Control Rig: Do not allow binding variables on non-root pins

diff --git a/Engine/Plugins/Animation/ControlRig/Source/ControlRigEditor/Private/ControlRigEditorModule.cpp b/Engine/Plugins/Animation/ControlRig/Source/ControlRigEditor/Private/ControlRigEditorModule.cpp
--- a/Engine/Plugins/Animation/ControlRig/Source/ControlRigEditor/Private/ControlRigEditorModule.cpp
+++ b/Engine/Plugins/Animation/ControlRig/Source/ControlRigEditor/Private/ControlRigEditorModule.cpp
@@ -1540,7 +1540,8 @@ void FControlRigEditorModule::GetContextMenuActions(const UControlRigGraphSchema
 						Cast<URigVMDispatchNode>(ModelPin->GetNode()) != nullptr || 
 						Cast<URigVMLibraryNode>(ModelPin->GetNode()) != nullptr)
 					{
-						if (ModelPin->GetDirection() == ERigVMPinDirection::Input && 
+						if (ModelPin->GetDirection() == ERigVMPinDirection::Input &&
+							ModelPin->IsRootPin() &&
 							bIsEditablePin)
 						{
 							if (!ModelPin->IsBoundToVariable())
diff --git a/Engine/Plugins/Runtime/RigVM/Source/RigVMDeveloper/Private/RigVMModel/RigVMController.cpp b/Engine/Plugins/Runtime/RigVM/Source/RigVMDeveloper/Private/RigVMModel/RigVMController.cpp
--- a/Engine/Plugins/Runtime/RigVM/Source/RigVMDeveloper/Private/RigVMModel/RigVMController.cpp
+++ b/Engine/Plugins/Runtime/RigVM/Source/RigVMDeveloper/Private/RigVMModel/RigVMController.cpp
@@ -4727,7 +4727,8 @@ URigVMCollapseNode* URigVMController::CollapseNodes(const TArray<URigVMNode*>& I
 #else
 	URigVMCollapseNode* CollapseNode = NewObject<URigVMCollapseNode>(Graph, *CollapseNodeName);
 #endif
-	CollapseNode->ContainedGraph = NewObject<URigVMGraph>(CollapseNode, TEXT("ContainedGraph"));
+	FString ContainedGraphName = CollapseNodeName + TEXT("_ContainedGraph");
+	CollapseNode->ContainedGraph = NewObject<URigVMGraph>(CollapseNode, *ContainedGraphName);
 
 #if UE_RIGVM_AGGREGATE_NODES_ENABLED
 	if (bIsAggregate)
@@ -9172,6 +9173,11 @@ bool URigVMController::BindPinToVariable(URigVMPin* InPin, const FString& InNewB
 		return false;
 	}
 
+	if (!InPin->IsRootPin())
+	{
+		return false;
+	}
+
 	FString VariableName = InNewBoundVariablePath, SegmentPath;
 	InNewBoundVariablePath.Split(TEXT("."), &VariableName, &SegmentPath);
 
@@ -11471,7 +11477,8 @@ URigVMLibraryNode* URigVMController::AddFunctionToLibrary(const FName& InFunctio
 
 	FString FunctionName = GetValidNodeName(InFunctionName.IsNone() ? FString(TEXT("Function")) : InFunctionName.ToString());
 	URigVMCollapseNode* CollapseNode = NewObject<URigVMCollapseNode>(Graph, *FunctionName);
-	CollapseNode->ContainedGraph = NewObject<URigVMGraph>(CollapseNode, TEXT("ContainedGraph"));
+	FString ContainedGraphName = FunctionName + TEXT("_ContainedGraph");
+	CollapseNode->ContainedGraph = NewObject<URigVMGraph>(CollapseNode, *ContainedGraphName);
 	CollapseNode->Position = InNodePosition;
 	Graph->Nodes.Add(CollapseNode);
 
