Control Rig: Fix compilation dirtying other assets

diff --git a/Engine/Plugins/Runtime/RigVM/Source/RigVMDeveloper/Private/RigVMCompiler/RigVMCompiler.cpp b/Engine/Plugins/Runtime/RigVM/Source/RigVMDeveloper/Private/RigVMCompiler/RigVMCompiler.cpp
--- a/Engine/Plugins/Runtime/RigVM/Source/RigVMDeveloper/Private/RigVMCompiler/RigVMCompiler.cpp
+++ b/Engine/Plugins/Runtime/RigVM/Source/RigVMDeveloper/Private/RigVMCompiler/RigVMCompiler.cpp
@@ -1183,8 +1183,6 @@ bool URigVMCompiler::CompileFunction(const URigVMLibraryNode* InLibraryNode, URi
 							}
 						}
 					}
-
-					ClientHost->GetRigVMClient()->UpdateFunctionReferences(Data->Header, true, false);
 				}
 		
 				Store->UpdateFunctionCompilationData(InLibraryNode->GetFunctionIdentifier(), *OutFunctionCompilationData);
