Control Rig: Removes query for bulk edit when adding/removing input argument variable node

diff --git a/Engine/Plugins/Animation/ControlRig/Source/ControlRigEditor/Private/Editor/ControlRigEditor.cpp b/Engine/Plugins/Animation/ControlRig/Source/ControlRigEditor/Private/Editor/ControlRigEditor.cpp
--- a/Engine/Plugins/Animation/ControlRig/Source/ControlRigEditor/Private/Editor/ControlRigEditor.cpp
+++ b/Engine/Plugins/Animation/ControlRig/Source/ControlRigEditor/Private/Editor/ControlRigEditor.cpp
@@ -229,6 +229,7 @@ FControlRigEditor::FControlRigEditor()
 	, RigHierarchyTabCount(0)
 	, HaltedAtNode(nullptr)
 	, bSuspendDetailsPanelRefresh(false)
+	, bAllowBulkEdits(false)
 	, bIsConstructionEventRunning(false)
 	, LastHierarchyHash(INDEX_NONE)
 {
@@ -5108,6 +5109,14 @@ void FControlRigEditor::OnRequestLocalizeFunctionDialog(FRigVMGraphFunctionIdent
 FRigVMController_BulkEditResult FControlRigEditor::OnRequestBulkEditDialog(UControlRigBlueprint* InBlueprint, URigVMController* InController,
 	URigVMLibraryNode* InFunction, ERigVMControllerBulkEditType InEditType)
 {
+	if (bAllowBulkEdits)
+	{
+		FRigVMController_BulkEditResult Result;
+		Result.bCanceled = false; 
+		Result.bSetupUndoRedo = false;
+		return Result;
+	}
+	
 	const TArray<FAssetData> FirstLevelReferenceAssets = InController->GetAffectedAssets(InEditType, false);
 	if(FirstLevelReferenceAssets.Num() == 0)
 	{
@@ -5123,6 +5132,12 @@ FRigVMController_BulkEditResult FControlRigEditor::OnRequestBulkEditDialog(UCont
 	FRigVMController_BulkEditResult Result;
 	Result.bCanceled = BulkEditDialog->ShowModal() == EAppReturnType::Cancel; 
 	Result.bSetupUndoRedo = false;
+
+	if (!Result.bCanceled)
+	{
+		bAllowBulkEdits = true;
+	}
+	
 	return Result;
 }
 
diff --git a/Engine/Plugins/Animation/ControlRig/Source/ControlRigEditor/Private/Editor/ControlRigEditor.h b/Engine/Plugins/Animation/ControlRig/Source/ControlRigEditor/Private/Editor/ControlRigEditor.h
--- a/Engine/Plugins/Animation/ControlRig/Source/ControlRigEditor/Private/Editor/ControlRigEditor.h
+++ b/Engine/Plugins/Animation/ControlRig/Source/ControlRigEditor/Private/Editor/ControlRigEditor.h
@@ -498,6 +498,8 @@ protected:
 
 	bool bSuspendDetailsPanelRefresh;
 
+	bool bAllowBulkEdits;
+
 	TArray<TStrongObjectPtr<UDetailsViewWrapperObject>> WrapperObjects;
 	TWeakObjectPtr<AStaticMeshActor> WeakGroundActorPtr;
 
diff --git a/Engine/Plugins/Animation/ControlRig/Source/ControlRigEditor/Private/Graph/SControlRigGraphNode.cpp b/Engine/Plugins/Animation/ControlRig/Source/ControlRigEditor/Private/Graph/SControlRigGraphNode.cpp
--- a/Engine/Plugins/Animation/ControlRig/Source/ControlRigEditor/Private/Graph/SControlRigGraphNode.cpp
+++ b/Engine/Plugins/Animation/ControlRig/Source/ControlRigEditor/Private/Graph/SControlRigGraphNode.cpp
@@ -70,6 +70,11 @@ void SControlRigGraphNode::Construct( const FArguments& InArgs )
 	{
 		return;
 	}
+	URigVMController* Controller = EdGraphNode->GetController();
+	if (!Controller)
+	{
+		return;
+	}
 
 	Blueprint = Cast<UControlRigBlueprint>(FBlueprintEditorUtils::FindBlueprintForNode(this->GraphNode));
 
@@ -78,7 +83,6 @@ void SControlRigGraphNode::Construct( const FArguments& InArgs )
 
 	SetIsEditable(false);
 
-	URigVMController* Controller = EdGraphNode->GetController();
 	Controller->OnModified().AddSP(this, &SControlRigGraphNode::HandleModifiedEvent);
 
 	UpdatePinTreeView();
diff --git a/Engine/Plugins/Runtime/RigVM/Source/RigVMDeveloper/Private/RigVMModel/RigVMController.cpp b/Engine/Plugins/Runtime/RigVM/Source/RigVMDeveloper/Private/RigVMModel/RigVMController.cpp
--- a/Engine/Plugins/Runtime/RigVM/Source/RigVMDeveloper/Private/RigVMModel/RigVMController.cpp
+++ b/Engine/Plugins/Runtime/RigVM/Source/RigVMDeveloper/Private/RigVMModel/RigVMController.cpp
@@ -802,15 +802,30 @@ void URigVMController::HandleModifiedEvent(ERigVMGraphNotifType InNotifType, URi
 				URigVMNode* Node = CastChecked<URigVMNode>(InSubject);
 				check(Node);
 
-				if(const URigVMLibraryNode* Function = FunctionLibrary->FindFunctionForNode(Node))
+				bool bIsLocal = false;
+				if (InNotifType == ERigVMGraphNotifType::VariableAdded || InNotifType == ERigVMGraphNotifType::VariableRemoved)
 				{
-					FunctionLibrary->ForEachReference(Function->GetFName(), [this](URigVMFunctionReferenceNode* Reference)
-                    {
-						if(const URigVMController* ReferenceController = GetControllerForGraph(Reference->GetGraph()))
+					if (URigVMVariableNode* VariableNode = Cast<URigVMVariableNode>(Node))
+					{
+						if (VariableNode->IsLocalVariable() || VariableNode->IsInputArgument())
 						{
-							ReferenceController->Notify(ERigVMGraphNotifType::VariableRemappingChanged, Reference);
+							bIsLocal = true;
 						}
-                    });
+					}
+				}
+
+				if (!bIsLocal)
+				{
+					if(const URigVMLibraryNode* Function = FunctionLibrary->FindFunctionForNode(Node))
+					{
+						FunctionLibrary->ForEachReference(Function->GetFName(), [this](URigVMFunctionReferenceNode* Reference)
+                    	{
+							if(const URigVMController* ReferenceController = GetControllerForGraph(Reference->GetGraph()))
+							{
+								ReferenceController->Notify(ERigVMGraphNotifType::VariableRemappingChanged, Reference);
+							}
+                    	});
+					}
 				}
 			}
 		}
@@ -1593,7 +1608,7 @@ URigVMVariableNode* URigVMController::AddVariableNode(const FName& InVariableNam
 			{
 				// Make sure there is no local variable with that name
 				bool bFoundLocalVariable = false;
-				for (FRigVMGraphVariableDescription& LocalVariable : OuterFunction->GetContainedGraph()->LocalVariables)
+				for (FRigVMGraphVariableDescription& LocalVariable : OuterFunction->GetContainedGraph()->GetLocalVariables(true))
 				{
 					if (LocalVariable.Name == InVariableName)
 					{
@@ -6289,7 +6304,7 @@ bool URigVMController::RemoveNodes(TArray<URigVMNode*> InNodes, bool bSetupUndoR
 						{
 							const FName VariableToRemove = VariableNode->GetVariableName();
 							bool bIsLocalVariable = false;
-							for (FRigVMGraphVariableDescription VariableDescription : OuterFunction->GetContainedGraph()->LocalVariables)
+							for (FRigVMGraphVariableDescription VariableDescription : OuterFunction->GetContainedGraph()->GetLocalVariables(true))
 							{
 								if (VariableDescription.Name == VariableToRemove)
 								{
