RigVM: ArrayGetNum dispatch returns incorrect opcode

diff --git a/Engine/Plugins/Runtime/RigVM/Source/RigVM/Public/RigVMFunctions/RigVMDispatch_Array.h b/Engine/Plugins/Runtime/RigVM/Source/RigVM/Public/RigVMFunctions/RigVMDispatch_Array.h
--- a/Engine/Plugins/Runtime/RigVM/Source/RigVM/Public/RigVMFunctions/RigVMDispatch_Array.h
+++ b/Engine/Plugins/Runtime/RigVM/Source/RigVM/Public/RigVMFunctions/RigVMDispatch_Array.h
@@ -96,7 +96,7 @@ struct RIGVM_API FRigVMDispatch_ArrayGetNum : public FRigVMDispatch_ArrayBase
 	GENERATED_BODY()
 
 public:
-	virtual ERigVMOpCode GetOpCode() const override { return ERigVMOpCode::ArrayReset; }
+	virtual ERigVMOpCode GetOpCode() const override { return ERigVMOpCode::ArrayGetNum; }
 	virtual FName GetArgumentNameForOperandIndex(int32 InOperandIndex, int32 InTotalOperands) const override;
 	virtual TArray<FRigVMTemplateArgument> GetArguments() const override;
 	virtual FRigVMTemplateTypeMap OnNewArgumentType(const FName& InArgumentName, TRigVMTypeIndex InTypeIndex) const override;
