Control Rig: Fix functions not updating compilation data when needed

diff --git a/Engine/Plugins/Runtime/RigVM/Source/RigVMDeveloper/Private/RigVMCompiler/RigVMCompiler.cpp b/Engine/Plugins/Runtime/RigVM/Source/RigVMDeveloper/Private/RigVMCompiler/RigVMCompiler.cpp
--- a/Engine/Plugins/Runtime/RigVM/Source/RigVMDeveloper/Private/RigVMCompiler/RigVMCompiler.cpp
+++ b/Engine/Plugins/Runtime/RigVM/Source/RigVMDeveloper/Private/RigVMCompiler/RigVMCompiler.cpp
@@ -366,7 +366,7 @@ bool URigVMCompiler::Compile(TArray<URigVMGraph*> InGraphs, URigVMController* In
 						// Clear compilation data if compiled with outdated dependency data
 						if (FunctionData->CompilationData.IsValid())
 						{
-							for (const TPair<FRigVMGraphFunctionIdentifier, uint32>& Pair : ReferenceNode->GetReferencedFunctionHeader().Dependencies)
+							for (const TPair<FRigVMGraphFunctionIdentifier, uint32>& Pair : FunctionData->Header.Dependencies)
 							{
 								if (IRigVMGraphFunctionHost* HostObj = Cast<IRigVMGraphFunctionHost>(Pair.Key.HostObject.ResolveObject()))
 								{
diff --git a/Engine/Plugins/Runtime/RigVM/Source/RigVMDeveloper/Private/RigVMModel/Nodes/RigVMLibraryNode.cpp b/Engine/Plugins/Runtime/RigVM/Source/RigVMDeveloper/Private/RigVMModel/Nodes/RigVMLibraryNode.cpp
--- a/Engine/Plugins/Runtime/RigVM/Source/RigVMDeveloper/Private/RigVMModel/Nodes/RigVMLibraryNode.cpp
+++ b/Engine/Plugins/Runtime/RigVM/Source/RigVMDeveloper/Private/RigVMModel/Nodes/RigVMLibraryNode.cpp
@@ -196,7 +196,7 @@ TMap<FRigVMGraphFunctionIdentifier, uint32> URigVMLibraryNode::GetDependencies()
     			Hash = Data->CompilationData.Hash;
     		}
     		Dependencies.Add(RefNode->GetReferencedFunctionHeader().LibraryPointer, Hash);
-    		for (const TPair<FRigVMGraphFunctionIdentifier, uint32>& Pair : RefNode->GetReferencedFunctionHeader().Dependencies)
+    		for (const TPair<FRigVMGraphFunctionIdentifier, uint32>& Pair : RefNode->GetReferencedFunctionData()->Header.Dependencies)
     		{
     			Dependencies.Add(Pair);
     		}
diff --git a/Engine/Plugins/Runtime/RigVM/Source/RigVMDeveloper/Private/RigVMModel/RigVMBuildData.cpp b/Engine/Plugins/Runtime/RigVM/Source/RigVMDeveloper/Private/RigVMModel/RigVMBuildData.cpp
--- a/Engine/Plugins/Runtime/RigVM/Source/RigVMDeveloper/Private/RigVMModel/RigVMBuildData.cpp
+++ b/Engine/Plugins/Runtime/RigVM/Source/RigVMDeveloper/Private/RigVMModel/RigVMBuildData.cpp
@@ -13,6 +13,9 @@
 static FDelayedAutoRegisterHelper GRigVMBuildDataSingletonHelper(EDelayedRegisterRunPhase::EndOfEngineInit, []() -> void
 {
 	URigVMBuildData::Get()->InitializeIfNeeded();
+
+	const FAssetRegistryModule& AssetRegistryModule = FModuleManager::LoadModuleChecked<FAssetRegistryModule>("AssetRegistry");
+	AssetRegistryModule.Get().OnAssetAdded().AddStatic(&URigVMBuildData::RegisterReferencesFromAsset);
 });
 
 FRigVMReferenceNodeData::FRigVMReferenceNodeData(URigVMFunctionReferenceNode* InReferenceNode)
@@ -109,32 +112,47 @@ void URigVMBuildData::InitializeIfNeeded()
 			// loop over all found assets
 			for(const FAssetData& ControlRigAssetData : ControlRigAssetDatas)
 			{
-				const FString ReferenceNodeDataString =
-					ControlRigAssetData.GetTagValueRef<FString>(ReferenceNodeDataProperty->GetFName());
-				if(ReferenceNodeDataString.IsEmpty())
+				RegisterReferencesFromAsset(ControlRigAssetData);
+			}
+		}
+	}
+	
+	bInitialized = true;
+}
+
+void URigVMBuildData::RegisterReferencesFromAsset(const FAssetData& InAssetData)
+{
+	URigVMBuildData* BuildData = URigVMBuildData::Get();
+
+	if (UClass* Class = InAssetData.GetClass())
+	{
+		const FArrayProperty* ReferenceNodeDataProperty =
+			  CastField<FArrayProperty>(Class->FindPropertyByName(TEXT("FunctionReferenceNodeData")));
+		if(ReferenceNodeDataProperty)
+		{
+			const FString ReferenceNodeDataString =
+						   InAssetData.GetTagValueRef<FString>(ReferenceNodeDataProperty->GetFName());
+			if(ReferenceNodeDataString.IsEmpty())
+			{
+				return;
+			}
+
+			// See if it has reference node data, and register the references
+			TArray<FRigVMReferenceNodeData> ReferenceNodeDatas;
+			ReferenceNodeDataProperty->ImportText_Direct(*ReferenceNodeDataString, &ReferenceNodeDatas, nullptr, EPropertyPortFlags::PPF_None);	
+			for(FRigVMReferenceNodeData& ReferenceNodeData : ReferenceNodeDatas)
+			{
+				if (ReferenceNodeData.ReferencedHeader.IsValid())
 				{
-					continue;
+					BuildData->RegisterFunctionReference(ReferenceNodeData.ReferencedHeader.LibraryPointer, ReferenceNodeData.GetReferenceNodeObjectPath());
 				}
-
-				// See if it has reference node data, and register the references
-				TArray<FRigVMReferenceNodeData> ReferenceNodeDatas;
-				ReferenceNodeDataProperty->ImportText_Direct(*ReferenceNodeDataString, &ReferenceNodeDatas, nullptr, EPropertyPortFlags::PPF_None);	
-				for(FRigVMReferenceNodeData& ReferenceNodeData : ReferenceNodeDatas)
+				else if (!ReferenceNodeData.ReferencedFunctionPath_DEPRECATED.IsEmpty())
 				{
-					if (ReferenceNodeData.ReferencedHeader.IsValid())
-					{
-						RegisterFunctionReference(ReferenceNodeData.ReferencedHeader.LibraryPointer, ReferenceNodeData.GetReferenceNodeObjectPath());
-					}
-					else if (!ReferenceNodeData.ReferencedFunctionPath_DEPRECATED.IsEmpty())
-					{
-						RegisterFunctionReference(ReferenceNodeData);							
-					}
+					BuildData->RegisterFunctionReference(ReferenceNodeData);
 				}
 			}
 		}
 	}
-	
-	bInitialized = true;
 }
 
 const FRigVMFunctionReferenceArray* URigVMBuildData::FindFunctionReferences(const FRigVMGraphFunctionIdentifier& InFunction) const
diff --git a/Engine/Plugins/Runtime/RigVM/Source/RigVMDeveloper/Public/RigVMModel/RigVMBuildData.h b/Engine/Plugins/Runtime/RigVM/Source/RigVMDeveloper/Public/RigVMModel/RigVMBuildData.h
--- a/Engine/Plugins/Runtime/RigVM/Source/RigVMDeveloper/Public/RigVMModel/RigVMBuildData.h
+++ b/Engine/Plugins/Runtime/RigVM/Source/RigVMDeveloper/Public/RigVMModel/RigVMBuildData.h
@@ -74,6 +74,8 @@ public:
 	// Looks for all function references (in RigVMClientHost metadata) and initializes the URigVMBuildData
 	void InitializeIfNeeded();
 
+	static void RegisterReferencesFromAsset(const FAssetData& AssetData);
+
 	// Returns the list of references for a given function definition
 	const FRigVMFunctionReferenceArray* FindFunctionReferences(const FRigVMGraphFunctionIdentifier& InFunction) const;
 
