RigVM: Use Dispatch Constructor to set factory struct, new make struct, break struct and constant dispatches

diff --git a/Engine/Plugins/Animation/ControlRig/Source/ControlRig/Private/Units/Animation/RigUnit_AnimAttribute.h b/Engine/Plugins/Animation/ControlRig/Source/ControlRig/Private/Units/Animation/RigUnit_AnimAttribute.h
--- a/Engine/Plugins/Animation/ControlRig/Source/ControlRig/Private/Units/Animation/RigUnit_AnimAttribute.h
+++ b/Engine/Plugins/Animation/ControlRig/Source/ControlRig/Private/Units/Animation/RigUnit_AnimAttribute.h
@@ -182,6 +182,11 @@ struct CONTROLRIG_API FRigDispatch_GetAnimAttribute: public FRigDispatch_AnimAtt
 {
 	GENERATED_BODY()
 
+	FRigDispatch_GetAnimAttribute()
+	{
+		FactoryScriptStruct = StaticStruct();
+	}
+
 	virtual TArray<FRigVMTemplateArgument> GetArguments() const override;
 	virtual FRigVMTemplateTypeMap OnNewArgumentType(const FName& InArgumentName, TRigVMTypeIndex InTypeIndex) const override;
 
@@ -255,6 +260,12 @@ USTRUCT(meta=(DisplayName="Set Animation Attribute"))
 struct CONTROLRIG_API FRigDispatch_SetAnimAttribute: public FRigDispatch_AnimAttributeBase
 {
 	GENERATED_BODY()
+
+	FRigDispatch_SetAnimAttribute()
+	{
+		FactoryScriptStruct = StaticStruct();
+	}
+
 	virtual bool IsSet() const override { return true; }
 	virtual TArray<FRigVMTemplateArgument> GetArguments() const override;
 	virtual TArray<FRigVMExecuteArgument> GetExecuteArguments_Impl(const FRigVMDispatchContext& InContext) const override;
diff --git a/Engine/Plugins/Animation/ControlRig/Source/ControlRig/Private/Units/Core/RigUnit_UserData.h b/Engine/Plugins/Animation/ControlRig/Source/ControlRig/Private/Units/Core/RigUnit_UserData.h
--- a/Engine/Plugins/Animation/ControlRig/Source/ControlRig/Private/Units/Core/RigUnit_UserData.h
+++ b/Engine/Plugins/Animation/ControlRig/Source/ControlRig/Private/Units/Core/RigUnit_UserData.h
@@ -14,6 +14,10 @@ struct FRigDispatch_GetUserData : public FRigDispatchFactory
 	GENERATED_BODY()
 
 public:
+	FRigDispatch_GetUserData()
+	{
+		FactoryScriptStruct = StaticStruct();
+	}
 
 #if WITH_EDITOR
 	virtual FString GetArgumentMetaData(const FName& InArgumentName, const FName& InMetaDataKey) const override;
diff --git a/Engine/Plugins/Animation/ControlRig/Source/ControlRig/Private/Units/Hierarchy/RigUnit_Metadata.h b/Engine/Plugins/Animation/ControlRig/Source/ControlRig/Private/Units/Hierarchy/RigUnit_Metadata.h
--- a/Engine/Plugins/Animation/ControlRig/Source/ControlRig/Private/Units/Hierarchy/RigUnit_Metadata.h
+++ b/Engine/Plugins/Animation/ControlRig/Source/ControlRig/Private/Units/Hierarchy/RigUnit_Metadata.h
@@ -10,6 +10,11 @@ struct CONTROLRIG_API FRigDispatch_MetadataBase : public FRigDispatchFactory
 {
 	GENERATED_BODY()
 
+	FRigDispatch_MetadataBase()
+	{
+		FactoryScriptStruct = StaticStruct();
+	}
+
 #if WITH_EDITOR
 	virtual FString GetNodeTitle(const FRigVMTemplateTypeMap& InTypes) const override;;
 #endif
diff --git a/Engine/Plugins/Runtime/RigVM/Source/RigVM/Private/RigVMCore/RigVMDispatchFactory.cpp b/Engine/Plugins/Runtime/RigVM/Source/RigVM/Private/RigVMCore/RigVMDispatchFactory.cpp
--- a/Engine/Plugins/Runtime/RigVM/Source/RigVM/Private/RigVMCore/RigVMDispatchFactory.cpp
+++ b/Engine/Plugins/Runtime/RigVM/Source/RigVM/Private/RigVMCore/RigVMDispatchFactory.cpp
@@ -233,21 +233,27 @@ bool FRigVMDispatchFactory::CopyProperty(const FProperty* InTargetProperty, uint
 
 const FRigVMTemplate* FRigVMDispatchFactory::GetTemplate() const
 {
+	// make sure to rely on the instance of this factory that's stored under the registry
+	FRigVMRegistry& Registry = FRigVMRegistry::Get();
+	const FName FactoryName = GetFactoryName();
+	const FRigVMDispatchFactory* ThisFactory = Registry.FindDispatchFactory(FactoryName);
+	if(ThisFactory != this)
+	{
+		return ThisFactory->GetTemplate();
+	}
+
 	static bool bIsDispatchingTemplate = false;
 	if(bIsDispatchingTemplate)
 	{
-			return nullptr;
+		return nullptr;
 	}
 	TGuardValue<bool> ReEntryGuard(bIsDispatchingTemplate, true);
-	
+
 	if(CachedTemplate)
 	{
 		return CachedTemplate;
 	}
-	
-	FRigVMRegistry& Registry = FRigVMRegistry::Get();
 
-	const FName FactoryName = GetFactoryName();
 	const TArray<FRigVMTemplateArgument> Arguments = GetArguments();
 
 	// we don't allow execute types on arguments
@@ -288,3 +294,11 @@ const FRigVMTemplate* FRigVMDispatchFactory::GetTemplate() const
 	return CachedTemplate;
 }
 
+FName FRigVMDispatchFactory::GetTemplateNotation() const
+{
+	if(const FRigVMTemplate* Template = GetTemplate())
+	{
+		return Template->GetNotation();
+	}
+	return NAME_None;
+}
diff --git a/Engine/Plugins/Runtime/RigVM/Source/RigVM/Private/RigVMCore/RigVMStruct.cpp b/Engine/Plugins/Runtime/RigVM/Source/RigVM/Private/RigVMCore/RigVMStruct.cpp
--- a/Engine/Plugins/Runtime/RigVM/Source/RigVM/Private/RigVMCore/RigVMStruct.cpp
+++ b/Engine/Plugins/Runtime/RigVM/Source/RigVM/Private/RigVMCore/RigVMStruct.cpp
@@ -295,6 +295,8 @@ const FName FRigVMStruct::IconMetaName("Icon");
 const FName FRigVMStruct::KeywordsMetaName("Keywords");
 const FName FRigVMStruct::TemplateNameMetaName = FRigVMRegistry::TemplateNameMetaName;
 const FName FRigVMStruct::FixedSizeArrayMetaName("FixedSizeArray");
+const FName FRigVMStruct::ShowOnlySubPinsMetaName("ShowOnlySubPins");
+const FName FRigVMStruct::HideSubPinsMetaName("HideSubPins");
 const FName FRigVMStruct::ArraySizeMetaName("ArraySize");
 const FName FRigVMStruct::AggregateMetaName("Aggregate");
 const FName FRigVMStruct::ExpandPinByDefaultMetaName("ExpandByDefault");
diff --git a/Engine/Plugins/Runtime/RigVM/Source/RigVM/Private/RigVMFunctions/RigVMDispatch_Constant.cpp b/Engine/Plugins/Runtime/RigVM/Source/RigVM/Private/RigVMFunctions/RigVMDispatch_Constant.cpp
new file mode 100644
--- /dev/null
+++ b/Engine/Plugins/Runtime/RigVM/Source/RigVM/Private/RigVMFunctions/RigVMDispatch_Constant.cpp
@@ -0,0 +1,77 @@
+// Copyright Epic Games, Inc. All Rights Reserved.
+
+#include "RigVMFunctions/RigVMDispatch_Constant.h"
+#include "RigVMCore/RigVMRegistry.h"
+#include "RigVMCore/RigVMStruct.h"
+#include "RigVMStringUtils.h"
+
+#include UE_INLINE_GENERATED_CPP_BY_NAME(RigVMDispatch_Constant)
+#define LOCTEXT_NAMESPACE "RigVMDispatch_Constant"
+
+const FName FRigVMDispatch_Constant::ValueName = TEXT("Value");
+
+TArray<FRigVMTemplateArgument> FRigVMDispatch_Constant::GetArguments() const
+{
+	const TArray<FRigVMTemplateArgument::ETypeCategory> Categories = {
+		FRigVMTemplateArgument::ETypeCategory_SingleAnyValue
+	};
+	return {
+		{ValueName, ERigVMPinDirection::IO, Categories}
+	};
+}
+
+FRigVMTemplateTypeMap FRigVMDispatch_Constant::OnNewArgumentType(const FName& InArgumentName, TRigVMTypeIndex InTypeIndex) const
+{
+	return {
+		{ValueName, InTypeIndex},
+	};
+}
+
+#if WITH_EDITOR
+
+FString FRigVMDispatch_Constant::GetNodeTitle(const FRigVMTemplateTypeMap& InTypes) const
+{
+	const FString SuperNodeTitle = FRigVMDispatch_CoreBase::GetNodeTitle(InTypes);
+	if(const TRigVMTypeIndex* TypeIndexPtr = InTypes.Find(ValueName))
+	{
+		const TRigVMTypeIndex& TypeIndex = *TypeIndexPtr;
+		if(TypeIndex != INDEX_NONE)
+		{
+			const FRigVMRegistry& Registry = FRigVMRegistry::Get();
+			if(!Registry.IsWildCardType(TypeIndex))
+			{
+				const FRigVMTemplateArgumentType& Type = Registry.GetType(TypeIndex);
+				static constexpr TCHAR Format[] = TEXT("%s %s");
+				if(Type.CPPTypeObject)
+				{
+					return FString::Printf(Format, *SuperNodeTitle, *Type.CPPTypeObject->GetName());
+				}
+				return FString::Printf(Format, *SuperNodeTitle, *Type.CPPType.ToString());
+			}
+		}
+	}
+	return SuperNodeTitle;
+}
+
+FText FRigVMDispatch_Constant::GetNodeTooltip(const FRigVMTemplateTypeMap& InTypes) const
+{
+	return LOCTEXT("ConstantToolTip", "Provides a constant value to the graph");
+}
+
+FString FRigVMDispatch_Constant::GetArgumentMetaData(const FName& InArgumentName, const FName& InMetaDataKey) const
+{
+	if(InArgumentName == ValueName && InMetaDataKey == FRigVMStruct::ConstantMetaName)
+	{
+		return TrueString;
+	}
+	return FRigVMDispatch_CoreBase::GetArgumentMetaData(InArgumentName, InMetaDataKey);
+}
+
+#endif
+
+void FRigVMDispatch_Constant::Execute(FRigVMExtendedExecuteContext& InContext, FRigVMMemoryHandleArray Handles)
+{
+	// nothing to do here - it's a constant
+}
+
+#undef LOCTEXT_NAMESPACE
diff --git a/Engine/Plugins/Runtime/RigVM/Source/RigVM/Private/RigVMFunctions/RigVMDispatch_MakeStruct.cpp b/Engine/Plugins/Runtime/RigVM/Source/RigVM/Private/RigVMFunctions/RigVMDispatch_MakeStruct.cpp
new file mode 100644
--- /dev/null
+++ b/Engine/Plugins/Runtime/RigVM/Source/RigVM/Private/RigVMFunctions/RigVMDispatch_MakeStruct.cpp
@@ -0,0 +1,107 @@
+// Copyright Epic Games, Inc. All Rights Reserved.
+
+#include "RigVMFunctions/RigVMDispatch_MakeStruct.h"
+#include "RigVMCore/RigVMRegistry.h"
+#include "RigVMCore/RigVMStruct.h"
+#include "RigVMStringUtils.h"
+
+#include UE_INLINE_GENERATED_CPP_BY_NAME(RigVMDispatch_MakeStruct)
+#define LOCTEXT_NAMESPACE "RigVMDispatch_MakeStruct"
+
+const FName FRigVMDispatch_MakeStruct::ElementsName = TEXT("Elements");
+const FName FRigVMDispatch_MakeStruct::StructName = TEXT("Struct");
+
+TArray<FRigVMTemplateArgument> FRigVMDispatch_MakeStruct::GetArguments() const
+{
+	const TArray<FRigVMTemplateArgument::ETypeCategory> Categories = {
+		FRigVMTemplateArgument::ETypeCategory_SingleScriptStructValue
+	};
+	return {
+		{ElementsName, ERigVMPinDirection::Input, Categories},
+		{StructName, ERigVMPinDirection::Output, Categories}
+	};
+}
+
+FRigVMTemplateTypeMap FRigVMDispatch_MakeStruct::OnNewArgumentType(const FName& InArgumentName, TRigVMTypeIndex InTypeIndex) const
+{
+	return {
+		{ElementsName, InTypeIndex},
+		{StructName, InTypeIndex}
+	};
+}
+
+#if WITH_EDITOR
+
+FString FRigVMDispatch_MakeStruct::GetNodeTitle(const FRigVMTemplateTypeMap& InTypes) const
+{
+	const FString SuperNodeTitle = FRigVMDispatch_CoreBase::GetNodeTitle(InTypes);
+	if(const TRigVMTypeIndex* StructTypePtr = InTypes.Find(StructName))
+	{
+		const TRigVMTypeIndex& StructType = *StructTypePtr;
+		if(StructType != INDEX_NONE)
+		{
+			const FRigVMRegistry& Registry = FRigVMRegistry::Get();
+			if(!Registry.IsWildCardType(StructType))
+			{
+				if(UScriptStruct* Struct = Cast<UScriptStruct>(Registry.GetType(StructType).CPPTypeObject))
+				{
+					static constexpr TCHAR Format[] = TEXT("%s %s");
+					return FString::Printf(Format, *SuperNodeTitle, *Struct->GetName());
+				}
+			}
+		}
+	}
+	return SuperNodeTitle;
+}
+
+FText FRigVMDispatch_MakeStruct::GetNodeTooltip(const FRigVMTemplateTypeMap& InTypes) const
+{
+	return LOCTEXT("MakeStructToolTip", "Composes a struct out of its elements");
+}
+
+FString FRigVMDispatch_MakeStruct::GetArgumentMetaData(const FName& InArgumentName, const FName& InMetaDataKey) const
+{
+	if(InArgumentName == ElementsName && InMetaDataKey == FRigVMStruct::ShowOnlySubPinsMetaName)
+	{
+		return TrueString;
+	}
+	if(InArgumentName == StructName && InMetaDataKey == FRigVMStruct::HideSubPinsMetaName)
+	{
+		return TrueString;
+	}
+	return FRigVMDispatch_CoreBase::GetArgumentMetaData(InArgumentName, InMetaDataKey);
+}
+
+#endif
+
+void FRigVMDispatch_MakeStruct::Execute(FRigVMExtendedExecuteContext& InContext, FRigVMMemoryHandleArray Handles)
+{
+	const FProperty* SourceProperty = Handles[0].GetProperty();
+	const FProperty* TargetProperty = Handles[1].GetProperty();
+	check(SourceProperty->SameType(TargetProperty));
+	const uint8* SourceMemory = Handles[0].GetData();
+	uint8* TargetMemory = Handles[1].GetData();
+	URigVMMemoryStorage::CopyProperty(TargetProperty, TargetMemory, SourceProperty, SourceMemory);
+}
+
+TArray<FRigVMTemplateArgument> FRigVMDispatch_BreakStruct::GetArguments() const
+{
+	const TArray<FRigVMTemplateArgument::ETypeCategory> Categories = {
+		FRigVMTemplateArgument::ETypeCategory_SingleScriptStructValue
+	};
+	return {
+		{StructName, ERigVMPinDirection::Input, Categories},
+		{ElementsName, ERigVMPinDirection::Output, Categories}
+	};
+}
+
+#if WITH_EDITOR
+
+FText FRigVMDispatch_BreakStruct::GetNodeTooltip(const FRigVMTemplateTypeMap& InTypes) const
+{
+	return LOCTEXT("BreakStructToolTip", "Decomposes a struct into its elements");
+}
+
+#endif
+
+#undef LOCTEXT_NAMESPACE
diff --git a/Engine/Plugins/Runtime/RigVM/Source/RigVM/Public/RigVMCore/RigVMDispatchFactory.h b/Engine/Plugins/Runtime/RigVM/Source/RigVM/Public/RigVMCore/RigVMDispatchFactory.h
--- a/Engine/Plugins/Runtime/RigVM/Source/RigVM/Public/RigVMCore/RigVMDispatchFactory.h
+++ b/Engine/Plugins/Runtime/RigVM/Source/RigVM/Public/RigVMCore/RigVMDispatchFactory.h
@@ -152,6 +152,9 @@ public:
 	// builds and returns the template
 	RIGVM_API const FRigVMTemplate* GetTemplate() const;
 
+	// returns the name of the factory template
+	RIGVM_API FName GetTemplateNotation() const;
+
 	// returns the name of the permutation for a given set of types
 	RIGVM_API FString GetPermutationName(const FRigVMTemplateTypeMap& InTypes) const;
 
diff --git a/Engine/Plugins/Runtime/RigVM/Source/RigVM/Public/RigVMCore/RigVMStruct.h b/Engine/Plugins/Runtime/RigVM/Source/RigVM/Public/RigVMCore/RigVMStruct.h
--- a/Engine/Plugins/Runtime/RigVM/Source/RigVM/Public/RigVMCore/RigVMStruct.h
+++ b/Engine/Plugins/Runtime/RigVM/Source/RigVM/Public/RigVMCore/RigVMStruct.h
@@ -248,6 +248,8 @@ public:
 	static const FName KeywordsMetaName;
 	static const FName TemplateNameMetaName;
 	static const FName FixedSizeArrayMetaName;
+	static const FName ShowOnlySubPinsMetaName;
+	static const FName HideSubPinsMetaName;
 	static const FName ArraySizeMetaName;
 	static const FName AggregateMetaName;
 	static const FName ExpandPinByDefaultMetaName;
diff --git a/Engine/Plugins/Runtime/RigVM/Source/RigVM/Public/RigVMFunctions/RigVMDispatch_Array.h b/Engine/Plugins/Runtime/RigVM/Source/RigVM/Public/RigVMFunctions/RigVMDispatch_Array.h
--- a/Engine/Plugins/Runtime/RigVM/Source/RigVM/Public/RigVMFunctions/RigVMDispatch_Array.h
+++ b/Engine/Plugins/Runtime/RigVM/Source/RigVM/Public/RigVMFunctions/RigVMDispatch_Array.h
@@ -38,6 +38,8 @@ protected:
 	static const FName RatioName;
 	static const FName ResultName;
 	static const FName& CompletedName;
+
+	friend class URigVMController;
 };
 
 USTRUCT(meta=(Abstract))
@@ -55,6 +57,11 @@ struct RIGVM_API FRigVMDispatch_ArrayMake : public FRigVMDispatch_ArrayBase
 	GENERATED_BODY()
 
 public:
+	FRigVMDispatch_ArrayMake()
+	{
+		FactoryScriptStruct = StaticStruct();
+	}
+
 	virtual TArray<FRigVMTemplateArgument> GetArguments() const override;
 	virtual FRigVMTemplateTypeMap OnNewArgumentType(const FName& InArgumentName, TRigVMTypeIndex InTypeIndex) const override;
 #if WITH_EDITOR
@@ -76,6 +83,11 @@ struct RIGVM_API FRigVMDispatch_ArrayReset : public FRigVMDispatch_ArrayBaseMuta
 	GENERATED_BODY()
 
 public:
+	FRigVMDispatch_ArrayReset()
+	{
+		FactoryScriptStruct = StaticStruct();
+	}
+
 	virtual ERigVMOpCode GetOpCode() const override { return ERigVMOpCode::ArrayReset; }
 	virtual FName GetArgumentNameForOperandIndex(int32 InOperandIndex, int32 InTotalOperands) const override;
 	virtual TArray<FRigVMTemplateArgument> GetArguments() const override;
@@ -96,6 +108,11 @@ struct RIGVM_API FRigVMDispatch_ArrayGetNum : public FRigVMDispatch_ArrayBase
 	GENERATED_BODY()
 
 public:
+	FRigVMDispatch_ArrayGetNum()
+	{
+		FactoryScriptStruct = StaticStruct();
+	}
+
 	virtual ERigVMOpCode GetOpCode() const override { return ERigVMOpCode::ArrayGetNum; }
 	virtual FName GetArgumentNameForOperandIndex(int32 InOperandIndex, int32 InTotalOperands) const override;
 	virtual TArray<FRigVMTemplateArgument> GetArguments() const override;
@@ -116,6 +133,11 @@ struct RIGVM_API FRigVMDispatch_ArraySetNum : public FRigVMDispatch_ArrayBaseMut
 	GENERATED_BODY()
 
 public:
+	FRigVMDispatch_ArraySetNum()
+	{
+		FactoryScriptStruct = StaticStruct();
+	}
+
 	virtual ERigVMOpCode GetOpCode() const override { return ERigVMOpCode::ArraySetNum; }
 	virtual FName GetArgumentNameForOperandIndex(int32 InOperandIndex, int32 InTotalOperands) const override;
 	virtual TArray<FRigVMTemplateArgument> GetArguments() const override;
@@ -136,6 +158,11 @@ struct RIGVM_API FRigVMDispatch_ArrayGetAtIndex : public FRigVMDispatch_ArrayBas
 	GENERATED_BODY()
 
 public:
+	FRigVMDispatch_ArrayGetAtIndex()
+	{
+		FactoryScriptStruct = StaticStruct();
+	}
+
 	virtual ERigVMOpCode GetOpCode() const override { return ERigVMOpCode::ArrayGetAtIndex; }
 	virtual FName GetArgumentNameForOperandIndex(int32 InOperandIndex, int32 InTotalOperands) const override;
 	virtual TArray<FRigVMTemplateArgument> GetArguments() const override;
@@ -156,6 +183,11 @@ struct RIGVM_API FRigVMDispatch_ArraySetAtIndex : public FRigVMDispatch_ArrayBas
 	GENERATED_BODY()
 
 public:
+	FRigVMDispatch_ArraySetAtIndex()
+	{
+		FactoryScriptStruct = StaticStruct();
+	}
+
 	virtual ERigVMOpCode GetOpCode() const override { return ERigVMOpCode::ArraySetAtIndex; }
 	virtual FName GetArgumentNameForOperandIndex(int32 InOperandIndex, int32 InTotalOperands) const override;
 	virtual TArray<FRigVMTemplateArgument> GetArguments() const override;
@@ -176,6 +208,11 @@ struct RIGVM_API FRigVMDispatch_ArrayAdd : public FRigVMDispatch_ArraySetAtIndex
 	GENERATED_BODY()
 
 public:
+	FRigVMDispatch_ArrayAdd()
+	{
+		FactoryScriptStruct = StaticStruct();
+	}
+
 	virtual ERigVMOpCode GetOpCode() const override { return ERigVMOpCode::ArrayAdd; }
 	virtual FName GetArgumentNameForOperandIndex(int32 InOperandIndex, int32 InTotalOperands) const override;
 	virtual TArray<FRigVMTemplateArgument> GetArguments() const override;
@@ -195,6 +232,11 @@ struct RIGVM_API FRigVMDispatch_ArrayInsert : public FRigVMDispatch_ArraySetAtIn
 	GENERATED_BODY()
 
 public:
+	FRigVMDispatch_ArrayInsert()
+	{
+		FactoryScriptStruct = StaticStruct();
+	}
+
 	virtual ERigVMOpCode GetOpCode() const override { return ERigVMOpCode::ArrayInsert; }
 #if WITH_EDITOR
    	virtual FText GetNodeTooltip(const FRigVMTemplateTypeMap& InTypes) const override;
@@ -212,6 +254,11 @@ struct RIGVM_API FRigVMDispatch_ArrayRemove : public FRigVMDispatch_ArrayBaseMut
 	GENERATED_BODY()
 
 public:
+	FRigVMDispatch_ArrayRemove()
+	{
+		FactoryScriptStruct = StaticStruct();
+	}
+
 	virtual ERigVMOpCode GetOpCode() const override { return ERigVMOpCode::ArrayRemove; }
 	virtual FName GetArgumentNameForOperandIndex(int32 InOperandIndex, int32 InTotalOperands) const override;
 	virtual TArray<FRigVMTemplateArgument> GetArguments() const override;
@@ -232,6 +279,11 @@ struct RIGVM_API FRigVMDispatch_ArrayReverse : public FRigVMDispatch_ArrayReset
 	GENERATED_BODY()
 
 public:
+	FRigVMDispatch_ArrayReverse()
+	{
+		FactoryScriptStruct = StaticStruct();
+	}
+
 	virtual ERigVMOpCode GetOpCode() const override { return ERigVMOpCode::ArrayReverse; }
 #if WITH_EDITOR
    	virtual FText GetNodeTooltip(const FRigVMTemplateTypeMap& InTypes) const override;
@@ -250,6 +302,11 @@ struct RIGVM_API FRigVMDispatch_ArrayFind : public FRigVMDispatch_ArrayBase
 	GENERATED_BODY()
 
 public:
+	FRigVMDispatch_ArrayFind()
+	{
+		FactoryScriptStruct = StaticStruct();
+	}
+
 	virtual ERigVMOpCode GetOpCode() const override { return ERigVMOpCode::ArrayFind; }
 	virtual FName GetArgumentNameForOperandIndex(int32 InOperandIndex, int32 InTotalOperands) const override;
 	virtual TArray<FRigVMTemplateArgument> GetArguments() const override;
@@ -270,6 +327,11 @@ struct RIGVM_API FRigVMDispatch_ArrayAppend : public FRigVMDispatch_ArrayBaseMut
 	GENERATED_BODY()
 
 public:
+	FRigVMDispatch_ArrayAppend()
+	{
+		FactoryScriptStruct = StaticStruct();
+	}
+
 	virtual ERigVMOpCode GetOpCode() const override { return ERigVMOpCode::ArrayAppend; }
 	virtual FName GetArgumentNameForOperandIndex(int32 InOperandIndex, int32 InTotalOperands) const override;
 	virtual TArray<FRigVMTemplateArgument> GetArguments() const override;
@@ -290,6 +352,11 @@ struct RIGVM_API FRigVMDispatch_ArrayClone : public FRigVMDispatch_ArrayBase
 	GENERATED_BODY()
 
 public:
+	FRigVMDispatch_ArrayClone()
+	{
+		FactoryScriptStruct = StaticStruct();
+	}
+
 	virtual ERigVMOpCode GetOpCode() const override { return ERigVMOpCode::ArrayClone; }
 	virtual FName GetArgumentNameForOperandIndex(int32 InOperandIndex, int32 InTotalOperands) const override;
 	virtual TArray<FRigVMTemplateArgument> GetArguments() const override;
@@ -310,6 +377,11 @@ struct RIGVM_API FRigVMDispatch_ArrayUnion : public FRigVMDispatch_ArrayAppend
 	GENERATED_BODY()
 
 public:
+	FRigVMDispatch_ArrayUnion()
+	{
+		FactoryScriptStruct = StaticStruct();
+	}
+
 	virtual ERigVMOpCode GetOpCode() const override { return ERigVMOpCode::ArrayUnion; }
 #if WITH_EDITOR
    	virtual FText GetNodeTooltip(const FRigVMTemplateTypeMap& InTypes) const override;
@@ -327,6 +399,11 @@ struct RIGVM_API FRigVMDispatch_ArrayDifference : public FRigVMDispatch_ArrayBas
 	GENERATED_BODY()
 
 public:
+	FRigVMDispatch_ArrayDifference()
+	{
+		FactoryScriptStruct = StaticStruct();
+	}
+
 	virtual ERigVMOpCode GetOpCode() const override { return ERigVMOpCode::ArrayDifference; }
 	virtual FName GetArgumentNameForOperandIndex(int32 InOperandIndex, int32 InTotalOperands) const override;
 	virtual TArray<FRigVMTemplateArgument> GetArguments() const override;
@@ -347,6 +424,11 @@ struct RIGVM_API FRigVMDispatch_ArrayIntersection : public FRigVMDispatch_ArrayD
 	GENERATED_BODY()
 
 public:
+	FRigVMDispatch_ArrayIntersection()
+	{
+		FactoryScriptStruct = StaticStruct();
+	}
+
 	virtual ERigVMOpCode GetOpCode() const override { return ERigVMOpCode::ArrayIntersection; }
 #if WITH_EDITOR
    	virtual FText GetNodeTooltip(const FRigVMTemplateTypeMap& InTypes) const override;
@@ -364,6 +446,11 @@ struct RIGVM_API FRigVMDispatch_ArrayIterator : public FRigVMDispatch_ArrayBaseM
 	GENERATED_BODY()
 
 public:
+	FRigVMDispatch_ArrayIterator()
+	{
+		FactoryScriptStruct = StaticStruct();
+	}
+
 	virtual ERigVMOpCode GetOpCode() const override { return ERigVMOpCode::ArrayIterator; }
 	virtual FName GetArgumentNameForOperandIndex(int32 InOperandIndex, int32 InTotalOperands) const override;
 	virtual TArray<FRigVMTemplateArgument> GetArguments() const override;
diff --git a/Engine/Plugins/Runtime/RigVM/Source/RigVM/Public/RigVMFunctions/RigVMDispatch_Constant.h b/Engine/Plugins/Runtime/RigVM/Source/RigVM/Public/RigVMFunctions/RigVMDispatch_Constant.h
new file mode 100644
--- /dev/null
+++ b/Engine/Plugins/Runtime/RigVM/Source/RigVM/Public/RigVMFunctions/RigVMDispatch_Constant.h
@@ -0,0 +1,37 @@
+// Copyright Epic Games, Inc. All Rights Reserved.
+
+#pragma once
+
+#include "RigVMCore/RigVMDispatchFactory.h"
+#include "RigVMDispatch_Core.h"
+#include "RigVMCore/RigVMByteCode.h"
+#include "RigVMDispatch_Constant.generated.h"
+
+USTRUCT(meta=(DisplayName = "Constant", Category = "Core", Keywords = "Value", NodeColor = "1,1,1,1"))
+struct RIGVM_API FRigVMDispatch_Constant : public FRigVMDispatch_CoreBase
+{
+	GENERATED_BODY()
+
+public:
+	FRigVMDispatch_Constant()
+	{
+		FactoryScriptStruct = StaticStruct();
+	}
+
+	virtual TArray<FRigVMTemplateArgument> GetArguments() const override;
+	virtual FRigVMTemplateTypeMap OnNewArgumentType(const FName& InArgumentName, TRigVMTypeIndex InTypeIndex) const override;
+#if WITH_EDITOR
+	virtual FString GetNodeTitle(const FRigVMTemplateTypeMap& InTypes) const override;
+	virtual FText GetNodeTooltip(const FRigVMTemplateTypeMap& InTypes) const override;
+	virtual FString GetArgumentMetaData(const FName& InArgumentName, const FName& InMetaDataKey) const override;
+#endif
+	virtual bool IsSingleton() const override { return true; }
+	
+protected:
+	virtual FRigVMFunctionPtr GetDispatchFunctionImpl(const FRigVMTemplateTypeMap& InTypes) const override { return &FRigVMDispatch_Constant::Execute; }
+	static void Execute(FRigVMExtendedExecuteContext& InContext, FRigVMMemoryHandleArray Handles);
+
+	static const FName ValueName;
+
+	friend class URigVMController;
+};
diff --git a/Engine/Plugins/Runtime/RigVM/Source/RigVM/Public/RigVMFunctions/RigVMDispatch_Core.h b/Engine/Plugins/Runtime/RigVM/Source/RigVM/Public/RigVMFunctions/RigVMDispatch_Core.h
--- a/Engine/Plugins/Runtime/RigVM/Source/RigVM/Public/RigVMFunctions/RigVMDispatch_Core.h
+++ b/Engine/Plugins/Runtime/RigVM/Source/RigVM/Public/RigVMFunctions/RigVMDispatch_Core.h
@@ -21,6 +21,11 @@ struct RIGVM_API FRigVMDispatch_CoreEquals : public FRigVMDispatch_CoreBase
 
 public:
 
+	FRigVMDispatch_CoreEquals()
+	{
+		FactoryScriptStruct = StaticStruct();
+	}
+
 	virtual FName GetArgumentNameForOperandIndex(int32 InOperandIndex, int32 InTotalOperands) const override;
 	virtual TArray<FRigVMTemplateArgument> GetArguments() const override;
 	virtual FRigVMTemplateTypeMap OnNewArgumentType(const FName& InArgumentName, TRigVMTypeIndex InTypeIndex) const override;
@@ -90,6 +95,11 @@ struct RIGVM_API FRigVMDispatch_CoreNotEquals : public FRigVMDispatch_CoreEquals
 {
 	GENERATED_BODY()
 
+	FRigVMDispatch_CoreNotEquals()
+	{
+		FactoryScriptStruct = StaticStruct();
+	}
+
 	// we are inheriting everything from the equals dispatch,
 	// and due to the check of the factory within FRigVMDispatch_CoreEquals::Execute we can
 	// rely on that completely. we only need this class for the displayname and
diff --git a/Engine/Plugins/Runtime/RigVM/Source/RigVM/Public/RigVMFunctions/RigVMDispatch_If.h b/Engine/Plugins/Runtime/RigVM/Source/RigVM/Public/RigVMFunctions/RigVMDispatch_If.h
--- a/Engine/Plugins/Runtime/RigVM/Source/RigVM/Public/RigVMFunctions/RigVMDispatch_If.h
+++ b/Engine/Plugins/Runtime/RigVM/Source/RigVM/Public/RigVMFunctions/RigVMDispatch_If.h
@@ -16,6 +16,11 @@ struct RIGVM_API FRigVMDispatch_If : public FRigVMDispatch_CoreBase
 
 public:
 
+	FRigVMDispatch_If()
+	{
+		FactoryScriptStruct = StaticStruct();
+	}
+
 	virtual FName GetArgumentNameForOperandIndex(int32 InOperandIndex, int32 InTotalOperands) const override;
 	virtual TArray<FRigVMTemplateArgument> GetArguments() const override;
 	virtual FRigVMTemplateTypeMap OnNewArgumentType(const FName& InArgumentName, TRigVMTypeIndex InTypeIndex) const override;
diff --git a/Engine/Plugins/Runtime/RigVM/Source/RigVM/Public/RigVMFunctions/RigVMDispatch_MakeStruct.h b/Engine/Plugins/Runtime/RigVM/Source/RigVM/Public/RigVMFunctions/RigVMDispatch_MakeStruct.h
new file mode 100644
--- /dev/null
+++ b/Engine/Plugins/Runtime/RigVM/Source/RigVM/Public/RigVMFunctions/RigVMDispatch_MakeStruct.h
@@ -0,0 +1,46 @@
+// Copyright Epic Games, Inc. All Rights Reserved.
+
+#pragma once
+
+#include "RigVMCore/RigVMDispatchFactory.h"
+#include "RigVMDispatch_Core.h"
+#include "RigVMCore/RigVMByteCode.h"
+#include "RigVMDispatch_MakeStruct.generated.h"
+
+USTRUCT(meta=(DisplayName = "Make", Category = "Core", Keywords = "Compose,Composition,Create", NodeColor = "1,1,1,1"))
+struct RIGVM_API FRigVMDispatch_MakeStruct : public FRigVMDispatch_CoreBase
+{
+	GENERATED_BODY()
+
+public:
+	virtual TArray<FRigVMTemplateArgument> GetArguments() const override;
+	virtual FRigVMTemplateTypeMap OnNewArgumentType(const FName& InArgumentName, TRigVMTypeIndex InTypeIndex) const override;
+#if WITH_EDITOR
+	virtual FString GetNodeTitle(const FRigVMTemplateTypeMap& InTypes) const override;
+	virtual FText GetNodeTooltip(const FRigVMTemplateTypeMap& InTypes) const override;
+	virtual FString GetArgumentMetaData(const FName& InArgumentName, const FName& InMetaDataKey) const override;
+#endif
+
+protected:
+	virtual FRigVMFunctionPtr GetDispatchFunctionImpl(const FRigVMTemplateTypeMap& InTypes) const override { return &FRigVMDispatch_MakeStruct::Execute; }
+	static void Execute(FRigVMExtendedExecuteContext& InContext, FRigVMMemoryHandleArray Handles);
+
+	static const FName ElementsName;
+	static const FName StructName;
+
+	friend struct FRigVMDispatch_BreakStruct;
+	friend class URigVMController;
+};
+
+USTRUCT(meta=(DisplayName = "Break", Keywords = "Decompose,Decomposition"))
+struct RIGVM_API FRigVMDispatch_BreakStruct : public FRigVMDispatch_MakeStruct
+{
+	GENERATED_BODY()
+
+public:
+	virtual TArray<FRigVMTemplateArgument> GetArguments() const override;
+#if WITH_EDITOR
+	virtual FText GetNodeTooltip(const FRigVMTemplateTypeMap& InTypes) const override;
+#endif
+};
+
diff --git a/Engine/Plugins/Runtime/RigVM/Source/RigVM/Public/RigVMFunctions/RigVMDispatch_Print.h b/Engine/Plugins/Runtime/RigVM/Source/RigVM/Public/RigVMFunctions/RigVMDispatch_Print.h
--- a/Engine/Plugins/Runtime/RigVM/Source/RigVM/Public/RigVMFunctions/RigVMDispatch_Print.h
+++ b/Engine/Plugins/Runtime/RigVM/Source/RigVM/Public/RigVMFunctions/RigVMDispatch_Print.h
@@ -15,6 +15,11 @@ struct RIGVM_API FRigVMDispatch_Print : public FRigVMDispatchFactory
 
 public:
 
+	FRigVMDispatch_Print()
+	{
+		FactoryScriptStruct = StaticStruct();
+	}
+
 	virtual FName GetArgumentNameForOperandIndex(int32 InOperandIndex, int32 InTotalOperands) const override;
 	virtual TArray<FRigVMTemplateArgument> GetArguments() const override;
 	virtual TArray<FRigVMExecuteArgument> GetExecuteArguments_Impl(const FRigVMDispatchContext& InContext) const override;
diff --git a/Engine/Plugins/Runtime/RigVM/Source/RigVM/Public/RigVMFunctions/RigVMDispatch_Select.h b/Engine/Plugins/Runtime/RigVM/Source/RigVM/Public/RigVMFunctions/RigVMDispatch_Select.h
--- a/Engine/Plugins/Runtime/RigVM/Source/RigVM/Public/RigVMFunctions/RigVMDispatch_Select.h
+++ b/Engine/Plugins/Runtime/RigVM/Source/RigVM/Public/RigVMFunctions/RigVMDispatch_Select.h
@@ -16,6 +16,11 @@ struct RIGVM_API FRigVMDispatch_SelectInt32 : public FRigVMDispatch_CoreBase
 
 public:
 
+	FRigVMDispatch_SelectInt32()
+	{
+		FactoryScriptStruct = StaticStruct();
+	}
+
 	virtual FName GetArgumentNameForOperandIndex(int32 InOperandIndex, int32 InTotalOperands) const override;
 	virtual TArray<FRigVMTemplateArgument> GetArguments() const override;
 	virtual FRigVMTemplateTypeMap OnNewArgumentType(const FName& InArgumentName, TRigVMTypeIndex InTypeIndex) const override;
diff --git a/Engine/Plugins/Runtime/RigVM/Source/RigVM/Public/RigVMFunctions/RigVMDispatch_Switch.h b/Engine/Plugins/Runtime/RigVM/Source/RigVM/Public/RigVMFunctions/RigVMDispatch_Switch.h
--- a/Engine/Plugins/Runtime/RigVM/Source/RigVM/Public/RigVMFunctions/RigVMDispatch_Switch.h
+++ b/Engine/Plugins/Runtime/RigVM/Source/RigVM/Public/RigVMFunctions/RigVMDispatch_Switch.h
@@ -17,7 +17,9 @@ struct RIGVM_API FRigVMDispatch_SwitchInt32 : public FRigVMDispatch_CoreBase
 public:
 
 	FRigVMDispatch_SwitchInt32()
-	{}
+	{
+		FactoryScriptStruct = StaticStruct();
+	}
 
 	virtual FName GetArgumentNameForOperandIndex(int32 InOperandIndex, int32 InTotalOperands) const override;
 	virtual TArray<FRigVMTemplateArgument> GetArguments() const override;
