Control Rig: Get external variables of a function reference form the function data, instead of the header

diff --git a/Engine/Plugins/Animation/ControlRig/Source/ControlRigDeveloper/Private/Graph/ControlRigGraphNode.cpp b/Engine/Plugins/Animation/ControlRig/Source/ControlRigDeveloper/Private/Graph/ControlRigGraphNode.cpp
--- a/Engine/Plugins/Animation/ControlRig/Source/ControlRigDeveloper/Private/Graph/ControlRigGraphNode.cpp
+++ b/Engine/Plugins/Animation/ControlRig/Source/ControlRigDeveloper/Private/Graph/ControlRigGraphNode.cpp
@@ -1211,7 +1211,7 @@ void UControlRigGraphNode::AllocateDefaultPins()
 	{
 		if(FunctionReferenceNode->RequiresVariableRemapping())
 		{
-			TArray<FRigVMExternalVariable> CurrentExternalVariables = FunctionReferenceNode->GetReferencedFunctionHeader().ExternalVariables;
+			TArray<FRigVMExternalVariable> CurrentExternalVariables = FunctionReferenceNode->GetExternalVariables(false);
 			for(const FRigVMExternalVariable& CurrentExternalVariable : CurrentExternalVariables)
 			{
 				ExternalVariables.Add(MakeShared<FRigVMExternalVariable>(CurrentExternalVariable));
diff --git a/Engine/Plugins/Animation/ControlRig/Source/ControlRigEditor/Private/Graph/SControlRigGraphPinVariableBinding.cpp b/Engine/Plugins/Animation/ControlRig/Source/ControlRigEditor/Private/Graph/SControlRigGraphPinVariableBinding.cpp
--- a/Engine/Plugins/Animation/ControlRig/Source/ControlRigEditor/Private/Graph/SControlRigGraphPinVariableBinding.cpp
+++ b/Engine/Plugins/Animation/ControlRig/Source/ControlRigEditor/Private/Graph/SControlRigGraphPinVariableBinding.cpp
@@ -181,7 +181,7 @@ bool SControlRigVariableBinding::OnCanBindProperty(FProperty* InProperty) const
 		}
 		else if(FunctionReferenceNode && !InnerVariableName.IsNone())
 		{
-			TArray<FRigVMExternalVariable> InnerVariables = FunctionReferenceNode->GetReferencedFunctionHeader().ExternalVariables;
+			TArray<FRigVMExternalVariable> InnerVariables = FunctionReferenceNode->GetExternalVariables(false);
 			for(const FRigVMExternalVariable& InnerVariable : InnerVariables)
 			{
 				if(InnerVariable.Name == InnerVariableName)
diff --git a/Engine/Plugins/Runtime/RigVM/Source/RigVMDeveloper/Private/RigVMModel/Nodes/RigVMFunctionReferenceNode.cpp b/Engine/Plugins/Runtime/RigVM/Source/RigVMDeveloper/Private/RigVMModel/Nodes/RigVMFunctionReferenceNode.cpp
--- a/Engine/Plugins/Runtime/RigVM/Source/RigVMDeveloper/Private/RigVMModel/Nodes/RigVMFunctionReferenceNode.cpp
+++ b/Engine/Plugins/Runtime/RigVM/Source/RigVMDeveloper/Private/RigVMModel/Nodes/RigVMFunctionReferenceNode.cpp
@@ -49,7 +49,7 @@ bool URigVMFunctionReferenceNode::RequiresVariableRemappingInternal(TArray<FRigV
 		
 	if(bHostedInDifferencePackage)
 	{
-		InnerVariables = ReferencedFunctionHeader.ExternalVariables;
+		InnerVariables = GetExternalVariables(false);
 		if(InnerVariables.Num() == 0)
 		{
 			return false;
@@ -92,7 +92,10 @@ TArray<FRigVMExternalVariable> URigVMFunctionReferenceNode::GetExternalVariables
 	
 	if(!bRemapped)
 	{
-		Variables = ReferencedFunctionHeader.ExternalVariables;
+		if (const FRigVMGraphFunctionData* FunctionData = GetReferencedFunctionData())
+		{
+			Variables = FunctionData->Header.ExternalVariables;
+		}
 	}
 	else
 	{
diff --git a/Engine/Plugins/Runtime/RigVM/Source/RigVMDeveloper/Private/RigVMModel/RigVMController.cpp b/Engine/Plugins/Runtime/RigVM/Source/RigVMDeveloper/Private/RigVMModel/RigVMController.cpp
--- a/Engine/Plugins/Runtime/RigVM/Source/RigVMDeveloper/Private/RigVMModel/RigVMController.cpp
+++ b/Engine/Plugins/Runtime/RigVM/Source/RigVMDeveloper/Private/RigVMModel/RigVMController.cpp
@@ -11410,7 +11410,8 @@ bool URigVMController::SetRemappedVariable(URigVMFunctionReferenceNode* InFuncti
 
 	FRigVMExternalVariable InnerExternalVariable;
 	{
-		if (FRigVMExternalVariable* Variable = InFunctionRefNode->ReferencedFunctionHeader.ExternalVariables.FindByPredicate([InInnerVariableName](const FRigVMExternalVariable& Variable)
+		TArray<FRigVMExternalVariable> Variables = InFunctionRefNode->GetExternalVariables(false);
+		if (FRigVMExternalVariable* Variable = Variables.FindByPredicate([InInnerVariableName](const FRigVMExternalVariable& Variable)
 		{
 			return Variable.Name == InInnerVariableName;
 		}))
