Control rig: Fix crash when opening context menu in ControlRigExample

diff --git a/Engine/Plugins/Runtime/RigVM/Source/RigVMDeveloper/Private/RigVMModel/Nodes/RigVMLibraryNode.cpp b/Engine/Plugins/Runtime/RigVM/Source/RigVMDeveloper/Private/RigVMModel/Nodes/RigVMLibraryNode.cpp
--- a/Engine/Plugins/Runtime/RigVM/Source/RigVMDeveloper/Private/RigVMModel/Nodes/RigVMLibraryNode.cpp
+++ b/Engine/Plugins/Runtime/RigVM/Source/RigVMDeveloper/Private/RigVMModel/Nodes/RigVMLibraryNode.cpp
@@ -191,15 +191,16 @@ TMap<FRigVMGraphFunctionIdentifier, uint32> URigVMLibraryNode::GetDependencies()
     	if (const URigVMFunctionReferenceNode* RefNode = Cast<URigVMFunctionReferenceNode>(Node))
     	{
     		uint32 Hash = 0;
-    		if (const FRigVMGraphFunctionData* Data = RefNode->GetReferencedFunctionData())
+    		const FRigVMGraphFunctionData* FunctionData = RefNode->GetReferencedFunctionData();
+    		if (FunctionData)
     		{
-    			Hash = Data->CompilationData.Hash;
+    			Hash = FunctionData->CompilationData.Hash;
+    			for (const TPair<FRigVMGraphFunctionIdentifier, uint32>& Pair : FunctionData->Header.Dependencies)
+    			{
+    				Dependencies.Add(Pair);
+    			}
     		}
     		Dependencies.Add(RefNode->GetReferencedFunctionHeader().LibraryPointer, Hash);
-    		for (const TPair<FRigVMGraphFunctionIdentifier, uint32>& Pair : RefNode->GetReferencedFunctionData()->Header.Dependencies)
-    		{
-    			Dependencies.Add(Pair);
-    		}
     	}
     	if (const URigVMCollapseNode* CollapseNode = Cast<URigVMCollapseNode>(Node))
     	{
