Control Rig: Mark obsolete work memory classes with CLASS_NewerVersionExists to ensure that stray objects that reference that class won't be gathered for package saving.

#rnx

diff --git a/Engine/Plugins/Runtime/RigVM/Source/RigVM/Private/RigVMCore/RigVMMemoryStorage.cpp b/Engine/Plugins/Runtime/RigVM/Source/RigVM/Private/RigVMCore/RigVMMemoryStorage.cpp
--- a/Engine/Plugins/Runtime/RigVM/Source/RigVM/Private/RigVMCore/RigVMMemoryStorage.cpp
+++ b/Engine/Plugins/Runtime/RigVM/Source/RigVM/Private/RigVMCore/RigVMMemoryStorage.cpp
@@ -528,12 +528,14 @@ URigVMMemoryStorageGeneratorClass* URigVMMemoryStorageGeneratorClass::CreateStor
 	// does not guarantee that condition. As a result we will do OldClass->ConditionalPostLoad() to ensure the OldClass
 	// is fully loaded, ready to be replaced.
 	
-	UObject* OldClass = StaticFindObjectFastInternal( /*Class=*/ NULL, Package, *ClassName, true, RF_NoFlags, EInternalObjectFlags::None);
-	if(OldClass)
+	UObject* OldClassObject = StaticFindObjectFastInternal( /*Class=*/ NULL, Package, *ClassName, true, RF_NoFlags, EInternalObjectFlags::None);
+	if(OldClassObject)
 	{
 		// ensure the OldClass is completely loaded so that we can remove it and replace it with the new class
 		// otherwise we get an assertion trying to replace objects currently being loaded
-		OldClass->ConditionalPostLoad();
+		OldClassObject->ConditionalPostLoad();
+
+		UClass* OldClass = Cast<UClass>(OldClassObject);
 		
 		FString DiscardedMemoryClassName;
 		static const TCHAR DiscardedMemoryClassTemplate[] = TEXT("DiscardedMemoryClassTemplate_%d");
@@ -548,7 +550,13 @@ URigVMMemoryStorageGeneratorClass* URigVMMemoryStorageGeneratorClass::CreateStor
 		}
 		while (DiscardedMemoryClassIndex < INT_MAX);
 
+		OldClass->ClassFlags |= CLASS_NewerVersionExists;
 		OldClass->Rename(*DiscardedMemoryClassName, GetTransientPackage(), REN_ForceNoResetLoaders | REN_DoNotDirty | REN_DontCreateRedirectors | REN_NonTransactional);
+
+		if (OldClass->ClassDefaultObject)
+		{
+			OldClass->ClassDefaultObject->Rename(nullptr, GetTransientPackage(), REN_ForceNoResetLoaders | REN_DoNotDirty | REN_DontCreateRedirectors | REN_NonTransactional);
+		}
 	}
 
 	// create the new class
