Control Rig: Fix test_create_rename_duplicate failing

diff --git a/Engine/Plugins/Runtime/RigVM/Source/RigVMDeveloper/Private/RigVMModel/RigVMClient.cpp b/Engine/Plugins/Runtime/RigVM/Source/RigVMDeveloper/Private/RigVMModel/RigVMClient.cpp
--- a/Engine/Plugins/Runtime/RigVM/Source/RigVMDeveloper/Private/RigVMModel/RigVMClient.cpp
+++ b/Engine/Plugins/Runtime/RigVM/Source/RigVMDeveloper/Private/RigVMModel/RigVMClient.cpp
@@ -73,6 +73,8 @@ void FRigVMClient::Reset()
 	Models.Reset();
 	Controllers.Reset();
 	FunctionLibrary = nullptr;
+
+	ResetActionStack();
 }
 
 URigVMGraph* FRigVMClient::GetDefaultModel() const
@@ -709,19 +711,25 @@ URigVMActionStack* FRigVMClient::GetOrCreateActionStack()
 {
 	if(UObject* Outer = GetOuter())
 	{
-		if(ActionStack == nullptr)
+		if (ActionStack && ActionStack->GetOuter() != Outer)
 		{
-			ActionStack = NewObject<URigVMActionStack>(Outer, TEXT("ActionStack"));
+			ResetActionStack();
 		}
-		else if(ActionStack->GetOuter() != Outer)
+		
+		if(ActionStack == nullptr)
 		{
-			// make sure the action stack is parented correctly
-			ActionStack->Rename(nullptr, Outer, REN_ForceNoResetLoaders | REN_DoNotDirty | REN_DontCreateRedirectors | REN_NonTransactional);
+			ActionStack = NewObject<URigVMActionStack>(Outer, TEXT("ActionStack"));
 		}
 	}
 	return ActionStack;
 }
 
+void FRigVMClient::ResetActionStack()
+{
+	DestroyObject(ActionStack);
+	ActionStack = nullptr;
+}
+
 FName FRigVMClient::GetUniqueName(const FName& InDesiredName) const
 {
 	return GetUniqueName(GetOuter(), InDesiredName);
diff --git a/Engine/Plugins/Runtime/RigVM/Source/RigVMDeveloper/Public/RigVMModel/RigVMClient.h b/Engine/Plugins/Runtime/RigVM/Source/RigVMDeveloper/Public/RigVMModel/RigVMClient.h
--- a/Engine/Plugins/Runtime/RigVM/Source/RigVMDeveloper/Public/RigVMModel/RigVMClient.h
+++ b/Engine/Plugins/Runtime/RigVM/Source/RigVMDeveloper/Public/RigVMModel/RigVMClient.h
@@ -85,6 +85,7 @@ public:
 	FRigVMClient()
 		: ExecuteContextStruct(nullptr)
 		, FunctionLibrary(nullptr)
+		, ActionStack(nullptr)
 		, bSuspendNotifications(false)
 		, bIgnoreModelNotifications(false)
 		, OuterClientHost(nullptr)
@@ -188,6 +189,7 @@ private:
 
 	URigVMController* CreateController(const URigVMGraph* InModel);
 	URigVMActionStack* GetOrCreateActionStack();
+	void ResetActionStack();
 
 	UPROPERTY(transient)
 	TObjectPtr<UScriptStruct> ExecuteContextStruct;
