Control Rig: Fix crash when changing a function on a duplicated blueprint

diff --git a/Engine/Plugins/Runtime/RigVM/Source/RigVMDeveloper/Private/RigVMModel/RigVMClient.cpp b/Engine/Plugins/Runtime/RigVM/Source/RigVMDeveloper/Private/RigVMModel/RigVMClient.cpp
--- a/Engine/Plugins/Runtime/RigVM/Source/RigVMDeveloper/Private/RigVMModel/RigVMClient.cpp
+++ b/Engine/Plugins/Runtime/RigVM/Source/RigVMDeveloper/Private/RigVMModel/RigVMClient.cpp
@@ -843,8 +843,16 @@ void FRigVMClient::HandleGraphModifiedEvent(ERigVMGraphNotifType InNotifType, UR
 #if WITH_EDITOR
 
 	IRigVMClientHost* ClientHost = Cast<IRigVMClientHost>(GetOuter());
-	IRigVMGraphFunctionHost* FunctionHost = ClientHost->GetRigVMGraphFunctionHost();
-	FRigVMGraphFunctionStore* FunctionStore = FunctionHost->GetRigVMGraphFunctionStore();
+	IRigVMGraphFunctionHost* FunctionHost = nullptr;
+	FRigVMGraphFunctionStore* FunctionStore = nullptr;
+	if (ClientHost)
+	{
+		FunctionHost = ClientHost->GetRigVMGraphFunctionHost();
+		if (FunctionHost)
+		{
+			FunctionStore = FunctionHost->GetRigVMGraphFunctionStore();
+		}
+	}
 
 	switch (InNotifType)
 	{
@@ -857,7 +865,10 @@ void FRigVMClient::HandleGraphModifiedEvent(ERigVMGraphNotifType InNotifType, UR
 				{
 					if (GetOuter()->Implements<URigVMClientHost>())
 					{
-						FunctionStore->AddFunction(CollapseNode->GetFunctionHeader(FunctionHost), false);
+						if (FunctionStore && FunctionHost)
+						{
+							FunctionStore->AddFunction(CollapseNode->GetFunctionHeader(FunctionHost), false);
+						}
 					}
 				}
 			}
@@ -881,7 +892,10 @@ void FRigVMClient::HandleGraphModifiedEvent(ERigVMGraphNotifType InNotifType, UR
 				{
 					if (GetOuter()->Implements<URigVMClientHost>())
 					{
-						FunctionStore->RemoveFunction(FRigVMGraphFunctionIdentifier (Cast<UObject>(FunctionHost), CollapseNode));
+						if (FunctionStore && FunctionHost)
+						{
+							FunctionStore->RemoveFunction(FRigVMGraphFunctionIdentifier (Cast<UObject>(FunctionHost), CollapseNode));
+						}
 					}
 				}
 			}
@@ -905,7 +919,7 @@ void FRigVMClient::HandleGraphModifiedEvent(ERigVMGraphNotifType InNotifType, UR
 			if(URigVMLibraryNode* LibraryNode = Cast<URigVMNode>(InSubject)->FindFunctionForNode())
 			{
 				DirtyGraphFunctionCompilationData(LibraryNode);
-				UpdateExternalVariablesForFunction(LibraryNode);				
+				UpdateExternalVariablesForFunction(LibraryNode);
 			}
 			break;
 		}
@@ -1014,7 +1028,7 @@ void FRigVMClient::HandleGraphModifiedEvent(ERigVMGraphNotifType InNotifType, UR
 					}
 				}
 			}
-			break;	
+			break;
 		}
 		
 		
