Control Rig: Fix python unit tests failing

diff --git a/Engine/Plugins/Runtime/RigVM/Source/RigVMDeveloper/Private/RigVMModel/RigVMController.cpp b/Engine/Plugins/Runtime/RigVM/Source/RigVMDeveloper/Private/RigVMModel/RigVMController.cpp
--- a/Engine/Plugins/Runtime/RigVM/Source/RigVMDeveloper/Private/RigVMModel/RigVMController.cpp
+++ b/Engine/Plugins/Runtime/RigVM/Source/RigVMDeveloper/Private/RigVMModel/RigVMController.cpp
@@ -18586,6 +18586,18 @@ URigVMActionStack* URigVMController::GetActionStack() const
 	{
 		return WeakActionStack.Get();
 	}
+	if (bIsRunningUnitTest)
+	{
+		URigVMActionStack* ActionStack = NewObject<URigVMActionStack>(GetTransientPackage(), TEXT("ActionStack"));
+		WeakActionStack = ActionStack;
+		ActionStackHandle = ActionStack->OnModified().AddLambda([&](ERigVMGraphNotifType NotifType, URigVMGraph* InGraph, UObject* InSubject) -> void {
+			if(InGraph == GetGraph())
+			{
+				Notify(NotifType, InSubject);
+			}
+		});
+		return WeakActionStack.Get();
+	}
 	checkNoEntry();
 	return nullptr;
 }
diff --git a/Engine/Plugins/Runtime/RigVM/Source/RigVMDeveloper/Public/RigVMModel/RigVMController.h b/Engine/Plugins/Runtime/RigVM/Source/RigVMDeveloper/Public/RigVMModel/RigVMController.h
--- a/Engine/Plugins/Runtime/RigVM/Source/RigVMDeveloper/Public/RigVMModel/RigVMController.h
+++ b/Engine/Plugins/Runtime/RigVM/Source/RigVMDeveloper/Public/RigVMModel/RigVMController.h
@@ -1353,8 +1353,8 @@ private:
 	UPROPERTY(transient)
 	TArray<TObjectPtr<URigVMGraph>> Graphs;
 
-	TWeakObjectPtr<URigVMActionStack> WeakActionStack;
-	FDelegateHandle ActionStackHandle;
+	mutable TWeakObjectPtr<URigVMActionStack> WeakActionStack;
+	mutable FDelegateHandle ActionStackHandle;
 
 	bool bSuspendNotifications;
 	bool bReportWarningsAndErrors;
