Control Rig: Fix debug vm execution

diff --git a/Engine/Plugins/Runtime/RigVM/Source/RigVM/Private/RigVMCore/RigVM.cpp b/Engine/Plugins/Runtime/RigVM/Source/RigVM/Private/RigVMCore/RigVM.cpp
--- a/Engine/Plugins/Runtime/RigVM/Source/RigVM/Private/RigVMCore/RigVM.cpp
+++ b/Engine/Plugins/Runtime/RigVM/Source/RigVM/Private/RigVMCore/RigVM.cpp
@@ -1509,18 +1509,21 @@ bool URigVM::Initialize(TArrayView<URigVMMemoryStorage*> Memory)
 					{
 						Property->CopyCompleteValue_InContainer(WorkMemory, WorkMemoryCDO);
 #if UE_RIGVM_DEBUG_EXECUTION
-						FString DefaultValue;
-						const uint8* PropertyMemory = Property->ContainerPtrToValuePtr<uint8>(WorkMemory);
-						Property->ExportText_Direct(
-							DefaultValue,
-							PropertyMemory,
-							PropertyMemory,
-							nullptr,
-							PPF_None,
-							nullptr);
-
-						UE_LOG(LogRigVM, Warning, TEXT("Property %s defaults to '%s'."), *Property->GetName(), *DefaultValue);
-						UE_LOG(LogRigVM, Warning, TEXT("Property %s defaults to '%s'."), *Property->GetName(), *DefaultValue);
+						if (Context.GetPublicData<FRigVMExecuteContext>().bDebugExecution)
+						{
+							FString DefaultValue;
+							const uint8* PropertyMemory = Property->ContainerPtrToValuePtr<uint8>(WorkMemory);
+							Property->ExportText_Direct(
+								DefaultValue,
+								PropertyMemory,
+								PropertyMemory,
+								nullptr,
+								PPF_None,
+								nullptr);
+
+							UE_LOG(LogRigVM, Warning, TEXT("Property %s defaults to '%s'."), *Property->GetName(), *DefaultValue);
+							UE_LOG(LogRigVM, Warning, TEXT("Property %s defaults to '%s'."), *Property->GetName(), *DefaultValue);
+						}
 #endif
 					}
 				}
@@ -1662,15 +1665,17 @@ ERigVMExecuteResult URigVM::Execute(TArrayView<URigVMMemoryStorage*> Memory, con
 	StartProfiling();
 	
 #if UE_RIGVM_DEBUG_EXECUTION
-	ContextPublicData.DebugMemoryString;
-	ContextPublicData.InstanceOpCodeEnum = StaticEnum<ERigVMOpCode>();
-	URigVMMemoryStorage* LiteralMemory = GetLiteralMemory(false);
-	ContextPublicData.DebugMemoryString = FString("\n\nLiteral Memory\n\n");
-	for (int32 PropertyIndex=0; PropertyIndex<LiteralMemory->Num(); ++PropertyIndex)
+	if (ContextPublicData.bDebugExecution)
 	{
-		ContextPublicData.DebugMemoryString += FString::Printf(TEXT("%s: %s\n"), *LiteralMemory->GetProperties()[PropertyIndex]->GetFullName(), *LiteralMemory->GetDataAsString(PropertyIndex));				
+		ContextPublicData.InstanceOpCodeEnum = StaticEnum<ERigVMOpCode>();
+		URigVMMemoryStorage* LiteralMemory = GetLiteralMemory(false);
+		ContextPublicData.DebugMemoryString = FString("\n\nLiteral Memory\n\n");
+		for (int32 PropertyIndex=0; PropertyIndex<LiteralMemory->Num(); ++PropertyIndex)
+		{
+			ContextPublicData.DebugMemoryString += FString::Printf(TEXT("%s: %s\n"), *LiteralMemory->GetProperties()[PropertyIndex]->GetFullName(), *LiteralMemory->GetDataAsString(PropertyIndex));				
+		}
+		ContextPublicData.DebugMemoryString += FString(TEXT("\n\nWork Memory\n\n"));
 	}
-	ContextPublicData.DebugMemoryString += FString(TEXT("\n\nWork Memory\n\n"));
 	
 #endif
 	
@@ -1717,7 +1722,10 @@ ERigVMExecuteResult URigVM::ExecuteInstructions(int32 InFirstInstruction, int32
 		if (ShouldHaltAtInstruction(CurrentEntryName, ContextPublicData.InstructionIndex))
 		{
 #if UE_RIGVM_DEBUG_EXECUTION
-			ContextPublicData.Log(EMessageSeverity::Info, ContextPublicData.DebugMemoryString);					
+			if (ContextPublicData.bDebugExecution)
+			{
+				ContextPublicData.Log(EMessageSeverity::Info, ContextPublicData.DebugMemoryString);
+			}
 #endif
 			// we'll recursively exit all invoked
 			// entries here.
@@ -1748,33 +1756,35 @@ ERigVMExecuteResult URigVM::ExecuteInstructions(int32 InFirstInstruction, int32
 
 #if WITH_EDITOR
 #if UE_RIGVM_DEBUG_EXECUTION
-		if (Instruction.OpCode >= ERigVMOpCode::Execute_0_Operands && Instruction.OpCode <= ERigVMOpCode::Execute_64_Operands)
+		if (ContextPublicData.bDebugExecution)
 		{
-			const FRigVMExecuteOp& Op = ByteCode.GetOpAt<FRigVMExecuteOp>(Instruction);
-			FRigVMOperandArray Operands = ByteCode.GetOperandsForExecuteOp(Instructions[ContextPublicData.InstructionIndex]);
+			if (Instruction.OpCode >= ERigVMOpCode::Execute_0_Operands && Instruction.OpCode <= ERigVMOpCode::Execute_64_Operands)
+			{
+				const FRigVMExecuteOp& Op = ByteCode.GetOpAt<FRigVMExecuteOp>(Instruction);
+				FRigVMOperandArray Operands = ByteCode.GetOperandsForExecuteOp(Instructions[ContextPublicData.InstructionIndex]);
 
-			TArray<FString> Labels;
-			for (const FRigVMOperand& Operand : Operands)
+				TArray<FString> Labels;
+				for (const FRigVMOperand& Operand : Operands)
+				{
+					Labels.Add(GetOperandLabel(Operand));
+				}
+
+				ContextPublicData.DebugMemoryString += FString::Printf(TEXT("Instruction %d: %s(%s)\n"), ContextPublicData.InstructionIndex, *FunctionNames[Op.FunctionIndex].ToString(), *FString::Join(Labels, TEXT(", ")));
+			}
+			else if(Instruction.OpCode == ERigVMOpCode::Copy)
 			{
-				Labels.Add(GetOperandLabel(Operand));
+				static auto FormatFunction = [](const FString& RegisterName, const FString& RegisterOffsetName) -> FString
+				{
+					return FString::Printf(TEXT("%s%s"), *RegisterName, *RegisterOffsetName);
+				};
+				const FRigVMCopyOp& Op = ByteCode.GetOpAt<FRigVMCopyOp>(Instruction);
+				ContextPublicData.DebugMemoryString += FString::Printf(TEXT("Instruction %d: Copy %s -> %s\n"), ContextPublicData.InstructionIndex, *GetOperandLabel(Op.Source, FormatFunction), *GetOperandLabel(Op.Target, FormatFunction));
 			}
-
-			ContextPublicData.DebugMemoryString += FString::Printf(TEXT("Instruction %d: %s(%s)\n"), ContextPublicData.InstructionIndex, *FunctionNames[Op.FunctionIndex].ToString(), *FString::Join(Labels, TEXT(", ")));
-		}
-		else if(Instruction.OpCode == ERigVMOpCode::Copy)
-		{
-			static auto FormatFunction = [](const FString& RegisterName, const FString& RegisterOffsetName) -> FString
+			else
 			{
-				return FString::Printf(TEXT("%s%s"), *RegisterName, *RegisterOffsetName);
-			};
-			const FRigVMCopyOp& Op = ByteCode.GetOpAt<FRigVMCopyOp>(Instruction);
-			ContextPublicData.DebugMemoryString += FString::Printf(TEXT("Instruction %d: Copy %s -> %s\n"), ContextPublicData.InstructionIndex, *GetOperandLabel(Op.Source, FormatFunction), *GetOperandLabel(Op.Target, FormatFunction));
-		}
-		else
-		{
-			ContextPublicData.DebugMemoryString += FString::Printf(TEXT("Instruction %d: %s\n"), ContextPublicData.InstructionIndex, *ContextPublicData.InstanceOpCodeEnum->GetNameByIndex((uint8)Instruction.OpCode).ToString());
+				ContextPublicData.DebugMemoryString += FString::Printf(TEXT("Instruction %d: %s\n"), ContextPublicData.InstructionIndex, *ContextPublicData.InstanceOpCodeEnum->GetNameByIndex((uint8)Instruction.OpCode).ToString());
+			}
 		}
-		
 #endif
 #endif
 
@@ -2044,7 +2054,10 @@ ERigVMExecuteResult URigVM::ExecuteInstructions(int32 InFirstInstruction, int32
 						ExecutionHalted().Broadcast(INDEX_NONE, nullptr, CurrentEntryName);
 					}
 #if UE_RIGVM_DEBUG_EXECUTION
-					ContextPublicData.Log(EMessageSeverity::Info, ContextPublicData.DebugMemoryString);					
+					if (ContextPublicData.bDebugExecution)
+					{
+						ContextPublicData.Log(EMessageSeverity::Info, ContextPublicData.DebugMemoryString);
+					}
 #endif
 #endif
 				}
@@ -2173,46 +2186,49 @@ ERigVMExecuteResult URigVM::ExecuteInstructions(int32 InFirstInstruction, int32
 		}
 
 #if UE_RIGVM_DEBUG_EXECUTION
-		TArray<FString> CurrentWorkMemory;
-		URigVMMemoryStorage* WorkMemory = GetWorkMemory(false);
-		int32 LineIndex = 0;
-		for (int32 PropertyIndex=0; PropertyIndex<WorkMemory->Num(); ++PropertyIndex, ++LineIndex)
+		if (ContextPublicData.bDebugExecution)
 		{
-			FString Line = FString::Printf(TEXT("%s: %s"), *WorkMemory->GetProperties()[PropertyIndex]->GetFullName(), *WorkMemory->GetDataAsString(PropertyIndex));
-			if (ContextPublicData.PreviousWorkMemory.Num() > 0 && ContextPublicData.PreviousWorkMemory.IsValidIndex(PropertyIndex) && ContextPublicData.PreviousWorkMemory[PropertyIndex].StartsWith(TEXT(" -- ")))
+			TArray<FString> CurrentWorkMemory;
+			URigVMMemoryStorage* WorkMemory = GetWorkMemory(false);
+			int32 LineIndex = 0;
+			for (int32 PropertyIndex=0; PropertyIndex<WorkMemory->Num(); ++PropertyIndex, ++LineIndex)
 			{
-				ContextPublicData.PreviousWorkMemory[PropertyIndex].RightChopInline(4);
-			}
-			if (ContextPublicData.PreviousWorkMemory.Num() == 0 || (ContextPublicData.PreviousWorkMemory.IsValidIndex(PropertyIndex) && Line == ContextPublicData.PreviousWorkMemory[PropertyIndex]))
-			{
-				CurrentWorkMemory.Add(Line);
-			}
-			else
-			{
-				CurrentWorkMemory.Add(FString::Printf(TEXT(" -- %s"), *Line));
-			}
-		}
-		for (const FRigVMExternalVariable& ExternalVariable : ExternalVariables)
-		{
-			FString Value;
-			ExternalVariable.Property->ExportTextItem_Direct(Value, ExternalVariable.Memory, nullptr, nullptr, PPF_None);
-			FString Line = FString::Printf(TEXT("External %s: %s"), *ExternalVariable.Name.ToString(), *Value);
-			if (ContextPublicData.PreviousWorkMemory.Num() > 0 && ContextPublicData.PreviousWorkMemory.IsValidIndex(LineIndex) && ContextPublicData.PreviousWorkMemory[LineIndex].StartsWith(TEXT(" -- ")))
-			{
-				ContextPublicData.PreviousWorkMemory[LineIndex].RightChopInline(4);
-			}
-			if (ContextPublicData.PreviousWorkMemory.Num() == 0 || (ContextPublicData.PreviousWorkMemory.IsValidIndex(LineIndex) && Line == ContextPublicData.PreviousWorkMemory[LineIndex]))
-			{
-				CurrentWorkMemory.Add(Line);
+				FString Line = FString::Printf(TEXT("%s: %s"), *WorkMemory->GetProperties()[PropertyIndex]->GetFullName(), *WorkMemory->GetDataAsString(PropertyIndex));
+				if (ContextPublicData.PreviousWorkMemory.Num() > 0 && ContextPublicData.PreviousWorkMemory.IsValidIndex(PropertyIndex) && ContextPublicData.PreviousWorkMemory[PropertyIndex].StartsWith(TEXT(" -- ")))
+				{
+					ContextPublicData.PreviousWorkMemory[PropertyIndex].RightChopInline(4);
+				}
+				if (ContextPublicData.PreviousWorkMemory.Num() == 0 || (ContextPublicData.PreviousWorkMemory.IsValidIndex(PropertyIndex) && Line == ContextPublicData.PreviousWorkMemory[PropertyIndex]))
+				{
+					CurrentWorkMemory.Add(Line);
+				}
+				else
+				{
+					CurrentWorkMemory.Add(FString::Printf(TEXT(" -- %s"), *Line));
+				}
 			}
-			else
+			for (const FRigVMExternalVariable& ExternalVariable : ExternalVariables)
 			{
-				CurrentWorkMemory.Add(FString::Printf(TEXT(" -- %s"), *Line));
+				FString Value;
+				ExternalVariable.Property->ExportTextItem_Direct(Value, ExternalVariable.Memory, nullptr, nullptr, PPF_None);
+				FString Line = FString::Printf(TEXT("External %s: %s"), *ExternalVariable.Name.ToString(), *Value);
+				if (ContextPublicData.PreviousWorkMemory.Num() > 0 && ContextPublicData.PreviousWorkMemory.IsValidIndex(LineIndex) && ContextPublicData.PreviousWorkMemory[LineIndex].StartsWith(TEXT(" -- ")))
+				{
+					ContextPublicData.PreviousWorkMemory[LineIndex].RightChopInline(4);
+				}
+				if (ContextPublicData.PreviousWorkMemory.Num() == 0 || (ContextPublicData.PreviousWorkMemory.IsValidIndex(LineIndex) && Line == ContextPublicData.PreviousWorkMemory[LineIndex]))
+				{
+					CurrentWorkMemory.Add(Line);
+				}
+				else
+				{
+					CurrentWorkMemory.Add(FString::Printf(TEXT(" -- %s"), *Line));
+				}
+				++LineIndex;
 			}
-			++LineIndex;
+			ContextPublicData.DebugMemoryString += FString::Join(CurrentWorkMemory, TEXT("\n")) + FString(TEXT("\n\n"));
+			ContextPublicData.PreviousWorkMemory = CurrentWorkMemory;
 		}
-		ContextPublicData.DebugMemoryString += FString::Join(CurrentWorkMemory, TEXT("\n")) + FString(TEXT("\n\n"));
-		ContextPublicData.PreviousWorkMemory = CurrentWorkMemory;
 #endif
 #endif
 	}
diff --git a/Engine/Plugins/Runtime/RigVM/Source/RigVM/Public/RigVMCore/RigVMMemoryStorage.h b/Engine/Plugins/Runtime/RigVM/Source/RigVM/Public/RigVMCore/RigVMMemoryStorage.h
--- a/Engine/Plugins/Runtime/RigVM/Source/RigVM/Public/RigVMCore/RigVMMemoryStorage.h
+++ b/Engine/Plugins/Runtime/RigVM/Source/RigVM/Public/RigVMCore/RigVMMemoryStorage.h
@@ -817,20 +817,20 @@ private:
 						for (int32 i=FirstAddedIndex; i<ArrayHelper.Num(); ++i)
 						{
 #if UE_RIGVM_DEBUG_EXECUTION
-							FString DefaultValue;
-							ElementProperty->ExportText_Direct(
-								DefaultValue,
-								DefaultElementMemory,
-								DefaultElementMemory,
-								nullptr,
-								PPF_None,
-								nullptr);
-
-							UE_LOG(LogRigVM, Display, TEXT("Adding slice %d for Property '%s', defaulting to '%s'."),
-								InSliceIndex,
-								*ArrayProperty->GetName(),
-								*DefaultValue
-							);
+							// FString DefaultValue;
+							// ElementProperty->ExportText_Direct(
+							// 	DefaultValue,
+							// 	DefaultElementMemory,
+							// 	DefaultElementMemory,
+							// 	nullptr,
+							// 	PPF_None,
+							// 	nullptr);
+							//
+							// UE_LOG(LogRigVM, Display, TEXT("Adding slice %d for Property '%s', defaulting to '%s'."),
+							// 	InSliceIndex,
+							// 	*ArrayProperty->GetName(),
+							// 	*DefaultValue
+							// );
 #endif
 
 							uint8* DestMemory = ArrayHelper.GetRawPtr(i);
@@ -841,22 +841,22 @@ private:
 			}
 
 #if UE_RIGVM_DEBUG_EXECUTION
-			const FProperty* ElementProperty = ArrayProperty->Inner;
-			FString DefaultValue;
-			const uint8* DefaultElementMemory = ArrayHelper.GetRawPtr(InSliceIndex);
-			ElementProperty->ExportText_Direct(
-				DefaultValue,
-				DefaultElementMemory,
-				DefaultElementMemory,
-				nullptr,
-				PPF_None,
-				nullptr);
-
-			UE_LOG(LogRigVM, Display, TEXT("Getting slice %d for Property '%s', currently '%s'."),
-				InSliceIndex,
-				*ArrayProperty->GetName(),
-				*DefaultValue
-			);
+			// const FProperty* ElementProperty = ArrayProperty->Inner;
+			// FString DefaultValue;
+			// const uint8* DefaultElementMemory = ArrayHelper.GetRawPtr(InSliceIndex);
+			// ElementProperty->ExportText_Direct(
+			// 	DefaultValue,
+			// 	DefaultElementMemory,
+			// 	DefaultElementMemory,
+			// 	nullptr,
+			// 	PPF_None,
+			// 	nullptr);
+			//
+			// UE_LOG(LogRigVM, Display, TEXT("Getting slice %d for Property '%s', currently '%s'."),
+			// 	InSliceIndex,
+			// 	*ArrayProperty->GetName(),
+			// 	*DefaultValue
+			// );
 #endif
 
 			return ArrayHelper.GetRawPtr(InSliceIndex);
