Control Rig: Patch float/double casting on array types altering vm execution

diff --git a/Engine/Plugins/Runtime/RigVM/Source/RigVMDeveloper/Private/RigVMModel/RigVMClient.cpp b/Engine/Plugins/Runtime/RigVM/Source/RigVMDeveloper/Private/RigVMModel/RigVMClient.cpp
--- a/Engine/Plugins/Runtime/RigVM/Source/RigVMDeveloper/Private/RigVMModel/RigVMClient.cpp
+++ b/Engine/Plugins/Runtime/RigVM/Source/RigVMDeveloper/Private/RigVMModel/RigVMClient.cpp
@@ -796,6 +796,7 @@ FRigVMClientPatchResult FRigVMClient::PatchModelsOnLoad()
 		Result.Merge(Controller->PatchBranchNodesOnLoad());
 		Result.Merge(Controller->PatchIfSelectNodesOnLoad());
 		Result.Merge(Controller->PatchArrayNodesOnLoad());
+		Result.Merge(Controller->PatchReduceArrayFloatDoubleConvertsionsOnLoad());
 		Result.Merge(Controller->PatchInvalidLinksOnWildcards());
 	}
 
diff --git a/Engine/Plugins/Runtime/RigVM/Source/RigVMDeveloper/Private/RigVMModel/RigVMController.cpp b/Engine/Plugins/Runtime/RigVM/Source/RigVMDeveloper/Private/RigVMModel/RigVMController.cpp
--- a/Engine/Plugins/Runtime/RigVM/Source/RigVMDeveloper/Private/RigVMModel/RigVMController.cpp
+++ b/Engine/Plugins/Runtime/RigVM/Source/RigVMDeveloper/Private/RigVMModel/RigVMController.cpp
@@ -16441,7 +16441,7 @@ bool URigVMController::ResolveWildCardPin(URigVMPin* InPin, TRigVMTypeIndex InTy
 
 }
 
-bool URigVMController::UpdateTemplateNodePinTypes(URigVMTemplateNode* InNode, bool bSetupUndoRedo, bool bInitializeDefaultValue)
+bool URigVMController::UpdateTemplateNodePinTypes(URigVMTemplateNode* InNode, bool bSetupUndoRedo, bool bInitializeDefaultValue, TMap<URigVMPin*, TArray<TRigVMTypeIndex>> ProposedTypes)
 {
 	URigVMGraph* Graph = GetGraph();
 	check(InNode->GetGraph() == Graph);
@@ -16538,6 +16538,7 @@ bool URigVMController::UpdateTemplateNodePinTypes(URigVMTemplateNode* InNode, bo
 	FinalPinTypes.SetNumUninitialized(Pins.Num());
 	
 	// Reduce compatible types, try to find the permutation in the reduced types
+	FRigVMTemplateTypeMap AlreadyResolvedTypeMap;
 	for(int32 PinIndex=0; PinIndex < Pins.Num(); ++PinIndex)
 	{
 		URigVMPin* Pin = Pins[PinIndex];
@@ -16562,6 +16563,33 @@ bool URigVMController::UpdateTemplateNodePinTypes(URigVMTemplateNode* InNode, bo
 			}
 		}
 
+		if (TArray<TRigVMTypeIndex>* Proposed = ProposedTypes.Find(Pin))
+		{
+			// Intersect the types with the input proposed types
+			TArray<TRigVMTypeIndex> Intersection = Types.FilterByPredicate([Proposed](const TRigVMTypeIndex& Type) { return Proposed->Contains(Type); });
+			if (Intersection.Num() == 1)
+			{
+				PreferredType = Intersection[0];
+				TypesFoundInReduced = 1;
+			}
+		}
+
+		// Find the options left in the already reduced pins
+		if (TypesFoundInReduced > 1 && !AlreadyResolvedTypeMap.IsEmpty())
+		{
+			FRigVMTemplateTypeMap OutTypeMap = AlreadyResolvedTypeMap;
+			TArray<int32> OutPermutations;
+			Template->Resolve(OutTypeMap, OutPermutations, false);
+			if (TRigVMTypeIndex* TypeIndex = OutTypeMap.Find(Pin->GetFName()))
+			{
+				if (*TypeIndex != INDEX_NONE)
+				{
+					PreferredType = *TypeIndex;
+					TypesFoundInReduced = 1;
+				}
+			}
+		}
+
 		if (TypesFoundInReduced > 1)
 		{
 			PreferredType = Pin->GetTypeIndex();
@@ -16570,6 +16598,7 @@ bool URigVMController::UpdateTemplateNodePinTypes(URigVMTemplateNode* InNode, bo
 		if (FinalPinTypes[PinIndex] != INDEX_NONE)
 		{
 			WasReduced[PinIndex] =  true;
+			AlreadyResolvedTypeMap.Add(Pin->GetFName(), FinalPinTypes[PinIndex]);
 		}
 
 		// Remove reduced types which do not match this type
@@ -18494,6 +18523,95 @@ FRigVMClientPatchResult URigVMController::PatchArrayNodesOnLoad()
 	return Result;
 }
 
+FRigVMClientPatchResult URigVMController::PatchReduceArrayFloatDoubleConvertsionsOnLoad()
+{
+	FRigVMClientPatchResult Result;
+
+	const FRigVMRegistry& Registry = FRigVMRegistry::Get();
+	if (const URigVMGraph* Graph = GetGraph())
+	{
+		bool bChangedType = false;
+		TArray<FString> NodesModified;
+		do
+		{
+			bChangedType = false;
+			for (URigVMLink* Link : Graph->GetLinks())
+			{
+				URigVMPin* SourcePin = Link->GetSourcePin();
+				URigVMPin* TargetPin = Link->GetTargetPin();
+				if (!SourcePin || !TargetPin)
+				{
+					continue;
+				}
+
+				URigVMTemplateNode* TemplateNode = Cast<URigVMTemplateNode>(TargetPin->GetNode());
+				if (!TemplateNode)
+				{
+					continue;
+				}
+				
+				const TRigVMTypeIndex SourceType = SourcePin->GetTypeIndex();
+				const TRigVMTypeIndex TargetType = TargetPin->GetTypeIndex();
+
+				// Only patch when target is array
+				if (!Registry.IsArrayType(TargetType))
+				{
+					continue;
+				}
+
+				TRigVMTypeIndex SourceBaseType = SourceType;
+				TRigVMTypeIndex TargetBaseType = TargetType;
+				while (Registry.IsArrayType(SourceBaseType) || Registry.IsArrayType(TargetBaseType))
+				{
+					SourceBaseType = Registry.GetBaseTypeFromArrayTypeIndex(SourceBaseType);
+					TargetBaseType = Registry.GetBaseTypeFromArrayTypeIndex(TargetBaseType);
+				}
+
+				if ((SourceBaseType != RigVMTypeUtils::TypeIndex::Float && SourceBaseType != RigVMTypeUtils::TypeIndex::Double) ||
+				   (TargetBaseType != RigVMTypeUtils::TypeIndex::Float && TargetBaseType != RigVMTypeUtils::TypeIndex::Double))
+				{
+					continue;
+				}
+
+				if (SourceBaseType == TargetBaseType)
+				{
+					continue;
+				}
+		
+				// If we have already changed the type of this node, we might be inside an infinite loop
+				// Lets break the loop and return failure
+				if (NodesModified.Contains(TemplateNode->GetName()))
+				{
+					Result.bSucceeded = false;
+					break;
+				}
+		
+				if (TemplateNode->SupportsType(TargetPin, SourceType))
+				{
+					FRigVMTemplate::FTypeMap TypeMap;
+					TypeMap.Add(TargetPin->GetFName(), SourceType);
+					TArray<int32> Permutations;
+					TemplateNode->GetTemplate()->Resolve(TypeMap, Permutations, false);
+					if (Permutations.Num() == 1)
+					{
+						FRigVMTemplateTypeMap NewTypes = TemplateNode->GetTemplate()->GetTypesForPermutation(Permutations[0]);
+						for (auto Pair : NewTypes)
+						{
+							const FRigVMTemplateArgumentType Type = Registry.GetType(Pair.Value);
+							ChangePinType(TemplateNode->FindPin(Pair.Key.ToString()), Type.CPPType.ToString(), Type.GetCPPTypeObjectPath(), false, false, false, false, false);
+							bChangedType = true;
+							NodesModified.Add(TemplateNode->GetName());
+							Result.bChangedContent = true;
+						}
+					}
+				}
+			}
+		} while(bChangedType && Result.bSucceeded);
+	}
+
+	return Result;
+}
+
 FRigVMClientPatchResult URigVMController::PatchInvalidLinksOnWildcards()
 {
 	FRigVMClientPatchResult Result;
diff --git a/Engine/Plugins/Runtime/RigVM/Source/RigVMDeveloper/Public/RigVMModel/RigVMController.h b/Engine/Plugins/Runtime/RigVM/Source/RigVMDeveloper/Public/RigVMModel/RigVMController.h
--- a/Engine/Plugins/Runtime/RigVM/Source/RigVMDeveloper/Public/RigVMModel/RigVMController.h
+++ b/Engine/Plugins/Runtime/RigVM/Source/RigVMDeveloper/Public/RigVMModel/RigVMController.h
@@ -1163,7 +1163,7 @@ private:
 	void ResolveTemplateNodeMetaData(URigVMTemplateNode* InNode, bool bSetupUndoRedo);
 
 	// Changes Pin types if filtered types of a pin are unique
-	bool UpdateTemplateNodePinTypes(URigVMTemplateNode* InNode, bool bSetupUndoRedo, bool bInitializeDefaultValue = true);
+	bool UpdateTemplateNodePinTypes(URigVMTemplateNode* InNode, bool bSetupUndoRedo, bool bInitializeDefaultValue = true, TMap<URigVMPin*, TArray<TRigVMTypeIndex>> ProposedTypes = TMap<URigVMPin*, TArray<TRigVMTypeIndex>>());
 
 	bool ChangePinType(const FString& InPinPath, const FString& InCPPType, const FName& InCPPTypeObjectPath, bool bSetupUndoRedo, bool bSetupOrphanPins = true, bool bBreakLinks = true, bool bRemoveSubPins = true, bool bInitializeDefaultValue = true);
 	bool ChangePinType(URigVMPin* InPin, const FString& InCPPType, const FName& InCPPTypeObjectPath, bool bSetupUndoRedo, bool bSetupOrphanPins = true, bool bBreakLinks = true, bool bRemoveSubPins = true, bool bInitializeDefaultValue = true);
@@ -1286,6 +1286,7 @@ protected:
 	FRigVMClientPatchResult PatchBranchNodesOnLoad();
 	FRigVMClientPatchResult PatchIfSelectNodesOnLoad();
 	FRigVMClientPatchResult PatchArrayNodesOnLoad();
+	FRigVMClientPatchResult PatchReduceArrayFloatDoubleConvertsionsOnLoad();
 	FRigVMClientPatchResult PatchInvalidLinksOnWildcards();
 
 	// work to do after a duplication of the host asset
