Control Rig: Reduce cooking times by not serializing function graphs during cook

diff --git a/Engine/Plugins/Animation/ControlRig/Source/ControlRigDeveloper/Private/ControlRigBlueprint.cpp b/Engine/Plugins/Animation/ControlRig/Source/ControlRigDeveloper/Private/ControlRigBlueprint.cpp
--- a/Engine/Plugins/Animation/ControlRig/Source/ControlRigDeveloper/Private/ControlRigBlueprint.cpp
+++ b/Engine/Plugins/Animation/ControlRig/Source/ControlRigDeveloper/Private/ControlRigBlueprint.cpp
@@ -796,7 +796,7 @@ void UControlRigBlueprint::PreSave(FObjectPreSaveContext ObjectSaveContext)
 {
 	Super::PreSave(ObjectSaveContext);
 
-	RigVMClient.PreSave();
+	RigVMClient.PreSave(ObjectSaveContext);
 
 	SupportedEventNames.Reset();
 	if (URigVMBlueprintGeneratedClass* RigClass = GetControlRigBlueprintGeneratedClass())
diff --git a/Engine/Plugins/Runtime/RigVM/Source/RigVMDeveloper/Private/RigVMModel/RigVMClient.cpp b/Engine/Plugins/Runtime/RigVM/Source/RigVMDeveloper/Private/RigVMModel/RigVMClient.cpp
--- a/Engine/Plugins/Runtime/RigVM/Source/RigVMDeveloper/Private/RigVMModel/RigVMClient.cpp
+++ b/Engine/Plugins/Runtime/RigVM/Source/RigVMDeveloper/Private/RigVMModel/RigVMClient.cpp
@@ -5,6 +5,7 @@
 #include "Exporters/Exporter.h"
 #include "Algo/Sort.h"
 #include "UnrealExporter.h"
+#include "UObject/ObjectSaveContext.h"
 #include "RigVMModel/RigVMControllerActions.h"
 
 #include UE_INLINE_GENERATED_CPP_BY_NAME(RigVMClient)
@@ -823,13 +824,29 @@ void FRigVMClient::PostDuplicateHost(const FString& InOldPathName, const FString
 	}
 }
 
-void FRigVMClient::PreSave()
+void FRigVMClient::PreSave(FObjectPreSaveContext ObjectSaveContext)
 {
-	for (URigVMNode* Node : FunctionLibrary->GetNodes())
+	if (!ObjectSaveContext.IsCooking())
 	{
-		if (URigVMLibraryNode* LibraryNode = Cast<URigVMLibraryNode>(Node))
+		if (IRigVMClientHost* ClientHost = Cast<IRigVMClientHost>(GetOuter()))
 		{
-			UpdateGraphFunctionSerializedGraph(LibraryNode);
+			if (IRigVMGraphFunctionHost* FunctionHost = ClientHost->GetRigVMGraphFunctionHost())
+			{
+				if (FRigVMGraphFunctionStore* Store = FunctionHost->GetRigVMGraphFunctionStore())
+				{
+					for (URigVMNode* Node : FunctionLibrary->GetNodes())
+					{
+						if (URigVMLibraryNode* LibraryNode = Cast<URigVMLibraryNode>(Node))
+						{
+							FRigVMGraphFunctionIdentifier Identifier = LibraryNode->GetFunctionIdentifier();
+							if (Store->IsFunctionPublic(Identifier))
+							{
+								UpdateGraphFunctionSerializedGraph(LibraryNode);
+							}
+						}
+					}
+				}
+			}
 		}
 	}
 }
diff --git a/Engine/Plugins/Runtime/RigVM/Source/RigVMDeveloper/Public/RigVMModel/RigVMClient.h b/Engine/Plugins/Runtime/RigVM/Source/RigVMDeveloper/Public/RigVMModel/RigVMClient.h
--- a/Engine/Plugins/Runtime/RigVM/Source/RigVMDeveloper/Public/RigVMModel/RigVMClient.h
+++ b/Engine/Plugins/Runtime/RigVM/Source/RigVMDeveloper/Public/RigVMModel/RigVMClient.h
@@ -156,7 +156,7 @@ public:
 	void PostDuplicateHost(const FString& InOldPathName, const FString& InNewPathName);
 
 	// work to be done before saving
-	void PreSave();
+	void PreSave(FObjectPreSaveContext ObjectSaveContext);
 
 	void HandleGraphModifiedEvent(ERigVMGraphNotifType InNotifType, URigVMGraph* InGraph, UObject* InSubject);
 
