Control Rig: Improve performance for RefreshEngineTypes

diff --git a/Engine/Plugins/Runtime/RigVM/Source/RigVM/Private/RigVMCore/RigVMTemplate.cpp b/Engine/Plugins/Runtime/RigVM/Source/RigVM/Private/RigVMCore/RigVMTemplate.cpp
--- a/Engine/Plugins/Runtime/RigVM/Source/RigVM/Private/RigVMCore/RigVMTemplate.cpp
+++ b/Engine/Plugins/Runtime/RigVM/Source/RigVM/Private/RigVMCore/RigVMTemplate.cpp
@@ -1363,6 +1363,51 @@ bool FRigVMTemplate::Resolve(FTypeMap& InOutTypes, TArray<int32>& OutPermutation
 	return !OutPermutationIndices.IsEmpty();
 }
 
+bool FRigVMTemplate::ContainsPermutation(const FTypeMap& InTypes) const
+{
+	const FRigVMRegistry& Registry = FRigVMRegistry::Get();
+	
+	TArray<int32> PossiblePermutations;
+	for (const TPair<FName, TRigVMTypeIndex>& Pair : InTypes)
+	{
+		if (const FRigVMTemplateArgument* Argument = FindArgument(Pair.Key))
+		{
+			if (const TArray<int32>* ArgumentPermutations = Argument->TypeToPermutations.Find(Pair.Value))
+			{
+				// If possible permutations is empty, initialize it
+				if (PossiblePermutations.IsEmpty())
+				{
+					PossiblePermutations = *ArgumentPermutations;
+				}
+				else
+				{
+					// Intersect possible permutations and the permutations found for this argument
+					PossiblePermutations = ArgumentPermutations->FilterByPredicate([PossiblePermutations](const int32& ArgPermutation)
+					{
+						return PossiblePermutations.Contains(ArgPermutation);
+					});
+					if (PossiblePermutations.IsEmpty())
+					{
+						return false;
+					}
+				}
+			}
+			else
+			{
+				// The argument does not support the given type
+				return false;
+			}
+		}
+		else
+		{
+			// The argument cannot be found
+			return false;
+		}
+	}
+	
+	return true;
+}
+
 bool FRigVMTemplate::ResolveArgument(const FName& InArgumentName, const TRigVMTypeIndex InTypeIndex,
 	FTypeMap& InOutTypes) const
 {
@@ -1559,13 +1604,9 @@ bool FRigVMTemplate::AddTypeForArgument(const FName& InArgumentName, TRigVMTypeI
 			
 			// make sure this permutation doesn't exist yet
 			FRigVMTemplateTypeMap TempTypes = Types;
-			TArray<int32> ExistingPermutations;
-			if(Resolve(TempTypes, ExistingPermutations, false))
+			if(ContainsPermutation(TempTypes))
 			{
-				if(ExistingPermutations.Num() == 1)
-				{
-					return false;
-				}
+				return false;
 			}
 			
 			for(FRigVMTemplateArgument& Argument : Arguments)
@@ -1583,8 +1624,7 @@ bool FRigVMTemplate::AddTypeForArgument(const FName& InArgumentName, TRigVMTypeI
 
 			// Find if these types were already registered
 			FRigVMTemplateTypeMap TestTypes = Types;
-			TArray<int32> TestPermutations;
-			if (Resolve(TestTypes, TestPermutations, false))
+			if (ContainsPermutation(TestTypes))
 			{
 				return false;
 			}
diff --git a/Engine/Plugins/Runtime/RigVM/Source/RigVM/Public/RigVMCore/RigVMTemplate.h b/Engine/Plugins/Runtime/RigVM/Source/RigVM/Public/RigVMCore/RigVMTemplate.h
--- a/Engine/Plugins/Runtime/RigVM/Source/RigVM/Public/RigVMCore/RigVMTemplate.h
+++ b/Engine/Plugins/Runtime/RigVM/Source/RigVM/Public/RigVMCore/RigVMTemplate.h
@@ -346,6 +346,9 @@ public:
 	// returns true if the template was able to resolve to at least one permutation
 	bool Resolve(FTypeMap& InOutTypes, TArray<int32> & OutPermutationIndices, bool bAllowFloatingPointCasts) const;
 
+	// returns true if the template was able to resolve to at least one permutation
+	bool ContainsPermutation(const FTypeMap& InTypes) const;
+
 	// returns true if the template can resolve an argument to a new type
 	bool ResolveArgument(const FName& InArgumentName, const TRigVMTypeIndex InTypeIndex, FTypeMap& InOutTypes) const;
 
