Control Rig: Fix pin type changed during cooking in CitySample

diff --git a/Engine/Plugins/Runtime/RigVM/Source/RigVMDeveloper/Private/RigVMModel/RigVMController.cpp b/Engine/Plugins/Runtime/RigVM/Source/RigVMDeveloper/Private/RigVMModel/RigVMController.cpp
--- a/Engine/Plugins/Runtime/RigVM/Source/RigVMDeveloper/Private/RigVMModel/RigVMController.cpp
+++ b/Engine/Plugins/Runtime/RigVM/Source/RigVMDeveloper/Private/RigVMModel/RigVMController.cpp
@@ -16537,8 +16537,54 @@ bool URigVMController::UpdateTemplateNodePinTypes(URigVMTemplateNode* InNode, bo
 	TArray<TRigVMTypeIndex> FinalPinTypes;
 	FinalPinTypes.SetNumUninitialized(Pins.Num());
 	
-	// Reduce compatible types, try to find the permutation in the reduced types
+	// Reduce compatible types, try to find the permutation in the reduced types (from the proposed types)
 	FRigVMTemplateTypeMap AlreadyResolvedTypeMap;
+	for(int32 PinIndex=0; PinIndex < Pins.Num(); ++PinIndex)
+	{
+		if (WasReduced[PinIndex])
+		{
+			continue;
+		}
+		
+		URigVMPin* Pin = Pins[PinIndex];
+		if (Pin->GetDirection() == ERigVMPinDirection::Hidden)
+		{
+			continue;
+		}
+		
+		if (TArray<TRigVMTypeIndex>* Proposed = ProposedTypes.Find(Pin))
+		{
+			TArray<TRigVMTypeIndex> ReducedPinTypes;
+			ReducedPinTypes.Reserve(ReducedTypes.Num());
+			for (auto Pair : ReducedTypes)
+			{
+				ReducedPinTypes.Add(Pair.Value[PinIndex]);
+			}
+			// Intersect the types with the input proposed types
+			TArray<TRigVMTypeIndex> Intersection = ReducedPinTypes.FilterByPredicate([Proposed](const TRigVMTypeIndex& Type) { return Proposed->Contains(Type); });
+			if (Intersection.Num() == 1)
+			{
+				WasReduced[PinIndex] = true;
+				FinalPinTypes[PinIndex] = Intersection[0];
+				AlreadyResolvedTypeMap.Add(Pin->GetFName(), FinalPinTypes[PinIndex]);
+
+				// If the argument is hidden, the available PinTypes will all be INDEX_NONE
+				if (Pin->Direction != ERigVMPinDirection::Hidden)
+				{
+					ReducedTypes = ReducedTypes.FilterByPredicate([Intersection, PinIndex](const TPair<int32, TArray<TRigVMTypeIndex>>& Permutation)
+					{
+						return Permutation.Value[PinIndex] == Intersection[0];
+					});
+				}
+			}
+		}
+	}
+
+	if (PinTypes.Num() > ReducedTypes.Num())
+	{
+		ensureMsgf(!ReducedTypes.IsEmpty(), TEXT("Found incompatible preferred types"));
+	}
+	
 	for(int32 PinIndex=0; PinIndex < Pins.Num(); ++PinIndex)
 	{
 		URigVMPin* Pin = Pins[PinIndex];
@@ -16563,17 +16609,6 @@ bool URigVMController::UpdateTemplateNodePinTypes(URigVMTemplateNode* InNode, bo
 			}
 		}
 
-		if (TArray<TRigVMTypeIndex>* Proposed = ProposedTypes.Find(Pin))
-		{
-			// Intersect the types with the input proposed types
-			TArray<TRigVMTypeIndex> Intersection = Types.FilterByPredicate([Proposed](const TRigVMTypeIndex& Type) { return Proposed->Contains(Type); });
-			if (Intersection.Num() == 1)
-			{
-				PreferredType = Intersection[0];
-				TypesFoundInReduced = 1;
-			}
-		}
-
 		// Find the options left in the already reduced pins
 		if (TypesFoundInReduced > 1 && !AlreadyResolvedTypeMap.IsEmpty())
 		{
