Control Rig: Improve performance when opening context menu in graph. Avoid loading other control rigs.

diff --git a/Engine/Plugins/Runtime/RigVM/Source/RigVM/Private/RigVMCore/RigVMGraphFunctionDefinition.cpp b/Engine/Plugins/Runtime/RigVM/Source/RigVM/Private/RigVMCore/RigVMGraphFunctionDefinition.cpp
--- a/Engine/Plugins/Runtime/RigVM/Source/RigVM/Private/RigVMCore/RigVMGraphFunctionDefinition.cpp
+++ b/Engine/Plugins/Runtime/RigVM/Source/RigVM/Private/RigVMCore/RigVMGraphFunctionDefinition.cpp
@@ -38,10 +38,10 @@ bool FRigVMGraphFunctionHeader::IsMutable() const
 	return false;
 }
 
-IRigVMGraphFunctionHost* FRigVMGraphFunctionHeader::GetFunctionHost() const
+IRigVMGraphFunctionHost* FRigVMGraphFunctionHeader::GetFunctionHost(bool bLoadIfNecessary) const
 {
 	UObject* HostObj = LibraryPointer.HostObject.ResolveObject();
-	if (!HostObj)
+	if (!HostObj && bLoadIfNecessary)
 	{
 		HostObj = LibraryPointer.HostObject.TryLoad();
 	}
@@ -52,9 +52,9 @@ IRigVMGraphFunctionHost* FRigVMGraphFunctionHeader::GetFunctionHost() const
 	return nullptr;
 }
 
-FRigVMGraphFunctionData* FRigVMGraphFunctionHeader::GetFunctionData() const
+FRigVMGraphFunctionData* FRigVMGraphFunctionHeader::GetFunctionData(bool bLoadIfNecessary) const
 {
-	if (IRigVMGraphFunctionHost* Host = GetFunctionHost())
+	if (IRigVMGraphFunctionHost* Host = GetFunctionHost(bLoadIfNecessary))
 	{
 		return Host->GetRigVMGraphFunctionStore()->FindFunction(LibraryPointer);
 	}
diff --git a/Engine/Plugins/Runtime/RigVM/Source/RigVM/Public/RigVMCore/RigVMGraphFunctionDefinition.h b/Engine/Plugins/Runtime/RigVM/Source/RigVM/Public/RigVMCore/RigVMGraphFunctionDefinition.h
--- a/Engine/Plugins/Runtime/RigVM/Source/RigVM/Public/RigVMCore/RigVMGraphFunctionDefinition.h
+++ b/Engine/Plugins/Runtime/RigVM/Source/RigVM/Public/RigVMCore/RigVMGraphFunctionDefinition.h
@@ -386,9 +386,9 @@ struct RIGVM_API FRigVMGraphFunctionHeader
 		return LibraryPointer == Other.LibraryPointer;
 	}
 
-	IRigVMGraphFunctionHost* GetFunctionHost() const;
+	IRigVMGraphFunctionHost* GetFunctionHost(bool bLoadIfNecessary = true) const;
 
-	FRigVMGraphFunctionData* GetFunctionData() const;
+	FRigVMGraphFunctionData* GetFunctionData(bool bLoadIfNecessary = true) const;
 
 	friend FArchive& operator<<(FArchive& Ar, FRigVMGraphFunctionHeader& Data)
 	{
diff --git a/Engine/Plugins/Runtime/RigVM/Source/RigVMDeveloper/Private/RigVMModel/Nodes/RigVMFunctionReferenceNode.cpp b/Engine/Plugins/Runtime/RigVM/Source/RigVMDeveloper/Private/RigVMModel/Nodes/RigVMFunctionReferenceNode.cpp
--- a/Engine/Plugins/Runtime/RigVM/Source/RigVMDeveloper/Private/RigVMModel/Nodes/RigVMFunctionReferenceNode.cpp
+++ b/Engine/Plugins/Runtime/RigVM/Source/RigVMDeveloper/Private/RigVMModel/Nodes/RigVMFunctionReferenceNode.cpp
@@ -92,10 +92,14 @@ TArray<FRigVMExternalVariable> URigVMFunctionReferenceNode::GetExternalVariables
 	
 	if(!bRemapped)
 	{
-		if (const FRigVMGraphFunctionData* FunctionData = GetReferencedFunctionData())
+		if (const FRigVMGraphFunctionData* FunctionData = GetReferencedFunctionData(false))
 		{
 			Variables = FunctionData->Header.ExternalVariables;
 		}
+		else
+		{
+			Variables = GetReferencedFunctionHeader().ExternalVariables;
+		}
 	}
 	else
 	{
@@ -162,9 +166,9 @@ uint32 URigVMFunctionReferenceNode::GetStructureHash() const
 	return Hash;
 }
 
-const FRigVMGraphFunctionData* URigVMFunctionReferenceNode::GetReferencedFunctionData() const
+const FRigVMGraphFunctionData* URigVMFunctionReferenceNode::GetReferencedFunctionData(bool bLoadIfNecessary) const
 {
-	if (IRigVMGraphFunctionHost* Host = ReferencedFunctionHeader.GetFunctionHost())
+	if (IRigVMGraphFunctionHost* Host = ReferencedFunctionHeader.GetFunctionHost(bLoadIfNecessary))
 	{
 		return Host->GetRigVMGraphFunctionStore()->FindFunction(ReferencedFunctionHeader.LibraryPointer);
 	}
diff --git a/Engine/Plugins/Runtime/RigVM/Source/RigVMDeveloper/Public/RigVMModel/Nodes/RigVMFunctionReferenceNode.h b/Engine/Plugins/Runtime/RigVM/Source/RigVMDeveloper/Public/RigVMModel/Nodes/RigVMFunctionReferenceNode.h
--- a/Engine/Plugins/Runtime/RigVM/Source/RigVMDeveloper/Public/RigVMModel/Nodes/RigVMFunctionReferenceNode.h
+++ b/Engine/Plugins/Runtime/RigVM/Source/RigVMDeveloper/Public/RigVMModel/Nodes/RigVMFunctionReferenceNode.h
@@ -49,7 +49,7 @@ public:
 	virtual uint32 GetStructureHash() const override;
 	const FRigVMGraphFunctionHeader& GetReferencedFunctionHeader() const { return ReferencedFunctionHeader; }
 
-	const FRigVMGraphFunctionData* GetReferencedFunctionData() const;
+	const FRigVMGraphFunctionData* GetReferencedFunctionData(bool bLoadIfNecessary = true) const;
 
 	
 private:
