ControlRig: 5.2 fix, renaming the ControlRig BP asset causes compilation error and thus broken rigs. When compiler tries to find the function defintion for function reference nodes in the graph that reference a private function, it can fail to find a match because even though the soft object path in the function ref node is updated to point to the renamed asset, the soft object path in the function definition storage still points to the old asset.

diff --git a/Engine/Plugins/Runtime/RigVM/Source/RigVM/Public/RigVMCore/RigVMGraphFunctionHost.h b/Engine/Plugins/Runtime/RigVM/Source/RigVM/Public/RigVMCore/RigVMGraphFunctionHost.h
--- a/Engine/Plugins/Runtime/RigVM/Source/RigVM/Public/RigVMCore/RigVMGraphFunctionHost.h
+++ b/Engine/Plugins/Runtime/RigVM/Source/RigVM/Public/RigVMCore/RigVMGraphFunctionHost.h
@@ -73,6 +73,13 @@ public:
 	friend FArchive& operator<<(FArchive& Ar, FRigVMGraphFunctionStore& Host)
 	{
 		Ar << Host.PublicFunctions;
+		
+		// This is only added to make sure SoftObjectPaths can be gathered and fixed up in the case of asset rename
+		// It should not affect data on disk
+		if (Ar.IsObjectReferenceCollector())
+		{
+			Ar << Host.PrivateFunctions;
+		}
 		return Ar;
 	}
 
