RigVM: Fix RigVM.Functions.NestedVariableRemapping failing in linux

diff --git a/Engine/Plugins/Runtime/RigVM/Source/RigVMDeveloper/Private/RigVMModel/RigVMController.cpp b/Engine/Plugins/Runtime/RigVM/Source/RigVMDeveloper/Private/RigVMModel/RigVMController.cpp
--- a/Engine/Plugins/Runtime/RigVM/Source/RigVMDeveloper/Private/RigVMModel/RigVMController.cpp
+++ b/Engine/Plugins/Runtime/RigVM/Source/RigVMDeveloper/Private/RigVMModel/RigVMController.cpp
@@ -2024,7 +2024,7 @@ bool URigVMController::OnExternalVariableRenamed(const FName& InOldVarName, cons
 					{
 						Reference->Modify();
 
-						const FName& MappedVariable = Reference->VariableMap.FindChecked(InOldVarName);
+						const FName MappedVariable = Reference->VariableMap.FindChecked(InOldVarName);
 						Reference->VariableMap.Remove(InOldVarName);
 						Reference->VariableMap.FindOrAdd(InNewVarName) = MappedVariable; 
 
