Fixed mismatched copy constructor/assignment operator in FRigVMPropertyPath.

diff --git a/Engine/Plugins/Runtime/RigVM/Source/RigVM/Private/RigVMCore/RigVMPropertyPath.cpp b/Engine/Plugins/Runtime/RigVM/Source/RigVM/Private/RigVMCore/RigVMPropertyPath.cpp
--- a/Engine/Plugins/Runtime/RigVM/Source/RigVM/Private/RigVMCore/RigVMPropertyPath.cpp
+++ b/Engine/Plugins/Runtime/RigVM/Source/RigVM/Private/RigVMCore/RigVMPropertyPath.cpp
@@ -125,12 +125,6 @@ FRigVMPropertyPath::FRigVMPropertyPath(const FProperty* InProperty, const FStrin
 	}
 }
 
-FRigVMPropertyPath::FRigVMPropertyPath(const FRigVMPropertyPath& InOther)
-: Path(InOther.Path)
-, Segments(InOther.Segments)
-{
-}
-
 const FProperty* FRigVMPropertyPath::GetTailProperty() const
 {
 	if(Segments.IsEmpty())
diff --git a/Engine/Plugins/Runtime/RigVM/Source/RigVM/Public/RigVMCore/RigVMPropertyPath.h b/Engine/Plugins/Runtime/RigVM/Source/RigVM/Public/RigVMCore/RigVMPropertyPath.h
--- a/Engine/Plugins/Runtime/RigVM/Source/RigVM/Public/RigVMCore/RigVMPropertyPath.h
+++ b/Engine/Plugins/Runtime/RigVM/Source/RigVM/Public/RigVMCore/RigVMPropertyPath.h
@@ -113,9 +113,6 @@ public:
 	 */
 	FRigVMPropertyPath(const FProperty* InProperty, const FString& InSegmentPath);
 
-	// Copy constructor
-	FRigVMPropertyPath(const FRigVMPropertyPath& InOther);
-
 	// Returns the property path as a string (the sanitized SegmentPath)
 	const FString& ToString() const { return Path; }
 
