RigVM: Fix access to execute context when UE_RIGVM_DEBUG_EXECUTION is activated

diff --git a/Engine/Plugins/Runtime/RigVM/Source/RigVM/Private/RigVMHost.cpp b/Engine/Plugins/Runtime/RigVM/Source/RigVM/Private/RigVMHost.cpp
--- a/Engine/Plugins/Runtime/RigVM/Source/RigVM/Private/RigVMHost.cpp
+++ b/Engine/Plugins/Runtime/RigVM/Source/RigVM/Private/RigVMHost.cpp
@@ -826,7 +826,7 @@ void URigVMHost::DeclareConstructClasses(TArray<FTopLevelAssetPath>& OutConstruc
 const FString URigVMHost::GetDebugExecutionString()
 {
 	TGuardValue<bool> DebugExecutionGuard(bDebugExecutionEnabled, true);
-	FRigVMExecuteContext& PublicContext = GetVM()->GetPublicData<FRigVMExecuteContext>();
+	FRigVMExecuteContext& PublicContext = GetExtendedExecuteContext().GetPublicData<FRigVMExecuteContext>();
 	PublicContext.DebugMemoryString.Reset();
 	
 	Evaluate_AnyThread();
