Fix for Linux Engine test compile error

diff --git a/Engine/Plugins/Runtime/RigVM/Source/RigVM/Private/RigVMCore/RigVMNativized.cpp b/Engine/Plugins/Runtime/RigVM/Source/RigVM/Private/RigVMCore/RigVMNativized.cpp
--- a/Engine/Plugins/Runtime/RigVM/Source/RigVM/Private/RigVMCore/RigVMNativized.cpp
+++ b/Engine/Plugins/Runtime/RigVM/Source/RigVM/Private/RigVMCore/RigVMNativized.cpp
@@ -27,7 +27,7 @@ void URigVMNativized::Reset(bool IsIgnoringArchetypeRef)
 	ByteCodeStorage.Reset();
 }
 
-bool URigVMNativized::Initialize(TArrayView<URigVMMemoryStorage*> Memory)
+bool URigVMNativized::Initialize(FRigVMExtendedExecuteContext& Context, TArrayView<URigVMMemoryStorage*> Memory)
 {
 	// nothing to do here 
 	return true;
diff --git a/Engine/Plugins/Runtime/RigVM/Source/RigVM/Public/RigVMCore/RigVMNativized.h b/Engine/Plugins/Runtime/RigVM/Source/RigVM/Public/RigVMCore/RigVMNativized.h
--- a/Engine/Plugins/Runtime/RigVM/Source/RigVM/Public/RigVMCore/RigVMNativized.h
+++ b/Engine/Plugins/Runtime/RigVM/Source/RigVM/Public/RigVMCore/RigVMNativized.h
@@ -50,7 +50,7 @@ public:
 	virtual bool IsNativized() const override { return true; }
 	virtual void Empty() override { return; }
 	virtual void CopyFrom(URigVM* InVM, bool bDeferCopy = false, bool bReferenceLiteralMemory = false, bool bReferenceByteCode = false, bool bCopyExternalVariables = false, bool bCopyDynamicRegisters = false) override { return; }
-	virtual bool Initialize(TArrayView<URigVMMemoryStorage*> Memory) override;
+	virtual bool Initialize(FRigVMExtendedExecuteContext& Context, TArrayView<URigVMMemoryStorage*> Memory) override;
 	virtual ERigVMExecuteResult Execute(TArrayView<URigVMMemoryStorage*> Memory, const FName& InEntryName = NAME_None) override;
 	virtual int32 AddRigVMFunction(UScriptStruct* InRigVMStruct, const FName& InMethodName) override { return INDEX_NONE; }
 	virtual FString GetRigVMFunctionName(int32 InFunctionIndex) const override { return FString(); }
