Fix for Linux and Mac EngineTests compilation : added missing parameter to overloaded function

diff --git a/Engine/Plugins/Runtime/RigVM/Source/RigVM/Private/RigVMCore/RigVMNativized.cpp b/Engine/Plugins/Runtime/RigVM/Source/RigVM/Private/RigVMCore/RigVMNativized.cpp
--- a/Engine/Plugins/Runtime/RigVM/Source/RigVM/Private/RigVMCore/RigVMNativized.cpp
+++ b/Engine/Plugins/Runtime/RigVM/Source/RigVM/Private/RigVMCore/RigVMNativized.cpp
@@ -33,7 +33,7 @@ bool URigVMNativized::Initialize(FRigVMExtendedExecuteContext& Context, TArrayVi
 	return true;
 }
 
-ERigVMExecuteResult URigVMNativized::Execute(TArrayView<URigVMMemoryStorage*> Memory, const FName& InEntryName)
+ERigVMExecuteResult URigVMNativized::Execute(FRigVMExtendedExecuteContext& Context, TArrayView<URigVMMemoryStorage*> Memory, const FName& InEntryName)
 {
 	// to be implemented by the generated code
 	return ERigVMExecuteResult::Failed;
diff --git a/Engine/Plugins/Runtime/RigVM/Source/RigVM/Public/RigVMCore/RigVMNativized.h b/Engine/Plugins/Runtime/RigVM/Source/RigVM/Public/RigVMCore/RigVMNativized.h
--- a/Engine/Plugins/Runtime/RigVM/Source/RigVM/Public/RigVMCore/RigVMNativized.h
+++ b/Engine/Plugins/Runtime/RigVM/Source/RigVM/Public/RigVMCore/RigVMNativized.h
@@ -51,7 +51,7 @@ public:
 	virtual void Empty() override { return; }
 	virtual void CopyFrom(URigVM* InVM, bool bDeferCopy = false, bool bReferenceLiteralMemory = false, bool bReferenceByteCode = false, bool bCopyExternalVariables = false, bool bCopyDynamicRegisters = false) override { return; }
 	virtual bool Initialize(FRigVMExtendedExecuteContext& Context, TArrayView<URigVMMemoryStorage*> Memory) override;
-	virtual ERigVMExecuteResult Execute(TArrayView<URigVMMemoryStorage*> Memory, const FName& InEntryName = NAME_None) override;
+	virtual ERigVMExecuteResult Execute(FRigVMExtendedExecuteContext& Context, TArrayView<URigVMMemoryStorage*> Memory, const FName& InEntryName = NAME_None) override;
 	virtual int32 AddRigVMFunction(UScriptStruct* InRigVMStruct, const FName& InMethodName) override { return INDEX_NONE; }
 	virtual FString GetRigVMFunctionName(int32 InFunctionIndex) const override { return FString(); }
 	virtual URigVMMemoryStorage* GetMemoryByType(ERigVMMemoryType InMemoryType, bool bCreateIfNeeded = true) override { return nullptr; }
