ControlRig: AlphaInterp node was broken due to RigVM refactor, where the the intiailization stage was removed but a mistake in the the refactored code causes the node to be constantly reinitialized.

diff --git a/Engine/Plugins/Runtime/RigVM/Source/RigVM/Private/RigVMFunctions/Simulation/RigVMFunction_AlphaInterp.cpp b/Engine/Plugins/Runtime/RigVM/Source/RigVM/Private/RigVMFunctions/Simulation/RigVMFunction_AlphaInterp.cpp
--- a/Engine/Plugins/Runtime/RigVM/Source/RigVM/Private/RigVMFunctions/Simulation/RigVMFunction_AlphaInterp.cpp
+++ b/Engine/Plugins/Runtime/RigVM/Source/RigVM/Private/RigVMFunctions/Simulation/RigVMFunction_AlphaInterp.cpp
@@ -10,11 +10,6 @@ FRigVMFunction_AlphaInterp_Execute()
 	ScaleBiasClamp.bClampResult = bClampResult;
 	ScaleBiasClamp.bInterpResult = bInterpResult;
 
-	if (ScaleBiasClamp.bInitialized)
-	{
-		ScaleBiasClamp.Reinitialize();
-	}
-
 	ScaleBiasClamp.InRange = InRange;
 	ScaleBiasClamp.OutRange = OutRange;
 	ScaleBiasClamp.ClampMin = ClampMin;
@@ -56,11 +51,6 @@ FRigVMFunction_AlphaInterpVector_Execute()
 	ScaleBiasClamp.bClampResult = bClampResult;
 	ScaleBiasClamp.bInterpResult = bInterpResult;
 
-	if (ScaleBiasClamp.bInitialized)
-	{
-		ScaleBiasClamp.Reinitialize();
-	}
-
 	ScaleBiasClamp.InRange = InRange;
 	ScaleBiasClamp.OutRange = OutRange;
 	ScaleBiasClamp.ClampMin = ClampMin;
@@ -104,11 +94,6 @@ FRigVMFunction_AlphaInterpQuat_Execute()
 	ScaleBiasClamp.bClampResult = bClampResult;
 	ScaleBiasClamp.bInterpResult = bInterpResult;
 
-	if (ScaleBiasClamp.bInitialized)
-	{
-		ScaleBiasClamp.Reinitialize();
-	}
-
 	ScaleBiasClamp.InRange = InRange;
 	ScaleBiasClamp.OutRange = OutRange;
 	ScaleBiasClamp.ClampMin = ClampMin;
