ControlRig: Fixed reroutes created on a link or from a pin not resolving to the expected type, causing link missing on asset reload. The fix does not address all issue with reroute currently but should handle 90% use cases with reroute, see code comment for details

diff --git a/Engine/Plugins/Runtime/RigVM/Source/RigVMDeveloper/Private/RigVMModel/RigVMController.cpp b/Engine/Plugins/Runtime/RigVM/Source/RigVMDeveloper/Private/RigVMModel/RigVMController.cpp
--- a/Engine/Plugins/Runtime/RigVM/Source/RigVMDeveloper/Private/RigVMModel/RigVMController.cpp
+++ b/Engine/Plugins/Runtime/RigVM/Source/RigVMDeveloper/Private/RigVMModel/RigVMController.cpp
@@ -2943,6 +2943,22 @@ URigVMRerouteNode* URigVMController::AddRerouteNodeOnPin(const FString& InPinPat
 		Pin->bIsExpanded = true;
 	});
 
+	if (!ValuePin->IsWildCard())
+	{
+		// Given that the pin already knows its type, let's make sure the template system is also aware
+		// The bottom line here is that Adding preferred type here is safe for all types:
+		// 1. For most types, considered this code for defensive purpose only, it is not needed and does no harm,
+		// as the preferred type will be inferred during AddLink() later anyway,
+		// which is why reroute node on non execute link have been working
+		// 2. For reroute node with a IO execution pin, this code is needed to address use cases where reroute is created
+		// from a pin/link directly. Ideally We should also rely on AllLink() to resolve the template properly such that
+		// Linking to a existing wild card reroute (a very rare use case) can also be resolved properly,
+		// but because execute pin has been refactored to be FRigVMExecuteArgument instead of FRigVMTemplateArgument
+		// as of the writing of this comment, AddLink() is excluding template resolution for execute pins, which is not ideal
+		
+		AddPreferredType(Node, ValuePin->GetFName(), ValuePin->GetTypeIndex(), bSetupUndoRedo);
+	}
+
 	Graph->Nodes.Add(Node);
 	if (!bSuspendNotifications)
 	{
