Fixed mismatched copy constructor/assignment operator in FRigVMASTProxy.

diff --git a/Engine/Plugins/Runtime/RigVM/Source/RigVMDeveloper/Public/RigVMCompiler/RigVMASTProxy.h b/Engine/Plugins/Runtime/RigVM/Source/RigVMDeveloper/Public/RigVMCompiler/RigVMASTProxy.h
--- a/Engine/Plugins/Runtime/RigVM/Source/RigVMDeveloper/Public/RigVMCompiler/RigVMASTProxy.h
+++ b/Engine/Plugins/Runtime/RigVM/Source/RigVMDeveloper/Public/RigVMCompiler/RigVMASTProxy.h
@@ -147,6 +147,18 @@ public:
 #endif
 	}
 
+	FRigVMASTProxy& operator=(const FRigVMASTProxy& InOther)
+	{
+		if (this != &InOther)
+		{
+			Callstack.Stack = InOther.Callstack.Stack;
+#if UE_BUILD_DEBUG
+			DebugName = Callstack.GetCallPath();
+#endif
+		}
+		return *this;
+	}
+
 	static FRigVMASTProxy MakeFromUObject(UObject* InSubject);
 	static FRigVMASTProxy MakeFromCallPath(const FString& InCallPath, UObject* InRootObject);
 	static FRigVMASTProxy MakeFromCallstack(const FRigVMCallstack& InCallstack);
