Fixed mismatched copy constructor/assignment operator for FRigVMASTLinkDescription.

diff --git a/Engine/Plugins/Runtime/RigVM/Source/RigVMDeveloper/Public/RigVMCompiler/RigVMAST.h b/Engine/Plugins/Runtime/RigVM/Source/RigVMDeveloper/Public/RigVMCompiler/RigVMAST.h
--- a/Engine/Plugins/Runtime/RigVM/Source/RigVMDeveloper/Public/RigVMCompiler/RigVMAST.h
+++ b/Engine/Plugins/Runtime/RigVM/Source/RigVMDeveloper/Public/RigVMCompiler/RigVMAST.h
@@ -55,6 +55,9 @@ struct RIGVMDEVELOPER_API FRigVMASTLinkDescription
 	, LinkIndex(INDEX_NONE)
 	{}
 
+	// Note that unlike the copy constructor, LinkIndex is copied by the assignment operator.
+	FRigVMASTLinkDescription& operator=(const FRigVMASTLinkDescription&) = default;
+
 	FRigVMASTProxy SourceProxy;
 	FRigVMASTProxy TargetProxy;
 	FString SegmentPath;
