RigVM: Fix reporting all log messages when Debug Execution is enabled

diff --git a/Engine/Plugins/Runtime/RigVM/Source/RigVM/Public/RigVMCore/RigVMExecuteContext.h b/Engine/Plugins/Runtime/RigVM/Source/RigVM/Public/RigVMCore/RigVMExecuteContext.h
--- a/Engine/Plugins/Runtime/RigVM/Source/RigVM/Public/RigVMCore/RigVMExecuteContext.h
+++ b/Engine/Plugins/Runtime/RigVM/Source/RigVM/Public/RigVMCore/RigVMExecuteContext.h
@@ -254,9 +254,7 @@ struct FRigVMExecuteContext
 		{
 			(*RuntimeSettings.LogFunction)(InSeverity, this, InMessage);
 		}
-#if !UE_RIGVM_DEBUG_EXECUTION
 		else
-#endif
 		{
 			if(InSeverity == EMessageSeverity::Error)
 			{
