RigVM: Fix crash when connecting execute reroute to a non-execute pin

diff --git a/Engine/Plugins/Runtime/RigVM/Source/RigVMDeveloper/Private/RigVMModel/RigVMController.cpp b/Engine/Plugins/Runtime/RigVM/Source/RigVMDeveloper/Private/RigVMModel/RigVMController.cpp
--- a/Engine/Plugins/Runtime/RigVM/Source/RigVMDeveloper/Private/RigVMModel/RigVMController.cpp
+++ b/Engine/Plugins/Runtime/RigVM/Source/RigVMDeveloper/Private/RigVMModel/RigVMController.cpp
@@ -9959,17 +9959,33 @@ bool URigVMController::AddLink(URigVMPin* OutputPin, URigVMPin* InputPin, bool b
 	// Before adding the link, let's resolve input and ouput pin types
 	// If templates, we will filter the permutations that support this link
 	// If any links need to be broken before perfoming this connection, try to find them and break them
-	if (!bIsTransacting && !InputPin->IsExecuteContext() && !OutputPin->IsExecuteContext())
+	if (!bIsTransacting)
 	{
-		URigVMPin* FirstToResolve = (InUserDirection == ERigVMPinDirection::Input) ? OutputPin : InputPin;
-		URigVMPin* SecondToResolve = (FirstToResolve == OutputPin) ? InputPin : OutputPin;
-		if (!PrepareToLink(FirstToResolve, SecondToResolve, bSetupUndoRedo))
+		if (!InputPin->IsExecuteContext() && !OutputPin->IsExecuteContext())
 		{
-			if (bSetupUndoRedo)
+			URigVMPin* FirstToResolve = (InUserDirection == ERigVMPinDirection::Input) ? OutputPin : InputPin;
+			URigVMPin* SecondToResolve = (FirstToResolve == OutputPin) ? InputPin : OutputPin;
+			if (!PrepareToLink(FirstToResolve, SecondToResolve, bSetupUndoRedo))
 			{
-				GetActionStack()->CancelAction(Action);
+				if (bSetupUndoRedo)
+				{
+					GetActionStack()->CancelAction(Action);
+				}
+				return false;
+			}
+		}
+		else if (InputPin->GetNode()->IsA<URigVMRerouteNode>() || OutputPin->GetNode()->IsA<URigVMRerouteNode>())
+		{
+			URigVMPin* PinToResolve = (OutputPin->GetNode()->IsA<URigVMRerouteNode>()) ? OutputPin : InputPin;
+			URigVMPin* PinToSkip = (PinToResolve == OutputPin) ? InputPin : OutputPin;
+			if (!ResolveWildCardPin(PinToResolve, PinToSkip->GetTypeIndex(), bSetupUndoRedo))
+			{
+				if (bSetupUndoRedo)
+				{
+					GetActionStack()->CancelAction(Action);
+				}
+				return false;
 			}
-			return false;
 		}
 	}
 
