RigVM: Fix compiler not getting proper default values of output pin in function reference nodes

diff --git a/Engine/Plugins/Runtime/RigVM/Source/RigVMDeveloper/Private/RigVMCompiler/RigVMAST.cpp b/Engine/Plugins/Runtime/RigVM/Source/RigVMDeveloper/Private/RigVMCompiler/RigVMAST.cpp
--- a/Engine/Plugins/Runtime/RigVM/Source/RigVMDeveloper/Private/RigVMCompiler/RigVMAST.cpp
+++ b/Engine/Plugins/Runtime/RigVM/Source/RigVMDeveloper/Private/RigVMCompiler/RigVMAST.cpp
@@ -1061,7 +1061,8 @@ FRigVMExprAST* FRigVMParserAST::TraversePin(const FRigVMASTProxy& InPinProxy, FR
 		Pin->GetDirection() == ERigVMPinDirection::Visible) &&
 		LinkIndices.Num() == 0)
 	{
-		if (Cast<URigVMVariableNode>(Pin->GetNode()))
+		if (Cast<URigVMVariableNode>(Pin->GetNode()) ||
+			(LibraryNodeBeingCompiled && Cast<URigVMFunctionReturnNode>(Pin->GetNode())))
 		{
 			PinExpr = MakeExpr<FRigVMVarExprAST>(FRigVMExprAST::EType::Var, InPinProxy);
 			FRigVMExprAST* PinLiteralExpr = MakeExpr<FRigVMLiteralExprAST>(InPinProxy);
@@ -2941,6 +2942,29 @@ void FRigVMParserAST::Inline(TArray<URigVMGraph*> InGraphs, const TArray<FRigVMA
 				{
 					if (IsValidLinkForAST(SourcePinProxy, InPinProxyForMap, OutTraversalInfo))
 					{
+						if (URigVMFunctionReferenceNode* FunctionReferenceNode = Cast<URigVMFunctionReferenceNode>(SourceNode))
+						{
+							if (const FRigVMGraphFunctionData* FunctionData = FunctionReferenceNode->GetReferencedFunctionData())
+							{
+								FString PinName = SourcePin->GetName();
+								PinName.ReplaceInline(TEXT(" "), TEXT("_"));
+								if (const FRigVMFunctionCompilationPropertyDescription* Description = FunctionData->CompilationData.WorkPropertyDescriptions.FindByPredicate([PinName](const FRigVMFunctionCompilationPropertyDescription& Description)
+								{
+									return Description.Name.ToString().EndsWith(PinName);
+								}))
+								{
+									// when asking for the default value of the array - we may need to get the previously stored override
+									URigVMPin::FPinOverrideValue OverrideValue;
+									OverrideValue.DefaultValue = Description->DefaultValue;
+									
+									// we query the default value now since the potential array elements have been visited and their
+									// potential default value override has been stored to the map.
+									OutTraversalInfo.PinOverrides->FindOrAdd(SourcePinProxy) = OverrideValue;
+								}
+								
+							}
+						}
+						
 						FRigVMASTLinkDescription Link(SourcePinProxy, InPinProxyForMap, InSegmentPath);
 						Link.LinkIndex = OutTraversalInfo.Links->Num();
 						OutTraversalInfo.Links->Add(Link);
diff --git a/Engine/Plugins/Runtime/RigVM/Source/RigVMDeveloper/Private/RigVMCompiler/RigVMCompiler.cpp b/Engine/Plugins/Runtime/RigVM/Source/RigVMDeveloper/Private/RigVMCompiler/RigVMCompiler.cpp
--- a/Engine/Plugins/Runtime/RigVM/Source/RigVMDeveloper/Private/RigVMCompiler/RigVMCompiler.cpp
+++ b/Engine/Plugins/Runtime/RigVM/Source/RigVMDeveloper/Private/RigVMCompiler/RigVMCompiler.cpp
@@ -1936,7 +1936,9 @@ int32 URigVMCompiler::TraverseInlineFunction(const FRigVMInlineFunctionExprAST*
 					{
 						continue;
 					}
-					if (!Properties[NumProperties].Name.ToString().Contains(Argument.Name.ToString()))
+					FString PropertyName = Argument.Name.ToString();
+					PropertyName.ReplaceInline(TEXT(" "), TEXT("_"));
+					if (!Properties[NumProperties].Name.ToString().Contains(PropertyName))
 					{
 						continue;
 					}
@@ -1958,7 +1960,7 @@ int32 URigVMCompiler::TraverseInlineFunction(const FRigVMInlineFunctionExprAST*
 				FRigVMCompilerWorkData::FFunctionRegisterData Data = {FunctionReferenceNode, MemoryType, PropertyIndex};
 				WorkData.FunctionRegisterToOperand.Add(Data, Operand);
 
-				// @todo Try to reuse literal operands				
+				// @todo Try to reuse literal operands
 			}
 		}		
 	}
@@ -2018,6 +2020,10 @@ int32 URigVMCompiler::TraverseInlineFunction(const FRigVMInlineFunctionExprAST*
 		FRigVMInstructionArray Instructions = ByteCode.GetInstructions();
 		for (int32 i=InstructionIndexStart; i<=InstructionIndexEnd; ++i)
 		{
+			if (!Instructions.IsValidIndex(i))
+			{
+				return INDEX_NONE;
+			}
 			const FRigVMInstruction& Instruction = Instructions[i];
 			const FRigVMOperandArray OperandArray = ByteCode.GetOperandsForOp(Instruction);
 			uint64 OperandsIndex = ByteCode.GetFirstOperandByteIndex(Instruction);
