RigVM: Fixed issue when adding link to execute reroute

diff --git a/Engine/Plugins/Runtime/RigVM/Source/RigVMDeveloper/Private/RigVMModel/RigVMController.cpp b/Engine/Plugins/Runtime/RigVM/Source/RigVMDeveloper/Private/RigVMModel/RigVMController.cpp
--- a/Engine/Plugins/Runtime/RigVM/Source/RigVMDeveloper/Private/RigVMModel/RigVMController.cpp
+++ b/Engine/Plugins/Runtime/RigVM/Source/RigVMDeveloper/Private/RigVMModel/RigVMController.cpp
@@ -8926,6 +8926,10 @@ bool URigVMController::RemoveArrayPin(const FString& InArrayElementPinPath, bool
 	{
 		if(ArrayPin->GetArraySize() == 1)
 		{
+			if (bSetupUndoRedo)
+			{
+				GetActionStack()->CancelAction(Action);
+			}
 			return false;
 		}
 	}
@@ -8933,6 +8937,10 @@ bool URigVMController::RemoveArrayPin(const FString& InArrayElementPinPath, bool
 	int32 IndexToRemove = ArrayElementPin->GetPinIndex();
 	if (!RemovePin(ArrayElementPin, bSetupUndoRedo))
 	{
+		if (bSetupUndoRedo)
+		{
+			GetActionStack()->CancelAction(Action);
+		}
 		return false;
 	}
 
@@ -9978,13 +9986,16 @@ bool URigVMController::AddLink(URigVMPin* OutputPin, URigVMPin* InputPin, bool b
 		{
 			URigVMPin* PinToResolve = (OutputPin->GetNode()->IsA<URigVMRerouteNode>()) ? OutputPin : InputPin;
 			URigVMPin* PinToSkip = (PinToResolve == OutputPin) ? InputPin : OutputPin;
-			if (!ResolveWildCardPin(PinToResolve, PinToSkip->GetTypeIndex(), bSetupUndoRedo))
+			if (PinToResolve->GetTypeIndex() != PinToSkip->GetTypeIndex())
 			{
-				if (bSetupUndoRedo)
+				if (!ResolveWildCardPin(PinToResolve, PinToSkip->GetTypeIndex(), bSetupUndoRedo))
 				{
-					GetActionStack()->CancelAction(Action);
+					if (bSetupUndoRedo)
+					{
+						GetActionStack()->CancelAction(Action);
+					}
+					return false;
 				}
-				return false;
 			}
 		}
 	}
