ControlRig: PatchBranchNode should assign new names to the new branch nodes to avoid the following linkerload error.

Failed import: class 'RigVMBranchNode' name 'Branch_1_1' outer 'RigVMModel' There is another object (of 'RigVMUnitNode' class) at the path.

diff --git a/Engine/Plugins/Runtime/RigVM/Source/RigVMDeveloper/Private/RigVMModel/RigVMController.cpp b/Engine/Plugins/Runtime/RigVM/Source/RigVMDeveloper/Private/RigVMModel/RigVMController.cpp
--- a/Engine/Plugins/Runtime/RigVM/Source/RigVMDeveloper/Private/RigVMModel/RigVMController.cpp
+++ b/Engine/Plugins/Runtime/RigVM/Source/RigVMDeveloper/Private/RigVMModel/RigVMController.cpp
@@ -18455,7 +18455,12 @@ FRigVMClientPatchResult URigVMController::PatchBranchNodesOnLoad()
 			FastBreakLinkedPaths(LinkedPaths);
 			RemoveNode(BranchNode, false, false);
 
-			const URigVMNode* NewNode = AddUnitNode(FRigVMFunction_ControlFlowBranch::StaticStruct(), FRigVMStruct::ExecuteName, NodePosition, NodeName, false, false);
+			// Cannot reuse the name of deprecated nodes, otherwise we may get the following error when PIE a second time
+			// Failed import: class 'RigVMBranchNode' name 'Branch_1_1' outer 'RigVMModel'. There is another object (of 'RigVMUnitNode' class) at the path.
+			const FString NewNodeName = NodeName + TEXT("_") + URigVMUnitNode::StaticClass()->GetName();
+			FRestoreLinkedPathSettings RestoreLinkedPathSettings;
+			RestoreLinkedPathSettings.NodeNameMap.Add(NodeName, NewNodeName);
+			const URigVMNode* NewNode = AddUnitNode(FRigVMFunction_ControlFlowBranch::StaticStruct(), FRigVMStruct::ExecuteName, NodePosition, NewNodeName, false, false);
 
 			Result.AddedNodes.Add(NewNode);
 
@@ -18464,7 +18469,7 @@ FRigVMClientPatchResult URigVMController::PatchBranchNodesOnLoad()
 				const URigVMPin* ConditionPin = NewNode->FindPin(GET_MEMBER_NAME_CHECKED(FRigVMFunction_ControlFlowBranch, Condition).ToString());
 				SetPinDefaultValue(ConditionPin->GetPinPath(), ConditionDefault, false, false, false, false);
 			}
-			RestoreLinkedPaths(LinkedPaths);
+			RestoreLinkedPaths(LinkedPaths, RestoreLinkedPathSettings);
 		}
 	}
 	
