RigVM: Replace spaces done with the proper sanitize name function

diff --git a/Engine/Plugins/Runtime/RigVM/Source/RigVMDeveloper/Private/RigVMCompiler/RigVMAST.cpp b/Engine/Plugins/Runtime/RigVM/Source/RigVMDeveloper/Private/RigVMCompiler/RigVMAST.cpp
--- a/Engine/Plugins/Runtime/RigVM/Source/RigVMDeveloper/Private/RigVMCompiler/RigVMAST.cpp
+++ b/Engine/Plugins/Runtime/RigVM/Source/RigVMDeveloper/Private/RigVMCompiler/RigVMAST.cpp
@@ -2946,8 +2946,7 @@ void FRigVMParserAST::Inline(TArray<URigVMGraph*> InGraphs, const TArray<FRigVMA
 						{
 							if (const FRigVMGraphFunctionData* FunctionData = FunctionReferenceNode->GetReferencedFunctionData())
 							{
-								FString PinName = SourcePin->GetName();
-								PinName.ReplaceInline(TEXT(" "), TEXT("_"));
+								const FString PinName = FRigVMPropertyDescription::SanitizeName(SourcePin->GetFName()).ToString();
 								if (const FRigVMFunctionCompilationPropertyDescription* Description = FunctionData->CompilationData.WorkPropertyDescriptions.FindByPredicate([PinName](const FRigVMFunctionCompilationPropertyDescription& Description)
 								{
 									return Description.Name.ToString().EndsWith(PinName);
diff --git a/Engine/Plugins/Runtime/RigVM/Source/RigVMDeveloper/Private/RigVMCompiler/RigVMCompiler.cpp b/Engine/Plugins/Runtime/RigVM/Source/RigVMDeveloper/Private/RigVMCompiler/RigVMCompiler.cpp
--- a/Engine/Plugins/Runtime/RigVM/Source/RigVMDeveloper/Private/RigVMCompiler/RigVMCompiler.cpp
+++ b/Engine/Plugins/Runtime/RigVM/Source/RigVMDeveloper/Private/RigVMCompiler/RigVMCompiler.cpp
@@ -1936,9 +1936,8 @@ int32 URigVMCompiler::TraverseInlineFunction(const FRigVMInlineFunctionExprAST*
 					{
 						continue;
 					}
-					FString PropertyName = Argument.Name.ToString();
-					PropertyName.ReplaceInline(TEXT(" "), TEXT("_"));
-					if (!Properties[NumProperties].Name.ToString().Contains(PropertyName))
+					const FString PropertyName = FRigVMPropertyDescription::SanitizeName(Argument.Name).ToString();
+					if (!Properties[NumProperties].Name.ToString().EndsWith(PropertyName))
 					{
 						continue;
 					}
