Control Rig: Ensure that we get a valid context ptr to reset. Cleaned up duplicated initializations by moving them to the member side.

#rnx

diff --git a/Engine/Plugins/Runtime/RigVM/Source/RigVM/Public/RigVMCore/RigVMExecuteContext.h b/Engine/Plugins/Runtime/RigVM/Source/RigVM/Public/RigVMCore/RigVMExecuteContext.h
--- a/Engine/Plugins/Runtime/RigVM/Source/RigVM/Public/RigVMCore/RigVMExecuteContext.h
+++ b/Engine/Plugins/Runtime/RigVM/Source/RigVM/Public/RigVMCore/RigVMExecuteContext.h
@@ -481,9 +481,6 @@ struct RIGVM_API FRigVMExtendedExecuteContext
 
 	FRigVMExtendedExecuteContext()
 	: PublicDataScope(FRigVMExecuteContext::StaticStruct())
-	, VM(nullptr)
-	, LastExecutionMicroSeconds()
-	, Factory(nullptr)
 	{
 		Reset();
 		SetDefaultNameCache();
@@ -491,9 +488,6 @@ struct RIGVM_API FRigVMExtendedExecuteContext
 
 	FRigVMExtendedExecuteContext(const UScriptStruct* InExecuteContextStruct)
 		: PublicDataScope(InExecuteContextStruct)
-		, VM(nullptr)
-		, LastExecutionMicroSeconds()
-		, Factory(nullptr)
 	{
 		if(InExecuteContextStruct)
 		{
@@ -504,10 +498,6 @@ struct RIGVM_API FRigVMExtendedExecuteContext
 	}
 
 	FRigVMExtendedExecuteContext(const FRigVMExtendedExecuteContext& InOther)
-		: PublicDataScope()
-		, VM(nullptr)
-		, LastExecutionMicroSeconds()
-		, Factory(nullptr)
 	{
 		*this = InOther;
 	}
@@ -519,9 +509,11 @@ struct RIGVM_API FRigVMExtendedExecuteContext
 
 	void Reset()
 	{
-		FRigVMExecuteContext* ExecuteContext = (FRigVMExecuteContext*)(PublicDataScope.GetStructMemory());
-		ExecuteContext->Reset();
-		ExecuteContext->ExtendedExecuteContext = this;
+		if (FRigVMExecuteContext* ExecuteContext = reinterpret_cast<FRigVMExecuteContext*>(PublicDataScope.GetStructMemory()))
+		{
+			ExecuteContext->Reset();
+			ExecuteContext->ExtendedExecuteContext = this;
+		}
 	
 		VM = nullptr;
 		Slices.Reset();
@@ -688,10 +680,10 @@ struct RIGVM_API FRigVMExtendedExecuteContext
 	}
 
 	FStructOnScope PublicDataScope;
-	URigVM* VM;
+	URigVM* VM = nullptr;
 	TArray<FRigVMSlice> Slices;
 	TArray<uint16> SliceOffsets;
-	double LastExecutionMicroSeconds;
-	const FRigVMDispatchFactory* Factory;
+	double LastExecutionMicroSeconds = 0.0;
+	const FRigVMDispatchFactory* Factory = nullptr;
 	FRigVMNameCache NameCache;
 };
