RigVM: New generic Execute opcode, store operandcount on operator

diff --git a/Engine/Plugins/Animation/ControlRig/Source/ControlRigEditor/Private/Editor/SControlRigStackView.cpp b/Engine/Plugins/Animation/ControlRig/Source/ControlRigEditor/Private/Editor/SControlRigStackView.cpp
--- a/Engine/Plugins/Animation/ControlRig/Source/ControlRigEditor/Private/Editor/SControlRigStackView.cpp
+++ b/Engine/Plugins/Animation/ControlRig/Source/ControlRigEditor/Private/Editor/SControlRigStackView.cpp
@@ -504,71 +504,7 @@ void SControlRigStackView::PopulateStackView(URigVM* InVM)
 
 				switch (Instructions[InstructionIndex].OpCode)
 				{
-					case ERigVMOpCode::Execute_0_Operands:
-					case ERigVMOpCode::Execute_1_Operands:
-					case ERigVMOpCode::Execute_2_Operands:
-					case ERigVMOpCode::Execute_3_Operands:
-					case ERigVMOpCode::Execute_4_Operands:
-					case ERigVMOpCode::Execute_5_Operands:
-					case ERigVMOpCode::Execute_6_Operands:
-					case ERigVMOpCode::Execute_7_Operands:
-					case ERigVMOpCode::Execute_8_Operands:
-					case ERigVMOpCode::Execute_9_Operands:
-					case ERigVMOpCode::Execute_10_Operands:
-					case ERigVMOpCode::Execute_11_Operands:
-					case ERigVMOpCode::Execute_12_Operands:
-					case ERigVMOpCode::Execute_13_Operands:
-					case ERigVMOpCode::Execute_14_Operands:
-					case ERigVMOpCode::Execute_15_Operands:
-					case ERigVMOpCode::Execute_16_Operands:
-					case ERigVMOpCode::Execute_17_Operands:
-					case ERigVMOpCode::Execute_18_Operands:
-					case ERigVMOpCode::Execute_19_Operands:
-					case ERigVMOpCode::Execute_20_Operands:
-					case ERigVMOpCode::Execute_21_Operands:
-					case ERigVMOpCode::Execute_22_Operands:
-					case ERigVMOpCode::Execute_23_Operands:
-					case ERigVMOpCode::Execute_24_Operands:
-					case ERigVMOpCode::Execute_25_Operands:
-					case ERigVMOpCode::Execute_26_Operands:
-					case ERigVMOpCode::Execute_27_Operands:
-					case ERigVMOpCode::Execute_28_Operands:
-					case ERigVMOpCode::Execute_29_Operands:
-					case ERigVMOpCode::Execute_30_Operands:
-					case ERigVMOpCode::Execute_31_Operands:
-					case ERigVMOpCode::Execute_32_Operands:
-					case ERigVMOpCode::Execute_33_Operands:
-					case ERigVMOpCode::Execute_34_Operands:
-					case ERigVMOpCode::Execute_35_Operands:
-					case ERigVMOpCode::Execute_36_Operands:
-					case ERigVMOpCode::Execute_37_Operands:
-					case ERigVMOpCode::Execute_38_Operands:
-					case ERigVMOpCode::Execute_39_Operands:
-					case ERigVMOpCode::Execute_40_Operands:
-					case ERigVMOpCode::Execute_41_Operands:
-					case ERigVMOpCode::Execute_42_Operands:
-					case ERigVMOpCode::Execute_43_Operands:
-					case ERigVMOpCode::Execute_44_Operands:
-					case ERigVMOpCode::Execute_45_Operands:
-					case ERigVMOpCode::Execute_46_Operands:
-					case ERigVMOpCode::Execute_47_Operands:
-					case ERigVMOpCode::Execute_48_Operands:
-					case ERigVMOpCode::Execute_49_Operands:
-					case ERigVMOpCode::Execute_50_Operands:
-					case ERigVMOpCode::Execute_51_Operands:
-					case ERigVMOpCode::Execute_52_Operands:
-					case ERigVMOpCode::Execute_53_Operands:
-					case ERigVMOpCode::Execute_54_Operands:
-					case ERigVMOpCode::Execute_55_Operands:
-					case ERigVMOpCode::Execute_56_Operands:
-					case ERigVMOpCode::Execute_57_Operands:
-					case ERigVMOpCode::Execute_58_Operands:
-					case ERigVMOpCode::Execute_59_Operands:
-					case ERigVMOpCode::Execute_60_Operands:
-					case ERigVMOpCode::Execute_61_Operands:
-					case ERigVMOpCode::Execute_62_Operands:
-					case ERigVMOpCode::Execute_63_Operands:
-					case ERigVMOpCode::Execute_64_Operands:
+					case ERigVMOpCode::Execute:
 					{
 						const FRigVMExecuteOp& Op = ByteCode.GetOpAt<FRigVMExecuteOp>(Instructions[InstructionIndex]);
 
diff --git a/Engine/Plugins/Runtime/RigVM/Source/RigVM/Private/RigVMCore/RigVM.cpp b/Engine/Plugins/Runtime/RigVM/Source/RigVM/Private/RigVMCore/RigVM.cpp
--- a/Engine/Plugins/Runtime/RigVM/Source/RigVM/Private/RigVMCore/RigVM.cpp
+++ b/Engine/Plugins/Runtime/RigVM/Source/RigVM/Private/RigVMCore/RigVM.cpp
@@ -394,71 +394,7 @@ bool URigVM::ValidateAllOperandsDuringLoad()
 	{
 		switch (ByteCodeInstruction.OpCode)
 		{
-			case ERigVMOpCode::Execute_0_Operands:
-			case ERigVMOpCode::Execute_1_Operands:
-			case ERigVMOpCode::Execute_2_Operands:
-			case ERigVMOpCode::Execute_3_Operands:
-			case ERigVMOpCode::Execute_4_Operands:
-			case ERigVMOpCode::Execute_5_Operands:
-			case ERigVMOpCode::Execute_6_Operands:
-			case ERigVMOpCode::Execute_7_Operands:
-			case ERigVMOpCode::Execute_8_Operands:
-			case ERigVMOpCode::Execute_9_Operands:
-			case ERigVMOpCode::Execute_10_Operands:
-			case ERigVMOpCode::Execute_11_Operands:
-			case ERigVMOpCode::Execute_12_Operands:
-			case ERigVMOpCode::Execute_13_Operands:
-			case ERigVMOpCode::Execute_14_Operands:
-			case ERigVMOpCode::Execute_15_Operands:
-			case ERigVMOpCode::Execute_16_Operands:
-			case ERigVMOpCode::Execute_17_Operands:
-			case ERigVMOpCode::Execute_18_Operands:
-			case ERigVMOpCode::Execute_19_Operands:
-			case ERigVMOpCode::Execute_20_Operands:
-			case ERigVMOpCode::Execute_21_Operands:
-			case ERigVMOpCode::Execute_22_Operands:
-			case ERigVMOpCode::Execute_23_Operands:
-			case ERigVMOpCode::Execute_24_Operands:
-			case ERigVMOpCode::Execute_25_Operands:
-			case ERigVMOpCode::Execute_26_Operands:
-			case ERigVMOpCode::Execute_27_Operands:
-			case ERigVMOpCode::Execute_28_Operands:
-			case ERigVMOpCode::Execute_29_Operands:
-			case ERigVMOpCode::Execute_30_Operands:
-			case ERigVMOpCode::Execute_31_Operands:
-			case ERigVMOpCode::Execute_32_Operands:
-			case ERigVMOpCode::Execute_33_Operands:
-			case ERigVMOpCode::Execute_34_Operands:
-			case ERigVMOpCode::Execute_35_Operands:
-			case ERigVMOpCode::Execute_36_Operands:
-			case ERigVMOpCode::Execute_37_Operands:
-			case ERigVMOpCode::Execute_38_Operands:
-			case ERigVMOpCode::Execute_39_Operands:
-			case ERigVMOpCode::Execute_40_Operands:
-			case ERigVMOpCode::Execute_41_Operands:
-			case ERigVMOpCode::Execute_42_Operands:
-			case ERigVMOpCode::Execute_43_Operands:
-			case ERigVMOpCode::Execute_44_Operands:
-			case ERigVMOpCode::Execute_45_Operands:
-			case ERigVMOpCode::Execute_46_Operands:
-			case ERigVMOpCode::Execute_47_Operands:
-			case ERigVMOpCode::Execute_48_Operands:
-			case ERigVMOpCode::Execute_49_Operands:
-			case ERigVMOpCode::Execute_50_Operands:
-			case ERigVMOpCode::Execute_51_Operands:
-			case ERigVMOpCode::Execute_52_Operands:
-			case ERigVMOpCode::Execute_53_Operands:
-			case ERigVMOpCode::Execute_54_Operands:
-			case ERigVMOpCode::Execute_55_Operands:
-			case ERigVMOpCode::Execute_56_Operands:
-			case ERigVMOpCode::Execute_57_Operands:
-			case ERigVMOpCode::Execute_58_Operands:
-			case ERigVMOpCode::Execute_59_Operands:
-			case ERigVMOpCode::Execute_60_Operands:
-			case ERigVMOpCode::Execute_61_Operands:
-			case ERigVMOpCode::Execute_62_Operands:
-			case ERigVMOpCode::Execute_63_Operands:
-			case ERigVMOpCode::Execute_64_Operands:
+			case ERigVMOpCode::Execute:
 			{
 				const FRigVMExecuteOp& Op = ByteCodeStorage.GetOpAt<FRigVMExecuteOp>(ByteCodeInstruction);
 				FRigVMOperandArray Operands = ByteCodeStorage.GetOperandsForExecuteOp(ByteCodeInstruction);
@@ -1088,7 +1024,7 @@ void URigVM::CacheMemoryHandlesIfRequired(FRigVMExtendedExecuteContext& Context,
 	{
 		const ERigVMOpCode OpCode = Instructions[InstructionIndex].OpCode; 
 
-		if (OpCode >= ERigVMOpCode::Execute_0_Operands && OpCode <= ERigVMOpCode::Execute_64_Operands)
+		if (OpCode == ERigVMOpCode::Execute)
 		{
 			const FRigVMExecuteOp& Op = ByteCode.GetOpAt<FRigVMExecuteOp>(Instructions[InstructionIndex]);
 			FRigVMOperandArray Operands = ByteCode.GetOperandsForExecuteOp(Instructions[InstructionIndex]);
@@ -1732,7 +1668,7 @@ ERigVMExecuteResult URigVM::ExecuteInstructions(FRigVMExtendedExecuteContext& Co
 #if UE_RIGVM_DEBUG_EXECUTION
 		if (ContextPublicData.bDebugExecution)
 		{
-			if (Instruction.OpCode >= ERigVMOpCode::Execute_0_Operands && Instruction.OpCode <= ERigVMOpCode::Execute_64_Operands)
+			if (Instruction.OpCode == ERigVMOpCode::Execute)
 			{
 				const FRigVMExecuteOp& Op = ByteCode.GetOpAt<FRigVMExecuteOp>(Instruction);
 				FRigVMOperandArray Operands = ByteCode.GetOperandsForExecuteOp(Instructions[ContextPublicData.InstructionIndex]);
@@ -1764,71 +1700,7 @@ ERigVMExecuteResult URigVM::ExecuteInstructions(FRigVMExtendedExecuteContext& Co
 
 		switch (Instruction.OpCode)
 		{
-			case ERigVMOpCode::Execute_0_Operands:
-			case ERigVMOpCode::Execute_1_Operands:
-			case ERigVMOpCode::Execute_2_Operands:
-			case ERigVMOpCode::Execute_3_Operands:
-			case ERigVMOpCode::Execute_4_Operands:
-			case ERigVMOpCode::Execute_5_Operands:
-			case ERigVMOpCode::Execute_6_Operands:
-			case ERigVMOpCode::Execute_7_Operands:
-			case ERigVMOpCode::Execute_8_Operands:
-			case ERigVMOpCode::Execute_9_Operands:
-			case ERigVMOpCode::Execute_10_Operands:
-			case ERigVMOpCode::Execute_11_Operands:
-			case ERigVMOpCode::Execute_12_Operands:
-			case ERigVMOpCode::Execute_13_Operands:
-			case ERigVMOpCode::Execute_14_Operands:
-			case ERigVMOpCode::Execute_15_Operands:
-			case ERigVMOpCode::Execute_16_Operands:
-			case ERigVMOpCode::Execute_17_Operands:
-			case ERigVMOpCode::Execute_18_Operands:
-			case ERigVMOpCode::Execute_19_Operands:
-			case ERigVMOpCode::Execute_20_Operands:
-			case ERigVMOpCode::Execute_21_Operands:
-			case ERigVMOpCode::Execute_22_Operands:
-			case ERigVMOpCode::Execute_23_Operands:
-			case ERigVMOpCode::Execute_24_Operands:
-			case ERigVMOpCode::Execute_25_Operands:
-			case ERigVMOpCode::Execute_26_Operands:
-			case ERigVMOpCode::Execute_27_Operands:
-			case ERigVMOpCode::Execute_28_Operands:
-			case ERigVMOpCode::Execute_29_Operands:
-			case ERigVMOpCode::Execute_30_Operands:
-			case ERigVMOpCode::Execute_31_Operands:
-			case ERigVMOpCode::Execute_32_Operands:
-			case ERigVMOpCode::Execute_33_Operands:
-			case ERigVMOpCode::Execute_34_Operands:
-			case ERigVMOpCode::Execute_35_Operands:
-			case ERigVMOpCode::Execute_36_Operands:
-			case ERigVMOpCode::Execute_37_Operands:
-			case ERigVMOpCode::Execute_38_Operands:
-			case ERigVMOpCode::Execute_39_Operands:
-			case ERigVMOpCode::Execute_40_Operands:
-			case ERigVMOpCode::Execute_41_Operands:
-			case ERigVMOpCode::Execute_42_Operands:
-			case ERigVMOpCode::Execute_43_Operands:
-			case ERigVMOpCode::Execute_44_Operands:
-			case ERigVMOpCode::Execute_45_Operands:
-			case ERigVMOpCode::Execute_46_Operands:
-			case ERigVMOpCode::Execute_47_Operands:
-			case ERigVMOpCode::Execute_48_Operands:
-			case ERigVMOpCode::Execute_49_Operands:
-			case ERigVMOpCode::Execute_50_Operands:
-			case ERigVMOpCode::Execute_51_Operands:
-			case ERigVMOpCode::Execute_52_Operands:
-			case ERigVMOpCode::Execute_53_Operands:
-			case ERigVMOpCode::Execute_54_Operands:
-			case ERigVMOpCode::Execute_55_Operands:
-			case ERigVMOpCode::Execute_56_Operands:
-			case ERigVMOpCode::Execute_57_Operands:
-			case ERigVMOpCode::Execute_58_Operands:
-			case ERigVMOpCode::Execute_59_Operands:
-			case ERigVMOpCode::Execute_60_Operands:
-			case ERigVMOpCode::Execute_61_Operands:
-			case ERigVMOpCode::Execute_62_Operands:
-			case ERigVMOpCode::Execute_63_Operands:
-			case ERigVMOpCode::Execute_64_Operands:
+			case ERigVMOpCode::Execute:
 			{
 				const FRigVMExecuteOp& Op = ByteCode.GetOpAt<FRigVMExecuteOp>(Instruction);
 				const int32 OperandCount = FirstHandleForInstruction[ContextPublicData.InstructionIndex + 1] - FirstHandleForInstruction[ContextPublicData.InstructionIndex];
@@ -2271,71 +2143,7 @@ TArray<FString> URigVM::DumpByteCodeAsTextArray(const TArray<int32>& InInstructi
 
 		switch (Instructions[InstructionIndex].OpCode)
 		{
-			case ERigVMOpCode::Execute_0_Operands:
-			case ERigVMOpCode::Execute_1_Operands:
-			case ERigVMOpCode::Execute_2_Operands:
-			case ERigVMOpCode::Execute_3_Operands:
-			case ERigVMOpCode::Execute_4_Operands:
-			case ERigVMOpCode::Execute_5_Operands:
-			case ERigVMOpCode::Execute_6_Operands:
-			case ERigVMOpCode::Execute_7_Operands:
-			case ERigVMOpCode::Execute_8_Operands:
-			case ERigVMOpCode::Execute_9_Operands:
-			case ERigVMOpCode::Execute_10_Operands:
-			case ERigVMOpCode::Execute_11_Operands:
-			case ERigVMOpCode::Execute_12_Operands:
-			case ERigVMOpCode::Execute_13_Operands:
-			case ERigVMOpCode::Execute_14_Operands:
-			case ERigVMOpCode::Execute_15_Operands:
-			case ERigVMOpCode::Execute_16_Operands:
-			case ERigVMOpCode::Execute_17_Operands:
-			case ERigVMOpCode::Execute_18_Operands:
-			case ERigVMOpCode::Execute_19_Operands:
-			case ERigVMOpCode::Execute_20_Operands:
-			case ERigVMOpCode::Execute_21_Operands:
-			case ERigVMOpCode::Execute_22_Operands:
-			case ERigVMOpCode::Execute_23_Operands:
-			case ERigVMOpCode::Execute_24_Operands:
-			case ERigVMOpCode::Execute_25_Operands:
-			case ERigVMOpCode::Execute_26_Operands:
-			case ERigVMOpCode::Execute_27_Operands:
-			case ERigVMOpCode::Execute_28_Operands:
-			case ERigVMOpCode::Execute_29_Operands:
-			case ERigVMOpCode::Execute_30_Operands:
-			case ERigVMOpCode::Execute_31_Operands:
-			case ERigVMOpCode::Execute_32_Operands:
-			case ERigVMOpCode::Execute_33_Operands:
-			case ERigVMOpCode::Execute_34_Operands:
-			case ERigVMOpCode::Execute_35_Operands:
-			case ERigVMOpCode::Execute_36_Operands:
-			case ERigVMOpCode::Execute_37_Operands:
-			case ERigVMOpCode::Execute_38_Operands:
-			case ERigVMOpCode::Execute_39_Operands:
-			case ERigVMOpCode::Execute_40_Operands:
-			case ERigVMOpCode::Execute_41_Operands:
-			case ERigVMOpCode::Execute_42_Operands:
-			case ERigVMOpCode::Execute_43_Operands:
-			case ERigVMOpCode::Execute_44_Operands:
-			case ERigVMOpCode::Execute_45_Operands:
-			case ERigVMOpCode::Execute_46_Operands:
-			case ERigVMOpCode::Execute_47_Operands:
-			case ERigVMOpCode::Execute_48_Operands:
-			case ERigVMOpCode::Execute_49_Operands:
-			case ERigVMOpCode::Execute_50_Operands:
-			case ERigVMOpCode::Execute_51_Operands:
-			case ERigVMOpCode::Execute_52_Operands:
-			case ERigVMOpCode::Execute_53_Operands:
-			case ERigVMOpCode::Execute_54_Operands:
-			case ERigVMOpCode::Execute_55_Operands:
-			case ERigVMOpCode::Execute_56_Operands:
-			case ERigVMOpCode::Execute_57_Operands:
-			case ERigVMOpCode::Execute_58_Operands:
-			case ERigVMOpCode::Execute_59_Operands:
-			case ERigVMOpCode::Execute_60_Operands:
-			case ERigVMOpCode::Execute_61_Operands:
-			case ERigVMOpCode::Execute_62_Operands:
-			case ERigVMOpCode::Execute_63_Operands:
-			case ERigVMOpCode::Execute_64_Operands:
+			case ERigVMOpCode::Execute:
 			{
 				const FRigVMExecuteOp& Op = ByteCode.GetOpAt<FRigVMExecuteOp>(Instructions[InstructionIndex]);
 				FString FunctionName = FunctionNames[Op.FunctionIndex].ToString();
diff --git a/Engine/Plugins/Runtime/RigVM/Source/RigVM/Private/RigVMCore/RigVMByteCode.cpp b/Engine/Plugins/Runtime/RigVM/Source/RigVM/Private/RigVMCore/RigVMByteCode.cpp
--- a/Engine/Plugins/Runtime/RigVM/Source/RigVM/Private/RigVMCore/RigVMByteCode.cpp
+++ b/Engine/Plugins/Runtime/RigVM/Source/RigVM/Private/RigVMCore/RigVMByteCode.cpp
@@ -11,6 +11,25 @@ void FRigVMExecuteOp::Serialize(FArchive& Ar)
 {
 	Ar << OpCode;
 	Ar << FunctionIndex;
+
+	if(Ar.IsLoading())
+	{
+		// backwards compatibility for old opcodes
+		if(OpCode >= ERigVMOpCode::Execute_0_Operands && OpCode <= ERigVMOpCode::Execute_64_Operands)
+		{
+			ArgumentCount = uint16(OpCode) - uint16(ERigVMOpCode::Execute_0_Operands);
+			OpCode = ERigVMOpCode::Execute;
+		}
+		else
+		{
+			check(OpCode == ERigVMOpCode::Execute);
+			Ar << ArgumentCount;
+		}
+	}
+	else
+	{
+		Ar << ArgumentCount;
+	}
 }
 
 void FRigVMUnaryOp::Serialize(FArchive& Ar)
@@ -171,7 +190,7 @@ FRigVMInstructionArray::FRigVMInstructionArray(const FRigVMByteCode& InByteCode,
 				}
 			}
 
-			if (OpCode >= ERigVMOpCode::Execute_0_Operands && OpCode < ERigVMOpCode::Execute_64_Operands)
+			if (OpCode == ERigVMOpCode::Execute)
 			{
 				uint64 OperandByteIndex = ByteIndex + (uint64)sizeof(FRigVMExecuteOp);
 
@@ -264,71 +283,7 @@ void FRigVMByteCode::Save(FArchive& Ar) const
 
 		switch (OpCode)
 		{
-			case ERigVMOpCode::Execute_0_Operands:
-			case ERigVMOpCode::Execute_1_Operands:
-			case ERigVMOpCode::Execute_2_Operands:
-			case ERigVMOpCode::Execute_3_Operands:
-			case ERigVMOpCode::Execute_4_Operands:
-			case ERigVMOpCode::Execute_5_Operands:
-			case ERigVMOpCode::Execute_6_Operands:
-			case ERigVMOpCode::Execute_7_Operands:
-			case ERigVMOpCode::Execute_8_Operands:
-			case ERigVMOpCode::Execute_9_Operands:
-			case ERigVMOpCode::Execute_10_Operands:
-			case ERigVMOpCode::Execute_11_Operands:
-			case ERigVMOpCode::Execute_12_Operands:
-			case ERigVMOpCode::Execute_13_Operands:
-			case ERigVMOpCode::Execute_14_Operands:
-			case ERigVMOpCode::Execute_15_Operands:
-			case ERigVMOpCode::Execute_16_Operands:
-			case ERigVMOpCode::Execute_17_Operands:
-			case ERigVMOpCode::Execute_18_Operands:
-			case ERigVMOpCode::Execute_19_Operands:
-			case ERigVMOpCode::Execute_20_Operands:
-			case ERigVMOpCode::Execute_21_Operands:
-			case ERigVMOpCode::Execute_22_Operands:
-			case ERigVMOpCode::Execute_23_Operands:
-			case ERigVMOpCode::Execute_24_Operands:
-			case ERigVMOpCode::Execute_25_Operands:
-			case ERigVMOpCode::Execute_26_Operands:
-			case ERigVMOpCode::Execute_27_Operands:
-			case ERigVMOpCode::Execute_28_Operands:
-			case ERigVMOpCode::Execute_29_Operands:
-			case ERigVMOpCode::Execute_30_Operands:
-			case ERigVMOpCode::Execute_31_Operands:
-			case ERigVMOpCode::Execute_32_Operands:
-			case ERigVMOpCode::Execute_33_Operands:
-			case ERigVMOpCode::Execute_34_Operands:
-			case ERigVMOpCode::Execute_35_Operands:
-			case ERigVMOpCode::Execute_36_Operands:
-			case ERigVMOpCode::Execute_37_Operands:
-			case ERigVMOpCode::Execute_38_Operands:
-			case ERigVMOpCode::Execute_39_Operands:
-			case ERigVMOpCode::Execute_40_Operands:
-			case ERigVMOpCode::Execute_41_Operands:
-			case ERigVMOpCode::Execute_42_Operands:
-			case ERigVMOpCode::Execute_43_Operands:
-			case ERigVMOpCode::Execute_44_Operands:
-			case ERigVMOpCode::Execute_45_Operands:
-			case ERigVMOpCode::Execute_46_Operands:
-			case ERigVMOpCode::Execute_47_Operands:
-			case ERigVMOpCode::Execute_48_Operands:
-			case ERigVMOpCode::Execute_49_Operands:
-			case ERigVMOpCode::Execute_50_Operands:
-			case ERigVMOpCode::Execute_51_Operands:
-			case ERigVMOpCode::Execute_52_Operands:
-			case ERigVMOpCode::Execute_53_Operands:
-			case ERigVMOpCode::Execute_54_Operands:
-			case ERigVMOpCode::Execute_55_Operands:
-			case ERigVMOpCode::Execute_56_Operands:
-			case ERigVMOpCode::Execute_57_Operands:
-			case ERigVMOpCode::Execute_58_Operands:
-			case ERigVMOpCode::Execute_59_Operands:
-			case ERigVMOpCode::Execute_60_Operands:
-			case ERigVMOpCode::Execute_61_Operands:
-			case ERigVMOpCode::Execute_62_Operands:
-			case ERigVMOpCode::Execute_63_Operands:
-			case ERigVMOpCode::Execute_64_Operands:
+			case ERigVMOpCode::Execute:
 			{
 				FRigVMExecuteOp Op = GetOpAt<FRigVMExecuteOp>(Instruction.ByteCodeIndex);
 				Ar << Op;
@@ -493,73 +448,15 @@ void FRigVMByteCode::Load(FArchive& Ar)
 
 		Ar << OpCode;
 
+		// backwards compatibility
+		if(OpCode >= ERigVMOpCode::Execute_0_Operands && OpCode <= ERigVMOpCode::Execute_64_Operands)
+		{
+			OpCode = ERigVMOpCode::Execute;
+		}
+
 		switch (OpCode)
 		{
-			case ERigVMOpCode::Execute_0_Operands:
-			case ERigVMOpCode::Execute_1_Operands:
-			case ERigVMOpCode::Execute_2_Operands:
-			case ERigVMOpCode::Execute_3_Operands:
-			case ERigVMOpCode::Execute_4_Operands:
-			case ERigVMOpCode::Execute_5_Operands:
-			case ERigVMOpCode::Execute_6_Operands:
-			case ERigVMOpCode::Execute_7_Operands:
-			case ERigVMOpCode::Execute_8_Operands:
-			case ERigVMOpCode::Execute_9_Operands:
-			case ERigVMOpCode::Execute_10_Operands:
-			case ERigVMOpCode::Execute_11_Operands:
-			case ERigVMOpCode::Execute_12_Operands:
-			case ERigVMOpCode::Execute_13_Operands:
-			case ERigVMOpCode::Execute_14_Operands:
-			case ERigVMOpCode::Execute_15_Operands:
-			case ERigVMOpCode::Execute_16_Operands:
-			case ERigVMOpCode::Execute_17_Operands:
-			case ERigVMOpCode::Execute_18_Operands:
-			case ERigVMOpCode::Execute_19_Operands:
-			case ERigVMOpCode::Execute_20_Operands:
-			case ERigVMOpCode::Execute_21_Operands:
-			case ERigVMOpCode::Execute_22_Operands:
-			case ERigVMOpCode::Execute_23_Operands:
-			case ERigVMOpCode::Execute_24_Operands:
-			case ERigVMOpCode::Execute_25_Operands:
-			case ERigVMOpCode::Execute_26_Operands:
-			case ERigVMOpCode::Execute_27_Operands:
-			case ERigVMOpCode::Execute_28_Operands:
-			case ERigVMOpCode::Execute_29_Operands:
-			case ERigVMOpCode::Execute_30_Operands:
-			case ERigVMOpCode::Execute_31_Operands:
-			case ERigVMOpCode::Execute_32_Operands:
-			case ERigVMOpCode::Execute_33_Operands:
-			case ERigVMOpCode::Execute_34_Operands:
-			case ERigVMOpCode::Execute_35_Operands:
-			case ERigVMOpCode::Execute_36_Operands:
-			case ERigVMOpCode::Execute_37_Operands:
-			case ERigVMOpCode::Execute_38_Operands:
-			case ERigVMOpCode::Execute_39_Operands:
-			case ERigVMOpCode::Execute_40_Operands:
-			case ERigVMOpCode::Execute_41_Operands:
-			case ERigVMOpCode::Execute_42_Operands:
-			case ERigVMOpCode::Execute_43_Operands:
-			case ERigVMOpCode::Execute_44_Operands:
-			case ERigVMOpCode::Execute_45_Operands:
-			case ERigVMOpCode::Execute_46_Operands:
-			case ERigVMOpCode::Execute_47_Operands:
-			case ERigVMOpCode::Execute_48_Operands:
-			case ERigVMOpCode::Execute_49_Operands:
-			case ERigVMOpCode::Execute_50_Operands:
-			case ERigVMOpCode::Execute_51_Operands:
-			case ERigVMOpCode::Execute_52_Operands:
-			case ERigVMOpCode::Execute_53_Operands:
-			case ERigVMOpCode::Execute_54_Operands:
-			case ERigVMOpCode::Execute_55_Operands:
-			case ERigVMOpCode::Execute_56_Operands:
-			case ERigVMOpCode::Execute_57_Operands:
-			case ERigVMOpCode::Execute_58_Operands:
-			case ERigVMOpCode::Execute_59_Operands:
-			case ERigVMOpCode::Execute_60_Operands:
-			case ERigVMOpCode::Execute_61_Operands:
-			case ERigVMOpCode::Execute_62_Operands:
-			case ERigVMOpCode::Execute_63_Operands:
-			case ERigVMOpCode::Execute_64_Operands:
+			case ERigVMOpCode::Execute:
 			{
 				FRigVMExecuteOp Op;
 				Ar << Op;
@@ -796,71 +693,7 @@ uint32 FRigVMByteCode::GetOperatorHash(const FRigVMInstruction& InInstruction) c
 	const ERigVMOpCode OpCode = GetOpCodeAt(InInstruction.ByteCodeIndex);
 	switch (OpCode)
 	{
-		case ERigVMOpCode::Execute_0_Operands:
-		case ERigVMOpCode::Execute_1_Operands:
-		case ERigVMOpCode::Execute_2_Operands:
-		case ERigVMOpCode::Execute_3_Operands:
-		case ERigVMOpCode::Execute_4_Operands:
-		case ERigVMOpCode::Execute_5_Operands:
-		case ERigVMOpCode::Execute_6_Operands:
-		case ERigVMOpCode::Execute_7_Operands:
-		case ERigVMOpCode::Execute_8_Operands:
-		case ERigVMOpCode::Execute_9_Operands:
-		case ERigVMOpCode::Execute_10_Operands:
-		case ERigVMOpCode::Execute_11_Operands:
-		case ERigVMOpCode::Execute_12_Operands:
-		case ERigVMOpCode::Execute_13_Operands:
-		case ERigVMOpCode::Execute_14_Operands:
-		case ERigVMOpCode::Execute_15_Operands:
-		case ERigVMOpCode::Execute_16_Operands:
-		case ERigVMOpCode::Execute_17_Operands:
-		case ERigVMOpCode::Execute_18_Operands:
-		case ERigVMOpCode::Execute_19_Operands:
-		case ERigVMOpCode::Execute_20_Operands:
-		case ERigVMOpCode::Execute_21_Operands:
-		case ERigVMOpCode::Execute_22_Operands:
-		case ERigVMOpCode::Execute_23_Operands:
-		case ERigVMOpCode::Execute_24_Operands:
-		case ERigVMOpCode::Execute_25_Operands:
-		case ERigVMOpCode::Execute_26_Operands:
-		case ERigVMOpCode::Execute_27_Operands:
-		case ERigVMOpCode::Execute_28_Operands:
-		case ERigVMOpCode::Execute_29_Operands:
-		case ERigVMOpCode::Execute_30_Operands:
-		case ERigVMOpCode::Execute_31_Operands:
-		case ERigVMOpCode::Execute_32_Operands:
-		case ERigVMOpCode::Execute_33_Operands:
-		case ERigVMOpCode::Execute_34_Operands:
-		case ERigVMOpCode::Execute_35_Operands:
-		case ERigVMOpCode::Execute_36_Operands:
-		case ERigVMOpCode::Execute_37_Operands:
-		case ERigVMOpCode::Execute_38_Operands:
-		case ERigVMOpCode::Execute_39_Operands:
-		case ERigVMOpCode::Execute_40_Operands:
-		case ERigVMOpCode::Execute_41_Operands:
-		case ERigVMOpCode::Execute_42_Operands:
-		case ERigVMOpCode::Execute_43_Operands:
-		case ERigVMOpCode::Execute_44_Operands:
-		case ERigVMOpCode::Execute_45_Operands:
-		case ERigVMOpCode::Execute_46_Operands:
-		case ERigVMOpCode::Execute_47_Operands:
-		case ERigVMOpCode::Execute_48_Operands:
-		case ERigVMOpCode::Execute_49_Operands:
-		case ERigVMOpCode::Execute_50_Operands:
-		case ERigVMOpCode::Execute_51_Operands:
-		case ERigVMOpCode::Execute_52_Operands:
-		case ERigVMOpCode::Execute_53_Operands:
-		case ERigVMOpCode::Execute_54_Operands:
-		case ERigVMOpCode::Execute_55_Operands:
-		case ERigVMOpCode::Execute_56_Operands:
-		case ERigVMOpCode::Execute_57_Operands:
-		case ERigVMOpCode::Execute_58_Operands:
-		case ERigVMOpCode::Execute_59_Operands:
-		case ERigVMOpCode::Execute_60_Operands:
-		case ERigVMOpCode::Execute_61_Operands:
-		case ERigVMOpCode::Execute_62_Operands:
-		case ERigVMOpCode::Execute_63_Operands:
-		case ERigVMOpCode::Execute_64_Operands:
+		case ERigVMOpCode::Execute:
 		{
 			const FRigVMExecuteOp& Op = GetOpAt<FRigVMExecuteOp>(InInstruction);
 			uint32 Hash = GetTypeHash(Op);
@@ -991,71 +824,7 @@ uint64 FRigVMByteCode::GetOpNumBytesAt(uint64 InByteCodeIndex, bool bIncludeOper
 	ERigVMOpCode OpCode = GetOpCodeAt(InByteCodeIndex);
 	switch (OpCode)
 	{
-		case ERigVMOpCode::Execute_0_Operands:
-		case ERigVMOpCode::Execute_1_Operands:
-		case ERigVMOpCode::Execute_2_Operands:
-		case ERigVMOpCode::Execute_3_Operands:
-		case ERigVMOpCode::Execute_4_Operands:
-		case ERigVMOpCode::Execute_5_Operands:
-		case ERigVMOpCode::Execute_6_Operands:
-		case ERigVMOpCode::Execute_7_Operands:
-		case ERigVMOpCode::Execute_8_Operands:
-		case ERigVMOpCode::Execute_9_Operands:
-		case ERigVMOpCode::Execute_10_Operands:
-		case ERigVMOpCode::Execute_11_Operands:
-		case ERigVMOpCode::Execute_12_Operands:
-		case ERigVMOpCode::Execute_13_Operands:
-		case ERigVMOpCode::Execute_14_Operands:
-		case ERigVMOpCode::Execute_15_Operands:
-		case ERigVMOpCode::Execute_16_Operands:
-		case ERigVMOpCode::Execute_17_Operands:
-		case ERigVMOpCode::Execute_18_Operands:
-		case ERigVMOpCode::Execute_19_Operands:
-		case ERigVMOpCode::Execute_20_Operands:
-		case ERigVMOpCode::Execute_21_Operands:
-		case ERigVMOpCode::Execute_22_Operands:
-		case ERigVMOpCode::Execute_23_Operands:
-		case ERigVMOpCode::Execute_24_Operands:
-		case ERigVMOpCode::Execute_25_Operands:
-		case ERigVMOpCode::Execute_26_Operands:
-		case ERigVMOpCode::Execute_27_Operands:
-		case ERigVMOpCode::Execute_28_Operands:
-		case ERigVMOpCode::Execute_29_Operands:
-		case ERigVMOpCode::Execute_30_Operands:
-		case ERigVMOpCode::Execute_31_Operands:
-		case ERigVMOpCode::Execute_32_Operands:
-		case ERigVMOpCode::Execute_33_Operands:
-		case ERigVMOpCode::Execute_34_Operands:
-		case ERigVMOpCode::Execute_35_Operands:
-		case ERigVMOpCode::Execute_36_Operands:
-		case ERigVMOpCode::Execute_37_Operands:
-		case ERigVMOpCode::Execute_38_Operands:
-		case ERigVMOpCode::Execute_39_Operands:
-		case ERigVMOpCode::Execute_40_Operands:
-		case ERigVMOpCode::Execute_41_Operands:
-		case ERigVMOpCode::Execute_42_Operands:
-		case ERigVMOpCode::Execute_43_Operands:
-		case ERigVMOpCode::Execute_44_Operands:
-		case ERigVMOpCode::Execute_45_Operands:
-		case ERigVMOpCode::Execute_46_Operands:
-		case ERigVMOpCode::Execute_47_Operands:
-		case ERigVMOpCode::Execute_48_Operands:
-		case ERigVMOpCode::Execute_49_Operands:
-		case ERigVMOpCode::Execute_50_Operands:
-		case ERigVMOpCode::Execute_51_Operands:
-		case ERigVMOpCode::Execute_52_Operands:
-		case ERigVMOpCode::Execute_53_Operands:
-		case ERigVMOpCode::Execute_54_Operands:
-		case ERigVMOpCode::Execute_55_Operands:
-		case ERigVMOpCode::Execute_56_Operands:
-		case ERigVMOpCode::Execute_57_Operands:
-		case ERigVMOpCode::Execute_58_Operands:
-		case ERigVMOpCode::Execute_59_Operands:
-		case ERigVMOpCode::Execute_60_Operands:
-		case ERigVMOpCode::Execute_61_Operands:
-		case ERigVMOpCode::Execute_62_Operands:
-		case ERigVMOpCode::Execute_63_Operands:
-		case ERigVMOpCode::Execute_64_Operands:
+		case ERigVMOpCode::Execute:
 		{
 			uint64 NumBytes = (uint64)sizeof(FRigVMExecuteOp);
 			if(bIncludeOperands)
@@ -1310,71 +1079,7 @@ FString FRigVMByteCode::DumpToText() const
 
 		switch (OpCode)
 		{
-			case ERigVMOpCode::Execute_0_Operands:
-			case ERigVMOpCode::Execute_1_Operands:
-			case ERigVMOpCode::Execute_2_Operands:
-			case ERigVMOpCode::Execute_3_Operands:
-			case ERigVMOpCode::Execute_4_Operands:
-			case ERigVMOpCode::Execute_5_Operands:
-			case ERigVMOpCode::Execute_6_Operands:
-			case ERigVMOpCode::Execute_7_Operands:
-			case ERigVMOpCode::Execute_8_Operands:
-			case ERigVMOpCode::Execute_9_Operands:
-			case ERigVMOpCode::Execute_10_Operands:
-			case ERigVMOpCode::Execute_11_Operands:
-			case ERigVMOpCode::Execute_12_Operands:
-			case ERigVMOpCode::Execute_13_Operands:
-			case ERigVMOpCode::Execute_14_Operands:
-			case ERigVMOpCode::Execute_15_Operands:
-			case ERigVMOpCode::Execute_16_Operands:
-			case ERigVMOpCode::Execute_17_Operands:
-			case ERigVMOpCode::Execute_18_Operands:
-			case ERigVMOpCode::Execute_19_Operands:
-			case ERigVMOpCode::Execute_20_Operands:
-			case ERigVMOpCode::Execute_21_Operands:
-			case ERigVMOpCode::Execute_22_Operands:
-			case ERigVMOpCode::Execute_23_Operands:
-			case ERigVMOpCode::Execute_24_Operands:
-			case ERigVMOpCode::Execute_25_Operands:
-			case ERigVMOpCode::Execute_26_Operands:
-			case ERigVMOpCode::Execute_27_Operands:
-			case ERigVMOpCode::Execute_28_Operands:
-			case ERigVMOpCode::Execute_29_Operands:
-			case ERigVMOpCode::Execute_30_Operands:
-			case ERigVMOpCode::Execute_31_Operands:
-			case ERigVMOpCode::Execute_32_Operands:
-			case ERigVMOpCode::Execute_33_Operands:
-			case ERigVMOpCode::Execute_34_Operands:
-			case ERigVMOpCode::Execute_35_Operands:
-			case ERigVMOpCode::Execute_36_Operands:
-			case ERigVMOpCode::Execute_37_Operands:
-			case ERigVMOpCode::Execute_38_Operands:
-			case ERigVMOpCode::Execute_39_Operands:
-			case ERigVMOpCode::Execute_40_Operands:
-			case ERigVMOpCode::Execute_41_Operands:
-			case ERigVMOpCode::Execute_42_Operands:
-			case ERigVMOpCode::Execute_43_Operands:
-			case ERigVMOpCode::Execute_44_Operands:
-			case ERigVMOpCode::Execute_45_Operands:
-			case ERigVMOpCode::Execute_46_Operands:
-			case ERigVMOpCode::Execute_47_Operands:
-			case ERigVMOpCode::Execute_48_Operands:
-			case ERigVMOpCode::Execute_49_Operands:
-			case ERigVMOpCode::Execute_50_Operands:
-			case ERigVMOpCode::Execute_51_Operands:
-			case ERigVMOpCode::Execute_52_Operands:
-			case ERigVMOpCode::Execute_53_Operands:
-			case ERigVMOpCode::Execute_54_Operands:
-			case ERigVMOpCode::Execute_55_Operands:
-			case ERigVMOpCode::Execute_56_Operands:
-			case ERigVMOpCode::Execute_57_Operands:
-			case ERigVMOpCode::Execute_58_Operands:
-			case ERigVMOpCode::Execute_59_Operands:
-			case ERigVMOpCode::Execute_60_Operands:
-			case ERigVMOpCode::Execute_61_Operands:
-			case ERigVMOpCode::Execute_62_Operands:
-			case ERigVMOpCode::Execute_63_Operands:
-			case ERigVMOpCode::Execute_64_Operands:
+			case ERigVMOpCode::Execute:
 			{
 				const FRigVMExecuteOp& Op = GetOpAt<FRigVMExecuteOp>(Instruction.ByteCodeIndex);
 				Line += FString::Printf(TEXT(", FunctionIndex %d"), (int32)Op.FunctionIndex);
@@ -1562,71 +1267,7 @@ FRigVMOperandArray FRigVMByteCode::GetOperandsForOp(const FRigVMInstruction& InI
 {
 	switch(InInstruction.OpCode)
 	{
-		case ERigVMOpCode::Execute_0_Operands:
-		case ERigVMOpCode::Execute_1_Operands:
-		case ERigVMOpCode::Execute_2_Operands:
-		case ERigVMOpCode::Execute_3_Operands:
-		case ERigVMOpCode::Execute_4_Operands:
-		case ERigVMOpCode::Execute_5_Operands:
-		case ERigVMOpCode::Execute_6_Operands:
-		case ERigVMOpCode::Execute_7_Operands:
-		case ERigVMOpCode::Execute_8_Operands:
-		case ERigVMOpCode::Execute_9_Operands:
-		case ERigVMOpCode::Execute_10_Operands:
-		case ERigVMOpCode::Execute_11_Operands:
-		case ERigVMOpCode::Execute_12_Operands:
-		case ERigVMOpCode::Execute_13_Operands:
-		case ERigVMOpCode::Execute_14_Operands:
-		case ERigVMOpCode::Execute_15_Operands:
-		case ERigVMOpCode::Execute_16_Operands:
-		case ERigVMOpCode::Execute_17_Operands:
-		case ERigVMOpCode::Execute_18_Operands:
-		case ERigVMOpCode::Execute_19_Operands:
-		case ERigVMOpCode::Execute_20_Operands:
-		case ERigVMOpCode::Execute_21_Operands:
-		case ERigVMOpCode::Execute_22_Operands:
-		case ERigVMOpCode::Execute_23_Operands:
-		case ERigVMOpCode::Execute_24_Operands:
-		case ERigVMOpCode::Execute_25_Operands:
-		case ERigVMOpCode::Execute_26_Operands:
-		case ERigVMOpCode::Execute_27_Operands:
-		case ERigVMOpCode::Execute_28_Operands:
-		case ERigVMOpCode::Execute_29_Operands:
-		case ERigVMOpCode::Execute_30_Operands:
-		case ERigVMOpCode::Execute_31_Operands:
-		case ERigVMOpCode::Execute_32_Operands:
-		case ERigVMOpCode::Execute_33_Operands:
-		case ERigVMOpCode::Execute_34_Operands:
-		case ERigVMOpCode::Execute_35_Operands:
-		case ERigVMOpCode::Execute_36_Operands:
-		case ERigVMOpCode::Execute_37_Operands:
-		case ERigVMOpCode::Execute_38_Operands:
-		case ERigVMOpCode::Execute_39_Operands:
-		case ERigVMOpCode::Execute_40_Operands:
-		case ERigVMOpCode::Execute_41_Operands:
-		case ERigVMOpCode::Execute_42_Operands:
-		case ERigVMOpCode::Execute_43_Operands:
-		case ERigVMOpCode::Execute_44_Operands:
-		case ERigVMOpCode::Execute_45_Operands:
-		case ERigVMOpCode::Execute_46_Operands:
-		case ERigVMOpCode::Execute_47_Operands:
-		case ERigVMOpCode::Execute_48_Operands:
-		case ERigVMOpCode::Execute_49_Operands:
-		case ERigVMOpCode::Execute_50_Operands:
-		case ERigVMOpCode::Execute_51_Operands:
-		case ERigVMOpCode::Execute_52_Operands:
-		case ERigVMOpCode::Execute_53_Operands:
-		case ERigVMOpCode::Execute_54_Operands:
-		case ERigVMOpCode::Execute_55_Operands:
-		case ERigVMOpCode::Execute_56_Operands:
-		case ERigVMOpCode::Execute_57_Operands:
-		case ERigVMOpCode::Execute_58_Operands:
-		case ERigVMOpCode::Execute_59_Operands:
-		case ERigVMOpCode::Execute_60_Operands:
-		case ERigVMOpCode::Execute_61_Operands:
-		case ERigVMOpCode::Execute_62_Operands:
-		case ERigVMOpCode::Execute_63_Operands:
-		case ERigVMOpCode::Execute_64_Operands:
+		case ERigVMOpCode::Execute:
 		{
 			return GetOperandsForExecuteOp(InInstruction);
 		}
@@ -1713,7 +1354,7 @@ FRigVMOperandArray FRigVMByteCode::GetOperandsForOp(const FRigVMInstruction& InI
 
 uint64 FRigVMByteCode::GetFirstOperandByteIndex(const FRigVMInstruction& InInstruction) const
 {
-	if (InInstruction.OpCode >= ERigVMOpCode::Execute_0_Operands && InInstruction.OpCode <= ERigVMOpCode::Execute_64_Operands)
+	if (InInstruction.OpCode == ERigVMOpCode::Execute)
 	{
 		const uint64 ByteCodeIndex = InInstruction.ByteCodeIndex;
 		// if the bytecode is not aligned the OperandAlignment needs to be 0
@@ -1824,71 +1465,7 @@ uint64 FRigVMByteCode::GetOpAlignment(ERigVMOpCode InOpCode) const
 {
 	switch (InOpCode)
 	{
-		case ERigVMOpCode::Execute_0_Operands:
-		case ERigVMOpCode::Execute_1_Operands:
-		case ERigVMOpCode::Execute_2_Operands:
-		case ERigVMOpCode::Execute_3_Operands:
-		case ERigVMOpCode::Execute_4_Operands:
-		case ERigVMOpCode::Execute_5_Operands:
-		case ERigVMOpCode::Execute_6_Operands:
-		case ERigVMOpCode::Execute_7_Operands:
-		case ERigVMOpCode::Execute_8_Operands:
-		case ERigVMOpCode::Execute_9_Operands:
-		case ERigVMOpCode::Execute_10_Operands:
-		case ERigVMOpCode::Execute_11_Operands:
-		case ERigVMOpCode::Execute_12_Operands:
-		case ERigVMOpCode::Execute_13_Operands:
-		case ERigVMOpCode::Execute_14_Operands:
-		case ERigVMOpCode::Execute_15_Operands:
-		case ERigVMOpCode::Execute_16_Operands:
-		case ERigVMOpCode::Execute_17_Operands:
-		case ERigVMOpCode::Execute_18_Operands:
-		case ERigVMOpCode::Execute_19_Operands:
-		case ERigVMOpCode::Execute_20_Operands:
-		case ERigVMOpCode::Execute_21_Operands:
-		case ERigVMOpCode::Execute_22_Operands:
-		case ERigVMOpCode::Execute_23_Operands:
-		case ERigVMOpCode::Execute_24_Operands:
-		case ERigVMOpCode::Execute_25_Operands:
-		case ERigVMOpCode::Execute_26_Operands:
-		case ERigVMOpCode::Execute_27_Operands:
-		case ERigVMOpCode::Execute_28_Operands:
-		case ERigVMOpCode::Execute_29_Operands:
-		case ERigVMOpCode::Execute_30_Operands:
-		case ERigVMOpCode::Execute_31_Operands:
-		case ERigVMOpCode::Execute_32_Operands:
-		case ERigVMOpCode::Execute_33_Operands:
-		case ERigVMOpCode::Execute_34_Operands:
-		case ERigVMOpCode::Execute_35_Operands:
-		case ERigVMOpCode::Execute_36_Operands:
-		case ERigVMOpCode::Execute_37_Operands:
-		case ERigVMOpCode::Execute_38_Operands:
-		case ERigVMOpCode::Execute_39_Operands:
-		case ERigVMOpCode::Execute_40_Operands:
-		case ERigVMOpCode::Execute_41_Operands:
-		case ERigVMOpCode::Execute_42_Operands:
-		case ERigVMOpCode::Execute_43_Operands:
-		case ERigVMOpCode::Execute_44_Operands:
-		case ERigVMOpCode::Execute_45_Operands:
-		case ERigVMOpCode::Execute_46_Operands:
-		case ERigVMOpCode::Execute_47_Operands:
-		case ERigVMOpCode::Execute_48_Operands:
-		case ERigVMOpCode::Execute_49_Operands:
-		case ERigVMOpCode::Execute_50_Operands:
-		case ERigVMOpCode::Execute_51_Operands:
-		case ERigVMOpCode::Execute_52_Operands:
-		case ERigVMOpCode::Execute_53_Operands:
-		case ERigVMOpCode::Execute_54_Operands:
-		case ERigVMOpCode::Execute_55_Operands:
-		case ERigVMOpCode::Execute_56_Operands:
-		case ERigVMOpCode::Execute_57_Operands:
-		case ERigVMOpCode::Execute_58_Operands:
-		case ERigVMOpCode::Execute_59_Operands:
-		case ERigVMOpCode::Execute_60_Operands:
-		case ERigVMOpCode::Execute_61_Operands:
-		case ERigVMOpCode::Execute_62_Operands:
-		case ERigVMOpCode::Execute_63_Operands:
-		case ERigVMOpCode::Execute_64_Operands:
+		case ERigVMOpCode::Execute:
 		{
 			static const uint64 Alignment = FRigVMExecuteOp::StaticStruct()->GetCppStructOps()->GetAlignment();
 			return Alignment;
@@ -2020,7 +1597,7 @@ void FRigVMByteCode::AlignByteCode()
 		const FRigVMInstruction& Instruction = Instructions[InstructionIndex];
 		BytesToReserve += GetOpAlignment(Instruction.OpCode);
 
-		if (Instruction.OpCode >= ERigVMOpCode::Execute_0_Operands && Instruction.OpCode < ERigVMOpCode::Execute_64_Operands)
+		if (Instruction.OpCode == ERigVMOpCode::Execute)
 		{
 			BytesToReserve += GetOperandAlignment();
 		}
@@ -2055,7 +1632,7 @@ void FRigVMByteCode::AlignByteCode()
 			AlignedByteCode[AlignedByteCodeIndex + ByteIndex] = ByteCode[OriginalByteCodeIndex + ByteIndex];
 		}
 
-		if (Instruction.OpCode >= ERigVMOpCode::Execute_0_Operands && Instruction.OpCode < ERigVMOpCode::Execute_64_Operands)
+		if (Instruction.OpCode == ERigVMOpCode::Execute)
 		{
 			AlignedByteCodeIndex += NumBytesToCopy;
 
diff --git a/Engine/Plugins/Runtime/RigVM/Source/RigVM/Public/RigVMCore/RigVMByteCode.h b/Engine/Plugins/Runtime/RigVM/Source/RigVM/Public/RigVMCore/RigVMByteCode.h
--- a/Engine/Plugins/Runtime/RigVM/Source/RigVM/Public/RigVMCore/RigVMByteCode.h
+++ b/Engine/Plugins/Runtime/RigVM/Source/RigVM/Public/RigVMCore/RigVMByteCode.h
@@ -99,71 +99,71 @@ struct FRigVMBranchInfoKey
 UENUM()
 enum class ERigVMOpCode : uint8
 {
-	Execute_0_Operands, // execute a rig function with 0 operands
-	Execute_1_Operands, // execute a rig function with 1 operands
-	Execute_2_Operands, // execute a rig function with 2 operands
-	Execute_3_Operands, // execute a rig function with 3 operands
-	Execute_4_Operands, // execute a rig function with 4 operands
-	Execute_5_Operands, // execute a rig function with 5 operands
-	Execute_6_Operands, // execute a rig function with 6 operands
-	Execute_7_Operands, // execute a rig function with 7 operands
-	Execute_8_Operands, // execute a rig function with 8 operands
-	Execute_9_Operands, // execute a rig function with 9 operands
-	Execute_10_Operands, // execute a rig function with 10 operands
-	Execute_11_Operands, // execute a rig function with 11 operands
-	Execute_12_Operands, // execute a rig function with 12 operands
-	Execute_13_Operands, // execute a rig function with 13 operands
-	Execute_14_Operands, // execute a rig function with 14 operands
-	Execute_15_Operands, // execute a rig function with 15 operands
-	Execute_16_Operands, // execute a rig function with 16 operands
-	Execute_17_Operands, // execute a rig function with 17 operands
-	Execute_18_Operands, // execute a rig function with 18 operands
-	Execute_19_Operands, // execute a rig function with 19 operands
-	Execute_20_Operands, // execute a rig function with 20 operands
-	Execute_21_Operands, // execute a rig function with 21 operands
-	Execute_22_Operands, // execute a rig function with 22 operands
-	Execute_23_Operands, // execute a rig function with 23 operands
-	Execute_24_Operands, // execute a rig function with 24 operands
-	Execute_25_Operands, // execute a rig function with 25 operands
-	Execute_26_Operands, // execute a rig function with 26 operands
-	Execute_27_Operands, // execute a rig function with 27 operands
-	Execute_28_Operands, // execute a rig function with 28 operands
-	Execute_29_Operands, // execute a rig function with 29 operands
-	Execute_30_Operands, // execute a rig function with 30 operands
-	Execute_31_Operands, // execute a rig function with 31 operands
-	Execute_32_Operands, // execute a rig function with 32 operands
-	Execute_33_Operands, // execute a rig function with 33 operands
-	Execute_34_Operands, // execute a rig function with 34 operands
-	Execute_35_Operands, // execute a rig function with 35 operands
-	Execute_36_Operands, // execute a rig function with 36 operands
-	Execute_37_Operands, // execute a rig function with 37 operands
-	Execute_38_Operands, // execute a rig function with 38 operands
-	Execute_39_Operands, // execute a rig function with 39 operands
-	Execute_40_Operands, // execute a rig function with 40 operands
-	Execute_41_Operands, // execute a rig function with 41 operands
-	Execute_42_Operands, // execute a rig function with 42 operands
-	Execute_43_Operands, // execute a rig function with 43 operands
-	Execute_44_Operands, // execute a rig function with 44 operands
-	Execute_45_Operands, // execute a rig function with 45 operands
-	Execute_46_Operands, // execute a rig function with 46 operands
-	Execute_47_Operands, // execute a rig function with 47 operands
-	Execute_48_Operands, // execute a rig function with 48 operands
-	Execute_49_Operands, // execute a rig function with 49 operands
-	Execute_50_Operands, // execute a rig function with 50 operands
-	Execute_51_Operands, // execute a rig function with 51 operands
-	Execute_52_Operands, // execute a rig function with 52 operands
-	Execute_53_Operands, // execute a rig function with 53 operands
-	Execute_54_Operands, // execute a rig function with 54 operands
-	Execute_55_Operands, // execute a rig function with 55 operands
-	Execute_56_Operands, // execute a rig function with 56 operands
-	Execute_57_Operands, // execute a rig function with 57 operands
-	Execute_58_Operands, // execute a rig function with 58 operands
-	Execute_59_Operands, // execute a rig function with 59 operands
-	Execute_60_Operands, // execute a rig function with 60 operands
-	Execute_61_Operands, // execute a rig function with 61 operands
-	Execute_62_Operands, // execute a rig function with 62 operands
-	Execute_63_Operands, // execute a rig function with 63 operands
-	Execute_64_Operands, // execute a rig function with 64 operands
+	Execute_0_Operands, // (DEPRECATED) execute a rig function with 0 operands
+	Execute_1_Operands, // (DEPRECATED) execute a rig function with 1 operands
+	Execute_2_Operands, // (DEPRECATED) execute a rig function with 2 operands
+	Execute_3_Operands, // (DEPRECATED) execute a rig function with 3 operands
+	Execute_4_Operands, // (DEPRECATED) execute a rig function with 4 operands
+	Execute_5_Operands, // (DEPRECATED) execute a rig function with 5 operands
+	Execute_6_Operands, // (DEPRECATED) execute a rig function with 6 operands
+	Execute_7_Operands, // (DEPRECATED) execute a rig function with 7 operands
+	Execute_8_Operands, // (DEPRECATED) execute a rig function with 8 operands
+	Execute_9_Operands, // (DEPRECATED) execute a rig function with 9 operands
+	Execute_10_Operands, // (DEPRECATED) execute a rig function with 10 operands
+	Execute_11_Operands, // (DEPRECATED) execute a rig function with 11 operands
+	Execute_12_Operands, // (DEPRECATED) execute a rig function with 12 operands
+	Execute_13_Operands, // (DEPRECATED) execute a rig function with 13 operands
+	Execute_14_Operands, // (DEPRECATED) execute a rig function with 14 operands
+	Execute_15_Operands, // (DEPRECATED) execute a rig function with 15 operands
+	Execute_16_Operands, // (DEPRECATED) execute a rig function with 16 operands
+	Execute_17_Operands, // (DEPRECATED) execute a rig function with 17 operands
+	Execute_18_Operands, // (DEPRECATED) execute a rig function with 18 operands
+	Execute_19_Operands, // (DEPRECATED) execute a rig function with 19 operands
+	Execute_20_Operands, // (DEPRECATED) execute a rig function with 20 operands
+	Execute_21_Operands, // (DEPRECATED) execute a rig function with 21 operands
+	Execute_22_Operands, // (DEPRECATED) execute a rig function with 22 operands
+	Execute_23_Operands, // (DEPRECATED) execute a rig function with 23 operands
+	Execute_24_Operands, // (DEPRECATED) execute a rig function with 24 operands
+	Execute_25_Operands, // (DEPRECATED) execute a rig function with 25 operands
+	Execute_26_Operands, // (DEPRECATED) execute a rig function with 26 operands
+	Execute_27_Operands, // (DEPRECATED) execute a rig function with 27 operands
+	Execute_28_Operands, // (DEPRECATED) execute a rig function with 28 operands
+	Execute_29_Operands, // (DEPRECATED) execute a rig function with 29 operands
+	Execute_30_Operands, // (DEPRECATED) execute a rig function with 30 operands
+	Execute_31_Operands, // (DEPRECATED) execute a rig function with 31 operands
+	Execute_32_Operands, // (DEPRECATED) execute a rig function with 32 operands
+	Execute_33_Operands, // (DEPRECATED) execute a rig function with 33 operands
+	Execute_34_Operands, // (DEPRECATED) execute a rig function with 34 operands
+	Execute_35_Operands, // (DEPRECATED) execute a rig function with 35 operands
+	Execute_36_Operands, // (DEPRECATED) execute a rig function with 36 operands
+	Execute_37_Operands, // (DEPRECATED) execute a rig function with 37 operands
+	Execute_38_Operands, // (DEPRECATED) execute a rig function with 38 operands
+	Execute_39_Operands, // (DEPRECATED) execute a rig function with 39 operands
+	Execute_40_Operands, // (DEPRECATED) execute a rig function with 40 operands
+	Execute_41_Operands, // (DEPRECATED) execute a rig function with 41 operands
+	Execute_42_Operands, // (DEPRECATED) execute a rig function with 42 operands
+	Execute_43_Operands, // (DEPRECATED) execute a rig function with 43 operands
+	Execute_44_Operands, // (DEPRECATED) execute a rig function with 44 operands
+	Execute_45_Operands, // (DEPRECATED) execute a rig function with 45 operands
+	Execute_46_Operands, // (DEPRECATED) execute a rig function with 46 operands
+	Execute_47_Operands, // (DEPRECATED) execute a rig function with 47 operands
+	Execute_48_Operands, // (DEPRECATED) execute a rig function with 48 operands
+	Execute_49_Operands, // (DEPRECATED) execute a rig function with 49 operands
+	Execute_50_Operands, // (DEPRECATED) execute a rig function with 50 operands
+	Execute_51_Operands, // (DEPRECATED) execute a rig function with 51 operands
+	Execute_52_Operands, // (DEPRECATED) execute a rig function with 52 operands
+	Execute_53_Operands, // (DEPRECATED) execute a rig function with 53 operands
+	Execute_54_Operands, // (DEPRECATED) execute a rig function with 54 operands
+	Execute_55_Operands, // (DEPRECATED) execute a rig function with 55 operands
+	Execute_56_Operands, // (DEPRECATED) execute a rig function with 56 operands
+	Execute_57_Operands, // (DEPRECATED) execute a rig function with 57 operands
+	Execute_58_Operands, // (DEPRECATED) execute a rig function with 58 operands
+	Execute_59_Operands, // (DEPRECATED) execute a rig function with 59 operands
+	Execute_60_Operands, // (DEPRECATED) execute a rig function with 60 operands
+	Execute_61_Operands, // (DEPRECATED) execute a rig function with 61 operands
+	Execute_62_Operands, // (DEPRECATED) execute a rig function with 62 operands
+	Execute_63_Operands, // (DEPRECATED) execute a rig function with 63 operands
+	Execute_64_Operands, // (DEPRECATED) execute a rig function with 64 operands
 	Zero, // zero the memory of a given register
 	BoolFalse, // set a given register to false
 	BoolTrue, // set a given register to true
@@ -200,6 +200,7 @@ enum class ERigVMOpCode : uint8
 	ArrayReverse, // (DEPRECATED) returns the reverse of the input array (unary op, in out array)
 	InvokeEntry, // invokes an entry from the entry list
 	JumpToBranch, // jumps to a branch based on a name operand
+	Execute, // single execute op (formerly Execute_0_Operands to Execute_64_Operands)
 	Invalid,
 	FirstArrayOpCode = ArrayReset,
 	LastArrayOpCode = ArrayReverse,
@@ -234,23 +235,32 @@ struct RIGVM_API FRigVMExecuteOp : public FRigVMBaseOp
 	FRigVMExecuteOp()
 	: FRigVMBaseOp()
 	, FunctionIndex(INDEX_NONE)
+	, ArgumentCount(0)
 	{
 	}
 
 	FRigVMExecuteOp(uint16 InFunctionIndex, uint8 InArgumentCount)
-	: FRigVMBaseOp((ERigVMOpCode)(uint8(ERigVMOpCode::Execute_0_Operands) + InArgumentCount))
+	: FRigVMBaseOp(ERigVMOpCode::Execute)
 	, FunctionIndex(InFunctionIndex)
+	, ArgumentCount(InArgumentCount)
 	{
 	}
 
 	uint16 FunctionIndex;
+	uint16 ArgumentCount;
 
 	friend uint32 GetTypeHash(const FRigVMExecuteOp& Op)
 	{
-		return HashCombine(GetTypeHash((const FRigVMBaseOp&)Op), GetTypeHash(Op.FunctionIndex));
+		return HashCombine(
+			GetTypeHash((const FRigVMBaseOp&)Op),
+			HashCombine(
+				GetTypeHash(Op.FunctionIndex),
+				GetTypeHash(Op.ArgumentCount)
+			)
+		);
 	}
 
-	uint8 GetOperandCount() const { return uint8(OpCode) - uint8(ERigVMOpCode::Execute_0_Operands); }
+	uint16 GetOperandCount() const { return ArgumentCount; }
 
 	void Serialize(FArchive& Ar);
 	friend FArchive& operator<<(FArchive& Ar, FRigVMExecuteOp& P)
@@ -1209,6 +1219,12 @@ private:
 	template<class OpType>
 	uint64 AddOp(const OpType& InOp)
 	{
+		check(
+			(InOp.OpCode > ERigVMOpCode::Execute_64_Operands) && 
+			!(InOp.OpCode >= ERigVMOpCode::ArrayReset && InOp.OpCode <= ERigVMOpCode::ArrayReverse) && 
+			(InOp.OpCode < ERigVMOpCode::Invalid) 
+		);
+		
 		uint64 ByteIndex = (uint64)ByteCode.AddZeroed(sizeof(OpType));
 		uint8* Pointer = &ByteCode[ByteIndex];
 		FMemory::Memcpy(Pointer, &InOp, sizeof(OpType));
diff --git a/Engine/Plugins/Runtime/RigVM/Source/RigVMDeveloper/Private/RigVMCompiler/RigVMCodeGenerator.cpp b/Engine/Plugins/Runtime/RigVM/Source/RigVMDeveloper/Private/RigVMCompiler/RigVMCodeGenerator.cpp
--- a/Engine/Plugins/Runtime/RigVM/Source/RigVMDeveloper/Private/RigVMCompiler/RigVMCodeGenerator.cpp
+++ b/Engine/Plugins/Runtime/RigVM/Source/RigVMDeveloper/Private/RigVMCompiler/RigVMCodeGenerator.cpp
@@ -643,71 +643,7 @@ FString FRigVMCodeGenerator::DumpInstructions(const FString& InPrefix, const TAr
 		const FRigVMInstruction& Instruction = Instructions[InstructionIndex];
 		switch(Instruction.OpCode)
 		{
-			case ERigVMOpCode::Execute_0_Operands:
-			case ERigVMOpCode::Execute_1_Operands:
-			case ERigVMOpCode::Execute_2_Operands:
-			case ERigVMOpCode::Execute_3_Operands:
-			case ERigVMOpCode::Execute_4_Operands:
-			case ERigVMOpCode::Execute_5_Operands:
-			case ERigVMOpCode::Execute_6_Operands:
-			case ERigVMOpCode::Execute_7_Operands:
-			case ERigVMOpCode::Execute_8_Operands:
-			case ERigVMOpCode::Execute_9_Operands:
-			case ERigVMOpCode::Execute_10_Operands:
-			case ERigVMOpCode::Execute_11_Operands:
-			case ERigVMOpCode::Execute_12_Operands:
-			case ERigVMOpCode::Execute_13_Operands:
-			case ERigVMOpCode::Execute_14_Operands:
-			case ERigVMOpCode::Execute_15_Operands:
-			case ERigVMOpCode::Execute_16_Operands:
-			case ERigVMOpCode::Execute_17_Operands:
-			case ERigVMOpCode::Execute_18_Operands:
-			case ERigVMOpCode::Execute_19_Operands:
-			case ERigVMOpCode::Execute_20_Operands:
-			case ERigVMOpCode::Execute_21_Operands:
-			case ERigVMOpCode::Execute_22_Operands:
-			case ERigVMOpCode::Execute_23_Operands:
-			case ERigVMOpCode::Execute_24_Operands:
-			case ERigVMOpCode::Execute_25_Operands:
-			case ERigVMOpCode::Execute_26_Operands:
-			case ERigVMOpCode::Execute_27_Operands:
-			case ERigVMOpCode::Execute_28_Operands:
-			case ERigVMOpCode::Execute_29_Operands:
-			case ERigVMOpCode::Execute_30_Operands:
-			case ERigVMOpCode::Execute_31_Operands:
-			case ERigVMOpCode::Execute_32_Operands:
-			case ERigVMOpCode::Execute_33_Operands:
-			case ERigVMOpCode::Execute_34_Operands:
-			case ERigVMOpCode::Execute_35_Operands:
-			case ERigVMOpCode::Execute_36_Operands:
-			case ERigVMOpCode::Execute_37_Operands:
-			case ERigVMOpCode::Execute_38_Operands:
-			case ERigVMOpCode::Execute_39_Operands:
-			case ERigVMOpCode::Execute_40_Operands:
-			case ERigVMOpCode::Execute_41_Operands:
-			case ERigVMOpCode::Execute_42_Operands:
-			case ERigVMOpCode::Execute_43_Operands:
-			case ERigVMOpCode::Execute_44_Operands:
-			case ERigVMOpCode::Execute_45_Operands:
-			case ERigVMOpCode::Execute_46_Operands:
-			case ERigVMOpCode::Execute_47_Operands:
-			case ERigVMOpCode::Execute_48_Operands:
-			case ERigVMOpCode::Execute_49_Operands:
-			case ERigVMOpCode::Execute_50_Operands:
-			case ERigVMOpCode::Execute_51_Operands:
-			case ERigVMOpCode::Execute_52_Operands:
-			case ERigVMOpCode::Execute_53_Operands:
-			case ERigVMOpCode::Execute_54_Operands:
-			case ERigVMOpCode::Execute_55_Operands:
-			case ERigVMOpCode::Execute_56_Operands:
-			case ERigVMOpCode::Execute_57_Operands:
-			case ERigVMOpCode::Execute_58_Operands:
-			case ERigVMOpCode::Execute_59_Operands:
-			case ERigVMOpCode::Execute_60_Operands:
-			case ERigVMOpCode::Execute_61_Operands:
-			case ERigVMOpCode::Execute_62_Operands:
-			case ERigVMOpCode::Execute_63_Operands:
-			case ERigVMOpCode::Execute_64_Operands:
+			case ERigVMOpCode::Execute:
 			{
 				const FRigVMExecuteOp& Op = ByteCode.GetOpAt<FRigVMExecuteOp>(Instruction);
 				FRigVMOperandArray Operands = ByteCode.GetOperandsForExecuteOp(Instruction);
@@ -1256,8 +1192,7 @@ void FRigVMCodeGenerator::ParseRequiredUProperties()
 	{
 		const FRigVMInstruction& Instruction = Instructions[InstructionIndex];
 		
-		if(Instruction.OpCode >= ERigVMOpCode::Execute_0_Operands &&
-			Instruction.OpCode <= ERigVMOpCode::Execute_64_Operands)
+		if(Instruction.OpCode == ERigVMOpCode::Execute)
 		{
 			const FRigVMExecuteOp& Op = ByteCode.GetOpAt<FRigVMExecuteOp>(Instruction);
 			const FRigVMFunction* Function = FRigVMRegistry::Get().FindFunction(*Functions[Op.FunctionIndex].ToString());
@@ -1544,71 +1479,7 @@ void FRigVMCodeGenerator::ParseInstructionGroups()
 			const FRigVMInstruction& Instruction = Instructions[InstructionIndex];
 			switch(Instruction.OpCode)
 			{
-				case ERigVMOpCode::Execute_0_Operands:
-				case ERigVMOpCode::Execute_1_Operands:
-				case ERigVMOpCode::Execute_2_Operands:
-				case ERigVMOpCode::Execute_3_Operands:
-				case ERigVMOpCode::Execute_4_Operands:
-				case ERigVMOpCode::Execute_5_Operands:
-				case ERigVMOpCode::Execute_6_Operands:
-				case ERigVMOpCode::Execute_7_Operands:
-				case ERigVMOpCode::Execute_8_Operands:
-				case ERigVMOpCode::Execute_9_Operands:
-				case ERigVMOpCode::Execute_10_Operands:
-				case ERigVMOpCode::Execute_11_Operands:
-				case ERigVMOpCode::Execute_12_Operands:
-				case ERigVMOpCode::Execute_13_Operands:
-				case ERigVMOpCode::Execute_14_Operands:
-				case ERigVMOpCode::Execute_15_Operands:
-				case ERigVMOpCode::Execute_16_Operands:
-				case ERigVMOpCode::Execute_17_Operands:
-				case ERigVMOpCode::Execute_18_Operands:
-				case ERigVMOpCode::Execute_19_Operands:
-				case ERigVMOpCode::Execute_20_Operands:
-				case ERigVMOpCode::Execute_21_Operands:
-				case ERigVMOpCode::Execute_22_Operands:
-				case ERigVMOpCode::Execute_23_Operands:
-				case ERigVMOpCode::Execute_24_Operands:
-				case ERigVMOpCode::Execute_25_Operands:
-				case ERigVMOpCode::Execute_26_Operands:
-				case ERigVMOpCode::Execute_27_Operands:
-				case ERigVMOpCode::Execute_28_Operands:
-				case ERigVMOpCode::Execute_29_Operands:
-				case ERigVMOpCode::Execute_30_Operands:
-				case ERigVMOpCode::Execute_31_Operands:
-				case ERigVMOpCode::Execute_32_Operands:
-				case ERigVMOpCode::Execute_33_Operands:
-				case ERigVMOpCode::Execute_34_Operands:
-				case ERigVMOpCode::Execute_35_Operands:
-				case ERigVMOpCode::Execute_36_Operands:
-				case ERigVMOpCode::Execute_37_Operands:
-				case ERigVMOpCode::Execute_38_Operands:
-				case ERigVMOpCode::Execute_39_Operands:
-				case ERigVMOpCode::Execute_40_Operands:
-				case ERigVMOpCode::Execute_41_Operands:
-				case ERigVMOpCode::Execute_42_Operands:
-				case ERigVMOpCode::Execute_43_Operands:
-				case ERigVMOpCode::Execute_44_Operands:
-				case ERigVMOpCode::Execute_45_Operands:
-				case ERigVMOpCode::Execute_46_Operands:
-				case ERigVMOpCode::Execute_47_Operands:
-				case ERigVMOpCode::Execute_48_Operands:
-				case ERigVMOpCode::Execute_49_Operands:
-				case ERigVMOpCode::Execute_50_Operands:
-				case ERigVMOpCode::Execute_51_Operands:
-				case ERigVMOpCode::Execute_52_Operands:
-				case ERigVMOpCode::Execute_53_Operands:
-				case ERigVMOpCode::Execute_54_Operands:
-				case ERigVMOpCode::Execute_55_Operands:
-				case ERigVMOpCode::Execute_56_Operands:
-				case ERigVMOpCode::Execute_57_Operands:
-				case ERigVMOpCode::Execute_58_Operands:
-				case ERigVMOpCode::Execute_59_Operands:
-				case ERigVMOpCode::Execute_60_Operands:
-				case ERigVMOpCode::Execute_61_Operands:
-				case ERigVMOpCode::Execute_62_Operands:
-				case ERigVMOpCode::Execute_63_Operands:
-				case ERigVMOpCode::Execute_64_Operands:
+				case ERigVMOpCode::Execute:
 				{
 					const FRigVMExecuteOp& Op = ByteCode.GetOpAt<FRigVMExecuteOp>(Instruction);
 					const FRigVMFunction* Function = FRigVMRegistry::Get().FindFunction(*Functions[Op.FunctionIndex].ToString());
diff --git a/Engine/Plugins/Runtime/RigVM/Source/RigVMDeveloper/Private/RigVMCompiler/RigVMCompiler.cpp b/Engine/Plugins/Runtime/RigVM/Source/RigVMDeveloper/Private/RigVMCompiler/RigVMCompiler.cpp
--- a/Engine/Plugins/Runtime/RigVM/Source/RigVMDeveloper/Private/RigVMCompiler/RigVMCompiler.cpp
+++ b/Engine/Plugins/Runtime/RigVM/Source/RigVMDeveloper/Private/RigVMCompiler/RigVMCompiler.cpp
@@ -2096,7 +2096,7 @@ int32 URigVMCompiler::TraverseInlineFunction(const FRigVMInlineFunctionExprAST*
 				}
 			}
 
-			if (Instruction.OpCode >= ERigVMOpCode::Execute_0_Operands && Instruction.OpCode <= ERigVMOpCode::Execute_64_Operands)
+			if (Instruction.OpCode == ERigVMOpCode::Execute)
 			{
 				FRigVMExecuteOp& Op = ByteCode.GetOpAt<FRigVMExecuteOp>(Instruction);
 				const int32 FunctionIndex = WorkData.VM->AddRigVMFunction(FunctionCompilationData->FunctionNames[Op.FunctionIndex].ToString());
