RigVM fix for CitySample packaging crash

diff --git a/Engine/Plugins/Runtime/RigVM/Source/RigVM/Public/RigVMCore/RigVMByteCode.h b/Engine/Plugins/Runtime/RigVM/Source/RigVM/Public/RigVMCore/RigVMByteCode.h
--- a/Engine/Plugins/Runtime/RigVM/Source/RigVM/Public/RigVMCore/RigVMByteCode.h
+++ b/Engine/Plugins/Runtime/RigVM/Source/RigVM/Public/RigVMCore/RigVMByteCode.h
@@ -1221,7 +1221,7 @@ private:
 	{
 		check(
 			(InOp.OpCode > ERigVMOpCode::Execute_64_Operands) && 
-			!(InOp.OpCode >= ERigVMOpCode::ArrayReset && InOp.OpCode <= ERigVMOpCode::ArrayReverse) && 
+			//!(InOp.OpCode >= ERigVMOpCode::ArrayReset && InOp.OpCode <= ERigVMOpCode::ArrayReverse) && 
 			(InOp.OpCode < ERigVMOpCode::Invalid) 
 		);
 		
