RigVM: Improve performance to find out if a template permutation already exist given a type map

diff --git a/Engine/Plugins/Runtime/RigVM/Source/RigVM/Private/RigVMCore/RigVMRegistry.cpp b/Engine/Plugins/Runtime/RigVM/Source/RigVM/Private/RigVMCore/RigVMRegistry.cpp
--- a/Engine/Plugins/Runtime/RigVM/Source/RigVM/Private/RigVMCore/RigVMRegistry.cpp
+++ b/Engine/Plugins/Runtime/RigVM/Source/RigVM/Private/RigVMCore/RigVMRegistry.cpp
@@ -1561,6 +1561,7 @@ const FRigVMTemplate* FRigVMRegistry::GetOrAddTemplateFromArguments(const FName&
 	{
 		Template.Permutations[Index] = INDEX_NONE;
 	}
+	Template.RecomputeTypesHashToPermutations();
 
 	const int32 Index = Templates.AddElement(Template);
 	Templates[Index].Index = Index;
diff --git a/Engine/Plugins/Runtime/RigVM/Source/RigVM/Private/RigVMCore/RigVMTemplate.cpp b/Engine/Plugins/Runtime/RigVM/Source/RigVM/Private/RigVMCore/RigVMTemplate.cpp
--- a/Engine/Plugins/Runtime/RigVM/Source/RigVM/Private/RigVMCore/RigVMTemplate.cpp
+++ b/Engine/Plugins/Runtime/RigVM/Source/RigVM/Private/RigVMCore/RigVMTemplate.cpp
@@ -479,6 +479,8 @@ FRigVMTemplate::FRigVMTemplate(UScriptStruct* InStruct, const FString& InTemplat
 		const FString NotationStr = FString::Printf(TEXT("%s(%s)"), *InTemplateName, *FString::Join(ArgumentNotations, TEXT(",")));
 		Notation = *NotationStr;
 		Permutations.Add(InFunctionIndex);
+
+		UpdateTypesHashToPermutation(Permutations.Num()-1);
 	}
 }
 
@@ -508,6 +510,8 @@ FRigVMTemplate::FRigVMTemplate(const FName& InTemplateName, const TArray<FRigVMT
 		const FString NotationStr = FString::Printf(TEXT("%s(%s)"), *InTemplateName.ToString(), *FString::Join(ArgumentNotations, TEXT(",")));
 		Notation = *NotationStr;
 		Permutations.Add(InFunctionIndex);
+
+		UpdateTypesHashToPermutation(Permutations.Num()-1);
 	}
 }
 
@@ -1048,6 +1052,9 @@ bool FRigVMTemplate::Merge(const FRigVMTemplate& InOther)
 	Arguments = NewArgs;
 
 	Permutations.Add(InOther.Permutations[0]);
+
+	UpdateTypesHashToPermutation(Permutations.Num()-1);
+	
 	return true;
 }
 
@@ -1363,9 +1370,49 @@ bool FRigVMTemplate::Resolve(FTypeMap& InOutTypes, TArray<int32>& OutPermutation
 	return !OutPermutationIndices.IsEmpty();
 }
 
+uint32 FRigVMTemplate::GetTypesHashFromTypes(const FTypeMap& InTypes) const
+{
+	// It is only a valid type map if it includes all arguments, and non of the types is a wildcard
+	
+	uint32 TypeHash = 0;
+	if (InTypes.Num() != NumArguments())
+	{
+		return TypeHash;
+	}
+
+	const FRigVMRegistry& Registry = FRigVMRegistry::Get();
+	for (const TPair<FName, TRigVMTypeIndex>& Pair : InTypes)
+	{
+		if (Registry.IsWildCardType(Pair.Value))
+		{
+			return TypeHash;
+		}
+	}
+
+	for (const FRigVMTemplateArgument& Argument : Arguments)
+	{
+		const TRigVMTypeIndex* ArgType = InTypes.Find(Argument.Name);
+		if (!ArgType)
+		{
+			return 0;
+		}
+		TypeHash = HashCombine(TypeHash, GetTypeHash(*ArgType));
+	}
+	return TypeHash;
+}
+
 bool FRigVMTemplate::ContainsPermutation(const FTypeMap& InTypes) const
 {
 	const FRigVMRegistry& Registry = FRigVMRegistry::Get();
+
+	// If they type map is valid (full description of arguments), then we can rely on
+	// the TypesHashToPermutation cache. Otherwise, we will have to search for a specific permutation
+	// by filtering types
+	uint32 TypesHash = GetTypesHashFromTypes(InTypes);
+	if (TypesHash != 0)
+	{
+		return TypesHashToPermutation.Contains(TypesHash);
+	}
 	
 	TArray<int32> PossiblePermutations;
 	for (const TPair<FName, TRigVMTypeIndex>& Pair : InTypes)
@@ -1602,13 +1649,6 @@ bool FRigVMTemplate::AddTypeForArgument(const FName& InArgumentName, TRigVMTypeI
 				Registry.ConvertExecuteContextToBaseType(ArgumentAndType.Value);
 			}
 			
-			// make sure this permutation doesn't exist yet
-			FRigVMTemplateTypeMap TempTypes = Types;
-			if(ContainsPermutation(TempTypes))
-			{
-				return false;
-			}
-			
 			for(FRigVMTemplateArgument& Argument : Arguments)
 			{
 				const TRigVMTypeIndex* TypeIndex = Types.Find(Argument.Name);
@@ -1637,6 +1677,9 @@ bool FRigVMTemplate::AddTypeForArgument(const FName& InArgumentName, TRigVMTypeI
 			}
 
 			Permutations.Add(INDEX_NONE);
+
+			UpdateTypesHashToPermutation(Permutations.Num()-1);
+			
 			return true;
 		}
 	}
@@ -1680,6 +1723,34 @@ void FRigVMTemplate::HandleTypeRemoval(TRigVMTypeIndex InTypeIndex)
 	}
 }
 
+void FRigVMTemplate::RecomputeTypesHashToPermutations()
+{
+	TypesHashToPermutation.Reset();
+	for (int32 PermutationIndex=0; PermutationIndex<NumPermutations(); ++PermutationIndex)
+	{
+		uint32 TypesHash=0;
+		for (int32 ArgIndex=0; ArgIndex<NumArguments(); ++ArgIndex)
+		{
+			TypesHash = HashCombine(Hash, GetTypeHash(Arguments[ArgIndex].TypeIndices[PermutationIndex]));
+		}
+		TypesHashToPermutation.Add(TypesHash, PermutationIndex);
+	}
+}
+
+void FRigVMTemplate::UpdateTypesHashToPermutation(const int32& InPermutation)
+{
+	if (InPermutation < 0 || InPermutation >= NumPermutations())
+	{
+		return;
+	}
+	
+	uint32 TypeHash=0;
+	for (const FRigVMTemplateArgument& Argument : Arguments)
+	{
+		TypeHash = HashCombine(TypeHash, GetTypeHash(Argument.TypeIndices[InPermutation]));
+	}
+	TypesHashToPermutation.Add(TypeHash, InPermutation);
+}
 
 uint32 GetTypeHash(const FRigVMTemplateArgument& InArgument)
 {
diff --git a/Engine/Plugins/Runtime/RigVM/Source/RigVM/Public/RigVMCore/RigVMTemplate.h b/Engine/Plugins/Runtime/RigVM/Source/RigVM/Public/RigVMCore/RigVMTemplate.h
--- a/Engine/Plugins/Runtime/RigVM/Source/RigVM/Public/RigVMCore/RigVMTemplate.h
+++ b/Engine/Plugins/Runtime/RigVM/Source/RigVM/Public/RigVMCore/RigVMTemplate.h
@@ -352,6 +352,10 @@ public:
 	// returns true if the template was able to resolve to at least one permutation
 	bool Resolve(FTypeMap& InOutTypes, TArray<int32> & OutPermutationIndices, bool bAllowFloatingPointCasts) const;
 
+	// Will return the hash of the input type map, if it is a valid type map. Otherwise, will return 0.
+	// It is only a valid type map if it includes all arguments, and non of the types is a wildcard.
+	uint32 GetTypesHashFromTypes(const FTypeMap& InTypes) const;
+
 	// returns true if the template was able to resolve to at least one permutation
 	bool ContainsPermutation(const FTypeMap& InTypes) const;
 
@@ -426,7 +430,11 @@ public:
 	{
 		return Delegates.RequestDispatchFunctionDelegate.IsBound() &&
 			Delegates.GetDispatchFactoryDelegate.IsBound();
-	}  
+	}
+
+	void RecomputeTypesHashToPermutations();
+
+	void UpdateTypesHashToPermutation(const int32& InPermutation);
 
 	RIGVM_API friend uint32 GetTypeHash(const FRigVMTemplate& InTemplate);
 
@@ -448,6 +456,7 @@ private:
 	TArray<FRigVMTemplateArgument> Arguments;
 	mutable TArray<FRigVMExecuteArgument> ExecuteArguments;
 	TArray<int32> Permutations;
+	TMap<uint32, int32> TypesHashToPermutation;
 	mutable uint32 Hash;
 
 	FRigVMTemplateDelegates Delegates;
