RigVM: refactoring reroute into simple reroute, make and break nodes

diff --git a/Engine/Plugins/Animation/ControlRig/Source/ControlRigDeveloper/Private/Graph/ControlRigGraph.cpp b/Engine/Plugins/Animation/ControlRig/Source/ControlRigDeveloper/Private/Graph/ControlRigGraph.cpp
--- a/Engine/Plugins/Animation/ControlRig/Source/ControlRigDeveloper/Private/Graph/ControlRigGraph.cpp
+++ b/Engine/Plugins/Animation/ControlRig/Source/ControlRigDeveloper/Private/Graph/ControlRigGraph.cpp
@@ -611,28 +611,6 @@ void UControlRigGraph::HandleModifiedEvent(ERigVMGraphNotifType InNotifType, URi
 			}
 			break;
 		}
-		case ERigVMGraphNotifType::RerouteCompactnessChanged:
-		{
-			if (URigVMRerouteNode* ModelNode = Cast<URigVMRerouteNode>(InSubject))
-			{
-				UEdGraphNode* EdNode = Cast<UEdGraphNode>(FindNodeForModelNodeName(ModelNode->GetFName()));
-				if (EdNode)
-				{
-					if (UControlRigGraphNode* RigNode = Cast<UControlRigGraphNode>(EdNode))
-					{
-						// start at index 2 (the subpins below the top level value pin)
-						// and hide the pins (or show them if they were hidden previously)
-						for (int32 PinIndex = 2; PinIndex < RigNode->Pins.Num(); PinIndex++)
-						{
-							RigNode->Pins[PinIndex]->bHidden = !ModelNode->GetShowsAsFullNode();
-						}
-						NotifyGraphChanged(FEdGraphEditAction(EEdGraphActionType::GRAPHACTION_RemoveNode, EdNode->GetGraph(), EdNode, true));
-						NotifyGraphChanged(FEdGraphEditAction(EEdGraphActionType::GRAPHACTION_AddNode, EdNode->GetGraph(), EdNode, true));
-					}
-				}
-			}
-			break;
-		}
 		case ERigVMGraphNotifType::NodeColorChanged:
 		{
 			if (URigVMNode* ModelNode = Cast<URigVMNode>(InSubject))
diff --git a/Engine/Plugins/Animation/ControlRig/Source/ControlRigDeveloper/Private/Graph/ControlRigGraphNode.cpp b/Engine/Plugins/Animation/ControlRig/Source/ControlRigDeveloper/Private/Graph/ControlRigGraphNode.cpp
--- a/Engine/Plugins/Animation/ControlRig/Source/ControlRigDeveloper/Private/Graph/ControlRigGraphNode.cpp
+++ b/Engine/Plugins/Animation/ControlRig/Source/ControlRigDeveloper/Private/Graph/ControlRigGraphNode.cpp
@@ -1698,29 +1698,23 @@ bool UControlRigGraphNode::IsSelectedInEditor() const
 
 bool UControlRigGraphNode::ShouldDrawNodeAsControlPointOnly(int32& OutInputPinIndex, int32& OutOutputPinIndex) const
 {
-	if (URigVMRerouteNode* Reroute = Cast<URigVMRerouteNode>(GetModelNode()))
+	if (Cast<URigVMRerouteNode>(GetModelNode()))
 	{
-		if (!Reroute->GetShowsAsFullNode())
+		OutInputPinIndex = 0;
+		OutOutputPinIndex = Pins.Num()/2;
+
+		if (Pins[OutOutputPinIndex]->Direction != EGPD_Output || Pins[OutOutputPinIndex]->ParentPin != nullptr)
 		{
-			if (Pins.Num() >= 2)
+			for (int32 i=0; i<Pins.Num(); ++i)
 			{
-				OutInputPinIndex = 0;
-				OutOutputPinIndex = Pins.Num()/2;
-
-				if (Pins[OutOutputPinIndex]->Direction != EGPD_Output || Pins[OutOutputPinIndex]->ParentPin != nullptr)
+				if (Pins[i]->Direction == EGPD_Output && Pins[i]->ParentPin == nullptr)
 				{
-					for (int32 i=0; i<Pins.Num(); ++i)
-					{
-						if (Pins[i]->Direction == EGPD_Output && Pins[i]->ParentPin == nullptr)
-						{
-							OutOutputPinIndex = i;
-							break;
-						}
-					}
+					OutOutputPinIndex = i;
+					break;
 				}
-				return true;
 			}
 		}
+		return true;
 	}
 	return false;
 }
@@ -1752,9 +1746,9 @@ void UControlRigGraphNode::ConfigurePin(UEdGraphPin* EdGraphPin, const URigVMPin
 	bool bHidden = ModelPin->GetDirection() == ERigVMPinDirection::Hidden;
 	if(!bHidden && !ModelPin->IsRootPin())
 	{
-		if(const URigVMRerouteNode* RerouteNode = Cast<URigVMRerouteNode>(ModelPin->GetNode()))
+		if(Cast<URigVMRerouteNode>(ModelPin->GetNode()))
 		{
-			bHidden = !RerouteNode->GetShowsAsFullNode();
+			bHidden = true;
 		}
 	}
 
diff --git a/Engine/Plugins/Animation/ControlRig/Source/ControlRigDeveloper/Private/Graph/ControlRigGraphSchema.cpp b/Engine/Plugins/Animation/ControlRig/Source/ControlRigDeveloper/Private/Graph/ControlRigGraphSchema.cpp
--- a/Engine/Plugins/Animation/ControlRig/Source/ControlRigDeveloper/Private/Graph/ControlRigGraphSchema.cpp
+++ b/Engine/Plugins/Animation/ControlRig/Source/ControlRigDeveloper/Private/Graph/ControlRigGraphSchema.cpp
@@ -20,6 +20,7 @@
 #include "Framework/Application/SlateApplication.h"
 #include "Framework/Notifications/NotificationManager.h"
 #include "EulerTransform.h"
+#include "RigVMStringUtils.h"
 #include "Curves/CurveFloat.h"
 #include "RigVMModel/Nodes/RigVMLibraryNode.h"
 #include "RigVMModel/Nodes/RigVMCollapseNode.h"
@@ -29,6 +30,10 @@
 #include "RigVMCore/RigVMUnknownType.h"
 #include "Kismet2/Kismet2NameValidators.h"
 #include "Algo/Count.h"
+#include "RigVMFunctions/RigVMDispatch_Array.h"
+#include "RigVMFunctions/RigVMDispatch_Constant.h"
+#include "RigVMFunctions/RigVMDispatch_MakeStruct.h"
+#include "RigVMModel/Nodes/RigVMDispatchNode.h"
 #include "Units/Execution/RigUnit_BeginExecution.h"
 #include "Units/Execution/RigUnit_PrepareForExecution.h"
 #include "Units/Execution/RigUnit_InverseExecution.h"
@@ -1316,7 +1321,7 @@ void UControlRigGraphSchema::OnPinConnectionDoubleCicked(UEdGraphPin* PinA, UEdG
 	{
 		if (URigVMLink* Link = Node->GetModel()->FindLink(FString::Printf(TEXT("%s -> %s"), *PinA->GetName(), *PinB->GetName())))
 		{
-			Node->GetController()->AddRerouteNodeOnLink(Link, false, GraphPosition, FString(), true, true);
+			Node->GetController()->AddRerouteNodeOnLink(Link, GraphPosition, FString(), true, true);
 		}
 	}
 }
@@ -2342,6 +2347,47 @@ bool UControlRigGraphSchema::AutowireNewNode(UControlRigGraphNode* NewNode, UEdG
 							}
 						}
 					}
+
+					// copy the default value over if the node is a reroute,
+					// a make array, make struct or make constant
+					if(!ModelPin->IsExecuteContext() && FromPin->Direction == EGPD_Input)
+					{
+						if(const URigVMPin* OtherModelPin = ModelPin->GetGraph()->FindPin(FromPin->GetName()))
+						{
+							FString PinNameToSet;
+							if(ModelPin->GetNode()->IsA<URigVMRerouteNode>())
+							{
+								PinNameToSet = TEXT("Value");
+							}
+							else if(const URigVMDispatchNode* DispatchNode = Cast<URigVMDispatchNode>(ModelPin->GetNode()))
+							{
+								if(const FRigVMDispatchFactory* Factory = DispatchNode->GetFactory())
+								{
+									if(Factory->GetFactoryName() == FRigVMDispatch_Constant().GetFactoryName())
+									{
+										PinNameToSet = TEXT("Value");
+									}
+									else if(Factory->GetFactoryName() == FRigVMDispatch_MakeStruct().GetFactoryName())
+									{
+										PinNameToSet = TEXT("Elements");
+									}
+									else if(Factory->GetFactoryName() == FRigVMDispatch_ArrayMake().GetFactoryName())
+									{
+										PinNameToSet = TEXT("Values");
+									}
+								}
+							}
+
+							if(!PinNameToSet.IsEmpty())
+							{
+								const FString DefaultValue = OtherModelPin->GetDefaultValue();
+								if(!DefaultValue.IsEmpty())
+								{
+									NewNode->GetController()->SetPinDefaultValue(RigVMStringUtils::JoinPinPath(NewNode->GetName(), PinNameToSet), DefaultValue, true, true, false, true);
+								}
+							}
+						}
+					}
 				}
 				return true;
 			}
diff --git a/Engine/Plugins/Animation/ControlRig/Source/ControlRigDeveloper/Public/Graph/ControlRigGraph.h b/Engine/Plugins/Animation/ControlRig/Source/ControlRigDeveloper/Public/Graph/ControlRigGraph.h
--- a/Engine/Plugins/Animation/ControlRig/Source/ControlRigDeveloper/Public/Graph/ControlRigGraph.h
+++ b/Engine/Plugins/Animation/ControlRig/Source/ControlRigDeveloper/Public/Graph/ControlRigGraph.h
@@ -182,7 +182,6 @@ private:
 	friend class UControlRigUnitNodeSpawner;
 	friend class UControlRigVariableNodeSpawner;
 	friend class UControlRigParameterNodeSpawner;
-	friend class UControlRigRerouteNodeSpawner;
 	friend class UControlRigBranchNodeSpawner;
 	friend class UControlRigIfNodeSpawner;
 	friend class UControlRigSelectNodeSpawner;
diff --git a/Engine/Plugins/Animation/ControlRig/Source/ControlRigDeveloper/Public/Graph/ControlRigGraphNode.h b/Engine/Plugins/Animation/ControlRig/Source/ControlRigDeveloper/Public/Graph/ControlRigGraphNode.h
--- a/Engine/Plugins/Animation/ControlRig/Source/ControlRigDeveloper/Public/Graph/ControlRigGraphNode.h
+++ b/Engine/Plugins/Animation/ControlRig/Source/ControlRigDeveloper/Public/Graph/ControlRigGraphNode.h
@@ -279,7 +279,6 @@ private:
 	friend class FControlRigArgumentLayout;
 	friend class FControlRigGraphDetails;
 	friend class UControlRigTemplateNodeSpawner;
-	friend class UControlRigRerouteNodeSpawner;
 	friend class UControlRigArrayNodeSpawner;
 	friend class UControlRigIfNodeSpawner;
 	friend class UControlRigSelectNodeSpawner;
diff --git a/Engine/Plugins/Animation/ControlRig/Source/ControlRigDeveloper/Public/Graph/ControlRigGraphSchema.h b/Engine/Plugins/Animation/ControlRig/Source/ControlRigDeveloper/Public/Graph/ControlRigGraphSchema.h
--- a/Engine/Plugins/Animation/ControlRig/Source/ControlRigDeveloper/Public/Graph/ControlRigGraphSchema.h
+++ b/Engine/Plugins/Animation/ControlRig/Source/ControlRigDeveloper/Public/Graph/ControlRigGraphSchema.h
@@ -342,7 +342,6 @@ private:
 	mutable TMap<FName, FVector2D> NodePositionsDuringStart;
 #endif
 
-	friend class UControlRigRerouteNodeSpawner;
 	friend class UControlRigIfNodeSpawner;
 	friend class UControlRigSelectNodeSpawner;
 	friend class UControlRigUnitNodeSpawner;
diff --git a/Engine/Plugins/Animation/ControlRig/Source/ControlRigEditor/Private/ControlRigEditorModule.cpp b/Engine/Plugins/Animation/ControlRig/Source/ControlRigEditor/Private/ControlRigEditorModule.cpp
--- a/Engine/Plugins/Animation/ControlRig/Source/ControlRigEditor/Private/ControlRigEditorModule.cpp
+++ b/Engine/Plugins/Animation/ControlRig/Source/ControlRigEditor/Private/ControlRigEditorModule.cpp
@@ -48,7 +48,6 @@
 #include "EdGraph/EdGraph.h"
 #include "Graph/NodeSpawners/ControlRigUnitNodeSpawner.h"
 #include "Graph/NodeSpawners/ControlRigVariableNodeSpawner.h"
-#include "Graph/NodeSpawners/ControlRigRerouteNodeSpawner.h"
 #include "Graph/NodeSpawners/ControlRigTemplateNodeSpawner.h"
 #include "Graph/NodeSpawners/ControlRigEnumNodeSpawner.h"
 #include "Graph/NodeSpawners/ControlRigFunctionRefNodeSpawner.h"
@@ -133,6 +132,9 @@
 #include "RigVMFunctions/Math/RigVMMathLibrary.h"
 #include "Constraints/ControlRigTransformableHandle.h"
 #include "Constraints/TransformConstraintChannelInterface.h"
+#include "RigVMFunctions/RigVMDispatch_Array.h"
+#include "RigVMFunctions/RigVMDispatch_MakeStruct.h"
+#include "RigVMFunctions/RigVMDispatch_Constant.h"
 
 #define LOCTEXT_NAMESPACE "ControlRigEditorModule"
 
@@ -1109,12 +1111,6 @@ void FControlRigEditorModule::GetTypeActions(UControlRigBlueprint* CRB, FBluepri
 		ActionRegistrar.AddBlueprintAction(ActionKey, NodeSpawner);
 	};
 
-	UBlueprintNodeSpawner* RerouteNodeSpawner = UControlRigRerouteNodeSpawner::CreateGeneric(
-		LOCTEXT("RerouteSpawnerDesc", "Reroute"),
-		LOCTEXT("RerouteSpawnerCategory", "Organization"), 
-		LOCTEXT("RerouteSpawnerTooltip", "Adds a new reroute node to the graph"));
-	ActionRegistrar.AddBlueprintAction(ActionKey, RerouteNodeSpawner);
-
 	for (TObjectIterator<UEnum> EnumIt; EnumIt; ++EnumIt)
 	{
 		UEnum* EnumToConsider = (*EnumIt);
@@ -1497,6 +1493,60 @@ void FControlRigEditorModule::GetContextMenuActions(const UControlRigGraphSchema
 						}
 					}
 
+					// menu entries to convert the node
+					if (URigVMRerouteNode* RerouteNode = Cast<URigVMRerouteNode>(ModelPin->GetNode()))
+					{
+						if(RerouteNode->GetLinkedSourceNodes().Num() == 0)
+						{
+							if(const URigVMPin* ValuePin = RerouteNode->FindPin(TEXT("Value")))
+							{
+								if(!ValuePin->IsExecuteContext())
+								{
+									FToolMenuSection& ConversionSection = Menu->AddSection("EdGraphSchemaConversion", LOCTEXT("ConversionHeader", "Conversion"));
+
+									if(ValuePin->IsArray())
+									{
+										ConversionSection.AddMenuEntry(
+											"Convert Reroute to Make Array",
+											LOCTEXT("ConvertReroutetoMakeArray", "Convert Reroute to Make Array"),
+											LOCTEXT("ConvertReroutetoMakeArray_Tooltip", "Converts the Reroute node to a to Make Array node and sustains the value"),
+											FSlateIcon(),
+											FUIAction(FExecuteAction::CreateLambda([Controller, RerouteNode]()
+											{
+												Controller->ConvertRerouteNodeToDispatch(RerouteNode, FRigVMDispatch_ArrayMake().GetTemplateNotation(), true, true);
+											}
+										)));
+									}
+									else if(ValuePin->IsStruct())
+									{
+										ConversionSection.AddMenuEntry(
+											"Convert Reroute to Make Struct",
+											LOCTEXT("ConvertReroutetoMakeStruct", "Convert Reroute to Make Struct"),
+											LOCTEXT("ConvertReroutetoMakeStruct_Tooltip", "Converts the Reroute node to a to Make Struct node and sustains the value"),
+											FSlateIcon(),
+											FUIAction(FExecuteAction::CreateLambda([Controller, RerouteNode]()
+											{
+												Controller->ConvertRerouteNodeToDispatch(RerouteNode, FRigVMDispatch_MakeStruct().GetTemplateNotation(), true, true);
+											}
+										)));
+									}
+									else
+									{
+										ConversionSection.AddMenuEntry(
+											"Convert Reroute to Constant",
+											LOCTEXT("ConvertReroutetoConstant", "Convert Reroute to Constant"),
+											LOCTEXT("ConvertReroutetoConstant_Tooltip", "Converts the Reroute node to a to Constant node and sustains the value"),
+											FSlateIcon(),
+											FUIAction(FExecuteAction::CreateLambda([Controller, RerouteNode]()
+											{
+												Controller->ConvertRerouteNodeToDispatch(RerouteNode, FRigVMDispatch_Constant().GetTemplateNotation(), true, true);
+											}
+										)));
+									}
+								}
+							}
+						}
+					}
 
 					if (ModelPin->GetDirection() == ERigVMPinDirection::Input &&
 							bIsEditablePin)
@@ -2201,6 +2251,46 @@ void FControlRigEditorModule::GetContextMenuActions(const UControlRigGraphSchema
 								FCanExecuteAction::CreateLambda([bCanRunOnce](){ return bCanRunOnce; }))
 							);
 						}
+
+						if(URigVMDispatchNode* DispatchNode = Cast<URigVMDispatchNode>(ModelNode))
+						{
+							if(DispatchNode->GetFactory()->GetFactoryName() == FRigVMDispatch_Constant().GetFactoryName())
+							{
+								if(const URigVMPin* ValuePin = DispatchNode->FindPin(TEXT("Value")))
+								{
+									// if the value pin has only links on the root pin
+									if((ValuePin->GetSourceLinks(false).Num() > 0) && (ValuePin->GetTargetLinks(false).Num() > 0))
+									{
+										FToolMenuSection& ConversionSection = Menu->AddSection("EdGraphSchemaConversion", LOCTEXT("ConversionHeader", "Conversion"));
+										ConversionSection.AddMenuEntry(
+											"Convert Constant to Reroute",
+											LOCTEXT("ConvertConstantToReroute", "Convert Constant to Reroute"),
+											LOCTEXT("ConvertConstantToReroute_Tooltip", "Converts the Constant node to a to Reroute node and sustains the value"),
+											FSlateIcon(),
+											FUIAction(FExecuteAction::CreateLambda([Controller, ValuePin, DispatchNode]()
+											{
+												Controller->OpenUndoBracket(TEXT("Replace Constant with Reroute"));
+												TArray<URigVMController::FLinkedPath> LinkedPaths = Controller->GetLinkedPaths(DispatchNode);
+												Controller->FastBreakLinkedPaths(LinkedPaths, true);
+												const FVector2D Position = DispatchNode->GetPosition();
+												const FString NodeName = DispatchNode->GetName();
+												const FString CPPType = ValuePin->GetCPPType();
+												FName CPPTypeObjectPath = NAME_None;
+												if(const UObject* CPPTypeObject = ValuePin->GetCPPTypeObject())
+												{
+													CPPTypeObjectPath = *CPPTypeObject->GetPathName();
+												}
+												
+												Controller->RemoveNode(DispatchNode, true, true);
+												Controller->AddFreeRerouteNode(CPPType, CPPTypeObjectPath, false, NAME_None, FString(), Position, NodeName, true);
+												Controller->RestoreLinkedPaths(LinkedPaths, URigVMController::FRestoreLinkedPathSettings(), true);
+												Controller->CloseUndoBracket();
+											}
+										)));
+									}
+								}
+							}
+						}
 						
 						const TArray<FRigVMUserWorkflow> Workflows = ModelNode->GetSupportedWorkflows(ERigVMUserWorkflowType::NodeContext, ModelNode);
 						if(!Workflows.IsEmpty())
diff --git a/Engine/Plugins/Animation/ControlRig/Source/ControlRigEditor/Private/Graph/NodeSpawners/ControlRigRerouteNodeSpawner.cpp b/Engine/Plugins/Animation/ControlRig/Source/ControlRigEditor/Private/Graph/NodeSpawners/ControlRigRerouteNodeSpawner.cpp
deleted file mode 100644
diff --git a/Engine/Plugins/Animation/ControlRig/Source/ControlRigEditor/Private/Graph/NodeSpawners/ControlRigRerouteNodeSpawner.h b/Engine/Plugins/Animation/ControlRig/Source/ControlRigEditor/Private/Graph/NodeSpawners/ControlRigRerouteNodeSpawner.h
deleted file mode 100644
diff --git a/Engine/Plugins/Animation/ControlRig/Source/ControlRigEditor/Private/Graph/SControlRigGraphNode.cpp b/Engine/Plugins/Animation/ControlRig/Source/ControlRigEditor/Private/Graph/SControlRigGraphNode.cpp
--- a/Engine/Plugins/Animation/ControlRig/Source/ControlRigEditor/Private/Graph/SControlRigGraphNode.cpp
+++ b/Engine/Plugins/Animation/ControlRig/Source/ControlRigEditor/Private/Graph/SControlRigGraphNode.cpp
@@ -482,6 +482,7 @@ void SControlRigGraphNode::AddPin(const TSharedRef<SGraphPin>& PinToAdd)
 void SControlRigGraphNode::CreateStandardPinWidget(UEdGraphPin* CurPin)
 {
 	bool bShowPin = true;
+	bool bIsFixedArray = false;
 	FString CPPType;
 	FString BoundVariableName;
 	if(const UControlRigGraphNode* RigGraphNode = Cast<UControlRigGraphNode>(GraphNode))
@@ -496,6 +497,7 @@ void SControlRigGraphNode::CreateStandardPinWidget(UEdGraphPin* CurPin)
 			
 			CPPType = ModelPin->GetCPPType();
 			BoundVariableName = ModelPin->GetBoundVariableName();
+			bIsFixedArray = ModelPin->IsFixedSizeArray();
 		}
 	}
 	
@@ -528,6 +530,14 @@ void SControlRigGraphNode::CreateStandardPinWidget(UEdGraphPin* CurPin)
 			NewPin->AddMetadata(MakeShared<FPinInfoMetaData>(CPPType, BoundVariableName));
 			check(NewPin.IsValid());
 			PinsToKeep.Remove(CurPin);
+
+			if(bIsFixedArray)
+			{
+				if(TSharedPtr<SWidget> PinImage = NewPin->GetPinImageWidget())
+				{
+					PinImage->SetVisibility(EVisibility::Collapsed);
+				}
+			}
 		}
 		
 		if(const UControlRigGraphNode* RigGraphNode = Cast<UControlRigGraphNode>(GraphNode))
@@ -1315,7 +1325,7 @@ EVisibility SControlRigGraphNode::GetPinVisibility(int32 InPinInfoIndex, bool bA
 {
 	if(PinInfos.IsValidIndex(InPinInfoIndex))
 	{
-		if(PinInfos[InPinInfoIndex].bFixedArray)
+		if(PinInfos[InPinInfoIndex].bShowOnlySubPins)
 		{
 			return bAskingForSubPin ? EVisibility::Visible : EVisibility::Collapsed;
 		}
@@ -1329,7 +1339,7 @@ EVisibility SControlRigGraphNode::GetPinVisibility(int32 InPinInfoIndex, bool bA
 				return ParentPinVisibility;
 			}
 
-			if(!PinInfos[InPinInfoIndex].bFixedArray)
+			if(!PinInfos[InPinInfoIndex].bShowOnlySubPins)
 			{
 				if(!PinInfos[ParentPinIndex].bExpanded)
 				{
@@ -1545,6 +1555,11 @@ void SControlRigGraphNode::UpdatePinTreeView()
 				return;
 			}
 
+			if(InPin->ShouldHideSubPins())
+			{
+				return;
+			}
+
 			if (InPin->GetCPPType() == TEXT("FRotator"))
 			{
 				const TArray<URigVMPin*>& SubPins = InPin->GetSubPins();
@@ -1628,7 +1643,7 @@ void SControlRigGraphNode::UpdatePinTreeView()
 		PinInfo.bExpanded = ModelPin->IsExpanded();
 		PinInfo.ModelPinPath = ModelPin->GetPinPath();
 		PinInfo.bAutoHeight = false;
-		PinInfo.bFixedArray = ModelPin->IsFixedSizeArray();
+		PinInfo.bShowOnlySubPins = ModelPin->ShouldOnlyShowSubPins();
 
 		if(!bSupportSubPins)
 		{
@@ -1661,7 +1676,7 @@ void SControlRigGraphNode::UpdatePinTreeView()
 			}
 			PinInfo.ParentIndex = *ParentIndexPtr;
 			PinInfo.Depth = PinInfos[PinInfo.ParentIndex].Depth + 1;
-			if(PinInfos[PinInfo.ParentIndex].bFixedArray)
+			if(PinInfos[PinInfo.ParentIndex].bShowOnlySubPins)
 			{
 				PinInfo.Depth--;
 			}
diff --git a/Engine/Plugins/Animation/ControlRig/Source/ControlRigEditor/Private/Graph/SControlRigGraphNode.h b/Engine/Plugins/Animation/ControlRig/Source/ControlRigEditor/Private/Graph/SControlRigGraphNode.h
--- a/Engine/Plugins/Animation/ControlRig/Source/ControlRigEditor/Private/Graph/SControlRigGraphNode.h
+++ b/Engine/Plugins/Animation/ControlRig/Source/ControlRigEditor/Private/Graph/SControlRigGraphNode.h
@@ -127,7 +127,7 @@ private:
 		TSharedPtr<SGraphPin> OutputPinWidget;
 		bool bExpanded;
 		bool bAutoHeight;
-		bool bFixedArray;
+		bool bShowOnlySubPins;
 	};
 
 	
diff --git a/Engine/Plugins/Runtime/RigVM/Source/RigVM/Private/RigVMFunctions/RigVMDispatch_Constant.cpp b/Engine/Plugins/Runtime/RigVM/Source/RigVM/Private/RigVMFunctions/RigVMDispatch_Constant.cpp
--- a/Engine/Plugins/Runtime/RigVM/Source/RigVM/Private/RigVMFunctions/RigVMDispatch_Constant.cpp
+++ b/Engine/Plugins/Runtime/RigVM/Source/RigVM/Private/RigVMFunctions/RigVMDispatch_Constant.cpp
@@ -58,15 +58,6 @@ FText FRigVMDispatch_Constant::GetNodeTooltip(const FRigVMTemplateTypeMap& InTyp
 	return LOCTEXT("ConstantToolTip", "Provides a constant value to the graph");
 }
 
-FString FRigVMDispatch_Constant::GetArgumentMetaData(const FName& InArgumentName, const FName& InMetaDataKey) const
-{
-	if(InArgumentName == ValueName && InMetaDataKey == FRigVMStruct::ConstantMetaName)
-	{
-		return TrueString;
-	}
-	return FRigVMDispatch_CoreBase::GetArgumentMetaData(InArgumentName, InMetaDataKey);
-}
-
 #endif
 
 void FRigVMDispatch_Constant::Execute(FRigVMExtendedExecuteContext& InContext, FRigVMMemoryHandleArray Handles)
diff --git a/Engine/Plugins/Runtime/RigVM/Source/RigVM/Public/RigVMFunctions/RigVMDispatch_Array.h b/Engine/Plugins/Runtime/RigVM/Source/RigVM/Public/RigVMFunctions/RigVMDispatch_Array.h
--- a/Engine/Plugins/Runtime/RigVM/Source/RigVM/Public/RigVMFunctions/RigVMDispatch_Array.h
+++ b/Engine/Plugins/Runtime/RigVM/Source/RigVM/Public/RigVMFunctions/RigVMDispatch_Array.h
@@ -51,7 +51,7 @@ public:
 	virtual TArray<FRigVMExecuteArgument> GetExecuteArguments_Impl(const FRigVMDispatchContext& InContext) const override;
 };
 
-USTRUCT(meta=(DisplayName = "Array", Keywords = "Make,MakeArray"))
+USTRUCT(meta=(DisplayName = "Make Array", Keywords = "Make,MakeArray,Constant,Reroute"))
 struct RIGVM_API FRigVMDispatch_ArrayMake : public FRigVMDispatch_ArrayBase
 {
 	GENERATED_BODY()
diff --git a/Engine/Plugins/Runtime/RigVM/Source/RigVM/Public/RigVMFunctions/RigVMDispatch_Constant.h b/Engine/Plugins/Runtime/RigVM/Source/RigVM/Public/RigVMFunctions/RigVMDispatch_Constant.h
--- a/Engine/Plugins/Runtime/RigVM/Source/RigVM/Public/RigVMFunctions/RigVMDispatch_Constant.h
+++ b/Engine/Plugins/Runtime/RigVM/Source/RigVM/Public/RigVMFunctions/RigVMDispatch_Constant.h
@@ -7,7 +7,7 @@
 #include "RigVMCore/RigVMByteCode.h"
 #include "RigVMDispatch_Constant.generated.h"
 
-USTRUCT(meta=(DisplayName = "Constant", Category = "Core", Keywords = "Value", NodeColor = "1,1,1,1"))
+USTRUCT(meta=(DisplayName = "Constant", Category = "Core", Keywords = "Value,Reroute", NodeColor = "1,1,1,1"))
 struct RIGVM_API FRigVMDispatch_Constant : public FRigVMDispatch_CoreBase
 {
 	GENERATED_BODY()
@@ -23,7 +23,6 @@ public:
 #if WITH_EDITOR
 	virtual FString GetNodeTitle(const FRigVMTemplateTypeMap& InTypes) const override;
 	virtual FText GetNodeTooltip(const FRigVMTemplateTypeMap& InTypes) const override;
-	virtual FString GetArgumentMetaData(const FName& InArgumentName, const FName& InMetaDataKey) const override;
 #endif
 	virtual bool IsSingleton() const override { return true; }
 	
diff --git a/Engine/Plugins/Runtime/RigVM/Source/RigVM/Public/RigVMFunctions/RigVMDispatch_MakeStruct.h b/Engine/Plugins/Runtime/RigVM/Source/RigVM/Public/RigVMFunctions/RigVMDispatch_MakeStruct.h
--- a/Engine/Plugins/Runtime/RigVM/Source/RigVM/Public/RigVMFunctions/RigVMDispatch_MakeStruct.h
+++ b/Engine/Plugins/Runtime/RigVM/Source/RigVM/Public/RigVMFunctions/RigVMDispatch_MakeStruct.h
@@ -7,12 +7,17 @@
 #include "RigVMCore/RigVMByteCode.h"
 #include "RigVMDispatch_MakeStruct.generated.h"
 
-USTRUCT(meta=(DisplayName = "Make", Category = "Core", Keywords = "Compose,Composition,Create", NodeColor = "1,1,1,1"))
+USTRUCT(meta=(DisplayName = "Make", Category = "Core", Keywords = "Compose,Composition,Create,Constant", NodeColor = "1,1,1,1"))
 struct RIGVM_API FRigVMDispatch_MakeStruct : public FRigVMDispatch_CoreBase
 {
 	GENERATED_BODY()
 
 public:
+	FRigVMDispatch_MakeStruct()
+	{
+		FactoryScriptStruct = StaticStruct();
+	}
+	
 	virtual TArray<FRigVMTemplateArgument> GetArguments() const override;
 	virtual FRigVMTemplateTypeMap OnNewArgumentType(const FName& InArgumentName, TRigVMTypeIndex InTypeIndex) const override;
 #if WITH_EDITOR
@@ -38,6 +43,11 @@ struct RIGVM_API FRigVMDispatch_BreakStruct : public FRigVMDispatch_MakeStruct
 	GENERATED_BODY()
 
 public:
+	FRigVMDispatch_BreakStruct()
+	{
+		FactoryScriptStruct = StaticStruct();
+	}
+
 	virtual TArray<FRigVMTemplateArgument> GetArguments() const override;
 #if WITH_EDITOR
 	virtual FText GetNodeTooltip(const FRigVMTemplateTypeMap& InTypes) const override;
diff --git a/Engine/Plugins/Runtime/RigVM/Source/RigVMDeveloper/Private/RigVMCompiler/RigVMAST.cpp b/Engine/Plugins/Runtime/RigVM/Source/RigVMDeveloper/Private/RigVMCompiler/RigVMAST.cpp
--- a/Engine/Plugins/Runtime/RigVM/Source/RigVMDeveloper/Private/RigVMCompiler/RigVMAST.cpp
+++ b/Engine/Plugins/Runtime/RigVM/Source/RigVMDeveloper/Private/RigVMCompiler/RigVMAST.cpp
@@ -1234,18 +1234,6 @@ FRigVMExprAST* FRigVMParserAST::TraverseLink(int32 InLinkIndex, FRigVMExprAST* I
 		}
 	}
 
-	if (!bRequiresCopy)
-	{
-		// Connections from a composition reroute always require a copy
-		if (URigVMRerouteNode* RerouteNode = Cast<URigVMRerouteNode>(SourceRootPin->GetNode()))
-		{
-			if (RerouteNode->IsCompositionNode())
-			{
-				bRequiresCopy = true;
-			}
-		}
-	}
-
 	if (!bRequiresCopy)
 	{
 		if (SourcePin->GetTypeIndex() != TargetPin->GetTypeIndex())
@@ -1681,11 +1669,7 @@ void FRigVMParserAST::FoldAssignments()
 		}
 		
 		// non-input pins on anything but a reroute (passthrough or literal) or array node should be skipped
-		bool bIsReroute = false;
-		if (URigVMRerouteNode* RerouteNode = Cast<URigVMRerouteNode>(TargetPin->GetNode()))
-		{
-			bIsReroute = !RerouteNode->IsCompositionNode();
-		}
+		const bool bIsReroute = TargetPin->GetNode()->IsA<URigVMRerouteNode>();
 		if (TargetPin->GetDirection() != ERigVMPinDirection::Input && !bIsReroute)
 		{
 			continue;
@@ -2583,10 +2567,7 @@ void FRigVMParserAST::Inline(TArray<URigVMGraph*> InGraphs, const TArray<FRigVMA
 			
 			if (URigVMRerouteNode* RerouteNode = Cast<URigVMRerouteNode>(Node))
 			{
-				if (!RerouteNode->IsCompositionNode())
-				{
-					return true;
-				}
+				return true;
 			}
 			
 			return Node->IsA<URigVMLibraryNode>() ||
@@ -2907,13 +2888,7 @@ void FRigVMParserAST::Inline(TArray<URigVMGraph*> InGraphs, const TArray<FRigVMA
 				URigVMPin* SourcePin = SourcePinProxy.GetSubjectChecked<URigVMPin>();
 				URigVMNode* SourceNode = SourcePin->GetNode();
 
-				bool bIsPassThroughOrLiteralReroute = false;
-				if (URigVMRerouteNode* RerouteNode = Cast<URigVMRerouteNode>(SourceNode))
-				{
-					bIsPassThroughOrLiteralReroute = RerouteNode->IsPassThrough() || RerouteNode->IsLiteral();
-				}
-				
-				if (bIsPassThroughOrLiteralReroute ||
+				if (SourceNode->IsA<URigVMRerouteNode>() ||
 					SourceNode->IsA<URigVMCollapseNode>() ||
 					SourceNode->IsA<URigVMFunctionReturnNode>())
 				{
@@ -2987,12 +2962,9 @@ void FRigVMParserAST::Inline(TArray<URigVMGraph*> InGraphs, const TArray<FRigVMA
 
 		static void VisitNode(const FRigVMASTProxy& InNodeProxy, LocalPinTraversalInfo& OutTraversalInfo)
 		{
-			if (URigVMRerouteNode* RerouteNode = Cast<URigVMRerouteNode>(InNodeProxy.GetSubject()))
+			if (InNodeProxy.GetSubject()->IsA<URigVMRerouteNode>())
 			{
-				if (RerouteNode->IsPassThrough() || RerouteNode->IsLiteral())
-				{
-					return;
-				}
+				return;
 			}
 			
 			const bool bIsCompilingFunction = OutTraversalInfo.LibraryNodeBeingCompiled != nullptr;
diff --git a/Engine/Plugins/Runtime/RigVM/Source/RigVMDeveloper/Private/RigVMCompiler/RigVMCompiler.cpp b/Engine/Plugins/Runtime/RigVM/Source/RigVMDeveloper/Private/RigVMCompiler/RigVMCompiler.cpp
--- a/Engine/Plugins/Runtime/RigVM/Source/RigVMDeveloper/Private/RigVMCompiler/RigVMCompiler.cpp
+++ b/Engine/Plugins/Runtime/RigVM/Source/RigVMDeveloper/Private/RigVMCompiler/RigVMCompiler.cpp
@@ -785,6 +785,33 @@ bool URigVMCompiler::Compile(TArray<URigVMGraph*> InGraphs, URigVMController* In
 				}
 			}
 
+			if(ModelNode->IsA<URigVMUnitNode>() || ModelNode->IsA<URigVMDispatchNode>())
+			{
+				// Maximum number of operands is 65535
+				int32 NumOperands = 0;
+				for (URigVMPin* Pin : ModelNode->GetPins())
+				{
+					if (Pin->IsExecuteContext())
+					{
+						continue;
+					}
+					
+					if (Pin->IsFixedSizeArray())
+					{
+						NumOperands += Pin->GetSubPins().Num();
+					}
+
+					NumOperands++;
+				}
+
+				if (NumOperands > 65535)
+				{
+					FString Format = FString::Printf(TEXT("Maximum number of operands for node @@ is 65535. Number of operands provided %d."), NumOperands);
+					Settings.ASTSettings.Report(EMessageSeverity::Error, ModelNode, Format);
+					bEncounteredGraphError = true;
+				}
+			}
+
 			for(URigVMPin* Pin : ModelNode->Pins)
 			{
 				if(!URigVMController::EnsurePinValidity(Pin, true))
@@ -1728,6 +1755,11 @@ int32 URigVMCompiler::TraverseCallExtern(const FRigVMCallExternExprAST* InExpr,
 			WorkData.VM->GetByteCode().AddZeroOp(BlockToRunOperand);
 		}
 
+		if (Operands.Num() > 64)
+		{
+			return INDEX_NONE;
+		}
+
 		// setup the instruction
 		const int32 FunctionIndex = WorkData.VM->AddRigVMFunction(Function->GetName());
 		check(FunctionIndex != INDEX_NONE);
@@ -1763,19 +1795,41 @@ int32 URigVMCompiler::TraverseCallExtern(const FRigVMCallExternExprAST* InExpr,
 			{
 				continue;
 			}
-			const FRigVMOperand& Operand = Operands[OperandIndex];
-			
-			if(InputPin->GetDirection() == ERigVMPinDirection::Output || InputPin->GetDirection() == ERigVMPinDirection::IO)
+			if (InputPin->IsFixedSizeArray())
 			{
-				OutputOperands.Add(Operand);
-			}
+				for (int32 SubPinIndex=0; SubPinIndex<InputPin->GetSubPins().Num(); ++SubPinIndex)
+				{
+					const URigVMPin* SubPin = InputPin->GetSubPins()[SubPinIndex];
+					const FRigVMOperand& Operand = Operands[OperandIndex+SubPinIndex];
+					if(SubPin->GetDirection() == ERigVMPinDirection::Output || SubPin->GetDirection() == ERigVMPinDirection::IO)
+					{
+						OutputOperands.Add(Operand);
+					}
 
-			if(InputPin->GetDirection() != ERigVMPinDirection::Input && InputPin->GetDirection() != ERigVMPinDirection::IO)
-			{
-				continue;
+					if(SubPin->GetDirection() != ERigVMPinDirection::Input && SubPin->GetDirection() != ERigVMPinDirection::IO)
+					{
+						continue;
+					}
+
+					InputsOperands.Add(Operand);
+				}
 			}
+			else
+			{
+				const FRigVMOperand& Operand = Operands[OperandIndex];
+				
+				if(InputPin->GetDirection() == ERigVMPinDirection::Output || InputPin->GetDirection() == ERigVMPinDirection::IO)
+				{
+					OutputOperands.Add(Operand);
+				}
 
-			InputsOperands.Add(Operand);
+				if(InputPin->GetDirection() != ERigVMPinDirection::Input && InputPin->GetDirection() != ERigVMPinDirection::IO)
+				{
+					continue;
+				}
+
+				InputsOperands.Add(Operand);
+			}
 		}
 
 		WorkData.VM->GetByteCode().SetOperandsForInstruction(
diff --git a/Engine/Plugins/Runtime/RigVM/Source/RigVMDeveloper/Private/RigVMModel/Nodes/RigVMRerouteNode.cpp b/Engine/Plugins/Runtime/RigVM/Source/RigVMDeveloper/Private/RigVMModel/Nodes/RigVMRerouteNode.cpp
--- a/Engine/Plugins/Runtime/RigVM/Source/RigVMDeveloper/Private/RigVMModel/Nodes/RigVMRerouteNode.cpp
+++ b/Engine/Plugins/Runtime/RigVM/Source/RigVMDeveloper/Private/RigVMModel/Nodes/RigVMRerouteNode.cpp
@@ -9,7 +9,6 @@ const FString URigVMRerouteNode::RerouteName = TEXT("Reroute");
 const FString URigVMRerouteNode::ValueName = TEXT("Value");
 
 URigVMRerouteNode::URigVMRerouteNode()
-: bShowAsFullNode(true)
 {
 }
 
@@ -56,81 +55,13 @@ FString URigVMRerouteNode::GetNodeTitle() const
 	return RerouteName;
 }
 
-bool URigVMRerouteNode::GetShowsAsFullNode() const
-{
-	return bShowAsFullNode;
-}
-
 FLinearColor URigVMRerouteNode::GetNodeColor() const
 {
 	return FLinearColor::White;
 }
 
-FName URigVMRerouteNode::GetNotation() const
-{
-	static constexpr TCHAR Format[] = TEXT("%s(io %s)");
-	static const FName RerouteNotation = *FString::Printf(Format, *RerouteName, *ValueName);
-	return RerouteNotation;
-}
-
-const FRigVMTemplate* URigVMRerouteNode::GetTemplate() const
-{
-	if(const FRigVMTemplate* SuperTemplate = Super::GetTemplate())
-	{
-		return SuperTemplate;
-	}
-	
-	if(CachedTemplate == nullptr)
-	{
-		static const FRigVMTemplate* RerouteNodeTemplate = nullptr;
-		if(RerouteNodeTemplate)
-		{
-			return RerouteNodeTemplate;
-		}
-
-		static TArray<FRigVMTemplateArgument> Arguments;
-		if(Arguments.IsEmpty())
-		{
-			static const TArray<FRigVMTemplateArgument::ETypeCategory> Categories = {
-				FRigVMTemplateArgument::ETypeCategory_Execute,
-				FRigVMTemplateArgument::ETypeCategory_SingleAnyValue,
-				FRigVMTemplateArgument::ETypeCategory_ArrayAnyValue
-			};
-			Arguments.Emplace(TEXT("Value"), ERigVMPinDirection::IO, Categories);
-		}
-
-		FRigVMTemplateDelegates Delegates;
-		Delegates.NewArgumentTypeDelegate = 
-			FRigVMTemplate_NewArgumentTypeDelegate::CreateLambda([](const FRigVMTemplate*, const FName& InArgumentName, int32 InTypeIndex)
-			{
-				// since a reroute has only one argument this is simple
-				FRigVMTemplateTypeMap Types;
-				Types.Add(InArgumentName, InTypeIndex);
-				return Types;
-			});
-
-		RerouteNodeTemplate = CachedTemplate = FRigVMRegistry::Get().GetOrAddTemplateFromArguments(*RerouteName, Arguments, Delegates);
-	}
-	return CachedTemplate;
-}
-
-bool URigVMRerouteNode::IsPassThrough() const
-{
-	return !GetPins()[0]->GetSourceLinks().IsEmpty();
-}
-
 bool URigVMRerouteNode::IsLiteral() const
 {
 	return GetPins()[0]->GetSourceLinks(true).IsEmpty();
 }
 
-bool URigVMRerouteNode::IsCompositionNode() const
-{
-	if (IsPassThrough() || IsLiteral())
-	{
-		return false;
-	}
-	
-	return true;
-}
-
diff --git a/Engine/Plugins/Runtime/RigVM/Source/RigVMDeveloper/Private/RigVMModel/RigVMClient.cpp b/Engine/Plugins/Runtime/RigVM/Source/RigVMDeveloper/Private/RigVMModel/RigVMClient.cpp
--- a/Engine/Plugins/Runtime/RigVM/Source/RigVMDeveloper/Private/RigVMModel/RigVMClient.cpp
+++ b/Engine/Plugins/Runtime/RigVM/Source/RigVMDeveloper/Private/RigVMModel/RigVMClient.cpp
@@ -792,6 +792,7 @@ FRigVMClientPatchResult FRigVMClient::PatchModelsOnLoad()
 		TGuardValue<bool> GuardSuspendTemplateComputation(Controller->bSuspendTemplateComputation, true);
 		TGuardValue<bool> GuardIsTransacting(Controller->bIsTransacting, true);
 
+		Result.Merge(Controller->PatchRerouteNodesOnLoad());
 		Result.Merge(Controller->PatchUnitNodesOnLoad());
 		Result.Merge(Controller->PatchDispatchNodesOnLoad());
 		Result.Merge(Controller->PatchBranchNodesOnLoad());
diff --git a/Engine/Plugins/Runtime/RigVM/Source/RigVMDeveloper/Private/RigVMModel/RigVMController.cpp b/Engine/Plugins/Runtime/RigVM/Source/RigVMDeveloper/Private/RigVMModel/RigVMController.cpp
--- a/Engine/Plugins/Runtime/RigVM/Source/RigVMDeveloper/Private/RigVMModel/RigVMController.cpp
+++ b/Engine/Plugins/Runtime/RigVM/Source/RigVMDeveloper/Private/RigVMModel/RigVMController.cpp
@@ -12,6 +12,9 @@
 #include "RigVMCore/RigVMUnknownType.h"
 #include "RigVMCore/RigVMByteCode.h"
 #include "RigVMFunctions/RigVMFunction_ControlFlow.h"
+#include "RigVMFunctions/RigVMDispatch_Array.h"
+#include "RigVMFunctions/RigVMDispatch_Constant.h"
+#include "RigVMFunctions/RigVMDispatch_MakeStruct.h"
 #include "RigVMCompiler/RigVMCompiler.h"
 #include "RigVMDeveloperModule.h"
 #include "UObject/PropertyPortFlags.h"
@@ -162,6 +165,12 @@ void URigVMController::SetGraph(URigVMGraph* InGraph)
 		PushGraph(InGraph, false);
 	}
 
+	// make sure we have "some" action stack. this is mainly relevant for unit tests
+	if(!WeakActionStack.IsValid())
+	{
+		SetActionStack(NewObject<URigVMActionStack>(this, TEXT("ActionStack")));
+	}
+
 	HandleModifiedEvent(ERigVMGraphNotifType::GraphChanged, GetGraph(), nullptr);
 }
 
@@ -518,11 +527,13 @@ int32 URigVMController::FPinInfoArray::AddPin(const URigVMPin* InPin, int32 InPa
 
 URigVMController::FPinInfoArray::FPinInfoArray(const URigVMNode* InNode, URigVMController* InController, const FPinInfoArray* InPreviousPinInfos)
 {
+	const bool bAddSubPins = !InNode->IsA<URigVMRerouteNode>();
+	
 	// this method adds pins as needed based on the property structure
 	for(const URigVMPin* Pin : InNode->GetPins())
 	{
 		const FString DefaultValue = Pin->GetDefaultValue();
-		(void)AddPin(InController, INDEX_NONE, Pin->GetFName(), Pin->GetDirection(), Pin->GetTypeIndex(), DefaultValue, nullptr, InPreviousPinInfos);
+		(void)AddPin(InController, INDEX_NONE, Pin->GetFName(), Pin->GetDirection(), Pin->GetTypeIndex(), DefaultValue, nullptr, InPreviousPinInfos, bAddSubPins);
 	}
 }
 
@@ -534,12 +545,12 @@ URigVMController::FPinInfoArray::FPinInfoArray(const FRigVMGraphFunctionHeader&
 	for(const FRigVMGraphFunctionArgument& FunctionArgument : FunctionHeader.Arguments)
 	{
 		const TRigVMTypeIndex TypeIndex = Registry.GetTypeIndexFromCPPType(FunctionArgument.CPPType.ToString());
-		(void)AddPin(InController, INDEX_NONE, FunctionArgument.Name, FunctionArgument.Direction, TypeIndex, FunctionArgument.DefaultValue, nullptr, InPreviousPinInfos);
+		(void)AddPin(InController, INDEX_NONE, FunctionArgument.Name, FunctionArgument.Direction, TypeIndex, FunctionArgument.DefaultValue, nullptr, InPreviousPinInfos, true);
 	}
 }
 
 int32 URigVMController::FPinInfoArray::AddPin(FProperty* InProperty, URigVMController* InController,
-                                              ERigVMPinDirection InDirection, int32 InParentIndex, const uint8* InDefaultValueMemory)
+                                              ERigVMPinDirection InDirection, int32 InParentIndex, const uint8* InDefaultValueMemory, bool bAddSubPins)
 {
 	// this method adds pins as needed based on the property structure
 
@@ -551,26 +562,28 @@ int32 URigVMController::FPinInfoArray::AddPin(FProperty* InProperty, URigVMContr
 		Pins[InParentIndex].SubPins.Add(Index);
 	}
 
-	if (const FStructProperty* StructProperty = CastField<FStructProperty>(InProperty))
-	{
-		(void)AddPins(StructProperty->Struct, InController, Pins[Index].Direction, Index, InDefaultValueMemory);
-	}
-	else if (const FArrayProperty* ArrayProperty = CastField<FArrayProperty>(InProperty))
+	if(bAddSubPins)
 	{
-		FScriptArrayHelper ArrayHelper(ArrayProperty, InDefaultValueMemory);
-		for(int32 ElementIndex = 0; ElementIndex < ArrayHelper.Num(); ElementIndex++)
+		if (const FStructProperty* StructProperty = CastField<FStructProperty>(InProperty))
+		{
+			(void)AddPins(StructProperty->Struct, InController, Pins[Index].Direction, Index, InDefaultValueMemory, bAddSubPins);
+		}
+		else if (const FArrayProperty* ArrayProperty = CastField<FArrayProperty>(InProperty))
 		{
-			const uint8* ElementDefaultValueMemory = ArrayHelper.GetRawPtr(ElementIndex);
-			const int32 SubIndex = AddPin(ArrayProperty->Inner, InController, Pins[Index].Direction, Index, ElementDefaultValueMemory);
-			Pins[SubIndex].Name = *FString::FormatAsNumber(ElementIndex);
+			FScriptArrayHelper ArrayHelper(ArrayProperty, InDefaultValueMemory);
+			for(int32 ElementIndex = 0; ElementIndex < ArrayHelper.Num(); ElementIndex++)
+			{
+				const uint8* ElementDefaultValueMemory = ArrayHelper.GetRawPtr(ElementIndex);
+				const int32 SubIndex = AddPin(ArrayProperty->Inner, InController, Pins[Index].Direction, Index, ElementDefaultValueMemory, bAddSubPins);
+				Pins[SubIndex].Name = *FString::FormatAsNumber(ElementIndex);
+			}
 		}
 	}
-
 	return Index;
 }
 
 int32 URigVMController::FPinInfoArray::AddPin(URigVMController* InController, int32 InParentIndex, const FName& InName, ERigVMPinDirection InDirection,
-	TRigVMTypeIndex InTypeIndex, const FString& InDefaultValue, const uint8* InDefaultValueMemory, const FPinInfoArray* InPreviousPinInfos)
+	TRigVMTypeIndex InTypeIndex, const FString& InDefaultValue, const uint8* InDefaultValueMemory, const FPinInfoArray* InPreviousPinInfos, bool bAddSubPins)
 {
 	const FRigVMRegistry& Registry = FRigVMRegistry::Get();
 		
@@ -600,7 +613,7 @@ int32 URigVMController::FPinInfoArray::AddPin(URigVMController* InController, in
 	}
 	
 	const FRigVMTemplateArgumentType& Type = Registry.GetType(InTypeIndex);
-	if(!Type.IsWildCard())
+	if(!Type.IsWildCard() && bAddSubPins)
 	{
 		if(Info.bIsArray)
 		{
@@ -623,7 +636,7 @@ int32 URigVMController::FPinInfoArray::AddPin(URigVMController* InController, in
 					ElementDefaultValueMemory = ElementDefaultValueMemoryScope.GetStructMemory();
 				}
 
-				(void)AddPin(InController, Index, *FString::FormatAsNumber(ElementIndex), InDirection, ElementTypeIndex, ElementDefaultValue, ElementDefaultValueMemory, InPreviousPinInfos);
+				(void)AddPin(InController, Index, *FString::FormatAsNumber(ElementIndex), InDirection, ElementTypeIndex, ElementDefaultValue, ElementDefaultValueMemory, InPreviousPinInfos, bAddSubPins);
 			}
 		}
 		else if(UScriptStruct* ScriptStruct = Cast<UScriptStruct>(Type.CPPTypeObject))
@@ -637,7 +650,7 @@ int32 URigVMController::FPinInfoArray::AddPin(URigVMController* InController, in
 				ScriptStruct->ImportText(*Pins[Index].DefaultValue, DefaultValueMemoryScope.GetStructMemory(), nullptr, PPF_None, &ErrorPipe, FString());
 				DefaultValueMemory = DefaultValueMemoryScope.GetStructMemory();
 			}
-			AddPins(ScriptStruct, InController, InDirection, Index, DefaultValueMemory);
+			AddPins(ScriptStruct, InController, InDirection, Index, DefaultValueMemory, bAddSubPins);
 		}
 	}
 
@@ -650,7 +663,7 @@ int32 URigVMController::FPinInfoArray::AddPin(URigVMController* InController, in
 }
 
 void URigVMController::FPinInfoArray::AddPins(UScriptStruct* InScriptStruct, URigVMController* InController,
-	ERigVMPinDirection InDirection, int32 InParentIndex, const uint8* InDefaultValueMemory)
+	ERigVMPinDirection InDirection, int32 InParentIndex, const uint8* InDefaultValueMemory, bool bAddSubPins)
 {
 	if (InController->ShouldStructBeUnfolded(InScriptStruct))
 	{
@@ -666,7 +679,18 @@ void URigVMController::FPinInfoArray::AddPins(UScriptStruct* InScriptStruct, URi
 				{
 					DefaultValueMemory = It->ContainerPtrToValuePtr<uint8>(InDefaultValueMemory);
 				}
-				(void)AddPin(*It, InController, InDirection, InParentIndex, DefaultValueMemory);
+
+				bool bAddSubPinsForProperty = bAddSubPins;
+#if WITH_EDITOR
+				if(bAddSubPins)
+				{
+					if(It->HasMetaData(FRigVMStruct::HideSubPinsMetaName))
+					{
+						bAddSubPinsForProperty = false;
+					}
+				}
+#endif
+				(void)AddPin(*It, InController, InDirection, InParentIndex, DefaultValueMemory, bAddSubPinsForProperty);
 			}
 		}
 	}
@@ -898,22 +922,6 @@ TArray<FString> URigVMController::GeneratePythonCommands()
 					*TargetPinPath));
 	}
 
-	// Reroutes
-	{
-		TArray<URigVMRerouteNode*> Reroutes;
-		for (URigVMNode* Node : GetGraph()->GetNodes())
-		{
-			if (URigVMRerouteNode* Reroute = Cast<URigVMRerouteNode>(Node))
-			{
-				// SetRerouteCompactnessByName(const FName& InNodeName, bool bShowAsFullNode, bool bSetupUndoRedo = true);
-				Commands.Add(FString::Printf(TEXT("blueprint.get_controller_by_name('%s').set_reroute_compactness_by_name('%s', %s)"),
-					*GraphName,
-					*Reroute->GetName(),
-					Reroute->GetShowsAsFullNode() ? TEXT("True") : TEXT("False")));
-			}
-		}
-	}
-
 	return Commands;
 }
 
@@ -1091,10 +1099,9 @@ TArray<FString> URigVMController::GetAddNodePythonCommands(URigVMNode* Node) con
 	}
 	else if (const URigVMRerouteNode* RerouteNode = Cast<URigVMRerouteNode>(Node))
 	{
-		// add_free_reroute_node(bool bShowAsFullNode, const FString& InCPPType, const FName& InCPPTypeObjectPath, bool bIsConstant, const FName& InCustomWidgetName, const FString& InDefaultValue, const FVector2D& InPosition = FVector2D::ZeroVector, const FString& InNodeName = TEXT(""), bool bSetupUndoRedo = true);
-		Commands.Add(FString::Printf(TEXT("blueprint.get_controller_by_name('%s').add_free_reroute_node(%s, '%s', '%s', %s, '%s', '%s', %s, '%s')"),
+		// add_free_reroute_node(const FString& InCPPType, const FName& InCPPTypeObjectPath, bool bIsConstant, const FName& InCustomWidgetName, const FString& InDefaultValue, const FVector2D& InPosition = FVector2D::ZeroVector, const FString& InNodeName = TEXT(""), bool bSetupUndoRedo = true);
+		Commands.Add(FString::Printf(TEXT("blueprint.get_controller_by_name('%s').add_free_reroute_node('%s', '%s', %s, '%s', '%s', %s, '%s')"),
 				*GraphName,
-				RerouteNode->GetShowsAsFullNode() ? TEXT("True") : TEXT("False"),
 				*RerouteNode->GetPins()[0]->GetCPPType(),
 				*RerouteNode->GetPins()[0]->GetCPPTypeObject()->GetPathName(),
 				RerouteNode->GetPins()[0]->IsDefinedAsConstant() ? TEXT("True") : TEXT("False"),
@@ -2787,7 +2794,7 @@ URigVMCommentNode* URigVMController::AddCommentNode(const FString& InCommentText
 	return Node;
 }
 
-URigVMRerouteNode* URigVMController::AddRerouteNodeOnLink(URigVMLink* InLink, bool bShowAsFullNode, const FVector2D& InPosition, const FString& InNodeName, bool bSetupUndoRedo, bool bPrintPythonCommand)
+URigVMRerouteNode* URigVMController::AddRerouteNodeOnLink(URigVMLink* InLink, const FVector2D& InPosition, const FString& InNodeName, bool bSetupUndoRedo, bool bPrintPythonCommand)
 {
 	if(!IsValidLinkForGraph(InLink))
 	{
@@ -2818,7 +2825,7 @@ URigVMRerouteNode* URigVMController::AddRerouteNodeOnLink(URigVMLink* InLink, bo
 		GetActionStack()->BeginAction(Action);
 	}
 
-	URigVMRerouteNode* Node = AddRerouteNodeOnPin(TargetPin->GetPinPath(), true, bShowAsFullNode, InPosition, InNodeName, bSetupUndoRedo);
+	URigVMRerouteNode* Node = AddRerouteNodeOnPin(TargetPin->GetPinPath(), true, InPosition, InNodeName, bSetupUndoRedo);
 	if (Node == nullptr)
 	{
 		if (bSetupUndoRedo)
@@ -2844,8 +2851,7 @@ URigVMRerouteNode* URigVMController::AddRerouteNodeOnLink(URigVMLink* InLink, bo
 		RigVMPythonUtils::Print(GetGraphOuterName(), 
 							FString::Printf(TEXT("blueprint.get_controller_by_name('%s').add_reroute_node_on_link_path('%s', %s, %s, '%s')"),
 											*GraphName,
-											*InLink->GetPinPathRepresentation(),
-											(bShowAsFullNode) ? TEXT("True") : TEXT("False"),
+											*URigVMLink::GetPinPathRepresentation(SourcePin->GetPinPath(), TargetPin->GetPinPath()),
 											*RigVMPythonUtils::Vector2DToPythonString(Node->GetPosition()),
 											*NodeName));
 	}
@@ -2853,7 +2859,7 @@ URigVMRerouteNode* URigVMController::AddRerouteNodeOnLink(URigVMLink* InLink, bo
 	return Node;
 }
 
-URigVMRerouteNode* URigVMController::AddRerouteNodeOnLinkPath(const FString& InLinkPinPathRepresentation, bool bShowAsFullNode, const FVector2D& InPosition, const FString&
+URigVMRerouteNode* URigVMController::AddRerouteNodeOnLinkPath(const FString& InLinkPinPathRepresentation, const FVector2D& InPosition, const FString&
                                                               InNodeName, bool bSetupUndoRedo, bool bPrintPythonCommand)
 {
 	if (!IsValidGraph())
@@ -2870,10 +2876,10 @@ URigVMRerouteNode* URigVMController::AddRerouteNodeOnLinkPath(const FString& InL
 	check(Graph);
 
 	URigVMLink* Link = Graph->FindLink(InLinkPinPathRepresentation);
-	return AddRerouteNodeOnLink(Link, bShowAsFullNode, InPosition, InNodeName, bSetupUndoRedo, bPrintPythonCommand);
+	return AddRerouteNodeOnLink(Link, InPosition, InNodeName, bSetupUndoRedo, bPrintPythonCommand);
 }
 
-URigVMRerouteNode* URigVMController::AddRerouteNodeOnPin(const FString& InPinPath, bool bAsInput, bool bShowAsFullNode, const FVector2D& InPosition, const FString& InNodeName, bool bSetupUndoRedo, bool bPrintPythonCommand)
+URigVMRerouteNode* URigVMController::AddRerouteNodeOnPin(const FString& InPinPath, bool bAsInput, const FVector2D& InPosition, const FString& InNodeName, bool bSetupUndoRedo, bool bPrintPythonCommand)
 {
 	if (!IsValidGraph())
 	{
@@ -2920,19 +2926,12 @@ URigVMRerouteNode* URigVMController::AddRerouteNodeOnPin(const FString& InPinPat
 	FString Name = GetValidNodeName(InNodeName.IsEmpty() ? FString(TEXT("RerouteNode")) : InNodeName);
 	URigVMRerouteNode* Node = NewObject<URigVMRerouteNode>(Graph, *Name);
 	Node->Position = InPosition;
-	Node->bShowAsFullNode = bShowAsFullNode;
 
 	URigVMPin* ValuePin = NewObject<URigVMPin>(Node, *URigVMRerouteNode::ValueName);
 	ConfigurePinFromPin(ValuePin, Pin);
 	ValuePin->Direction = ERigVMPinDirection::IO;
 	AddNodePin(Node, ValuePin);
 
-	if (ValuePin->IsStruct())
-	{
-		TGuardValue<bool> SuspendNotifications(bSuspendNotifications, true);
-		AddPinsForStruct(ValuePin->GetScriptStruct(), Node, ValuePin, ValuePin->Direction, FString(), false);
-	}
-
 	FString DefaultValue = Pin->GetDefaultValue();
 	if (!DefaultValue.IsEmpty())
 	{
@@ -2974,14 +2973,13 @@ URigVMRerouteNode* URigVMController::AddRerouteNodeOnPin(const FString& InPinPat
 	{
 		const FString GraphName = GetSanitizedGraphName(GetGraph()->GetGraphName());
 		const FString NodeName = GetSanitizedNodeName(Node->GetName());
-		// AddRerouteNodeOnPin(const FString& InPinPath, bool bAsInput, bool bShowAsFullNode, const FVector2D& InPosition, const FString& InNodeName, bool bSetupUndoRedo, bool bPrintPythonCommand)
+		// AddRerouteNodeOnPin(const FString& InPinPath, bool bAsInput, const FVector2D& InPosition, const FString& InNodeName, bool bSetupUndoRedo, bool bPrintPythonCommand)
 		
 		RigVMPythonUtils::Print(GetGraphOuterName(), 
 							FString::Printf(TEXT("blueprint.get_controller_by_name('%s').add_reroute_node_on_pin('%s', %s, %s, %s '%s')"),
 											*GraphName,
 											*GetSanitizedPinPath(InPinPath),
 											(bAsInput) ? TEXT("True") : TEXT("False"),
-											(bShowAsFullNode) ? TEXT("True") : TEXT("False"),
 											*RigVMPythonUtils::Vector2DToPythonString(Node->GetPosition()),
 											*NodeName));
 	}
@@ -5312,7 +5310,7 @@ TArray<URigVMNode*> URigVMController::ExpandLibraryNode(URigVMLibraryNode* InNod
 	// b) create a map of new links to create by following the links to / from the library node
 	TMap<FString, TArray<FString>> ToLibraryNode;
 	TMap<FString, TArray<FString>> FromLibraryNode;
-	TArray<URigVMPin*> LibraryPinsToReroute;
+	TArray<URigVMPin*> LibraryPinsToTurnIntoConstant;
 
 	TArray<URigVMLink*> LibraryLinks = InNode->GetLinks();
 	for (URigVMLink* Link : LibraryLinks)
@@ -5321,7 +5319,7 @@ TArray<URigVMNode*> URigVMController::ExpandLibraryNode(URigVMLibraryNode* InNod
 		{
 			if (!Link->GetTargetPin()->IsRootPin())
 			{
-				LibraryPinsToReroute.AddUnique(Link->GetTargetPin()->GetRootPin());
+				LibraryPinsToTurnIntoConstant.AddUnique(Link->GetTargetPin()->GetRootPin());
 			}
 
 			FString NodeName, PinPath;
@@ -5332,7 +5330,7 @@ TArray<URigVMNode*> URigVMController::ExpandLibraryNode(URigVMLibraryNode* InNod
 		{
 			if (!Link->GetSourcePin()->IsRootPin())
 			{
-				LibraryPinsToReroute.AddUnique(Link->GetSourcePin()->GetRootPin());
+				LibraryPinsToTurnIntoConstant.AddUnique(Link->GetSourcePin()->GetRootPin());
 			}
 
 			FString NodeName, PinPath;
@@ -5367,7 +5365,7 @@ TArray<URigVMNode*> URigVMController::ExpandLibraryNode(URigVMLibraryNode* InNod
 
 			if (!Link->GetSourcePin()->IsRootPin())
 			{
-				LibraryPinsToReroute.AddUnique(InNode->FindPin(Link->GetSourcePin()->GetRootPin()->GetName()));
+				LibraryPinsToTurnIntoConstant.AddUnique(InNode->FindPin(Link->GetSourcePin()->GetRootPin()->GetName()));
 			}
 
 			FString NodeName, PinPath;
@@ -5408,7 +5406,7 @@ TArray<URigVMNode*> URigVMController::ExpandLibraryNode(URigVMLibraryNode* InNod
 
 			if (!Link->GetTargetPin()->IsRootPin())
 			{
-				LibraryPinsToReroute.AddUnique(InNode->FindPin(Link->GetTargetPin()->GetRootPin()->GetName()));
+				LibraryPinsToTurnIntoConstant.AddUnique(InNode->FindPin(Link->GetTargetPin()->GetRootPin()->GetName()));
 			}
 
 			FString NodeName, PinPath;
@@ -5449,57 +5447,51 @@ TArray<URigVMNode*> URigVMController::ExpandLibraryNode(URigVMLibraryNode* InNod
 		}
 	}
 
-	// f) create reroutes for all pins which had wires on sub pins
-	TMap<FString, URigVMPin*> ReroutedInputPins;
-	TMap<FString, URigVMPin*> ReroutedOutputPins;
-	FVector2D RerouteInputPosition = InNode->Position + FVector2D(-Diagonal.X, -Diagonal.Y) * 0.5 + FVector2D(-200.f, 0.f);
-	FVector2D RerouteOutputPosition = InNode->Position + FVector2D(Diagonal.X, -Diagonal.Y) * 0.5 + FVector2D(250.f, 0.f);
-	for (URigVMPin* LibraryPinToReroute : LibraryPinsToReroute)
+	// f) create constants for all pins which had wires on sub pins
+	TMap<FString, URigVMPin*> ConstantInputPins;
+	TMap<FString, URigVMPin*> ConstantOutputPins;
+	FVector2D ConstantInputPosition = InNode->Position + FVector2D(-Diagonal.X, -Diagonal.Y) * 0.5 + FVector2D(-200.f, 0.f);
+	FVector2D ConstantOutputPosition = InNode->Position + FVector2D(Diagonal.X, -Diagonal.Y) * 0.5 + FVector2D(250.f, 0.f);
+	for (URigVMPin* LibraryPinToTurnIntoConstant : LibraryPinsToTurnIntoConstant)
 	{
-		if (LibraryPinToReroute->GetDirection() == ERigVMPinDirection::Input ||
-			LibraryPinToReroute->GetDirection() == ERigVMPinDirection::IO)
+		if (LibraryPinToTurnIntoConstant->GetDirection() == ERigVMPinDirection::Input ||
+			LibraryPinToTurnIntoConstant->GetDirection() == ERigVMPinDirection::IO)
 		{
-			URigVMRerouteNode* RerouteNode =
-				AddFreeRerouteNode(
-					true,
-					LibraryPinToReroute->GetCPPType(),
-					*LibraryPinToReroute->GetCPPTypeObject()->GetPathName(),
-					false,
-					NAME_None,
-					LibraryPinToReroute->GetDefaultValue(),
-					RerouteInputPosition,
-					FString::Printf(TEXT("Reroute_%s"), *LibraryPinToReroute->GetName()),
+			URigVMTemplateNode* ConstantNode =
+				AddConstantNode(
+					LibraryPinToTurnIntoConstant->GetCPPType(),
+					*LibraryPinToTurnIntoConstant->GetCPPTypeObject()->GetPathName(),
+					LibraryPinToTurnIntoConstant->GetDefaultValue(),
+					ConstantInputPosition,
+					FString::Printf(TEXT("Constant_%s"), *LibraryPinToTurnIntoConstant->GetName()),
 					false);
 
-			RerouteInputPosition += FVector2D(0.f, 150.f);
+			ConstantInputPosition += FVector2D(0.f, 150.f);
 
-			URigVMPin* ReroutePin = RerouteNode->FindPin(URigVMRerouteNode::ValueName);
-			ApplyPinState(ReroutePin, GetPinState(LibraryPinToReroute));
-			ReroutedInputPins.Add(LibraryPinToReroute->GetName(), ReroutePin);
-			ExpandedNodes.Add(RerouteNode);
+			URigVMPin* ValuePin = ConstantNode->FindPin(FRigVMDispatch_Constant::ValueName.ToString());
+			ApplyPinState(ValuePin, GetPinState(LibraryPinToTurnIntoConstant));
+			ConstantInputPins.Add(LibraryPinToTurnIntoConstant->GetName(), ValuePin);
+			ExpandedNodes.Add(ConstantNode);
 		}
 
-		if (LibraryPinToReroute->GetDirection() == ERigVMPinDirection::Output ||
-			LibraryPinToReroute->GetDirection() == ERigVMPinDirection::IO)
+		if (LibraryPinToTurnIntoConstant->GetDirection() == ERigVMPinDirection::Output ||
+			LibraryPinToTurnIntoConstant->GetDirection() == ERigVMPinDirection::IO)
 		{
-			URigVMRerouteNode* RerouteNode =
-				AddFreeRerouteNode(
-					true,
-					LibraryPinToReroute->GetCPPType(),
-					*LibraryPinToReroute->GetCPPTypeObject()->GetPathName(),
-					false,
-					NAME_None,
-					LibraryPinToReroute->GetDefaultValue(),
-					RerouteOutputPosition,
-					FString::Printf(TEXT("Reroute_%s"), *LibraryPinToReroute->GetName()),
+			URigVMTemplateNode* ConstantNode =
+				AddConstantNode(
+					LibraryPinToTurnIntoConstant->GetCPPType(),
+					*LibraryPinToTurnIntoConstant->GetCPPTypeObject()->GetPathName(),
+					LibraryPinToTurnIntoConstant->GetDefaultValue(),
+					ConstantOutputPosition,
+					FString::Printf(TEXT("Constant_%s"), *LibraryPinToTurnIntoConstant->GetName()),
 					false);
 
-			RerouteOutputPosition += FVector2D(0.f, 150.f);
+			ConstantOutputPosition += FVector2D(0.f, 150.f);
 
-			URigVMPin* ReroutePin = RerouteNode->FindPin(URigVMRerouteNode::ValueName);
-			ApplyPinState(ReroutePin, GetPinState(LibraryPinToReroute));
-			ReroutedOutputPins.Add(LibraryPinToReroute->GetName(), ReroutePin);
-			ExpandedNodes.Add(RerouteNode);
+			URigVMPin* ValuePin = ConstantNode->FindPin(FRigVMDispatch_Constant::ValueName.ToString());
+			ApplyPinState(ValuePin, GetPinState(LibraryPinToTurnIntoConstant));
+			ConstantOutputPins.Add(LibraryPinToTurnIntoConstant->GetName(), ValuePin);
+			ExpandedNodes.Add(ConstantNode);
 		}
 	}
 
@@ -5546,13 +5538,13 @@ TArray<URigVMNode*> URigVMController::ExpandLibraryNode(URigVMLibraryNode* InNod
 				const FString SourcePinPath = LibraryPinPath;
 				FString TargetPinPath = LibraryPinLinks[0];
 
-				// if the pin on the library node is represented by a reroute
+				// if the pin on the library node is represented by a constant
 				// we need to remap to that instead.
-				if(URigVMPin** ReroutedPinPtr = ReroutedInputPins.Find(SourcePinPath))
+				if(URigVMPin** ConstantPinPtr = ConstantInputPins.Find(SourcePinPath))
 				{
-					if(URigVMPin* ReroutedPin = *ReroutedPinPtr)
+					if(URigVMPin* ConstantPin = *ConstantPinPtr)
 					{
-						TargetPinPath = ReroutedPin->GetPinPath();
+						TargetPinPath = ConstantPin->GetPinPath();
 					}
 				}
 
@@ -5570,13 +5562,13 @@ TArray<URigVMNode*> URigVMController::ExpandLibraryNode(URigVMLibraryNode* InNod
 				const FString SourcePinPath = LibraryPinPath;
 				FString TargetPinPath = LibraryPinLinks[0];
 
-				// if the pin on the library node is represented by a reroute
+				// if the pin on the library node is represented by a constant
 				// we need to remap to that instead.
-				if(URigVMPin** ReroutedPinPtr = ReroutedOutputPins.Find(SourcePinPath))
+				if(URigVMPin** ConstantPinPtr = ConstantOutputPins.Find(SourcePinPath))
 				{
-					if(URigVMPin* ReroutedPin = *ReroutedPinPtr)
+					if(URigVMPin* ConstantPin = *ConstantPinPtr)
 					{
-						TargetPinPath = ReroutedPin->GetPinPath();
+						TargetPinPath = ConstantPin->GetPinPath();
 					}
 				}
 
@@ -5586,7 +5578,7 @@ TArray<URigVMNode*> URigVMController::ExpandLibraryNode(URigVMLibraryNode* InNod
 	}
 
 	// h) re-establish all of the links going to the left of the library node
-	//    in this pass we only care about pins which have reroutes
+	//    in this pass we only care about pins which have constants
 	for (const TPair<FString, TArray<FString>>& ToLibraryNodePair : ToLibraryNode)
 	{
 		FString LibraryNodePinName, LibraryNodePinPathSuffix;
@@ -5595,13 +5587,13 @@ TArray<URigVMNode*> URigVMController::ExpandLibraryNode(URigVMLibraryNode* InNod
 			LibraryNodePinName = ToLibraryNodePair.Key;
 		}
 
-		if (!ReroutedInputPins.Contains(LibraryNodePinName))
+		if (!ConstantInputPins.Contains(LibraryNodePinName))
 		{
 			continue;
 		}
 
-		URigVMPin* ReroutedPin = ReroutedInputPins.FindChecked(LibraryNodePinName);
-		URigVMPin* TargetPin = LibraryNodePinPathSuffix.IsEmpty() ? ReroutedPin : ReroutedPin->FindSubPin(LibraryNodePinPathSuffix);
+		URigVMPin* ConstantPin = ConstantInputPins.FindChecked(LibraryNodePinName);
+		URigVMPin* TargetPin = LibraryNodePinPathSuffix.IsEmpty() ? ConstantPin : ConstantPin->FindSubPin(LibraryNodePinPathSuffix);
 		check(TargetPin);
 
 		for (const FString& SourcePinPath : ToLibraryNodePair.Value)
@@ -5655,17 +5647,17 @@ TArray<URigVMNode*> URigVMController::ExpandLibraryNode(URigVMLibraryNode* InNod
 			RemappedSourcePinPath = URigVMPin::JoinPinPath(RemappedSourcePinPath, EntryPinPathSuffix);
 		}
 
-		// remap the top level pin in case we need to insert a reroute
+		// remap the top level pin in case we need to insert a constant
 		FString EntryPinName;
 		if (!URigVMPin::SplitPinPathAtStart(FromEntryNodePair.Key, EntryPinPath, EntryPinPathSuffix))
 		{
 			EntryPinName = FromEntryNodePair.Key;
 			EntryPinPathSuffix.Reset();
 		}
-		if (ReroutedInputPins.Contains(EntryPinName))
+		if (ConstantInputPins.Contains(EntryPinName))
 		{
-			URigVMPin* ReroutedPin = ReroutedInputPins.FindChecked(EntryPinName);
-			URigVMPin* TargetPin = EntryPinPathSuffix.IsEmpty() ? ReroutedPin : ReroutedPin->FindSubPin(EntryPinPathSuffix);
+			URigVMPin* ConstantPin = ConstantInputPins.FindChecked(EntryPinName);
+			URigVMPin* TargetPin = EntryPinPathSuffix.IsEmpty() ? ConstantPin : ConstantPin->FindSubPin(EntryPinPathSuffix);
 			check(TargetPin);
 			RemappedSourcePinPath = TargetPin->GetPinPath();
 		}
@@ -5721,7 +5713,7 @@ TArray<URigVMNode*> URigVMController::ExpandLibraryNode(URigVMLibraryNode* InNod
 	}
 
 	// j) re-establish all of the links going from the right of the library node
-	//    in this pass we only check pins which have a reroute
+	//    in this pass we only check pins which have a constant
 	for (const TPair<FString, TArray<FString>>& ToReturnNodePair : ToReturnNode)
 	{
 		FString LibraryNodePinName, LibraryNodePinPathSuffix;
@@ -5730,13 +5722,13 @@ TArray<URigVMNode*> URigVMController::ExpandLibraryNode(URigVMLibraryNode* InNod
 			LibraryNodePinName = ToReturnNodePair.Key;
 		}
 
-		if (!ReroutedOutputPins.Contains(LibraryNodePinName))
+		if (!ConstantOutputPins.Contains(LibraryNodePinName))
 		{
 			continue;
 		}
 
-		URigVMPin* ReroutedPin = ReroutedOutputPins.FindChecked(LibraryNodePinName);
-		URigVMPin* TargetPin = LibraryNodePinPathSuffix.IsEmpty() ? ReroutedPin : ReroutedPin->FindSubPin(LibraryNodePinPathSuffix);
+		URigVMPin* ConstantPin = ConstantOutputPins.FindChecked(LibraryNodePinName);
+		URigVMPin* TargetPin = LibraryNodePinPathSuffix.IsEmpty() ? ConstantPin : ConstantPin->FindSubPin(LibraryNodePinPathSuffix);
 		check(TargetPin);
 
 		for (const FString& SourcePinpath : ToReturnNodePair.Value)
@@ -5790,17 +5782,17 @@ TArray<URigVMNode*> URigVMController::ExpandLibraryNode(URigVMLibraryNode* InNod
 			RemappedSourcePinPath = URigVMPin::JoinPinPath(RemappedSourcePinPath, FromLibraryNodePinPathSuffix);
 		}
 
-		// remap the top level pin in case we need to insert a reroute
+		// remap the top level pin in case we need to insert a constant
 		FString ReturnPinName, ReturnPinPathSuffix;
 		if (!URigVMPin::SplitPinPathAtStart(FromLibraryNodePair.Key, ReturnPinName, ReturnPinPathSuffix))
 		{
 			ReturnPinName = FromLibraryNodePair.Key;
 			ReturnPinPathSuffix.Reset();
 		}
-		if (ReroutedOutputPins.Contains(ReturnPinName))
+		if (ConstantOutputPins.Contains(ReturnPinName))
 		{
-			URigVMPin* ReroutedPin = ReroutedOutputPins.FindChecked(ReturnPinName);
-			URigVMPin* SourcePin = ReturnPinPathSuffix.IsEmpty() ? ReroutedPin : ReroutedPin->FindSubPin(ReturnPinPathSuffix);
+			URigVMPin* ConstantPin = ConstantOutputPins.FindChecked(ReturnPinName);
+			URigVMPin* SourcePin = ReturnPinPathSuffix.IsEmpty() ? ConstantPin : ConstantPin->FindSubPin(ReturnPinPathSuffix);
 			check(SourcePin);
 			RemappedSourcePinPath = SourcePin->GetPinPath();
 		}
@@ -7376,66 +7368,6 @@ bool URigVMController::SetCommentTextByName(const FName& InNodeName, const FStri
 	return SetCommentText(Node, InCommentText, InCommentFontSize, bInCommentBubbleVisible, bInCommentColorBubble, bSetupUndoRedo, bPrintPythonCommand);
 }
 
-bool URigVMController::SetRerouteCompactness(URigVMNode* InNode, bool bShowAsFullNode, bool bSetupUndoRedo, bool bPrintPythonCommand)
-{
-	if (!bIsTransacting && !IsGraphEditable())
-	{
-		return false;
-	}
-	
-	if (!IsValidNodeForGraph(InNode))
-	{
-		return false;
-	}
-
-	if (URigVMRerouteNode* RerouteNode = Cast<URigVMRerouteNode>(InNode))
-	{
-		if (RerouteNode->bShowAsFullNode == bShowAsFullNode)
-		{
-			return false;
-		}
-
-		FRigVMSetRerouteCompactnessAction Action;
-		if (bSetupUndoRedo)
-		{
-			Action = FRigVMSetRerouteCompactnessAction(this, RerouteNode, bShowAsFullNode);
-			Action.SetTitle(FString::Printf(TEXT("Set Reroute Size")));
-			GetActionStack()->BeginAction(Action);
-		}
-
-		RerouteNode->bShowAsFullNode = bShowAsFullNode;
-		Notify(ERigVMGraphNotifType::RerouteCompactnessChanged, InNode);
-
-		if (bSetupUndoRedo)
-		{
-			GetActionStack()->EndAction(Action);
-		}
-
-		return true;
-	}
-
-	return false;
-}
-
-bool URigVMController::SetRerouteCompactnessByName(const FName& InNodeName, bool bShowAsFullNode, bool bSetupUndoRedo)
-{
-	if(!IsValidGraph())
-	{
-		return false;
-	}
-
-	if (!bIsTransacting && !IsGraphEditable())
-	{
-		return false;
-	}
-
-	URigVMGraph* Graph = GetGraph();
-	check(Graph);
-
-	URigVMNode* Node = Graph->FindNodeByName(InNodeName);
-	return SetRerouteCompactness(Node, bShowAsFullNode, bSetupUndoRedo);
-}
-
 bool URigVMController::RenameVariable(const FName& InOldName, const FName& InNewName, bool bSetupUndoRedo)
 {
 	if(!IsValidGraph())
@@ -7528,36 +7460,6 @@ bool URigVMController::RenameParameter(const FName& InOldName, const FName& InNe
 	return false;
 }
 
-void URigVMController::UpdateRerouteNodeAfterChangingLinks(URigVMPin* PinChanged, bool bSetupUndoRedo)
-{
-	if (bIgnoreRerouteCompactnessChanges)
-	{
-		return;
-	}
-
-	if (!IsValidGraph())
-	{
-		return;
-	}
-
-	URigVMRerouteNode* Node = Cast<URigVMRerouteNode>(PinChanged->GetNode());
-	if (Node == nullptr)
-	{
-		return;
-	}
-
-	int32 NbTotalSources = Node->Pins[0]->GetSourceLinks(true /* recursive */).Num();
-	int32 NbTotalTargets = Node->Pins[0]->GetTargetLinks(true /* recursive */).Num();
-	int32 NbToplevelSources = Node->Pins[0]->GetSourceLinks(false /* recursive */).Num();
-	int32 NbToplevelTargets = Node->Pins[0]->GetTargetLinks(false /* recursive */).Num();
-
-	bool bJustTopLevelConnections = (NbTotalSources == NbToplevelSources) && (NbTotalTargets == NbToplevelTargets);
-	bool bOnlyConnectionsOnOneSide = (NbTotalSources == 0) || (NbTotalTargets == 0);
-	bool bShowAsFullNode = (!bJustTopLevelConnections) || bOnlyConnectionsOnOneSide;
-
-	SetRerouteCompactness(Node, bShowAsFullNode, bSetupUndoRedo);
-}
-
 bool URigVMController::SetPinExpansion(const FString& InPinPath, bool bIsExpanded, bool bSetupUndoRedo, bool bPrintPythonCommand)
 {
 	if (!IsValidGraph())
@@ -7840,11 +7742,17 @@ bool URigVMController::SetPinDefaultValue(URigVMPin* InPin, const FString& InDef
 			{
 				while (Elements.Num() > InPin->SubPins.Num())
 				{
-					InsertArrayPin(InPin, INDEX_NONE, FString(), bSetupUndoRedo);
+					if(!InsertArrayPin(InPin, INDEX_NONE, FString(), bSetupUndoRedo))
+					{
+						break;
+					}
 				}
 				while (Elements.Num() < InPin->SubPins.Num())
 				{
-					RemoveArrayPin(InPin->SubPins.Last()->GetPinPath(), bSetupUndoRedo);
+					if(!RemoveArrayPin(InPin->SubPins.Last()->GetPinPath(), bSetupUndoRedo))
+					{
+						break;
+					}
 				}
 			}
 			else
@@ -7854,12 +7762,14 @@ bool URigVMController::SetPinDefaultValue(URigVMPin* InPin, const FString& InDef
 
 			for (int32 ElementIndex = 0; ElementIndex < Elements.Num(); ElementIndex++)
 			{
-				URigVMPin* SubPin = InPin->SubPins[ElementIndex];
-				PostProcessDefaultValue(SubPin, Elements[ElementIndex]);
-				if (!Elements[ElementIndex].IsEmpty())
+				if(URigVMPin* SubPin = InPin->FindSubPin(FString::FromInt(ElementIndex)))
 				{
-					SetPinDefaultValue(SubPin, Elements[ElementIndex], bResizeArrays, false, false);
-					bSetPinDefaultValueSucceeded = true;
+					PostProcessDefaultValue(SubPin, Elements[ElementIndex]);
+					if (!Elements[ElementIndex].IsEmpty())
+					{
+						SetPinDefaultValue(SubPin, Elements[ElementIndex], bResizeArrays, false, false);
+						bSetPinDefaultValueSucceeded = true;
+					}
 				}
 			}
 		}
@@ -10038,12 +9948,6 @@ bool URigVMController::AddLink(URigVMPin* OutputPin, URigVMPin* InputPin, bool b
 	}
 	Notify(ERigVMGraphNotifType::LinkAdded, Link);
 
-	if (bSetupUndoRedo)
-	{
-		UpdateRerouteNodeAfterChangingLinks(OutputPin, bSetupUndoRedo);
-		UpdateRerouteNodeAfterChangingLinks(InputPin, bSetupUndoRedo);
-	}
-
 	if (bSetupUndoRedo)
 	{
 #if WITH_EDITOR
@@ -10273,12 +10177,6 @@ bool URigVMController::BreakLink(URigVMPin* OutputPin, URigVMPin* InputPin, bool
 
 			DestroyObject(Link);
 
-			if (bSetupUndoRedo)
-			{
-				UpdateRerouteNodeAfterChangingLinks(OutputPin, bSetupUndoRedo);
-				UpdateRerouteNodeAfterChangingLinks(InputPin, bSetupUndoRedo);
-			}
-
 			if (bSetupUndoRedo)
 			{
 				GetActionStack()->EndAction(Action);
@@ -12634,7 +12532,7 @@ bool URigVMController::SetVariableName(URigVMVariableNode* InVariableNode, const
 	return true;
 }
 
-URigVMRerouteNode* URigVMController::AddFreeRerouteNode(bool bShowAsFullNode, const FString& InCPPType, const FName& InCPPTypeObjectPath, bool bIsConstant, const FName& InCustomWidgetName, const FString& InDefaultValue, const FVector2D& InPosition, const FString& InNodeName, bool bSetupUndoRedo)
+URigVMRerouteNode* URigVMController::AddFreeRerouteNode(const FString& InCPPType, const FName& InCPPTypeObjectPath, bool bIsConstant, const FName& InCustomWidgetName, const FString& InDefaultValue, const FVector2D& InPosition, const FString& InNodeName, bool bSetupUndoRedo)
 {
 	if (!IsValidGraph())
 	{
@@ -12666,7 +12564,6 @@ URigVMRerouteNode* URigVMController::AddFreeRerouteNode(bool bShowAsFullNode, co
 	FString Name = GetValidNodeName(InNodeName.IsEmpty() ? FString(TEXT("RerouteNode")) : InNodeName);
 	URigVMRerouteNode* Node = NewObject<URigVMRerouteNode>(Graph, *Name);
 	Node->Position = InPosition;
-	Node->bShowAsFullNode = bShowAsFullNode;
 
 	URigVMPin* ValuePin = NewObject<URigVMPin>(Node, *URigVMRerouteNode::ValueName);
 	ValuePin->CPPType = InCPPType;
@@ -12677,16 +12574,7 @@ URigVMRerouteNode* URigVMController::AddFreeRerouteNode(bool bShowAsFullNode, co
 	AddNodePin(Node, ValuePin);
 	Graph->Nodes.Add(Node);
 
-	if (ValuePin->IsStruct())
-	{
-		FString DefaultValue = InDefaultValue;
-		CreateDefaultValueForStructIfRequired(ValuePin->GetScriptStruct(), DefaultValue);
-		{
-			TGuardValue<bool> SuspendNotifications(bSuspendNotifications, true);
-			AddPinsForStruct(ValuePin->GetScriptStruct(), Node, ValuePin, ValuePin->Direction, DefaultValue, false);
-		}
-	}
-	else if (!InDefaultValue.IsEmpty() && InDefaultValue != TEXT("()"))
+	if (!InDefaultValue.IsEmpty() && InDefaultValue != TEXT("()"))
 	{
 		SetPinDefaultValue(ValuePin, InDefaultValue, true, false, false);
 	}
@@ -12706,138 +12594,103 @@ URigVMRerouteNode* URigVMController::AddFreeRerouteNode(bool bShowAsFullNode, co
 	return Node;
 }
 
-URigVMNode* URigVMController::AddBranchNode(const FVector2D& InPosition, const FString& InNodeName, bool bSetupUndoRedo, bool bPrintPythonCommand)
-{
-	return AddUnitNode(FRigVMFunction_ControlFlowBranch::StaticStruct(), FRigVMStruct::ExecuteName, InPosition, InNodeName, bSetupUndoRedo, bPrintPythonCommand);
-}
-
-URigVMNode* URigVMController::AddIfNode(const FString& InCPPType, const FName& InCPPTypeObjectPath, const FVector2D& InPosition, const FString& InNodeName, bool  bSetupUndoRedo, bool bPrintPythonCommand)
+URigVMTemplateNode* URigVMController::AddConstantNode(const FString& InCPPType, const FName& InCPPTypeObjectPath,
+	const FString& InDefaultValue, const FVector2D& InPosition, const FString& InNodeName, bool bSetupUndoRedo)
 {
 	if(!IsValidGraph())
 	{
 		return nullptr;
 	}
-
+	
 	if (!bIsTransacting && !IsGraphEditable())
 	{
 		return nullptr;
 	}
 
-	URigVMGraph* Graph = GetGraph();
-	check(Graph);
-
-	ensure(!InCPPType.IsEmpty());
+	if(bSetupUndoRedo)
+	{
+		OpenUndoBracket(TEXT("Add Constant Node"));
+	}
 
-	UObject* CPPTypeObject = nullptr;
-	if(!InCPPTypeObjectPath.IsNone())
+	URigVMTemplateNode* TemplateNode = AddTemplateNode(FRigVMDispatch_Constant().GetTemplateNotation(), InPosition, InNodeName, bSetupUndoRedo, bSetupUndoRedo);
+	if(!TemplateNode)
 	{
-		CPPTypeObject = RigVMTypeUtils::FindObjectFromCPPTypeObjectPath<UObject>(InCPPTypeObjectPath.ToString());
-		if (CPPTypeObject == nullptr)
+		if(bSetupUndoRedo)
 		{
-			ReportErrorf(TEXT("Cannot find cpp type object for path '%s'."), *InCPPTypeObjectPath.ToString());
-			return nullptr;
+			CancelUndoBracket();
 		}
+		return nullptr;
 	}
 
-	if(bSetupUndoRedo)
+	if(const URigVMPin* ValuePin = TemplateNode->FindPin(FRigVMDispatch_Constant::ValueName.ToString()))
 	{
-		OpenUndoBracket(TEXT("Add If Node"));
+		ResolveWildCardPin(ValuePin->GetPinPath(), InCPPType, InCPPTypeObjectPath, bSetupUndoRedo, bSetupUndoRedo);
+		if(!InDefaultValue.IsEmpty())
+		{
+			SetPinDefaultValue(ValuePin->GetPinPath(), InDefaultValue, true, bSetupUndoRedo, false, bSetupUndoRedo);
+		}
 	}
 
-	const FString CPPType = RigVMTypeUtils::PostProcessCPPType(InCPPType, CPPTypeObject);
-	const FString Name = GetValidNodeName(InNodeName.IsEmpty() ? FString(TEXT("IfNode")) : InNodeName);
-	const TRigVMTypeIndex& TypeIndex = FRigVMRegistry::Get().FindOrAddType({*CPPType, CPPTypeObject});
-	
-	const FRigVMDispatchFactory* Factory = FRigVMRegistry::Get().FindOrAddDispatchFactory(FRigVMDispatch_If::StaticStruct());
-	URigVMNode* Node = AddTemplateNode(Factory->GetTemplate()->GetNotation(), InPosition, Name, bSetupUndoRedo, bPrintPythonCommand);
-	if(Node)
-	{
-		ResolveWildCardPin(Node->GetPins().Last(), TypeIndex, bSetupUndoRedo, bPrintPythonCommand);
-	}
-	
 	if(bSetupUndoRedo)
 	{
 		CloseUndoBracket();
 	}
 
-	return Node;
+	return TemplateNode;
 }
 
-URigVMNode* URigVMController::AddIfNodeFromStruct(UScriptStruct* InScriptStruct, const FVector2D& InPosition, const FString& InNodeName, bool bSetupUndoRedo)
+URigVMTemplateNode* URigVMController::AddMakeStructNode(const FString& InCPPType, const FName& InCPPTypeObjectPath,
+                                                        const FString& InDefaultValue, const FVector2D& InPosition, const FString& InNodeName, bool bSetupUndoRedo)
 {
-	if (!InScriptStruct)
+	if(!IsValidGraph())
 	{
 		return nullptr;
 	}
-
-	return AddIfNode(RigVMTypeUtils::GetUniqueStructTypeName(InScriptStruct), FName(InScriptStruct->GetPathName()), InPosition, InNodeName, bSetupUndoRedo);
-}
-
-URigVMNode* URigVMController::AddSelectNode(const FString& InCPPType, const FName& InCPPTypeObjectPath, const FVector2D& InPosition, const FString& InNodeName, bool bSetupUndoRedo, bool bPrintPythonCommand)
-{
-	if (!IsValidGraph())
+	
+	if (!bIsTransacting && !IsGraphEditable())
 	{
 		return nullptr;
 	}
 
-	if (!bIsTransacting && !IsGraphEditable())
+	if(bSetupUndoRedo)
 	{
-		return nullptr;
+		OpenUndoBracket(TEXT("Add Make Struct Node"));
 	}
 
-	URigVMGraph* Graph = GetGraph();
-	check(Graph);
-
-	ensure(!InCPPType.IsEmpty());
-
-	UObject* CPPTypeObject = nullptr;
-	if (!InCPPTypeObjectPath.IsNone())
+	URigVMTemplateNode* TemplateNode = AddTemplateNode(FRigVMDispatch_MakeStruct().GetTemplateNotation(), InPosition, InNodeName, bSetupUndoRedo, bSetupUndoRedo);
+	if(!TemplateNode)
 	{
-		CPPTypeObject = RigVMTypeUtils::FindObjectFromCPPTypeObjectPath<UObject>(InCPPTypeObjectPath.ToString());
-		if (CPPTypeObject == nullptr)
+		if(bSetupUndoRedo)
 		{
-			ReportErrorf(TEXT("Cannot find cpp type object for path '%s'."), *InCPPTypeObjectPath.ToString());
-			return nullptr;
+			CancelUndoBracket();
 		}
+		return nullptr;
 	}
 
-	if(bSetupUndoRedo)
+	if(const URigVMPin* StructPin = TemplateNode->FindPin(FRigVMDispatch_MakeStruct::StructName.ToString()))
 	{
-		OpenUndoBracket(TEXT("Add Select Node"));
+		ResolveWildCardPin(StructPin->GetPinPath(), InCPPType, InCPPTypeObjectPath, bSetupUndoRedo, bSetupUndoRedo);
 	}
-
-	const FString CPPType = RigVMTypeUtils::PostProcessCPPType(InCPPType, CPPTypeObject);
-	const FString Name = GetValidNodeName(InNodeName.IsEmpty() ? FString(TEXT("SelectNode")) : InNodeName);
-	const TRigVMTypeIndex& TypeIndex = FRigVMRegistry::Get().FindOrAddType({*CPPType, CPPTypeObject});
-	
-	const FRigVMDispatchFactory* Factory = FRigVMRegistry::Get().FindOrAddDispatchFactory(FRigVMDispatch_SelectInt32::StaticStruct());
-	URigVMNode* Node = AddTemplateNode(Factory->GetTemplate()->GetNotation(), InPosition, Name, bSetupUndoRedo, bPrintPythonCommand);
-	if(Node)
+	if(!InDefaultValue.IsEmpty())
 	{
-		ResolveWildCardPin(Node->GetPins().Last(), TypeIndex, bSetupUndoRedo, bPrintPythonCommand);
+		if(const URigVMPin* ElementsPin = TemplateNode->FindPin(FRigVMDispatch_MakeStruct::ElementsName.ToString()))
+		{
+			SetPinDefaultValue(ElementsPin->GetPinPath(), InDefaultValue, true, bSetupUndoRedo, false, bSetupUndoRedo);
+		}
 	}
 
 	if(bSetupUndoRedo)
 	{
 		CloseUndoBracket();
 	}
-	return Node;
-}
-
-URigVMNode* URigVMController::AddSelectNodeFromStruct(UScriptStruct* InScriptStruct, const FVector2D& InPosition,
-	const FString& InNodeName, bool bSetupUndoRedo)
-{
-	if (!InScriptStruct)
-	{
-		return nullptr;
-	}
 
-	return AddSelectNode(RigVMTypeUtils::GetUniqueStructTypeName(InScriptStruct), FName(InScriptStruct->GetPathName()), InPosition, InNodeName, bSetupUndoRedo);
+	return TemplateNode;
 }
 
-URigVMTemplateNode* URigVMController::AddTemplateNode(const FName& InNotation, const FVector2D& InPosition, const FString& InNodeName, bool bSetupUndoRedo, bool bPrintPythonCommand)
+URigVMTemplateNode* URigVMController::AddBreakStructNode(const FString& InCPPType, const FName& InCPPTypeObjectPath,
+	const FString& InDefaultValue, const FVector2D& InPosition, const FString& InNodeName, bool bSetupUndoRedo)
 {
-	if (!IsValidGraph())
+	if(!IsValidGraph())
 	{
 		return nullptr;
 	}
@@ -12847,56 +12700,414 @@ URigVMTemplateNode* URigVMController::AddTemplateNode(const FName& InNotation, c
 		return nullptr;
 	}
 
-	URigVMGraph* Graph = GetGraph();
-	check(Graph);
-
-	ensure(!InNotation.IsNone());
+	if(bSetupUndoRedo)
+	{
+		OpenUndoBracket(TEXT("Add Break Struct Node"));
+	}
 
-	const FRigVMTemplate* Template = FRigVMRegistry::Get().FindTemplate(InNotation);
-	if (Template == nullptr)
+	URigVMTemplateNode* TemplateNode = AddTemplateNode(FRigVMDispatch_BreakStruct().GetTemplateNotation(), InPosition, InNodeName, bSetupUndoRedo, bSetupUndoRedo);
+	if(!TemplateNode)
 	{
-		ReportErrorf(TEXT("Template '%s' cannot be found."), *InNotation.ToString());
+		if(bSetupUndoRedo)
+		{
+			CancelUndoBracket();
+		}
 		return nullptr;
 	}
 
-	if(IRigVMClientHost* ClientHost = GetImplementingOuter<IRigVMClientHost>())
+	if(const URigVMPin* StructPin = TemplateNode->FindPin(FRigVMDispatch_BreakStruct::StructName.ToString()))
 	{
-		if(const FRigVMClient* Client = ClientHost->GetRigVMClient())
+		ResolveWildCardPin(StructPin->GetPinPath(), InCPPType, InCPPTypeObjectPath, bSetupUndoRedo, bSetupUndoRedo);
+		if(!InDefaultValue.IsEmpty())
 		{
-			if(!Template->SupportsExecuteContextStruct(Client->GetExecuteContextStruct()))
-			{
-				ReportErrorf(TEXT("Cannot add node for template '%s' - incompatible execute context: '%s' vs '%s'."),
-						*Template->GetNotation().ToString(),
-						*Template->GetExecuteContextStruct()->GetStructCPPName(),
-						*Client->GetExecuteContextStruct()->GetStructCPPName());
-				return nullptr;
-			}
+			SetPinDefaultValue(StructPin->GetPinPath(), InDefaultValue, true, bSetupUndoRedo, false, bSetupUndoRedo);
 		}
 	}
 
-	FString Name = GetValidNodeName(InNodeName.IsEmpty() ? Template->GetName().ToString() : InNodeName);
-	URigVMTemplateNode* Node = nullptr;
-
-	// determine what kind of node we need to create
-	if(Template->UsesDispatch())
+	if(bSetupUndoRedo)
 	{
-		Node = NewObject<URigVMDispatchNode>(Graph, *Name);
+		CloseUndoBracket();
+	}
+
+	return TemplateNode;
+}
+
+URigVMTemplateNode* URigVMController::AddConstantNodeOnPin(const FString& InPinPath, const FVector2D& InPosition,
+	const FString& InNodeName, bool bSetupUndoRedo)
+{
+	if(!IsValidGraph())
+	{
+		return nullptr;
+	}
+	
+	if (!bIsTransacting && !IsGraphEditable())
+	{
+		return nullptr;
+	}
+
+	const URigVMGraph* Graph = GetGraph();
+	check(Graph);
+
+	const URigVMPin* Pin = Graph->FindPin(InPinPath);
+	if(Pin == nullptr)
+	{
+		ReportErrorf(TEXT("Pin '%s' cannot be found."), *InPinPath);
+		return nullptr;
+	}
+
+	if(Pin->GetDirection() != ERigVMPinDirection::Input && Pin->GetDirection() != ERigVMPinDirection::IO)
+	{
+		ReportError(TEXT("Constant nodes can only be added to input / io pins.."));
+		return nullptr;
+	}
+
+	const FString& CPPType = Pin->GetCPPType();
+	FName CPPTypeObjectPath = NAME_None;
+	if(const UObject* CPPTypeObject = Pin->GetCPPTypeObject())
+	{
+		CPPTypeObjectPath = *CPPTypeObject->GetPathName();
+	}
+	const FString DefaultValue = Pin->GetDefaultValue();
+
+	if(bSetupUndoRedo)
+	{
+		OpenUndoBracket(TEXT("Add Constant Node On Pin"));
+	}
+
+	URigVMTemplateNode* TemplateNode = AddConstantNode(CPPType, CPPTypeObjectPath, DefaultValue, InPosition, InNodeName, bSetupUndoRedo);
+	if(TemplateNode == nullptr)
+	{
+		if(bSetupUndoRedo)
+		{
+			CancelUndoBracket();
+		}
+		return nullptr;
+	}
+
+	if(const URigVMPin* ValuePin = TemplateNode->FindPin(FRigVMDispatch_Constant::ValueName.ToString()))
+	{
+		AddLink(ValuePin->GetPinPath(), Pin->GetPinPath(), bSetupUndoRedo, bSetupUndoRedo);
+	}
+
+	if(bSetupUndoRedo)
+	{
+		CloseUndoBracket();
+	}
+
+	return TemplateNode;
+}
+
+URigVMTemplateNode* URigVMController::AddMakeStructNodeOnPin(const FString& InPinPath, const FVector2D& InPosition,
+	const FString& InNodeName, bool bSetupUndoRedo)
+{
+	if(!IsValidGraph())
+	{
+		return nullptr;
+	}
+	
+	if (!bIsTransacting && !IsGraphEditable())
+	{
+		return nullptr;
+	}
+
+	const URigVMGraph* Graph = GetGraph();
+	check(Graph);
+
+	const URigVMPin* Pin = Graph->FindPin(InPinPath);
+	if(Pin == nullptr)
+	{
+		ReportErrorf(TEXT("Pin '%s' cannot be found."), *InPinPath);
+		return nullptr;
+	}
+
+	if(Pin->GetDirection() != ERigVMPinDirection::Input && Pin->GetDirection() != ERigVMPinDirection::IO)
+	{
+		ReportError(TEXT("Make Struct nodes can only be added to input / io pins.."));
+		return nullptr;
+	}
+
+	const FString& CPPType = Pin->GetCPPType();
+	FName CPPTypeObjectPath = NAME_None;
+	if(const UObject* CPPTypeObject = Pin->GetCPPTypeObject())
+	{
+		CPPTypeObjectPath = *CPPTypeObject->GetPathName();
+	}
+	const FString DefaultValue = Pin->GetDefaultValue();
+
+	if(bSetupUndoRedo)
+	{
+		OpenUndoBracket(TEXT("Add Make Struct Node On Pin"));
+	}
+
+	URigVMTemplateNode* TemplateNode = AddMakeStructNode(CPPType, CPPTypeObjectPath, DefaultValue, InPosition, InNodeName, bSetupUndoRedo);
+	if(TemplateNode == nullptr)
+	{
+		if(bSetupUndoRedo)
+		{
+			CancelUndoBracket();
+		}
+		return nullptr;
+	}
+
+	if(const URigVMPin* StructPin = TemplateNode->FindPin(FRigVMDispatch_MakeStruct::StructName.ToString()))
+	{
+		AddLink(StructPin->GetPinPath(), Pin->GetPinPath(), bSetupUndoRedo, bSetupUndoRedo);
+	}
+
+	if(bSetupUndoRedo)
+	{
+		CloseUndoBracket();
+	}
+
+	return TemplateNode;
+}
+
+URigVMTemplateNode* URigVMController::AddBreakStructNodeOnPin(const FString& InPinPath, const FVector2D& InPosition,
+	const FString& InNodeName, bool bSetupUndoRedo)
+{
+	if(!IsValidGraph())
+	{
+		return nullptr;
+	}
+	
+	if (!bIsTransacting && !IsGraphEditable())
+	{
+		return nullptr;
+	}
+
+	const URigVMGraph* Graph = GetGraph();
+	check(Graph);
+
+	const URigVMPin* Pin = Graph->FindPin(InPinPath);
+	if(Pin == nullptr)
+	{
+		ReportErrorf(TEXT("Pin '%s' cannot be found."), *InPinPath);
+		return nullptr;
+	}
+
+	if(Pin->GetDirection() != ERigVMPinDirection::Output && Pin->GetDirection() != ERigVMPinDirection::IO)
+	{
+		ReportError(TEXT("Break Struct nodes can only be added to output / io pins.."));
+		return nullptr;
+	}
+
+	const FString& CPPType = Pin->GetCPPType();
+	FName CPPTypeObjectPath = NAME_None;
+	if(const UObject* CPPTypeObject = Pin->GetCPPTypeObject())
+	{
+		CPPTypeObjectPath = *CPPTypeObject->GetPathName();
+	}
+
+	if(bSetupUndoRedo)
+	{
+		OpenUndoBracket(TEXT("Add Break Struct Node On Pin"));
+	}
+
+	URigVMTemplateNode* TemplateNode = AddBreakStructNode(CPPType, CPPTypeObjectPath, FString(), InPosition, InNodeName, bSetupUndoRedo);
+	if(TemplateNode == nullptr)
+	{
+		if(bSetupUndoRedo)
+		{
+			CancelUndoBracket();
+		}
+		return nullptr;
+	}
+
+	if(const URigVMPin* StructPin = TemplateNode->FindPin(FRigVMDispatch_MakeStruct::StructName.ToString()))
+	{
+		AddLink(Pin->GetPinPath(), StructPin->GetPinPath(), bSetupUndoRedo, bSetupUndoRedo);
+	}
+
+	if(bSetupUndoRedo)
+	{
+		CloseUndoBracket();
+	}
+
+	return TemplateNode;
+}
+
+URigVMNode* URigVMController::AddBranchNode(const FVector2D& InPosition, const FString& InNodeName, bool bSetupUndoRedo, bool bPrintPythonCommand)
+{
+	return AddUnitNode(FRigVMFunction_ControlFlowBranch::StaticStruct(), FRigVMStruct::ExecuteName, InPosition, InNodeName, bSetupUndoRedo, bPrintPythonCommand);
+}
+
+URigVMNode* URigVMController::AddIfNode(const FString& InCPPType, const FName& InCPPTypeObjectPath, const FVector2D& InPosition, const FString& InNodeName, bool  bSetupUndoRedo, bool bPrintPythonCommand)
+{
+	if(!IsValidGraph())
+	{
+		return nullptr;
+	}
+
+	if (!bIsTransacting && !IsGraphEditable())
+	{
+		return nullptr;
+	}
+
+	URigVMGraph* Graph = GetGraph();
+	check(Graph);
+
+	ensure(!InCPPType.IsEmpty());
+
+	UObject* CPPTypeObject = nullptr;
+	if(!InCPPTypeObjectPath.IsNone())
+	{
+		CPPTypeObject = RigVMTypeUtils::FindObjectFromCPPTypeObjectPath<UObject>(InCPPTypeObjectPath.ToString());
+		if (CPPTypeObject == nullptr)
+		{
+			ReportErrorf(TEXT("Cannot find cpp type object for path '%s'."), *InCPPTypeObjectPath.ToString());
+			return nullptr;
+		}
+	}
+
+	if(bSetupUndoRedo)
+	{
+		OpenUndoBracket(TEXT("Add If Node"));
+	}
+
+	const FString CPPType = RigVMTypeUtils::PostProcessCPPType(InCPPType, CPPTypeObject);
+	const FString Name = GetValidNodeName(InNodeName.IsEmpty() ? FString(TEXT("IfNode")) : InNodeName);
+	const TRigVMTypeIndex& TypeIndex = FRigVMRegistry::Get().FindOrAddType({*CPPType, CPPTypeObject});
+	
+	const FRigVMDispatchFactory* Factory = FRigVMRegistry::Get().FindOrAddDispatchFactory(FRigVMDispatch_If::StaticStruct());
+	URigVMNode* Node = AddTemplateNode(Factory->GetTemplate()->GetNotation(), InPosition, Name, bSetupUndoRedo, bPrintPythonCommand);
+	if(Node)
+	{
+		ResolveWildCardPin(Node->GetPins().Last(), TypeIndex, bSetupUndoRedo, bPrintPythonCommand);
+	}
+	
+	if(bSetupUndoRedo)
+	{
+		CloseUndoBracket();
+	}
+
+	return Node;
+}
+
+URigVMNode* URigVMController::AddIfNodeFromStruct(UScriptStruct* InScriptStruct, const FVector2D& InPosition, const FString& InNodeName, bool bSetupUndoRedo)
+{
+	if (!InScriptStruct)
+	{
+		return nullptr;
+	}
+
+	return AddIfNode(RigVMTypeUtils::GetUniqueStructTypeName(InScriptStruct), FName(InScriptStruct->GetPathName()), InPosition, InNodeName, bSetupUndoRedo);
+}
+
+URigVMNode* URigVMController::AddSelectNode(const FString& InCPPType, const FName& InCPPTypeObjectPath, const FVector2D& InPosition, const FString& InNodeName, bool bSetupUndoRedo, bool bPrintPythonCommand)
+{
+	if (!IsValidGraph())
+	{
+		return nullptr;
+	}
+
+	if (!bIsTransacting && !IsGraphEditable())
+	{
+		return nullptr;
+	}
+
+	URigVMGraph* Graph = GetGraph();
+	check(Graph);
+
+	ensure(!InCPPType.IsEmpty());
+
+	UObject* CPPTypeObject = nullptr;
+	if (!InCPPTypeObjectPath.IsNone())
+	{
+		CPPTypeObject = RigVMTypeUtils::FindObjectFromCPPTypeObjectPath<UObject>(InCPPTypeObjectPath.ToString());
+		if (CPPTypeObject == nullptr)
+		{
+			ReportErrorf(TEXT("Cannot find cpp type object for path '%s'."), *InCPPTypeObjectPath.ToString());
+			return nullptr;
+		}
+	}
+
+	if(bSetupUndoRedo)
+	{
+		OpenUndoBracket(TEXT("Add Select Node"));
+	}
+
+	const FString CPPType = RigVMTypeUtils::PostProcessCPPType(InCPPType, CPPTypeObject);
+	const FString Name = GetValidNodeName(InNodeName.IsEmpty() ? FString(TEXT("SelectNode")) : InNodeName);
+	const TRigVMTypeIndex& TypeIndex = FRigVMRegistry::Get().FindOrAddType({*CPPType, CPPTypeObject});
+	
+	const FRigVMDispatchFactory* Factory = FRigVMRegistry::Get().FindOrAddDispatchFactory(FRigVMDispatch_SelectInt32::StaticStruct());
+	URigVMNode* Node = AddTemplateNode(Factory->GetTemplate()->GetNotation(), InPosition, Name, bSetupUndoRedo, bPrintPythonCommand);
+	if(Node)
+	{
+		ResolveWildCardPin(Node->GetPins().Last(), TypeIndex, bSetupUndoRedo, bPrintPythonCommand);
+	}
+
+	if(bSetupUndoRedo)
+	{
+		CloseUndoBracket();
+	}
+	return Node;
+}
+
+URigVMNode* URigVMController::AddSelectNodeFromStruct(UScriptStruct* InScriptStruct, const FVector2D& InPosition,
+	const FString& InNodeName, bool bSetupUndoRedo)
+{
+	if (!InScriptStruct)
+	{
+		return nullptr;
+	}
+
+	return AddSelectNode(RigVMTypeUtils::GetUniqueStructTypeName(InScriptStruct), FName(InScriptStruct->GetPathName()), InPosition, InNodeName, bSetupUndoRedo);
+}
+
+URigVMTemplateNode* URigVMController::AddTemplateNode(const FName& InNotation, const FVector2D& InPosition, const FString& InNodeName, bool bSetupUndoRedo, bool bPrintPythonCommand)
+{
+	if (!IsValidGraph())
+	{
+		return nullptr;
 	}
-	else if(const FRigVMFunction* FirstFunction = Template->GetPermutation(0))
+
+	if (!bIsTransacting && !IsGraphEditable())
 	{
-		const UScriptStruct* PotentialUnitStruct = FirstFunction->Struct;
-		if(PotentialUnitStruct && PotentialUnitStruct->IsChildOf(FRigVMStruct::StaticStruct()))
+		return nullptr;
+	}
+
+	URigVMGraph* Graph = GetGraph();
+	check(Graph);
+
+	ensure(!InNotation.IsNone());
+
+	const FRigVMTemplate* Template = FRigVMRegistry::Get().FindTemplate(InNotation);
+	if (Template == nullptr)
+	{
+		ReportErrorf(TEXT("Template '%s' cannot be found."), *InNotation.ToString());
+		return nullptr;
+	}
+
+	if(IRigVMClientHost* ClientHost = GetImplementingOuter<IRigVMClientHost>())
+	{
+		if(const FRigVMClient* Client = ClientHost->GetRigVMClient())
 		{
-			Node = NewObject<URigVMUnitNode>(Graph, *Name); 
+			if(!Template->SupportsExecuteContextStruct(Client->GetExecuteContextStruct()))
+			{
+				ReportErrorf(TEXT("Cannot add node for template '%s' - incompatible execute context: '%s' vs '%s'."),
+						*Template->GetNotation().ToString(),
+						*Template->GetExecuteContextStruct()->GetStructCPPName(),
+						*Client->GetExecuteContextStruct()->GetStructCPPName());
+				return nullptr;
+			}
 		}
 	}
 
-	if(Node == nullptr)
+	FString Name = GetValidNodeName(InNodeName.IsEmpty() ? Template->GetName().ToString() : InNodeName);
+	URigVMTemplateNode* Node = nullptr;
+
+	// determine what kind of node we need to create
+	if(Template->UsesDispatch())
+	{
+		Node = NewObject<URigVMDispatchNode>(Graph, *Name);
+	}
+	else if(const FRigVMFunction* FirstFunction = Template->GetPermutation(0))
 	{
-		const FString TemplateName = Template->GetName().ToString();
-		if(TemplateName == URigVMRerouteNode::RerouteName)
+		const UScriptStruct* PotentialUnitStruct = FirstFunction->Struct;
+		if(PotentialUnitStruct && PotentialUnitStruct->IsChildOf(FRigVMStruct::StaticStruct()))
 		{
-			Node = NewObject<URigVMRerouteNode>(Graph, *Name);
+			Node = NewObject<URigVMUnitNode>(Graph, *Name); 
 		}
 	}
 
@@ -13777,6 +13988,11 @@ void URigVMController::AddPinsForStruct(UStruct* InStruct, URigVMNode* InNode, U
 		return;
 	}
 
+	if(InParentPin && InParentPin->ShouldHideSubPins())
+	{
+		return;
+	}
+
 	// todo: reuse the default values when creating the pins
 	
 	TArray<FString> MemberNameValuePairs = URigVMPin::SplitDefaultValue(InDefaultValue);
@@ -13879,6 +14095,11 @@ void URigVMController::AddPinsForArray(FArrayProperty* InArrayProperty, URigVMNo
 		return;
 	}
 
+	if(InParentPin && InParentPin->ShouldHideSubPins())
+	{
+		return;
+	}
+
 	for (int32 ElementIndex = 0; ElementIndex < InDefaultValues.Num(); ElementIndex++)
 	{
 		FString ElementName = FString::FormatAsNumber(InParentPin->SubPins.Num());
@@ -14597,7 +14818,7 @@ void URigVMController::RepopulatePinsOnNode(URigVMNode* InNode, bool bFollowCore
 		}
 
 		const TSharedPtr<FStructOnScope> DefaultValueContent = UnitNode->ConstructStructInstance(false);
-		NewPinInfos.AddPins(ScriptStruct, this, ERigVMPinDirection::Invalid, INDEX_NONE, DefaultValueContent->GetStructMemory());
+		NewPinInfos.AddPins(ScriptStruct, this, ERigVMPinDirection::Invalid, INDEX_NONE, DefaultValueContent->GetStructMemory(), true);
 	}
 	else if (DispatchNode)
 	{
@@ -14627,7 +14848,7 @@ void URigVMController::RepopulatePinsOnNode(URigVMNode* InNode, bool bFollowCore
 					}
 				}
 
-				(void)NewPinInfos.AddPin(this, INDEX_NONE, Arg->Name, Arg->Direction, TypeIndex, DefaultValue, nullptr, &PreviousPinInfos);
+				(void)NewPinInfos.AddPin(this, INDEX_NONE, Arg->Name, Arg->Direction, TypeIndex, DefaultValue, nullptr, &PreviousPinInfos, false);
 			}
 		};
 
@@ -14692,8 +14913,15 @@ void URigVMController::RepopulatePinsOnNode(URigVMNode* InNode, bool bFollowCore
 				ArgumentScriptStruct->ImportText(*DefaultValue, DefaultValueMemoryScope.GetStructMemory(), nullptr, PPF_None, &ErrorPipe, FString());
 				DefaultValueMemory = DefaultValueMemoryScope.GetStructMemory();
 			}
-			
-			(void)NewPinInfos.AddPin(this, INDEX_NONE, Arg->Name, Arg->GetDirection(), TypeIndex, DefaultValue, DefaultValueMemory, &PreviousPinInfos);
+
+			bool bAddSubPinsForArgument = true;
+#if WITH_EDITOR
+			if(!Template->GetArgumentMetaData(Arg->Name, FRigVMStruct::HideSubPinsMetaName).IsEmpty())
+			{
+				bAddSubPinsForArgument = false;
+			}
+#endif
+			(void)NewPinInfos.AddPin(this, INDEX_NONE, Arg->Name, Arg->GetDirection(), TypeIndex, DefaultValue, DefaultValueMemory, &PreviousPinInfos, bAddSubPinsForArgument);
 		}
 
 		AddExecutePins(ERigVMPinDirection::Output);
@@ -16042,7 +16270,7 @@ bool URigVMController::FullyResolveTemplateNode(URigVMTemplateNode* InNode, int3
 		{
 			FRigVMTemplateTypeMap TypeMap = InNode->GetTemplate()->GetTypesForPermutation(InNode->ResolvedPermutation);
 			const FRigVMFunctionPtr DispatchFunction = Factory->GetDispatchFunction(TypeMap);
-			const FRigVMFunction* ResolvedFunction = Template->GetPermutation(InNode->ResolvedPermutation);
+			const FRigVMFunction* ResolvedFunction = const_cast<FRigVMTemplate*>(Template)->GetOrCreatePermutation(InNode->ResolvedPermutation);
 			check(DispatchFunction);
 			check(ResolvedFunction);
 			check(ResolvedFunction->FunctionPtr == DispatchFunction);
@@ -17486,6 +17714,12 @@ void URigVMController::AddNodePin(URigVMNode* InNode, URigVMPin* InPin)
 	ValidatePin(InPin);
 	checkf(!InNode->Pins.Contains(InPin), TEXT("Node %s already contains pin %s"), *InNode->GetPathName(), *InPin->GetName());
 	InNode->Pins.Add(InPin);
+
+	// the first time we add a fixed size array we want to expand it
+	if(InPin->IsFixedSizeArray())
+	{
+		InPin->bIsExpanded = true;
+	}
 }
 
 void URigVMController::AddSubPin(URigVMPin* InParentPin, URigVMPin* InPin)
@@ -18293,6 +18527,410 @@ FRigVMTemplate::FTypeMap URigVMController::GetCommonlyUsedTypesForTemplate(
 	return Template->GetArgumentTypesFromString(TypesString);
 }
 
+URigVMNode* URigVMController::ConvertRerouteNodeToDispatch(URigVMRerouteNode* InRerouteNode,
+	const FName& InTemplateNotation, bool bSetupUndoRedo, bool bPrintPythonCommand)
+{
+	static const TMap<FString, FString> EmptyRedirects = {};
+	
+	static const TMap<FString, FString> StructRedirects = {
+		{
+			URigVMRerouteNode::ValueName,
+			FRigVMDispatch_MakeStruct::StructName.ToString(),
+		}};
+
+	static const TMap<FString, FString> ElementsRedirects = {
+		{
+			URigVMRerouteNode::ValueName,
+			FRigVMDispatch_MakeStruct::ElementsName.ToString(),
+		}};
+
+	static const TMap<FString, FString> ValuesRedirects = {
+		{
+			URigVMRerouteNode::ValueName,
+			FRigVMDispatch_ArrayMake::ValuesName.ToString(),
+		}};
+
+	static const TMap<FString, FString> ArrayRedirects = {
+		{
+			URigVMRerouteNode::ValueName,
+			FRigVMDispatch_ArrayMake::ArrayName.ToString(),
+		}};
+
+	const TMap<FString, FString>* InputRedirects = &EmptyRedirects;
+	const TMap<FString, FString>* OutputRedirects = &EmptyRedirects;
+	FString PinToResolveName;
+
+	// constant is empty
+	if(InTemplateNotation == FRigVMDispatch_Constant().GetTemplateNotation())
+	{
+		PinToResolveName = FRigVMDispatch_Constant::ValueName.ToString();
+	}
+	else if(InTemplateNotation == FRigVMDispatch_MakeStruct().GetTemplateNotation())
+	{
+		InputRedirects = &ElementsRedirects;
+		OutputRedirects = &StructRedirects;
+		PinToResolveName = FRigVMDispatch_MakeStruct::StructName.ToString();
+	}
+	else if(InTemplateNotation == FRigVMDispatch_BreakStruct().GetTemplateNotation())
+	{
+		InputRedirects = &StructRedirects;
+		OutputRedirects = &ElementsRedirects;
+		PinToResolveName = FRigVMDispatch_MakeStruct::StructName.ToString();
+	}
+	else if(InTemplateNotation == FRigVMDispatch_ArrayMake().GetTemplateNotation())
+	{
+		InputRedirects = &ValuesRedirects;
+		OutputRedirects = &ArrayRedirects;
+		PinToResolveName = FRigVMDispatch_ArrayMake::ArrayName.ToString();
+	}
+	else
+	{
+		static constexpr TCHAR Format[] = TEXT("Template '%s' not supported when converting a reroute node.");
+		ReportErrorf(Format, *InTemplateNotation.ToString());
+		return nullptr;
+	}
+
+	if(bSetupUndoRedo)
+	{
+		OpenUndoBracket(TEXT("Converting reroute node"));
+	}
+
+	TMap<FString, FPinState> PinStates = GetPinStates(InRerouteNode);
+	const TArray<FLinkedPath> LinkedPaths = GetLinkedPaths(InRerouteNode);
+	FastBreakLinkedPaths(LinkedPaths, bSetupUndoRedo);
+
+	const URigVMPin* ValuePin = InRerouteNode->FindPin(URigVMRerouteNode::ValueName);
+	if(ValuePin == nullptr)
+	{
+		if(bSetupUndoRedo)
+		{
+			CancelUndoBracket();
+		}
+		return nullptr;
+	}
+
+	const FString NodeName = InRerouteNode->GetName();
+	const FVector2D NodePosition = InRerouteNode->GetPosition();
+	const TRigVMTypeIndex TypeIndex = ValuePin->GetTypeIndex();
+
+	if(!RemoveNode(InRerouteNode, bSetupUndoRedo, bPrintPythonCommand))
+	{
+		if(bSetupUndoRedo)
+		{
+			CancelUndoBracket();
+		}
+		return nullptr;
+	}
+
+	URigVMNode* NewNode = AddTemplateNode(InTemplateNotation, NodePosition, NodeName, bSetupUndoRedo, bPrintPythonCommand);
+	if(NewNode)
+	{
+		URigVMPin* PinToResolve = NewNode->FindPin(PinToResolveName);
+		check(PinToResolve);
+		verify(ResolveWildCardPin(PinToResolve, TypeIndex, bSetupUndoRedo, bPrintPythonCommand));
+		
+		TMap<FString, FPinState> RemappedPinStates;
+		for(const TPair<FString, FPinState>& Pair : PinStates)
+		{
+			FString PinPath = Pair.Key;
+			TArray<FString> Parts;
+			if(!URigVMPin::SplitPinPath(PinPath, Parts))
+			{
+				Parts = {PinPath};
+			}
+			if(const FString* RedirectedPart = InputRedirects->Find(Parts[0]))
+			{
+				Parts[0] = *RedirectedPart;
+				PinPath = URigVMPin::JoinPinPath(Parts);
+			}
+			RemappedPinStates.Add(PinPath, Pair.Value);
+		}
+		ApplyPinStates(NewNode, RemappedPinStates, {}, bSetupUndoRedo);
+
+		const FString NodeNamePrefix = URigVMPin::JoinPinPath({NodeName, FString()});
+		FRestoreLinkedPathSettings RestoreSettings;
+		RestoreSettings.RemapDelegates.Add(NodeName,
+			FRigVMController_PinPathRemapDelegate::CreateLambda([NodeNamePrefix, InputRedirects, OutputRedirects](const FString& InPinPath, bool bIsInput) -> FString
+			{
+				TArray<FString> Parts;
+				if(URigVMPin::SplitPinPath(InPinPath, Parts))
+				{
+					const TMap<FString, FString>* Redirects = bIsInput ? InputRedirects : OutputRedirects;
+					if(const FString* RedirectedPart = Redirects->Find(Parts[1]))
+					{
+						Parts[1] = *RedirectedPart;
+						return URigVMPin::JoinPinPath(Parts);
+					}
+				}
+				return InPinPath;
+			})
+		);
+		RestoreLinkedPaths(LinkedPaths, RestoreSettings, bSetupUndoRedo);
+	}
+
+	if(bSetupUndoRedo)
+	{
+		CloseUndoBracket();
+	}
+	
+	return NewNode;
+}
+
+FRigVMClientPatchResult URigVMController::PatchRerouteNodesOnLoad()
+{
+	FRigVMClientPatchResult Result;
+
+	if (const URigVMGraph* Graph = GetGraph())
+	{
+		TArray<URigVMRerouteNode*> ReroutesToTurnIntoConstant;
+		TArray<URigVMRerouteNode*> ReroutesToTurnIntoMakeStruct;
+		TArray<URigVMRerouteNode*> ReroutesToTurnIntoBreakStruct;
+		TArray<URigVMRerouteNode*> ReroutesToTurnIntoMakeAndBreakStruct;
+		TArray<URigVMRerouteNode*> ReroutesToTurnIntoMakeArray;
+		TArray<URigVMRerouteNode*> ReroutesToRemove;
+
+		for(URigVMNode* Node : Graph->GetNodes())
+		{
+			if(URigVMRerouteNode* RerouteNode = Cast<URigVMRerouteNode>(Node))
+			{
+				const URigVMPin* ValuePin = RerouteNode->FindPin(URigVMRerouteNode::ValueName);
+				check(ValuePin);
+
+				// skip reroutes which are execute contexts
+				if(ValuePin->IsExecuteContext())
+				{
+					continue;
+				}
+
+				const int32 TopLevelSourceLinks = ValuePin->GetSourceLinks(false).Num();
+				const int32 TopLevelTargetLinks = ValuePin->GetTargetLinks(false).Num();
+				const int32 SubPinSourceLinks = ValuePin->GetSourceLinks(true).Num() - TopLevelSourceLinks;
+				const int32 SubPinTargetLinks = ValuePin->GetTargetLinks(true).Num() - TopLevelTargetLinks;
+
+				const bool bOnlyTopLevelSourceLinks = (TopLevelSourceLinks > 0) && (SubPinSourceLinks == 0);
+				const bool bOnlyTopLevelTargetLinks = (TopLevelTargetLinks > 0) && (SubPinTargetLinks == 0);
+				const bool bOnlySubPinSourceLinks = (TopLevelSourceLinks == 0) && (SubPinSourceLinks > 0);
+				const bool bOnlySubPinTargetLinks = (TopLevelTargetLinks == 0) && (SubPinTargetLinks > 0);
+				const bool bHasSourceLinks = bOnlyTopLevelSourceLinks || bOnlySubPinSourceLinks;
+				const bool bHasTargetLinks = bOnlyTopLevelTargetLinks || bOnlySubPinTargetLinks;
+
+				if(bHasSourceLinks && bHasTargetLinks)
+				{
+					if(bOnlyTopLevelSourceLinks && bOnlyTopLevelTargetLinks)
+					{
+						// this is a normal reroute and we can keep it
+					}
+					else if(bOnlyTopLevelSourceLinks) // && bOnlySubPinTargetLinks
+					{
+						check(ValuePin->IsStruct());
+						ReroutesToTurnIntoBreakStruct.Add(RerouteNode);
+					}
+					else if(bOnlyTopLevelTargetLinks) // && bOnlySubPinSourceLinks
+					{
+						if(ValuePin->IsArray())
+						{
+							ReroutesToTurnIntoMakeArray.Add(RerouteNode);
+						}
+						else
+						{
+							check(ValuePin->IsStruct());
+							ReroutesToTurnIntoMakeStruct.Add(RerouteNode);
+						}
+					}
+					else // bOnlySubPinSourceLinks && bOnlySubPinTargetLinks
+					{
+						check(ValuePin->IsStruct());
+						ReroutesToTurnIntoMakeAndBreakStruct.Add(RerouteNode);
+					}
+				}
+				else if(bHasSourceLinks) // && !bHasTargetLinks
+				{
+					if(bOnlyTopLevelSourceLinks)
+					{
+						// don't do anything - keep this node
+					}
+					else // bOnlySubPinSourceLinks
+					{
+						if(ValuePin->IsArray())
+						{
+							ReroutesToTurnIntoMakeArray.Add(RerouteNode);
+						}
+						else
+						{
+							check(ValuePin->IsStruct());
+							ReroutesToTurnIntoMakeStruct.Add(RerouteNode);
+						}
+					}
+				}
+				else if(bHasTargetLinks) // && !bHasSourceLinks
+				{
+					if(bOnlyTopLevelTargetLinks)
+					{
+						if(ValuePin->IsStruct())
+						{
+							ReroutesToTurnIntoMakeStruct.Add(RerouteNode);
+						}
+						else if(ValuePin->IsArray())
+						{
+							ReroutesToTurnIntoMakeArray.Add(RerouteNode);
+						}
+						else
+						{
+							ReroutesToTurnIntoConstant.Add(RerouteNode);
+						}
+					}
+					else // bOnlySubPinTargetLinks
+					{
+						check(ValuePin->IsStruct());
+						ReroutesToTurnIntoBreakStruct.Add(RerouteNode);
+					}
+				}
+				else
+				{
+					ReroutesToRemove.Add(RerouteNode);
+				}
+			}
+		}
+		
+		// remove obsolete reroutes - reroutes with no connections at all
+		if(ReroutesToRemove.Num() > 0)
+		{
+			for(URigVMRerouteNode* RerouteNode : ReroutesToRemove)
+			{
+				const FString PathName = RerouteNode->GetPathName();
+				if(RemoveNode(RerouteNode, false, false))
+				{
+					Result.RemovedNodes.Add(PathName);
+					Result.bChangedContent = true;
+				}
+			}
+			ReroutesToRemove.Reset();
+		}
+		
+		// convert some reroutes to make constant nodes
+		if(ReroutesToTurnIntoConstant.Num() > 0)
+		{
+			const FName TemplateNotation = FRigVMDispatch_Constant().GetTemplateNotation();
+			const TMap<FString, FString> PinRedirects;
+			for(URigVMRerouteNode* RerouteNode : ReroutesToTurnIntoConstant)
+			{
+				ConvertRerouteNodeToDispatch(RerouteNode, TemplateNotation, false, false);
+			}
+			ReroutesToTurnIntoConstant.Reset();
+		}
+
+		// convert some reroutes to make struct nodes
+		if(ReroutesToTurnIntoMakeStruct.Num() > 0)
+		{
+			const FName TemplateNotation = FRigVMDispatch_MakeStruct().GetTemplateNotation();
+			const TMap<FString, FString> PinRedirects;
+			for(URigVMRerouteNode* RerouteNode : ReroutesToTurnIntoMakeStruct)
+			{
+				ConvertRerouteNodeToDispatch(RerouteNode, TemplateNotation, false, false);
+			}
+			ReroutesToTurnIntoMakeStruct.Reset();
+		}
+
+		// convert some reroutes to break struct nodes
+		if(ReroutesToTurnIntoBreakStruct.Num() > 0)
+		{
+			const FName TemplateNotation = FRigVMDispatch_BreakStruct().GetTemplateNotation();
+			const TMap<FString, FString> PinRedirects;
+			for(URigVMRerouteNode* RerouteNode : ReroutesToTurnIntoBreakStruct)
+			{
+				ConvertRerouteNodeToDispatch(RerouteNode, TemplateNotation, false, false);
+			}
+			ReroutesToTurnIntoBreakStruct.Reset();
+		}
+
+		// convert some reroutes to a pair of a a make struct and a break struct node
+		if(ReroutesToTurnIntoMakeAndBreakStruct.Num() > 0)
+		{
+			const FName MakeTemplateNotation = FRigVMDispatch_MakeStruct().GetTemplateNotation();
+			const FName BreakTemplateNotation = FRigVMDispatch_BreakStruct().GetTemplateNotation();
+			for(URigVMRerouteNode* RerouteNode : ReroutesToTurnIntoMakeAndBreakStruct)
+			{
+				const URigVMPin* ValuePin = RerouteNode->FindPin(URigVMRerouteNode::ValueName);
+				check(ValuePin);
+				const TRigVMTypeIndex& TypeIndex = ValuePin->GetTypeIndex();
+				const FString NodeName = RerouteNode->GetName();
+				const TArray<FLinkedPath> LinkedPaths = GetLinkedPaths(RerouteNode);
+				FastBreakLinkedPaths(LinkedPaths);
+
+				if(const URigVMNode* MakeStructNode = ConvertRerouteNodeToDispatch(RerouteNode, MakeTemplateNotation, false, false))
+				{
+					if(const URigVMNode* BreakStructNode = AddTemplateNode(BreakTemplateNotation, MakeStructNode->GetPosition() + FVector2D(150, 0), FString(), false, false))
+					{
+						URigVMPin* PinToResolve = BreakStructNode->FindPin(FRigVMDispatch_BreakStruct::StructName.ToString());
+						check(PinToResolve);
+						verify(ResolveWildCardPin(PinToResolve, TypeIndex, false, false));
+
+						AddLink(
+							MakeStructNode->FindPin(FRigVMDispatch_MakeStruct::StructName.ToString()),
+							BreakStructNode->FindPin(FRigVMDispatch_BreakStruct::StructName.ToString()),
+							false);
+
+						const FString NodeNamePrefix = URigVMPin::JoinPinPath({NodeName, FString()});
+						FRestoreLinkedPathSettings RestoreSettings;
+						RestoreSettings.RemapDelegates.Add(NodeName,
+							FRigVMController_PinPathRemapDelegate::CreateLambda(
+								[NodeNamePrefix, MakeStructNode, BreakStructNode](const FString& InPinPath, bool bIsInput) -> FString
+							{
+								if(InPinPath.StartsWith(NodeNamePrefix, ESearchCase::CaseSensitive))
+								{
+									const URigVMNode* NewNode = bIsInput ? MakeStructNode : BreakStructNode;
+									const FString RemainingPinPath = InPinPath.Mid(NodeNamePrefix.Len());
+									FString Left, Right;
+									if(!URigVMPin::SplitPinPathAtStart(RemainingPinPath, Left, Right))
+									{
+										Left = RemainingPinPath;
+									}
+									else
+									{
+										Right = URigVMPin::JoinPinPath({FString(), Right});
+									}
+
+									static const FString& ValueName = URigVMRerouteNode::ValueName;
+									static const FString ElementsName = FRigVMDispatch_MakeStruct::ElementsName.ToString();
+
+									if(Left.Equals(ValueName, ESearchCase::CaseSensitive))
+									{
+										const FString NewNodeNamePrefix = URigVMPin::JoinPinPath({NewNode->GetName(), FString()});
+										return NewNodeNamePrefix + ElementsName + Right;
+									}
+
+								}
+								return InPinPath;
+							})
+						);
+						RestoreLinkedPaths(LinkedPaths, RestoreSettings);
+					}
+				}
+				else
+				{
+					// clean up detached links
+					ProcessDetachedLinks();
+				}
+			}
+			ReroutesToTurnIntoMakeAndBreakStruct.Reset();
+		}
+
+		// convert some reroutes to make array nodes
+		if(ReroutesToTurnIntoMakeArray.Num() > 0)
+		{
+			const FName TemplateNotation = FRigVMDispatch_ArrayMake().GetTemplateNotation();
+			const TMap<FString, FString> PinRedirects;
+			for(URigVMRerouteNode* RerouteNode : ReroutesToTurnIntoMakeArray)
+			{
+				ConvertRerouteNodeToDispatch(RerouteNode, TemplateNotation, false, false);
+			}
+			ReroutesToTurnIntoMakeArray.Reset();
+		}
+	}
+
+	return Result;
+}
+
 FRigVMClientPatchResult URigVMController::PatchUnitNodesOnLoad()
 {
 	FRigVMClientPatchResult Result;
diff --git a/Engine/Plugins/Runtime/RigVM/Source/RigVMDeveloper/Private/RigVMModel/RigVMControllerActions.cpp b/Engine/Plugins/Runtime/RigVM/Source/RigVMDeveloper/Private/RigVMModel/RigVMControllerActions.cpp
--- a/Engine/Plugins/Runtime/RigVM/Source/RigVMDeveloper/Private/RigVMModel/RigVMControllerActions.cpp
+++ b/Engine/Plugins/Runtime/RigVM/Source/RigVMDeveloper/Private/RigVMModel/RigVMControllerActions.cpp
@@ -958,44 +958,6 @@ bool FRigVMSetCommentTextAction::Redo()
 	return FRigVMBaseAction::Redo();
 }
 
-FRigVMSetRerouteCompactnessAction::FRigVMSetRerouteCompactnessAction()
-: FRigVMBaseAction(nullptr)
-, NodePath()
-, OldShowAsFullNode(false)
-, NewShowAsFullNode(false)
-{
-
-}
-FRigVMSetRerouteCompactnessAction::FRigVMSetRerouteCompactnessAction(URigVMController* InController, URigVMRerouteNode* InNode, bool InShowAsFullNode)
-: FRigVMBaseAction(InController)
-, NodePath(InNode->GetNodePath())
-, OldShowAsFullNode(InNode->GetShowsAsFullNode())
-, NewShowAsFullNode(InShowAsFullNode)
-{
-}
-
-bool FRigVMSetRerouteCompactnessAction::Undo()
-{
-	if(!FRigVMBaseAction::Undo())
-	{
-		return false;
-	}
-	return GetController()->SetRerouteCompactnessByName(*NodePath, OldShowAsFullNode, false);
-}
-
-bool FRigVMSetRerouteCompactnessAction::Redo()
-{
-	if(!CanUndoRedo())
-	{
-		return false;
-	}
-	if(!GetController()->SetRerouteCompactnessByName(*NodePath, NewShowAsFullNode, false))
-	{
-		return false;
-	}
-	return FRigVMBaseAction::Redo();
-}
-
 FRigVMRenameVariableAction::FRigVMRenameVariableAction()
 : FRigVMBaseAction(nullptr)
 {
diff --git a/Engine/Plugins/Runtime/RigVM/Source/RigVMDeveloper/Private/RigVMModel/RigVMPin.cpp b/Engine/Plugins/Runtime/RigVM/Source/RigVMDeveloper/Private/RigVMModel/RigVMPin.cpp
--- a/Engine/Plugins/Runtime/RigVM/Source/RigVMDeveloper/Private/RigVMModel/RigVMPin.cpp
+++ b/Engine/Plugins/Runtime/RigVM/Source/RigVMDeveloper/Private/RigVMModel/RigVMPin.cpp
@@ -399,7 +399,7 @@ bool URigVMPin::IsExpanded() const
 {
 	if(!bIsExpanded)
 	{
-		if(IsFixedSizeArray())
+		if(ShouldOnlyShowSubPins())
 		{
 			return true;
 		}
@@ -687,6 +687,74 @@ bool URigVMPin::IsFixedSizeArray() const
 	return false;
 }
 
+bool URigVMPin::ShouldOnlyShowSubPins() const
+{
+#if WITH_EDITOR
+
+	if(IsRootPin())
+	{
+		if(const URigVMNode* Node = GetNode())
+		{
+			if(const URigVMUnitNode* UnitNode = Cast<URigVMUnitNode>(Node))
+			{
+				if(const UScriptStruct* Struct = UnitNode->GetScriptStruct())
+				{
+					if(const FProperty* Property = Struct->FindPropertyByName(GetFName()))
+					{
+						return Property->HasMetaData(FRigVMStruct::ShowOnlySubPinsMetaName);
+					}
+				}
+			}
+			else if(const URigVMDispatchNode* DispatchNode = Cast<URigVMDispatchNode>(Node))
+			{
+				if(const FRigVMDispatchFactory* Factory = DispatchNode->GetFactory())
+				{
+					return Factory->HasArgumentMetaData(GetFName(), FRigVMStruct::ShowOnlySubPinsMetaName);
+				}
+			}
+		}
+	}
+#endif
+
+	return false;
+}
+
+bool URigVMPin::ShouldHideSubPins() const
+{
+#if WITH_EDITOR
+	if(ShouldOnlyShowSubPins())
+	{
+		return false;
+	}
+	
+	if(IsRootPin())
+	{
+		if(const URigVMNode* Node = GetNode())
+		{
+			if(const URigVMUnitNode* UnitNode = Cast<URigVMUnitNode>(Node))
+			{
+				if(const UScriptStruct* Struct = UnitNode->GetScriptStruct())
+				{
+					if(const FProperty* Property = Struct->FindPropertyByName(GetFName()))
+					{
+						return Property->HasMetaData(FRigVMStruct::HideSubPinsMetaName);
+					}
+				}
+			}
+			else if(const URigVMDispatchNode* DispatchNode = Cast<URigVMDispatchNode>(Node))
+			{
+				if(const FRigVMDispatchFactory* Factory = DispatchNode->GetFactory())
+				{
+					return Factory->HasArgumentMetaData(GetFName(), FRigVMStruct::HideSubPinsMetaName);
+				}
+			}
+		}
+	}
+#endif
+
+	return false;
+}
+
 FString URigVMPin::GetDefaultValue() const
 {
 	return GetDefaultValue(EmptyPinOverride);
@@ -1663,7 +1731,25 @@ bool URigVMPin::CanLink(URigVMPin* InSourcePin, URigVMPin* InTargetPin, FString*
 		}
 		return false;
 	}
-	
+
+	if(InSourcePin->ShouldOnlyShowSubPins() || InSourcePin->IsFixedSizeArray())
+	{
+		if (OutFailureReason)
+		{
+			*OutFailureReason = TEXT("Source pin only allows links to sub-pins.");
+		}
+		return false;
+	}
+
+	if(InTargetPin->ShouldOnlyShowSubPins() || InTargetPin->IsFixedSizeArray())
+	{
+		if (OutFailureReason)
+		{
+			*OutFailureReason = TEXT("Target pin only allows links to sub-pins.");
+		}
+		return false;
+	}
+
 	URigVMNode* SourceNode = InSourcePin->GetNode();
 	URigVMNode* TargetNode = InTargetPin->GetNode();
 	if (SourceNode == TargetNode)
diff --git a/Engine/Plugins/Runtime/RigVM/Source/RigVMDeveloper/Public/RigVMModel/Nodes/RigVMRerouteNode.h b/Engine/Plugins/Runtime/RigVM/Source/RigVMDeveloper/Public/RigVMModel/Nodes/RigVMRerouteNode.h
--- a/Engine/Plugins/Runtime/RigVM/Source/RigVMDeveloper/Public/RigVMModel/Nodes/RigVMRerouteNode.h
+++ b/Engine/Plugins/Runtime/RigVM/Source/RigVMDeveloper/Public/RigVMModel/Nodes/RigVMRerouteNode.h
@@ -2,7 +2,7 @@
 
 #pragma once
 
-#include "RigVMModel/Nodes/RigVMTemplateNode.h"
+#include "RigVMModel/RigVMNode.h"
 #include "RigVMRerouteNode.generated.h"
 
 /**
@@ -13,7 +13,7 @@
  * displayed as full nodes or as small circles.
  */
 UCLASS(BlueprintType)
-class RIGVMDEVELOPER_API URigVMRerouteNode : public URigVMTemplateNode
+class RIGVMDEVELOPER_API URigVMRerouteNode : public URigVMNode
 {
 	GENERATED_BODY()
 
@@ -25,34 +25,16 @@ public:
 	// Override of node title
 	virtual FString GetNodeTitle() const override;
 
-	// Returns true if this node should be shown as a full node,
-	// or false if this node should be shown as a small circle.
-	UFUNCTION(BlueprintCallable, Category = RigVMRerouteNode)
-	bool GetShowsAsFullNode() const;
-
 	virtual FLinearColor GetNodeColor() const override;
 
-	virtual FName GetNotation() const override;
-	virtual const FRigVMTemplate* GetTemplate() const override;
-	virtual bool IsSingleton() const override { return false; }
-
-	// Is a dot node with a root pin source connection
-	bool IsPassThrough() const;
-
 	// Has no source connections
 	bool IsLiteral() const;
 
-	// Is not a passthrough or a literal reroute
-	bool IsCompositionNode() const;
-
 private:
 
 	static const FString RerouteName;
 	static const FString ValueName;
 
-	UPROPERTY()
-	bool bShowAsFullNode;
-
 	friend class URigVMController;
 	friend class URigVMCompiler;
 	friend class FRigVMParserAST;
diff --git a/Engine/Plugins/Runtime/RigVM/Source/RigVMDeveloper/Public/RigVMModel/RigVMController.h b/Engine/Plugins/Runtime/RigVM/Source/RigVMDeveloper/Public/RigVMModel/RigVMController.h
--- a/Engine/Plugins/Runtime/RigVM/Source/RigVMDeveloper/Public/RigVMModel/RigVMController.h
+++ b/Engine/Plugins/Runtime/RigVM/Source/RigVMDeveloper/Public/RigVMModel/RigVMController.h
@@ -358,25 +358,43 @@ public:
 	// they don't require any additional memory though and are purely
 	// cosmetic. This causes a NodeAdded modified event.
 	UFUNCTION(BlueprintCallable, Category = RigVMController)
-	URigVMRerouteNode* AddRerouteNodeOnLink(URigVMLink* InLink, bool bShowAsFullNode, const FVector2D& InPosition = FVector2D::ZeroVector, const FString& InNodeName = TEXT(""), bool bSetupUndoRedo = true, bool bPrintPythonCommand = false);
+	URigVMRerouteNode* AddRerouteNodeOnLink(URigVMLink* InLink, const FVector2D& InPosition = FVector2D::ZeroVector, const FString& InNodeName = TEXT(""), bool bSetupUndoRedo = true, bool bPrintPythonCommand = false);
 
 	// Adds a Reroute Node on an existing Link to the edited Graph given the Link's string representation.
 	// Reroute Nodes can be used to visually improve the data flow,
 	// they don't require any additional memory though and are purely
 	// cosmetic. This causes a NodeAdded modified event.
 	UFUNCTION(BlueprintCallable, Category = RigVMController)
-	URigVMRerouteNode* AddRerouteNodeOnLinkPath(const FString& InLinkPinPathRepresentation, bool bShowAsFullNode, const FVector2D& InPosition = FVector2D::ZeroVector, const FString& InNodeName = TEXT(""), bool bSetupUndoRedo = true, bool bPrintPythonCommand = false);
+	URigVMRerouteNode* AddRerouteNodeOnLinkPath(const FString& InLinkPinPathRepresentation, const FVector2D& InPosition = FVector2D::ZeroVector, const FString& InNodeName = TEXT(""), bool bSetupUndoRedo = true, bool bPrintPythonCommand = false);
 
 	// Adds a Reroute Node on an existing Pin to the editor Graph.
 	// Reroute Nodes can be used to visually improve the data flow,
 	// they don't require any additional memory though and are purely
 	// cosmetic. This causes a NodeAdded modified event.
 	UFUNCTION(BlueprintCallable, Category = RigVMController)
-	URigVMRerouteNode* AddRerouteNodeOnPin(const FString& InPinPath, bool bAsInput, bool bShowAsFullNode, const FVector2D& InPosition = FVector2D::ZeroVector, const FString& InNodeName = TEXT(""), bool bSetupUndoRedo = true, bool bPrintPythonCommand = false);
+	URigVMRerouteNode* AddRerouteNodeOnPin(const FString& InPinPath, bool bAsInput, const FVector2D& InPosition = FVector2D::ZeroVector, const FString& InNodeName = TEXT(""), bool bSetupUndoRedo = true, bool bPrintPythonCommand = false);
 
 	// Adds a free Reroute Node
 	UFUNCTION(BlueprintCallable, Category = RigVMController)
-	URigVMRerouteNode* AddFreeRerouteNode(bool bShowAsFullNode, const FString& InCPPType, const FName& InCPPTypeObjectPath, bool bIsConstant, const FName& InCustomWidgetName, const FString& InDefaultValue, const FVector2D& InPosition = FVector2D::ZeroVector, const FString& InNodeName = TEXT(""), bool bSetupUndoRedo = true);
+	URigVMRerouteNode* AddFreeRerouteNode(const FString& InCPPType, const FName& InCPPTypeObjectPath, bool bIsConstant, const FName& InCustomWidgetName, const FString& InDefaultValue, const FVector2D& InPosition = FVector2D::ZeroVector, const FString& InNodeName = TEXT(""), bool bSetupUndoRedo = true);
+
+	// helper function to set up a constant node
+	URigVMTemplateNode* AddConstantNode(const FString& InCPPType, const FName& InCPPTypeObjectPath, const FString& InDefaultValue, const FVector2D& InPosition = FVector2D::ZeroVector, const FString& InNodeName = TEXT(""), bool bSetupUndoRedo = true);
+
+	// helper function to set up a make struct node
+	URigVMTemplateNode* AddMakeStructNode(const FString& InCPPType, const FName& InCPPTypeObjectPath, const FString& InDefaultValue, const FVector2D& InPosition = FVector2D::ZeroVector, const FString& InNodeName = TEXT(""), bool bSetupUndoRedo = true);
+
+	// helper function to set up a break struct node
+	URigVMTemplateNode* AddBreakStructNode(const FString& InCPPType, const FName& InCPPTypeObjectPath, const FString& InDefaultValue, const FVector2D& InPosition = FVector2D::ZeroVector, const FString& InNodeName = TEXT(""), bool bSetupUndoRedo = true);
+
+	// helper function to set up a constant node on a pin
+	URigVMTemplateNode* AddConstantNodeOnPin(const FString& InPinPath, const FVector2D& InPosition = FVector2D::ZeroVector, const FString& InNodeName = TEXT(""), bool bSetupUndoRedo = true);
+
+	// helper function to set up a make struct node on a pin
+	URigVMTemplateNode* AddMakeStructNodeOnPin(const FString& InPinPath, const FVector2D& InPosition = FVector2D::ZeroVector, const FString& InNodeName = TEXT(""), bool bSetupUndoRedo = true);
+
+	// helper function to set up a break struct node on a pin
+	URigVMTemplateNode* AddBreakStructNodeOnPin(const FString& InPinPath, const FVector2D& InPosition = FVector2D::ZeroVector, const FString& InNodeName = TEXT(""), bool bSetupUndoRedo = true);
 
 	// Adds a branch node to the graph.
 	// Branch nodes can be used to split the execution of into multiple branches,
@@ -669,16 +687,6 @@ public:
 	UFUNCTION(BlueprintCallable, Category = RigVMController)
 	bool SetCommentTextByName(const FName& InNodeName, const FString& InCommentText, const int32& InCommentFontSize, const bool& bInCommentBubbleVisible, const bool& bInCommentColorBubble, bool bSetupUndoRedo = true, bool bPrintPythonCommand = false);
 
-	// Sets the compactness of a reroute node in the graph.
-	// This causes a RerouteCompactnessChanged modified event.
-	UFUNCTION(BlueprintCallable, Category = RigVMController)
-	bool SetRerouteCompactness(URigVMNode* InNode, bool bShowAsFullNode, bool bSetupUndoRedo = true, bool bPrintPythonCommand = false);
-
-	// Sets the compactness of a reroute node in the graph by name.
-	// This causes a RerouteCompactnessChanged modified event.
-	UFUNCTION(BlueprintCallable, Category = RigVMController)
-	bool SetRerouteCompactnessByName(const FName& InNodeName, bool bShowAsFullNode, bool bSetupUndoRedo = true);
-
 	// Renames a variable in the graph.
 	// This causes a VariableRenamed modified event.
 	UFUNCTION(BlueprintCallable, Category = RigVMController, meta=(DeprecatedFunction))
@@ -1052,9 +1060,9 @@ public:
 		TArray<FPinInfo>::RangedForIteratorType end() const { return Pins.end(); }
 
 		int32 AddPin(const URigVMPin* InPin, int32 InParentIndex, ERigVMPinDirection InDirection = ERigVMPinDirection::Invalid);
-		int32 AddPin(FProperty* InProperty, URigVMController* InController, ERigVMPinDirection InDirection, int32 InParentIndex, const uint8* InDefaultValueMemory);
-		int32 AddPin(URigVMController* InController, int32 InParentIndex, const FName& InName, ERigVMPinDirection InDirection, TRigVMTypeIndex InTypeIndex, const FString& InDefaultValue, const uint8* InDefaultValueMemory, const FPinInfoArray* InPreviousPinInfos);
-		void AddPins(UScriptStruct* InScriptStruct, URigVMController* InController, ERigVMPinDirection InDirection, int32 InParentIndex, const uint8* InDefaultValueMemory);
+		int32 AddPin(FProperty* InProperty, URigVMController* InController, ERigVMPinDirection InDirection, int32 InParentIndex, const uint8* InDefaultValueMemory, bool bAddSubPins);
+		int32 AddPin(URigVMController* InController, int32 InParentIndex, const FName& InName, ERigVMPinDirection InDirection, TRigVMTypeIndex InTypeIndex, const FString& InDefaultValue, const uint8* InDefaultValueMemory, const FPinInfoArray* InPreviousPinInfos, bool bAddSubPins);
+		void AddPins(UScriptStruct* InScriptStruct, URigVMController* InController, ERigVMPinDirection InDirection, int32 InParentIndex, const uint8* InDefaultValueMemory, bool bAddSubPins);
 
 		const FString& GetPinPath(const int32 InIndex) const;
 		int32 GetIndexFromPinPath(const FString& InPinPath) const;
@@ -1121,7 +1129,6 @@ public:
 
 private:
 	bool BreakAllLinksRecursive(URigVMPin* Pin, bool bAsInput, bool bTowardsParent, bool bSetupUndoRedo);
-	void UpdateRerouteNodeAfterChangingLinks(URigVMPin* PinChanged, bool bSetupUndoRedo = true);
 	bool SetPinExpansion(URigVMPin* InPin, bool bIsExpanded, bool bSetupUndoRedo = true);
 	void ExpandPinRecursively(URigVMPin* InPin, bool bSetupUndoRedo);
 	bool SetPinIsWatched(URigVMPin* InPin, bool bIsWatched, bool bSetupUndoRedo);
@@ -1278,9 +1285,12 @@ public:
 	URigVMActionStack* GetActionStack() const;
 	void SetActionStack(URigVMActionStack* InActionStack);
 
+	URigVMNode* ConvertRerouteNodeToDispatch(URigVMRerouteNode* InRerouteNode, const FName& InTemplateNotation, bool bSetupUndoRedo, bool bPrintPythonCommand);
+
 protected:
 
 	// backwards compatibility code
+	FRigVMClientPatchResult PatchRerouteNodesOnLoad();
 	FRigVMClientPatchResult PatchUnitNodesOnLoad();
 	FRigVMClientPatchResult PatchDispatchNodesOnLoad();
 	FRigVMClientPatchResult PatchBranchNodesOnLoad();
diff --git a/Engine/Plugins/Runtime/RigVM/Source/RigVMDeveloper/Public/RigVMModel/RigVMControllerActions.h b/Engine/Plugins/Runtime/RigVM/Source/RigVMDeveloper/Public/RigVMModel/RigVMControllerActions.h
--- a/Engine/Plugins/Runtime/RigVM/Source/RigVMDeveloper/Public/RigVMModel/RigVMControllerActions.h
+++ b/Engine/Plugins/Runtime/RigVM/Source/RigVMDeveloper/Public/RigVMModel/RigVMControllerActions.h
@@ -754,33 +754,6 @@ public:
 	bool bNewColorBubble;
 };
 
-/**
- * An action setting a reroute node's compactness in the graph.
- */
-USTRUCT()
-struct FRigVMSetRerouteCompactnessAction : public FRigVMBaseAction
-{
-	GENERATED_BODY()
-
-public:
-
-	FRigVMSetRerouteCompactnessAction();
-	FRigVMSetRerouteCompactnessAction(URigVMController* InController, URigVMRerouteNode* InNode, bool InShowAsFullNode);
-	virtual ~FRigVMSetRerouteCompactnessAction() {};
-	virtual UScriptStruct* GetScriptStruct() const override { return FRigVMSetRerouteCompactnessAction::StaticStruct(); }
-	virtual bool Undo() override;
-	virtual bool Redo() override;
-
-	UPROPERTY()
-	FString NodePath;
-
-	UPROPERTY()
-	bool OldShowAsFullNode;
-
-	UPROPERTY()
-	bool NewShowAsFullNode;
-};
-
 /**
  * An action renaming a variable in the graph.
  */
diff --git a/Engine/Plugins/Runtime/RigVM/Source/RigVMDeveloper/Public/RigVMModel/RigVMNotifications.h b/Engine/Plugins/Runtime/RigVM/Source/RigVMDeveloper/Public/RigVMModel/RigVMNotifications.h
--- a/Engine/Plugins/Runtime/RigVM/Source/RigVMDeveloper/Public/RigVMModel/RigVMNotifications.h
+++ b/Engine/Plugins/Runtime/RigVM/Source/RigVMDeveloper/Public/RigVMModel/RigVMNotifications.h
@@ -37,7 +37,6 @@ enum class ERigVMGraphNotifType : uint8
 	LinkAdded, // A link has been added (Subject == URigVMLink)
 	LinkRemoved, // A link has been removed (Subject == URigVMLink)
 	CommentTextChanged, // A comment node's text has changed (Subject == URigVMCommentNode)
-	RerouteCompactnessChanged, // A reroute node's compactness has changed (Subject == URigVMRerouteNode)
 	VariableAdded, // A variable has been added (Subject == URigVMVariableNode)
 	VariableRemoved, // A variable has been removed (Subject == URigVMVariableNode)
 	VariableRenamed, // A variable has been renamed (Subject == URigVMVariableNode)
diff --git a/Engine/Plugins/Runtime/RigVM/Source/RigVMDeveloper/Public/RigVMModel/RigVMPin.h b/Engine/Plugins/Runtime/RigVM/Source/RigVMDeveloper/Public/RigVMModel/RigVMPin.h
--- a/Engine/Plugins/Runtime/RigVM/Source/RigVMDeveloper/Public/RigVMModel/RigVMPin.h
+++ b/Engine/Plugins/Runtime/RigVM/Source/RigVMDeveloper/Public/RigVMModel/RigVMPin.h
@@ -274,6 +274,14 @@ public:
 	// Returns true if this pin is an array that should be displayed as elements only
 	UFUNCTION(BlueprintCallable, Category = RigVMPin)
 	bool IsFixedSizeArray() const;
+
+	// Returns true if this pin is an array that should be displayed as elements only
+	UFUNCTION(BlueprintCallable, Category = RigVMPin)
+	bool ShouldOnlyShowSubPins() const;
+
+	// Returns true if this pin's subpins should be hidden in the UI
+	UFUNCTION(BlueprintCallable, Category = RigVMPin)
+	bool ShouldHideSubPins() const;
 	
 	// Returns the default value of the Pin as a string.
 	// Note that this value is computed based on the Pin's
