Control Rig: Improve performance of context menu repopulation

diff --git a/Engine/Plugins/Animation/ControlRig/Source/ControlRigDeveloper/Private/ControlRigBlueprint.cpp b/Engine/Plugins/Animation/ControlRig/Source/ControlRigDeveloper/Private/ControlRigBlueprint.cpp
--- a/Engine/Plugins/Animation/ControlRig/Source/ControlRigDeveloper/Private/ControlRigBlueprint.cpp
+++ b/Engine/Plugins/Animation/ControlRig/Source/ControlRigDeveloper/Private/ControlRigBlueprint.cpp
@@ -870,6 +870,8 @@ void UControlRigBlueprint::PreSave(FObjectPreSaveContext ObjectSaveContext)
 
 	FunctionReferenceNodeData = GetReferenceNodeData();
 	IAssetRegistry::GetChecked().AssetTagsFinalized(*this);
+
+	CachedAssetTags.Reset();
 }
 
 void UControlRigBlueprint::PostLoad()
@@ -2740,6 +2742,19 @@ void UControlRigBlueprint::PostDuplicate(bool bDuplicateForPIE)
 	check(HasAnyFlags(RF_NeedPostLoad));
 }
 
+void UControlRigBlueprint::GetAssetRegistryTags(TArray<FAssetRegistryTag>& OutTags) const
+{
+	if (CachedAssetTags.IsEmpty())
+	{
+		Super::GetAssetRegistryTags(OutTags);
+		CachedAssetTags = OutTags;
+	}
+	else
+	{
+		OutTags = CachedAssetTags;
+	}
+}
+
 FRigVMGraphModifiedEvent& UControlRigBlueprint::OnModified()
 {
 	return ModifiedEvent;
diff --git a/Engine/Plugins/Animation/ControlRig/Source/ControlRigDeveloper/Public/ControlRigBlueprint.h b/Engine/Plugins/Animation/ControlRig/Source/ControlRigDeveloper/Public/ControlRigBlueprint.h
--- a/Engine/Plugins/Animation/ControlRig/Source/ControlRigDeveloper/Public/ControlRigBlueprint.h
+++ b/Engine/Plugins/Animation/ControlRig/Source/ControlRigDeveloper/Public/ControlRigBlueprint.h
@@ -246,6 +246,7 @@ public:
 	virtual void ReplaceDeprecatedNodes() override;
 	virtual void PreDuplicate(FObjectDuplicationParameters& DupParams) override;
 	virtual void PostDuplicate(bool bDuplicateForPIE) override;
+	virtual void GetAssetRegistryTags(TArray<FAssetRegistryTag>& OutTags) const override;
 
 	virtual bool SupportsGlobalVariables() const override { return true; }
 	virtual bool SupportsLocalVariables() const override { return true; }
@@ -447,6 +448,8 @@ protected:
 	UPROPERTY(transient, DuplicateTransient)
 	TObjectPtr<URigVMController> TemplateController;
 
+	mutable TArray<FAssetRegistryTag> CachedAssetTags;
+
 #endif
 
 public:
diff --git a/Engine/Plugins/Runtime/RigVM/Source/RigVMDeveloper/Private/RigVMModel/RigVMBuildData.cpp b/Engine/Plugins/Runtime/RigVM/Source/RigVMDeveloper/Private/RigVMModel/RigVMBuildData.cpp
--- a/Engine/Plugins/Runtime/RigVM/Source/RigVMDeveloper/Private/RigVMModel/RigVMBuildData.cpp
+++ b/Engine/Plugins/Runtime/RigVM/Source/RigVMDeveloper/Private/RigVMModel/RigVMBuildData.cpp
@@ -20,7 +20,7 @@ FRigVMReferenceNodeData::FRigVMReferenceNodeData(URigVMFunctionReferenceNode* In
 	check(InReferenceNode);
 	ReferenceNodePtr = TSoftObjectPtr<URigVMFunctionReferenceNode>(InReferenceNode);
 	ReferenceNodePath = ReferenceNodePtr.ToString();
-	ReferencedHeader = InReferenceNode->GetReferencedFunctionHeader();
+	ReferencedFunctionIdentifier = InReferenceNode->GetReferencedFunctionHeader().LibraryPointer;
 }
 
 TSoftObjectPtr<URigVMFunctionReferenceNode> FRigVMReferenceNodeData::GetReferenceNodeObjectPath()
@@ -142,9 +142,13 @@ void URigVMBuildData::RegisterReferencesFromAsset(const FAssetData& InAssetData)
 			ReferenceNodeDataProperty->ImportText_Direct(*ReferenceNodeDataString, &ReferenceNodeDatas, nullptr, EPropertyPortFlags::PPF_None);	
 			for(FRigVMReferenceNodeData& ReferenceNodeData : ReferenceNodeDatas)
 			{
-				if (ReferenceNodeData.ReferencedHeader.IsValid())
+				if (ReferenceNodeData.ReferencedFunctionIdentifier.LibraryNode.IsValid())
 				{
-					BuildData->RegisterFunctionReference(ReferenceNodeData.ReferencedHeader.LibraryPointer, ReferenceNodeData.GetReferenceNodeObjectPath());
+					BuildData->RegisterFunctionReference(ReferenceNodeData.ReferencedFunctionIdentifier, ReferenceNodeData.GetReferenceNodeObjectPath());
+				}
+				else if (ReferenceNodeData.ReferencedHeader_DEPRECATED.IsValid())
+				{
+					BuildData->RegisterFunctionReference(ReferenceNodeData.ReferencedHeader_DEPRECATED.LibraryPointer, ReferenceNodeData.GetReferenceNodeObjectPath());
 				}
 				else if (!ReferenceNodeData.ReferencedFunctionPath_DEPRECATED.IsEmpty())
 				{
@@ -236,19 +240,24 @@ void URigVMBuildData::RegisterFunctionReference(const FRigVMGraphFunctionIdentif
 
 void URigVMBuildData::RegisterFunctionReference(FRigVMReferenceNodeData InReferenceNodeData)
 {
-	if (InReferenceNodeData.ReferencedHeader.IsValid())
+	if (InReferenceNodeData.ReferencedFunctionIdentifier.LibraryNode.IsValid())
+	{
+		return RegisterFunctionReference(InReferenceNodeData.ReferencedFunctionIdentifier, InReferenceNodeData.GetReferenceNodeObjectPath());
+	}
+
+	if (!InReferenceNodeData.ReferencedFunctionIdentifier.LibraryNode.IsValid())
 	{
-		return RegisterFunctionReference(InReferenceNodeData.ReferencedHeader.LibraryPointer, InReferenceNodeData.GetReferenceNodeObjectPath());
+		InReferenceNodeData.ReferencedFunctionIdentifier = InReferenceNodeData.ReferencedHeader_DEPRECATED.LibraryPointer;
 	}
 
-	if (!InReferenceNodeData.ReferencedHeader.LibraryPointer.LibraryNode.IsValid())
+	if (!InReferenceNodeData.ReferencedFunctionIdentifier.LibraryNode.IsValid())
 	{
-		InReferenceNodeData.ReferencedHeader.LibraryPointer.LibraryNode = InReferenceNodeData.ReferencedFunctionPath_DEPRECATED;
+		InReferenceNodeData.ReferencedFunctionIdentifier.LibraryNode = InReferenceNodeData.ReferencedFunctionPath_DEPRECATED;
 	}
 	
-	check(InReferenceNodeData.ReferencedHeader.LibraryPointer.LibraryNode.IsValid());
+	check(InReferenceNodeData.ReferencedFunctionIdentifier.LibraryNode.IsValid());
 
-	FSoftObjectPath LibraryNodePath = InReferenceNodeData.ReferencedHeader.LibraryPointer.LibraryNode;
+	FSoftObjectPath LibraryNodePath = InReferenceNodeData.ReferencedFunctionIdentifier.LibraryNode;
 	TSoftObjectPtr<URigVMLibraryNode> LibraryNodePtr = TSoftObjectPtr<URigVMLibraryNode>(LibraryNodePath);
 
 	// Try to find a FunctionIdentifier with the same LibraryNodePath
diff --git a/Engine/Plugins/Runtime/RigVM/Source/RigVMDeveloper/Public/RigVMModel/RigVMBuildData.h b/Engine/Plugins/Runtime/RigVM/Source/RigVMDeveloper/Public/RigVMModel/RigVMBuildData.h
--- a/Engine/Plugins/Runtime/RigVM/Source/RigVMDeveloper/Public/RigVMModel/RigVMBuildData.h
+++ b/Engine/Plugins/Runtime/RigVM/Source/RigVMDeveloper/Public/RigVMModel/RigVMBuildData.h
@@ -34,7 +34,7 @@ struct RIGVMDEVELOPER_API FRigVMReferenceNodeData
 public:
 	
 	FRigVMReferenceNodeData()
-		:ReferenceNodePath()
+		:ReferenceNodePath(), ReferencedFunctionIdentifier()
 	{}
 
 	FRigVMReferenceNodeData(URigVMFunctionReferenceNode* InReferenceNode);
@@ -45,8 +45,11 @@ public:
 	UPROPERTY(meta=(DeprecatedProperty))
 	FString ReferencedFunctionPath_DEPRECATED;
 	
+	UPROPERTY(meta=(DeprecatedProperty))
+	FRigVMGraphFunctionHeader ReferencedHeader_DEPRECATED;
+
 	UPROPERTY()
-	FRigVMGraphFunctionHeader ReferencedHeader;
+	FRigVMGraphFunctionIdentifier ReferencedFunctionIdentifier;
 
 	TSoftObjectPtr<URigVMFunctionReferenceNode> GetReferenceNodeObjectPath();
 	URigVMFunctionReferenceNode* GetReferenceNode();
