RigVM: Ensure thread safety for FRigVMRegistry::FindFunction

diff --git a/Engine/Plugins/Runtime/RigVM/Source/RigVM/Private/RigVMCore/RigVM.cpp b/Engine/Plugins/Runtime/RigVM/Source/RigVM/Private/RigVMCore/RigVM.cpp
--- a/Engine/Plugins/Runtime/RigVM/Source/RigVM/Private/RigVMCore/RigVM.cpp
+++ b/Engine/Plugins/Runtime/RigVM/Source/RigVM/Private/RigVMCore/RigVM.cpp
@@ -223,6 +223,13 @@ void URigVM::Load(FArchive& Ar)
 		Ar << ByteCodeStorage;
 		Ar << Parameters;
 	}
+
+	// ensure to load the required functions
+	if(!ResolveFunctionsIfRequired())
+	{
+		Reset();
+		return;
+	}
 }
 
 void URigVM::PostLoad()
@@ -905,8 +912,12 @@ FRigVMParameter URigVM::GetParameterByName(const FName& InParameterName)
 	return FRigVMParameter();
 }
 
-void URigVM::ResolveFunctionsIfRequired()
+bool URigVM::ResolveFunctionsIfRequired()
 {
+	FScopeLock ResolveFunctionsScopeLock(&ResolveFunctionsMutex);
+	
+	bool bSuccess = true;
+	
 	if (GetFunctions().Num() != GetFunctionNames().Num())
 	{
 		GetFunctions().Reset();
@@ -925,9 +936,12 @@ void URigVM::ResolveFunctionsIfRequired()
 			{
 				// We cannot recover from missing functions. 
 				UE_LOG(LogRigVM, Fatal, TEXT("No handler found for function '%s'"), *GetFunctionNames()[FunctionIndex].ToString());
+				bSuccess = false;
 			}
 		}
 	}
+
+	return bSuccess;
 }
 
 void URigVM::RefreshInstructionsIfRequired()
diff --git a/Engine/Plugins/Runtime/RigVM/Source/RigVM/Private/RigVMCore/RigVMDispatchFactory.cpp b/Engine/Plugins/Runtime/RigVM/Source/RigVM/Private/RigVMCore/RigVMDispatchFactory.cpp
--- a/Engine/Plugins/Runtime/RigVM/Source/RigVM/Private/RigVMCore/RigVMDispatchFactory.cpp
+++ b/Engine/Plugins/Runtime/RigVM/Source/RigVM/Private/RigVMCore/RigVMDispatchFactory.cpp
@@ -9,6 +9,7 @@
 
 const FString FRigVMDispatchFactory::DispatchPrefix = TEXT("DISPATCH_");
 const FString FRigVMDispatchFactory::TrueString = TEXT("True");
+FCriticalSection FRigVMDispatchFactory::GetTemplateMutex;
 
 FName FRigVMDispatchFactory::GetFactoryName() const
 {
@@ -195,12 +196,15 @@ TArray<FRigVMExecuteArgument> FRigVMDispatchFactory::GetExecuteArguments(const F
 	return Arguments;
 }
 
-FRigVMFunctionPtr FRigVMDispatchFactory::GetDispatchFunction(const FRigVMTemplateTypeMap& InTypes) const
+FRigVMFunctionPtr FRigVMDispatchFactory::GetDispatchFunction(const FRigVMTemplateTypeMap& InTypes, bool bForceCreation) const
 {
-	const FString PermutationName = GetPermutationNameImpl(InTypes);
-	if(const FRigVMFunction* ExistingFunction = FRigVMRegistry::Get().FindFunction(*PermutationName))
+	if(!bForceCreation)
 	{
-		return ExistingFunction->FunctionPtr;
+		const FString PermutationName = GetPermutationNameImpl(InTypes);
+		if(const FRigVMFunction* ExistingFunction = FRigVMRegistry::Get().FindFunction(*PermutationName))
+		{
+			return ExistingFunction->FunctionPtr;
+		}
 	}
 	return GetDispatchFunctionImpl(InTypes);
 }
@@ -231,7 +235,7 @@ bool FRigVMDispatchFactory::CopyProperty(const FProperty* InTargetProperty, uint
 	return URigVMMemoryStorage::CopyProperty(InTargetProperty, InTargetPtr, InSourceProperty, InSourcePtr);
 }
 
-const FRigVMTemplate* FRigVMDispatchFactory::GetTemplate() const
+const FRigVMTemplate* FRigVMDispatchFactory::GetTemplate(bool bIsDispatchingTemplate) const
 {
 	// make sure to rely on the instance of this factory that's stored under the registry
 	FRigVMRegistry& Registry = FRigVMRegistry::Get();
@@ -242,12 +246,7 @@ const FRigVMTemplate* FRigVMDispatchFactory::GetTemplate() const
 		return ThisFactory->GetTemplate();
 	}
 
-	static bool bIsDispatchingTemplate = false;
-	if(bIsDispatchingTemplate)
-	{
-		return nullptr;
-	}
-	TGuardValue<bool> ReEntryGuard(bIsDispatchingTemplate, true);
+	FScopeLock GetTemplateScopeLock(&GetTemplateMutex);
 
 	if(CachedTemplate)
 	{
@@ -287,10 +286,10 @@ const FRigVMTemplate* FRigVMDispatchFactory::GetTemplate() const
 	Delegates.RequestDispatchFunctionDelegate = FRigVMTemplate_RequestDispatchFunctionDelegate::CreateLambda(
 	[this](const FRigVMTemplate*, const FRigVMTemplateTypeMap& InTypes)
 	{
-		return GetDispatchFunction(InTypes);
+		return GetDispatchFunction(InTypes, true /* force creation */);
 	});
 
-	CachedTemplate = Registry.GetOrAddTemplateFromArguments(GetFactoryName(), Arguments, Delegates); 
+	CachedTemplate = Registry.GetOrAddTemplateFromArguments(GetFactoryName(), Arguments, Delegates, bIsDispatchingTemplate); 
 	return CachedTemplate;
 }
 
diff --git a/Engine/Plugins/Runtime/RigVM/Source/RigVM/Private/RigVMCore/RigVMRegistry.cpp b/Engine/Plugins/Runtime/RigVM/Source/RigVM/Private/RigVMCore/RigVMRegistry.cpp
--- a/Engine/Plugins/Runtime/RigVM/Source/RigVM/Private/RigVMCore/RigVMRegistry.cpp
+++ b/Engine/Plugins/Runtime/RigVM/Source/RigVM/Private/RigVMCore/RigVMRegistry.cpp
@@ -15,6 +15,14 @@
 #include "RigVMFunctions/RigVMDispatch_Core.h"
 
 const FName FRigVMRegistry::TemplateNameMetaName = TEXT("TemplateName");
+FCriticalSection FRigVMRegistry::RefreshTypesMutex;
+FCriticalSection FRigVMRegistry::RegisterFunctionMutex;
+FCriticalSection FRigVMRegistry::RegisterTemplateMutex;
+FCriticalSection FRigVMRegistry::RegisterFactoryMutex;
+FCriticalSection FRigVMRegistry::FindFunctionMutex;
+FCriticalSection FRigVMRegistry::FindTemplateMutex;
+FCriticalSection FRigVMRegistry::FindFactoryMutex;
+FCriticalSection FRigVMRegistry::GetPermutationMutex;
 
 
 // When the object system has been completely loaded, load in all the engine types that we haven't registered already in InitializeIfNeeded 
@@ -202,6 +210,9 @@ uint32 FRigVMRegistry::GetHashForProperty(const FProperty* InProperty) const
 
 void FRigVMRegistry::InitializeIfNeeded()
 {
+	// we don't need to use a mutex here since this is called on
+	// the main thread during engine startup for the first time
+	
 	if(!Types.IsEmpty())
 	{
 		return;
@@ -305,6 +316,8 @@ void FRigVMRegistry::InitializeIfNeeded()
 
 void FRigVMRegistry::RefreshEngineTypes()
 {
+	FScopeLock RefreshTypesScopeLock(&RefreshTypesMutex);
+	
 	// Register all user-defined types that the engine knows about. Enumerating over the entire object hierarchy is
 	// slow, so we do it for structs, enums and dispatch factories in one shot.
 	TArray<UScriptStruct*> DispatchFactoriesToRegister;
@@ -393,6 +406,9 @@ void FRigVMRegistry::Reset()
 
 TRigVMTypeIndex FRigVMRegistry::FindOrAddType_Internal(const FRigVMTemplateArgumentType& InType, bool bForce)
 {
+	// we don't use a mutex here since by the time the engine relies on worker
+	// thread for execution or async loading all types will have been registered.
+	
 	TRigVMTypeIndex Index = GetTypeIndex(InType);
 	if(Index == INDEX_NONE)
 	{
@@ -1123,6 +1139,8 @@ bool FRigVMRegistry::IsAllowedType(const UClass* InClass)
 
 void FRigVMRegistry::Register(const TCHAR* InName, FRigVMFunctionPtr InFunctionPtr, UScriptStruct* InStruct, const TArray<FRigVMFunctionArgument>& InArguments)
 {
+	FScopeLock RegisterFunctionScopeLock(&RegisterFunctionMutex);
+	
 	if (FindFunction(InName) != nullptr)
 	{
 		return;
@@ -1203,6 +1221,8 @@ void FRigVMRegistry::Register(const TCHAR* InName, FRigVMFunctionPtr InFunctionP
 
 const FRigVMDispatchFactory* FRigVMRegistry::RegisterFactory(UScriptStruct* InFactoryStruct)
 {
+	FScopeLock RegisterFactoryScopeLock(&RegisterFactoryMutex);
+
 	check(InFactoryStruct);
 	check(InFactoryStruct != FRigVMDispatchFactory::StaticStruct());
 	check(InFactoryStruct->IsChildOf(FRigVMDispatchFactory::StaticStruct()));
@@ -1240,18 +1260,14 @@ const FRigVMDispatchFactory* FRigVMRegistry::RegisterFactory(UScriptStruct* InFa
 
 const FRigVMFunction* FRigVMRegistry::FindFunction(const TCHAR* InName) const
 {
+	FScopeLock FindFunctionScopeLock(&FindFunctionMutex);
+
 	// Check first if the function is provided by internally registered rig units. 
 	if(const int32* FunctionIndexPtr = FunctionNameToIndex.Find(InName))
 	{
 		return &Functions[*FunctionIndexPtr];
 	}
 
-	static bool IsDispatchingFunction = false;
-	if(IsDispatchingFunction)
-	{
-		return nullptr;
-	}
-
 	// Otherwise ask the associated dispatch factory for a function matching this signature.
 	const FString NameString(InName);
 	FString StructOrFactoryName, SuffixString;
@@ -1268,7 +1284,6 @@ const FRigVMFunction* FRigVMRegistry::FindFunction(const TCHAR* InName) const
 					const int32 PermutationIndex = Template->FindPermutation(ArgumentTypes);
 					if(PermutationIndex != INDEX_NONE)
 					{
-						TGuardValue<bool> ReEntryGuard(IsDispatchingFunction, true);
 						return ((FRigVMTemplate*)Template)->GetOrCreatePermutation(PermutationIndex);
 					}
 				}
@@ -1341,17 +1356,13 @@ const FRigVMTemplate* FRigVMRegistry::FindTemplate(const FName& InNotation, bool
 		return nullptr;
 	}
 
+	FScopeLock FindTemplateScopeLock(&FindTemplateMutex);
+
 	if(const int32* TemplateIndexPtr = TemplateNotationToIndex.Find(InNotation))
 	{
 		return &Templates[*TemplateIndexPtr];
 	}
 
-	static bool IsDispatchingFunction = false;
-	if(IsDispatchingFunction)
-	{
-		return nullptr;
-	}
-	
 	const FString NotationString(InNotation.ToString());
 	FString FactoryName, ArgumentsString;
 	if(NotationString.Split(TEXT("("), &FactoryName, &ArgumentsString))
@@ -1378,8 +1389,7 @@ const FRigVMTemplate* FRigVMRegistry::FindTemplate(const FName& InNotation, bool
 		
 		if(const FRigVMDispatchFactory* Factory = FindDispatchFactory(*FactoryName))
 		{
-			TGuardValue<bool> ReEntryGuard(IsDispatchingFunction, true);
-			return Factory->GetTemplate();
+			return Factory->GetTemplate(true /* dispatching template */);
 		}
 	}
 
@@ -1410,7 +1420,7 @@ const FRigVMTemplate* FRigVMRegistry::FindTemplate(const FName& InNotation, bool
 				const FString NewDispatchFactoryName = FRigVMDispatchFactory::DispatchPrefix + Redirect->NewName.ObjectName.ToString();
 				if(const FRigVMDispatchFactory* NewDispatchFactory = FindDispatchFactory(*NewDispatchFactoryName))
 				{
-					return NewDispatchFactory->GetTemplate();
+					return NewDispatchFactory->GetTemplate(true /* dispatching template */);
 				}
 			}
 		}
@@ -1458,12 +1468,20 @@ const TChunkedArray<FRigVMTemplate>& FRigVMRegistry::GetTemplates() const
 	return Templates;
 }
 
-const FRigVMTemplate* FRigVMRegistry::GetOrAddTemplateFromArguments(const FName& InName, const TArray<FRigVMTemplateArgument>& InArguments, const FRigVMTemplateDelegates& InDelegates)
+const FRigVMTemplate* FRigVMRegistry::GetOrAddTemplateFromArguments(const FName& InName, const TArray<FRigVMTemplateArgument>& InArguments, const FRigVMTemplateDelegates& InDelegates, bool bIsDispatchingTemplate)
 {
+	FScopeLock RegisterTemplateScopeLock(&RegisterTemplateMutex);
+	
 	FRigVMTemplate Template(InName, InArguments, INDEX_NONE);
-	if(const FRigVMTemplate* ExistingTemplate = FindTemplate(Template.GetNotation()))
+
+	// avoid reentry in FindTemplate. try to find an existing
+	// template only if we are not yet in ::FindTemplate.
+	if(!bIsDispatchingTemplate)
 	{
-		return ExistingTemplate;
+		if(const FRigVMTemplate* ExistingTemplate = FindTemplate(Template.GetNotation()))
+		{
+			return ExistingTemplate;
+		}
 	}
 
 	// we only support to ask for templates here which provide singleton types
@@ -1581,6 +1599,8 @@ const FRigVMTemplate* FRigVMRegistry::GetOrAddTemplateFromArguments(const FName&
 
 FRigVMDispatchFactory* FRigVMRegistry::FindDispatchFactory(const FName& InFactoryName) const 
 {
+	FScopeLock FindFactoryScopeLock(&FindFactoryMutex);
+
 	FRigVMDispatchFactory* const* FactoryPtr = Factories.FindByPredicate([InFactoryName](const FRigVMDispatchFactory* Factory) -> bool
 	{
 		return Factory->GetFactoryName() == InFactoryName;
diff --git a/Engine/Plugins/Runtime/RigVM/Source/RigVM/Private/RigVMCore/RigVMTemplate.cpp b/Engine/Plugins/Runtime/RigVM/Source/RigVM/Private/RigVMCore/RigVMTemplate.cpp
--- a/Engine/Plugins/Runtime/RigVM/Source/RigVM/Private/RigVMCore/RigVMTemplate.cpp
+++ b/Engine/Plugins/Runtime/RigVM/Source/RigVM/Private/RigVMCore/RigVMTemplate.cpp
@@ -1159,6 +1159,8 @@ const FRigVMFunction* FRigVMTemplate::GetPrimaryPermutation() const
 
 const FRigVMFunction* FRigVMTemplate::GetPermutation(int32 InIndex) const
 {
+	FScopeLock FindPermutationScopeLock(&FRigVMRegistry::GetPermutationMutex);
+	
 	const FRigVMRegistry& Registry = FRigVMRegistry::Get();
 	const int32 FunctionIndex = Permutations[InIndex];
 	if(Registry.GetFunctions().IsValidIndex(FunctionIndex))
@@ -1170,6 +1172,8 @@ const FRigVMFunction* FRigVMTemplate::GetPermutation(int32 InIndex) const
 
 const FRigVMFunction* FRigVMTemplate::GetOrCreatePermutation(int32 InIndex)
 {
+	FScopeLock FindPermutationScopeLock(&FRigVMRegistry::GetPermutationMutex);
+
 	if(const FRigVMFunction* Function = GetPermutation(InIndex))
 	{
 		return Function;
@@ -1177,6 +1181,7 @@ const FRigVMFunction* FRigVMTemplate::GetOrCreatePermutation(int32 InIndex)
 	
 	if(Permutations[InIndex] == INDEX_NONE && UsesDispatch())
 	{
+		FScopeLock RegisterFunctionScopeLock(&FRigVMRegistry::RegisterFunctionMutex);
 		FRigVMRegistry& Registry = FRigVMRegistry::Get();
 		
 		FTypeMap Types;
diff --git a/Engine/Plugins/Runtime/RigVM/Source/RigVM/Public/RigVMCore/RigVM.h b/Engine/Plugins/Runtime/RigVM/Source/RigVM/Public/RigVMCore/RigVM.h
--- a/Engine/Plugins/Runtime/RigVM/Source/RigVM/Public/RigVMCore/RigVM.h
+++ b/Engine/Plugins/Runtime/RigVM/Source/RigVM/Public/RigVMCore/RigVM.h
@@ -730,7 +730,7 @@ public:
 
 private:
 
-	void ResolveFunctionsIfRequired();
+	bool ResolveFunctionsIfRequired();
 	void RefreshInstructionsIfRequired();
 public:
 	void InvalidateCachedMemory();
@@ -903,6 +903,8 @@ protected:
 	void SetupInstructionTracking(FRigVMExtendedExecuteContext& Context, int32 InInstructionCount);
 	void StartProfiling(FRigVMExtendedExecuteContext& Context);
 	void StopProfiling(FRigVMExtendedExecuteContext& Context);
+	
+	FCriticalSection ResolveFunctionsMutex;
 
 	friend class URigVMCompiler;
 	friend struct FRigVMCompilerWorkData;
diff --git a/Engine/Plugins/Runtime/RigVM/Source/RigVM/Public/RigVMCore/RigVMDispatchFactory.h b/Engine/Plugins/Runtime/RigVM/Source/RigVM/Public/RigVMCore/RigVMDispatchFactory.h
--- a/Engine/Plugins/Runtime/RigVM/Source/RigVM/Public/RigVMCore/RigVMDispatchFactory.h
+++ b/Engine/Plugins/Runtime/RigVM/Source/RigVM/Public/RigVMCore/RigVMDispatchFactory.h
@@ -147,10 +147,10 @@ public:
 	virtual FRigVMStructUpgradeInfo GetUpgradeInfo(const FRigVMTemplateTypeMap& InTypes, const FRigVMDispatchContext& InContext) const { return FRigVMStructUpgradeInfo(); }
 
 	// returns the dispatch function for a given type set
-	RIGVM_API FRigVMFunctionPtr GetDispatchFunction(const FRigVMTemplateTypeMap& InTypes) const;
+	RIGVM_API FRigVMFunctionPtr GetDispatchFunction(const FRigVMTemplateTypeMap& InTypes, bool bForceCreation = false) const;
 
 	// builds and returns the template
-	RIGVM_API const FRigVMTemplate* GetTemplate() const;
+	RIGVM_API const FRigVMTemplate* GetTemplate(bool bIsDispatchingTemplate = false) const;
 
 	// returns the name of the factory template
 	RIGVM_API FName GetTemplateNotation() const;
@@ -222,6 +222,7 @@ protected:
 
 	UScriptStruct* FactoryScriptStruct;
 	mutable const FRigVMTemplate* CachedTemplate;
+	static FCriticalSection GetTemplateMutex;
 	friend struct FRigVMTemplate;
 	friend struct FRigVMRegistry;
 };
diff --git a/Engine/Plugins/Runtime/RigVM/Source/RigVM/Public/RigVMCore/RigVMRegistry.h b/Engine/Plugins/Runtime/RigVM/Source/RigVM/Public/RigVMCore/RigVMRegistry.h
--- a/Engine/Plugins/Runtime/RigVM/Source/RigVM/Public/RigVMCore/RigVMRegistry.h
+++ b/Engine/Plugins/Runtime/RigVM/Source/RigVM/Public/RigVMCore/RigVMRegistry.h
@@ -220,7 +220,8 @@ public:
 	const FRigVMTemplate* GetOrAddTemplateFromArguments(
 		const FName& InName,
 		const TArray<FRigVMTemplateArgument>& InArguments,
-		const FRigVMTemplateDelegates& InDelegates);
+		const FRigVMTemplateDelegates& InDelegates,
+		bool bIsDispatchingTemplate = false);
 
 	// Returns a dispatch factory given its name (or nullptr)
 	FRigVMDispatchFactory* FindDispatchFactory(const FName& InFactoryName) const;
@@ -355,7 +356,16 @@ private:
 	
 	// Notifies other system that types have been added/removed, and template permutations have been updated
 	FOnRigVMRegistryChanged OnRigVMRegistryChangedDelegate;
-	
+
+	static FCriticalSection RefreshTypesMutex;
+	static FCriticalSection RegisterFunctionMutex;
+	static FCriticalSection RegisterTemplateMutex;
+	static FCriticalSection RegisterFactoryMutex;
+	static FCriticalSection FindFunctionMutex;
+	static FCriticalSection FindTemplateMutex;
+	static FCriticalSection FindFactoryMutex;
+	static FCriticalSection GetPermutationMutex;
+
 	friend struct FRigVMStruct;
 	friend struct FRigVMTemplate;
 	friend struct FRigVMTemplateArgument;
