RigVM: Ensure thread safety for FRigVMRegistry::FindFunction

diff --git a/Engine/Plugins/Animation/ControlRig/Source/ControlRig/Private/ControlRigBlueprintGeneratedClass.cpp b/Engine/Plugins/Animation/ControlRig/Source/ControlRig/Private/ControlRigBlueprintGeneratedClass.cpp
--- a/Engine/Plugins/Animation/ControlRig/Source/ControlRig/Private/ControlRigBlueprintGeneratedClass.cpp
+++ b/Engine/Plugins/Animation/ControlRig/Source/ControlRig/Private/ControlRigBlueprintGeneratedClass.cpp
@@ -28,7 +28,11 @@ void UControlRigBlueprintGeneratedClass::Serialize(FArchive& Ar)
 		return;
 	}
 
-	URigVM* VM = NewObject<URigVM>(GetTransientPackage());
+	// for debugging purposes we'll give this VM a name that's useful.
+	static TAtomic<uint32> NumVMs{ 0 };
+	static constexpr TCHAR Format[] = TEXT("%s_VM_%zu");
+	const FString VMDebugName = FString::Printf(Format, *GetName(), uint32(++NumVMs));
+	URigVM* VM = NewObject<URigVM>(GetTransientPackage(), *VMDebugName);
 
 	if (UControlRig* CDO = Cast<UControlRig>(GetDefaultObject(true)))
 	{
diff --git a/Engine/Plugins/Runtime/RigVM/Source/RigVM/Private/RigVMCore/RigVM.cpp b/Engine/Plugins/Runtime/RigVM/Source/RigVM/Private/RigVMCore/RigVM.cpp
--- a/Engine/Plugins/Runtime/RigVM/Source/RigVM/Private/RigVMCore/RigVM.cpp
+++ b/Engine/Plugins/Runtime/RigVM/Source/RigVM/Private/RigVMCore/RigVM.cpp
@@ -153,6 +153,8 @@ void URigVM::Save(FArchive& Ar)
 	// by the RigVMCompiler.
 	if(!Ar.IsIgnoringArchetypeRef())
 	{
+		ResolveFunctionsIfRequired();
+	
 		Ar << ExternalPropertyPathDescriptions;
 		Ar << FunctionNamesStorage;
 		Ar << ByteCodeStorage;
@@ -223,6 +225,13 @@ void URigVM::Load(FArchive& Ar)
 		Ar << ByteCodeStorage;
 		Ar << Parameters;
 	}
+
+	// ensure to load the required functions
+	if(!ResolveFunctionsIfRequired())
+	{
+		Reset();
+		return;
+	}
 }
 
 void URigVM::PostLoad()
@@ -905,8 +914,12 @@ FRigVMParameter URigVM::GetParameterByName(const FName& InParameterName)
 	return FRigVMParameter();
 }
 
-void URigVM::ResolveFunctionsIfRequired()
+bool URigVM::ResolveFunctionsIfRequired()
 {
+	FScopeLock ResolveFunctionsScopeLock(&ResolveFunctionsMutex);
+	
+	bool bSuccess = true;
+	
 	if (GetFunctions().Num() != GetFunctionNames().Num())
 	{
 		GetFunctions().Reset();
@@ -914,20 +927,33 @@ void URigVM::ResolveFunctionsIfRequired()
 		GetFactories().Reset();
 		GetFactories().SetNumZeroed(GetFunctionNames().Num());
 
-		for (int32 FunctionIndex = 0; FunctionIndex < GetFunctionNames().Num(); FunctionIndex++)
+		TArray<FName>& FunctionNames = GetFunctionNames();
+		for (int32 FunctionIndex = 0; FunctionIndex < FunctionNames.Num(); FunctionIndex++)
 		{
-			if(const FRigVMFunction* Function = FRigVMRegistry::Get().FindFunction(*GetFunctionNames()[FunctionIndex].ToString()))
+			const FString FunctionNameString = FunctionNames[FunctionIndex].ToString();
+			if(const FRigVMFunction* Function = FRigVMRegistry::Get().FindFunction(*FunctionNameString))
 			{
 				GetFunctions()[FunctionIndex] = Function;
 				GetFactories()[FunctionIndex] = Function->Factory;
+				
+				// update the name in the function name list. the resolved function
+				// may differ since it may rely on a core redirect.
+				if(!FunctionNameString.Equals(Function->Name, ESearchCase::CaseSensitive))
+				{
+					FunctionNames[FunctionIndex] = *Function->Name;
+					UE_LOG(LogRigVM, Display, TEXT("Redirected function '%s' to '%s' for VM '%s'"), *FunctionNameString, *Function->Name, *GetPathName());
+				}
 			}
 			else
 			{
-				// We cannot recover from missing functions. 
-				UE_LOG(LogRigVM, Fatal, TEXT("No handler found for function '%s'"), *GetFunctionNames()[FunctionIndex].ToString());
+				// We cannot recover from missing functions.
+				UE_LOG(LogRigVM, Fatal, TEXT("No handler found for function '%s' for VM '%s'"), *FunctionNameString, *GetPathName());
+				bSuccess = false;
 			}
 		}
 	}
+
+	return bSuccess;
 }
 
 void URigVM::RefreshInstructionsIfRequired()
diff --git a/Engine/Plugins/Runtime/RigVM/Source/RigVM/Private/RigVMCore/RigVMDispatchFactory.cpp b/Engine/Plugins/Runtime/RigVM/Source/RigVM/Private/RigVMCore/RigVMDispatchFactory.cpp
--- a/Engine/Plugins/Runtime/RigVM/Source/RigVM/Private/RigVMCore/RigVMDispatchFactory.cpp
+++ b/Engine/Plugins/Runtime/RigVM/Source/RigVM/Private/RigVMCore/RigVMDispatchFactory.cpp
@@ -9,6 +9,7 @@
 
 const FString FRigVMDispatchFactory::DispatchPrefix = TEXT("DISPATCH_");
 const FString FRigVMDispatchFactory::TrueString = TEXT("True");
+FCriticalSection FRigVMDispatchFactory::GetTemplateMutex;
 
 FName FRigVMDispatchFactory::GetFactoryName() const
 {
@@ -195,13 +196,27 @@ TArray<FRigVMExecuteArgument> FRigVMDispatchFactory::GetExecuteArguments(const F
 	return Arguments;
 }
 
-FRigVMFunctionPtr FRigVMDispatchFactory::GetDispatchFunction(const FRigVMTemplateTypeMap& InTypes) const
+FRigVMFunctionPtr FRigVMDispatchFactory::GetOrCreateDispatchFunction(const FRigVMTemplateTypeMap& InTypes) const
 {
+	FScopeLock GetTemplateScopeLock(&FRigVMRegistry::GetDispatchFunctionMutex);
+
 	const FString PermutationName = GetPermutationNameImpl(InTypes);
 	if(const FRigVMFunction* ExistingFunction = FRigVMRegistry::Get().FindFunction(*PermutationName))
 	{
 		return ExistingFunction->FunctionPtr;
 	}
+	
+	return CreateDispatchFunction_NoLock(InTypes);
+}
+
+FRigVMFunctionPtr FRigVMDispatchFactory::CreateDispatchFunction(const FRigVMTemplateTypeMap& InTypes) const
+{
+	FScopeLock GetTemplateScopeLock(&FRigVMRegistry::GetDispatchFunctionMutex);
+	return CreateDispatchFunction_NoLock(InTypes);
+}
+
+FRigVMFunctionPtr FRigVMDispatchFactory::CreateDispatchFunction_NoLock(const FRigVMTemplateTypeMap& InTypes) const
+{
 	return GetDispatchFunctionImpl(InTypes);
 }
 
@@ -242,12 +257,9 @@ const FRigVMTemplate* FRigVMDispatchFactory::GetTemplate() const
 		return ThisFactory->GetTemplate();
 	}
 
-	static bool bIsDispatchingTemplate = false;
-	if(bIsDispatchingTemplate)
-	{
-			return nullptr;
-	}
-	TGuardValue<bool> ReEntryGuard(bIsDispatchingTemplate, true);	if(CachedTemplate)
+	FScopeLock GetTemplateScopeLock(&GetTemplateMutex);
+	
+	if(CachedTemplate)
 	{
 		return CachedTemplate;
 	}
@@ -285,10 +297,10 @@ const FRigVMTemplate* FRigVMDispatchFactory::GetTemplate() const
 	Delegates.RequestDispatchFunctionDelegate = FRigVMTemplate_RequestDispatchFunctionDelegate::CreateLambda(
 	[this](const FRigVMTemplate*, const FRigVMTemplateTypeMap& InTypes)
 	{
-		return GetDispatchFunction(InTypes);
+		return CreateDispatchFunction(InTypes);
 	});
 
-	CachedTemplate = Registry.GetOrAddTemplateFromArguments(GetFactoryName(), Arguments, Delegates); 
+	CachedTemplate = Registry.AddTemplateFromArguments(GetFactoryName(), Arguments, Delegates); 
 	return CachedTemplate;
 }
 
diff --git a/Engine/Plugins/Runtime/RigVM/Source/RigVM/Private/RigVMCore/RigVMRegistry.cpp b/Engine/Plugins/Runtime/RigVM/Source/RigVM/Private/RigVMCore/RigVMRegistry.cpp
--- a/Engine/Plugins/Runtime/RigVM/Source/RigVM/Private/RigVMCore/RigVMRegistry.cpp
+++ b/Engine/Plugins/Runtime/RigVM/Source/RigVM/Private/RigVMCore/RigVMRegistry.cpp
@@ -15,6 +15,15 @@
 #include "RigVMFunctions/RigVMDispatch_Core.h"
 
 const FName FRigVMRegistry::TemplateNameMetaName = TEXT("TemplateName");
+FCriticalSection FRigVMRegistry::RefreshTypesMutex;
+FCriticalSection FRigVMRegistry::RegisterFunctionMutex;
+FCriticalSection FRigVMRegistry::RegisterTemplateMutex;
+FCriticalSection FRigVMRegistry::RegisterFactoryMutex;
+FCriticalSection FRigVMRegistry::FindFunctionMutex;
+FCriticalSection FRigVMRegistry::FindTemplateMutex;
+FCriticalSection FRigVMRegistry::FindFactoryMutex;
+FCriticalSection FRigVMRegistry::GetDispatchFunctionMutex;
+FCriticalSection FRigVMRegistry::GetPermutationMutex;
 
 
 // When the object system has been completely loaded, load in all the engine types that we haven't registered already in InitializeIfNeeded 
@@ -202,6 +211,9 @@ uint32 FRigVMRegistry::GetHashForProperty(const FProperty* InProperty) const
 
 void FRigVMRegistry::InitializeIfNeeded()
 {
+	// we don't need to use a mutex here since this is called on
+	// the main thread during engine startup for the first time
+	
 	if(!Types.IsEmpty())
 	{
 		return;
@@ -305,6 +317,8 @@ void FRigVMRegistry::InitializeIfNeeded()
 
 void FRigVMRegistry::RefreshEngineTypes()
 {
+	FScopeLock RefreshTypesScopeLock(&RefreshTypesMutex);
+	
 	// Register all user-defined types that the engine knows about. Enumerating over the entire object hierarchy is
 	// slow, so we do it for structs, enums and dispatch factories in one shot.
 	TArray<UScriptStruct*> DispatchFactoriesToRegister;
@@ -393,6 +407,9 @@ void FRigVMRegistry::Reset()
 
 TRigVMTypeIndex FRigVMRegistry::FindOrAddType_Internal(const FRigVMTemplateArgumentType& InType, bool bForce)
 {
+	// we don't use a mutex here since by the time the engine relies on worker
+	// thread for execution or async loading all types will have been registered.
+	
 	TRigVMTypeIndex Index = GetTypeIndex(InType);
 	if(Index == INDEX_NONE)
 	{
@@ -1123,6 +1140,8 @@ bool FRigVMRegistry::IsAllowedType(const UClass* InClass)
 
 void FRigVMRegistry::Register(const TCHAR* InName, FRigVMFunctionPtr InFunctionPtr, UScriptStruct* InStruct, const TArray<FRigVMFunctionArgument>& InArguments)
 {
+	FScopeLock RegisterFunctionScopeLock(&RegisterFunctionMutex);
+	
 	if (FindFunction(InName) != nullptr)
 	{
 		return;
@@ -1203,6 +1222,8 @@ void FRigVMRegistry::Register(const TCHAR* InName, FRigVMFunctionPtr InFunctionP
 
 const FRigVMDispatchFactory* FRigVMRegistry::RegisterFactory(UScriptStruct* InFactoryStruct)
 {
+	FScopeLock RegisterFactoryScopeLock(&RegisterFactoryMutex);
+
 	check(InFactoryStruct);
 	check(InFactoryStruct != FRigVMDispatchFactory::StaticStruct());
 	check(InFactoryStruct->IsChildOf(FRigVMDispatchFactory::StaticStruct()));
@@ -1239,6 +1260,12 @@ const FRigVMDispatchFactory* FRigVMRegistry::RegisterFactory(UScriptStruct* InFa
 }
 
 const FRigVMFunction* FRigVMRegistry::FindFunction(const TCHAR* InName) const
+{
+	FScopeLock FindFunctionScopeLock(&FindFunctionMutex);
+	return FindFunction_NoLock(InName);
+}
+
+const FRigVMFunction* FRigVMRegistry::FindFunction_NoLock(const TCHAR* InName) const
 {
 	// Check first if the function is provided by internally registered rig units. 
 	if(const int32* FunctionIndexPtr = FunctionNameToIndex.Find(InName))
@@ -1246,12 +1273,6 @@ const FRigVMFunction* FRigVMRegistry::FindFunction(const TCHAR* InName) const
 		return &Functions[*FunctionIndexPtr];
 	}
 
-	static bool IsDispatchingFunction = false;
-	if(IsDispatchingFunction)
-	{
-		return nullptr;
-	}
-
 	// Otherwise ask the associated dispatch factory for a function matching this signature.
 	const FString NameString(InName);
 	FString StructOrFactoryName, SuffixString;
@@ -1268,7 +1289,6 @@ const FRigVMFunction* FRigVMRegistry::FindFunction(const TCHAR* InName) const
 					const int32 PermutationIndex = Template->FindPermutation(ArgumentTypes);
 					if(PermutationIndex != INDEX_NONE)
 					{
-						TGuardValue<bool> ReEntryGuard(IsDispatchingFunction, true);
 						return ((FRigVMTemplate*)Template)->GetOrCreatePermutation(PermutationIndex);
 					}
 				}
@@ -1276,7 +1296,6 @@ const FRigVMFunction* FRigVMRegistry::FindFunction(const TCHAR* InName) const
 		}
 	}
 
-#if WITH_EDITOR
 	// if we haven't been able to find the function - try to see if we can get the dispatch or rigvmstruct
 	// from a core redirect
 	if(!StructOrFactoryName.IsEmpty())
@@ -1287,8 +1306,7 @@ const FRigVMFunction* FRigVMRegistry::FindFunction(const TCHAR* InName) const
 		{
 			StructOrFactoryName = StructOrFactoryName.Mid(FRigVMDispatchFactory::DispatchPrefix.Len());
 		}
-
-		if(StructOrFactoryName.StartsWith(StructPrefix, ESearchCase::CaseSensitive))
+		else if(StructOrFactoryName.StartsWith(StructPrefix, ESearchCase::CaseSensitive))
 		{
 			StructOrFactoryName = StructOrFactoryName.Mid(StructPrefix.Len());
 		}
@@ -1300,12 +1318,15 @@ const FRigVMFunction* FRigVMRegistry::FindFunction(const TCHAR* InName) const
 			for(const FCoreRedirect* Redirect : Redirects)
 			{
 				FString NewStructOrFactoryName = Redirect->NewName.ObjectName.ToString();
-				NewStructOrFactoryName = StructPrefix + NewStructOrFactoryName;
 				if(bIsDispatchFactory)
 				{
 					NewStructOrFactoryName = FRigVMDispatchFactory::DispatchPrefix + NewStructOrFactoryName;
 				}
-				const FRigVMFunction* RedirectedFunction = FindFunction(*(NewStructOrFactoryName + TEXT("::") + SuffixString));
+				else
+				{
+					NewStructOrFactoryName = StructPrefix + NewStructOrFactoryName;
+				}
+				const FRigVMFunction* RedirectedFunction = FindFunction_NoLock(*(NewStructOrFactoryName + TEXT("::") + SuffixString));
 				if(RedirectedFunction)
 				{
 					FRigVMRegistry& MutableRegistry = FRigVMRegistry::Get();
@@ -1315,7 +1336,6 @@ const FRigVMFunction* FRigVMRegistry::FindFunction(const TCHAR* InName) const
 			}
 		}
 	}
-#endif
 	
 	return nullptr;
 }
@@ -1341,17 +1361,22 @@ const FRigVMTemplate* FRigVMRegistry::FindTemplate(const FName& InNotation, bool
 		return nullptr;
 	}
 
-	if(const int32* TemplateIndexPtr = TemplateNotationToIndex.Find(InNotation))
+	FScopeLock FindTemplateScopeLock(&FindTemplateMutex);
+	return FindTemplate_NoLock(InNotation, bIncludeDeprecated);
+}
+
+const FRigVMTemplate* FRigVMRegistry::FindTemplate_NoLock(const FName& InNotation, bool bIncludeDeprecated) const
+{
+	if (InNotation.IsNone())
 	{
-		return &Templates[*TemplateIndexPtr];
+		return nullptr;
 	}
 
-	static bool IsDispatchingFunction = false;
-	if(IsDispatchingFunction)
+	if(const int32* TemplateIndexPtr = TemplateNotationToIndex.Find(InNotation))
 	{
-		return nullptr;
+		return &Templates[*TemplateIndexPtr];
 	}
-	
+
 	const FString NotationString(InNotation.ToString());
 	FString FactoryName, ArgumentsString;
 	if(NotationString.Split(TEXT("("), &FactoryName, &ArgumentsString))
@@ -1378,7 +1403,6 @@ const FRigVMTemplate* FRigVMRegistry::FindTemplate(const FName& InNotation, bool
 		
 		if(const FRigVMDispatchFactory* Factory = FindDispatchFactory(*FactoryName))
 		{
-			TGuardValue<bool> ReEntryGuard(IsDispatchingFunction, true);
 			return Factory->GetTemplate();
 		}
 	}
@@ -1446,7 +1470,7 @@ const FRigVMTemplate* FRigVMRegistry::FindTemplate(const FName& InNotation, bool
 
 		if(SanitizedNotation != OriginalNotation)
 		{
-			return FindTemplate(*SanitizedNotation, bIncludeDeprecated);
+			return FindTemplate_NoLock(*SanitizedNotation, bIncludeDeprecated);
 		}
 	}
 
@@ -1460,12 +1484,30 @@ const TChunkedArray<FRigVMTemplate>& FRigVMRegistry::GetTemplates() const
 
 const FRigVMTemplate* FRigVMRegistry::GetOrAddTemplateFromArguments(const FName& InName, const TArray<FRigVMTemplateArgument>& InArguments, const FRigVMTemplateDelegates& InDelegates)
 {
+	FScopeLock RegisterTemplateScopeLock(&RegisterTemplateMutex);
+	
 	FRigVMTemplate Template(InName, InArguments, INDEX_NONE);
+
+	// avoid reentry in FindTemplate. try to find an existing
+	// template only if we are not yet in ::FindTemplate.
 	if(const FRigVMTemplate* ExistingTemplate = FindTemplate(Template.GetNotation()))
 	{
 		return ExistingTemplate;
 	}
 
+	return AddTemplateFromArguments_NoLock(InName, InArguments, InDelegates);
+}
+
+const FRigVMTemplate* FRigVMRegistry::AddTemplateFromArguments(const FName& InName, const TArray<FRigVMTemplateArgument>& InArguments, const FRigVMTemplateDelegates& InDelegates)
+{
+	FScopeLock RegisterTemplateScopeLock(&RegisterTemplateMutex);
+	return AddTemplateFromArguments_NoLock(InName, InArguments, InDelegates);
+}
+
+const FRigVMTemplate* FRigVMRegistry::AddTemplateFromArguments_NoLock(const FName& InName, const TArray<FRigVMTemplateArgument>& InArguments, const FRigVMTemplateDelegates& InDelegates)
+{
+	FRigVMTemplate Template(InName, InArguments, INDEX_NONE);
+	
 	// we only support to ask for templates here which provide singleton types
 	int32 NumPermutations = 1;
 	for(const FRigVMTemplateArgument& Argument : InArguments)
@@ -1579,7 +1621,13 @@ const FRigVMTemplate* FRigVMRegistry::GetOrAddTemplateFromArguments(const FName&
 	return &Templates[Index];
 }
 
-FRigVMDispatchFactory* FRigVMRegistry::FindDispatchFactory(const FName& InFactoryName) const 
+FRigVMDispatchFactory* FRigVMRegistry::FindDispatchFactory(const FName& InFactoryName) const
+{
+	FScopeLock FindFactoryScopeLock(&FindFactoryMutex);
+	return FindDispatchFactory_NoLock(InFactoryName);
+}
+
+FRigVMDispatchFactory* FRigVMRegistry::FindDispatchFactory_NoLock(const FName& InFactoryName) const
 {
 	FRigVMDispatchFactory* const* FactoryPtr = Factories.FindByPredicate([InFactoryName](const FRigVMDispatchFactory* Factory) -> bool
 	{
diff --git a/Engine/Plugins/Runtime/RigVM/Source/RigVM/Private/RigVMCore/RigVMTemplate.cpp b/Engine/Plugins/Runtime/RigVM/Source/RigVM/Private/RigVMCore/RigVMTemplate.cpp
--- a/Engine/Plugins/Runtime/RigVM/Source/RigVM/Private/RigVMCore/RigVMTemplate.cpp
+++ b/Engine/Plugins/Runtime/RigVM/Source/RigVM/Private/RigVMCore/RigVMTemplate.cpp
@@ -1158,6 +1158,12 @@ const FRigVMFunction* FRigVMTemplate::GetPrimaryPermutation() const
 }
 
 const FRigVMFunction* FRigVMTemplate::GetPermutation(int32 InIndex) const
+{
+	FScopeLock FindPermutationScopeLock(&FRigVMRegistry::GetPermutationMutex);
+	return GetPermutation_NoLock(InIndex);
+}
+
+const FRigVMFunction* FRigVMTemplate::GetPermutation_NoLock(int32 InIndex) const
 {
 	const FRigVMRegistry& Registry = FRigVMRegistry::Get();
 	const int32 FunctionIndex = Permutations[InIndex];
@@ -1170,13 +1176,16 @@ const FRigVMFunction* FRigVMTemplate::GetPermutation(int32 InIndex) const
 
 const FRigVMFunction* FRigVMTemplate::GetOrCreatePermutation(int32 InIndex)
 {
-	if(const FRigVMFunction* Function = GetPermutation(InIndex))
+	FScopeLock FindPermutationScopeLock(&FRigVMRegistry::GetPermutationMutex);
+
+	if(const FRigVMFunction* Function = GetPermutation_NoLock(InIndex))
 	{
 		return Function;
 	}
 	
 	if(Permutations[InIndex] == INDEX_NONE && UsesDispatch())
 	{
+		FScopeLock RegisterFunctionScopeLock(&FRigVMRegistry::RegisterFunctionMutex);
 		FRigVMRegistry& Registry = FRigVMRegistry::Get();
 		
 		FTypeMap Types;
diff --git a/Engine/Plugins/Runtime/RigVM/Source/RigVM/Public/RigVMCore/RigVM.h b/Engine/Plugins/Runtime/RigVM/Source/RigVM/Public/RigVMCore/RigVM.h
--- a/Engine/Plugins/Runtime/RigVM/Source/RigVM/Public/RigVMCore/RigVM.h
+++ b/Engine/Plugins/Runtime/RigVM/Source/RigVM/Public/RigVMCore/RigVM.h
@@ -730,7 +730,7 @@ public:
 
 private:
 
-	void ResolveFunctionsIfRequired();
+	bool ResolveFunctionsIfRequired();
 	void RefreshInstructionsIfRequired();
 public:
 	void InvalidateCachedMemory();
@@ -904,6 +904,8 @@ protected:
 	void StartProfiling(FRigVMExtendedExecuteContext& Context);
 	void StopProfiling(FRigVMExtendedExecuteContext& Context);
 
+	FCriticalSection ResolveFunctionsMutex;
+
 	friend class URigVMCompiler;
 	friend struct FRigVMCompilerWorkData;
 	friend struct FRigVMCodeGenerator;
diff --git a/Engine/Plugins/Runtime/RigVM/Source/RigVM/Public/RigVMCore/RigVMDispatchFactory.h b/Engine/Plugins/Runtime/RigVM/Source/RigVM/Public/RigVMCore/RigVMDispatchFactory.h
--- a/Engine/Plugins/Runtime/RigVM/Source/RigVM/Public/RigVMCore/RigVMDispatchFactory.h
+++ b/Engine/Plugins/Runtime/RigVM/Source/RigVM/Public/RigVMCore/RigVMDispatchFactory.h
@@ -147,7 +147,7 @@ public:
 	virtual FRigVMStructUpgradeInfo GetUpgradeInfo(const FRigVMTemplateTypeMap& InTypes, const FRigVMDispatchContext& InContext) const { return FRigVMStructUpgradeInfo(); }
 
 	// returns the dispatch function for a given type set
-	RIGVM_API FRigVMFunctionPtr GetDispatchFunction(const FRigVMTemplateTypeMap& InTypes) const;
+	RIGVM_API FRigVMFunctionPtr GetOrCreateDispatchFunction(const FRigVMTemplateTypeMap& InTypes) const;
 
 	// builds and returns the template
 	RIGVM_API const FRigVMTemplate* GetTemplate() const;
@@ -166,6 +166,9 @@ protected:
 	// returns the name of the permutation for a given set of types
 	RIGVM_API FString GetPermutationNameImpl(const FRigVMTemplateTypeMap& InTypes) const;
 
+	RIGVM_API FRigVMFunctionPtr CreateDispatchFunction(const FRigVMTemplateTypeMap& InTypes) const;
+	FRigVMFunctionPtr CreateDispatchFunction_NoLock(const FRigVMTemplateTypeMap& InTypes) const;
+
 	virtual FRigVMFunctionPtr GetDispatchFunctionImpl(const FRigVMTemplateTypeMap& InTypes) const { return nullptr; }
 
 	virtual TArray<FRigVMExecuteArgument> GetExecuteArguments_Impl(const FRigVMDispatchContext& InContext) const { return TArray<FRigVMExecuteArgument>(); }
@@ -222,6 +225,7 @@ protected:
 
 	UScriptStruct* FactoryScriptStruct;
 	mutable const FRigVMTemplate* CachedTemplate;
+	static FCriticalSection GetTemplateMutex;
 	friend struct FRigVMTemplate;
 	friend struct FRigVMRegistry;
 };
diff --git a/Engine/Plugins/Runtime/RigVM/Source/RigVM/Public/RigVMCore/RigVMRegistry.h b/Engine/Plugins/Runtime/RigVM/Source/RigVM/Public/RigVMCore/RigVMRegistry.h
--- a/Engine/Plugins/Runtime/RigVM/Source/RigVM/Public/RigVMCore/RigVMRegistry.h
+++ b/Engine/Plugins/Runtime/RigVM/Source/RigVM/Public/RigVMCore/RigVMRegistry.h
@@ -222,6 +222,12 @@ public:
 		const TArray<FRigVMTemplateArgument>& InArguments,
 		const FRigVMTemplateDelegates& InDelegates);
 
+	// Adds a new template given its arguments
+	const FRigVMTemplate* AddTemplateFromArguments(
+		const FName& InName,
+		const TArray<FRigVMTemplateArgument>& InArguments,
+		const FRigVMTemplateDelegates& InDelegates);
+
 	// Returns a dispatch factory given its name (or nullptr)
 	FRigVMDispatchFactory* FindDispatchFactory(const FName& InFactoryName) const;
 
@@ -317,6 +323,15 @@ private:
 	
 	void RemoveTypeInCategory(FRigVMTemplateArgument::ETypeCategory InCategory, TRigVMTypeIndex InTypeIndex);
 
+	const FRigVMFunction* FindFunction_NoLock(const TCHAR* InName) const;
+	const FRigVMTemplate* FindTemplate_NoLock(const FName& InNotation, bool bIncludeDeprecated) const;
+	FRigVMDispatchFactory* FindDispatchFactory_NoLock(const FName& InFactoryName) const;
+
+	const FRigVMTemplate* AddTemplateFromArguments_NoLock(
+		const FName& InName,
+		const TArray<FRigVMTemplateArgument>& InArguments,
+		const FRigVMTemplateDelegates& InDelegates);
+
 	// memory for all (known) types
 	TArray<FTypeInfo> Types;
 	TMap<FRigVMTemplateArgumentType, TRigVMTypeIndex> TypeToIndex;
@@ -355,8 +370,19 @@ private:
 	
 	// Notifies other system that types have been added/removed, and template permutations have been updated
 	FOnRigVMRegistryChanged OnRigVMRegistryChangedDelegate;
-	
+
+	static FCriticalSection RefreshTypesMutex;
+	static FCriticalSection RegisterFunctionMutex;
+	static FCriticalSection RegisterTemplateMutex;
+	static FCriticalSection RegisterFactoryMutex;
+	static FCriticalSection FindFunctionMutex;
+	static FCriticalSection FindTemplateMutex;
+	static FCriticalSection FindFactoryMutex;
+	static FCriticalSection GetDispatchFunctionMutex;
+	static FCriticalSection GetPermutationMutex;
+
 	friend struct FRigVMStruct;
 	friend struct FRigVMTemplate;
 	friend struct FRigVMTemplateArgument;
+	friend struct FRigVMDispatchFactory;
 };
diff --git a/Engine/Plugins/Runtime/RigVM/Source/RigVM/Public/RigVMCore/RigVMTemplate.h b/Engine/Plugins/Runtime/RigVM/Source/RigVM/Public/RigVMCore/RigVMTemplate.h
--- a/Engine/Plugins/Runtime/RigVM/Source/RigVM/Public/RigVMCore/RigVMTemplate.h
+++ b/Engine/Plugins/Runtime/RigVM/Source/RigVM/Public/RigVMCore/RigVMTemplate.h
@@ -451,6 +451,8 @@ private:
 	void InvalidateHash() { Hash = UINT32_MAX; }
 	const TArray<FRigVMExecuteArgument>& GetExecuteArguments(const FRigVMDispatchContext& InContext) const;
 
+	const FRigVMFunction* GetPermutation_NoLock(int32 InIndex) const;
+
 	int32 Index;
 	FName Notation;
 	TArray<FRigVMTemplateArgument> Arguments;
diff --git a/Engine/Plugins/Runtime/RigVM/Source/RigVMDeveloper/Private/RigVMCompiler/RigVMCompiler.cpp b/Engine/Plugins/Runtime/RigVM/Source/RigVMDeveloper/Private/RigVMCompiler/RigVMCompiler.cpp
--- a/Engine/Plugins/Runtime/RigVM/Source/RigVMDeveloper/Private/RigVMCompiler/RigVMCompiler.cpp
+++ b/Engine/Plugins/Runtime/RigVM/Source/RigVMDeveloper/Private/RigVMCompiler/RigVMCompiler.cpp
@@ -1144,6 +1144,9 @@ bool URigVMCompiler::Compile(TArray<URigVMGraph*> InGraphs, URigVMController* In
 		OutFunctionCompilationData->Hash = GetTypeHash(OutFunctionCompilationData);
 	}
 
+	// make sure all functions are known and resolved now.
+	WorkData.VM->ResolveFunctionsIfRequired();
+
 	if (!CurrentCompilationFunction)
 	{
 		CompileTimer.Stop();
diff --git a/Engine/Plugins/Runtime/RigVM/Source/RigVMDeveloper/Private/RigVMModel/RigVMController.cpp b/Engine/Plugins/Runtime/RigVM/Source/RigVMDeveloper/Private/RigVMModel/RigVMController.cpp
--- a/Engine/Plugins/Runtime/RigVM/Source/RigVMDeveloper/Private/RigVMModel/RigVMController.cpp
+++ b/Engine/Plugins/Runtime/RigVM/Source/RigVMDeveloper/Private/RigVMModel/RigVMController.cpp
@@ -16269,7 +16269,7 @@ bool URigVMController::FullyResolveTemplateNode(URigVMTemplateNode* InNode, int3
 		if(Factory)
 		{
 			FRigVMTemplateTypeMap TypeMap = InNode->GetTemplate()->GetTypesForPermutation(InNode->ResolvedPermutation);
-			const FRigVMFunctionPtr DispatchFunction = Factory->GetDispatchFunction(TypeMap);
+			const FRigVMFunctionPtr DispatchFunction = Factory->GetOrCreateDispatchFunction(TypeMap);
 			const FRigVMFunction* ResolvedFunction = const_cast<FRigVMTemplate*>(Template)->GetOrCreatePermutation(InNode->ResolvedPermutation);
 			check(DispatchFunction);
 			check(ResolvedFunction);
