Force integrate from old to new location to properly backout files from CL# 25363516
#rnx

diff --git a/Engine/Plugins/Runtime/RigVM/Source/RigVMDeveloper/Private/RigVMCompiler/RigVMCompiler.cpp b/Engine/Plugins/Runtime/RigVM/Source/RigVMDeveloper/Private/RigVMCompiler/RigVMCompiler.cpp
--- a/Engine/Plugins/Runtime/RigVM/Source/RigVMDeveloper/Private/RigVMCompiler/RigVMCompiler.cpp
+++ b/Engine/Plugins/Runtime/RigVM/Source/RigVMDeveloper/Private/RigVMCompiler/RigVMCompiler.cpp
@@ -1144,9 +1144,6 @@ bool URigVMCompiler::Compile(TArray<URigVMGraph*> InGraphs, URigVMController* In
 		OutFunctionCompilationData->Hash = GetTypeHash(OutFunctionCompilationData);
 	}
 
-	// make sure all functions are known and resolved now.
-	WorkData.VM->ResolveFunctionsIfRequired();
-
 	if (!CurrentCompilationFunction)
 	{
 		CompileTimer.Stop();
@@ -3455,4 +3452,3 @@ void URigVMCompiler::ReportError(const FString& InMessage)
 {
 	Settings.Report(EMessageSeverity::Error, nullptr, InMessage);
 }
-
diff --git a/Engine/Plugins/Runtime/RigVM/Source/RigVMDeveloper/Private/RigVMModel/RigVMController.cpp b/Engine/Plugins/Runtime/RigVM/Source/RigVMDeveloper/Private/RigVMModel/RigVMController.cpp
--- a/Engine/Plugins/Runtime/RigVM/Source/RigVMDeveloper/Private/RigVMModel/RigVMController.cpp
+++ b/Engine/Plugins/Runtime/RigVM/Source/RigVMDeveloper/Private/RigVMModel/RigVMController.cpp
@@ -16269,7 +16269,7 @@ bool URigVMController::FullyResolveTemplateNode(URigVMTemplateNode* InNode, int3
 		if(Factory)
 		{
 			FRigVMTemplateTypeMap TypeMap = InNode->GetTemplate()->GetTypesForPermutation(InNode->ResolvedPermutation);
-			const FRigVMFunctionPtr DispatchFunction = Factory->GetOrCreateDispatchFunction(TypeMap);
+			const FRigVMFunctionPtr DispatchFunction = Factory->GetDispatchFunction(TypeMap);
 			const FRigVMFunction* ResolvedFunction = const_cast<FRigVMTemplate*>(Template)->GetOrCreatePermutation(InNode->ResolvedPermutation);
 			check(DispatchFunction);
 			check(ResolvedFunction);
