RigVM: Avoid warning about already linked pins when reattaching links

diff --git a/Engine/Plugins/Runtime/RigVM/Source/RigVMDeveloper/Private/RigVMModel/RigVMController.cpp b/Engine/Plugins/Runtime/RigVM/Source/RigVMDeveloper/Private/RigVMModel/RigVMController.cpp
--- a/Engine/Plugins/Runtime/RigVM/Source/RigVMDeveloper/Private/RigVMModel/RigVMController.cpp
+++ b/Engine/Plugins/Runtime/RigVM/Source/RigVMDeveloper/Private/RigVMModel/RigVMController.cpp
@@ -18398,6 +18398,13 @@ bool URigVMController::RestoreLinkedPaths(
 				bSingleLinkSuccess = false;
 			}
 		}
+		else
+		{
+			if (ExistingLink)
+			{
+				GetGraph()->DetachedLinks.Remove(ExistingLink);
+			}
+		}
 
 		if(bSingleLinkSuccess)
 		{
