FRigVMClient: ensure that objects are not rooted before destroying them
    - follow-up for 25005143

diff --git a/Engine/Plugins/Runtime/RigVM/Source/RigVMDeveloper/Private/RigVMModel/RigVMClient.cpp b/Engine/Plugins/Runtime/RigVM/Source/RigVMDeveloper/Private/RigVMModel/RigVMClient.cpp
--- a/Engine/Plugins/Runtime/RigVM/Source/RigVMDeveloper/Private/RigVMModel/RigVMClient.cpp
+++ b/Engine/Plugins/Runtime/RigVM/Source/RigVMDeveloper/Private/RigVMModel/RigVMClient.cpp
@@ -752,7 +752,10 @@ void FRigVMClient::DestroyObject(UObject* InObject)
 		static constexpr TCHAR ObjectNameFormat[] = TEXT("RigVMClient_ObjectToBeDestroyed_%d");
 		const FString NewObjectName = FString::Printf(ObjectNameFormat, ObjectIndexToBeDestroyed++);
 		InObject->Rename(*NewObjectName, GetTransientPackage(), REN_ForceNoResetLoaders | REN_DontCreateRedirectors);
-		InObject->MarkAsGarbage();
+		if (!InObject->IsRooted())
+		{
+			InObject->MarkAsGarbage();
+		}
 	}
 }
 
