RigVM: Prevent injection nodes from being deleting while patching array nodes

diff --git a/Engine/Plugins/Runtime/RigVM/Source/RigVMDeveloper/Private/RigVMModel/RigVMController.cpp b/Engine/Plugins/Runtime/RigVM/Source/RigVMDeveloper/Private/RigVMModel/RigVMController.cpp
--- a/Engine/Plugins/Runtime/RigVM/Source/RigVMDeveloper/Private/RigVMModel/RigVMController.cpp
+++ b/Engine/Plugins/Runtime/RigVM/Source/RigVMDeveloper/Private/RigVMModel/RigVMController.cpp
@@ -15874,6 +15874,10 @@ void URigVMController::ApplyPinState(URigVMPin* InPin, const FPinState& InPinSta
 				if (FormerlyInjectedNode->IsInjected())
 				{
 					URigVMInjectionInfo* Injection = Cast<URigVMInjectionInfo>(FormerlyInjectedNode->GetOuter());
+					if(URigVMPin* OuterPin = Cast<URigVMPin>(Injection->GetOuter()))
+					{
+						OuterPin->InjectionInfos.Remove(Injection);
+					}
 					RenameObject(FormerlyInjectedNode, nullptr, InPin->GetGraph());
 					DestroyObject(Injection);
 				}
@@ -19219,13 +19223,19 @@ FRigVMClientPatchResult URigVMController::PatchArrayNodesOnLoad()
 			Result.bChangedContent = true;
 
 			FastBreakLinkedPaths(LinkedPaths);
-			RemoveNode(ArrayNode, false, false);
 
+			// Before removing the node (which will remove all injected nodes), lets rename it
+			// create the new node, and apply pin states to transfer ownership of injected nodes
+			// to the new node. Only after everything has been moved, we can safely remove the
+			// old node.
+			const FString DeletedName = GetValidNodeName(FString::Printf(TEXT("%s_Deleted"), *NodeName));
+			RenameNode(ArrayNode, *DeletedName, false);
+			
 			URigVMNode* NewNode = AddArrayNode(OpCode, CPPType, CPPTypeObject, NodePosition, NodeName, false, false, true);
-
-			Result.AddedNodes.Add(NewNode);
-
 			ApplyPinStates(NewNode, PinStates, {}, false);
+			Result.AddedNodes.Add(NewNode);
+			
+			RemoveNode(ArrayNode, false, false);
 
 			FRestoreLinkedPathSettings Settings;
 			RestoreLinkedPaths(LinkedPaths, Settings);
