RigVM: Fix reroute refactor loading issues

diff --git a/Engine/Plugins/Runtime/RigVM/Source/RigVMDeveloper/Private/RigVMModel/RigVMController.cpp b/Engine/Plugins/Runtime/RigVM/Source/RigVMDeveloper/Private/RigVMModel/RigVMController.cpp
--- a/Engine/Plugins/Runtime/RigVM/Source/RigVMDeveloper/Private/RigVMModel/RigVMController.cpp
+++ b/Engine/Plugins/Runtime/RigVM/Source/RigVMDeveloper/Private/RigVMModel/RigVMController.cpp
@@ -18292,6 +18292,7 @@ bool URigVMController::RestoreLinkedPaths(
 
 	bool bSuccess = true;
 	bool bAffectedAnyTemplateNode = false;
+	URigVMGraph* Graph = GetGraph();
 	
 	for(const FLinkedPath& LinkedPath : RemappedLinkedPaths)
 	{
@@ -18306,14 +18307,14 @@ bool URigVMController::RestoreLinkedPaths(
 		
 		const FString SourcePath = LinkedPath.SourcePinPath;
 		const FString TargetPath = LinkedPath.TargetPinPath;
-		URigVMPin* SourcePin = GetGraph()->FindPin(SourcePath);
-		URigVMPin* TargetPin = GetGraph()->FindPin(TargetPath);
+		URigVMPin* SourcePin = Graph->FindPin(SourcePath);
+		URigVMPin* TargetPin = Graph->FindPin(TargetPath);
 
 		URigVMLink* ExistingLink = FindLinkFromPinPathRepresentation(OriginalPinPathRepresentation, true);
-		auto OnFailedToRestoreLink = [this, ExistingLink, LinkedPath, &bSuccess](const FString& InFailureReason)
+		auto OnFailedToRestoreLink = [this, Graph, ExistingLink, LinkedPath, &bSuccess](const FString& InFailureReason)
 		{
-			const URigVMPin* SourcePin = LinkedPath.GetSourcePin(GetGraph());
-			const URigVMPin* TargetPin = LinkedPath.GetTargetPin(GetGraph());
+			const URigVMPin* SourcePin = LinkedPath.GetSourcePin(Graph);
+			const URigVMPin* TargetPin = LinkedPath.GetTargetPin(Graph);
 
 			bool bNotify = true;
 
@@ -18342,7 +18343,7 @@ bool URigVMController::RestoreLinkedPaths(
 				{
 					Notify(ERigVMGraphNotifType::LinkRemoved, ExistingLink);
 				}
-				GetGraph()->DetachedLinks.Remove(ExistingLink);
+				Graph->DetachedLinks.Remove(ExistingLink);
 				DestroyObject(ExistingLink);
 			}
 
@@ -18406,7 +18407,7 @@ bool URigVMController::RestoreLinkedPaths(
 		{
 			if (ExistingLink)
 			{
-				GetGraph()->DetachedLinks.Remove(ExistingLink);
+				Graph->DetachedLinks.Remove(ExistingLink);
 			}
 		}
 
@@ -18423,6 +18424,21 @@ bool URigVMController::RestoreLinkedPaths(
 		}
 	}
 
+	for (int32 PathIndex=0; PathIndex<InLinkedPaths.Num(); ++PathIndex)
+	{
+		if ((InLinkedPaths[PathIndex].SourcePinPath != RemappedLinkedPaths[PathIndex].SourcePinPath) ||
+			(InLinkedPaths[PathIndex].TargetPinPath != RemappedLinkedPaths[PathIndex].TargetPinPath))
+		{
+			const FString& SourcePath = InLinkedPaths[PathIndex].SourcePinPath;
+			const FString& TargetPath = InLinkedPaths[PathIndex].TargetPinPath;
+			Graph->DetachedLinks = Graph->DetachedLinks.FilterByPredicate(
+				[SourcePath, TargetPath](const URigVMLink* Link) -> bool
+				{
+					return !Link->GetSourcePinPath().Equals(SourcePath, ESearchCase::CaseSensitive) ||
+						!Link->GetTargetPinPath().Equals(TargetPath, ESearchCase::CaseSensitive);
+				});
+		}
+	}
 	return bSuccess;
 }
 
@@ -18624,14 +18640,12 @@ URigVMNode* URigVMController::ConvertRerouteNodeToDispatch(URigVMRerouteNode* In
 	const FVector2D NodePosition = InRerouteNode->GetPosition();
 	const TRigVMTypeIndex TypeIndex = ValuePin->GetTypeIndex();
 
-	if(!RemoveNode(InRerouteNode, bSetupUndoRedo, bPrintPythonCommand))
-	{
-		if(bSetupUndoRedo)
-		{
-			CancelUndoBracket();
-		}
-		return nullptr;
-	}
+	// Before removing the node (which will remove all injected nodes), lets rename it
+	// create the new node, and apply pin states to transfer ownership of injected nodes
+	// to the new node. Only after everything has been moved, we can safely remove the
+	// old node.
+	const FString DeletedName = GetValidNodeName(FString::Printf(TEXT("%s_Deleted"), *NodeName));
+	RenameNode(InRerouteNode, *DeletedName, false);
 
 	URigVMNode* NewNode = AddTemplateNode(InTemplateNotation, NodePosition, NodeName, bSetupUndoRedo, bPrintPythonCommand);
 	if(NewNode)
@@ -18677,6 +18691,15 @@ URigVMNode* URigVMController::ConvertRerouteNodeToDispatch(URigVMRerouteNode* In
 			})
 		);
 		RestoreLinkedPaths(LinkedPaths, RestoreSettings, bSetupUndoRedo);
+
+		if(!RemoveNode(InRerouteNode, bSetupUndoRedo, bPrintPythonCommand))
+		{
+			if(bSetupUndoRedo)
+			{
+				CancelUndoBracket();
+			}
+			return nullptr;
+		}
 	}
 
 	if(bSetupUndoRedo)
@@ -18697,6 +18720,7 @@ FRigVMClientPatchResult URigVMController::PatchRerouteNodesOnLoad()
 		TArray<URigVMRerouteNode*> ReroutesToTurnIntoMakeStruct;
 		TArray<URigVMRerouteNode*> ReroutesToTurnIntoBreakStruct;
 		TArray<URigVMRerouteNode*> ReroutesToTurnIntoMakeAndBreakStruct;
+		TArray<URigVMRerouteNode*> ReroutesToTurnIntoRerouteAndBreakStruct;
 		TArray<URigVMRerouteNode*> ReroutesToTurnIntoMakeArray;
 		TArray<URigVMRerouteNode*> ReroutesToRemove;
 
@@ -18722,8 +18746,8 @@ FRigVMClientPatchResult URigVMController::PatchRerouteNodesOnLoad()
 				const bool bOnlyTopLevelTargetLinks = (TopLevelTargetLinks > 0) && (SubPinTargetLinks == 0);
 				const bool bOnlySubPinSourceLinks = (TopLevelSourceLinks == 0) && (SubPinSourceLinks > 0);
 				const bool bOnlySubPinTargetLinks = (TopLevelTargetLinks == 0) && (SubPinTargetLinks > 0);
-				const bool bHasSourceLinks = bOnlyTopLevelSourceLinks || bOnlySubPinSourceLinks;
-				const bool bHasTargetLinks = bOnlyTopLevelTargetLinks || bOnlySubPinTargetLinks;
+				const bool bHasSourceLinks = (TopLevelSourceLinks + SubPinSourceLinks > 0);
+				const bool bHasTargetLinks = (TopLevelTargetLinks + SubPinTargetLinks > 0);
 
 				if(bHasSourceLinks && bHasTargetLinks)
 				{
@@ -18731,12 +18755,12 @@ FRigVMClientPatchResult URigVMController::PatchRerouteNodesOnLoad()
 					{
 						// this is a normal reroute and we can keep it
 					}
-					else if(bOnlyTopLevelSourceLinks) // && bOnlySubPinTargetLinks
+					else if(bOnlyTopLevelSourceLinks && bOnlySubPinTargetLinks)
 					{
 						check(ValuePin->IsStruct());
 						ReroutesToTurnIntoBreakStruct.Add(RerouteNode);
 					}
-					else if(bOnlyTopLevelTargetLinks) // && bOnlySubPinSourceLinks
+					else if(bOnlyTopLevelTargetLinks && bOnlySubPinSourceLinks)
 					{
 						if(ValuePin->IsArray())
 						{
@@ -18748,6 +18772,11 @@ FRigVMClientPatchResult URigVMController::PatchRerouteNodesOnLoad()
 							ReroutesToTurnIntoMakeStruct.Add(RerouteNode);
 						}
 					}
+					else if(bOnlyTopLevelSourceLinks && !bOnlyTopLevelTargetLinks && !bOnlySubPinTargetLinks)
+					{
+						check(ValuePin->IsStruct());
+						ReroutesToTurnIntoRerouteAndBreakStruct.Add(RerouteNode);
+					}
 					else // bOnlySubPinSourceLinks && bOnlySubPinTargetLinks
 					{
 						check(ValuePin->IsStruct());
@@ -18926,6 +18955,68 @@ FRigVMClientPatchResult URigVMController::PatchRerouteNodesOnLoad()
 			ReroutesToTurnIntoMakeAndBreakStruct.Reset();
 		}
 
+		// convert some reroutes to a pair of a reroute and a break struct node
+		if(ReroutesToTurnIntoRerouteAndBreakStruct.Num() > 0)
+		{
+			const FName BreakTemplateNotation = FRigVMDispatch_BreakStruct().GetTemplateNotation();
+			for(URigVMRerouteNode* RerouteNode : ReroutesToTurnIntoRerouteAndBreakStruct)
+			{
+				const URigVMPin* ValuePin = RerouteNode->FindPin(URigVMRerouteNode::ValueName);
+				check(ValuePin);
+				const TRigVMTypeIndex& TypeIndex = ValuePin->GetTypeIndex();
+				const FString NodeName = RerouteNode->GetName();
+				const TArray<FLinkedPath> LinkedPaths = GetLinkedPaths(RerouteNode);
+				FastBreakLinkedPaths(LinkedPaths);
+
+				if(const URigVMNode* BreakStructNode = AddTemplateNode(BreakTemplateNotation, RerouteNode->GetPosition() + FVector2D(150, 0), FString(), false, false))
+				{
+					URigVMPin* PinToResolve = BreakStructNode->FindPin(FRigVMDispatch_BreakStruct::StructName.ToString());
+					check(PinToResolve);
+					verify(ResolveWildCardPin(PinToResolve, TypeIndex, false, false));
+
+					AddLink(
+						RerouteNode->FindPin(URigVMRerouteNode::ValueName),
+						BreakStructNode->FindPin(FRigVMDispatch_BreakStruct::StructName.ToString()),
+						false);
+
+					const FString NodeNamePrefix = URigVMPin::JoinPinPath({NodeName, FString()});
+					FRestoreLinkedPathSettings RestoreSettings;
+					RestoreSettings.RemapDelegates.Add(NodeName,
+						FRigVMController_PinPathRemapDelegate::CreateLambda(
+							[NodeNamePrefix, RerouteNode, BreakStructNode](const FString& InPinPath, bool bIsInput) -> FString
+						{
+							if(InPinPath.StartsWith(NodeNamePrefix, ESearchCase::CaseSensitive))
+							{
+								const URigVMNode* NewNode = bIsInput ? RerouteNode : BreakStructNode;
+								const FString RemainingPinPath = InPinPath.Mid(NodeNamePrefix.Len());
+								FString Left, Right;
+								if(!URigVMPin::SplitPinPathAtStart(RemainingPinPath, Left, Right))
+								{
+									Left = RemainingPinPath;
+								}
+								else
+								{
+									Right = URigVMPin::JoinPinPath({FString(), Right});
+								}
+
+								static const FString& ValueName = URigVMRerouteNode::ValueName;
+								static const FString ElementsName = FRigVMDispatch_BreakStruct::ElementsName.ToString();
+
+								if(Left.Equals(ValueName, ESearchCase::CaseSensitive) && !Right.IsEmpty())
+								{
+									const FString NewNodeNamePrefix = URigVMPin::JoinPinPath({NewNode->GetName(), FString()});
+									return NewNodeNamePrefix + ElementsName + Right;
+								}
+							}
+							return InPinPath;
+						})
+					);
+					RestoreLinkedPaths(LinkedPaths, RestoreSettings);
+				}
+			}
+			ReroutesToTurnIntoRerouteAndBreakStruct.Reset();
+		}
+
 		// convert some reroutes to make array nodes
 		if(ReroutesToTurnIntoMakeArray.Num() > 0)
 		{
@@ -18980,8 +19071,14 @@ FRigVMClientPatchResult URigVMController::PatchUnitNodesOnLoad()
 			Result.bChangedContent = true;
 
 			FastBreakLinkedPaths(LinkedPaths);
-			RemoveNode(UnitNode, false, false);
 
+			// Before removing the node (which will remove all injected nodes), lets rename it
+			// create the new node, and apply pin states to transfer ownership of injected nodes
+			// to the new node. Only after everything has been moved, we can safely remove the
+			// old node.
+			const FString DeletedName = GetValidNodeName(FString::Printf(TEXT("%s_Deleted"), *NodeName));
+			RenameNode(UnitNode, *DeletedName, false, false);
+			
 			URigVMTemplateNode* NewNode = AddTemplateNode(
 				Template->GetNotation(),
 				NodePosition, 
@@ -19010,6 +19107,8 @@ FRigVMClientPatchResult URigVMController::PatchUnitNodesOnLoad()
 
 			ApplyPinStates(NewNode, PinStates, {}, false);
 
+			RemoveNode(UnitNode, false, false);
+
 			FRestoreLinkedPathSettings Settings;
 			RestoreLinkedPaths(LinkedPaths, Settings);
 		}
@@ -19102,7 +19201,13 @@ FRigVMClientPatchResult URigVMController::PatchBranchNodesOnLoad()
 			Result.bChangedContent = true;
 
 			FastBreakLinkedPaths(LinkedPaths);
-			RemoveNode(BranchNode, false, false);
+
+			// Before removing the node (which will remove all injected nodes), lets rename it
+			// create the new node, and apply pin states to transfer ownership of injected nodes
+			// to the new node. Only after everything has been moved, we can safely remove the
+			// old node.
+			const FString DeletedName = GetValidNodeName(FString::Printf(TEXT("%s_Deleted"), *NodeName));
+			RenameNode(BranchNode, *DeletedName, false, false);
 
 			// Cannot reuse the name of deprecated nodes, otherwise we may get the following error when PIE a second time
 			// Failed import: class 'RigVMBranchNode' name 'Branch_1_1' outer 'RigVMModel'. There is another object (of 'RigVMUnitNode' class) at the path.
@@ -19119,6 +19224,8 @@ FRigVMClientPatchResult URigVMController::PatchBranchNodesOnLoad()
 				SetPinDefaultValue(ConditionPin->GetPinPath(), ConditionDefault, false, false, false, false);
 			}
 			RestoreLinkedPaths(LinkedPaths, RestoreLinkedPathSettings);
+
+			RemoveNode(BranchNode, false, false);
 		}
 	}
 	
@@ -19150,7 +19257,13 @@ FRigVMClientPatchResult URigVMController::PatchIfSelectNodesOnLoad()
 			Result.bChangedContent = true;
 			
 			FastBreakLinkedPaths(LinkedPaths);
-			RemoveNode(IfOrSelectNode, false, false);
+			
+			// Before removing the node (which will remove all injected nodes), lets rename it
+			// create the new node, and apply pin states to transfer ownership of injected nodes
+			// to the new node. Only after everything has been moved, we can safely remove the
+			// old node.
+			const FString DeletedName = GetValidNodeName(FString::Printf(TEXT("%s_Deleted"), *NodeName));
+			RenameNode(IfOrSelectNode, *DeletedName, false, false);
 
 			const FRigVMDispatchFactory* Factory = FRigVMRegistry::Get().FindOrAddDispatchFactory(
 				bIsIfNode ? FRigVMDispatch_If::StaticStruct() : FRigVMDispatch_SelectInt32::StaticStruct());
@@ -19191,6 +19304,8 @@ FRigVMClientPatchResult URigVMController::PatchIfSelectNodesOnLoad()
 
 			FRestoreLinkedPathSettings Settings;
 			RestoreLinkedPaths(LinkedPaths, Settings);
+
+			RemoveNode(IfOrSelectNode, false, false);
 		}
 	}
 
