RigVM: Fix reroute with no inputs and top level output links and subelement output links failing to patch

diff --git a/Engine/Plugins/Runtime/RigVM/Source/RigVMDeveloper/Private/RigVMModel/RigVMController.cpp b/Engine/Plugins/Runtime/RigVM/Source/RigVMDeveloper/Private/RigVMModel/RigVMController.cpp
--- a/Engine/Plugins/Runtime/RigVM/Source/RigVMDeveloper/Private/RigVMModel/RigVMController.cpp
+++ b/Engine/Plugins/Runtime/RigVM/Source/RigVMDeveloper/Private/RigVMModel/RigVMController.cpp
@@ -18819,6 +18819,11 @@ FRigVMClientPatchResult URigVMController::PatchRerouteNodesOnLoad()
 							ReroutesToTurnIntoConstant.Add(RerouteNode);
 						}
 					}
+					else if(!bOnlyTopLevelTargetLinks && !bOnlySubPinTargetLinks)
+					{
+						check(ValuePin->IsStruct());
+						ReroutesToTurnIntoRerouteAndBreakStruct.Add(RerouteNode);
+					}
 					else // bOnlySubPinTargetLinks
 					{
 						check(ValuePin->IsStruct());
