RigVM: Lazy branches are not identifying instructions indices correctly when connected to the same node

diff --git a/Engine/Plugins/Runtime/RigVM/Source/RigVM/Private/RigVMFunctions/RigVMDispatch_Print.cpp b/Engine/Plugins/Runtime/RigVM/Source/RigVM/Private/RigVMFunctions/RigVMDispatch_Print.cpp
--- a/Engine/Plugins/Runtime/RigVM/Source/RigVM/Private/RigVMFunctions/RigVMDispatch_Print.cpp
+++ b/Engine/Plugins/Runtime/RigVM/Source/RigVM/Private/RigVMFunctions/RigVMDispatch_Print.cpp
@@ -120,7 +120,7 @@ void FRigVMDispatch_Print::Execute(FRigVMExtendedExecuteContext& InContext, FRig
 	}
 
 	static constexpr TCHAR LogFormat[] = TEXT("%s[%04d] %s%s");
-	UE_LOG(LogRigVM, Display, LogFormat, *ObjectPath, InContext.GetPublicData<>().GetInstructionIndex(), *Prefix, *String);
+	InContext.GetPublicData<>().Logf(EMessageSeverity::Info, LogFormat, *ObjectPath, InContext.GetPublicData<>().GetInstructionIndex(), *Prefix, *String);
 	const UObject* WorldObject = (const UObject*)InContext.VM;
 
 	if(ScreenDuration > SMALL_NUMBER && WorldObject)
diff --git a/Engine/Plugins/Runtime/RigVM/Source/RigVMDeveloper/Private/RigVMCompiler/RigVMAST.cpp b/Engine/Plugins/Runtime/RigVM/Source/RigVMDeveloper/Private/RigVMCompiler/RigVMAST.cpp
--- a/Engine/Plugins/Runtime/RigVM/Source/RigVMDeveloper/Private/RigVMCompiler/RigVMAST.cpp
+++ b/Engine/Plugins/Runtime/RigVM/Source/RigVMDeveloper/Private/RigVMCompiler/RigVMAST.cpp
@@ -1242,6 +1242,16 @@ FRigVMExprAST* FRigVMParserAST::TraverseLink(int32 InLinkIndex, FRigVMExprAST* I
 		}
 	}
 
+	if (!bRequiresCopy)
+	{
+		// Due to the unpredictability of lazy branches, we need to make sure that non-lazy pins are not
+		// affected by the execution of lazy evaluation.
+		if (!TargetRootPin->IsLazy() && TargetRootPin->GetNode()->HasLazyPin(true))
+		{
+			bRequiresCopy = true;
+		}
+	}
+
 	FRigVMAssignExprAST* AssignExpr = nullptr;
 	if (bRequiresCopy)
 	{
diff --git a/Engine/Plugins/Runtime/RigVM/Source/RigVMDeveloper/Private/RigVMCompiler/RigVMCompiler.cpp b/Engine/Plugins/Runtime/RigVM/Source/RigVMDeveloper/Private/RigVMCompiler/RigVMCompiler.cpp
--- a/Engine/Plugins/Runtime/RigVM/Source/RigVMDeveloper/Private/RigVMCompiler/RigVMCompiler.cpp
+++ b/Engine/Plugins/Runtime/RigVM/Source/RigVMDeveloper/Private/RigVMCompiler/RigVMCompiler.cpp
@@ -1431,8 +1431,19 @@ void URigVMCompiler::TraverseBlock(const FRigVMBlockExprAST* InExpr, FRigVMCompi
 			}
 		}
 	}
-	
-	TraverseChildren(InExpr, WorkData);
+
+	if(!WorkData.bSetupMemory && !BranchInfo.Label.IsNone())
+	{
+		// We need to make sure lazy blocks are properly populated, with all the operations that need to run
+		// during the evaluation of the lazy branch, even if some of the expressions
+		// have already been visited in other parts of the traversal. See RigVM.Compiler.IfFromSameNode unit test. 
+		TGuardValue<TMap<const FRigVMExprAST*, bool>> ExprCompletedGuard(WorkData.ExprComplete, {});
+		TraverseChildren(InExpr, WorkData);
+	}
+	else
+	{
+		TraverseChildren(InExpr, WorkData);
+	}
 
 	if(!BranchInfo.Label.IsNone())
 	{
@@ -3355,6 +3366,12 @@ const FRigVMCompilerWorkData::FRigVMASTProxyArray& URigVMCompiler::FindProxiesWi
 					continue;
 				}
 
+				// Non-lazy pins in node with lazy pins cannot share operands
+				if (Pin->GetDirection() == ERigVMPinDirection::Input && !Pin->IsLazy() && Pin->GetNode()->HasLazyPin())
+				{
+					continue;
+				}
+
 				// when compiling a function, avoid folding pins from entry and return nodes
 				if (CurrentCompilationFunction)
 				{
