RigVM: Introducing RigVMSchema

diff --git a/Engine/Plugins/Animation/ControlRig/Source/ControlRigDeveloper/Private/ControlRigBlueprint.cpp b/Engine/Plugins/Animation/ControlRig/Source/ControlRigDeveloper/Private/ControlRigBlueprint.cpp
--- a/Engine/Plugins/Animation/ControlRig/Source/ControlRigDeveloper/Private/ControlRigBlueprint.cpp
+++ b/Engine/Plugins/Animation/ControlRig/Source/ControlRigDeveloper/Private/ControlRigBlueprint.cpp
@@ -34,6 +34,7 @@
 #include "Rigs/RigControlHierarchy.h"
 #include "RigVMModel/Nodes/RigVMDispatchNode.h"
 #include "UObject/UE5MainStreamObjectVersion.h"
+#include "ControlRigSchema.h"
 
 #include UE_INLINE_GENERATED_CPP_BY_NAME(ControlRigBlueprint)
 
@@ -140,6 +141,7 @@ UControlRigBlueprint::UControlRigBlueprint(const FObjectInitializer& ObjectIniti
 	VMRecompilationBracket = 0;
 
 	RigVMClient.SetOuterClientHost(this, GET_MEMBER_NAME_CHECKED(UControlRigBlueprint, RigVMClient));
+	RigVMClient.SetSchemaClass(UControlRigSchema::StaticClass());
 	{
 		TGuardValue<bool> DisableClientNotifs(RigVMClient.bSuspendNotifications, true);
 		RigVMClient.GetOrCreateFunctionLibrary(false, &ObjectInitializer, false);
@@ -449,7 +451,7 @@ void UControlRigBlueprint::HandleRigVMGraphAdded(const FRigVMClient* InClient, c
 #if WITH_EDITOR
 		if(!bSuspendPythonMessagesForRigVMClient)
 		{
-			const FString BlueprintName = URigVMController::GetSanitizedName(GetName(), true, false);
+			const FString BlueprintName = InClient->GetSchema()->GetSanitizedName(GetName(), true, false);
 			RigVMPythonUtils::Print(BlueprintName, 
 				FString::Printf(TEXT("blueprint.add_model('%s')"),
 					*Model->GetName()));
@@ -468,7 +470,7 @@ void UControlRigBlueprint::HandleRigVMGraphRemoved(const FRigVMClient* InClient,
 #if WITH_EDITOR
 		if(!bSuspendPythonMessagesForRigVMClient)
 		{
-			const FString BlueprintName = URigVMController::GetSanitizedName(GetName(), true, false);
+			const FString BlueprintName = InClient->GetSchema()->GetSanitizedName(GetName(), true, false);
 			RigVMPythonUtils::Print(BlueprintName, 
 				FString::Printf(TEXT("blueprint.remove_model('%s')"),
 					*Model->GetName()));
@@ -501,23 +503,6 @@ void UControlRigBlueprint::HandleConfigureRigVMController(const FRigVMClient* In
 {
 	InControllerToConfigure->OnModified().AddUObject(this, &UControlRigBlueprint::HandleModifiedEvent);
 
-	InControllerToConfigure->UnfoldStructDelegate.BindLambda([](const UStruct* InStruct) -> bool {
-
-		if (InStruct == TBaseStructure<FQuat>::Get())
-		{
-			return false;
-		}
-		if (InStruct == FRuntimeFloatCurve::StaticStruct())
-		{
-			return false;
-		}
-		if (InStruct == FRigPose::StaticStruct())
-		{
-			return false;
-		}
-		return true;
-		});
-
 	TWeakObjectPtr<UControlRigBlueprint> WeakThis(this);
 
 	// this delegate is used by the controller to determine variable validity
@@ -2388,6 +2373,7 @@ URigVMController* UControlRigBlueprint::GetTemplateController(bool bIsFunctionLi
 		TemplateController->SetGraph(GetTemplateModel(bIsFunctionLibrary));
 		TemplateController->EnableReporting(false);
 		TemplateController->SetFlags(RF_Transient);
+		TemplateController->SetSchema(RigVMClient.GetOrCreateSchema());
 	}
 	return TemplateController;
 #else
diff --git a/Engine/Plugins/Animation/ControlRig/Source/ControlRigDeveloper/Private/ControlRigSchema.cpp b/Engine/Plugins/Animation/ControlRig/Source/ControlRigDeveloper/Private/ControlRigSchema.cpp
new file mode 100644
--- /dev/null
+++ b/Engine/Plugins/Animation/ControlRig/Source/ControlRigDeveloper/Private/ControlRigSchema.cpp
@@ -0,0 +1,37 @@
+// Copyright Epic Games, Inc. All Rights Reserved.
+
+#include "ControlRigSchema.h"
+#include "RigVMModel/RigVMController.h"
+#include "Rigs/RigHierarchyPose.h"
+#include "Curves/CurveFloat.h"
+#include "Units/RigUnitContext.h"
+
+UControlRigSchema::UControlRigSchema(const FObjectInitializer& ObjectInitializer)
+	: Super(ObjectInitializer)
+{
+	SetExecuteContextStruct(FControlRigExecuteContext::StaticStruct());
+}
+
+bool UControlRigSchema::ShouldUnfoldStruct(URigVMController* InController, const UStruct* InStruct) const
+{
+	RIGVMSCHEMA_DEFAULT_FUNCTION_BODY
+	
+	if(!Super::ShouldUnfoldStruct(InController, InStruct))
+	{
+		return false;
+	}
+	if (InStruct == TBaseStructure<FQuat>::Get())
+	{
+		return false;
+	}
+	if (InStruct == FRuntimeFloatCurve::StaticStruct())
+	{
+		return false;
+	}
+	if (InStruct == FRigPose::StaticStruct())
+	{
+		return false;
+	}
+	
+	return true;
+}
diff --git a/Engine/Plugins/Animation/ControlRig/Source/ControlRigDeveloper/Private/Graph/ControlRigGraphSchema.cpp b/Engine/Plugins/Animation/ControlRig/Source/ControlRigDeveloper/Private/Graph/ControlRigGraphSchema.cpp
--- a/Engine/Plugins/Animation/ControlRig/Source/ControlRigDeveloper/Private/Graph/ControlRigGraphSchema.cpp
+++ b/Engine/Plugins/Animation/ControlRig/Source/ControlRigDeveloper/Private/Graph/ControlRigGraphSchema.cpp
@@ -94,7 +94,8 @@ EValidatorResult FControlRigLocalVariableNameValidator::IsValid(const FString& N
 	EValidatorResult Result = FStringSetNameValidator::IsValid(Name, bOriginal);
 	if (Result == EValidatorResult::Ok)
 	{
-		if (URigVMController::GetSanitizedName(Name, false, true) == Name)
+		const URigVMSchema* RigSchema = Cast<URigVMSchema>(URigVMSchema::StaticClass()->GetDefaultObject(true));
+		if (RigSchema->GetSanitizedName(Name, false, true) == Name)
 		{
 			return Result;
 		}
@@ -133,7 +134,8 @@ EValidatorResult FControlRigNameValidator::IsValid(const FString& Name, bool bOr
 	EValidatorResult Result = FStringSetNameValidator::IsValid(Name, bOriginal);
 	if (Result == EValidatorResult::Ok)
 	{
-		if (URigVMController::GetSanitizedName(Name, false, true) == Name)
+		const URigVMSchema* RigSchema = Cast<URigVMSchema>(URigVMSchema::StaticClass()->GetDefaultObject(true));
+		if (RigSchema->GetSanitizedName(Name, false, true) == Name)
 		{
 			return Result;
 		}
diff --git a/Engine/Plugins/Animation/ControlRig/Source/ControlRigDeveloper/Public/ControlRigSchema.h b/Engine/Plugins/Animation/ControlRig/Source/ControlRigDeveloper/Public/ControlRigSchema.h
new file mode 100644
--- /dev/null
+++ b/Engine/Plugins/Animation/ControlRig/Source/ControlRigDeveloper/Public/ControlRigSchema.h
@@ -0,0 +1,16 @@
+// Copyright Epic Games, Inc. All Rights Reserved.
+
+#pragma once
+
+#include "RigVMModel/RigVMSchema.h"
+#include "ControlRigSchema.generated.h"
+
+UCLASS(BlueprintType)
+class CONTROLRIGDEVELOPER_API UControlRigSchema : public URigVMSchema
+{
+	GENERATED_UCLASS_BODY()
+
+public:
+
+	virtual bool ShouldUnfoldStruct(URigVMController* InController, const UStruct* InStruct) const override;
+};
\ No newline at end of file
diff --git a/Engine/Plugins/Animation/ControlRig/Source/ControlRigEditor/Private/Graph/NodeSpawners/ControlRigFunctionRefNodeSpawner.cpp b/Engine/Plugins/Animation/ControlRig/Source/ControlRigEditor/Private/Graph/NodeSpawners/ControlRigFunctionRefNodeSpawner.cpp
--- a/Engine/Plugins/Animation/ControlRig/Source/ControlRigEditor/Private/Graph/NodeSpawners/ControlRigFunctionRefNodeSpawner.cpp
+++ b/Engine/Plugins/Animation/ControlRig/Source/ControlRigEditor/Private/Graph/NodeSpawners/ControlRigFunctionRefNodeSpawner.cpp
@@ -339,7 +339,8 @@ UControlRigGraphNode* UControlRigFunctionRefNodeSpawner::SpawnNode(UEdGraph* Par
 			Controller->OpenUndoBracket(FString::Printf(TEXT("Add '%s' Node"), *Name.ToString()));
 		}
 
-		if (URigVMFunctionReferenceNode* ModelNode = Controller->AddFunctionReferenceNodeFromDescription(InFunction, Location, Name.ToString(), bIsUserFacingNode, !bIsTemplateNode, bIsTemplateNode))
+		TGuardValue<bool> AllowPrivateFunctionsOnTemplateController(Controller->bAllowPrivateFunctions, bIsTemplateNode);
+		if (URigVMFunctionReferenceNode* ModelNode = Controller->AddFunctionReferenceNodeFromDescription(InFunction, Location, Name.ToString(), bIsUserFacingNode, !bIsTemplateNode))
 		{
 			NewNode = Cast<UControlRigGraphNode>(RigGraph->FindNodeForModelNodeName(ModelNode->GetFName()));
 			check(NewNode);
@@ -395,7 +396,9 @@ UControlRigGraphNode* UControlRigFunctionRefNodeSpawner::SpawnNode(UEdGraph* Par
 			}
 			else
 			{
+				Controller->EnableReporting(false);
 				Controller->RemoveNode(ModelNode, false);
+				Controller->EnableReporting(true);
 			}
 		}
 		else
diff --git a/Engine/Plugins/Animation/ControlRig/Source/ControlRigEditor/Private/Graph/NodeSpawners/ControlRigUnitNodeSpawner.cpp b/Engine/Plugins/Animation/ControlRig/Source/ControlRigEditor/Private/Graph/NodeSpawners/ControlRigUnitNodeSpawner.cpp
--- a/Engine/Plugins/Animation/ControlRig/Source/ControlRigEditor/Private/Graph/NodeSpawners/ControlRigUnitNodeSpawner.cpp
+++ b/Engine/Plugins/Animation/ControlRig/Source/ControlRigEditor/Private/Graph/NodeSpawners/ControlRigUnitNodeSpawner.cpp
@@ -252,7 +252,7 @@ UControlRigGraphNode* UControlRigUnitNodeSpawner::SpawnNode(UEdGraph* ParentGrap
 						}
 					}
 
-					FName SafeEventName = URigVMController::GetUniqueName(EventName, [ExistingEventNames](const FName& NameToCheck) -> bool
+					FName SafeEventName = Controller->GetSchema()->GetUniqueName(EventName, [ExistingEventNames](const FName& NameToCheck) -> bool
 					{
 						return !ExistingEventNames.Contains(NameToCheck);
 					}, false, true);
diff --git a/Engine/Plugins/Experimental/AnimNext/Source/AnimNextUncookedOnly/Internal/Graph/AnimNextGraph_EditorData.h b/Engine/Plugins/Experimental/AnimNext/Source/AnimNextUncookedOnly/Internal/Graph/AnimNextGraph_EditorData.h
--- a/Engine/Plugins/Experimental/AnimNext/Source/AnimNextUncookedOnly/Internal/Graph/AnimNextGraph_EditorData.h
+++ b/Engine/Plugins/Experimental/AnimNext/Source/AnimNextUncookedOnly/Internal/Graph/AnimNextGraph_EditorData.h
@@ -26,6 +26,12 @@ enum class EAnimNextGraphLoadType : uint8
 	CheckUserDefinedStructs
 };
 
+UCLASS()
+class UAnimNextGraph_Schema : public URigVMSchema
+{
+	GENERATED_BODY()
+};
+
 UCLASS(MinimalAPI)
 class UAnimNextGraph_EditorData : public UObject, public IRigVMClientHost, public IRigVMGraphFunctionHost
 {
diff --git a/Engine/Plugins/Experimental/AnimNext/Source/AnimNextUncookedOnly/Internal/Param/AnimNextParameterBlock_EditorData.h b/Engine/Plugins/Experimental/AnimNext/Source/AnimNextUncookedOnly/Internal/Param/AnimNextParameterBlock_EditorData.h
--- a/Engine/Plugins/Experimental/AnimNext/Source/AnimNextUncookedOnly/Internal/Param/AnimNextParameterBlock_EditorData.h
+++ b/Engine/Plugins/Experimental/AnimNext/Source/AnimNextUncookedOnly/Internal/Param/AnimNextParameterBlock_EditorData.h
@@ -42,6 +42,12 @@ enum class EAnimNextParameterLoadType : uint8
 	CheckUserDefinedStructs
 };
 
+UCLASS()
+class UAnimNextParameterBlockLibrary_Schema : public URigVMSchema
+{
+	GENERATED_BODY()
+};
+
 // Script-callable editor API hoisted onto UAnimNextParameterBlock
 UCLASS()
 class ANIMNEXTUNCOOKEDONLY_API UAnimNextParameterBlockLibrary : public UBlueprintFunctionLibrary
diff --git a/Engine/Plugins/Experimental/AnimNext/Source/AnimNextUncookedOnly/Private/Graph/AnimNextGraph_EditorData.cpp b/Engine/Plugins/Experimental/AnimNext/Source/AnimNextUncookedOnly/Private/Graph/AnimNextGraph_EditorData.cpp
--- a/Engine/Plugins/Experimental/AnimNext/Source/AnimNextUncookedOnly/Private/Graph/AnimNextGraph_EditorData.cpp
+++ b/Engine/Plugins/Experimental/AnimNext/Source/AnimNextUncookedOnly/Private/Graph/AnimNextGraph_EditorData.cpp
@@ -18,6 +18,7 @@
 UAnimNextGraph_EditorData::UAnimNextGraph_EditorData(const FObjectInitializer& ObjectInitializer)
 {
 	RigVMClient.Reset();
+	RigVMClient.SetSchemaClass(UAnimNextGraph_Schema::StaticClass());
 	RigVMClient.SetOuterClientHost(this, GET_MEMBER_NAME_CHECKED(UAnimNextGraph_EditorData, RigVMClient));
 	{
 		TGuardValue<bool> DisableClientNotifs(RigVMClient.bSuspendNotifications, true);
@@ -268,23 +269,6 @@ void UAnimNextGraph_EditorData::HandleConfigureRigVMController(const FRigVMClien
 {
 	InControllerToConfigure->OnModified().AddUObject(this, &UAnimNextGraph_EditorData::HandleModifiedEvent);
 
-	InControllerToConfigure->UnfoldStructDelegate.BindLambda([](const UStruct* InStruct) -> bool
-	{
-		if (InStruct == TBaseStructure<FQuat>::Get())
-		{
-			return false;
-		}
-		if (InStruct == FRuntimeFloatCurve::StaticStruct())
-		{
-			return false;
-		}
-		if (InStruct == FRigPose::StaticStruct())
-		{
-			return false;
-		}
-		return true;
-	});
-
 	TWeakObjectPtr<UAnimNextGraph_EditorData> WeakThis(this);
 
 	// this delegate is used by the controller to determine variable validity
diff --git a/Engine/Plugins/Experimental/AnimNext/Source/AnimNextUncookedOnly/Private/Param/AnimNextParameterBlock_EditorData.cpp b/Engine/Plugins/Experimental/AnimNext/Source/AnimNextUncookedOnly/Private/Param/AnimNextParameterBlock_EditorData.cpp
--- a/Engine/Plugins/Experimental/AnimNext/Source/AnimNextUncookedOnly/Private/Param/AnimNextParameterBlock_EditorData.cpp
+++ b/Engine/Plugins/Experimental/AnimNext/Source/AnimNextUncookedOnly/Private/Param/AnimNextParameterBlock_EditorData.cpp
@@ -56,13 +56,14 @@ static EntryClassType* CreateNewParameterBlockEntry(UAnimNextParameterBlock_Edit
 UAnimNextParameterBlock_EditorData::UAnimNextParameterBlock_EditorData(const FObjectInitializer& ObjectInitializer)
 {
 	RigVMClient.Reset();
+	RigVMClient.SetSchemaClass(UAnimNextParameterBlockLibrary_Schema::StaticClass());
+	RigVMClient.SetExecuteContextStruct(FAnimNextParametersExecuteContext::StaticStruct());
 	RigVMClient.SetOuterClientHost(this, GET_MEMBER_NAME_CHECKED(UAnimNextParameterBlock_EditorData, RigVMClient));
 	{
 		TGuardValue<bool> DisableClientNotifs(RigVMClient.bSuspendNotifications, true);
 		RigVMClient.AddModel(TEXT("RigVMGraph"), false, &ObjectInitializer);
 		RigVMClient.GetOrCreateFunctionLibrary(false, &ObjectInitializer);
 	}
-	RigVMClient.SetExecuteContextStruct(FAnimNextParametersExecuteContext::StaticStruct());
 }
 
 #if WITH_EDITORONLY_DATA
@@ -631,7 +632,7 @@ void UAnimNextParameterBlock_EditorData::HandleRigVMGraphAdded(const FRigVMClien
 #if WITH_EDITOR
 		if(!bSuspendPythonMessagesForRigVMClient)
 		{
-			const FString BlueprintName = URigVMController::GetSanitizedName(GetName(), true, false);
+			const FString BlueprintName = InClient->GetSchema()->GetSanitizedName(GetName(), true, false);
 			RigVMPythonUtils::Print(BlueprintName, FString::Printf(TEXT("block.add_binding('%s')"), *RigVMGraph->GetName()));
 		}
 #endif
@@ -648,7 +649,7 @@ void UAnimNextParameterBlock_EditorData::HandleRigVMGraphRemoved(const FRigVMCli
 #if WITH_EDITOR
 		if(!bSuspendPythonMessagesForRigVMClient)
 		{
-			const FString BlueprintName = URigVMController::GetSanitizedName(GetName(), true, false);
+			const FString BlueprintName = InClient->GetSchema()->GetSanitizedName(GetName(), true, false);
 			RigVMPythonUtils::Print(BlueprintName, FString::Printf(TEXT("block.remove_binding('%s')"), *RigVMGraph->GetName()));
 		}
 #endif
@@ -659,23 +660,6 @@ void UAnimNextParameterBlock_EditorData::HandleConfigureRigVMController(const FR
 {
 	InControllerToConfigure->OnModified().AddUObject(this, &UAnimNextParameterBlock_EditorData::HandleModifiedEvent);
 
-	InControllerToConfigure->UnfoldStructDelegate.BindLambda([](const UStruct* InStruct) -> bool
-	{
-		if (InStruct == TBaseStructure<FQuat>::Get())
-		{
-			return false;
-		}
-		if (InStruct == FRuntimeFloatCurve::StaticStruct())
-		{
-			return false;
-		}
-		if (InStruct == FRigPose::StaticStruct())
-		{
-			return false;
-		}
-		return true;
-	});
-
 	TWeakObjectPtr<UAnimNextParameterBlock_EditorData> WeakThis(this);
 
 	// this delegate is used by the controller to determine variable validity
diff --git a/Engine/Plugins/Runtime/RigVM/Source/RigVM/Private/RigVMCore/RigVMExternalVariable.cpp b/Engine/Plugins/Runtime/RigVM/Source/RigVM/Private/RigVMCore/RigVMExternalVariable.cpp
--- a/Engine/Plugins/Runtime/RigVM/Source/RigVM/Private/RigVMCore/RigVMExternalVariable.cpp
+++ b/Engine/Plugins/Runtime/RigVM/Source/RigVM/Private/RigVMCore/RigVMExternalVariable.cpp
@@ -1,6 +1,7 @@
 // Copyright Epic Games, Inc. All Rights Reserved.
 
 #include "RigVMCore/RigVMExternalVariable.h"
+#include "RigVMCore/RigVMRegistry.h"
 
 #if WITH_EDITOR
 
@@ -239,4 +240,13 @@ FRigVMExternalVariable RigVMTypeUtils::ExternalVariableFromCPPType(const FName&
 	return Variable;
 }
 
-#endif
\ No newline at end of file
+#endif
+
+TRigVMTypeIndex FRigVMExternalVariable::GetTypeIndex() const
+{
+	if(IsValid(true))
+	{
+		return FRigVMRegistry::Get().GetTypeIndexFromCPPType(GetExtendedCPPType().ToString());
+	}
+	return INDEX_NONE;
+}
diff --git a/Engine/Plugins/Runtime/RigVM/Source/RigVM/Public/RigVMCore/RigVMDispatchFactory.h b/Engine/Plugins/Runtime/RigVM/Source/RigVM/Public/RigVMCore/RigVMDispatchFactory.h
--- a/Engine/Plugins/Runtime/RigVM/Source/RigVM/Public/RigVMCore/RigVMDispatchFactory.h
+++ b/Engine/Plugins/Runtime/RigVM/Source/RigVM/Public/RigVMCore/RigVMDispatchFactory.h
@@ -13,6 +13,7 @@
 #include "RigVMCore/RigVMStructUpgradeInfo.h"
 #include "RigVMCore/RigVMTemplate.h"
 #include "RigVMCore/RigVMTypeIndex.h"
+#include "RigVMCore/RigVMExecuteContext.h"
 #include "UObject/Class.h"
 #include "UObject/NameTypes.h"
 #include "UObject/ObjectMacros.h"
diff --git a/Engine/Plugins/Runtime/RigVM/Source/RigVM/Public/RigVMCore/RigVMExternalVariable.h b/Engine/Plugins/Runtime/RigVM/Source/RigVM/Public/RigVMCore/RigVMExternalVariable.h
--- a/Engine/Plugins/Runtime/RigVM/Source/RigVM/Public/RigVMCore/RigVMExternalVariable.h
+++ b/Engine/Plugins/Runtime/RigVM/Source/RigVM/Public/RigVMCore/RigVMExternalVariable.h
@@ -7,6 +7,7 @@
 #include "RigVMDefines.h"
 #include "RigVMMemory.h"
 #include "RigVMModule.h"
+#include "RigVMTypeIndex.h"
 #include "RigVMTypeUtils.h"
 #include "UObject/UnrealType.h"
 #if WITH_EDITOR
@@ -594,6 +595,8 @@ struct RIGVM_API FRigVMExternalVariable
 		return TypeName;
 	}
 
+	TRigVMTypeIndex GetTypeIndex() const;
+
 	static void MergeExternalVariable(TArray<FRigVMExternalVariable>& OutVariables, const FRigVMExternalVariable& InVariable)
 	{
 		if(!InVariable.IsValid(true))
diff --git a/Engine/Plugins/Runtime/RigVM/Source/RigVMDeveloper/Private/RigVMModel/RigVMClient.cpp b/Engine/Plugins/Runtime/RigVM/Source/RigVMDeveloper/Private/RigVMModel/RigVMClient.cpp
--- a/Engine/Plugins/Runtime/RigVM/Source/RigVMDeveloper/Private/RigVMModel/RigVMClient.cpp
+++ b/Engine/Plugins/Runtime/RigVM/Source/RigVMDeveloper/Private/RigVMModel/RigVMClient.cpp
@@ -14,6 +14,28 @@
 #include "ScopedTransaction.h"
 #endif
 
+void FRigVMClient::SetSchemaClass(TSubclassOf<URigVMSchema> InSchemaClass)
+{
+	check(InSchemaClass);
+	
+	if(InSchemaClass == SchemaClass)
+	{
+		return;
+	}
+	
+	if(SchemaPtr)
+	{
+		if(!SchemaPtr->HasAnyFlags(RF_ClassDefaultObject))
+		{
+			DestroyObject(SchemaPtr);
+		}
+		SchemaPtr = nullptr;
+	}
+
+	SchemaClass = InSchemaClass;
+	GetOrCreateSchema();
+}
+
 void FRigVMClient::SetOuterClientHost(UObject* InOuterClientHost, const FName& InOuterClientHostPropertyName)
 {
 	OuterClientHost = InOuterClientHost;
@@ -76,6 +98,33 @@ void FRigVMClient::Reset()
 	FunctionLibrary = nullptr;
 
 	ResetActionStack();
+
+	if(SchemaPtr)
+	{
+		if(!SchemaPtr->HasAnyFlags(RF_ClassDefaultObject))
+		{
+			DestroyObject(SchemaPtr);
+		}
+		SchemaPtr = nullptr;
+	}
+}
+
+URigVMSchema* FRigVMClient::GetOrCreateSchema()
+{
+	check(SchemaClass);
+	if(SchemaPtr && !IsValid(SchemaPtr))
+	{
+		SchemaPtr = nullptr;
+	}
+	if(SchemaPtr && SchemaPtr->GetClass() != SchemaClass)
+	{
+		SchemaPtr = nullptr;
+	}
+	if(SchemaPtr == nullptr)
+	{
+		SchemaPtr = CastChecked<URigVMSchema>(SchemaClass->GetDefaultObject(true));
+	}
+	return SchemaPtr;
 }
 
 URigVMGraph* FRigVMClient::GetDefaultModel() const
@@ -305,7 +354,7 @@ void FRigVMClient::AddModel(URigVMGraph* InModel, bool bCreateController)
 		Models.Add(InModel);
 	}
 
-	InModel->SetExecuteContextStruct(GetExecuteContextStruct());
+	InModel->SetExecuteContextStruct(GetOrCreateSchema()->GetExecuteContextStruct());
 
 	if(bCreateController)
 	{
@@ -394,6 +443,17 @@ TArray<FName> FRigVMClient::GetEntryNames() const
 	return EntryNames;
 }
 
+UScriptStruct* FRigVMClient::GetExecuteContextStruct() const
+{
+	check(SchemaPtr);
+	return SchemaPtr->GetExecuteContextStruct();
+}
+
+void FRigVMClient::SetExecuteContextStruct(UScriptStruct* InExecuteContextStruct)
+{
+	GetOrCreateSchema()->SetExecuteContextStruct(InExecuteContextStruct);
+}
+
 bool FRigVMClient::RemoveModel(const FString& InNodePathOrName, bool bSetupUndoRedo)
 {
 	if(URigVMGraph* Model = GetModel(InNodePathOrName))
@@ -685,6 +745,7 @@ URigVMController* FRigVMClient::CreateController(const URigVMGraph* InModel)
 	const FName SafeControllerName = GetUniqueName(*FString::Printf(TEXT("%s_Controller"), *ModelName));
 	URigVMController* Controller = NewObject<URigVMController>(GetOuter(), SafeControllerName);
 	Controllers.Add(InModel, Controller);
+	Controller->SetSchema(GetOrCreateSchema());
 	Controller->SetActionStack(GetOrCreateActionStack());
 	if(InModel)
 	{
@@ -731,6 +792,13 @@ void FRigVMClient::ResetActionStack()
 	ActionStack = nullptr;
 }
 
+void FRigVMClient::SetSchema(URigVMSchema* InSchema)
+{
+	check(InSchema);
+	SchemaPtr = InSchema;
+	SchemaClass = InSchema->GetClass();
+}
+
 FName FRigVMClient::GetUniqueName(const FName& InDesiredName) const
 {
 	return GetUniqueName(GetOuter(), InDesiredName);
@@ -738,7 +806,7 @@ FName FRigVMClient::GetUniqueName(const FName& InDesiredName) const
 
 FName FRigVMClient::GetUniqueName(UObject* InOuter, const FName& InDesiredName)
 {
-	return URigVMController::GetUniqueName(*InDesiredName.ToString(), [InOuter](const FName& InName) -> bool
+	return URigVMSchema::GetUniqueName(*InDesiredName.ToString(), [InOuter](const FName& InName) -> bool
 	{
 		return FindObjectWithOuter(InOuter, nullptr, InName) == nullptr; 
 	}, false, true);
@@ -752,7 +820,7 @@ void FRigVMClient::DestroyObject(UObject* InObject)
 		static constexpr TCHAR ObjectNameFormat[] = TEXT("RigVMClient_ObjectToBeDestroyed_%d");
 		const FString NewObjectName = FString::Printf(ObjectNameFormat, ObjectIndexToBeDestroyed++);
 		InObject->Rename(*NewObjectName, GetTransientPackage(), REN_ForceNoResetLoaders | REN_DontCreateRedirectors);
-		if (!InObject->IsRooted())
+		if(!InObject->IsRooted())
 		{
 			InObject->MarkAsGarbage();
 		}
@@ -820,10 +888,14 @@ void FRigVMClient::ProcessDetachedLinks()
 
 void FRigVMClient::PostDuplicateHost(const FString& InOldPathName, const FString& InNewPathName)
 {
+	// this also takes care of wiping invalid schemas on the client
+	URigVMSchema* Schema = GetOrCreateSchema();
+	
 	TArray<URigVMGraph*> AllModels = GetAllModels(true, true);
 	for(URigVMGraph* Model : AllModels)
 	{
 		URigVMController* Controller = GetOrCreateController(Model);
+		Controller->SetSchema(Schema);
 		Controller->PostDuplicateHost(InOldPathName, InNewPathName);
 	}
 }
diff --git a/Engine/Plugins/Runtime/RigVM/Source/RigVMDeveloper/Private/RigVMModel/RigVMController.cpp b/Engine/Plugins/Runtime/RigVM/Source/RigVMDeveloper/Private/RigVMModel/RigVMController.cpp
--- a/Engine/Plugins/Runtime/RigVM/Source/RigVMDeveloper/Private/RigVMModel/RigVMController.cpp
+++ b/Engine/Plugins/Runtime/RigVM/Source/RigVMDeveloper/Private/RigVMModel/RigVMController.cpp
@@ -9,7 +9,6 @@
 #include "RigVMModel/Nodes/RigVMDispatchNode.h"
 #include "RigVMCore/RigVMRegistry.h"
 #include "RigVMCore/RigVMExecuteContext.h"
-#include "RigVMCore/RigVMUnknownType.h"
 #include "RigVMCore/RigVMByteCode.h"
 #include "RigVMFunctions/RigVMFunction_ControlFlow.h"
 #include "RigVMFunctions/RigVMDispatch_Array.h"
@@ -19,7 +18,6 @@
 #include "RigVMDeveloperModule.h"
 #include "UObject/PropertyPortFlags.h"
 #include "UObject/Package.h"
-#include "UObject/StrongObjectPtr.h"
 #include "Misc/CoreMisc.h"
 #include "Algo/Count.h"
 #include "Algo/Reverse.h"
@@ -29,7 +27,6 @@
 #include "Engine/UserDefinedStruct.h"
 #include "RigVMFunctions/RigVMDispatch_If.h"
 #include "RigVMFunctions/RigVMDispatch_Select.h"
-#include "RigVMFunctions/RigVMDispatch_Array.h"
 #include "RigVMModel/RigVMClient.h"
 #include "RigVMModel/Nodes/RigVMBranchNode.h"
 #include "RigVMModel/Nodes/RigVMArrayNode.h"
@@ -85,11 +82,13 @@ void FRigVMClientPatchResult::Merge(const FRigVMClientPatchResult& InOther)
 
 URigVMController::URigVMController()
 	: bValidatePinDefaults(true)
+	, SchemaPtr(nullptr)
 	, bSuspendNotifications(false)
 	, bReportWarningsAndErrors(true)
 	, bIgnoreRerouteCompactnessChanges(false)
 	, UserLinkDirection(ERigVMPinDirection::Invalid)
 	, bEnableTypeCasting(true)
+	, bAllowPrivateFunctions(false)
 	, bIsTransacting(false)
 	, bIsRunningUnitTest(false)
 	, bIsFullyResolvingTemplateNode(false)
@@ -103,11 +102,13 @@ URigVMController::URigVMController()
 URigVMController::URigVMController(const FObjectInitializer& ObjectInitializer)
 	: Super(ObjectInitializer)
 	, bValidatePinDefaults(true)
+	, SchemaPtr(nullptr)
 	, bSuspendNotifications(false)
 	, bReportWarningsAndErrors(true)
 	, bIgnoreRerouteCompactnessChanges(false)
 	, UserLinkDirection(ERigVMPinDirection::Invalid)
 	, bEnableTypeCasting(true)
+	, bAllowPrivateFunctions(false)
 	, bIsTransacting(false)
 	, bIsRunningUnitTest(false)
 	, bIsFullyResolvingTemplateNode(false)
@@ -174,6 +175,11 @@ void URigVMController::SetGraph(URigVMGraph* InGraph)
 	HandleModifiedEvent(ERigVMGraphNotifType::GraphChanged, GetGraph(), nullptr);
 }
 
+void URigVMController::SetSchema(URigVMSchema* InSchema)
+{
+	SchemaPtr = InSchema;
+}
+
 bool URigVMController::PushGraph(URigVMGraph* InGraph, bool bSetupUndoRedo)
 {
 	if (URigVMGraph* LastGraph = GetGraph())
@@ -231,7 +237,7 @@ URigVMGraph* URigVMController::GetTopLevelGraph() const
 			Graph = OuterGraph;
 			Outer = Outer->GetOuter();
 		}
-		else if (URigVMLibraryNode* LibraryNode = Cast<URigVMLibraryNode>(Outer))
+		else if (Outer->IsA<URigVMLibraryNode>())
 		{
 			Outer = Outer->GetOuter();
 		}
@@ -313,7 +319,7 @@ void URigVMController::ResendAllNotifications()
 		{
 			Notify(ERigVMGraphNotifType::NodeAdded, Node);
 
-			if (URigVMCommentNode* CommentNode = Cast<URigVMCommentNode>(Node))
+			if (Node->IsA<URigVMCommentNode>())
 			{
 				Notify(ERigVMGraphNotifType::CommentTextChanged, Node);
 			}
@@ -665,7 +671,7 @@ int32 URigVMController::FPinInfoArray::AddPin(URigVMController* InController, in
 void URigVMController::FPinInfoArray::AddPins(UScriptStruct* InScriptStruct, URigVMController* InController,
 	ERigVMPinDirection InDirection, int32 InParentIndex, const uint8* InDefaultValueMemory, bool bAddSubPins)
 {
-	if (InController->ShouldStructBeUnfolded(InScriptStruct))
+	if (InController->GetSchema()->ShouldUnfoldStruct(InController, InScriptStruct))
 	{
 		TArray<UStruct*> StructsToVisit = FRigVMTemplate::GetSuperStructs(InScriptStruct, true);
 		for(UStruct* StructToVisit : StructsToVisit)
@@ -866,12 +872,12 @@ TArray<FString> URigVMController::GeneratePythonCommands()
 {
 	TArray<FString> Commands;
 
-	const FString GraphName = GetSanitizedGraphName(GetGraph()->GetGraphName());
+	const FString GraphName = GetSchema()->GetSanitizedGraphName(GetGraph()->GetGraphName());
 
 	// Add local variables
 	for (const FRigVMGraphVariableDescription& Variable : GetGraph()->LocalVariables)
 	{
-		const FString VariableName = GetSanitizedVariableName(Variable.Name.ToString());
+		const FString VariableName = GetSchema()->GetSanitizedVariableName(Variable.Name.ToString());
 
 		if (Variable.CPPTypeObject)
 		{
@@ -912,8 +918,8 @@ TArray<FString> URigVMController::GeneratePythonCommands()
 			continue;
 		}
 
-		const FString SourcePinPath = GetSanitizedPinPath(SourcePin->GetPinPath());
-		const FString TargetPinPath = GetSanitizedPinPath(TargetPin->GetPinPath());
+		const FString SourcePinPath = GetSchema()->GetSanitizedPinPath(SourcePin->GetPinPath());
+		const FString TargetPinPath = GetSchema()->GetSanitizedPinPath(TargetPin->GetPinPath());
 
 		//bool AddLink(const FString& InOutputPinPath, const FString& InInputPinPath, bool bSetupUndoRedo = true);
 		Commands.Add(FString::Printf(TEXT("blueprint.get_controller_by_name('%s').add_link('%s', '%s')"),
@@ -929,8 +935,8 @@ TArray<FString> URigVMController::GetAddNodePythonCommands(URigVMNode* Node) con
 {
 	TArray<FString> Commands;
 
-	const FString GraphName = GetSanitizedGraphName(GetGraph()->GetGraphName());
-	const FString NodeName = GetSanitizedNodeName(Node->GetName());
+	const FString GraphName = GetSchema()->GetSanitizedGraphName(GetGraph()->GetGraphName());
+	const FString NodeName = GetSchema()->GetSanitizedNodeName(Node->GetName());
 
 	auto GetResolveWildcardPinsPythonCommands = [](const FString& InGraphName, const URigVMTemplateNode* InNode, const FRigVMTemplate* InTemplate)
 	{
@@ -974,9 +980,9 @@ TArray<FString> URigVMController::GetAddNodePythonCommands(URigVMNode* Node) con
 		if (const URigVMInjectionInfo* InjectionInfo = Cast<URigVMInjectionInfo>(UnitNode->GetOuter()))
 		{
 			const URigVMPin* InjectionInfoPin = InjectionInfo->GetPin();
-			const FString InjectionInfoPinPath = GetSanitizedPinPath(InjectionInfoPin->GetPinPath());
-			const FString InjectionInfoInputPinName = InjectionInfo->InputPin ? GetSanitizedPinName(InjectionInfo->InputPin->GetName()) : FString();
-			const FString InjectionInfoOutputPinName = InjectionInfo->OutputPin ? GetSanitizedPinName(InjectionInfo->OutputPin->GetName()) : FString();
+			const FString InjectionInfoPinPath = GetSchema()->GetSanitizedPinPath(InjectionInfoPin->GetPinPath());
+			const FString InjectionInfoInputPinName = InjectionInfo->InputPin ? GetSchema()->GetSanitizedPinName(InjectionInfo->InputPin->GetName()) : FString();
+			const FString InjectionInfoOutputPinName = InjectionInfo->OutputPin ? GetSchema()->GetSanitizedPinName(InjectionInfo->OutputPin->GetName()) : FString();
 
 			//URigVMInjectionInfo* AddInjectedNodeFromStructPath(const FString& InPinPath, bool bAsInput, const FString& InScriptStructPath, const FName& InMethodName, const FName& InInputPinName, const FName& InOutputPinName, const FString& InNodeName = TEXT(""), bool bSetupUndoRedo = true);
 			Commands.Add(FString::Printf(TEXT("%s_info = blueprint.get_controller_by_name('%s').add_injected_node_from_struct_path('%s', %s, '%s', '%s', '%s', '%s', '%s')"),
@@ -1058,7 +1064,7 @@ TArray<FString> URigVMController::GetAddNodePythonCommands(URigVMNode* Node) con
 	{
 		if (!VariableNode->IsInjected())
 		{
-			const FString VariableName = GetSanitizedVariableName(VariableNode->GetVariableName().ToString());
+			const FString VariableName = GetSchema()->GetSanitizedVariableName(VariableNode->GetVariableName().ToString());
 
 			// add_variable_node(variable_name, cpp_type, cpp_type_object, is_getter, default_value, position=[0.0, 0.0], node_name='', undo=True)
 			if (VariableNode->GetVariableDescription().CPPTypeObject)
@@ -1141,7 +1147,7 @@ TArray<FString> URigVMController::GetAddNodePythonCommands(URigVMNode* Node) con
 	}
 	else if (const URigVMCollapseNode* CollapseNode = Cast<URigVMCollapseNode>(Node))
 	{
-		const FString ContainedGraphName = GetSanitizedGraphName(CollapseNode->GetContainedGraph()->GetGraphName());
+		const FString ContainedGraphName = GetSchema()->GetSanitizedGraphName(CollapseNode->GetContainedGraph()->GetGraphName());
 		// AddFunctionReferenceNode(URigVMLibraryNode* InFunctionDefinition, const FVector2D& InNodePosition = FVector2D::ZeroVector, const FString& InNodeName = TEXT(""), bool bSetupUndoRedo = true);
 		Commands.Add(FString::Printf(TEXT("blueprint.get_controller_by_name('%s').add_function_reference_node(function_%s, %s, '%s')"),
 					*GraphName,
@@ -1188,7 +1194,7 @@ TArray<FString> URigVMController::GetAddNodePythonCommands(URigVMNode* Node) con
 			const FString DefaultValue = Pin->GetDefaultValue();
 			if (!DefaultValue.IsEmpty() && DefaultValue != TEXT("()"))
 			{
-				const FString PinPath = GetSanitizedPinPath(Pin->GetPinPath());
+				const FString PinPath = GetSchema()->GetSanitizedPinPath(Pin->GetPinPath());
 
 				Commands.Add(FString::Printf(TEXT("blueprint.get_controller_by_name('%s').set_pin_default_value('%s', '%s')"),
 							*GraphName,
@@ -1202,7 +1208,7 @@ TArray<FString> URigVMController::GetAddNodePythonCommands(URigVMNode* Node) con
 					if (SubPins[i]->IsStruct() || SubPins[i]->IsArray())
 					{
 						SubPins.Append(SubPins[i]->GetSubPins());
-						const FString SubPinPath = GetSanitizedPinPath(SubPins[i]->GetPinPath());
+						const FString SubPinPath = GetSchema()->GetSanitizedPinPath(SubPins[i]->GetPinPath());
 
 						Commands.Add(FString::Printf(TEXT("blueprint.get_controller_by_name('%s').set_pin_expansion('%s', %s)"),
 							*GraphName,
@@ -1214,7 +1220,7 @@ TArray<FString> URigVMController::GetAddNodePythonCommands(URigVMNode* Node) con
 
 			if (!Pin->GetBoundVariablePath().IsEmpty())
 			{
-				const FString PinPath = GetSanitizedPinPath(Pin->GetPinPath());
+				const FString PinPath = GetSchema()->GetSanitizedPinPath(Pin->GetPinPath());
 
 				Commands.Add(FString::Printf(TEXT("blueprint.get_controller_by_name('%s').bind_pin_to_variable('%s', '%s')"),
 							*GraphName,
@@ -1268,16 +1274,21 @@ URigVMUnitNode* URigVMController::AddUnitNode(UScriptStruct* InScriptStruct, con
 		return nullptr;
 	}
 
+	if(!GetSchema()->SupportsUnitFunction(this, Function))
+	{
+		return nullptr;
+	}
+
 	if(IRigVMClientHost* ClientHost = GetImplementingOuter<IRigVMClientHost>())
 	{
 		if(const FRigVMClient* Client = ClientHost->GetRigVMClient())
 		{
-			if(!Function->SupportsExecuteContextStruct(Client->GetExecuteContextStruct()))
+			if(!Function->SupportsExecuteContextStruct(GetSchema()->GetExecuteContextStruct()))
 			{
 				ReportErrorf(TEXT("Cannot add node for function '%s' - incompatible execute context: '%s' vs '%s'."),
 						*Function->GetName(),
 						*Function->GetExecuteContextStruct()->GetStructCPPName(),
-						*Client->GetExecuteContextStruct()->GetStructCPPName());
+						*GetSchema()->GetExecuteContextStruct()->GetStructCPPName());
 				return nullptr;
 			}
 		}
@@ -1297,7 +1308,7 @@ URigVMUnitNode* URigVMController::AddUnitNode(UScriptStruct* InScriptStruct, con
 			OpenUndoBracket(FString::Printf(TEXT("Add %s Node"), *Template->GetName().ToString()));
 		}
 
-		const FString Name = GetValidNodeName(InNodeName.IsEmpty() ? InScriptStruct->GetName() : InNodeName);
+		const FString Name = GetSchema()->GetValidNodeName(Graph, InNodeName.IsEmpty() ? InScriptStruct->GetName() : InNodeName);
 		URigVMUnitNode* TemplateNode = Cast<URigVMUnitNode>(AddTemplateNode(Template->GetNotation(), InPosition, Name, bSetupUndoRedo, bPrintPythonCommand));
 		if(TemplateNode == nullptr)
 		{
@@ -1355,7 +1366,7 @@ URigVMUnitNode* URigVMController::AddUnitNode(UScriptStruct* InScriptStruct, con
 		if(StructMemory->CanOnlyExistOnce())
 		{
 			// don't allow several event nodes in the main graph
-			TObjectPtr<URigVMNode> EventNode = FindEventNode(InScriptStruct);
+			TObjectPtr<URigVMNode> EventNode = GetSchema()->FindEventNode(this, InScriptStruct);
 			if (EventNode != nullptr)
 			{
 				const FString ErrorMessage = FString::Printf(TEXT("Rig Graph can only contain one single %s node."),
@@ -1366,7 +1377,7 @@ URigVMUnitNode* URigVMController::AddUnitNode(UScriptStruct* InScriptStruct, con
 		}
 	}
 	
-	FString Name = GetValidNodeName(InNodeName.IsEmpty() ? InScriptStruct->GetName() : InNodeName);
+	FString Name = GetSchema()->GetValidNodeName(Graph, InNodeName.IsEmpty() ? InScriptStruct->GetName() : InNodeName);
 	URigVMUnitNode* Node = NewObject<URigVMUnitNode>(Graph, *Name);
 	Node->ResolvedFunctionName = Function->GetName();
 	Node->Position = InPosition;
@@ -1386,7 +1397,11 @@ URigVMUnitNode* URigVMController::AddUnitNode(UScriptStruct* InScriptStruct, con
 		AddPinsForStruct(InScriptStruct, Node, nullptr, ERigVMPinDirection::Invalid, ExportedDefaultValue, true);
 	}
 
-	Graph->Nodes.Add(Node);
+	if(!AddGraphNode(Node, true))
+	{
+		return nullptr;
+	}
+	
 	if (!bSuspendNotifications)
 	{
 		Graph->MarkPackageDirty();
@@ -1402,8 +1417,6 @@ URigVMUnitNode* URigVMController::AddUnitNode(UScriptStruct* InScriptStruct, con
 		GetActionStack()->AddAction(FRigVMImportFromTextAction(this, Node));
 	}
 
-	Notify(ERigVMGraphNotifType::NodeAdded, Node);
-
 	if (UnitNodeCreatedContext.IsValid())
 	{
 		if (TSharedPtr<FStructOnScope> StructScope = Node->ConstructStructInstance())
@@ -1424,7 +1437,7 @@ URigVMUnitNode* URigVMController::AddUnitNode(UScriptStruct* InScriptStruct, con
 		TArray<FString> Commands = GetAddNodePythonCommands(Node);
 		for (const FString& Command : Commands)
 		{
-			RigVMPythonUtils::Print(GetGraphOuterName(),
+			RigVMPythonUtils::Print(GetSchema()->GetGraphOuterName(GetGraph()),
 				FString::Printf(TEXT("%s"), *Command));
 		}
 	}
@@ -1675,7 +1688,12 @@ URigVMVariableNode* URigVMController::AddVariableNode(const FName& InVariableNam
 		return nullptr;
 	}
 
-	FString Name = GetValidNodeName(InNodeName.IsEmpty() ? FString(TEXT("VariableNode")) : InNodeName);
+	if(!GetSchema()->SupportsExternalVariable(this, &ExternalVariable))
+	{
+		return nullptr;
+	}
+
+	FString Name = GetSchema()->GetValidNodeName(Graph, InNodeName.IsEmpty() ? FString(TEXT("VariableNode")) : InNodeName);
 	URigVMVariableNode* Node = NewObject<URigVMVariableNode>(Graph, *Name);
 	Node->Position = InPosition;
 
@@ -1710,7 +1728,10 @@ URigVMVariableNode* URigVMController::AddVariableNode(const FName& InVariableNam
 	ValuePin->Direction = bIsGetter ? ERigVMPinDirection::Output : ERigVMPinDirection::Input;
 	AddNodePin(Node, ValuePin);
 
-	Graph->Nodes.Add(Node);
+	if(!AddGraphNode(Node, false))
+	{
+		return nullptr;
+	}
 
 	if (ValuePin->IsStruct())
 	{
@@ -1757,7 +1778,7 @@ URigVMVariableNode* URigVMController::AddVariableNode(const FName& InVariableNam
 		TArray<FString> Commands = GetAddNodePythonCommands(Node);
 		for (const FString& Command : Commands)
 		{
-			RigVMPythonUtils::Print(GetGraphOuterName(), 
+			RigVMPythonUtils::Print(GetSchema()->GetGraphOuterName(GetGraph()), 
 								FString::Printf(TEXT("%s"), *Command));
 		}
 	}
@@ -2277,16 +2298,16 @@ bool URigVMController::UnresolveTemplateNodes(const TArray<FName>& InNodeNames,
 	{
 		if(bPrintPythonCommand)
 		{
-			const FString GraphName = GetSanitizedGraphName(GetGraph()->GetGraphName());
+			const FString GraphName = GetSchema()->GetSanitizedGraphName(GetGraph()->GetGraphName());
 			TArray<FString> NodeNames;
 			for(const FName& NodeName : InNodeNames)
 			{
-				NodeNames.Add(GetSanitizedNodeName(NodeName.ToString()));
+				NodeNames.Add(GetSchema()->GetSanitizedNodeName(NodeName.ToString()));
 			}
 			const FString NodeNamesJoined = FString::Join(NodeNames, TEXT("','"));
 
 			// UnresolveTemplateNodes(const TArray<FName>& InNodeNames)
-			RigVMPythonUtils::Print(GetGraphOuterName(),
+			RigVMPythonUtils::Print(GetSchema()->GetGraphOuterName(GetGraph()),
 					FString::Printf(TEXT("blueprint.get_controller_by_name('%s').unresolve_template_nodes(['%s'])"),
 					*GraphName,
 					*NodeNamesJoined));
@@ -2437,16 +2458,16 @@ TArray<URigVMNode*> URigVMController::UpgradeNodes(const TArray<FName>& InNodeNa
 
 	if(bPrintPythonCommand)
 	{
-		const FString GraphName = GetSanitizedGraphName(GetGraph()->GetGraphName());
+		const FString GraphName = GetSchema()->GetSanitizedGraphName(GetGraph()->GetGraphName());
 		TArray<FString> NodeNames;
 		for(const FName& NodeName : InNodeNames)
 		{
-			NodeNames.Add(GetSanitizedNodeName(NodeName.ToString()));
+			NodeNames.Add(GetSchema()->GetSanitizedNodeName(NodeName.ToString()));
 		}
 		const FString NodeNamesJoined = FString::Join(NodeNames, TEXT("','"));
 
 		// UpgradeNodes(const TArray<FName>& InNodeNames)
-		RigVMPythonUtils::Print(GetGraphOuterName(),
+		RigVMPythonUtils::Print(GetSchema()->GetGraphOuterName(GetGraph()),
 				FString::Printf(TEXT("blueprint.get_controller_by_name('%s').upgrade_nodes(['%s'])"),
 				*GraphName,
 				*NodeNamesJoined));
@@ -2752,14 +2773,18 @@ URigVMCommentNode* URigVMController::AddCommentNode(const FString& InCommentText
 		return nullptr;
 	}
 
-	FString Name = GetValidNodeName(InNodeName.IsEmpty() ? FString(TEXT("CommentNode")) : InNodeName);
+	FString Name = GetSchema()->GetValidNodeName(Graph, InNodeName.IsEmpty() ? FString(TEXT("CommentNode")) : InNodeName);
 	URigVMCommentNode* Node = NewObject<URigVMCommentNode>(Graph, *Name);
 	Node->Position = InPosition;
 	Node->Size = InSize;
 	Node->NodeColor = InColor;
 	Node->CommentText = InCommentText;
 
-	Graph->Nodes.Add(Node);
+	if(!AddGraphNode(Node, true))
+	{
+		return nullptr;
+	}
+	
 	if (!bSuspendNotifications)
 	{
 		Graph->MarkPackageDirty();
@@ -2774,8 +2799,6 @@ URigVMCommentNode* URigVMController::AddCommentNode(const FString& InCommentText
 		GetActionStack()->AddAction(FRigVMImportFromTextAction(this, Node));
 	}
 
-	Notify(ERigVMGraphNotifType::NodeAdded, Node);
-
 	if (bSetupUndoRedo)
 	{
 		GetActionStack()->EndAction(Action);
@@ -2786,7 +2809,7 @@ URigVMCommentNode* URigVMController::AddCommentNode(const FString& InCommentText
 		TArray<FString> Commands = GetAddNodePythonCommands(Node);
 		for (const FString& Command : Commands)
 		{
-			RigVMPythonUtils::Print(GetGraphOuterName(), 
+			RigVMPythonUtils::Print(GetSchema()->GetGraphOuterName(GetGraph()), 
 								FString::Printf(TEXT("%s"), *Command));
 		}
 	}
@@ -2845,10 +2868,10 @@ URigVMRerouteNode* URigVMController::AddRerouteNodeOnLink(URigVMLink* InLink, co
 
 	if (bPrintPythonCommand)
 	{
-		const FString GraphName = GetSanitizedGraphName(GetGraph()->GetGraphName());
-		const FString NodeName = GetSanitizedNodeName(Node->GetName());
+		const FString GraphName = GetSchema()->GetSanitizedGraphName(GetGraph()->GetGraphName());
+		const FString NodeName = GetSchema()->GetSanitizedNodeName(Node->GetName());
 
-		RigVMPythonUtils::Print(GetGraphOuterName(), 
+		RigVMPythonUtils::Print(GetSchema()->GetGraphOuterName(GetGraph()), 
 							FString::Printf(TEXT("blueprint.get_controller_by_name('%s').add_reroute_node_on_link_path('%s', %s, %s, '%s')"),
 											*GraphName,
 											*URigVMLink::GetPinPathRepresentation(SourcePin->GetPinPath(), TargetPin->GetPinPath()),
@@ -2923,7 +2946,7 @@ URigVMRerouteNode* URigVMController::AddRerouteNodeOnPin(const FString& InPinPat
 		BreakAllLinks(PinForLink, bAsInput, bSetupUndoRedo);
 	}
 
-	FString Name = GetValidNodeName(InNodeName.IsEmpty() ? FString(TEXT("RerouteNode")) : InNodeName);
+	FString Name = GetSchema()->GetValidNodeName(Graph, InNodeName.IsEmpty() ? FString(TEXT("RerouteNode")) : InNodeName);
 	URigVMRerouteNode* Node = NewObject<URigVMRerouteNode>(Graph, *Name);
 	Node->Position = InPosition;
 
@@ -2942,7 +2965,11 @@ URigVMRerouteNode* URigVMController::AddRerouteNodeOnPin(const FString& InPinPat
 		Pin->bIsExpanded = true;
 	});
 
-	Graph->Nodes.Add(Node);
+	if(!AddGraphNode(Node, true))
+	{
+		return nullptr;
+	}
+	
 	if (!bSuspendNotifications)
 	{
 		Graph->MarkPackageDirty();
@@ -2953,8 +2980,6 @@ URigVMRerouteNode* URigVMController::AddRerouteNodeOnPin(const FString& InPinPat
 		GetActionStack()->AddAction(FRigVMImportFromTextAction(this, Node, false));
 	}
 
-	Notify(ERigVMGraphNotifType::NodeAdded, Node);
-
 	if (bAsInput)
 	{
 		AddLink(ValuePin, PinForLink, bSetupUndoRedo);
@@ -2971,14 +2996,14 @@ URigVMRerouteNode* URigVMController::AddRerouteNodeOnPin(const FString& InPinPat
 
 	if (bPrintPythonCommand)
 	{
-		const FString GraphName = GetSanitizedGraphName(GetGraph()->GetGraphName());
-		const FString NodeName = GetSanitizedNodeName(Node->GetName());
+		const FString GraphName = GetSchema()->GetSanitizedGraphName(GetGraph()->GetGraphName());
+		const FString NodeName = GetSchema()->GetSanitizedNodeName(Node->GetName());
 		// AddRerouteNodeOnPin(const FString& InPinPath, bool bAsInput, const FVector2D& InPosition, const FString& InNodeName, bool bSetupUndoRedo, bool bPrintPythonCommand)
 		
-		RigVMPythonUtils::Print(GetGraphOuterName(), 
+		RigVMPythonUtils::Print(GetSchema()->GetGraphOuterName(GetGraph()), 
 							FString::Printf(TEXT("blueprint.get_controller_by_name('%s').add_reroute_node_on_pin('%s', %s, %s, %s '%s')"),
 											*GraphName,
-											*GetSanitizedPinPath(InPinPath),
+											*GetSchema()->GetSanitizedPinPath(InPinPath),
 											(bAsInput) ? TEXT("True") : TEXT("False"),
 											*RigVMPythonUtils::Vector2DToPythonString(Node->GetPosition()),
 											*NodeName));
@@ -3186,18 +3211,18 @@ URigVMInjectionInfo* URigVMController::AddInjectedNode(const FString& InPinPath,
 
 	if (bPrintPythonCommand)
 	{
-		const FString GraphName = GetSanitizedGraphName(GetGraph()->GetGraphName());
+		const FString GraphName = GetSchema()->GetSanitizedGraphName(GetGraph()->GetGraphName());
 		
-		RigVMPythonUtils::Print(GetGraphOuterName(), 
+		RigVMPythonUtils::Print(GetSchema()->GetGraphOuterName(GetGraph()), 
 							FString::Printf(TEXT("blueprint.get_controller_by_name('%s').add_injected_node_from_struct_path('%s', %s, '%s', '%s', '%s', '%s', '%s')"),
 											*GraphName,
-											*GetSanitizedPinPath(InPinPath),
+											*GetSchema()->GetSanitizedPinPath(InPinPath),
 											(bAsInput) ? TEXT("True") : TEXT("False"),
 											*InScriptStruct->GetPathName(),
 											*InMethodName.ToString(),
-											*GetSanitizedPinName(InInputPinName.ToString()),
-											*GetSanitizedPinName(InOutputPinName.ToString()),
-											*GetSanitizedNodeName(InNodeName)));
+											*GetSchema()->GetSanitizedPinName(InInputPinName.ToString()),
+											*GetSchema()->GetSanitizedPinName(InOutputPinName.ToString()),
+											*GetSchema()->GetSanitizedNodeName(InNodeName)));
 	}
 
 	return InjectionInfo;
@@ -3319,12 +3344,12 @@ bool URigVMController::RemoveInjectedNode(const FString& InPinPath, bool bAsInpu
 
 	if (bPrintPythonCommand)
 	{
-		const FString GraphName = GetSanitizedGraphName(GetGraph()->GetGraphName());
+		const FString GraphName = GetSchema()->GetSanitizedGraphName(GetGraph()->GetGraphName());
 		
-		RigVMPythonUtils::Print(GetGraphOuterName(), 
+		RigVMPythonUtils::Print(GetSchema()->GetGraphOuterName(GetGraph()), 
 							FString::Printf(TEXT("blueprint.get_controller_by_name('%s').remove_injected_node('%s', %s)"),
 											*GraphName,
-											*GetSanitizedPinPath(InPinPath),
+											*GetSchema()->GetSanitizedPinPath(InPinPath),
 											(bAsInput) ? TEXT("True") : TEXT("False")));
 	}
 
@@ -3605,12 +3630,12 @@ URigVMNode* URigVMController::EjectNodeFromPin(URigVMPin* InPin, bool bSetupUndo
 
 	if (bPrintPythonCommand)
 	{
-		const FString GraphName = GetSanitizedGraphName(GetGraph()->GetGraphName());
+		const FString GraphName = GetSchema()->GetSanitizedGraphName(GetGraph()->GetGraphName());
 
-		RigVMPythonUtils::Print(GetGraphOuterName(), 
+		RigVMPythonUtils::Print(GetSchema()->GetGraphOuterName(GetGraph()), 
 							FString::Printf(TEXT("blueprint.get_controller_by_name('%s').eject_node_from_pin('%s')"),
 											*GraphName,
-											*GetSanitizedPinPath(InPin->GetPinPath())));
+											*GetSchema()->GetSanitizedPinPath(InPin->GetPinPath())));
 	}
 
 	return NodeToEject;
@@ -3838,6 +3863,26 @@ protected:
 			DefaultObject->IsA<URigVMPin>();
 	}
 
+	virtual bool CanCreateObject(UObject* InParent, UClass* ObjectClass, const FName& InDesiredName) const override
+	{
+		if(const URigVMGraph* Graph = Cast<URigVMGraph>(InParent))
+		{
+			if(const URigVMLibraryNode* LibraryNode = Cast<URigVMLibraryNode>(Graph->GetOuter()))
+			{
+				// avoid creating duplicate entry / return nodes
+				if(ObjectClass == URigVMFunctionEntryNode::StaticClass() && LibraryNode->GetEntryNode() != nullptr)
+				{
+					return false;
+				}
+				if(ObjectClass == URigVMFunctionReturnNode::StaticClass() && LibraryNode->GetReturnNode() != nullptr)
+				{
+					return false;
+				}
+			}		
+		}
+		return true;
+	}
+
 	virtual void UpdateObjectName(UClass* ObjectClass, UObject* InParent, FName& InOutObjName) override
 	{
 		if (ObjectClass->GetDefaultObject()->IsA<URigVMNode>())
@@ -3856,13 +3901,17 @@ protected:
 			TArray<FName>& CreateNodeNames = CreateNodeNamesPerGraph.FindOrAdd(Graph);
 			TMap<FString, FString>& NodeNameMap = NodeNameMapPerGraph.FindOrAdd(Graph);
 
-			const FName ValidName = URigVMController::GetUniqueName(InOutObjName, [Graph, this, CreateNodeNames](const FName& InName) {
+			const FName ValidName = URigVMSchema::GetUniqueName(InOutObjName, [Graph, this, CreateNodeNames](const FName& InName) {
 				return !CreateNodeNames.Contains(InName) && Graph->IsNameAvailable(InName.ToString());
 			}, false, true);
 
 			if(!InOutObjName.IsEqual(ValidName, ENameCase::CaseSensitive))
 			{
-				NodeNameMap.Add(InOutObjName.ToString(), ValidName.ToString());
+				if(ObjectClass != URigVMFunctionEntryNode::StaticClass() &&
+					ObjectClass != URigVMFunctionReturnNode::StaticClass())
+				{
+					NodeNameMap.Add(InOutObjName.ToString(), ValidName.ToString());
+				}
 			}
 			
 			CreateNodeNames.Add(ValidName);
@@ -3998,14 +4047,11 @@ TArray<FName> URigVMController::ImportNodesFromText(const FString& InText, bool
 
 			if(URigVMController* ControllerForGraph = GetControllerForGraph(Graph))
 			{
-				if(!ControllerForGraph->CanAddNode(CreatedNode, true))
+				if(!ControllerForGraph->AddGraphNode(CreatedNode, false))
 				{
-					Graph->Nodes.Remove(CreatedNode);
-					DestroyObject(CreatedNode);
 					continue;
 				}
 
-				Graph->Nodes.AddUnique(CreatedNode);
 				FilteredNodes.Add(CreatedNode);
 			}
 		}
@@ -4198,9 +4244,9 @@ TArray<FName> URigVMController::ImportNodesFromText(const FString& InText, bool
 		PythonContent = InText.Replace(TEXT("'"), TEXT("\\'"));
 		PythonContent = PythonContent.Replace(TEXT("\r\n"), TEXT("\\r\\n'\r\n'"));
 
-		const FString GraphName = GetSanitizedGraphName(GetGraph()->GetGraphName());
+		const FString GraphName = GetSchema()->GetSanitizedGraphName(GetGraph()->GetGraphName());
 
-		RigVMPythonUtils::Print(GetGraphOuterName(),
+		RigVMPythonUtils::Print(GetSchema()->GetGraphOuterName(GetGraph()),
 			FString::Printf(TEXT("blueprint.get_controller_by_name('%s').import_nodes_from_text('%s')"),
 			*GraphName,
 			*PythonContent));
@@ -4437,11 +4483,11 @@ TMap<FRigVMGraphFunctionIdentifier, URigVMLibraryNode*> URigVMController::Locali
 	{
 		for (const FRigVMGraphFunctionIdentifier& Identifier : InFunctionDefinitions)
 		{
-			const FString GraphName = GetSanitizedGraphName(GetGraph()->GetGraphName());
+			const FString GraphName = GetSchema()->GetSanitizedGraphName(GetGraph()->GetGraphName());
 			FRigVMGraphFunctionData* FunctionData = FRigVMGraphFunctionData::FindFunctionData(Identifier);
-			const FString FunctionDefinitionName = GetSanitizedNodeName(FunctionData->Header.Name.ToString());
+			const FString FunctionDefinitionName = GetSchema()->GetSanitizedNodeName(FunctionData->Header.Name.ToString());
 
-			RigVMPythonUtils::Print(GetGraphOuterName(), 
+			RigVMPythonUtils::Print(GetSchema()->GetGraphOuterName(GetGraph()), 
 				FString::Printf(TEXT("blueprint.get_controller_by_name('%s').localize_function_from_path('%s', '%s', %s)"),
 						*GraphName,
 						*Identifier.HostObject.ToString(),
@@ -4453,21 +4499,6 @@ TMap<FRigVMGraphFunctionIdentifier, URigVMLibraryNode*> URigVMController::Locali
 	return LocalizedFunctions;
 }
 
-FName URigVMController::GetUniqueName(const FName& InName, TFunction<bool(const FName&)> IsNameAvailableFunction, bool bAllowPeriod, bool bAllowSpace)
-{
-	FString SanitizedPrefix = InName.ToString();
-	SanitizeName(SanitizedPrefix, bAllowPeriod, bAllowSpace);
-
-	int32 NameSuffix = 0;
-	FString Name = SanitizedPrefix;
-	while (!IsNameAvailableFunction(*Name))
-	{
-		NameSuffix++;
-		Name = FString::Printf(TEXT("%s_%d"), *SanitizedPrefix, NameSuffix);
-	}
-	return *Name;
-}
-
 URigVMCollapseNode* URigVMController::CollapseNodes(const TArray<FName>& InNodeNames, const FString& InCollapseNodeName, bool bSetupUndoRedo, bool bPrintPythonCommand, bool bIsAggregate)
 {
 	if (!IsValidGraph())
@@ -4509,9 +4540,9 @@ URigVMCollapseNode* URigVMController::CollapseNodes(const TArray<FName>& InNodeN
 		}
 		ArrayStr += TEXT("]");
 
-		const FString GraphName = GetSanitizedGraphName(GetGraph()->GetGraphName());
+		const FString GraphName = GetSchema()->GetSanitizedGraphName(GetGraph()->GetGraphName());
 
-		RigVMPythonUtils::Print(GetGraphOuterName(), 
+		RigVMPythonUtils::Print(GetSchema()->GetGraphOuterName(GetGraph()), 
 							FString::Printf(TEXT("blueprint.get_controller_by_name('%s').collapse_nodes(%s, '%s')"),
 											*GraphName,
 											*ArrayStr,
@@ -4554,10 +4585,10 @@ TArray<URigVMNode*> URigVMController::ExpandLibraryNode(const FName& InNodeName,
 
 	if (!Nodes.IsEmpty() && bPrintPythonCommand)
 	{
-		const FString GraphName = GetSanitizedGraphName(GetGraph()->GetGraphName());
-		const FString NodeName = GetSanitizedNodeName(Node->GetName());
+		const FString GraphName = GetSchema()->GetSanitizedGraphName(GetGraph()->GetGraphName());
+		const FString NodeName = GetSchema()->GetSanitizedNodeName(Node->GetName());
 		
-		RigVMPythonUtils::Print(GetGraphOuterName(), 
+		RigVMPythonUtils::Print(GetSchema()->GetGraphOuterName(GetGraph()), 
 							FString::Printf(TEXT("blueprint.get_controller_by_name('%s').expand_library_node('%s')"),
 											*GraphName,
 											*NodeName));
@@ -4590,6 +4621,14 @@ URigVMCollapseNode* URigVMController::CollapseNodes(const TArray<URigVMNode*>& I
 		return nullptr;
 	}
 
+	{
+		const TArrayView<URigVMNode* const> NodesView((URigVMNode** const)InNodes.GetData(), InNodes.Num());
+		if(!GetSchema()->CanCollapseNodes(this, NodesView))
+		{
+			return nullptr;
+		}
+	}
+
 #if UE_RIGVM_AGGREGATE_NODES_ENABLED
 	if (bIsAggregate)
 	{
@@ -4718,7 +4757,7 @@ URigVMCollapseNode* URigVMController::CollapseNodes(const TArray<URigVMNode*>& I
 	FRigVMControllerCompileBracketScope CompileScope(this);
 	FRigVMCollapseNodesAction CollapseAction;
 
-	FString CollapseNodeName = GetValidNodeName(InCollapseNodeName.IsEmpty() ? FString(TEXT("CollapseNode")) : InCollapseNodeName);
+	FString CollapseNodeName = GetSchema()->GetValidNodeName(Graph, InCollapseNodeName.IsEmpty() ? FString(TEXT("CollapseNode")) : InCollapseNodeName);
 
 	if (bSetupUndoRedo)
 	{
@@ -4752,7 +4791,11 @@ URigVMCollapseNode* URigVMController::CollapseNodes(const TArray<URigVMNode*>& I
 #endif
 	
 	CollapseNode->Position = Center;
-	Graph->Nodes.Add(CollapseNode);
+
+	if(!AddGraphNode(CollapseNode, false))
+	{
+		return nullptr;
+	}
 
 	// now looper over the links to be rewired
 	for (URigVMLink* Link : LinksToRewire)
@@ -4804,7 +4847,7 @@ URigVMCollapseNode* URigVMController::CollapseNodes(const TArray<URigVMNode*>& I
 			}
 		}
 
-		FName PinName = GetUniqueName(PinToCollapse->GetFName(), [CollapseNode](const FName& InName) {
+		FName PinName = URigVMSchema::GetUniqueName(PinToCollapse->GetFName(), [CollapseNode](const FName& InName) {
 			return CollapseNode->FindPin(InName.ToString()) == nullptr;
 		}, false, true);
 
@@ -4846,26 +4889,32 @@ URigVMCollapseNode* URigVMController::CollapseNodes(const TArray<URigVMNode*>& I
 	if(URigVMController* CollapseController = GetControllerForGraph(CollapseNode->GetContainedGraph()))
 	{
 		EntryNode = NewObject<URigVMFunctionEntryNode>(CollapseNode->ContainedGraph, TEXT("Entry"));
-		CollapseNode->ContainedGraph->Nodes.Add(EntryNode);
-		EntryNode->Position = -Diagonal * 0.5f - FVector2D(250.f, 0.f);
+
+		if(CollapseController->AddGraphNode(EntryNode, false))
 		{
-			TGuardValue<bool> SuspendNotifications(CollapseController->bSuspendNotifications, true);
-			CollapseController->RefreshFunctionPins(EntryNode);
-		}
+			EntryNode->Position = -Diagonal * 0.5f - FVector2D(250.f, 0.f);
+			{
+				TGuardValue<bool> SuspendNotifications(CollapseController->bSuspendNotifications, true);
+				CollapseController->RefreshFunctionPins(EntryNode);
+			}
 
-		CollapseController->Notify(ERigVMGraphNotifType::NodeAdded, EntryNode);
+			CollapseController->Notify(ERigVMGraphNotifType::NodeAdded, EntryNode);
+		}
 
 		if (bContainsOutputs)
 		{
 			ReturnNode = NewObject<URigVMFunctionReturnNode>(CollapseNode->ContainedGraph, TEXT("Return"));
-			CollapseNode->ContainedGraph->Nodes.Add(ReturnNode);
-			ReturnNode->Position = FVector2D(Diagonal.X, -Diagonal.Y) * 0.5f + FVector2D(300.f, 0.f);
+
+			if(CollapseController->AddGraphNode(ReturnNode, false))
 			{
-				TGuardValue<bool> SuspendNotifications(CollapseController->bSuspendNotifications, true);
-				CollapseController->RefreshFunctionPins(ReturnNode);
-			}
+				ReturnNode->Position = FVector2D(Diagonal.X, -Diagonal.Y) * 0.5f + FVector2D(300.f, 0.f);
+				{
+					TGuardValue<bool> SuspendNotifications(CollapseController->bSuspendNotifications, true);
+					CollapseController->RefreshFunctionPins(ReturnNode);
+				}
 
-			CollapseController->Notify(ERigVMGraphNotifType::NodeAdded, ReturnNode);
+				CollapseController->Notify(ERigVMGraphNotifType::NodeAdded, ReturnNode);
+			}
 		}
 	}
 
@@ -5098,7 +5147,7 @@ URigVMCollapseNode* URigVMController::CollapseNodes(const TArray<URigVMNode*>& I
 
 	if (!InCollapseNodeName.IsEmpty() && CollapseNodeName != InCollapseNodeName)
 	{		
-		FString ValidName = GetValidNodeName(InCollapseNodeName);
+		FString ValidName = GetSchema()->GetValidNodeName(Graph, InCollapseNodeName);
 		if (ValidName == InCollapseNodeName)
 		{
 			RenameNode(CollapseNode, *ValidName, false);
@@ -5128,9 +5177,8 @@ TArray<URigVMNode*> URigVMController::ExpandLibraryNode(URigVMLibraryNode* InNod
 	URigVMGraph* Graph = GetGraph();
 	check(Graph);
 
-	if (Graph->IsA<URigVMFunctionLibrary>())
+	if(!GetSchema()->CanExpandNode(this, InNode))
 	{
-		ReportError(TEXT("Cannot expand nodes in function library graphs."));
 		return TArray<URigVMNode*>();
 	}
 
@@ -5846,12 +5894,12 @@ FName URigVMController::PromoteCollapseNodeToFunctionReferenceNode(const FName&
 	{
 		if (bPrintPythonCommand)
 		{
-			const FString GraphName = GetSanitizedGraphName(GetGraph()->GetGraphName());
+			const FString GraphName = GetSchema()->GetSanitizedGraphName(GetGraph()->GetGraphName());
 			
-			RigVMPythonUtils::Print(GetGraphOuterName(), 
+			RigVMPythonUtils::Print(GetSchema()->GetGraphOuterName(GetGraph()), 
 								FString::Printf(TEXT("blueprint.get_controller_by_name('%s').promote_collapse_node_to_function_reference_node('%s')"),
 													*GraphName,
-													*GetSanitizedNodeName(InNodeName.ToString())));
+													*GetSchema()->GetSanitizedNodeName(InNodeName.ToString())));
 		}
 		
 		return Result->GetFName();
@@ -5929,7 +5977,7 @@ URigVMFunctionReferenceNode* URigVMController::PromoteCollapseNodeToFunctionRefe
 		{
 			FRigVMControllerCompileBracketScope LibraryCompileScope(LibraryController);
 
-			const FString FunctionName = LibraryController->GetValidNodeName(InCollapseNode->GetName());
+			const FString FunctionName = LibraryController->GetSchema()->GetValidNodeName(FunctionLibrary, InCollapseNode->GetName());
 			FunctionDefinition = LibraryController->AddFunctionToLibrary(
 				*FunctionName,
 				InCollapseNode->GetPins().ContainsByPredicate([](const URigVMPin* Pin) -> bool
@@ -6144,8 +6192,11 @@ URigVMCollapseNode* URigVMController::PromoteFunctionReferenceNodeToCollapseNode
 				
 		CollapseNode->NodeColor = FLinearColor::White;
 		CollapseNode->Position = NodePosition;
-		Graph->Nodes.Add(CollapseNode);
-		Notify(ERigVMGraphNotifType::NodeAdded, CollapseNode);
+
+		if(!AddGraphNode(CollapseNode, false))
+		{
+			return nullptr;
+		}
 
 		ApplyPinStates(CollapseNode, PinStates);
 		RestoreLinkedPaths(LinkedPaths, RestoreSettings);
@@ -6261,6 +6312,11 @@ bool URigVMController::RemoveNodes(TArray<URigVMNode*> InNodes, bool bSetupUndoR
 			return false;
 		}
 
+		if(!GetSchema()->CanRemoveNode(this, InNode))
+		{
+			return false;
+		}
+
 		if (InNode->IsInjected())
 		{
 			URigVMInjectionInfo* InjectionInfo = InNode->GetInjectionInfo();
@@ -6526,15 +6582,15 @@ bool URigVMController::RemoveNodes(TArray<URigVMNode*> InNodes, bool bSetupUndoR
 
 	if (bPrintPythonCommand)
 	{
-		const FString GraphName = GetSanitizedGraphName(GetGraph()->GetGraphName());
+		const FString GraphName = GetSchema()->GetSanitizedGraphName(GetGraph()->GetGraphName());
 
 		if (Graph->IsA<URigVMFunctionLibrary>())
 		{
 			for(const URigVMNode* InNode : FilteredNodes)
 			{
-				const FString NodeName = GetSanitizedNodeName(InNode->GetName());
+				const FString NodeName = GetSchema()->GetSanitizedNodeName(InNode->GetName());
 				
-				RigVMPythonUtils::Print(GetGraphOuterName(), 
+				RigVMPythonUtils::Print(GetSchema()->GetGraphOuterName(GetGraph()), 
 									FString::Printf(TEXT("library_controller.remove_function_from_library('%s')"),
 													*NodeName));
 			}
@@ -6542,10 +6598,10 @@ bool URigVMController::RemoveNodes(TArray<URigVMNode*> InNodes, bool bSetupUndoR
 		else
 		{
 			TArray<FString> NodePaths;
-			Algo::Transform(FilteredNodes, NodePaths, [](const URigVMNode* Node)
+			Algo::Transform(FilteredNodes, NodePaths, [this](const URigVMNode* Node)
 			{
 				static constexpr TCHAR Format[] = TEXT("'%s'");
-				return FString::Printf(Format, *GetSanitizedPinPath(Node->GetNodePath()));
+				return FString::Printf(Format, *GetSchema()->GetSanitizedPinPath(Node->GetNodePath()));
 			});
 
 			const FString NodePathsJoined = FString::Join(NodePaths, TEXT(", "));
@@ -6553,7 +6609,7 @@ bool URigVMController::RemoveNodes(TArray<URigVMNode*> InNodes, bool bSetupUndoR
 			FString PythonCmd = FString::Printf(TEXT("blueprint.get_controller_by_name('%s')."), *GraphName );
 			PythonCmd += FString::Printf(TEXT("remove_node_by_name([%s])"), *NodePathsJoined);
 			
-			RigVMPythonUtils::Print(GetGraphOuterName(), PythonCmd );
+			RigVMPythonUtils::Print(GetSchema()->GetGraphOuterName(GetGraph()), PythonCmd );
 		}
 	}
 
@@ -6641,8 +6697,9 @@ bool URigVMController::RenameNode(URigVMNode* InNode, const FName& InNewName, bo
 		return false;
 	}
 
-	FName ValidNewName = *GetValidNodeName(InNewName.ToString());
-	if (InNode->GetFName() == ValidNewName)
+	FName ValidNewName = *GetSchema()->GetValidNodeName(GetGraph(), InNewName.ToString());
+
+	if(!GetSchema()->CanRenameNode(this, InNode, ValidNewName))
 	{
 		return false;
 	}
@@ -6705,7 +6762,7 @@ bool URigVMController::RenameNode(URigVMNode* InNode, const FName& InNewName, bo
 
 	if (GetGraph()->IsA<URigVMFunctionLibrary>())
 	{
-		RigVMPythonUtils::Print(GetGraphOuterName(), 
+		RigVMPythonUtils::Print(GetSchema()->GetGraphOuterName(GetGraph()), 
 							FString::Printf(TEXT("library_controller.rename_function('%s', '%s')"),
 										*OldName,
 										*InNewName.ToString()));
@@ -6857,7 +6914,7 @@ bool URigVMController::SetNodeSelection(const TArray<FName>& InNodeNames, bool b
 		FString ArrayStr = TEXT("[");
 		for (auto It = InNodeNames.CreateConstIterator(); It; ++It)
 		{
-			ArrayStr += TEXT("'") + GetSanitizedNodeName(It->ToString()) + TEXT("'");
+			ArrayStr += TEXT("'") + GetSchema()->GetSanitizedNodeName(It->ToString()) + TEXT("'");
 			if (It.GetIndex() < InNodeNames.Num() - 1)
 			{
 				ArrayStr += TEXT(", ");
@@ -6865,9 +6922,9 @@ bool URigVMController::SetNodeSelection(const TArray<FName>& InNodeNames, bool b
 		}
 		ArrayStr += TEXT("]");
 
-		const FString GraphName = GetSanitizedGraphName(GetGraph()->GetGraphName());
+		const FString GraphName = GetSchema()->GetSanitizedGraphName(GetGraph()->GetGraphName());
 
-		RigVMPythonUtils::Print(GetGraphOuterName(), 
+		RigVMPythonUtils::Print(GetSchema()->GetGraphOuterName(GetGraph()), 
 							FString::Printf(TEXT("blueprint.get_controller_by_name('%s').set_node_selection(%s)"),
 											*GraphName,
 											*ArrayStr));
@@ -6888,7 +6945,7 @@ bool URigVMController::SetNodePosition(URigVMNode* InNode, const FVector2D& InPo
 		return false;
 	}
 
-	if ((InNode->Position - InPosition).IsNearlyZero())
+	if(!GetSchema()->CanMoveNode(this, InNode, InPosition))
 	{
 		return false;
 	}
@@ -6911,10 +6968,10 @@ bool URigVMController::SetNodePosition(URigVMNode* InNode, const FVector2D& InPo
 
 	if (bPrintPythonCommand)
 	{
-		const FString GraphName = GetSanitizedGraphName(GetGraph()->GetGraphName());
-		const FString NodePath = GetSanitizedPinPath(InNode->GetNodePath());
+		const FString GraphName = GetSchema()->GetSanitizedGraphName(GetGraph()->GetGraphName());
+		const FString NodePath = GetSchema()->GetSanitizedPinPath(InNode->GetNodePath());
 
-		RigVMPythonUtils::Print(GetGraphOuterName(), 
+		RigVMPythonUtils::Print(GetSchema()->GetGraphOuterName(GetGraph()), 
 							FString::Printf(TEXT("blueprint.get_controller_by_name('%s').set_node_position_by_name('%s', %s)"),
 											*GraphName,
 											*NodePath,
@@ -6956,7 +7013,7 @@ bool URigVMController::SetNodeSize(URigVMNode* InNode, const FVector2D& InSize,
 		return false;
 	}
 
-	if ((InNode->Size - InSize).IsNearlyZero())
+	if(!GetSchema()->CanResizeNode(this, InNode, InSize))
 	{
 		return false;
 	}
@@ -6979,10 +7036,10 @@ bool URigVMController::SetNodeSize(URigVMNode* InNode, const FVector2D& InSize,
 
 	if (bPrintPythonCommand)
 	{
-		const FString GraphName = GetSanitizedGraphName(GetGraph()->GetGraphName());
-		const FString NodePath = GetSanitizedPinPath(InNode->GetNodePath());
+		const FString GraphName = GetSchema()->GetSanitizedGraphName(GetGraph()->GetGraphName());
+		const FString NodePath = GetSchema()->GetSanitizedPinPath(InNode->GetNodePath());
 
-		RigVMPythonUtils::Print(GetGraphOuterName(), 
+		RigVMPythonUtils::Print(GetSchema()->GetGraphOuterName(GetGraph()), 
 							FString::Printf(TEXT("blueprint.get_controller_by_name('%s').set_node_size_by_name('%s', %s)"),
 											*GraphName,
 											*NodePath,
@@ -7023,7 +7080,7 @@ bool URigVMController::SetNodeColor(URigVMNode* InNode, const FLinearColor& InCo
 		return false;
 	}
 
-	if ((FVector4(InNode->NodeColor) - FVector4(InColor)).IsNearlyZero3())
+	if(!GetSchema()->CanRecolorNode(this, InNode, InColor))
 	{
 		return false;
 	}
@@ -7060,10 +7117,10 @@ bool URigVMController::SetNodeColor(URigVMNode* InNode, const FLinearColor& InCo
 
 	if (bPrintPythonCommand)
 	{
-		const FString GraphName = GetSanitizedGraphName(GetGraph()->GetGraphName());
-		const FString NodePath = GetSanitizedPinPath(InNode->GetNodePath());
+		const FString GraphName = GetSchema()->GetSanitizedGraphName(GetGraph()->GetGraphName());
+		const FString NodePath = GetSchema()->GetSanitizedPinPath(InNode->GetNodePath());
 
-		RigVMPythonUtils::Print(GetGraphOuterName(),
+		RigVMPythonUtils::Print(GetSchema()->GetGraphOuterName(GetGraph()),
 			FString::Printf(TEXT("blueprint.get_controller_by_name('%s').set_node_color_by_name('%s', %s)"),
 				*GraphName,
 				*NodePath,
@@ -7127,10 +7184,10 @@ bool URigVMController::SetNodeCategory(URigVMCollapseNode* InNode, const FString
 
 	if (bPrintPythonCommand)
 	{
-		const FString GraphName = GetSanitizedGraphName(GetGraph()->GetGraphName());
-		const FString NodePath = GetSanitizedPinPath(InNode->GetNodePath());
+		const FString GraphName = GetSchema()->GetSanitizedGraphName(GetGraph()->GetGraphName());
+		const FString NodePath = GetSchema()->GetSanitizedPinPath(InNode->GetNodePath());
 
-		RigVMPythonUtils::Print(GetGraphOuterName(),
+		RigVMPythonUtils::Print(GetSchema()->GetGraphOuterName(GetGraph()),
 			FString::Printf(TEXT("blueprint.get_controller_by_name('%s').set_node_category_by_name('%s', '%s')"),
 				*GraphName,
 				*NodePath,
@@ -7194,10 +7251,10 @@ bool URigVMController::SetNodeKeywords(URigVMCollapseNode* InNode, const FString
 
 	if (bPrintPythonCommand)
 	{
-		const FString GraphName = GetSanitizedGraphName(GetGraph()->GetGraphName());
-		const FString NodePath = GetSanitizedPinPath(InNode->GetNodePath());
+		const FString GraphName = GetSchema()->GetSanitizedGraphName(GetGraph()->GetGraphName());
+		const FString NodePath = GetSchema()->GetSanitizedPinPath(InNode->GetNodePath());
 
-		RigVMPythonUtils::Print(GetGraphOuterName(),
+		RigVMPythonUtils::Print(GetSchema()->GetGraphOuterName(GetGraph()),
 			FString::Printf(TEXT("blueprint.get_controller_by_name('%s').set_node_keywords_by_name('%s', '%s')"),
 				*GraphName,
 				*NodePath,
@@ -7261,10 +7318,10 @@ bool URigVMController::SetNodeDescription(URigVMCollapseNode* InNode, const FStr
 
 	if (bPrintPythonCommand)
 	{
-		const FString GraphName = GetSanitizedGraphName(GetGraph()->GetGraphName());
-		const FString NodePath = GetSanitizedPinPath(InNode->GetNodePath());
+		const FString GraphName = GetSchema()->GetSanitizedGraphName(GetGraph()->GetGraphName());
+		const FString NodePath = GetSchema()->GetSanitizedPinPath(InNode->GetNodePath());
 
-		RigVMPythonUtils::Print(GetGraphOuterName(),
+		RigVMPythonUtils::Print(GetSchema()->GetGraphOuterName(GetGraph()),
 			FString::Printf(TEXT("blueprint.get_controller_by_name('%s').set_node_description_by_name('%s', '%s')"),
 				*GraphName,
 				*NodePath,
@@ -7333,10 +7390,10 @@ bool URigVMController::SetCommentText(URigVMNode* InNode, const FString& InComme
 
 		if (bPrintPythonCommand)
 		{
-			const FString GraphName = GetSanitizedGraphName(GetGraph()->GetGraphName());
-			const FString NodePath = GetSanitizedPinPath(CommentNode->GetNodePath());
+			const FString GraphName = GetSchema()->GetSanitizedGraphName(GetGraph()->GetGraphName());
+			const FString NodePath = GetSchema()->GetSanitizedPinPath(CommentNode->GetNodePath());
 
-			RigVMPythonUtils::Print(GetGraphOuterName(),
+			RigVMPythonUtils::Print(GetSchema()->GetGraphOuterName(GetGraph()),
 				FString::Printf(TEXT("blueprint.get_controller_by_name('%s').set_comment_text_by_name('%s', '%s')"),
 				*GraphName,
 				*NodePath,
@@ -7485,12 +7542,12 @@ bool URigVMController::SetPinExpansion(const FString& InPinPath, bool bIsExpande
 	const bool bSuccess = SetPinExpansion(Pin, bIsExpanded, bSetupUndoRedo);
 	if (bSuccess && bPrintPythonCommand)
 	{
-		const FString GraphName = GetSanitizedGraphName(GetGraph()->GetGraphName());
+		const FString GraphName = GetSchema()->GetSanitizedGraphName(GetGraph()->GetGraphName());
 
-		RigVMPythonUtils::Print(GetGraphOuterName(),
+		RigVMPythonUtils::Print(GetSchema()->GetGraphOuterName(GetGraph()),
 			FString::Printf(TEXT("blueprint.get_controller_by_name('%s').set_pin_expansion('%s', %s)"),
 			*GraphName,
-			*GetSanitizedPinPath(InPinPath),
+			*GetSchema()->GetSanitizedPinPath(InPinPath),
 			(bIsExpanded) ? TEXT("True") : TEXT("False")));
 	}
 
@@ -7679,12 +7736,12 @@ bool URigVMController::SetPinDefaultValue(const FString& InPinPath, const FStrin
 
 	if (bPrintPythonCommand)
 	{
-		const FString GraphName = GetSanitizedGraphName(GetGraph()->GetGraphName());
+		const FString GraphName = GetSchema()->GetSanitizedGraphName(GetGraph()->GetGraphName());
 
-		RigVMPythonUtils::Print(GetGraphOuterName(),
+		RigVMPythonUtils::Print(GetSchema()->GetGraphOuterName(GetGraph()),
 			FString::Printf(TEXT("blueprint.get_controller_by_name('%s').set_pin_default_value('%s', '%s', %s)"),
 			*GraphName,
-			*GetSanitizedPinPath(InPinPath),
+			*GetSchema()->GetSanitizedPinPath(InPinPath),
 			*InDefaultValue,
 			(bResizeArrays) ? TEXT("True") : TEXT("False")));
 	}
@@ -7734,7 +7791,7 @@ bool URigVMController::SetPinDefaultValue(URigVMPin* InPin, const FString& InDef
 	bool bSetPinDefaultValueSucceeded = false;
 	if (InPin->IsArray())
 	{
-		if (ShouldPinBeUnfolded(InPin))
+		if (GetSchema()->CanUnfoldPin(this, InPin))
 		{
 			TArray<FString> Elements = URigVMPin::SplitDefaultValue(ClampedDefaultValue);
 
@@ -7851,12 +7908,12 @@ bool URigVMController::ResetPinDefaultValue(const FString& InPinPath, bool bSetu
 	const bool bSuccess = ResetPinDefaultValue(Pin, bSetupUndoRedo);
 	if (bSuccess && bPrintPythonCommand)
 	{
-		const FString GraphName = GetSanitizedGraphName(GetGraph()->GetGraphName());
+		const FString GraphName = GetSchema()->GetSanitizedGraphName(GetGraph()->GetGraphName());
 
-		RigVMPythonUtils::Print(GetGraphOuterName(),
+		RigVMPythonUtils::Print(GetSchema()->GetGraphOuterName(GetGraph()),
 			FString::Printf(TEXT("blueprint.get_controller_by_name('%s').reset_pin_default_value('%s')"),
 			*GraphName,
-			*GetSanitizedPinPath(InPinPath)));
+			*GetSchema()->GetSanitizedPinPath(InPinPath)));
 	}
 
 	return bSuccess;
@@ -8371,10 +8428,10 @@ FString URigVMController::AddAggregatePin(URigVMNode* InNode, const FString& InP
 
 	if (bPrintPythonCommand)
 	{
-		const FString GraphName = GetSanitizedGraphName(GetGraph()->GetGraphName());
-		const FString NodePath = GetSanitizedPinPath(InNode->GetNodePath());
+		const FString GraphName = GetSchema()->GetSanitizedGraphName(GetGraph()->GetGraphName());
+		const FString NodePath = GetSchema()->GetSanitizedPinPath(InNode->GetNodePath());
 
-		RigVMPythonUtils::Print(GetGraphOuterName(),
+		RigVMPythonUtils::Print(GetSchema()->GetGraphOuterName(GetGraph()),
 			FString::Printf(TEXT("blueprint.get_controller_by_name('%s').add_aggregate_pin('%s', '%s', '%s')"),
 			*GraphName,
 			*NodePath,
@@ -8599,10 +8656,10 @@ bool URigVMController::RemoveAggregatePin(URigVMPin* InPin, bool bSetupUndoRedo,
 
 	if (bSuccess && bPrintPythonCommand)
 	{
-		const FString GraphName = GetSanitizedGraphName(GetGraph()->GetGraphName());
-		const FString PinPath = GetSanitizedPinPath(InPin->GetPinPath());
+		const FString GraphName = GetSchema()->GetSanitizedGraphName(GetGraph()->GetGraphName());
+		const FString PinPath = GetSchema()->GetSanitizedPinPath(InPin->GetPinPath());
 
-		RigVMPythonUtils::Print(GetGraphOuterName(),
+		RigVMPythonUtils::Print(GetSchema()->GetGraphOuterName(GetGraph()),
 			FString::Printf(TEXT("blueprint.get_controller_by_name('%s').remove_aggregate_pin('%s')"),
 			*GraphName,
 			*PinPath));
@@ -8683,12 +8740,12 @@ FString URigVMController::InsertArrayPin(const FString& InArrayPinPath, int32 In
 	{
 		if (bPrintPythonCommand)
 		{
-			const FString GraphName = GetSanitizedGraphName(GetGraph()->GetGraphName());
+			const FString GraphName = GetSchema()->GetSanitizedGraphName(GetGraph()->GetGraphName());
 			
-			RigVMPythonUtils::Print(GetGraphOuterName(),
+			RigVMPythonUtils::Print(GetSchema()->GetGraphOuterName(GetGraph()),
 				FString::Printf(TEXT("blueprint.get_controller_by_name('%s').insert_array_pin('%s', %d, '%s')"),
 				*GraphName,
-				*GetSanitizedPinPath(InArrayPinPath),
+				*GetSchema()->GetSanitizedPinPath(InArrayPinPath),
 				InIndex,
 				*InDefaultValue));
 		}
@@ -8712,7 +8769,7 @@ URigVMPin* URigVMController::InsertArrayPin(URigVMPin* ArrayPin, int32 InIndex,
 		return nullptr;
 	}
 
-	if (!ShouldPinBeUnfolded(ArrayPin))
+	if (!GetSchema()->CanUnfoldPin(this, ArrayPin))
 	{
 		ReportErrorf(TEXT("Cannot insert array pin under '%s'."), *ArrayPin->GetPinPath());
 		return nullptr;
@@ -8874,12 +8931,12 @@ bool URigVMController::RemoveArrayPin(const FString& InArrayElementPinPath, bool
 
 	if (bPrintPythonCommand)
 	{
-		const FString GraphName = GetSanitizedGraphName(GetGraph()->GetGraphName());
+		const FString GraphName = GetSchema()->GetSanitizedGraphName(GetGraph()->GetGraphName());
 
-		RigVMPythonUtils::Print(GetGraphOuterName(),
+		RigVMPythonUtils::Print(GetSchema()->GetGraphOuterName(GetGraph()),
 			FString::Printf(TEXT("blueprint.get_controller_by_name('%s').remove_array_pin('%s')"),
 			*GraphName,
-			*GetSanitizedPinPath(InArrayElementPinPath)));
+			*GetSchema()->GetSanitizedPinPath(InArrayElementPinPath)));
 	}
 
 	return true;
@@ -9066,12 +9123,12 @@ bool URigVMController::BindPinToVariable(const FString& InPinPath, const FString
 	
 	if (bSuccess && bPrintPythonCommand)
 	{
-		const FString GraphName = GetSanitizedGraphName(GetGraph()->GetGraphName());
+		const FString GraphName = GetSchema()->GetSanitizedGraphName(GetGraph()->GetGraphName());
 		
-		RigVMPythonUtils::Print(GetGraphOuterName(),
+		RigVMPythonUtils::Print(GetSchema()->GetGraphOuterName(GetGraph()),
 			FString::Printf(TEXT("blueprint.get_controller_by_name('%s').bind_pin_to_variable('%s', '%s')"),
 			*GraphName,
-			*GetSanitizedPinPath(InPinPath),
+			*GetSchema()->GetSanitizedPinPath(InPinPath),
 			*InNewBoundVariablePath));
 	}
 	
@@ -9090,27 +9147,6 @@ bool URigVMController::BindPinToVariable(URigVMPin* InPin, const FString& InNewB
 		return false;
 	}
 
-	if (GetGraph()->IsA<URigVMFunctionLibrary>())
-	{
-		ReportError(TEXT("Cannot bind pins to variables in function library graphs."));
-		return false;
-	}
-
-	if (InPin->GetBoundVariablePath() == InNewBoundVariablePath)
-	{
-		return false;
-	}
-
-	if (InPin->GetDirection() != ERigVMPinDirection::Input)
-	{
-		return false;
-	}
-
-	if (!InPin->IsRootPin())
-	{
-		return false;
-	}
-
 	FString VariableName = InNewBoundVariablePath, SegmentPath;
 	InNewBoundVariablePath.Split(TEXT("."), &VariableName, &SegmentPath);
 
@@ -9124,6 +9160,11 @@ bool URigVMController::BindPinToVariable(URigVMPin* InPin, const FString& InNewB
 		}
 	}
 
+	if(!GetSchema()->CanBindVariable(this, InPin, &Variable, InNewBoundVariablePath))
+	{
+		return false;
+	}
+
 	if (!Variable.Name.IsValid())
 	{
 		ReportError(TEXT("Cannot find variable in this graph."));
@@ -9243,12 +9284,12 @@ bool URigVMController::UnbindPinFromVariable(const FString& InPinPath, bool bSet
 	const bool bSuccess = UnbindPinFromVariable(Pin, bSetupUndoRedo);
 	if (bSuccess && bPrintPythonCommand)
 	{
-		const FString GraphName = GetSanitizedGraphName(GetGraph()->GetGraphName());
+		const FString GraphName = GetSchema()->GetSanitizedGraphName(GetGraph()->GetGraphName());
 		
-		RigVMPythonUtils::Print(GetGraphOuterName(),
+		RigVMPythonUtils::Print(GetSchema()->GetGraphOuterName(GetGraph()),
 			FString::Printf(TEXT("blueprint.get_controller_by_name('%s').unbind_pin_from_variable('%s')"),
 			*GraphName,
-			*GetSanitizedPinPath(InPinPath)));
+			*GetSchema()->GetSanitizedPinPath(InPinPath)));
 	}
 	
 	return bSuccess;
@@ -9266,16 +9307,8 @@ bool URigVMController::UnbindPinFromVariable(URigVMPin* InPin, bool bSetupUndoRe
 		return false;
 	}
 
-
-	if (GetGraph()->IsA<URigVMFunctionLibrary>())
-	{
-		ReportError(TEXT("Cannot unbind pins from variables in function library graphs."));
-		return false;
-	}
-
-	if (!InPin->IsBoundToVariable())
+	if(!GetSchema()->CanUnbindVariable(this, InPin))
 	{
-		ReportError(TEXT("Pin is not bound to any variable."));
 		return false;
 	}
 
@@ -9437,12 +9470,12 @@ bool URigVMController::PromotePinToVariable(const FString& InPinPath, bool bCrea
 	const bool bSuccess = PromotePinToVariable(Pin, bCreateVariableNode, InNodePosition, bSetupUndoRedo);
 	if (bSuccess && bPrintPythonCommand)
 	{
-		const FString GraphName = GetSanitizedGraphName(GetGraph()->GetGraphName());
+		const FString GraphName = GetSchema()->GetSanitizedGraphName(GetGraph()->GetGraphName());
 		
-		RigVMPythonUtils::Print(GetGraphOuterName(), 
+		RigVMPythonUtils::Print(GetSchema()->GetGraphOuterName(GetGraph()), 
 			FString::Printf(TEXT("blueprint.get_controller_by_name('%s').promote_pin_to_variable('%s', %s, %s)"),
 			*GraphName,
-			*GetSanitizedPinPath(InPinPath),
+			*GetSchema()->GetSanitizedPinPath(InPinPath),
 			(bCreateVariableNode) ? TEXT("True") : TEXT("False"),
 			*RigVMPythonUtils::Vector2DToPythonString(InNodePosition)));
 	}
@@ -9583,12 +9616,12 @@ bool URigVMController::AddLink(const FString& InOutputPinPath, const FString& In
 	const bool bSuccess = AddLink(OutputPin, InputPin, bSetupUndoRedo, InUserDirection, bCreateCastNode);
 	if (bSuccess && bPrintPythonCommand)
 	{
-		const FString GraphName = GetSanitizedGraphName(GetGraph()->GetGraphName());
+		const FString GraphName = GetSchema()->GetSanitizedGraphName(GetGraph()->GetGraphName());
 
-		const FString SanitizedInputPinPath = GetSanitizedPinPath(InputPin->GetPinPath());
-		const FString SanitizedOutputPinPath = GetSanitizedPinPath(OutputPin->GetPinPath());
+		const FString SanitizedInputPinPath = GetSchema()->GetSanitizedPinPath(InputPin->GetPinPath());
+		const FString SanitizedOutputPinPath = GetSchema()->GetSanitizedPinPath(OutputPin->GetPinPath());
 
-		RigVMPythonUtils::Print(GetGraphOuterName(), 
+		RigVMPythonUtils::Print(GetSchema()->GetGraphOuterName(GetGraph()), 
 			FString::Printf(TEXT("blueprint.get_controller_by_name('%s').add_link('%s', '%s')"),
 			*GraphName,
 			*SanitizedOutputPinPath,
@@ -9610,57 +9643,18 @@ bool URigVMController::AddLink(URigVMPin* OutputPin, URigVMPin* InputPin, bool b
 		ReportError(FailureReason);
 		return false;
 	}
-	
-	if(OutputPin == nullptr)
-	{
-		static const FString FailureReason = TEXT("OutputPin is nullptr."); 
-		if(OutFailureReason)
-		{
-			*OutFailureReason = FailureReason;
-		}
-		ReportError(FailureReason);
-		return false;
-	}
-
-	if(InputPin == nullptr)
-	{
-		static const FString FailureReason = TEXT("InputPin is nullptr."); 
-		if(OutFailureReason)
-		{
-			*OutFailureReason = FailureReason;
-		}
-		ReportError(FailureReason);
-		return false;
-	}
-
-	if(!IsValidPinForGraph(OutputPin) || !IsValidPinForGraph(InputPin))
-	{
-		static const FString FailureReason = TEXT("Pin is not valid for graph."); 
-		if(OutFailureReason)
-		{
-			*OutFailureReason = FailureReason;
-		}
-		ReportError(FailureReason);
-		return false;
-	}
 
 	URigVMGraph* Graph = GetGraph();
 	check(Graph);
 
-	if (Graph->IsA<URigVMFunctionLibrary>())
+	TGuardValue<ERigVMPinDirection> UserLinkDirectionGuard(UserLinkDirection,	
+	InUserDirection == ERigVMPinDirection::Invalid ? UserLinkDirection : InUserDirection);
+
+	if(!GetSchema()->CanAddLink(this, OutputPin, InputPin, nullptr, UserLinkDirection, true, true, OutFailureReason))
 	{
-		static const FString FailureReason = TEXT("Cannot add links in function library graphs."); 
-		if(OutFailureReason)
-		{
-			*OutFailureReason = FailureReason;
-		}
-		ReportError(FailureReason);
 		return false;
 	}
-
-	TGuardValue<ERigVMPinDirection> UserLinkDirectionGuard(UserLinkDirection,
-	InUserDirection == ERigVMPinDirection::Invalid ? UserLinkDirection : InUserDirection);
-
+	
 	if (!bIsTransacting)
 	{
 		FString FailureReason;
@@ -10109,13 +10103,13 @@ bool URigVMController::BreakLink(const FString& InOutputPinPath, const FString&
 	const bool bSuccess = BreakLink(OutputPin, InputPin, bSetupUndoRedo);
 	if (bSuccess && bPrintPythonCommand)
 	{
-		const FString GraphName = GetSanitizedGraphName(GetGraph()->GetGraphName());
+		const FString GraphName = GetSchema()->GetSanitizedGraphName(GetGraph()->GetGraphName());
 
-		RigVMPythonUtils::Print(GetGraphOuterName(), 
+		RigVMPythonUtils::Print(GetSchema()->GetGraphOuterName(GetGraph()), 
 			FString::Printf(TEXT("blueprint.get_controller_by_name('%s').break_link('%s', '%s')"),
 			*GraphName,
-			*GetSanitizedPinPath(OutputPin->GetPinPath()),
-			*GetSanitizedPinPath(InputPin->GetPinPath())));
+			*GetSchema()->GetSanitizedPinPath(OutputPin->GetPinPath()),
+			*GetSchema()->GetSanitizedPinPath(InputPin->GetPinPath())));
 	}
 	return bSuccess;
 }
@@ -10126,24 +10120,12 @@ bool URigVMController::BreakLink(URigVMPin* OutputPin, URigVMPin* InputPin, bool
 	{
 		return false;
 	}
-	
-	if(!IsValidPinForGraph(OutputPin) || !IsValidPinForGraph(InputPin))
-	{
-		return false;
-	}
-
-	if (!OutputPin->IsLinkedTo(InputPin))
-	{
-		return false;
-	}
-	ensure(InputPin->IsLinkedTo(OutputPin));
 
 	URigVMGraph* Graph = GetGraph();
 	check(Graph);
 
-	if (Graph->IsA<URigVMFunctionLibrary>())
+	if(!GetSchema()->CanBreakLink(this, OutputPin, InputPin))
 	{
-		ReportError(TEXT("Cannot break links in function library graphs."));
 		return false;
 	}
 
@@ -10220,12 +10202,12 @@ bool URigVMController::BreakAllLinks(const FString& InPinPath, bool bAsInput, bo
 	const bool bSuccess = BreakAllLinks(Pin, bAsInput, bSetupUndoRedo);
 	if (bSuccess && bPrintPythonCommand)
 	{
-		const FString GraphName = GetSanitizedGraphName(GetGraph()->GetGraphName());
+		const FString GraphName = GetSchema()->GetSanitizedGraphName(GetGraph()->GetGraphName());
 
-		RigVMPythonUtils::Print(GetGraphOuterName(), 
+		RigVMPythonUtils::Print(GetSchema()->GetGraphOuterName(GetGraph()), 
 			FString::Printf(TEXT("blueprint.get_controller_by_name('%s').break_all_links('%s', %s)"),
 			*GraphName,
-			*GetSanitizedPinPath(Pin->GetPinPath()),
+			*GetSchema()->GetSanitizedPinPath(Pin->GetPinPath()),
 			bAsInput ? TEXT("True") : TEXT("False")));
 	}
 	return bSuccess;
@@ -10393,7 +10375,7 @@ FName URigVMController::AddExposedPin(const FName& InPinName, ERigVMPinDirection
 		}
 	}
 
-	FName PinName = GetUniqueName(InPinName, [LibraryNode](const FName& InName) {
+	FName PinName = URigVMSchema::GetUniqueName(InPinName, [LibraryNode](const FName& InName) {
 
 		if(LibraryNode->FindPin(InName.ToString()) != nullptr)
 		{
@@ -10412,8 +10394,16 @@ FName URigVMController::AddExposedPin(const FName& InPinName, ERigVMPinDirection
 
 	}, false, true);
 
+	const FString CPPType = RigVMTypeUtils::PostProcessCPPType(InCPPType, CPPTypeObject);
+	const FRigVMTemplateArgumentType Type(*CPPType, CPPTypeObject);
+	const TRigVMTypeIndex TypeIndex = FRigVMRegistry::Get().FindOrAddType(Type);
+	if(!GetSchema()->SupportsType(this, TypeIndex))
+	{
+		return NAME_None;
+	}
+
 	URigVMPin* Pin = NewObject<URigVMPin>(LibraryNode, PinName);
-	Pin->CPPType = RigVMTypeUtils::PostProcessCPPType(InCPPType, CPPTypeObject);
+	Pin->CPPType = CPPType;
 	Pin->CPPTypeObjectPath = InCPPTypeObjectPath;
 	Pin->bIsConstant = false;
 	Pin->Direction = InDirection;
@@ -10456,26 +10446,25 @@ FName URigVMController::AddExposedPin(const FName& InPinName, ERigVMPinDirection
 	if (!EntryNode)
 	{
 		EntryNode = NewObject<URigVMFunctionEntryNode>(Graph, TEXT("Entry"));
-		Graph->Nodes.Add(EntryNode);
+
+		if(AddGraphNode(EntryNode, false))
 		{
 			TGuardValue<bool> SuspendNotifications(bSuspendNotifications, true);
 			RefreshFunctionPins(EntryNode);
+			Notify(ERigVMGraphNotifType::NodeAdded, EntryNode);
 		}
-
-		Notify(ERigVMGraphNotifType::NodeAdded, EntryNode);
 	}
 
 	URigVMFunctionReturnNode* ReturnNode = Graph->GetReturnNode();
 	if (!ReturnNode)
 	{
 		ReturnNode = NewObject<URigVMFunctionReturnNode>(Graph, TEXT("Return"));
-		Graph->Nodes.Add(ReturnNode);
+		if(AddGraphNode(ReturnNode, false))
 		{
 			TGuardValue<bool> SuspendNotifications(bSuspendNotifications, true);
 			RefreshFunctionPins(ReturnNode);
+			Notify(ERigVMGraphNotifType::NodeAdded, ReturnNode);
 		}
-
-		Notify(ERigVMGraphNotifType::NodeAdded, ReturnNode);
 	}
 
 	
@@ -10493,13 +10482,13 @@ FName URigVMController::AddExposedPin(const FName& InPinName, ERigVMPinDirection
 	{
 		//AddExposedPin(const FName& InPinName, ERigVMPinDirection InDirection, const FString& InCPPType, const FName& InCPPTypeObjectPath, const FString& InDefaultValue, bool bSetupUndoRedo, bool bPrintPythonCommand)
 
-		const FString GraphName = GetSanitizedGraphName(GetGraph()->GetGraphName());
+		const FString GraphName = GetSchema()->GetSanitizedGraphName(GetGraph()->GetGraphName());
 		static constexpr TCHAR AddExposedPinFormat[] = TEXT("blueprint.get_controller_by_name('%s').add_exposed_pin('%s', %s, '%s', '%s', '%s')");
 		
-		RigVMPythonUtils::Print(GetGraphOuterName(), 
+		RigVMPythonUtils::Print(GetSchema()->GetGraphOuterName(GetGraph()), 
 			FString::Printf(AddExposedPinFormat,
 				*GraphName,
-				*GetSanitizedPinName(InPinName.ToString()),
+				*GetSchema()->GetSanitizedPinName(InPinName.ToString()),
 				*RigVMPythonUtils::EnumValueToPythonString<ERigVMPinDirection>((int64)InDirection),
 				*InCPPType,
 				*InCPPTypeObjectPath.ToString(),
@@ -10607,12 +10596,12 @@ bool URigVMController::RemoveExposedPin(const FName& InPinName, bool bSetupUndoR
 
 	if (bSuccessfullyRemovedPin && bPrintPythonCommand)
 	{
-		const FString GraphName = GetSanitizedGraphName(GetGraph()->GetGraphName());
+		const FString GraphName = GetSchema()->GetSanitizedGraphName(GetGraph()->GetGraphName());
 
-		RigVMPythonUtils::Print(GetGraphOuterName(), 
+		RigVMPythonUtils::Print(GetSchema()->GetGraphOuterName(GetGraph()), 
 			FString::Printf(TEXT("blueprint.get_controller_by_name('%s').remove_exposed_pin('%s')"),
 				*GraphName,
-				*GetSanitizedPinName(InPinName.ToString())));
+				*GetSchema()->GetSanitizedPinName(InPinName.ToString())));
 	}
 
 	return bSuccessfullyRemovedPin;
@@ -10672,7 +10661,7 @@ bool URigVMController::RenameExposedPin(const FName& InOldPinName, const FName&
 		}
 	}
 
-	FName PinName = GetUniqueName(InNewPinName, [LibraryNode](const FName& InName) {
+	FName PinName = GetSchema()->GetUniqueName(InNewPinName, [LibraryNode](const FName& InName) {
 		const TArray<FRigVMGraphVariableDescription>& LocalVariables = LibraryNode->GetContainedGraph()->GetLocalVariables(true);
 		for(const FRigVMGraphVariableDescription& VariableDescription : LocalVariables)
 		{
@@ -10795,13 +10784,13 @@ bool URigVMController::RenameExposedPin(const FName& InOldPinName, const FName&
 
 	if (bPrintPythonCommand)
 	{
-		const FString GraphName = GetSanitizedGraphName(GetGraph()->GetGraphName());
+		const FString GraphName = GetSchema()->GetSanitizedGraphName(GetGraph()->GetGraphName());
 
-		RigVMPythonUtils::Print(GetGraphOuterName(), 
+		RigVMPythonUtils::Print(GetSchema()->GetGraphOuterName(GetGraph()), 
 			FString::Printf(TEXT("blueprint.get_controller_by_name('%s').rename_exposed_pin('%s', '%s')"),
 				*GraphName,
-				*GetSanitizedPinName(InOldPinName.ToString()),
-				*GetSanitizedPinName(InNewPinName.ToString())));
+				*GetSchema()->GetSanitizedPinName(InOldPinName.ToString()),
+				*GetSchema()->GetSanitizedPinName(InNewPinName.ToString())));
 	}
 
 	return true;
@@ -10876,6 +10865,14 @@ bool URigVMController::ChangeExposedPinType(const FName& InPinName, const FStrin
 		}
 	}
 
+	const FString CPPType = RigVMTypeUtils::PostProcessCPPType(InCPPType, CPPTypeObject);
+	const FRigVMTemplateArgumentType Type(*CPPType, CPPTypeObject);
+	const TRigVMTypeIndex TypeIndex = FRigVMRegistry::Get().FindOrAddType(Type);
+	if(!GetSchema()->SupportsType(this, TypeIndex))
+	{
+		return false;
+	}
+
 	if (Pin->GetDirection() == ERigVMPinDirection::IO)
 	{
 		bool bIsExecute = false;
@@ -11037,12 +11034,12 @@ bool URigVMController::ChangeExposedPinType(const FName& InPinName, const FStrin
 
 	if (bPrintPythonCommand)
 	{
-		const FString GraphName = GetSanitizedGraphName(GetGraph()->GetGraphName());
+		const FString GraphName = GetSchema()->GetSanitizedGraphName(GetGraph()->GetGraphName());
 		
-		RigVMPythonUtils::Print(GetGraphOuterName(), 
+		RigVMPythonUtils::Print(GetSchema()->GetGraphOuterName(GetGraph()), 
 			FString::Printf(TEXT("blueprint.get_controller_by_name('%s').change_exposed_pin_type('%s', '%s', '%s', %s)"),
 				*GraphName,
-				*GetSanitizedPinName(InPinName.ToString()),
+				*GetSchema()->GetSanitizedPinName(InPinName.ToString()),
 				*InCPPType,
 				*InCPPTypeObjectPath.ToString(),
 				(bSetupUndoRedo) ? TEXT("True") : TEXT("False")));
@@ -11122,12 +11119,12 @@ bool URigVMController::SetExposedPinIndex(const FName& InPinName, int32 InNewInd
 
 	if (bPrintPythonCommand)
 	{
-		const FString GraphName = GetSanitizedGraphName(GetGraph()->GetGraphName());
+		const FString GraphName = GetSchema()->GetSanitizedGraphName(GetGraph()->GetGraphName());
 
-		RigVMPythonUtils::Print(GetGraphOuterName(), 
+		RigVMPythonUtils::Print(GetSchema()->GetGraphOuterName(GetGraph()), 
 			FString::Printf(TEXT("blueprint.get_controller_by_name('%s').set_exposed_pin_index('%s', %d)"),
 				*GraphName,
-				*GetSanitizedPinName(InPinName.ToString()),
+				*GetSchema()->GetSanitizedPinName(InPinName.ToString()),
 				InNewIndex));
 	}
 
@@ -11154,7 +11151,7 @@ URigVMFunctionReferenceNode* URigVMController::AddFunctionReferenceNode(URigVMLi
 	return nullptr;
 }
 
-URigVMFunctionReferenceNode* URigVMController::AddFunctionReferenceNodeFromDescription(const FRigVMGraphFunctionHeader& InFunctionDefinition, const FVector2D& InNodePosition, const FString& InNodeName, bool bSetupUndoRedo, bool bPrintPythonCommand, bool bAllowPrivateFunctions)
+URigVMFunctionReferenceNode* URigVMController::AddFunctionReferenceNodeFromDescription(const FRigVMGraphFunctionHeader& InFunctionDefinition, const FVector2D& InNodePosition, const FString& InNodeName, bool bSetupUndoRedo, bool bPrintPythonCommand)
 {
 	if (!IsValidGraph())
 	{
@@ -11175,16 +11172,20 @@ URigVMFunctionReferenceNode* URigVMController::AddFunctionReferenceNodeFromDescr
 		return nullptr;
 	}
 
-	if(!CanAddFunctionRefForDefinition(InFunctionDefinition, true, bAllowPrivateFunctions))
+	if(!GetSchema()->SupportsGraphFunction(this, &InFunctionDefinition))
 	{
 		return nullptr;
 	}
 
-	FString NodeName = GetValidNodeName(InNodeName.IsEmpty() ? InFunctionDefinition.Name.ToString() : InNodeName);
+	FString NodeName = GetSchema()->GetValidNodeName(Graph, InNodeName.IsEmpty() ? InFunctionDefinition.Name.ToString() : InNodeName);
 	URigVMFunctionReferenceNode* FunctionRefNode = NewObject<URigVMFunctionReferenceNode>(Graph, *NodeName);
 	FunctionRefNode->Position = InNodePosition;
 	FunctionRefNode->ReferencedFunctionHeader = InFunctionDefinition;
-	Graph->Nodes.Add(FunctionRefNode);
+
+	if(!AddGraphNode(FunctionRefNode, false))
+	{
+		return nullptr;
+	}
 
 	FRigVMControllerCompileBracketScope CompileScope(this);
 	
@@ -11218,8 +11219,8 @@ URigVMFunctionReferenceNode* URigVMController::AddFunctionReferenceNodeFromDescr
 
 	if (bPrintPythonCommand)
 	{
-		const FString GraphName = GetSanitizedGraphName(GetGraph()->GetGraphName());
-		const FString FunctionDefinitionName = GetSanitizedNodeName(InFunctionDefinition.Name.ToString());
+		const FString GraphName = GetSchema()->GetSanitizedGraphName(GetGraph()->GetGraphName());
+		const FString FunctionDefinitionName = GetSchema()->GetSanitizedNodeName(InFunctionDefinition.Name.ToString());
 
 		bool bLocal = false;
 		if(IRigVMClientHost* ClientHost = GetImplementingOuter<IRigVMClientHost>())
@@ -11227,7 +11228,7 @@ URigVMFunctionReferenceNode* URigVMController::AddFunctionReferenceNodeFromDescr
 			if (InFunctionDefinition.LibraryPointer.HostObject == Cast<UObject>(ClientHost->GetRigVMGraphFunctionHost()))
 			{
 				bLocal = true;
-				RigVMPythonUtils::Print(GetGraphOuterName(), 
+				RigVMPythonUtils::Print(GetSchema()->GetGraphOuterName(GetGraph()), 
 					FString::Printf(TEXT("blueprint.get_controller_by_name('%s').add_function_reference_node(library.find_function('%s'), %s, '%s')"),
 							*GraphName,
 							*FunctionDefinitionName,
@@ -11238,7 +11239,7 @@ URigVMFunctionReferenceNode* URigVMController::AddFunctionReferenceNodeFromDescr
 		
 		if (!bLocal)
 		{
-			RigVMPythonUtils::Print(GetGraphOuterName(), 
+			RigVMPythonUtils::Print(GetSchema()->GetGraphOuterName(GetGraph()), 
 				FString::Printf(TEXT("blueprint.get_controller_by_name('%s').add_external_function_reference_node('%s', '%s', %s, '%s')"),
 						*GraphName,
 						*InFunctionDefinition.LibraryPointer.HostObject.ToString(),
@@ -11410,23 +11411,24 @@ URigVMLibraryNode* URigVMController::AddFunctionToLibrary(const FName& InFunctio
 	URigVMGraph* Graph = GetGraph();
 	check(Graph);
 
-	if (!Graph->IsA<URigVMFunctionLibrary>())
+	if(!GetSchema()->CanAddFunction(this, nullptr))
 	{
-		ReportError(TEXT("Can only add function definitions to function library graphs."));
 		return nullptr;
 	}
 
-	FString FunctionName = GetValidNodeName(InFunctionName.IsNone() ? FString(TEXT("Function")) : InFunctionName.ToString());
+	FString FunctionName = GetSchema()->GetValidNodeName(Graph, InFunctionName.IsNone() ? FString(TEXT("Function")) : InFunctionName.ToString());
 	URigVMCollapseNode* CollapseNode = NewObject<URigVMCollapseNode>(Graph, *FunctionName);
 	FString ContainedGraphName = FunctionName + TEXT("_ContainedGraph");
 	CollapseNode->ContainedGraph = NewObject<URigVMGraph>(CollapseNode, *ContainedGraphName);
 	CollapseNode->Position = InNodePosition;
-	Graph->Nodes.Add(CollapseNode);
+
+	if(!AddGraphNode(CollapseNode, true))
+	{
+		return nullptr;
+	}
 
 	FRigVMControllerCompileBracketScope CompileScope(this);
 	
-	Notify(ERigVMGraphNotifType::NodeAdded, CollapseNode);
-	
 	if (bMutable)
 	{
 		const UScriptStruct* ExecuteContextStruct = FRigVMExecuteContext::StaticStruct();
@@ -11452,12 +11454,12 @@ URigVMLibraryNode* URigVMController::AddFunctionToLibrary(const FName& InFunctio
 		if (EntryNode == nullptr)
 		{
 			EntryNode = NewObject<URigVMFunctionEntryNode>(CollapseNode->ContainedGraph, TEXT("Entry"));
-			CollapseNode->ContainedGraph->Nodes.Add(EntryNode);
+			if(CollapseController->AddGraphNode(EntryNode, false))
 			{
 				TGuardValue<bool> SuspendNotifications(CollapseController->bSuspendNotifications, true);
 				CollapseController->RefreshFunctionPins(EntryNode);
+				CollapseController->Notify(ERigVMGraphNotifType::NodeAdded, EntryNode);
 			}
-			CollapseController->Notify(ERigVMGraphNotifType::NodeAdded, EntryNode);
 		}
 
 		EntryNode->Position = FVector2D(-250.f, 0.f);
@@ -11466,12 +11468,12 @@ URigVMLibraryNode* URigVMController::AddFunctionToLibrary(const FName& InFunctio
 		if (ReturnNode == nullptr)
 		{
 			ReturnNode = NewObject<URigVMFunctionReturnNode>(CollapseNode->ContainedGraph, TEXT("Return"));
-			CollapseNode->ContainedGraph->Nodes.Add(ReturnNode);
+			if(CollapseController->AddGraphNode(ReturnNode, false))
 			{
 				TGuardValue<bool> SuspendNotifications(CollapseController->bSuspendNotifications, true);
 				CollapseController->RefreshFunctionPins(ReturnNode);
+				CollapseController->Notify(ERigVMGraphNotifType::NodeAdded, ReturnNode);
 			}
-			CollapseController->Notify(ERigVMGraphNotifType::NodeAdded, ReturnNode);
 		}
 
 		ReturnNode->Position = FVector2D(250.f, 0.f);
@@ -11492,12 +11494,12 @@ URigVMLibraryNode* URigVMController::AddFunctionToLibrary(const FName& InFunctio
 
 	if (bPrintPythonCommand)
 	{
-		const FString GraphName = GetSanitizedGraphName(GetGraph()->GetGraphName());
+		const FString GraphName = GetSchema()->GetSanitizedGraphName(GetGraph()->GetGraphName());
 
 		//AddFunctionToLibrary(const FName& InFunctionName, bool bMutable, const FVector2D& InNodePosition, bool bSetupUndoRedo, bool bPrintPythonCommand)
-		RigVMPythonUtils::Print(GetGraphOuterName(), 
+		RigVMPythonUtils::Print(GetSchema()->GetGraphOuterName(GetGraph()), 
 			FString::Printf(TEXT("library_controller.add_function_to_library('%s', %s, %s)"),
-				*GetSanitizedNodeName(InFunctionName.ToString()),
+				*GetSchema()->GetSanitizedNodeName(InFunctionName.ToString()),
 				(bMutable) ? TEXT("True") : TEXT("False"),
 				*RigVMPythonUtils::Vector2DToPythonString(InNodePosition)));
 	}
@@ -11520,14 +11522,17 @@ bool URigVMController::RemoveFunctionFromLibrary(const FName& InFunctionName, bo
 	URigVMGraph* Graph = GetGraph();
 	check(Graph);
 
-
-	if (!Graph->IsA<URigVMFunctionLibrary>())
+	if(URigVMNode* Node = Graph->FindNodeByName(InFunctionName))
 	{
-		ReportError(TEXT("Can only remove function definitions from function library graphs."));
-		return false;
+		if(!GetSchema()->CanRemoveFunction(this, Node))
+		{
+			return false;
+		}
+		return RemoveNodeByName(InFunctionName, bSetupUndoRedo);
 	}
 
-	return RemoveNodeByName(InFunctionName, bSetupUndoRedo);
+	ReportErrorf(TEXT("Cannot find node '%s'."), *InFunctionName.ToString());
+	return false;
 }
 
 bool URigVMController::RenameFunction(const FName& InOldFunctionName, const FName& InNewFunctionName, bool bSetupUndoRedo)
@@ -11627,9 +11632,9 @@ bool URigVMController::MarkFunctionAsPublic(const FName& InFunctionName, bool bI
 
 	if (bPrintPythonCommand)
 	{
-		RigVMPythonUtils::Print(GetGraphOuterName(), 
+		RigVMPythonUtils::Print(GetSchema()->GetGraphOuterName(GetGraph()), 
 			FString::Printf(TEXT("library_controller.mark_function_as_public('%s', %s)"),
-				*GetSanitizedNodeName(InFunctionName.ToString()),
+				*GetSchema()->GetSanitizedNodeName(InFunctionName.ToString()),
 				(bInIsPublic) ? TEXT("True") : TEXT("False")));
 	}
 
@@ -11704,7 +11709,7 @@ FRigVMGraphVariableDescription URigVMController::AddLocalVariable(const FName& I
 		}
 	}
 
-	FName VariableName = GetUniqueName(InVariableName, [Graph](const FName& InName) {
+	FName VariableName = URigVMSchema::GetUniqueName(InVariableName, [Graph](const FName& InName) {
 		for (FRigVMGraphVariableDescription LocalVariable : Graph->GetLocalVariables(true))
 		{
 			if (LocalVariable.Name == InName)
@@ -11749,9 +11754,9 @@ FRigVMGraphVariableDescription URigVMController::AddLocalVariable(const FName& I
 
 	if (bPrintPythonCommand)
 	{
-		const FString GraphName = GetSanitizedGraphName(GetGraph()->GetGraphName());
+		const FString GraphName = GetSchema()->GetSanitizedGraphName(GetGraph()->GetGraphName());
 
-		RigVMPythonUtils::Print(GetGraphOuterName(), 
+		RigVMPythonUtils::Print(GetSchema()->GetGraphOuterName(GetGraph()), 
 			FString::Printf(TEXT("blueprint.get_controller_by_name('%s').add_local_variable_from_object_path('%s', '%s', '%s', '%s')"),
 				*GraphName,
 				*NewVariable.Name.ToString(),
@@ -11915,12 +11920,12 @@ bool URigVMController::RemoveLocalVariable(const FName& InVariableName, bool bSe
 
 		if (bPrintPythonCommand)
 		{
-			const FString GraphName = GetSanitizedGraphName(GetGraph()->GetGraphName());
+			const FString GraphName = GetSchema()->GetSanitizedGraphName(GetGraph()->GetGraphName());
 
-			RigVMPythonUtils::Print(GetGraphOuterName(), 
+			RigVMPythonUtils::Print(GetSchema()->GetGraphOuterName(GetGraph()), 
 				FString::Printf(TEXT("blueprint.get_controller_by_name('%s').remove_local_variable('%s')"),
 					*GraphName,
-					*GetSanitizedVariableName(InVariableName.ToString())));
+					*GetSchema()->GetSanitizedVariableName(InVariableName.ToString())));
 		}
 		return true;
 	}
@@ -12009,13 +12014,13 @@ bool URigVMController::RenameLocalVariable(const FName& InVariableName, const FN
 
 	if (bPrintPythonCommand)
 	{
-		const FString GraphName = GetSanitizedGraphName(GetGraph()->GetGraphName());
+		const FString GraphName = GetSchema()->GetSanitizedGraphName(GetGraph()->GetGraphName());
 
-		RigVMPythonUtils::Print(GetGraphOuterName(), 
+		RigVMPythonUtils::Print(GetSchema()->GetGraphOuterName(GetGraph()), 
 			FString::Printf(TEXT("blueprint.get_controller_by_name('%s').rename_local_variable('%s', '%s')"),
 				*GraphName,
-				*GetSanitizedVariableName(InVariableName.ToString()),
-				*GetSanitizedVariableName(InNewVariableName.ToString())));
+				*GetSchema()->GetSanitizedVariableName(InVariableName.ToString()),
+				*GetSchema()->GetSanitizedVariableName(InNewVariableName.ToString())));
 	}
 
 	return true;
@@ -12116,13 +12121,13 @@ bool URigVMController::SetLocalVariableType(const FName& InVariableName, const F
 
 	if (bPrintPythonCommand)
 	{
-		const FString GraphName = GetSanitizedGraphName(GetGraph()->GetGraphName());
+		const FString GraphName = GetSchema()->GetSanitizedGraphName(GetGraph()->GetGraphName());
 
 		//bool URigVMController::SetLocalVariableType(const FName& InVariableName, const FString& InCPPType, UObject* InCPPTypeObject, bool bSetupUndoRedo, bool bPrintPythonCommand)
-		RigVMPythonUtils::Print(GetGraphOuterName(), 
+		RigVMPythonUtils::Print(GetSchema()->GetGraphOuterName(GetGraph()), 
 			FString::Printf(TEXT("blueprint.get_controller_by_name('%s').set_local_variable_type_from_object_path('%s', '%s', '%s')"),
 				*GraphName,
-				*GetSanitizedVariableName(InVariableName.ToString()),
+				*GetSchema()->GetSanitizedVariableName(InVariableName.ToString()),
 				*InCPPType,
 				(InCPPTypeObject) ? *InCPPTypeObject->GetPathName() : *FString()));
 	}
@@ -12195,12 +12200,12 @@ bool URigVMController::SetLocalVariableDefaultValue(const FName& InVariableName,
 
 	if (bPrintPythonCommand)
 	{
-		const FString GraphName = GetSanitizedGraphName(GetGraph()->GetGraphName());
+		const FString GraphName = GetSchema()->GetSanitizedGraphName(GetGraph()->GetGraphName());
 
-		RigVMPythonUtils::Print(GetGraphOuterName(), 
+		RigVMPythonUtils::Print(GetSchema()->GetGraphOuterName(GetGraph()), 
 			FString::Printf(TEXT("blueprint.get_controller_by_name('%s').set_local_variable_default_value('%s', '%s')"),
 				*GraphName,
-				*GetSanitizedVariableName(InVariableName.ToString()),
+				*GetSchema()->GetSanitizedVariableName(InVariableName.ToString()),
 				*InDefaultValue));
 	}
 	
@@ -12496,7 +12501,7 @@ bool URigVMController::SetVariableName(URigVMVariableNode* InVariableNode, const
 	}
 
 	const FRigVMExternalVariable VariableType = RigVMTypeUtils::ExternalVariableFromCPPType(InVariableName, InVariableNode->GetCPPType(), InVariableNode->GetCPPTypeObject());
-	FName VariableName = GetUniqueName(InVariableName, [Descriptions, NameToIndex, VariableType](const FName& InName) {
+	FName VariableName = URigVMSchema::GetUniqueName(InVariableName, [Descriptions, NameToIndex, VariableType](const FName& InName) {
 		const int32* FoundIndex = NameToIndex.Find(InName);
 		if (FoundIndex == nullptr)
 		{
@@ -12561,7 +12566,7 @@ URigVMRerouteNode* URigVMController::AddFreeRerouteNode(const FString& InCPPType
 		GetActionStack()->BeginAction(Action);
 	}
 
-	FString Name = GetValidNodeName(InNodeName.IsEmpty() ? FString(TEXT("RerouteNode")) : InNodeName);
+	FString Name = GetSchema()->GetValidNodeName(Graph, InNodeName.IsEmpty() ? FString(TEXT("RerouteNode")) : InNodeName);
 	URigVMRerouteNode* Node = NewObject<URigVMRerouteNode>(Graph, *Name);
 	Node->Position = InPosition;
 
@@ -12572,7 +12577,10 @@ URigVMRerouteNode* URigVMController::AddFreeRerouteNode(const FString& InCPPType
 	ValuePin->CustomWidgetName = InCustomWidgetName;
 	ValuePin->Direction = ERigVMPinDirection::IO;
 	AddNodePin(Node, ValuePin);
-	Graph->Nodes.Add(Node);
+	if(!AddGraphNode(Node, false))
+	{
+		return nullptr;
+	}
 
 	if (!InDefaultValue.IsEmpty() && InDefaultValue != TEXT("()"))
 	{
@@ -12965,7 +12973,7 @@ URigVMNode* URigVMController::AddIfNode(const FString& InCPPType, const FName& I
 	}
 
 	const FString CPPType = RigVMTypeUtils::PostProcessCPPType(InCPPType, CPPTypeObject);
-	const FString Name = GetValidNodeName(InNodeName.IsEmpty() ? FString(TEXT("IfNode")) : InNodeName);
+	const FString Name = GetSchema()->GetValidNodeName(Graph, InNodeName.IsEmpty() ? FString(TEXT("IfNode")) : InNodeName);
 	const TRigVMTypeIndex& TypeIndex = FRigVMRegistry::Get().FindOrAddType({*CPPType, CPPTypeObject});
 	
 	const FRigVMDispatchFactory* Factory = FRigVMRegistry::Get().FindOrAddDispatchFactory(FRigVMDispatch_If::StaticStruct());
@@ -13027,7 +13035,7 @@ URigVMNode* URigVMController::AddSelectNode(const FString& InCPPType, const FNam
 	}
 
 	const FString CPPType = RigVMTypeUtils::PostProcessCPPType(InCPPType, CPPTypeObject);
-	const FString Name = GetValidNodeName(InNodeName.IsEmpty() ? FString(TEXT("SelectNode")) : InNodeName);
+	const FString Name = GetSchema()->GetValidNodeName(Graph, InNodeName.IsEmpty() ? FString(TEXT("SelectNode")) : InNodeName);
 	const TRigVMTypeIndex& TypeIndex = FRigVMRegistry::Get().FindOrAddType({*CPPType, CPPTypeObject});
 	
 	const FRigVMDispatchFactory* Factory = FRigVMRegistry::Get().FindOrAddDispatchFactory(FRigVMDispatch_SelectInt32::StaticStruct());
@@ -13083,18 +13091,18 @@ URigVMTemplateNode* URigVMController::AddTemplateNode(const FName& InNotation, c
 	{
 		if(const FRigVMClient* Client = ClientHost->GetRigVMClient())
 		{
-			if(!Template->SupportsExecuteContextStruct(Client->GetExecuteContextStruct()))
+			if(!Template->SupportsExecuteContextStruct(GetSchema()->GetExecuteContextStruct()))
 			{
 				ReportErrorf(TEXT("Cannot add node for template '%s' - incompatible execute context: '%s' vs '%s'."),
 						*Template->GetNotation().ToString(),
 						*Template->GetExecuteContextStruct()->GetStructCPPName(),
-						*Client->GetExecuteContextStruct()->GetStructCPPName());
+						*GetSchema()->GetExecuteContextStruct()->GetStructCPPName());
 				return nullptr;
 			}
 		}
 	}
 
-	FString Name = GetValidNodeName(InNodeName.IsEmpty() ? Template->GetName().ToString() : InNodeName);
+	FString Name = GetSchema()->GetValidNodeName(Graph, InNodeName.IsEmpty() ? Template->GetName().ToString() : InNodeName);
 	URigVMTemplateNode* Node = nullptr;
 
 	// determine what kind of node we need to create
@@ -13139,9 +13147,10 @@ URigVMTemplateNode* URigVMController::AddTemplateNode(const FName& InNotation, c
 		}
 	}
 
-	Graph->Nodes.Add(Node);
-
-	Notify(ERigVMGraphNotifType::NodeAdded, Node);
+	if(!AddGraphNode(Node, true))
+	{
+		return nullptr;
+	}
 
 	FRigVMBaseAction Action(this);
 	FRigVMControllerCompileBracketScope CompileScope(this);
@@ -13163,7 +13172,7 @@ URigVMTemplateNode* URigVMController::AddTemplateNode(const FName& InNotation, c
 		TArray<FString> Commands = GetAddNodePythonCommands(Node);
 		for (const FString& Command : Commands)
 		{
-			RigVMPythonUtils::Print(GetGraphOuterName(), 
+			RigVMPythonUtils::Print(GetSchema()->GetGraphOuterName(GetGraph()), 
 				FString::Printf(TEXT("%s"), *Command));
 		}
 	}
@@ -13270,7 +13279,7 @@ URigVMEnumNode* URigVMController::AddEnumNode(const FName& InCPPTypeObjectPath,
 		return nullptr;
 	}
 
-	FString Name = GetValidNodeName(InNodeName.IsEmpty() ? FString(TEXT("IfNode")) : InNodeName);
+	FString Name = GetSchema()->GetValidNodeName(Graph, InNodeName.IsEmpty() ? FString(TEXT("IfNode")) : InNodeName);
 	URigVMEnumNode* Node = NewObject<URigVMEnumNode>(Graph, *Name);
 	Node->Position = InPosition;
 
@@ -13288,9 +13297,10 @@ URigVMEnumNode* URigVMController::AddEnumNode(const FName& InCPPTypeObjectPath,
 	EnumIndexPin->DisplayName = TEXT("Result");
 	AddNodePin(Node, EnumIndexPin);
 
-	Graph->Nodes.Add(Node);
-
-	Notify(ERigVMGraphNotifType::NodeAdded, Node);
+	if(!AddGraphNode(Node, true))
+	{
+		return nullptr;
+	}
 
 	FRigVMControllerCompileBracketScope CompileScope(this);
 	if (bSetupUndoRedo)
@@ -13303,7 +13313,7 @@ URigVMEnumNode* URigVMController::AddEnumNode(const FName& InCPPTypeObjectPath,
 		TArray<FString> Commands = GetAddNodePythonCommands(Node);
 		for (const FString& Command : Commands)
 		{
-			RigVMPythonUtils::Print(GetGraphOuterName(), 
+			RigVMPythonUtils::Print(GetSchema()->GetGraphOuterName(GetGraph()), 
 				FString::Printf(TEXT("%s"), *Command));
 		}
 	}
@@ -13474,7 +13484,7 @@ URigVMInvokeEntryNode* URigVMController::AddInvokeEntryNode(const FName& InEntry
 		return nullptr;
 	}
 
-	FString Name = GetValidNodeName(InNodeName.IsEmpty() ? FString(TEXT("InvokeEntryNode")) : InNodeName);
+	FString Name = GetSchema()->GetValidNodeName(Graph, InNodeName.IsEmpty() ? FString(TEXT("InvokeEntryNode")) : InNodeName);
 	URigVMInvokeEntryNode* Node = NewObject<URigVMInvokeEntryNode>(Graph, *Name);
 	Node->Position = InPosition;
 
@@ -13490,7 +13500,10 @@ URigVMInvokeEntryNode* URigVMController::AddInvokeEntryNode(const FName& InEntry
 	EntryNamePin->CustomWidgetName = TEXT("EntryName");
 	AddNodePin(Node, EntryNamePin);
 
-	Graph->Nodes.Add(Node);
+	if(!AddGraphNode(Node, true))
+	{
+		return nullptr;
+	}
 
 	if (!bSuspendNotifications)
 	{
@@ -13499,7 +13512,6 @@ URigVMInvokeEntryNode* URigVMController::AddInvokeEntryNode(const FName& InEntry
 
 	FRigVMControllerCompileBracketScope CompileScope(this);
 
-	Notify(ERigVMGraphNotifType::NodeAdded, Node);
 	Notify(ERigVMGraphNotifType::VariableAdded, Node);
 
 	if (bSetupUndoRedo)
@@ -13514,7 +13526,7 @@ URigVMInvokeEntryNode* URigVMController::AddInvokeEntryNode(const FName& InEntry
 		TArray<FString> Commands = GetAddNodePythonCommands(Node);
 		for (const FString& Command : Commands)
 		{
-			RigVMPythonUtils::Print(GetGraphOuterName(), 
+			RigVMPythonUtils::Print(GetSchema()->GetGraphOuterName(GetGraph()), 
 								FString::Printf(TEXT("%s"), *Command));
 		}
 	}
@@ -13539,18 +13551,13 @@ void URigVMController::ForEveryPinRecursively(URigVMNode* InNode, TFunction<void
 	}
 }
 
-FString URigVMController::GetValidNodeName(const FString& InPrefix)
-{
-	URigVMGraph* Graph = GetGraph();
-	check(Graph);
-
-	return GetUniqueName(*InPrefix, [&](const FName& InName) {
-		return Graph->IsNameAvailable(InName.ToString());
-	}, false, true).ToString();
-}
-
 bool URigVMController::IsValidGraph() const
 {
+	if(!IsValidSchema())
+	{
+		return false;
+	}
+	
 	URigVMGraph* Graph = GetGraph();
 	if (Graph == nullptr)
 	{
@@ -13566,17 +13573,32 @@ bool URigVMController::IsValidGraph() const
 	return true;
 }
 
-bool URigVMController::IsGraphEditable() const
+bool URigVMController::IsValidSchema() const
 {
-	if(!IsValidGraph())
+	if (SchemaPtr == nullptr)
 	{
+		ReportError(TEXT("Controller does not have a schema associated."));
 		return false;
 	}
 
-	URigVMGraph* Graph = GetGraph();
-	check(Graph);
+	if (!IsValid(SchemaPtr))
+	{
+		return false;
+	}
 
-	return Graph->bEditable;
+	return true;
+}
+
+bool URigVMController::IsGraphEditable() const
+{
+	if(const URigVMSchema* Schema = GetSchema())
+	{
+		if(const URigVMGraph* Graph = GetGraph())
+		{
+			return Schema->IsGraphEditable(Graph);
+		}
+	}
+	return false;
 }
 
 bool URigVMController::IsValidNodeForGraph(const URigVMNode* InNode)
@@ -13697,300 +13719,19 @@ bool URigVMController::IsValidLinkForGraph(const URigVMLink* InLink)
 	return true;
 }
 
-bool URigVMController::CanAddNode(URigVMNode* InNode, bool bReportErrors, bool bIgnoreFunctionEntryReturnNodes)
+void URigVMController::AddPinsForStruct(UStruct* InStruct, URigVMNode* InNode, URigVMPin* InParentPin, ERigVMPinDirection InPinDirection, const FString& InDefaultValue, bool bAutoExpandArrays, const FPinInfoArray* PreviousPins)
 {
-	if(!IsValidGraph())
+	if(!InStruct->IsChildOf(FRigVMStruct::StaticStruct()))
 	{
-		return false;
+		if(!GetSchema()->ShouldUnfoldStruct(this, InStruct))
+		{
+			return;
+		}
 	}
 
-	if (!bIsTransacting && !IsGraphEditable())
+	if(InParentPin && InParentPin->ShouldHideSubPins())
 	{
-		return false;
-	}
-
-	check(InNode);
-
-	URigVMGraph* Graph = GetGraph();
-	check(Graph);
-
-	if (Graph->IsA<URigVMFunctionLibrary>())
-	{
-		if (!InNode->IsA<URigVMCollapseNode>())
-		{
-			return false;
-		}
-	}
-
-	if (URigVMFunctionReferenceNode* FunctionRefNode = Cast<URigVMFunctionReferenceNode>(InNode))
-	{
-		if (URigVMFunctionLibrary* FunctionLibrary = FunctionRefNode->GetLibrary())
-		{			
-			FRigVMGraphFunctionHeader FunctionDefinition = FunctionRefNode->GetReferencedFunctionHeader();
-			if(!CanAddFunctionRefForDefinition(FunctionDefinition, false))
-			{
-				URigVMFunctionLibrary* TargetLibrary = Graph->GetDefaultFunctionLibrary();
-				URigVMLibraryNode* NewFunctionDefinition = TargetLibrary->FindPreviouslyLocalizedFunction(FunctionDefinition.LibraryPointer);
-			
-				if((NewFunctionDefinition == nullptr) && RequestLocalizeFunctionDelegate.IsBound())
-				{
-					if(RequestLocalizeFunctionDelegate.Execute(FunctionDefinition.LibraryPointer))
-					{
-						NewFunctionDefinition = TargetLibrary->FindPreviouslyLocalizedFunction(FunctionDefinition.LibraryPointer);
-					}
-				}
-
-				if(NewFunctionDefinition == nullptr)
-				{
-					return false;
-				}
-			
-				SetReferencedFunction(FunctionRefNode, NewFunctionDefinition, false);
-				FunctionDefinition = FunctionRefNode->GetReferencedFunctionHeader();
-			}
-			
-			if(!CanAddFunctionRefForDefinition(FunctionDefinition, bReportErrors))
-			{
-				DestroyObject(InNode);
-				return false;
-			}
-		}			
-	}
-	else if(!bIgnoreFunctionEntryReturnNodes &&
-		(InNode->IsA<URigVMFunctionEntryNode>() ||
-		InNode->IsA<URigVMFunctionReturnNode>()))
-	{
-		// only allow entry / return nodes on sub graphs
-		if(Graph->IsRootGraph())
-		{
-			return false;
-		}
-
-		// only allow one function entry node
-		if(InNode->IsA<URigVMFunctionEntryNode>())
-		{
-			if(const URigVMFunctionEntryNode* ExistingEntryNode = Graph->GetEntryNode())
-			{
-				if(ExistingEntryNode != InNode)
-				{
-					return false;
-				}
-			}
-		}
-		// only allow one function return node
-		else if(InNode->IsA<URigVMFunctionReturnNode>())
-		{
-			if(const URigVMFunctionReturnNode* ExistingReturnNode = Graph->GetReturnNode())
-			{
-				if(ExistingReturnNode != InNode)
-				{
-					return false;
-				}
-			}
-		}
-	}
-	else if(const URigVMCollapseNode* CollapseNode = Cast<URigVMCollapseNode>(InNode))
-	{
-		if(URigVMController* CollapseController = GetControllerForGraph(CollapseNode->GetContainedGraph()))
-		{
-			TArray<URigVMNode*> ContainedNodes = CollapseNode->GetContainedNodes();
-			for(URigVMNode* ContainedNode : ContainedNodes)
-			{
-				if(!CollapseController->CanAddNode(ContainedNode, bReportErrors, true))
-				{
-					return false;
-				}
-			}
-		}
-	}
-	else if(URigVMVariableNode* VariableNode = Cast<URigVMVariableNode>(InNode))
-	{
-		if (URigVMPin* NamePin = VariableNode->FindPin(URigVMVariableNode::VariableName))
-		{
-			FString VarName = NamePin->GetDefaultValue();
-			if (!VarName.IsEmpty())
-			{
-				TArray<FRigVMExternalVariable> AllVariables = URigVMController::GetAllVariables(true);
-				for(const FRigVMExternalVariable& Variable : AllVariables)
-				{
-					if(Variable.Name.ToString() == VarName)
-					{
-						return true;
-					}
-				}
-				return false;
-			}
-		}
-	}
-	else if (InNode->IsEvent())
-	{
-		if (URigVMUnitNode* InUnitNode = Cast<URigVMUnitNode>(InNode))
-		{
-			if (!CanAddEventNode(InUnitNode->GetScriptStruct(), bReportErrors))
-			{
-				return false;
-			}
-		}
-	}
-	
-	return true;
-}
-
-TObjectPtr<URigVMNode> URigVMController::FindEventNode(const UScriptStruct* InScriptStruct) const
-{
-	check(InScriptStruct);
-
-	// construct equivalent default struct
-	FStructOnScope InDefaultStructScope(InScriptStruct);
-	
-	if (URigVMGraph* Graph = GetGraph())
-	{
-		TObjectPtr<URigVMNode>* FoundNode = 
-		Graph->Nodes.FindByPredicate( [&InDefaultStructScope](const TObjectPtr<URigVMNode>& Node) {
-			if (Node->IsEvent())
-			{
-				if (URigVMUnitNode* UnitNode = Cast<URigVMUnitNode>(Node))
-				{
-					// compare default structures
-					TSharedPtr<FStructOnScope> DefaultStructScope = UnitNode->ConstructStructInstance(true);
-					if (DefaultStructScope.IsValid() && InDefaultStructScope.GetStruct() == DefaultStructScope->GetStruct())
-					{
-						return true;
-					}
-				}
-			}
-			return false;
-		});
-
-		if (FoundNode)
-		{
-			return *FoundNode;
-		}
-	}
-	
-	return TObjectPtr<URigVMNode>();
-}
-
-bool URigVMController::CanAddEventNode(UScriptStruct* InScriptStruct, const bool bReportErrors) const
-{
-	if(!IsValidGraph())
-	{
-		return false;
-	}
-
-	if (!bIsTransacting && !IsGraphEditable())
-	{
-		return false;
-	}
-	
-	check(InScriptStruct);
-	
-	URigVMGraph* Graph = GetGraph();
-	check(Graph);
-
-	// check if we're trying to add a node within a graph which is not the top level one
-	if (!Graph->IsTopLevelGraph())
-	{
-		if (bReportErrors)
-		{
-			ReportAndNotifyError(TEXT("Event nodes can only be added to top level graphs."));
-		}
-		return false;
-	}
-
-	TObjectPtr<URigVMNode> EventNode = FindEventNode(InScriptStruct);
-	const bool bHasEventNode = (EventNode != nullptr) && EventNode->CanOnlyExistOnce();
-	if (bHasEventNode && bReportErrors)
-	{
-		const FString ErrorMessage = FString::Printf(TEXT("Rig Graph can only contain one single %s node."),
-													 *InScriptStruct->GetDisplayNameText().ToString());
-		ReportAndNotifyError(ErrorMessage);
-	}
-		
-	return !bHasEventNode;
-}
-
-bool URigVMController::CanAddFunctionRefForDefinition(const FRigVMGraphFunctionHeader& InFunctionDefinition, bool bReportErrors, bool bAllowPrivateFunctions)
-{
-	if(!IsValidGraph())
-	{
-		return false;
-	}
-
-	if (!bIsTransacting && !IsGraphEditable())
-	{
-		return false;
-	}
-
-	URigVMGraph* Graph = GetGraph();
-	check(Graph);
-
-	if (!bAllowPrivateFunctions)
-	{
-		if(IRigVMClientHost* ClientHost = GetImplementingOuter<IRigVMClientHost>())
-		{
-			bool bIsAvailable = ClientHost->GetRigVMClient()->GetFunctionLibrary()->GetFunctionHostObjectPath() ==
-				InFunctionDefinition.LibraryPointer.HostObject;
-			if (!bIsAvailable)
-			{
-				if (IRigVMGraphFunctionHost* Host = Cast<IRigVMGraphFunctionHost>(InFunctionDefinition.LibraryPointer.HostObject.TryLoad()))
-				{
-					bIsAvailable = Host->GetRigVMGraphFunctionStore()->IsFunctionPublic(InFunctionDefinition.LibraryPointer);		
-				}
-			}
-			if (!bIsAvailable)
-			{
-				if(bReportErrors)
-				{
-					ReportAndNotifyError(TEXT("Function is not available for placement in another graph host."));
-				}
-				return false;
-			}
-		}
-	}
-
-	if (URigVMNode* Node = Cast<URigVMNode>(Graph->GetOuter()))
-	{
-		if (URigVMLibraryNode* LibraryNode = Node->FindFunctionForNode())
-		{
-			if (InFunctionDefinition.Dependencies.Contains(LibraryNode->GetFunctionIdentifier()))
-			{
-				if(bReportErrors)
-				{
-					ReportAndNotifyError(TEXT("Function is not available for placement in this graph host due to dependency cycles."));
-				}
-				return false;
-			}
-		}
-	}
-	
-	URigVMLibraryNode* ParentLibraryNode = Cast<URigVMLibraryNode>(Graph->GetOuter());
-	while (ParentLibraryNode)
-	{
-		if (TSoftObjectPtr<UObject>(ParentLibraryNode).ToSoftObjectPath() == InFunctionDefinition.LibraryPointer.LibraryNode)
-		{
-			if(bReportErrors)
-			{
-				ReportAndNotifyError(TEXT("You cannot place functions inside of itself or an indirect recursion."));
-			}
-			return false;
-		}
-		ParentLibraryNode = Cast<URigVMLibraryNode>(ParentLibraryNode->GetGraph()->GetOuter());
-	}
-
-	return true;
-}
-
-void URigVMController::AddPinsForStruct(UStruct* InStruct, URigVMNode* InNode, URigVMPin* InParentPin, ERigVMPinDirection InPinDirection, const FString& InDefaultValue, bool bAutoExpandArrays, const FPinInfoArray* PreviousPins)
-{
-	if(!ShouldStructBeUnfolded(InStruct))
-	{
-		return;
-	}
-
-	if(InParentPin && InParentPin->ShouldHideSubPins())
-	{
-		return;
+		return;
 	}
 
 	// todo: reuse the default values when creating the pins
@@ -14032,7 +13773,7 @@ void URigVMController::AddPinsForStruct(UStruct* InStruct, URigVMNode* InNode, U
 			FStructProperty* StructProperty = CastField<FStructProperty>(*It);
 			if (StructProperty)
 			{
-				if (ShouldStructBeUnfolded(StructProperty->Struct))
+				if (GetSchema()->ShouldUnfoldStruct(this, StructProperty->Struct))
 				{
 					FString DefaultValue;
 					if (DefaultValuePtr != nullptr)
@@ -14058,7 +13799,7 @@ void URigVMController::AddPinsForStruct(UStruct* InStruct, URigVMNode* InNode, U
 
 				if (DefaultValuePtr)
 				{
-					if (ShouldPinBeUnfolded(Pin))
+					if (GetSchema()->CanUnfoldPin(this, Pin))
 					{
 						TArray<FString> ElementDefaultValues = URigVMPin::SplitDefaultValue(*DefaultValuePtr);
 						AddPinsForArray(ArrayProperty, InNode, Pin, Pin->Direction, ElementDefaultValues, bAutoExpandArrays);
@@ -14090,7 +13831,7 @@ void URigVMController::AddPinsForStruct(UStruct* InStruct, URigVMNode* InNode, U
 void URigVMController::AddPinsForArray(FArrayProperty* InArrayProperty, URigVMNode* InNode, URigVMPin* InParentPin, ERigVMPinDirection InPinDirection, const TArray<FString>& InDefaultValues, bool bAutoExpandArrays)
 {
 	check(InParentPin);
-	if (!ShouldPinBeUnfolded(InParentPin))
+	if (!GetSchema()->CanUnfoldPin(this, InParentPin))
 	{
 		return;
 	}
@@ -14119,7 +13860,7 @@ void URigVMController::AddPinsForArray(FArrayProperty* InArrayProperty, URigVMNo
 		FStructProperty* StructProperty = CastField<FStructProperty>(InArrayProperty->Inner);
 		if (StructProperty)
 		{
-			if (ShouldPinBeUnfolded(Pin))
+			if (GetSchema()->CanUnfoldPin(this, Pin))
 			{
 				// DefaultValue before this point only contains parent struct overrides,
 				// see comments in CreateDefaultValueForStructIfRequired
@@ -14143,7 +13884,7 @@ void URigVMController::AddPinsForArray(FArrayProperty* InArrayProperty, URigVMNo
 		FArrayProperty* ArrayProperty = CastField<FArrayProperty>(InArrayProperty->Inner);
 		if (ArrayProperty)
 		{
-			if (ShouldPinBeUnfolded(Pin))
+			if (GetSchema()->CanUnfoldPin(this, Pin))
 			{
 				TArray<FString> ElementDefaultValues = URigVMPin::SplitDefaultValue(DefaultValue);
 				AddPinsForArray(ArrayProperty, InNode, Pin, Pin->Direction, ElementDefaultValues, bAutoExpandArrays);
@@ -14453,49 +14194,6 @@ void URigVMController::ConfigurePinFromArgument(URigVMPin* InOutPin, const FRigV
 	}
 }
 
-bool URigVMController::ShouldStructBeUnfolded(const UStruct* Struct)
-{
-	if (Struct == nullptr)
-	{
-		return false;
-	}
-	if (Struct->IsChildOf(UClass::StaticClass()))
-	{
-		return false;
-	}
-	if(Struct->IsChildOf(FRigVMExecuteContext::StaticStruct()))
-	{
-		return false;
-	}
-	if(Struct->IsChildOf(RigVMTypeUtils::GetWildCardCPPTypeObject()))
-	{
-		return false;
-	}
-	if (UnfoldStructDelegate.IsBound())
-	{
-		if (!UnfoldStructDelegate.Execute(Struct))
-		{
-			return false;
-		}
-	}
-	return true;
-}
-
-bool URigVMController::ShouldPinBeUnfolded(URigVMPin* InPin)
-{
-	if (InPin->IsStruct())
-	{
-		return ShouldStructBeUnfolded(InPin->GetScriptStruct());
-	}
-	if (InPin->IsArray())
-	{
-		return InPin->GetDirection() == ERigVMPinDirection::Input ||
-			InPin->GetDirection() == ERigVMPinDirection::IO ||
-			InPin->IsFixedSizeArray();
-	}
-	return false;
-}
-
 FProperty* URigVMController::FindPropertyForPin(const FString& InPinPath)
 {
 	if(!IsValidGraph())
@@ -16556,10 +16254,10 @@ bool URigVMController::ResolveWildCardPin(const FString& InPinPath, const FStrin
 		{
 			if (bPrintPythonCommand)
 			{
-				const FString GraphName = GetSanitizedGraphName(GetGraph()->GetGraphName());
+				const FString GraphName = GetSchema()->GetSanitizedGraphName(GetGraph()->GetGraphName());
 
 				// bool ResolveWildCardPin(const FString& InPinPath, const FString& InCPPType, const FName& InCPPTypeObjectPath, bool bSetupUndoRedo = true, bool bPrintPythonCommand = false);
-				RigVMPythonUtils::Print(GetGraphOuterName(), 
+				RigVMPythonUtils::Print(GetSchema()->GetGraphOuterName(GetGraph()), 
 									FString::Printf(TEXT("blueprint.get_controller_by_name('%s').resolve_wild_card_pin('%s', '%s', '%s')"),
 													*GraphName,
 													*InPinPath,
@@ -17264,6 +16962,11 @@ bool URigVMController::ChangePinType(URigVMPin* InPin, TRigVMTypeIndex InTypeInd
 		return false;
 	}
 
+	if(!GetSchema()->SupportsType(this, InTypeIndex))
+	{
+		return false;
+	}
+
 	// only allow valid pin cpp types on template nodes
 	TRigVMTypeIndex TypeIndex = InTypeIndex;
 	if(URigVMTemplateNode* TemplateNode = Cast<URigVMTemplateNode>(InPin->GetNode()))
@@ -17713,6 +17416,29 @@ void URigVMController::MakeExecutePin(URigVMPin* InOutPin)
 	}
 }
 
+bool URigVMController::AddGraphNode(URigVMNode* InNode, bool bNotify)
+{
+	URigVMGraph* Graph = GetGraph();
+	
+	check(Graph);
+	check(InNode);
+
+	if(!GetSchema()->CanAddNode(this, InNode))
+	{
+		Graph->Nodes.Remove(InNode);
+		DestroyObject(InNode);
+		return false;
+	}
+
+	Graph->Nodes.AddUnique(InNode);
+
+	if(bNotify)
+	{
+		Notify(ERigVMGraphNotifType::NodeAdded, InNode);
+	}
+	return true;
+}
+
 void URigVMController::AddNodePin(URigVMNode* InNode, URigVMPin* InPin)
 {
 	ValidatePin(InPin);
@@ -17801,7 +17527,7 @@ void URigVMController::EnsureLocalVariableValidity()
 	}
 }
 
-FRigVMExternalVariable URigVMController::GetVariableByName(const FName& InExternalVariableName, const bool bIncludeInputArguments)
+FRigVMExternalVariable URigVMController::GetVariableByName(const FName& InExternalVariableName, const bool bIncludeInputArguments) const
 {
 	TArray<FRigVMExternalVariable> Variables = GetAllVariables(bIncludeInputArguments);
 	for (const FRigVMExternalVariable& Variable : Variables)
@@ -17815,7 +17541,7 @@ FRigVMExternalVariable URigVMController::GetVariableByName(const FName& InExtern
 	return FRigVMExternalVariable();
 }
 
-TArray<FRigVMExternalVariable> URigVMController::GetAllVariables(const bool bIncludeInputArguments)
+TArray<FRigVMExternalVariable> URigVMController::GetAllVariables(const bool bIncludeInputArguments) const
 {
 	TArray<FRigVMExternalVariable> ExternalVariables;
 
@@ -17850,7 +17576,7 @@ void URigVMController::RefreshFunctionReferences(URigVMLibraryNode* InFunctionDe
 
 	if (const URigVMFunctionLibrary* FunctionLibrary = Cast<URigVMFunctionLibrary>(InFunctionDefinition->GetGraph()))
 	{
-		FunctionLibrary->ForEachReference(InFunctionDefinition->GetFName(), [this, bSetupUndoRedo](URigVMFunctionReferenceNode* ReferenceNode)
+		FunctionLibrary->ForEachReference(InFunctionDefinition->GetFName(), [this](URigVMFunctionReferenceNode* ReferenceNode)
 		{
 			if(URigVMController* ReferenceController = GetControllerForGraph(ReferenceNode->GetGraph()))
 			{
@@ -17864,70 +17590,6 @@ void URigVMController::RefreshFunctionReferences(URigVMLibraryNode* InFunctionDe
 	}
 }
 
-FString URigVMController::GetGraphOuterName() const
-{
-	check(GetGraph() != nullptr);
-	return GetSanitizedName(GetGraph()->GetRootGraph()->GetOuter()->GetFName().ToString(), true, false);
-}
-
-FString URigVMController::GetSanitizedName(const FString& InName, bool bAllowPeriod, bool bAllowSpace)
-{
-	FString CopiedName = InName;
-	SanitizeName(CopiedName, bAllowPeriod, bAllowSpace);
-	return CopiedName;
-}
-
-FString URigVMController::GetSanitizedGraphName(const FString& InName)
-{
-	return GetSanitizedName(InName, true, true);
-}
-
-FString URigVMController::GetSanitizedNodeName(const FString& InName)
-{
-	return GetSanitizedName(InName, false, true);
-}
-
-FString URigVMController::GetSanitizedVariableName(const FString& InName)
-{
-	return GetSanitizedName(InName, false, true);
-}
-
-FString URigVMController::GetSanitizedPinName(const FString& InName)
-{
-	return GetSanitizedName(InName, false, true);
-}
-
-FString URigVMController::GetSanitizedPinPath(const FString& InName)
-{
-	return GetSanitizedName(InName, true, true);
-}
-
-void URigVMController::SanitizeName(FString& InOutName, bool bAllowPeriod, bool bAllowSpace)
-{
-	// Sanitize the name
-	for (int32 i = 0; i < InOutName.Len(); ++i)
-	{
-		TCHAR& C = InOutName[i];
-
-		const bool bGoodChar =
-			FChar::IsAlpha(C) ||											// Any letter (upper and lowercase) anytime
-			(C == '_') || (C == '-') || 									// _  and - anytime
-			(bAllowPeriod && (C == '.')) ||
-			(bAllowSpace && (C == ' ')) ||
-			((i > 0) && FChar::IsDigit(C));									// 0-9 after the first character
-
-		if (!bGoodChar)
-		{
-			C = '_';
-		}
-	}
-
-	if (InOutName.Len() > GetMaxNameLength())
-	{
-		InOutName.LeftChopInline(InOutName.Len() - GetMaxNameLength());
-	}
-}
-
 URigVMController::FLinkedPath::FLinkedPath(URigVMLink* InLink)
 	: GraphPtr(InLink->GetGraph())
 	, SourcePinPath(InLink->GetSourcePinPath())
@@ -18644,7 +18306,7 @@ URigVMNode* URigVMController::ConvertRerouteNodeToDispatch(URigVMRerouteNode* In
 	// create the new node, and apply pin states to transfer ownership of injected nodes
 	// to the new node. Only after everything has been moved, we can safely remove the
 	// old node.
-	const FString DeletedName = GetValidNodeName(FString::Printf(TEXT("%s_Deleted"), *NodeName));
+	const FString DeletedName = GetSchema()->GetValidNodeName(GetGraph(), FString::Printf(TEXT("%s_Deleted"), *NodeName));
 	RenameNode(InRerouteNode, *DeletedName, false);
 
 	URigVMNode* NewNode = AddTemplateNode(InTemplateNotation, NodePosition, NodeName, bSetupUndoRedo, bPrintPythonCommand);
@@ -19081,7 +18743,7 @@ FRigVMClientPatchResult URigVMController::PatchUnitNodesOnLoad()
 			// create the new node, and apply pin states to transfer ownership of injected nodes
 			// to the new node. Only after everything has been moved, we can safely remove the
 			// old node.
-			const FString DeletedName = GetValidNodeName(FString::Printf(TEXT("%s_Deleted"), *NodeName));
+			const FString DeletedName = GetSchema()->GetValidNodeName(Graph, FString::Printf(TEXT("%s_Deleted"), *NodeName));
 			RenameNode(UnitNode, *DeletedName, false, false);
 			
 			URigVMTemplateNode* NewNode = AddTemplateNode(
@@ -19211,7 +18873,7 @@ FRigVMClientPatchResult URigVMController::PatchBranchNodesOnLoad()
 			// create the new node, and apply pin states to transfer ownership of injected nodes
 			// to the new node. Only after everything has been moved, we can safely remove the
 			// old node.
-			const FString DeletedName = GetValidNodeName(FString::Printf(TEXT("%s_Deleted"), *NodeName));
+			const FString DeletedName = GetSchema()->GetValidNodeName(Graph, FString::Printf(TEXT("%s_Deleted"), *NodeName));
 			RenameNode(BranchNode, *DeletedName, false, false);
 
 			// Cannot reuse the name of deprecated nodes, otherwise we may get the following error when PIE a second time
@@ -19267,7 +18929,7 @@ FRigVMClientPatchResult URigVMController::PatchIfSelectNodesOnLoad()
 			// create the new node, and apply pin states to transfer ownership of injected nodes
 			// to the new node. Only after everything has been moved, we can safely remove the
 			// old node.
-			const FString DeletedName = GetValidNodeName(FString::Printf(TEXT("%s_Deleted"), *NodeName));
+			const FString DeletedName = GetSchema()->GetValidNodeName(Graph, FString::Printf(TEXT("%s_Deleted"), *NodeName));
 			RenameNode(IfOrSelectNode, *DeletedName, false, false);
 
 			const FRigVMDispatchFactory* Factory = FRigVMRegistry::Get().FindOrAddDispatchFactory(
@@ -19348,7 +19010,7 @@ FRigVMClientPatchResult URigVMController::PatchArrayNodesOnLoad()
 			// create the new node, and apply pin states to transfer ownership of injected nodes
 			// to the new node. Only after everything has been moved, we can safely remove the
 			// old node.
-			const FString DeletedName = GetValidNodeName(FString::Printf(TEXT("%s_Deleted"), *NodeName));
+			const FString DeletedName = GetSchema()->GetValidNodeName(Graph, FString::Printf(TEXT("%s_Deleted"), *NodeName));
 			RenameNode(ArrayNode, *DeletedName, false);
 			
 			URigVMNode* NewNode = AddArrayNode(OpCode, CPPType, CPPTypeObject, NodePosition, NodeName, false, false, true);
diff --git a/Engine/Plugins/Runtime/RigVM/Source/RigVMDeveloper/Private/RigVMModel/RigVMNode.cpp b/Engine/Plugins/Runtime/RigVM/Source/RigVMDeveloper/Private/RigVMModel/RigVMNode.cpp
--- a/Engine/Plugins/Runtime/RigVM/Source/RigVMDeveloper/Private/RigVMModel/RigVMNode.cpp
+++ b/Engine/Plugins/Runtime/RigVM/Source/RigVMDeveloper/Private/RigVMModel/RigVMNode.cpp
@@ -396,9 +396,9 @@ uint32 URigVMNode::GetStructureHash() const
 	return Hash;
 }
 
-URigVMLibraryNode* URigVMNode::FindFunctionForNode() 
+URigVMLibraryNode* URigVMNode::FindFunctionForNode() const  
 {
-	UObject* Subject = this;
+	const UObject* Subject = this;
 	while (Subject->GetOuter() && !Subject->GetOuter()->IsA<URigVMFunctionLibrary>())
 	{
 		Subject = Subject->GetOuter();
@@ -408,7 +408,7 @@ URigVMLibraryNode* URigVMNode::FindFunctionForNode()
 		}
 	}
 
-	return Cast<URigVMLibraryNode>(Subject);
+	return const_cast<URigVMLibraryNode*>(Cast<URigVMLibraryNode>(Subject));
 }
 
 bool URigVMNode::IsLinkedToRecursive(URigVMPin* InPin, URigVMNode* InNode) const
diff --git a/Engine/Plugins/Runtime/RigVM/Source/RigVMDeveloper/Private/RigVMModel/RigVMPin.cpp b/Engine/Plugins/Runtime/RigVM/Source/RigVMDeveloper/Private/RigVMModel/RigVMPin.cpp
--- a/Engine/Plugins/Runtime/RigVM/Source/RigVMDeveloper/Private/RigVMModel/RigVMPin.cpp
+++ b/Engine/Plugins/Runtime/RigVM/Source/RigVMDeveloper/Private/RigVMModel/RigVMPin.cpp
@@ -1712,7 +1712,7 @@ URigVMGraph* URigVMPin::GetGraph() const
 	return nullptr;
 }
 
-bool URigVMPin::CanLink(URigVMPin* InSourcePin, URigVMPin* InTargetPin, FString* OutFailureReason, const FRigVMByteCode* InByteCode, ERigVMPinDirection InUserLinkDirection, bool bInAllowNonArgumentPins, bool bEnableTypeCasting)
+bool URigVMPin::CanLink(const URigVMPin* InSourcePin, const URigVMPin* InTargetPin, FString* OutFailureReason, const FRigVMByteCode* InByteCode, ERigVMPinDirection InUserLinkDirection, bool bInAllowNonArgumentPins, bool bEnableTypeCasting)
 {
 	if (InSourcePin == nullptr || InTargetPin == nullptr)
 	{
@@ -1819,7 +1819,7 @@ bool URigVMPin::CanLink(URigVMPin* InSourcePin, URigVMPin* InTargetPin, FString*
 
 			if(bCPPTypesDiffer)
 			{
-				auto TemplateNodeSupportsType = [](URigVMPin* InPin, const int32& InTypeIndex, FString* OutFailureReason) -> bool
+				auto TemplateNodeSupportsType = [](const URigVMPin* InPin, const int32& InTypeIndex, FString* OutFailureReason) -> bool
 				{
 					if (URigVMTemplateNode* TemplateNode = Cast<URigVMTemplateNode>(InPin->GetNode()))
 					{
@@ -1851,8 +1851,8 @@ bool URigVMPin::CanLink(URigVMPin* InSourcePin, URigVMPin* InTargetPin, FString*
 					// Find out if these pins have any type in common
 					uint8 SourceLevels = 0;
 					uint8 TargetLevels = 0;
-					URigVMPin* RootSourcePin = InSourcePin;
-					URigVMPin* RootTargetPin = InTargetPin;
+					const URigVMPin* RootSourcePin = InSourcePin;
+					const URigVMPin* RootTargetPin = InTargetPin;
 					while (RootSourcePin->IsArrayElement())
 					{
 						SourceLevels++;
diff --git a/Engine/Plugins/Runtime/RigVM/Source/RigVMDeveloper/Private/RigVMModel/RigVMSchema.cpp b/Engine/Plugins/Runtime/RigVM/Source/RigVMDeveloper/Private/RigVMModel/RigVMSchema.cpp
new file mode 100644
--- /dev/null
+++ b/Engine/Plugins/Runtime/RigVM/Source/RigVMDeveloper/Private/RigVMModel/RigVMSchema.cpp
@@ -0,0 +1,848 @@
+// Copyright Epic Games, Inc. All Rights Reserved.
+
+#include "RigVMModel/RigVMSchema.h"
+#include "RigVMModel/RigVMClient.h"
+#include "RigVMModel/RigVMController.h"
+
+URigVMSchema::URigVMSchema()
+	: ExecuteContextStruct(nullptr)
+	, Registry(&FRigVMRegistry::Get())
+{
+	SetExecuteContextStruct(FRigVMExecuteContext::StaticStruct());
+}
+
+URigVMSchema::URigVMSchema(const FObjectInitializer& ObjectInitializer)
+	: Super(ObjectInitializer)
+	, ExecuteContextStruct(nullptr)
+	, Registry(&FRigVMRegistry::Get())
+{
+	SetExecuteContextStruct(FRigVMExecuteContext::StaticStruct());
+}
+
+bool URigVMSchema::SupportsType(URigVMController* InController, TRigVMTypeIndex InTypeIndex) const
+{
+	// filter out incompatible execute types
+	if(Registry->IsExecuteType(InTypeIndex))
+	{
+		const UStruct* Struct = CastChecked<UStruct>(Registry->GetType(InTypeIndex).CPPTypeObject);
+		if(!ValidExecuteContextStructs.Contains(Struct))
+		{
+			static constexpr TCHAR Format[] = TEXT("ExecuteContext struct '%s' is not supported.");
+			InController->ReportErrorf(Format, *Struct->GetName());
+			return false;
+		}
+	}
+	return true;
+}
+
+bool URigVMSchema::SupportsUnitFunction(URigVMController* InController, const FRigVMFunction* InUnitFunction) const
+{
+	if(InUnitFunction)
+	{
+		if(const UScriptStruct* FunctionExecuteContextStruct = InUnitFunction->GetExecuteContextStruct())
+		{
+			const TRigVMTypeIndex TypeIndex = Registry->GetTypeIndexFromCPPType(FunctionExecuteContextStruct->GetStructCPPName());
+			if(!SupportsType(InController, TypeIndex))
+			{
+				return false;
+			}
+		}
+		
+		const TArray<TRigVMTypeIndex>& ArgumentTypes = InUnitFunction->GetArgumentTypeIndices();
+		for(const TRigVMTypeIndex& ArgumentType : ArgumentTypes)
+		{
+			if(!SupportsType(InController, ArgumentType))
+			{
+				const FString CPPTypeString = Registry->GetType(ArgumentType).CPPType.ToString();
+				static constexpr TCHAR Format[] = TEXT("Unit function '%s' is not supported since type '%s' is not supported.");
+				InController->ReportErrorf(Format, *InUnitFunction->GetName(), *CPPTypeString);
+				return false;
+			}
+		}
+		return true;
+	}
+
+	return false;
+}
+
+bool URigVMSchema::SupportsDispatchFactory(URigVMController* InController, const FRigVMDispatchFactory* InDispatchFactory) const
+{
+	if(InDispatchFactory)
+	{
+		if(const UScriptStruct* DispatchExecuteContextStruct = InDispatchFactory->GetExecuteContextStruct())
+		{
+			const TRigVMTypeIndex TypeIndex = Registry->GetTypeIndexFromCPPType(DispatchExecuteContextStruct->GetStructCPPName());
+			if(!SupportsType(InController, TypeIndex))
+			{
+				return false;
+			}
+		}
+
+		if(const FRigVMTemplate* Template = InDispatchFactory->GetTemplate())
+		{
+			if(!SupportsTemplate(InController, Template))
+			{
+				return false;
+			}
+		}
+		return true;
+	}
+
+	return false;
+}
+
+bool URigVMSchema::SupportsTemplate(URigVMController* InController, const FRigVMTemplate* InTemplate) const
+{
+	if(InTemplate)
+	{
+		// we are only checking for execute arguments here to make sure that we don't
+		// support access to templates that cannot be compiled for the given graph.
+		// normal template arguments (non-execute) will be allowed - but during resolval
+		// we'll check if the target types are supported by the graph.
+		static const FRigVMDispatchContext DispatchContext;
+		for(int32 Index = 0; Index < InTemplate->NumExecuteArguments(DispatchContext); Index++)
+		{
+			const TRigVMTypeIndex& TypeIndex = InTemplate->GetExecuteArgument(Index, DispatchContext)->TypeIndex;
+			if(!SupportsType(InController, TypeIndex))
+			{
+				const FString CPPTypeString = Registry->GetType(TypeIndex).CPPType.ToString();
+				static constexpr TCHAR Format[] = TEXT("Template '%s' is not supported since type '%s' is not supported.");
+				InController->ReportErrorf(Format, *InTemplate->GetNotation().ToString(), *CPPTypeString);
+				return false;
+			}
+		}
+		return true;
+	}
+
+	return false;
+
+}
+
+bool URigVMSchema::SupportsGraphFunction(URigVMController* InController, const FRigVMGraphFunctionHeader* InGraphFunction) const
+{
+	RIGVMSCHEMA_DEFAULT_FUNCTION_BODY
+
+	check(InGraphFunction);
+
+	if(Graph->IsA<URigVMFunctionLibrary>())
+	{
+		static const FString Message = TEXT("You cannot place functions inside of the top level function library graph.");
+		InController->ReportError(Message);
+		return false;
+	}
+
+	if (!InController->bAllowPrivateFunctions)
+	{
+		if(IRigVMClientHost* ClientHost = Graph->GetImplementingOuter<IRigVMClientHost>())
+		{
+			bool bIsAvailable = ClientHost->GetRigVMClient()->GetFunctionLibrary()->GetFunctionHostObjectPath() ==
+				InGraphFunction->LibraryPointer.HostObject;
+			if (!bIsAvailable)
+			{
+				if (IRigVMGraphFunctionHost* Host = Cast<IRigVMGraphFunctionHost>(InGraphFunction->LibraryPointer.HostObject.TryLoad()))
+				{
+					bIsAvailable = Host->GetRigVMGraphFunctionStore()->IsFunctionPublic(InGraphFunction->LibraryPointer);		
+				}
+			}
+			if (!bIsAvailable)
+			{
+				InController->ReportError(TEXT("Function is not available for placement in another graph host."));
+				return false;
+			}
+		}
+	}
+
+	if (const URigVMNode* Node = Cast<URigVMNode>(Graph->GetOuter()))
+	{
+		if (const URigVMLibraryNode* LibraryNode = Node->FindFunctionForNode())
+		{
+			if (InGraphFunction->Dependencies.Contains(LibraryNode->GetFunctionIdentifier()))
+			{
+				static const FString Message = TEXT("Function is not available for placement in this graph host due to dependency cycles."); 
+				InController->ReportError(Message);
+				return false;
+			}
+		}
+	}
+	
+	URigVMLibraryNode* ParentLibraryNode = Cast<URigVMLibraryNode>(Graph->GetOuter());
+	while (ParentLibraryNode)
+	{
+		if (TSoftObjectPtr<UObject>(ParentLibraryNode).ToSoftObjectPath() == InGraphFunction->LibraryPointer.LibraryNode)
+		{
+			static const FString Message = TEXT("You cannot place functions inside of itself or an indirect recursion.");
+			InController->ReportError(Message);
+			return false;
+		}
+		ParentLibraryNode = Cast<URigVMLibraryNode>(ParentLibraryNode->GetGraph()->GetOuter());
+	}
+
+	return true;
+}
+
+bool URigVMSchema::SupportsExternalVariable(URigVMController* InController, const FRigVMExternalVariable* InExternalVariable) const
+{
+	RIGVMSCHEMA_DEFAULT_FUNCTION_BODY
+
+	if(InExternalVariable)
+	{
+		TRigVMTypeIndex Type = Registry->GetTypeIndex(InExternalVariable->GetExtendedCPPType(), InExternalVariable->TypeObject);
+		if(Type == INDEX_NONE)
+		{
+			return false;
+		}
+		return SupportsType(InController, Type);
+	}
+
+	return false;
+}
+
+bool URigVMSchema::ShouldUnfoldStruct(URigVMController* InController, const UStruct* InStruct) const
+{
+	RIGVMSCHEMA_DEFAULT_FUNCTION_BODY
+
+	if (InStruct == nullptr)
+	{
+		return false;
+	}
+	if (InStruct->IsChildOf(UClass::StaticClass()))
+	{
+		return false;
+	}
+	if(InStruct->IsChildOf(FRigVMExecuteContext::StaticStruct()))
+	{
+		return false;
+	}
+	if(InStruct->IsChildOf(RigVMTypeUtils::GetWildCardCPPTypeObject()))
+	{
+		return false;
+	}
+
+	return true;
+}
+
+bool URigVMSchema::IsValidNodeName(const URigVMGraph* InGraph, const FName& InNodeName) const
+{
+	return InGraph->IsNameAvailable(InNodeName.ToString());
+}
+
+bool URigVMSchema::CanAddNode(URigVMController* InController, const URigVMNode* InNode) const
+{
+	RIGVMSCHEMA_DEFAULT_FUNCTION_BODY
+
+	if(InNode == nullptr)
+	{
+		return false;
+	}
+
+	if (Graph->IsA<URigVMFunctionLibrary>())
+	{
+		if(!CanAddFunction(InController, InNode))
+		{
+			return false;
+		}
+	}
+
+	if (const URigVMFunctionReferenceNode* FunctionRefNode = Cast<URigVMFunctionReferenceNode>(InNode))
+	{
+		FRigVMGraphFunctionHeader FunctionDefinition = FunctionRefNode->GetReferencedFunctionHeader();
+
+		bool bSupportsGraphFunction;
+		{
+			TGuardValue<bool> GuardErrorReporting(InController->bReportWarningsAndErrors, false);
+			bSupportsGraphFunction = SupportsGraphFunction(InController, &FunctionDefinition);
+		}
+		if(!bSupportsGraphFunction)
+		{
+			URigVMFunctionLibrary* TargetLibrary = Graph->GetDefaultFunctionLibrary();
+			URigVMLibraryNode* LocalizedFunctionDefinition = TargetLibrary->FindPreviouslyLocalizedFunction(FunctionDefinition.LibraryPointer);
+
+			if(InController)
+			{
+				if((LocalizedFunctionDefinition == nullptr) && InController->RequestLocalizeFunctionDelegate.IsBound())
+				{
+					if(InController->RequestLocalizeFunctionDelegate.Execute(FunctionDefinition.LibraryPointer))
+					{
+						LocalizedFunctionDefinition = TargetLibrary->FindPreviouslyLocalizedFunction(FunctionDefinition.LibraryPointer);
+					}
+				}
+
+				if(LocalizedFunctionDefinition == nullptr)
+				{
+					return false;
+				}
+			
+				InController->SetReferencedFunction(const_cast<URigVMFunctionReferenceNode*>(FunctionRefNode), LocalizedFunctionDefinition, false);
+				FunctionDefinition = FunctionRefNode->GetReferencedFunctionHeader();
+
+				if(!SupportsGraphFunction(InController, &FunctionDefinition))
+				{
+					return false;
+				}
+			}
+		}
+	}
+	else if(InNode->IsA<URigVMFunctionEntryNode>() ||
+		InNode->IsA<URigVMFunctionReturnNode>())
+	{
+		// only allow entry / return nodes on sub graphs
+		if(Graph->IsRootGraph())
+		{
+			static const FString Message("Entry and Return nodes can only be added to sub graphs.");
+			InController->ReportError(Message);
+			return false;
+		}
+
+		// only allow one function entry node
+		if(InNode->IsA<URigVMFunctionEntryNode>())
+		{
+			if(const URigVMFunctionEntryNode* ExistingEntryNode = Graph->GetEntryNode())
+			{
+				if(ExistingEntryNode != InNode)
+				{
+					static const FString Message("Graphs can only contain on Entry node.");
+					InController->ReportError(Message);
+					return false;
+				}
+			}
+		}
+		// only allow one function return node
+		else if(InNode->IsA<URigVMFunctionReturnNode>())
+		{
+			if(const URigVMFunctionReturnNode* ExistingReturnNode = Graph->GetReturnNode())
+			{
+				if(ExistingReturnNode != InNode)
+				{
+					static const FString Message("Graphs can only contain on Return node.");
+					InController->ReportError(Message);
+					return false;
+				}
+			}
+		}
+	}
+	else if(const URigVMCollapseNode* CollapseNode = Cast<URigVMCollapseNode>(InNode))
+	{
+		const URigVMGraph* ContainedGraph = CollapseNode->GetContainedGraph();
+		const TArray<URigVMNode*> ContainedNodes = CollapseNode->GetContainedNodes();
+		URigVMController* ContainedController = InController->GetControllerForGraph(ContainedGraph);
+		for(const URigVMNode* ContainedNode : ContainedNodes)
+		{
+			if(!CanAddNode(ContainedController, ContainedNode))
+			{
+				return false;
+			}
+		}
+	}
+	else if(const URigVMVariableNode* VariableNode = Cast<URigVMVariableNode>(InNode))
+	{
+		if (const URigVMPin* NamePin = VariableNode->FindPin(URigVMVariableNode::VariableName))
+		{
+			const FString VarNameString = NamePin->GetDefaultValue();
+			if (!VarNameString.IsEmpty())
+			{
+				const FName VarName = *VarNameString;
+
+				TArray<FRigVMExternalVariable> AllVariables = InController->GetAllVariables(true);
+				for(const FRigVMExternalVariable& Variable : AllVariables)
+				{
+					if(Variable.Name.IsEqual(VarName, ENameCase::CaseSensitive))
+					{
+						return true;
+					}
+				}
+				return false;
+			}
+		}
+	}
+	else if (InNode->IsEvent())
+	{
+		if (const URigVMUnitNode* InUnitNode = Cast<URigVMUnitNode>(InNode))
+		{
+			if(const UScriptStruct* EventStruct = InUnitNode->GetScriptStruct())
+			{
+				// check if we're trying to add a node within a graph which is not the top level one
+				if (!Graph->IsTopLevelGraph())
+				{
+					static const FString Message = TEXT("Event nodes can only be added to top level graphs."); 
+					InController->ReportError(Message);
+					return false;
+				}
+
+				const TObjectPtr<URigVMNode> EventNode = FindEventNode(InController, EventStruct);
+				const bool bHasEventNode = (EventNode != nullptr) && EventNode->CanOnlyExistOnce();
+				if (bHasEventNode)
+				{
+					const FString ErrorMessage = FString::Printf(TEXT("Rig Graph can only contain one single %s node."),
+																 *EventStruct->GetDisplayNameText().ToString());
+					InController->ReportError(ErrorMessage);
+					return false;
+				}
+			}
+		}
+	}
+	
+	return true;
+}
+
+bool URigVMSchema::CanRemoveNode(URigVMController* InController, const URigVMNode* InNode) const
+{
+	RIGVMSCHEMA_DEFAULT_FUNCTION_BODY
+
+	if(InNode->IsA<URigVMFunctionEntryNode>() ||
+		InNode->IsA<URigVMFunctionReturnNode>())
+	{
+		return false;
+	}
+
+	return true;
+}
+
+bool URigVMSchema::CanRenameNode(URigVMController* InController, const URigVMNode* InNode, const FName& InNewNodeName) const
+{
+	RIGVMSCHEMA_DEFAULT_FUNCTION_BODY
+
+	if(Graph)
+	{
+		if(InNode)
+		{
+			if(InNode->GetName() != InNewNodeName)
+			{
+				return IsValidNodeName(Graph, InNewNodeName);
+			}
+		}
+	}
+
+	return false;
+}
+
+bool URigVMSchema::CanMoveNode(URigVMController* InController, const URigVMNode* InNode, const FVector2D& InNewPosition) const
+{
+	RIGVMSCHEMA_DEFAULT_FUNCTION_BODY
+
+	if ((InNode->Position - InNewPosition).IsNearlyZero())
+	{
+		return false;
+	}
+
+	return true;
+}
+
+bool URigVMSchema::CanResizeNode(URigVMController* InController, const URigVMNode* InNode, const FVector2D& InNewSize) const
+{
+	RIGVMSCHEMA_DEFAULT_FUNCTION_BODY
+
+	if ((InNode->Size - InNewSize).IsNearlyZero())
+	{
+		return false;
+	}
+
+	return InNode->IsA<URigVMCommentNode>();
+}
+
+bool URigVMSchema::CanRecolorNode(URigVMController* InController, const URigVMNode* InNode, const FLinearColor& InNewColor) const
+{
+	RIGVMSCHEMA_DEFAULT_FUNCTION_BODY
+
+	if (InNode->NodeColor.Equals(InNewColor, 0.001))
+	{
+		return false;
+	}
+
+	return InNode->IsA<URigVMCommentNode>();
+}
+
+bool URigVMSchema::CanAddLink(URigVMController* InController, const URigVMPin* InSourcePin, const URigVMPin* InTargetPin, const FRigVMByteCode* InByteCode, ERigVMPinDirection InUserLinkDirection, bool bInAllowWildcard, bool bEnableTypeCasting, FString* OutFailureReason) const
+{
+	RIGVMSCHEMA_DEFAULT_FUNCTION_BODY
+
+	if (Graph->IsA<URigVMFunctionLibrary>())
+	{
+		static const FString FailureReason = TEXT("Cannot add links in function library graphs."); 
+		if(OutFailureReason)
+		{
+			*OutFailureReason = FailureReason;
+		}		
+		return false;
+	}
+
+	if(InSourcePin == nullptr)
+	{
+		static const FString FailureReason = TEXT("SourcePin is nullptr."); 
+		if(OutFailureReason)
+		{
+			*OutFailureReason = FailureReason;
+		}		
+		return false;
+	}
+
+	if(InTargetPin == nullptr)
+	{
+		static const FString FailureReason = TEXT("TargetPin is nullptr."); 
+		if(OutFailureReason)
+		{
+			*OutFailureReason = FailureReason;
+		}		
+		return false;
+	}
+
+	if((InSourcePin->GetGraph() != Graph) || (InTargetPin->GetGraph() != Graph))
+	{
+		static const FString FailureReason = TEXT("Pin is not valid for graph."); 
+		if(OutFailureReason)
+		{
+			*OutFailureReason = FailureReason;
+		}		
+		return false;
+	}
+
+	if(!URigVMPin::CanLink(InSourcePin, InTargetPin, OutFailureReason, InByteCode, InUserLinkDirection, bInAllowWildcard))
+	{
+		return false;
+	}
+
+	return true;
+}
+
+bool URigVMSchema::CanBreakLink(URigVMController* InController, const URigVMPin* InSourcePin, const URigVMPin* InTargetPin) const
+{
+	RIGVMSCHEMA_DEFAULT_FUNCTION_BODY
+
+	if (Graph->IsA<URigVMFunctionLibrary>())
+	{
+		static const FString FailureReason = TEXT("Cannot break links in function library graphs."); 
+		InController->ReportError(FailureReason);
+		return false;
+	}
+
+	if(InSourcePin == nullptr)
+	{
+		static const FString FailureReason = TEXT("SourcePin is nullptr."); 
+		InController->ReportError(FailureReason);
+		return false;
+	}
+
+	if(InTargetPin == nullptr)
+	{
+		static const FString FailureReason = TEXT("TargetPin is nullptr."); 
+		InController->ReportError(FailureReason);
+		return false;
+	}
+
+	if((InSourcePin->GetGraph() != Graph) || (InTargetPin->GetGraph() != Graph))
+	{
+		static const FString FailureReason = TEXT("Pin is not valid for graph."); 
+		InController->ReportError(FailureReason);
+		return false;
+	}
+
+	if (!InSourcePin->IsLinkedTo(InTargetPin))
+	{
+		return false;
+	}
+	ensure(InTargetPin->IsLinkedTo(InSourcePin));
+	
+	return true;
+}
+
+bool URigVMSchema::CanCollapseNodes(URigVMController* InController, const TArrayView<URigVMNode* const>& InNodesToCollapse) const
+{
+	if(InNodesToCollapse.IsEmpty())
+	{
+		return false;
+	}
+
+	RIGVMSCHEMA_DEFAULT_FUNCTION_BODY
+
+	for(int32 Index = 0; Index < InNodesToCollapse.Num(); Index++)
+	{
+		if(InNodesToCollapse[Index]->GetGraph() != Graph)
+		{
+			static const FString Message = TEXT("You can only collapse nodes within the same graph.");
+			InController->ReportError(Message);
+			return false;
+		}
+	}
+
+	return true;
+}
+
+bool URigVMSchema::CanExpandNode(URigVMController* InController, const URigVMNode* InNodeToExpand) const
+{
+	check(InNodeToExpand);
+
+	RIGVMSCHEMA_DEFAULT_FUNCTION_BODY
+
+	if(Graph->IsA<URigVMFunctionLibrary>())
+	{
+		static const FString Message = TEXT("You cannot expand subgraphs within a function library graph.");
+		InController->ReportError(Message);
+		return false;
+	}
+
+	return true;
+}
+
+bool URigVMSchema::CanUnfoldPin(URigVMController* InController, const URigVMPin* InPinToUnfold) const
+{
+	RIGVMSCHEMA_DEFAULT_FUNCTION_BODY
+
+	if (InPinToUnfold->IsStruct())
+	{
+		return ShouldUnfoldStruct(InController, InPinToUnfold->GetScriptStruct());
+	}
+	if (InPinToUnfold->IsArray())
+	{
+		return InPinToUnfold->GetDirection() == ERigVMPinDirection::Input ||
+			InPinToUnfold->GetDirection() == ERigVMPinDirection::IO ||
+			InPinToUnfold->IsFixedSizeArray();
+	}
+	return false;
+}
+
+bool URigVMSchema::CanBindVariable(URigVMController* InController, const URigVMPin* InPinToBind, const FRigVMExternalVariable* InVariableToBind, const FString& InNewBoundVariablePath) const
+{
+	RIGVMSCHEMA_DEFAULT_FUNCTION_BODY
+
+	if (Graph->IsA<URigVMFunctionLibrary>())
+	{
+		InController->ReportError(TEXT("Cannot bind pins to variables in function library graphs."));
+		return false;
+	}
+
+	if (InPinToBind->GetBoundVariablePath() == InNewBoundVariablePath)
+	{
+		return false;
+	}
+
+	if (InPinToBind->GetDirection() != ERigVMPinDirection::Input)
+	{
+		InController->ReportError(TEXT("Variables can only be bound to input pins."));
+		return false;
+	}
+
+	if (!InPinToBind->IsRootPin())
+	{
+		InController->ReportError(TEXT("Variables can only be bound to root pins."));
+		return false;
+	}
+
+	if(!SupportsType(InController, InVariableToBind->GetTypeIndex()))
+	{
+		return false;
+	}
+
+	return true;
+}
+
+bool URigVMSchema::CanUnbindVariable(URigVMController* InController, const URigVMPin* InBoundPin) const
+{
+	RIGVMSCHEMA_DEFAULT_FUNCTION_BODY
+
+	check(InBoundPin);
+
+	if (Graph->IsA<URigVMFunctionLibrary>())
+	{
+		InController->ReportError(TEXT("Cannot unbind pins from variables in function library graphs."));
+		return false;
+	}
+	if (!InBoundPin->IsBoundToVariable())
+	{
+		InController->ReportError(TEXT("Pin is not bound to any variable."));
+		return false;
+	}
+
+	return true;
+}
+
+bool URigVMSchema::CanAddFunction(URigVMController* InController, const URigVMNode* InFunctionNode) const
+{
+	RIGVMSCHEMA_DEFAULT_FUNCTION_BODY
+
+	if (!Graph->IsA<URigVMFunctionLibrary>())
+	{
+		InController->ReportError(TEXT("Can only add function definitions to function library graphs."));
+		return false;
+	}
+
+	if(InFunctionNode)
+	{
+		if (!InFunctionNode->IsA<URigVMCollapseNode>())
+		{
+			return false;
+		}
+	}
+
+	return true;
+}
+
+bool URigVMSchema::CanRemoveFunction(URigVMController* InController, const URigVMNode* InFunctionNode) const
+{
+	RIGVMSCHEMA_DEFAULT_FUNCTION_BODY
+
+	check(InFunctionNode);
+
+	if (!Graph->IsA<URigVMFunctionLibrary>())
+	{
+		InController->ReportError(TEXT("Can only remove function definitions from function library graphs."));
+		return false;
+	}
+
+	return true;
+}
+
+FString URigVMSchema::GetSanitizedName(const FString& InName, bool bAllowPeriod, bool bAllowSpace) const
+{
+	FString CopiedName = InName;
+	SanitizeName(CopiedName, bAllowPeriod, bAllowSpace);
+	return CopiedName;
+}
+
+FString URigVMSchema::GetSanitizedGraphName(const FString& InName) const
+{
+	return GetSanitizedName(InName, true, true);
+}
+
+FString URigVMSchema::GetSanitizedNodeName(const FString& InName) const
+{
+	return GetSanitizedName(InName, false, true);
+}
+
+FString URigVMSchema::GetSanitizedVariableName(const FString& InName) const
+{
+	return GetSanitizedName(InName, false, true);
+}
+
+FString URigVMSchema::GetSanitizedPinName(const FString& InName) const
+{
+	return GetSanitizedName(InName, false, true);
+}
+
+FString URigVMSchema::GetSanitizedPinPath(const FString& InName) const
+{
+	return GetSanitizedName(InName, true, true);
+}
+
+FString URigVMSchema::GetGraphOuterName(const URigVMGraph* InGraph) const
+{
+	check(InGraph);
+	return GetSanitizedName(InGraph->GetRootGraph()->GetOuter()->GetFName().ToString(), true, false);
+}
+
+FString URigVMSchema::GetValidNodeName(const URigVMGraph* InGraph, const FString& InPrefix) const
+{
+	check(InGraph);
+
+	return GetUniqueName(*InPrefix, [&](const FName& InName) {
+		return InGraph->IsNameAvailable(InName.ToString());
+	}, false, true).ToString();
+}
+
+void URigVMSchema::SanitizeName(FString& InOutName, bool bAllowPeriod, bool bAllowSpace)
+{
+	// Sanitize the name
+	for (int32 i = 0; i < InOutName.Len(); ++i)
+	{
+		TCHAR& C = InOutName[i];
+
+		const bool bGoodChar =
+			FChar::IsAlpha(C) ||											// Any letter (upper and lowercase) anytime
+			(C == '_') || (C == '-') || 									// _  and - anytime
+			(bAllowPeriod && (C == '.')) ||
+			(bAllowSpace && (C == ' ')) ||
+			((i > 0) && FChar::IsDigit(C));									// 0-9 after the first character
+
+		if (!bGoodChar)
+		{
+			C = '_';
+		}
+	}
+
+	if (InOutName.Len() > GetMaxNameLength())
+	{
+		InOutName.LeftChopInline(InOutName.Len() - GetMaxNameLength());
+	}
+}
+
+FName URigVMSchema::GetUniqueName(const FName& InName, TFunction<bool(const FName&)> IsNameAvailablePredicate,
+	bool bAllowPeriod, bool bAllowSpace)
+{
+	FString SanitizedPrefix = InName.ToString();
+	SanitizeName(SanitizedPrefix, bAllowPeriod, bAllowSpace);
+
+	int32 NameSuffix = 0;
+	FString Name = SanitizedPrefix;
+	while (!IsNameAvailablePredicate(*Name))
+	{
+		NameSuffix++;
+		Name = FString::Printf(TEXT("%s_%d"), *SanitizedPrefix, NameSuffix);
+	}
+	return *Name;
+}
+
+void URigVMSchema::SetExecuteContextStruct(UScriptStruct* InExecuteContextStruct)
+{
+	if(InExecuteContextStruct != ExecuteContextStruct)
+	{
+		if(GetClass() == URigVMSchema::StaticClass() && HasAnyFlags(RF_ClassDefaultObject))
+		{
+			// only allow the default execute context on the base class.
+			// please create a child class of URigVMSchema for your use case.
+			verify(InExecuteContextStruct == FRigVMExecuteContext::StaticStruct());
+			return;
+		}
+		
+		ExecuteContextStruct = InExecuteContextStruct;
+		ValidExecuteContextStructs.Reset();
+		
+		if(ExecuteContextStruct)
+		{
+			ValidExecuteContextStructs = FRigVMTemplate::GetSuperStructs(ExecuteContextStruct, true);
+		}
+	}
+}
+
+bool URigVMSchema::IsGraphEditable(const URigVMGraph* InGraph) const
+{
+	if(InGraph)
+	{
+		return InGraph->bEditable;
+	}
+	return false;
+}
+
+TObjectPtr<URigVMNode> URigVMSchema::FindEventNode(URigVMController* InController, const UScriptStruct* InScriptStruct) const
+{
+	const URigVMGraph* Graph = InController->GetGraph();
+	check(Graph);
+	check(InScriptStruct);
+
+	if (Graph)
+	{
+		// construct equivalent default struct
+		FStructOnScope InDefaultStructScope(InScriptStruct);
+
+		const TObjectPtr<URigVMNode>* FoundNode = 
+		Graph->Nodes.FindByPredicate( [&InDefaultStructScope](const TObjectPtr<URigVMNode>& Node) {
+			if (Node->IsEvent())
+			{
+				if (URigVMUnitNode* UnitNode = Cast<URigVMUnitNode>(Node))
+				{
+					// compare default structures
+					TSharedPtr<FStructOnScope> DefaultStructScope = UnitNode->ConstructStructInstance(true);
+					if (DefaultStructScope.IsValid() && InDefaultStructScope.GetStruct() == DefaultStructScope->GetStruct())
+					{
+						return true;
+					}
+				}
+			}
+			return false;
+		});
+
+		if (FoundNode)
+		{
+			return *FoundNode;
+		}
+	}
+	
+	return TObjectPtr<URigVMNode>();
+}
diff --git a/Engine/Plugins/Runtime/RigVM/Source/RigVMDeveloper/Public/RigVMModel/Nodes/RigVMVariableNode.h b/Engine/Plugins/Runtime/RigVM/Source/RigVMDeveloper/Public/RigVMModel/Nodes/RigVMVariableNode.h
--- a/Engine/Plugins/Runtime/RigVM/Source/RigVMDeveloper/Public/RigVMModel/Nodes/RigVMVariableNode.h
+++ b/Engine/Plugins/Runtime/RigVM/Source/RigVMDeveloper/Public/RigVMModel/Nodes/RigVMVariableNode.h
@@ -82,5 +82,6 @@ private:
 	friend class URigVMCompiler;
 	friend class FRigVMVarExprAST;
 	friend class FRigVMParserAST;
+	friend class URigVMSchema;
 };
 
diff --git a/Engine/Plugins/Runtime/RigVM/Source/RigVMDeveloper/Public/RigVMModel/RigVMClient.h b/Engine/Plugins/Runtime/RigVM/Source/RigVMDeveloper/Public/RigVMModel/RigVMClient.h
--- a/Engine/Plugins/Runtime/RigVM/Source/RigVMDeveloper/Public/RigVMModel/RigVMClient.h
+++ b/Engine/Plugins/Runtime/RigVM/Source/RigVMDeveloper/Public/RigVMModel/RigVMClient.h
@@ -3,6 +3,7 @@
 #pragma once
 
 #include "RigVMGraph.h"
+#include "RigVMSchema.h"
 #include "RigVMFunctionLibrary.h"
 #include "RigVMController.h"
 #include "RigVMClient.generated.h"
@@ -83,7 +84,8 @@ public:
 	GENERATED_BODY()
 
 	FRigVMClient()
-		: ExecuteContextStruct(nullptr)
+		: SchemaPtr(nullptr)
+		, SchemaClass(URigVMSchema::StaticClass())
 		, FunctionLibrary(nullptr)
 		, ActionStack(nullptr)
 		, bSuspendNotifications(false)
@@ -91,14 +93,16 @@ public:
 		, OuterClientHost(nullptr)
 		, OuterClientPropertyName(NAME_None)
 	{
-		ExecuteContextStruct = FRigVMExecuteContext::StaticStruct();
 	}
 
+	void SetSchemaClass(TSubclassOf<URigVMSchema> InSchemaClass);
 	void SetOuterClientHost(UObject* InOuterClientHost, const FName& InOuterClientHostPropertyName);
 	void SetFromDeprecatedData(URigVMGraph* InDefaultGraph, URigVMFunctionLibrary* InFunctionLibrary);
 
 	void Reset();
 	int32 Num() const { return Models.Num(); }
+	const URigVMSchema* GetSchema() const { return SchemaPtr; }
+	URigVMSchema* GetOrCreateSchema();
 	URigVMGraph* GetDefaultModel() const;
 	URigVMGraph* GetModel(int32 InIndex) const { return Models[InIndex]; }
 	URigVMGraph* GetModel(const FString& InNodePathOrName) const;
@@ -117,9 +121,9 @@ public:
 	URigVMFunctionLibrary* GetFunctionLibrary() const { return FunctionLibrary; }
 	URigVMFunctionLibrary* GetOrCreateFunctionLibrary(bool bSetupUndoRedo, const FObjectInitializer* ObjectInitializer = nullptr, bool bCreateController = true);
 	TArray<FName> GetEntryNames() const;
-	UScriptStruct* GetExecuteContextStruct() const { return ExecuteContextStruct; }
-	void SetExecuteContextStruct(UScriptStruct* InExecuteContextStruct) { ExecuteContextStruct = InExecuteContextStruct; }
-
+	UScriptStruct* GetExecuteContextStruct() const;
+	void SetExecuteContextStruct(UScriptStruct* InExecuteContextStruct);
+	
 	URigVMGraph* AddModel(const FName& InName, bool bSetupUndoRedo, const FObjectInitializer* ObjectInitializer = nullptr, bool bCreateController = true);
 	void AddModel(URigVMGraph* InModel, bool bCreateController);
 	bool RemoveModel(const FString& InNodePathOrName, bool bSetupUndoRedo);
@@ -191,8 +195,13 @@ private:
 	URigVMActionStack* GetOrCreateActionStack();
 	void ResetActionStack();
 
+	void SetSchema(URigVMSchema* InSchema);
+
 	UPROPERTY(transient)
-	TObjectPtr<UScriptStruct> ExecuteContextStruct;
+	TObjectPtr<URigVMSchema> SchemaPtr;
+
+	UPROPERTY(transient)
+	TSubclassOf<URigVMSchema> SchemaClass;
 
 	UPROPERTY()
 	TArray<TObjectPtr<URigVMGraph>> Models;
@@ -218,4 +227,6 @@ public:
 private:
 	TWeakObjectPtr<UObject> OuterClientHost;
 	FName OuterClientPropertyName;
+
+	friend class UEngineTestClientHost;
 };
diff --git a/Engine/Plugins/Runtime/RigVM/Source/RigVMDeveloper/Public/RigVMModel/RigVMController.h b/Engine/Plugins/Runtime/RigVM/Source/RigVMDeveloper/Public/RigVMModel/RigVMController.h
--- a/Engine/Plugins/Runtime/RigVM/Source/RigVMDeveloper/Public/RigVMModel/RigVMController.h
+++ b/Engine/Plugins/Runtime/RigVM/Source/RigVMDeveloper/Public/RigVMModel/RigVMController.h
@@ -4,6 +4,7 @@
 
 #include "RigVMGraph.h"
 #include "RigVMFunctionLibrary.h"
+#include "RigVMSchema.h"
 #include "RigVMModel/Nodes/RigVMUnitNode.h"
 #include "RigVMModel/Nodes/RigVMVariableNode.h"
 #include "RigVMModel/Nodes/RigVMParameterNode.h"
@@ -203,6 +204,13 @@ public:
 	UFUNCTION(BlueprintCallable, Category = RigVMController, meta=(DeprecatedFunction, DeprecationMessage="Function has been deprecated, please rely on GetControllerForGraph instead."))
 	void SetGraph(URigVMGraph* InGraph);
 
+	// Returns the schema used by this controller
+	UFUNCTION(BlueprintCallable, Category = RigVMController)
+	URigVMSchema* GetSchema() const { return SchemaPtr; }
+
+	// Sets the schema on the controller
+	void SetSchema(URigVMSchema* InSchema);
+
 	// Pushes a new graph to the stack
 	// This causes a GraphChanged modified event.
 	UFUNCTION(BlueprintCallable, Category = RigVMController, meta=(DeprecatedFunction, DeprecationMessage="Function has been deprecated, please rely on GetControllerForGraph instead."))
@@ -242,6 +250,10 @@ public:
 	UFUNCTION(BlueprintPure, Category = RigVMController)
 	bool IsReportingEnabled() const { return bReportWarningsAndErrors; }
 
+	// Returns true if the controller is currently transacting
+	UFUNCTION(BlueprintPure, Category = RigVMController)
+	bool IsTransacting() const { return bIsTransacting; }
+
 	UFUNCTION(BlueprintCallable, Category = RigVMController)
 	TArray<FString> GeneratePythonCommands();
 
@@ -551,9 +563,6 @@ public:
 		bool bPrintPythonCommand = false);
 
 
-	// Returns a unique name
-	static FName GetUniqueName(const FName& InName, TFunction<bool(const FName&)> IsNameAvailableFunction, bool bAllowPeriod, bool bAllowSpace);
-
 	// Turns a series of nodes into a Collapse node
 	UFUNCTION(BlueprintCallable, Category = RigVMController)
 	URigVMCollapseNode* CollapseNodes(const TArray<FName>& InNodeNames, const FString& InCollapseNodeName = TEXT(""), bool bSetupUndoRedo = true, bool bPrintPythonCommand = false, bool bIsAggregate = false);
@@ -828,7 +837,7 @@ public:
 	// Adds a function reference / invocation to the graph
 	UFUNCTION(BlueprintCallable, Category = RigVMController)
 	URigVMFunctionReferenceNode* AddFunctionReferenceNodeFromDescription(const FRigVMGraphFunctionHeader& InFunctionDefinition, const FVector2D& InNodePosition = FVector2D::ZeroVector, const
-	                                                                     FString& InNodeName = TEXT(""), bool bSetupUndoRedo = true, bool bPrintPythonCommand = false, bool bAllowPrivateFunctions = false);
+	                                                                     FString& InNodeName = TEXT(""), bool bSetupUndoRedo = true, bool bPrintPythonCommand = false);
 	UFUNCTION(BlueprintCallable, Category = RigVMController)
 	URigVMFunctionReferenceNode* AddExternalFunctionReferenceNode(const FString& InHostPath, const FName& InFunctionName, const FVector2D& InNodePosition = FVector2D::ZeroVector, const FString& InNodeName = TEXT(""), bool bSetupUndoRedo = true, bool bPrintPythonCommand = false);
 	
@@ -898,9 +907,6 @@ public:
 	// Determine affected assets for a potential bulk edit on a library node
 	TArray<FAssetData> GetAffectedAssets(ERigVMControllerBulkEditType InEditType, bool bForceLoad = false);
 
-	// A delegate that can be set to change the struct unfolding behaviour
-	FRigVMController_ShouldStructUnfoldDelegate UnfoldStructDelegate;
-
 	// A delegate to retrieve the list of external variables
 	FRigVMController_GetExternalVariablesDelegate GetExternalVariablesDelegate;
 
@@ -1083,24 +1089,18 @@ private:
 	FRigVMGraphModifiedEvent ModifiedEventStatic;
 	void HandleModifiedEvent(ERigVMGraphNotifType InNotifType, URigVMGraph* InGraph, UObject* InSubject);
 
-	FString GetValidNodeName(const FString& InPrefix);
 	bool IsValidGraph() const;
+	bool IsValidSchema() const;
 	bool IsGraphEditable() const;
 	bool IsValidNodeForGraph(const URigVMNode* InNode);
 	bool IsValidPinForGraph(const URigVMPin* InPin);
 	bool IsValidLinkForGraph(const URigVMLink* InLink);
-	bool CanAddNode(URigVMNode* InNode, bool bReportErrors, bool bIgnoreFunctionEntryReturnNodes = false);
-	TObjectPtr<URigVMNode> FindEventNode(const UScriptStruct* InScriptStruct) const;
-	bool CanAddEventNode(UScriptStruct* InScriptStruct, const bool bReportErrors) const;
-	bool CanAddFunctionRefForDefinition(const FRigVMGraphFunctionHeader& InFunctionDefinition, bool bReportErrors, bool bAllowPrivateFunctions=false);
 	void AddPinsForStruct(UStruct* InStruct, URigVMNode* InNode, URigVMPin* InParentPin, ERigVMPinDirection InPinDirection, const FString& InDefaultValue, bool bAutoExpandArrays, const FPinInfoArray* PreviousPins = nullptr);
 	void AddPinsForArray(FArrayProperty* InArrayProperty, URigVMNode* InNode, URigVMPin* InParentPin, ERigVMPinDirection InPinDirection, const TArray<FString>& InDefaultValues, bool bAutoExpandArrays);
 	void AddPinsForTemplate(const FRigVMTemplate* InTemplate, const FRigVMTemplateTypeMap& InPinTypeMap, URigVMNode* InNode);
 	void ConfigurePinFromProperty(FProperty* InProperty, URigVMPin* InOutPin, ERigVMPinDirection InPinDirection = ERigVMPinDirection::Invalid);
 	void ConfigurePinFromPin(URigVMPin* InOutPin, URigVMPin* InPin, bool bCopyDisplayName = false);
 	void ConfigurePinFromArgument(URigVMPin* InOutPin, const FRigVMGraphFunctionArgument& InArgument, bool bCopyDisplayName = false);
-	virtual bool ShouldStructBeUnfolded(const UStruct* InStruct);
-	virtual bool ShouldPinBeUnfolded(URigVMPin* InPin);
 	bool SetPinDefaultValue(URigVMPin* InPin, const FString& InDefaultValue, bool bResizeArrays, bool bSetupUndoRedo, bool bMergeUndoAction);
 	bool ResetPinDefaultValue(URigVMPin* InPin, bool bSetupUndoRedo);
 	static FString GetPinInitialDefaultValue(const URigVMPin* InPin);
@@ -1186,6 +1186,7 @@ private:
 	void DestroyObject(UObject* InObjectToDestroy) const ;
 	static URigVMPin* MakeExecutePin(URigVMNode* InNode, const FName& InName);
 	static void MakeExecutePin(URigVMPin* InOutPin);
+	bool AddGraphNode(URigVMNode* InNode, bool bNotify);
 	static void AddNodePin(URigVMNode* InNode, URigVMPin* InPin);
 	static void AddSubPin(URigVMPin* InParentPin, URigVMPin* InPin);
 	static bool EnsurePinValidity(URigVMPin* InPin, bool bRecursive);
@@ -1195,22 +1196,12 @@ private:
 	// they can get out of sync when the variable references a user defined struct
 	void EnsureLocalVariableValidity();
 	
-	FRigVMExternalVariable GetVariableByName(const FName& InExternalVariableName, const bool bIncludeInputArguments = false);
-	TArray<FRigVMExternalVariable> GetAllVariables(const bool bIncludeInputArguments = false);
+	FRigVMExternalVariable GetVariableByName(const FName& InExternalVariableName, const bool bIncludeInputArguments = false) const;
+	TArray<FRigVMExternalVariable> GetAllVariables(const bool bIncludeInputArguments = false) const;
 
 	void RefreshFunctionReferences(URigVMLibraryNode* InFunctionDefinition, bool bSetupUndoRedo);
 
-	FString GetGraphOuterName() const;
-
 public:
-	static int32 GetMaxNameLength() { return 100; }
-	static FString GetSanitizedName(const FString& InName, bool bAllowPeriod, bool bAllowSpace);
-	static FString GetSanitizedGraphName(const FString& InName);
-	static FString GetSanitizedNodeName(const FString& InName);
-	static FString GetSanitizedVariableName(const FString& InName);
-	static FString GetSanitizedPinName(const FString& InName);
-	static FString GetSanitizedPinPath(const FString& InName);
-	static void SanitizeName(FString& InOutName, bool bAllowPeriod, bool bAllowSpace);
 
 	struct FLinkedPath
 	{
@@ -1364,6 +1355,9 @@ private:
 	UPROPERTY(transient)
 	TArray<TObjectPtr<URigVMGraph>> Graphs;
 
+	UPROPERTY(transient, DuplicateTransient)
+	TObjectPtr<URigVMSchema> SchemaPtr;
+
 	mutable TWeakObjectPtr<URigVMActionStack> WeakActionStack;
 	mutable FDelegateHandle ActionStackHandle;
 
@@ -1372,6 +1366,7 @@ private:
 	bool bIgnoreRerouteCompactnessChanges;
 	ERigVMPinDirection UserLinkDirection;
 	bool bEnableTypeCasting;
+	bool bAllowPrivateFunctions;
 
 	// temporary maps used for pin redirection
 	// only valid between Detach & ReattachLinksToPinObjects
@@ -1444,6 +1439,8 @@ private:
 	friend class FRigVMControllerNotifGuard;
 	friend struct FRigVMClient;
 	friend struct FRigVMActionWrapper;
+	friend class URigVMSchema;
+	friend class UControlRigFunctionRefNodeSpawner;
 };
 
 class FRigVMControllerNotifGuard
diff --git a/Engine/Plugins/Runtime/RigVM/Source/RigVMDeveloper/Public/RigVMModel/RigVMGraph.h b/Engine/Plugins/Runtime/RigVM/Source/RigVMDeveloper/Public/RigVMModel/RigVMGraph.h
--- a/Engine/Plugins/Runtime/RigVM/Source/RigVMDeveloper/Public/RigVMModel/RigVMGraph.h
+++ b/Engine/Plugins/Runtime/RigVM/Source/RigVMDeveloper/Public/RigVMModel/RigVMGraph.h
@@ -208,5 +208,6 @@ private:
 	friend class FRigVMControllerCompileBracketScope;
 	friend class URigVMCompiler;
 	friend struct FRigVMControllerObjectFactory;
+	friend class URigVMSchema;
 };
 
diff --git a/Engine/Plugins/Runtime/RigVM/Source/RigVMDeveloper/Public/RigVMModel/RigVMNode.h b/Engine/Plugins/Runtime/RigVM/Source/RigVMDeveloper/Public/RigVMModel/RigVMNode.h
--- a/Engine/Plugins/Runtime/RigVM/Source/RigVMDeveloper/Public/RigVMModel/RigVMNode.h
+++ b/Engine/Plugins/Runtime/RigVM/Source/RigVMDeveloper/Public/RigVMModel/RigVMNode.h
@@ -273,7 +273,7 @@ public:
 	virtual FName GetNextAggregateName(const FName& InLastAggregatePinName) const { return NAME_None; }
 
 	UFUNCTION(BlueprintCallable, Category = RigVMNode)
-	virtual URigVMLibraryNode* FindFunctionForNode();
+	virtual URigVMLibraryNode* FindFunctionForNode() const;
 
 	virtual FRigVMStructUpgradeInfo GetUpgradeInfo() const { return FRigVMStructUpgradeInfo(); }
 
@@ -341,5 +341,6 @@ private:
 	friend class URigVMPin;
 	friend class URigVMCompiler;
 	friend class FRigVMLexer;
+	friend class URigVMSchema;
 };
 
diff --git a/Engine/Plugins/Runtime/RigVM/Source/RigVMDeveloper/Public/RigVMModel/RigVMPin.h b/Engine/Plugins/Runtime/RigVM/Source/RigVMDeveloper/Public/RigVMModel/RigVMPin.h
--- a/Engine/Plugins/Runtime/RigVM/Source/RigVMDeveloper/Public/RigVMModel/RigVMPin.h
+++ b/Engine/Plugins/Runtime/RigVM/Source/RigVMDeveloper/Public/RigVMModel/RigVMPin.h
@@ -406,7 +406,7 @@ public:
 
 	// Returns true is the two provided source and target Pins
 	// can be linked to one another.
-	static bool CanLink(URigVMPin* InSourcePin, URigVMPin* InTargetPin, FString* OutFailureReason, const FRigVMByteCode* InByteCode, ERigVMPinDirection InUserLinkDirection = ERigVMPinDirection::IO, bool bInAllowWildcard = false, bool bEnableTypeCasting = true);
+	static bool CanLink(const URigVMPin* InSourcePin, const URigVMPin* InTargetPin, FString* OutFailureReason, const FRigVMByteCode* InByteCode, ERigVMPinDirection InUserLinkDirection = ERigVMPinDirection::IO, bool bInAllowWildcard = false, bool bEnableTypeCasting = true);
 
 	// Returns true if this pin has injected nodes
 	bool HasInjectedNodes() const { return InjectionInfos.Num() > 0; }
diff --git a/Engine/Plugins/Runtime/RigVM/Source/RigVMDeveloper/Public/RigVMModel/RigVMSchema.h b/Engine/Plugins/Runtime/RigVM/Source/RigVMDeveloper/Public/RigVMModel/RigVMSchema.h
new file mode 100644
--- /dev/null
+++ b/Engine/Plugins/Runtime/RigVM/Source/RigVMDeveloper/Public/RigVMModel/RigVMSchema.h
@@ -0,0 +1,148 @@
+// Copyright Epic Games, Inc. All Rights Reserved.
+
+#pragma once
+
+#include "RigVMCore/RigVMRegistry.h"
+#include "RigVMCore/RigVMFunction.h"
+#include "RigVMCore/RigVMTemplate.h"
+#include "RigVMCore/RigVMDispatchFactory.h"
+#include "RigVMCore/RigVMGraphFunctionDefinition.h"
+#include "RigVMGraph.h"
+#include "RigVMSchema.generated.h"
+
+#define RIGVMSCHEMA_DEFAULT_FUNCTION_BODY \
+const URigVMGraph* Graph = InController->GetGraph(); \
+check(Graph); \
+if (!InController->IsTransacting() && !IsGraphEditable(Graph)) \
+{ \
+	return false; \
+}
+
+class URigVMController;
+
+/**
+ * The Schema is used to determine which actions are allowed
+ * on a graph. This includes any topological change.
+ */
+UCLASS()
+class RIGVMDEVELOPER_API URigVMSchema : public UObject
+{
+	GENERATED_UCLASS_BODY()
+
+public:
+
+	// Default constructor
+	URigVMSchema();
+
+	// Returns the execute context struct this schema is using
+	UScriptStruct* GetExecuteContextStruct() const { return ExecuteContextStruct; }
+
+	// Returns true if a graph supports a given type
+	virtual bool SupportsType(URigVMController* InController, TRigVMTypeIndex InTypeIndex) const;
+
+	// Returns true if a graph supports a given unit function
+	virtual bool SupportsUnitFunction(URigVMController* InController, const FRigVMFunction* InUnitFunction) const;
+
+	// Returns true if a graph supports a given dispatch factory
+	virtual bool SupportsDispatchFactory(URigVMController* InController, const FRigVMDispatchFactory* InDispatchFactory) const;
+
+	// Returns true if a graph supports a given template
+	virtual bool SupportsTemplate(URigVMController* InController, const FRigVMTemplate* InTemplate) const;
+
+	// Returns true if a graph supports a given graph based function
+	virtual bool SupportsGraphFunction(URigVMController* InController, const FRigVMGraphFunctionHeader* InGraphFunction) const;
+
+	// Returns true if a graph supports a given external variable
+	virtual bool SupportsExternalVariable(URigVMController* InController, const FRigVMExternalVariable* InExternalVariable) const;
+
+	// Returns true if the pin for a given struct should be unfolded into subpins for a graph
+	virtual bool ShouldUnfoldStruct(URigVMController* InController, const UStruct* InStruct) const;
+
+	// Returns true if a node name is valid for a given graph
+	virtual bool IsValidNodeName(const URigVMGraph* InGraph, const FName& InNodeName) const;
+
+	// Returns true if a given node can be added the graph
+	virtual bool CanAddNode(URigVMController* InController, const URigVMNode* InNode) const;
+
+	// Returns true if a given node can be removed from the graph
+	virtual bool CanRemoveNode(URigVMController* InController, const URigVMNode* InNode) const;
+
+	// Returns true if a given node can be added the graph
+	virtual bool CanRenameNode(URigVMController* InController, const URigVMNode* InNode, const FName& InNewNodeName) const;
+
+	// Returns true if a node can moved to a new position
+	virtual bool CanMoveNode(URigVMController* InController, const URigVMNode* InNode, const FVector2D& InNewPosition = FVector2D::ZeroVector) const;
+
+	// Returns true if a node can resized to a new size
+	virtual bool CanResizeNode(URigVMController* InController, const URigVMNode* InNode, const FVector2D& InNewSize = {100, 100}) const;
+
+	// Returns true if a node can recolored to a new color
+	virtual bool CanRecolorNode(URigVMController* InController, const URigVMNode* InNode, const FLinearColor& InNewColor = FLinearColor::White) const;
+
+	// Returns true if a link can be added between two pins
+	virtual bool CanAddLink(URigVMController* InController, const URigVMPin* InSourcePin, const URigVMPin* InTargetPin, const FRigVMByteCode* InByteCode, ERigVMPinDirection InUserLinkDirection = ERigVMPinDirection::IO, bool bInAllowWildcard = false, bool bEnableTypeCasting = true, FString* OutFailureReason = nullptr) const;
+
+	// Returns true if a link exists between two pins and can be broken / removed
+	virtual bool CanBreakLink(URigVMController* InController, const URigVMPin* InSourcePin, const URigVMPin* InTargetPin) const;
+
+	// Returns true if a set of nodes can be collapsed
+	virtual bool CanCollapseNodes(URigVMController* InController, const TArrayView<URigVMNode* const>& InNodesToCollapse) const;
+
+	// Returns true if a node can be expanded within a graph
+	virtual bool CanExpandNode(URigVMController* InController, const URigVMNode* InNodeToExpand) const;
+
+	// Returns true if a pin should be unfolded (represented by its subpins)
+	virtual bool CanUnfoldPin(URigVMController* InController, const URigVMPin* InPinToUnfold) const;
+
+	// Returns true if a variable can be bound to a pin
+	virtual bool CanBindVariable(URigVMController* InController, const URigVMPin* InPinToBind, const FRigVMExternalVariable* InVariableToBind, const FString& InNewBoundVariablePath) const;
+
+	// Returns true if a variable can be unbound from a pin
+	virtual bool CanUnbindVariable(URigVMController* InController, const URigVMPin* InBoundPin) const;
+
+	// Returns true if functions definitions (not refs) can be added to the graph
+	virtual bool CanAddFunction(URigVMController* InController, const URigVMNode* InFunctionNode) const;
+
+	// Returns true if a function can be removed
+	virtual bool CanRemoveFunction(URigVMController* InController, const URigVMNode* InFunctionNode) const;
+	
+	static int32 GetMaxNameLength() { return 100; }
+	virtual FString GetSanitizedName(const FString& InName, bool bAllowPeriod, bool bAllowSpace) const;
+	virtual FString GetSanitizedGraphName(const FString& InName) const;
+	virtual FString GetSanitizedNodeName(const FString& InName) const;
+	virtual FString GetSanitizedVariableName(const FString& InName) const;
+	virtual FString GetSanitizedPinName(const FString& InName) const;
+	virtual FString GetSanitizedPinPath(const FString& InName) const;
+	virtual FString GetGraphOuterName(const URigVMGraph* InGraph) const;
+	virtual	FString GetValidNodeName(const URigVMGraph* InGraph, const FString& InPrefix) const;
+	static void SanitizeName(FString& InOutName, bool bAllowPeriod, bool bAllowSpace);
+
+	// Returns a unique name based on a IsNameAvailable predicate
+	static FName GetUniqueName(const FName& InName, TFunction<bool(const FName&)> IsNameAvailablePredicate, bool bAllowPeriod, bool bAllowSpace);
+
+protected:
+
+	// Sets the execute context struct this schema is using
+	void SetExecuteContextStruct(UScriptStruct* InExecuteContextStruct);
+
+	bool IsGraphEditable(const URigVMGraph* InGraph) const;
+	TObjectPtr<URigVMNode> FindEventNode(URigVMController* InController, const UScriptStruct* InScriptStruct) const;
+
+	UPROPERTY(transient)
+	TObjectPtr<UScriptStruct> ExecuteContextStruct;
+	TArray<UStruct*> ValidExecuteContextStructs;
+	FRigVMRegistry* Registry;
+
+	/*
+	bool bIsTransacting; // Performing undo/redo transaction
+	bool bAllowPrivateFunctions;
+	bool bIgnoreFunctionEntryReturnNodes;
+	mutable FString LastError;
+	mutable FString LastWarning;
+	mutable FString LastErrorOrWarning;
+	*/
+
+	friend struct FRigVMClient;
+	friend class URigVMController;
+	friend struct FRigVMBaseAction;
+};
\ No newline at end of file
