RigVM: Ensure to send notifications for entry / return nodes on new functions

diff --git a/Engine/Plugins/Runtime/RigVM/Source/RigVMDeveloper/Private/RigVMModel/RigVMController.cpp b/Engine/Plugins/Runtime/RigVM/Source/RigVMDeveloper/Private/RigVMModel/RigVMController.cpp
--- a/Engine/Plugins/Runtime/RigVM/Source/RigVMDeveloper/Private/RigVMModel/RigVMController.cpp
+++ b/Engine/Plugins/Runtime/RigVM/Source/RigVMDeveloper/Private/RigVMModel/RigVMController.cpp
@@ -11461,6 +11461,7 @@ URigVMLibraryNode* URigVMController::AddFunctionToLibrary(const FName& InFunctio
 				CollapseController->Notify(ERigVMGraphNotifType::NodeAdded, EntryNode);
 			}
 		}
+		CollapseController->Notify(ERigVMGraphNotifType::NodeAdded, EntryNode);
 
 		EntryNode->Position = FVector2D(-250.f, 0.f);
 		CollapseController->Notify(ERigVMGraphNotifType::NodePositionChanged, EntryNode);
@@ -11475,6 +11476,7 @@ URigVMLibraryNode* URigVMController::AddFunctionToLibrary(const FName& InFunctio
 				CollapseController->Notify(ERigVMGraphNotifType::NodeAdded, ReturnNode);
 			}
 		}
+		CollapseController->Notify(ERigVMGraphNotifType::NodeAdded, ReturnNode);
 
 		ReturnNode->Position = FVector2D(250.f, 0.f);
 		CollapseController->Notify(ERigVMGraphNotifType::NodePositionChanged, ReturnNode);
