FRigVMRegistry::GetOrAddTemplateFromArguments performance improvement
    - use of TSet instead of TArray to speed-up the O(N^2) call due to the Contains function

diff --git a/Engine/Plugins/Runtime/RigVM/Source/RigVM/Private/RigVMCore/RigVMRegistry.cpp b/Engine/Plugins/Runtime/RigVM/Source/RigVM/Private/RigVMCore/RigVMRegistry.cpp
--- a/Engine/Plugins/Runtime/RigVM/Source/RigVM/Private/RigVMCore/RigVMRegistry.cpp
+++ b/Engine/Plugins/Runtime/RigVM/Source/RigVM/Private/RigVMCore/RigVMRegistry.cpp
@@ -1563,26 +1563,48 @@ const FRigVMTemplate* FRigVMRegistry::AddTemplateFromArguments_NoLock(const FNam
 	// Remove duplicate permutations
 	{
 		TArray<int32> ToRemove;
-		TArray<TArray<TRigVMTypeIndex>> PermutationTypes;
-		PermutationTypes.Reserve(NumPermutations);
-		for(int32 Index=0;Index<NumPermutations;Index++)
+		
+		const int32 NumArguments = Template.Arguments.Num();
+		if (NumArguments == 1)
 		{
-			TArray<TRigVMTypeIndex> ArgTypes;
-			ArgTypes.Reserve(Template.Arguments.Num());
-			for(FRigVMTemplateArgument& Argument : Template.Arguments)
-			{
-				ArgTypes.Add(Argument.TypeIndices[Index]);
-			}
-
-			if (PermutationTypes.Contains(ArgTypes))
+			TSet< TRigVMTypeIndex > PermutationTypes; PermutationTypes.Reserve(NumPermutations);
+			for(int32 Index = 0; Index < NumPermutations; Index++)
+ 
 			{
-				ToRemove.Add(Index);
+				const TRigVMTypeIndex ArgType = Template.Arguments[0].TypeIndices[Index];
+				if (PermutationTypes.Contains(ArgType))
+				{
+					ToRemove.Add(Index);
+				}
+				else
+				{
+					PermutationTypes.Add(ArgType);
+				}
 			}
-			else
+		}
+		else
+		{
+			TSet< TArray<TRigVMTypeIndex> > PermutationTypes;
+			PermutationTypes.Reserve(NumPermutations);
+			TArray<TRigVMTypeIndex> ArgTypes; ArgTypes.SetNum(NumArguments);			
+			for(int32 Index = 0; Index < NumPermutations; Index++)
 			{
-				PermutationTypes.Add(ArgTypes);
+				for(int32 ArgIndex = 0; ArgIndex < NumArguments; ArgIndex++)
+				{
+					ArgTypes[ArgIndex] = Template.Arguments[ArgIndex].TypeIndices[Index];
+				}
+				
+				if (PermutationTypes.Contains(ArgTypes))
+				{
+					ToRemove.Add(Index);
+				}
+				else
+				{
+					PermutationTypes.Add(ArgTypes);
+				}
 			}
 		}
+		
 		for (int32 i=ToRemove.Num()-1; i>=0; --i)
 		{
 			for(FRigVMTemplateArgument& Argument : Template.Arguments)
