RigVM: Python side unit test breakage

diff --git a/Engine/Plugins/Runtime/RigVM/Source/RigVMDeveloper/Private/RigVMModel/RigVMController.cpp b/Engine/Plugins/Runtime/RigVM/Source/RigVMDeveloper/Private/RigVMModel/RigVMController.cpp
--- a/Engine/Plugins/Runtime/RigVM/Source/RigVMDeveloper/Private/RigVMModel/RigVMController.cpp
+++ b/Engine/Plugins/Runtime/RigVM/Source/RigVMDeveloper/Private/RigVMModel/RigVMController.cpp
@@ -2780,7 +2780,7 @@ URigVMCommentNode* URigVMController::AddCommentNode(const FString& InCommentText
 	Node->NodeColor = InColor;
 	Node->CommentText = InCommentText;
 
-	if(!AddGraphNode(Node, true))
+	if(!AddGraphNode(Node, false))
 	{
 		return nullptr;
 	}
@@ -2799,6 +2799,8 @@ URigVMCommentNode* URigVMController::AddCommentNode(const FString& InCommentText
 		GetActionStack()->AddAction(FRigVMImportFromTextAction(this, Node));
 	}
 
+	Notify(ERigVMGraphNotifType::NodeAdded, Node);
+
 	if (bSetupUndoRedo)
 	{
 		GetActionStack()->EndAction(Action);
diff --git a/Engine/Plugins/Runtime/RigVM/Source/RigVMDeveloper/Public/RigVMModel/RigVMController.h b/Engine/Plugins/Runtime/RigVM/Source/RigVMDeveloper/Public/RigVMModel/RigVMController.h
--- a/Engine/Plugins/Runtime/RigVM/Source/RigVMDeveloper/Public/RigVMModel/RigVMController.h
+++ b/Engine/Plugins/Runtime/RigVM/Source/RigVMDeveloper/Public/RigVMModel/RigVMController.h
@@ -209,6 +209,7 @@ public:
 	URigVMSchema* GetSchema() const { return SchemaPtr; }
 
 	// Sets the schema on the controller
+	UFUNCTION(BlueprintCallable, Category = RigVMController)
 	void SetSchema(URigVMSchema* InSchema);
 
 	// Pushes a new graph to the stack
