RigVM Editor Module first steps

diff --git a/Engine/Plugins/Animation/ControlRig/Source/ControlRigEditor/ControlRigEditor.Build.cs b/Engine/Plugins/Animation/ControlRig/Source/ControlRigEditor/ControlRigEditor.Build.cs
--- a/Engine/Plugins/Animation/ControlRig/Source/ControlRigEditor/ControlRigEditor.Build.cs
+++ b/Engine/Plugins/Animation/ControlRig/Source/ControlRigEditor/ControlRigEditor.Build.cs
@@ -72,6 +72,7 @@ namespace UnrealBuildTool.Rules
 					"ToolMenus",
                     "RigVM",
                     "RigVMDeveloper",
+                    "RigVMEditor",
 					"AnimationEditor",
 					"MessageLog",
                     "SequencerScripting",
diff --git a/Engine/Plugins/Runtime/RigVM/RigVM.uplugin b/Engine/Plugins/Runtime/RigVM/RigVM.uplugin
--- a/Engine/Plugins/Runtime/RigVM/RigVM.uplugin
+++ b/Engine/Plugins/Runtime/RigVM/RigVM.uplugin
@@ -25,6 +25,11 @@
 			"Name" : "RigVMDeveloper",
 			"Type" : "UncookedOnly",
 			"LoadingPhase" : "PreDefault"
+		},
+		{
+			"Name" : "RigVMEditor",
+			"Type" : "Editor",
+			"LoadingPhase" : "PostEngineInit"
 		}
 	]
 }
\ No newline at end of file
diff --git a/Engine/Plugins/Runtime/RigVM/Source/RigVMEditor/Private/RigVMEditorModule.cpp b/Engine/Plugins/Runtime/RigVM/Source/RigVMEditor/Private/RigVMEditorModule.cpp
new file mode 100644
--- /dev/null
+++ b/Engine/Plugins/Runtime/RigVM/Source/RigVMEditor/Private/RigVMEditorModule.cpp
@@ -0,0 +1,25 @@
+// Copyright Epic Games, Inc. All Rights Reserved.
+
+/*=============================================================================
+	RigVMEditor.cpp: Module implementation.
+=============================================================================*/
+
+#include "RigVMEditorModule.h"
+#include "Modules/ModuleManager.h"
+
+DEFINE_LOG_CATEGORY(LogRigVMEditor);
+
+IMPLEMENT_MODULE(FRigVMEditorModule, RigVMEditor)
+
+FRigVMEditorModule& FRigVMEditorModule::Get()
+{
+	return FModuleManager::LoadModuleChecked< FRigVMEditorModule >(TEXT("RigVMEditor"));
+}
+
+void FRigVMEditorModule::StartupModule()
+{
+}
+
+void FRigVMEditorModule::ShutdownModule()
+{
+}
diff --git a/Engine/Plugins/Runtime/RigVM/Source/RigVMEditor/Public/RigVMEditorModule.h b/Engine/Plugins/Runtime/RigVM/Source/RigVMEditor/Public/RigVMEditorModule.h
new file mode 100644
--- /dev/null
+++ b/Engine/Plugins/Runtime/RigVM/Source/RigVMEditor/Public/RigVMEditorModule.h
@@ -0,0 +1,24 @@
+// Copyright Epic Games, Inc. All Rights Reserved.
+
+/*=============================================================================
+	RigVMEditor.h: Module definitions.
+=============================================================================*/
+
+#pragma once
+
+#include "CoreMinimal.h"
+#include "Modules/ModuleInterface.h"
+#include "Logging/LogMacros.h"
+
+DECLARE_LOG_CATEGORY_EXTERN(LogRigVMEditor, Log, All);
+
+class FRigVMEditorModule : public IModuleInterface
+{
+public:
+
+	static FRigVMEditorModule& Get();
+
+	/** IModuleInterface implementation */
+	virtual void StartupModule() override;
+	virtual void ShutdownModule() override;
+};
\ No newline at end of file
diff --git a/Engine/Plugins/Runtime/RigVM/Source/RigVMEditor/RigVMEditor.Build.cs b/Engine/Plugins/Runtime/RigVM/Source/RigVMEditor/RigVMEditor.Build.cs
new file mode 100644
--- /dev/null
+++ b/Engine/Plugins/Runtime/RigVM/Source/RigVMEditor/RigVMEditor.Build.cs
@@ -0,0 +1,19 @@
+// Copyright Epic Games, Inc. All Rights Reserved.
+
+using UnrealBuildTool;
+
+public class RigVMEditor : ModuleRules
+{
+    public RigVMEditor(ReadOnlyTargetRules Target) : base(Target)
+    {
+        PublicDependencyModuleNames.AddRange(
+            new string[] {
+                "Core",
+                "CoreUObject",
+                "Engine",
+                "RigVM",
+                "RigVMDeveloper"
+            }
+        );
+    }
+}
