RigVM: Avoid ensure by not offering a function reference node if any of the argument types is invalid

diff --git a/Engine/Plugins/Runtime/RigVM/Source/RigVMDeveloper/Private/RigVMModel/RigVMController.cpp b/Engine/Plugins/Runtime/RigVM/Source/RigVMDeveloper/Private/RigVMModel/RigVMController.cpp
--- a/Engine/Plugins/Runtime/RigVM/Source/RigVMDeveloper/Private/RigVMModel/RigVMController.cpp
+++ b/Engine/Plugins/Runtime/RigVM/Source/RigVMDeveloper/Private/RigVMModel/RigVMController.cpp
@@ -539,6 +539,7 @@ URigVMController::FPinInfoArray::FPinInfoArray(const URigVMNode* InNode, URigVMC
 	for(const URigVMPin* Pin : InNode->GetPins())
 	{
 		const FString DefaultValue = Pin->GetDefaultValue();
+		ensureMsgf(Pin->GetTypeIndex() != INDEX_NONE, TEXT("Invalid pin type %s in %s"), *Pin->GetCPPType(), *InNode->GetPackage()->GetPathName());
 		(void)AddPin(InController, INDEX_NONE, Pin->GetFName(), Pin->GetDirection(), Pin->GetTypeIndex(), DefaultValue, nullptr, InPreviousPinInfos, bAddSubPins);
 	}
 }
@@ -551,6 +552,7 @@ URigVMController::FPinInfoArray::FPinInfoArray(const FRigVMGraphFunctionHeader&
 	for(const FRigVMGraphFunctionArgument& FunctionArgument : FunctionHeader.Arguments)
 	{
 		const TRigVMTypeIndex TypeIndex = Registry.GetTypeIndexFromCPPType(FunctionArgument.CPPType.ToString());
+		ensureMsgf(TypeIndex != INDEX_NONE, TEXT("Invalid pin type %s in %s"), *FunctionArgument.CPPType.ToString(), *InController->GetPackage()->GetPathName());
 		(void)AddPin(InController, INDEX_NONE, FunctionArgument.Name, FunctionArgument.Direction, TypeIndex, FunctionArgument.DefaultValue, nullptr, InPreviousPinInfos, true);
 	}
 }
@@ -591,6 +593,7 @@ int32 URigVMController::FPinInfoArray::AddPin(FProperty* InProperty, URigVMContr
 int32 URigVMController::FPinInfoArray::AddPin(URigVMController* InController, int32 InParentIndex, const FName& InName, ERigVMPinDirection InDirection,
 	TRigVMTypeIndex InTypeIndex, const FString& InDefaultValue, const uint8* InDefaultValueMemory, const FPinInfoArray* InPreviousPinInfos, bool bAddSubPins)
 {
+	ensureMsgf(InTypeIndex != INDEX_NONE, TEXT("Invalid pin type for pin %s in %s"), *InName.ToString(), *InController->GetPackage()->GetPathName());
 	const FRigVMRegistry& Registry = FRigVMRegistry::Get();
 		
 	FPinInfo Info;
diff --git a/Engine/Plugins/Runtime/RigVM/Source/RigVMDeveloper/Private/RigVMModel/RigVMSchema.cpp b/Engine/Plugins/Runtime/RigVM/Source/RigVMDeveloper/Private/RigVMModel/RigVMSchema.cpp
--- a/Engine/Plugins/Runtime/RigVM/Source/RigVMDeveloper/Private/RigVMModel/RigVMSchema.cpp
+++ b/Engine/Plugins/Runtime/RigVM/Source/RigVMDeveloper/Private/RigVMModel/RigVMSchema.cpp
@@ -131,6 +131,17 @@ bool URigVMSchema::SupportsGraphFunction(URigVMController* InController, const F
 		return false;
 	}
 
+	// Make sure all the argument types are supported before creating the node
+	for (const FRigVMGraphFunctionArgument& Argument : InGraphFunction->Arguments)
+	{
+		const TRigVMTypeIndex Type = Registry->GetTypeIndexFromCPPType(Argument.CPPType.ToString());
+		if (Type == INDEX_NONE)
+		{
+			InController->ReportErrorf(TEXT("Cannot add function reference to %s because argument %s has invalid type %s."), *InGraphFunction->Name.ToString(), *Argument.Name.ToString(), *Argument.CPPType.ToString());
+			return false;
+		}
+	}
+
 	if (!InController->bAllowPrivateFunctions)
 	{
 		if(IRigVMClientHost* ClientHost = Graph->GetImplementingOuter<IRigVMClientHost>())
