RigVM: Fix RigVM.Functions.PromoteToFunction unit test failing

diff --git a/Engine/Plugins/Runtime/RigVM/Source/RigVMDeveloper/Private/RigVMModel/RigVMPin.cpp b/Engine/Plugins/Runtime/RigVM/Source/RigVMDeveloper/Private/RigVMModel/RigVMPin.cpp
--- a/Engine/Plugins/Runtime/RigVM/Source/RigVMDeveloper/Private/RigVMModel/RigVMPin.cpp
+++ b/Engine/Plugins/Runtime/RigVM/Source/RigVMDeveloper/Private/RigVMModel/RigVMPin.cpp
@@ -115,7 +115,10 @@ URigVMPin::URigVMPin()
 	, LastKnownTypeIndex(INDEX_NONE)
 {
 #if UE_BUILD_DEBUG
-	CachedPinPath = GetPinPath();
+	if (GetNode())
+	{
+		CachedPinPath = GetPinPath();
+	}
 #endif
 }
 
