RigVM: Predicates Core

diff --git a/Engine/Plugins/Animation/ControlRig/Source/ControlRig/Private/Units/Animation/RigUnit_AnimAttribute.cpp b/Engine/Plugins/Animation/ControlRig/Source/ControlRig/Private/Units/Animation/RigUnit_AnimAttribute.cpp
--- a/Engine/Plugins/Animation/ControlRig/Source/ControlRig/Private/Units/Animation/RigUnit_AnimAttribute.cpp
+++ b/Engine/Plugins/Animation/ControlRig/Source/ControlRig/Private/Units/Animation/RigUnit_AnimAttribute.cpp
@@ -153,7 +153,7 @@ static uint8* GetAnimAttributeValue(
 }
 
 void FRigDispatch_GetAnimAttribute::GetAnimAttributeDispatch(FRigVMExtendedExecuteContext& InContext,
-	FRigVMMemoryHandleArray Handles)
+	FRigVMMemoryHandleArray Handles, FRigVMPredicateBranchArray Predicates)
 {
 	const FRigDispatch_GetAnimAttribute* Factory = static_cast<const FRigDispatch_GetAnimAttribute*>(InContext.Factory);
 		
@@ -216,7 +216,7 @@ void FRigDispatch_GetAnimAttribute::GetAnimAttributeDispatch(FRigVMExtendedExecu
 }
 
 void FRigDispatch_SetAnimAttribute::SetAnimAttributeDispatch(FRigVMExtendedExecuteContext& InContext,
-	FRigVMMemoryHandleArray Handles)
+	FRigVMMemoryHandleArray Handles, FRigVMPredicateBranchArray Predicates)
 {
 	const FRigDispatch_SetAnimAttribute* Factory = static_cast<const FRigDispatch_SetAnimAttribute*>(InContext.Factory);
 		
diff --git a/Engine/Plugins/Animation/ControlRig/Source/ControlRig/Private/Units/Animation/RigUnit_AnimAttribute.h b/Engine/Plugins/Animation/ControlRig/Source/ControlRig/Private/Units/Animation/RigUnit_AnimAttribute.h
--- a/Engine/Plugins/Animation/ControlRig/Source/ControlRig/Private/Units/Animation/RigUnit_AnimAttribute.h
+++ b/Engine/Plugins/Animation/ControlRig/Source/ControlRig/Private/Units/Animation/RigUnit_AnimAttribute.h
@@ -219,7 +219,7 @@ protected:
 
 	// dispatch function for built-in types
 	template<typename ValueType>	
-	static void GetAnimAttributeDispatch(FRigVMExtendedExecuteContext& InContext, FRigVMMemoryHandleArray Handles)
+	static void GetAnimAttributeDispatch(FRigVMExtendedExecuteContext& InContext, FRigVMMemoryHandleArray Handles, FRigVMPredicateBranchArray Predicates)
 	{
 		const FRigDispatch_GetAnimAttribute* Factory = static_cast<const FRigDispatch_GetAnimAttribute*>(InContext.Factory);
 		
@@ -249,7 +249,7 @@ protected:
 	}
 
 	// dispatch function for user/dev defined types
-	static void GetAnimAttributeDispatch(FRigVMExtendedExecuteContext& InContext, FRigVMMemoryHandleArray Handles);
+	static void GetAnimAttributeDispatch(FRigVMExtendedExecuteContext& InContext, FRigVMMemoryHandleArray Handles, FRigVMPredicateBranchArray Predicates);
 	
 };
 
@@ -298,7 +298,7 @@ protected:
 
 	// dispatch function for built-in types
 	template<typename ValueType>	
-	static void SetAnimAttributeDispatch(FRigVMExtendedExecuteContext& InContext, FRigVMMemoryHandleArray Handles)
+	static void SetAnimAttributeDispatch(FRigVMExtendedExecuteContext& InContext, FRigVMMemoryHandleArray Handles, FRigVMPredicateBranchArray Predicates)
 	{
 		const FRigDispatch_SetAnimAttribute* Factory = static_cast<const FRigDispatch_SetAnimAttribute*>(InContext.Factory);
 		
@@ -331,6 +331,6 @@ protected:
 	}
 
 	// dispatch function for user/dev defined types
-	static void SetAnimAttributeDispatch(FRigVMExtendedExecuteContext& InContext, FRigVMMemoryHandleArray Handles);
+	static void SetAnimAttributeDispatch(FRigVMExtendedExecuteContext& InContext, FRigVMMemoryHandleArray Handles, FRigVMPredicateBranchArray Predicates);
 };
 
diff --git a/Engine/Plugins/Animation/ControlRig/Source/ControlRig/Private/Units/Core/RigUnit_UserData.cpp b/Engine/Plugins/Animation/ControlRig/Source/ControlRig/Private/Units/Core/RigUnit_UserData.cpp
--- a/Engine/Plugins/Animation/ControlRig/Source/ControlRig/Private/Units/Core/RigUnit_UserData.cpp
+++ b/Engine/Plugins/Animation/ControlRig/Source/ControlRig/Private/Units/Core/RigUnit_UserData.cpp
@@ -65,7 +65,7 @@ FRigVMFunctionPtr FRigDispatch_GetUserData::GetDispatchFunctionImpl(const FRigVM
 	return &FRigDispatch_GetUserData::Execute;
 }
 
-void FRigDispatch_GetUserData::Execute(FRigVMExtendedExecuteContext& InContext, FRigVMMemoryHandleArray Handles)
+void FRigDispatch_GetUserData::Execute(FRigVMExtendedExecuteContext& InContext, FRigVMMemoryHandleArray Handles, FRigVMPredicateBranchArray RigVMBranches)
 {
 	const FControlRigExecuteContext& ControlRigContext = InContext.GetPublicData<FControlRigExecuteContext>();
 
diff --git a/Engine/Plugins/Animation/ControlRig/Source/ControlRig/Private/Units/Core/RigUnit_UserData.h b/Engine/Plugins/Animation/ControlRig/Source/ControlRig/Private/Units/Core/RigUnit_UserData.h
--- a/Engine/Plugins/Animation/ControlRig/Source/ControlRig/Private/Units/Core/RigUnit_UserData.h
+++ b/Engine/Plugins/Animation/ControlRig/Source/ControlRig/Private/Units/Core/RigUnit_UserData.h
@@ -28,7 +28,7 @@ public:
 protected:
 
 	virtual FRigVMFunctionPtr GetDispatchFunctionImpl(const FRigVMTemplateTypeMap& InTypes) const override;
-	static void Execute(FRigVMExtendedExecuteContext& InContext, FRigVMMemoryHandleArray Handles);
+	static void Execute(FRigVMExtendedExecuteContext& InContext, FRigVMMemoryHandleArray Handles, FRigVMPredicateBranchArray RigVMBranches);
 
 	static inline constexpr TCHAR ArgNameSpaceName[] = TEXT("NameSpace");
 	static inline constexpr TCHAR ArgPathName[] = TEXT("Path");
diff --git a/Engine/Plugins/Animation/ControlRig/Source/ControlRig/Private/Units/Hierarchy/RigUnit_Metadata.h b/Engine/Plugins/Animation/ControlRig/Source/ControlRig/Private/Units/Hierarchy/RigUnit_Metadata.h
--- a/Engine/Plugins/Animation/ControlRig/Source/ControlRig/Private/Units/Hierarchy/RigUnit_Metadata.h
+++ b/Engine/Plugins/Animation/ControlRig/Source/ControlRig/Private/Units/Hierarchy/RigUnit_Metadata.h
@@ -77,7 +77,7 @@ protected:
 #endif
 
 	template<typename ValueType, typename MetadataType, ERigMetadataType EnumValue>
-	static void GetMetadataDispatch(FRigVMExtendedExecuteContext& InContext, FRigVMMemoryHandleArray Handles)
+	static void GetMetadataDispatch(FRigVMExtendedExecuteContext& InContext, FRigVMMemoryHandleArray Handles, FRigVMPredicateBranchArray Predicates)
 	{
 		const FRigDispatch_GetMetadata* Factory = static_cast<const FRigDispatch_GetMetadata*>(InContext.Factory);
 
@@ -140,7 +140,7 @@ protected:
 #endif
 	
 	template<typename ValueType, typename MetadataType, ERigMetadataType EnumValue>
-	static void SetMetadataDispatch(FRigVMExtendedExecuteContext& InContext, FRigVMMemoryHandleArray Handles)
+	static void SetMetadataDispatch(FRigVMExtendedExecuteContext& InContext, FRigVMMemoryHandleArray Handles, FRigVMPredicateBranchArray Predicates)
 	{
 		const FRigDispatch_SetMetadata* Factory = static_cast<const FRigDispatch_SetMetadata*>(InContext.Factory);
 
diff --git a/Engine/Plugins/Experimental/AnimNext/Source/AnimNext/Internal/Param/RigVMDispatch_SetParameter.h b/Engine/Plugins/Experimental/AnimNext/Source/AnimNext/Internal/Param/RigVMDispatch_SetParameter.h
--- a/Engine/Plugins/Experimental/AnimNext/Source/AnimNext/Internal/Param/RigVMDispatch_SetParameter.h
+++ b/Engine/Plugins/Experimental/AnimNext/Source/AnimNext/Internal/Param/RigVMDispatch_SetParameter.h
@@ -33,7 +33,7 @@ private:
 	{
 		return &FRigVMDispatch_SetParameter::Execute;
 	}
-	static void Execute(FRigVMExtendedExecuteContext& InContext, FRigVMMemoryHandleArray Handles);
+	static void Execute(FRigVMExtendedExecuteContext& InContext, FRigVMMemoryHandleArray Handles, FRigVMPredicateBranchArray RigVMBranches);
 
 	static const FName ValueName;
 	static const FName ExecuteContextName;
diff --git a/Engine/Plugins/Experimental/AnimNext/Source/AnimNext/Private/Param/RigVMDispatch_SetParameter.cpp b/Engine/Plugins/Experimental/AnimNext/Source/AnimNext/Private/Param/RigVMDispatch_SetParameter.cpp
--- a/Engine/Plugins/Experimental/AnimNext/Source/AnimNext/Private/Param/RigVMDispatch_SetParameter.cpp
+++ b/Engine/Plugins/Experimental/AnimNext/Source/AnimNext/Private/Param/RigVMDispatch_SetParameter.cpp
@@ -47,7 +47,7 @@ FRigVMTemplateTypeMap FRigVMDispatch_SetParameter::OnNewArgumentType(const FName
 	return Types;
 }
 
-void FRigVMDispatch_SetParameter::Execute(FRigVMExtendedExecuteContext& InContext, FRigVMMemoryHandleArray Handles)
+void FRigVMDispatch_SetParameter::Execute(FRigVMExtendedExecuteContext& InContext, FRigVMMemoryHandleArray Handles, FRigVMPredicateBranchArray RigVMBranches)
 {
 	const FProperty* ValueProperty = Handles[0].GetResolvedProperty(); 
 	check(ValueProperty);
diff --git a/Engine/Plugins/Runtime/RigVM/Source/RigVM/Private/RigVMCore/RigVM.cpp b/Engine/Plugins/Runtime/RigVM/Source/RigVM/Private/RigVMCore/RigVM.cpp
--- a/Engine/Plugins/Runtime/RigVM/Source/RigVM/Private/RigVMCore/RigVM.cpp
+++ b/Engine/Plugins/Runtime/RigVM/Source/RigVM/Private/RigVMCore/RigVM.cpp
@@ -1763,11 +1763,17 @@ ERigVMExecuteResult URigVM::ExecuteInstructions(FRigVMExtendedExecuteContext& Co
 				{
 					Handles = FRigVMMemoryHandleArray(&Context.CachedMemoryHandles[FirstHandleForInstruction[ContextPublicData.InstructionIndex]], OperandCount);
 				}
+
+				FRigVMPredicateBranchArray Predicates;
+				if (Op.PredicateCount > 0)
+				{
+					Predicates = FRigVMPredicateBranchArray(&ByteCode.PredicateBranches[Op.FirstPredicateIndex], Op.PredicateCount);
+				}
 #if WITH_EDITOR
 				ContextPublicData.FunctionName = FunctionNames[Op.FunctionIndex];
 #endif
 				Context.Factory = Factories[Op.FunctionIndex];
-				(*Functions[Op.FunctionIndex]->FunctionPtr)(Context, Handles);
+				(*Functions[Op.FunctionIndex]->FunctionPtr)(Context, Handles, Predicates);
 
 #if WITH_EDITOR
 
@@ -2141,7 +2147,7 @@ bool URigVM::Execute(FRigVMExtendedExecuteContext& Context, const FName& InEntry
 	return Execute(Context, TArray<URigVMMemoryStorage*>(), InEntryName) != ERigVMExecuteResult::Failed;
 }
 
-ERigVMExecuteResult URigVM::ExecuteLazyBranch(FRigVMExtendedExecuteContext& Context, const FRigVMBranchInfo& InBranchToRun)
+ERigVMExecuteResult URigVM::ExecuteBranch(FRigVMExtendedExecuteContext& Context, const FRigVMBranchInfo& InBranchToRun)
 {
 	// likely have to optimize this
 	TGuardValue<FRigVMExtendedExecuteContext> ContextGuard(Context, Context);
diff --git a/Engine/Plugins/Runtime/RigVM/Source/RigVM/Private/RigVMCore/RigVMByteCode.cpp b/Engine/Plugins/Runtime/RigVM/Source/RigVM/Private/RigVMCore/RigVMByteCode.cpp
--- a/Engine/Plugins/Runtime/RigVM/Source/RigVM/Private/RigVMCore/RigVMByteCode.cpp
+++ b/Engine/Plugins/Runtime/RigVM/Source/RigVM/Private/RigVMCore/RigVMByteCode.cpp
@@ -4,9 +4,21 @@
 #include "UObject/AnimObjectVersion.h"
 #include "UObject/UE5MainStreamObjectVersion.h"
 #include "UObject/FortniteMainBranchObjectVersion.h"
+#include "RigVMCore/RigVM.h"
+#include "RigVMObjectVersion.h"
 
 #include UE_INLINE_GENERATED_CPP_BY_NAME(RigVMByteCode)
 
+ERigVMExecuteResult FRigVMPredicateBranch::Execute(FRigVMExtendedExecuteContext& Context)
+{
+	check(VM);
+	if(BranchInfo.IsValid())
+	{
+		return VM->ExecuteBranch(Context, BranchInfo);
+	}
+	return ERigVMExecuteResult::Failed;
+}
+
 void FRigVMExecuteOp::Serialize(FArchive& Ar)
 {
 	Ar << OpCode;
@@ -30,6 +42,12 @@ void FRigVMExecuteOp::Serialize(FArchive& Ar)
 	{
 		Ar << ArgumentCount;
 	}
+
+	if (Ar.CustomVer(FRigVMObjectVersion::GUID) >= FRigVMObjectVersion::PredicatesAddedToExecuteOps)
+	{
+		Ar << FirstPredicateIndex;
+		Ar << PredicateCount;
+	}
 }
 
 void FRigVMUnaryOp::Serialize(FArchive& Ar)
@@ -470,7 +488,7 @@ void FRigVMByteCode::Load(FArchive& Ar)
 					Operands.Add(Operand);
 				}
 
-				AddExecuteOp(Op.FunctionIndex, Operands);
+				AddExecuteOp(Op.FunctionIndex, Operands, Op.FirstPredicateIndex, Op.PredicateCount);
 				break;
 			}
 			case ERigVMOpCode::Copy:
@@ -1039,9 +1057,11 @@ uint64 FRigVMByteCode::AddJumpIfOp(ERigVMOpCode InOpCode, uint16 InInstructionIn
 	return AddOp(Op);
 }
 
-uint64 FRigVMByteCode::AddExecuteOp(uint16 InFunctionIndex, const FRigVMOperandArray& InOperands)
+uint64 FRigVMByteCode::AddExecuteOp(uint16 InFunctionIndex, const FRigVMOperandArray& InOperands, const int32& StartPredicateIndex, const int32 PredicateCount)
 {
 	FRigVMExecuteOp Op(InFunctionIndex, (uint8)InOperands.Num());
+	Op.FirstPredicateIndex = StartPredicateIndex;
+	Op.PredicateCount = PredicateCount;
 	uint64 OpByteIndex = AddOp(Op);
 
 	uint64 OperandsByteIndex = (uint64)ByteCode.AddZeroed(sizeof(FRigVMOperand) * InOperands.Num());
@@ -1263,6 +1283,11 @@ int32 FRigVMByteCode::AddBranchInfo(const FName& InBranchLabel, int32 InInstruct
 	return AddBranchInfo(BranchInfo);
 }
 
+int32 FRigVMByteCode::AddPredicateBranch(const FRigVMPredicateBranch& InPredicateBranch)
+{
+	return PredicateBranches.Add(InPredicateBranch);
+}
+
 FRigVMOperandArray FRigVMByteCode::GetOperandsForOp(const FRigVMInstruction& InInstruction) const
 {
 	switch(InInstruction.OpCode)
diff --git a/Engine/Plugins/Runtime/RigVM/Source/RigVM/Private/RigVMCore/RigVMDispatchFactory.cpp b/Engine/Plugins/Runtime/RigVM/Source/RigVM/Private/RigVMCore/RigVMDispatchFactory.cpp
--- a/Engine/Plugins/Runtime/RigVM/Source/RigVM/Private/RigVMCore/RigVMDispatchFactory.cpp
+++ b/Engine/Plugins/Runtime/RigVM/Source/RigVM/Private/RigVMCore/RigVMDispatchFactory.cpp
@@ -220,6 +220,17 @@ FRigVMFunctionPtr FRigVMDispatchFactory::CreateDispatchFunction_NoLock(const FRi
 	return GetDispatchFunctionImpl(InTypes);
 }
 
+TArray<FRigVMFunction> FRigVMDispatchFactory::CreateDispatchPredicates(const FRigVMTemplateTypeMap& InTypes) const
+{
+	FScopeLock GetTemplateScopeLock(&FRigVMRegistry::GetDispatchPredicatesMutex);
+	return CreateDispatchPredicates_NoLock(InTypes);
+}
+
+TArray<FRigVMFunction> FRigVMDispatchFactory::CreateDispatchPredicates_NoLock(const FRigVMTemplateTypeMap& InTypes) const
+{
+	return GetDispatchPredicatesImpl(InTypes);
+}
+
 FString FRigVMDispatchFactory::GetPermutationName(const FRigVMTemplateTypeMap& InTypes) const
 {
 #if WITH_EDITOR
@@ -300,6 +311,12 @@ const FRigVMTemplate* FRigVMDispatchFactory::GetTemplate() const
 		return CreateDispatchFunction(InTypes);
 	});
 
+	Delegates.RequestDispatchPredicatesDelegate = FRigVMTemplate_RequestDispatchPredicatesDelegate::CreateLambda(
+	[this](const FRigVMTemplate*, const FRigVMTemplateTypeMap& InTypes)
+	{
+		return CreateDispatchPredicates(InTypes);
+	});
+
 	CachedTemplate = Registry.AddTemplateFromArguments(GetFactoryName(), Arguments, Delegates); 
 	return CachedTemplate;
 }
diff --git a/Engine/Plugins/Runtime/RigVM/Source/RigVM/Private/RigVMCore/RigVMMemoryStorage.cpp b/Engine/Plugins/Runtime/RigVM/Source/RigVM/Private/RigVMCore/RigVMMemoryStorage.cpp
--- a/Engine/Plugins/Runtime/RigVM/Source/RigVM/Private/RigVMCore/RigVMMemoryStorage.cpp
+++ b/Engine/Plugins/Runtime/RigVM/Source/RigVM/Private/RigVMCore/RigVMMemoryStorage.cpp
@@ -84,7 +84,7 @@ ERigVMExecuteResult FRigVMLazyBranch::Execute(FRigVMExtendedExecuteContext& Cont
 	check(VM);
 	if(BranchInfo.IsValid())
 	{
-		return VM->ExecuteLazyBranch(Context, BranchInfo);
+		return VM->ExecuteBranch(Context, BranchInfo);
 	}
 
 	check(FunctionPtr);
diff --git a/Engine/Plugins/Runtime/RigVM/Source/RigVM/Private/RigVMCore/RigVMRegistry.cpp b/Engine/Plugins/Runtime/RigVM/Source/RigVM/Private/RigVMCore/RigVMRegistry.cpp
--- a/Engine/Plugins/Runtime/RigVM/Source/RigVM/Private/RigVMCore/RigVMRegistry.cpp
+++ b/Engine/Plugins/Runtime/RigVM/Source/RigVM/Private/RigVMCore/RigVMRegistry.cpp
@@ -23,6 +23,7 @@ FCriticalSection FRigVMRegistry::FindFunctionMutex;
 FCriticalSection FRigVMRegistry::FindTemplateMutex;
 FCriticalSection FRigVMRegistry::FindFactoryMutex;
 FCriticalSection FRigVMRegistry::GetDispatchFunctionMutex;
+FCriticalSection FRigVMRegistry::GetDispatchPredicatesMutex;
 FCriticalSection FRigVMRegistry::GetPermutationMutex;
 
 
@@ -1259,6 +1260,23 @@ const FRigVMDispatchFactory* FRigVMRegistry::RegisterFactory(UScriptStruct* InFa
 	return Factory;
 }
 
+void FRigVMRegistry::RegisterPredicate(UScriptStruct* InStruct, const TCHAR* InName, const TArray<FRigVMFunctionArgument>& InArguments)
+{
+	// Make sure the predicate does not already exist
+	TArray<FRigVMFunction>& Predicates = StructNameToPredicates.FindOrAdd(InStruct->GetFName());
+	if (Predicates.ContainsByPredicate([InName](const FRigVMFunction& Predicate)
+	{
+		return Predicate.Name == InName;
+	}))
+	{
+		
+		return;
+	}
+
+	FRigVMFunction Function(InName, nullptr, InStruct, Predicates.Num(), InArguments);
+	Predicates.Add(Function);
+}
+
 const FRigVMFunction* FRigVMRegistry::FindFunction(const TCHAR* InName) const
 {
 	FScopeLock FindFunctionScopeLock(&FindFunctionMutex);
@@ -1707,3 +1725,8 @@ const TArray<FRigVMDispatchFactory*>& FRigVMRegistry::GetFactories() const
 {
 	return Factories;
 }
+
+const TArray<FRigVMFunction>* FRigVMRegistry::GetPredicatesForStruct(const FName& InStructName) const
+{
+	return StructNameToPredicates.Find(InStructName);
+}
diff --git a/Engine/Plugins/Runtime/RigVM/Source/RigVM/Private/RigVMCore/RigVMTemplate.cpp b/Engine/Plugins/Runtime/RigVM/Source/RigVM/Private/RigVMCore/RigVMTemplate.cpp
--- a/Engine/Plugins/Runtime/RigVM/Source/RigVM/Private/RigVMCore/RigVMTemplate.cpp
+++ b/Engine/Plugins/Runtime/RigVM/Source/RigVM/Private/RigVMCore/RigVMTemplate.cpp
@@ -1232,6 +1232,12 @@ const FRigVMFunction* FRigVMTemplate::GetOrCreatePermutation(int32 InIndex)
 			Registry.Functions[FunctionIndex].TemplateIndex = Index;
 			Registry.FunctionNameToIndex.Add(*PermutationName, FunctionIndex);
 
+			if (Delegates.RequestDispatchPredicatesDelegate.IsBound())
+			{
+				TArray<FRigVMFunction> Predicates = Delegates.RequestDispatchPredicatesDelegate.Execute(this, Types);
+				Registry.StructNameToPredicates.Add(*PermutationName, Predicates);
+			}
+
 			return &Registry.Functions[FunctionIndex];
 		}
 	}
diff --git a/Engine/Plugins/Runtime/RigVM/Source/RigVM/Private/RigVMFunctions/RigVMDispatch_Array.cpp b/Engine/Plugins/Runtime/RigVM/Source/RigVM/Private/RigVMFunctions/RigVMDispatch_Array.cpp
--- a/Engine/Plugins/Runtime/RigVM/Source/RigVM/Private/RigVMFunctions/RigVMDispatch_Array.cpp
+++ b/Engine/Plugins/Runtime/RigVM/Source/RigVM/Private/RigVMFunctions/RigVMDispatch_Array.cpp
@@ -282,7 +282,7 @@ FName FRigVMDispatch_ArrayMake::GetArgumentNameForOperandIndex(int32 InOperandIn
 	return ArrayName;
 }
 
-void FRigVMDispatch_ArrayMake::Execute(FRigVMExtendedExecuteContext& InContext, FRigVMMemoryHandleArray Handles)
+void FRigVMDispatch_ArrayMake::Execute(FRigVMExtendedExecuteContext& InContext, FRigVMMemoryHandleArray Handles, FRigVMPredicateBranchArray Predicates)
 {
 	const FArrayProperty* ArrayProperty = CastFieldChecked<FArrayProperty>(Handles.Last().GetProperty());
 	FScriptArrayHelper ArrayHelper(ArrayProperty, Handles.Last().GetData());
@@ -328,7 +328,7 @@ FRigVMTemplateTypeMap FRigVMDispatch_ArrayReset::OnNewArgumentType(const FName&
 	};
 }
 
-void FRigVMDispatch_ArrayReset::Execute(FRigVMExtendedExecuteContext& InContext, FRigVMMemoryHandleArray Handles)
+void FRigVMDispatch_ArrayReset::Execute(FRigVMExtendedExecuteContext& InContext, FRigVMMemoryHandleArray Handles, FRigVMPredicateBranchArray Predicates)
 {
 	const FArrayProperty* ArrayProperty = CastFieldChecked<FArrayProperty>(Handles[0].GetProperty());
 	FScriptArrayHelper ArrayHelper(ArrayProperty, Handles[0].GetData());
@@ -382,7 +382,7 @@ FRigVMTemplateTypeMap FRigVMDispatch_ArrayGetNum::OnNewArgumentType(const FName&
 	};
 }
 
-void FRigVMDispatch_ArrayGetNum::Execute(FRigVMExtendedExecuteContext& InContext, FRigVMMemoryHandleArray Handles)
+void FRigVMDispatch_ArrayGetNum::Execute(FRigVMExtendedExecuteContext& InContext, FRigVMMemoryHandleArray Handles, FRigVMPredicateBranchArray Predicates)
 {
 	const FArrayProperty* ArrayProperty = CastFieldChecked<FArrayProperty>(Handles[0].GetProperty());
 	const FScriptArrayHelper ArrayHelper(ArrayProperty, Handles[0].GetData());
@@ -441,7 +441,7 @@ FRigVMTemplateTypeMap FRigVMDispatch_ArraySetNum::OnNewArgumentType(const FName&
 	};
 }
 
-void FRigVMDispatch_ArraySetNum::Execute(FRigVMExtendedExecuteContext& InContext, FRigVMMemoryHandleArray Handles)
+void FRigVMDispatch_ArraySetNum::Execute(FRigVMExtendedExecuteContext& InContext, FRigVMMemoryHandleArray Handles, FRigVMPredicateBranchArray Predicates)
 {
 	const FArrayProperty* ArrayProperty = CastFieldChecked<FArrayProperty>(Handles[0].GetProperty());
 	FScriptArrayHelper ArrayHelper(ArrayProperty, Handles[0].GetData());
@@ -516,7 +516,7 @@ FRigVMTemplateTypeMap FRigVMDispatch_ArrayGetAtIndex::OnNewArgumentType(const FN
 	};
 }
 
-void FRigVMDispatch_ArrayGetAtIndex::Execute(FRigVMExtendedExecuteContext& InContext, FRigVMMemoryHandleArray Handles)
+void FRigVMDispatch_ArrayGetAtIndex::Execute(FRigVMExtendedExecuteContext& InContext, FRigVMMemoryHandleArray Handles, FRigVMPredicateBranchArray Predicates)
 {
 	const FArrayProperty* ArrayProperty = CastFieldChecked<FArrayProperty>(Handles[0].GetProperty());
 	FScriptArrayHelper ArrayHelper(ArrayProperty, Handles[0].GetData());
@@ -598,7 +598,7 @@ FRigVMTemplateTypeMap FRigVMDispatch_ArraySetAtIndex::OnNewArgumentType(const FN
 	};
 }
 
-void FRigVMDispatch_ArraySetAtIndex::Execute(FRigVMExtendedExecuteContext& InContext, FRigVMMemoryHandleArray Handles)
+void FRigVMDispatch_ArraySetAtIndex::Execute(FRigVMExtendedExecuteContext& InContext, FRigVMMemoryHandleArray Handles, FRigVMPredicateBranchArray Predicates)
 {
 	const FArrayProperty* ArrayProperty = CastFieldChecked<FArrayProperty>(Handles[0].GetProperty());
 	FScriptArrayHelper ArrayHelper(ArrayProperty, Handles[0].GetData());
@@ -660,7 +660,7 @@ TArray<FRigVMTemplateArgument> FRigVMDispatch_ArrayAdd::GetArguments() const
 	};
 }
 
-void FRigVMDispatch_ArrayAdd::Execute(FRigVMExtendedExecuteContext& InContext, FRigVMMemoryHandleArray Handles)
+void FRigVMDispatch_ArrayAdd::Execute(FRigVMExtendedExecuteContext& InContext, FRigVMMemoryHandleArray Handles, FRigVMPredicateBranchArray Predicates)
 {
 	const FArrayProperty* ArrayProperty = CastFieldChecked<FArrayProperty>(Handles[0].GetProperty());
 	FScriptArrayHelper ArrayHelper(ArrayProperty, Handles[0].GetData());
@@ -707,7 +707,7 @@ FText FRigVMDispatch_ArrayAdd::GetArgumentTooltip(const FName& InArgumentName, T
 
 ////////////////////////////////////////////////////////////////////////////////////////////////////////////////////////
 
-void FRigVMDispatch_ArrayInsert::Execute(FRigVMExtendedExecuteContext& InContext, FRigVMMemoryHandleArray Handles)
+void FRigVMDispatch_ArrayInsert::Execute(FRigVMExtendedExecuteContext& InContext, FRigVMMemoryHandleArray Handles, FRigVMPredicateBranchArray Predicates)
 {
 	const FArrayProperty* ArrayProperty = CastFieldChecked<FArrayProperty>(Handles[0].GetProperty());
 	FScriptArrayHelper ArrayHelper(ArrayProperty, Handles[0].GetData());
@@ -786,7 +786,7 @@ FRigVMTemplateTypeMap FRigVMDispatch_ArrayRemove::OnNewArgumentType(const FName&
 	};
 }
 
-void FRigVMDispatch_ArrayRemove::Execute(FRigVMExtendedExecuteContext& InContext, FRigVMMemoryHandleArray Handles)
+void FRigVMDispatch_ArrayRemove::Execute(FRigVMExtendedExecuteContext& InContext, FRigVMMemoryHandleArray Handles, FRigVMPredicateBranchArray Predicates)
 {
 	const FArrayProperty* ArrayProperty = CastFieldChecked<FArrayProperty>(Handles[0].GetProperty());
 	FScriptArrayHelper ArrayHelper(ArrayProperty, Handles[0].GetData());
@@ -822,7 +822,7 @@ FText FRigVMDispatch_ArrayRemove::GetArgumentTooltip(const FName& InArgumentName
 
 ////////////////////////////////////////////////////////////////////////////////////////////////////////////////////////
 
-void FRigVMDispatch_ArrayReverse::Execute(FRigVMExtendedExecuteContext& InContext, FRigVMMemoryHandleArray Handles)
+void FRigVMDispatch_ArrayReverse::Execute(FRigVMExtendedExecuteContext& InContext, FRigVMMemoryHandleArray Handles, FRigVMPredicateBranchArray Predicates)
 {
 	const FArrayProperty* ArrayProperty = CastFieldChecked<FArrayProperty>(Handles[0].GetProperty());
 	FScriptArrayHelper ArrayHelper(ArrayProperty, Handles[0].GetData());
@@ -896,7 +896,7 @@ FRigVMTemplateTypeMap FRigVMDispatch_ArrayFind::OnNewArgumentType(const FName& I
 	};
 }
 
-void FRigVMDispatch_ArrayFind::Execute(FRigVMExtendedExecuteContext& InContext, FRigVMMemoryHandleArray Handles)
+void FRigVMDispatch_ArrayFind::Execute(FRigVMExtendedExecuteContext& InContext, FRigVMMemoryHandleArray Handles, FRigVMPredicateBranchArray Predicates)
 {
 	const FArrayProperty* ArrayProperty = CastFieldChecked<FArrayProperty>(Handles[0].GetProperty());
 	FScriptArrayHelper ArrayHelper(ArrayProperty, Handles[0].GetData());
@@ -991,7 +991,7 @@ FRigVMTemplateTypeMap FRigVMDispatch_ArrayAppend::OnNewArgumentType(const FName&
 	};
 }
 
-void FRigVMDispatch_ArrayAppend::Execute(FRigVMExtendedExecuteContext& InContext, FRigVMMemoryHandleArray Handles)
+void FRigVMDispatch_ArrayAppend::Execute(FRigVMExtendedExecuteContext& InContext, FRigVMMemoryHandleArray Handles, FRigVMPredicateBranchArray Predicates)
 {
 	const FArrayProperty* OtherProperty = CastFieldChecked<FArrayProperty>(Handles[1].GetProperty());
 	FScriptArrayHelper OtherHelper(OtherProperty, Handles[1].GetData());
@@ -1069,7 +1069,7 @@ FRigVMTemplateTypeMap FRigVMDispatch_ArrayClone::OnNewArgumentType(const FName&
 	};
 }
 
-void FRigVMDispatch_ArrayClone::Execute(FRigVMExtendedExecuteContext& InContext, FRigVMMemoryHandleArray Handles)
+void FRigVMDispatch_ArrayClone::Execute(FRigVMExtendedExecuteContext& InContext, FRigVMMemoryHandleArray Handles, FRigVMPredicateBranchArray Predicates)
 {
 	const FArrayProperty* ArrayProperty = CastFieldChecked<FArrayProperty>(Handles[0].GetProperty());
 	FScriptArrayHelper ArrayHelper(ArrayProperty, Handles[0].GetData());
@@ -1114,7 +1114,7 @@ FText FRigVMDispatch_ArrayClone::GetArgumentTooltip(const FName& InArgumentName,
 
 ////////////////////////////////////////////////////////////////////////////////////////////////////////////////////////
 
-void FRigVMDispatch_ArrayUnion::Execute(FRigVMExtendedExecuteContext& InContext, FRigVMMemoryHandleArray Handles)
+void FRigVMDispatch_ArrayUnion::Execute(FRigVMExtendedExecuteContext& InContext, FRigVMMemoryHandleArray Handles, FRigVMPredicateBranchArray Predicates)
 {
 	const FArrayProperty* OtherProperty = CastFieldChecked<FArrayProperty>(Handles[1].GetProperty());
 	FScriptArrayHelper OtherHelper(OtherProperty, Handles[1].GetData());
@@ -1235,7 +1235,7 @@ FRigVMTemplateTypeMap FRigVMDispatch_ArrayDifference::OnNewArgumentType(const FN
 	};
 }
 
-void FRigVMDispatch_ArrayDifference::Execute(FRigVMExtendedExecuteContext& InContext, FRigVMMemoryHandleArray Handles)
+void FRigVMDispatch_ArrayDifference::Execute(FRigVMExtendedExecuteContext& InContext, FRigVMMemoryHandleArray Handles, FRigVMPredicateBranchArray Predicates)
 {
 	const FArrayProperty* ArrayProperty = CastFieldChecked<FArrayProperty>(Handles[0].GetProperty());
 	FScriptArrayHelper ArrayHelper(ArrayProperty, Handles[0].GetData());
@@ -1309,7 +1309,7 @@ FText FRigVMDispatch_ArrayDifference::GetArgumentTooltip(const FName& InArgument
 
 ////////////////////////////////////////////////////////////////////////////////////////////////////////////////////////
 
-void FRigVMDispatch_ArrayIntersection::Execute(FRigVMExtendedExecuteContext& InContext, FRigVMMemoryHandleArray Handles)
+void FRigVMDispatch_ArrayIntersection::Execute(FRigVMExtendedExecuteContext& InContext, FRigVMMemoryHandleArray Handles, FRigVMPredicateBranchArray Predicates)
 {
 	const FArrayProperty* ArrayProperty = CastFieldChecked<FArrayProperty>(Handles[0].GetProperty());
 	FScriptArrayHelper ArrayHelper(ArrayProperty, Handles[0].GetData());
@@ -1441,7 +1441,7 @@ FRigVMTemplateTypeMap FRigVMDispatch_ArrayIterator::OnNewArgumentType(const FNam
 		{FRigVMStruct::ControlFlowBlockToRunName, RigVMTypeUtils::TypeIndex::FName}
 	};}
 
-void FRigVMDispatch_ArrayIterator::Execute(FRigVMExtendedExecuteContext& InContext, FRigVMMemoryHandleArray Handles)
+void FRigVMDispatch_ArrayIterator::Execute(FRigVMExtendedExecuteContext& InContext, FRigVMMemoryHandleArray Handles, FRigVMPredicateBranchArray Predicates)
 {
 	const FArrayProperty* ArrayProperty = CastFieldChecked<FArrayProperty>(Handles[0].GetProperty());
 	FScriptArrayHelper ArrayHelper(ArrayProperty, Handles[0].GetData());
diff --git a/Engine/Plugins/Runtime/RigVM/Source/RigVM/Private/RigVMFunctions/RigVMDispatch_Constant.cpp b/Engine/Plugins/Runtime/RigVM/Source/RigVM/Private/RigVMFunctions/RigVMDispatch_Constant.cpp
--- a/Engine/Plugins/Runtime/RigVM/Source/RigVM/Private/RigVMFunctions/RigVMDispatch_Constant.cpp
+++ b/Engine/Plugins/Runtime/RigVM/Source/RigVM/Private/RigVMFunctions/RigVMDispatch_Constant.cpp
@@ -60,7 +60,7 @@ FText FRigVMDispatch_Constant::GetNodeTooltip(const FRigVMTemplateTypeMap& InTyp
 
 #endif
 
-void FRigVMDispatch_Constant::Execute(FRigVMExtendedExecuteContext& InContext, FRigVMMemoryHandleArray Handles)
+void FRigVMDispatch_Constant::Execute(FRigVMExtendedExecuteContext& InContext, FRigVMMemoryHandleArray Handles, FRigVMPredicateBranchArray RigVMBranches)
 {
 	// nothing to do here - it's a constant
 }
diff --git a/Engine/Plugins/Runtime/RigVM/Source/RigVM/Private/RigVMFunctions/RigVMDispatch_Core.cpp b/Engine/Plugins/Runtime/RigVM/Source/RigVM/Private/RigVMFunctions/RigVMDispatch_Core.cpp
--- a/Engine/Plugins/Runtime/RigVM/Source/RigVM/Private/RigVMFunctions/RigVMDispatch_Core.cpp
+++ b/Engine/Plugins/Runtime/RigVM/Source/RigVM/Private/RigVMFunctions/RigVMDispatch_Core.cpp
@@ -109,7 +109,7 @@ bool FRigVMDispatch_CoreEquals::AdaptResult(bool bResult, const FRigVMExtendedEx
 	return bResult;
 }
 
-void FRigVMDispatch_CoreEquals::Execute(FRigVMExtendedExecuteContext& InContext, FRigVMMemoryHandleArray Handles)
+void FRigVMDispatch_CoreEquals::Execute(FRigVMExtendedExecuteContext& InContext, FRigVMMemoryHandleArray Handles, FRigVMPredicateBranchArray Predicates)
 {
 	const FProperty* PropertyA = Handles[0].GetResolvedProperty(); 
 	const FProperty* PropertyB = Handles[1].GetResolvedProperty(); 
diff --git a/Engine/Plugins/Runtime/RigVM/Source/RigVM/Private/RigVMFunctions/RigVMDispatch_If.cpp b/Engine/Plugins/Runtime/RigVM/Source/RigVM/Private/RigVMFunctions/RigVMDispatch_If.cpp
--- a/Engine/Plugins/Runtime/RigVM/Source/RigVM/Private/RigVMFunctions/RigVMDispatch_If.cpp
+++ b/Engine/Plugins/Runtime/RigVM/Source/RigVM/Private/RigVMFunctions/RigVMDispatch_If.cpp
@@ -69,7 +69,7 @@ FRigVMFunctionPtr FRigVMDispatch_If::GetDispatchFunctionImpl(const FRigVMTemplat
 	return &FRigVMDispatch_If::Execute;
 }
 
-void FRigVMDispatch_If::Execute(FRigVMExtendedExecuteContext& InContext, FRigVMMemoryHandleArray Handles)
+void FRigVMDispatch_If::Execute(FRigVMExtendedExecuteContext& InContext, FRigVMMemoryHandleArray Handles, FRigVMPredicateBranchArray Predicates)
 {
 	const FProperty* Property = Handles[1].GetResolvedProperty(); 
 
diff --git a/Engine/Plugins/Runtime/RigVM/Source/RigVM/Private/RigVMFunctions/RigVMDispatch_MakeStruct.cpp b/Engine/Plugins/Runtime/RigVM/Source/RigVM/Private/RigVMFunctions/RigVMDispatch_MakeStruct.cpp
--- a/Engine/Plugins/Runtime/RigVM/Source/RigVM/Private/RigVMFunctions/RigVMDispatch_MakeStruct.cpp
+++ b/Engine/Plugins/Runtime/RigVM/Source/RigVM/Private/RigVMFunctions/RigVMDispatch_MakeStruct.cpp
@@ -74,7 +74,7 @@ FString FRigVMDispatch_MakeStruct::GetArgumentMetaData(const FName& InArgumentNa
 
 #endif
 
-void FRigVMDispatch_MakeStruct::Execute(FRigVMExtendedExecuteContext& InContext, FRigVMMemoryHandleArray Handles)
+void FRigVMDispatch_MakeStruct::Execute(FRigVMExtendedExecuteContext& InContext, FRigVMMemoryHandleArray Handles, FRigVMPredicateBranchArray RigVMBranches)
 {
 	const FProperty* SourceProperty = Handles[0].GetProperty();
 	const FProperty* TargetProperty = Handles[1].GetProperty();
diff --git a/Engine/Plugins/Runtime/RigVM/Source/RigVM/Private/RigVMFunctions/RigVMDispatch_Print.cpp b/Engine/Plugins/Runtime/RigVM/Source/RigVM/Private/RigVMFunctions/RigVMDispatch_Print.cpp
--- a/Engine/Plugins/Runtime/RigVM/Source/RigVM/Private/RigVMFunctions/RigVMDispatch_Print.cpp
+++ b/Engine/Plugins/Runtime/RigVM/Source/RigVM/Private/RigVMFunctions/RigVMDispatch_Print.cpp
@@ -89,7 +89,7 @@ FString FRigVMDispatch_Print::GetArgumentMetaData(const FName& InArgumentName, c
 
 #endif
 
-void FRigVMDispatch_Print::Execute(FRigVMExtendedExecuteContext& InContext, FRigVMMemoryHandleArray Handles)
+void FRigVMDispatch_Print::Execute(FRigVMExtendedExecuteContext& InContext, FRigVMMemoryHandleArray Handles, FRigVMPredicateBranchArray Predicates)
 {
 #if WITH_EDITOR
 	const FProperty* ValueProperty = Handles[1].GetResolvedProperty(); 
diff --git a/Engine/Plugins/Runtime/RigVM/Source/RigVM/Private/RigVMFunctions/RigVMDispatch_Select.cpp b/Engine/Plugins/Runtime/RigVM/Source/RigVM/Private/RigVMFunctions/RigVMDispatch_Select.cpp
--- a/Engine/Plugins/Runtime/RigVM/Source/RigVM/Private/RigVMFunctions/RigVMDispatch_Select.cpp
+++ b/Engine/Plugins/Runtime/RigVM/Source/RigVM/Private/RigVMFunctions/RigVMDispatch_Select.cpp
@@ -109,7 +109,7 @@ FRigVMFunctionPtr FRigVMDispatch_SelectInt32::GetDispatchFunctionImpl(const FRig
 	return &FRigVMDispatch_SelectInt32::Execute;
 }
 
-void FRigVMDispatch_SelectInt32::Execute(FRigVMExtendedExecuteContext& InContext, FRigVMMemoryHandleArray Handles)
+void FRigVMDispatch_SelectInt32::Execute(FRigVMExtendedExecuteContext& InContext, FRigVMMemoryHandleArray Handles, FRigVMPredicateBranchArray Predicates)
 {
 	const FProperty* Property = Handles.Last().GetResolvedProperty();
 	const int32 NumCases = Handles.Num() - 2;
diff --git a/Engine/Plugins/Runtime/RigVM/Source/RigVM/Private/RigVMFunctions/RigVMDispatch_Switch.cpp b/Engine/Plugins/Runtime/RigVM/Source/RigVM/Private/RigVMFunctions/RigVMDispatch_Switch.cpp
--- a/Engine/Plugins/Runtime/RigVM/Source/RigVM/Private/RigVMFunctions/RigVMDispatch_Switch.cpp
+++ b/Engine/Plugins/Runtime/RigVM/Source/RigVM/Private/RigVMFunctions/RigVMDispatch_Switch.cpp
@@ -125,7 +125,7 @@ FRigVMFunctionPtr FRigVMDispatch_SwitchInt32::GetDispatchFunctionImpl(const FRig
 	return &FRigVMDispatch_SwitchInt32::Execute;
 }
 
-void FRigVMDispatch_SwitchInt32::Execute(FRigVMExtendedExecuteContext& InContext, FRigVMMemoryHandleArray Handles)
+void FRigVMDispatch_SwitchInt32::Execute(FRigVMExtendedExecuteContext& InContext, FRigVMMemoryHandleArray Handles, FRigVMPredicateBranchArray Predicates)
 {
 #if WITH_EDITOR
 	check(Handles[0].IsInt32());
diff --git a/Engine/Plugins/Runtime/RigVM/Source/RigVM/Private/RigVMFunctions/RigVMFunction_String.cpp b/Engine/Plugins/Runtime/RigVM/Source/RigVM/Private/RigVMFunctions/RigVMFunction_String.cpp
--- a/Engine/Plugins/Runtime/RigVM/Source/RigVM/Private/RigVMFunctions/RigVMFunction_String.cpp
+++ b/Engine/Plugins/Runtime/RigVM/Source/RigVM/Private/RigVMFunctions/RigVMFunction_String.cpp
@@ -279,7 +279,7 @@ FRigVMFunctionPtr FRigDispatch_ToString::GetDispatchFunctionImpl(const FRigVMTem
 	const TRigVMTypeIndex& ValueTypeIndex = InTypes.FindChecked(TEXT("Value"));
 	if(ValueTypeIndex == RigVMTypeUtils::TypeIndex::FName)
 	{
-		return [](FRigVMExtendedExecuteContext& InContext, FRigVMMemoryHandleArray Handles)
+		return [](FRigVMExtendedExecuteContext& InContext, FRigVMMemoryHandleArray Handles, FRigVMPredicateBranchArray Predicates)
 		{
 			check(Handles[0].IsName());
 			check(Handles[1].IsString());
@@ -291,7 +291,7 @@ FRigVMFunctionPtr FRigDispatch_ToString::GetDispatchFunctionImpl(const FRigVMTem
 	}
 	if(ValueTypeIndex == RigVMTypeUtils::TypeIndex::FString)
 	{
-		return [](FRigVMExtendedExecuteContext& InContext, FRigVMMemoryHandleArray Handles)
+		return [](FRigVMExtendedExecuteContext& InContext, FRigVMMemoryHandleArray Handles, FRigVMPredicateBranchArray Predicates)
 		{
 			check(Handles[0].IsString());
 			check(Handles[1].IsString());
@@ -305,7 +305,7 @@ FRigVMFunctionPtr FRigDispatch_ToString::GetDispatchFunctionImpl(const FRigVMTem
 	return &FRigDispatch_ToString::Execute;
 }
 
-void FRigDispatch_ToString::Execute(FRigVMExtendedExecuteContext& InContext, FRigVMMemoryHandleArray Handles)
+void FRigDispatch_ToString::Execute(FRigVMExtendedExecuteContext& InContext, FRigVMMemoryHandleArray Handles, FRigVMPredicateBranchArray Predicates)
 {
 	const FProperty* ValueProperty = Handles[0].GetResolvedProperty(); 
 	check(ValueProperty);
@@ -344,7 +344,7 @@ FRigVMFunctionPtr FRigDispatch_FromString::GetDispatchFunctionImpl(const FRigVMT
 	const TRigVMTypeIndex& ResultTypeIndex = InTypes.FindChecked(TEXT("Result"));
 	if(ResultTypeIndex == RigVMTypeUtils::TypeIndex::FName)
 	{
-		return [](FRigVMExtendedExecuteContext& InContext, FRigVMMemoryHandleArray Handles)
+		return [](FRigVMExtendedExecuteContext& InContext, FRigVMMemoryHandleArray Handles, FRigVMPredicateBranchArray Predicates)
 		{
 			check(Handles[0].IsString());
 			check(Handles[1].IsName());
@@ -356,7 +356,7 @@ FRigVMFunctionPtr FRigDispatch_FromString::GetDispatchFunctionImpl(const FRigVMT
 	}
 	if(ResultTypeIndex == RigVMTypeUtils::TypeIndex::FString)
 	{
-		return [](FRigVMExtendedExecuteContext& InContext, FRigVMMemoryHandleArray Handles)
+		return [](FRigVMExtendedExecuteContext& InContext, FRigVMMemoryHandleArray Handles, FRigVMPredicateBranchArray Predicates)
 		{
 			check(Handles[0].IsString());
 			check(Handles[1].IsString());
@@ -370,7 +370,7 @@ FRigVMFunctionPtr FRigDispatch_FromString::GetDispatchFunctionImpl(const FRigVMT
 	return &FRigDispatch_FromString::Execute;
 }
 
-void FRigDispatch_FromString::Execute(FRigVMExtendedExecuteContext& InContext, FRigVMMemoryHandleArray Handles)
+void FRigDispatch_FromString::Execute(FRigVMExtendedExecuteContext& InContext, FRigVMMemoryHandleArray Handles, FRigVMPredicateBranchArray Predicates)
 {
 	const FProperty* ValueProperty = Handles[1].GetResolvedProperty(); 
 	check(ValueProperty);
diff --git a/Engine/Plugins/Runtime/RigVM/Source/RigVM/Public/RigVMCore/RigVM.h b/Engine/Plugins/Runtime/RigVM/Source/RigVM/Public/RigVMCore/RigVM.h
--- a/Engine/Plugins/Runtime/RigVM/Source/RigVM/Public/RigVMCore/RigVM.h
+++ b/Engine/Plugins/Runtime/RigVM/Source/RigVM/Public/RigVMCore/RigVM.h
@@ -234,7 +234,7 @@ public:
 	ERigVMExecuteResult ExecuteLazyBranch(const FRigVMBranchInfo& InBranchToRun) { return ERigVMExecuteResult::Failed; }
 
 	// Executes a single branch on the VM. We assume that the memory is already set correctly at this point.
-	ERigVMExecuteResult ExecuteLazyBranch(FRigVMExtendedExecuteContext& Context, const FRigVMBranchInfo& InBranchToRun);
+	ERigVMExecuteResult ExecuteBranch(FRigVMExtendedExecuteContext& Context, const FRigVMBranchInfo& InBranchToRun);
 
 private:
 	ERigVMExecuteResult ExecuteInstructions(FRigVMExtendedExecuteContext& Context, int32 InFirstInstruction, int32 InLastInstruction);
diff --git a/Engine/Plugins/Runtime/RigVM/Source/RigVM/Public/RigVMCore/RigVMByteCode.h b/Engine/Plugins/Runtime/RigVM/Source/RigVM/Public/RigVMCore/RigVMByteCode.h
--- a/Engine/Plugins/Runtime/RigVM/Source/RigVM/Public/RigVMCore/RigVMByteCode.h
+++ b/Engine/Plugins/Runtime/RigVM/Source/RigVM/Public/RigVMCore/RigVMByteCode.h
@@ -94,6 +94,25 @@ struct FRigVMBranchInfoKey
 	FName Label;
 };
 
+// A description of a predicate branch in the VM's bytecode
+USTRUCT()
+struct RIGVM_API FRigVMPredicateBranch
+{
+	GENERATED_BODY()
+
+	FRigVMPredicateBranch()
+		: VM(nullptr)
+	{
+		
+	}
+	
+	ERigVMExecuteResult Execute(FRigVMExtendedExecuteContext& Context);
+	bool IsValid() { return BranchInfo.IsValid(); }
+	
+	FRigVMBranchInfo BranchInfo;
+	FRigVMMemoryHandleArray MemoryHandles;
+	URigVM* VM;
+};
 
 // The code for a single operation within the RigVM
 UENUM()
@@ -236,6 +255,8 @@ struct RIGVM_API FRigVMExecuteOp : public FRigVMBaseOp
 	: FRigVMBaseOp()
 	, FunctionIndex(INDEX_NONE)
 	, ArgumentCount(0)
+	, FirstPredicateIndex(INDEX_NONE)
+	, PredicateCount(INDEX_NONE)
 	{
 	}
 
@@ -248,6 +269,8 @@ struct RIGVM_API FRigVMExecuteOp : public FRigVMBaseOp
 
 	uint16 FunctionIndex;
 	uint16 ArgumentCount;
+	uint16 FirstPredicateIndex;
+	uint16 PredicateCount;
 
 	friend uint32 GetTypeHash(const FRigVMExecuteOp& Op)
 	{
@@ -988,7 +1011,7 @@ public:
 	int32 FindEntryIndex(const FName& InEntryName) const;
 
 	// adds an execute operator given its function index operands
-	uint64 AddExecuteOp(uint16 InFunctionIndex, const FRigVMOperandArray& InOperands);
+	uint64 AddExecuteOp(uint16 InFunctionIndex, const FRigVMOperandArray& InOperands, const int32& StartPredicateIndex, const int32 PredicateCount);
 	
 	// adds an execute operator given its function index operands
 	uint64 InlineFunction(const FRigVMByteCode* FunctionByteCode, const FRigVMOperandArray& InOperands);
@@ -1045,6 +1068,9 @@ public:
 	int32 AddBranchInfo(const FRigVMBranchInfo& InBranchInfo);
 	int32 AddBranchInfo(const FName& InBranchLabel, int32 InInstructionIndex, int32 InArgumentIndex, int32 InFirstBranchInstruction, int32 InLastBranchInstruction);
 
+	// adds information about a predicate branch for an instruction
+	int32 AddPredicateBranch(const FRigVMPredicateBranch& InPredicateBranch);
+
 	// returns an instruction array for iterating over all operators
 	FRigVMInstructionArray GetInstructions() const
 	{
@@ -1270,6 +1296,10 @@ private:
 	UPROPERTY()
 	TArray<FRigVMBranchInfo> BranchInfos;
 
+	// a list of all predicate branches
+	UPROPERTY()
+	TArray<FRigVMPredicateBranch> PredicateBranches;
+
 	const FRigVMBranchInfo* GetBranchInfo(const FRigVMBranchInfoKey& InBranchInfoKey) const;
 	mutable TMap<FRigVMBranchInfoKey, const FRigVMBranchInfo*> BranchInfoLookup;
 
diff --git a/Engine/Plugins/Runtime/RigVM/Source/RigVM/Public/RigVMCore/RigVMDispatchFactory.h b/Engine/Plugins/Runtime/RigVM/Source/RigVM/Public/RigVMCore/RigVMDispatchFactory.h
--- a/Engine/Plugins/Runtime/RigVM/Source/RigVM/Public/RigVMCore/RigVMDispatchFactory.h
+++ b/Engine/Plugins/Runtime/RigVM/Source/RigVM/Public/RigVMCore/RigVMDispatchFactory.h
@@ -171,6 +171,11 @@ protected:
 	FRigVMFunctionPtr CreateDispatchFunction_NoLock(const FRigVMTemplateTypeMap& InTypes) const;
 
 	virtual FRigVMFunctionPtr GetDispatchFunctionImpl(const FRigVMTemplateTypeMap& InTypes) const { return nullptr; }
+	
+	RIGVM_API TArray<FRigVMFunction> CreateDispatchPredicates(const FRigVMTemplateTypeMap& InTypes) const;
+	TArray<FRigVMFunction> CreateDispatchPredicates_NoLock(const FRigVMTemplateTypeMap& InTypes) const;
+	
+	virtual TArray<FRigVMFunction> GetDispatchPredicatesImpl(const FRigVMTemplateTypeMap& InTypes) const { return TArray<FRigVMFunction>(); }
 
 	virtual TArray<FRigVMExecuteArgument> GetExecuteArguments_Impl(const FRigVMDispatchContext& InContext) const { return TArray<FRigVMExecuteArgument>(); }
 
diff --git a/Engine/Plugins/Runtime/RigVM/Source/RigVM/Public/RigVMCore/RigVMExecuteContext.h b/Engine/Plugins/Runtime/RigVM/Source/RigVM/Public/RigVMCore/RigVMExecuteContext.h
--- a/Engine/Plugins/Runtime/RigVM/Source/RigVM/Public/RigVMCore/RigVMExecuteContext.h
+++ b/Engine/Plugins/Runtime/RigVM/Source/RigVM/Public/RigVMCore/RigVMExecuteContext.h
@@ -456,6 +456,7 @@ private:
 	FRigVMExtendedExecuteContext* ExtendedExecuteContext = nullptr; 
 
 	friend struct TRigVMLazyValueBase;
+	friend struct FRigVMPredicateBranch;
 	friend struct FRigVMExtendedExecuteContext;
 	friend class URigVM;
 	friend class URigVMNativized;
diff --git a/Engine/Plugins/Runtime/RigVM/Source/RigVM/Public/RigVMCore/RigVMFunction.h b/Engine/Plugins/Runtime/RigVM/Source/RigVM/Public/RigVMCore/RigVMFunction.h
--- a/Engine/Plugins/Runtime/RigVM/Source/RigVM/Public/RigVMCore/RigVMFunction.h
+++ b/Engine/Plugins/Runtime/RigVM/Source/RigVM/Public/RigVMCore/RigVMFunction.h
@@ -20,10 +20,12 @@
 class UScriptStruct;
 struct FRigVMExtendedExecuteContext;
 struct FRigVMMemoryHandle;
+struct FRigVMPredicateBranch;
 
 typedef TArrayView<FRigVMMemoryHandle> FRigVMMemoryHandleArray;
+typedef TArrayView<FRigVMPredicateBranch> FRigVMPredicateBranchArray;
 
-typedef void (*FRigVMFunctionPtr)(FRigVMExtendedExecuteContext& RigVMExecuteContext, FRigVMMemoryHandleArray RigVMMemoryHandles);
+typedef void (*FRigVMFunctionPtr)(FRigVMExtendedExecuteContext& RigVMExecuteContext, FRigVMMemoryHandleArray RigVMMemoryHandles, FRigVMPredicateBranchArray RigVMBranches);
 
 struct FRigVMDispatchFactory;
 struct FRigVMTemplate;
@@ -43,6 +45,18 @@ enum class ERigVMPinDirection : uint8
 	Invalid // The max value for this enum - used for guarding.
 };
 
+/**
+ * The Function Argument is used to differentiate different kinds of 
+ * pins in the data flow graph - inputs or outputs 
+ */
+UENUM(BlueprintType)
+enum class ERigVMFunctionArgumentDirection : uint8
+{
+	Input, // A const input value
+	Output, // A mutable output value
+	Invalid // The max value for this enum - used for guarding.
+};
+
 /**
  * The FRigVMFunctionArgument describes an argument necessary for the C++ invocation of the RIGVM_METHOD backed function
  */
@@ -52,22 +66,26 @@ struct RIGVM_API FRigVMFunctionArgument
 	const TCHAR* Type;
 	TSharedPtr<FString> NameString;
 	TSharedPtr<FString> TypeString;
+	ERigVMFunctionArgumentDirection Direction;
 
 	FRigVMFunctionArgument()
 		: Name(nullptr)
 		, Type(nullptr)
+		, Direction(ERigVMFunctionArgumentDirection::Invalid)
 	{
 	}
 
-	FRigVMFunctionArgument(const TCHAR* InName, const TCHAR* InType)
+	FRigVMFunctionArgument(const TCHAR* InName, const TCHAR* InType, const ERigVMFunctionArgumentDirection InDirection = ERigVMFunctionArgumentDirection::Input)
 		: Name(InName)
 		, Type(InType)
+		, Direction(InDirection)
 	{
 	}
 
-	FRigVMFunctionArgument(const FString& InName, const FString& InType)
+	FRigVMFunctionArgument(const FString& InName, const FString& InType, const ERigVMFunctionArgumentDirection InDirection = ERigVMFunctionArgumentDirection::Input)
 		: Name(nullptr)
 		, Type(nullptr)
+		, Direction(InDirection)
 	{
 		NameString = MakeShareable(new FString(InName));
 		TypeString = MakeShareable(new FString(InType));
@@ -76,6 +94,8 @@ struct RIGVM_API FRigVMFunctionArgument
 	}
 };
 
+typedef TArray<FRigVMFunctionArgument> FRigVMFunctionArgumentArray;
+
 /**
  * The FRigVMExecuteArgument describes an execute argument.
  */
diff --git a/Engine/Plugins/Runtime/RigVM/Source/RigVM/Public/RigVMCore/RigVMRegistry.h b/Engine/Plugins/Runtime/RigVM/Source/RigVM/Public/RigVMCore/RigVMRegistry.h
--- a/Engine/Plugins/Runtime/RigVM/Source/RigVM/Public/RigVMCore/RigVMRegistry.h
+++ b/Engine/Plugins/Runtime/RigVM/Source/RigVM/Public/RigVMCore/RigVMRegistry.h
@@ -59,6 +59,9 @@ public:
 	// Registers a dispatch factory given its struct.
 	const FRigVMDispatchFactory* RegisterFactory(UScriptStruct* InFactoryStruct);
 
+	// Register a predicate contained in the input struct
+	void RegisterPredicate(UScriptStruct* InStruct, const TCHAR* InName, const TArray<FRigVMFunctionArgument>& InArguments);
+
 	// Initializes the registry by storing the defaults
 	void InitializeIfNeeded();
 	
@@ -260,6 +263,9 @@ public:
 	// Returns all dispatch factories
 	const TArray<FRigVMDispatchFactory*>& GetFactories() const;
 
+	// Given a struct name, return the predicates
+	const TArray<FRigVMFunction>* GetPredicatesForStruct(const FName& InStructName) const;
+
 	static const TArray<UScriptStruct*>& GetMathTypes();
 
 	// Notifies other system that types have been added/removed, and template permutations have been updated
@@ -352,6 +358,9 @@ private:
 	// name lookup for functions
 	TMap<FName, int32> FunctionNameToIndex;
 
+	// lookup all the predicate functions of this struct
+	TMap<FName, TArray<FRigVMFunction>> StructNameToPredicates;
+
 	// name lookup for non-deprecated templates
 	TMap<FName, int32> TemplateNotationToIndex;
 
@@ -379,6 +388,7 @@ private:
 	static FCriticalSection FindTemplateMutex;
 	static FCriticalSection FindFactoryMutex;
 	static FCriticalSection GetDispatchFunctionMutex;
+	static FCriticalSection GetDispatchPredicatesMutex;
 	static FCriticalSection GetPermutationMutex;
 
 	friend struct FRigVMStruct;
diff --git a/Engine/Plugins/Runtime/RigVM/Source/RigVM/Public/RigVMCore/RigVMTemplate.h b/Engine/Plugins/Runtime/RigVM/Source/RigVM/Public/RigVMCore/RigVMTemplate.h
--- a/Engine/Plugins/Runtime/RigVM/Source/RigVM/Public/RigVMCore/RigVMTemplate.h
+++ b/Engine/Plugins/Runtime/RigVM/Source/RigVM/Public/RigVMCore/RigVMTemplate.h
@@ -34,6 +34,7 @@ typedef TMap<FName, TRigVMTypeIndex> FRigVMTemplateTypeMap;
 
 DECLARE_DELEGATE_RetVal_ThreeParams(FRigVMTemplateTypeMap, FRigVMTemplate_NewArgumentTypeDelegate, const FRigVMTemplate* /* InTemplate */, const FName& /* InArgumentName */, TRigVMTypeIndex /* InTypeIndexToAdd */);
 DECLARE_DELEGATE_RetVal_TwoParams(FRigVMFunctionPtr, FRigVMTemplate_RequestDispatchFunctionDelegate, const FRigVMTemplate* /* InTemplate */,  const FRigVMTemplateTypeMap& /* InTypes */);
+DECLARE_DELEGATE_RetVal_TwoParams(TArray<FRigVMFunction>, FRigVMTemplate_RequestDispatchPredicatesDelegate, const FRigVMTemplate* /* InTemplate */,  const FRigVMTemplateTypeMap& /* InTypes */);
 DECLARE_DELEGATE_RetVal(FRigVMDispatchFactory*, FRigVMTemplate_GetDispatchFactoryDelegate);
 
 struct RIGVM_API FRigVMTemplateDelegates
@@ -41,6 +42,7 @@ struct RIGVM_API FRigVMTemplateDelegates
 	FRigVMTemplate_NewArgumentTypeDelegate NewArgumentTypeDelegate;
 	FRigVMTemplate_GetDispatchFactoryDelegate GetDispatchFactoryDelegate;
 	FRigVMTemplate_RequestDispatchFunctionDelegate RequestDispatchFunctionDelegate;
+	FRigVMTemplate_RequestDispatchPredicatesDelegate RequestDispatchPredicatesDelegate;
 };
 
 USTRUCT()
diff --git a/Engine/Plugins/Runtime/RigVM/Source/RigVM/Public/RigVMFunctions/RigVMDispatch_Array.h b/Engine/Plugins/Runtime/RigVM/Source/RigVM/Public/RigVMFunctions/RigVMDispatch_Array.h
--- a/Engine/Plugins/Runtime/RigVM/Source/RigVM/Public/RigVMFunctions/RigVMDispatch_Array.h
+++ b/Engine/Plugins/Runtime/RigVM/Source/RigVM/Public/RigVMFunctions/RigVMDispatch_Array.h
@@ -74,7 +74,7 @@ public:
 
 protected:
 	virtual FRigVMFunctionPtr GetDispatchFunctionImpl(const FRigVMTemplateTypeMap& InTypes) const override { return &FRigVMDispatch_ArrayMake::Execute; }
-	static void Execute(FRigVMExtendedExecuteContext& InContext, FRigVMMemoryHandleArray Handles);
+	static void Execute(FRigVMExtendedExecuteContext& InContext, FRigVMMemoryHandleArray Handles, FRigVMPredicateBranchArray Predicates);
 };
 
 USTRUCT(meta=(DisplayName = "Reset"))
@@ -99,7 +99,7 @@ public:
 
 protected:
 	virtual FRigVMFunctionPtr GetDispatchFunctionImpl(const FRigVMTemplateTypeMap& InTypes) const override { return &FRigVMDispatch_ArrayReset::Execute; }
-	static void Execute(FRigVMExtendedExecuteContext& InContext, FRigVMMemoryHandleArray Handles);
+	static void Execute(FRigVMExtendedExecuteContext& InContext, FRigVMMemoryHandleArray Handles, FRigVMPredicateBranchArray Predicates);
 };
 
 USTRUCT(meta=(DisplayName = "Num"))
@@ -124,7 +124,7 @@ public:
 
 protected:
 	virtual FRigVMFunctionPtr GetDispatchFunctionImpl(const FRigVMTemplateTypeMap& InTypes) const override { return &FRigVMDispatch_ArrayGetNum::Execute; }
-	static void Execute(FRigVMExtendedExecuteContext& InContext, FRigVMMemoryHandleArray Handles);
+	static void Execute(FRigVMExtendedExecuteContext& InContext, FRigVMMemoryHandleArray Handles, FRigVMPredicateBranchArray Predicates);
 };
 
 USTRUCT(meta=(DisplayName = "Set Num"))
@@ -149,7 +149,7 @@ public:
 
 protected:
 	virtual FRigVMFunctionPtr GetDispatchFunctionImpl(const FRigVMTemplateTypeMap& InTypes) const override { return &FRigVMDispatch_ArraySetNum::Execute; }
-	static void Execute(FRigVMExtendedExecuteContext& InContext, FRigVMMemoryHandleArray Handles);
+	static void Execute(FRigVMExtendedExecuteContext& InContext, FRigVMMemoryHandleArray Handles, FRigVMPredicateBranchArray Predicates);
 };
 
 USTRUCT(meta=(DisplayName = "At"))
@@ -174,7 +174,7 @@ public:
 
 protected:
 	virtual FRigVMFunctionPtr GetDispatchFunctionImpl(const FRigVMTemplateTypeMap& InTypes) const override { return &FRigVMDispatch_ArrayGetAtIndex::Execute; }
-	static void Execute(FRigVMExtendedExecuteContext& InContext, FRigVMMemoryHandleArray Handles);
+	static void Execute(FRigVMExtendedExecuteContext& InContext, FRigVMMemoryHandleArray Handles, FRigVMPredicateBranchArray Predicates);
 };
 
 USTRUCT(meta=(DisplayName = "Set At"))
@@ -199,7 +199,7 @@ public:
 
 protected:
 	virtual FRigVMFunctionPtr GetDispatchFunctionImpl(const FRigVMTemplateTypeMap& InTypes) const override { return &FRigVMDispatch_ArraySetAtIndex::Execute; }
-	static void Execute(FRigVMExtendedExecuteContext& InContext, FRigVMMemoryHandleArray Handles);
+	static void Execute(FRigVMExtendedExecuteContext& InContext, FRigVMMemoryHandleArray Handles, FRigVMPredicateBranchArray Predicates);
 };
 
 USTRUCT(meta=(DisplayName = "Add"))
@@ -223,7 +223,7 @@ public:
 
 protected:
 	virtual FRigVMFunctionPtr GetDispatchFunctionImpl(const FRigVMTemplateTypeMap& InTypes) const override { return &FRigVMDispatch_ArrayAdd::Execute; }
-	static void Execute(FRigVMExtendedExecuteContext& InContext, FRigVMMemoryHandleArray Handles);
+	static void Execute(FRigVMExtendedExecuteContext& InContext, FRigVMMemoryHandleArray Handles, FRigVMPredicateBranchArray Predicates);
 };
 
 USTRUCT(meta=(DisplayName = "Insert"))
@@ -245,7 +245,7 @@ public:
 
 protected:
 	virtual FRigVMFunctionPtr GetDispatchFunctionImpl(const FRigVMTemplateTypeMap& InTypes) const override { return &FRigVMDispatch_ArrayInsert::Execute; }
-	static void Execute(FRigVMExtendedExecuteContext& InContext, FRigVMMemoryHandleArray Handles);
+	static void Execute(FRigVMExtendedExecuteContext& InContext, FRigVMMemoryHandleArray Handles, FRigVMPredicateBranchArray Predicates);
 };
 
 USTRUCT(meta=(DisplayName = "Remove"))
@@ -270,7 +270,7 @@ public:
 
 protected:
 	virtual FRigVMFunctionPtr GetDispatchFunctionImpl(const FRigVMTemplateTypeMap& InTypes) const override { return &FRigVMDispatch_ArrayRemove::Execute; }
-	static void Execute(FRigVMExtendedExecuteContext& InContext, FRigVMMemoryHandleArray Handles);
+	static void Execute(FRigVMExtendedExecuteContext& InContext, FRigVMMemoryHandleArray Handles, FRigVMPredicateBranchArray Predicates);
 };
 
 USTRUCT(meta=(DisplayName = "Reverse"))
@@ -293,7 +293,7 @@ public:
 protected:
 
 	virtual FRigVMFunctionPtr GetDispatchFunctionImpl(const FRigVMTemplateTypeMap& InTypes) const override { return &FRigVMDispatch_ArrayReverse::Execute; }
-	static void Execute(FRigVMExtendedExecuteContext& InContext, FRigVMMemoryHandleArray Handles);
+	static void Execute(FRigVMExtendedExecuteContext& InContext, FRigVMMemoryHandleArray Handles, FRigVMPredicateBranchArray Predicates);
 };
 
 USTRUCT(meta=(DisplayName = "Find"))
@@ -318,7 +318,7 @@ public:
 
 protected:
 	virtual FRigVMFunctionPtr GetDispatchFunctionImpl(const FRigVMTemplateTypeMap& InTypes) const override { return &FRigVMDispatch_ArrayFind::Execute; }
-	static void Execute(FRigVMExtendedExecuteContext& InContext, FRigVMMemoryHandleArray Handles);
+	static void Execute(FRigVMExtendedExecuteContext& InContext, FRigVMMemoryHandleArray Handles, FRigVMPredicateBranchArray Predicates);
 };
 
 USTRUCT(meta=(DisplayName = "Append"))
@@ -343,7 +343,7 @@ public:
 
 protected:
 	virtual FRigVMFunctionPtr GetDispatchFunctionImpl(const FRigVMTemplateTypeMap& InTypes) const override { return &FRigVMDispatch_ArrayAppend::Execute; }
-	static void Execute(FRigVMExtendedExecuteContext& InContext, FRigVMMemoryHandleArray Handles);
+	static void Execute(FRigVMExtendedExecuteContext& InContext, FRigVMMemoryHandleArray Handles, FRigVMPredicateBranchArray Predicates);
 };
 
 USTRUCT(meta=(DisplayName = "Clone"))
@@ -368,7 +368,7 @@ public:
 
 protected:
 	virtual FRigVMFunctionPtr GetDispatchFunctionImpl(const FRigVMTemplateTypeMap& InTypes) const override { return &FRigVMDispatch_ArrayClone::Execute; }
-	static void Execute(FRigVMExtendedExecuteContext& InContext, FRigVMMemoryHandleArray Handles);
+	static void Execute(FRigVMExtendedExecuteContext& InContext, FRigVMMemoryHandleArray Handles, FRigVMPredicateBranchArray Predicates);
 };
 
 USTRUCT(meta=(DisplayName = "Union"))
@@ -390,7 +390,7 @@ public:
 
 protected:
 	virtual FRigVMFunctionPtr GetDispatchFunctionImpl(const FRigVMTemplateTypeMap& InTypes) const override { return &FRigVMDispatch_ArrayUnion::Execute; }
-	static void Execute(FRigVMExtendedExecuteContext& InContext, FRigVMMemoryHandleArray Handles);
+	static void Execute(FRigVMExtendedExecuteContext& InContext, FRigVMMemoryHandleArray Handles, FRigVMPredicateBranchArray Predicates);
 };
 
 USTRUCT(meta=(DisplayName = "Difference"))
@@ -415,7 +415,7 @@ public:
 
 protected:
 	virtual FRigVMFunctionPtr GetDispatchFunctionImpl(const FRigVMTemplateTypeMap& InTypes) const override { return &FRigVMDispatch_ArrayDifference::Execute; }
-	static void Execute(FRigVMExtendedExecuteContext& InContext, FRigVMMemoryHandleArray Handles);
+	static void Execute(FRigVMExtendedExecuteContext& InContext, FRigVMMemoryHandleArray Handles, FRigVMPredicateBranchArray Predicates);
 };
 
 USTRUCT(meta=(DisplayName = "Intersection"))
@@ -437,7 +437,7 @@ public:
 
 protected:
 	virtual FRigVMFunctionPtr GetDispatchFunctionImpl(const FRigVMTemplateTypeMap& InTypes) const override { return &FRigVMDispatch_ArrayIntersection::Execute; }
-	static void Execute(FRigVMExtendedExecuteContext& InContext, FRigVMMemoryHandleArray Handles);
+	static void Execute(FRigVMExtendedExecuteContext& InContext, FRigVMMemoryHandleArray Handles, FRigVMPredicateBranchArray Predicates);
 };
 
 USTRUCT(meta=(DisplayName = "For Each", Icon="EditorStyle|GraphEditor.Macro.ForEach_16x"))
@@ -465,5 +465,5 @@ public:
 
 protected:
 	virtual FRigVMFunctionPtr GetDispatchFunctionImpl(const FRigVMTemplateTypeMap& InTypes) const override { return &FRigVMDispatch_ArrayIterator::Execute; }
-	static void Execute(FRigVMExtendedExecuteContext& InContext, FRigVMMemoryHandleArray Handles);
+	static void Execute(FRigVMExtendedExecuteContext& InContext, FRigVMMemoryHandleArray Handles, FRigVMPredicateBranchArray Predicates);
 };
diff --git a/Engine/Plugins/Runtime/RigVM/Source/RigVM/Public/RigVMFunctions/RigVMDispatch_Constant.h b/Engine/Plugins/Runtime/RigVM/Source/RigVM/Public/RigVMFunctions/RigVMDispatch_Constant.h
--- a/Engine/Plugins/Runtime/RigVM/Source/RigVM/Public/RigVMFunctions/RigVMDispatch_Constant.h
+++ b/Engine/Plugins/Runtime/RigVM/Source/RigVM/Public/RigVMFunctions/RigVMDispatch_Constant.h
@@ -28,7 +28,7 @@ public:
 	
 protected:
 	virtual FRigVMFunctionPtr GetDispatchFunctionImpl(const FRigVMTemplateTypeMap& InTypes) const override { return &FRigVMDispatch_Constant::Execute; }
-	static void Execute(FRigVMExtendedExecuteContext& InContext, FRigVMMemoryHandleArray Handles);
+	static void Execute(FRigVMExtendedExecuteContext& InContext, FRigVMMemoryHandleArray Handles, FRigVMPredicateBranchArray RigVMBranches);
 
 	static const FName ValueName;
 
diff --git a/Engine/Plugins/Runtime/RigVM/Source/RigVM/Public/RigVMFunctions/RigVMDispatch_Core.h b/Engine/Plugins/Runtime/RigVM/Source/RigVM/Public/RigVMFunctions/RigVMDispatch_Core.h
--- a/Engine/Plugins/Runtime/RigVM/Source/RigVM/Public/RigVMFunctions/RigVMDispatch_Core.h
+++ b/Engine/Plugins/Runtime/RigVM/Source/RigVM/Public/RigVMFunctions/RigVMDispatch_Core.h
@@ -34,10 +34,10 @@ protected:
 
 	virtual FRigVMFunctionPtr GetDispatchFunctionImpl(const FRigVMTemplateTypeMap& InTypes) const override;
 	static bool AdaptResult(bool bResult, const FRigVMExtendedExecuteContext& InContext);
-	static void Execute(FRigVMExtendedExecuteContext& InContext, FRigVMMemoryHandleArray Handles);
+	static void Execute(FRigVMExtendedExecuteContext& InContext, FRigVMMemoryHandleArray Handles, FRigVMPredicateBranchArray Predicates);
 
 	template<typename T>
-	static void Equals(FRigVMExtendedExecuteContext& InContext, FRigVMMemoryHandleArray Handles)
+	static void Equals(FRigVMExtendedExecuteContext& InContext, FRigVMMemoryHandleArray Handles, FRigVMPredicateBranchArray Predicates)
 	{
 		check(Handles[0].IsType<T>());
 		check(Handles[1].IsType<T>());
@@ -48,7 +48,7 @@ protected:
 		Result = AdaptResult(A == B, InContext);
 	}
 
-	static void NameEquals(FRigVMExtendedExecuteContext& InContext, FRigVMMemoryHandleArray Handles)
+	static void NameEquals(FRigVMExtendedExecuteContext& InContext, FRigVMMemoryHandleArray Handles, FRigVMPredicateBranchArray Predicates)
 	{
 		check(Handles[0].IsType<FName>());
 		check(Handles[1].IsType<FName>());
@@ -59,7 +59,7 @@ protected:
 		Result = AdaptResult(A.IsEqual(B, ENameCase::CaseSensitive), InContext);
 	}
 
-	static void StringEquals(FRigVMExtendedExecuteContext& InContext, FRigVMMemoryHandleArray Handles)
+	static void StringEquals(FRigVMExtendedExecuteContext& InContext, FRigVMMemoryHandleArray Handles, FRigVMPredicateBranchArray Predicates)
 	{
 		check(Handles[0].IsType<FString>());
 		check(Handles[1].IsType<FString>());
@@ -71,7 +71,7 @@ protected:
 	}
 
 	template<typename T>
-	static void MathTypeEquals(FRigVMExtendedExecuteContext& InContext, FRigVMMemoryHandleArray Handles)
+	static void MathTypeEquals(FRigVMExtendedExecuteContext& InContext, FRigVMMemoryHandleArray Handles, FRigVMPredicateBranchArray Predicates)
 	{
 		check(Handles[0].IsType<T>());
 		check(Handles[1].IsType<T>());
diff --git a/Engine/Plugins/Runtime/RigVM/Source/RigVM/Public/RigVMFunctions/RigVMDispatch_If.h b/Engine/Plugins/Runtime/RigVM/Source/RigVM/Public/RigVMFunctions/RigVMDispatch_If.h
--- a/Engine/Plugins/Runtime/RigVM/Source/RigVM/Public/RigVMFunctions/RigVMDispatch_If.h
+++ b/Engine/Plugins/Runtime/RigVM/Source/RigVM/Public/RigVMFunctions/RigVMDispatch_If.h
@@ -33,7 +33,7 @@ public:
 protected:
 
 	virtual FRigVMFunctionPtr GetDispatchFunctionImpl(const FRigVMTemplateTypeMap& InTypes) const override;
-	static void Execute(FRigVMExtendedExecuteContext& InContext, FRigVMMemoryHandleArray Handles);
+	static void Execute(FRigVMExtendedExecuteContext& InContext, FRigVMMemoryHandleArray Handles, FRigVMPredicateBranchArray Predicates);
 
 	static const FName ConditionName;
 	static const FName TrueName;
diff --git a/Engine/Plugins/Runtime/RigVM/Source/RigVM/Public/RigVMFunctions/RigVMDispatch_MakeStruct.h b/Engine/Plugins/Runtime/RigVM/Source/RigVM/Public/RigVMFunctions/RigVMDispatch_MakeStruct.h
--- a/Engine/Plugins/Runtime/RigVM/Source/RigVM/Public/RigVMFunctions/RigVMDispatch_MakeStruct.h
+++ b/Engine/Plugins/Runtime/RigVM/Source/RigVM/Public/RigVMFunctions/RigVMDispatch_MakeStruct.h
@@ -28,7 +28,7 @@ public:
 
 protected:
 	virtual FRigVMFunctionPtr GetDispatchFunctionImpl(const FRigVMTemplateTypeMap& InTypes) const override { return &FRigVMDispatch_MakeStruct::Execute; }
-	static void Execute(FRigVMExtendedExecuteContext& InContext, FRigVMMemoryHandleArray Handles);
+	static void Execute(FRigVMExtendedExecuteContext& InContext, FRigVMMemoryHandleArray Handles, FRigVMPredicateBranchArray RigVMBranches);
 
 	static const FName ElementsName;
 	static const FName StructName;
diff --git a/Engine/Plugins/Runtime/RigVM/Source/RigVM/Public/RigVMFunctions/RigVMDispatch_Print.h b/Engine/Plugins/Runtime/RigVM/Source/RigVM/Public/RigVMFunctions/RigVMDispatch_Print.h
--- a/Engine/Plugins/Runtime/RigVM/Source/RigVM/Public/RigVMFunctions/RigVMDispatch_Print.h
+++ b/Engine/Plugins/Runtime/RigVM/Source/RigVM/Public/RigVMFunctions/RigVMDispatch_Print.h
@@ -37,7 +37,7 @@ protected:
 	{
 		return &FRigVMDispatch_Print::Execute;
 	}
-	static void Execute(FRigVMExtendedExecuteContext& InContext, FRigVMMemoryHandleArray Handles);
+	static void Execute(FRigVMExtendedExecuteContext& InContext, FRigVMMemoryHandleArray Handles, FRigVMPredicateBranchArray Predicates);
 
 	static const FName PrefixName;
 	static const FName ValueName;
diff --git a/Engine/Plugins/Runtime/RigVM/Source/RigVM/Public/RigVMFunctions/RigVMDispatch_Select.h b/Engine/Plugins/Runtime/RigVM/Source/RigVM/Public/RigVMFunctions/RigVMDispatch_Select.h
--- a/Engine/Plugins/Runtime/RigVM/Source/RigVM/Public/RigVMFunctions/RigVMDispatch_Select.h
+++ b/Engine/Plugins/Runtime/RigVM/Source/RigVM/Public/RigVMFunctions/RigVMDispatch_Select.h
@@ -33,7 +33,7 @@ public:
 protected:
 
 	virtual FRigVMFunctionPtr GetDispatchFunctionImpl(const FRigVMTemplateTypeMap& InTypes) const override;
-	static void Execute(FRigVMExtendedExecuteContext& InContext, FRigVMMemoryHandleArray Handles);
+	static void Execute(FRigVMExtendedExecuteContext& InContext, FRigVMMemoryHandleArray Handles, FRigVMPredicateBranchArray Predicates);
 
 	static const FName IndexName;
 	static const FName ValuesName;
diff --git a/Engine/Plugins/Runtime/RigVM/Source/RigVM/Public/RigVMFunctions/RigVMDispatch_Switch.h b/Engine/Plugins/Runtime/RigVM/Source/RigVM/Public/RigVMFunctions/RigVMDispatch_Switch.h
--- a/Engine/Plugins/Runtime/RigVM/Source/RigVM/Public/RigVMFunctions/RigVMDispatch_Switch.h
+++ b/Engine/Plugins/Runtime/RigVM/Source/RigVM/Public/RigVMFunctions/RigVMDispatch_Switch.h
@@ -37,7 +37,7 @@ public:
 protected:
 
 	virtual FRigVMFunctionPtr GetDispatchFunctionImpl(const FRigVMTemplateTypeMap& InTypes) const override;
-	static void Execute(FRigVMExtendedExecuteContext& InContext, FRigVMMemoryHandleArray Handles);
+	static void Execute(FRigVMExtendedExecuteContext& InContext, FRigVMMemoryHandleArray Handles, FRigVMPredicateBranchArray Predicates);
 	static FName GetCaseName(int32 InIndex);
 	static FName GetCaseDisplayName(int32 InIndex);
 	
diff --git a/Engine/Plugins/Runtime/RigVM/Source/RigVM/Public/RigVMFunctions/RigVMFunction_String.h b/Engine/Plugins/Runtime/RigVM/Source/RigVM/Public/RigVMFunctions/RigVMFunction_String.h
--- a/Engine/Plugins/Runtime/RigVM/Source/RigVM/Public/RigVMFunctions/RigVMFunction_String.h
+++ b/Engine/Plugins/Runtime/RigVM/Source/RigVM/Public/RigVMFunctions/RigVMFunction_String.h
@@ -516,7 +516,7 @@ public:
 protected:
 
 	virtual FRigVMFunctionPtr GetDispatchFunctionImpl(const FRigVMTemplateTypeMap& InTypes) const override;
-	static void Execute(FRigVMExtendedExecuteContext& InContext, FRigVMMemoryHandleArray Handles);
+	static void Execute(FRigVMExtendedExecuteContext& InContext, FRigVMMemoryHandleArray Handles, FRigVMPredicateBranchArray Predicates);
 };
 
 /*
@@ -535,5 +535,5 @@ public:
 protected:
 
 	virtual FRigVMFunctionPtr GetDispatchFunctionImpl(const FRigVMTemplateTypeMap& InTypes) const override;
-	static void Execute(FRigVMExtendedExecuteContext& InContext, FRigVMMemoryHandleArray Handles);
+	static void Execute(FRigVMExtendedExecuteContext& InContext, FRigVMMemoryHandleArray Handles, FRigVMPredicateBranchArray Predicates);
 };
\ No newline at end of file
diff --git a/Engine/Plugins/Runtime/RigVM/Source/RigVM/Public/RigVMObjectVersion.h b/Engine/Plugins/Runtime/RigVM/Source/RigVM/Public/RigVMObjectVersion.h
--- a/Engine/Plugins/Runtime/RigVM/Source/RigVM/Public/RigVMObjectVersion.h
+++ b/Engine/Plugins/Runtime/RigVM/Source/RigVM/Public/RigVMObjectVersion.h
@@ -14,6 +14,9 @@ struct RIGVM_API FRigVMObjectVersion
 		
 		AddedVMHashChecks,
 
+		// Predicates added to execute operations
+		PredicatesAddedToExecuteOps,
+
 		// -----<new versions can be added above this line>-------------------------------------------------
 		VersionPlusOne,
 		LatestVersion = VersionPlusOne - 1
diff --git a/Engine/Plugins/Runtime/RigVM/Source/RigVMDeveloper/Private/RigVMCompiler/RigVMCompiler.cpp b/Engine/Plugins/Runtime/RigVM/Source/RigVMDeveloper/Private/RigVMCompiler/RigVMCompiler.cpp
--- a/Engine/Plugins/Runtime/RigVM/Source/RigVMDeveloper/Private/RigVMCompiler/RigVMCompiler.cpp
+++ b/Engine/Plugins/Runtime/RigVM/Source/RigVMDeveloper/Private/RigVMCompiler/RigVMCompiler.cpp
@@ -1490,7 +1490,7 @@ void URigVMCompiler::TraverseEntry(const FRigVMEntryExprAST* InExpr, FRigVMCompi
 
 			// setup the instruction
 			int32 FunctionIndex = WorkData.VM->AddRigVMFunction(UnitNode->GetScriptStruct(), UnitNode->GetMethodName());
-			WorkData.VM->GetByteCode().AddExecuteOp(FunctionIndex, Operands);
+			WorkData.VM->GetByteCode().AddExecuteOp(FunctionIndex, Operands, 0, 0);
 		
 			int32 EntryInstructionIndex = WorkData.VM->GetByteCode().GetNumInstructions() - 1;
 			FName Entryname = UnitNode->GetEventName();
@@ -1551,6 +1551,7 @@ int32 URigVMCompiler::TraverseCallExtern(const FRigVMCallExternExprAST* InExpr,
 	URigVMNode* Node = InExpr->GetNode();
 	URigVMUnitNode* UnitNode = Cast<URigVMUnitNode>(Node);
 	URigVMDispatchNode* DispatchNode = Cast<URigVMDispatchNode>(Node);
+	const FRigVMRegistry& Registry = FRigVMRegistry::Get();
 	if(!ValidateNode(UnitNode, false) && !ValidateNode(DispatchNode, false))
 	{
 		return INDEX_NONE;
@@ -1628,7 +1629,7 @@ int32 URigVMCompiler::TraverseCallExtern(const FRigVMCallExternExprAST* InExpr,
 
 		if(UnitNode)
 		{
-			Function = FRigVMRegistry::Get().FindFunction(UnitNode->GetScriptStruct(), *UnitNode->GetMethodName().ToString());
+			Function = Registry.FindFunction(UnitNode->GetScriptStruct(), *UnitNode->GetMethodName().ToString());
 		}
 		else if(DispatchNode)
 		{
@@ -1766,6 +1767,29 @@ int32 URigVMCompiler::TraverseCallExtern(const FRigVMCallExternExprAST* InExpr,
 			WorkData.VM->GetByteCode().GetOpAt<FRigVMJumpOp>(JumpToCallExternByte).InstructionIndex = InstructionsToJump;
 		}
 
+		int32 StartPredicateIndex = INDEX_NONE;
+		int32 PredicateCount = 0;
+		const TArray<FRigVMFunction>* Predicates = nullptr;
+		if (UnitNode)
+		{
+			Predicates = Registry.GetPredicatesForStruct(Function->Struct->GetFName());
+		}
+		else if(DispatchNode)
+		{
+			FRigVMTemplateTypeMap TypeMap = DispatchNode->GetTemplatePinTypeMap(true);
+			const FString PermutationName = DispatchNode->GetFactory()->GetPermutationName(TypeMap);
+			Predicates = Registry.GetPredicatesForStruct(*PermutationName);
+		}
+		if (Predicates)
+		{
+			StartPredicateIndex = WorkData.VM->GetByteCode().PredicateBranches.Num();
+			for (const FRigVMFunction& Predicate : *Predicates)
+			{
+				WorkData.VM->GetByteCode().AddPredicateBranch(FRigVMPredicateBranch());
+				PredicateCount++;
+			}
+		}
+		
 		if(Node->IsControlFlowNode())
 		{
 			check(BlockToRunOperand.IsValid());
@@ -1780,7 +1804,7 @@ int32 URigVMCompiler::TraverseCallExtern(const FRigVMCallExternExprAST* InExpr,
 		// setup the instruction
 		const int32 FunctionIndex = WorkData.VM->AddRigVMFunction(Function->GetName());
 		check(FunctionIndex != INDEX_NONE);
-		WorkData.VM->GetByteCode().AddExecuteOp(FunctionIndex, Operands);
+		WorkData.VM->GetByteCode().AddExecuteOp(FunctionIndex, Operands, StartPredicateIndex, PredicateCount);
 		CallExternInstructionIndex = WorkData.VM->GetByteCode().GetNumInstructions() - 1;
 
 		// setup the branch infos for this call extern instruction
@@ -2525,7 +2549,7 @@ void URigVMCompiler::AddCopyOperator(const FRigVMCopyOp& InOp, const FRigVMAssig
 			const FRigVMOperand Target = InOp.Target;
 
 			const int32 FunctionIndex = WorkData.VM->AddRigVMFunction(CastFunction->Name);
-			WorkData.VM->GetByteCode().AddExecuteOp(FunctionIndex, {Source, Target});
+			WorkData.VM->GetByteCode().AddExecuteOp(FunctionIndex, {Source, Target}, 0, 0);
 
 			bAddCopyOp = false;
 		}
@@ -2545,7 +2569,7 @@ void URigVMCompiler::AddCopyOperator(const FRigVMCopyOp& InOp, const FRigVMAssig
 					static const FString ArrayCloneName =
 						FRigVMRegistry::Get().FindOrAddSingletonDispatchFunction<FRigVMDispatch_ArrayClone>();
 					const int32 FunctionIndex = WorkData.VM->AddRigVMFunction(ArrayCloneName);
-					WorkData.VM->GetByteCode().AddExecuteOp(FunctionIndex, {InOp.Source, InOp.Target});
+					WorkData.VM->GetByteCode().AddExecuteOp(FunctionIndex, {InOp.Source, InOp.Target}, 0, 0);
 					bAddCopyOp = false;
 				}
 			}
diff --git a/Engine/Source/Programs/Shared/EpicGames.UHT/Exporters/CodeGen/UhtHeaderCodeGeneratorCppFile.cs b/Engine/Source/Programs/Shared/EpicGames.UHT/Exporters/CodeGen/UhtHeaderCodeGeneratorCppFile.cs
--- a/Engine/Source/Programs/Shared/EpicGames.UHT/Exporters/CodeGen/UhtHeaderCodeGeneratorCppFile.cs
+++ b/Engine/Source/Programs/Shared/EpicGames.UHT/Exporters/CodeGen/UhtHeaderCodeGeneratorCppFile.cs
@@ -555,32 +555,67 @@ namespace EpicGames.UHT.Exporters.CodeGen
 				{
 					foreach (UhtRigVMMethodInfo methodInfo in scriptStruct.RigVMStructInfo.Methods)
 					{
-						builder.Append("\t\tTArray<FRigVMFunctionArgument> ").AppendArgumentsName(scriptStruct, methodInfo).Append(";\r\n");
-						foreach (UhtRigVMParameter parameter in scriptStruct.RigVMStructInfo.Members)
+						if (methodInfo.IsPredicate)
 						{
+							builder.Append("\t\tTArray<FRigVMFunctionArgument> ").AppendArgumentsName(scriptStruct, methodInfo).Append(";\r\n");
+							foreach (UhtRigVMParameter parameter in methodInfo.Parameters)
+							{
+								builder
+									.Append("\t\t")
+									.AppendArgumentsName(scriptStruct, methodInfo)
+									.Append(".Emplace(TEXT(\"")
+									.Append(parameter.NameOriginal())
+									.Append("\"), TEXT(\"")
+									.Append(parameter.TypeOriginal())
+									.Append("\"));\r\n");
+							}
 							builder
 								.Append("\t\t")
 								.AppendArgumentsName(scriptStruct, methodInfo)
-								.Append(".Emplace(TEXT(\"")
-								.Append(parameter.NameOriginal())
-								.Append("\"), TEXT(\"")
-								.Append(parameter.TypeOriginal())
-								.Append("\"));\r\n");
+								.Append(".Emplace(TEXT(\"Return\"), TEXT(\"")
+								.Append(methodInfo.ReturnType)
+								.Append("\"), ERigVMFunctionArgumentDirection::Output);\r\n");
+							
+							builder
+								.Append("\t\tFRigVMRegistry::Get().RegisterPredicate(")
+								.Append(registrationName)
+								.Append(".OuterSingleton, ")
+								.Append("TEXT(\"")
+								.Append(methodInfo.Name)
+								.Append("\"), ")
+								.AppendArgumentsName(scriptStruct, methodInfo)
+								.Append(");\r\n");
+						}
+						else
+						{
+							builder.Append("\t\tTArray<FRigVMFunctionArgument> ").AppendArgumentsName(scriptStruct, methodInfo).Append(";\r\n");
+							foreach (UhtRigVMParameter parameter in scriptStruct.RigVMStructInfo.Members)
+							{
+								builder
+									.Append("\t\t")
+									.AppendArgumentsName(scriptStruct, methodInfo)
+									.Append(".Emplace(TEXT(\"")
+									.Append(parameter.NameOriginal())
+									.Append("\"), TEXT(\"")
+									.Append(parameter.TypeOriginal())
+									.Append("\"));\r\n");
+							}
+							
+							builder
+								.Append("\t\tFRigVMRegistry::Get().Register(TEXT(\"")
+								.Append(scriptStruct.SourceName)
+								.Append("::")
+								.Append(methodInfo.Name)
+								.Append("\"), &")
+								.Append(scriptStruct.SourceName)
+								.Append("::RigVM")
+								.Append(methodInfo.Name)
+								.Append(", ")
+								.Append(registrationName)
+								.Append(".OuterSingleton, ")
+								.AppendArgumentsName(scriptStruct, methodInfo)
+								.Append(");\r\n");
 						}
-						builder
-							.Append("\t\tFRigVMRegistry::Get().Register(TEXT(\"")
-							.Append(scriptStruct.SourceName)
-							.Append("::")
-							.Append(methodInfo.Name)
-							.Append("\"), &")
-							.Append(scriptStruct.SourceName)
-							.Append("::RigVM")
-							.Append(methodInfo.Name)
-							.Append(", ")
-							.Append(registrationName)
-							.Append(".OuterSingleton, ")
-							.AppendArgumentsName(scriptStruct, methodInfo)
-							.Append(");\r\n");
 					}
 				}
 
@@ -714,6 +749,11 @@ namespace EpicGames.UHT.Exporters.CodeGen
 				
 				foreach (UhtRigVMMethodInfo methodInfo in scriptStruct.RigVMStructInfo.Methods)
 				{
+					if (methodInfo.IsPredicate)
+					{
+						continue;
+					}
+					
 					builder.Append("\r\n");
 
 					builder
@@ -771,11 +811,28 @@ namespace EpicGames.UHT.Exporters.CodeGen
 						//COMPATIBILITY-TODO - Remove the tab
 						builder.Append("\t\r\n");
 					}
+					
+					foreach (UhtRigVMMethodInfo predicateInfo in scriptStruct.RigVMStructInfo.Methods)
+					{
+						if (predicateInfo.IsPredicate)
+						{
+							builder.Append("\t").Append(predicateInfo.Name).Append("Struct ").Append(predicateInfo.Name).Append("Predicate; \r\n");
+						}
+					}
 
 					//COMPATIBILITY-TODO - Replace spaces with \t
 					builder.Append('\t').Append(methodInfo.ReturnPrefix()).Append("Static").Append(methodInfo.Name).Append("(\r\n");
 					builder.Append("\t\tInExecuteContext");
 					builder.AppendParameterNames(scriptStruct.RigVMStructInfo.Members, true, ",\r\n\t\t", true);
+					
+					foreach (UhtRigVMMethodInfo predicateInfo in scriptStruct.RigVMStructInfo.Methods)
+					{
+						if (predicateInfo.IsPredicate)
+						{
+							builder.Append(", \r\n\t\t").Append(predicateInfo.Name).Append("Predicate");
+						}
+					}
+					
 					builder.Append("\r\n");
 					builder.Append("\t);\r\n");
 					builder.Append("}\r\n");
diff --git a/Engine/Source/Programs/Shared/EpicGames.UHT/Exporters/CodeGen/UhtHeaderCodeGeneratorHFile.cs b/Engine/Source/Programs/Shared/EpicGames.UHT/Exporters/CodeGen/UhtHeaderCodeGeneratorHFile.cs
--- a/Engine/Source/Programs/Shared/EpicGames.UHT/Exporters/CodeGen/UhtHeaderCodeGeneratorHFile.cs
+++ b/Engine/Source/Programs/Shared/EpicGames.UHT/Exporters/CodeGen/UhtHeaderCodeGeneratorHFile.cs
@@ -8,6 +8,7 @@ using System.Text;
 using EpicGames.Core;
 using EpicGames.UHT.Types;
 using EpicGames.UHT.Utils;
+using Microsoft.Extensions.Primitives;
 
 namespace EpicGames.UHT.Exporters.CodeGen
 {
@@ -142,6 +143,11 @@ namespace EpicGames.UHT.Exporters.CodeGen
 				builder.Append("\r\n");
 				foreach (UhtRigVMMethodInfo methodInfo in scriptStruct.RigVMStructInfo.Methods)
 				{
+					if (methodInfo.IsPredicate)
+					{
+						continue;
+					}
+					
 					builder.Append("#define ").Append(scriptStruct.SourceName).Append('_').Append(methodInfo.Name).Append("() \\\r\n");
 					builder.Append('\t').Append(methodInfo.ReturnType).Append(' ').Append(scriptStruct.SourceName).Append("::Static").Append(methodInfo.Name).Append("( \\\r\n");
 					builder.Append("\t\t");
@@ -151,6 +157,13 @@ namespace EpicGames.UHT.Exporters.CodeGen
 					}
 					builder.Append(scriptStruct.RigVMStructInfo.ExecuteContextType).Append("& ").Append(RigVMExecuteContextParamName);
 					builder.AppendParameterDecls(scriptStruct.RigVMStructInfo.Members, true, ", \\\r\n\t\t", true, false);
+					foreach (UhtRigVMMethodInfo predicateInfo in scriptStruct.RigVMStructInfo.Methods)
+					{
+						if (predicateInfo.IsPredicate)
+						{
+							builder.Append(", \\\r\n\t\t").Append(predicateInfo.Name).Append("Struct& ").Append(predicateInfo.Name);
+						}
+					}
 					builder.Append(" \\\r\n");
 					builder.Append("\t)\r\n");
 				}
@@ -181,6 +194,57 @@ namespace EpicGames.UHT.Exporters.CodeGen
 						
 						foreach (UhtRigVMMethodInfo methodInfo in scriptStruct.RigVMStructInfo.Methods)
 						{
+							if (methodInfo.IsPredicate)
+							{
+								builder.Append('\t').Append("struct ").Append(methodInfo.Name).Append("Struct \\\r\n");
+								builder.Append("\t{ \\\r\n");
+								builder.Append("\t\t").Append(methodInfo.Name).Append("Struct(){  } \\\r\n");
+								builder.Append("\t\t").Append(methodInfo.Name).Append("Struct(FRigVMExtendedExecuteContext& InContext, FRigVMPredicateBranch InBranch){ Context = &InContext; Branch = InBranch; } \\\r\n");
+								builder.Append("\t\t").Append(methodInfo.ReturnType).Append(" Execute(\\\r\n\t\t\t");
+								builder.AppendParameterDecls(methodInfo.Parameters, false, ", \\\r\n\t\t\t", true, false);
+								builder.Append("\t\t)  \\\r\n");
+								builder.Append("\t\t{  \\\r\n");
+								builder.Append("\t\t\tif (Branch.IsValid())  \\\r\n");
+								builder.Append("\t\t\t{  \\\r\n");
+								int parameterIndex = 0;
+								foreach (UhtRigVMParameter parameter in methodInfo.Parameters)
+								{
+									string baseType = parameter.TypeOriginal().ToString();
+									if (baseType.StartsWith("const"))
+									{
+										baseType = baseType.Substring(5).Trim();
+									}
+
+									if (baseType.EndsWith("&"))
+									{
+										baseType = baseType.Substring(0, baseType.Length - 1).Trim();
+									}
+									builder.Append("\t\t\t\t*(").Append(baseType).Append("*) Branch.MemoryHandles[")
+										.Append(parameterIndex++).Append("].GetData(false) = ")
+										.Append(parameter.Name).Append(";  \\\r\n");
+									
+								}
+								builder.Append("\t\t\t\tBranch.Execute(*Context);  \\\r\n");
+								builder.Append("\t\t\t\treturn *(").Append(methodInfo.ReturnType).Append("*)Branch.MemoryHandles[").Append(methodInfo.Parameters.Count).Append("].GetData(false);  \\\r\n");
+								builder.Append("\t\t\t}  \\\r\n");
+								builder.Append("\t\t\treturn ").Append(methodInfo.Name).Append("(");
+								builder.AppendParameterNames(methodInfo.Parameters, false, ", ", false, false);
+								builder.Append("); \\\r\n");
+								builder.Append("\t\t}  \\\r\n");
+								builder.Append("\tprivate: \\\r\n");
+								builder.Append("\t\tFRigVMExtendedExecuteContext* Context; \\\r\n");
+								builder.Append("\t\tFRigVMPredicateBranch Branch; \\\r\n");
+								builder.Append("\t};  \\\r\n");
+							}
+						}
+
+						foreach (UhtRigVMMethodInfo methodInfo in scriptStruct.RigVMStructInfo.Methods)
+						{
+							if (methodInfo.IsPredicate)
+							{
+								continue;
+							}
+							
 							builder.Append('\t').Append(methodInfo.ReturnType).Append(' ').Append(methodInfo.Name).Append('(').Append(constPrefix).Append(scriptStruct.RigVMStructInfo.ExecuteContextType).Append("& InExecuteContext); \\\r\n");
 							builder.Append("\tstatic ").Append(methodInfo.ReturnType).Append(" Static").Append(methodInfo.Name).Append("( \\\r\n");
 							builder.Append("\t\t");
@@ -190,13 +254,23 @@ namespace EpicGames.UHT.Exporters.CodeGen
 							}
 							builder.Append(scriptStruct.RigVMStructInfo.ExecuteContextType).Append("& ").Append(RigVMExecuteContextParamName);
 							builder.AppendParameterDecls(scriptStruct.RigVMStructInfo.Members, true, ", \\\r\n\t\t", true, false);
+
+							foreach (UhtRigVMMethodInfo predicateInfo in scriptStruct.RigVMStructInfo.Methods)
+							{
+								if (predicateInfo.IsPredicate)
+								{
+									builder.Append(", \\\r\n\t\t").Append(predicateInfo.Name).Append("Struct& ").Append(predicateInfo.Name);
+								}
+							}
+							
 							builder.Append(" \\\r\n");
 							builder.Append("\t); \\\r\n");
 
 							builder.Append("\tFORCEINLINE_DEBUGGABLE static ").Append(methodInfo.ReturnType).Append(" RigVM").Append(methodInfo.Name).Append("( \\\r\n");
 							builder.Append("\t\t");
 							builder.Append(RigVMExecuteContextDeclaration).Append(", \\\r\n");
-							builder.Append("\t\tFRigVMMemoryHandleArray RigVMMemoryHandles \\\r\n");
+							builder.Append("\t\tFRigVMMemoryHandleArray RigVMMemoryHandles, \\\r\n");
+							builder.Append("\t\tFRigVMPredicateBranchArray RigVMBranches \\\r\n");
 							builder.Append("\t) \\\r\n");
 							builder.Append("\t{ \\\r\n");
 
@@ -294,9 +368,25 @@ namespace EpicGames.UHT.Exporters.CodeGen
 								builder.Append("\t\t \\\r\n");
 							}
 
+							int predicateId = 0;
+							foreach (UhtRigVMMethodInfo predicateInfo in scriptStruct.RigVMStructInfo.Methods)
+							{
+								if (predicateInfo.IsPredicate)
+								{
+									builder.Append("\t\t").Append(predicateInfo.Name).Append("Struct ").Append(predicateInfo.Name).Append("Predicate(RigVMExecuteContext, RigVMBranches[").Append(predicateId++).Append("]); \\\r\n");
+								}
+							}
+							
 							builder.Append("\t\t").Append(methodInfo.ReturnPrefix()).Append("Static").Append(methodInfo.Name).Append("( \\\r\n");
 							builder.Append("\t\t\tRigVMExecuteContext.GetPublicData<").Append(scriptStruct.RigVMStructInfo.ExecuteContextType).Append(">()");
 							builder.AppendParameterNames(scriptStruct.RigVMStructInfo.Members, true, ", \\\r\n\t\t\t", false);
+							foreach (UhtRigVMMethodInfo predicateInfo in scriptStruct.RigVMStructInfo.Methods)
+							{
+								if (predicateInfo.IsPredicate)
+								{
+									builder.Append(", \\\r\n\t\t\t").Append(predicateInfo.Name).Append("Predicate");
+								}
+							}
 							builder.Append(" \\\r\n");
 							builder.Append("\t\t); \\\r\n");
 							builder.Append("\t} \\\r\n");
diff --git a/Engine/Source/Programs/Shared/EpicGames.UHT/Parsers/UhtScriptStructParser.cs b/Engine/Source/Programs/Shared/EpicGames.UHT/Parsers/UhtScriptStructParser.cs
--- a/Engine/Source/Programs/Shared/EpicGames.UHT/Parsers/UhtScriptStructParser.cs
+++ b/Engine/Source/Programs/Shared/EpicGames.UHT/Parsers/UhtScriptStructParser.cs
@@ -169,9 +169,37 @@ namespace EpicGames.UHT.Parsers
 				// Create a new method information and add it
 				UhtRigVMMethodInfo methodInfo = new();
 
+				topScope.TokenReader
+					.Require('(');
+				if (topScope.TokenReader.PeekToken().IsValue("meta"))
+				{
+					topScope.TokenReader
+						.RequireIdentifier("meta")
+						.Require('=')
+						.RequireList('(', ')', ',', false, (IEnumerable<UhtToken> tokens) =>
+						{
+							foreach (UhtToken token in tokens)
+							{
+								if (token.IsIdentifier("Predicate"))
+								{
+									methodInfo.IsPredicate = true;
+								}
+							}
+						})
+						.Require(')');
+				}
+				else
+				{
+					topScope.TokenReader.Require(')');
+				}
+
+				if (methodInfo.IsPredicate)
+				{
+					topScope.TokenReader.RequireIdentifier("static");
+				}
+
 				// NOTE: The argument list reader doesn't handle templates with multiple arguments (i.e. the ',' issue)
 				topScope.TokenReader
-					.RequireList('(', ')')
 					.Optional("virtual")
 					.RequireIdentifier((ref UhtToken identifier) => methodInfo.ReturnType = identifier.Value.ToString())
 					.RequireIdentifier((ref UhtToken identifier) => methodInfo.Name = identifier.Value.ToString());
@@ -219,7 +247,7 @@ namespace EpicGames.UHT.Parsers
 
 				if (!isGetUpgradeInfo && !isGetNextAggregateName)
 				{
-					if (methodInfo.Parameters.Count > 0)
+					if (methodInfo.Parameters.Count > 0 && !methodInfo.IsPredicate)
 					{
 						topScope.TokenReader.LogError($"RIGVM_METHOD {scriptStruct.SourceName}::{methodInfo.Name} has {methodInfo.Parameters.Count} parameters. Since 5.2 parameters are no longer allowed for RIGVM_METHOD functions.");
 						methodInfo.Parameters.Clear();
diff --git a/Engine/Source/Programs/Shared/EpicGames.UHT/Types/UhtScriptStruct.cs b/Engine/Source/Programs/Shared/EpicGames.UHT/Types/UhtScriptStruct.cs
--- a/Engine/Source/Programs/Shared/EpicGames.UHT/Types/UhtScriptStruct.cs
+++ b/Engine/Source/Programs/Shared/EpicGames.UHT/Types/UhtScriptStruct.cs
@@ -500,6 +500,11 @@ namespace EpicGames.UHT.Types
 		{
 			return (ReturnType.Length == 0 || ReturnType == "void") ? s_noPrefix : ReturnPrefixInternal;
 		}
+		/// <summary>
+		/// Whether or not this method is a predicate.
+		/// </summary>
+		/// <returns>True if this method is a predicate</returns>
+		public bool IsPredicate { get; set; } = false;
 	}
 
 	/// <summary>
diff --git a/Engine/Source/Programs/UnrealHeaderTool/Private/CodeGenerator.cpp b/Engine/Source/Programs/UnrealHeaderTool/Private/CodeGenerator.cpp
--- a/Engine/Source/Programs/UnrealHeaderTool/Private/CodeGenerator.cpp
+++ b/Engine/Source/Programs/UnrealHeaderTool/Private/CodeGenerator.cpp
@@ -3653,6 +3653,11 @@ void FNativeClassHeaderGenerator::ExportGeneratedStructBodyMacros(FOutputDevice&
 			FString StructMembers = StructRigVMInfo.Members.Declarations(false, TEXT(",\r\n\t\t"), true, false, true, false);
 			for (const FRigVMMethodInfo& MethodInfo : StructRigVMInfo.Methods)
 			{
+				if (MethodInfo.bIsPredicate)
+				{
+					continue;
+				}
+				
 				FString StructMembersForStub = StructRigVMInfo.Members.Names(false, TEXT(",\r\n\t\t\t"), false, true, false);
 				FString RigVMParameterPrefix2 = RigVMExecuteContextPublicDeclaration + FString(StructMembers.IsEmpty() ? TEXT("") : TEXT(",\r\n\t\t"));
 				if(StructRigVMInfo.ExecuteContextMember.IsEmpty())
diff --git a/Engine/Source/Programs/UnrealHeaderTool/Private/HeaderParser.cpp b/Engine/Source/Programs/UnrealHeaderTool/Private/HeaderParser.cpp
--- a/Engine/Source/Programs/UnrealHeaderTool/Private/HeaderParser.cpp
+++ b/Engine/Source/Programs/UnrealHeaderTool/Private/HeaderParser.cpp
@@ -6669,6 +6669,7 @@ void FHeaderParser::CompileRigVMMethodDeclaration(FUnrealStructDefinitionInfo& S
 	}
 
 	// find the next close brace
+	TArray<FToken> MetaPropertyTokens;
 	while (!MatchSymbol(TEXT(")")))
 	{
 		FToken Token;
@@ -6676,6 +6677,26 @@ void FHeaderParser::CompileRigVMMethodDeclaration(FUnrealStructDefinitionInfo& S
 		{
 			break;
 		}
+		if (Token.IsIdentifier(TEXT("meta"), ESearchCase::CaseSensitive))
+		{
+			if (!MatchSymbol(TEXT("=")))
+			{
+				Throwf(TEXT("Bad RIGVM_METHOD definition"));
+			}
+			if (!MatchSymbol(TEXT("(")))
+			{
+				Throwf(TEXT("Bad RIGVM_METHOD definition"));
+			}
+			while (!MatchSymbol(TEXT(")")))
+			{
+				FToken MetaPropertyToken;
+				if (!GetToken(MetaPropertyToken))
+				{
+					break;
+				}
+				MetaPropertyTokens.Add(MetaPropertyToken);
+			}
+		}
 	}
 
 	FToken PrefixToken, ReturnTypeToken, NameToken, PostfixToken;
@@ -6691,6 +6712,13 @@ void FHeaderParser::CompileRigVMMethodDeclaration(FUnrealStructDefinitionInfo& S
 			return;
 		}
 	}
+	else if (PrefixToken.IsIdentifier(TEXT("static"), ESearchCase::CaseSensitive))
+	{
+		if (!GetToken(ReturnTypeToken))
+		{
+			return;
+		}
+	}
 	else
 	{
 		ReturnTypeToken = PrefixToken;
@@ -6728,6 +6756,19 @@ void FHeaderParser::CompileRigVMMethodDeclaration(FUnrealStructDefinitionInfo& S
 	FRigVMMethodInfo MethodInfo;
 	MethodInfo.ReturnType = FString(ReturnTypeToken.Value);
 	MethodInfo.Name = FString(NameToken.Value);
+	if (MetaPropertyTokens.ContainsByPredicate([](FToken& Token){return Token.IsIdentifier(TEXT("Predicate"), ESearchCase::CaseSensitive);}))
+	{
+		if (!PrefixToken.IsIdentifier(TEXT("static"), ESearchCase::CaseSensitive))
+		{
+			LogError(TEXT("RIGVM_METHOD predicate F%s::%s is not defined as static. Predicates must be static methods."), *StructDef.GetName(), *MethodInfo.Name);
+			return;
+		}
+		MethodInfo.bIsPredicate = true;
+	}
+	else
+	{
+		MethodInfo.bIsPredicate = false;
+	}
 
 	// look out for the upgrade info method
 	static const FString GetUpgradeInfoString = TEXT("GetUpgradeInfo");
@@ -6799,7 +6840,7 @@ void FHeaderParser::CompileRigVMMethodDeclaration(FUnrealStructDefinitionInfo& S
 	FRigVMStructInfo& StructRigVMInfo = StructDef.GetRigVMInfo();
 
 	// disable support for opaque arguments
-	if (MethodInfo.Parameters.Num() > 0)
+	if (MethodInfo.Parameters.Num() > 0 && !MethodInfo.bIsPredicate)
 	{
 		LogError(TEXT("RIGVM_METHOD F%s::%s has %d parameters. Since 5.2 parameters are no longer allowed for RIGVM_METHOD functions."), *StructRigVMInfo.Name, *MethodInfo.Name, MethodInfo.Parameters.Num());
 		MethodInfo.Parameters = FRigVMParameterArray();
@@ -6822,6 +6863,7 @@ const FName FHeaderParser::NAME_OutputText(TEXT("Output"));
 const FName FHeaderParser::NAME_ConstantText(TEXT("Constant"));
 const FName FHeaderParser::NAME_VisibleText(TEXT("Visible"));
 const FName FHeaderParser::NAME_LazyText(TEXT("Lazy"));
+const FName FHeaderParser::NAME_PredicateText(TEXT("Predicate"));
 
 const FName FHeaderParser::NAME_SingletonText(TEXT("Singleton"));
 
diff --git a/Engine/Source/Programs/UnrealHeaderTool/Private/HeaderParser.h b/Engine/Source/Programs/UnrealHeaderTool/Private/HeaderParser.h
--- a/Engine/Source/Programs/UnrealHeaderTool/Private/HeaderParser.h
+++ b/Engine/Source/Programs/UnrealHeaderTool/Private/HeaderParser.h
@@ -663,6 +663,7 @@ public:
 	static const FName NAME_ConstantText;
 	static const FName NAME_VisibleText;
 	static const FName NAME_LazyText;
+	static const FName NAME_PredicateText;
 	static const FName NAME_SingletonText;
 	static const TCHAR* TArrayText;
 	static const TCHAR* TEnumAsByteText;
diff --git a/Engine/Source/Programs/UnrealHeaderTool/Private/ParserHelper.h b/Engine/Source/Programs/UnrealHeaderTool/Private/ParserHelper.h
--- a/Engine/Source/Programs/UnrealHeaderTool/Private/ParserHelper.h
+++ b/Engine/Source/Programs/UnrealHeaderTool/Private/ParserHelper.h
@@ -734,6 +734,7 @@ struct FRigVMMethodInfo
 	FString ReturnType;
 	FString Name;
 	FRigVMParameterArray Parameters;
+	bool bIsPredicate;
 
 	FString ReturnPrefix() const
 	{
