RigVM: During recursive remove node actions disable schema check

diff --git a/Engine/Plugins/Runtime/RigVM/Source/RigVMDeveloper/Private/RigVMModel/RigVMController.cpp b/Engine/Plugins/Runtime/RigVM/Source/RigVMDeveloper/Private/RigVMModel/RigVMController.cpp
--- a/Engine/Plugins/Runtime/RigVM/Source/RigVMDeveloper/Private/RigVMModel/RigVMController.cpp
+++ b/Engine/Plugins/Runtime/RigVM/Source/RigVMDeveloper/Private/RigVMModel/RigVMController.cpp
@@ -96,6 +96,7 @@ URigVMController::URigVMController()
 #if WITH_EDITOR
 	, bRegisterTemplateNodeUsage(true)
 #endif
+	, bEnableSchemaRemoveNodeCheck(true)
 {
 }
 
@@ -116,6 +117,7 @@ URigVMController::URigVMController(const FObjectInitializer& ObjectInitializer)
 #if WITH_EDITOR
 	, bRegisterTemplateNodeUsage(true)
 #endif
+	, bEnableSchemaRemoveNodeCheck(true)
 {
 }
 
@@ -6205,6 +6207,8 @@ URigVMCollapseNode* URigVMController::PromoteFunctionReferenceNodeToCollapseNode
 
 		ApplyPinStates(CollapseNode, PinStates);
 		RestoreLinkedPaths(LinkedPaths, RestoreSettings);
+
+		Notify(ERigVMGraphNotifType::NodeAdded, CollapseNode);
 	}
 
 	if(bRemoveFunctionDefinition)
@@ -6317,9 +6321,12 @@ bool URigVMController::RemoveNodes(TArray<URigVMNode*> InNodes, bool bSetupUndoR
 			return false;
 		}
 
-		if(!GetSchema()->CanRemoveNode(this, InNode))
+		if(bEnableSchemaRemoveNodeCheck)
 		{
-			return false;
+			if(!GetSchema()->CanRemoveNode(this, InNode))
+			{
+				continue;
+			}
 		}
 
 		if (InNode->IsInjected())
@@ -6564,6 +6571,7 @@ bool URigVMController::RemoveNodes(TArray<URigVMNode*> InNodes, bool bSetupUndoR
 			{
 				TArray<URigVMNode*> ContainedNodes = SubGraph->GetNodes();
 				TGuardValue<bool> SuspendTemplateComputation(bSuspendTemplateComputation, true);
+				TGuardValue<bool> DisableSchemaRemoveNodeCheck(bEnableSchemaRemoveNodeCheck, false);
 				SubGraphController->RemoveNodes(ContainedNodes, false, false);
 
 				if(IRigVMClientHost* ClientHost = GetImplementingOuter<IRigVMClientHost>())
diff --git a/Engine/Plugins/Runtime/RigVM/Source/RigVMDeveloper/Public/RigVMModel/RigVMController.h b/Engine/Plugins/Runtime/RigVM/Source/RigVMDeveloper/Public/RigVMModel/RigVMController.h
--- a/Engine/Plugins/Runtime/RigVM/Source/RigVMDeveloper/Public/RigVMModel/RigVMController.h
+++ b/Engine/Plugins/Runtime/RigVM/Source/RigVMDeveloper/Public/RigVMModel/RigVMController.h
@@ -1424,6 +1424,8 @@ private:
 	bool bRegisterTemplateNodeUsage;
 #endif
 
+	bool bEnableSchemaRemoveNodeCheck;
+
 	friend class URigVMGraph;
 	friend class URigVMPin;
 	friend class URigVMActionStack;
