RigVM : Fixed reset cached memory at Execute, as it was not clearing the context data

diff --git a/Engine/Plugins/Runtime/RigVM/Source/RigVM/Private/RigVMCore/RigVM.cpp b/Engine/Plugins/Runtime/RigVM/Source/RigVM/Private/RigVMCore/RigVM.cpp
--- a/Engine/Plugins/Runtime/RigVM/Source/RigVM/Private/RigVMCore/RigVM.cpp
+++ b/Engine/Plugins/Runtime/RigVM/Source/RigVM/Private/RigVMCore/RigVM.cpp
@@ -991,6 +991,12 @@ void URigVM::InvalidateCachedMemory()
 	LazyBranches.Reset();
 }
 
+void URigVM::InvalidateCachedMemory(FRigVMExtendedExecuteContext& Context)
+{
+	InvalidateCachedMemory();
+	Context.InvalidateCachedMemory();
+}
+
 void URigVM::CopyDeferredVMIfRequired()
 {
 	ensure(ActiveExecutions.load() == 0);  // we require that no active worker threads are running while we serialize
@@ -1012,17 +1018,17 @@ void URigVM::CacheMemoryHandlesIfRequired(FRigVMExtendedExecuteContext& Context,
 
 	if (Instructions.Num() == 0 || InMemory.Num() == 0)
 	{
-		InvalidateCachedMemory();
+		InvalidateCachedMemory(Context);
 		return;
 	}
 
 	if ((Instructions.Num() + 1) != FirstHandleForInstruction.Num())
 	{
-		InvalidateCachedMemory();
+		InvalidateCachedMemory(Context);
 	}
 	else if (InMemory.Num() != Context.CachedMemory.Num())
 	{
-		InvalidateCachedMemory();
+		InvalidateCachedMemory(Context);
 	}
 	else
 	{
@@ -1030,7 +1036,7 @@ void URigVM::CacheMemoryHandlesIfRequired(FRigVMExtendedExecuteContext& Context,
 		{
 			if (InMemory[Index] != Context.CachedMemory[Index])
 			{
-				InvalidateCachedMemory();
+				InvalidateCachedMemory(Context);
 				break;
 			}
 		}
diff --git a/Engine/Plugins/Runtime/RigVM/Source/RigVM/Public/RigVMCore/RigVM.h b/Engine/Plugins/Runtime/RigVM/Source/RigVM/Public/RigVMCore/RigVM.h
--- a/Engine/Plugins/Runtime/RigVM/Source/RigVM/Public/RigVMCore/RigVM.h
+++ b/Engine/Plugins/Runtime/RigVM/Source/RigVM/Public/RigVMCore/RigVM.h
@@ -812,6 +812,7 @@ private:
 
 public:
 	void InvalidateCachedMemory();
+	void InvalidateCachedMemory(FRigVMExtendedExecuteContext& Context);
 	
 private:
 	void CacheMemoryHandlesIfRequired(FRigVMExtendedExecuteContext& Context, TArrayView<URigVMMemoryStorage*> InMemory);
