RigVM: Fix crash when  undoing PromoteFunctionToCollapseNode

diff --git a/Engine/Plugins/Runtime/RigVM/Source/RigVMDeveloper/Private/RigVMModel/RigVMController.cpp b/Engine/Plugins/Runtime/RigVM/Source/RigVMDeveloper/Private/RigVMModel/RigVMController.cpp
--- a/Engine/Plugins/Runtime/RigVM/Source/RigVMDeveloper/Private/RigVMModel/RigVMController.cpp
+++ b/Engine/Plugins/Runtime/RigVM/Source/RigVMDeveloper/Private/RigVMModel/RigVMController.cpp
@@ -4160,7 +4160,14 @@ TArray<FName> URigVMController::ImportNodesFromText(const FString& InText, bool
 
 			if(Graph != LastGraph)
 			{
-				Factory.RestoreLinksSettings.NodeNameMap = Factory.NodeNameMapPerGraph.FindChecked(Graph);
+				if (TMap<FString, FString>* Map = Factory.NodeNameMapPerGraph.Find(Graph))
+				{
+					Factory.RestoreLinksSettings.NodeNameMap = *Map;
+				}
+				else
+				{
+					Factory.RestoreLinksSettings.NodeNameMap.Reset();
+				}
 				LastGraph = Graph;
 			}
 
