Fixed mismatched copy constructor/assignment operator in FRigVMExternalVariableDef.

diff --git a/Engine/Plugins/Runtime/RigVM/Source/RigVM/Public/RigVMCore/RigVMExternalVariable.h b/Engine/Plugins/Runtime/RigVM/Source/RigVM/Public/RigVMCore/RigVMExternalVariable.h
--- a/Engine/Plugins/Runtime/RigVM/Source/RigVM/Public/RigVMCore/RigVMExternalVariable.h
+++ b/Engine/Plugins/Runtime/RigVM/Source/RigVM/Public/RigVMCore/RigVMExternalVariable.h
@@ -25,7 +25,6 @@ struct RIGVM_API FRigVMExternalVariableDef
 	GENERATED_BODY()
 
 	FRigVMExternalVariableDef() = default;
-	FRigVMExternalVariableDef(const FRigVMExternalVariableDef& Other) = default;
 
 	static void GetTypeFromProperty(const FProperty* InProperty, FName& OutTypeName, UObject*& OutTypeObject)
 	{
