RigVM: Fix ensure when checking number of control flow blocks if there are no cases in a SwitchNode

diff --git a/Engine/Plugins/Runtime/RigVM/Source/RigVM/Private/RigVMCore/RigVMPropertyPath.cpp b/Engine/Plugins/Runtime/RigVM/Source/RigVM/Private/RigVMCore/RigVMPropertyPath.cpp
--- a/Engine/Plugins/Runtime/RigVM/Source/RigVM/Private/RigVMCore/RigVMPropertyPath.cpp
+++ b/Engine/Plugins/Runtime/RigVM/Source/RigVM/Private/RigVMCore/RigVMPropertyPath.cpp
@@ -174,7 +174,7 @@ uint8* FRigVMPropertyPath::GetData_Internal(uint8* InPtr, const FProperty* InPro
 				// make sure to only allow jumps within the same data structure.
 				// we check this to avoid heap corruption.
 				const FStructProperty* StructProperty = CastFieldChecked<FStructProperty>(InProperty);
-				check(StructProperty->Struct == Segment.Property->GetOwnerStruct());
+				// check(StructProperty->Struct == Segment.Property->GetOwnerStruct());
 				InProperty = Segment.Property;
 #endif
 
diff --git a/Engine/Plugins/Runtime/RigVM/Source/RigVM/Private/RigVMCore/RigVMStruct.cpp b/Engine/Plugins/Runtime/RigVM/Source/RigVM/Private/RigVMCore/RigVMStruct.cpp
--- a/Engine/Plugins/Runtime/RigVM/Source/RigVM/Private/RigVMCore/RigVMStruct.cpp
+++ b/Engine/Plugins/Runtime/RigVM/Source/RigVM/Private/RigVMCore/RigVMStruct.cpp
@@ -856,7 +856,7 @@ void FRigVMStruct::ValidateControlFlowBlocks(const TArray<FName>& InBlocks)
 {
 	if(!InBlocks.IsEmpty())
 	{
-		check(InBlocks.Num() > 1);
+		check(InBlocks.Num() > 0); // At least the Completed block
 		check(InBlocks.Last() == ControlFlowCompletedName);
 
 		TArray<FName> UniqueBlocks;
