RigVM: Avoid crash when modifying variables inside public function
#rnx

diff --git a/Engine/Plugins/Runtime/RigVM/Source/RigVM/Public/RigVMCore/RigVM.h b/Engine/Plugins/Runtime/RigVM/Source/RigVM/Public/RigVMCore/RigVM.h
--- a/Engine/Plugins/Runtime/RigVM/Source/RigVM/Public/RigVMCore/RigVM.h
+++ b/Engine/Plugins/Runtime/RigVM/Source/RigVM/Public/RigVMCore/RigVM.h
@@ -688,9 +688,9 @@ public:
 	FRigVMOperand AddExternalVariable(const FRigVMExternalVariable& InExternalVariable) { return FRigVMOperand(); }
 
 	// Adds a new external / unowned variable to the VM
-	FRigVMOperand AddExternalVariable(FRigVMExtendedExecuteContext& Context, const FRigVMExternalVariable& InExternalVariable)
+	FRigVMOperand AddExternalVariable(FRigVMExtendedExecuteContext& Context, const FRigVMExternalVariable& InExternalVariable, bool bAllowNullMemory = false)
 	{
-		check(InExternalVariable.Memory != nullptr);
+		check(bAllowNullMemory || InExternalVariable.Memory != nullptr);
 
 		const int32 VariableIndex = ExternalVariables.Add(InExternalVariable);
 		Context.ExternalVariableRuntimeData.Add(FRigVMExternalVariableRuntimeData(InExternalVariable.Memory));
diff --git a/Engine/Plugins/Runtime/RigVM/Source/RigVMDeveloper/Private/RigVMCompiler/RigVMCompiler.cpp b/Engine/Plugins/Runtime/RigVM/Source/RigVMDeveloper/Private/RigVMCompiler/RigVMCompiler.cpp
--- a/Engine/Plugins/Runtime/RigVM/Source/RigVMDeveloper/Private/RigVMCompiler/RigVMCompiler.cpp
+++ b/Engine/Plugins/Runtime/RigVM/Source/RigVMDeveloper/Private/RigVMCompiler/RigVMCompiler.cpp
@@ -832,7 +832,7 @@ bool URigVMCompiler::Compile(TArray<URigVMGraph*> InGraphs, URigVMController* In
 	
 	for (const FRigVMExternalVariable& ExternalVariable : InExternalVariables)
 	{
-		FRigVMOperand Operand = OutVM->AddExternalVariable(OutVMContext, ExternalVariable);
+		FRigVMOperand Operand = OutVM->AddExternalVariable(OutVMContext, ExternalVariable, CurrentCompilationFunction != nullptr);
 		FString Hash = FString::Printf(TEXT("Variable::%s"), *ExternalVariable.Name.ToString());
 		OutOperands->Add(Hash, Operand);
 	}
