RigVM: Protect ArrayUnion node from creating an array bigger than allowed

diff --git a/Engine/Plugins/Runtime/RigVM/Source/RigVM/Private/RigVMFunctions/RigVMDispatch_Array.cpp b/Engine/Plugins/Runtime/RigVM/Source/RigVM/Private/RigVMFunctions/RigVMDispatch_Array.cpp
--- a/Engine/Plugins/Runtime/RigVM/Source/RigVM/Private/RigVMFunctions/RigVMDispatch_Array.cpp
+++ b/Engine/Plugins/Runtime/RigVM/Source/RigVM/Private/RigVMFunctions/RigVMDispatch_Array.cpp
@@ -1130,50 +1130,66 @@ void FRigVMDispatch_ArrayUnion::Execute(FRigVMExtendedExecuteContext& InContext,
 		const TMap<uint32, int32> HashA = GetArrayHash(ArrayHelper, ArrayProperty);
 		const TMap<uint32, int32> HashB = GetArrayHash(OtherHelper, OtherProperty);
 
-		// copy the complete array to a temp storage
-		TArray<uint8, TAlignedHeapAllocator<16>> TempStorage;
-		const int32 NumElementsA = ArrayHelper.Num();
-		TempStorage.AddZeroed(NumElementsA * ArrayElementProperty->GetSize());
-		uint8* TempMemory = TempStorage.GetData();
-		for(int32 Index = 0; Index < NumElementsA; Index++)
+		int32 FinalCount = HashA.Num();
+		TArray<int32> OtherIndicesToAdd;
+		OtherIndicesToAdd.Reserve(HashB.Num());
+		for(const TPair<uint32, int32>& Pair : HashB)
 		{
-			ArrayElementProperty->InitializeValue(TempMemory);
-			ArrayElementProperty->CopyCompleteValue(TempMemory, ArrayHelper.GetRawPtr(Index));
-			TempMemory += ArrayElementProperty->GetSize();
+			if(!HashA.Contains(Pair.Key))
+			{
+				FinalCount++;
+				if (!InContext.IsValidArraySize(FinalCount))
+				{
+					break;
+				}
+				OtherIndicesToAdd.Add(Pair.Value);
+			}
 		}
 
-		ArrayHelper.Resize(0);
-
-		for(const TPair<uint32, int32>& Pair : HashA)
+		if (InContext.IsValidArraySize(FinalCount))
 		{
-			const int32 AddedIndex = ArrayHelper.AddValue();
-			TempMemory = TempStorage.GetData() + Pair.Value * ArrayElementProperty->GetSize();
-						
-			URigVMMemoryStorage::CopyProperty(
-				ArrayElementProperty,
-				ArrayHelper.GetRawPtr(AddedIndex),
-				ArrayElementProperty,
-				TempMemory
-			);
-		}
+			// copy the complete array to a temp storage
+			TArray<uint8, TAlignedHeapAllocator<16>> TempStorage;
+			const int32 NumElementsA = ArrayHelper.Num();
+			TempStorage.AddZeroed(NumElementsA * ArrayElementProperty->GetSize());
+			uint8* TempMemory = TempStorage.GetData();
+			for(int32 Index = 0; Index < NumElementsA; Index++)
+			{
+				ArrayElementProperty->InitializeValue(TempMemory);
+				ArrayElementProperty->CopyCompleteValue(TempMemory, ArrayHelper.GetRawPtr(Index));
+				TempMemory += ArrayElementProperty->GetSize();
+			}
 
-		TempMemory = TempStorage.GetData();
-		for(int32 Index = 0; Index < NumElementsA; Index++)
-		{
-			ArrayElementProperty->DestroyValue(TempMemory);
-			TempMemory += ArrayElementProperty->GetSize();
-		}
+			ArrayHelper.Resize(0);
 
-		for(const TPair<uint32, int32>& Pair : HashB)
-		{
-			if(!HashA.Contains(Pair.Key))
+			for(const TPair<uint32, int32>& Pair : HashA)
+			{
+				const int32 AddedIndex = ArrayHelper.AddValue();
+				TempMemory = TempStorage.GetData() + Pair.Value * ArrayElementProperty->GetSize();
+					
+				URigVMMemoryStorage::CopyProperty(
+					ArrayElementProperty,
+					ArrayHelper.GetRawPtr(AddedIndex),
+					ArrayElementProperty,
+					TempMemory
+				);
+			}
+
+			TempMemory = TempStorage.GetData();
+			for(int32 Index = 0; Index < NumElementsA; Index++)
+			{
+				ArrayElementProperty->DestroyValue(TempMemory);
+				TempMemory += ArrayElementProperty->GetSize();
+			}
+
+			for (const int32& OtherIndex : OtherIndicesToAdd)
 			{
 				const int32 AddedIndex = ArrayHelper.AddValue();
 				URigVMMemoryStorage::CopyProperty(
 					ArrayElementProperty,
 					ArrayHelper.GetRawPtr(AddedIndex),
 					OtherElementProperty,
-					OtherHelper.GetRawPtr(Pair.Value)
+					OtherHelper.GetRawPtr(OtherIndex)
 				);
 			}
 		}
