Control Rig: Fix a potential array overrun when appending elements.

#rnx

diff --git a/Engine/Plugins/Runtime/RigVM/Source/RigVM/Private/RigVMFunctions/RigVMDispatch_Array.cpp b/Engine/Plugins/Runtime/RigVM/Source/RigVM/Private/RigVMFunctions/RigVMDispatch_Array.cpp
--- a/Engine/Plugins/Runtime/RigVM/Source/RigVM/Private/RigVMFunctions/RigVMDispatch_Array.cpp
+++ b/Engine/Plugins/Runtime/RigVM/Source/RigVM/Private/RigVMFunctions/RigVMDispatch_Array.cpp
@@ -993,24 +993,25 @@ FRigVMTemplateTypeMap FRigVMDispatch_ArrayAppend::OnNewArgumentType(const FName&
 
 void FRigVMDispatch_ArrayAppend::Execute(FRigVMExtendedExecuteContext& InContext, FRigVMMemoryHandleArray Handles, FRigVMPredicateBranchArray Predicates)
 {
-	const FArrayProperty* OtherProperty = CastFieldChecked<FArrayProperty>(Handles[1].GetProperty());
-	FScriptArrayHelper OtherHelper(OtherProperty, Handles[1].GetData());
+	const FArrayProperty* SourceArrayProperty = CastFieldChecked<FArrayProperty>(Handles[1].GetProperty());
+	FScriptArrayHelper SourceArrayHelper(SourceArrayProperty, Handles[1].GetData());
 
-	if(OtherHelper.Num() > 0)
+	if(SourceArrayHelper.Num() > 0)
 	{
-		const FArrayProperty* ArrayProperty = CastFieldChecked<FArrayProperty>(Handles[0].GetProperty());
-		FScriptArrayHelper ArrayHelper(ArrayProperty, Handles[0].GetData());
+		const FArrayProperty* TargetArrayProperty = CastFieldChecked<FArrayProperty>(Handles[0].GetProperty());
+		FScriptArrayHelper TargetArrayHelper(TargetArrayProperty, Handles[0].GetData());
 
-		if(InContext.IsValidArraySize(ArrayHelper.Num() + OtherHelper.Num()))
+		if(InContext.IsValidArraySize(TargetArrayHelper.Num() + SourceArrayHelper.Num()))
 		{
-			const FProperty* TargetProperty = ArrayProperty->Inner;
-			const FProperty* SourceProperty = OtherProperty->Inner;
+			const FProperty* TargetProperty = TargetArrayProperty->Inner;
+			const FProperty* SourceProperty = SourceArrayProperty->Inner;
 
-			int32 TargetIndex = ArrayHelper.AddValues(OtherHelper.Num());
-			for(int32 SourceIndex = 0; SourceIndex < OtherHelper.Num(); SourceIndex++, TargetIndex++)
+			const int32 SourceItemCount = SourceArrayHelper.Num();
+			int32 TargetIndex = TargetArrayHelper.AddValues(SourceItemCount);
+			for(int32 SourceIndex = 0; SourceIndex < SourceItemCount; SourceIndex++, TargetIndex++)
 			{
-				uint8* TargetMemory = ArrayHelper.GetRawPtr(TargetIndex);
-				const uint8* SourceMemory = OtherHelper.GetRawPtr(SourceIndex);
+				uint8* TargetMemory = TargetArrayHelper.GetRawPtr(TargetIndex);
+				const uint8* SourceMemory = SourceArrayHelper.GetRawPtr(SourceIndex);
 				URigVMMemoryStorage::CopyProperty(TargetProperty, TargetMemory, SourceProperty, SourceMemory);
 			}
 		}
