RigVM: Fix MakeArray having 2 subpins when created
#rnx

diff --git a/Engine/Plugins/Runtime/RigVM/Source/RigVMDeveloper/Private/RigVMModel/RigVMController.cpp b/Engine/Plugins/Runtime/RigVM/Source/RigVMDeveloper/Private/RigVMModel/RigVMController.cpp
--- a/Engine/Plugins/Runtime/RigVM/Source/RigVMDeveloper/Private/RigVMModel/RigVMController.cpp
+++ b/Engine/Plugins/Runtime/RigVM/Source/RigVMDeveloper/Private/RigVMModel/RigVMController.cpp
@@ -8899,29 +8899,26 @@ bool URigVMController::RemoveArrayPin(const FString& InArrayElementPinPath, bool
 	check(ArrayPin);
 	ensure(ArrayPin->IsArray());
 
-	FRigVMControllerCompileBracketScope CompileScope(this);
-	FRigVMRemoveArrayPinAction Action;
-	if (bSetupUndoRedo)
-	{
-		Action = FRigVMRemoveArrayPinAction(this, ArrayElementPin);
-		Action.SetTitle(FString::Printf(TEXT("Remove Array Pin")));
-		GetActionStack()->BeginAction(Action);
-	}
-
 	// we need to keep at least one element for fixed size arrays
 	if(ArrayPin->IsExecuteContext() || ArrayPin->IsFixedSizeArray())
 	{
 		if(ArrayPin->GetArraySize() == 1)
 		{
-			if (bSetupUndoRedo)
-			{
-				GetActionStack()->CancelAction(Action);
-			}
+			ReportErrorf(TEXT("Cannot remove last element of a fixed size array %s."), *InArrayElementPinPath);
 			return false;
 		}
 	}
 
+	FRigVMControllerCompileBracketScope CompileScope(this);
+	FRigVMBaseAction Action(this);
+	if (bSetupUndoRedo)
+	{
+		Action.SetTitle(FString::Printf(TEXT("Remove Array Pin")));
+		GetActionStack()->BeginAction(Action);
+	}
+
 	int32 IndexToRemove = ArrayElementPin->GetPinIndex();
+	FRigVMRemoveArrayPinAction RemovePinAction(this, ArrayElementPin);
 	if (!RemovePin(ArrayElementPin, bSetupUndoRedo))
 	{
 		if (bSetupUndoRedo)
@@ -8931,6 +8928,11 @@ bool URigVMController::RemoveArrayPin(const FString& InArrayElementPinPath, bool
 		return false;
 	}
 
+	if (bSetupUndoRedo)
+	{
+		GetActionStack()->AddAction(RemovePinAction);
+	}
+
 	for (int32 ExistingIndex = IndexToRemove; ExistingIndex < ArrayPin->GetArraySize(); ExistingIndex++)
 	{
 		URigVMPin* ExistingPin = ArrayPin->GetSubPins()[ExistingIndex];
