RigVM: Avoid unwanted warnings/errors when setting pin default values

diff --git a/Engine/Plugins/Runtime/RigVM/Source/RigVMDeveloper/Private/RigVMModel/RigVMController.cpp b/Engine/Plugins/Runtime/RigVM/Source/RigVMDeveloper/Private/RigVMModel/RigVMController.cpp
--- a/Engine/Plugins/Runtime/RigVM/Source/RigVMDeveloper/Private/RigVMModel/RigVMController.cpp
+++ b/Engine/Plugins/Runtime/RigVM/Source/RigVMDeveloper/Private/RigVMModel/RigVMController.cpp
@@ -7817,6 +7817,7 @@ bool URigVMController::SetPinDefaultValue(URigVMPin* InPin, const FString& InDef
 
 			if (bResizeArrays)
 			{
+				TGuardValue<bool> GuardReporting(bReportWarningsAndErrors, false);
 				while (Elements.Num() > InPin->SubPins.Num())
 				{
 					if(!InsertArrayPin(InPin, INDEX_NONE, FString(), bSetupUndoRedo))
