RigVM: Fix crash when undoing delete function

diff --git a/Engine/Plugins/Runtime/RigVM/Source/RigVMDeveloper/Private/RigVMModel/RigVMClient.cpp b/Engine/Plugins/Runtime/RigVM/Source/RigVMDeveloper/Private/RigVMModel/RigVMClient.cpp
--- a/Engine/Plugins/Runtime/RigVM/Source/RigVMDeveloper/Private/RigVMModel/RigVMClient.cpp
+++ b/Engine/Plugins/Runtime/RigVM/Source/RigVMDeveloper/Private/RigVMModel/RigVMClient.cpp
@@ -296,6 +296,7 @@ bool FRigVMClient::RemoveController(const URigVMGraph* InModel)
 	const bool bSuccess = Controllers.Remove(Key) > 0;
 	if(Controller)
 	{
+		Controller->SetActionStack(nullptr);
 		Controller->Rename(nullptr, GetTransientPackage(), REN_ForceNoResetLoaders | REN_DoNotDirty | REN_DontCreateRedirectors | REN_NonTransactional);
 		Controller->RemoveFromRoot();
 		Controller->MarkAsGarbage();
diff --git a/Engine/Plugins/Runtime/RigVM/Source/RigVMDeveloper/Private/RigVMModel/RigVMController.cpp b/Engine/Plugins/Runtime/RigVM/Source/RigVMDeveloper/Private/RigVMModel/RigVMController.cpp
--- a/Engine/Plugins/Runtime/RigVM/Source/RigVMDeveloper/Private/RigVMModel/RigVMController.cpp
+++ b/Engine/Plugins/Runtime/RigVM/Source/RigVMDeveloper/Private/RigVMModel/RigVMController.cpp
@@ -123,6 +123,11 @@ URigVMController::URigVMController(const FObjectInitializer& ObjectInitializer)
 
 URigVMController::~URigVMController()
 {
+	if(URigVMActionStack* ActionStack = WeakActionStack.Get())
+	{
+		ActionStack->OnModified().Remove(ActionStackHandle);
+		ActionStackHandle.Reset();
+	}
 }
 
 #if WITH_EDITORONLY_DATA
