RigVM: Fix bug around number of operands

We moved towards uint16 size - one code path was not updated. The array causing the problems has 65 elements.

diff --git a/Engine/Plugins/Runtime/RigVM/Source/RigVMDeveloper/Private/RigVMCompiler/RigVMCompiler.cpp b/Engine/Plugins/Runtime/RigVM/Source/RigVMDeveloper/Private/RigVMCompiler/RigVMCompiler.cpp
--- a/Engine/Plugins/Runtime/RigVM/Source/RigVMDeveloper/Private/RigVMCompiler/RigVMCompiler.cpp
+++ b/Engine/Plugins/Runtime/RigVM/Source/RigVMDeveloper/Private/RigVMCompiler/RigVMCompiler.cpp
@@ -1796,7 +1796,7 @@ int32 URigVMCompiler::TraverseCallExtern(const FRigVMCallExternExprAST* InExpr,
 			WorkData.VM->GetByteCode().AddZeroOp(BlockToRunOperand);
 		}
 
-		if (Operands.Num() > 64)
+		if (Operands.Num() > 65535)
 		{
 			return INDEX_NONE;
 		}
