Control Rig: Don't error out, just report a warning, if a type cannot be resolved. Rely on downstream mechanisms to recover from it.

#rnx

diff --git a/Engine/Plugins/Runtime/RigVM/Source/RigVM/Public/RigVMCore/RigVMTemplate.h b/Engine/Plugins/Runtime/RigVM/Source/RigVM/Public/RigVMCore/RigVMTemplate.h
--- a/Engine/Plugins/Runtime/RigVM/Source/RigVM/Public/RigVMCore/RigVMTemplate.h
+++ b/Engine/Plugins/Runtime/RigVM/Source/RigVM/Public/RigVMCore/RigVMTemplate.h
@@ -13,6 +13,7 @@
 #include "Math/Color.h"
 #include "Misc/AssertionMacros.h"
 #include "RigVMFunction.h"
+#include "RigVMModule.h"
 #include "RigVMTraits.h"
 #include "RigVMTypeIndex.h"
 #include "RigVMTypeUtils.h"
@@ -73,8 +74,11 @@ struct RIGVM_API FRigVMTemplateArgumentType
 		// so here we override the CppType name with the actual name used in the registry
 		const FString InCPPTypeString = CPPType.ToString();
 		CPPType = *RigVMTypeUtils::PostProcessCPPType(InCPPTypeString, CPPTypeObject);
-		
-		checkf(!CPPType.IsNone(), TEXT("FRigVMTemplateArgumentType(): Input CPPType '%s' could not be resolved."), *InCPPTypeString);
+
+		if (CPPType.IsNone())
+		{
+			UE_LOG(LogRigVM, Warning, TEXT("FRigVMTemplateArgumentType(): Input CPPType '%s' could not be resolved."), *InCPPTypeString);
+		}
 	}
 
 	FRigVMTemplateArgumentType(UClass* InClass)
