RigVMController: Ensure permutations are created when needed

diff --git a/Engine/Plugins/Runtime/RigVM/Source/RigVMDeveloper/Private/RigVMModel/RigVMController.cpp b/Engine/Plugins/Runtime/RigVM/Source/RigVMDeveloper/Private/RigVMModel/RigVMController.cpp
--- a/Engine/Plugins/Runtime/RigVM/Source/RigVMDeveloper/Private/RigVMModel/RigVMController.cpp
+++ b/Engine/Plugins/Runtime/RigVM/Source/RigVMDeveloper/Private/RigVMModel/RigVMController.cpp
@@ -13110,7 +13110,7 @@ URigVMTemplateNode* URigVMController::AddTemplateNode(const FName& InNotation, c
 
 	ensure(!InNotation.IsNone());
 
-	const FRigVMTemplate* Template = FRigVMRegistry::Get().FindTemplate(InNotation);
+	FRigVMTemplate* Template = const_cast<FRigVMTemplate*>(FRigVMRegistry::Get().FindTemplate(InNotation));
 	if (Template == nullptr)
 	{
 		ReportErrorf(TEXT("Template '%s' cannot be found."), *InNotation.ToString());
@@ -13140,7 +13140,7 @@ URigVMTemplateNode* URigVMController::AddTemplateNode(const FName& InNotation, c
 	{
 		Node = NewObject<URigVMDispatchNode>(Graph, *Name);
 	}
-	else if(const FRigVMFunction* FirstFunction = Template->GetPermutation(0))
+	else if(const FRigVMFunction* FirstFunction = Template->GetOrCreatePermutation(0))
 	{
 		const UScriptStruct* PotentialUnitStruct = FirstFunction->Struct;
 		if(PotentialUnitStruct && PotentialUnitStruct->IsChildOf(FRigVMStruct::StaticStruct()))
@@ -13253,14 +13253,14 @@ TArray<UScriptStruct*> URigVMController::GetUnitStructsForTemplate(const FName&
 {
 	TArray<UScriptStruct*> UnitStructs;
 
-	const FRigVMTemplate* Template = FRigVMRegistry::Get().FindTemplate(InNotation);
+	FRigVMTemplate* Template = const_cast<FRigVMTemplate*>(FRigVMRegistry::Get().FindTemplate(InNotation));
 	if(Template)
 	{
 		if(!Template->UsesDispatch())
 		{
 			for(int32 PermutationIndex = 0; PermutationIndex < Template->NumPermutations(); PermutationIndex++)
 			{
-				UnitStructs.Add(Template->GetPermutation(PermutationIndex)->Struct);
+				UnitStructs.Add(Template->GetOrCreatePermutation(PermutationIndex)->Struct);
 			}
 		}
 	}
@@ -13392,7 +13392,7 @@ URigVMNode* URigVMController::AddArrayNode(ERigVMOpCode InOpCode, const FString&
 		OpenUndoBracket(TEXT("Add Array Node"));
 	}
 	
-	const FRigVMTemplate* Template = Factory->GetTemplate();
+	FRigVMTemplate* Template = const_cast<FRigVMTemplate*>(Factory->GetTemplate());
 	URigVMTemplateNode* Node = AddTemplateNode(
 		Template->GetNotation(),
 		InPosition, 
@@ -13435,6 +13435,7 @@ URigVMNode* URigVMController::AddArrayNode(ERigVMOpCode InOpCode, const FString&
 				TArray<int32> Permutations;
 				Template->Resolve(TypeMap, Permutations, false);
 				check(Permutations.Num() == 1);
+				Template->GetOrCreatePermutation(Permutations[0]);
 
 				for(const FRigVMTemplate::FTypePair& Pair : TypeMap)
 				{
@@ -15985,7 +15986,7 @@ bool URigVMController::FullyResolveTemplateNode(URigVMTemplateNode* InNode, int3
 		return true;
 	}
 
-	const FRigVMTemplate* Template = InNode->GetTemplate();
+	FRigVMTemplate* Template = const_cast<FRigVMTemplate*>(InNode->GetTemplate());
 	const FRigVMDispatchFactory* Factory = Template->GetDispatchFactory();
 	
 	InNode->ResolvedPermutation = InPermutationIndex;
@@ -16002,7 +16003,7 @@ bool URigVMController::FullyResolveTemplateNode(URigVMTemplateNode* InNode, int3
 		{
 			FRigVMTemplateTypeMap TypeMap = InNode->GetTemplate()->GetTypesForPermutation(InNode->ResolvedPermutation);
 			const FRigVMFunctionPtr DispatchFunction = Factory->GetOrCreateDispatchFunction(TypeMap);
-			const FRigVMFunction* ResolvedFunction = const_cast<FRigVMTemplate*>(Template)->GetOrCreatePermutation(InNode->ResolvedPermutation);
+			const FRigVMFunction* ResolvedFunction = Template->GetOrCreatePermutation(InNode->ResolvedPermutation);
 			check(DispatchFunction);
 			check(ResolvedFunction);
 			check(ResolvedFunction->FunctionPtr == DispatchFunction);
@@ -16010,7 +16011,7 @@ bool URigVMController::FullyResolveTemplateNode(URigVMTemplateNode* InNode, int3
 	}
 	
 
-	const FRigVMFunction* ResolvedFunction = Template->GetPermutation(InNode->ResolvedPermutation);
+	const FRigVMFunction* ResolvedFunction = Template->GetOrCreatePermutation(InNode->ResolvedPermutation);
 	const TArray<int32> PermutationIndices = {InNode->ResolvedPermutation};
 
 	// find all existing pins that we may need to change
@@ -18965,7 +18966,7 @@ FRigVMClientPatchResult URigVMController::PatchIfSelectNodesOnLoad()
 			const FRigVMDispatchFactory* Factory = FRigVMRegistry::Get().FindOrAddDispatchFactory(
 				bIsIfNode ? FRigVMDispatch_If::StaticStruct() : FRigVMDispatch_SelectInt32::StaticStruct());
 
-			const FRigVMTemplate* Template = Factory->GetTemplate();
+			FRigVMTemplate* Template = const_cast<FRigVMTemplate*>(Factory->GetTemplate());
 			URigVMTemplateNode* NewNode = AddTemplateNode(
 				Template->GetNotation(),
 				NodePosition, 
@@ -18981,7 +18982,11 @@ FRigVMClientPatchResult URigVMController::PatchIfSelectNodesOnLoad()
 				FRigVMTemplateTypeMap Types;
 				Types.Add(IfOrSelectNode->GetPins().Last()->GetFName(), TypeIndex);
 				Template->Resolve(Types, Permutations, false);
-
+				if(Permutations.Num() == 1)
+				{
+					Template->GetOrCreatePermutation(Permutations[0]);
+				}
+				
 				for(URigVMPin* Pin : NewNode->GetPins())
 				{
 					if(Pin->IsWildCard())
