Fixed a crash in RigVM compiler when a rig function has more than 255 operands (was casting the size to 8 bits)

diff --git a/Engine/Plugins/Runtime/RigVM/Source/RigVM/Private/RigVMCore/RigVMByteCode.cpp b/Engine/Plugins/Runtime/RigVM/Source/RigVM/Private/RigVMCore/RigVMByteCode.cpp
--- a/Engine/Plugins/Runtime/RigVM/Source/RigVM/Private/RigVMCore/RigVMByteCode.cpp
+++ b/Engine/Plugins/Runtime/RigVM/Source/RigVM/Private/RigVMCore/RigVMByteCode.cpp
@@ -1059,7 +1059,7 @@ uint64 FRigVMByteCode::AddJumpIfOp(ERigVMOpCode InOpCode, uint16 InInstructionIn
 
 uint64 FRigVMByteCode::AddExecuteOp(uint16 InFunctionIndex, const FRigVMOperandArray& InOperands, const int32& StartPredicateIndex, const int32 PredicateCount)
 {
-	FRigVMExecuteOp Op(InFunctionIndex, (uint8)InOperands.Num());
+	FRigVMExecuteOp Op(InFunctionIndex, IntCastChecked<uint16>(InOperands.Num()));
 	Op.FirstPredicateIndex = StartPredicateIndex;
 	Op.PredicateCount = PredicateCount;
 	uint64 OpByteIndex = AddOp(Op);
diff --git a/Engine/Plugins/Runtime/RigVM/Source/RigVM/Public/RigVMCore/RigVMByteCode.h b/Engine/Plugins/Runtime/RigVM/Source/RigVM/Public/RigVMCore/RigVMByteCode.h
--- a/Engine/Plugins/Runtime/RigVM/Source/RigVM/Public/RigVMCore/RigVMByteCode.h
+++ b/Engine/Plugins/Runtime/RigVM/Source/RigVM/Public/RigVMCore/RigVMByteCode.h
@@ -260,7 +260,7 @@ struct RIGVM_API FRigVMExecuteOp : public FRigVMBaseOp
 	{
 	}
 
-	FRigVMExecuteOp(uint16 InFunctionIndex, uint8 InArgumentCount)
+	FRigVMExecuteOp(uint16 InFunctionIndex, uint16 InArgumentCount)
 	: FRigVMBaseOp(ERigVMOpCode::Execute)
 	, FunctionIndex(InFunctionIndex)
 	, ArgumentCount(InArgumentCount)
