RigVM: Incorrect reload of pins in DispatchSwitch nodes

diff --git a/Engine/Plugins/Runtime/RigVM/Source/RigVMDeveloper/Private/RigVMModel/RigVMController.cpp b/Engine/Plugins/Runtime/RigVM/Source/RigVMDeveloper/Private/RigVMModel/RigVMController.cpp
--- a/Engine/Plugins/Runtime/RigVM/Source/RigVMDeveloper/Private/RigVMModel/RigVMController.cpp
+++ b/Engine/Plugins/Runtime/RigVM/Source/RigVMDeveloper/Private/RigVMModel/RigVMController.cpp
@@ -14565,19 +14565,22 @@ void URigVMController::RepopulatePinsOnNode(URigVMNode* InNode, bool bFollowCore
 			for (int32 ArgIndex = 0; ArgIndex < Template->NumExecuteArguments(DispatchContext); ArgIndex++)
 			{
 				const FRigVMExecuteArgument* Arg = Template->GetExecuteArgument(ArgIndex, DispatchContext);
-				const FRigVMTemplateArgumentType Type = Registry.GetType(Arg->TypeIndex);
-				const TRigVMTypeIndex TypeIndex = Registry.GetTypeIndex(Type);
-
-				FString DefaultValue;
-				if(Registry.IsArrayType(Arg->TypeIndex))
+				if (Arg->Direction == InPinDirection)
 				{
-					if(const FRigVMDispatchFactory* Factory = DispatchNode->GetFactory())
+					const FRigVMTemplateArgumentType Type = Registry.GetType(Arg->TypeIndex);
+					const TRigVMTypeIndex TypeIndex = Registry.GetTypeIndex(Type);
+
+					FString DefaultValue;
+					if(Registry.IsArrayType(Arg->TypeIndex))
 					{
-						DefaultValue = Factory->GetArgumentDefaultValue(Arg->Name, Arg->TypeIndex);
+						if(const FRigVMDispatchFactory* Factory = DispatchNode->GetFactory())
+						{
+							DefaultValue = Factory->GetArgumentDefaultValue(Arg->Name, Arg->TypeIndex);
+						}
 					}
-				}
 
-				(void)NewPinInfos.AddPin(this, INDEX_NONE, Arg->Name, Arg->Direction, TypeIndex, DefaultValue, nullptr, &PreviousPinInfos, false);
+					(void)NewPinInfos.AddPin(this, INDEX_NONE, Arg->Name, Arg->Direction, TypeIndex, DefaultValue, nullptr, &PreviousPinInfos, true);
+				}
 			}
 		};
 
