RigVM: Fix crash when moving an asset and trying to rename a function

diff --git a/Engine/Plugins/Runtime/RigVM/Source/RigVMDeveloper/Private/RigVMModel/RigVMClient.cpp b/Engine/Plugins/Runtime/RigVM/Source/RigVMDeveloper/Private/RigVMModel/RigVMClient.cpp
--- a/Engine/Plugins/Runtime/RigVM/Source/RigVMDeveloper/Private/RigVMModel/RigVMClient.cpp
+++ b/Engine/Plugins/Runtime/RigVM/Source/RigVMDeveloper/Private/RigVMModel/RigVMClient.cpp
@@ -1025,29 +1025,33 @@ void FRigVMClient::HandleGraphModifiedEvent(ERigVMGraphNotifType InNotifType, UR
 				if (URigVMCollapseNode* CollapseNode = Cast<URigVMCollapseNode>(InSubject))
 				{
 					URigVMBuildData* BuildData = URigVMBuildData::Get();
-					IRigVMGraphFunctionHost* Host = Cast<IRigVMGraphFunctionHost>(CollapseNode->GetFunctionIdentifier().HostObject.ResolveObject());
-					FRigVMGraphFunctionData* Data = Host->GetRigVMGraphFunctionStore()->FindFunctionByName(CollapseNode->GetPreviousFName());
-					const FRigVMGraphFunctionIdentifier PreviousFunctionId = Data->Header.LibraryPointer;
-					Data->Header = CollapseNode->GetFunctionHeader();
-
-					if (const FRigVMFunctionReferenceArray* FunctionReferencesPtr = BuildData->GraphFunctionReferences.Find(PreviousFunctionId))
+					if (IRigVMGraphFunctionHost* Host = Cast<IRigVMGraphFunctionHost>(CollapseNode->GetFunctionIdentifier().HostObject.ResolveObject()))
 					{
-						const FRigVMFunctionReferenceArray& FunctionReferences = *FunctionReferencesPtr;
-						BuildData->Modify();
-						FRigVMFunctionReferenceArray& NewFunctionReferences = BuildData->GraphFunctionReferences.Add(Data->Header.LibraryPointer, FunctionReferences);
-						BuildData->GraphFunctionReferences.Remove(PreviousFunctionId);
-						BuildData->MarkPackageDirty();
-
-						for (int32 i=0; i<NewFunctionReferences.Num(); ++i)
+						if (FRigVMGraphFunctionData* Data = Host->GetRigVMGraphFunctionStore()->FindFunctionByName(CollapseNode->GetPreviousFName()))
 						{
-							NewFunctionReferences[i]->ReferencedFunctionHeader = Data->Header;
+							const FRigVMGraphFunctionIdentifier PreviousFunctionId = Data->Header.LibraryPointer;
+							Data->Header = CollapseNode->GetFunctionHeader();
+
+							if (const FRigVMFunctionReferenceArray* FunctionReferencesPtr = BuildData->GraphFunctionReferences.Find(PreviousFunctionId))
+							{
+								const FRigVMFunctionReferenceArray& FunctionReferences = *FunctionReferencesPtr;
+								BuildData->Modify();
+								FRigVMFunctionReferenceArray& NewFunctionReferences = BuildData->GraphFunctionReferences.Add(Data->Header.LibraryPointer, FunctionReferences);
+								BuildData->GraphFunctionReferences.Remove(PreviousFunctionId);
+								BuildData->MarkPackageDirty();
+
+								for (int32 i=0; i<NewFunctionReferences.Num(); ++i)
+								{
+									NewFunctionReferences[i]->ReferencedFunctionHeader = Data->Header;
+								}
+							}
+
+							UpdateGraphFunctionData(CollapseNode);
 						}
 					}
-					
-					UpdateGraphFunctionData(CollapseNode);
 				}
 			}
-			break;	
+			break;
 		}
 		case ERigVMGraphNotifType::NodeColorChanged: // A node's color has changed (Subject == URigVMNode)
 		case ERigVMGraphNotifType::NodeCategoryChanged: // A node's category has changed (Subject == URigVMNode)
