RigVM: Fix compilation issue in ClangEditor

diff --git a/Engine/Plugins/Runtime/RigVM/Source/RigVM/Private/RigVMCore/RigVMTemplate.cpp b/Engine/Plugins/Runtime/RigVM/Source/RigVM/Private/RigVMCore/RigVMTemplate.cpp
--- a/Engine/Plugins/Runtime/RigVM/Source/RigVM/Private/RigVMCore/RigVMTemplate.cpp
+++ b/Engine/Plugins/Runtime/RigVM/Source/RigVM/Private/RigVMCore/RigVMTemplate.cpp
@@ -12,6 +12,23 @@
 
 ////////////////////////////////////////////////////////////////////////////////////////////////
 
+FRigVMTemplateArgumentType::FRigVMTemplateArgumentType(const FName& InCPPType, UObject* InCPPTypeObject)
+		: CPPType(InCPPType)
+		, CPPTypeObject(InCPPTypeObject)
+{
+	// InCppType is unreliable because not all caller knows that
+	// we use generated unique names for user defined structs
+	// so here we override the CppType name with the actual name used in the registry
+	const FString InCPPTypeString = CPPType.ToString();
+	CPPType = *RigVMTypeUtils::PostProcessCPPType(InCPPTypeString, CPPTypeObject);
+#if WITH_EDITOR
+	if (CPPType.IsNone())
+	{
+		UE_LOG(LogRigVM, Warning, TEXT("FRigVMTemplateArgumentType(): Input CPPType '%s' could not be resolved."), *InCPPTypeString);
+	}
+#endif
+}
+
 FRigVMTemplateArgument::FRigVMTemplateArgument()
 	: Index(INDEX_NONE)
 	, Name(NAME_None)
diff --git a/Engine/Plugins/Runtime/RigVM/Source/RigVM/Public/RigVMCore/RigVMTemplate.h b/Engine/Plugins/Runtime/RigVM/Source/RigVM/Public/RigVMCore/RigVMTemplate.h
--- a/Engine/Plugins/Runtime/RigVM/Source/RigVM/Public/RigVMCore/RigVMTemplate.h
+++ b/Engine/Plugins/Runtime/RigVM/Source/RigVM/Public/RigVMCore/RigVMTemplate.h
@@ -65,21 +65,7 @@ struct RIGVM_API FRigVMTemplateArgumentType
 		CPPTypeObject = RigVMTypeUtils::GetWildCardCPPTypeObject();
 	}
 
-	FRigVMTemplateArgumentType(const FName& InCPPType, UObject* InCPPTypeObject = nullptr)
-		: CPPType(InCPPType)
-		, CPPTypeObject(InCPPTypeObject)
-	{
-		// InCppType is unreliable because not all caller knows that
-		// we use generated unique names for user defined structs
-		// so here we override the CppType name with the actual name used in the registry
-		const FString InCPPTypeString = CPPType.ToString();
-		CPPType = *RigVMTypeUtils::PostProcessCPPType(InCPPTypeString, CPPTypeObject);
-
-		if (CPPType.IsNone())
-		{
-			UE_LOG(LogRigVM, Warning, TEXT("FRigVMTemplateArgumentType(): Input CPPType '%s' could not be resolved."), *InCPPTypeString);
-		}
-	}
+	FRigVMTemplateArgumentType(const FName& InCPPType, UObject* InCPPTypeObject = nullptr);
 
 	FRigVMTemplateArgumentType(UClass* InClass)
 	: CPPType(*RigVMTypeUtils::CPPTypeFromObject(InClass))
