RigVMEditor: EdGraph, EdGraphNode, EdGraphSchema and Blueprint

diff --git a/Engine/Plugins/Runtime/RigVM/Source/RigVMDeveloper/Private/RigVMBlueprint.cpp b/Engine/Plugins/Runtime/RigVM/Source/RigVMDeveloper/Private/RigVMBlueprint.cpp
new file mode 100644
--- /dev/null
+++ b/Engine/Plugins/Runtime/RigVM/Source/RigVMDeveloper/Private/RigVMBlueprint.cpp
@@ -0,0 +1,3988 @@
+// Copyright Epic Games, Inc. All Rights Reserved.
+
+#include "RigVMBlueprint.h"
+
+#include "RigVMBlueprintGeneratedClass.h"
+#include "EdGraph/EdGraph.h"
+#include "EdGraphNode_Comment.h"
+#include "BlueprintActionDatabaseRegistrar.h"
+#include "EdGraph/RigVMEdGraph.h"
+#include "EdGraph/RigVMEdGraphNode.h"
+#include "EdGraph/RigVMEdGraphSchema.h"
+#include "UObject/ObjectSaveContext.h"
+#include "UObject/UObjectGlobals.h"
+#include "RigVMObjectVersion.h"
+#include "BlueprintCompilationManager.h"
+#include "RigVMCompiler/RigVMCompiler.h"
+#include "RigVMCore/RigVMRegistry.h"
+#include "AssetRegistry/AssetRegistryModule.h"
+#include "RigVMPythonUtils.h"
+#include "RigVMTypeUtils.h"
+#include "Algo/Count.h"
+#include "RigVMModel/Nodes/RigVMAggregateNode.h"
+#include "RigVMModel/Nodes/RigVMDispatchNode.h"
+#include "UObject/UE5MainStreamObjectVersion.h"
+
+#include UE_INLINE_GENERATED_CPP_BY_NAME(RigVMBlueprint)
+
+#if WITH_EDITOR
+#include "Kismet2/KismetDebugUtilities.h"
+#include "Kismet2/WatchedPin.h"
+#include "Kismet2/BlueprintEditorUtils.h"
+#include "RigVMBlueprintUtils.h"
+#include "Editor/UnrealEdEngine.h"
+#include "Editor/Transactor.h"
+#include "CookOnTheSide/CookOnTheFlyServer.h"
+#endif//WITH_EDITOR
+
+#define LOCTEXT_NAMESPACE "RigVMBlueprint"
+
+TAutoConsoleVariable<bool> CVarRigVMEnablePreLoadFiltering(
+	TEXT("RigVM.EnablePreLoadFiltering"),
+	true,
+	TEXT("When true the RigVMGraphs will be skipped during preload to speed up load times."));
+
+TAutoConsoleVariable<bool> CVarRigVMEnablePostLoadHashing(
+	TEXT("RigVM.EnablePostLoadHashing"),
+	true,
+	TEXT("When true refreshing the RigVMGraphs will be skipped if the hash matches the serialized hash."));
+
+static TArray<UClass*> GetClassObjectsInPackage(UPackage* InPackage)
+{
+	TArray<UObject*> Objects;
+	GetObjectsWithOuter(InPackage, Objects, false);
+
+	TArray<UClass*> ClassObjects;
+	for (UObject* Object : Objects)
+	{
+		if (UClass* Class = Cast<UClass>(Object))
+		{
+			ClassObjects.Add(Class);
+		}
+	}
+
+	return ClassObjects;
+}
+
+FSoftObjectPath URigVMBlueprint::PreDuplicateAssetPath;
+FSoftObjectPath URigVMBlueprint::PreDuplicateHostPath;
+
+URigVMBlueprint::URigVMBlueprint(const FObjectInitializer& ObjectInitializer)
+{
+	bSuspendModelNotificationsForSelf = false;
+	bSuspendModelNotificationsForOthers = false;
+	bSuspendAllNotifications = false;
+	bSuspendPythonMessagesForRigVMClient = true;
+	bMarkBlueprintAsStructurallyModifiedPending = false;
+
+#if WITH_EDITORONLY_DATA
+	ReferencedObjectPathsStored = false;
+#endif
+
+	bRecompileOnLoad = 0;
+	bAutoRecompileVM = true;
+	bVMRecompilationRequired = false;
+	bIsCompiling = false;
+	VMRecompilationBracket = 0;
+
+	RigVMClient.SetOuterClientHost(this, GET_MEMBER_NAME_CHECKED(URigVMBlueprint, RigVMClient));
+	RigVMClient.SetSchemaClass(GetRigVMSchemaClass());
+	{
+		TGuardValue<bool> DisableClientNotifs(RigVMClient.bSuspendNotifications, true);
+		RigVMClient.GetOrCreateFunctionLibrary(false, &ObjectInitializer, false);
+		RigVMClient.AddModel(FRigVMClient::RigVMModelPrefix, false, &ObjectInitializer, false);
+	}
+	RigVMClient.SetExecuteContextStruct(GetRigVMExecuteContextStruct());
+	
+	FunctionLibraryEdGraph = Cast<URigVMEdGraph>(ObjectInitializer.CreateDefaultSubobject(this, TEXT("RigVMFunctionLibraryEdGraph"), GetRigVMEdGraphClass(), GetRigVMEdGraphClass()));
+	FunctionLibraryEdGraph->Schema = GetRigVMEdGraphSchemaClass();
+	FunctionLibraryEdGraph->bAllowRenaming = 0;
+	FunctionLibraryEdGraph->bEditable = 0;
+	FunctionLibraryEdGraph->bAllowDeletion = 0;
+	FunctionLibraryEdGraph->bIsFunctionDefinition = false;
+	FunctionLibraryEdGraph->ModelNodePath = RigVMClient.GetFunctionLibrary()->GetNodePath();
+	FunctionLibraryEdGraph->Initialize(this);
+
+	bUpdatingExternalVariables = false;
+	
+	bDirtyDuringLoad = false;
+	bErrorsDuringCompilation = false;
+
+	SupportedEventNames.Reset();
+
+	VMCompileSettings.ASTSettings.ReportDelegate.BindUObject(this, &URigVMBlueprint::HandleReportFromCompiler);
+
+#if WITH_EDITOR
+	CompileLog.SetSourcePath(GetPathName());
+	CompileLog.bLogDetailedResults = false;
+	CompileLog.EventDisplayThresholdMs = false;
+#endif
+}
+
+URigVMBlueprint::URigVMBlueprint()
+{
+}
+
+void URigVMBlueprint::InitializeModelIfRequired(bool bRecompileVM)
+{
+	DECLARE_SCOPE_HIERARCHICAL_COUNTER_FUNC()
+
+	if (RigVMClient.GetController(0) == nullptr)
+	{
+		check(RigVMClient.Num() == 1);
+		check(RigVMClient.GetFunctionLibrary());
+		
+		RigVMClient.GetOrCreateController(RigVMClient.GetDefaultModel());
+		RigVMClient.GetOrCreateController(RigVMClient.GetFunctionLibrary());
+
+		bool bRecompileRequired = false;
+		for (int32 i = 0; i < UbergraphPages.Num(); ++i)
+		{
+			if (URigVMEdGraph* Graph = Cast<URigVMEdGraph>(UbergraphPages[i]))
+			{
+				if (bRecompileVM)
+				{
+					bRecompileRequired = true;
+				}
+
+				Graph->Initialize(this);
+			}
+		}
+
+		if(bRecompileRequired)
+		{
+			RecompileVM();
+		}
+
+		FunctionLibraryEdGraph->Initialize(this);
+	}
+}
+
+URigVMBlueprintGeneratedClass* URigVMBlueprint::GetRigVMBlueprintGeneratedClass() const
+{
+	URigVMBlueprintGeneratedClass* Result = Cast<URigVMBlueprintGeneratedClass>(*GeneratedClass);
+	return Result;
+}
+
+URigVMBlueprintGeneratedClass* URigVMBlueprint::GetRigVMBlueprintSkeletonClass() const
+{
+	URigVMBlueprintGeneratedClass* Result = Cast<URigVMBlueprintGeneratedClass>(*SkeletonGeneratedClass);
+	return Result;
+}
+
+UClass* URigVMBlueprint::GetBlueprintClass() const
+{
+	return URigVMBlueprintGeneratedClass::StaticClass();
+}
+
+UClass* URigVMBlueprint::RegenerateClass(UClass* ClassToRegenerate, UObject* PreviousCDO)
+{
+	UClass* Result;
+	{
+		TGuardValue<bool> NotificationGuard(bSuspendAllNotifications, true);
+		Result = Super::RegenerateClass(ClassToRegenerate, PreviousCDO);
+	}
+	return Result;
+}
+
+void URigVMBlueprint::LoadModulesRequiredForCompilation() 
+{
+}
+
+bool URigVMBlueprint::ExportGraphToText(UEdGraph* InEdGraph, FString& OutText)
+{
+	OutText.Empty();
+
+	if (URigVMGraph* RigGraph = GetModel(InEdGraph))
+	{
+		if (URigVMCollapseNode* CollapseNode = Cast<URigVMCollapseNode>(RigGraph->GetOuter()))
+		{
+			if (URigVMController* Controller = GetOrCreateController(CollapseNode->GetGraph()))
+			{
+				TArray<FName> NodeNamesToExport;
+				NodeNamesToExport.Add(CollapseNode->GetFName());
+				OutText = Controller->ExportNodesToText(NodeNamesToExport);
+			}
+		}
+	}
+
+	// always return true so that the default mechanism doesn't take over
+	return true;
+}
+
+bool URigVMBlueprint::CanImportGraphFromText(const FString& InClipboardText)
+{
+	return GetTemplateController(true)->CanImportNodesFromText(InClipboardText);
+}
+
+bool URigVMBlueprint::RequiresForceLoadMembers(UObject* InObject) const
+{
+	// only filter if the console variable is enabled
+	if(!CVarRigVMEnablePreLoadFiltering->GetBool())
+	{
+		return UBlueprint::RequiresForceLoadMembers(InObject);
+	}
+
+	/* todoooo
+	// old assets don't support preload filtering
+	if (GetLinkerCustomVersion(FControlRigObjectVersion::GUID) < FControlRigObjectVersion::RemoveParameters)
+	{
+		return UBlueprint::RequiresForceLoadMembers(InObject);
+	}
+	*/
+	
+	// we can stop traversing when hitting a URigVMNode
+	// except for collapse nodes - since they contain a graphs again
+	// and variable  nodes - since they are needed during preload by the BP compiler
+	if(InObject->IsA<URigVMNode>())
+	{
+		if(!InObject->IsA<URigVMCollapseNode>() &&
+			!InObject->IsA<URigVMVariableNode>())
+		{
+			return false;
+		}
+	}
+	return UBlueprint::RequiresForceLoadMembers(InObject);
+}
+
+void URigVMBlueprint::PostEditChangeChainProperty(FPropertyChangedChainEvent& PropertyChangedEvent)
+{
+	Super::PostEditChangeChainProperty(PropertyChangedEvent);
+	PostEditChangeChainPropertyEvent.Broadcast(PropertyChangedEvent);
+}
+
+void URigVMBlueprint::PostRename(UObject* OldOuter, const FName OldName)
+{
+	Super::PostRename(OldOuter, OldName);
+
+	// Whenever the asset is renamed/moved, generated classes parented to the old package
+	// are not moved to the new package automatically (see FAssetRenameManager), so we
+	// have to manually perform the move/rename, to avoid invalid reference to the old package
+	
+	// Note: while asset duplication doesn't duplicate the classes either, it is not a problem there
+	// because we always recompile in post duplicate.
+	TArray<UClass*> ClassObjects = GetClassObjectsInPackage(OldOuter->GetPackage());
+
+	for (UClass* ClassObject : ClassObjects)
+	{
+		if (URigVMMemoryStorageGeneratorClass* MemoryClass = Cast<URigVMMemoryStorageGeneratorClass>(ClassObject))
+		{
+			MemoryClass->Rename(nullptr, GetPackage(), REN_ForceNoResetLoaders | REN_DoNotDirty | REN_DontCreateRedirectors | REN_NonTransactional);
+		}
+	}
+}
+
+void URigVMBlueprint::GetPreloadDependencies(TArray<UObject*>& OutDeps)
+{
+	Super::GetPreloadDependencies(OutDeps);
+
+	TArray<UClass*> ClassObjects = GetClassObjectsInPackage(GetPackage());
+
+	for (UClass* ClassObject : ClassObjects)
+	{
+		if (URigVMMemoryStorageGeneratorClass* MemoryClass = Cast<URigVMMemoryStorageGeneratorClass>(ClassObject))
+		{
+			OutDeps.Add(MemoryClass);
+		}
+	}
+}
+
+FRigVMClient* URigVMBlueprint::GetRigVMClient()
+{
+	return &RigVMClient;
+}
+
+const FRigVMClient* URigVMBlueprint::GetRigVMClient() const
+{
+	return &RigVMClient;
+}
+
+IRigVMGraphFunctionHost* URigVMBlueprint::GetRigVMGraphFunctionHost() 
+{
+	return GetRigVMBlueprintGeneratedClass();
+}
+
+const IRigVMGraphFunctionHost* URigVMBlueprint::GetRigVMGraphFunctionHost() const 
+{
+	return GetRigVMBlueprintGeneratedClass();
+}
+
+UObject* URigVMBlueprint::GetEditorObjectForRigVMGraph(URigVMGraph* InVMGraph) const
+{
+	if(InVMGraph)
+	{
+		if(InVMGraph->GetOutermost() != GetOutermost())
+		{
+			return nullptr;
+		}
+
+		if(InVMGraph->IsA<URigVMFunctionLibrary>())
+		{
+			return FunctionLibraryEdGraph;
+		}
+
+		TArray<UEdGraph*> EdGraphs;
+		GetAllGraphs(EdGraphs);
+
+		bool bIsFunctionDefinition = false;
+		if (URigVMLibraryNode* LibraryNode = Cast<URigVMLibraryNode>(InVMGraph->GetOuter()))
+		{
+			bIsFunctionDefinition = LibraryNode->GetGraph()->IsA<URigVMFunctionLibrary>();
+		}
+
+		for (UEdGraph* EdGraph : EdGraphs)
+		{
+			if (URigVMEdGraph* RigGraph = Cast<URigVMEdGraph>(EdGraph))
+			{
+				if (RigGraph->bIsFunctionDefinition != bIsFunctionDefinition)
+				{
+					continue;
+				}
+
+				if ((RigGraph->ModelNodePath == InVMGraph->GetNodePath()) ||
+					(RigGraph->ModelNodePath.IsEmpty() && (RigVMClient.GetDefaultModel() == InVMGraph)))
+				{
+					return RigGraph;
+				}
+			}
+		}
+	}
+	
+	return nullptr;
+}
+
+void URigVMBlueprint::HandleRigVMGraphAdded(const FRigVMClient* InClient, const FString& InNodePath)
+{
+	if(URigVMGraph* Model = InClient->GetModel(InNodePath))
+	{
+		if(!HasAnyFlags(RF_ClassDefaultObject | RF_NeedInitialization | RF_NeedLoad | RF_NeedPostLoad) &&
+			GetOuter() != GetTransientPackage())
+		{
+			CreateEdGraph(Model, true);
+			RecompileVM();
+		}
+
+#if WITH_EDITOR
+		if(!bSuspendPythonMessagesForRigVMClient)
+		{
+			const FString BlueprintName = InClient->GetSchema()->GetSanitizedName(GetName(), true, false);
+			RigVMPythonUtils::Print(BlueprintName, 
+				FString::Printf(TEXT("blueprint.add_model('%s')"),
+					*Model->GetName()));
+		}
+#endif
+	}
+}
+
+void URigVMBlueprint::HandleRigVMGraphRemoved(const FRigVMClient* InClient, const FString& InNodePath)
+{
+	if(URigVMGraph* Model = InClient->GetModel(InNodePath))
+	{
+		RemoveEdGraph(Model);
+		RecompileVM();
+
+#if WITH_EDITOR
+		if(!bSuspendPythonMessagesForRigVMClient)
+		{
+			const FString BlueprintName = InClient->GetSchema()->GetSanitizedName(GetName(), true, false);
+			RigVMPythonUtils::Print(BlueprintName, 
+				FString::Printf(TEXT("blueprint.remove_model('%s')"),
+					*Model->GetName()));
+		}
+#endif
+	}
+}
+
+void URigVMBlueprint::HandleRigVMGraphRenamed(const FRigVMClient* InClient, const FString& InOldNodePath, const FString& InNewNodePath)
+{
+	if(InClient->GetModel(InNewNodePath))
+	{
+		TArray<UEdGraph*> EdGraphs;
+		GetAllGraphs(EdGraphs);
+
+		for (UEdGraph* EdGraph : EdGraphs)
+		{
+			if (URigVMEdGraph* RigGraph = Cast<URigVMEdGraph>(EdGraph))
+			{
+				RigGraph->HandleRigVMGraphRenamed(InOldNodePath, InNewNodePath);
+			}
+		}
+
+		FBlueprintEditorUtils::MarkBlueprintAsStructurallyModified(this);
+	}
+}
+
+void URigVMBlueprint::HandleConfigureRigVMController(const FRigVMClient* InClient,
+	URigVMController* InControllerToConfigure)
+{
+	InControllerToConfigure->OnModified().AddUObject(this, &URigVMBlueprint::HandleModifiedEvent);
+
+	TWeakObjectPtr<URigVMBlueprint> WeakThis(this);
+
+	// this delegate is used by the controller to determine variable validity
+	// during a bind process. the controller itself doesn't own the variables,
+	// so we need a delegate to request them from the owning blueprint
+	InControllerToConfigure->GetExternalVariablesDelegate.BindLambda([](URigVMGraph* InGraph) -> TArray<FRigVMExternalVariable> {
+
+		if (InGraph)
+		{
+			if(URigVMBlueprint* Blueprint = InGraph->GetTypedOuter<URigVMBlueprint>())
+			{
+				if (URigVMBlueprintGeneratedClass* RigClass = Blueprint->GetRigVMBlueprintGeneratedClass())
+				{
+                    if (URigVMHost* CDO = Cast<URigVMHost>(RigClass->GetDefaultObject(true /* create if needed */)))
+                    {
+                        return CDO->GetExternalVariablesImpl(true /* rely on variables within blueprint */);
+                    }
+                }
+			}
+		}
+		return TArray<FRigVMExternalVariable>();
+
+	});
+
+
+	// this delegate is used by the controller to retrieve the current bytecode of the VM
+	InControllerToConfigure->GetCurrentByteCodeDelegate.BindLambda([WeakThis]() -> const FRigVMByteCode* {
+
+		if (WeakThis.IsValid())
+		{
+			if (URigVMBlueprintGeneratedClass* RigClass = WeakThis->GetRigVMBlueprintGeneratedClass())
+			{
+				if (URigVMHost* CDO = Cast<URigVMHost>(RigClass->GetDefaultObject(false)))
+				{
+					if (URigVM* VM = CDO->GetVM())
+					{
+						return &VM->GetByteCode();
+					}
+				}
+			}
+		}
+		return nullptr;
+
+	});
+
+#if WITH_EDITOR
+
+	// this sets up three delegates:
+	// a) get external variables (mapped to Controller->GetExternalVariables)
+	// b) bind pin to variable (mapped to Controller->BindPinToVariable)
+	// c) create external variable (mapped to the passed in tfunction)
+	// the last one is defined within the blueprint since the controller
+	// doesn't own the variables and can't create one itself.
+	InControllerToConfigure->SetupDefaultUnitNodeDelegates(TDelegate<FName(FRigVMExternalVariable, FString)>::CreateLambda(
+		[WeakThis](FRigVMExternalVariable InVariableToCreate, FString InDefaultValue) -> FName {
+			if (WeakThis.IsValid())
+			{
+				return WeakThis->AddHostMemberVariableFromExternal(InVariableToCreate, InDefaultValue);
+			}
+			return NAME_None;
+		}
+	));
+
+	TWeakObjectPtr<URigVMController> WeakController = InControllerToConfigure;
+	InControllerToConfigure->RequestBulkEditDialogDelegate.BindLambda([WeakThis, WeakController](URigVMLibraryNode* InFunction, ERigVMControllerBulkEditType InEditType) -> FRigVMController_BulkEditResult 
+	{
+		if(WeakThis.IsValid() && WeakController.IsValid())
+		{
+			URigVMBlueprint* StrongThis = WeakThis.Get();
+            URigVMController* StrongController = WeakController.Get();
+            if(StrongThis->OnRequestBulkEditDialog().IsBound())
+			{
+				return StrongThis->OnRequestBulkEditDialog().Execute(StrongThis, StrongController, InFunction, InEditType);
+			}
+		}
+		return FRigVMController_BulkEditResult();
+	});
+
+	InControllerToConfigure->RequestBreakLinksDialogDelegate.BindLambda([WeakThis, WeakController](TArray<URigVMLink*> InLinks) -> bool 
+	{
+		if(WeakThis.IsValid() && WeakController.IsValid())
+		{
+			URigVMBlueprint* StrongThis = WeakThis.Get();
+			if(StrongThis->OnRequestBreakLinksDialog().IsBound())
+			{
+				return StrongThis->OnRequestBreakLinksDialog().Execute(InLinks);
+			}
+		}
+		return false;
+	});
+
+	InControllerToConfigure->RequestPinTypeSelectionDelegate.BindLambda([WeakThis](const TArray<TRigVMTypeIndex>& InTypes) -> TRigVMTypeIndex 
+	{
+		if(WeakThis.IsValid())
+		{
+			URigVMBlueprint* StrongThis = WeakThis.Get();
+			if(StrongThis->OnRequestPinTypeSelectionDialog().IsBound())
+			{
+				return StrongThis->OnRequestPinTypeSelectionDialog().Execute(InTypes);
+			}
+		}
+		return INDEX_NONE;
+	});
+
+	InControllerToConfigure->RequestNewExternalVariableDelegate.BindLambda([WeakThis](FRigVMGraphVariableDescription InVariable, bool bInIsPublic, bool bInIsReadOnly) -> FName
+	{
+		if (WeakThis.IsValid())
+		{
+			for (FBPVariableDescription& ExistingVariable : WeakThis->NewVariables)
+			{
+				if (ExistingVariable.VarName == InVariable.Name)
+				{
+					return FName();
+				}
+			}
+
+			FRigVMExternalVariable ExternalVariable = InVariable.ToExternalVariable();
+			return WeakThis->AddMemberVariable(InVariable.Name,
+				ExternalVariable.TypeObject ? ExternalVariable.TypeObject->GetPathName() : ExternalVariable.TypeName.ToString(),
+				bInIsPublic,
+				bInIsReadOnly,
+				InVariable.DefaultValue);
+		}
+		
+		return FName();
+	});
+
+
+	InControllerToConfigure->RequestJumpToHyperlinkDelegate.BindLambda([WeakThis](const UObject* InSubject)
+	{
+		if (WeakThis.IsValid())
+		{
+			URigVMBlueprint* StrongThis = WeakThis.Get();
+			if(StrongThis->OnRequestJumpToHyperlink().IsBound())
+			{
+				StrongThis->OnRequestJumpToHyperlink().Execute(InSubject);
+			}
+		}
+	});
+
+	/*
+	 * todooooo
+	InControllerToConfigure->ConfigureWorkflowOptionsDelegate.BindLambda([WeakThis](URigVMUserWorkflowOptions* Options)
+	{
+		if(UControlRigWorkflowOptions* ControlRigNodeWorkflowOptions = Cast<UControlRigWorkflowOptions>(Options))
+		{
+			ControlRigNodeWorkflowOptions->Hierarchy = nullptr;
+			ControlRigNodeWorkflowOptions->Selection.Reset();
+			
+			if(const URigVMBlueprint* StrongThis = WeakThis.Get())
+			{
+				if(UControlRig* ControlRig = Cast<UControlRig>(StrongThis->GetObjectBeingDebugged()))
+				{
+					ControlRigNodeWorkflowOptions->Hierarchy = ControlRig->GetHierarchy();
+				}
+				ControlRigNodeWorkflowOptions->Selection = StrongThis->Hierarchy->GetSelectedKeys();
+			}
+		}
+	});
+	*/
+
+#endif
+}
+
+bool URigVMBlueprint::TryImportGraphFromText(const FString& InClipboardText, UEdGraph** OutGraphPtr)
+{
+	if (OutGraphPtr)
+	{
+		*OutGraphPtr = nullptr;
+	}
+
+	if (URigVMController* FunctionLibraryController = GetOrCreateController(GetLocalFunctionLibrary()))
+	{
+		TGuardValue<FRigVMController_RequestLocalizeFunctionDelegate> RequestLocalizeDelegateGuard(
+            FunctionLibraryController->RequestLocalizeFunctionDelegate,
+            FRigVMController_RequestLocalizeFunctionDelegate::CreateLambda([this](FRigVMGraphFunctionIdentifier& InFunctionToLocalize)
+            {
+            	BroadcastRequestLocalizeFunctionDialog(InFunctionToLocalize);
+				const URigVMLibraryNode* LocalizedFunctionNode = GetLocalFunctionLibrary()->FindPreviouslyLocalizedFunction(InFunctionToLocalize);
+				return LocalizedFunctionNode != nullptr;
+            })
+        );
+		
+		TArray<FName> ImportedNodeNames = FunctionLibraryController->ImportNodesFromText(InClipboardText, true, true);
+		if (ImportedNodeNames.Num() == 0)
+		{
+			return false;
+		}
+
+		URigVMCollapseNode* CollapseNode = Cast<URigVMCollapseNode>(GetLocalFunctionLibrary()->FindFunction(ImportedNodeNames[0]));
+		if (ImportedNodeNames.Num() > 1 || CollapseNode == nullptr || CollapseNode->GetContainedGraph() == nullptr)
+		{
+			FunctionLibraryController->Undo();
+			return false;
+		}
+
+		UEdGraph* EdGraph = GetEdGraph(CollapseNode->GetContainedGraph());
+		if (OutGraphPtr)
+		{
+			*OutGraphPtr = EdGraph;
+		}
+
+		BroadcastGraphImported(EdGraph);
+	}
+
+	// always return true so that the default mechanism doesn't take over
+	return true;
+}
+
+void URigVMBlueprint::Serialize(FArchive& Ar)
+{
+	RigVMClient.SetOuterClientHost(this, GET_MEMBER_NAME_CHECKED(URigVMBlueprint, RigVMClient));
+	
+	Super::Serialize(Ar);
+
+	if(Ar.IsObjectReferenceCollector())
+	{
+		Ar.UsingCustomVersion(FRigVMObjectVersion::GUID);
+
+#if WITH_EDITORONLY_DATA
+		if (Ar.IsCooking() && ReferencedObjectPathsStored)
+		{
+			for (FSoftObjectPath ObjectPath : ReferencedObjectPaths)
+			{
+				ObjectPath.Serialize(Ar);
+			}
+		}
+		else
+#endif
+		{
+			TArray<IRigVMGraphFunctionHost*> ReferencedFunctionHosts = GetReferencedFunctionHosts(false);
+
+			for(IRigVMGraphFunctionHost* ReferencedFunctionHost : ReferencedFunctionHosts)
+			{
+				if (URigVMBlueprintGeneratedClass* BPGeneratedClass = Cast<URigVMBlueprintGeneratedClass>(ReferencedFunctionHost))
+				{
+					Ar << BPGeneratedClass;
+				}
+			}
+		}
+	}
+
+	if(Ar.IsLoading())
+	{
+		if(Model_DEPRECATED || FunctionLibrary_DEPRECATED)
+		{
+			TGuardValue<bool> DisableClientNotifs(RigVMClient.bSuspendNotifications, true);
+			RigVMClient.SetFromDeprecatedData(Model_DEPRECATED, FunctionLibrary_DEPRECATED);
+		}
+	}
+}
+
+void URigVMBlueprint::PreSave(const class ITargetPlatform* TargetPlatform)
+{
+	PRAGMA_DISABLE_DEPRECATION_WARNINGS;
+	Super::PreSave(TargetPlatform);
+	PRAGMA_ENABLE_DEPRECATION_WARNINGS;
+}
+
+void URigVMBlueprint::PreSave(FObjectPreSaveContext ObjectSaveContext)
+{
+	Super::PreSave(ObjectSaveContext);
+
+	RigVMClient.PreSave(ObjectSaveContext);
+
+	SupportedEventNames.Reset();
+	if (URigVMBlueprintGeneratedClass* RigClass = GetRigVMBlueprintGeneratedClass())
+	{
+		if (const URigVMHost* CDO = Cast<URigVMHost>(RigClass->GetDefaultObject(true /* create if needed */)))
+		{
+			SupportedEventNames = CDO->GetSupportedEvents();
+		}
+
+		PublicGraphFunctions.Reset();
+		PublicGraphFunctions.SetNum(RigClass->GetRigVMGraphFunctionStore()->PublicFunctions.Num());
+		for (int32 i=0; i<PublicGraphFunctions.Num(); ++i)
+		{
+			PublicGraphFunctions[i] = RigClass->GetRigVMGraphFunctionStore()->PublicFunctions[i].Header;
+		}
+	}
+
+#if WITH_EDITORONLY_DATA
+	ReferencedObjectPaths.Reset();
+
+	TArray<IRigVMGraphFunctionHost*> ReferencedFunctionHosts = GetReferencedFunctionHosts(false);
+	for(IRigVMGraphFunctionHost* ReferencedFunctionHost : ReferencedFunctionHosts)
+	{
+		if (URigVMBlueprintGeneratedClass* BPGeneratedClass = Cast<URigVMBlueprintGeneratedClass>(ReferencedFunctionHost))
+		{
+			ReferencedObjectPaths.AddUnique(BPGeneratedClass);
+		}
+	}
+
+	ReferencedObjectPathsStored = true;
+#endif
+
+	FunctionReferenceNodeData = GetReferenceNodeData();
+	IAssetRegistry::GetChecked().AssetTagsFinalized(*this);
+
+	CachedAssetTags.Reset();
+}
+
+void URigVMBlueprint::PostLoad()
+{
+	Super::PostLoad();
+
+	bVMRecompilationRequired = true;
+	{
+		TGuardValue<bool> IsCompilingGuard(bIsCompiling, true);
+		
+		TArray<IRigVMGraphFunctionHost*> ReferencedFunctionHosts = GetReferencedFunctionHosts(true);
+
+		// PostLoad all referenced function hosts so that their function data are fully loaded 
+		// and ready to be inlined into this BP during compilation
+		for (IRigVMGraphFunctionHost* FunctionHost : ReferencedFunctionHosts)
+		{
+			if (URigVMBlueprintGeneratedClass* BPGeneratedClass = Cast<URigVMBlueprintGeneratedClass>(FunctionHost))
+			{
+				if (BPGeneratedClass->HasAllFlags(RF_NeedPostLoad))
+				{
+					BPGeneratedClass->ConditionalPostLoad();
+				}
+			}
+		}
+		
+		// temporarily disable default value validation during load time, serialized values should always be accepted
+		TGuardValue<bool> DisablePinDefaultValueValidation(GetOrCreateController()->bValidatePinDefaults, false);
+
+		// remove all non-controlrig-graphs
+		TArray<UEdGraph*> NewUberGraphPages;
+		for (UEdGraph* Graph : UbergraphPages)
+		{
+			URigVMEdGraph* RigGraph = Cast<URigVMEdGraph>(Graph);
+			if (RigGraph)
+			{
+				NewUberGraphPages.Add(RigGraph);
+			}
+			else
+			{
+				Graph->MarkAsGarbage();
+				Graph->Rename(nullptr, GetTransientPackage(), REN_ForceNoResetLoaders);
+			}
+		}
+		UbergraphPages = NewUberGraphPages;
+
+		TArray<TGuardValue<bool>> EditableGraphGuards;
+		{
+			for (URigVMGraph* Graph : GetAllModels())
+			{
+				EditableGraphGuards.Emplace(Graph->bEditable, true);
+			}
+		}
+		
+		InitializeModelIfRequired(false /* recompile vm */);
+		{
+			TGuardValue<bool> GuardNotifications(bSuspendModelNotificationsForSelf, true);
+			
+			const FRigVMClientPatchResult PatchResult = GetRigVMClient()->PatchModelsOnLoad();
+			if(PatchResult.RequiresToMarkPackageDirty())
+			{
+				(void)MarkPackageDirty();
+				bDirtyDuringLoad = true;
+			}
+			
+			PatchFunctionReferencesOnLoad();
+			PatchVariableNodesOnLoad();
+			PatchVariableNodesWithIncorrectType();
+			PathDomainSpecificContentOnLoad();
+			PatchBoundVariables();
+			PatchParameterNodesOnLoad();
+			PatchLinksWithCast();
+			PatchFunctionsOnLoad();
+		}
+
+#if WITH_EDITOR
+
+		{
+			TGuardValue<bool> GuardNotifications(bSuspendModelNotificationsForSelf, true);
+
+			// refresh the graph such that the pin hierarchies matches their CPPTypeObject
+			// this step is needed everytime we open a BP in the editor, b/c even after load
+			// model data can change while the Control Rig BP is not opened
+			// for example, if a user defined struct changed after BP load,
+			// any pin that references the struct needs to be regenerated
+			RefreshAllModels();
+		}
+		
+		// at this point we may still have links which are detached. we may or may not be able to 
+		// reattach them.
+		GetRigVMClient()->ProcessDetachedLinks();
+
+		GetRigVMBlueprintGeneratedClass()->GetRigVMGraphFunctionStore()->RemoveAllCompilationData();
+
+		// perform backwards compat value upgrades
+		TArray<URigVMGraph*> GraphsToValidate = GetAllModels();
+		for (int32 GraphIndex = 0; GraphIndex < GraphsToValidate.Num(); GraphIndex++)
+		{
+			URigVMGraph* GraphToValidate = GraphsToValidate[GraphIndex];
+			if(GraphToValidate == nullptr)
+			{
+				continue;
+			}
+
+			for(URigVMNode* Node : GraphToValidate->GetNodes())
+			{
+				URigVMController* Controller = GetOrCreateController(GraphToValidate);
+				FRigVMControllerNotifGuard NotifGuard(Controller, true);
+				Controller->RemoveUnusedOrphanedPins(Node, true);
+			}
+				
+			for(URigVMNode* Node : GraphToValidate->GetNodes())
+			{
+				// avoid function reference related validation for temp assets, a temp asset may get generated during
+				// certain content validation process. It is usually just a simple file-level copy of the source asset
+				// so these references are usually not fixed-up properly. Thus, it is meaningless to validate them.
+				// They should not be allowed to dirty the source asset either.
+				if (!this->GetPackage()->GetName().StartsWith("/Temp/"))
+				{
+					if(URigVMFunctionReferenceNode* FunctionReferenceNode = Cast<URigVMFunctionReferenceNode>(Node))
+					{
+						if(URigVMBuildData* BuildData = URigVMBuildData::Get())
+						{
+							BuildData->RegisterFunctionReference(FunctionReferenceNode);
+						}
+					}
+				}
+			}
+		}
+
+		CompileLog.Messages.Reset();
+		CompileLog.NumErrors = CompileLog.NumWarnings = 0;
+#endif
+	}
+
+#if WITH_EDITOR
+	if(GIsEditor)
+	{
+		// delay compilation until the package has been loaded
+		FCoreUObjectDelegates::OnEndLoadPackage.AddUObject(this, &URigVMBlueprint::HandlePackageDone);
+	}
+#else
+	RecompileVMIfRequired();
+#endif
+	RequestRigVMInit();
+
+	FCoreUObjectDelegates::OnObjectModified.RemoveAll(this);
+	OnChanged().RemoveAll(this);
+	FCoreUObjectDelegates::OnObjectModified.AddUObject(this, &URigVMBlueprint::OnPreVariableChange);
+	OnChanged().AddUObject(this, &URigVMBlueprint::OnPostVariableChange);
+
+	if (UPackage* Package = GetOutermost())
+	{
+		Package->SetDirtyFlag(bDirtyDuringLoad);
+	}
+
+#if WITH_EDITOR
+	// if we are running with -game we are in editor code,
+	// but GIsEditor is turned off
+	if(!GIsEditor)
+	{
+		HandlePackageDone();
+	}
+#endif
+
+	// RigVMRegistry changes can be triggered when new user defined types(structs/enums) are added/removed
+	// in which case we have to refresh the model
+	FRigVMRegistry::Get().OnRigVMRegistryChanged().RemoveAll(this);
+	FRigVMRegistry::Get().OnRigVMRegistryChanged().AddUObject(this, &URigVMBlueprint::OnRigVMRegistryChanged);
+	
+	UEdGraphPin::ResolveAllPinReferences();
+}
+
+#if WITH_EDITORONLY_DATA
+void URigVMBlueprint::DeclareConstructClasses(TArray<FTopLevelAssetPath>& OutConstructClasses, const UClass* SpecificSubclass)
+{
+	Super::DeclareConstructClasses(OutConstructClasses, SpecificSubclass);
+	OutConstructClasses.Add(FTopLevelAssetPath(URigVMController::StaticClass()));
+}
+#endif
+
+#if WITH_EDITOR
+void URigVMBlueprint::HandlePackageDone(const FEndLoadPackageContext& Context)
+{
+	if (!Context.LoadedPackages.Contains(GetPackage()))
+	{
+		return;
+	}
+	HandlePackageDone();
+}
+
+void URigVMBlueprint::HandlePackageDone()
+{
+	FCoreUObjectDelegates::OnEndLoadPackage.RemoveAll(this);
+
+	if(URigVMBuildData* BuildData = URigVMBuildData::Get())
+	{
+		if(URigVMFunctionLibrary* FunctionLibrary = RigVMClient.GetFunctionLibrary())
+		{
+			// for backwards compatibility load the function references from the
+			// model's storage over to the centralized build data
+			if(!FunctionLibrary->FunctionReferences_DEPRECATED.IsEmpty())
+			{
+				// let's also update the asset data of the dependents
+				IAssetRegistry& AssetRegistry = IAssetRegistry::GetChecked();
+				
+				for(const TTuple< TObjectPtr<URigVMLibraryNode>, FRigVMFunctionReferenceArray >& Pair :
+					FunctionLibrary->FunctionReferences_DEPRECATED)
+				{
+					TSoftObjectPtr<URigVMLibraryNode> FunctionKey(Pair.Key);
+						
+					for(int32 ReferenceIndex = 0; ReferenceIndex < Pair.Value.Num(); ReferenceIndex++)
+					{
+						// update the build data
+						BuildData->RegisterFunctionReference(FunctionKey->GetFunctionIdentifier(), Pair.Value[ReferenceIndex]);
+
+						// find all control rigs matching the reference node
+						FAssetData AssetData = AssetRegistry.GetAssetByObjectPath(
+							Pair.Value[ReferenceIndex].ToSoftObjectPath().GetWithoutSubPath());
+
+						// if the asset has never been loaded - make sure to load it once and mark as dirty
+						if(AssetData.IsValid() && !AssetData.IsAssetLoaded())
+						{
+							if(URigVMBlueprint* Dependent = Cast<URigVMBlueprint>(AssetData.GetAsset()))
+							{
+								if(Dependent != this)
+								{
+									(void)Dependent->MarkPackageDirty();
+								}
+							}
+						}
+					}
+				}
+				
+				FunctionLibrary->FunctionReferences_DEPRECATED.Reset();
+				(void)MarkPackageDirty();
+			}
+		}
+
+		// update the build data from the current function references
+		const TArray<FRigVMReferenceNodeData> ReferenceNodeDatas = GetReferenceNodeData();
+		for(const FRigVMReferenceNodeData& ReferenceNodeData : ReferenceNodeDatas)
+		{
+			BuildData->RegisterFunctionReference(ReferenceNodeData);
+		}
+
+		BuildData->ClearInvalidReferences();
+	}
+	
+	RemoveDeprecatedVMMemoryClass();
+	RecompileVM();
+	RequestRigVMInit();
+	BroadcastRigVMPackageDone();
+}
+
+void URigVMBlueprint::BroadcastRigVMPackageDone()
+{
+	if (URigVMBlueprintGeneratedClass* RigClass = GetRigVMBlueprintGeneratedClass())
+	{
+		URigVMHost* CDO = Cast<URigVMHost>(RigClass->GetDefaultObject(true /* create if needed */));
+		CDO->BroadCastEndLoadPackage();
+
+		TArray<UObject*> ArchetypeInstances;
+		CDO->GetArchetypeInstances(ArchetypeInstances);
+		for (UObject* Instance : ArchetypeInstances)
+		{
+			if (URigVMHost* InstanceHost = Cast<URigVMHost>(Instance))
+			{
+				InstanceHost->BroadCastEndLoadPackage();
+			}
+		}
+	}
+}
+
+void URigVMBlueprint::RemoveDeprecatedVMMemoryClass() const
+{
+	TArray<UObject*> Objects;
+	GetObjectsWithOuter(this, Objects, false);
+
+	for (UObject* Object : Objects)
+	{
+		if (URigVMMemoryStorageGeneratorClass* DeprecatedClass = Cast<URigVMMemoryStorageGeneratorClass>(Object))
+		{
+			DeprecatedClass->Rename(nullptr, GetTransientPackage(), REN_ForceNoResetLoaders | REN_DoNotDirty | REN_DontCreateRedirectors | REN_NonTransactional);
+		}
+	}
+}
+#endif
+
+void URigVMBlueprint::RecompileVM()
+{
+	if(bIsCompiling)
+	{
+		return;
+	}
+
+	URigVMBlueprintGeneratedClass* RigClass = GetRigVMBlueprintGeneratedClass();
+	if(RigClass == nullptr)
+	{
+		return;
+	}
+
+	TGuardValue<bool> CompilingGuard(bIsCompiling, true);
+	
+	bErrorsDuringCompilation = false;
+
+	RigGraphDisplaySettings.MinMicroSeconds = RigGraphDisplaySettings.LastMinMicroSeconds = DBL_MAX;
+	RigGraphDisplaySettings.MaxMicroSeconds = RigGraphDisplaySettings.LastMaxMicroSeconds = (double)INDEX_NONE;
+
+	URigVMHost* CDO = Cast<URigVMHost>(RigClass->GetDefaultObject(true /* create if needed */));
+	if (CDO && CDO->VM != nullptr)
+	{
+		TGuardValue<bool> ReentrantGuardSelf(bSuspendModelNotificationsForSelf, true);
+		TGuardValue<bool> ReentrantGuardOthers(bSuspendModelNotificationsForOthers, true);
+
+		CDO->PostInitInstanceIfRequired();
+		CDO->VMRuntimeSettings = VMRuntimeSettings;
+
+		if (!HasAnyFlags(RF_Transient | RF_Transactional))
+		{
+			CDO->Modify(false);
+		}
+		CDO->VM->Reset();
+
+		// Clear all Errors
+		CompileLog.Messages.Reset();
+		CompileLog.NumErrors = CompileLog.NumWarnings = 0;
+		
+		TArray<UEdGraph*> EdGraphs;
+		GetAllGraphs(EdGraphs);
+
+		for (UEdGraph* Graph : EdGraphs)
+		{
+			URigVMEdGraph* RigGraph = Cast<URigVMEdGraph>(Graph);
+			if (RigGraph == nullptr)
+			{
+				continue;
+			}
+
+			for (UEdGraphNode* GraphNode : Graph->Nodes)
+			{
+				if (URigVMEdGraphNode* RigVMEdGraphNode = Cast<URigVMEdGraphNode>(GraphNode))
+				{
+					RigVMEdGraphNode->ClearErrorInfo();
+				}
+			}
+		}
+
+		URigVMCompiler* Compiler = URigVMCompiler::StaticClass()->GetDefaultObject<URigVMCompiler>();
+		VMCompileSettings.SetExecuteContextStruct(RigVMClient.GetExecuteContextStruct());
+		Compiler->Settings = (bCompileInDebugMode) ? FRigVMCompileSettings::Fast(VMCompileSettings.GetExecuteContextStruct()) : VMCompileSettings;
+		Compiler->Compile(RigVMClient.GetAllModels(false, false), GetOrCreateController(), CDO->VM, CDO->GetExtendedExecuteContext(), CDO->GetExternalVariablesImpl(false), &PinToOperandMap);
+
+		if (bErrorsDuringCompilation)
+		{
+			if(Compiler->Settings.SurpressErrors)
+			{
+				Compiler->Settings.Reportf(EMessageSeverity::Info, this,
+					TEXT("Compilation Errors may be suppressed for ControlRigBlueprint: %s. See VM Compile Setting in Class Settings for more Details"), *this->GetName());
+			}
+			bVMRecompilationRequired = false;
+			if(CDO->VM)
+			{
+				VMCompiledEvent.Broadcast(this, CDO->VM);
+			}
+			return;
+		}
+
+		TArray<UObject*> ArchetypeInstances;
+		CDO->GetArchetypeInstances(ArchetypeInstances);
+		for (UObject* Instance : ArchetypeInstances)
+		{
+			if (URigVMHost* InstanceHost = Cast<URigVMHost>(Instance))
+			{
+				// No objects should be created during load, so PostInitInstanceIfRequired, which creates a new VM and
+				// DynamicHierarchy, should not be called during load
+				if (!InstanceHost->HasAllFlags(RF_NeedPostLoad))
+				{
+					InstanceHost->PostInitInstanceIfRequired();
+				}
+				InstanceHost->InstantiateVMFromCDO();
+				InstanceHost->CopyExternalVariableDefaultValuesFromCDO();
+			}
+		}
+
+		bVMRecompilationRequired = false;
+		VMCompiledEvent.Broadcast(this, CDO->VM);
+
+#if WITH_EDITOR
+		RefreshControlRigBreakpoints();
+#endif
+	}
+}
+
+void URigVMBlueprint::RecompileVMIfRequired()
+{
+	if (bVMRecompilationRequired)
+	{
+		RecompileVM();
+	}
+}
+
+void URigVMBlueprint::RequestAutoVMRecompilation()
+{
+	bVMRecompilationRequired = true;
+	if (bAutoRecompileVM && VMRecompilationBracket == 0)
+	{
+		RecompileVMIfRequired();
+	}
+}
+
+void URigVMBlueprint::IncrementVMRecompileBracket()
+{
+	VMRecompilationBracket++;
+}
+
+void URigVMBlueprint::DecrementVMRecompileBracket()
+{
+	if (VMRecompilationBracket == 1)
+	{
+		if (bAutoRecompileVM)
+		{
+			RecompileVMIfRequired();
+		}
+		VMRecompilationBracket = 0;
+	}
+	else if (VMRecompilationBracket > 0)
+	{
+		VMRecompilationBracket--;
+	}
+}
+
+void URigVMBlueprint::RefreshAllModels(ERigVMBlueprintLoadType InLoadType)
+{
+	const bool bIsPostLoad = InLoadType == ERigVMBlueprintLoadType::PostLoad;
+
+	// avoid any compute if the current structure hashes match with the serialized ones
+	if(CVarRigVMEnablePostLoadHashing->GetBool() && RigVMClient.GetStructureHash() == RigVMClient.GetSerializedStructureHash())
+	{
+		if(bIsPostLoad)
+		{
+			TArray<URigVMGraph*> ModelGraphs = RigVMClient.GetAllModels(true, true);
+			Algo::Reverse(ModelGraphs);
+			for (URigVMGraph* ModelGraph : ModelGraphs)
+			{
+				URigVMController* Controller = GetOrCreateController(ModelGraph);
+				URigVMController::FRestoreLinkedPathSettings Settings;
+				Settings.bFollowCoreRedirectors = true;
+				Settings.bRelayToOrphanPins = true;
+				Controller->ProcessDetachedLinks(Settings);
+			}
+		}
+		return;
+	}
+	
+	TGuardValue<bool> IsCompilingGuard(bIsCompiling, true);
+	TGuardValue<bool> ClientIgnoreModificationsGuard(RigVMClient.bIgnoreModelNotifications, true);
+	
+	TArray<URigVMGraph*> AllModelsLeavesFirst = RigVMClient.GetAllModelsLeavesFirst(true);
+	TMap<const URigVMGraph*, TArray<URigVMController::FLinkedPath>> LinkedPaths;
+
+	if (ensure(IsInGameThread()))
+	{
+		for (URigVMGraph* Graph : AllModelsLeavesFirst)
+		{
+			URigVMController* Controller = GetOrCreateController(Graph);
+			// temporarily disable default value validation during load time, serialized values should always be accepted
+			TGuardValue<bool> PerGraphDisablePinDefaultValueValidation(Controller->bValidatePinDefaults, false);
+			TGuardValue<bool> GuardEditGraph(Graph->bEditable, true);
+			FRigVMControllerNotifGuard NotifGuard(Controller, true);
+			LinkedPaths.Add(Graph, Controller->GetLinkedPaths());
+			Controller->FastBreakLinkedPaths(LinkedPaths.FindChecked(Graph));
+			TArray<URigVMNode*> Nodes = Graph->GetNodes();
+			for (URigVMNode* Node : Nodes)
+			{
+				Controller->RepopulatePinsOnNode(Node, true, true);
+			}
+		}
+		SetupPinRedirectorsForBackwardsCompatibility();
+	}
+
+	for (URigVMGraph* Graph : AllModelsLeavesFirst)
+	{
+		URigVMController* Controller = GetOrCreateController(Graph);
+		TGuardValue<bool> GuardEditGraph(Graph->bEditable, true);
+		FRigVMControllerNotifGuard NotifGuard(Controller, true);
+		{
+			URigVMController::FRestoreLinkedPathSettings Settings;
+			Settings.bFollowCoreRedirectors = true;
+			Controller->RestoreLinkedPaths(LinkedPaths.FindChecked(Graph));
+		}
+
+		for(URigVMNode* ModelNode : Graph->GetNodes())
+		{
+			Controller->RemoveUnusedOrphanedPins(ModelNode);
+		}
+
+		if(bIsPostLoad)
+		{
+			for(URigVMNode* ModelNode : Graph->GetNodes())
+			{
+				if (URigVMTemplateNode* TemplateNode = Cast<URigVMTemplateNode>(ModelNode))
+				{
+					TemplateNode->InvalidateCache();
+					TemplateNode->PostLoad();
+				}
+			}
+		}
+
+#if WITH_EDITOR
+
+		if(bIsPostLoad)
+		{
+			for(URigVMNode* ModelNode : Graph->GetNodes())
+			{
+				if(URigVMUnitNode* UnitNode = Cast<URigVMUnitNode>(ModelNode))
+				{
+					if (!UnitNode->HasWildCardPin())
+					{
+						UScriptStruct* ScriptStruct = UnitNode->GetScriptStruct(); 
+						if(ScriptStruct == nullptr)
+						{
+							Controller->FullyResolveTemplateNode(UnitNode, INDEX_NONE, false);
+						}
+
+						// Try to find a deprecated template
+						if (UnitNode->GetScriptStruct() == nullptr && !UnitNode->TemplateNotation.IsNone())
+						{
+							const FRigVMTemplate* Template = FRigVMRegistry::Get().FindTemplate(UnitNode->TemplateNotation, true);
+							FRigVMTemplate::FTypeMap TypeMap = UnitNode->GetTemplatePinTypeMap();
+
+							int32 Permutation;
+							if (Template->FullyResolve(TypeMap, Permutation))
+							{
+								const FRigVMFunction* Function = Template->GetPermutation(Permutation);
+								UnitNode->ResolvedFunctionName = Function->GetName();
+							}
+						}
+
+						if (UnitNode->GetScriptStruct() == nullptr)
+						{
+							static const TCHAR UnresolvedUnitNodeMessage[] = TEXT("Node %s could not be resolved.");
+							Controller->ReportErrorf(UnresolvedUnitNodeMessage, *ModelNode->GetNodePath(true));
+						}
+					}
+				}
+				if (URigVMDispatchNode* DispatchNode = Cast<URigVMDispatchNode>(ModelNode))
+				{
+					if (DispatchNode->GetFactory() == nullptr)
+					{
+						static const TCHAR UnresolvedDispatchNodeMessage[] = TEXT("Dispatch node %s has no factory..");
+						Controller->ReportErrorf(UnresolvedDispatchNodeMessage, *ModelNode->GetNodePath(true));
+					}
+					else if (!DispatchNode->HasWildCardPin())
+					{
+						if (DispatchNode->GetResolvedFunction() == nullptr)
+						{
+							Controller->FullyResolveTemplateNode(DispatchNode, INDEX_NONE, false);
+						}
+						if (DispatchNode->GetResolvedFunction() == nullptr)
+						{
+							static const TCHAR UnresolvedDispatchNodeMessage[] = TEXT("Node %s could not be resolved.");
+							Controller->ReportErrorf(UnresolvedDispatchNodeMessage, *ModelNode->GetNodePath(true));
+						}
+					}
+				}
+			}
+		}
+#endif
+
+	}
+}
+
+void URigVMBlueprint::OnRigVMRegistryChanged()
+{
+	RefreshAllModels();
+	RebuildGraphFromModel();
+	// avoids slate crash
+	FRigVMBlueprintUtils::HandleRefreshAllNodes(this);
+}
+
+void URigVMBlueprint::HandleReportFromCompiler(EMessageSeverity::Type InSeverity, UObject* InSubject, const FString& InMessage)
+{
+	UObject* SubjectForMessage = InSubject;
+	if(URigVMNode* ModelNode = Cast<URigVMNode>(SubjectForMessage))
+	{
+		if(URigVMBlueprint* RigBlueprint = ModelNode->GetTypedOuter<URigVMBlueprint>())
+		{
+			if(URigVMEdGraph* EdGraph = Cast<URigVMEdGraph>(RigBlueprint->GetEdGraph(ModelNode->GetGraph())))
+			{
+				if(UEdGraphNode* EdNode = EdGraph->FindNodeForModelNodeName(ModelNode->GetFName()))
+				{
+					SubjectForMessage = EdNode;
+				}
+			}
+		}
+	}
+
+	FCompilerResultsLog* Log = CurrentMessageLog ? CurrentMessageLog : &CompileLog;
+	if (InSeverity == EMessageSeverity::Error)
+	{
+		Status = BS_Error;
+		(void)MarkPackageDirty();
+
+		// see UnitTest "ControlRig.Basics.OrphanedPins" to learn why errors are suppressed this way
+		if (VMCompileSettings.SurpressErrors)
+		{
+			Log->bSilentMode = true;
+		}
+
+		if(InMessage.Contains(TEXT("@@")))
+		{
+			Log->Error(*InMessage, SubjectForMessage);
+		}
+		else
+		{
+			Log->Error(*InMessage);
+		}
+
+		BroadCastReportCompilerMessage(InSeverity, InSubject, InMessage);
+
+		// see UnitTest "ControlRig.Basics.OrphanedPins" to learn why errors are suppressed this way
+		if (!VMCompileSettings.SurpressErrors)
+		{ 
+			FScriptExceptionHandler::Get().HandleException(ELogVerbosity::Error, *InMessage, *FString());
+		}
+		
+		bErrorsDuringCompilation = true;
+	}
+	else if (InSeverity == EMessageSeverity::Warning)
+	{
+		if(InMessage.Contains(TEXT("@@")))
+		{
+			Log->Warning(*InMessage, SubjectForMessage);
+		}
+		else
+		{
+			Log->Warning(*InMessage);
+		}
+
+		BroadCastReportCompilerMessage(InSeverity, InSubject, InMessage);
+		FScriptExceptionHandler::Get().HandleException(ELogVerbosity::Warning, *InMessage, *FString());
+	}
+	else
+	{
+		if(InMessage.Contains(TEXT("@@")))
+		{
+			Log->Note(*InMessage, SubjectForMessage);
+		}
+		else
+		{
+			Log->Note(*InMessage);
+		}
+
+		UE_LOG(LogRigVMDeveloper, Display, TEXT("%s"), *InMessage);
+	}
+
+	if (URigVMEdGraphNode* EdGraphNode = Cast<URigVMEdGraphNode>(SubjectForMessage))
+	{
+		EdGraphNode->ErrorType = (int32)InSeverity;
+		EdGraphNode->ErrorMsg = InMessage;
+		EdGraphNode->bHasCompilerMessage = EdGraphNode->ErrorType <= int32(EMessageSeverity::Info);
+	}
+}
+
+TArray<IRigVMGraphFunctionHost*> URigVMBlueprint::GetReferencedFunctionHosts(bool bForceLoad)
+{
+	TArray<IRigVMGraphFunctionHost*> ReferencedBlueprints;
+	
+	TArray<UEdGraph*> EdGraphs;
+	GetAllGraphs(EdGraphs);
+	for (UEdGraph* EdGraph : EdGraphs)
+	{
+		for(UEdGraphNode* Node : EdGraph->Nodes)
+		{
+			if(URigVMEdGraphNode* RigNode = Cast<URigVMEdGraphNode>(Node))
+			{
+				if(URigVMFunctionReferenceNode* FunctionRefNode = Cast<URigVMFunctionReferenceNode>(RigNode->GetModelNode()))
+				{
+					IRigVMGraphFunctionHost* Host = nullptr;
+					if (bForceLoad || FunctionRefNode->IsReferencedFunctionHostLoaded())
+					{
+						// Load the function host
+						Host = FunctionRefNode->GetReferencedFunctionHeader().GetFunctionHost();
+					}
+					else if (bForceLoad || FunctionRefNode->IsReferencedNodeLoaded())
+					{
+						// Load the reference library node
+						if(const URigVMLibraryNode* ReferencedNode = FunctionRefNode->LoadReferencedNode())
+						{
+							if(URigVMFunctionLibrary* ReferencedFunctionLibrary = ReferencedNode->GetLibrary())
+							{
+								FSoftObjectPath FunctionHostPath = ReferencedFunctionLibrary->GetFunctionHostObjectPath();
+								if (UObject* FunctionHostObj = FunctionHostPath.TryLoad())
+								{
+									Host = Cast<IRigVMGraphFunctionHost>(FunctionHostObj);									
+								}
+							}
+						}
+					}
+
+					if (Host != nullptr && Host != GetRigVMBlueprintGeneratedClass())
+					{
+						ReferencedBlueprints.Add(Host);
+					}
+				}
+			}
+		}
+	}
+	
+	return ReferencedBlueprints;
+}
+
+#if WITH_EDITOR
+
+void URigVMBlueprint::ClearBreakpoints()
+{
+	for(URigVMNode* Node : RigVMBreakpointNodes)
+	{
+		Node->SetHasBreakpoint(false);		
+	}
+	
+	RigVMBreakpointNodes.Empty();
+	RefreshControlRigBreakpoints();
+}
+
+bool URigVMBlueprint::AddBreakpoint(const FString& InBreakpointNodePath)
+{
+	URigVMLibraryNode* FunctionNode = nullptr;
+	
+	// Find the node in the graph
+	for(const URigVMGraph* Model : RigVMClient)
+	{
+		URigVMNode* BreakpointNode = Model->FindNode(InBreakpointNodePath);
+		if (BreakpointNode == nullptr)
+		{
+			// If we cannot find the node, it might be because it is inside a function
+			FString FunctionName = InBreakpointNodePath, Right;
+			URigVMNode::SplitNodePathAtStart(InBreakpointNodePath, FunctionName, Right);
+
+			// Look inside the local function library
+			if (URigVMLibraryNode* LibraryNode = GetLocalFunctionLibrary()->FindFunction(FName(FunctionName)))
+			{
+				BreakpointNode = LibraryNode->GetContainedGraph()->FindNode(Right);
+				FunctionNode = LibraryNode;
+			}
+		}
+
+		if(BreakpointNode)
+		{
+			return AddBreakpoint(BreakpointNode, FunctionNode);
+		}
+	}
+	return false;
+}
+
+bool URigVMBlueprint::AddBreakpoint(URigVMNode* InBreakpointNode, URigVMLibraryNode* LibraryNode)
+{
+	if (InBreakpointNode == nullptr)
+	{
+		return false;
+	}
+
+	bool bSuccess = true;
+	if (LibraryNode)
+	{
+		// If the breakpoint node is inside a library node, find all references to the library node
+		TArray<TSoftObjectPtr<URigVMFunctionReferenceNode>> References = LibraryNode->GetLibrary()->GetReferencesForFunction(LibraryNode->GetFName());
+		for (TSoftObjectPtr<URigVMFunctionReferenceNode> Reference : References)
+		{
+			if (!Reference.IsValid())
+			{
+				continue;
+			}
+
+			URigVMBlueprint* ReferenceBlueprint = Reference->GetTypedOuter<URigVMBlueprint>();
+
+			// If the reference is not inside another function, add a breakpoint in the blueprint containing the
+			// reference, without a function specified
+			bool bIsInsideFunction = Reference->GetRootGraph()->IsA<URigVMFunctionLibrary>();
+			if(!bIsInsideFunction)
+			{
+				bSuccess &= ReferenceBlueprint->AddBreakpoint(InBreakpointNode);
+			}
+			else
+			{
+				// Otherwise, we need to add breakpoints to all the blueprints that reference this
+				// function (when the blueprint graph is flattened)
+				
+				// Get all the functions containing this reference
+				URigVMNode* Node = Reference.Get();
+				while (Node->GetGraph() != ReferenceBlueprint->GetLocalFunctionLibrary())
+				{
+					if (URigVMLibraryNode* ParentLibraryNode = Cast<URigVMLibraryNode>(Node->GetGraph()->GetOuter()))
+					{
+						// Recursively add breakpoints to the reference blueprint, specifying the parent function
+						bSuccess &= ReferenceBlueprint->AddBreakpoint(InBreakpointNode, ParentLibraryNode);
+					}
+
+					Node = Cast<URigVMNode>(Node->GetGraph()->GetOuter());
+				}
+			}
+		}
+	}
+	else
+	{
+		if (!RigVMBreakpointNodes.Contains(InBreakpointNode))
+		{
+			// Add the breakpoint to the VM
+			bSuccess = AddBreakpointToControlRig(InBreakpointNode);
+			BreakpointAddedEvent.Broadcast();
+		}
+	}
+
+	return bSuccess;
+}
+
+bool URigVMBlueprint::AddBreakpointToControlRig(URigVMNode* InBreakpointNode)
+{
+	URigVMBlueprintGeneratedClass* RigClass = GetRigVMBlueprintGeneratedClass();
+	URigVMHost* CDO = Cast<URigVMHost>(RigClass->GetDefaultObject(false));
+	const FRigVMByteCode* ByteCode = GetController()->GetCurrentByteCode();
+	TSet<FString> AddedCallpaths;
+
+	if (CDO && ByteCode)
+	{
+		FRigVMInstructionArray Instructions = ByteCode->GetInstructions();	
+		
+		// For each instruction, see if the node is in the callpath
+		// Only add one breakpoint for each callpath related to this node (i.e. if a node produces multiple
+		// instructions, only add a breakpoint to the first instruction)
+		for (int32 i = 0; i< Instructions.Num(); ++i)
+		{
+			for(URigVMGraph* Model : RigVMClient)
+			{
+				const FRigVMASTProxy Proxy = FRigVMASTProxy::MakeFromCallPath(ByteCode->GetCallPathForInstruction(i), Model);
+				if (Proxy.GetCallstack().Contains(InBreakpointNode))
+				{
+					// Find the callpath related to the breakpoint node
+					FRigVMASTProxy BreakpointProxy = Proxy;
+					while(BreakpointProxy.GetSubject() != InBreakpointNode)
+					{
+						BreakpointProxy = BreakpointProxy.GetParent();
+					}
+					const FString& BreakpointCallPath = BreakpointProxy.GetCallstack().GetCallPath();
+
+					// Only add this callpath breakpoint once
+					if (!AddedCallpaths.Contains(BreakpointCallPath))
+					{
+						AddedCallpaths.Add(BreakpointCallPath);
+						CDO->AddBreakpoint(i, InBreakpointNode, BreakpointProxy.GetCallstack().Num());
+					}
+				}
+			}
+		}
+	}
+
+	if (AddedCallpaths.Num() > 0)
+	{
+		RigVMBreakpointNodes.AddUnique(InBreakpointNode);
+		return true;
+	}
+	
+	return false;
+}
+
+bool URigVMBlueprint::RemoveBreakpoint(const FString& InBreakpointNodePath)
+{
+	// Find the node in the graph
+	URigVMNode* BreakpointNode = nullptr;
+	
+	for(URigVMGraph* Model : RigVMClient)
+	{
+		BreakpointNode = Model->FindNode(InBreakpointNodePath);
+		if (BreakpointNode == nullptr)
+		{
+			// If we cannot find the node, it might be because it is inside a function
+			FString FunctionName = InBreakpointNodePath, Right;
+			URigVMNode::SplitNodePathAtStart(InBreakpointNodePath, FunctionName, Right);
+
+			// Look inside the local function library
+			if (URigVMLibraryNode* LibraryNode = GetLocalFunctionLibrary()->FindFunction(FName(FunctionName)))
+			{
+				BreakpointNode = LibraryNode->GetContainedGraph()->FindNode(Right);
+			}
+		}
+		if(BreakpointNode)
+		{
+			break;
+		}
+	}
+
+	if(BreakpointNode)
+	{
+		bool bSuccess = RemoveBreakpoint(BreakpointNode);
+
+		// Remove the breakpoint from all the loaded dependent blueprints
+		TArray<URigVMBlueprint*> DependentBlueprints = GetDependentBlueprints(true, true);
+		DependentBlueprints.Remove(this);
+		for (URigVMBlueprint* Dependent : DependentBlueprints)
+		{
+			bSuccess &= Dependent->RemoveBreakpoint(BreakpointNode);
+		}
+		return bSuccess;
+	}
+
+	return false;
+}
+
+bool URigVMBlueprint::RemoveBreakpoint(URigVMNode* InBreakpointNode)
+{
+	if (RigVMBreakpointNodes.Contains(InBreakpointNode))
+	{
+		RigVMBreakpointNodes.Remove(InBreakpointNode);
+
+		// Multiple breakpoint nodes might set a breakpoint to the same instruction. When we remove
+		// one of the breakpoint nodes, we do not want to remove the instruction breakpoint if there
+		// is another breakpoint node addressing it. For that reason, we just recompute all the
+		// breakpoint instructions.
+		// Refreshing breakpoints in the control rig will keep the state it had before.
+		RefreshControlRigBreakpoints();
+		return true;
+	}
+
+	return false;
+}
+
+
+void URigVMBlueprint::RefreshControlRigBreakpoints()
+{
+	if (URigVMBlueprintGeneratedClass* RigClass = GetRigVMBlueprintGeneratedClass())
+	{
+		URigVMHost* CDO = Cast<URigVMHost>(RigClass->GetDefaultObject(false));
+		CDO->GetDebugInfo().Reset();
+		for (URigVMNode* Node : RigVMBreakpointNodes)
+		{
+			AddBreakpointToControlRig(Node);
+		}
+	}
+}
+
+TArray<FRigVMReferenceNodeData> URigVMBlueprint::GetReferenceNodeData() const
+{
+	TArray<FRigVMReferenceNodeData> Data;
+	
+	const TArray<URigVMGraph*> AllModels = GetAllModels();
+	for (URigVMGraph* ModelToVisit : AllModels)
+	{
+		for(URigVMNode* Node : ModelToVisit->GetNodes())
+		{
+			if(URigVMFunctionReferenceNode* ReferenceNode = Cast<URigVMFunctionReferenceNode>(Node))
+			{
+				Data.Add(FRigVMReferenceNodeData(ReferenceNode));
+			}
+		}
+	}
+	return Data;
+}
+
+#endif
+
+void URigVMBlueprint::RequestRigVMInit()
+{
+	if (URigVMBlueprintGeneratedClass* RigClass = GetRigVMBlueprintGeneratedClass())
+	{
+		URigVMHost* CDO = Cast<URigVMHost>(RigClass->GetDefaultObject(true /* create if needed */));
+		CDO->RequestInit();
+
+		TArray<UObject*> ArchetypeInstances;
+		CDO->GetArchetypeInstances(ArchetypeInstances);
+		for (UObject* Instance : ArchetypeInstances)
+		{
+			if (URigVMHost* InstanceHost = Cast<URigVMHost>(Instance))
+			{
+				InstanceHost->RequestInit();
+			}
+		}
+	}
+}
+
+URigVMGraph* URigVMBlueprint::GetModel(const UEdGraph* InEdGraph) const
+{
+	if (InEdGraph == nullptr)
+	{
+		return GetDefaultModel();
+	}
+
+	if(InEdGraph->GetOutermost() != GetOutermost())
+	{
+		return nullptr;
+	}
+
+#if WITH_EDITORONLY_DATA
+	if (InEdGraph == FunctionLibraryEdGraph)
+	{
+		return RigVMClient.GetFunctionLibrary();
+	}
+#endif
+
+	const URigVMEdGraph* RigGraph = Cast< URigVMEdGraph>(InEdGraph);
+	check(RigGraph);
+	return GetModel(RigGraph->ModelNodePath);
+}
+
+URigVMGraph* URigVMBlueprint::GetModel(const FString& InNodePath) const
+{
+	return RigVMClient.GetModel(InNodePath);
+}
+
+URigVMGraph* URigVMBlueprint::GetDefaultModel() const
+{
+	return RigVMClient.GetDefaultModel();
+}
+
+TArray<URigVMGraph*> URigVMBlueprint::GetAllModels() const
+{
+	return RigVMClient.GetAllModels(true, true);
+}
+
+URigVMFunctionLibrary* URigVMBlueprint::GetLocalFunctionLibrary() const
+{
+	return RigVMClient.GetFunctionLibrary();
+}
+
+URigVMGraph* URigVMBlueprint::AddModel(FString InName, bool bSetupUndoRedo, bool bPrintPythonCommand)
+{
+	const FString DesiredName = FString::Printf(TEXT("%s %s"),
+    	FRigVMClient::RigVMModelPrefix, *InName);
+
+	TGuardValue<bool> EnablePythonPrint(bSuspendPythonMessagesForRigVMClient, !bPrintPythonCommand);
+	return RigVMClient.AddModel(*DesiredName, bSetupUndoRedo);
+}
+
+bool URigVMBlueprint::RemoveModel(FString InName, bool bSetupUndoRedo, bool bPrintPythonCommand)
+{
+	TGuardValue<bool> EnablePythonPrint(bSuspendPythonMessagesForRigVMClient, !bPrintPythonCommand);
+	return RigVMClient.RemoveModel(InName, bSetupUndoRedo);
+}
+
+URigVMGraph* URigVMBlueprint::GetFocusedModel() const
+{
+#if WITH_EDITOR
+	if(OnGetFocusedGraphDelegate.IsBound())
+	{
+		return OnGetFocusedGraphDelegate.Execute(); 
+	}
+#endif
+	return RigVMClient.GetDefaultModel();
+}
+
+URigVMController* URigVMBlueprint::GetController(const URigVMGraph* InGraph) const
+{
+	return RigVMClient.GetController(InGraph);
+}
+
+URigVMController* URigVMBlueprint::GetControllerByName(const FString InGraphName) const
+{
+	if(InGraphName.IsEmpty())
+	{
+		if(const URigVMGraph* DefaultModel = GetRigVMClient()->GetDefaultModel())
+		{
+			return GetController(DefaultModel);
+		}
+	}
+	
+	for (const URigVMGraph* Graph : GetAllModels())
+	{
+		if (Graph->GetName() == InGraphName || Graph->GetGraphName() == InGraphName)
+		{
+			return GetController(Graph);
+		}
+	}
+	
+	return nullptr;
+}
+
+URigVMController* URigVMBlueprint::GetOrCreateController(URigVMGraph* InGraph)
+{
+	return RigVMClient.GetOrCreateController(InGraph);
+}
+
+URigVMController* URigVMBlueprint::GetController(const UEdGraph* InEdGraph) const
+{
+	return RigVMClient.GetController(InEdGraph);
+}
+
+URigVMController* URigVMBlueprint::GetOrCreateController(const UEdGraph* InEdGraph)
+{
+	return RigVMClient.GetOrCreateController(InEdGraph);
+}
+
+TArray<FString> URigVMBlueprint::GeneratePythonCommands(const FString InNewBlueprintName)
+{
+	TArray<FString> InternalCommands;
+	InternalCommands.Add(TEXT("unreal.load_module('ControlRigDeveloper')"));
+	InternalCommands.Add(TEXT("factory = unreal.ControlRigBlueprintFactory"));
+	InternalCommands.Add(FString::Printf(TEXT("blueprint = factory.create_new_control_rig_asset(desired_package_path = '%s')"), *InNewBlueprintName));
+	InternalCommands.Add(TEXT("library = blueprint.get_local_function_library()"));
+	InternalCommands.Add(TEXT("library_controller = blueprint.get_controller(library)"));
+	InternalCommands.Add(TEXT("hierarchy = blueprint.hierarchy"));
+	InternalCommands.Add(TEXT("hierarchy_controller = hierarchy.get_controller()"));
+	InternalCommands.Add(TEXT("blueprint.set_auto_vm_recompile(False)"));
+	
+	// Add variables
+	for (const FBPVariableDescription& Variable : NewVariables)
+	{
+		const FRigVMExternalVariable ExternalVariable = RigVMTypeUtils::ExternalVariableFromBPVariableDescription(Variable);
+		FString CPPType;
+		UObject* CPPTypeObject = nullptr;
+		RigVMTypeUtils::CPPTypeFromExternalVariable(ExternalVariable, CPPType, &CPPTypeObject);
+		if (CPPTypeObject)
+		{
+			if (ExternalVariable.bIsArray)
+			{
+				CPPType = RigVMTypeUtils::ArrayTypeFromBaseType(CPPTypeObject->GetPathName());
+			}
+			else
+			{
+				CPPType = CPPTypeObject->GetPathName();
+			}
+		}
+		// FName AddMemberVariable(const FName& InName, const FString& InCPPType, bool bIsPublic = false, bool bIsReadOnly = false, FString InDefaultValue = TEXT(""));
+		InternalCommands.Add(FString::Printf(TEXT("blueprint.add_member_variable('%s', '%s', %s, %s)"),
+					*ExternalVariable.Name.ToString(),
+					*CPPType,
+					ExternalVariable.bIsPublic ? TEXT("True") : TEXT("False"),
+					ExternalVariable.bIsReadOnly ? TEXT("True") : TEXT("False")));	
+	}
+	
+	// Create graphs
+	{
+		TArray<URigVMGraph*> AllModels = GetAllModels();
+		AllModels.RemoveAll([](const URigVMGraph* GraphToRemove) -> bool
+		{
+			return GraphToRemove->GetTypedOuter<URigVMAggregateNode>() != nullptr;
+		});
+		
+		// Find all graphs to process and sort them by dependencies
+		TArray<URigVMGraph*> ProcessedGraphs;
+		while (ProcessedGraphs.Num() < AllModels.Num())
+		{
+			for (URigVMGraph* Graph : AllModels)
+			{
+				if (ProcessedGraphs.Contains(Graph))
+				{
+					continue;
+				}
+
+				bool bFoundUnprocessedReference = false;
+				for (auto Node : Graph->GetNodes())
+				{
+					if (URigVMFunctionReferenceNode* Reference = Cast<URigVMFunctionReferenceNode>(Node))
+					{
+						if (Reference->GetReferencedFunctionHeader().LibraryPointer.HostObject != GetRigVMBlueprintGeneratedClass())
+						{
+							continue;
+						}
+
+						URigVMLibraryNode* LibraryNode = Cast<URigVMLibraryNode>(Reference->GetReferencedFunctionHeader().LibraryPointer.LibraryNode.ResolveObject());
+						if (!ProcessedGraphs.Contains(LibraryNode->GetContainedGraph()))
+						{
+							bFoundUnprocessedReference = true;
+							break;
+						}
+					}
+					else if (URigVMCollapseNode* CollapseNode = Cast<URigVMCollapseNode>(Node))
+					{
+						if(!CollapseNode->IsA<URigVMAggregateNode>())
+						{
+							if (!ProcessedGraphs.Contains(CollapseNode->GetContainedGraph()))
+							{
+								bFoundUnprocessedReference = true;
+								break;
+							}
+						}
+					}
+				}
+
+				if (!bFoundUnprocessedReference)
+				{
+					ProcessedGraphs.Add(Graph);
+				}
+			}
+		}	
+
+		// Dump python commands for each graph
+		for (URigVMGraph* Graph : ProcessedGraphs)
+		{
+			if (Graph->IsA<URigVMFunctionLibrary>())
+			{
+				continue;
+			}
+
+			URigVMController* Controller = GetController(Graph);
+			if (Graph->GetParentGraph()) 
+			{
+				// Add them all as functions (even collapsed graphs)
+				// The controller will deal with deleting collapsed graph function when it creates the collapse node
+				{						
+					// Add Function
+					InternalCommands.Add(FString::Printf(TEXT("function_%s = library_controller.add_function_to_library('%s', mutable=%s)\ngraph = function_%s.get_contained_graph()"),
+							*RigVMPythonUtils::PythonizeName(Graph->GetGraphName()),
+							*Graph->GetGraphName(),
+							Graph->GetEntryNode()->IsMutable() ? TEXT("True") : TEXT("False"),
+							*RigVMPythonUtils::PythonizeName(Graph->GetGraphName())));
+
+					URigVMLibraryNode* LibraryNode = Cast<URigVMLibraryNode>(Graph->GetOuter());
+
+					InternalCommands.Add(FString::Printf(TEXT("library_controller.set_node_category_by_name('%s', '%s')"),
+						*Graph->GetGraphName(),
+						*LibraryNode->GetNodeCategory()));
+
+					InternalCommands.Add(FString::Printf(TEXT("library_controller.set_node_keywords_by_name('%s', '%s')"),
+						*Graph->GetGraphName(),
+						*LibraryNode->GetNodeKeywords() ));
+
+					InternalCommands.Add(FString::Printf(TEXT("library_controller.set_node_description_by_name('%s', '%s')"),
+						*Graph->GetGraphName(),
+						*LibraryNode->GetNodeDescription()));
+
+					InternalCommands.Add(FString::Printf(TEXT("library_controller.set_node_color_by_name('%s', %s)"),
+						*Graph->GetGraphName(),
+						*RigVMPythonUtils::LinearColorToPythonString(LibraryNode->GetNodeColor()) ));
+					
+					URigVMFunctionEntryNode* EntryNode = Graph->GetEntryNode();
+					URigVMFunctionReturnNode* ReturnNode = Graph->GetReturnNode();
+					
+					// Set Entry and Return nodes in the correct position
+					{
+						//bool SetNodePositionByName(const FName& InNodeName, const FVector2D& InPosition, bool bSetupUndoRedo = true, bool bMergeUndoAction = false);
+						InternalCommands.Add(FString::Printf(TEXT("blueprint.get_controller_by_name('%s').set_node_position_by_name('Entry', unreal.Vector2D(%f, %f))"),
+								*Graph->GetGraphName(),
+								EntryNode->GetPosition().X, 
+								EntryNode->GetPosition().Y));
+
+						InternalCommands.Add(FString::Printf(TEXT("blueprint.get_controller_by_name('%s').set_node_position_by_name('Return', unreal.Vector2D(%f, %f))"),
+								*Graph->GetGraphName(),
+								ReturnNode->GetPosition().X, 
+								ReturnNode->GetPosition().Y));
+					}
+					
+					// Add Exposed Pins
+					{
+
+						bool bHitFirstExecute = false;
+						bool bRenamedExecute = false;
+						for (auto Pin : EntryNode->GetPins())
+						{
+							if (Pin->GetDirection() != ERigVMPinDirection::Output)
+							{
+								continue;
+							}
+
+							if(Pin->IsExecuteContext())
+							{
+								if(!bHitFirstExecute)
+								{
+									bHitFirstExecute = true;
+									if (Pin->GetName() != FRigVMStruct::ExecuteContextName.ToString())
+									{
+										bRenamedExecute = true;
+										InternalCommands.Add(FString::Printf(TEXT("blueprint.get_controller_by_name('%s').rename_exposed_pin('%s', '%s')"),
+											*Graph->GetGraphName(),
+											*FRigVMStruct::ExecuteContextName.ToString(),
+											*Pin->GetName()));
+									}
+									continue;
+								}
+							}
+
+							// FName AddExposedPin(const FName& InPinName, ERigVMPinDirection InDirection, const FString& InCPPType, const FName& InCPPTypeObjectPath, const FString& InDefaultValue, bool bSetupUndoRedo = true);
+							InternalCommands.Add(FString::Printf(TEXT("blueprint.get_controller_by_name('%s').add_exposed_pin('%s', %s, '%s', '%s', '%s')"),
+									*Graph->GetGraphName(),
+									*Pin->GetName(),
+									*RigVMPythonUtils::EnumValueToPythonString<ERigVMPinDirection>((int64)ERigVMPinDirection::Input),
+									*Pin->GetCPPType(),
+									Pin->GetCPPTypeObject() ? *Pin->GetCPPTypeObject()->GetPathName() : TEXT(""),
+									*Pin->GetDefaultValue()));
+						}
+
+						bHitFirstExecute = false;
+						for (auto Pin : ReturnNode->GetPins())
+						{
+							if (Pin->GetDirection() != ERigVMPinDirection::Input)
+							{
+								continue;
+							}
+
+							if(Pin->IsExecuteContext())
+							{
+								if(!bHitFirstExecute)
+								{
+									bHitFirstExecute = true;
+									if (!bRenamedExecute && Pin->GetName() != FRigVMStruct::ExecuteContextName.ToString())
+									{
+										bRenamedExecute = true;
+										InternalCommands.Add(FString::Printf(TEXT("blueprint.get_controller_by_name('%s').rename_exposed_pin('%s', '%s')"),
+											*Graph->GetGraphName(),
+											*FRigVMStruct::ExecuteContextName.ToString(),
+											*Pin->GetName()));
+									}
+									continue;
+								}
+							}
+
+							// FName AddExposedPin(const FName& InPinName, ERigVMPinDirection InDirection, const FString& InCPPType, const FName& InCPPTypeObjectPath, const FString& InDefaultValue, bool bSetupUndoRedo = true);
+							InternalCommands.Add(FString::Printf(TEXT("blueprint.get_controller_by_name('%s').add_exposed_pin('%s', %s, '%s', '%s', '%s')"),
+									*Graph->GetGraphName(),
+									*Pin->GetName(),
+									*RigVMPythonUtils::EnumValueToPythonString<ERigVMPinDirection>((int64)ERigVMPinDirection::Output),
+									*Pin->GetCPPType(),
+									Pin->GetCPPTypeObject() ? *Pin->GetCPPTypeObject()->GetPathName() : TEXT(""),
+									*Pin->GetDefaultValue()));
+						}
+					}
+				}
+			}
+			else if(Graph != GetDefaultModel())
+			{
+				InternalCommands.Add(FString::Printf(TEXT("blueprint.add_model('%s')"),
+						*Graph->GetName()));
+			}
+								
+			InternalCommands.Append(Controller->GeneratePythonCommands());
+		}
+	}
+
+	InternalCommands.Add(TEXT("blueprint.set_auto_vm_recompile(True)"));
+
+	// Split multiple commands into different array elements
+	TArray<FString> InnerFunctionCmds;
+	for (FString Cmd : InternalCommands)
+	{
+		FString Left, Right=Cmd;
+		while (Right.Split(TEXT("\n"), &Left, &Right))
+		{
+			InnerFunctionCmds.Add(Left);			
+		}
+		InnerFunctionCmds.Add(Right);
+	}
+
+	// Define a function and insert all the commands
+	// We do not want to pollute the global state with our definitions
+	TArray<FString> Commands;
+	Commands.Add(FString::Printf(TEXT("import unreal\n"
+		"def create_control_rig():\n")));
+	for (const FString& InnerCmd : InnerFunctionCmds)
+	{
+		Commands.Add(FString::Printf(TEXT("\t%s"), *InnerCmd));
+	}
+
+	Commands.Add(TEXT("create_control_rig()\n"));
+	return Commands;
+}
+
+
+URigVMGraph* URigVMBlueprint::GetTemplateModel(bool bIsFunctionLibrary)
+{
+#if WITH_EDITORONLY_DATA
+	if (TemplateModel == nullptr)
+	{
+		if (bIsFunctionLibrary)
+		{
+			TemplateModel = NewObject<URigVMFunctionLibrary>(this, TEXT("TemplateFunctionLibrary"));
+		}
+		else
+		{
+			TemplateModel = NewObject<URigVMGraph>(this, TEXT("TemplateModel"));
+		}
+		TemplateModel->SetFlags(RF_Transient);
+		TemplateModel->SetExecuteContextStruct(RigVMClient.GetExecuteContextStruct());
+	}
+	return TemplateModel;
+#else
+	return nullptr;
+#endif
+}
+
+URigVMController* URigVMBlueprint::GetTemplateController(bool bIsFunctionLibrary)
+{
+#if WITH_EDITORONLY_DATA
+	if (TemplateController == nullptr)
+	{
+		TemplateController = NewObject<URigVMController>(this, TEXT("TemplateController"));
+		TemplateController->SetGraph(GetTemplateModel(bIsFunctionLibrary));
+		TemplateController->EnableReporting(false);
+		TemplateController->SetFlags(RF_Transient);
+		TemplateController->SetSchema(RigVMClient.GetOrCreateSchema());
+	}
+	return TemplateController;
+#else
+	return nullptr;
+#endif
+}
+
+UEdGraph* URigVMBlueprint::GetEdGraph(URigVMGraph* InModel) const
+{
+	return Cast<UEdGraph>(GetEditorObjectForRigVMGraph(InModel));
+}
+
+UEdGraph* URigVMBlueprint::GetEdGraph(const FString& InNodePath) const
+{
+	if (URigVMGraph* ModelForNodePath = GetModel(InNodePath))
+	{
+		return GetEdGraph(ModelForNodePath);
+	}
+	return nullptr;
+}
+
+bool URigVMBlueprint::IsFunctionPublic(const FName& InFunctionName) const
+{
+	return GetLocalFunctionLibrary()->IsFunctionPublic(InFunctionName);	
+}
+
+void URigVMBlueprint::MarkFunctionPublic(const FName& InFunctionName, bool bIsPublic)
+{
+	if(IsFunctionPublic(InFunctionName) == bIsPublic)
+	{
+		return;
+	}
+	
+	URigVMController* Controller = RigVMClient.GetOrCreateController(GetLocalFunctionLibrary());
+	Controller->MarkFunctionAsPublic(InFunctionName, bIsPublic);
+}
+
+TArray<URigVMBlueprint*> URigVMBlueprint::GetDependencies(bool bRecursive) const
+{
+	TArray<URigVMBlueprint*> Dependencies;
+
+	TArray<URigVMGraph*> Graphs = GetAllModels();
+	for(URigVMGraph* Graph : Graphs)
+	{
+		for(URigVMNode* Node : Graph->GetNodes())
+		{
+			if(URigVMFunctionReferenceNode* FunctionReferenceNode = Cast<URigVMFunctionReferenceNode>(Node))
+			{
+				if(const URigVMLibraryNode* LibraryNode = FunctionReferenceNode->LoadReferencedNode())
+				{
+					if(URigVMBlueprint* DependencyBlueprint = LibraryNode->GetTypedOuter<URigVMBlueprint>())
+					{
+						if(DependencyBlueprint != this)
+						{
+							if(!Dependencies.Contains(DependencyBlueprint))
+							{
+								Dependencies.Add(DependencyBlueprint);
+
+								if(bRecursive)
+								{
+									TArray<URigVMBlueprint*> ChildDependencies = DependencyBlueprint->GetDependencies(true);
+									for(URigVMBlueprint* ChildDependency : ChildDependencies)
+									{
+										Dependencies.AddUnique(ChildDependency);
+									}
+								}
+							}
+						}
+					}
+				}
+			}
+		}
+	}
+
+	return Dependencies;
+}
+
+TArray<FAssetData> URigVMBlueprint::GetDependentAssets() const
+{
+	TArray<FAssetData> Dependents;
+	TArray<FSoftObjectPath> AssetPaths;
+
+	if(URigVMFunctionLibrary* FunctionLibrary = RigVMClient.GetFunctionLibrary())
+	{
+		IAssetRegistry& AssetRegistry = IAssetRegistry::GetChecked();
+		TArray<URigVMLibraryNode*> Functions = FunctionLibrary->GetFunctions();
+		for(URigVMLibraryNode* Function : Functions)
+		{
+			const FName FunctionName = Function->GetFName();
+			if(IsFunctionPublic(FunctionName))
+			{
+				TArray<TSoftObjectPtr<URigVMFunctionReferenceNode>> References = FunctionLibrary->GetReferencesForFunction(FunctionName);
+				for(const TSoftObjectPtr<URigVMFunctionReferenceNode>& Reference : References)
+				{
+					if (const URigVMFunctionReferenceNode* ReferencePtr = Reference.Get())
+					{
+						if (const URigVMBlueprint* ControlRigBlueprint = ReferencePtr->GetTypedOuter<URigVMBlueprint>())
+						{
+							const TSoftObjectPtr<UPackage> Blueprint = ControlRigBlueprint;
+							const FSoftObjectPath AssetPath = Blueprint.ToSoftObjectPath();
+							if(AssetPath.GetLongPackageName().StartsWith(TEXT("/Engine/Transient")))
+							{
+								continue;
+							}
+				
+							if(!AssetPaths.Contains(AssetPath))
+							{
+								AssetPaths.Add(AssetPath);
+
+								const FAssetData AssetData = AssetRegistry.GetAssetByObjectPath(AssetPath);
+								if(AssetData.IsValid())
+								{
+									Dependents.Add(AssetData);
+								}
+							}						
+						}
+					}
+				}
+			}
+		}
+	}
+
+	return Dependents;
+}
+
+TArray<URigVMBlueprint*> URigVMBlueprint::GetDependentBlueprints(bool bRecursive, bool bOnlyLoaded) const
+{
+	TArray<FAssetData> Assets = GetDependentAssets();
+	TArray<URigVMBlueprint*> Dependents;
+
+	for(const FAssetData& Asset : Assets)
+	{
+		if (!bOnlyLoaded || Asset.IsAssetLoaded())
+		{
+			if(URigVMBlueprint* Dependent = Cast<URigVMBlueprint>(Asset.GetAsset()))
+			{
+				if(!Dependents.Contains(Dependent))
+				{
+					Dependents.Add(Dependent);
+
+					if(bRecursive && Dependent != this)
+					{
+						TArray<URigVMBlueprint*> ParentDependents = Dependent->GetDependentBlueprints(true);
+						for(URigVMBlueprint* ParentDependent : ParentDependents)
+						{
+							Dependents.AddUnique(ParentDependent);
+						}
+					}
+				}
+			}
+		}
+	}
+
+	return Dependents;
+}
+
+void URigVMBlueprint::GetTypeActions(FBlueprintActionDatabaseRegistrar& ActionRegistrar) const
+{
+	DECLARE_SCOPE_HIERARCHICAL_COUNTER_FUNC()
+
+	// todooooo
+	// IControlRigEditorModule::Get().GetTypeActions((URigVMBlueprint*)this, ActionRegistrar);
+}
+
+void URigVMBlueprint::GetInstanceActions(FBlueprintActionDatabaseRegistrar& ActionRegistrar) const
+{
+	DECLARE_SCOPE_HIERARCHICAL_COUNTER_FUNC()
+
+	// todooooo
+	// IControlRigEditorModule::Get().GetInstanceActions((URigVMBlueprint*)this, ActionRegistrar);
+}
+
+void URigVMBlueprint::SetObjectBeingDebugged(UObject* NewObject)
+{
+	URigVMHost* PreviousRigBeingDebugged = Cast<URigVMHost>(GetObjectBeingDebugged());
+	if (PreviousRigBeingDebugged && PreviousRigBeingDebugged != NewObject)
+	{
+		PreviousRigBeingDebugged->DrawInterface.Reset();
+		PreviousRigBeingDebugged->RigVMLog = nullptr;
+	}
+
+	Super::SetObjectBeingDebugged(NewObject);
+}
+
+void URigVMBlueprint::PostTransacted(const FTransactionObjectEvent& TransactionEvent)
+{
+	DECLARE_SCOPE_HIERARCHICAL_COUNTER_FUNC()
+	Super::PostTransacted(TransactionEvent);
+
+	if (TransactionEvent.GetEventType() == ETransactionObjectEventType::UndoRedo)
+	{
+		TArray<FName> PropertiesChanged = TransactionEvent.GetChangedProperties();
+
+		if (PropertiesChanged.Contains(GET_MEMBER_NAME_CHECKED(URigVMBlueprint, VMRuntimeSettings)))
+		{
+			PropagateRuntimeSettingsFromBPToInstances();
+		}
+
+		if (PropertiesChanged.Contains(GET_MEMBER_NAME_CHECKED(URigVMBlueprint, NewVariables)))
+		{
+			if (RefreshEditorEvent.IsBound())
+			{
+				RefreshEditorEvent.Broadcast(this);
+			}
+			(void)MarkPackageDirty();			
+		}
+
+		if (PropertiesChanged.Contains(GET_MEMBER_NAME_CHECKED(URigVMBlueprint, RigVMClient)) ||
+			PropertiesChanged.Contains(GET_MEMBER_NAME_CHECKED(URigVMBlueprint, UbergraphPages)))
+		{
+			UbergraphPages.RemoveAll([](const UEdGraph* UberGraph) -> bool
+			{
+ 				return UberGraph == nullptr || !IsValid(UberGraph);
+			});
+			RigVMClient.PostTransacted(TransactionEvent);
+
+			RecompileVM();
+			(void)MarkPackageDirty();			
+		}
+	}
+}
+
+void URigVMBlueprint::ReplaceDeprecatedNodes()
+{
+	TArray<UEdGraph*> EdGraphs;
+	GetAllGraphs(EdGraphs);
+
+	for (UEdGraph* EdGraph : EdGraphs)
+	{
+		if (URigVMEdGraph* RigGraph = Cast<URigVMEdGraph>(EdGraph))
+		{
+			RigGraph->Schema = URigVMEdGraphSchema::StaticClass();
+		}
+	}
+
+	Super::ReplaceDeprecatedNodes();
+}
+
+void URigVMBlueprint::PreDuplicate(FObjectDuplicationParameters& DupParams)
+{
+	Super::PreDuplicate(DupParams);
+	PreDuplicateAssetPath = GetPathName();
+	if(URigVMBlueprintGeneratedClass* CRGeneratedClass = GetRigVMBlueprintGeneratedClass())
+	{
+		PreDuplicateHostPath = FSoftObjectPath(CRGeneratedClass->GetPathName());
+	}
+	else
+	{
+		PreDuplicateHostPath.Reset();
+	}
+
+}
+
+void URigVMBlueprint::PostDuplicate(bool bDuplicateForPIE)
+{
+	// assuming PostDuplicate is always followed by a PostLoad:
+	// so theoretically, PostDuplicate just makes corrections to the serialized data and does nothing more,
+	// while PostLoad looks at whatever is serialized and load it into memory according to the version of the editor used
+	// note: how to know if we have corrected everything?
+	// ans: check the reference viewer for the duplicated BP and make sure that the original BP does not appear in there
+	
+	{
+		// pause compilation because we need to patch some stuff first
+		TGuardValue<bool> CompilingGuard(bIsCompiling, true);
+		// this will create the new EMPTY generated class to be used as the function store for this BP
+		// it will be filled during PostLoad based on the graph model
+		Super::PostDuplicate(bDuplicateForPIE);
+	}
+	
+	auto UpdateFunctionHeaders = [this](const FString& InOldPath, const FString& InNewPath)
+	{
+		if(InOldPath.IsEmpty() || InNewPath.IsEmpty())
+		{
+			return;
+		}
+		if(!InNewPath.Equals(InOldPath, ESearchCase::CaseSensitive))
+		{
+			if(URigVMBlueprintGeneratedClass* CRGeneratedClass = GetRigVMBlueprintGeneratedClass())
+			{
+				FRigVMGraphFunctionStore& Store = CRGeneratedClass->GraphFunctionStore;
+				// technically not needed, the store should be empty, it will not be populated until PostLoad
+				// this code is kept here just in case things change in the future
+				if (!(ensure(Store.PublicFunctions.Num() == 0) && ensure(Store.PrivateFunctions.Num() == 0)))
+				{
+					Store.PostDuplicateHost(InOldPath, InNewPath);
+				}
+			}
+			RigVMClient.PostDuplicateHost(InOldPath, InNewPath);
+		}
+	};
+
+	// update the paths once for the blueprint and once for the generated class
+	// make sure all function headers pointing to things in the old BP are changed to
+	// point to their duplicates in the new BP
+	UpdateFunctionHeaders(PreDuplicateAssetPath.ToString(), GetPathName());
+	if(const URigVMBlueprintGeneratedClass* CRGeneratedClass = GetRigVMBlueprintGeneratedClass())
+	{
+		UpdateFunctionHeaders(PreDuplicateHostPath.ToString(), CRGeneratedClass->GetPathName());
+	}
+
+	PreDuplicateAssetPath.Reset();
+	PreDuplicateHostPath.Reset();
+
+	FBlueprintEditorUtils::MarkBlueprintAsStructurallyModified(this);
+
+	// now that the data in the new BP is correct, make sure there is a call to PostLoad()
+	// to complete the loading and get the BP ready, including refresh all models, patch function store, and recompile VM
+	check(HasAnyFlags(RF_NeedPostLoad));
+}
+
+void URigVMBlueprint::GetAssetRegistryTags(TArray<FAssetRegistryTag>& OutTags) const
+{
+	if (CachedAssetTags.IsEmpty())
+	{
+		Super::GetAssetRegistryTags(OutTags);
+		CachedAssetTags = OutTags;
+	}
+	else
+	{
+		OutTags = CachedAssetTags;
+	}
+}
+
+FRigVMGraphModifiedEvent& URigVMBlueprint::OnModified()
+{
+	return ModifiedEvent;
+}
+
+
+FOnVMCompiledEvent& URigVMBlueprint::OnVMCompiled()
+{
+	return VMCompiledEvent;
+}
+
+UClass* URigVMBlueprint::GetRigVMHostClass()
+{
+	return GeneratedClass;
+}
+
+URigVMHost* URigVMBlueprint::CreateRigVMHost()
+{
+	RecompileVMIfRequired();
+
+	URigVMHost* Host = NewObject<URigVMHost>(this, GetRigVMHostClass());
+	Host->Initialize(true);
+	return Host;
+}
+
+TArray<UStruct*> URigVMBlueprint::GetAvailableRigVMStructs() const
+{
+	TArray<UStruct*> Structs;
+	UStruct* BaseStruct = FRigVMStruct::StaticStruct();
+
+	for (const FRigVMFunction& Function : FRigVMRegistry::Get().GetFunctions())
+	{
+		if (Function.Struct)
+		{
+			if (Function.Struct->IsChildOf(BaseStruct))
+			{
+				Structs.Add(Function.Struct);
+				// todo: filter by available types
+				// todo: filter by execute context
+			}
+		}
+	}
+
+	return Structs;
+}
+
+#if WITH_EDITOR
+
+TArray<FRigVMGraphVariableDescription> URigVMBlueprint::GetMemberVariables() const
+{
+	TArray<FRigVMGraphVariableDescription> Variables;
+	for (const FBPVariableDescription& BPVariable : NewVariables)
+	{
+		FRigVMGraphVariableDescription NewVariable;
+		NewVariable.Name = BPVariable.VarName;
+		NewVariable.DefaultValue = BPVariable.DefaultValue;
+		FString CPPType;
+		UObject* CPPTypeObject;
+		RigVMTypeUtils::CPPTypeFromPinType(BPVariable.VarType, CPPType, &CPPTypeObject);
+		NewVariable.CPPType = CPPType;
+		NewVariable.CPPTypeObject = CPPTypeObject;
+		Variables.Add(NewVariable);
+	}
+
+	return Variables;
+}
+
+FName URigVMBlueprint::AddMemberVariable(const FName& InName, const FString& InCPPType, bool bIsPublic, bool bIsReadOnly, FString InDefaultValue)
+{
+	FRigVMExternalVariable Variable = RigVMTypeUtils::ExternalVariableFromCPPTypePath(InName, InCPPType, bIsPublic, bIsReadOnly);
+	FName Result = AddHostMemberVariableFromExternal(Variable, InDefaultValue);
+	if (!Result.IsNone())
+	{
+		FBPCompileRequest Request(this, EBlueprintCompileOptions::None, nullptr);
+		FBlueprintCompilationManager::CompileSynchronously(Request);
+	}
+	return Result;
+}
+
+bool URigVMBlueprint::RemoveMemberVariable(const FName& InName)
+{
+	const int32 VarIndex = FBlueprintEditorUtils::FindNewVariableIndex(this, InName);
+	if (VarIndex == INDEX_NONE)
+	{
+		return false;
+	}
+	
+	FBlueprintEditorUtils::RemoveMemberVariable(this, InName);
+	return true;
+}
+
+bool URigVMBlueprint::RenameMemberVariable(const FName& InOldName, const FName& InNewName)
+{
+	int32 VarIndex = FBlueprintEditorUtils::FindNewVariableIndex(this, InOldName);
+	if (VarIndex == INDEX_NONE)
+	{
+		return false;
+	}
+
+	VarIndex = FBlueprintEditorUtils::FindNewVariableIndex(this, InNewName);
+	if (VarIndex != INDEX_NONE)
+	{
+		return false;
+	}
+	
+	FBlueprintEditorUtils::RenameMemberVariable(this, InOldName, InNewName);
+	return true;
+}
+
+bool URigVMBlueprint::ChangeMemberVariableType(const FName& InName, const FString& InCPPType, bool bIsPublic,
+	bool bIsReadOnly, FString InDefaultValue)
+{
+	int32 VarIndex = FBlueprintEditorUtils::FindNewVariableIndex(this, InName);
+	if (VarIndex == INDEX_NONE)
+	{
+		return false;
+	}
+
+	FRigVMExternalVariable Variable;
+	Variable.Name = InName;
+	Variable.bIsPublic = bIsPublic;
+	Variable.bIsReadOnly = bIsReadOnly;
+
+	FString CPPType = InCPPType;
+	if (CPPType.StartsWith(TEXT("TMap<")))
+	{
+		UE_LOG(LogRigVMDeveloper, Warning, TEXT("TMap Variables are not supported."));
+		return false;
+	}
+
+	Variable.bIsArray = RigVMTypeUtils::IsArrayType(CPPType);
+	if (Variable.bIsArray)
+	{
+		CPPType = RigVMTypeUtils::BaseTypeFromArrayType(CPPType);
+	}
+
+	if (CPPType == TEXT("bool"))
+	{
+		Variable.TypeName = *CPPType;
+		Variable.Size = sizeof(bool);
+	}
+	else if (CPPType == TEXT("float"))
+	{
+		Variable.TypeName = *CPPType;
+		Variable.Size = sizeof(float);
+	}
+	else if (CPPType == TEXT("double"))
+	{
+		Variable.TypeName = *CPPType;
+		Variable.Size = sizeof(double);
+	}
+	else if (CPPType == TEXT("int32"))
+	{
+		Variable.TypeName = *CPPType;
+		Variable.Size = sizeof(int32);
+	}
+	else if (CPPType == TEXT("FString"))
+	{
+		Variable.TypeName = *CPPType;
+		Variable.Size = sizeof(FString);
+	}
+	else if (CPPType == TEXT("FName"))
+	{
+		Variable.TypeName = *CPPType;
+		Variable.Size = sizeof(FName);
+	}
+	else if(UScriptStruct* ScriptStruct = RigVMTypeUtils::FindObjectFromCPPTypeObjectPath<UScriptStruct>(CPPType))
+	{
+		Variable.TypeName = *RigVMTypeUtils::GetUniqueStructTypeName(ScriptStruct);
+		Variable.TypeObject = ScriptStruct;
+		Variable.Size = ScriptStruct->GetStructureSize();
+	}
+	else if (UEnum* Enum= RigVMTypeUtils::FindObjectFromCPPTypeObjectPath<UEnum>(CPPType))
+	{
+		Variable.TypeName = *RigVMTypeUtils::CPPTypeFromEnum(Enum);
+		Variable.TypeObject = Enum;
+		Variable.Size = Enum->GetResourceSizeBytes(EResourceSizeMode::EstimatedTotal);
+	}
+
+	FEdGraphPinType PinType = RigVMTypeUtils::PinTypeFromExternalVariable(Variable);
+	if (!PinType.PinCategory.IsValid())
+	{
+		return false;
+	}
+
+	FBlueprintEditorUtils::ChangeMemberVariableType(this, InName, PinType);
+
+	return true;
+}
+
+#endif
+
+void URigVMBlueprint::RebuildGraphFromModel()
+{
+	DECLARE_SCOPE_HIERARCHICAL_COUNTER_FUNC()
+
+	TGuardValue<bool> SelfGuard(bSuspendModelNotificationsForSelf, true);
+	TGuardValue<bool> ClientIgnoreModificationsGuard(RigVMClient.bIgnoreModelNotifications, true);
+	
+	verify(GetOrCreateController());
+
+	TArray<UEdGraph*> EdGraphs;
+	GetAllGraphs(EdGraphs);
+
+	for (UEdGraph* Graph : EdGraphs)
+	{
+		TArray<UEdGraphNode*> Nodes = Graph->Nodes;
+		for (UEdGraphNode* Node : Nodes)
+		{
+			Graph->RemoveNode(Node);
+		}
+
+		if (URigVMEdGraph* RigGraph = Cast<URigVMEdGraph>(Graph))
+		{
+			if (RigGraph->bIsFunctionDefinition)
+			{
+				FunctionGraphs.Remove(RigGraph);
+			}
+		}
+	}
+
+	if(FunctionLibraryEdGraph && RigVMClient.GetFunctionLibrary())
+	{
+		FunctionLibraryEdGraph->ModelNodePath = RigVMClient.GetFunctionLibrary()->GetNodePath();
+	}
+
+	TArray<URigVMGraph*> RigGraphs = RigVMClient.GetAllModels(true, true);
+
+	for (int32 RigGraphIndex = 0; RigGraphIndex < RigGraphs.Num(); RigGraphIndex++)
+	{
+		GetOrCreateController(RigGraphs[RigGraphIndex])->ResendAllNotifications();
+	}
+
+	for (int32 RigGraphIndex = 0; RigGraphIndex < RigGraphs.Num(); RigGraphIndex++)
+	{
+		URigVMGraph* RigGraph = RigGraphs[RigGraphIndex];
+
+		for (URigVMNode* RigNode : RigGraph->GetNodes())
+		{
+			if (URigVMCollapseNode* CollapseNode = Cast<URigVMCollapseNode>(RigNode))
+			{
+				CreateEdGraphForCollapseNodeIfNeeded(CollapseNode, true);
+			}
+		}
+	}
+}
+
+void URigVMBlueprint::Notify(ERigVMGraphNotifType InNotifType, UObject* InSubject)
+{
+	GetOrCreateController()->Notify(InNotifType, InSubject);
+}
+
+void URigVMBlueprint::HandleModifiedEvent(ERigVMGraphNotifType InNotifType, URigVMGraph* InGraph, UObject* InSubject)
+{
+	DECLARE_SCOPE_HIERARCHICAL_COUNTER_FUNC()
+
+#if WITH_EDITOR
+
+	if (bSuspendAllNotifications)
+	{
+		return;
+	}
+
+	// since it's possible that a notification will be already sent / forwarded to the
+	// listening objects within the switch statement below - we keep a flag to mark
+	// the notify for still pending (or already sent)
+	bool bNotifForOthersPending = true;
+
+	auto MarkBlueprintAsStructurallyModified = [this]()
+	{
+		if(VMRecompilationBracket == 0)
+		{
+			if(bMarkBlueprintAsStructurallyModifiedPending)
+			{
+				bMarkBlueprintAsStructurallyModifiedPending = false;
+				FBlueprintEditorUtils::MarkBlueprintAsStructurallyModified(this);
+			}
+		}
+		else
+		{
+			bMarkBlueprintAsStructurallyModifiedPending = true;
+		}
+	};
+
+	if (!bSuspendModelNotificationsForSelf)
+	{
+		switch (InNotifType)
+		{
+			case ERigVMGraphNotifType::InteractionBracketOpened:
+			{
+				IncrementVMRecompileBracket();
+				break;
+			}
+			case ERigVMGraphNotifType::InteractionBracketClosed:
+			case ERigVMGraphNotifType::InteractionBracketCanceled:
+			{
+				DecrementVMRecompileBracket();
+				MarkBlueprintAsStructurallyModified();
+				break;
+			}
+			case ERigVMGraphNotifType::PinDefaultValueChanged:
+			{
+				if (URigVMPin* Pin = Cast<URigVMPin>(InSubject))
+				{
+					bool bRequiresRecompile = false;
+
+					URigVMPin* RootPin = Pin->GetRootPin();
+					static const FString ConstSuffix = TEXT(":Const");
+					const FString PinHash = RootPin->GetPinPath(true) + ConstSuffix;
+					
+					if (const FRigVMOperand* Operand = PinToOperandMap.Find(PinHash))
+					{
+						FRigVMASTProxy RootPinProxy = FRigVMASTProxy::MakeFromUObject(RootPin);
+						if(const FRigVMExprAST* Expression = InGraph->GetRuntimeAST()->GetExprForSubject(RootPinProxy))
+						{
+							bRequiresRecompile = Expression->NumParents() > 1;
+						}
+						else
+						{
+							bRequiresRecompile = true;
+						}
+
+						// If we are only changing a pin's default value, we need to
+						// check if there is a connection to a sub-pin of the root pin
+						// that has its value is directly stored in the root pin due to optimization, if so,
+						// we want to recompile to make sure the pin's new default value and values from other connections
+						// are both applied to the root pin because GetDefaultValue() alone cannot account for values
+						// from other connections.
+						if(!bRequiresRecompile)
+						{
+							TArray<URigVMPin*> SourcePins = RootPin->GetLinkedSourcePins(true);
+							for (const URigVMPin* SourcePin : SourcePins)
+							{
+								// check if the source node is optimized out, if so, only a recompile will allows us
+								// to re-query its value.
+								FRigVMASTProxy SourceNodeProxy = FRigVMASTProxy::MakeFromUObject(SourcePin->GetNode());
+								if (InGraph->GetRuntimeAST()->GetExprForSubject(SourceNodeProxy) == nullptr)
+								{
+									bRequiresRecompile = true;
+									break;
+								}
+							}
+						} 
+						
+						if(!bRequiresRecompile)
+						{
+							const FString DefaultValue = RootPin->GetDefaultValue();
+
+							URigVMBlueprintGeneratedClass* RigClass = GetRigVMBlueprintGeneratedClass();
+							URigVMHost* CDO = Cast<URigVMHost>(RigClass->GetDefaultObject(true /* create if needed */));
+							if (CDO->VM != nullptr)
+							{
+								CDO->VM->SetPropertyValueFromString(*Operand, DefaultValue);
+							}
+
+							TArray<UObject*> ArchetypeInstances;
+							CDO->GetArchetypeInstances(ArchetypeInstances);
+							for (UObject* ArchetypeInstance : ArchetypeInstances)
+							{
+								URigVMHost* InstancedHost = Cast<URigVMHost>(ArchetypeInstance);
+								if (InstancedHost)
+								{
+									if (InstancedHost->VM)
+									{
+										InstancedHost->VM->SetPropertyValueFromString(*Operand, DefaultValue);
+									}
+								}
+							}
+
+							if (Pin->IsDefinedAsConstant() || Pin->GetRootPin()->IsDefinedAsConstant())
+							{
+								// re-init the rigs
+								RequestRigVMInit();
+								bRequiresRecompile = true;
+							}
+						}
+					}
+					else
+					{
+						bRequiresRecompile = true;
+					}
+				
+					if(bRequiresRecompile)
+					{
+						RequestAutoVMRecompilation();
+					}
+				}
+				(void)MarkPackageDirty();
+				break;
+			}
+			case ERigVMGraphNotifType::NodeAdded:
+			case ERigVMGraphNotifType::NodeRemoved:
+			{
+				bool bAdded = InNotifType == ERigVMGraphNotifType::NodeAdded;
+				if (!bAdded)
+				{
+					if (URigVMNode* RigVMNode = Cast<URigVMNode>(InSubject))
+					{
+						RemoveBreakpoint(RigVMNode);
+					}
+				}
+
+				if (URigVMCollapseNode* CollapseNode = Cast<URigVMCollapseNode>(InSubject))
+				{
+					if (bAdded)
+					{
+						// If the controller for this graph already exist, make sure it is referencing the correct graph
+						if (URigVMController* Controller = RigVMClient.GetController(CollapseNode->GetContainedGraph()))
+						{
+							Controller->SetGraph(CollapseNode->GetContainedGraph());
+						}
+						
+						CreateEdGraphForCollapseNodeIfNeeded(CollapseNode);
+					}
+					else
+					{
+						bNotifForOthersPending = !RemoveEdGraphForCollapseNode(CollapseNode, true);
+
+						// Cannot remove from the Controllers array because we would lose the action stack on that graph
+						// Controllers.Remove(CollapseNode->GetContainedGraph();
+					}
+
+					RequestAutoVMRecompilation();
+
+					(void)MarkPackageDirty();
+					MarkBlueprintAsStructurallyModified();
+					break;
+				}
+
+				if (URigVMNode* RigVMNode = Cast<URigVMNode>(InSubject))
+				{
+					if(RigVMNode->IsEvent() && RigVMNode->GetGraph()->IsRootGraph())
+					{
+						// let the UI know the title for the graph may have changed.
+						RigVMClient.NotifyOuterOfPropertyChange();
+
+						if(URigVMEdGraph* EdGraph = Cast<URigVMEdGraph>(GetEdGraph(RigVMNode->GetGraph())))
+						{
+							// decide if this graph should be renameable
+							const int32 NumberOfEvents = Algo::CountIf(RigVMNode->GetGraph()->GetNodes(), [](const URigVMNode* NodeToCount) -> bool
+							{
+								return NodeToCount->IsEvent() && NodeToCount->CanOnlyExistOnce();
+							});
+							EdGraph->bAllowRenaming = NumberOfEvents != 1;
+						}
+					}
+				}
+				// fall through to the next case
+			}
+			case ERigVMGraphNotifType::LinkAdded:
+			case ERigVMGraphNotifType::LinkRemoved:
+			case ERigVMGraphNotifType::PinArraySizeChanged:
+			case ERigVMGraphNotifType::PinDirectionChanged:
+			{
+				RequestAutoVMRecompilation();
+				(void)MarkPackageDirty();
+
+				// we don't need to mark the blueprint as modified since we only
+				// need to recompile the VM here - unless we don't auto recompile.
+				if(!bAutoRecompileVM)
+				{
+					MarkBlueprintAsStructurallyModified();
+				}
+				break;
+			}
+			case ERigVMGraphNotifType::PinWatchedChanged:
+			{
+				if (URigVMHost* DebuggedHost = Cast<URigVMHost>(GetObjectBeingDebugged()))
+				{
+					URigVMPin* Pin = CastChecked<URigVMPin>(InSubject)->GetRootPin(); 
+					URigVMCompiler* Compiler = URigVMCompiler::StaticClass()->GetDefaultObject<URigVMCompiler>();
+					Compiler->Settings = VMCompileSettings;
+
+					TSharedPtr<FRigVMParserAST> RuntimeAST = GetDefaultModel()->GetRuntimeAST();
+					
+					if(Pin->RequiresWatch())
+					{
+						// check if the node is optimized out - in that case we need to recompile
+						if(DebuggedHost->GetVM()->GetByteCode().GetFirstInstructionIndexForSubject(Pin->GetNode()) == INDEX_NONE)
+						{
+							RequestAutoVMRecompilation();
+							(void)MarkPackageDirty();
+						}
+						else
+						{
+							if(DebuggedHost->GetVM()->GetDebugMemory()->Num() == 0)
+							{
+								RequestAutoVMRecompilation();
+								(void)MarkPackageDirty();
+							}
+							else
+							{
+								Compiler->MarkDebugWatch(true, Pin, DebuggedHost->GetVM(), &PinToOperandMap, RuntimeAST);
+							}
+						}
+					}
+					else
+					{
+						Compiler->MarkDebugWatch(false, Pin, DebuggedHost->GetVM(), &PinToOperandMap, RuntimeAST);
+					}
+				}
+				// break; fall through
+			}
+			case ERigVMGraphNotifType::PinTypeChanged:
+			case ERigVMGraphNotifType::PinIndexChanged:
+			{
+				if (URigVMPin* ModelPin = Cast<URigVMPin>(InSubject))
+				{
+					if (UEdGraph* EdGraph = GetEdGraph(InGraph))
+					{							
+						if (URigVMEdGraph* Graph = Cast<URigVMEdGraph>(EdGraph))
+						{
+							if (UEdGraphNode* EdNode = Graph->FindNodeForModelNodeName(ModelPin->GetNode()->GetFName()))
+							{
+								if (UEdGraphPin* EdPin = EdNode->FindPin(*ModelPin->GetPinPath()))
+								{
+									if (ModelPin->RequiresWatch())
+									{
+										if (!FKismetDebugUtilities::IsPinBeingWatched(this, EdPin))
+										{
+											FKismetDebugUtilities::AddPinWatch(this, FBlueprintWatchedPin(EdPin));
+										}
+									}
+									else
+									{
+										FKismetDebugUtilities::RemovePinWatch(this, EdPin);
+									}
+
+									if(InNotifType == ERigVMGraphNotifType::PinWatchedChanged)
+									{
+										return;
+									}
+									RequestAutoVMRecompilation();
+									(void)MarkPackageDirty();
+								}
+							}
+						}
+					}
+				}
+				// fall through another time
+			}
+			case ERigVMGraphNotifType::PinAdded:
+			case ERigVMGraphNotifType::PinRemoved:
+			case ERigVMGraphNotifType::PinRenamed:
+			{
+				// exposed pin changes like this (as well as type change etc)
+				// require to mark the blueprint as structurally modified,
+				// so that the instance actions work out.
+				if (URigVMPin* ModelPin = Cast<URigVMPin>(InSubject))
+				{
+					if(URigVMCollapseNode* CollapseNode = Cast<URigVMCollapseNode>(ModelPin->GetNode()))
+					{
+						if(Cast<URigVMFunctionLibrary>(CollapseNode->GetOuter()))
+						{
+							MarkBlueprintAsStructurallyModified();
+						}
+					}
+				}
+				break;
+			}
+			case ERigVMGraphNotifType::PinBoundVariableChanged:
+			case ERigVMGraphNotifType::VariableRemappingChanged:
+			{
+				RequestAutoVMRecompilation();
+				(void)MarkPackageDirty();
+				break;
+			}
+			case ERigVMGraphNotifType::NodeRenamed:
+			{
+				if (URigVMCollapseNode* CollapseNode = Cast<URigVMCollapseNode>(InSubject))
+				{
+					FString NewNodePath = CollapseNode->GetNodePath(true /* recursive */);
+					FString Left, Right = NewNodePath;
+					URigVMNode::SplitNodePathAtEnd(NewNodePath, Left, Right);
+					FString OldNodePath = CollapseNode->GetPreviousFName().ToString();
+					if (!Left.IsEmpty())
+					{
+						OldNodePath = URigVMNode::JoinNodePath(Left, OldNodePath);
+					}
+
+					HandleRigVMGraphRenamed(GetRigVMClient(), OldNodePath, NewNodePath);
+
+					if (UEdGraph* ContainedEdGraph = GetEdGraph(CollapseNode->GetContainedGraph()))
+					{
+						ContainedEdGraph->Rename(*CollapseNode->GetEditorSubGraphName(), nullptr);
+					}
+
+					MarkBlueprintAsStructurallyModified();
+				}
+				break;
+			}
+			case ERigVMGraphNotifType::NodeCategoryChanged:
+			case ERigVMGraphNotifType::NodeKeywordsChanged:
+			case ERigVMGraphNotifType::NodeDescriptionChanged:
+			{
+				MarkBlueprintAsStructurallyModified();
+				break;
+			}
+			default:
+			{
+				break;
+			}
+		}
+	}
+
+	// if the notification still has to be sent...
+	if (bNotifForOthersPending && !bSuspendModelNotificationsForOthers)
+	{
+		if (ModifiedEvent.IsBound())
+		{
+			ModifiedEvent.Broadcast(InNotifType, InGraph, InSubject);
+		}
+	}
+#endif
+}
+
+void URigVMBlueprint::SuspendNotifications(bool bSuspendNotifs)
+{
+	if (bSuspendAllNotifications == bSuspendNotifs)
+	{
+		return;
+	}
+
+	bSuspendAllNotifications = bSuspendNotifs;
+	if (!bSuspendNotifs)
+	{
+		RebuildGraphFromModel();
+		RefreshEditorEvent.Broadcast(this);
+		RequestAutoVMRecompilation();
+	}
+}
+
+void URigVMBlueprint::CreateMemberVariablesOnLoad()
+{
+#if WITH_EDITOR
+
+	AddedMemberVariableMap.Reset();
+	for (int32 VariableIndex = 0; VariableIndex < NewVariables.Num(); VariableIndex++)
+	{
+		AddedMemberVariableMap.Add(NewVariables[VariableIndex].VarName, VariableIndex);
+	}
+
+	if (RigVMClient.Num() == 0)
+	{
+		return;
+	}
+
+#endif
+}
+
+#if WITH_EDITOR
+
+FName URigVMBlueprint::FindHostMemberVariableUniqueName(TSharedPtr<FKismetNameValidator> InNameValidator, const FString& InBaseName)
+{
+	FString BaseName = InBaseName;
+	if (InNameValidator->IsValid(BaseName) == EValidatorResult::ContainsInvalidCharacters)
+	{
+		for (TCHAR& TestChar : BaseName)
+		{
+			for (TCHAR BadChar : UE_BLUEPRINT_INVALID_NAME_CHARACTERS)
+			{
+				if (TestChar == BadChar)
+				{
+					TestChar = TEXT('_');
+					break;
+				}
+			}
+		}
+	}
+
+	FString KismetName = BaseName;
+
+	int32 Suffix = 0;
+	while (InNameValidator->IsValid(KismetName) != EValidatorResult::Ok)
+	{
+		KismetName = FString::Printf(TEXT("%s_%d"), *BaseName, Suffix);
+		Suffix++;
+	}
+
+
+	return *KismetName;
+}
+
+int32 URigVMBlueprint::AddHostMemberVariable(URigVMBlueprint* InBlueprint, const FName& InVarName, FEdGraphPinType InVarType, bool bIsPublic, bool bIsReadOnly, FString InDefaultValue)
+{
+	FBPVariableDescription NewVar;
+
+	NewVar.VarName = InVarName;
+	NewVar.VarGuid = FGuid::NewGuid();
+	NewVar.FriendlyName = FName::NameToDisplayString(InVarName.ToString(), (InVarType.PinCategory == UEdGraphSchema_K2::PC_Boolean) ? true : false);
+	NewVar.VarType = InVarType;
+
+	NewVar.PropertyFlags |= (CPF_Edit | CPF_BlueprintVisible | CPF_DisableEditOnInstance);
+
+	if (bIsPublic)
+	{
+		NewVar.PropertyFlags &= ~CPF_DisableEditOnInstance;
+	}
+
+	if (bIsReadOnly)
+	{
+		NewVar.PropertyFlags |= CPF_BlueprintReadOnly;
+	}
+
+	NewVar.ReplicationCondition = COND_None;
+
+	NewVar.Category = UEdGraphSchema_K2::VR_DefaultCategory;
+
+	// user created variables should be none of these things
+	NewVar.VarType.bIsConst = false;
+	NewVar.VarType.bIsWeakPointer = false;
+	NewVar.VarType.bIsReference = false;
+
+	// Text variables, etc. should default to multiline
+	NewVar.SetMetaData(TEXT("MultiLine"), TEXT("true"));
+
+	NewVar.DefaultValue = InDefaultValue;
+
+	return InBlueprint->NewVariables.Add(NewVar);
+}
+
+FName URigVMBlueprint::AddHostMemberVariableFromExternal(FRigVMExternalVariable InVariableToCreate, FString InDefaultValue)
+{
+	FEdGraphPinType PinType = RigVMTypeUtils::PinTypeFromExternalVariable(InVariableToCreate);
+	if (!PinType.PinCategory.IsValid())
+	{
+		return NAME_None;
+	}
+
+	Modify();
+
+	TSharedPtr<FKismetNameValidator> NameValidator = MakeShareable(new FKismetNameValidator(this, NAME_None, nullptr));
+	FName VarName = FindHostMemberVariableUniqueName(NameValidator, InVariableToCreate.Name.ToString());
+	int32 VariableIndex = AddHostMemberVariable(this, VarName, PinType, InVariableToCreate.bIsPublic, InVariableToCreate.bIsReadOnly, InDefaultValue);
+	if (VariableIndex != INDEX_NONE)
+	{
+		FBlueprintEditorUtils::MarkBlueprintAsStructurallyModified(this);
+		return VarName;
+	}
+
+	return NAME_None;
+}
+
+void URigVMBlueprint::PatchFunctionReferencesOnLoad()
+{
+	// If the asset was copied from one project to another, the function referenced might have a different
+	// path, even if the function is internal to the contorl rig. In that case, let's try to find the function
+	// in the local function library.
+
+	for(URigVMGraph* Model : RigVMClient)
+	{
+		TArray<URigVMNode*> Nodes = Model->GetNodes();
+		for (URigVMLibraryNode* Library : RigVMClient.GetFunctionLibrary()->GetFunctions())
+		{
+			Nodes.Append(Library->GetContainedNodes());
+		}
+		
+		for (int32 i=0; i<Nodes.Num(); ++i)
+		{
+			URigVMNode* Node = Nodes[i];
+			if (URigVMFunctionReferenceNode* FunctionReferenceNode = Cast<URigVMFunctionReferenceNode>(Node))
+			{
+				if (!FunctionReferenceNode->ReferencedNodePtr_DEPRECATED.IsValid())
+				{
+					(void)FunctionReferenceNode->ReferencedNodePtr_DEPRECATED.LoadSynchronous();
+				}
+				if (!FunctionReferenceNode->ReferencedNodePtr_DEPRECATED)
+				{
+					if(URigVMFunctionLibrary* FunctionLibrary = RigVMClient.GetFunctionLibrary())
+					{
+						FString FunctionPath = FunctionReferenceNode->ReferencedNodePtr_DEPRECATED.ToSoftObjectPath().GetSubPathString();
+						
+						FString Left, Right;
+						if(FunctionPath.Split(TEXT("."), &Left, &Right))
+						{
+							FString LibraryNodePath = FunctionLibrary->GetNodePath();
+							if(Left == FunctionLibrary->GetName())
+							{
+								if (URigVMLibraryNode* LibraryNode = Cast<URigVMLibraryNode>(FunctionLibrary->FindNode(Right)))
+								{
+									FunctionReferenceNode->ReferencedNodePtr_DEPRECATED = LibraryNode;
+								}
+							}
+						}
+					}
+				}
+
+				if (FunctionReferenceNode->ReferencedNodePtr_DEPRECATED.IsValid())
+				{
+					FunctionReferenceNode->ReferencedFunctionHeader = FunctionReferenceNode->ReferencedNodePtr_DEPRECATED->GetFunctionHeader();					
+				}
+				else if (!FunctionReferenceNode->ReferencedNodePtr_DEPRECATED.IsNull())
+				{
+					// At least lets make sure we store the path in the header
+					FunctionReferenceNode->ReferencedFunctionHeader.LibraryPointer.LibraryNode = FunctionReferenceNode->ReferencedNodePtr_DEPRECATED.ToSoftObjectPath();
+				}
+			}
+
+			if (URigVMCollapseNode* CollapseNode = Cast<URigVMCollapseNode>(Node))
+			{
+				Nodes.Append(CollapseNode->GetContainedNodes());
+			}
+			
+		}
+	}
+	FunctionReferenceNodeData = GetReferenceNodeData();
+}
+
+#endif
+
+void URigVMBlueprint::PatchVariableNodesOnLoad()
+{
+#if WITH_EDITOR
+	AddedMemberVariableMap.Reset();
+	LastNewVariables = NewVariables;
+#endif
+}
+
+void URigVMBlueprint::PatchBoundVariables()
+{
+}
+
+void URigVMBlueprint::PatchVariableNodesWithIncorrectType()
+{
+	TGuardValue<bool> GuardNotifsSelf(bSuspendModelNotificationsForSelf, true);
+
+	struct Local
+	{
+		static bool RefreshIfNeeded(URigVMController* Controller, URigVMVariableNode* VariableNode, const FString& CPPType, UObject* CPPTypeObject)
+		{
+			if (URigVMPin* ValuePin = VariableNode->GetValuePin())
+			{
+				if (ValuePin->GetCPPType() != CPPType || ValuePin->GetCPPTypeObject() != CPPTypeObject)
+				{
+					Controller->RefreshVariableNode(VariableNode->GetFName(), VariableNode->GetVariableName(), CPPType, CPPTypeObject, false);
+					if (RigVMTypeUtils::AreCompatible(*ValuePin->GetCPPType(), ValuePin->GetCPPTypeObject(), *CPPType, CPPTypeObject))
+					{
+						return false;
+					}
+					return true;
+				}
+			}
+			return false;
+		}
+	};
+
+	for (URigVMGraph* Graph : GetAllModels())
+	{
+		URigVMController* Controller = GetOrCreateController(Graph);
+		TArray<URigVMNode*> Nodes = Graph->GetNodes();
+		for (URigVMNode* Node : Nodes)
+		{
+			if (URigVMVariableNode* VariableNode = Cast<URigVMVariableNode>(Node))
+			{
+				if (VariableNode->IsInputArgument())
+				{
+					continue;
+				}
+				
+				FRigVMGraphVariableDescription Description = VariableNode->GetVariableDescription();
+
+				// Check for local variables
+				if (VariableNode->IsLocalVariable())
+				{
+					TArray<FRigVMGraphVariableDescription> LocalVariables = Graph->GetLocalVariables(false);
+					for (FRigVMGraphVariableDescription Variable : LocalVariables)
+					{
+						if (Variable.Name == Description.Name)
+						{
+							if (Local::RefreshIfNeeded(Controller, VariableNode, Variable.CPPType, Variable.CPPTypeObject))
+							{
+								bDirtyDuringLoad = true;
+							}
+							break;
+						}
+					}
+				}
+				else
+				{
+					for (struct FBPVariableDescription& Variable : NewVariables)
+					{
+						if (Variable.VarName == Description.Name)
+						{
+							FString CPPType;
+							UObject* CPPTypeObject = nullptr;
+							RigVMTypeUtils::CPPTypeFromPinType(Variable.VarType, CPPType, &CPPTypeObject);
+							if (Local::RefreshIfNeeded(Controller, VariableNode, CPPType, CPPTypeObject))
+							{
+								bDirtyDuringLoad = true;
+							}
+						}
+					}
+				}
+			}
+		}
+	}
+}
+
+void URigVMBlueprint::PatchLinksWithCast()
+{
+#if WITH_EDITOR
+
+	{
+		TGuardValue<bool> DisableAutoCompile(bAutoRecompileVM, false);
+
+		// find all links containing a cast
+		TArray<TTuple<URigVMGraph*,TWeakObjectPtr<URigVMLink>,FString,FString>> LinksWithCast;
+		for (URigVMGraph* Graph : GetAllModels())
+		{
+			for(URigVMLink* Link : Graph->GetLinks())
+			{
+				const URigVMPin* SourcePin = Link->GetSourcePin();
+				const URigVMPin* TargetPin = Link->GetTargetPin();
+				if (SourcePin && TargetPin)
+				{
+					const TRigVMTypeIndex SourceTypeIndex = SourcePin->GetTypeIndex();
+					const TRigVMTypeIndex TargetTypeIndex = TargetPin->GetTypeIndex();
+					
+					if(SourceTypeIndex != TargetTypeIndex)
+					{
+						if(!FRigVMRegistry::Get().CanMatchTypes(SourceTypeIndex, TargetTypeIndex, true))
+						{
+							LinksWithCast.Emplace(Graph, TWeakObjectPtr<URigVMLink>(Link), SourcePin->GetPinPath(), TargetPin->GetPinPath());
+						}
+					}
+				}
+			}
+		}
+
+		// remove all of those links
+		for(const auto& Tuple : LinksWithCast)
+		{
+			URigVMController* Controller = GetController(Tuple.Get<0>());
+
+			if(URigVMLink* Link = Tuple.Get<1>().Get())
+			{
+				// the link may be detached, attach it first so that removal works.
+				const URigVMPin* SourcePin = Link->GetSourcePin();
+				URigVMPin* TargetPin = Link->GetTargetPin();
+				if(!SourcePin->IsLinkedTo(TargetPin))
+				{
+					const TArray<URigVMController::FLinkedPath> LinkedPaths = Controller->GetLinkedPaths({Link});
+					Controller->RestoreLinkedPaths(LinkedPaths);
+				}
+			}
+
+			Controller->BreakLink(Tuple.Get<2>(), Tuple.Get<3>(), false);
+
+			// notify the user that the link has been broken.
+			UE_LOG(LogRigVMDeveloper, Warning,
+				TEXT("A link was removed in %s (%s) - it contained different types on source and target pin (former cast link?)."),
+				*Controller->GetGraph()->GetNodePath(),
+				*URigVMLink::GetPinPathRepresentation(Tuple.Get<2>(), Tuple.Get<3>())
+			);
+		}
+	}
+#endif
+}
+
+void URigVMBlueprint::PatchFunctionsOnLoad()
+{
+	URigVMBlueprintGeneratedClass* CRGeneratedClass = GetRigVMBlueprintGeneratedClass();
+	FRigVMGraphFunctionStore& Store = CRGeneratedClass->GraphFunctionStore;
+	const URigVMFunctionLibrary* Library = GetLocalFunctionLibrary();
+
+	TMap<URigVMLibraryNode*, FRigVMGraphFunctionHeader> OldHeaders;
+
+	// Backwards compatibility. Store public access in the model
+	TArray<FName> BackwardsCompatiblePublicFunctions;
+	if (GetLinkerCustomVersion(FUE5MainStreamObjectVersion::GUID) < FUE5MainStreamObjectVersion::RigVMSaveFunctionAccessInModel)
+	{
+		for (const FRigVMGraphFunctionData& FunctionData : Store.PublicFunctions)
+		{
+			BackwardsCompatiblePublicFunctions.Add(FunctionData.Header.Name);
+			URigVMLibraryNode* LibraryNode = Cast<URigVMLibraryNode>(FunctionData.Header.LibraryPointer.LibraryNode.ResolveObject());
+			OldHeaders.Add(LibraryNode, FunctionData.Header);
+		}
+	}
+
+	// Addressing issue where PublicGraphFunctions is populated, but the model PublicFunctionNames is not
+	URigVMFunctionLibrary* FunctionLibrary = GetLocalFunctionLibrary();
+	if (FunctionLibrary)
+	{
+		if (PublicGraphFunctions.Num() > FunctionLibrary->PublicFunctionNames.Num())
+		{
+			for (const FRigVMGraphFunctionHeader& PublicHeader : PublicGraphFunctions)
+			{
+				BackwardsCompatiblePublicFunctions.Add(PublicHeader.Name);
+			}
+		}
+	}
+
+	// Lets rebuild the FunctionStore from the model
+	if (FunctionLibrary)
+	{
+		Store.PublicFunctions.Reset();
+		Store.PrivateFunctions.Reset();
+
+		for (URigVMLibraryNode* LibraryNode : FunctionLibrary->GetFunctions())
+		{
+			bool bIsPublic = FunctionLibrary->IsFunctionPublic(LibraryNode->GetFName());
+			if (!bIsPublic)
+			{
+				bIsPublic = BackwardsCompatiblePublicFunctions.Contains(LibraryNode->GetFName());
+				if (bIsPublic)
+				{
+					FunctionLibrary->PublicFunctionNames.Add(LibraryNode->GetFName());
+				}
+			}
+
+			FRigVMGraphFunctionHeader Header = LibraryNode->GetFunctionHeader(CRGeneratedClass);
+			if (FRigVMGraphFunctionHeader* OldHeader = OldHeaders.Find(LibraryNode))
+			{				
+				Header.ExternalVariables = OldHeader->ExternalVariables;
+				Header.Dependencies = OldHeader->Dependencies;
+			}
+			Store.AddFunction(Header, bIsPublic);
+			
+		}
+	}
+
+	// Update dependencies and external variables if needed
+	for (URigVMLibraryNode* LibraryNode : Library->GetFunctions())
+	{
+		GetRigVMClient()->UpdateExternalVariablesForFunction(LibraryNode);
+		GetRigVMClient()->UpdateDependenciesForFunction(LibraryNode);
+	}
+}
+
+void URigVMBlueprint::PropagateRuntimeSettingsFromBPToInstances()
+{
+	if (const UClass* MyControlRigClass = GeneratedClass)
+	{
+		if (URigVMHost* DefaultObject = Cast<URigVMHost>(MyControlRigClass->GetDefaultObject(false)))
+		{
+			DefaultObject->VMRuntimeSettings = VMRuntimeSettings;
+
+			TArray<UObject*> ArchetypeInstances;
+			DefaultObject->GetArchetypeInstances(ArchetypeInstances);
+
+			for (UObject* ArchetypeInstance : ArchetypeInstances)
+			{
+				if (URigVMHost* InstanceHost = Cast<URigVMHost>(ArchetypeInstance))
+				{
+					InstanceHost->VMRuntimeSettings = VMRuntimeSettings;
+				}
+			}
+		}
+	}
+
+	TArray<UEdGraph*> EdGraphs;
+	GetAllGraphs(EdGraphs);
+
+	for (UEdGraph* Graph : EdGraphs)
+	{
+		TArray<UEdGraphNode*> Nodes = Graph->Nodes;
+		for (UEdGraphNode* Node : Nodes)
+		{
+			if(URigVMEdGraphNode* RigNode = Cast<URigVMEdGraphNode>(Node))
+			{
+				RigNode->ReconstructNode_Internal(true);
+			}
+		}
+	}
+}
+
+#if WITH_EDITOR
+
+void URigVMBlueprint::OnPreVariableChange(UObject* InObject)
+{
+	if (InObject != this)
+	{
+		return;
+	}
+	LastNewVariables = NewVariables;
+}
+
+void URigVMBlueprint::OnPostVariableChange(UBlueprint* InBlueprint)
+{
+	if (InBlueprint != this)
+	{
+		return;
+	}
+
+	if (bUpdatingExternalVariables)
+	{
+		return;
+	}
+
+	TGuardValue<bool> UpdatingVariablesGuard(bUpdatingExternalVariables, true);
+	TArray<FBPVariableDescription> LocalLastNewVariables = LastNewVariables;
+
+	TMap<FGuid, int32> NewVariablesByGuid;
+	for (int32 VarIndex = 0; VarIndex < NewVariables.Num(); VarIndex++)
+	{
+		NewVariablesByGuid.Add(NewVariables[VarIndex].VarGuid, VarIndex);
+	}
+
+	TMap<FGuid, int32> OldVariablesByGuid;
+	for (int32 VarIndex = 0; VarIndex < LocalLastNewVariables.Num(); VarIndex++)
+	{
+		OldVariablesByGuid.Add(LocalLastNewVariables[VarIndex].VarGuid, VarIndex);
+	}
+
+	for (const FBPVariableDescription& OldVariable : LocalLastNewVariables)
+	{
+		if (!NewVariablesByGuid.Contains(OldVariable.VarGuid))
+		{
+			OnVariableRemoved(OldVariable.VarName);
+			continue;
+		}
+	}
+
+	for (const FBPVariableDescription& NewVariable : NewVariables)
+	{
+		if (!OldVariablesByGuid.Contains(NewVariable.VarGuid))
+		{
+			OnVariableAdded(NewVariable.VarName);
+			continue;
+		}
+
+		int32 OldVarIndex = OldVariablesByGuid.FindChecked(NewVariable.VarGuid);
+		const FBPVariableDescription& OldVariable = LocalLastNewVariables[OldVarIndex];
+		if (OldVariable.VarName != NewVariable.VarName)
+		{
+			OnVariableRenamed(OldVariable.VarName, NewVariable.VarName);
+		}
+
+		if (OldVariable.VarType != NewVariable.VarType)
+		{
+			OnVariableTypeChanged(NewVariable.VarName, OldVariable.VarType, NewVariable.VarType);
+		}
+	}
+
+	LastNewVariables = NewVariables;
+}
+
+void URigVMBlueprint::OnVariableAdded(const FName& InVarName)
+{
+	FBPVariableDescription Variable;
+	for (FBPVariableDescription& NewVariable : NewVariables)
+	{
+		if (NewVariable.VarName == InVarName)
+		{
+			Variable = NewVariable;
+			break;
+		}
+	}
+
+	const FRigVMExternalVariable ExternalVariable = RigVMTypeUtils::ExternalVariableFromBPVariableDescription(Variable);
+    FString CPPType;
+    UObject* CPPTypeObject = nullptr;
+    RigVMTypeUtils::CPPTypeFromExternalVariable(ExternalVariable, CPPType, &CPPTypeObject);
+	if (CPPTypeObject)
+	{
+		if (ExternalVariable.bIsArray)
+		{
+			CPPType = RigVMTypeUtils::ArrayTypeFromBaseType(CPPTypeObject->GetPathName());
+		}
+		else
+		{
+			CPPType = CPPTypeObject->GetPathName();
+		}
+	}
+
+	// register the type in the registry
+	FRigVMRegistry::Get().FindOrAddType(FRigVMTemplateArgumentType(*CPPType, CPPTypeObject));
+	
+    RigVMPythonUtils::Print(GetFName().ToString(),
+		FString::Printf(TEXT("blueprint.add_member_variable('%s', '%s', %s, %s, '%s')"),
+			*InVarName.ToString(),
+			*CPPType,
+			(ExternalVariable.bIsPublic) ? TEXT("False") : TEXT("True"), 
+			(ExternalVariable.bIsReadOnly) ? TEXT("True") : TEXT("False"), 
+			*Variable.DefaultValue)); 
+	
+	BroadcastExternalVariablesChangedEvent();
+}
+
+void URigVMBlueprint::OnVariableRemoved(const FName& InVarName)
+{
+	TArray<URigVMGraph*> AllGraphs = GetAllModels();
+	for (URigVMGraph* Graph : AllGraphs)
+	{
+		if (URigVMController* Controller = GetOrCreateController(Graph))
+		{
+#if WITH_EDITOR
+			const bool bSetupUndoRedo = !GIsTransacting;
+#else
+			const bool bSetupUndoRedo = false;
+#endif
+			Controller->OnExternalVariableRemoved(InVarName, bSetupUndoRedo);
+		}
+	}
+
+	RigVMPythonUtils::Print(GetFName().ToString(),
+		FString::Printf(TEXT("blueprint.remove_member_variable('%s')"),
+			*InVarName.ToString()));
+	
+	BroadcastExternalVariablesChangedEvent();
+}
+
+void URigVMBlueprint::OnVariableRenamed(const FName& InOldVarName, const FName& InNewVarName)
+{
+	TArray<URigVMGraph*> AllGraphs = GetAllModels();
+	for (URigVMGraph* Graph : AllGraphs)
+	{
+		if (URigVMController* Controller = GetOrCreateController(Graph))
+		{
+#if WITH_EDITOR
+			const bool bSetupUndoRedo = !GIsTransacting;
+#else
+			const bool bSetupUndoRedo = false;
+#endif
+			Controller->OnExternalVariableRenamed(InOldVarName, InNewVarName, bSetupUndoRedo);
+		}
+	}
+
+	RigVMPythonUtils::Print(GetFName().ToString(),
+		FString::Printf(TEXT("blueprint.rename_member_variable('%s', '%s')"),
+			*InOldVarName.ToString(),
+			*InNewVarName.ToString()));
+	
+	BroadcastExternalVariablesChangedEvent();
+}
+
+void URigVMBlueprint::OnVariableTypeChanged(const FName& InVarName, FEdGraphPinType InOldPinType, FEdGraphPinType InNewPinType)
+{
+	FString CPPType;
+	UObject* CPPTypeObject = nullptr;
+	RigVMTypeUtils::CPPTypeFromPinType(InNewPinType, CPPType, &CPPTypeObject);
+	
+	TArray<URigVMGraph*> AllGraphs = GetAllModels();
+	for (URigVMGraph* Graph : AllGraphs)
+	{
+		if (URigVMController* Controller = GetOrCreateController(Graph))
+		{
+#if WITH_EDITOR
+			const bool bSetupUndoRedo = !GIsTransacting;
+#else
+			const bool bSetupUndoRedo = false;
+#endif
+
+			if (!CPPType.IsEmpty())
+			{
+				Controller->OnExternalVariableTypeChanged(InVarName, CPPType, CPPTypeObject, bSetupUndoRedo);
+			}
+			else
+			{
+				Controller->OnExternalVariableRemoved(InVarName, bSetupUndoRedo);
+			}
+		}
+	}
+
+	if(UScriptStruct* ScriptStruct = Cast<UScriptStruct>(CPPTypeObject))
+	{
+		for (auto Var : NewVariables)
+		{
+			if (Var.VarName == InVarName)
+			{
+				CPPType = ScriptStruct->GetName();
+			}
+		}
+	}
+	else if (UEnum* Enum = Cast<UEnum>(CPPTypeObject))
+	{
+		for (auto Var : NewVariables)
+		{
+			if (Var.VarName == InVarName)
+			{
+				CPPType = Enum->GetName();
+			}
+		}
+	}
+
+	// register the type in the registry
+	FRigVMRegistry::Get().FindOrAddType(FRigVMTemplateArgumentType(*CPPType, CPPTypeObject));
+
+	RigVMPythonUtils::Print(GetFName().ToString(),
+		FString::Printf(TEXT("blueprint.change_member_variable_type('%s', '%s')"),
+		*InVarName.ToString(),
+		*CPPType));
+
+	BroadcastExternalVariablesChangedEvent();
+}
+
+void URigVMBlueprint::BroadcastExternalVariablesChangedEvent()
+{
+	if (URigVMBlueprintGeneratedClass* RigClass = GetRigVMBlueprintGeneratedClass())
+	{
+		if (URigVMHost* CDO = Cast<URigVMHost>(RigClass->GetDefaultObject(true /* create if needed */)))
+		{
+			ExternalVariablesChangedEvent.Broadcast(CDO->GetExternalVariables());
+		}
+	}
+}
+
+void URigVMBlueprint::BroadcastNodeDoubleClicked(URigVMNode* InNode)
+{
+	NodeDoubleClickedEvent.Broadcast(this, InNode);
+}
+
+void URigVMBlueprint::BroadcastGraphImported(UEdGraph* InGraph)
+{
+	GraphImportedEvent.Broadcast(InGraph);
+}
+
+void URigVMBlueprint::BroadcastPostEditChangeChainProperty(FPropertyChangedChainEvent& PropertyChangedChainEvent)
+{
+	PostEditChangeChainPropertyEvent.Broadcast(PropertyChangedChainEvent);
+}
+
+void URigVMBlueprint::BroadcastRequestLocalizeFunctionDialog(FRigVMGraphFunctionIdentifier InFunction, bool bForce)
+{
+	RequestLocalizeFunctionDialog.Broadcast(InFunction, this, bForce);
+}
+
+void URigVMBlueprint::BroadCastReportCompilerMessage(EMessageSeverity::Type InSeverity, UObject* InSubject, const FString& InMessage)
+{
+	ReportCompilerMessageEvent.Broadcast(InSeverity, InSubject, InMessage);
+}
+
+#endif
+
+UEdGraph* URigVMBlueprint::CreateEdGraph(URigVMGraph* InModel, bool bForce)
+{
+	check(InModel);
+
+#if WITH_EDITORONLY_DATA
+	if(InModel->IsA<URigVMFunctionLibrary>())
+	{
+		return FunctionLibraryEdGraph;
+	}
+#endif
+	
+	if(bForce)
+	{
+		RemoveEdGraph(InModel);
+	}
+
+	FString GraphName = InModel->GetName();
+	GraphName.RemoveFromStart(FRigVMClient::RigVMModelPrefix);
+	GraphName.TrimStartAndEndInline();
+
+	if(GraphName.IsEmpty())
+	{
+		GraphName = URigVMEdGraphSchema::GraphName_RigVM.ToString();
+	}
+
+	GraphName = RigVMClient.GetUniqueName(*GraphName).ToString();
+
+	URigVMEdGraph* RigVMEdGraph = NewObject<URigVMEdGraph>(this, *GraphName, RF_Transactional);
+	RigVMEdGraph->Schema = URigVMEdGraphSchema::StaticClass();
+	RigVMEdGraph->bAllowDeletion = true;
+	RigVMEdGraph->ModelNodePath = InModel->GetNodePath();
+	RigVMEdGraph->Initialize(this);
+	
+	FBlueprintEditorUtils::AddUbergraphPage(this, RigVMEdGraph);
+	LastEditedDocuments.AddUnique(RigVMEdGraph);
+
+	return RigVMEdGraph;
+}
+
+bool URigVMBlueprint::RemoveEdGraph(URigVMGraph* InModel)
+{
+	if(URigVMEdGraph* RigGraph = Cast<URigVMEdGraph>(GetEdGraph(InModel)))
+	{
+		if(UbergraphPages.Contains(RigGraph))
+		{
+			Modify();
+			UbergraphPages.Remove(RigGraph);
+		}
+		DestroyObject(RigGraph);
+		return true;
+	}
+	return false;
+}
+
+void URigVMBlueprint::DestroyObject(UObject* InObject)
+{
+	RigVMClient.DestroyObject(InObject);
+}
+
+void URigVMBlueprint::RenameGraph(const FString& InNodePath, const FName& InNewName)
+{
+	FName OldName = NAME_None;
+	UEdGraph* EdGraph = GetEdGraph(InNodePath);
+	if(EdGraph)
+	{
+		OldName = EdGraph->GetFName();
+	}
+	
+	RigVMClient.RenameModel(InNodePath, InNewName, true);
+
+	if(EdGraph)
+	{
+		NotifyGraphRenamed(EdGraph, OldName, EdGraph->GetFName());
+	}
+}
+
+void URigVMBlueprint::CreateEdGraphForCollapseNodeIfNeeded(URigVMCollapseNode* InNode, bool bForce)
+{
+	check(InNode);
+
+	if (bForce)
+	{
+		RemoveEdGraphForCollapseNode(InNode, false);
+	}
+
+	if (InNode->GetGraph()->IsA<URigVMFunctionLibrary>())
+	{
+		if (URigVMGraph* ContainedGraph = InNode->GetContainedGraph())
+		{
+			bool bFunctionGraphExists = false;
+			for (UEdGraph* FunctionGraph : FunctionGraphs)
+			{
+				if (URigVMEdGraph* RigFunctionGraph = Cast<URigVMEdGraph>(FunctionGraph))
+				{
+					if (RigFunctionGraph->ModelNodePath == ContainedGraph->GetNodePath())
+					{
+						bFunctionGraphExists = true;
+						break;
+					}
+				}
+			}
+
+			if (!bFunctionGraphExists)
+			{
+				// create a sub graph
+				URigVMEdGraph* RigFunctionGraph = NewObject<URigVMEdGraph>(this, *InNode->GetName(), RF_Transactional);
+				RigFunctionGraph->Schema = URigVMEdGraphSchema::StaticClass();
+				RigFunctionGraph->bAllowRenaming = 1;
+				RigFunctionGraph->bEditable = 1;
+				RigFunctionGraph->bAllowDeletion = 1;
+				RigFunctionGraph->ModelNodePath = ContainedGraph->GetNodePath();
+				RigFunctionGraph->bIsFunctionDefinition = true;
+
+				FunctionGraphs.Add(RigFunctionGraph);
+
+				RigFunctionGraph->Initialize(this);
+
+				GetOrCreateController(ContainedGraph)->ResendAllNotifications();
+			}
+
+		}
+	}
+	else if (URigVMEdGraph* RigGraph = Cast<URigVMEdGraph>(GetEdGraph(InNode->GetGraph())))
+	{
+		if (URigVMGraph* ContainedGraph = InNode->GetContainedGraph())
+		{
+			bool bSubGraphExists = false;
+			for (UEdGraph* SubGraph : RigGraph->SubGraphs)
+			{
+				if (URigVMEdGraph* SubRigGraph = Cast<URigVMEdGraph>(SubGraph))
+				{
+					if (SubRigGraph->ModelNodePath == ContainedGraph->GetNodePath())
+					{
+						bSubGraphExists = true;
+						break;
+					}
+				}
+			}
+
+			if (!bSubGraphExists)
+			{
+				bool bEditable = true;
+				if (InNode->IsA<URigVMAggregateNode>())
+				{
+					bEditable = false;
+				}
+				
+				// create a sub graph
+				URigVMEdGraph* SubRigGraph = NewObject<URigVMEdGraph>(RigGraph, *InNode->GetEditorSubGraphName(), RF_Transactional);
+				SubRigGraph->Schema = URigVMEdGraphSchema::StaticClass();
+				SubRigGraph->bAllowRenaming = 1;
+				SubRigGraph->bEditable = bEditable;
+				SubRigGraph->bAllowDeletion = 1;
+				SubRigGraph->ModelNodePath = ContainedGraph->GetNodePath();
+				SubRigGraph->bIsFunctionDefinition = false;
+
+				RigGraph->SubGraphs.Add(SubRigGraph);
+
+				SubRigGraph->Initialize(this);
+
+				GetOrCreateController(ContainedGraph)->ResendAllNotifications();
+			}
+		}
+	}
+}
+
+bool URigVMBlueprint::RemoveEdGraphForCollapseNode(URigVMCollapseNode* InNode, bool bNotify)
+{
+	check(InNode);
+
+	if (InNode->GetGraph()->IsA<URigVMFunctionLibrary>())
+	{
+		if (URigVMGraph* ContainedGraph = InNode->GetContainedGraph())
+		{
+			for (UEdGraph* FunctionGraph : FunctionGraphs)
+			{
+				if (URigVMEdGraph* RigFunctionGraph = Cast<URigVMEdGraph>(FunctionGraph))
+				{
+					if (RigFunctionGraph->ModelNodePath == ContainedGraph->GetNodePath())
+					{
+						if (URigVMController* SubController = GetController(ContainedGraph))
+						{
+							SubController->OnModified().RemoveAll(RigFunctionGraph);
+						}
+
+						if (ModifiedEvent.IsBound() && bNotify)
+						{
+							ModifiedEvent.Broadcast(ERigVMGraphNotifType::NodeRemoved, InNode->GetGraph(), InNode);
+						}
+
+						FunctionGraphs.Remove(RigFunctionGraph);
+						RigFunctionGraph->Rename(nullptr, GetTransientPackage(), REN_ForceNoResetLoaders | REN_DontCreateRedirectors);
+						RigFunctionGraph->MarkAsGarbage();
+						return bNotify;
+					}
+				}
+			}
+		}
+	}
+	else if (URigVMEdGraph* RigGraph = Cast<URigVMEdGraph>(GetEdGraph(InNode->GetGraph())))
+	{
+		if (URigVMGraph* ContainedGraph = InNode->GetContainedGraph())
+		{
+			for (UEdGraph* SubGraph : RigGraph->SubGraphs)
+			{
+				if (URigVMEdGraph* SubRigGraph = Cast<URigVMEdGraph>(SubGraph))
+				{
+					if (SubRigGraph->ModelNodePath == ContainedGraph->GetNodePath())
+					{
+						if (URigVMController* SubController = GetController(ContainedGraph))
+						{
+							SubController->OnModified().RemoveAll(SubRigGraph);
+						}
+
+						if (ModifiedEvent.IsBound() && bNotify)
+						{
+							ModifiedEvent.Broadcast(ERigVMGraphNotifType::NodeRemoved, InNode->GetGraph(), InNode);
+						}
+
+						RigGraph->SubGraphs.Remove(SubRigGraph);
+						SubRigGraph->Rename(nullptr, GetTransientPackage(), REN_ForceNoResetLoaders | REN_DontCreateRedirectors);
+						SubRigGraph->MarkAsGarbage();
+						return bNotify;
+					}
+				}
+			}
+		}
+	}
+
+	return false;
+}
+
+#undef LOCTEXT_NAMESPACE
+
+
