RigVMEditor: EdGraph, EdGraphNode, EdGraphSchema and Blueprint

diff --git a/Engine/Plugins/Runtime/RigVM/Source/RigVMDeveloper/Private/RigVMBlueprintUtils.cpp b/Engine/Plugins/Runtime/RigVM/Source/RigVMDeveloper/Private/RigVMBlueprintUtils.cpp
new file mode 100644
--- /dev/null
+++ b/Engine/Plugins/Runtime/RigVM/Source/RigVMDeveloper/Private/RigVMBlueprintUtils.cpp
@@ -0,0 +1,148 @@
+// Copyright Epic Games, Inc. All Rights Reserved.
+
+#include "RigVMBlueprintUtils.h"
+#include "Kismet2/BlueprintEditorUtils.h"
+#include "RigVMCore/RigVMStruct.h"
+#include "UObject/UObjectIterator.h"
+#include "EdGraph/RigVMEdGraphNode.h"
+#include "RigVMBlueprint.h"
+#include "Kismet2/Kismet2NameValidators.h"
+
+#define LOCTEXT_NAMESPACE "RigVMBlueprintUtils"
+
+FName FRigVMBlueprintUtils::ValidateName(UBlueprint* InBlueprint, const FString& InName)
+{
+	DECLARE_SCOPE_HIERARCHICAL_COUNTER_FUNC()
+
+	FString Name = InName;
+	if (Name.StartsWith(TEXT("RigUnit_")))
+	{
+		Name.RightChopInline(8, false);
+	}
+	else if (Name.StartsWith(TEXT("RigVMStruct_")))
+	{
+		Name.RightChopInline(12, false);
+	}
+
+	TSharedPtr<FKismetNameValidator> NameValidator;
+	NameValidator = MakeShareable(new FKismetNameValidator(InBlueprint));
+
+	// Clean up BaseName to not contain any invalid characters, which will mean we can never find a legal name no matter how many numbers we add
+	if (NameValidator->IsValid(Name) == EValidatorResult::ContainsInvalidCharacters)
+	{
+		for (TCHAR& TestChar : Name)
+		{
+			for (TCHAR BadChar : UE_BLUEPRINT_INVALID_NAME_CHARACTERS)
+			{
+				if (TestChar == BadChar)
+				{
+					TestChar = TEXT('_');
+					break;
+				}
+			}
+		}
+	}
+
+	if (UClass* ParentClass = InBlueprint->ParentClass)
+	{
+		FFieldVariant ExistingField = FindUFieldOrFProperty(ParentClass, *Name);
+		if (ExistingField)
+		{
+			Name = FString::Printf(TEXT("%s_%d"), *Name, 0);
+		}
+	}
+
+	int32 Count = 0;
+	FString BaseName = Name;
+	while (NameValidator->IsValid(Name) != EValidatorResult::Ok)
+	{
+		// Calculate the number of digits in the number, adding 2 (1 extra to correctly count digits, another to account for the '_' that will be added to the name
+		int32 CountLength = Count > 0 ? (int32)log((double)Count) + 2 : 2;
+
+		// If the length of the final string will be too long, cut off the end so we can fit the number
+		if (CountLength + BaseName.Len() > NameValidator->GetMaximumNameLength())
+		{
+			BaseName.LeftInline(NameValidator->GetMaximumNameLength() - CountLength);
+		}
+		Name = FString::Printf(TEXT("%s_%d"), *BaseName, Count);
+		Count++;
+	}
+
+	return *Name;
+}
+
+void FRigVMBlueprintUtils::ForAllRigVMStructs(TFunction<void(UScriptStruct*)> InFunction)
+{
+	DECLARE_SCOPE_HIERARCHICAL_COUNTER_FUNC()
+
+	// Run over all unit types
+	for(TObjectIterator<UStruct> StructIt; StructIt; ++StructIt)
+	{
+		if (*StructIt)
+		{
+			if(StructIt->IsChildOf(FRigVMStruct::StaticStruct()) && !StructIt->HasMetaData(FRigVMStruct::AbstractMetaName))
+			{
+				if (UScriptStruct* ScriptStruct = Cast<UScriptStruct>(*StructIt))
+				{
+					InFunction(ScriptStruct);
+				}
+			}
+		}
+	}
+}
+
+void FRigVMBlueprintUtils::HandleReconstructAllNodes(UBlueprint* InBlueprint)
+{
+	DECLARE_SCOPE_HIERARCHICAL_COUNTER_FUNC()
+
+	return HandleRefreshAllNodes(InBlueprint);
+}
+
+void FRigVMBlueprintUtils::HandleRefreshAllNodes(UBlueprint* InBlueprint)
+{
+	DECLARE_SCOPE_HIERARCHICAL_COUNTER_FUNC()
+#ifdef WITH_EDITORONLY_DATA
+	// Avoid refreshing EdGraph nodes during cook
+	if (GIsCookerLoadingPackage)
+	{
+		return;
+	}
+	
+	if(const URigVMBlueprint* RigVMBlueprint = Cast<URigVMBlueprint>(InBlueprint))
+	{
+		if (RigVMBlueprint->GetDefaultModel() == nullptr)
+		{
+			return;
+		}
+
+		TArray<URigVMEdGraphNode*> AllNodes;
+		FBlueprintEditorUtils::GetAllNodesOfClass(RigVMBlueprint, AllNodes);
+
+		for (URigVMEdGraphNode* Node : AllNodes)
+		{
+			Node->SetFlags(RF_Transient);
+		}
+
+		for(URigVMEdGraphNode* Node : AllNodes)
+		{
+			Node->ReconstructNode();
+		}
+
+		for (URigVMEdGraphNode* Node : AllNodes)
+		{
+			Node->ClearFlags(RF_Transient);
+		}
+	}
+#endif
+}
+
+void FRigVMBlueprintUtils::RemoveMemberVariableIfNotUsed(UBlueprint* Blueprint, const FName VarName)
+{
+	DECLARE_SCOPE_HIERARCHICAL_COUNTER_FUNC()
+
+	if (Blueprint->IsA<URigVMBlueprint>())
+	{
+		FBlueprintEditorUtils::RemoveMemberVariable(Blueprint, VarName);
+	}
+}
+#undef LOCTEXT_NAMESPACE
\ No newline at end of file
diff --git a/Engine/Plugins/Runtime/RigVM/Source/RigVMDeveloper/Private/RigVMModel/RigVMSchema.cpp b/Engine/Plugins/Runtime/RigVM/Source/RigVMDeveloper/Private/RigVMModel/RigVMSchema.cpp
--- a/Engine/Plugins/Runtime/RigVM/Source/RigVMDeveloper/Private/RigVMModel/RigVMSchema.cpp
+++ b/Engine/Plugins/Runtime/RigVM/Source/RigVMDeveloper/Private/RigVMModel/RigVMSchema.cpp
@@ -799,7 +799,6 @@ void URigVMSchema::SetExecuteContextStruct(UScriptStruct* InExecuteContextStruct
 			// only allow the default execute context on the base class.
 			// please create a child class of URigVMSchema for your use case.
 			verify(InExecuteContextStruct == FRigVMExecuteContext::StaticStruct());
-			return;
 		}
 		
 		ExecuteContextStruct = InExecuteContextStruct;
diff --git a/Engine/Plugins/Runtime/RigVM/Source/RigVMDeveloper/Public/EdGraph/RigVMEdGraph.h b/Engine/Plugins/Runtime/RigVM/Source/RigVMDeveloper/Public/EdGraph/RigVMEdGraph.h
new file mode 100644
--- /dev/null
+++ b/Engine/Plugins/Runtime/RigVM/Source/RigVMDeveloper/Public/EdGraph/RigVMEdGraph.h
@@ -0,0 +1,115 @@
+// Copyright Epic Games, Inc. All Rights Reserved.
+
+#pragma once
+
+#include "EdGraph/EdGraph.h"
+#include "EdGraph/RigVMEdGraphNode.h"
+#include "RigVMModel/RigVMGraph.h"
+#include "RigVMCore/RigVM.h"
+#include "RigVMModel/RigVMClient.h"
+#include "RigVMEdGraph.generated.h"
+
+class URigVMBlueprint;
+class URigVMEdGraphSchema;
+class URigVMController;
+
+DECLARE_MULTICAST_DELEGATE_OneParam(FRigVMEdGraphNodeClicked, URigVMEdGraphNode*);
+
+UCLASS()
+class RIGVMDEVELOPER_API URigVMEdGraph : public UEdGraph, public IRigVMEditorSideObject
+{
+	GENERATED_BODY()
+
+public:
+	URigVMEdGraph();
+
+	/** IRigVMEditorSideObject interface */
+	virtual FRigVMClient* GetRigVMClient() const override;
+	virtual FString GetRigVMNodePath() const override;
+	virtual void HandleRigVMGraphRenamed(const FString& InOldNodePath, const FString& InNewNodePath) override;
+
+	/** Set up this graph */
+	virtual void Initialize(URigVMBlueprint* InBlueprint);
+
+	/** Get the ed graph schema */
+	const URigVMEdGraphSchema* GetRigVMEdGraphSchema();
+
+#if WITH_EDITORONLY_DATA
+	/** Customize blueprint changes based on backwards compatibility */
+	virtual void Serialize(FArchive& Ar) override;
+#endif
+
+#if WITH_EDITOR
+
+	bool bSuspendModelNotifications;
+	bool bIsTemporaryGraphForCopyPaste;
+
+	UEdGraphNode* FindNodeForModelNodeName(const FName& InModelNodeName, const bool bCacheIfRequired = true);
+
+	URigVMBlueprint* GetBlueprint() const;
+	URigVMGraph* GetModel() const;
+	URigVMController* GetController() const;
+	bool IsRootGraph() const { return GetRootGraph() == this; }
+	const URigVMEdGraph* GetRootGraph() const;
+
+	void HandleModifiedEvent(ERigVMGraphNotifType InNotifType, URigVMGraph* InGraph, UObject* InSubject);
+	void ConsumeQueuedNotifications();
+
+	int32 GetInstructionIndex(const URigVMEdGraphNode* InNode, bool bAsInput);
+
+	UPROPERTY()
+	FString ModelNodePath;
+
+	UPROPERTY()
+	bool bIsFunctionDefinition;
+
+private:
+
+	bool bIsSelecting;
+
+	FRigVMEdGraphNodeClicked OnGraphNodeClicked;
+
+	TMap<URigVMNode*, TPair<int32, int32>> CachedInstructionIndices;
+
+	void RemoveNode(UEdGraphNode* InNode);
+
+#endif
+#if WITH_EDITORONLY_DATA
+
+	UPROPERTY(transient)
+	TObjectPtr<URigVMController> TemplateController;
+
+#endif
+#if WITH_EDITOR
+
+	URigVMController* GetTemplateController();
+
+protected:
+	void HandleVMCompiledEvent(UObject* InCompiledObject, URigVM* InVM);
+
+private:
+	TMap<FName, UEdGraphNode*> ModelNodePathToEdNode;
+	mutable TWeakObjectPtr<URigVMGraph> CachedModelGraph;
+
+	/*
+	friend class UControlRigUnitNodeSpawner;
+	friend class UControlRigVariableNodeSpawner;
+	friend class UControlRigParameterNodeSpawner;
+	friend class UControlRigBranchNodeSpawner;
+	friend class UControlRigIfNodeSpawner;
+	friend class UControlRigSelectNodeSpawner;
+	friend class UControlRigTemplateNodeSpawner;
+	friend class UControlRigEnumNodeSpawner;
+	friend class UControlRigFunctionRefNodeSpawner;
+	friend class UControlRigArrayNodeSpawner;
+	friend class UControlRigInvokeEntryNodeSpawner;
+	*/
+
+#endif
+	friend class URigVMEdGraphNode;
+	friend class URigVMBlueprint;
+	friend class UControlRigGraphNode;
+	friend class FControlRigEditor;
+	friend class SControlRigGraphNode;
+	friend class URigVMBlueprint;
+};
diff --git a/Engine/Plugins/Runtime/RigVM/Source/RigVMDeveloper/Public/EdGraph/RigVMEdGraphNode.h b/Engine/Plugins/Runtime/RigVM/Source/RigVMDeveloper/Public/EdGraph/RigVMEdGraphNode.h
new file mode 100644
--- /dev/null
+++ b/Engine/Plugins/Runtime/RigVM/Source/RigVMDeveloper/Public/EdGraph/RigVMEdGraphNode.h
@@ -0,0 +1,280 @@
+// Copyright Epic Games, Inc. All Rights Reserved.
+
+
+#pragma once
+
+#include "CoreMinimal.h"
+#include "UObject/WeakObjectPtrTemplates.h"
+#include "RigVMModel/RigVMGraph.h"
+#include "RigVMEdGraphNode.generated.h"
+
+class FBlueprintActionDatabaseRegistrar;
+class UEdGraph;
+struct FSlateIcon;
+class URigVMBlueprint;
+
+/** Base class for RigVM editor side nodes */
+UCLASS()
+class RIGVMDEVELOPER_API URigVMEdGraphNode : public UEdGraphNode
+{
+	GENERATED_BODY()
+
+	friend class FRigVMEdGraphNodeDetailsCustomization;
+	friend class FRigVMBlueprintCompilerContext;
+	friend class URigVMEdGraph;
+	friend class URigVMEdGraphSchema;
+	friend class URigVMBlueprint;
+	friend class FRigVMEdGraphTraverser;
+	friend class FRigVMEdGraphPanelPinFactory;
+	friend class FRigVMEditor;
+	friend struct FRigVMBlueprintUtils;
+	friend class SRigVMEdGraphPinCurveFloat;
+
+private:
+
+	UPROPERTY()
+	FString ModelNodePath;
+
+	UPROPERTY(transient)
+	TWeakObjectPtr<URigVMNode> CachedModelNode;
+
+	UPROPERTY(transient)
+	TMap<FString, TWeakObjectPtr<URigVMPin>> PinPathToModelPin;
+
+#if WITH_EDITORONLY_DATA
+	/** The property we represent. For template nodes this represents the struct/property type name. */
+	UPROPERTY()
+	FName PropertyName_DEPRECATED;
+
+	UPROPERTY()
+	FString StructPath_DEPRECATED;
+
+	/** Pin Type for property */
+	UPROPERTY()
+	FEdGraphPinType PinType_DEPRECATED;
+
+	/** The type of parameter */
+	UPROPERTY()
+	int32 ParameterType_DEPRECATED;
+
+	/** Expanded pins */
+	UPROPERTY()
+	TArray<FString> ExpandedPins_DEPRECATED;
+#endif
+
+	/** Cached dimensions of this node (used for auto-layout) */
+	FVector2D Dimensions;
+
+	/** The cached node titles */
+	mutable FText NodeTitle;
+
+	/** The cached fulol node title */
+	mutable FText FullNodeTitle;
+
+public:
+
+	DECLARE_MULTICAST_DELEGATE(FNodeTitleDirtied);
+	DECLARE_MULTICAST_DELEGATE(FNodePinsChanged);
+	DECLARE_MULTICAST_DELEGATE(FNodePinExpansionChanged);
+	DECLARE_MULTICAST_DELEGATE(FNodeBeginRemoval);
+
+	struct FPinPair
+	{
+		FPinPair()
+			: InputPin(nullptr)
+			, OutputPin(nullptr)
+		{}
+		UEdGraphPin* InputPin;
+		UEdGraphPin* OutputPin;
+
+		bool IsValid() const { return InputPin != nullptr || OutputPin != nullptr; }
+	};
+
+	URigVMEdGraphNode();
+
+	// UObject Interface.
+#if WITH_EDITOR
+	virtual bool Modify( bool bAlwaysMarkDirty=true ) { return false; }
+#endif
+	
+	// UEdGraphNode Interface.
+	virtual FText GetNodeTitle(ENodeTitleType::Type TitleType) const override;
+	virtual FLinearColor GetNodeTitleColor() const override;
+	virtual FLinearColor GetNodeBodyTintColor() const;
+	virtual bool ShowPaletteIconOnNode() const override;
+	virtual FSlateIcon GetIconAndTint(FLinearColor& OutColor) const override;
+	virtual void AllocateDefaultPins() override;
+	virtual void ReconstructNode() override;
+	virtual void ReconstructNode_Internal(bool bForce = false);
+	virtual void GetNodeContextMenuActions(class UToolMenu* Menu, class UGraphNodeContextMenuContext* Context) const override;
+	virtual void DestroyNode() override;
+	virtual void PinDefaultValueChanged(UEdGraphPin* Pin) override;
+	virtual FText GetTooltipText() const override;
+	virtual bool CanCreateUnderSpecifiedSchema(const UEdGraphSchema* InSchema) const override;
+	virtual void AutowireNewNode(UEdGraphPin* FromPin) override;	
+	virtual bool SupportsCommentBubble() const override { return false; }
+	virtual bool IsSelectedInEditor() const;
+	virtual bool ShouldDrawNodeAsControlPointOnly(int32& OutInputPinIndex, int32& OutOutputPinIndex) const override;
+	virtual void BeginDestroy() override;
+
+	virtual bool IsDeprecated() const override;
+	virtual FEdGraphNodeDeprecationResponse GetDeprecationResponse(EEdGraphNodeDeprecationType DeprecationType) const override;
+
+	/** Set the cached dimensions of this node */
+	void SetDimensions(const FVector2D& InDimensions) { Dimensions = InDimensions; }
+
+	/** Get the cached dimensions of this node */
+	const FVector2D& GetDimensions() const { return Dimensions; }
+
+	/** Check a pin's expansion state */
+	bool IsPinExpanded(const FString& InPinPath);
+
+	/** Propagate pin defaults to underlying properties if they have changed */
+	void CopyPinDefaultsToModel(UEdGraphPin* Pin, bool bUndo = false, bool bPrintPythonCommand = false);
+
+	/** Get the blueprint that this node is contained within */
+	URigVMBlueprint* GetBlueprint() const;
+
+	/** Get the VM model this node lives within */
+	URigVMGraph* GetModel() const;
+
+	/** Get the blueprint that this node is contained within */
+	URigVMController* GetController() const;
+
+	/** Get the VM node this is node is wrapping */
+	URigVMNode* GetModelNode() const;
+
+	/** Get the VM node name this node is wrapping */
+	FName GetModelNodeName() const;
+
+	URigVMPin* GetModelPinFromPinPath(const FString& InPinPath) const;
+
+	/** Add a new element to the aggregate node referred to by the property path */
+	void HandleAddAggregateElement(const FString& InNodePath);
+
+	/** Add a new array element to the array referred to by the property path */
+	void HandleAddArrayElement(FString InPinPath);
+	
+	/** Clear the array referred to by the property path */
+	void HandleClearArray(FString InPinPath);
+
+	/** Remove the array element referred to by the property path */
+	void HandleRemoveArrayElement(FString InPinPath);
+
+	/** Insert a new array element after the element referred to by the property path */
+	void HandleInsertArrayElement(FString InPinPath);
+
+	int32 GetInstructionIndex(bool bAsInput) const;
+
+	const FRigVMTemplate* GetTemplate() const;
+
+	void ClearErrorInfo();
+
+	URigVMPin* FindModelPinFromGraphPin(const UEdGraphPin* InGraphPin) const;
+	UEdGraphPin* FindGraphPinFromModelPin(const URigVMPin* InModelPin, bool bAsInput) const;
+
+	/// Synchronize the stored name/value/type on the graph pin with the value stored on the node. 
+	/// If the pin has sub-pins, the value update is done recursively.
+	void SynchronizeGraphPinNameWithModelPin(const URigVMPin* InModelPin, bool bNotify = true);
+	void SynchronizeGraphPinValueWithModelPin(const URigVMPin* InModelPin);
+	void SynchronizeGraphPinTypeWithModelPin(const URigVMPin* InModelPin);
+	void SynchronizeGraphPinExpansionWithModelPin(const URigVMPin* InModelPin);
+	
+	void SyncGraphNodeTitleWithModelNodeTitle();
+	void SyncGraphNodeNameWithModelNodeName(const URigVMNode* InModelNode);
+
+	FNodeTitleDirtied& OnNodeTitleDirtied() { return NodeTitleDirtied; }
+	FNodePinsChanged& OnNodePinsChanged() { return NodePinsChanged; }
+	FNodePinExpansionChanged& OnNodePinExpansionChanged() { return NodePinExpansionChanged; }
+	FNodeBeginRemoval& OnNodeBeginRemoval() { return NodeBeginRemoval; }
+
+	/** Called when there's a drastic change in the pins */
+	bool ModelPinsChanged(bool bForce = false);
+
+	/** Called when a model pin is added after the node creation */
+	bool ModelPinAdded(const URigVMPin* InModelPin);
+
+	/** Called when a model pin is being removed */
+	bool ModelPinRemoved(const URigVMPin* InModelPin);
+
+	/** Returns true if this node is relying on the cast template */
+	bool DrawAsCompactNode() const;
+
+protected:
+
+	FLinearColor GetNodeProfilingColor() const;
+	FLinearColor GetNodeOpacityColor() const;
+
+	/** Helper function for AllocateDefaultPins */
+	void UpdatePinLists();
+	bool CreateGraphPinFromModelPin(const URigVMPin* InModelPin, EEdGraphPinDirection InDirection,  UEdGraphPin* InParentPin = nullptr);
+	void RemoveGraphSubPins(UEdGraphPin *InParentPin, const TArray<UEdGraphPin*>& InPinsToKeep = TArray<UEdGraphPin*>());
+	bool ModelPinAdded_Internal(const URigVMPin* InModelPin);
+	bool ModelPinRemoved_Internal(const URigVMPin* InModelPin);
+
+	/** Copies default values from underlying properties into pin defaults, for editing */
+	void SetupPinDefaultsFromModel(UEdGraphPin* Pin, const URigVMPin* InModelPin = nullptr);
+
+	/** Recreate pins when we reconstruct this node */
+	virtual void ReallocatePinsDuringReconstruction(const TArray<UEdGraphPin*>& OldPins);
+
+	/** Wire-up new pins given old pin wiring */
+	virtual void RewireOldPinsToNewPins(TArray<UEdGraphPin*>& InOldPins, TArray<UEdGraphPin*>& InNewPins);
+
+	/** Handle anything post-reconstruction */
+	virtual void PostReconstructNode();
+
+	/** Something that could change our title has changed */
+	void InvalidateNodeTitle() const;
+
+	/** Destroy all pins in an array */
+	void DestroyPinList(TArray<UEdGraphPin*>& InPins);
+
+	/** Sets the body + title color from a color provided by the model */
+	void SetColorFromModel(const FLinearColor& InColor);
+
+	UClass* GetRigVMGeneratedClass() const;
+	UClass* GetRigVMSkeletonGeneratedClass() const;
+
+	static FEdGraphPinType GetPinTypeForModelPin(const URigVMPin* InModelPin);
+
+private:
+
+	int32 GetNodeTopologyVersion() const { return NodeTopologyVersion; }
+	int32 NodeTopologyVersion;
+
+	static void ConfigurePin(UEdGraphPin* EdGraphPin, const URigVMPin* ModelPin);
+	TArray<URigVMPin*>& PinListForPin(const URigVMPin* InModelPin);
+
+	FLinearColor CachedTitleColor;
+	FLinearColor CachedNodeColor;
+
+#if WITH_EDITOR
+	bool bEnableProfiling;
+#endif
+
+	TArray<URigVMPin*> ExecutePins;
+	TArray<URigVMPin*> InputOutputPins;
+	TArray<URigVMPin*> InputPins;
+	TArray<URigVMPin*> OutputPins;
+	TArray<TSharedPtr<FRigVMExternalVariable>> ExternalVariables;
+	TArray<UEdGraphPin*> LastEdGraphPins;
+	
+	TMap<URigVMPin*, FPinPair> CachedPins;
+
+	FNodeTitleDirtied NodeTitleDirtied;
+	FNodePinsChanged NodePinsChanged;
+	FNodePinExpansionChanged NodePinExpansionChanged;
+	FNodeBeginRemoval NodeBeginRemoval;
+
+	mutable const FRigVMTemplate* CachedTemplate;
+	mutable TOptional<bool> DrawAsCompactNodeCache;
+
+	friend class SRigVMEdGraphNode;
+	friend class FRigVMArgumentLayout;
+	friend class FRigVMEdGraphDetails;
+	friend class URigVMTemplateNodeSpawner;
+	friend class URigVMArrayNodeSpawner;
+	friend class URigVMIfNodeSpawner;
+	friend class URigVMSelectNodeSpawner;
+};
diff --git a/Engine/Plugins/Runtime/RigVM/Source/RigVMDeveloper/Public/EdGraph/RigVMEdGraphSchema.h b/Engine/Plugins/Runtime/RigVM/Source/RigVMDeveloper/Public/EdGraph/RigVMEdGraphSchema.h
new file mode 100644
--- /dev/null
+++ b/Engine/Plugins/Runtime/RigVM/Source/RigVMDeveloper/Public/EdGraph/RigVMEdGraphSchema.h
@@ -0,0 +1,347 @@
+// Copyright Epic Games, Inc. All Rights Reserved.
+
+#pragma once
+
+#include "EdGraph/RigVMEdGraphNode.h"
+#include "EdGraph/EdGraphSchema.h"
+#include "GraphEditorDragDropAction.h"
+#include "RigVMModel/RigVMGraph.h"
+#include "EdGraphSchema_K2_Actions.h"
+#include "Kismet2/Kismet2NameValidators.h"
+
+#include "RigVMEdGraphSchema.generated.h"
+
+class URigVMBlueprint;
+class URigVMEdGraph;
+class URigVMEdGraphNode;
+
+/** Extra operations that can be performed on pin connection */
+enum class ECanCreateConnectionResponse_Extended
+{
+	None,
+
+	BreakChildren,
+
+	BreakParent,
+};
+
+/** Struct used to extend the response to a connection request to include breaking parents/children */
+struct FRigVMPinConnectionResponse
+{
+	FRigVMPinConnectionResponse(const ECanCreateConnectionResponse InResponse, FText InMessage, ECanCreateConnectionResponse_Extended InExtendedResponse = ECanCreateConnectionResponse_Extended::None)
+		: Response(InResponse, MoveTemp(InMessage))
+		, ExtendedResponse(InExtendedResponse)
+	{
+	}
+
+	friend bool operator==(const FRigVMPinConnectionResponse& A, const FRigVMPinConnectionResponse& B)
+	{
+		return (A.Response == B.Response) && (A.ExtendedResponse == B.ExtendedResponse);
+	}	
+
+	FPinConnectionResponse Response;
+	ECanCreateConnectionResponse_Extended ExtendedResponse;
+};
+
+/////////////////////////////////////////////////////
+// FRigVMLocalVariableNameValidator
+class RIGVMDEVELOPER_API FRigVMLocalVariableNameValidator : public FStringSetNameValidator
+{
+
+public:
+	FRigVMLocalVariableNameValidator(const class UBlueprint* Blueprint, const URigVMGraph* Graph, FName InExistingName = NAME_None);
+
+	// Begin FNameValidatorInterface
+	virtual EValidatorResult IsValid(const FString& Name, bool bOriginal) override;
+	virtual EValidatorResult IsValid(const FName& Name, bool bOriginal) override;
+	// End FNameValidatorInterface
+};
+
+/////////////////////////////////////////////////////
+// FRigVMNameValidator
+class RIGVMDEVELOPER_API FRigVMNameValidator : public FStringSetNameValidator
+{
+
+public:
+	FRigVMNameValidator(const class UBlueprint* Blueprint, const UStruct* ValidationScope, FName InExistingName = NAME_None);
+
+	// Begin FNameValidatorInterface
+	virtual EValidatorResult IsValid(const FString& Name, bool bOriginal) override;
+	virtual EValidatorResult IsValid(const FName& Name, bool bOriginal) override;
+	// End FNameValidatorInterface
+};
+
+USTRUCT()
+struct RIGVMDEVELOPER_API FRigVMEdGraphSchemaAction_LocalVar : public FEdGraphSchemaAction_BlueprintVariableBase
+{
+	GENERATED_BODY()
+
+	
+public:
+
+	// Simple type info
+	static FName StaticGetTypeId() {static FName Type("FRigVMEdGraphSchemaAction_LocalVar"); return Type;}
+	virtual FName GetTypeId() const override { return StaticGetTypeId(); } 
+
+	FRigVMEdGraphSchemaAction_LocalVar()
+		: FEdGraphSchemaAction_BlueprintVariableBase()
+	{}
+
+	FRigVMEdGraphSchemaAction_LocalVar(FText InNodeCategory, FText InMenuDesc, FText InToolTip, const int32 InGrouping, const int32 InSectionID)
+		: FEdGraphSchemaAction_BlueprintVariableBase(MoveTemp(InNodeCategory), MoveTemp(InMenuDesc), MoveTemp(InToolTip), InGrouping, InSectionID)
+	{}
+
+	virtual bool IsA(const FName& InType) const override
+	{
+		return InType == GetTypeId() || InType == FEdGraphSchemaAction_BlueprintVariableBase::StaticGetTypeId();
+	}
+
+	virtual FEdGraphPinType GetPinType() const override;
+
+	virtual void ChangeVariableType(const FEdGraphPinType& NewPinType) override;
+
+	virtual void RenameVariable(const FName& NewName) override;
+
+	virtual bool IsValidName(const FName& NewName, FText& OutErrorMessage) const override;
+
+	virtual void DeleteVariable() override;
+
+	virtual bool IsVariableUsed() override;
+};
+
+USTRUCT()
+struct RIGVMDEVELOPER_API FRigVMEdGraphSchemaAction_PromoteToVariable : public FEdGraphSchemaAction
+{
+	GENERATED_BODY()
+
+public:
+
+	// Simple type info
+	static FName StaticGetTypeId() {static FName Type("FRigVMEdGraphSchemaAction_PromoteToVariable"); return Type;}
+	virtual FName GetTypeId() const override { return StaticGetTypeId(); } 
+
+	FRigVMEdGraphSchemaAction_PromoteToVariable()
+		: FEdGraphSchemaAction()
+		, EdGraphPin(nullptr)
+		, bLocalVariable(false)
+	{}
+
+	FRigVMEdGraphSchemaAction_PromoteToVariable(UEdGraphPin* InEdGraphPin, bool InLocalVariable);
+
+	virtual bool IsA(const FName& InType) const override
+	{
+		return InType == GetTypeId();
+	}
+
+	virtual UEdGraphNode* PerformAction(UEdGraph* ParentGraph, UEdGraphPin* FromPin, const FVector2D Location, bool bSelectNewNode) override;
+
+private:
+
+	UEdGraphPin* EdGraphPin;
+	bool bLocalVariable;
+};
+
+USTRUCT()
+struct RIGVMDEVELOPER_API FRigVMEdGraphSchemaAction_PromoteToExposedPin : public FEdGraphSchemaAction
+{
+	GENERATED_BODY()
+
+public:
+
+	// Simple type info
+	static FName StaticGetTypeId() {static FName Type("FRigVMEdGraphSchemaAction_PromoteToExposedPin"); return Type;}
+	virtual FName GetTypeId() const override { return StaticGetTypeId(); } 
+
+	FRigVMEdGraphSchemaAction_PromoteToExposedPin()
+		: FEdGraphSchemaAction()
+		, EdGraphPin(nullptr)
+	{}
+
+	FRigVMEdGraphSchemaAction_PromoteToExposedPin(UEdGraphPin* InEdGraphPin);
+
+	virtual bool IsA(const FName& InType) const override
+	{
+		return InType == GetTypeId();
+	}
+
+	virtual UEdGraphNode* PerformAction(UEdGraph* ParentGraph, UEdGraphPin* FromPin, const FVector2D Location, bool bSelectNewNode) override;
+
+private:
+
+	UEdGraphPin* EdGraphPin;
+};
+
+USTRUCT()
+struct RIGVMDEVELOPER_API FRigVMEdGraphSchemaAction_Event : public FEdGraphSchemaAction
+{
+	GENERATED_BODY()
+
+public:
+
+	// Simple type info
+	static FName StaticGetTypeId() {static FName Type("FRigVMEdGraphSchemaAction_Event"); return Type;}
+	virtual FName GetTypeId() const override { return StaticGetTypeId(); } 
+
+	FRigVMEdGraphSchemaAction_Event()
+		: FEdGraphSchemaAction()
+	{}
+
+	FRigVMEdGraphSchemaAction_Event(const FName& InEventName, const FString& InNodePath, const FText& InNodeCategory);
+
+	virtual bool IsA(const FName& InType) const override
+	{
+		return InType == GetTypeId() || Super::IsA(InType);
+	}
+	virtual bool IsParentable() const override { return true; }
+	virtual FReply OnDoubleClick(UBlueprint* InBlueprint) override;
+	virtual bool CanBeRenamed() const override { return false; }
+	virtual FSlateBrush const* GetPaletteIcon() const override;
+
+private:
+
+	FString NodePath;
+};
+
+/** DragDropAction class for drag and dropping an item from the My Blueprints tree (e.g., variable or function) */
+class RIGVMDEVELOPER_API FRigVMFunctionDragDropAction : public FGraphSchemaActionDragDropAction
+{
+public:
+
+	DRAG_DROP_OPERATOR_TYPE(FRigVMFunctionDragDropAction, FGraphSchemaActionDragDropAction)
+
+	// FGraphEditorDragDropAction interface
+	virtual FReply DroppedOnPanel(const TSharedRef< class SWidget >& Panel, FVector2D ScreenPosition, FVector2D GraphPosition, UEdGraph& Graph) override;
+	virtual FReply DroppedOnPin(FVector2D ScreenPosition, FVector2D GraphPosition) override;
+	virtual FReply DroppedOnAction(TSharedRef<FEdGraphSchemaAction> Action) override;
+	virtual FReply DroppedOnCategory(FText Category) override;
+	virtual void HoverTargetChanged() override;
+	// End of FGraphEditorDragDropAction
+
+	/** Set if operation is modified by alt */
+	void SetAltDrag(bool InIsAltDrag) { bAltDrag = InIsAltDrag; }
+
+	/** Set if operation is modified by the ctrl key */
+	void SetCtrlDrag(bool InIsCtrlDrag) { bControlDrag = InIsCtrlDrag; }
+
+protected:
+
+	/** Constructor */
+	FRigVMFunctionDragDropAction();
+
+	static TSharedRef<FRigVMFunctionDragDropAction> New(TSharedPtr<FEdGraphSchemaAction> InAction, URigVMBlueprint* InRigBlueprint, URigVMEdGraph* InRigGraph);
+
+protected:
+
+	URigVMBlueprint* SourceRigBlueprint;
+	URigVMEdGraph* SourceRigGraph;
+	bool bControlDrag;
+	bool bAltDrag;
+
+	friend class URigVMEdGraphSchema;
+};
+
+UCLASS()
+class RIGVMDEVELOPER_API URigVMEdGraphSchema : public UEdGraphSchema
+{
+	GENERATED_BODY()
+
+public:
+	/** Name constants */
+	static const FName GraphName_RigVM;
+
+public:
+	URigVMEdGraphSchema();
+
+	// UEdGraphSchema interface
+	virtual void GetGraphContextActions(FGraphContextMenuBuilder& ContextMenuBuilder) const override;
+	virtual void GetContextMenuActions(class UToolMenu* Menu, class UGraphNodeContextMenuContext* Context) const override;
+	virtual bool TryCreateConnection(UEdGraphPin* PinA, UEdGraphPin* PinB) const override;
+	virtual const FPinConnectionResponse CanCreateConnection(const UEdGraphPin* A, const UEdGraphPin* B) const override;
+	virtual FLinearColor GetPinTypeColor(const FEdGraphPinType& PinType) const override;
+	virtual bool SupportsPinType(TWeakPtr<const FEdGraphSchemaAction> SchemaAction, const FEdGraphPinType& PinType) const override;
+	virtual bool SupportsPinTypeContainer(TWeakPtr<const FEdGraphSchemaAction> SchemaAction, const FEdGraphPinType& PinType, const EPinContainerType& ContainerType) const override;
+	virtual void BreakPinLinks(UEdGraphPin& TargetPin, bool bSendsNodeNotifcation) const override;
+	virtual void BreakSinglePinLink(UEdGraphPin* SourcePin, UEdGraphPin* TargetPin) const override;
+	virtual bool CanGraphBeDropped(TSharedPtr<FEdGraphSchemaAction> InAction) const override;
+	virtual FString GetFindReferenceSearchTerm(const FEdGraphSchemaAction* InGraphAction) const override;
+	virtual FReply BeginGraphDragAction(TSharedPtr<FEdGraphSchemaAction> InAction, const FPointerEvent& MouseEvent = FPointerEvent() ) const override;
+	virtual class FConnectionDrawingPolicy* CreateConnectionDrawingPolicy(int32 InBackLayerID, int32 InFrontLayerID, float InZoomFactor, const FSlateRect& InClippingRect, class FSlateWindowElementList& InDrawElements, class UEdGraph* InGraphObj) const override;
+	virtual bool ShouldHidePinDefaultValue(UEdGraphPin* Pin) const override;
+	virtual void TrySetDefaultValue(UEdGraphPin& InPin, const FString& InNewDefaultValue, bool bMarkAsModified = true) const override;
+	virtual void TrySetDefaultObject(UEdGraphPin& InPin, UObject* InNewDefaultObject, bool bMarkAsModified) const override;
+	virtual void TrySetDefaultText(UEdGraphPin& InPin, const FText& InNewDefaultText, bool bMarkAsModified) const override;
+	virtual bool ShouldAlwaysPurgeOnModification() const override { return false; }
+	virtual bool ArePinsCompatible(const UEdGraphPin* PinA, const UEdGraphPin* PinB, const UClass* CallingContext, bool bIgnoreArray /*= false*/) const override;
+	virtual bool DoesSupportPinWatching() const	override { return true; }
+	virtual bool IsPinBeingWatched(UEdGraphPin const* Pin) const override;
+	virtual void ClearPinWatch(UEdGraphPin const* Pin) const override;
+	virtual void OnPinConnectionDoubleCicked(UEdGraphPin* PinA, UEdGraphPin* PinB, const FVector2D& GraphPosition) const override;
+	virtual bool MarkBlueprintDirtyFromNewNode(UBlueprint* InBlueprint, UEdGraphNode* InEdGraphNode) const override;
+	virtual bool SafeDeleteNodeFromGraph(UEdGraph* Graph, UEdGraphNode* Node) const override;
+	virtual bool CanVariableBeDropped(UEdGraph* InGraph, FProperty* InVariableToDrop) const override;
+	virtual bool RequestVariableDropOnPanel(UEdGraph* InGraph, FProperty* InVariableToDrop, const FVector2D& InDropPosition, const FVector2D& InScreenPosition) override;
+	virtual bool RequestVariableDropOnPin(UEdGraph* InGraph, FProperty* InVariableToDrop, UEdGraphPin* InPin, const FVector2D& InDropPosition, const FVector2D& InScreenPosition) override;
+	virtual bool IsStructEditable(UStruct* InStruct) const;
+	virtual void SetNodePosition(UEdGraphNode* Node, const FVector2D& Position) const override;
+	void SetNodePosition(UEdGraphNode* Node, const FVector2D& Position, bool bSetupUndo) const;
+	virtual void GetGraphDisplayInformation(const UEdGraph& Graph, /*out*/ FGraphDisplayInfo& DisplayInfo) const override;
+	virtual bool GetLocalVariables(const UEdGraph* InGraph, TArray<FBPVariableDescription>& OutLocalVariables) const override;
+	virtual TSharedPtr<FEdGraphSchemaAction> MakeActionFromVariableDescription(const UEdGraph* InEdGraph, const FBPVariableDescription& Variable) const override;
+	virtual FText GetGraphCategory(const UEdGraph* InGraph) const override;
+	virtual EGraphType GetGraphType(const UEdGraph* TestEdGraph) const override;
+	virtual FReply TrySetGraphCategory(const UEdGraph* InGraph, const FText& InCategory) override;
+	virtual bool TryDeleteGraph(UEdGraph* GraphToDelete) const override;
+	virtual bool TryRenameGraph(UEdGraph* GraphToRename, const FName& InNewName) const override;
+	virtual bool TryToGetChildEvents(const UEdGraph* Graph, const int32 SectionId, TArray<TSharedPtr<FEdGraphSchemaAction>>& Actions, const FText& ParentCategory) const override;
+	virtual bool CanDuplicateGraph(UEdGraph* InSourceGraph) const { return false; }
+	virtual UEdGraphPin* DropPinOnNode(UEdGraphNode* InTargetNode, const FName& InSourcePinName, const FEdGraphPinType& InSourcePinType, EEdGraphPinDirection InSourcePinDirection) const override;
+	virtual bool SupportsDropPinOnNode(UEdGraphNode* InTargetNode, const FEdGraphPinType& InSourcePinType, EEdGraphPinDirection InSourcePinDirection, FText& OutErrorMessage) const override;
+	virtual void SetPinBeingDroppedOnNode(UEdGraphPin* InSourcePin) const override { PinBeingDropped = InSourcePin; }
+	virtual void InsertAdditionalActions(TArray<UBlueprint*> InBlueprints, TArray<UEdGraph*> InGraphs, TArray<UEdGraphPin*> InPins, FGraphActionListBuilderBase& OutAllActions) const override;
+	virtual TSharedPtr<INameValidatorInterface> GetNameValidator(const UBlueprint* BlueprintObj, const FName& OriginalName, const UStruct* ValidationScope, const FName& ActionTypeId) const override;
+
+	/** Returns true if the schema supports the script type */
+	bool SupportsPinType(const UScriptStruct* ScriptStruct) const;
+
+	/** Create a graph node for a rig */
+	URigVMEdGraphNode* CreateGraphNode(URigVMEdGraph* InGraph, const FName& InPropertyName) const;
+
+	/** Helper function to rename a node */
+	void RenameNode(URigVMEdGraphNode* Node, const FName& InNewNodeName) const;
+
+	/** Helper function to recursively reset the pin defaults */
+	virtual void ResetPinDefaultsRecursive(UEdGraphPin* InPin) const;
+
+	/** Returns all of the applicable pin types for variables within a control rig */
+	virtual void GetVariablePinTypes(TArray<FEdGraphPinType>& PinTypes) const;
+
+	void StartGraphNodeInteraction(UEdGraphNode* InNode) const;
+	void EndGraphNodeInteraction(UEdGraphNode* InNode) const;
+	static TArray<UEdGraphNode*> GetNodesToMoveForNode(UEdGraphNode* InNode);
+	FVector2D GetNodePositionAtStartOfInteraction(const UEdGraphNode* InNode) const;
+
+	bool AutowireNewNode(URigVMEdGraphNode* NewNode, UEdGraphPin* FromPin) const;
+
+	void HandleModifiedEvent(ERigVMGraphNotifType InNotifType, URigVMGraph* InGraph, UObject* InSubject);
+
+	// Allow derived classes to spawn derived node classes
+	virtual TSubclassOf<URigVMEdGraphNode> GetGraphNodeClass() const { return URigVMEdGraphNode::StaticClass(); }
+
+	virtual bool IsRigVMDefaultEvent(const FName& InEventName) const;
+
+private:
+
+	const UEdGraphPin* LastPinForCompatibleCheck = nullptr;
+	bool bLastPinWasInput;
+	mutable UEdGraphPin* PinBeingDropped = nullptr;
+
+#if WITH_EDITOR
+	mutable TArray<UEdGraphNode*> NodesBeingInteracted;
+	mutable TMap<FName, FVector2D> NodePositionsDuringStart;
+#endif
+
+	friend class UControlRigIfNodeSpawner;
+	friend class UControlRigSelectNodeSpawner;
+	friend class UControlRigUnitNodeSpawner;
+	friend class UControlRigArrayNodeSpawner;
+};
diff --git a/Engine/Plugins/Runtime/RigVM/Source/RigVMDeveloper/Public/RigVMBlueprint.h b/Engine/Plugins/Runtime/RigVM/Source/RigVMDeveloper/Public/RigVMBlueprint.h
new file mode 100644
--- /dev/null
+++ b/Engine/Plugins/Runtime/RigVM/Source/RigVMDeveloper/Public/RigVMBlueprint.h
@@ -0,0 +1,633 @@
+// Copyright Epic Games, Inc. All Rights Reserved.
+
+#pragma once
+
+#include "CoreMinimal.h"
+#include "Engine/Blueprint.h"
+#include "RigVMCore/RigVM.h"
+#include "RigVMHost.h"
+#include "RigVMModel/RigVMClient.h"
+#include "RigVMCompiler/RigVMCompiler.h"
+#include "RigVMCore/RigVMGraphFunctionDefinition.h"
+#include "EdGraph/RigVMEdGraph.h"
+#include "EdGraph/RigVMEdGraphSchema.h"
+
+#include "RigVMBlueprint.generated.h"
+
+class URigVMBlueprintGeneratedClass;
+struct FEndLoadPackageContext;
+
+DECLARE_EVENT_TwoParams(URigVMBlueprint, FOnVMCompiledEvent, UObject*, URigVM*);
+DECLARE_EVENT_OneParam(URigVMBlueprint, FOnRefreshEditorEvent, URigVMBlueprint*);
+DECLARE_EVENT_FourParams(URigVMBlueprint, FOnVariableDroppedEvent, UObject*, FProperty*, const FVector2D&, const FVector2D&);
+DECLARE_EVENT_OneParam(URigVMBlueprint, FOnExternalVariablesChanged, const TArray<FRigVMExternalVariable>&);
+DECLARE_EVENT_TwoParams(URigVMBlueprint, FOnNodeDoubleClicked, URigVMBlueprint*, URigVMNode*);
+DECLARE_EVENT_OneParam(URigVMBlueprint, FOnGraphImported, UEdGraph*);
+DECLARE_EVENT_OneParam(URigVMBlueprint, FOnPostEditChangeChainProperty, FPropertyChangedChainEvent&);
+DECLARE_EVENT_ThreeParams(URigVMBlueprint, FOnLocalizeFunctionDialogRequested, FRigVMGraphFunctionIdentifier&, URigVMBlueprint*, bool);
+DECLARE_EVENT_ThreeParams(URigVMBlueprint, FOnReportCompilerMessage, EMessageSeverity::Type, UObject*, const FString&);
+DECLARE_DELEGATE_RetVal_FourParams(FRigVMController_BulkEditResult, FRigVMOnBulkEditDialogRequestedDelegate, URigVMBlueprint*, URigVMController*, URigVMLibraryNode*, ERigVMControllerBulkEditType);
+DECLARE_DELEGATE_RetVal_OneParam(bool, FRigVMOnBreakLinksDialogRequestedDelegate, TArray<URigVMLink*>);
+DECLARE_DELEGATE_RetVal_OneParam(TRigVMTypeIndex, FRigVMOnPinTypeSelectionRequestedDelegate, const TArray<TRigVMTypeIndex>&);
+DECLARE_EVENT(URigVMBlueprint, FOnBreakpointAdded);
+DECLARE_EVENT_OneParam(URigVMBlueprint, FOnRequestInspectObject, const TArray<UObject*>& );
+DECLARE_DELEGATE_RetVal(URigVMGraph*, FRigVMBlueprintGetFocusedGraph);
+
+USTRUCT()
+struct RIGVMDEVELOPER_API FRigVMEdGraphDisplaySettings
+{
+	GENERATED_BODY();
+
+	FRigVMEdGraphDisplaySettings()
+		: bShowNodeInstructionIndex(false)
+		, bShowNodeRunCounts(false)
+		, NodeRunLowerBound(1)
+		, NodeRunLimit(256)
+		, MinMicroSeconds(0.0)
+		, MaxMicroSeconds(1.0)
+		, TotalMicroSeconds(0.0)
+		, bAutoDetermineRange(true)
+		, LastMinMicroSeconds(0.0)
+		, LastMaxMicroSeconds(1.0)
+		, MinDurationColor(FLinearColor::Green)
+		, MaxDurationColor(FLinearColor::Red)
+	{
+	}
+
+	// When enabled shows the first node instruction index
+	// matching the execution stack window.
+	UPROPERTY(EditAnywhere, Category = "Graph Display Settings")
+	bool bShowNodeInstructionIndex;
+
+	// When enabled shows the node counts both in the graph view as
+	// we as in the execution stack window.
+	// The number on each node represents how often the node has been run.
+	// Keep in mind when looking at nodes in a function the count
+	// represents the sum of all counts for each node based on all
+	// references of the function currently running.
+	UPROPERTY(EditAnywhere, Category = "Graph Display Settings")
+	bool bShowNodeRunCounts;
+
+	// A lower limit for counts for nodes used for debugging.
+	// Any node lower than this count won't show the run count.
+	UPROPERTY(EditAnywhere, Category = "Graph Display Settings")
+	int32 NodeRunLowerBound;
+
+	// A upper limit for counts for nodes used for debugging.
+	// If a node reaches this count a warning will be issued for the
+	// node and displayed both in the execution stack as well as in the
+	// graph. Setting this to <= 1 disables the warning.
+	// Note: The count limit doesn't apply to functions / collapse nodes.
+	UPROPERTY(EditAnywhere, Category = "Graph Display Settings")
+	int32 NodeRunLimit;
+
+	// The duration in microseconds of the fastest instruction / node
+	UPROPERTY(EditAnywhere, Category = "Graph Display Settings", transient, meta = (EditCondition = "!bAutoDetermineRange"))
+	double MinMicroSeconds;
+
+	// The duration in microseconds of the slowest instruction / node
+	UPROPERTY(EditAnywhere, Category = "Graph Display Settings", transient, meta = (EditCondition = "!bAutoDetermineRange"))
+	double MaxMicroSeconds;
+
+	// The total duration of the last execution of the rig
+	UPROPERTY(VisibleAnywhere, Category = "Graph Display Settings", transient)
+	double TotalMicroSeconds;
+
+	UPROPERTY(EditAnywhere, Category = "Graph Display Settings")
+	bool bAutoDetermineRange;
+
+	UPROPERTY(transient)
+	double LastMinMicroSeconds;
+
+	UPROPERTY(transient)
+	double LastMaxMicroSeconds;
+
+	// The color of the fastest instruction / node
+	UPROPERTY(EditAnywhere, Category = "Graph Display Settings")
+	FLinearColor MinDurationColor;
+
+	// The color of the slowest instruction / node
+	UPROPERTY(EditAnywhere, Category = "Graph Display Settings")
+	FLinearColor MaxDurationColor;
+};
+
+enum class ERigVMBlueprintLoadType : uint8
+{
+	PostLoad,
+	CheckUserDefinedStructs
+};
+
+UCLASS(BlueprintType, meta=(IgnoreClassThumbnail))
+class RIGVMDEVELOPER_API URigVMBlueprint : public UBlueprint, public IRigVMClientHost
+{
+	GENERATED_UCLASS_BODY()
+
+public:
+	URigVMBlueprint();
+
+	void InitializeModelIfRequired(bool bRecompileVM = true);
+
+	/** Get the (full) generated class for this rigvm blueprint */
+	URigVMBlueprintGeneratedClass* GetRigVMBlueprintGeneratedClass() const;
+
+	/** Get the (skeleton) generated class for this rigvm blueprint */
+	URigVMBlueprintGeneratedClass* GetRigVMBlueprintSkeletonClass() const;
+
+	/** Returns the expected schema class to use for this blueprint */
+	virtual UClass* GetRigVMSchemaClass() const { return URigVMSchema::StaticClass(); }
+
+	/** Returns the expected execute context struct to use for this blueprint */
+	virtual UScriptStruct* GetRigVMExecuteContextStruct() const { return FRigVMExecuteContext::StaticStruct(); }
+
+	/** Returns the expected ed graph schema class to use for this blueprint */
+	virtual UClass* GetRigVMEdGraphSchemaClass() const { return URigVMEdGraphSchema::StaticClass(); }
+
+	/** Returns the expected ed graph class to use for this blueprint */
+	virtual UClass* GetRigVMEdGraphClass() const { return URigVMEdGraph::StaticClass(); }
+
+	virtual void Serialize(FArchive& Ar) override;
+
+#if WITH_EDITOR
+
+	// UBlueprint interface
+	virtual UClass* GetBlueprintClass() const override;
+	virtual UClass* RegenerateClass(UClass* ClassToRegenerate, UObject* PreviousCDO) override;
+	virtual bool SupportedByDefaultBlueprintFactory() const override { return false; }
+	virtual bool IsValidForBytecodeOnlyRecompile() const override { return false; }
+	virtual void LoadModulesRequiredForCompilation() override;
+	virtual void GetTypeActions(FBlueprintActionDatabaseRegistrar& ActionRegistrar) const override;
+	virtual void GetInstanceActions(FBlueprintActionDatabaseRegistrar& ActionRegistrar) const override;
+	virtual void SetObjectBeingDebugged(UObject* NewObject) override;
+	PRAGMA_DISABLE_DEPRECATION_WARNINGS // Suppress compiler warning on override of deprecated function
+	UE_DEPRECATED(5.0, "Use version that takes FObjectPreSaveContext instead.")
+	virtual void PreSave(const class ITargetPlatform* TargetPlatform) override;
+	PRAGMA_ENABLE_DEPRECATION_WARNINGS
+	virtual void PreSave(FObjectPreSaveContext ObjectSaveContext) override;
+	virtual void PostLoad() override;
+#if WITH_EDITORONLY_DATA
+	static void DeclareConstructClasses(TArray<FTopLevelAssetPath>& OutConstructClasses, const UClass* SpecificSubclass);
+#endif
+	virtual bool IsPostLoadThreadSafe() const override { return false; }
+	virtual void PostTransacted(const FTransactionObjectEvent& TransactionEvent) override;
+	virtual void ReplaceDeprecatedNodes() override;
+	virtual void PreDuplicate(FObjectDuplicationParameters& DupParams) override;
+	virtual void PostDuplicate(bool bDuplicateForPIE) override;
+	virtual void GetAssetRegistryTags(TArray<FAssetRegistryTag>& OutTags) const override;
+
+	virtual bool SupportsGlobalVariables() const override { return true; }
+	virtual bool SupportsLocalVariables() const override { return true; }
+	virtual bool SupportsFunctions() const override { return true; }
+	virtual bool SupportsMacros() const override { return false; }
+	virtual bool SupportsDelegates() const override { return false; }
+	virtual bool SupportsEventGraphs() const override { return true; }
+	virtual bool SupportsAnimLayers() const override { return false; }
+	virtual bool ExportGraphToText(UEdGraph* InEdGraph, FString& OutText) override;
+	virtual bool TryImportGraphFromText(const FString& InClipboardText, UEdGraph** OutGraphPtr = nullptr) override;
+	virtual bool CanImportGraphFromText(const FString& InClipboardText) override;
+	virtual bool RequiresForceLoadMembers(UObject* InObject) const override;
+
+	// UObject interface
+	virtual void PostEditChangeChainProperty(struct FPropertyChangedChainEvent& PropertyChangedEvent) override;
+	virtual void PostRename(UObject* OldOuter, const FName OldName) override;
+	/** Called during cooking. Must return all objects that will be Preload()ed when this is serialized at load time. */
+	void GetPreloadDependencies(TArray<UObject*>& OutDeps) override;
+
+	// IRigVMClientHost interface
+	virtual FRigVMClient* GetRigVMClient() override;
+	virtual const FRigVMClient* GetRigVMClient() const override;
+	virtual IRigVMGraphFunctionHost* GetRigVMGraphFunctionHost() override;
+	virtual const IRigVMGraphFunctionHost* GetRigVMGraphFunctionHost() const override;
+	virtual UObject* GetEditorObjectForRigVMGraph(URigVMGraph* InVMGraph) const override;
+	virtual void HandleRigVMGraphAdded(const FRigVMClient* InClient, const FString& InNodePath) override;
+	virtual void HandleRigVMGraphRemoved(const FRigVMClient* InClient, const FString& InNodePath) override;
+	virtual void HandleRigVMGraphRenamed(const FRigVMClient* InClient, const FString& InOldNodePath, const FString& InNewNodePath) override;
+	virtual void HandleConfigureRigVMController(const FRigVMClient* InClient, URigVMController* InControllerToConfigure) override;
+
+	FOnRequestInspectObject& OnRequestInspectObject() { return OnRequestInspectObjectEvent; }
+	void RequestInspectObject(const TArray<UObject*>& InObjects) { OnRequestInspectObjectEvent.Broadcast(InObjects); }
+
+#endif	// #if WITH_EDITOR
+
+	FRigVMBlueprintGetFocusedGraph& OnGetFocusedGraph() { return OnGetFocusedGraphDelegate; }
+
+	virtual bool ShouldBeMarkedDirtyUponTransaction() const override { return false; }
+
+	UFUNCTION(BlueprintCallable, Category = "RigVM Blueprint")
+	void RecompileVM();
+
+	UFUNCTION(BlueprintCallable, Category = "RigVM Blueprint")
+	void RecompileVMIfRequired();
+	
+	UFUNCTION(BlueprintCallable, Category = "RigVM Blueprint")
+	void RequestAutoVMRecompilation();
+
+	UFUNCTION(BlueprintCallable, Category = "RigVM Blueprint")
+	void SetAutoVMRecompile(bool bAutoRecompile) { bAutoRecompileVM = bAutoRecompile; }
+
+	void IncrementVMRecompileBracket();
+	void DecrementVMRecompileBracket();
+
+	// this is needed since even after load
+	// model data can change while the Control Rig BP is not opened
+	// for example, if a user defined struct changed after BP load,
+	// any pin that references the struct needs to be regenerated
+	void RefreshAllModels(ERigVMBlueprintLoadType InLoadType = ERigVMBlueprintLoadType::PostLoad);
+
+	// RigVMRegistry changes can be triggered when new user defined types(structs/enums) are added/removed
+	// in which case we have to refresh the model
+	void OnRigVMRegistryChanged();
+	
+	UFUNCTION(BlueprintCallable, Category = "RigVM Blueprint")
+	void RequestRigVMInit();
+
+	UFUNCTION(BlueprintCallable, Category = "RigVM Blueprint")
+	URigVMGraph* GetModel(const UEdGraph* InEdGraph = nullptr) const;
+	URigVMGraph* GetModel(const FString& InNodePath) const;
+
+	UFUNCTION(BlueprintCallable, Category = "RigVM Blueprint")
+	URigVMGraph* GetDefaultModel() const;
+
+	UFUNCTION(BlueprintCallable, Category = "RigVM Blueprint")
+	TArray<URigVMGraph*> GetAllModels() const;
+
+	UFUNCTION(BlueprintCallable, Category = "RigVM Blueprint")
+	URigVMFunctionLibrary* GetLocalFunctionLibrary() const;
+
+	UFUNCTION(BlueprintCallable, Category = "RigVM Blueprint")
+	URigVMGraph* AddModel(FString InName = TEXT("Rig Graph"), bool bSetupUndoRedo = true, bool bPrintPythonCommand = true);
+
+	UFUNCTION(BlueprintCallable, Category = "RigVM Blueprint")
+	bool RemoveModel(FString InName = TEXT("Rig Graph"), bool bSetupUndoRedo = true, bool bPrintPythonCommand = true);
+
+	UFUNCTION(BlueprintCallable, Category = "RigVM Blueprint")
+	URigVMGraph* GetFocusedModel() const;
+
+	UFUNCTION(BlueprintCallable, Category = "RigVM Blueprint")
+	URigVMController* GetController(const URigVMGraph* InGraph = nullptr) const;
+
+	UFUNCTION(BlueprintCallable, Category = "RigVM Blueprint")
+	URigVMController* GetControllerByName(const FString InGraphName = TEXT("")) const;
+
+	UFUNCTION(BlueprintCallable, Category = "RigVM Blueprint")
+	URigVMController* GetOrCreateController(URigVMGraph* InGraph = nullptr);
+
+	URigVMController* GetController(const UEdGraph* InEdGraph) const;
+	URigVMController* GetOrCreateController(const UEdGraph* InGraph);
+
+	UFUNCTION(BlueprintCallable, Category = "RigVM Blueprint")
+	TArray<FString> GeneratePythonCommands(const FString InNewBlueprintName);
+
+	URigVMGraph* GetTemplateModel(bool bIsFunctionLibrary = false);
+	URigVMController* GetTemplateController(bool bIsFunctionLibrary = false);
+
+#if WITH_EDITOR
+	UEdGraph* GetEdGraph(URigVMGraph* InModel) const;
+	UEdGraph* GetEdGraph(const FString& InNodePath) const;
+#endif
+
+#if WITH_EDITORONLY_DATA
+	UPROPERTY()
+	TObjectPtr<URigVMEdGraph> FunctionLibraryEdGraph;
+#endif
+
+	bool IsFunctionPublic(const FName& InFunctionName) const;
+	void MarkFunctionPublic(const FName& InFunctionName, bool bIsPublic = true);
+
+	// Returns a list of dependencies of this blueprint.
+	// Dependencies are blueprints that contain functions used in this blueprint
+	TArray<URigVMBlueprint*> GetDependencies(bool bRecursive = false) const;
+
+	// Returns a list of dependents as unresolved soft object pointers.
+	// A dependent is a blueprint which uses a function defined in this blueprint.
+	// This function is not recursive, since it avoids opening the asset.
+	// Use GetDependentBlueprints as an alternative.
+	TArray<FAssetData> GetDependentAssets() const;
+
+	// Returns a list of dependents as resolved blueprints.
+	// A dependent is a blueprint which uses a function defined in this blueprint.
+	// If bOnlyLoaded is false, this function loads the dependent assets and can introduce a large cost
+	// depending on the size / count of assets in the project.
+	TArray<URigVMBlueprint*> GetDependentBlueprints(bool bRecursive = false, bool bOnlyLoaded = false) const;
+
+	UPROPERTY(EditAnywhere, Category = "User Interface")
+	FRigVMEdGraphDisplaySettings RigGraphDisplaySettings;
+
+	UPROPERTY(EditAnywhere, Category = "VM")
+	FRigVMRuntimeSettings VMRuntimeSettings;
+
+	UPROPERTY(EditAnywhere, BlueprintReadWrite, Category = "VM")
+	FRigVMCompileSettings VMCompileSettings;
+
+protected:
+
+#if WITH_EDITORONLY_DATA
+	UPROPERTY()
+	TObjectPtr<URigVMGraph> Model_DEPRECATED;
+
+	UPROPERTY()
+	TObjectPtr<URigVMFunctionLibrary> FunctionLibrary_DEPRECATED;
+#endif
+
+	UPROPERTY()
+	FRigVMClient RigVMClient;
+
+#if WITH_EDITORONLY_DATA
+
+	UPROPERTY()
+	bool ReferencedObjectPathsStored;
+
+	UPROPERTY()
+	TArray<FSoftObjectPath> ReferencedObjectPaths;
+
+#endif
+
+	/** Asset searchable information about exposed public functions on this rig */
+	UPROPERTY(AssetRegistrySearchable)
+	TArray<FRigVMGraphFunctionHeader> PublicGraphFunctions;
+
+	/** Asset searchable information function references in this rig */
+	UPROPERTY(AssetRegistrySearchable)
+	TArray<FRigVMReferenceNodeData> FunctionReferenceNodeData;
+
+#if WITH_EDITORONLY_DATA
+
+	UPROPERTY(transient, DuplicateTransient)
+	TObjectPtr<URigVMGraph> TemplateModel;
+
+	UPROPERTY(transient, DuplicateTransient)
+	TObjectPtr<URigVMController> TemplateController;
+
+	mutable TArray<FAssetRegistryTag> CachedAssetTags;
+
+#endif
+
+public:
+
+	UPROPERTY(transient, DuplicateTransient)
+	TMap<FString, FRigVMOperand> PinToOperandMap;
+
+	bool bSuspendModelNotificationsForSelf;
+	bool bSuspendModelNotificationsForOthers;
+	bool bSuspendAllNotifications;
+
+	virtual void SetupPinRedirectorsForBackwardsCompatibility() const {};
+	void RebuildGraphFromModel();
+
+	FRigVMGraphModifiedEvent& OnModified();
+	FOnVMCompiledEvent& OnVMCompiled();
+
+	UFUNCTION(BlueprintCallable, Category = "VM")
+	UClass* GetRigVMHostClass();
+
+	UFUNCTION(BlueprintCallable, Category = "VM")
+	URigVMHost* CreateRigVMHost();
+
+	UFUNCTION(BlueprintCallable, Category = "VM")
+	virtual TArray<UStruct*> GetAvailableRigVMStructs() const;
+
+#if WITH_EDITOR
+	UFUNCTION(BlueprintCallable, Category = "Variables")
+	TArray<FRigVMGraphVariableDescription> GetMemberVariables() const;
+	
+	UFUNCTION(BlueprintCallable, Category = "Variables")
+	FName AddMemberVariable(const FName& InName, const FString& InCPPType, bool bIsPublic = false, bool bIsReadOnly = false, FString InDefaultValue = TEXT(""));
+
+	UFUNCTION(BlueprintCallable, Category = "Variables")
+	bool RemoveMemberVariable(const FName& InName);
+
+	UFUNCTION(BlueprintCallable, Category = "Variables")
+	bool RenameMemberVariable(const FName& InOldName, const FName& InNewName);
+
+	UFUNCTION(BlueprintCallable, Category = "Variables")
+	bool ChangeMemberVariableType(const FName& InName, const FString& InCPPType, bool bIsPublic = false, bool bIsReadOnly = false, FString InDefaultValue = TEXT(""));
+#endif
+
+private:
+
+	/** The event names this rigvm blueprint contains */
+	UPROPERTY(AssetRegistrySearchable)
+	TArray<FName> SupportedEventNames;
+
+	UPROPERTY(transient, DuplicateTransient)
+	bool bAutoRecompileVM;
+
+	UPROPERTY(transient, DuplicateTransient)
+	bool bVMRecompilationRequired;
+
+	UPROPERTY(transient, DuplicateTransient)
+	bool bIsCompiling;
+
+	UPROPERTY(transient, DuplicateTransient)
+	int32 VMRecompilationBracket;
+
+	FRigVMGraphModifiedEvent ModifiedEvent;
+	void Notify(ERigVMGraphNotifType InNotifType, UObject* InSubject);
+	void HandleModifiedEvent(ERigVMGraphNotifType InNotifType, URigVMGraph* InGraph, UObject* InSubject);
+
+#if WITH_EDITOR
+
+	UFUNCTION(BlueprintCallable, Category = "RigVM Blueprint")
+	void SuspendNotifications(bool bSuspendNotifs);
+
+	FOnRefreshEditorEvent RefreshEditorEvent;
+	FOnVariableDroppedEvent VariableDroppedEvent;
+	FOnBreakpointAdded BreakpointAddedEvent;
+
+public:
+
+	void BroadcastRefreshEditor() { return RefreshEditorEvent.Broadcast(this); }
+	FOnRefreshEditorEvent& OnRefreshEditor() { return RefreshEditorEvent; }
+	FOnVariableDroppedEvent& OnVariableDropped() { return VariableDroppedEvent; }
+	FOnBreakpointAdded& OnBreakpointAdded() { return BreakpointAddedEvent; }
+
+private:
+
+#endif
+
+	FOnVMCompiledEvent VMCompiledEvent;
+
+	virtual void CreateMemberVariablesOnLoad();
+#if WITH_EDITOR
+	static FName FindHostMemberVariableUniqueName(TSharedPtr<FKismetNameValidator> InNameValidator, const FString& InBaseName);
+	static int32 AddHostMemberVariable(URigVMBlueprint* InBlueprint, const FName& InVarName, FEdGraphPinType InVarType, bool bIsPublic, bool bIsReadOnly, FString InDefaultValue);
+	FName AddHostMemberVariableFromExternal(FRigVMExternalVariable InVariableToCreate, FString InDefaultValue = FString());
+#endif
+
+	virtual void PatchFunctionReferencesOnLoad();
+	virtual void PatchVariableNodesOnLoad();
+	virtual void PathDomainSpecificContentOnLoad() {}
+	virtual void PatchBoundVariables();
+	virtual void PatchVariableNodesWithIncorrectType();
+	virtual void PatchParameterNodesOnLoad() {}
+	virtual void PatchLinksWithCast();
+	virtual void PatchFunctionsOnLoad();
+
+	TMap<FName, int32> AddedMemberVariableMap;
+	TArray<FBPVariableDescription> LastNewVariables;
+
+public:
+	void PropagateRuntimeSettingsFromBPToInstances();
+
+private:
+
+#if WITH_EDITOR
+	void HandlePackageDone(const FEndLoadPackageContext& Context);
+	void HandlePackageDone();
+	// ControlRigBP, once end-loaded, will inform other ControlRig-Dependent systems that ControlRig instances are ready.
+	void BroadcastRigVMPackageDone();
+
+	// Previously some memory classes were parented to the asset object
+	// however it is no longer supported since classes are now identified 
+	// with only package name + class name, see FTopLevelAssetPath
+	// this function removes those deprecated class.
+	// new classes should be created by RecompileVM and parented to the Package
+	// during PostLoad
+	void RemoveDeprecatedVMMemoryClass() const;
+#endif
+
+#if WITH_EDITOR
+
+public:
+
+	FOnExternalVariablesChanged& OnExternalVariablesChanged() { return ExternalVariablesChangedEvent; }
+
+	virtual void OnPreVariableChange(UObject* InObject);
+	virtual void OnPostVariableChange(UBlueprint* InBlueprint);
+	virtual void OnVariableAdded(const FName& InVarName);
+	virtual void OnVariableRemoved(const FName& InVarName);
+	virtual void OnVariableRenamed(const FName& InOldVarName, const FName& InNewVarName);
+	virtual void OnVariableTypeChanged(const FName& InVarName, FEdGraphPinType InOldPinType, FEdGraphPinType InNewPinType);
+
+	FOnNodeDoubleClicked& OnNodeDoubleClicked() { return NodeDoubleClickedEvent; }
+	void BroadcastNodeDoubleClicked(URigVMNode* InNode);
+
+	FOnGraphImported& OnGraphImported() { return GraphImportedEvent; }
+	void BroadcastGraphImported(UEdGraph* InGraph);
+
+	FOnPostEditChangeChainProperty& OnPostEditChangeChainProperty() { return PostEditChangeChainPropertyEvent; }
+	void BroadcastPostEditChangeChainProperty(FPropertyChangedChainEvent& PropertyChangedChainEvent);
+
+	FOnLocalizeFunctionDialogRequested& OnRequestLocalizeFunctionDialog() { return RequestLocalizeFunctionDialog; }
+	void BroadcastRequestLocalizeFunctionDialog(FRigVMGraphFunctionIdentifier InFunction, bool bForce = false);
+
+	FRigVMOnBulkEditDialogRequestedDelegate& OnRequestBulkEditDialog() { return RequestBulkEditDialog; }
+
+	FRigVMOnBreakLinksDialogRequestedDelegate& OnRequestBreakLinksDialog() { return RequestBreakLinksDialog; }
+
+	FRigVMOnPinTypeSelectionRequestedDelegate& OnRequestPinTypeSelectionDialog() { return RequestPinTypeSelectionDialog; }
+
+	FRigVMController_RequestJumpToHyperlinkDelegate& OnRequestJumpToHyperlink() { return RequestJumpToHyperlink; };
+
+	FOnReportCompilerMessage& OnReportCompilerMessage() { return ReportCompilerMessageEvent; }
+	void BroadCastReportCompilerMessage(EMessageSeverity::Type InSeverity, UObject* InSubject, const FString& InMessage);
+
+private:
+
+	FOnExternalVariablesChanged ExternalVariablesChangedEvent;
+	bool bUpdatingExternalVariables;
+	void BroadcastExternalVariablesChangedEvent();
+	FCompilerResultsLog CompileLog;
+
+	FOnNodeDoubleClicked NodeDoubleClickedEvent;
+	FOnGraphImported GraphImportedEvent;
+	FOnPostEditChangeChainProperty PostEditChangeChainPropertyEvent;
+	FOnLocalizeFunctionDialogRequested RequestLocalizeFunctionDialog;
+	FOnReportCompilerMessage ReportCompilerMessageEvent;
+	FRigVMOnBulkEditDialogRequestedDelegate RequestBulkEditDialog;
+	FRigVMOnBreakLinksDialogRequestedDelegate RequestBreakLinksDialog;
+	FRigVMOnPinTypeSelectionRequestedDelegate RequestPinTypeSelectionDialog;
+	FRigVMController_RequestJumpToHyperlinkDelegate RequestJumpToHyperlink;
+
+#endif
+
+	UEdGraph* CreateEdGraph(URigVMGraph* InModel, bool bForce = false);
+	bool RemoveEdGraph(URigVMGraph* InModel);
+	void DestroyObject(UObject* InObject);
+	void RenameGraph(const FString& InNodePath, const FName& InNewName);
+	void CreateEdGraphForCollapseNodeIfNeeded(URigVMCollapseNode* InNode, bool bForce = false);
+	bool RemoveEdGraphForCollapseNode(URigVMCollapseNode* InNode, bool bNotify = false);
+	void HandleReportFromCompiler(EMessageSeverity::Type InSeverity, UObject* InSubject, const FString& InMessage);
+
+	TArray<IRigVMGraphFunctionHost*> GetReferencedFunctionHosts(bool bForceLoad);
+	
+#if WITH_EDITOR
+private:
+	bool bCompileInDebugMode;
+	
+	TArray<URigVMNode*> RigVMBreakpointNodes;
+
+	FOnRequestInspectObject OnRequestInspectObjectEvent;
+	FRigVMBlueprintGetFocusedGraph OnGetFocusedGraphDelegate;
+		
+public:
+
+	/** Sets the execution mode. In Release mode the rig will ignore all breakpoints. */
+	void SetDebugMode(const bool bValue) { bCompileInDebugMode = bValue; }
+
+	/** Removes all the breakpoints from the blueprint and the VM */
+	void ClearBreakpoints();
+
+	/** Adds a breakpoint to all loaded blueprints which use the node indicated by InBreakpointNodePath 
+	  * If the node is inside a public function, it will add a breakpoint to all blueprints calling this function. */
+	bool AddBreakpoint(const FString& InBreakpointNodePath);
+
+	/** Adds a breakpoint to all loaded blueprints which use the InBreakpointNode. 
+	  * If LibraryNode is not null, it indicates that the library uses the InBreapointNode, and the function will add
+	  * breakpoints to any other loaded blueprint that references this library. */
+	bool AddBreakpoint(URigVMNode* InBreakpointNode, URigVMLibraryNode* LibraryNode = nullptr);
+
+	/** Adds a breakpoint to the first instruction of each callpath related to the InBreakpointNode */
+	bool AddBreakpointToControlRig(URigVMNode* InBreakpointNode);
+
+	/** Removes the given breakpoint from all the loaded blueprints that use this node, and recomputes all breakpoints
+	  * in the VM  */
+	bool RemoveBreakpoint(const FString& InBreakpointNodePath);
+	bool RemoveBreakpoint(URigVMNode* InBreakpointNode);
+
+	/** Recomputes the instruction breakpoints given the node breakpoints in the blueprint */
+	void RefreshControlRigBreakpoints();
+
+	/** Shape libraries to load during package load completed */ 
+	TArray<FString> ShapeLibrariesToLoadOnPackageLoaded;
+
+	TArray<FRigVMReferenceNodeData> GetReferenceNodeData() const;
+
+#endif
+
+protected:
+
+	static FSoftObjectPath PreDuplicateAssetPath;
+	static FSoftObjectPath PreDuplicateHostPath;
+
+private:
+	bool bDirtyDuringLoad;
+	bool bErrorsDuringCompilation;
+	bool bSuspendPythonMessagesForRigVMClient;
+	bool bMarkBlueprintAsStructurallyModifiedPending;
+
+	friend class FRigVMBlueprintCompilerContext;
+	friend class FRigVMEditorModule;
+	friend class URigVMEdGraphSchema;
+	friend struct FRigVMEdGraphSchemaAction_PromoteToVariable;
+};
+
+class RIGVMDEVELOPER_API FRigVMBlueprintVMCompileScope
+{
+public:
+   
+	FRigVMBlueprintVMCompileScope(URigVMBlueprint *InBlueprint)
+	: Blueprint(InBlueprint)
+	{
+		check(Blueprint);
+		Blueprint->IncrementVMRecompileBracket();
+	}
+
+	~FRigVMBlueprintVMCompileScope()
+	{
+		Blueprint->DecrementVMRecompileBracket();
+	}
+   
+private:
+
+	URigVMBlueprint* Blueprint;
+};
\ No newline at end of file
diff --git a/Engine/Plugins/Runtime/RigVM/Source/RigVMDeveloper/Public/RigVMBlueprintUtils.h b/Engine/Plugins/Runtime/RigVM/Source/RigVMDeveloper/Public/RigVMBlueprintUtils.h
new file mode 100644
--- /dev/null
+++ b/Engine/Plugins/Runtime/RigVM/Source/RigVMDeveloper/Public/RigVMBlueprintUtils.h
@@ -0,0 +1,30 @@
+// Copyright Epic Games, Inc. All Rights Reserved.
+
+#pragma once
+
+#include "CoreMinimal.h"
+#include "Templates/SubclassOf.h"
+#include "EdGraph/EdGraphPin.h"
+
+class UStruct;
+class UBlueprint;
+class URigVMEdGraphNode;
+class UEdGraph;
+class UEdGraphPin;
+
+struct RIGVMDEVELOPER_API FRigVMBlueprintUtils
+{
+/** Call a function for each valid rig unit struct */
+static void ForAllRigVMStructs(TFunction<void(UScriptStruct*)> InFunction);
+
+/** Handle blueprint node reconstruction */
+static void HandleReconstructAllNodes(UBlueprint* InBlueprint);
+
+/** Handle blueprint node refresh */
+static void HandleRefreshAllNodes(UBlueprint* InBlueprint);
+
+/** remove the variable if not used by anybody else but ToBeDeleted*/
+static void RemoveMemberVariableIfNotUsed(UBlueprint* Blueprint, const FName VarName);
+
+static FName ValidateName(UBlueprint* InBlueprint, const FString& InName);
+};
\ No newline at end of file
diff --git a/Engine/Plugins/Runtime/RigVM/Source/RigVMDeveloper/Public/RigVMModel/Nodes/RigVMFunctionReferenceNode.h b/Engine/Plugins/Runtime/RigVM/Source/RigVMDeveloper/Public/RigVMModel/Nodes/RigVMFunctionReferenceNode.h
--- a/Engine/Plugins/Runtime/RigVM/Source/RigVMDeveloper/Public/RigVMModel/Nodes/RigVMFunctionReferenceNode.h
+++ b/Engine/Plugins/Runtime/RigVM/Source/RigVMDeveloper/Public/RigVMModel/Nodes/RigVMFunctionReferenceNode.h
@@ -71,6 +71,7 @@ private:
 	friend class URigVMController;
 	friend class FRigVMParserAST;
 	friend class UControlRigBlueprint;
+	friend class URigVMBlueprint;
 	friend struct FRigVMClient;
 	friend struct EngineTestRigVMFramework;
 };
diff --git a/Engine/Plugins/Runtime/RigVM/Source/RigVMDeveloper/Public/RigVMModel/Nodes/RigVMTemplateNode.h b/Engine/Plugins/Runtime/RigVM/Source/RigVMDeveloper/Public/RigVMModel/Nodes/RigVMTemplateNode.h
--- a/Engine/Plugins/Runtime/RigVM/Source/RigVMDeveloper/Public/RigVMModel/Nodes/RigVMTemplateNode.h
+++ b/Engine/Plugins/Runtime/RigVM/Source/RigVMDeveloper/Public/RigVMModel/Nodes/RigVMTemplateNode.h
@@ -170,6 +170,7 @@ protected:
 
 	friend class URigVMController;
 	friend class UControlRigBlueprint;
+	friend class URigVMBlueprint;
 	friend struct FRigVMSetTemplateFilteredPermutationsAction;
 	friend struct FRigVMSetPreferredTemplatePermutationsAction;
 	friend struct FRigVMRemoveNodeAction;
diff --git a/Engine/Plugins/Runtime/RigVM/Source/RigVMDeveloper/RigVMDeveloper.Build.cs b/Engine/Plugins/Runtime/RigVM/Source/RigVMDeveloper/RigVMDeveloper.Build.cs
--- a/Engine/Plugins/Runtime/RigVM/Source/RigVMDeveloper/RigVMDeveloper.Build.cs
+++ b/Engine/Plugins/Runtime/RigVM/Source/RigVMDeveloper/RigVMDeveloper.Build.cs
@@ -13,6 +13,7 @@ public class RigVMDeveloper : ModuleRules
                 "Engine",
                 "RigVM",
                 "VisualGraphUtils",
+                "KismetCompiler",
             }
         );
 
@@ -27,8 +28,10 @@ public class RigVMDeveloper : ModuleRules
 					"SlateCore",
 					
 					"MessageLog",
-					"BlueprintGraph"
-				}
+					"BlueprintGraph",
+					"GraphEditor",
+					"Kismet",
+                }
 			);
         }
     }
