RigVM: Missing include, non-unity breakage

diff --git a/Engine/Plugins/Runtime/RigVM/Source/RigVMDeveloper/Private/RigVMBlueprintUtils.cpp b/Engine/Plugins/Runtime/RigVM/Source/RigVMDeveloper/Private/RigVMBlueprintUtils.cpp
--- a/Engine/Plugins/Runtime/RigVM/Source/RigVMDeveloper/Private/RigVMBlueprintUtils.cpp
+++ b/Engine/Plugins/Runtime/RigVM/Source/RigVMDeveloper/Private/RigVMBlueprintUtils.cpp
@@ -7,6 +7,7 @@
 #include "EdGraph/RigVMEdGraphNode.h"
 #include "RigVMBlueprint.h"
 #include "Kismet2/Kismet2NameValidators.h"
+#include "Stats/StatsHierarchical.h"
 
 #define LOCTEXT_NAMESPACE "RigVMBlueprintUtils"
 
