ControlRig: Fix error messages "Cannot add function reference to ... because argument ... has invalid type"

diff --git a/Engine/Plugins/Runtime/RigVM/Source/RigVM/Private/RigVMCore/RigVMRegistry.cpp b/Engine/Plugins/Runtime/RigVM/Source/RigVM/Private/RigVMCore/RigVMRegistry.cpp
--- a/Engine/Plugins/Runtime/RigVM/Source/RigVM/Private/RigVMCore/RigVMRegistry.cpp
+++ b/Engine/Plugins/Runtime/RigVM/Source/RigVM/Private/RigVMCore/RigVMRegistry.cpp
@@ -1128,7 +1128,7 @@ bool FRigVMRegistry::IsAllowedType(const UStruct* InStruct)
 
 bool FRigVMRegistry::IsAllowedType(const UClass* InClass)
 {
-	if(InClass->HasAnyClassFlags(CLASS_Hidden | CLASS_Abstract))
+	if(!InClass || InClass->HasAnyClassFlags(CLASS_Hidden | CLASS_Abstract))
 	{
 		return false;
 	}
diff --git a/Engine/Plugins/Runtime/RigVM/Source/RigVMDeveloper/Private/RigVMModel/RigVMSchema.cpp b/Engine/Plugins/Runtime/RigVM/Source/RigVMDeveloper/Private/RigVMModel/RigVMSchema.cpp
--- a/Engine/Plugins/Runtime/RigVM/Source/RigVMDeveloper/Private/RigVMModel/RigVMSchema.cpp
+++ b/Engine/Plugins/Runtime/RigVM/Source/RigVMDeveloper/Private/RigVMModel/RigVMSchema.cpp
@@ -134,7 +134,16 @@ bool URigVMSchema::SupportsGraphFunction(URigVMController* InController, const F
 	// Make sure all the argument types are supported before creating the node
 	for (const FRigVMGraphFunctionArgument& Argument : InGraphFunction->Arguments)
 	{
-		const TRigVMTypeIndex Type = Registry->GetTypeIndexFromCPPType(Argument.CPPType.ToString());
+		TRigVMTypeIndex Type = Registry->GetTypeIndexFromCPPType(Argument.CPPType.ToString());
+		if (Type == INDEX_NONE)
+		{
+			if (Argument.CPPTypeObject.IsValid())
+			{
+				FRigVMTemplateArgumentType ArgumentType(Argument.CPPType, Argument.CPPTypeObject.Get());
+				Type = Registry->FindOrAddType(ArgumentType);
+			}
+		}
+		
 		if (Type == INDEX_NONE)
 		{
 			InController->ReportErrorf(TEXT("Cannot add function reference to %s because argument %s has invalid type %s."), *InGraphFunction->Name.ToString(), *Argument.Name.ToString(), *Argument.CPPType.ToString());
