RigVM: Core baseline support for decorators

diff --git a/Engine/Plugins/Animation/ControlRig/Source/ControlRigDeveloper/Private/Graph/ControlRigGraphSchema.cpp b/Engine/Plugins/Animation/ControlRig/Source/ControlRigDeveloper/Private/Graph/ControlRigGraphSchema.cpp
--- a/Engine/Plugins/Animation/ControlRig/Source/ControlRigDeveloper/Private/Graph/ControlRigGraphSchema.cpp
+++ b/Engine/Plugins/Animation/ControlRig/Source/ControlRigDeveloper/Private/Graph/ControlRigGraphSchema.cpp
@@ -914,6 +914,11 @@ FLinearColor UControlRigGraphSchema::GetPinTypeColor(const FEdGraphPinType& PinT
 				return FLinearColor::White;
 			}
 
+			if (Struct->IsChildOf(FRigVMDecorator::StaticStruct()))
+			{
+				return FLinearColor::White;
+			}
+
 			if (Struct->IsChildOf(RigVMTypeUtils::GetWildCardCPPTypeObject()))
 			{
 				return FLinearColor(FVector3f::OneVector * 0.25f);
diff --git a/Engine/Plugins/Animation/ControlRig/Source/ControlRigEditor/Private/ControlRigEditorModule.cpp b/Engine/Plugins/Animation/ControlRig/Source/ControlRigEditor/Private/ControlRigEditorModule.cpp
--- a/Engine/Plugins/Animation/ControlRig/Source/ControlRigEditor/Private/ControlRigEditorModule.cpp
+++ b/Engine/Plugins/Animation/ControlRig/Source/ControlRigEditor/Private/ControlRigEditorModule.cpp
@@ -2644,6 +2644,30 @@ void FControlRigEditorModule::GetContextMenuActions(const UControlRigGraphSchema
 							InSection.AddMenuEntry(FGraphEditorCommands::Get().DistributeNodesVertically);
 						}
 					}));
+
+					// this struct is only available in EngineTest for now
+					static const FString DecoratorObjectPath = TEXT("/Script/EngineTestEditor.EngineTestRigVM_SimpleDecorator");
+					if(const UScriptStruct* SimpleDecoratorStruct = Cast<UScriptStruct>(RigVMTypeUtils::FindObjectFromCPPTypeObjectPath(DecoratorObjectPath)))
+					{
+						if(URigVMNode* ModelNode = RigNode->GetModelNode())
+						{
+							FToolMenuSection& EngineTestSection = Menu->AddSection("EdGraphSchemaEngineTest", LOCTEXT("EngineTestHeader", "EngineTest"));
+							EngineTestSection.AddMenuEntry(
+								"Add simple decorator",
+								LOCTEXT("AddSimpleDecorator", "Add simple decorator"),
+								LOCTEXT("AddSimpleDecorator_Tooltip", "Adds a simple test decorator to the node"),
+								FSlateIcon(),
+								FUIAction(FExecuteAction::CreateLambda([Controller, ModelNode, SimpleDecoratorStruct]()
+								{
+									(void)Controller->AddDecorator(
+										ModelNode->GetFName(),
+										*SimpleDecoratorStruct->GetPathName(),
+										TEXT("Decorator"),
+										FString(), INDEX_NONE, true, true);
+								}))
+							);
+						}
+					}
 				}
 			}
 		}
diff --git a/Engine/Plugins/Runtime/RigVM/Source/RigVM/Private/RigVMCore/RigVMRegistry.cpp b/Engine/Plugins/Runtime/RigVM/Source/RigVM/Private/RigVMCore/RigVMRegistry.cpp
--- a/Engine/Plugins/Runtime/RigVM/Source/RigVM/Private/RigVMCore/RigVMRegistry.cpp
+++ b/Engine/Plugins/Runtime/RigVM/Source/RigVM/Private/RigVMCore/RigVMRegistry.cpp
@@ -3,6 +3,7 @@
 #include "RigVMCore/RigVMRegistry.h"
 #include "AssetRegistry/AssetData.h"
 #include "RigVMCore/RigVMStruct.h"
+#include <RigVMCore/RigVMDecorator.h>
 #include "RigVMTypeUtils.h"
 #include "RigVMModule.h"
 #include "Animation/AttributeTypes.h"
@@ -1101,7 +1102,8 @@ bool FRigVMRegistry::IsAllowedType(const UStruct* InStruct)
 	{
 		return false;
 	}
-	if(InStruct->IsChildOf(FRigVMStruct::StaticStruct()))
+	if(InStruct->IsChildOf(FRigVMStruct::StaticStruct()) &&
+		!InStruct->IsChildOf(FRigVMDecorator::StaticStruct()))
 	{
 		return false;
 	}
diff --git a/Engine/Plugins/Runtime/RigVM/Source/RigVM/Public/RigVMCore/RigVMDecorator.h b/Engine/Plugins/Runtime/RigVM/Source/RigVM/Public/RigVMCore/RigVMDecorator.h
new file mode 100644
--- /dev/null
+++ b/Engine/Plugins/Runtime/RigVM/Source/RigVM/Public/RigVMCore/RigVMDecorator.h
@@ -0,0 +1,76 @@
+// Copyright Epic Games, Inc. All Rights Reserved.
+
+#pragma once
+
+#include "RigVMDefines.h"
+#include "RigVMCore/RigVMStruct.h"
+
+#include "RigVMDecorator.generated.h"
+
+#if WITH_EDITOR
+
+class URigVMController;
+class URigVMNode;
+
+#endif
+
+/**
+ * The base class for all RigVM decorators.
+ */
+USTRUCT()
+struct RIGVM_API FRigVMDecorator : public FRigVMStruct
+{
+	GENERATED_BODY()
+
+public:
+
+	FRigVMDecorator()
+	: Name(NAME_None)
+	, DecoratorStruct(nullptr)
+	{}
+
+	virtual ~FRigVMDecorator() {}
+
+	// returns the name of the decorator (the instance of it on the node)
+	FName GetName() const { return Name; }
+
+	// returns the display name of the decorator
+	virtual FString GetDisplayName() const
+	{
+		return FString();
+	}
+
+	// returns the struct of this decorator
+	const UScriptStruct* GetScriptStruct() const { return DecoratorStruct; }
+
+	// returns true if the given decorator is of a certain type (or super type)
+	template<typename T>
+	bool IsA() const
+	{
+		if(const UScriptStruct* ScriptStruct = GetScriptStruct())
+		{
+			return ScriptStruct->IsChildOf(T::StaticStruct());
+		}
+		return false;
+	}
+
+#if WITH_EDITOR
+
+	// returns true if this decorator can be added to a given node
+	virtual bool CanBeAddedToNode(URigVMNode* InNode, FString* OutFailureReason) const { return true; }
+
+	// allows the decorator to react when added to a node
+	virtual void OnDecoratorAdded(URigVMController* InController, URigVMNode* InNode) {}
+
+#endif
+
+private:
+
+	// The name of the decorator on the node
+	FName Name;
+	// The struct backing up the decorator
+	UScriptStruct* DecoratorStruct;
+
+	friend class URigVMNode;
+	friend class URigVMController;
+};
diff --git a/Engine/Plugins/Runtime/RigVM/Source/RigVMDeveloper/Private/RigVMModel/RigVMController.cpp b/Engine/Plugins/Runtime/RigVM/Source/RigVMDeveloper/Private/RigVMModel/RigVMController.cpp
--- a/Engine/Plugins/Runtime/RigVM/Source/RigVMDeveloper/Private/RigVMModel/RigVMController.cpp
+++ b/Engine/Plugins/Runtime/RigVM/Source/RigVMDeveloper/Private/RigVMModel/RigVMController.cpp
@@ -30,6 +30,7 @@
 #include "RigVMModel/RigVMClient.h"
 #include "RigVMModel/Nodes/RigVMBranchNode.h"
 #include "RigVMModel/Nodes/RigVMArrayNode.h"
+#include "Logging/LogScopedVerbosityOverride.h"
 
 #include UE_INLINE_GENERATED_CPP_BY_NAME(RigVMController)
 
@@ -359,6 +360,7 @@ URigVMController::FPinInfo::FPinInfo()
 	, bIsExpanded(false)
 	, bIsConstant(false)
 	, bIsDynamicArray(false)
+	, bIsDecorator(false)
 {
 }
 
@@ -372,6 +374,7 @@ URigVMController::FPinInfo::FPinInfo(const URigVMPin* InPin, int32 InParentIndex
 	, bIsExpanded(InPin->IsExpanded())
 	, bIsConstant(InPin->IsDefinedAsConstant())
 	, bIsDynamicArray(InPin->IsDynamicArray())
+	, bIsDecorator(InPin->IsDecoratorPin() && InPin->IsRootPin())
 {
 	// this method describes the info as currently represented in the model.
 
@@ -395,6 +398,7 @@ URigVMController::FPinInfo::FPinInfo(FProperty* InProperty, ERigVMPinDirection I
 	, bIsExpanded(false)
 	, bIsConstant(false)
 	, bIsDynamicArray(false)
+	, bIsDecorator(false)
 {
 	// this method describes the info as needed based on the property structure
 
@@ -512,6 +516,7 @@ uint32 GetTypeHash(const URigVMController::FPinInfo& InPin)
 	Hash = HashCombine(Hash, GetTypeHash((int32)InPin.Direction));
 	Hash = HashCombine(Hash, GetTypeHash((int32)InPin.TypeIndex));
 	Hash = HashCombine(Hash, GetTypeHash(InPin.bIsArray));
+	Hash = HashCombine(Hash, GetTypeHash(InPin.bIsDecorator));
 	// we are not hashing the parent index,  pinpath, default value or the property since
 	// it doesn't matter for the structure validity of the node
 	return Hash;
@@ -1200,6 +1205,11 @@ TArray<FString> URigVMController::GetAddNodePythonCommands(URigVMNode* Node) con
 			{
 				continue;
 			}
+
+			if(Pin->IsDecoratorPin())
+			{
+				continue;
+			}
 			
 			const FString DefaultValue = Pin->GetDefaultValue();
 			if (!DefaultValue.IsEmpty() && DefaultValue != TEXT("()"))
@@ -1238,6 +1248,34 @@ TArray<FString> URigVMController::GetAddNodePythonCommands(URigVMNode* Node) con
 							*Pin->GetBoundVariablePath()));
 			}
 		}
+
+		for(const FString& DecoratorName : Node->GetDecoratorNames())
+		{
+			Commands.Append(GetAddDecoratorPythonCommands(Node, *DecoratorName));
+		}
+	}
+
+	return Commands;
+}
+
+TArray<FString> URigVMController::GetAddDecoratorPythonCommands(URigVMNode* Node, const FName& DecoratorName) const
+{
+	TArray<FString> Commands;
+
+	const FString GraphName = GetSchema()->GetSanitizedGraphName(GetGraph()->GetGraphName());
+	const FString NodeName = GetSchema()->GetSanitizedNodeName(Node->GetName());
+
+	if(const URigVMPin* DecoratorPin = Node->FindDecorator(DecoratorName))
+	{
+		const FString DecoratorStructPath = DecoratorPin->GetCPPTypeObject()->GetPathName();
+		const FString DefaultValue = DecoratorPin->GetDefaultValue();
+
+		Commands.Add(FString::Printf(TEXT("blueprint.get_controller_by_name('%s').add_decorator('%s', '%s', '%s', '%s')"),
+			*GraphName,
+			*NodeName,
+			*DecoratorStructPath,
+			*DecoratorName.ToString(),
+			*DefaultValue));
 	}
 
 	return Commands;
@@ -8976,7 +9014,9 @@ bool URigVMController::RemovePin(URigVMPin* InPinToRemove, bool bSetupUndoRedo,
 	{
 		return false;
 	}
-	
+
+	check(!(InPinToRemove->IsRootPin() && InPinToRemove->IsDecoratorPin()));
+
 	FRigVMControllerCompileBracketScope CompileScope(this);
 	if (bSetupUndoRedo || bForceBreakLinks)
 	{
@@ -10713,6 +10753,8 @@ bool URigVMController::RenameExposedPin(const FName& InOldPinName, const FName&
 	{
 		static bool RenamePin(const URigVMController* InController, URigVMPin* InPin, const FName& InNewName)
 		{
+			check(!InPin->IsDecoratorPin());
+			
 			URigVMController* PinController = InController->GetControllerForGraph(InPin->GetGraph());
 			if(PinController == nullptr)
 			{
@@ -13565,6 +13607,242 @@ URigVMInvokeEntryNode* URigVMController::AddInvokeEntryNode(const FName& InEntry
 	return Node;
 }
 
+FName URigVMController::AddDecorator(const FName& InNodeName, const FName& InDecoratorTypeObjectPath,
+	const FName& InDecoratorName, const FString& InDefaultValue, int32 InPinIndex, bool bSetupUndoRedo,
+	bool bPrintPythonCommand)
+{
+	if (!IsValidGraph())
+	{
+		return NAME_None;
+	}
+
+	if (!bIsTransacting && !IsGraphEditable())
+	{
+		return NAME_None;
+	}
+
+	const URigVMGraph* Graph = GetGraph();
+	check(Graph);
+
+	if (URigVMNode* Node = Graph->FindNodeByName(InNodeName))
+	{
+		UObject* DecoratorCPPTypeObject = RigVMTypeUtils::FindObjectFromCPPTypeObjectPath(InDecoratorTypeObjectPath.ToString());
+		if(DecoratorCPPTypeObject == nullptr)
+		{
+			DecoratorCPPTypeObject = URigVMCompiler::GetScriptStructForCPPType(InDecoratorTypeObjectPath.ToString());
+		}
+		if(DecoratorCPPTypeObject == nullptr)
+		{
+			ReportErrorf(TEXT("Cannot find decorator script struct '%s'."), *InDecoratorTypeObjectPath.ToString());
+			return NAME_None;
+		}
+
+		UScriptStruct* DecoratorScriptStruct = Cast<UScriptStruct>(DecoratorCPPTypeObject);
+		if(DecoratorScriptStruct == nullptr)
+		{
+			ReportErrorf(TEXT("CPP Type Object '%s' is not a struct."), *InDecoratorTypeObjectPath.ToString());
+			return NAME_None;
+		}
+
+		const FName DecoratorName = AddDecorator(Node, DecoratorScriptStruct, InDecoratorName, InDefaultValue, InPinIndex, bSetupUndoRedo);
+		if(!DecoratorName.IsNone() && bPrintPythonCommand)
+		{
+			const TArray<FString> DecoratorCommands = GetAddDecoratorPythonCommands(Node, DecoratorName);
+			for(const FString& DecoratorCommand : DecoratorCommands)
+			{
+				RigVMPythonUtils::Print(GetSchema()->GetGraphOuterName(GetGraph()), DecoratorCommand);
+			}
+		}
+		return DecoratorName;
+	}
+
+	ReportErrorf(TEXT("Cannot find node '%s'."), *InNodeName.ToString());
+	return NAME_None;
+}
+
+FName URigVMController::AddDecorator(URigVMNode* InNode, UScriptStruct* InDecoratorScriptStruct, const FName& InDecoratorName, const FString& InDefaultValue, int32 InPinIndex, bool bSetupUndoRedo)
+{
+	if (!IsValidNodeForGraph(InNode))
+	{
+		return NAME_None;
+	}
+
+	if (!bIsTransacting && !IsGraphEditable())
+	{
+		return NAME_None;
+	}
+
+	check(InDecoratorScriptStruct);
+
+	if(!InDecoratorScriptStruct->IsChildOf(FRigVMDecorator::StaticStruct()))
+	{
+		ReportErrorf(TEXT("CPP Type Object '%s' is not a struct."), *InDecoratorScriptStruct->GetPathName());
+		return NAME_None;
+	}
+
+	const FRigVMTemplateArgumentType DecoratorType(InDecoratorScriptStruct);
+	const TRigVMTypeIndex DecoratorTypeIndex = FRigVMRegistry::Get().FindOrAddType(DecoratorType);
+
+	if(const URigVMSchema* Schema = GetSchema())
+	{
+		if(!Schema->SupportsType(this, DecoratorTypeIndex))
+		{
+			ReportError(TEXT("Decorator cannot be added to node: Schema doesn't support the type."));
+			return NAME_None;
+		}
+	}
+
+	const FName ValidDecoratorName = URigVMSchema::GetUniqueName(InDecoratorName, [InNode](const FName& InName) {
+		return InNode->FindPin(InName.ToString()) == nullptr;
+	}, false, false);
+
+	TSharedPtr<FStructOnScope> DecoratorScope(new FStructOnScope(InDecoratorScriptStruct));
+	FRigVMDecorator* Decorator = (FRigVMDecorator*)DecoratorScope->GetStructMemory();
+
+	if(!InDefaultValue.IsEmpty())
+	{
+		FRigVMPinDefaultValueImportErrorContext ErrorPipe;
+		{
+			// force logging to the error pipe for error detection
+			LOG_SCOPE_VERBOSITY_OVERRIDE(LogExec, ELogVerbosity::Verbose); 
+			InDecoratorScriptStruct->ImportText(*InDefaultValue, Decorator, nullptr, PPF_None, &ErrorPipe, InDecoratorScriptStruct->GetName()); 
+		}
+	}
+
+	Decorator->Name = ValidDecoratorName;
+	Decorator->DecoratorStruct = InDecoratorScriptStruct;
+
+	FString FailureReason;
+	if(!Decorator->CanBeAddedToNode(InNode, &FailureReason))
+	{
+		ReportErrorf(TEXT("Decorator cannot be added to node: %s"), *FailureReason);
+		return NAME_None;
+	}
+
+	FRigVMBaseAction Action(this);
+	if (bSetupUndoRedo)
+	{
+		Action.SetTitle(FString::Printf(TEXT("Add Decorator")));
+		GetActionStack()->BeginAction(Action);
+		GetActionStack()->AddAction(FRigVMAddDecoratorAction(this, InNode, ValidDecoratorName, InDecoratorScriptStruct, InDefaultValue, InNode->GetPins().Num()));
+	}
+
+	InNode->DecoratorRootPinNames.Add(ValidDecoratorName.ToString());
+
+	URigVMPin* DecoratorPin = NewObject<URigVMPin>(InNode, ValidDecoratorName);
+	const FString DisplayName = Decorator->GetDisplayName();
+	DecoratorPin->DisplayName = DisplayName.IsEmpty() ? FName(NAME_None) : FName(*DisplayName); 
+	DecoratorPin->CPPType = InDecoratorScriptStruct->GetStructCPPName();
+	DecoratorPin->CPPTypeObject = InDecoratorScriptStruct;
+	DecoratorPin->CPPTypeObjectPath = *DecoratorPin->CPPTypeObject->GetPathName();
+	DecoratorPin->Direction = ERigVMPinDirection::Input;
+
+	AddNodePin(InNode, DecoratorPin);
+	Notify(ERigVMGraphNotifType::PinAdded, DecoratorPin);
+	
+	AddPinsForStruct(InDecoratorScriptStruct, InNode, DecoratorPin, DecoratorPin->GetDirection(), InDefaultValue, true);
+
+	// move the the pin to the right index as required
+	if(DecoratorPin->GetPinIndex() != InPinIndex &&
+		InPinIndex >=0 && InPinIndex < InNode->GetPins().Num())
+	{
+		URigVMPin* LastPin =  InNode->Pins.Pop();
+		InNode->Pins.Insert(LastPin, InPinIndex);
+	}
+
+	InNode->UpdateDecoratorRootPinNames();
+	Decorator->OnDecoratorAdded(this, InNode);
+
+	if (bSetupUndoRedo)
+	{
+		GetActionStack()->EndAction(Action);
+	}
+
+	return ValidDecoratorName;
+}
+
+bool URigVMController::RemoveDecorator(const FName& InNodeName, const FName& InDecoratorName, bool bSetupUndoRedo, bool bPrintPythonCommand)
+{
+	if (!IsValidGraph())
+	{
+		return false;
+	}
+
+	if (!bIsTransacting && !IsGraphEditable())
+	{
+		return false;
+	}
+
+	const URigVMGraph* Graph = GetGraph();
+	check(Graph);
+
+	if (URigVMNode* Node = Graph->FindNodeByName(InNodeName))
+	{
+		const bool bSuccess = RemoveDecorator(Node, InDecoratorName, bSetupUndoRedo);
+		if(bSuccess && bPrintPythonCommand)
+		{
+			const FString GraphName = GetSchema()->GetSanitizedGraphName(GetGraph()->GetGraphName());
+        	const FString NodeName = GetSchema()->GetSanitizedNodeName(Node->GetName());
+
+			RigVMPythonUtils::Print(GetSchema()->GetGraphOuterName(GetGraph()),
+				FString::Printf(TEXT("blueprint.get_controller_by_name('%s').remove_decorator('%s', '%s')"),
+				*GraphName,
+				*NodeName,
+				*InDecoratorName.ToString()));
+		}
+		return bSuccess;
+	}
+
+	ReportErrorf(TEXT("Cannot find node '%s'."), *InNodeName.ToString());
+	return false;
+
+}
+
+bool URigVMController::RemoveDecorator(URigVMNode* InNode, const FName& InDecoratorName, bool bSetupUndoRedo)
+{
+	if(!IsValidNodeForGraph(InNode))
+	{
+		return false;
+	}
+
+	if (!bIsTransacting && !IsGraphEditable())
+	{
+		return false;
+	}
+
+	URigVMPin* DecoratorPin = InNode->FindDecorator(InDecoratorName);
+	if(DecoratorPin == nullptr)
+	{
+		return false;
+	}
+
+	FRigVMBaseAction Action(this);
+	if (bSetupUndoRedo)
+	{
+		Action.SetTitle(FString::Printf(TEXT("Remove Decorator")));
+		GetActionStack()->BeginAction(Action);
+
+		const UScriptStruct* DecoratorScriptStruct = DecoratorPin->GetScriptStruct();
+		const FString DecoratorDefaultValue = DecoratorPin->GetDefaultValue();
+		GetActionStack()->AddAction(FRigVMRemoveDecoratorAction(this, InNode, InDecoratorName, DecoratorScriptStruct, DecoratorDefaultValue, DecoratorPin->GetPinIndex()));
+	}
+
+	const FString DecoratorNameString = InDecoratorName.ToString();
+	(void)InNode->DecoratorRootPinNames.RemoveAll([DecoratorNameString](const FString& DecoratorRootPinName) -> bool
+	{
+		return DecoratorNameString.Equals(DecoratorRootPinName, ESearchCase::CaseSensitive);
+	});
+
+	RemovePin(DecoratorPin, bSetupUndoRedo, true);
+
+	if(bSetupUndoRedo)
+	{
+		GetActionStack()->EndAction(Action);
+	}
+
+	return true;
+}
+
 void URigVMController::ForEveryPinRecursively(URigVMPin* InPin, TFunction<void(URigVMPin*)> OnEachPinFunction)
 {
 	OnEachPinFunction(InPin);
@@ -14839,6 +15117,17 @@ void URigVMController::RepopulatePinsOnNode(URigVMNode* InNode, bool bFollowCore
 			}
 		};
 
+	// make sure the new pin infos contains the decorator pins from the last run
+	for(int32 Index = 0; Index < PreviousPinInfos.Num(); Index++)
+	{
+		const FPinInfo& PreviousPin = PreviousPinInfos[Index]; 
+		if(PreviousPin.bIsDecorator)
+		{
+			const int32 NewPinIndex = NewPinInfos.AddPin(this, INDEX_NONE, PreviousPin.Name, PreviousPin.Direction, PreviousPin.TypeIndex, PreviousPin.DefaultValue, nullptr, &PreviousPinInfos, true);
+			NewPinInfos[NewPinIndex].bIsDecorator = true;
+		}
+	}
+
 	// if the nodes match in structure - nothing to do here
 	if(GetTypeHash(NewPinInfos) == PreviousPinHash)
 	{
@@ -15269,6 +15558,16 @@ void URigVMController::RepopulatePinsOnNode(URigVMNode* InNode, bool bFollowCore
 		ApplyPinStates(InNode, PinStates, RedirectedPinPaths);
 	}
 
+	InNode->DecoratorRootPinNames.Reset();
+	for(int32 Index = 0; Index < NewPinInfos.Num(); Index++)
+	{
+		if(NewPinInfos[Index].bIsDecorator)
+		{
+			InNode->DecoratorRootPinNames.Add(NewPinInfos[Index].Name.ToString());
+		}
+	}
+	InNode->UpdateDecoratorRootPinNames();
+
 	if(!LinkedPaths.IsEmpty())
 	{
 #if UE_RIGVMCONTROLLER_VERBOSE_REPOPULATE
@@ -17001,6 +17300,11 @@ bool URigVMController::ChangePinType(URigVMPin* InPin, TRigVMTypeIndex InTypeInd
 		return false;
 	}
 
+	if(InPin->IsDecoratorPin() && InPin->GetRootPin())
+	{
+		return false;
+	}
+
 	// only allow valid pin cpp types on template nodes
 	TRigVMTypeIndex TypeIndex = InTypeIndex;
 	if(URigVMTemplateNode* TemplateNode = Cast<URigVMTemplateNode>(InPin->GetNode()))
diff --git a/Engine/Plugins/Runtime/RigVM/Source/RigVMDeveloper/Private/RigVMModel/RigVMControllerActions.cpp b/Engine/Plugins/Runtime/RigVM/Source/RigVMDeveloper/Private/RigVMModel/RigVMControllerActions.cpp
--- a/Engine/Plugins/Runtime/RigVM/Source/RigVMDeveloper/Private/RigVMModel/RigVMControllerActions.cpp
+++ b/Engine/Plugins/Runtime/RigVM/Source/RigVMDeveloper/Private/RigVMModel/RigVMControllerActions.cpp
@@ -2096,3 +2096,78 @@ bool FRigVMReplaceNodesAction::Redo()
 	return FRigVMBaseAction::Redo();
 }
 
+FRigVMAddDecoratorAction::FRigVMAddDecoratorAction()
+: FRigVMBaseAction()
+, NodeName(NAME_None)
+, DecoratorName(NAME_None)
+, ScriptStructPath()
+, DecoratorDefault()
+, PinIndex(INDEX_NONE)
+{
+}
+
+FRigVMAddDecoratorAction::FRigVMAddDecoratorAction(URigVMController* InController, const URigVMNode* InNode, const FName& InDecoratorName,
+	const UScriptStruct* InDecoratorScriptStruct, const FString& InDecoratorDefault, int32 InPinIndex)
+: FRigVMBaseAction(InController)
+, NodeName(InNode->GetFName())
+, DecoratorName(InDecoratorName)
+, ScriptStructPath(InDecoratorScriptStruct->GetPathName())
+, DecoratorDefault(InDecoratorDefault)
+, PinIndex(InPinIndex)
+{
+}
+
+bool FRigVMAddDecoratorAction::Undo()
+{
+	if (!FRigVMBaseAction::Undo())
+	{
+		return false;
+	}
+	return GetController()->RemoveDecorator(NodeName, DecoratorName, false);
+}
+
+bool FRigVMAddDecoratorAction::Redo()
+{
+	if(!CanUndoRedo())
+	{
+		return false;
+	}
+	if (GetController()->AddDecorator(NodeName, *ScriptStructPath, DecoratorName, DecoratorDefault, PinIndex, false, false) == DecoratorName)
+	{
+		return FRigVMBaseAction::Redo();
+	}
+	return false;
+}
+
+FRigVMRemoveDecoratorAction::FRigVMRemoveDecoratorAction()
+: FRigVMAddDecoratorAction()
+{
+}
+
+FRigVMRemoveDecoratorAction::FRigVMRemoveDecoratorAction(URigVMController* InController, const URigVMNode* InNode,
+	const FName& InDecoratorName, const UScriptStruct* InDecoratorScriptStruct, const FString& InDecoratorDefault, int32 InPinIndex)
+: FRigVMAddDecoratorAction(InController, InNode, InDecoratorName, InDecoratorScriptStruct, InDecoratorDefault, InPinIndex)
+{
+}
+
+bool FRigVMRemoveDecoratorAction::Undo()
+{
+	if (!FRigVMBaseAction::Undo())
+	{
+		return false;
+	}
+	return GetController()->AddDecorator(NodeName, *ScriptStructPath, DecoratorName, DecoratorDefault, PinIndex, false, false) == DecoratorName;
+}
+
+bool FRigVMRemoveDecoratorAction::Redo()
+{
+	if(!CanUndoRedo())
+	{
+		return false;
+	}
+	if (GetController()->RemoveDecorator(NodeName, DecoratorName, false))
+	{
+		return FRigVMBaseAction::Redo();
+	}
+	return false;
+}
diff --git a/Engine/Plugins/Runtime/RigVM/Source/RigVMDeveloper/Private/RigVMModel/RigVMNode.cpp b/Engine/Plugins/Runtime/RigVM/Source/RigVMDeveloper/Private/RigVMModel/RigVMNode.cpp
--- a/Engine/Plugins/Runtime/RigVM/Source/RigVMDeveloper/Private/RigVMModel/RigVMNode.cpp
+++ b/Engine/Plugins/Runtime/RigVM/Source/RigVMDeveloper/Private/RigVMModel/RigVMNode.cpp
@@ -9,6 +9,7 @@
 #include "RigVMCore/RigVMExecuteContext.h"
 #include "RigVMCore/RigVMStruct.h"
 #include "RigVMUserWorkflowRegistry.h"
+#include "Logging/LogScopedVerbosityOverride.h"
 
 #include UE_INLINE_GENERATED_CPP_BY_NAME(RigVMNode)
 
@@ -224,6 +225,19 @@ FText URigVMNode::GetToolTipTextForPin(const URigVMPin* InPin) const
 	return FText::FromName(InPin->GetFName());
 }
 
+void URigVMNode::UpdateDecoratorRootPinNames()
+{
+	TArray<FString> NewDecoratorRootPinNames;
+	for(URigVMPin* Pin : GetPins())
+	{
+		if(Pin->IsDecoratorPin())
+		{
+			NewDecoratorRootPinNames.Add(Pin->GetName());
+		}
+	}
+	DecoratorRootPinNames = NewDecoratorRootPinNames;
+}
+
 bool URigVMNode::IsSelected() const
 {
 	URigVMGraph* Graph = GetGraph();
@@ -396,6 +410,103 @@ uint32 URigVMNode::GetStructureHash() const
 	return Hash;
 }
 
+TArray<URigVMPin*> URigVMNode::GetDecoratorPins() const
+{
+	TArray<URigVMPin*> DecoratorPins;
+	DecoratorPins.Reserve(DecoratorRootPinNames.Num());
+	
+	for(const FString& DecoratorRootPinName : DecoratorRootPinNames)
+	{
+		URigVMPin* DecoratorPin = FindPin(DecoratorRootPinName);
+		check(DecoratorPin);
+
+		DecoratorPins.Add(DecoratorPin);
+	}
+
+	return DecoratorPins;
+}
+
+bool URigVMNode::IsDecoratorPin(FName InName) const
+{
+	if(const URigVMPin* Pin = FindPin(InName.ToString()))
+	{
+		return IsDecoratorPin(Pin);
+	}
+	return false;
+}
+
+bool URigVMNode::IsDecoratorPin(const URigVMPin* InDecoratorPin) const
+{
+	return FindDecorator(InDecoratorPin) != nullptr;
+}
+
+URigVMPin* URigVMNode::FindDecorator(const FName& InName) const
+{
+	const FString NameString = InName.ToString();
+	for(const FString& DecoratorRootPinName : DecoratorRootPinNames)
+	{
+		if(DecoratorRootPinName.Equals(NameString, ESearchCase::CaseSensitive))
+		{
+			return FindPin(DecoratorRootPinName);
+		}
+	}
+	return nullptr;
+}
+
+URigVMPin* URigVMNode::FindDecorator(const URigVMPin* InDecoratorPin) const
+{
+	if(InDecoratorPin)
+	{
+		const URigVMPin* RootPin = InDecoratorPin->GetRootPin();
+		if(RootPin->GetNode() == this)
+		{
+			return FindDecorator(RootPin->GetFName());
+		}
+	}
+	return nullptr;
+}
+
+TSharedPtr<FStructOnScope> URigVMNode::GetDecoratorInstance(const FName& InName, bool bUseDefaultValueFromPin) const
+{
+	return GetDecoratorInstance(FindPin(InName.ToString()), bUseDefaultValueFromPin);
+}
+
+TSharedPtr<FStructOnScope> URigVMNode::GetDecoratorInstance(const URigVMPin* InDecoratorPin, bool bUseDefaultValueFromPin) const
+{
+	if(const URigVMPin* RootPin = FindDecorator(InDecoratorPin))
+	{
+		check(RootPin->IsStruct());
+
+		UScriptStruct* ScriptStruct = RootPin->GetScriptStruct();
+		check(ScriptStruct->IsChildOf(FRigVMDecorator::StaticStruct()));
+
+		TSharedPtr<FStructOnScope> Scope(new FStructOnScope(ScriptStruct));
+		FRigVMDecorator* Decorator = (FRigVMDecorator*)Scope->GetStructMemory();
+
+		if(bUseDefaultValueFromPin)
+		{
+			const FString DefaultValue = RootPin->GetDefaultValue();
+			if(!DefaultValue.IsEmpty())
+			{
+				FRigVMPinDefaultValueImportErrorContext ErrorPipe;
+				{
+					// force logging to the error pipe for error detection
+					LOG_SCOPE_VERBOSITY_OVERRIDE(LogExec, ELogVerbosity::Verbose); 
+					ScriptStruct->ImportText(*DefaultValue, Decorator, nullptr, PPF_None, &ErrorPipe, ScriptStruct->GetName()); 
+				}
+			}
+		}
+
+		Decorator->Name = RootPin->GetFName();
+		Decorator->DecoratorStruct = ScriptStruct;
+		
+		return Scope;
+	}
+
+	static const TSharedPtr<FStructOnScope> EmptyScope;
+	return EmptyScope;
+}
+
 URigVMLibraryNode* URigVMNode::FindFunctionForNode() const  
 {
 	const UObject* Subject = this;
diff --git a/Engine/Plugins/Runtime/RigVM/Source/RigVMDeveloper/Private/RigVMModel/RigVMPin.cpp b/Engine/Plugins/Runtime/RigVM/Source/RigVMDeveloper/Private/RigVMModel/RigVMPin.cpp
--- a/Engine/Plugins/Runtime/RigVM/Source/RigVMDeveloper/Private/RigVMModel/RigVMPin.cpp
+++ b/Engine/Plugins/Runtime/RigVM/Source/RigVMDeveloper/Private/RigVMModel/RigVMPin.cpp
@@ -1390,6 +1390,26 @@ uint32 URigVMPin::GetStructureHash() const
 	return Hash;
 }
 
+bool URigVMPin::IsDecoratorPin() const
+{
+	if(const URigVMNode* Node = GetNode())
+	{
+		return Node->IsDecoratorPin(GetRootPin());
+	}
+	return false;
+}
+
+TSharedPtr<FStructOnScope> URigVMPin::GetDecoratorInstance(bool bUseDefaultValueFromPin) const
+{
+	if(const URigVMNode* Node = GetNode())
+	{
+		return Node->GetDecoratorInstance(GetRootPin(), bUseDefaultValueFromPin);
+	}
+
+	static const TSharedPtr<FStructOnScope> EmptyDecorator;
+	return EmptyDecorator;
+}
+
 void URigVMPin::UpdateTypeInformationIfRequired() const
 {
 	if (CPPTypeObject == nullptr)
@@ -1753,6 +1773,16 @@ bool URigVMPin::CanLink(const URigVMPin* InSourcePin, const URigVMPin* InTargetP
 		return false;
 	}
 
+	if((InSourcePin->IsDecoratorPin() && InSourcePin->IsRootPin()) ||
+		(InTargetPin->IsDecoratorPin() && InTargetPin->IsRootPin()))
+	{
+		if(OutFailureReason)
+		{
+			*OutFailureReason = TEXT("Cannot add link to root decorator pins.");
+		}
+		return false;
+	}
+
 	URigVMNode* SourceNode = InSourcePin->GetNode();
 	URigVMNode* TargetNode = InTargetPin->GetNode();
 	if (SourceNode == TargetNode)
diff --git a/Engine/Plugins/Runtime/RigVM/Source/RigVMDeveloper/Public/RigVMModel/RigVMController.h b/Engine/Plugins/Runtime/RigVM/Source/RigVMDeveloper/Public/RigVMModel/RigVMController.h
--- a/Engine/Plugins/Runtime/RigVM/Source/RigVMDeveloper/Public/RigVMModel/RigVMController.h
+++ b/Engine/Plugins/Runtime/RigVM/Source/RigVMDeveloper/Public/RigVMModel/RigVMController.h
@@ -259,6 +259,7 @@ public:
 	TArray<FString> GeneratePythonCommands();
 
 	TArray<FString> GetAddNodePythonCommands(URigVMNode* Node) const;
+	TArray<FString> GetAddDecoratorPythonCommands(URigVMNode* Node, const FName& DecoratorName) const;
 
 #if WITH_EDITOR
 	// Note: The functions below are scoped with WITH_EDITOR since we are considering
@@ -497,6 +498,14 @@ public:
 	UFUNCTION(BlueprintCallable, Category = RigVMController)
 	URigVMInvokeEntryNode* AddInvokeEntryNode(const FName& InEntryName, const FVector2D& InPosition = FVector2D::ZeroVector, const FString& InNodeName = TEXT(""), bool bSetupUndoRedo = true, bool bPrintPythonCommand = false);
 
+	// Adds a decorator to a node
+	UFUNCTION(BlueprintCallable, Category = RigVMController)
+	FName AddDecorator(const FName& InNodeName, const FName& InDecoratorTypeObjectPath, const FName& InDecoratorName = NAME_None, const FString& InDefaultValue = TEXT(""), int32 InPinIndex = -1, bool bSetupUndoRedo = true, bool bPrintPythonCommand = false);
+
+	// Removes a decorator from a node
+	UFUNCTION(BlueprintCallable, Category = RigVMController)
+	bool RemoveDecorator(const FName& InNodeName, const FName& InDecoratorName, bool bSetupUndoRedo = true, bool bPrintPythonCommand = false);
+
 	// Un-does the last action on the stack.
 	// Note: This should really only be used for unit tests,
 	// use the GEditor's main Undo method instead.
@@ -1048,6 +1057,7 @@ public:
 		bool bIsExpanded;
 		bool bIsConstant;
 		bool bIsDynamicArray;
+		bool bIsDecorator;
 		TArray<int32> SubPins;
 
 		friend uint32 GetTypeHash(const FPinInfo& InPin);
@@ -1117,6 +1127,8 @@ private:
 	URigVMInjectionInfo* InjectNodeIntoPin(URigVMPin* InPin, bool bAsInput, const FName& InInputPinName, const FName& InOutputPinName, bool bSetupUndoRedo = true);
 	URigVMNode* EjectNodeFromPin(URigVMPin* InPin, bool bSetupUndoRedo = true, bool bPrintPythonCommands = false);
 	bool EjectAllInjectedNodes(URigVMNode* InNode, bool bSetupUndoRedo = true, bool bPrintPythonCommands = false);
+	FName AddDecorator(URigVMNode* InNode, UScriptStruct* InDecoratorScriptStruct, const FName& InDecoratorName, const FString& InDefaultValue, int32 InPinIndex = -1, bool bSetupUndoRedo = true);
+	bool RemoveDecorator(URigVMNode* InNode, const FName& InDecoratorName, bool bSetupUndoRedo = true);
 
 public:
 
diff --git a/Engine/Plugins/Runtime/RigVM/Source/RigVMDeveloper/Public/RigVMModel/RigVMControllerActions.h b/Engine/Plugins/Runtime/RigVM/Source/RigVMDeveloper/Public/RigVMModel/RigVMControllerActions.h
--- a/Engine/Plugins/Runtime/RigVM/Source/RigVMDeveloper/Public/RigVMModel/RigVMControllerActions.h
+++ b/Engine/Plugins/Runtime/RigVM/Source/RigVMDeveloper/Public/RigVMModel/RigVMControllerActions.h
@@ -1443,3 +1443,54 @@ public:
 	virtual bool Undo() override;
 	virtual bool Redo() override;
 };
+
+/**
+ * An action to add a decorator to a node
+ */
+USTRUCT()
+struct FRigVMAddDecoratorAction : public FRigVMBaseAction
+{
+	GENERATED_BODY()
+
+public:
+
+	FRigVMAddDecoratorAction();
+	FRigVMAddDecoratorAction(URigVMController* InController, const URigVMNode* InNode, const FName& InDecoratorName, const UScriptStruct* InDecoratorScriptStruct, const FString& InDecoratorDefault, int32 InPinIndex);
+	virtual ~FRigVMAddDecoratorAction() {};
+	virtual UScriptStruct* GetScriptStruct() const override { return FRigVMAddDecoratorAction::StaticStruct(); }
+	virtual bool Undo() override;
+	virtual bool Redo() override;
+
+	UPROPERTY()
+	FName NodeName;
+
+	UPROPERTY()
+	FName DecoratorName;
+
+	UPROPERTY()
+	FString ScriptStructPath;
+
+	UPROPERTY()
+	FString DecoratorDefault;
+
+	UPROPERTY()
+	int32 PinIndex;
+};
+
+/**
+ * An action to remove a decorator from a node
+ */
+USTRUCT()
+struct FRigVMRemoveDecoratorAction : public FRigVMAddDecoratorAction
+{
+	GENERATED_BODY()
+
+public:
+
+	FRigVMRemoveDecoratorAction();
+	FRigVMRemoveDecoratorAction(URigVMController* InController, const URigVMNode* InNode, const FName& InDecoratorName, const UScriptStruct* InDecoratorScriptStruct, const FString& InDecoratorDefault, int32 InPinIndex);
+	virtual ~FRigVMRemoveDecoratorAction() {};
+	virtual UScriptStruct* GetScriptStruct() const override { return FRigVMRemoveDecoratorAction::StaticStruct(); }
+	virtual bool Undo() override;
+	virtual bool Redo() override;
+};
\ No newline at end of file
diff --git a/Engine/Plugins/Runtime/RigVM/Source/RigVMDeveloper/Public/RigVMModel/RigVMNode.h b/Engine/Plugins/Runtime/RigVM/Source/RigVMDeveloper/Public/RigVMModel/RigVMNode.h
--- a/Engine/Plugins/Runtime/RigVM/Source/RigVMDeveloper/Public/RigVMModel/RigVMNode.h
+++ b/Engine/Plugins/Runtime/RigVM/Source/RigVMDeveloper/Public/RigVMModel/RigVMNode.h
@@ -5,8 +5,10 @@
 #include "RigVMPin.h"
 #include "RigVMCore/RigVM.h"
 #include "RigVMCore/RigVMStruct.h"
+#include "RigVMCore/RigVMDecorator.h"
 #include "RigVMCore/RigVMUserWorkflow.h"
 #include "RigVMCore/RigVMExecuteContext.h"
+#include "UObject/StructOnScope.h"
 #include "RigVMNode.generated.h"
 
 class URigVMGraph;
@@ -289,6 +291,24 @@ public:
 
 	virtual uint32 GetStructureHash() const;
 
+	UFUNCTION(BlueprintPure, Category = RigVMNode)
+	TArray<URigVMPin*> GetDecoratorPins() const;
+
+	const TArray<FString>& GetDecoratorNames() const { return DecoratorRootPinNames; }
+
+	UFUNCTION(BlueprintPure, Category = RigVMNode)
+	bool IsDecoratorPin(FName InName) const;
+
+	bool IsDecoratorPin(const URigVMPin* InDecoratorPin) const;
+
+	URigVMPin* FindDecorator(const FName& InName) const;
+
+	URigVMPin* FindDecorator(const URigVMPin* InDecoratorPin) const;
+
+	TSharedPtr<FStructOnScope> GetDecoratorInstance(const FName& InName, bool bUseDefaultValueFromPin = true) const;
+
+	TSharedPtr<FStructOnScope> GetDecoratorInstance(const URigVMPin* InDecoratorPin, bool bUseDefaultValueFromPin = true) const;
+
 private:
 
 	static const FString NodeColorName;
@@ -303,6 +323,7 @@ protected:
 	virtual FText GetToolTipTextForPin(const URigVMPin* InPin) const;
 	virtual bool AllowsLinksOn(const URigVMPin* InPin) const { return true; }
 	virtual bool ShouldInputPinComputeLazily(const URigVMPin* InPin) const { return false; }
+	void UpdateDecoratorRootPinNames();
 
 	UPROPERTY()
 	FString NodeTitle;
@@ -325,6 +346,9 @@ protected:
 	UPROPERTY(transient)
 	bool bHaltedAtThisNode;
 
+	UPROPERTY()
+	TArray<FString> DecoratorRootPinNames;
+
 private:
 
 	UPROPERTY()
diff --git a/Engine/Plugins/Runtime/RigVM/Source/RigVMDeveloper/Public/RigVMModel/RigVMPin.h b/Engine/Plugins/Runtime/RigVM/Source/RigVMDeveloper/Public/RigVMModel/RigVMPin.h
--- a/Engine/Plugins/Runtime/RigVM/Source/RigVMDeveloper/Public/RigVMModel/RigVMPin.h
+++ b/Engine/Plugins/Runtime/RigVM/Source/RigVMDeveloper/Public/RigVMModel/RigVMPin.h
@@ -9,6 +9,8 @@
 #include "RigVMCore/RigVMByteCode.h"
 #include "RigVMCore/RigVMTemplate.h"
 #include "RigVMCompiler/RigVMASTProxy.h"
+#include "UObject/StructOnScope.h"
+#include <RigVMCore/RigVMDecorator.h>
 #include "RigVMPin.generated.h"
 
 class URigVMGraph;
@@ -457,6 +459,13 @@ public:
 
 	uint32 GetStructureHash() const;
 
+	// Returns true if this pin represents a decorator 
+	UFUNCTION(BlueprintPure, Category = RigVMPin)
+	bool IsDecoratorPin() const;
+
+	// Returns the decorator backing up this pin
+	TSharedPtr<FStructOnScope> GetDecoratorInstance(bool bUseDefaultValueFromPin = true) const;
+
 private:
 
 	void UpdateTypeInformationIfRequired() const;
