RigVM: Added Int32 to Name node

diff --git a/Engine/Plugins/Runtime/RigVM/Source/RigVM/Private/RigVMFunctions/Math/RigVMFunction_MathInt.cpp b/Engine/Plugins/Runtime/RigVM/Source/RigVM/Private/RigVMFunctions/Math/RigVMFunction_MathInt.cpp
--- a/Engine/Plugins/Runtime/RigVM/Source/RigVM/Private/RigVMFunctions/Math/RigVMFunction_MathInt.cpp
+++ b/Engine/Plugins/Runtime/RigVM/Source/RigVM/Private/RigVMFunctions/Math/RigVMFunction_MathInt.cpp
@@ -2,6 +2,7 @@
 
 #include "RigVMFunctions/Math/RigVMFunction_MathInt.h"
 #include "RigVMFunctions/RigVMDispatch_Core.h"
+#include "RigVMCore/RigVMStructTest.h"
 
 #include UE_INLINE_GENERATED_CPP_BY_NAME(RigVMFunction_MathInt)
 
@@ -157,12 +158,31 @@ FRigVMFunction_MathIntArrayAverage_Execute()
 	}
 }
 
+FRigVMFunction_MathIntToString_Execute()
+{
+	Result = FString::FromInt(Number);
+	if(Number >= 0)
+	{
+		while(Result.Len() < FMath::Min<int32>(8, PaddedSize))
+		{
+			static const FString Zero = FString::FromInt(0);
+			Result = Zero + Result;
+		}
+	}
+}
 
+FRigVMFunction_MathIntToName_Execute()
+{
+	FString String;
+	FRigVMFunction_MathIntToString::StaticExecute(ExecuteContext, Number, PaddedSize, String);
+	Result = *String;
+}
 
-
-
-
-
-
-
-
+IMPLEMENT_RIGVMSTRUCT_AUTOMATION_TEST(FRigVMFunction_MathIntToName)
+{
+	Unit.Number = 13;
+	Unit.PaddedSize = 5;
+	Execute();
+	AddErrorIfFalse(Unit.Result == TEXT("00013"), TEXT("unexpected result"));
+	return true;
+}
\ No newline at end of file
diff --git a/Engine/Plugins/Runtime/RigVM/Source/RigVM/Public/RigVMFunctions/Math/RigVMFunction_MathInt.h b/Engine/Plugins/Runtime/RigVM/Source/RigVM/Public/RigVMFunctions/Math/RigVMFunction_MathInt.h
--- a/Engine/Plugins/Runtime/RigVM/Source/RigVM/Public/RigVMFunctions/Math/RigVMFunction_MathInt.h
+++ b/Engine/Plugins/Runtime/RigVM/Source/RigVM/Public/RigVMFunctions/Math/RigVMFunction_MathInt.h
@@ -533,4 +533,66 @@ struct RIGVM_API FRigVMFunction_MathIntArrayAverage : public FRigVMFunction_Math
 
 	UPROPERTY(meta = (Output))
 	int32 Average;
+};
+
+/**
+ * Converts an integer to a string
+ */
+USTRUCT(meta = (DisplayName = "Int to String", TemplateName="Int to Name"))
+struct RIGVM_API FRigVMFunction_MathIntToString : public FRigVMFunction_MathIntBase
+{
+	GENERATED_BODY()
+
+	FRigVMFunction_MathIntToString()
+	{
+		Number = PaddedSize = 0;
+	}
+
+	/** Execute logic for this function */
+	RIGVM_METHOD()
+	virtual void Execute() override;
+
+	UPROPERTY(meta = (Input))
+	int32 Number;
+
+	/*
+	 * For positive numbers you can pad the result to a required with
+	 * so rather than '13' return '00013' for a padded size of 5.
+	 */
+	UPROPERTY(meta = (Input))
+	int32 PaddedSize;
+
+	UPROPERTY(meta = (Output))
+	FString Result;
+};
+
+/**
+ * Converts an integer to a name
+ */
+USTRUCT(meta = (DisplayName = "Int to Name", TemplateName="Int to Name"))
+struct RIGVM_API FRigVMFunction_MathIntToName : public FRigVMFunction_MathIntBase
+{
+	GENERATED_BODY()
+
+	FRigVMFunction_MathIntToName()
+	{
+		Number = PaddedSize = 0;
+	}
+
+	/** Execute logic for this function */
+	RIGVM_METHOD()
+	virtual void Execute() override;
+
+	UPROPERTY(meta = (Input))
+	int32 Number;
+
+	/*
+	 * For positive numbers you can pad the result to a required with
+	 * so rather than '13' return '00013' for a padded size of 5.
+	 */
+	UPROPERTY(meta = (Input))
+	int32 PaddedSize;
+
+	UPROPERTY(meta = (Output))
+	FName Result;
 };
\ No newline at end of file
