RigVM: First steps towards rebasing ControlRigBlueprint

diff --git a/Engine/Plugins/Experimental/AnimNext/Source/AnimNextUncookedOnly/Internal/Graph/AnimNextGraph_EdGraphSchema.h b/Engine/Plugins/Experimental/AnimNext/Source/AnimNextUncookedOnly/Internal/Graph/AnimNextGraph_EdGraphSchema.h
--- a/Engine/Plugins/Experimental/AnimNext/Source/AnimNextUncookedOnly/Internal/Graph/AnimNextGraph_EdGraphSchema.h
+++ b/Engine/Plugins/Experimental/AnimNext/Source/AnimNextUncookedOnly/Internal/Graph/AnimNextGraph_EdGraphSchema.h
@@ -13,7 +13,7 @@ class UAnimNextGraph_EdGraphSchema : public UControlRigGraphSchema
 	GENERATED_BODY()
 
 	// UControlRigGraphSchema interface
-	virtual TSubclassOf<UControlRigGraphNode> GetGraphNodeClass() const override { return UAnimNextGraph_EdGraphNode::StaticClass(); }
+	virtual TSubclassOf<URigVMEdGraphNode> GetGraphNodeClass(const URigVMEdGraph* InGraph) const override { return UAnimNextGraph_EdGraphNode::StaticClass(); }
 
 	// UEdGraphSchema interface
 	virtual void TrySetDefaultValue(UEdGraphPin& InPin, const FString& InNewDefaultValue, bool bMarkAsModified) const override;
diff --git a/Engine/Plugins/Experimental/AnimNext/Source/AnimNextUncookedOnly/Internal/Graph/AnimNextGraph_EditorData.h b/Engine/Plugins/Experimental/AnimNext/Source/AnimNextUncookedOnly/Internal/Graph/AnimNextGraph_EditorData.h
--- a/Engine/Plugins/Experimental/AnimNext/Source/AnimNextUncookedOnly/Internal/Graph/AnimNextGraph_EditorData.h
+++ b/Engine/Plugins/Experimental/AnimNext/Source/AnimNextUncookedOnly/Internal/Graph/AnimNextGraph_EditorData.h
@@ -6,6 +6,7 @@
 #include "RigVMModel/RigVMGraph.h"
 #include "AnimNextGraph_EdGraph.h"
 #include "RigVMCore/RigVMGraphFunctionHost.h"
+#include "RigVMBlueprint.h"
 #include "AnimNextGraph_EditorData.generated.h"
 
 class UAnimNextGraph;
@@ -110,7 +111,7 @@ class UAnimNextGraph_EditorData : public UObject, public IRigVMClientHost, publi
 	TMap<TObjectPtr<URigVMGraph>, TObjectPtr<URigVMController>> Controllers;
 
 	UPROPERTY(EditAnywhere, Category = "User Interface")
-	FRigGraphDisplaySettings RigGraphDisplaySettings;
+	FRigVMEdGraphDisplaySettings RigGraphDisplaySettings;
 
 	UPROPERTY(EditAnywhere, Category = "VM")
 	FRigVMRuntimeSettings VMRuntimeSettings;
@@ -118,9 +119,6 @@ class UAnimNextGraph_EditorData : public UObject, public IRigVMClientHost, publi
 	UPROPERTY(EditAnywhere, BlueprintReadWrite, Category = "VM", meta = (AllowPrivateAccess = "true"))
 	FRigVMCompileSettings VMCompileSettings;
 
-	UPROPERTY(EditAnywhere, Category = "Python Log Settings")
-	FControlRigPythonSettings PythonLogSettings;
-
 	UPROPERTY(transient, DuplicateTransient)
 	TMap<FString, FRigVMOperand> PinToOperandMap;
 
@@ -132,7 +130,7 @@ class UAnimNextGraph_EditorData : public UObject, public IRigVMClientHost, publi
 	
 	FCompilerResultsLog CompileLog;
 
-	FOnVMCompiledEvent VMCompiledEvent;
+	FOnRigVMCompiledEvent VMCompiledEvent;
 	FRigVMGraphModifiedEvent ModifiedEvent;
 
 	bool bAutoRecompileVM = true;
diff --git a/Engine/Plugins/Experimental/AnimNext/Source/AnimNextUncookedOnly/Internal/Param/AnimNextParameterBlock_EdGraphSchema.h b/Engine/Plugins/Experimental/AnimNext/Source/AnimNextUncookedOnly/Internal/Param/AnimNextParameterBlock_EdGraphSchema.h
--- a/Engine/Plugins/Experimental/AnimNext/Source/AnimNextUncookedOnly/Internal/Param/AnimNextParameterBlock_EdGraphSchema.h
+++ b/Engine/Plugins/Experimental/AnimNext/Source/AnimNextUncookedOnly/Internal/Param/AnimNextParameterBlock_EdGraphSchema.h
@@ -13,7 +13,7 @@ class UAnimNextParameterBlock_EdGraphSchema : public UControlRigGraphSchema
 	GENERATED_BODY()
 
 	// UControlRigGraphSchema interface
-	virtual TSubclassOf<UControlRigGraphNode> GetGraphNodeClass() const override { return UAnimNextParameterBlock_EdGraphNode::StaticClass(); }
+	virtual TSubclassOf<URigVMEdGraphNode> GetGraphNodeClass(const URigVMEdGraph* InGraph) const override { return UAnimNextParameterBlock_EdGraphNode::StaticClass(); }
 
 	// UEdGraphSchema interface
 	virtual void TrySetDefaultValue(UEdGraphPin& InPin, const FString& InNewDefaultValue, bool bMarkAsModified) const override;
diff --git a/Engine/Plugins/Experimental/AnimNext/Source/AnimNextUncookedOnly/Internal/Param/AnimNextParameterBlock_EditorData.h b/Engine/Plugins/Experimental/AnimNext/Source/AnimNextUncookedOnly/Internal/Param/AnimNextParameterBlock_EditorData.h
--- a/Engine/Plugins/Experimental/AnimNext/Source/AnimNextUncookedOnly/Internal/Param/AnimNextParameterBlock_EditorData.h
+++ b/Engine/Plugins/Experimental/AnimNext/Source/AnimNextUncookedOnly/Internal/Param/AnimNextParameterBlock_EditorData.h
@@ -6,6 +6,7 @@
 #include "RigVMModel/RigVMGraph.h"
 #include "AnimNextParameterBlock_EdGraph.h"
 #include "RigVMCore/RigVMGraphFunctionHost.h"
+#include "RigVMBlueprint.h"
 #include "AnimNextParameterBlock_EditorData.generated.h"
 
 class UAnimNextParameterBlock;
@@ -179,7 +180,7 @@ class UAnimNextParameterBlock_EditorData : public UObject, public IRigVMClientHo
 	FRigVMGraphFunctionStore GraphFunctionStore;
 
 	UPROPERTY(EditAnywhere, Category = "User Interface")
-	FRigGraphDisplaySettings RigGraphDisplaySettings;
+	FRigVMEdGraphDisplaySettings RigGraphDisplaySettings;
 
 	UPROPERTY(EditAnywhere, Category = "VM")
 	FRigVMRuntimeSettings VMRuntimeSettings;
@@ -187,9 +188,6 @@ class UAnimNextParameterBlock_EditorData : public UObject, public IRigVMClientHo
 	UPROPERTY(EditAnywhere, BlueprintReadWrite, Category = "VM", meta = (AllowPrivateAccess = "true"))
 	FRigVMCompileSettings VMCompileSettings;
 
-	UPROPERTY(EditAnywhere, Category = "Python Log Settings")
-	FControlRigPythonSettings PythonLogSettings;
-
 	UPROPERTY(transient, DuplicateTransient)
 	TMap<FString, FRigVMOperand> PinToOperandMap;
 
@@ -214,7 +212,7 @@ class UAnimNextParameterBlock_EditorData : public UObject, public IRigVMClientHo
 	
 	FCompilerResultsLog CompileLog;
 
-	FOnVMCompiledEvent RigVMCompiledEvent;
+	FOnRigVMCompiledEvent RigVMCompiledEvent;
 	FRigVMGraphModifiedEvent RigVMGraphModifiedEvent;
 
 	// Delegate to subscribe to modifications to this block
diff --git a/Engine/Plugins/Experimental/Animation/PoseCorrectives/Source/PoseCorrectivesEditor/PoseCorrectivesEditor.Build.cs b/Engine/Plugins/Experimental/Animation/PoseCorrectives/Source/PoseCorrectivesEditor/PoseCorrectivesEditor.Build.cs
--- a/Engine/Plugins/Experimental/Animation/PoseCorrectives/Source/PoseCorrectivesEditor/PoseCorrectivesEditor.Build.cs
+++ b/Engine/Plugins/Experimental/Animation/PoseCorrectives/Source/PoseCorrectivesEditor/PoseCorrectivesEditor.Build.cs
@@ -23,6 +23,7 @@ public class PoseCorrectivesEditor : ModuleRules
 				"ControlRig",
 				"ControlRigDeveloper",
 				"ControlRigEditor",
+				"RigVMDeveloper",
 			}
 			);
 			
diff --git a/Engine/Plugins/Experimental/Animation/PoseCorrectives/Source/PoseCorrectivesEditor/Private/PoseCorrectivesEditorController.cpp b/Engine/Plugins/Experimental/Animation/PoseCorrectives/Source/PoseCorrectivesEditor/Private/PoseCorrectivesEditorController.cpp
--- a/Engine/Plugins/Experimental/Animation/PoseCorrectives/Source/PoseCorrectivesEditor/Private/PoseCorrectivesEditorController.cpp
+++ b/Engine/Plugins/Experimental/Animation/PoseCorrectives/Source/PoseCorrectivesEditor/Private/PoseCorrectivesEditorController.cpp
@@ -40,7 +40,7 @@ void FPoseCorrectivesEditorController::InitializeTargetControlRigBP()
 	TObjectPtr<const UControlRigBlueprint> ControlRigBP = Cast<UControlRigBlueprint>(Asset->ControlRigBlueprint);	
 	if (ControlRigBP)
 	{
-		if (const UClass* GeneratedClass = ControlRigBP->GetControlRigBlueprintGeneratedClass())
+		if (const UClass* GeneratedClass = ControlRigBP->GetRigVMBlueprintGeneratedClass())
 		{
 			if (ControlRig)
 			{
diff --git a/Engine/Plugins/Importers/USDImporter/Source/USDStage/Private/USDLevelSequenceHelper.cpp b/Engine/Plugins/Importers/USDImporter/Source/USDStage/Private/USDLevelSequenceHelper.cpp
--- a/Engine/Plugins/Importers/USDImporter/Source/USDStage/Private/USDLevelSequenceHelper.cpp
+++ b/Engine/Plugins/Importers/USDImporter/Source/USDStage/Private/USDLevelSequenceHelper.cpp
@@ -78,6 +78,7 @@
 #include "UObject/Package.h"
 #include "UObject/UObjectGlobals.h"
 #include "Widgets/Notifications/SNotificationList.h"
+#include "RigVMBlueprintGeneratedClass.h"
 
 #if WITH_EDITOR
 #include "ControlRigBlueprint.h"
@@ -2345,7 +2346,7 @@ void FUsdLevelSequenceHelperImpl::UpdateControlRigTracks(UUsdPrimTwin& PrimTwin)
 
 			if (UControlRigBlueprint* BP = Cast<UControlRigBlueprint>(FSoftObjectPath(ControlRigBPPath).TryLoad()))
 			{
-				ControlRigClass = BP->GetControlRigClass();
+				ControlRigClass = BP->GetRigVMBlueprintGeneratedClass();
 			}
 		}
 
diff --git a/Engine/Plugins/Runtime/RigVM/Source/RigVMDeveloper/Private/EdGraph/RigVMEdGraph.cpp b/Engine/Plugins/Runtime/RigVM/Source/RigVMDeveloper/Private/EdGraph/RigVMEdGraph.cpp
--- a/Engine/Plugins/Runtime/RigVM/Source/RigVMDeveloper/Private/EdGraph/RigVMEdGraph.cpp
+++ b/Engine/Plugins/Runtime/RigVM/Source/RigVMDeveloper/Private/EdGraph/RigVMEdGraph.cpp
@@ -22,7 +22,7 @@
 #include "Framework/Notifications/NotificationManager.h"
 #endif
 
-#define LOCTEXT_NAMESPACE "ControlRigGraph"
+#define LOCTEXT_NAMESPACE "RigVMEdGraph"
 
 URigVMEdGraph::URigVMEdGraph()
 {
@@ -97,7 +97,7 @@ void URigVMEdGraph::HandleRigVMGraphRenamed(const FString& InOldNodePath, const
 	}
 }
 
-void URigVMEdGraph::Initialize(URigVMBlueprint* InBlueprint)
+void URigVMEdGraph::InitializeFromBlueprint(URigVMBlueprint* InBlueprint)
 {
 	DECLARE_SCOPE_HIERARCHICAL_COUNTER_FUNC()
 
@@ -132,24 +132,29 @@ void URigVMEdGraph::Serialize(FArchive& Ar)
 #if WITH_EDITOR
 
 void URigVMEdGraph::HandleModifiedEvent(ERigVMGraphNotifType InNotifType, URigVMGraph* InGraph, UObject* InSubject)
+{
+	(void)HandleModifiedEvent_Internal(InNotifType, InGraph, InSubject);
+}
+
+bool URigVMEdGraph::HandleModifiedEvent_Internal(ERigVMGraphNotifType InNotifType, URigVMGraph* InGraph, UObject* InSubject)
 {
 	DECLARE_SCOPE_HIERARCHICAL_COUNTER_FUNC()
 
 	if (bSuspendModelNotifications)
 	{
-		return;
+		return false;
 	}
 
 	// only make sure to receive notifs for this graph - unless
 	// we are on a template graph (used by node spawners)
 	if (GetModel() != InGraph && TemplateController == nullptr)
 	{
-		return;
+		return false;
 	}
 
-	if(URigVMEdGraphSchema* ControlRigSchema = (URigVMEdGraphSchema*)GetRigVMEdGraphSchema())
+	if(URigVMEdGraphSchema* EdGraphSchema = (URigVMEdGraphSchema*)GetRigVMEdGraphSchema())
 	{
-		ControlRigSchema->HandleModifiedEvent(InNotifType, InGraph, InSubject);
+		EdGraphSchema->HandleModifiedEvent(InNotifType, InGraph, InSubject);
 	}
 
 	if(TemplateController)
@@ -165,7 +170,7 @@ void URigVMEdGraph::HandleModifiedEvent(ERigVMGraphNotifType InNotifType, URigVM
 			case ERigVMGraphNotifType::VariableRenamed:
 			case ERigVMGraphNotifType::GraphChanged:
 			{
-				return;
+				return false;
 			}
 			default:
 			{
@@ -213,7 +218,7 @@ void URigVMEdGraph::HandleModifiedEvent(ERigVMGraphNotifType InNotifType, URigVM
 		{
 			if (bIsSelecting)
 			{
-				return;
+				return false;
 			}
 			TGuardValue<bool> SelectionGuard(bIsSelecting, true);
 
@@ -290,7 +295,7 @@ void URigVMEdGraph::HandleModifiedEvent(ERigVMGraphNotifType InNotifType, URigVM
 				}
 				else // struct, library, parameter + variable
 				{
-					URigVMEdGraphNode* NewNode = NewObject<URigVMEdGraphNode>(this, GetRigVMEdGraphSchema()->GetGraphNodeClass(), ModelNode->GetFName());
+					URigVMEdGraphNode* NewNode = NewObject<URigVMEdGraphNode>(this, GetRigVMEdGraphSchema()->GetGraphNodeClass(this), ModelNode->GetFName());
 					AddNode(NewNode, false, false);
 
 					NewNode->ModelNodePath = ModelNode->GetNodePath();
@@ -681,6 +686,8 @@ void URigVMEdGraph::HandleModifiedEvent(ERigVMGraphNotifType InNotifType, URigVM
 			break;
 		}
 	}
+
+	return true;
 }
 
 int32 URigVMEdGraph::GetInstructionIndex(const URigVMEdGraphNode* InNode, bool bAsInput)
@@ -911,6 +918,27 @@ const URigVMEdGraph* URigVMEdGraph::GetRootGraph() const
 	return this;
 }
 
+void URigVMEdGraph::AddNode(UEdGraphNode* NodeToAdd, bool bUserAction, bool bSelectNewNode)
+{
+	// Comments are added outside of the RigVMEditor, so we add here the node to the model
+	if (const UEdGraphNode_Comment* CommentNode = Cast<const UEdGraphNode_Comment>(NodeToAdd))
+	{
+		if (URigVMController* Controller = GetBlueprint()->GetOrCreateController(GetModel()))
+		{
+			if (GetModel()->FindNodeByName(NodeToAdd->GetFName()) == nullptr) // When recreating nodes at RebuildGraphFromModel, the model node already exists
+			{
+				TGuardValue<bool> BlueprintNotifGuard(GetBlueprint()->bSuspendModelNotificationsForOthers, true);
+				FVector2D NodePos(CommentNode->NodePosX, CommentNode->NodePosY);
+				FVector2D NodeSize(CommentNode->NodeWidth, CommentNode->NodeHeight);
+				FLinearColor NodeColor = CommentNode->CommentColor;
+				Controller->AddCommentNode(CommentNode->NodeComment, NodePos, NodeSize, NodeColor, CommentNode->GetName(), GIsTransacting == false, true);
+			}
+		}
+	}
+
+	Super::AddNode(NodeToAdd, bUserAction, bSelectNewNode);
+}
+
 void URigVMEdGraph::RemoveNode(UEdGraphNode* InNode)
 {
 	// Make sure EdGraph is not part of the transaction
@@ -936,7 +964,7 @@ void URigVMEdGraph::RemoveNode(UEdGraphNode* InNode)
 		static int32 DeletedIndex = FMath::Rand();
 		do
 		{
-			DeletedName = FString::Printf(TEXT("ControlRigGraph_%s_Deleted_%d"), *InNode->GetName(), DeletedIndex++); 
+			DeletedName = FString::Printf(TEXT("EdGraph_%s_Deleted_%d"), *InNode->GetName(), DeletedIndex++); 
 			ExistingObject = StaticFindObject(/*Class=*/ NULL, this, *DeletedName, true);						
 		}
 		while (ExistingObject);
diff --git a/Engine/Plugins/Runtime/RigVM/Source/RigVMDeveloper/Private/EdGraph/RigVMEdGraphNode.cpp b/Engine/Plugins/Runtime/RigVM/Source/RigVMDeveloper/Private/EdGraph/RigVMEdGraphNode.cpp
--- a/Engine/Plugins/Runtime/RigVM/Source/RigVMDeveloper/Private/EdGraph/RigVMEdGraphNode.cpp
+++ b/Engine/Plugins/Runtime/RigVM/Source/RigVMDeveloper/Private/EdGraph/RigVMEdGraphNode.cpp
@@ -198,17 +198,6 @@ void URigVMEdGraphNode::ReconstructNode_Internal(bool bForce)
 		{
 			return;
 		}
-
-		/*
-		 * todoooooooo
-		// if this node has been saved prior to our custom version,
-		// don't reset the node
-		int32 LinkerVersion = RigGraph->GetLinkerCustomVersion(FControlRigObjectVersion::GUID);
-		if (LinkerVersion < FControlRigObjectVersion::SwitchedToRigVM)
-		{
-			return;
-		}
-		*/
 	}
 
 #if WITH_EDITOR
diff --git a/Engine/Plugins/Runtime/RigVM/Source/RigVMDeveloper/Private/EdGraph/RigVMEdGraphSchema.cpp b/Engine/Plugins/Runtime/RigVM/Source/RigVMDeveloper/Private/EdGraph/RigVMEdGraphSchema.cpp
--- a/Engine/Plugins/Runtime/RigVM/Source/RigVMDeveloper/Private/EdGraph/RigVMEdGraphSchema.cpp
+++ b/Engine/Plugins/Runtime/RigVM/Source/RigVMDeveloper/Private/EdGraph/RigVMEdGraphSchema.cpp
@@ -49,7 +49,7 @@
 
 #define LOCTEXT_NAMESPACE "CRigVMGraphSchema"
 
-const FName URigVMEdGraphSchema::GraphName_RigVM(TEXT("RigVM Graph"));
+const FName URigVMEdGraphSchema::GraphName_RigVM(TEXT("RigVM"));
 
 FRigVMLocalVariableNameValidator::FRigVMLocalVariableNameValidator(const UBlueprint* Blueprint, const URigVMGraph* Graph, FName InExistingName)
 	: FStringSetNameValidator(InExistingName.ToString())
@@ -204,9 +204,9 @@ void FRigVMEdGraphSchemaAction_LocalVar::RenameVariable(const FName& NewName)
 
 bool FRigVMEdGraphSchemaAction_LocalVar::IsValidName(const FName& NewName, FText& OutErrorMessage) const
 {
-	if (const URigVMEdGraph* ControlRigGraph = Cast<URigVMEdGraph>(GetVariableScope()))
+	if (const URigVMEdGraph* EdGraphGraph = Cast<URigVMEdGraph>(GetVariableScope()))
 	{
-		FRigVMLocalVariableNameValidator NameValidator(ControlRigGraph->GetBlueprint(), ControlRigGraph->GetModel(), GetVariableName());
+		FRigVMLocalVariableNameValidator NameValidator(EdGraphGraph->GetBlueprint(), EdGraphGraph->GetModel(), GetVariableName());
 		const EValidatorResult Result = NameValidator.IsValid(NewName.ToString(), false);
 		if (Result != EValidatorResult::Ok && Result != EValidatorResult::ExistingName)
 		{
@@ -236,10 +236,10 @@ void FRigVMEdGraphSchemaAction_LocalVar::DeleteVariable()
 
 bool FRigVMEdGraphSchemaAction_LocalVar::IsVariableUsed()
 {
-	if (const URigVMEdGraph* ControlRigGraph = Cast<URigVMEdGraph>(GetVariableScope()))
+	if (const URigVMEdGraph* EdGraph = Cast<URigVMEdGraph>(GetVariableScope()))
 	{
 		const FString VarNameStr = GetVariableName().ToString();
-		for (URigVMNode* Node : ControlRigGraph->GetModel()->GetNodes())
+		for (URigVMNode* Node : EdGraph->GetModel()->GetNodes())
 		{
 			if (const URigVMVariableNode* VarNode = Cast<URigVMVariableNode>(Node))
 			{
@@ -969,9 +969,9 @@ TSharedPtr<INameValidatorInterface> URigVMEdGraphSchema::GetNameValidator(const
 {
 	if (ActionTypeId == FRigVMEdGraphSchemaAction_LocalVar::StaticGetTypeId())
 	{
-		if (const URigVMEdGraph* ControlRigGraph = Cast<URigVMEdGraph>(ValidationScope))
+		if (const URigVMEdGraph* EdGraph = Cast<URigVMEdGraph>(ValidationScope))
 		{
-			if (const URigVMGraph* Graph = ControlRigGraph->GetModel())
+			if (const URigVMGraph* Graph = EdGraph->GetModel())
 			{
 				return MakeShareable(new FRigVMLocalVariableNameValidator(BlueprintObj, Graph, OriginalName));
 			}
@@ -988,6 +988,8 @@ bool URigVMEdGraphSchema::SupportsPinType(const UScriptStruct* ScriptStruct) con
 		return false;
 	}
 
+	// todo: Validate ExecuteContext structs to match URigVMBlueprint::GetExecuteContextStruct()
+
 	for (TFieldIterator<FProperty> It(ScriptStruct); It; ++It)
 	{
 		FProperty* Property = *It;
@@ -1698,7 +1700,6 @@ bool URigVMEdGraphSchema::TryToGetChildEvents(const UEdGraph* Graph, const int32
 			return true;
 		}
 	}
-
 	return false;
 }
 
@@ -1838,11 +1839,11 @@ URigVMEdGraphNode* URigVMEdGraphSchema::CreateGraphNode(URigVMEdGraph* InGraph,
 {
 	const bool bSelectNewNode = true;
 	FGraphNodeCreator<URigVMEdGraphNode> GraphNodeCreator(*InGraph);
-	URigVMEdGraphNode* ControlRigGraphNode = GraphNodeCreator.CreateNode(bSelectNewNode);
-	ControlRigGraphNode->ModelNodePath = InPropertyName.ToString();
+	URigVMEdGraphNode* EdGraphNode = GraphNodeCreator.CreateNode(bSelectNewNode, GetGraphNodeClass(InGraph));
+	EdGraphNode->ModelNodePath = InPropertyName.ToString();
 	GraphNodeCreator.Finalize();
 
-	return ControlRigGraphNode;
+	return EdGraphNode;
 }
 
 void URigVMEdGraphSchema::TrySetDefaultValue(UEdGraphPin& InPin, const FString& InNewDefaultValue, bool bMarkAsModified) const
@@ -2395,6 +2396,13 @@ void URigVMEdGraphSchema::HandleModifiedEvent(ERigVMGraphNotifType InNotifType,
 	}
 }
 
+TSubclassOf<URigVMEdGraphNode> URigVMEdGraphSchema::GetGraphNodeClass(const URigVMEdGraph* InGraph) const
+{
+	const UBlueprint* Blueprint = FBlueprintEditorUtils::FindBlueprintForGraph(InGraph);
+	const URigVMBlueprint* RigBlueprint = CastChecked<URigVMBlueprint>(Blueprint);
+	return RigBlueprint->GetRigVMEdGraphNodeClass();
+}
+
 bool URigVMEdGraphSchema::IsRigVMDefaultEvent(const FName& InEventName) const
 {
 	return false;
diff --git a/Engine/Plugins/Runtime/RigVM/Source/RigVMDeveloper/Private/RigVMBlueprint.cpp b/Engine/Plugins/Runtime/RigVM/Source/RigVMDeveloper/Private/RigVMBlueprint.cpp
--- a/Engine/Plugins/Runtime/RigVM/Source/RigVMDeveloper/Private/RigVMBlueprint.cpp
+++ b/Engine/Plugins/Runtime/RigVM/Source/RigVMDeveloper/Private/RigVMBlueprint.cpp
@@ -68,6 +68,10 @@ static TArray<UClass*> GetClassObjectsInPackage(UPackage* InPackage)
 FSoftObjectPath URigVMBlueprint::PreDuplicateAssetPath;
 FSoftObjectPath URigVMBlueprint::PreDuplicateHostPath;
 
+URigVMBlueprint::URigVMBlueprint()
+{
+}
+
 URigVMBlueprint::URigVMBlueprint(const FObjectInitializer& ObjectInitializer)
 {
 	bSuspendModelNotificationsForSelf = false;
@@ -86,24 +90,6 @@ URigVMBlueprint::URigVMBlueprint(const FObjectInitializer& ObjectInitializer)
 	bIsCompiling = false;
 	VMRecompilationBracket = 0;
 
-	RigVMClient.SetOuterClientHost(this, GET_MEMBER_NAME_CHECKED(URigVMBlueprint, RigVMClient));
-	RigVMClient.SetSchemaClass(GetRigVMSchemaClass());
-	{
-		TGuardValue<bool> DisableClientNotifs(RigVMClient.bSuspendNotifications, true);
-		RigVMClient.GetOrCreateFunctionLibrary(false, &ObjectInitializer, false);
-		RigVMClient.AddModel(FRigVMClient::RigVMModelPrefix, false, &ObjectInitializer, false);
-	}
-	RigVMClient.SetExecuteContextStruct(GetRigVMExecuteContextStruct());
-	
-	FunctionLibraryEdGraph = Cast<URigVMEdGraph>(ObjectInitializer.CreateDefaultSubobject(this, TEXT("RigVMFunctionLibraryEdGraph"), GetRigVMEdGraphClass(), GetRigVMEdGraphClass()));
-	FunctionLibraryEdGraph->Schema = GetRigVMEdGraphSchemaClass();
-	FunctionLibraryEdGraph->bAllowRenaming = 0;
-	FunctionLibraryEdGraph->bEditable = 0;
-	FunctionLibraryEdGraph->bAllowDeletion = 0;
-	FunctionLibraryEdGraph->bIsFunctionDefinition = false;
-	FunctionLibraryEdGraph->ModelNodePath = RigVMClient.GetFunctionLibrary()->GetNodePath();
-	FunctionLibraryEdGraph->Initialize(this);
-
 	bUpdatingExternalVariables = false;
 	
 	bDirtyDuringLoad = false;
@@ -118,10 +104,44 @@ URigVMBlueprint::URigVMBlueprint(const FObjectInitializer& ObjectInitializer)
 	CompileLog.bLogDetailedResults = false;
 	CompileLog.EventDisplayThresholdMs = false;
 #endif
+
+	if(GetClass() == URigVMBlueprint::StaticClass())
+	{
+		CommonInitialization(ObjectInitializer);
+	}
 }
 
-URigVMBlueprint::URigVMBlueprint()
+void URigVMBlueprint::CommonInitialization(const FObjectInitializer& ObjectInitializer)
 {
+	// guard against this running multiple times
+	check(GetRigVMClient()->GetSchema() == nullptr);
+	
+	RigVMClient.SetSchemaClass(GetRigVMSchemaClass());
+	RigVMClient.SetExecuteContextStruct(GetRigVMExecuteContextStruct());
+
+	for(UEdGraph* UberGraph : UbergraphPages)
+	{
+		if(URigVMEdGraph* EdGraph = Cast<URigVMEdGraph>(UberGraph))
+		{
+			EdGraph->Schema = GetRigVMEdGraphSchemaClass();
+		}
+	}
+
+	RigVMClient.SetOuterClientHost(this, GET_MEMBER_NAME_CHECKED(URigVMBlueprint, RigVMClient));
+	{
+		TGuardValue<bool> DisableClientNotifs(RigVMClient.bSuspendNotifications, true);
+		RigVMClient.GetOrCreateFunctionLibrary(false, &ObjectInitializer, false);
+		RigVMClient.AddModel(FRigVMClient::RigVMModelPrefix, false, &ObjectInitializer, false);
+	}
+
+ 	FunctionLibraryEdGraph = Cast<URigVMEdGraph>(CreateDefaultSubobject(TEXT("RigVMFunctionLibraryEdGraph"), GetRigVMEdGraphClass(), GetRigVMEdGraphClass(), true, true));
+	FunctionLibraryEdGraph->Schema = GetRigVMEdGraphSchemaClass();
+	FunctionLibraryEdGraph->bAllowRenaming = 0;
+	FunctionLibraryEdGraph->bEditable = 0;
+	FunctionLibraryEdGraph->bAllowDeletion = 0;
+	FunctionLibraryEdGraph->bIsFunctionDefinition = false;
+	FunctionLibraryEdGraph->ModelNodePath = RigVMClient.GetFunctionLibrary()->GetNodePath();
+	FunctionLibraryEdGraph->InitializeFromBlueprint(this);
 }
 
 void URigVMBlueprint::InitializeModelIfRequired(bool bRecompileVM)
@@ -146,7 +166,7 @@ void URigVMBlueprint::InitializeModelIfRequired(bool bRecompileVM)
 					bRecompileRequired = true;
 				}
 
-				Graph->Initialize(this);
+				Graph->InitializeFromBlueprint(this);
 			}
 		}
 
@@ -155,7 +175,7 @@ void URigVMBlueprint::InitializeModelIfRequired(bool bRecompileVM)
 			RecompileVM();
 		}
 
-		FunctionLibraryEdGraph->Initialize(this);
+		FunctionLibraryEdGraph->InitializeFromBlueprint(this);
 	}
 }
 
@@ -224,14 +244,6 @@ bool URigVMBlueprint::RequiresForceLoadMembers(UObject* InObject) const
 		return UBlueprint::RequiresForceLoadMembers(InObject);
 	}
 
-	/* todoooo
-	// old assets don't support preload filtering
-	if (GetLinkerCustomVersion(FControlRigObjectVersion::GUID) < FControlRigObjectVersion::RemoveParameters)
-	{
-		return UBlueprint::RequiresForceLoadMembers(InObject);
-	}
-	*/
-	
 	// we can stop traversing when hitting a URigVMNode
 	// except for collapse nodes - since they contain a graphs again
 	// and variable  nodes - since they are needed during preload by the BP compiler
@@ -667,6 +679,13 @@ void URigVMBlueprint::Serialize(FArchive& Ar)
 			TGuardValue<bool> DisableClientNotifs(RigVMClient.bSuspendNotifications, true);
 			RigVMClient.SetFromDeprecatedData(Model_DEPRECATED, FunctionLibrary_DEPRECATED);
 		}
+
+		TArray<UEdGraph*> EdGraphs;
+		GetAllGraphs(EdGraphs);
+		for (UEdGraph* EdGraph : EdGraphs)
+		{
+			EdGraph->Schema = GetRigVMEdGraphSchemaClass();
+		}
 	}
 }
 
@@ -751,7 +770,7 @@ void URigVMBlueprint::PostLoad()
 		for (UEdGraph* Graph : UbergraphPages)
 		{
 			URigVMEdGraph* RigGraph = Cast<URigVMEdGraph>(Graph);
-			if (RigGraph)
+			if (RigGraph && RigGraph->GetClass() == GetRigVMEdGraphClass())
 			{
 				NewUberGraphPages.Add(RigGraph);
 			}
@@ -1032,8 +1051,7 @@ void URigVMBlueprint::RecompileVM()
 		TGuardValue<bool> ReentrantGuardSelf(bSuspendModelNotificationsForSelf, true);
 		TGuardValue<bool> ReentrantGuardOthers(bSuspendModelNotificationsForOthers, true);
 
-		CDO->PostInitInstanceIfRequired();
-		CDO->VMRuntimeSettings = VMRuntimeSettings;
+		SetupDefaultObjectDuringCompilation(CDO);
 
 		if (!HasAnyFlags(RF_Transient | RF_Transactional))
 		{
@@ -1106,7 +1124,7 @@ void URigVMBlueprint::RecompileVM()
 		VMCompiledEvent.Broadcast(this, CDO->VM);
 
 #if WITH_EDITOR
-		RefreshControlRigBreakpoints();
+		RefreshBreakpoints();
 #endif
 	}
 }
@@ -1440,7 +1458,7 @@ void URigVMBlueprint::ClearBreakpoints()
 	}
 	
 	RigVMBreakpointNodes.Empty();
-	RefreshControlRigBreakpoints();
+	RefreshBreakpoints();
 }
 
 bool URigVMBlueprint::AddBreakpoint(const FString& InBreakpointNodePath)
@@ -1526,7 +1544,7 @@ bool URigVMBlueprint::AddBreakpoint(URigVMNode* InBreakpointNode, URigVMLibraryN
 		if (!RigVMBreakpointNodes.Contains(InBreakpointNode))
 		{
 			// Add the breakpoint to the VM
-			bSuccess = AddBreakpointToControlRig(InBreakpointNode);
+			bSuccess = AddBreakpointToHost(InBreakpointNode);
 			BreakpointAddedEvent.Broadcast();
 		}
 	}
@@ -1534,7 +1552,7 @@ bool URigVMBlueprint::AddBreakpoint(URigVMNode* InBreakpointNode, URigVMLibraryN
 	return bSuccess;
 }
 
-bool URigVMBlueprint::AddBreakpointToControlRig(URigVMNode* InBreakpointNode)
+bool URigVMBlueprint::AddBreakpointToHost(URigVMNode* InBreakpointNode)
 {
 	URigVMBlueprintGeneratedClass* RigClass = GetRigVMBlueprintGeneratedClass();
 	URigVMHost* CDO = Cast<URigVMHost>(RigClass->GetDefaultObject(false));
@@ -1637,15 +1655,14 @@ bool URigVMBlueprint::RemoveBreakpoint(URigVMNode* InBreakpointNode)
 		// is another breakpoint node addressing it. For that reason, we just recompute all the
 		// breakpoint instructions.
 		// Refreshing breakpoints in the control rig will keep the state it had before.
-		RefreshControlRigBreakpoints();
+		RefreshBreakpoints();
 		return true;
 	}
 
 	return false;
 }
 
-
-void URigVMBlueprint::RefreshControlRigBreakpoints()
+void URigVMBlueprint::RefreshBreakpoints()
 {
 	if (URigVMBlueprintGeneratedClass* RigClass = GetRigVMBlueprintGeneratedClass())
 	{
@@ -1653,7 +1670,7 @@ void URigVMBlueprint::RefreshControlRigBreakpoints()
 		CDO->GetDebugInfo().Reset();
 		for (URigVMNode* Node : RigVMBreakpointNodes)
 		{
-			AddBreakpointToControlRig(Node);
+			AddBreakpointToHost(Node);
 		}
 	}
 }
@@ -1678,6 +1695,12 @@ TArray<FRigVMReferenceNodeData> URigVMBlueprint::GetReferenceNodeData() const
 
 #endif
 
+void URigVMBlueprint::SetupDefaultObjectDuringCompilation(URigVMHost* InCDO)
+{
+	InCDO->PostInitInstanceIfRequired();
+	InCDO->VMRuntimeSettings = VMRuntimeSettings;
+}
+
 void URigVMBlueprint::RequestRigVMInit()
 {
 	if (URigVMBlueprintGeneratedClass* RigClass = GetRigVMBlueprintGeneratedClass())
@@ -1811,13 +1834,18 @@ URigVMController* URigVMBlueprint::GetOrCreateController(const UEdGraph* InEdGra
 TArray<FString> URigVMBlueprint::GeneratePythonCommands(const FString InNewBlueprintName)
 {
 	TArray<FString> InternalCommands;
-	InternalCommands.Add(TEXT("unreal.load_module('ControlRigDeveloper')"));
-	InternalCommands.Add(TEXT("factory = unreal.ControlRigBlueprintFactory"));
-	InternalCommands.Add(FString::Printf(TEXT("blueprint = factory.create_new_control_rig_asset(desired_package_path = '%s')"), *InNewBlueprintName));
+
+	if(GetClass() == StaticClass())
+	{
+		InternalCommands.Add(TEXT("import unreal"));
+		InternalCommands.Add(TEXT("unreal.load_module('RigVMDeveloper')"));
+		InternalCommands.Add(TEXT("blueprint = unreal.RigVMBlueprint()"));
+		InternalCommands.Add(TEXT("hierarchy = blueprint.hierarchy"));
+		InternalCommands.Add(TEXT("hierarchy_controller = hierarchy.get_controller()"));
+	}
+
 	InternalCommands.Add(TEXT("library = blueprint.get_local_function_library()"));
 	InternalCommands.Add(TEXT("library_controller = blueprint.get_controller(library)"));
-	InternalCommands.Add(TEXT("hierarchy = blueprint.hierarchy"));
-	InternalCommands.Add(TEXT("hierarchy_controller = hierarchy.get_controller()"));
 	InternalCommands.Add(TEXT("blueprint.set_auto_vm_recompile(False)"));
 	
 	// Add variables
@@ -2062,13 +2090,13 @@ TArray<FString> URigVMBlueprint::GeneratePythonCommands(const FString InNewBluep
 	// We do not want to pollute the global state with our definitions
 	TArray<FString> Commands;
 	Commands.Add(FString::Printf(TEXT("import unreal\n"
-		"def create_control_rig():\n")));
+		"def create_asset():\n")));
 	for (const FString& InnerCmd : InnerFunctionCmds)
 	{
 		Commands.Add(FString::Printf(TEXT("\t%s"), *InnerCmd));
 	}
 
-	Commands.Add(TEXT("create_control_rig()\n"));
+	Commands.Add(TEXT("create_asset()\n"));
 	return Commands;
 }
 
@@ -2334,10 +2362,7 @@ void URigVMBlueprint::ReplaceDeprecatedNodes()
 
 	for (UEdGraph* EdGraph : EdGraphs)
 	{
-		if (URigVMEdGraph* RigGraph = Cast<URigVMEdGraph>(EdGraph))
-		{
-			RigGraph->Schema = URigVMEdGraphSchema::StaticClass();
-		}
+		EdGraph->Schema = GetRigVMEdGraphSchemaClass();
 	}
 
 	Super::ReplaceDeprecatedNodes();
@@ -2434,7 +2459,7 @@ FRigVMGraphModifiedEvent& URigVMBlueprint::OnModified()
 }
 
 
-FOnVMCompiledEvent& URigVMBlueprint::OnVMCompiled()
+FOnRigVMCompiledEvent& URigVMBlueprint::OnVMCompiled()
 {
 	return VMCompiledEvent;
 }
@@ -3779,11 +3804,11 @@ UEdGraph* URigVMBlueprint::CreateEdGraph(URigVMGraph* InModel, bool bForce)
 
 	GraphName = RigVMClient.GetUniqueName(*GraphName).ToString();
 
-	URigVMEdGraph* RigVMEdGraph = NewObject<URigVMEdGraph>(this, *GraphName, RF_Transactional);
-	RigVMEdGraph->Schema = URigVMEdGraphSchema::StaticClass();
+	URigVMEdGraph* RigVMEdGraph = NewObject<URigVMEdGraph>(this, GetRigVMEdGraphClass(), *GraphName, RF_Transactional);
+	RigVMEdGraph->Schema = GetRigVMEdGraphSchemaClass();
 	RigVMEdGraph->bAllowDeletion = true;
 	RigVMEdGraph->ModelNodePath = InModel->GetNodePath();
-	RigVMEdGraph->Initialize(this);
+	RigVMEdGraph->InitializeFromBlueprint(this);
 	
 	FBlueprintEditorUtils::AddUbergraphPage(this, RigVMEdGraph);
 	LastEditedDocuments.AddUnique(RigVMEdGraph);
@@ -3857,8 +3882,8 @@ void URigVMBlueprint::CreateEdGraphForCollapseNodeIfNeeded(URigVMCollapseNode* I
 			if (!bFunctionGraphExists)
 			{
 				// create a sub graph
-				URigVMEdGraph* RigFunctionGraph = NewObject<URigVMEdGraph>(this, *InNode->GetName(), RF_Transactional);
-				RigFunctionGraph->Schema = URigVMEdGraphSchema::StaticClass();
+				URigVMEdGraph* RigFunctionGraph = NewObject<URigVMEdGraph>(this, GetRigVMEdGraphClass(), *InNode->GetName(), RF_Transactional);
+				RigFunctionGraph->Schema = GetRigVMEdGraphSchemaClass();
 				RigFunctionGraph->bAllowRenaming = 1;
 				RigFunctionGraph->bEditable = 1;
 				RigFunctionGraph->bAllowDeletion = 1;
@@ -3867,7 +3892,7 @@ void URigVMBlueprint::CreateEdGraphForCollapseNodeIfNeeded(URigVMCollapseNode* I
 
 				FunctionGraphs.Add(RigFunctionGraph);
 
-				RigFunctionGraph->Initialize(this);
+				RigFunctionGraph->InitializeFromBlueprint(this);
 
 				GetOrCreateController(ContainedGraph)->ResendAllNotifications();
 			}
@@ -3900,8 +3925,8 @@ void URigVMBlueprint::CreateEdGraphForCollapseNodeIfNeeded(URigVMCollapseNode* I
 				}
 				
 				// create a sub graph
-				URigVMEdGraph* SubRigGraph = NewObject<URigVMEdGraph>(RigGraph, *InNode->GetEditorSubGraphName(), RF_Transactional);
-				SubRigGraph->Schema = URigVMEdGraphSchema::StaticClass();
+				URigVMEdGraph* SubRigGraph = NewObject<URigVMEdGraph>(RigGraph, GetRigVMEdGraphClass(), *InNode->GetEditorSubGraphName(), RF_Transactional);
+				SubRigGraph->Schema = GetRigVMEdGraphSchemaClass();
 				SubRigGraph->bAllowRenaming = 1;
 				SubRigGraph->bEditable = bEditable;
 				SubRigGraph->bAllowDeletion = 1;
@@ -3910,7 +3935,7 @@ void URigVMBlueprint::CreateEdGraphForCollapseNodeIfNeeded(URigVMCollapseNode* I
 
 				RigGraph->SubGraphs.Add(SubRigGraph);
 
-				SubRigGraph->Initialize(this);
+				SubRigGraph->InitializeFromBlueprint(this);
 
 				GetOrCreateController(ContainedGraph)->ResendAllNotifications();
 			}
diff --git a/Engine/Plugins/Runtime/RigVM/Source/RigVMDeveloper/Private/RigVMDeveloperModule.cpp b/Engine/Plugins/Runtime/RigVM/Source/RigVMDeveloper/Private/RigVMDeveloperModule.cpp
--- a/Engine/Plugins/Runtime/RigVM/Source/RigVMDeveloper/Private/RigVMDeveloperModule.cpp
+++ b/Engine/Plugins/Runtime/RigVM/Source/RigVMDeveloper/Private/RigVMDeveloperModule.cpp
@@ -6,7 +6,23 @@
 
 #include "RigVMDeveloperModule.h"
 #include "Modules/ModuleManager.h"
-
-IMPLEMENT_MODULE(FDefaultModuleImpl, RigVMDeveloper);
+#include "RigVMBlueprint.h"
 
 DEFINE_LOG_CATEGORY(LogRigVMDeveloper);
+
+class FRigVMDeveloperModule : public IModuleInterface
+{
+public:
+	/** IModuleInterface implementation */
+	virtual void StartupModule() override
+	{
+	}
+
+	virtual void ShutdownModule() override
+	{
+	}
+};
+
+IMPLEMENT_MODULE(FRigVMDeveloperModule, RigVMDeveloper);
+
+
diff --git a/Engine/Plugins/Runtime/RigVM/Source/RigVMDeveloper/Private/RigVMModel/RigVMClient.cpp b/Engine/Plugins/Runtime/RigVM/Source/RigVMDeveloper/Private/RigVMModel/RigVMClient.cpp
--- a/Engine/Plugins/Runtime/RigVM/Source/RigVMDeveloper/Private/RigVMModel/RigVMClient.cpp
+++ b/Engine/Plugins/Runtime/RigVM/Source/RigVMDeveloper/Private/RigVMModel/RigVMClient.cpp
@@ -33,7 +33,12 @@ void FRigVMClient::SetSchemaClass(TSubclassOf<URigVMSchema> InSchemaClass)
 	}
 
 	SchemaClass = InSchemaClass;
-	GetOrCreateSchema();
+
+	URigVMSchema* Schema = GetOrCreateSchema();
+	for(auto& Pair : Controllers)
+	{
+		Pair.Value->SetSchema(Schema);
+	}
 }
 
 void FRigVMClient::SetOuterClientHost(UObject* InOuterClientHost, const FName& InOuterClientHostPropertyName)
diff --git a/Engine/Plugins/Runtime/RigVM/Source/RigVMDeveloper/Public/EdGraph/RigVMEdGraph.h b/Engine/Plugins/Runtime/RigVM/Source/RigVMDeveloper/Public/EdGraph/RigVMEdGraph.h
--- a/Engine/Plugins/Runtime/RigVM/Source/RigVMDeveloper/Public/EdGraph/RigVMEdGraph.h
+++ b/Engine/Plugins/Runtime/RigVM/Source/RigVMDeveloper/Public/EdGraph/RigVMEdGraph.h
@@ -29,7 +29,7 @@ public:
 	virtual void HandleRigVMGraphRenamed(const FString& InOldNodePath, const FString& InNewNodePath) override;
 
 	/** Set up this graph */
-	virtual void Initialize(URigVMBlueprint* InBlueprint);
+	virtual void InitializeFromBlueprint(URigVMBlueprint* InBlueprint);
 
 	/** Get the ed graph schema */
 	const URigVMEdGraphSchema* GetRigVMEdGraphSchema();
@@ -53,7 +53,7 @@ public:
 	const URigVMEdGraph* GetRootGraph() const;
 
 	void HandleModifiedEvent(ERigVMGraphNotifType InNotifType, URigVMGraph* InGraph, UObject* InSubject);
-	void ConsumeQueuedNotifications();
+	virtual bool HandleModifiedEvent_Internal(ERigVMGraphNotifType InNotifType, URigVMGraph* InGraph, UObject* InSubject);
 
 	int32 GetInstructionIndex(const URigVMEdGraphNode* InNode, bool bAsInput);
 
@@ -63,6 +63,10 @@ public:
 	UPROPERTY()
 	bool bIsFunctionDefinition;
 
+protected:
+	using Super::AddNode;
+	virtual void AddNode(UEdGraphNode* NodeToAdd, bool bUserAction = false, bool bSelectNewNode = true) override;
+
 private:
 
 	bool bIsSelecting;
@@ -112,4 +116,17 @@ private:
 	friend class FControlRigEditor;
 	friend class SControlRigGraphNode;
 	friend class URigVMBlueprint;
+
+	// todo
+	friend class UControlRigUnitNodeSpawner;
+	friend class UControlRigVariableNodeSpawner;
+	friend class UControlRigParameterNodeSpawner;
+	friend class UControlRigBranchNodeSpawner;
+	friend class UControlRigIfNodeSpawner;
+	friend class UControlRigSelectNodeSpawner;
+	friend class UControlRigTemplateNodeSpawner;
+	friend class UControlRigEnumNodeSpawner;
+	friend class UControlRigFunctionRefNodeSpawner;
+	friend class UControlRigArrayNodeSpawner;
+	friend class UControlRigInvokeEntryNodeSpawner;
 };
diff --git a/Engine/Plugins/Runtime/RigVM/Source/RigVMDeveloper/Public/EdGraph/RigVMEdGraphNode.h b/Engine/Plugins/Runtime/RigVM/Source/RigVMDeveloper/Public/EdGraph/RigVMEdGraphNode.h
--- a/Engine/Plugins/Runtime/RigVM/Source/RigVMDeveloper/Public/EdGraph/RigVMEdGraphNode.h
+++ b/Engine/Plugins/Runtime/RigVM/Source/RigVMDeveloper/Public/EdGraph/RigVMEdGraphNode.h
@@ -277,4 +277,14 @@ private:
 	friend class URigVMArrayNodeSpawner;
 	friend class URigVMIfNodeSpawner;
 	friend class URigVMSelectNodeSpawner;
+
+	// todo
+	friend class FControlRigEditor;
+	friend class UControlRigBlueprint;
+	friend class SControlRigGraphNode;
+	friend class UControlRigGraphSchema;
+	friend class UControlRigTemplateNodeSpawner;
+	friend class UControlRigArrayNodeSpawner;
+	friend class UControlRigIfNodeSpawner;
+	friend class UControlRigSelectNodeSpawner;
 };
diff --git a/Engine/Plugins/Runtime/RigVM/Source/RigVMDeveloper/Public/EdGraph/RigVMEdGraphSchema.h b/Engine/Plugins/Runtime/RigVM/Source/RigVMDeveloper/Public/EdGraph/RigVMEdGraphSchema.h
--- a/Engine/Plugins/Runtime/RigVM/Source/RigVMDeveloper/Public/EdGraph/RigVMEdGraphSchema.h
+++ b/Engine/Plugins/Runtime/RigVM/Source/RigVMDeveloper/Public/EdGraph/RigVMEdGraphSchema.h
@@ -252,6 +252,8 @@ public:
 public:
 	URigVMEdGraphSchema();
 
+	virtual const FName& GetGraphNameForUI() const { return GraphName_RigVM; }
+
 	// UEdGraphSchema interface
 	virtual void GetGraphContextActions(FGraphContextMenuBuilder& ContextMenuBuilder) const override;
 	virtual void GetContextMenuActions(class UToolMenu* Menu, class UGraphNodeContextMenuContext* Context) const override;
@@ -325,7 +327,7 @@ public:
 	void HandleModifiedEvent(ERigVMGraphNotifType InNotifType, URigVMGraph* InGraph, UObject* InSubject);
 
 	// Allow derived classes to spawn derived node classes
-	virtual TSubclassOf<URigVMEdGraphNode> GetGraphNodeClass() const { return URigVMEdGraphNode::StaticClass(); }
+	virtual TSubclassOf<URigVMEdGraphNode> GetGraphNodeClass(const URigVMEdGraph* InGraph) const;
 
 	virtual bool IsRigVMDefaultEvent(const FName& InEventName) const;
 
diff --git a/Engine/Plugins/Runtime/RigVM/Source/RigVMDeveloper/Public/RigVMBlueprint.h b/Engine/Plugins/Runtime/RigVM/Source/RigVMDeveloper/Public/RigVMBlueprint.h
--- a/Engine/Plugins/Runtime/RigVM/Source/RigVMDeveloper/Public/RigVMBlueprint.h
+++ b/Engine/Plugins/Runtime/RigVM/Source/RigVMDeveloper/Public/RigVMBlueprint.h
@@ -17,22 +17,32 @@
 class URigVMBlueprintGeneratedClass;
 struct FEndLoadPackageContext;
 
-DECLARE_EVENT_TwoParams(URigVMBlueprint, FOnVMCompiledEvent, UObject*, URigVM*);
-DECLARE_EVENT_OneParam(URigVMBlueprint, FOnRefreshEditorEvent, URigVMBlueprint*);
-DECLARE_EVENT_FourParams(URigVMBlueprint, FOnVariableDroppedEvent, UObject*, FProperty*, const FVector2D&, const FVector2D&);
-DECLARE_EVENT_OneParam(URigVMBlueprint, FOnExternalVariablesChanged, const TArray<FRigVMExternalVariable>&);
-DECLARE_EVENT_TwoParams(URigVMBlueprint, FOnNodeDoubleClicked, URigVMBlueprint*, URigVMNode*);
-DECLARE_EVENT_OneParam(URigVMBlueprint, FOnGraphImported, UEdGraph*);
-DECLARE_EVENT_OneParam(URigVMBlueprint, FOnPostEditChangeChainProperty, FPropertyChangedChainEvent&);
-DECLARE_EVENT_ThreeParams(URigVMBlueprint, FOnLocalizeFunctionDialogRequested, FRigVMGraphFunctionIdentifier&, URigVMBlueprint*, bool);
-DECLARE_EVENT_ThreeParams(URigVMBlueprint, FOnReportCompilerMessage, EMessageSeverity::Type, UObject*, const FString&);
+DECLARE_EVENT_TwoParams(URigVMBlueprint, FOnRigVMCompiledEvent, UObject*, URigVM*);
+DECLARE_EVENT_OneParam(URigVMBlueprint, FOnRigVMRefreshEditorEvent, URigVMBlueprint*);
+DECLARE_EVENT_FourParams(URigVMBlueprint, FOnRigVMVariableDroppedEvent, UObject*, FProperty*, const FVector2D&, const FVector2D&);
+DECLARE_EVENT_OneParam(URigVMBlueprint, FOnRigVMExternalVariablesChanged, const TArray<FRigVMExternalVariable>&);
+DECLARE_EVENT_TwoParams(URigVMBlueprint, FOnRigVMNodeDoubleClicked, URigVMBlueprint*, URigVMNode*);
+DECLARE_EVENT_OneParam(URigVMBlueprint, FOnRigVMGraphImported, UEdGraph*);
+DECLARE_EVENT_OneParam(URigVMBlueprint, FOnRigVMPostEditChangeChainProperty, FPropertyChangedChainEvent&);
+DECLARE_EVENT_ThreeParams(URigVMBlueprint, FOnRigVMLocalizeFunctionDialogRequested, FRigVMGraphFunctionIdentifier&, URigVMBlueprint*, bool);
+DECLARE_EVENT_ThreeParams(URigVMBlueprint, FOnRigVMReportCompilerMessage, EMessageSeverity::Type, UObject*, const FString&);
 DECLARE_DELEGATE_RetVal_FourParams(FRigVMController_BulkEditResult, FRigVMOnBulkEditDialogRequestedDelegate, URigVMBlueprint*, URigVMController*, URigVMLibraryNode*, ERigVMControllerBulkEditType);
 DECLARE_DELEGATE_RetVal_OneParam(bool, FRigVMOnBreakLinksDialogRequestedDelegate, TArray<URigVMLink*>);
 DECLARE_DELEGATE_RetVal_OneParam(TRigVMTypeIndex, FRigVMOnPinTypeSelectionRequestedDelegate, const TArray<TRigVMTypeIndex>&);
-DECLARE_EVENT(URigVMBlueprint, FOnBreakpointAdded);
-DECLARE_EVENT_OneParam(URigVMBlueprint, FOnRequestInspectObject, const TArray<UObject*>& );
+DECLARE_EVENT(URigVMBlueprint, FOnRigVMBreakpointAdded);
+DECLARE_EVENT_OneParam(URigVMBlueprint, FOnRigVMRequestInspectObject, const TArray<UObject*>& );
 DECLARE_DELEGATE_RetVal(URigVMGraph*, FRigVMBlueprintGetFocusedGraph);
 
+USTRUCT()
+struct RIGVMDEVELOPER_API FRigVMPythonSettings
+{
+	GENERATED_BODY();
+
+	FRigVMPythonSettings()
+	{
+	}
+};
+
 USTRUCT()
 struct RIGVMDEVELOPER_API FRigVMEdGraphDisplaySettings
 {
@@ -125,6 +135,8 @@ class RIGVMDEVELOPER_API URigVMBlueprint : public UBlueprint, public IRigVMClien
 public:
 	URigVMBlueprint();
 
+	void CommonInitialization(const FObjectInitializer& ObjectInitializer);
+	
 	void InitializeModelIfRequired(bool bRecompileVM = true);
 
 	/** Get the (full) generated class for this rigvm blueprint */
@@ -139,12 +151,15 @@ public:
 	/** Returns the expected execute context struct to use for this blueprint */
 	virtual UScriptStruct* GetRigVMExecuteContextStruct() const { return FRigVMExecuteContext::StaticStruct(); }
 
-	/** Returns the expected ed graph schema class to use for this blueprint */
-	virtual UClass* GetRigVMEdGraphSchemaClass() const { return URigVMEdGraphSchema::StaticClass(); }
-
 	/** Returns the expected ed graph class to use for this blueprint */
 	virtual UClass* GetRigVMEdGraphClass() const { return URigVMEdGraph::StaticClass(); }
 
+	/** Returns the expected ed graph node class to use for this blueprint */
+	virtual UClass* GetRigVMEdGraphNodeClass() const { return URigVMEdGraphNode::StaticClass(); }
+
+	/** Returns the expected ed graph schema class to use for this blueprint */
+	virtual UClass* GetRigVMEdGraphSchemaClass() const { return URigVMEdGraphSchema::StaticClass(); }
+
 	virtual void Serialize(FArchive& Ar) override;
 
 #if WITH_EDITOR
@@ -203,7 +218,7 @@ public:
 	virtual void HandleRigVMGraphRenamed(const FRigVMClient* InClient, const FString& InOldNodePath, const FString& InNewNodePath) override;
 	virtual void HandleConfigureRigVMController(const FRigVMClient* InClient, URigVMController* InControllerToConfigure) override;
 
-	FOnRequestInspectObject& OnRequestInspectObject() { return OnRequestInspectObjectEvent; }
+	FOnRigVMRequestInspectObject& OnRequestInspectObject() { return OnRequestInspectObjectEvent; }
 	void RequestInspectObject(const TArray<UObject*>& InObjects) { OnRequestInspectObjectEvent.Broadcast(InObjects); }
 
 #endif	// #if WITH_EDITOR
@@ -224,6 +239,9 @@ public:
 	UFUNCTION(BlueprintCallable, Category = "RigVM Blueprint")
 	void SetAutoVMRecompile(bool bAutoRecompile) { bAutoRecompileVM = bAutoRecompile; }
 
+	UFUNCTION(BlueprintPure, Category = "RigVM Blueprint")
+	bool GetAutoVMRecompile() const { return bAutoRecompileVM; }
+
 	void IncrementVMRecompileBracket();
 	void DecrementVMRecompileBracket();
 
@@ -275,7 +293,7 @@ public:
 	URigVMController* GetOrCreateController(const UEdGraph* InGraph);
 
 	UFUNCTION(BlueprintCallable, Category = "RigVM Blueprint")
-	TArray<FString> GeneratePythonCommands(const FString InNewBlueprintName);
+	virtual TArray<FString> GeneratePythonCommands(const FString InNewBlueprintName);
 
 	URigVMGraph* GetTemplateModel(bool bIsFunctionLibrary = false);
 	URigVMController* GetTemplateController(bool bIsFunctionLibrary = false);
@@ -318,6 +336,9 @@ public:
 	UPROPERTY(EditAnywhere, BlueprintReadWrite, Category = "VM")
 	FRigVMCompileSettings VMCompileSettings;
 
+	UPROPERTY(EditAnywhere, Category = "Python Log Settings")
+	FRigVMPythonSettings PythonLogSettings;
+
 protected:
 
 #if WITH_EDITORONLY_DATA
@@ -370,11 +391,11 @@ public:
 	bool bSuspendModelNotificationsForOthers;
 	bool bSuspendAllNotifications;
 
-	virtual void SetupPinRedirectorsForBackwardsCompatibility() const {};
+	virtual void SetupPinRedirectorsForBackwardsCompatibility() {};
 	void RebuildGraphFromModel();
 
 	FRigVMGraphModifiedEvent& OnModified();
-	FOnVMCompiledEvent& OnVMCompiled();
+	FOnRigVMCompiledEvent& OnVMCompiled();
 
 	UFUNCTION(BlueprintCallable, Category = "VM")
 	UClass* GetRigVMHostClass();
@@ -429,22 +450,22 @@ private:
 	UFUNCTION(BlueprintCallable, Category = "RigVM Blueprint")
 	void SuspendNotifications(bool bSuspendNotifs);
 
-	FOnRefreshEditorEvent RefreshEditorEvent;
-	FOnVariableDroppedEvent VariableDroppedEvent;
-	FOnBreakpointAdded BreakpointAddedEvent;
+	FOnRigVMRefreshEditorEvent RefreshEditorEvent;
+	FOnRigVMVariableDroppedEvent VariableDroppedEvent;
+	FOnRigVMBreakpointAdded BreakpointAddedEvent;
 
 public:
 
 	void BroadcastRefreshEditor() { return RefreshEditorEvent.Broadcast(this); }
-	FOnRefreshEditorEvent& OnRefreshEditor() { return RefreshEditorEvent; }
-	FOnVariableDroppedEvent& OnVariableDropped() { return VariableDroppedEvent; }
-	FOnBreakpointAdded& OnBreakpointAdded() { return BreakpointAddedEvent; }
+	FOnRigVMRefreshEditorEvent& OnRefreshEditor() { return RefreshEditorEvent; }
+	FOnRigVMVariableDroppedEvent& OnVariableDropped() { return VariableDroppedEvent; }
+	FOnRigVMBreakpointAdded& OnBreakpointAdded() { return BreakpointAddedEvent; }
 
 private:
 
 #endif
 
-	FOnVMCompiledEvent VMCompiledEvent;
+	FOnRigVMCompiledEvent VMCompiledEvent;
 
 	virtual void CreateMemberVariablesOnLoad();
 #if WITH_EDITOR
@@ -472,8 +493,14 @@ private:
 
 #if WITH_EDITOR
 	void HandlePackageDone(const FEndLoadPackageContext& Context);
-	void HandlePackageDone();
-	// ControlRigBP, once end-loaded, will inform other ControlRig-Dependent systems that ControlRig instances are ready.
+
+protected:
+	
+	virtual void HandlePackageDone();
+
+private:
+	
+	// RigVMBP, once end-loaded, will inform other RigVM-Dependent systems that Host instances are ready.
 	void BroadcastRigVMPackageDone();
 
 	// Previously some memory classes were parented to the asset object
@@ -489,7 +516,7 @@ private:
 
 public:
 
-	FOnExternalVariablesChanged& OnExternalVariablesChanged() { return ExternalVariablesChangedEvent; }
+	FOnRigVMExternalVariablesChanged& OnExternalVariablesChanged() { return ExternalVariablesChangedEvent; }
 
 	virtual void OnPreVariableChange(UObject* InObject);
 	virtual void OnPostVariableChange(UBlueprint* InBlueprint);
@@ -498,16 +525,16 @@ public:
 	virtual void OnVariableRenamed(const FName& InOldVarName, const FName& InNewVarName);
 	virtual void OnVariableTypeChanged(const FName& InVarName, FEdGraphPinType InOldPinType, FEdGraphPinType InNewPinType);
 
-	FOnNodeDoubleClicked& OnNodeDoubleClicked() { return NodeDoubleClickedEvent; }
+	FOnRigVMNodeDoubleClicked& OnNodeDoubleClicked() { return NodeDoubleClickedEvent; }
 	void BroadcastNodeDoubleClicked(URigVMNode* InNode);
 
-	FOnGraphImported& OnGraphImported() { return GraphImportedEvent; }
+	FOnRigVMGraphImported& OnGraphImported() { return GraphImportedEvent; }
 	void BroadcastGraphImported(UEdGraph* InGraph);
 
-	FOnPostEditChangeChainProperty& OnPostEditChangeChainProperty() { return PostEditChangeChainPropertyEvent; }
+	FOnRigVMPostEditChangeChainProperty& OnPostEditChangeChainProperty() { return PostEditChangeChainPropertyEvent; }
 	void BroadcastPostEditChangeChainProperty(FPropertyChangedChainEvent& PropertyChangedChainEvent);
 
-	FOnLocalizeFunctionDialogRequested& OnRequestLocalizeFunctionDialog() { return RequestLocalizeFunctionDialog; }
+	FOnRigVMLocalizeFunctionDialogRequested& OnRequestLocalizeFunctionDialog() { return RequestLocalizeFunctionDialog; }
 	void BroadcastRequestLocalizeFunctionDialog(FRigVMGraphFunctionIdentifier InFunction, bool bForce = false);
 
 	FRigVMOnBulkEditDialogRequestedDelegate& OnRequestBulkEditDialog() { return RequestBulkEditDialog; }
@@ -518,21 +545,24 @@ public:
 
 	FRigVMController_RequestJumpToHyperlinkDelegate& OnRequestJumpToHyperlink() { return RequestJumpToHyperlink; };
 
-	FOnReportCompilerMessage& OnReportCompilerMessage() { return ReportCompilerMessageEvent; }
+	FOnRigVMReportCompilerMessage& OnReportCompilerMessage() { return ReportCompilerMessageEvent; }
 	void BroadCastReportCompilerMessage(EMessageSeverity::Type InSeverity, UObject* InSubject, const FString& InMessage);
 
+	const FCompilerResultsLog& GetCompileLog() const { return CompileLog; }
+	FCompilerResultsLog& GetCompileLog() { return CompileLog; }
+
 private:
 
-	FOnExternalVariablesChanged ExternalVariablesChangedEvent;
+	FOnRigVMExternalVariablesChanged ExternalVariablesChangedEvent;
 	bool bUpdatingExternalVariables;
 	void BroadcastExternalVariablesChangedEvent();
 	FCompilerResultsLog CompileLog;
 
-	FOnNodeDoubleClicked NodeDoubleClickedEvent;
-	FOnGraphImported GraphImportedEvent;
-	FOnPostEditChangeChainProperty PostEditChangeChainPropertyEvent;
-	FOnLocalizeFunctionDialogRequested RequestLocalizeFunctionDialog;
-	FOnReportCompilerMessage ReportCompilerMessageEvent;
+	FOnRigVMNodeDoubleClicked NodeDoubleClickedEvent;
+	FOnRigVMGraphImported GraphImportedEvent;
+	FOnRigVMPostEditChangeChainProperty PostEditChangeChainPropertyEvent;
+	FOnRigVMLocalizeFunctionDialogRequested RequestLocalizeFunctionDialog;
+	FOnRigVMReportCompilerMessage ReportCompilerMessageEvent;
 	FRigVMOnBulkEditDialogRequestedDelegate RequestBulkEditDialog;
 	FRigVMOnBreakLinksDialogRequestedDelegate RequestBreakLinksDialog;
 	FRigVMOnPinTypeSelectionRequestedDelegate RequestPinTypeSelectionDialog;
@@ -548,6 +578,8 @@ private:
 	bool RemoveEdGraphForCollapseNode(URigVMCollapseNode* InNode, bool bNotify = false);
 	void HandleReportFromCompiler(EMessageSeverity::Type InSeverity, UObject* InSubject, const FString& InMessage);
 
+protected:
+	
 	TArray<IRigVMGraphFunctionHost*> GetReferencedFunctionHosts(bool bForceLoad);
 	
 #if WITH_EDITOR
@@ -556,7 +588,7 @@ private:
 	
 	TArray<URigVMNode*> RigVMBreakpointNodes;
 
-	FOnRequestInspectObject OnRequestInspectObjectEvent;
+	FOnRigVMRequestInspectObject OnRequestInspectObjectEvent;
 	FRigVMBlueprintGetFocusedGraph OnGetFocusedGraphDelegate;
 		
 public:
@@ -577,7 +609,7 @@ public:
 	bool AddBreakpoint(URigVMNode* InBreakpointNode, URigVMLibraryNode* LibraryNode = nullptr);
 
 	/** Adds a breakpoint to the first instruction of each callpath related to the InBreakpointNode */
-	bool AddBreakpointToControlRig(URigVMNode* InBreakpointNode);
+	bool AddBreakpointToHost(URigVMNode* InBreakpointNode);
 
 	/** Removes the given breakpoint from all the loaded blueprints that use this node, and recomputes all breakpoints
 	  * in the VM  */
@@ -585,7 +617,7 @@ public:
 	bool RemoveBreakpoint(URigVMNode* InBreakpointNode);
 
 	/** Recomputes the instruction breakpoints given the node breakpoints in the blueprint */
-	void RefreshControlRigBreakpoints();
+	void RefreshBreakpoints();
 
 	/** Shape libraries to load during package load completed */ 
 	TArray<FString> ShapeLibrariesToLoadOnPackageLoaded;
@@ -599,6 +631,14 @@ protected:
 	static FSoftObjectPath PreDuplicateAssetPath;
 	static FSoftObjectPath PreDuplicateHostPath;
 
+	void MarkDirtyDuringLoad() { bDirtyDuringLoad = true; }
+	
+	virtual void SetupDefaultObjectDuringCompilation(URigVMHost* InCDO);
+
+public:
+	
+	bool IsMarkedDirtyDuringLoad() const { return bDirtyDuringLoad; }
+
 private:
 	bool bDirtyDuringLoad;
 	bool bErrorsDuringCompilation;
@@ -609,20 +649,24 @@ private:
 	friend class FRigVMEditorModule;
 	friend class URigVMEdGraphSchema;
 	friend struct FRigVMEdGraphSchemaAction_PromoteToVariable;
+
+	// todo
+	friend class FControlRigEditorModule;
+	friend class FControlRigBlueprintCompilerContext;
 };
 
-class RIGVMDEVELOPER_API FRigVMBlueprintVMCompileScope
+class RIGVMDEVELOPER_API FRigVMBlueprintCompileScope
 {
 public:
    
-	FRigVMBlueprintVMCompileScope(URigVMBlueprint *InBlueprint)
+	FRigVMBlueprintCompileScope(URigVMBlueprint *InBlueprint)
 	: Blueprint(InBlueprint)
 	{
 		check(Blueprint);
 		Blueprint->IncrementVMRecompileBracket();
 	}
 
-	~FRigVMBlueprintVMCompileScope()
+	~FRigVMBlueprintCompileScope()
 	{
 		Blueprint->DecrementVMRecompileBracket();
 	}
