RigVM: Load user defined types for type checking

diff --git a/Engine/Plugins/Runtime/RigVM/Source/RigVM/Private/RigVMCore/RigVMGraphFunctionDefinition.cpp b/Engine/Plugins/Runtime/RigVM/Source/RigVM/Private/RigVMCore/RigVMGraphFunctionDefinition.cpp
--- a/Engine/Plugins/Runtime/RigVM/Source/RigVM/Private/RigVMCore/RigVMGraphFunctionDefinition.cpp
+++ b/Engine/Plugins/Runtime/RigVM/Source/RigVM/Private/RigVMCore/RigVMGraphFunctionDefinition.cpp
@@ -11,8 +11,8 @@ FRigVMExternalVariable FRigVMGraphFunctionArgument::GetExternalVariable() const
 	Variable.Name = Name;
 	Variable.bIsArray = bIsArray;
 	Variable.TypeName = CPPType;
-	
-	if(CPPTypeObject.IsValid())
+
+	if(IsCPPTypeObjectValid())
 	{
 		Variable.TypeObject = CPPTypeObject.Get();
 	}
@@ -20,11 +20,22 @@ FRigVMExternalVariable FRigVMGraphFunctionArgument::GetExternalVariable() const
 	return Variable;
 }
 
+bool FRigVMGraphFunctionArgument::IsCPPTypeObjectValid() const
+{
+	if(!CPPTypeObject.IsValid())
+	{
+		// this is potentially a user defined struct or user defined enum
+		// so we have to try to load it.
+		(void)CPPTypeObject.LoadSynchronous();
+	}
+	return CPPTypeObject.IsValid();
+}
+
 bool FRigVMGraphFunctionHeader::IsMutable() const
 {
 	for(const FRigVMGraphFunctionArgument& Arg : Arguments)
 	{
-		if(Arg.CPPTypeObject.IsValid())
+		if(Arg.IsCPPTypeObjectValid())
 		{
 			if(UScriptStruct* Struct = Cast<UScriptStruct>(Arg.CPPTypeObject.Get()))
 			{
diff --git a/Engine/Plugins/Runtime/RigVM/Source/RigVM/Public/RigVMCore/RigVMGraphFunctionDefinition.h b/Engine/Plugins/Runtime/RigVM/Source/RigVM/Public/RigVMCore/RigVMGraphFunctionDefinition.h
--- a/Engine/Plugins/Runtime/RigVM/Source/RigVM/Public/RigVMCore/RigVMGraphFunctionDefinition.h
+++ b/Engine/Plugins/Runtime/RigVM/Source/RigVM/Public/RigVMCore/RigVMGraphFunctionDefinition.h
@@ -254,6 +254,9 @@ struct RIGVM_API FRigVMGraphFunctionArgument
 
 	FRigVMExternalVariable GetExternalVariable() const;
 
+	// validates and potentially loads the CPP Type Object
+	bool IsCPPTypeObjectValid() const; 
+
 	friend uint32 GetTypeHash(const FRigVMGraphFunctionArgument& Argument)
 	{
 		uint32 Hash = HashCombine(GetTypeHash(Argument.Name), GetTypeHash(Argument.DisplayName));
diff --git a/Engine/Plugins/Runtime/RigVM/Source/RigVMDeveloper/Private/RigVMCompiler/RigVMCompiler.cpp b/Engine/Plugins/Runtime/RigVM/Source/RigVMDeveloper/Private/RigVMCompiler/RigVMCompiler.cpp
--- a/Engine/Plugins/Runtime/RigVM/Source/RigVMDeveloper/Private/RigVMCompiler/RigVMCompiler.cpp
+++ b/Engine/Plugins/Runtime/RigVM/Source/RigVMDeveloper/Private/RigVMCompiler/RigVMCompiler.cpp
@@ -2017,7 +2017,7 @@ int32 URigVMCompiler::TraverseInlineFunction(const FRigVMInlineFunctionExprAST*
 			{
 				for (const FRigVMGraphFunctionArgument& Argument : FunctionReferenceNode->GetReferencedFunctionHeader().Arguments)
 				{
-					if (Argument.CPPTypeObject.IsValid())
+					if (Argument.IsCPPTypeObjectValid())
 					{
 						if (const UScriptStruct* ScriptStruct = Cast<UScriptStruct>(Argument.CPPTypeObject.Get()))
 						{
diff --git a/Engine/Plugins/Runtime/RigVM/Source/RigVMDeveloper/Private/RigVMModel/RigVMSchema.cpp b/Engine/Plugins/Runtime/RigVM/Source/RigVMDeveloper/Private/RigVMModel/RigVMSchema.cpp
--- a/Engine/Plugins/Runtime/RigVM/Source/RigVMDeveloper/Private/RigVMModel/RigVMSchema.cpp
+++ b/Engine/Plugins/Runtime/RigVM/Source/RigVMDeveloper/Private/RigVMModel/RigVMSchema.cpp
@@ -137,7 +137,7 @@ bool URigVMSchema::SupportsGraphFunction(URigVMController* InController, const F
 		TRigVMTypeIndex Type = Registry->GetTypeIndexFromCPPType(Argument.CPPType.ToString());
 		if (Type == INDEX_NONE)
 		{
-			if (Argument.CPPTypeObject.IsValid())
+			if (Argument.IsCPPTypeObjectValid())
 			{
 				FRigVMTemplateArgumentType ArgumentType(Argument.CPPType, Argument.CPPTypeObject.Get());
 				Type = Registry->FindOrAddType(ArgumentType);
