RigVMEditor: Don't show RigVMStruct based variable types

diff --git a/Engine/Plugins/Runtime/RigVM/Source/RigVMDeveloper/Private/EdGraph/RigVMEdGraphSchema.cpp b/Engine/Plugins/Runtime/RigVM/Source/RigVMDeveloper/Private/EdGraph/RigVMEdGraphSchema.cpp
--- a/Engine/Plugins/Runtime/RigVM/Source/RigVMDeveloper/Private/EdGraph/RigVMEdGraphSchema.cpp
+++ b/Engine/Plugins/Runtime/RigVM/Source/RigVMDeveloper/Private/EdGraph/RigVMEdGraphSchema.cpp
@@ -988,6 +988,16 @@ bool URigVMEdGraphSchema::SupportsPinType(const UScriptStruct* ScriptStruct) con
 		return false;
 	}
 
+	if(ScriptStruct->IsChildOf(FRigVMStruct::StaticStruct()))
+	{
+		return false;
+	}
+
+	if(ScriptStruct->IsChildOf(FRigVMDispatchFactory::StaticStruct()))
+	{
+		return false;
+	}
+
 	// todo: Validate ExecuteContext structs to match URigVMBlueprint::GetExecuteContextStruct()
 
 	for (TFieldIterator<FProperty> It(ScriptStruct); It; ++It)
