Fix for crash when adding a pin on Sequence Node in Control Rig graph
This also fixes issues with Undo / Redo when we add a pin, then we add comments and we perform undo / redo of the comments

diff --git a/Engine/Plugins/Runtime/RigVM/Source/RigVMDeveloper/Private/EdGraph/RigVMEdGraph.cpp b/Engine/Plugins/Runtime/RigVM/Source/RigVMDeveloper/Private/EdGraph/RigVMEdGraph.cpp
--- a/Engine/Plugins/Runtime/RigVM/Source/RigVMDeveloper/Private/EdGraph/RigVMEdGraph.cpp
+++ b/Engine/Plugins/Runtime/RigVM/Source/RigVMDeveloper/Private/EdGraph/RigVMEdGraph.cpp
@@ -920,18 +920,32 @@ const URigVMEdGraph* URigVMEdGraph::GetRootGraph() const
 
 void URigVMEdGraph::AddNode(UEdGraphNode* NodeToAdd, bool bUserAction, bool bSelectNewNode)
 {
-	// Comments are added outside of the RigVMEditor, so we add here the node to the model
+	// Comments are added outside of the ControlRigEditor, so we add here the node to the model
 	if (const UEdGraphNode_Comment* CommentNode = Cast<const UEdGraphNode_Comment>(NodeToAdd))
 	{
 		if (URigVMController* Controller = GetBlueprint()->GetOrCreateController(GetModel()))
 		{
 			if (GetModel()->FindNodeByName(NodeToAdd->GetFName()) == nullptr) // When recreating nodes at RebuildGraphFromModel, the model node already exists
 			{
+				if (GEditor && !GIsTransacting)
+				{
+					// FBlueprintActionMenuItem::PerformAction has a super high level FScopedTransaction
+					// FScopedTransaction Transaction(LOCTEXT("AddNodeTransaction", "Add Node"));
+					// That scoped transaction stores the whole graph change, which is not what we want, as it deletes the Subgraphs
+					// when we perform an redo of a comment after performing the redo of an Add Pin of a Sequence.
+					// The issue happens because the Redo of the Add Pin changes the Subgraph object, so when we apply 
+					// the Comment Redo with all the Graph changes, it clears and serializes the SubGraphs array,
+					// making it invalid, as the subgraph has changed during the previous AddPin Redo
+					// For other node types, UControlRigUnitNodeSpawner::Invoke cancels the transaction when the node is added,
+					GEditor->CancelTransaction(0);
+				}
+
 				TGuardValue<bool> BlueprintNotifGuard(GetBlueprint()->bSuspendModelNotificationsForOthers, true);
 				FVector2D NodePos(CommentNode->NodePosX, CommentNode->NodePosY);
 				FVector2D NodeSize(CommentNode->NodeWidth, CommentNode->NodeHeight);
 				FLinearColor NodeColor = CommentNode->CommentColor;
-				Controller->AddCommentNode(CommentNode->NodeComment, NodePos, NodeSize, NodeColor, CommentNode->GetName(), GIsTransacting == false, true);
+				Controller->AddCommentNode(CommentNode->NodeComment, NodePos, NodeSize, NodeColor, CommentNode->GetName(), !GIsTransacting, true);
+				ModelNodePathToEdNode.Add(NodeToAdd->GetFName(), NodeToAdd);
 			}
 		}
 	}
