RigVMEdGraph Non Unityi build fix : GEditor'GEditor': undeclared identifier

diff --git a/Engine/Plugins/Runtime/RigVM/Source/RigVMDeveloper/Private/EdGraph/RigVMEdGraph.cpp b/Engine/Plugins/Runtime/RigVM/Source/RigVMDeveloper/Private/EdGraph/RigVMEdGraph.cpp
--- a/Engine/Plugins/Runtime/RigVM/Source/RigVMDeveloper/Private/EdGraph/RigVMEdGraph.cpp
+++ b/Engine/Plugins/Runtime/RigVM/Source/RigVMDeveloper/Private/EdGraph/RigVMEdGraph.cpp
@@ -3,6 +3,7 @@
 #include "EdGraph/RigVMEdGraph.h"
 #include "RigVMBlueprint.h"
 #include "RigVMBlueprintGeneratedClass.h"
+#include "Editor/EditorEngine.h"
 #include "EdGraph/RigVMEdGraphSchema.h"
 #include "RigVMModel/RigVMGraph.h"
 #include "RigVMObjectVersion.h"
@@ -24,6 +25,11 @@
 
 #define LOCTEXT_NAMESPACE "RigVMEdGraph"
 
+#if WITH_EDITOR
+/** The editor object. */
+extern UNREALED_API class UEditorEngine* GEditor;
+#endif
+
 URigVMEdGraph::URigVMEdGraph()
 {
 	bSuspendModelNotifications = false;
@@ -927,6 +933,7 @@ void URigVMEdGraph::AddNode(UEdGraphNode* NodeToAdd, bool bUserAction, bool bSel
 		{
 			if (GetModel()->FindNodeByName(NodeToAdd->GetFName()) == nullptr) // When recreating nodes at RebuildGraphFromModel, the model node already exists
 			{
+#if WITH_EDITOR
 				if (GEditor && !GIsTransacting)
 				{
 					// FBlueprintActionMenuItem::PerformAction has a super high level FScopedTransaction
@@ -939,7 +946,7 @@ void URigVMEdGraph::AddNode(UEdGraphNode* NodeToAdd, bool bUserAction, bool bSel
 					// For other node types, UControlRigUnitNodeSpawner::Invoke cancels the transaction when the node is added,
 					GEditor->CancelTransaction(0);
 				}
-
+#endif // WITH_EDITOR
 				TGuardValue<bool> BlueprintNotifGuard(GetBlueprint()->bSuspendModelNotificationsForOthers, true);
 				FVector2D NodePos(CommentNode->NodePosX, CommentNode->NodePosY);
 				FVector2D NodeSize(CommentNode->NodeWidth, CommentNode->NodeHeight);
