RigVM: Move Control Rig Node Spawners to RigVMEditor

diff --git a/Engine/Plugins/Animation/ControlRig/Config/BaseControlRig.ini b/Engine/Plugins/Animation/ControlRig/Config/BaseControlRig.ini
--- a/Engine/Plugins/Animation/ControlRig/Config/BaseControlRig.ini
+++ b/Engine/Plugins/Animation/ControlRig/Config/BaseControlRig.ini
@@ -434,3 +434,5 @@
 ; moving editor side types to RigVMEditor module
 +StructRedirects=(OldName="/Script/ControlRig.RigGraphDisplaySettings",NewName="/Script/RigVMDeveloper.RigVMEdGraphDisplaySettings")
 +StructRedirects=(OldName="/Script/ControlRig.ControlRigPythonSettings",NewName="/Script/RigVMDeveloper.RigVMPythonSettings")
++StructRedirects=(OldName="/Script/ControlRig.ControlRigPublicFunctionArg",NewName="/Script/RigVMDeveloper.RigVMOldPublicFunctionArg")
++StructRedirects=(OldName="/Script/ControlRig.ControlRigPublicFunctionData",NewName="/Script/RigVMDeveloper.RigVMOldPublicFunctionData")
diff --git a/Engine/Plugins/Animation/ControlRig/Source/ControlRigDeveloper/Private/ControlRigBlueprint.cpp b/Engine/Plugins/Animation/ControlRig/Source/ControlRigDeveloper/Private/ControlRigBlueprint.cpp
--- a/Engine/Plugins/Animation/ControlRig/Source/ControlRigDeveloper/Private/ControlRigBlueprint.cpp
+++ b/Engine/Plugins/Animation/ControlRig/Source/ControlRigDeveloper/Private/ControlRigBlueprint.cpp
@@ -38,40 +38,6 @@
 
 #define LOCTEXT_NAMESPACE "ControlRigBlueprint"
 
-FEdGraphPinType FControlRigPublicFunctionArg::GetPinType() const
-{
-	FRigVMExternalVariable Variable;
-	Variable.Name = Name;
-	Variable.bIsArray = bIsArray;
-	Variable.TypeName = CPPType;
-	
-	if(CPPTypeObjectPath.IsValid())
-	{
-		Variable.TypeObject = RigVMTypeUtils::FindObjectFromCPPTypeObjectPath(CPPTypeObjectPath.ToString());
-	}
-
-	return RigVMTypeUtils::PinTypeFromExternalVariable(Variable);
-}
-
-bool FControlRigPublicFunctionData::IsMutable() const
-{
-	for(const FControlRigPublicFunctionArg& Arg : Arguments)
-	{
-		if(!Arg.CPPTypeObjectPath.IsNone())
-		{
-			if(UScriptStruct* Struct = Cast<UScriptStruct>(
-				RigVMTypeUtils::FindObjectFromCPPTypeObjectPath(Arg.CPPTypeObjectPath.ToString())))
-			{
-				if(Struct->IsChildOf(FRigVMExecuteContext::StaticStruct()))
-				{
-					return true;
-				}
-			}
-		}
-	}
-	return false;
-}
-
 TArray<UControlRigBlueprint*> UControlRigBlueprint::sCurrentlyOpenedRigBlueprints;
 
 UControlRigBlueprint::UControlRigBlueprint(const FObjectInitializer& ObjectInitializer)
@@ -888,9 +854,9 @@ void UControlRigBlueprint::PatchFunctionsOnLoad()
 	TArray<FName> BackwardsCompatiblePublicFunctions;
 	if (GetLinkerCustomVersion(FControlRigObjectVersion::GUID) < FControlRigObjectVersion::StoreFunctionsInGeneratedClass)
 	{
-		for (const FControlRigPublicFunctionData& Data : PublicFunctions_DEPRECATED)
+		for (const FRigVMOldPublicFunctionData& OldPublicFunction : PublicFunctions_DEPRECATED)
 		{
-			BackwardsCompatiblePublicFunctions.Add(Data.Name);
+			BackwardsCompatiblePublicFunctions.Add(OldPublicFunction.Name);
 		}
 	}
 	else
diff --git a/Engine/Plugins/Animation/ControlRig/Source/ControlRigDeveloper/Public/ControlRigBlueprint.h b/Engine/Plugins/Animation/ControlRig/Source/ControlRigDeveloper/Public/ControlRigBlueprint.h
--- a/Engine/Plugins/Animation/ControlRig/Source/ControlRigDeveloper/Public/ControlRigBlueprint.h
+++ b/Engine/Plugins/Animation/ControlRig/Source/ControlRigDeveloper/Public/ControlRigBlueprint.h
@@ -27,66 +27,6 @@ class USkeletalMesh;
 class UControlRigGraph;
 struct FEndLoadPackageContext;
 
-USTRUCT(meta = (Deprecated = "5.2"))
-struct CONTROLRIGDEVELOPER_API FControlRigPublicFunctionArg
-{
-	GENERATED_BODY();
-	
-	FControlRigPublicFunctionArg()
-	: Name(NAME_None)
-	, CPPType(NAME_None)
-	, CPPTypeObjectPath(NAME_None)
-	, bIsArray(false)
-	, Direction(ERigVMPinDirection::Input)
-	{}
-
-	UPROPERTY()
-	FName Name;
-
-	UPROPERTY()
-	FName CPPType;
-
-	UPROPERTY()
-	FName CPPTypeObjectPath;
-
-	UPROPERTY()
-	bool bIsArray;
-
-	UPROPERTY()
-	ERigVMPinDirection Direction;
-
-	FEdGraphPinType GetPinType() const;
-};
-
-USTRUCT(meta = (Deprecated = "5.2"))
-struct CONTROLRIGDEVELOPER_API FControlRigPublicFunctionData
-{
-	GENERATED_BODY();
-
-	FControlRigPublicFunctionData()
-		:Name(NAME_None)
-	{}
-
-	UPROPERTY()
-	FName Name;
-
-	UPROPERTY()
-	FString DisplayName;
-
-	UPROPERTY()
-	FString Category;
-
-	UPROPERTY()
-	FString Keywords;
-
-	UPROPERTY()
-	FControlRigPublicFunctionArg ReturnValue;
-
-	UPROPERTY()
-	TArray<FControlRigPublicFunctionArg> Arguments;
-
-	bool IsMutable() const;
-};
 
 UCLASS(BlueprintType, meta=(IgnoreClassThumbnail))
 class CONTROLRIGDEVELOPER_API UControlRigBlueprint : public URigVMBlueprint, public IInterface_PreviewMeshProvider
@@ -141,7 +81,7 @@ protected:
 
 	/** Asset searchable information about exposed public functions on this rig */
 	UPROPERTY(AssetRegistrySearchable)
-	TArray<FControlRigPublicFunctionData> PublicFunctions_DEPRECATED;
+	TArray<FRigVMOldPublicFunctionData> PublicFunctions_DEPRECATED;
 
 	virtual void SetupDefaultObjectDuringCompilation(URigVMHost* InCDO) override;
 
diff --git a/Engine/Plugins/Animation/ControlRig/Source/ControlRigDeveloper/Public/Graph/ControlRigGraph.h b/Engine/Plugins/Animation/ControlRig/Source/ControlRigDeveloper/Public/Graph/ControlRigGraph.h
--- a/Engine/Plugins/Animation/ControlRig/Source/ControlRigDeveloper/Public/Graph/ControlRigGraph.h
+++ b/Engine/Plugins/Animation/ControlRig/Source/ControlRigDeveloper/Public/Graph/ControlRigGraph.h
@@ -124,18 +124,6 @@ private:
 
 #endif
 
-	friend class UControlRigUnitNodeSpawner;
-	friend class UControlRigVariableNodeSpawner;
-	friend class UControlRigParameterNodeSpawner;
-	friend class UControlRigBranchNodeSpawner;
-	friend class UControlRigIfNodeSpawner;
-	friend class UControlRigSelectNodeSpawner;
-	friend class UControlRigTemplateNodeSpawner;
-	friend class UControlRigEnumNodeSpawner;
-	friend class UControlRigFunctionRefNodeSpawner;
-	friend class UControlRigArrayNodeSpawner;
-	friend class UControlRigInvokeEntryNodeSpawner;
-
 	friend class UControlRigGraphNode;
 	friend class FControlRigEditor;
 	friend class SControlRigGraphNode;
diff --git a/Engine/Plugins/Animation/ControlRig/Source/ControlRigEditor/Private/ControlRigEditorModule.cpp b/Engine/Plugins/Animation/ControlRig/Source/ControlRigEditor/Private/ControlRigEditorModule.cpp
--- a/Engine/Plugins/Animation/ControlRig/Source/ControlRigEditor/Private/ControlRigEditorModule.cpp
+++ b/Engine/Plugins/Animation/ControlRig/Source/ControlRigEditor/Private/ControlRigEditorModule.cpp
@@ -46,12 +46,12 @@
 #include "Graph/ControlRigGraphSchema.h"
 #include "Graph/ControlRigGraph.h"
 #include "EdGraph/EdGraph.h"
-#include "Graph/NodeSpawners/ControlRigUnitNodeSpawner.h"
-#include "Graph/NodeSpawners/ControlRigVariableNodeSpawner.h"
-#include "Graph/NodeSpawners/ControlRigTemplateNodeSpawner.h"
-#include "Graph/NodeSpawners/ControlRigEnumNodeSpawner.h"
-#include "Graph/NodeSpawners/ControlRigFunctionRefNodeSpawner.h"
-#include "Graph/NodeSpawners/ControlRigInvokeEntryNodeSpawner.h"
+#include "EdGraph/NodeSpawners/RigVMEdGraphUnitNodeSpawner.h"
+#include "EdGraph/NodeSpawners/RigVMEdGraphVariableNodeSpawner.h"
+#include "EdGraph/NodeSpawners/RigVMEdGraphTemplateNodeSpawner.h"
+#include "EdGraph/NodeSpawners/RigVMEdGraphEnumNodeSpawner.h"
+#include "EdGraph/NodeSpawners/RigVMEdGraphFunctionRefNodeSpawner.h"
+#include "EdGraph/NodeSpawners/RigVMEdGraphInvokeEntryNodeSpawner.h"
 #include "Kismet2/BlueprintEditorUtils.h"
 #include "Kismet2/KismetDebugUtilities.h"
 #include "Graph/ControlRigGraphNode.h"
@@ -1036,7 +1036,7 @@ void FControlRigEditorModule::GetTypeActions(UControlRigBlueprint* CRB, FBluepri
 		FText MenuDesc = FText::FromName(Template.GetName());
 		FText ToolTip = Template.GetTooltipText();
 
-		UBlueprintNodeSpawner* NodeSpawner = UControlRigTemplateNodeSpawner::CreateFromNotation(Template.GetNotation(), MenuDesc, NodeCategory, ToolTip);
+		UBlueprintNodeSpawner* NodeSpawner = URigVMEdGraphTemplateNodeSpawner::CreateFromNotation(Template.GetNotation(), MenuDesc, NodeCategory, ToolTip);
 		check(NodeSpawner != nullptr);
 		ActionRegistrar.AddBlueprintAction(ActionKey, NodeSpawner);
 	};
@@ -1058,7 +1058,7 @@ void FControlRigEditorModule::GetTypeActions(UControlRigBlueprint* CRB, FBluepri
 		FText MenuDesc = FText::FromString(Factory->GetNodeTitle(FRigVMTemplateTypeMap()));
 		FText ToolTip = Factory->GetNodeTooltip(FRigVMTemplateTypeMap());
 
-		UBlueprintNodeSpawner* NodeSpawner = UControlRigTemplateNodeSpawner::CreateFromNotation(Template->GetNotation(), MenuDesc, NodeCategory, ToolTip);
+		UBlueprintNodeSpawner* NodeSpawner = URigVMEdGraphTemplateNodeSpawner::CreateFromNotation(Template->GetNotation(), MenuDesc, NodeCategory, ToolTip);
 		check(NodeSpawner != nullptr);
 		ActionRegistrar.AddBlueprintAction(ActionKey, NodeSpawner);
 	};
@@ -1106,7 +1106,7 @@ void FControlRigEditorModule::GetTypeActions(UControlRigBlueprint* CRB, FBluepri
 		FText MenuDesc = FText::FromString(DisplayNameMetadata + MenuDescSuffixMetadata);
 		FText ToolTip = Struct->GetToolTipText();
 
-		UBlueprintNodeSpawner* NodeSpawner = UControlRigUnitNodeSpawner::CreateFromStruct(Struct, Function.GetMethodName(), MenuDesc, NodeCategory, ToolTip);
+		UBlueprintNodeSpawner* NodeSpawner = URigVMEdGraphUnitNodeSpawner::CreateFromStruct(Struct, Function.GetMethodName(), MenuDesc, NodeCategory, ToolTip);
 		check(NodeSpawner != nullptr);
 		ActionRegistrar.AddBlueprintAction(ActionKey, NodeSpawner);
 	};
@@ -1134,7 +1134,7 @@ void FControlRigEditorModule::GetTypeActions(UControlRigBlueprint* CRB, FBluepri
 		FText MenuDesc = FText::FromString(FString::Printf(TEXT("Enum %s"), *EnumToConsider->GetName()));
 		FText ToolTip = MenuDesc;
 
-		UBlueprintNodeSpawner* NodeSpawner = UControlRigEnumNodeSpawner::CreateForEnum(EnumToConsider, MenuDesc, NodeCategory, ToolTip);
+		UBlueprintNodeSpawner* NodeSpawner = URigVMEdGraphEnumNodeSpawner::CreateForEnum(EnumToConsider, MenuDesc, NodeCategory, ToolTip);
 		check(NodeSpawner != nullptr);
 		ActionRegistrar.AddBlueprintAction(ActionKey, NodeSpawner);
 	}
@@ -1180,11 +1180,11 @@ void FControlRigEditorModule::GetTypeActions(UControlRigBlueprint* CRB, FBluepri
 
 			if (PublicFunctionsProperty && !PublicFunctionsString.IsEmpty())
 			{
-				TArray<FControlRigPublicFunctionData> PublicFunctions;
+				TArray<FRigVMOldPublicFunctionData> PublicFunctions;
 				PublicFunctionsProperty->ImportText_Direct(*PublicFunctionsString, &PublicFunctions, nullptr, EPropertyPortFlags::PPF_None);
-				for(const FControlRigPublicFunctionData& PublicFunction : PublicFunctions)
+				for(const FRigVMOldPublicFunctionData& PublicFunction : PublicFunctions)
 				{
-					UBlueprintNodeSpawner* NodeSpawner = UControlRigFunctionRefNodeSpawner::CreateFromAssetData(ControlRigAssetData, PublicFunction);
+					UBlueprintNodeSpawner* NodeSpawner = URigVMEdGraphFunctionRefNodeSpawner::CreateFromAssetData(ControlRigAssetData, PublicFunction);
 					check(NodeSpawner != nullptr);
 					ActionRegistrar.AddBlueprintAction(ActionKey, NodeSpawner);
 				}
@@ -1196,7 +1196,7 @@ void FControlRigEditorModule::GetTypeActions(UControlRigBlueprint* CRB, FBluepri
 				PublicGraphFunctionsProperty->ImportText_Direct(*PublicGraphFunctionsString, &PublicFunctions, nullptr, EPropertyPortFlags::PPF_None);
 				for(const FRigVMGraphFunctionHeader& PublicFunction : PublicFunctions)
 				{
-					UBlueprintNodeSpawner* NodeSpawner = UControlRigFunctionRefNodeSpawner::CreateFromAssetData(ControlRigAssetData, PublicFunction);
+					UBlueprintNodeSpawner* NodeSpawner = URigVMEdGraphFunctionRefNodeSpawner::CreateFromAssetData(ControlRigAssetData, PublicFunction);
 					check(NodeSpawner != nullptr);
 					ActionRegistrar.AddBlueprintAction(ActionKey, NodeSpawner);
 				}
@@ -1218,10 +1218,10 @@ void FControlRigEditorModule::GetInstanceActions(UControlRigBlueprint* CRB, FBlu
 			{
 				FText MenuDesc = FText::FromName(ExternalVariable.Name);
 				FText ToolTip = FText::FromString(FString::Printf(TEXT("Get the value of variable %s"), *ExternalVariable.Name.ToString()));
-				ActionRegistrar.AddBlueprintAction(GeneratedClass, UControlRigVariableNodeSpawner::CreateFromExternalVariable(CRB, ExternalVariable, true, MenuDesc, NodeCategory, ToolTip));
+				ActionRegistrar.AddBlueprintAction(GeneratedClass, URigVMEdGraphVariableNodeSpawner::CreateFromExternalVariable(CRB, ExternalVariable, true, MenuDesc, NodeCategory, ToolTip));
 
 				ToolTip = FText::FromString(FString::Printf(TEXT("Set the value of variable %s"), *ExternalVariable.Name.ToString()));
-				ActionRegistrar.AddBlueprintAction(GeneratedClass, UControlRigVariableNodeSpawner::CreateFromExternalVariable(CRB, ExternalVariable, false, MenuDesc, NodeCategory, ToolTip));
+				ActionRegistrar.AddBlueprintAction(GeneratedClass, URigVMEdGraphVariableNodeSpawner::CreateFromExternalVariable(CRB, ExternalVariable, false, MenuDesc, NodeCategory, ToolTip));
 			}
 		}
 
@@ -1230,7 +1230,7 @@ void FControlRigEditorModule::GetInstanceActions(UControlRigBlueprint* CRB, FBlu
 			TArray<URigVMLibraryNode*> Functions = LocalFunctionLibrary->GetFunctions();
 			for (URigVMLibraryNode* Function : Functions)
 			{
-				UBlueprintNodeSpawner* NodeSpawner = UControlRigFunctionRefNodeSpawner::CreateFromFunction(Function);
+				UBlueprintNodeSpawner* NodeSpawner = URigVMEdGraphFunctionRefNodeSpawner::CreateFromFunction(Function);
 				check(NodeSpawner != nullptr);
 				ActionRegistrar.AddBlueprintAction(GeneratedClass, NodeSpawner);
 			}
@@ -1242,10 +1242,10 @@ void FControlRigEditorModule::GetInstanceActions(UControlRigBlueprint* CRB, FBlu
 				{
 					FText MenuDesc = FText::FromName(LocalVariable.Name);
 					FText ToolTip = FText::FromString(FString::Printf(TEXT("Get the value of variable %s"), *LocalVariable.Name.ToString()));
-					ActionRegistrar.AddBlueprintAction(GeneratedClass, UControlRigVariableNodeSpawner::CreateFromLocalVariable(CRB, Function->GetContainedGraph(), LocalVariable, true, MenuDesc, NodeCategory, ToolTip));
+					ActionRegistrar.AddBlueprintAction(GeneratedClass, URigVMEdGraphVariableNodeSpawner::CreateFromLocalVariable(CRB, Function->GetContainedGraph(), LocalVariable, true, MenuDesc, NodeCategory, ToolTip));
 
 					ToolTip = FText::FromString(FString::Printf(TEXT("Set the value of variable %s"), *LocalVariable.Name.ToString()));
-					ActionRegistrar.AddBlueprintAction(GeneratedClass, UControlRigVariableNodeSpawner::CreateFromLocalVariable(CRB, Function->GetContainedGraph(), LocalVariable, false, MenuDesc, NodeCategory, ToolTip));
+					ActionRegistrar.AddBlueprintAction(GeneratedClass, URigVMEdGraphVariableNodeSpawner::CreateFromLocalVariable(CRB, Function->GetContainedGraph(), LocalVariable, false, MenuDesc, NodeCategory, ToolTip));
 				}
 			}
 		}
@@ -1259,7 +1259,7 @@ void FControlRigEditorModule::GetInstanceActions(UControlRigBlueprint* CRB, FBlu
 				{
 					FText MenuDesc = FText::FromName(InputArgument.Name);
 					FText ToolTip = FText::FromString(FString::Printf(TEXT("Get the value of input %s"), *InputArgument.Name.ToString()));
-					ActionRegistrar.AddBlueprintAction(GeneratedClass, UControlRigVariableNodeSpawner::CreateFromLocalVariable(CRB, Graph, InputArgument, true, MenuDesc, NodeCategory, ToolTip));
+					ActionRegistrar.AddBlueprintAction(GeneratedClass, URigVMEdGraphVariableNodeSpawner::CreateFromLocalVariable(CRB, Graph, InputArgument, true, MenuDesc, NodeCategory, ToolTip));
 				}			
 			}
 		}
@@ -1275,7 +1275,7 @@ void FControlRigEditorModule::GetInstanceActions(UControlRigBlueprint* CRB, FBlu
 				FString Suffix = EntryName.ToString().EndsWith(EventStr) ? FString() : EventSuffix;
 				FText MenuDesc = FText::FromString(FString::Printf(TEXT("Run %s%s"), *EntryName.ToString(), *Suffix));
 				FText ToolTip = FText::FromString(FString::Printf(TEXT("Runs the %s%s"), *EntryName.ToString(), *Suffix));
-				ActionRegistrar.AddBlueprintAction(GeneratedClass, UControlRigInvokeEntryNodeSpawner::CreateForEntry(CRB, EntryName, MenuDesc, NodeCategory, ToolTip));
+				ActionRegistrar.AddBlueprintAction(GeneratedClass, URigVMEdGraphInvokeEntryNodeSpawner::CreateForEntry(CRB, EntryName, MenuDesc, NodeCategory, ToolTip));
 			}
 		}
 	}
diff --git a/Engine/Plugins/Animation/ControlRig/Source/ControlRigEditor/Private/Editor/ControlRigEditor.cpp b/Engine/Plugins/Animation/ControlRig/Source/ControlRigEditor/Private/Editor/ControlRigEditor.cpp
--- a/Engine/Plugins/Animation/ControlRig/Source/ControlRigEditor/Private/Editor/ControlRigEditor.cpp
+++ b/Engine/Plugins/Animation/ControlRig/Source/ControlRigEditor/Private/Editor/ControlRigEditor.cpp
@@ -61,7 +61,7 @@
 #include "Units/Hierarchy/RigUnit_GetControlTransform.h"
 #include "Units/Hierarchy/RigUnit_SetCurveValue.h"
 #include "Units/Hierarchy/RigUnit_AddBoneTransform.h"
-#include "Graph/NodeSpawners/ControlRigUnitNodeSpawner.h"
+#include "EdGraph/NodeSpawners/RigVMEdGraphUnitNodeSpawner.h"
 #include "Graph/ControlRigGraphSchema.h"
 #include "ControlRigObjectVersion.h"
 #include "EdGraphUtilities.h"
@@ -6685,7 +6685,7 @@ void FControlRigEditor::HandleMakeElementGetterSetter(ERigElementGetterSetterTyp
 				}
 			}
 
-			UControlRigUnitNodeSpawner::HookupMutableNode(ModelNode, GetControlRigBlueprint());
+			URigVMEdGraphUnitNodeSpawner::HookupMutableNode(ModelNode, GetControlRigBlueprint());
 		}
 
 		NodePosition += NodePositionIncrement;
diff --git a/Engine/Plugins/Animation/ControlRig/Source/ControlRigEditor/Private/Graph/NodeSpawners/ControlRigEnumNodeSpawner.cpp b/Engine/Plugins/Animation/ControlRig/Source/ControlRigEditor/Private/Graph/NodeSpawners/ControlRigEnumNodeSpawner.cpp
deleted file mode 100644
diff --git a/Engine/Plugins/Animation/ControlRig/Source/ControlRigEditor/Private/Graph/NodeSpawners/ControlRigEnumNodeSpawner.h b/Engine/Plugins/Animation/ControlRig/Source/ControlRigEditor/Private/Graph/NodeSpawners/ControlRigEnumNodeSpawner.h
deleted file mode 100644
diff --git a/Engine/Plugins/Animation/ControlRig/Source/ControlRigEditor/Private/Graph/NodeSpawners/ControlRigFunctionRefNodeSpawner.cpp b/Engine/Plugins/Animation/ControlRig/Source/ControlRigEditor/Private/Graph/NodeSpawners/ControlRigFunctionRefNodeSpawner.cpp
deleted file mode 100644
diff --git a/Engine/Plugins/Animation/ControlRig/Source/ControlRigEditor/Private/Graph/NodeSpawners/ControlRigFunctionRefNodeSpawner.h b/Engine/Plugins/Animation/ControlRig/Source/ControlRigEditor/Private/Graph/NodeSpawners/ControlRigFunctionRefNodeSpawner.h
deleted file mode 100644
diff --git a/Engine/Plugins/Animation/ControlRig/Source/ControlRigEditor/Private/Graph/NodeSpawners/ControlRigInvokeEntryNodeSpawner.cpp b/Engine/Plugins/Animation/ControlRig/Source/ControlRigEditor/Private/Graph/NodeSpawners/ControlRigInvokeEntryNodeSpawner.cpp
deleted file mode 100644
diff --git a/Engine/Plugins/Animation/ControlRig/Source/ControlRigEditor/Private/Graph/NodeSpawners/ControlRigInvokeEntryNodeSpawner.h b/Engine/Plugins/Animation/ControlRig/Source/ControlRigEditor/Private/Graph/NodeSpawners/ControlRigInvokeEntryNodeSpawner.h
deleted file mode 100644
diff --git a/Engine/Plugins/Animation/ControlRig/Source/ControlRigEditor/Private/Graph/NodeSpawners/ControlRigTemplateNodeSpawner.cpp b/Engine/Plugins/Animation/ControlRig/Source/ControlRigEditor/Private/Graph/NodeSpawners/ControlRigTemplateNodeSpawner.cpp
deleted file mode 100644
diff --git a/Engine/Plugins/Animation/ControlRig/Source/ControlRigEditor/Private/Graph/NodeSpawners/ControlRigTemplateNodeSpawner.h b/Engine/Plugins/Animation/ControlRig/Source/ControlRigEditor/Private/Graph/NodeSpawners/ControlRigTemplateNodeSpawner.h
deleted file mode 100644
diff --git a/Engine/Plugins/Animation/ControlRig/Source/ControlRigEditor/Private/Graph/NodeSpawners/ControlRigUnitNodeSpawner.cpp b/Engine/Plugins/Animation/ControlRig/Source/ControlRigEditor/Private/Graph/NodeSpawners/ControlRigUnitNodeSpawner.cpp
deleted file mode 100644
diff --git a/Engine/Plugins/Animation/ControlRig/Source/ControlRigEditor/Private/Graph/NodeSpawners/ControlRigUnitNodeSpawner.h b/Engine/Plugins/Animation/ControlRig/Source/ControlRigEditor/Private/Graph/NodeSpawners/ControlRigUnitNodeSpawner.h
deleted file mode 100644
diff --git a/Engine/Plugins/Animation/ControlRig/Source/ControlRigEditor/Private/Graph/NodeSpawners/ControlRigVariableNodeSpawner.cpp b/Engine/Plugins/Animation/ControlRig/Source/ControlRigEditor/Private/Graph/NodeSpawners/ControlRigVariableNodeSpawner.cpp
deleted file mode 100644
diff --git a/Engine/Plugins/Animation/ControlRig/Source/ControlRigEditor/Private/Graph/NodeSpawners/ControlRigVariableNodeSpawner.h b/Engine/Plugins/Animation/ControlRig/Source/ControlRigEditor/Private/Graph/NodeSpawners/ControlRigVariableNodeSpawner.h
deleted file mode 100644
diff --git a/Engine/Plugins/Runtime/RigVM/Source/RigVMDeveloper/Private/RigVMBlueprint.cpp b/Engine/Plugins/Runtime/RigVM/Source/RigVMDeveloper/Private/RigVMBlueprint.cpp
--- a/Engine/Plugins/Runtime/RigVM/Source/RigVMDeveloper/Private/RigVMBlueprint.cpp
+++ b/Engine/Plugins/Runtime/RigVM/Source/RigVMDeveloper/Private/RigVMBlueprint.cpp
@@ -65,6 +65,40 @@ static TArray<UClass*> GetClassObjectsInPackage(UPackage* InPackage)
 	return ClassObjects;
 }
 
+FEdGraphPinType FRigVMOldPublicFunctionArg::GetPinType() const
+{
+	FRigVMExternalVariable Variable;
+	Variable.Name = Name;
+	Variable.bIsArray = bIsArray;
+	Variable.TypeName = CPPType;
+	
+	if(CPPTypeObjectPath.IsValid())
+	{
+		Variable.TypeObject = RigVMTypeUtils::FindObjectFromCPPTypeObjectPath(CPPTypeObjectPath.ToString());
+	}
+
+	return RigVMTypeUtils::PinTypeFromExternalVariable(Variable);
+}
+
+bool FRigVMOldPublicFunctionData::IsMutable() const
+{
+	for(const FRigVMOldPublicFunctionArg& Arg : Arguments)
+	{
+		if(!Arg.CPPTypeObjectPath.IsNone())
+		{
+			if(UScriptStruct* Struct = Cast<UScriptStruct>(
+				RigVMTypeUtils::FindObjectFromCPPTypeObjectPath(Arg.CPPTypeObjectPath.ToString())))
+			{
+				if(Struct->IsChildOf(FRigVMExecuteContext::StaticStruct()))
+				{
+					return true;
+				}
+			}
+		}
+	}
+	return false;
+}
+
 FSoftObjectPath URigVMBlueprint::PreDuplicateAssetPath;
 FSoftObjectPath URigVMBlueprint::PreDuplicateHostPath;
 
diff --git a/Engine/Plugins/Runtime/RigVM/Source/RigVMDeveloper/Public/EdGraph/RigVMEdGraph.h b/Engine/Plugins/Runtime/RigVM/Source/RigVMDeveloper/Public/EdGraph/RigVMEdGraph.h
--- a/Engine/Plugins/Runtime/RigVM/Source/RigVMDeveloper/Public/EdGraph/RigVMEdGraph.h
+++ b/Engine/Plugins/Runtime/RigVM/Source/RigVMDeveloper/Public/EdGraph/RigVMEdGraph.h
@@ -95,20 +95,6 @@ private:
 	TMap<FName, UEdGraphNode*> ModelNodePathToEdNode;
 	mutable TWeakObjectPtr<URigVMGraph> CachedModelGraph;
 
-	/*
-	friend class UControlRigUnitNodeSpawner;
-	friend class UControlRigVariableNodeSpawner;
-	friend class UControlRigParameterNodeSpawner;
-	friend class UControlRigBranchNodeSpawner;
-	friend class UControlRigIfNodeSpawner;
-	friend class UControlRigSelectNodeSpawner;
-	friend class UControlRigTemplateNodeSpawner;
-	friend class UControlRigEnumNodeSpawner;
-	friend class UControlRigFunctionRefNodeSpawner;
-	friend class UControlRigArrayNodeSpawner;
-	friend class UControlRigInvokeEntryNodeSpawner;
-	*/
-
 #endif
 	friend class URigVMEdGraphNode;
 	friend class URigVMBlueprint;
@@ -117,16 +103,15 @@ private:
 	friend class SControlRigGraphNode;
 	friend class URigVMBlueprint;
 
-	// todo
-	friend class UControlRigUnitNodeSpawner;
-	friend class UControlRigVariableNodeSpawner;
-	friend class UControlRigParameterNodeSpawner;
-	friend class UControlRigBranchNodeSpawner;
-	friend class UControlRigIfNodeSpawner;
-	friend class UControlRigSelectNodeSpawner;
-	friend class UControlRigTemplateNodeSpawner;
-	friend class UControlRigEnumNodeSpawner;
-	friend class UControlRigFunctionRefNodeSpawner;
-	friend class UControlRigArrayNodeSpawner;
-	friend class UControlRigInvokeEntryNodeSpawner;
+	friend class URigVMEdGraphUnitNodeSpawner;
+	friend class URigVMEdGraphVariableNodeSpawner;
+	friend class URigVMEdGraphParameterNodeSpawner;
+	friend class URigVMEdGraphBranchNodeSpawner;
+	friend class URigVMEdGraphIfNodeSpawner;
+	friend class URigVMEdGraphSelectNodeSpawner;
+	friend class URigVMEdGraphTemplateNodeSpawner;
+	friend class URigVMEdGraphEnumNodeSpawner;
+	friend class URigVMEdGraphFunctionRefNodeSpawner;
+	friend class URigVMEdGraphArrayNodeSpawner;
+	friend class URigVMEdGraphInvokeEntryNodeSpawner;
 };
diff --git a/Engine/Plugins/Runtime/RigVM/Source/RigVMDeveloper/Public/EdGraph/RigVMEdGraphNode.h b/Engine/Plugins/Runtime/RigVM/Source/RigVMDeveloper/Public/EdGraph/RigVMEdGraphNode.h
--- a/Engine/Plugins/Runtime/RigVM/Source/RigVMDeveloper/Public/EdGraph/RigVMEdGraphNode.h
+++ b/Engine/Plugins/Runtime/RigVM/Source/RigVMDeveloper/Public/EdGraph/RigVMEdGraphNode.h
@@ -273,18 +273,14 @@ private:
 	friend class SRigVMEdGraphNode;
 	friend class FRigVMArgumentLayout;
 	friend class FRigVMEdGraphDetails;
-	friend class URigVMTemplateNodeSpawner;
-	friend class URigVMArrayNodeSpawner;
-	friend class URigVMIfNodeSpawner;
-	friend class URigVMSelectNodeSpawner;
+	friend class URigVMEdGraphTemplateNodeSpawner;
+	friend class URigVMEdGraphArrayNodeSpawner;
+	friend class URigVMEdGraphIfNodeSpawner;
+	friend class URigVMEdGraphSelectNodeSpawner;
 
 	// todo
 	friend class FControlRigEditor;
 	friend class UControlRigBlueprint;
 	friend class SControlRigGraphNode;
 	friend class UControlRigGraphSchema;
-	friend class UControlRigTemplateNodeSpawner;
-	friend class UControlRigArrayNodeSpawner;
-	friend class UControlRigIfNodeSpawner;
-	friend class UControlRigSelectNodeSpawner;
 };
diff --git a/Engine/Plugins/Runtime/RigVM/Source/RigVMDeveloper/Public/EdGraph/RigVMEdGraphSchema.h b/Engine/Plugins/Runtime/RigVM/Source/RigVMDeveloper/Public/EdGraph/RigVMEdGraphSchema.h
--- a/Engine/Plugins/Runtime/RigVM/Source/RigVMDeveloper/Public/EdGraph/RigVMEdGraphSchema.h
+++ b/Engine/Plugins/Runtime/RigVM/Source/RigVMDeveloper/Public/EdGraph/RigVMEdGraphSchema.h
@@ -342,8 +342,8 @@ private:
 	mutable TMap<FName, FVector2D> NodePositionsDuringStart;
 #endif
 
-	friend class UControlRigIfNodeSpawner;
-	friend class UControlRigSelectNodeSpawner;
-	friend class UControlRigUnitNodeSpawner;
-	friend class UControlRigArrayNodeSpawner;
+	friend class URigVMEdGraphIfNodeSpawner;
+	friend class URigVMEdGraphSelectNodeSpawner;
+	friend class URigVMEdGraphUnitNodeSpawner;
+	friend class URigVMEdGraphArrayNodeSpawner;
 };
diff --git a/Engine/Plugins/Runtime/RigVM/Source/RigVMDeveloper/Public/RigVMBlueprint.h b/Engine/Plugins/Runtime/RigVM/Source/RigVMDeveloper/Public/RigVMBlueprint.h
--- a/Engine/Plugins/Runtime/RigVM/Source/RigVMDeveloper/Public/RigVMBlueprint.h
+++ b/Engine/Plugins/Runtime/RigVM/Source/RigVMDeveloper/Public/RigVMBlueprint.h
@@ -127,6 +127,68 @@ enum class ERigVMBlueprintLoadType : uint8
 	CheckUserDefinedStructs
 };
 
+USTRUCT(meta = (Deprecated = "5.2"))
+struct RIGVMDEVELOPER_API FRigVMOldPublicFunctionArg
+{
+	GENERATED_BODY();
+	
+	FRigVMOldPublicFunctionArg()
+	: Name(NAME_None)
+	, CPPType(NAME_None)
+	, CPPTypeObjectPath(NAME_None)
+	, bIsArray(false)
+	, Direction(ERigVMPinDirection::Input)
+	{}
+
+	UPROPERTY()
+	FName Name;
+
+	UPROPERTY()
+	FName CPPType;
+
+	UPROPERTY()
+	FName CPPTypeObjectPath;
+
+	UPROPERTY()
+	bool bIsArray;
+
+	UPROPERTY()
+	ERigVMPinDirection Direction;
+
+	FEdGraphPinType GetPinType() const;
+};
+
+USTRUCT(meta = (Deprecated = "5.2"))
+struct RIGVMDEVELOPER_API FRigVMOldPublicFunctionData
+{
+	GENERATED_BODY();
+
+	FRigVMOldPublicFunctionData()
+		:Name(NAME_None)
+	{}
+
+	UPROPERTY()
+	FName Name;
+
+	UPROPERTY()
+	FString DisplayName;
+
+	UPROPERTY()
+	FString Category;
+
+	UPROPERTY()
+	FString Keywords;
+
+	UPROPERTY()
+	FRigVMOldPublicFunctionArg ReturnValue;
+
+	UPROPERTY()
+	TArray<FRigVMOldPublicFunctionArg> Arguments;
+
+	bool IsMutable() const;
+};
+
+
 UCLASS(BlueprintType, meta=(IgnoreClassThumbnail))
 class RIGVMDEVELOPER_API URigVMBlueprint : public UBlueprint, public IRigVMClientHost
 {
diff --git a/Engine/Plugins/Runtime/RigVM/Source/RigVMDeveloper/Public/RigVMModel/Nodes/RigVMEnumNode.h b/Engine/Plugins/Runtime/RigVM/Source/RigVMDeveloper/Public/RigVMModel/Nodes/RigVMEnumNode.h
--- a/Engine/Plugins/Runtime/RigVM/Source/RigVMDeveloper/Public/RigVMModel/Nodes/RigVMEnumNode.h
+++ b/Engine/Plugins/Runtime/RigVM/Source/RigVMDeveloper/Public/RigVMModel/Nodes/RigVMEnumNode.h
@@ -60,7 +60,7 @@ private:
 	friend class URigVMController;
 	friend class URigVMCompiler;
 	friend struct FRigVMAddEnumNodeAction;
-	friend class UControlRigEnumNodeSpawner;
+	friend class URigVMEdGraphEnumNodeSpawner;
 	friend class FRigVMParserAST;
 };
 
diff --git a/Engine/Plugins/Runtime/RigVM/Source/RigVMDeveloper/Public/RigVMModel/Nodes/RigVMIfNode.h b/Engine/Plugins/Runtime/RigVM/Source/RigVMDeveloper/Public/RigVMModel/Nodes/RigVMIfNode.h
--- a/Engine/Plugins/Runtime/RigVM/Source/RigVMDeveloper/Public/RigVMModel/Nodes/RigVMIfNode.h
+++ b/Engine/Plugins/Runtime/RigVM/Source/RigVMDeveloper/Public/RigVMModel/Nodes/RigVMIfNode.h
@@ -43,7 +43,7 @@ private:
 	friend class URigVMController;
 	friend class URigVMCompiler;
 	friend struct FRigVMAddIfNodeAction;
-	friend class UControlRigIfNodeSpawner;
+	friend class URigVMEdGraphIfNodeSpawner;
 	friend class FRigVMParserAST;
 	friend struct FRigVMRemoveNodeAction;
 };
diff --git a/Engine/Plugins/Runtime/RigVM/Source/RigVMDeveloper/Public/RigVMModel/Nodes/RigVMSelectNode.h b/Engine/Plugins/Runtime/RigVM/Source/RigVMDeveloper/Public/RigVMModel/Nodes/RigVMSelectNode.h
--- a/Engine/Plugins/Runtime/RigVM/Source/RigVMDeveloper/Public/RigVMModel/Nodes/RigVMSelectNode.h
+++ b/Engine/Plugins/Runtime/RigVM/Source/RigVMDeveloper/Public/RigVMModel/Nodes/RigVMSelectNode.h
@@ -48,7 +48,7 @@ private:
 	friend class URigVMController;
 	friend class URigVMCompiler;
 	friend struct FRigVMAddSelectNodeAction;
-	friend class UControlRigSelectNodeSpawner;
+	friend class URigVMEdGraphSelectNodeSpawner;
 	friend class FRigVMParserAST;
 	friend class FRigVMSelectExprAST;
 	friend struct FRigVMRemoveNodeAction;
diff --git a/Engine/Plugins/Runtime/RigVM/Source/RigVMDeveloper/Public/RigVMModel/RigVMController.h b/Engine/Plugins/Runtime/RigVM/Source/RigVMDeveloper/Public/RigVMModel/RigVMController.h
--- a/Engine/Plugins/Runtime/RigVM/Source/RigVMDeveloper/Public/RigVMModel/RigVMController.h
+++ b/Engine/Plugins/Runtime/RigVM/Source/RigVMDeveloper/Public/RigVMModel/RigVMController.h
@@ -1455,7 +1455,7 @@ private:
 	friend struct FRigVMClient;
 	friend struct FRigVMActionWrapper;
 	friend class URigVMSchema;
-	friend class UControlRigFunctionRefNodeSpawner;
+	friend class URigVMEdGraphFunctionRefNodeSpawner;
 };
 
 class FRigVMControllerNotifGuard
diff --git a/Engine/Plugins/Runtime/RigVM/Source/RigVMEditor/Private/EdGraph/NodeSpawners/RigVMEdGraphEnumNodeSpawner.cpp b/Engine/Plugins/Runtime/RigVM/Source/RigVMEditor/Private/EdGraph/NodeSpawners/RigVMEdGraphEnumNodeSpawner.cpp
new file mode 100644
--- /dev/null
+++ b/Engine/Plugins/Runtime/RigVM/Source/RigVMEditor/Private/EdGraph/NodeSpawners/RigVMEdGraphEnumNodeSpawner.cpp
@@ -0,0 +1,133 @@
+// Copyright Epic Games, Inc. All Rights Reserved.
+
+#include "EdGraph/NodeSpawners/RigVMEdGraphEnumNodeSpawner.h"
+#include "EdGraph/NodeSpawners/RigVMEdGraphUnitNodeSpawner.h"
+#include "EdGraph/RigVMEdGraph.h"
+#include "EdGraph/RigVMEdGraphNode.h"
+#include "EdGraph/RigVMEdGraphSchema.h"
+#include "BlueprintNodeTemplateCache.h"
+#include "RigVMModel/RigVMController.h"
+#include "RigVMBlueprint.h"
+
+#include "RigVMModel/Nodes/RigVMEnumNode.h"
+
+#include UE_INLINE_GENERATED_CPP_BY_NAME(RigVMEdGraphEnumNodeSpawner)
+
+#if WITH_EDITOR
+#include "Editor.h"
+#include "Kismet2/BlueprintEditorUtils.h"
+#endif
+
+#define LOCTEXT_NAMESPACE "RigVMEdGraphEnumNodeSpawner"
+
+URigVMEdGraphEnumNodeSpawner* URigVMEdGraphEnumNodeSpawner::CreateForEnum(UEnum* InEnum, const FText& InMenuDesc, const FText& InCategory, const FText& InTooltip)
+{
+	URigVMEdGraphEnumNodeSpawner* NodeSpawner = NewObject<URigVMEdGraphEnumNodeSpawner>(GetTransientPackage());
+	NodeSpawner->NodeClass = URigVMEdGraphNode::StaticClass();
+	NodeSpawner->Enum = InEnum;
+
+	FBlueprintActionUiSpec& MenuSignature = NodeSpawner->DefaultMenuSignature;
+	
+	MenuSignature.MenuName = InMenuDesc;
+	MenuSignature.Tooltip  = InTooltip;
+	MenuSignature.Category = InCategory;
+	MenuSignature.Keywords = FText::FromString(TEXT("Enum"));
+	MenuSignature.Icon = FSlateIcon(TEXT("ControlRigEditorStyle"), TEXT("ControlRig.RigUnit"));
+
+	return NodeSpawner;
+}
+
+FBlueprintNodeSignature URigVMEdGraphEnumNodeSpawner::GetSpawnerSignature() const
+{
+	return FBlueprintNodeSignature(FString("RigUnit=" + Enum->GetFName().ToString()));
+}
+
+FBlueprintActionUiSpec URigVMEdGraphEnumNodeSpawner::GetUiSpec(FBlueprintActionContext const& Context, FBindingSet const& Bindings) const
+{
+	UEdGraph* TargetGraph = (Context.Graphs.Num() > 0) ? Context.Graphs[0] : nullptr;
+	FBlueprintActionUiSpec MenuSignature = PrimeDefaultUiSpec(TargetGraph);
+
+	DynamicUiSignatureGetter.ExecuteIfBound(Context, Bindings, &MenuSignature);
+	return MenuSignature;
+}
+
+UEdGraphNode* URigVMEdGraphEnumNodeSpawner::Invoke(UEdGraph* ParentGraph, FBindingSet const& Bindings, FVector2D const Location) const
+{
+	URigVMEdGraphNode* NewNode = nullptr;
+
+	bool const bIsTemplateNode = FBlueprintNodeTemplateCache::IsTemplateOuter(ParentGraph);
+	bool const bIsUserFacingNode = !bIsTemplateNode;
+
+	if (bIsTemplateNode)
+	{
+		NewNode = NewObject<URigVMEdGraphNode>(ParentGraph, TEXT("EnumNode"));
+		ParentGraph->AddNode(NewNode, false);
+
+		NewNode->CreateNewGuid();
+		NewNode->PostPlacedNewNode();
+
+		UEdGraphPin* OutputValuePin = UEdGraphPin::CreatePin(NewNode);
+		NewNode->Pins.Add(OutputValuePin);
+
+		OutputValuePin->PinType.PinCategory = TEXT("int32");
+		OutputValuePin->Direction = EGPD_Output;
+		NewNode->SetFlags(RF_Transactional);
+
+		return NewNode;
+	}
+
+	// First create a backing member for our node
+	URigVMEdGraph* RigGraph = Cast<URigVMEdGraph>(ParentGraph);
+	if(RigGraph == nullptr) return nullptr;
+	URigVMBlueprint* RigBlueprint = Cast<URigVMBlueprint>(FBlueprintEditorUtils::FindBlueprintForGraph(ParentGraph));
+	check(RigBlueprint);
+
+#if WITH_EDITOR
+	if (GEditor && !bIsTemplateNode)
+	{
+		GEditor->CancelTransaction(0);
+	}
+#endif
+
+	URigVMController* Controller = bIsTemplateNode ? RigGraph->GetTemplateController() : RigBlueprint->GetController(ParentGraph);
+
+	FName Name = *URigVMEnumNode::EnumName;
+
+	if (!bIsTemplateNode)
+	{
+		Controller->OpenUndoBracket(FString::Printf(TEXT("Add '%s' Node"), *Name.ToString()));
+	}
+
+	if (URigVMEnumNode* ModelNode = Controller->AddEnumNode(*Enum->GetPathName(), Location, Name.ToString(), bIsUserFacingNode, !bIsTemplateNode))
+	{
+		NewNode = Cast<URigVMEdGraphNode>(RigGraph->FindNodeForModelNodeName(ModelNode->GetFName()));
+
+		if (NewNode && bIsUserFacingNode)
+		{
+			Controller->ClearNodeSelection(true);
+			Controller->SelectNode(ModelNode, true, true);
+		}
+
+		if (bIsUserFacingNode)
+		{
+			Controller->CloseUndoBracket();
+		}
+		else
+		{
+			Controller->RemoveNode(ModelNode, false);
+		}
+	}
+	else
+	{
+		if (bIsUserFacingNode)
+		{
+			Controller->CancelUndoBracket();
+		}
+	}
+
+
+	return NewNode;
+}
+
+#undef LOCTEXT_NAMESPACE
+
diff --git a/Engine/Plugins/Runtime/RigVM/Source/RigVMEditor/Private/EdGraph/NodeSpawners/RigVMEdGraphFunctionRefNodeSpawner.cpp b/Engine/Plugins/Runtime/RigVM/Source/RigVMEditor/Private/EdGraph/NodeSpawners/RigVMEdGraphFunctionRefNodeSpawner.cpp
new file mode 100644
--- /dev/null
+++ b/Engine/Plugins/Runtime/RigVM/Source/RigVMEditor/Private/EdGraph/NodeSpawners/RigVMEdGraphFunctionRefNodeSpawner.cpp
@@ -0,0 +1,439 @@
+// Copyright Epic Games, Inc. All Rights Reserved.
+
+#include "EdGraph/NodeSpawners/RigVMEdGraphFunctionRefNodeSpawner.h"
+#include "EdGraph/NodeSpawners/RigVMEdGraphUnitNodeSpawner.h"
+#include "EdGraph/RigVMEdGraph.h"
+#include "EdGraph/RigVMEdGraphNode.h"
+#include "EdGraph/RigVMEdGraphSchema.h"
+#include "EdGraphSchema_K2.h"
+#include "Kismet2/BlueprintEditorUtils.h"
+#include "Settings/EditorStyleSettings.h"
+#include "Editor/EditorEngine.h"
+#include "ObjectEditorUtils.h"
+#include "EditorCategoryUtils.h"
+#include "K2Node_Variable.h"
+#include "BlueprintNodeTemplateCache.h"
+#include "RigVMBlueprintUtils.h"
+#include "ScopedTransaction.h"
+
+#include UE_INLINE_GENERATED_CPP_BY_NAME(RigVMEdGraphFunctionRefNodeSpawner)
+
+#if WITH_EDITOR
+#include "Editor.h"
+#include "SGraphActionMenu.h"
+#include "GraphEditorSettings.h" 
+#endif
+
+#define LOCTEXT_NAMESPACE "RigVMEdGraphFunctionRefNodeSpawner"
+
+// 3 possible creation types:
+// - CreateFromFunction(URigVMLibraryNode) --> This is a local function. Valid Header.
+// - CreateFromAssetData(const FAssetData& InAssetData, const FRigVMGraphFunctionHeader& InPublicFunction) --> Public function. Valid Header.
+// - CreateFromAssetData(const FAssetData& InAssetData, const FRigVMOldPublicFunctionData& InPublicFunction) --> Public function. Valid AssetData and Header.Name
+
+URigVMEdGraphFunctionRefNodeSpawner* URigVMEdGraphFunctionRefNodeSpawner::CreateFromFunction(URigVMLibraryNode* InFunction)
+{
+	check(InFunction);
+
+	URigVMEdGraphFunctionRefNodeSpawner* NodeSpawner = NewObject<URigVMEdGraphFunctionRefNodeSpawner>(GetTransientPackage());
+	NodeSpawner->ReferencedPublicFunctionHeader = InFunction->GetFunctionHeader();
+	NodeSpawner->NodeClass = URigVMEdGraphNode::StaticClass();
+	NodeSpawner->bIsLocalFunction = true;
+
+	FBlueprintActionUiSpec& MenuSignature = NodeSpawner->DefaultMenuSignature;
+
+	const FString Category = InFunction->GetNodeCategory();
+
+	MenuSignature.MenuName = FText::FromString(InFunction->GetName());
+	MenuSignature.Tooltip = InFunction->GetToolTipText();
+
+	static const FString LocalFunctionString = TEXT("Local Function");
+	if(MenuSignature.Tooltip.IsEmpty())
+	{
+		MenuSignature.Tooltip = FText::FromString(LocalFunctionString);
+	}
+	else
+	{
+		static constexpr TCHAR Format[] = TEXT("%s\n\n%s");
+		MenuSignature.Tooltip = FText::FromString(FString::Printf(Format, *MenuSignature.Tooltip.ToString(), *LocalFunctionString));
+	}
+	
+	MenuSignature.Category = FText::FromString(Category);
+	MenuSignature.Keywords = FText::FromString(InFunction->GetNodeKeywords());
+
+	// add at least one character, so that PrimeDefaultUiSpec() doesn't 
+	// attempt to query the template node
+	if (MenuSignature.Keywords.IsEmpty())
+	{
+		// want to set it to something so we won't end up back in this condition
+		MenuSignature.Keywords = FText::FromString(TEXT(" "));
+	}
+
+	MenuSignature.Icon = FSlateIcon(FAppStyle::GetAppStyleSetName(), "Kismet.AllClasses.FunctionIcon");
+
+#if WITH_EDITOR
+	if (InFunction->IsMutable())
+	{
+		MenuSignature.IconTint = GetDefault<UGraphEditorSettings>()->FunctionCallNodeTitleColor;
+	}
+	else
+	{
+		MenuSignature.IconTint = GetDefault<UGraphEditorSettings>()->PureFunctionCallNodeTitleColor;
+	}
+#endif
+	
+	return NodeSpawner;
+}
+
+URigVMEdGraphFunctionRefNodeSpawner* URigVMEdGraphFunctionRefNodeSpawner::CreateFromAssetData(const FAssetData& InAssetData, const FRigVMGraphFunctionHeader& InPublicFunction)
+{
+	URigVMEdGraphFunctionRefNodeSpawner* NodeSpawner = NewObject<URigVMEdGraphFunctionRefNodeSpawner>(GetTransientPackage());
+	NodeSpawner->ReferencedPublicFunctionHeader = InPublicFunction;
+	NodeSpawner->NodeClass = URigVMEdGraphNode::StaticClass();
+	NodeSpawner->bIsLocalFunction = false;
+
+	FBlueprintActionUiSpec& MenuSignature = NodeSpawner->DefaultMenuSignature;
+
+	MenuSignature.MenuName = FText::FromName(InPublicFunction.Name);
+	MenuSignature.Category = FText::FromString(InPublicFunction.Category);
+	MenuSignature.Keywords = FText::FromString(InPublicFunction.Keywords);
+	MenuSignature.Tooltip = InPublicFunction.Tooltip;
+
+	const FString PackagePathString = InAssetData.PackageName.ToString();
+	if(MenuSignature.Tooltip.IsEmpty())
+	{
+		MenuSignature.Tooltip = FText::FromString(PackagePathString);
+	}
+	else
+	{
+		static constexpr TCHAR Format[] = TEXT("%s\n\n%s");
+		MenuSignature.Tooltip = FText::FromString(FString::Printf(Format, *MenuSignature.Tooltip.ToString(), *PackagePathString));
+	}
+
+	// add at least one character, so that PrimeDefaultUiSpec() doesn't 
+	// attempt to query the template node
+	if (MenuSignature.Keywords.IsEmpty())
+	{
+		// want to set it to something so we won't end up back in this condition
+		MenuSignature.Keywords = FText::FromString(TEXT(" "));
+	}
+
+	MenuSignature.Icon = FSlateIcon(FAppStyle::GetAppStyleSetName(), "Kismet.AllClasses.FunctionIcon");
+
+#if WITH_EDITOR
+	if (InPublicFunction.IsMutable())
+	{
+		MenuSignature.IconTint = GetDefault<UGraphEditorSettings>()->FunctionCallNodeTitleColor;
+	}
+	else
+	{
+		MenuSignature.IconTint = GetDefault<UGraphEditorSettings>()->PureFunctionCallNodeTitleColor;
+	}
+#endif
+	
+	return NodeSpawner;
+}
+
+URigVMEdGraphFunctionRefNodeSpawner* URigVMEdGraphFunctionRefNodeSpawner::CreateFromAssetData(const FAssetData& InAssetData, const FRigVMOldPublicFunctionData& InPublicFunction)
+{
+	URigVMEdGraphFunctionRefNodeSpawner* NodeSpawner = NewObject<URigVMEdGraphFunctionRefNodeSpawner>(GetTransientPackage());
+	NodeSpawner->NodeClass = URigVMEdGraphNode::StaticClass();
+	NodeSpawner->bIsLocalFunction = false;
+
+	FRigVMGraphFunctionHeader& Header = NodeSpawner->ReferencedPublicFunctionHeader;
+	Header.Name = InPublicFunction.Name;
+	Header.Arguments.Reserve(InPublicFunction.Arguments.Num());
+	for (const FRigVMOldPublicFunctionArg& Arg : InPublicFunction.Arguments)
+	{
+		FRigVMGraphFunctionArgument NewArgument;
+		NewArgument.Name = Arg.Name;
+		NewArgument.Direction = Arg.Direction;
+		NewArgument.bIsArray = Arg.bIsArray;
+		NewArgument.CPPType = Arg.CPPType;
+		PRAGMA_DISABLE_DEPRECATION_WARNINGS
+		NewArgument.CPPTypeObject = FSoftObjectPath(Arg.CPPTypeObjectPath);
+		PRAGMA_ENABLE_DEPRECATION_WARNINGS
+		Header.Arguments.Add(NewArgument);
+	}
+	NodeSpawner->AssetPath = InAssetData.ToSoftObjectPath();
+
+	FBlueprintActionUiSpec& MenuSignature = NodeSpawner->DefaultMenuSignature;
+
+	const FString Category = InPublicFunction.Category;
+
+	MenuSignature.MenuName = FText::FromName(InPublicFunction.Name);
+	MenuSignature.Category = FText::FromString(Category);
+	MenuSignature.Keywords = FText::FromString(InPublicFunction.Keywords);
+
+	if(const URigVMBlueprint* ReferencedBlueprint = Cast<URigVMBlueprint>(InAssetData.FastGetAsset(false)))
+	{
+		if(const URigVMFunctionLibrary* FunctionLibrary = ReferencedBlueprint->GetLocalFunctionLibrary())
+		{
+			if(const URigVMLibraryNode* FunctionNode = FunctionLibrary->FindFunction(InPublicFunction.Name))
+			{
+				MenuSignature.Tooltip = FunctionNode->GetToolTipText();
+			}
+		}
+	}
+
+	const FString ObjectPathString = InAssetData.GetObjectPathString();
+	if(MenuSignature.Tooltip.IsEmpty())
+	{
+		MenuSignature.Tooltip = FText::FromString(ObjectPathString);
+	}
+	else
+	{
+		static constexpr TCHAR Format[] = TEXT("%s\n\n%s");
+		MenuSignature.Tooltip = FText::FromString(FString::Printf(Format, *MenuSignature.Tooltip.ToString(), *ObjectPathString));
+	}
+
+	// add at least one character, so that PrimeDefaultUiSpec() doesn't 
+	// attempt to query the template node
+	if (MenuSignature.Keywords.IsEmpty())
+	{
+		// want to set it to something so we won't end up back in this condition
+		MenuSignature.Keywords = FText::FromString(TEXT(" "));
+	}
+
+	MenuSignature.Icon = FSlateIcon(FAppStyle::GetAppStyleSetName(), "Kismet.AllClasses.FunctionIcon");
+
+#if WITH_EDITOR
+	if (InPublicFunction.IsMutable())
+	{
+		MenuSignature.IconTint = GetDefault<UGraphEditorSettings>()->FunctionCallNodeTitleColor;
+	}
+	else
+	{
+		MenuSignature.IconTint = GetDefault<UGraphEditorSettings>()->PureFunctionCallNodeTitleColor;
+	}
+#endif
+	
+	return NodeSpawner;
+}
+
+void URigVMEdGraphFunctionRefNodeSpawner::Prime()
+{
+	// we expect that you don't need a node template to construct menu entries
+	// from this, so we choose not to pre-cache one here
+}
+
+FBlueprintNodeSignature URigVMEdGraphFunctionRefNodeSpawner::GetSpawnerSignature() const
+{
+	FString SignatureString = TEXT("Invalid RigFunction");
+	if (ReferencedPublicFunctionHeader.IsValid())
+	{
+		SignatureString = 
+			FString::Printf(TEXT("RigFunction=%s"),
+			*ReferencedPublicFunctionHeader.GetHash());
+	}
+	else
+	{
+		SignatureString = 
+			FString::Printf(TEXT("RigFunction=%s:%s"),
+			*AssetPath.ToString(),
+			*ReferencedPublicFunctionHeader.Name.ToString());
+	}
+
+	if(bIsLocalFunction)
+	{
+		SignatureString += TEXT(" (local)");
+	}
+
+	return FBlueprintNodeSignature(SignatureString);
+}
+
+FBlueprintActionUiSpec URigVMEdGraphFunctionRefNodeSpawner::GetUiSpec(FBlueprintActionContext const& Context, FBindingSet const& Bindings) const
+{
+	UEdGraph* TargetGraph = (Context.Graphs.Num() > 0) ? Context.Graphs[0] : nullptr;
+	FBlueprintActionUiSpec MenuSignature = PrimeDefaultUiSpec(TargetGraph);
+
+	DynamicUiSignatureGetter.ExecuteIfBound(Context, Bindings, &MenuSignature);
+	return MenuSignature;
+}
+
+UEdGraphNode* URigVMEdGraphFunctionRefNodeSpawner::Invoke(UEdGraph* ParentGraph, FBindingSet const& Bindings, FVector2D const Location) const
+{
+	URigVMEdGraphNode* NewNode = nullptr;
+
+	// if we are trying to build the real function ref - but we haven't loaded the asset yet...
+	if(!FBlueprintNodeTemplateCache::IsTemplateOuter(ParentGraph))
+	{
+		if (!ReferencedPublicFunctionHeader.IsValid() && AssetPath.IsValid())
+		{
+			if (URigVMBlueprint* Blueprint = Cast<URigVMBlueprint>(AssetPath.TryLoad()))
+			{
+				ReferencedPublicFunctionHeader = Blueprint->GetLocalFunctionLibrary()->FindFunction(ReferencedPublicFunctionHeader.Name)->GetFunctionHeader();			
+			}
+		}
+	}
+		
+	if(ReferencedPublicFunctionHeader.IsValid())
+	{
+#if WITH_EDITOR
+		if (GEditor)
+		{
+			GEditor->CancelTransaction(0);
+		}
+#endif
+
+		UBlueprint* Blueprint = FBlueprintEditorUtils::FindBlueprintForGraph(ParentGraph);
+		NewNode = SpawnNode(ParentGraph, Blueprint, ReferencedPublicFunctionHeader, Location);
+	}
+	else
+	{
+		// we are only going to get here if we are spawning a template node
+		NewNode = NewObject<URigVMEdGraphNode>(ParentGraph);
+		ParentGraph->AddNode(NewNode, false);
+
+		NewNode->CreateNewGuid();
+		NewNode->PostPlacedNewNode();
+
+		for(const FRigVMGraphFunctionArgument& Arg : ReferencedPublicFunctionHeader.Arguments)
+		{
+			if(Arg.Direction ==  ERigVMPinDirection::Input ||
+				Arg.Direction ==  ERigVMPinDirection::IO)
+			{
+				UEdGraphPin* InputPin = UEdGraphPin::CreatePin(NewNode);
+				NewNode->Pins.Add(InputPin);
+
+				InputPin->Direction = EGPD_Input;
+				InputPin->PinType = RigVMTypeUtils::PinTypeFromCPPType(Arg.CPPType, Arg.CPPTypeObject.Get());
+			}
+
+			if(Arg.Direction ==  ERigVMPinDirection::Output ||
+				Arg.Direction ==  ERigVMPinDirection::IO)
+			{
+				UEdGraphPin* OutputPin = UEdGraphPin::CreatePin(NewNode);
+				NewNode->Pins.Add(OutputPin);
+
+				OutputPin->Direction = EGPD_Output;
+				OutputPin->PinType = RigVMTypeUtils::PinTypeFromCPPType(Arg.CPPType, Arg.CPPTypeObject.Get());
+			}
+		}
+
+		NewNode->SetFlags(RF_Transactional);
+}
+
+	return NewNode;
+}
+
+URigVMEdGraphNode* URigVMEdGraphFunctionRefNodeSpawner::SpawnNode(UEdGraph* ParentGraph, UBlueprint* Blueprint, FRigVMGraphFunctionHeader& InFunction, FVector2D const Location)
+{
+	URigVMEdGraphNode* NewNode = nullptr;
+	URigVMBlueprint* RigBlueprint = Cast<URigVMBlueprint>(Blueprint);
+	URigVMEdGraph* RigGraph = Cast<URigVMEdGraph>(ParentGraph);
+
+	if (RigBlueprint != nullptr && RigGraph != nullptr)
+	{
+		bool const bIsTemplateNode = FBlueprintNodeTemplateCache::IsTemplateOuter(ParentGraph);
+		bool const bIsUserFacingNode = !bIsTemplateNode;
+
+		FName Name = bIsTemplateNode ? InFunction.Name : FRigVMBlueprintUtils::ValidateName(RigBlueprint, InFunction.Name.ToString());
+		URigVMController* Controller = bIsTemplateNode ? RigGraph->GetTemplateController() : RigBlueprint->GetController(ParentGraph);
+
+		if (!bIsTemplateNode)
+		{
+			Controller->OpenUndoBracket(FString::Printf(TEXT("Add '%s' Node"), *Name.ToString()));
+		}
+
+		TGuardValue<bool> AllowPrivateFunctionsOnTemplateController(Controller->bAllowPrivateFunctions, bIsTemplateNode);
+		if (URigVMFunctionReferenceNode* ModelNode = Controller->AddFunctionReferenceNodeFromDescription(InFunction, Location, Name.ToString(), bIsUserFacingNode, !bIsTemplateNode))
+		{
+			NewNode = Cast<URigVMEdGraphNode>(RigGraph->FindNodeForModelNodeName(ModelNode->GetFName()));
+			check(NewNode);
+
+			if (NewNode && bIsUserFacingNode)
+			{
+				Controller->ClearNodeSelection(true);
+				Controller->SelectNode(ModelNode, true, true);
+
+				URigVMEdGraphUnitNodeSpawner::HookupMutableNode(ModelNode, RigBlueprint);
+			}
+
+			if (!bIsTemplateNode)
+			{
+				for(URigVMNode* OtherModelNode : ModelNode->GetGraph()->GetNodes())
+				{
+					if(OtherModelNode == ModelNode)
+					{
+						continue;
+					}
+					
+					URigVMFunctionReferenceNode* ExistingFunctionReferenceNode = Cast<URigVMFunctionReferenceNode>(OtherModelNode);
+					if(ExistingFunctionReferenceNode == nullptr)
+					{
+						continue;
+					}
+
+					if(!(ExistingFunctionReferenceNode->GetReferencedFunctionHeader().LibraryPointer == InFunction.LibraryPointer))
+					{
+						continue;
+					}
+
+					for(TPair<FName, FName> MappedVariablePair : ExistingFunctionReferenceNode->GetVariableMap())
+					{
+						if(!MappedVariablePair.Value.IsNone())
+						{
+							Controller->SetRemappedVariable(ModelNode, MappedVariablePair.Key, MappedVariablePair.Value, bIsUserFacingNode);
+						}
+					}
+				}
+			}
+			else
+			{
+				// If the package is a template, do not remove the EdGraphNode
+				// We might be spawning a node to populate the PROTO_ context menu for function declarations.
+				FRigVMControllerNotifGuard NotifGuard(Controller, true);
+				Controller->RemoveNode(ModelNode, false);
+			}
+
+			if (bIsUserFacingNode)
+			{
+				Controller->CloseUndoBracket();
+			}
+			else
+			{
+				Controller->EnableReporting(false);
+				Controller->RemoveNode(ModelNode, false);
+				Controller->EnableReporting(true);
+			}
+		}
+		else
+		{
+			if (bIsUserFacingNode)
+			{
+				Controller->CancelUndoBracket();
+			}
+		}
+	}
+	return NewNode;
+}
+
+bool URigVMEdGraphFunctionRefNodeSpawner::IsTemplateNodeFilteredOut(FBlueprintActionFilter const& Filter) const
+{
+	if(bIsLocalFunction)
+	{
+		if(ReferencedPublicFunctionHeader.IsValid())
+		{
+			for (UBlueprint* Blueprint : Filter.Context.Blueprints)
+			{
+				if(Cast<IRigVMGraphFunctionHost>(Blueprint->GeneratedClass) != ReferencedPublicFunctionHeader.GetFunctionHost())
+				{
+					return true;
+				}
+			}
+		}
+	}
+	const FString ReferencedAssetObjectPathString = ReferencedPublicFunctionHeader.LibraryPointer.LibraryNode.GetAssetName();
+	for (UBlueprint* Blueprint : Filter.Context.Blueprints)
+	{
+		if(Blueprint->GetPathName() == ReferencedAssetObjectPathString)
+		{
+			return true;
+		}
+	}
+	return false;
+}
+
+#undef LOCTEXT_NAMESPACE
+
diff --git a/Engine/Plugins/Runtime/RigVM/Source/RigVMEditor/Private/EdGraph/NodeSpawners/RigVMEdGraphInvokeEntryNodeSpawner.cpp b/Engine/Plugins/Runtime/RigVM/Source/RigVMEditor/Private/EdGraph/NodeSpawners/RigVMEdGraphInvokeEntryNodeSpawner.cpp
new file mode 100644
--- /dev/null
+++ b/Engine/Plugins/Runtime/RigVM/Source/RigVMEditor/Private/EdGraph/NodeSpawners/RigVMEdGraphInvokeEntryNodeSpawner.cpp
@@ -0,0 +1,142 @@
+// Copyright Epic Games, Inc. All Rights Reserved.
+
+#include "EdGraph/NodeSpawners/RigVMEdGraphInvokeEntryNodeSpawner.h"
+#include "EdGraph/RigVMEdGraph.h"
+#include "EdGraph/RigVMEdGraphNode.h"
+#include "EdGraphSchema_K2.h"
+#include "Kismet2/BlueprintEditorUtils.h"
+#include "Settings/EditorStyleSettings.h"
+#include "Editor/EditorEngine.h"
+#include "ObjectEditorUtils.h"
+#include "EditorCategoryUtils.h"
+#include "BlueprintNodeTemplateCache.h"
+#include "RigVMModel/RigVMController.h"
+#include "RigVMBlueprint.h"
+#include "RigVMModel/Nodes/RigVMInvokeEntryNode.h"
+
+#include UE_INLINE_GENERATED_CPP_BY_NAME(RigVMEdGraphInvokeEntryNodeSpawner)
+
+#if WITH_EDITOR
+#include "Editor.h"
+#endif
+
+#define LOCTEXT_NAMESPACE "RigVMEdGraphInvokeEntryNodeSpawner"
+
+URigVMEdGraphInvokeEntryNodeSpawner* URigVMEdGraphInvokeEntryNodeSpawner::CreateForEntry(URigVMBlueprint* InBlueprint, const FName& InEntryName, const FText& InMenuDesc, const FText& InCategory, const FText& InTooltip)
+{
+	URigVMEdGraphInvokeEntryNodeSpawner* NodeSpawner = NewObject<URigVMEdGraphInvokeEntryNodeSpawner>(GetTransientPackage());
+	NodeSpawner->Blueprint = InBlueprint;
+	NodeSpawner->EntryName = InEntryName;
+	NodeSpawner->NodeClass = URigVMEdGraphNode::StaticClass();
+
+	FBlueprintActionUiSpec& MenuSignature = NodeSpawner->DefaultMenuSignature;
+	
+	MenuSignature.MenuName = InMenuDesc;
+	MenuSignature.Tooltip  = InTooltip;
+	MenuSignature.Category = InCategory;
+	MenuSignature.Keywords = FText::FromString(TEXT("Event,Entry,Invoke,Run,Launch,Start"));
+	MenuSignature.Icon = FSlateIcon(FAppStyle::GetAppStyleSetName(), "GraphEditor.Event_16x");
+
+	return NodeSpawner;
+}
+
+void URigVMEdGraphInvokeEntryNodeSpawner::Prime()
+{
+	// we expect that you don't need a node template to construct menu entries
+	// from this, so we choose not to pre-cache one here
+}
+
+FBlueprintNodeSignature URigVMEdGraphInvokeEntryNodeSpawner::GetSpawnerSignature() const
+{
+	return FBlueprintNodeSignature(FString("InvokeEntryNode=" + EntryName.ToString()));
+}
+
+FBlueprintActionUiSpec URigVMEdGraphInvokeEntryNodeSpawner::GetUiSpec(FBlueprintActionContext const& Context, FBindingSet const& Bindings) const
+{
+	UEdGraph* TargetGraph = (Context.Graphs.Num() > 0) ? Context.Graphs[0] : nullptr;
+	FBlueprintActionUiSpec MenuSignature = PrimeDefaultUiSpec(TargetGraph);
+
+	DynamicUiSignatureGetter.ExecuteIfBound(Context, Bindings, &MenuSignature);
+	return MenuSignature;
+}
+
+UEdGraphNode* URigVMEdGraphInvokeEntryNodeSpawner::Invoke(UEdGraph* ParentGraph, FBindingSet const& Bindings, FVector2D const Location) const
+{
+	URigVMEdGraphNode* NewNode = nullptr;
+
+	bool const bIsTemplateNode = FBlueprintNodeTemplateCache::IsTemplateOuter(ParentGraph);
+	bool const bIsUserFacingNode = !bIsTemplateNode;
+
+	// First create a backing member for our node
+	URigVMEdGraph* RigGraph = Cast<URigVMEdGraph>(ParentGraph);
+	if(RigGraph == nullptr) return nullptr;
+	URigVMBlueprint* RigBlueprint = Cast<URigVMBlueprint>(FBlueprintEditorUtils::FindBlueprintForGraph(ParentGraph));
+	check(RigBlueprint);
+
+#if WITH_EDITOR
+	if (GEditor && !bIsTemplateNode)
+	{
+		GEditor->CancelTransaction(0);
+	}
+#endif
+
+	URigVMController* Controller = bIsTemplateNode ? RigGraph->GetTemplateController() : RigBlueprint->GetController(ParentGraph);
+
+	if (bIsUserFacingNode)
+	{
+		Controller->OpenUndoBracket(TEXT("Add Run Event"));
+	}
+
+	FString NodeName;
+	if(bIsTemplateNode)
+	{
+		// since we are removing the node at the end of this function
+		// we need to create a unique here.
+		static constexpr TCHAR InvokeEntryNodeNameFormat[] = TEXT("InvokeEntryNode_%s");
+		NodeName = FString::Printf(InvokeEntryNodeNameFormat, *EntryName.ToString());
+	}
+
+	if (URigVMNode* ModelNode = Controller->AddInvokeEntryNode(EntryName, Location, NodeName, !bIsTemplateNode, !bIsTemplateNode))
+	{
+		for (UEdGraphNode* Node : ParentGraph->Nodes)
+		{
+			if (URigVMEdGraphNode* RigNode = Cast<URigVMEdGraphNode>(Node))
+			{
+				if (RigNode->GetModelNodeName() == ModelNode->GetFName())
+				{
+					NewNode = RigNode;
+					break;
+				}
+			}
+		}
+
+		if (bIsUserFacingNode)
+		{
+			if (NewNode)
+			{
+				Controller->ClearNodeSelection(true);
+				Controller->SelectNode(ModelNode, true, true);
+			}
+			Controller->CloseUndoBracket();
+		}
+		else
+		{
+			// similar to UBlueprintNodeSpawner::Invoke -> UBlueprintNodeSpawner::SpawnEdGraphNode
+			// we simply want the node, but not actually adding it to a graph
+			Controller->RemoveNode(ModelNode, false);
+		}
+	}
+	else
+	{
+		if (!bIsTemplateNode)
+		{
+			Controller->CancelUndoBracket();
+		}
+	}
+
+
+	return NewNode;
+}
+
+#undef LOCTEXT_NAMESPACE
+
diff --git a/Engine/Plugins/Runtime/RigVM/Source/RigVMEditor/Private/EdGraph/NodeSpawners/RigVMEdGraphTemplateNodeSpawner.cpp b/Engine/Plugins/Runtime/RigVM/Source/RigVMEditor/Private/EdGraph/NodeSpawners/RigVMEdGraphTemplateNodeSpawner.cpp
new file mode 100644
--- /dev/null
+++ b/Engine/Plugins/Runtime/RigVM/Source/RigVMEditor/Private/EdGraph/NodeSpawners/RigVMEdGraphTemplateNodeSpawner.cpp
@@ -0,0 +1,234 @@
+// Copyright Epic Games, Inc. All Rights Reserved.
+
+#include "EdGraph/NodeSpawners/RigVMEdGraphTemplateNodeSpawner.h"
+#include "EdGraph/NodeSpawners/RigVMEdGraphUnitNodeSpawner.h"
+#include "EdGraph/RigVMEdGraph.h"
+#include "EdGraph/RigVMEdGraphNode.h"
+#include "EdGraph/RigVMEdGraphSchema.h"
+#include "EdGraphSchema_K2.h"
+#include "Kismet2/BlueprintEditorUtils.h"
+#include "Settings/EditorStyleSettings.h"
+#include "Editor/EditorEngine.h"
+#include "ObjectEditorUtils.h"
+#include "EditorCategoryUtils.h"
+#include "K2Node_Variable.h"
+#include "BlueprintNodeTemplateCache.h"
+#include "RigVMBlueprintUtils.h"
+#include "ScopedTransaction.h"
+#include "RigVMModel/Nodes/RigVMDispatchNode.h"
+
+#include UE_INLINE_GENERATED_CPP_BY_NAME(RigVMEdGraphTemplateNodeSpawner)
+
+#if WITH_EDITOR
+#include "Editor.h"
+#endif
+
+#define LOCTEXT_NAMESPACE "RigVMEdGraphTemplateNodeSpawner"
+
+URigVMEdGraphTemplateNodeSpawner* URigVMEdGraphTemplateNodeSpawner::CreateFromNotation(const FName& InNotation, const FText& InMenuDesc, const FText& InCategory, const FText& InTooltip)
+{
+	URigVMEdGraphTemplateNodeSpawner* NodeSpawner = NewObject<URigVMEdGraphTemplateNodeSpawner>(GetTransientPackage());
+	NodeSpawner->TemplateNotation = InNotation;
+	NodeSpawner->NodeClass = URigVMEdGraphNode::StaticClass();
+
+	FBlueprintActionUiSpec& MenuSignature = NodeSpawner->DefaultMenuSignature;
+	
+	MenuSignature.MenuName = InMenuDesc;
+	MenuSignature.Tooltip  = InTooltip;
+	MenuSignature.Category = InCategory;
+
+	if (const FRigVMTemplate* Template = FRigVMRegistry::Get().FindTemplate(InNotation))
+	{
+#if WITH_EDITOR
+		FString KeywordsMetadata = Template->GetKeywords();
+		MenuSignature.Keywords = FText::FromString(KeywordsMetadata);
+#endif
+	}
+
+	// add at least one character, so that PrimeDefaultUiSpec() doesn't 
+	// attempt to query the template node
+	//
+	// @TODO: maybe UPROPERTY() fields should have keyword metadata like functions
+	if (MenuSignature.Keywords.IsEmpty())
+	{
+		// want to set it to something so we won't end up back in this condition
+		MenuSignature.Keywords = FText::FromString(TEXT(" "));
+	}
+
+	// @TODO: should use details customization-like extensibility system to provide editor only data like this
+	MenuSignature.Icon = FSlateIcon(TEXT("ControlRigEditorStyle"), TEXT("ControlRig.RigUnit"));
+
+	return NodeSpawner;
+}
+
+void URigVMEdGraphTemplateNodeSpawner::Prime()
+{
+	// we expect that you don't need a node template to construct menu entries
+	// from this, so we choose not to pre-cache one here
+}
+
+FBlueprintNodeSignature URigVMEdGraphTemplateNodeSpawner::GetSpawnerSignature() const
+{
+	const int32 NotationHash = (int32)GetTypeHash(TemplateNotation);
+	return FBlueprintNodeSignature(FString("RigVMTemplate_") + FString::FromInt(NotationHash));
+}
+
+bool URigVMEdGraphTemplateNodeSpawner::IsTemplateNodeFilteredOut(FBlueprintActionFilter const& Filter) const
+{
+	bool bFilteredOut = true;
+	for (UEdGraphPin* Pin : Filter.Context.Pins)
+	{
+		for (UEdGraph* Graph : Filter.Context.Graphs)
+		{
+			if (!Graph->Nodes.Contains(Pin->GetOwningNode()))
+			{
+				continue;
+			}
+
+			if (URigVMEdGraph* EdGraph = Cast<URigVMEdGraph>(Graph))
+			{
+				if (URigVMGraph* RigGraph = EdGraph->GetModel())
+				{
+					if (URigVMPin* ModelPin = RigGraph->FindPin(Pin->GetName()))
+					{
+						if (URigVMTemplateNode* TemplateNode = Cast<URigVMTemplateNode>(ModelPin->GetNode()))
+						{
+							// Only filter when the sourece pin is not a wildcard
+							if (ModelPin->GetCPPType() != RigVMTypeUtils::GetWildCardCPPType())
+							{
+								if (const FRigVMTemplate* Template = FRigVMRegistry::Get().FindTemplate(TemplateNotation))
+								{
+									if(ModelPin->IsExecuteContext())
+									{
+										FRigVMDispatchContext DispatchContext;
+										if(const URigVMDispatchNode* DispatchNode = Cast<URigVMDispatchNode>(ModelPin->GetNode()))
+										{
+											DispatchContext = DispatchNode->GetDispatchContext();
+										}
+
+										if(Template->NumExecuteArguments(DispatchContext) > 0)
+										{
+											return false;
+										}
+									}
+									
+									for (int32 i=0; i<Template->NumArguments(); ++i)
+									{
+										const FRigVMTemplateArgument* Argument = Template->GetArgument(i);
+										if (Template->ArgumentSupportsTypeIndex(Argument->GetName(), ModelPin->GetTypeIndex()))
+										{
+											return false;
+										}
+									}
+									return true;
+								}									
+							}
+							else
+							{
+								// TODO: given the filtered types of the source pin, provide only the nodes that will not break any connections
+							}
+						}
+					}					
+				}
+			}
+		}
+	}
+	
+	return Super::IsTemplateNodeFilteredOut(Filter);
+}
+
+FBlueprintActionUiSpec URigVMEdGraphTemplateNodeSpawner::GetUiSpec(FBlueprintActionContext const& Context, FBindingSet const& Bindings) const
+{
+	UEdGraph* TargetGraph = (Context.Graphs.Num() > 0) ? Context.Graphs[0] : nullptr;
+	FBlueprintActionUiSpec MenuSignature = PrimeDefaultUiSpec(TargetGraph);
+
+	DynamicUiSignatureGetter.ExecuteIfBound(Context, Bindings, &MenuSignature);
+	return MenuSignature;
+}
+
+UEdGraphNode* URigVMEdGraphTemplateNodeSpawner::Invoke(UEdGraph* ParentGraph, FBindingSet const& Bindings, FVector2D const Location) const
+{
+	URigVMEdGraphNode* NewNode = nullptr;
+
+	if(!TemplateNotation.IsNone())
+	{
+#if WITH_EDITOR
+		if (GEditor)
+		{
+			GEditor->CancelTransaction(0);
+		}
+#endif
+
+		UBlueprint* Blueprint = FBlueprintEditorUtils::FindBlueprintForGraph(ParentGraph);
+		NewNode = SpawnNode(ParentGraph, Blueprint, TemplateNotation, Location);
+	}
+
+	return NewNode;
+}
+
+URigVMEdGraphNode* URigVMEdGraphTemplateNodeSpawner::SpawnNode(UEdGraph* ParentGraph, UBlueprint* Blueprint, const FName& InNotation, FVector2D const Location)
+{
+	URigVMEdGraphNode* NewNode = nullptr;
+	URigVMBlueprint* RigBlueprint = Cast<URigVMBlueprint>(Blueprint);
+	URigVMEdGraph* RigGraph = Cast<URigVMEdGraph>(ParentGraph);
+
+	if (RigBlueprint != nullptr && RigGraph != nullptr)
+	{
+		bool const bIsTemplateNode = FBlueprintNodeTemplateCache::IsTemplateOuter(ParentGraph);
+		bool const bIsUserFacingNode = !bIsTemplateNode;
+
+		const FRigVMTemplate* Template = FRigVMRegistry::Get().FindTemplate(InNotation);
+		if (Template == nullptr)
+		{
+			return nullptr;
+		}
+
+		const int32 NotationHash = (int32)GetTypeHash(InNotation);
+		const FString TemplateName = TEXT("RigVMTemplate_") + FString::FromInt(NotationHash);
+
+		FName Name = bIsTemplateNode ? *TemplateName : FRigVMBlueprintUtils::ValidateName(RigBlueprint, Template->GetName().ToString());
+		URigVMController* Controller = bIsTemplateNode ? RigGraph->GetTemplateController() : RigBlueprint->GetController(ParentGraph);
+
+		if (!bIsTemplateNode)
+		{
+			Controller->OpenUndoBracket(FString::Printf(TEXT("Add '%s' Node"), *Name.ToString()));
+		}
+
+		if (URigVMTemplateNode* ModelNode = Controller->AddTemplateNode(InNotation, Location, Name.ToString(), bIsUserFacingNode, !bIsTemplateNode))
+		{
+			NewNode = Cast<URigVMEdGraphNode>(RigGraph->FindNodeForModelNodeName(ModelNode->GetFName()));
+
+			if (NewNode && bIsUserFacingNode)
+			{
+				Controller->ClearNodeSelection(true);
+				Controller->SelectNode(ModelNode, true, true);
+
+				URigVMEdGraphUnitNodeSpawner::HookupMutableNode(ModelNode, RigBlueprint);
+			}
+
+			if (bIsUserFacingNode)
+			{
+				Controller->CloseUndoBracket();
+			}
+			else
+			{
+				if(NewNode)
+				{
+					NewNode->ModelNodePath = Template->GetNotation().ToString();
+				}
+				Controller->RemoveNode(ModelNode, false);
+			}
+		}
+		else
+		{
+			if (bIsUserFacingNode)
+			{
+				Controller->CancelUndoBracket();
+			}
+		}
+	}
+	return NewNode;
+}
+
+#undef LOCTEXT_NAMESPACE
+
diff --git a/Engine/Plugins/Runtime/RigVM/Source/RigVMEditor/Private/EdGraph/NodeSpawners/RigVMEdGraphUnitNodeSpawner.cpp b/Engine/Plugins/Runtime/RigVM/Source/RigVMEditor/Private/EdGraph/NodeSpawners/RigVMEdGraphUnitNodeSpawner.cpp
new file mode 100644
--- /dev/null
+++ b/Engine/Plugins/Runtime/RigVM/Source/RigVMEditor/Private/EdGraph/NodeSpawners/RigVMEdGraphUnitNodeSpawner.cpp
@@ -0,0 +1,433 @@
+// Copyright Epic Games, Inc. All Rights Reserved.
+
+#include "EdGraph/NodeSpawners/RigVMEdGraphUnitNodeSpawner.h"
+#include "EdGraph/RigVMEdGraph.h"
+#include "EdGraph/RigVMEdGraphNode.h"
+#include "EdGraph/RigVMEdGraphSchema.h"
+#include "EdGraphSchema_K2.h"
+#include "Kismet2/BlueprintEditorUtils.h"
+#include "Settings/EditorStyleSettings.h"
+#include "Editor/EditorEngine.h"
+#include "ObjectEditorUtils.h"
+#include "EditorCategoryUtils.h"
+#include "K2Node_Variable.h"
+#include "BlueprintNodeTemplateCache.h"
+#include "RigVMBlueprintUtils.h"
+#include "ScopedTransaction.h"
+#include "RigVMFunctions/Execution/RigVMFunction_UserDefinedEvent.h"
+
+#include UE_INLINE_GENERATED_CPP_BY_NAME(RigVMEdGraphUnitNodeSpawner)
+
+#if WITH_EDITOR
+#include "Editor.h"
+#include "SGraphActionMenu.h"
+#endif
+
+#define LOCTEXT_NAMESPACE "RigVMEdGraphUnitNodeSpawner"
+
+URigVMEdGraphUnitNodeSpawner* URigVMEdGraphUnitNodeSpawner::CreateFromStruct(UScriptStruct* InStruct, const FName& InMethodName, const FText& InMenuDesc, const FText& InCategory, const FText& InTooltip)
+{
+	URigVMEdGraphUnitNodeSpawner* NodeSpawner = NewObject<URigVMEdGraphUnitNodeSpawner>(GetTransientPackage());
+	NodeSpawner->StructTemplate = InStruct;
+	NodeSpawner->MethodName = InMethodName;
+	NodeSpawner->NodeClass = URigVMEdGraphNode::StaticClass();
+
+	FBlueprintActionUiSpec& MenuSignature = NodeSpawner->DefaultMenuSignature;
+	
+	MenuSignature.MenuName = InMenuDesc;
+	MenuSignature.Tooltip  = InTooltip;
+	MenuSignature.Category = InCategory;
+
+	FString KeywordsMetadata, TemplateNameMetadata;
+	InStruct->GetStringMetaDataHierarchical(FRigVMStruct::KeywordsMetaName, &KeywordsMetadata);
+	if(!TemplateNameMetadata.IsEmpty())
+	{
+		if(KeywordsMetadata.IsEmpty())
+		{
+			KeywordsMetadata = TemplateNameMetadata;
+		}
+		else
+		{
+			KeywordsMetadata = KeywordsMetadata + TEXT(",") + TemplateNameMetadata;
+		}
+	}
+	MenuSignature.Keywords = FText::FromString(KeywordsMetadata);
+
+	// add at least one character, so that PrimeDefaultUiSpec() doesn't 
+	// attempt to query the template node
+	//
+	// @TODO: maybe UPROPERTY() fields should have keyword metadata like functions
+	if (MenuSignature.Keywords.IsEmpty())
+	{
+		// want to set it to something so we won't end up back in this condition
+		MenuSignature.Keywords = FText::FromString(TEXT(" "));
+	}
+
+	// @TODO: should use details customization-like extensibility system to provide editor only data like this
+	FStructOnScope StructScope(InStruct);
+	const FRigVMStruct* StructInstance = (const FRigVMStruct*)StructScope.GetStructMemory();
+	if(StructInstance->GetEventName().IsNone())
+	{
+		if(InStruct->HasMetaDataHierarchical(FRigVMStruct::IconMetaName))
+		{
+			FString IconPath;
+			const int32 NumOfIconPathNames = 4;
+			
+			FName IconPathNames[NumOfIconPathNames] = {
+				NAME_None, // StyleSetName
+				NAME_None, // StyleName
+				NAME_None, // SmallStyleName
+				NAME_None  // StatusOverlayStyleName
+			};
+
+			InStruct->GetStringMetaDataHierarchical(FRigVMStruct::IconMetaName, &IconPath);
+
+			int32 NameIndex = 0;
+
+			while (!IconPath.IsEmpty() && NameIndex < NumOfIconPathNames)
+			{
+				FString Left;
+				FString Right;
+
+				if (!IconPath.Split(TEXT("|"), &Left, &Right))
+				{
+					Left = IconPath;
+				}
+
+				IconPathNames[NameIndex] = FName(*Left);
+
+				NameIndex++;
+				IconPath = Right;
+			}
+			MenuSignature.Icon = FSlateIcon(IconPathNames[0], IconPathNames[1], IconPathNames[2], IconPathNames[3]);
+		}
+		else
+		{
+			MenuSignature.Icon = FSlateIcon(TEXT("ControlRigEditorStyle"), TEXT("ControlRig.RigUnit"));
+		}
+	}
+	else
+	{
+		MenuSignature.Icon = FSlateIcon(FAppStyle::GetAppStyleSetName(), "GraphEditor.Event_16x");
+	}
+
+	return NodeSpawner;
+}
+
+void URigVMEdGraphUnitNodeSpawner::Prime()
+{
+	// we expect that you don't need a node template to construct menu entries
+	// from this, so we choose not to pre-cache one here
+}
+
+FBlueprintNodeSignature URigVMEdGraphUnitNodeSpawner::GetSpawnerSignature() const
+{
+	return FBlueprintNodeSignature(FString("RigUnit=" + StructTemplate->GetFName().ToString()));
+}
+
+FBlueprintActionUiSpec URigVMEdGraphUnitNodeSpawner::GetUiSpec(FBlueprintActionContext const& Context, FBindingSet const& Bindings) const
+{
+	UEdGraph* TargetGraph = (Context.Graphs.Num() > 0) ? Context.Graphs[0] : nullptr;
+	FBlueprintActionUiSpec MenuSignature = PrimeDefaultUiSpec(TargetGraph);
+
+	DynamicUiSignatureGetter.ExecuteIfBound(Context, Bindings, &MenuSignature);
+	return MenuSignature;
+}
+
+UEdGraphNode* URigVMEdGraphUnitNodeSpawner::Invoke(UEdGraph* ParentGraph, FBindingSet const& Bindings, FVector2D const Location) const
+{
+	URigVMEdGraphNode* NewNode = nullptr;
+
+	if(StructTemplate)
+	{
+#if WITH_EDITOR
+		if (GEditor)
+		{
+			GEditor->CancelTransaction(0);
+		}
+#endif
+
+		UBlueprint* Blueprint = FBlueprintEditorUtils::FindBlueprintForGraph(ParentGraph);
+		NewNode = SpawnNode(ParentGraph, Blueprint, StructTemplate, MethodName, Location);
+	}
+
+	return NewNode;
+}
+
+URigVMEdGraphNode* URigVMEdGraphUnitNodeSpawner::SpawnNode(UEdGraph* ParentGraph, UBlueprint* Blueprint, UScriptStruct* StructTemplate, const FName& InMethodName, FVector2D const Location)
+{
+	URigVMEdGraphNode* NewNode = nullptr;
+	URigVMBlueprint* RigBlueprint = Cast<URigVMBlueprint>(Blueprint);
+	URigVMEdGraph* RigGraph = Cast<URigVMEdGraph>(ParentGraph);
+
+	if (RigBlueprint != nullptr && RigGraph != nullptr)
+	{
+		bool const bIsTemplateNode = FBlueprintNodeTemplateCache::IsTemplateOuter(ParentGraph);
+		bool const bIsUserFacingNode = !bIsTemplateNode;
+
+		FName Name = bIsTemplateNode ? *StructTemplate->GetStructCPPName() : FRigVMBlueprintUtils::ValidateName(RigBlueprint, StructTemplate->GetFName().ToString());
+
+		FStructOnScope StructScope(StructTemplate);
+		const FRigVMStruct* StructInstance = (const FRigVMStruct*)StructScope.GetStructMemory();
+		const FName EventName = StructInstance->GetEventName();
+
+		// events can only exist once across all uber graphs
+		if(bIsUserFacingNode)
+		{
+			if(!EventName.IsNone())
+			{
+				if(StructInstance->CanOnlyExistOnce() &&
+					(StructTemplate != FRigVMFunction_UserDefinedEvent::StaticStruct()))
+				{
+					const TArray<URigVMGraph*> Models = RigBlueprint->GetAllModels();
+					for(URigVMGraph* Model : Models)
+					{
+						if(Model->IsRootGraph() && !Model->IsA<URigVMFunctionLibrary>())
+						{
+							for(const URigVMNode* Node : Model->GetNodes())
+							{
+								if(Node->GetEventName() == EventName)
+								{
+									RigBlueprint->OnRequestJumpToHyperlink().Execute(Node);
+									return nullptr;
+								}
+							}
+						}
+					}
+				}
+			}
+		}
+		
+		URigVMController* Controller = bIsTemplateNode ? RigGraph->GetTemplateController() : RigBlueprint->GetController(ParentGraph);
+		if(Controller == nullptr)
+		{
+			return nullptr;
+		}
+
+		if (!bIsTemplateNode)
+		{
+			Controller->OpenUndoBracket(FString::Printf(TEXT("Add '%s' Node"), *Name.ToString()));
+		}
+
+		FRigVMUnitNodeCreatedContext& UnitNodeCreatedContext = Controller->GetUnitNodeCreatedContext();
+		FRigVMUnitNodeCreatedContext::FScope ReasonScope(UnitNodeCreatedContext, ERigVMNodeCreatedReason::NodeSpawner);
+
+		if (URigVMUnitNode* ModelNode = Controller->AddUnitNode(StructTemplate, InMethodName, Location, Name.ToString(), bIsUserFacingNode, !bIsTemplateNode))
+		{
+			NewNode = Cast<URigVMEdGraphNode>(RigGraph->FindNodeForModelNodeName(ModelNode->GetFName()));
+			check(NewNode);
+
+			if (NewNode && bIsUserFacingNode)
+			{
+				if(StructTemplate == FRigVMFunction_UserDefinedEvent::StaticStruct())
+				{
+					const URigVMEdGraphSchema* Schema = CastChecked<URigVMEdGraphSchema>(NewNode->GetSchema());
+					
+					// ensure uniqueness for the event name
+					TArray<FName> ExistingEventNames;
+
+					const TArray<URigVMGraph*> Models = RigBlueprint->GetAllModels();
+					for(URigVMGraph* Model : Models)
+					{
+						for(const URigVMNode* Node : Model->GetNodes())
+						{
+							if(Node != ModelNode)
+							{
+								const FName ExistingEventName = Node->GetEventName();
+								if(!ExistingEventName.IsNone())
+								{
+									ExistingEventNames.AddUnique(ExistingEventName);
+								}
+							}
+						}
+					}
+
+					FName SafeEventName = Controller->GetSchema()->GetUniqueName(EventName, [ExistingEventNames, Schema](const FName& NameToCheck) -> bool
+					{
+						return !ExistingEventNames.Contains(NameToCheck) && !Schema->IsRigVMDefaultEvent(NameToCheck);
+					}, false, true);
+
+					Controller->SetPinDefaultValue(ModelNode->FindPin(TEXT("EventName"))->GetPinPath(), SafeEventName.ToString(), false, true, false, true);
+				}
+
+				Controller->ClearNodeSelection(true);
+				Controller->SelectNode(ModelNode, true, true);
+
+				HookupMutableNode(ModelNode, RigBlueprint);
+			}
+
+#if WITH_EDITORONLY_DATA
+			/*
+			// Todo: we should remove this functionality imho
+			if (!bIsTemplateNode)
+			{
+				const FControlRigSettingsPerPinBool* ExpansionMapPtr = UControlRigEditorSettings::Get()->RigUnitPinExpansion.Find(ModelNode->GetScriptStruct()->GetName());
+				if (ExpansionMapPtr)
+				{
+					const FControlRigSettingsPerPinBool& ExpansionMap = *ExpansionMapPtr;
+
+					for (const TPair<FString, bool>& Pair : ExpansionMap.Values)
+					{
+						FString PinPath = FString::Printf(TEXT("%s.%s"), *ModelNode->GetName(), *Pair.Key);
+						Controller->SetPinExpansion(PinPath, Pair.Value, bIsUserFacingNode);
+					}
+				}
+
+				FString UsedFilterString = SGraphActionMenu::LastUsedFilterText;
+				if (!UsedFilterString.IsEmpty())
+				{
+					UsedFilterString = UsedFilterString.ToLower();
+
+					if (UEnum* RigElementTypeEnum = StaticEnum<ERigElementType>())
+					{
+						ERigElementType UsedElementType = ERigElementType::None;
+						const int32 NumEnumValues = RigElementTypeEnum->NumEnums();
+
+						for (int32 EnumIndex = 0; EnumIndex < NumEnumValues; EnumIndex++)
+						{
+							const FString EnumText = RigElementTypeEnum->GetDisplayNameTextByIndex(EnumIndex).ToString().ToLower();
+							if (!EnumText.IsEmpty() && UsedFilterString.Contains(EnumText))
+							{
+								UsedElementType = (ERigElementType)RigElementTypeEnum->GetValueByIndex(EnumIndex);
+								break;
+							}
+						}
+
+						if (UsedElementType != ERigElementType::None)
+						{
+							TArray<URigVMPin*> ModelPins = ModelNode->GetAllPinsRecursively();
+							for (URigVMPin* ModelPin : ModelPins)
+							{
+								if (ModelPin->GetCPPTypeObject() == FRigElementKey::StaticStruct())
+								{
+									if (URigVMPin* TypePin = ModelPin->FindSubPin(TEXT("Type")))
+									{
+										FString DefaultValue = RigElementTypeEnum->GetDisplayNameTextByValue((int64)UsedElementType).ToString();
+										Controller->SetPinDefaultValue(TypePin->GetPinPath(), DefaultValue);
+										break;
+									}
+								}
+							}
+						}
+					}
+				}
+			}
+			*/
+#endif
+
+			if (bIsUserFacingNode)
+			{
+				Controller->CloseUndoBracket();
+			}
+			else
+			{
+				Controller->RemoveNode(ModelNode, false);
+			}
+		}
+		else
+		{
+			if (bIsUserFacingNode)
+			{
+				Controller->CancelUndoBracket();
+			}
+		}
+	}
+	return NewNode;
+}
+
+void URigVMEdGraphUnitNodeSpawner::HookupMutableNode(URigVMNode* InModelNode, URigVMBlueprint* InRigBlueprint)
+{
+	/*
+	// Todo in a follow up submit
+	if(!UControlRigEditorSettings::Get()->bAutoLinkMutableNodes)
+	{
+		return;
+	}
+	*/
+	
+	URigVMController* Controller = InRigBlueprint->GetController(InModelNode->GetGraph());
+
+	Controller->ClearNodeSelection(true);
+	Controller->SelectNode(InModelNode, true, true);
+
+	// see if the node has an execute pin
+	URigVMPin* ModelNodeExecutePin = nullptr;
+	for (URigVMPin* Pin : InModelNode->GetPins())
+	{
+		if (Pin->GetScriptStruct())
+		{
+			if (Pin->GetScriptStruct()->IsChildOf(FRigVMExecuteContext::StaticStruct()))
+			{
+				if (Pin->GetDirection() == ERigVMPinDirection::IO || Pin->GetDirection() == ERigVMPinDirection::Input)
+				{
+					ModelNodeExecutePin = Pin;
+					break;
+				}
+			}
+		}
+	}
+
+	// we have an execute pin - so we have to hook it up
+	if (ModelNodeExecutePin)
+	{
+		URigVMPin* ClosestOtherModelNodeExecutePin = nullptr;
+		float ClosestDistance = FLT_MAX;
+
+		const URigVMEdGraphSchema* Schema = GetDefault<URigVMEdGraphSchema>();
+		if (Schema->LastPinForCompatibleCheck)
+		{
+			URigVMPin* FromPin = Controller->GetGraph()->FindPin(Schema->LastPinForCompatibleCheck->GetName());
+			if (FromPin)
+			{
+				if (FromPin->IsExecuteContext() &&
+					(FromPin->GetDirection() == ERigVMPinDirection::IO || FromPin->GetDirection() == ERigVMPinDirection::Output))
+				{
+					ClosestOtherModelNodeExecutePin = FromPin;
+				}
+			}
+
+		}
+
+		if (ClosestOtherModelNodeExecutePin == nullptr)
+		{
+			for (URigVMNode* OtherModelNode : Controller->GetGraph()->GetNodes())
+			{
+				if (OtherModelNode == InModelNode)
+				{
+					continue;
+				}
+
+				for (URigVMPin* Pin : OtherModelNode->GetPins())
+				{
+					if (Pin->GetScriptStruct())
+					{
+						if (Pin->GetScriptStruct()->IsChildOf(FRigVMExecuteContext::StaticStruct()))
+						{
+							if (Pin->GetDirection() == ERigVMPinDirection::IO || Pin->GetDirection() == ERigVMPinDirection::Output)
+							{
+								if (Pin->GetLinkedTargetPins().Num() == 0)
+								{
+									float Distance = (OtherModelNode->GetPosition() - InModelNode->GetPosition()).Size();
+									if (Distance < ClosestDistance)
+									{
+										ClosestOtherModelNodeExecutePin = Pin;
+										ClosestDistance = Distance;
+										break;
+									}
+								}
+							}
+						}
+					}
+				}
+			}
+		}
+
+		if (ClosestOtherModelNodeExecutePin)
+		{
+			Controller->AddLink(ClosestOtherModelNodeExecutePin->GetPinPath(), ModelNodeExecutePin->GetPinPath(), true);
+		}
+	}
+}
+
+#undef LOCTEXT_NAMESPACE
+
diff --git a/Engine/Plugins/Runtime/RigVM/Source/RigVMEditor/Private/EdGraph/NodeSpawners/RigVMEdGraphVariableNodeSpawner.cpp b/Engine/Plugins/Runtime/RigVM/Source/RigVMEditor/Private/EdGraph/NodeSpawners/RigVMEdGraphVariableNodeSpawner.cpp
new file mode 100644
--- /dev/null
+++ b/Engine/Plugins/Runtime/RigVM/Source/RigVMEditor/Private/EdGraph/NodeSpawners/RigVMEdGraphVariableNodeSpawner.cpp
@@ -0,0 +1,220 @@
+// Copyright Epic Games, Inc. All Rights Reserved.
+
+#include "EdGraph/NodeSpawners/RigVMEdGraphVariableNodeSpawner.h"
+#include "EdGraph/RigVMEdGraph.h"
+#include "EdGraph/RigVMEdGraphNode.h"
+#include "EdGraphSchema_K2.h"
+#include "Kismet2/BlueprintEditorUtils.h"
+#include "Settings/EditorStyleSettings.h"
+#include "Editor/EditorEngine.h"
+#include "ObjectEditorUtils.h"
+#include "EditorCategoryUtils.h"
+#include "K2Node_Variable.h"
+#include "BlueprintNodeTemplateCache.h"
+#include "RigVMModel/RigVMController.h"
+#include "RigVMBlueprint.h"
+#include "RigVMModel/Nodes/RigVMVariableNode.h"
+
+#include UE_INLINE_GENERATED_CPP_BY_NAME(RigVMEdGraphVariableNodeSpawner)
+
+#if WITH_EDITOR
+#include "Editor.h"
+#endif
+
+#define LOCTEXT_NAMESPACE "RigVMEdGraphVariableNodeSpawner"
+
+URigVMEdGraphVariableNodeSpawner* URigVMEdGraphVariableNodeSpawner::CreateFromExternalVariable(URigVMBlueprint* InBlueprint, const FRigVMExternalVariable& InExternalVariable, bool bInIsGetter, const FText& InMenuDesc, const FText& InCategory, const FText& InTooltip)
+{
+	URigVMEdGraphVariableNodeSpawner* NodeSpawner = NewObject<URigVMEdGraphVariableNodeSpawner>(GetTransientPackage());
+	NodeSpawner->Blueprint = InBlueprint;
+	NodeSpawner->ExternalVariable = InExternalVariable;
+	NodeSpawner->bIsGetter = bInIsGetter;
+	NodeSpawner->NodeClass = URigVMEdGraphNode::StaticClass();
+
+	FBlueprintActionUiSpec& MenuSignature = NodeSpawner->DefaultMenuSignature;
+	
+	MenuSignature.MenuName = FText::FromString(FString::Printf(TEXT("%s %s"), bInIsGetter ? TEXT("Get") : TEXT("Set"), *InMenuDesc.ToString()));
+	MenuSignature.Tooltip  = InTooltip;
+	MenuSignature.Category = InCategory;
+	MenuSignature.Keywords = FText::FromString(TEXT("Variable"));
+
+	FEdGraphPinType PinType = RigVMTypeUtils::PinTypeFromExternalVariable(InExternalVariable);
+	MenuSignature.Icon = UK2Node_Variable::GetVarIconFromPinType(PinType, MenuSignature.IconTint);
+
+	return NodeSpawner;
+}
+
+URigVMEdGraphVariableNodeSpawner* URigVMEdGraphVariableNodeSpawner::CreateFromLocalVariable(
+	URigVMBlueprint* InBlueprint, URigVMGraph* InGraphOwner, const FRigVMGraphVariableDescription& InLocalVariable,
+	bool bInIsGetter, const FText& InMenuDesc, const FText& InCategory, const FText& InTooltip)
+{
+	URigVMEdGraphVariableNodeSpawner* Spawner = CreateFromExternalVariable(InBlueprint, InLocalVariable.ToExternalVariable(), bInIsGetter, InMenuDesc, InCategory, InTooltip);
+	Spawner->bIsLocalVariable = true;
+	Spawner->GraphOwner = InGraphOwner;
+	return Spawner;
+}
+
+void URigVMEdGraphVariableNodeSpawner::Prime()
+{
+	// we expect that you don't need a node template to construct menu entries
+	// from this, so we choose not to pre-cache one here
+}
+
+bool URigVMEdGraphVariableNodeSpawner::IsTemplateNodeFilteredOut(FBlueprintActionFilter const& Filter) const
+{
+	if (Blueprint.IsValid())
+	{
+		if (!Filter.Context.Blueprints.Contains(Blueprint.Get()))
+		{
+			return true;
+		}
+
+		if (bIsLocalVariable)
+		{
+			bool bIsFiltered = true;
+			if (Filter.Context.Graphs.Num() == 1 && GraphOwner.IsValid())
+			{
+				if (URigVMEdGraph* Graph = Cast<URigVMEdGraph>(Filter.Context.Graphs[0]))
+				{
+					if (GraphOwner.Get() == Graph->GetModel())
+					{
+						bIsFiltered = false;
+					}
+				}
+			}
+
+			if (bIsFiltered)
+			{
+				return true;
+			}
+		}
+	}
+	return Super::IsTemplateNodeFilteredOut(Filter);
+}
+
+FBlueprintNodeSignature URigVMEdGraphVariableNodeSpawner::GetSpawnerSignature() const
+{
+	return FBlueprintNodeSignature(FString("RigUnit=" + ExternalVariable.Name.ToString()));
+}
+
+FBlueprintActionUiSpec URigVMEdGraphVariableNodeSpawner::GetUiSpec(FBlueprintActionContext const& Context, FBindingSet const& Bindings) const
+{
+	UEdGraph* TargetGraph = (Context.Graphs.Num() > 0) ? Context.Graphs[0] : nullptr;
+	FBlueprintActionUiSpec MenuSignature = PrimeDefaultUiSpec(TargetGraph);
+
+	DynamicUiSignatureGetter.ExecuteIfBound(Context, Bindings, &MenuSignature);
+	return MenuSignature;
+}
+
+UEdGraphNode* URigVMEdGraphVariableNodeSpawner::Invoke(UEdGraph* ParentGraph, FBindingSet const& Bindings, FVector2D const Location) const
+{
+	URigVMEdGraphNode* NewNode = nullptr;
+
+	bool const bIsTemplateNode = FBlueprintNodeTemplateCache::IsTemplateOuter(ParentGraph);
+	bool const bIsUserFacingNode = !bIsTemplateNode;
+
+	// First create a backing member for our node
+	URigVMEdGraph* RigGraph = Cast<URigVMEdGraph>(ParentGraph);
+	if(RigGraph == nullptr) return nullptr;
+	URigVMBlueprint* RigBlueprint = Cast<URigVMBlueprint>(FBlueprintEditorUtils::FindBlueprintForGraph(ParentGraph));
+	check(RigBlueprint);
+
+	FName MemberName = NAME_None;
+
+#if WITH_EDITOR
+	if (GEditor && !bIsTemplateNode)
+	{
+		GEditor->CancelTransaction(0);
+	}
+#endif
+
+	URigVMController* Controller = bIsTemplateNode ? RigGraph->GetTemplateController() : RigBlueprint->GetController(ParentGraph);
+
+	if (bIsUserFacingNode)
+	{
+		Controller->OpenUndoBracket(TEXT("Add Variable"));
+	}
+
+	FString ObjectPath;
+	if (ExternalVariable.TypeObject)
+	{
+		ObjectPath = ExternalVariable.TypeObject->GetPathName();
+	}
+
+	FString TypeName = ExternalVariable.TypeName.ToString();
+	if (ExternalVariable.bIsArray)
+	{
+		TypeName = FString::Printf(TEXT("TArray<%s>"), *TypeName);
+	}
+
+	FString NodeName;
+	if(bIsTemplateNode)
+	{
+		// for template controllers let's rely on locally defined nodes
+		// without a backing model node. access to local variables or
+		// input arguments doesn't work on the template controller.
+		
+		static constexpr TCHAR VariableNodeNameFormat[] = TEXT("VariableNode_%s_%s");
+		static const FString GetterPrefix = TEXT("Getter");
+		static const FString SetterPrefix = TEXT("Setter");
+		NodeName = FString::Printf(VariableNodeNameFormat, bIsGetter ? *GetterPrefix : *SetterPrefix, *ExternalVariable.Name.ToString());
+
+		NewNode = NewObject<URigVMEdGraphNode>(ParentGraph, *NodeName);
+		ParentGraph->AddNode(NewNode, false);
+
+		NewNode->CreateNewGuid();
+		NewNode->PostPlacedNewNode();
+
+		UEdGraphPin* ValuePin = UEdGraphPin::CreatePin(NewNode);
+		NewNode->Pins.Add(ValuePin);
+
+		ValuePin->PinType = RigVMTypeUtils::PinTypeFromExternalVariable(ExternalVariable);
+		ValuePin->Direction = bIsGetter ? EGPD_Output : EGPD_Input;
+		NewNode->SetFlags(RF_Transactional);
+		return NewNode;
+	}
+
+	if (URigVMNode* ModelNode = Controller->AddVariableNodeFromObjectPath(ExternalVariable.Name, TypeName, ObjectPath, bIsGetter, FString(), Location, NodeName, !bIsTemplateNode, !bIsTemplateNode))
+	{
+		for (UEdGraphNode* Node : ParentGraph->Nodes)
+		{
+			if (URigVMEdGraphNode* RigNode = Cast<URigVMEdGraphNode>(Node))
+			{
+				if (RigNode->GetModelNodeName() == ModelNode->GetFName())
+				{
+					NewNode = RigNode;
+					break;
+				}
+			}
+		}
+
+		if (bIsUserFacingNode)
+		{
+			if (NewNode)
+			{
+				Controller->ClearNodeSelection(true);
+				Controller->SelectNode(ModelNode, true, true);
+			}
+			Controller->CloseUndoBracket();
+		}
+		else
+		{
+			// similar to UBlueprintNodeSpawner::Invoke -> UBlueprintNodeSpawner::SpawnEdGraphNode
+			// we simply want the node, but not actually adding it to a graph
+			Controller->RemoveNode(ModelNode, false);
+		}
+	}
+	else
+	{
+		if (!bIsTemplateNode)
+		{
+			Controller->CancelUndoBracket();
+		}
+	}
+
+
+	return NewNode;
+}
+
+#undef LOCTEXT_NAMESPACE
+
diff --git a/Engine/Plugins/Runtime/RigVM/Source/RigVMEditor/Public/EdGraph/NodeSpawners/RigVMEdGraphEnumNodeSpawner.h b/Engine/Plugins/Runtime/RigVM/Source/RigVMEditor/Public/EdGraph/NodeSpawners/RigVMEdGraphEnumNodeSpawner.h
new file mode 100644
--- /dev/null
+++ b/Engine/Plugins/Runtime/RigVM/Source/RigVMEditor/Public/EdGraph/NodeSpawners/RigVMEdGraphEnumNodeSpawner.h
@@ -0,0 +1,43 @@
+// Copyright Epic Games, Inc. All Rights Reserved.
+
+#pragma once
+
+#include "CoreMinimal.h"
+#include "UObject/ObjectMacros.h"
+#include "UObject/Object.h"
+#include "Templates/SubclassOf.h"
+#include "UObject/UnrealType.h"
+#include "EdGraph/EdGraphPin.h"
+#include "Engine/Blueprint.h"
+#include "BlueprintActionFilter.h"
+#include "BlueprintNodeSignature.h"
+#include "BlueprintFieldNodeSpawner.h"
+#include "RigVMEdGraphEnumNodeSpawner.generated.h"
+
+class URigVMEdGraphNode;
+
+UCLASS(Transient)
+class RIGVMEDITOR_API URigVMEdGraphEnumNodeSpawner : public UBlueprintNodeSpawner
+{
+	GENERATED_BODY()
+
+public:
+
+	/**
+	 * Creates a new URigVMEdGraphEnumNodeSpawner
+	 * 
+	 * @return A newly allocated instance of this class.
+	 */
+	static URigVMEdGraphEnumNodeSpawner* CreateForEnum(UEnum* InEnum, const FText& InMenuDesc, const FText& InCategory, const FText& InTooltip);
+
+	// UBlueprintNodeSpawner interface
+	virtual FBlueprintNodeSignature GetSpawnerSignature() const override;
+	virtual FBlueprintActionUiSpec GetUiSpec(FBlueprintActionContext const& Context, FBindingSet const& Bindings) const override;
+	virtual UEdGraphNode* Invoke(UEdGraph* ParentGraph, FBindingSet const& Bindings, FVector2D const Location) const override;
+	// End UBlueprintNodeSpawner interface
+
+private:
+
+	UEnum* Enum;
+	friend class UEngineTestControlRig;
+};
diff --git a/Engine/Plugins/Runtime/RigVM/Source/RigVMEditor/Public/EdGraph/NodeSpawners/RigVMEdGraphFunctionRefNodeSpawner.h b/Engine/Plugins/Runtime/RigVM/Source/RigVMEditor/Public/EdGraph/NodeSpawners/RigVMEdGraphFunctionRefNodeSpawner.h
new file mode 100644
--- /dev/null
+++ b/Engine/Plugins/Runtime/RigVM/Source/RigVMEditor/Public/EdGraph/NodeSpawners/RigVMEdGraphFunctionRefNodeSpawner.h
@@ -0,0 +1,70 @@
+// Copyright Epic Games, Inc. All Rights Reserved.
+
+#pragma once
+
+#include "CoreMinimal.h"
+#include "UObject/ObjectMacros.h"
+#include "UObject/Object.h"
+#include "Templates/SubclassOf.h"
+#include "UObject/UnrealType.h"
+#include "EdGraph/EdGraphPin.h"
+#include "Engine/Blueprint.h"
+#include "BlueprintActionFilter.h"
+#include "BlueprintNodeSignature.h"
+#include "BlueprintFieldNodeSpawner.h"
+#include "RigVMBlueprint.h"
+#include "EdGraph/RigVMEdGraph.h"
+#include "EdGraph/RigVMEdGraphNode.h"
+#include "RigVMModel/Nodes/RigVMUnitNode.h"
+#include "RigVMEdGraphFunctionRefNodeSpawner.generated.h"
+
+class URigVMEdGraphNode;
+
+UCLASS(Transient)
+class RIGVMEDITOR_API URigVMEdGraphFunctionRefNodeSpawner : public UBlueprintNodeSpawner
+{
+	GENERATED_BODY()
+
+public:
+
+	/**
+	 * Creates a new URigVMEdGraphFunctionRefNodeSpawner, charged with spawning a function reference
+	 * 
+	 * @return A newly allocated instance of this class.
+	 */
+	static URigVMEdGraphFunctionRefNodeSpawner* CreateFromFunction(URigVMLibraryNode* InFunction);
+
+	/**
+	 * Creates a new URigVMEdGraphFunctionRefNodeSpawner, charged with spawning a function reference
+	 * 
+	 * @return A newly allocated instance of this class.
+	 */
+	static URigVMEdGraphFunctionRefNodeSpawner* CreateFromAssetData(const FAssetData& InAssetData, const FRigVMGraphFunctionHeader& InPublicFunction);
+	static URigVMEdGraphFunctionRefNodeSpawner* CreateFromAssetData(const FAssetData& InAssetData, const FRigVMOldPublicFunctionData& InPublicFunction);
+
+	// UBlueprintNodeSpawner interface
+	virtual void Prime() override;
+	virtual FBlueprintNodeSignature GetSpawnerSignature() const override;
+	virtual FBlueprintActionUiSpec GetUiSpec(FBlueprintActionContext const& Context, FBindingSet const& Bindings) const override;
+	virtual UEdGraphNode* Invoke(UEdGraph* ParentGraph, FBindingSet const& Bindings, FVector2D const Location) const override;
+	virtual bool IsTemplateNodeFilteredOut(FBlueprintActionFilter const& Filter) const override;
+	// End UBlueprintNodeSpawner interface
+
+private:
+
+	/** The public function definition we will spawn from [optional] */
+	UPROPERTY(Transient)
+	mutable FRigVMGraphFunctionHeader ReferencedPublicFunctionHeader;
+
+	/** Marked as true for local function definitions */
+	UPROPERTY(Transient)
+	bool bIsLocalFunction;
+
+	UPROPERTY()
+	FSoftObjectPath AssetPath;
+
+	static URigVMEdGraphNode* SpawnNode(UEdGraph* ParentGraph, UBlueprint* Blueprint, FRigVMGraphFunctionHeader& InFunction, FVector2D const Location);
+
+	friend class UEngineTestControlRig;
+	friend class FControlRigEditor;
+};
diff --git a/Engine/Plugins/Runtime/RigVM/Source/RigVMEditor/Public/EdGraph/NodeSpawners/RigVMEdGraphInvokeEntryNodeSpawner.h b/Engine/Plugins/Runtime/RigVM/Source/RigVMEditor/Public/EdGraph/NodeSpawners/RigVMEdGraphInvokeEntryNodeSpawner.h
new file mode 100644
--- /dev/null
+++ b/Engine/Plugins/Runtime/RigVM/Source/RigVMEditor/Public/EdGraph/NodeSpawners/RigVMEdGraphInvokeEntryNodeSpawner.h
@@ -0,0 +1,39 @@
+// Copyright Epic Games, Inc. All Rights Reserved.
+
+#pragma once
+
+#include "CoreMinimal.h"
+#include "BlueprintActionFilter.h"
+#include "BlueprintNodeSignature.h"
+#include "BlueprintFieldNodeSpawner.h"
+#include "RigVMModel/RigVMGraph.h"
+#include "RigVMEdGraphInvokeEntryNodeSpawner.generated.h"
+
+class URigVMEdGraphNode;
+class URigVMBlueprint;
+
+UCLASS(Transient)
+class RIGVMEDITOR_API URigVMEdGraphInvokeEntryNodeSpawner : public UBlueprintNodeSpawner
+{
+	GENERATED_BODY()
+
+public:
+
+	static URigVMEdGraphInvokeEntryNodeSpawner* CreateForEntry(URigVMBlueprint* InBlueprint, const FName& InEntryName, const FText& InMenuDesc, const FText& InCategory, const FText& InTooltip);
+
+	// UBlueprintNodeSpawner interface
+	virtual void Prime() override;
+	virtual FBlueprintNodeSignature GetSpawnerSignature() const override;
+	virtual FBlueprintActionUiSpec GetUiSpec(FBlueprintActionContext const& Context, FBindingSet const& Bindings) const override;
+	virtual UEdGraphNode* Invoke(UEdGraph* ParentGraph, FBindingSet const& Bindings, FVector2D const Location) const override;
+	// End UBlueprintNodeSpawner interface
+
+private:
+
+	/** The pin type we will spawn */
+	TWeakObjectPtr<URigVMBlueprint> Blueprint;
+	TWeakObjectPtr<URigVMGraph> GraphOwner;
+	FName EntryName;
+
+	friend class UEngineTestControlRig;
+};
diff --git a/Engine/Plugins/Runtime/RigVM/Source/RigVMEditor/Public/EdGraph/NodeSpawners/RigVMEdGraphTemplateNodeSpawner.h b/Engine/Plugins/Runtime/RigVM/Source/RigVMEditor/Public/EdGraph/NodeSpawners/RigVMEdGraphTemplateNodeSpawner.h
new file mode 100644
--- /dev/null
+++ b/Engine/Plugins/Runtime/RigVM/Source/RigVMEditor/Public/EdGraph/NodeSpawners/RigVMEdGraphTemplateNodeSpawner.h
@@ -0,0 +1,54 @@
+// Copyright Epic Games, Inc. All Rights Reserved.
+
+#pragma once
+
+#include "CoreMinimal.h"
+#include "UObject/ObjectMacros.h"
+#include "UObject/Object.h"
+#include "Templates/SubclassOf.h"
+#include "UObject/UnrealType.h"
+#include "EdGraph/EdGraphPin.h"
+#include "Engine/Blueprint.h"
+#include "BlueprintActionFilter.h"
+#include "BlueprintNodeSignature.h"
+#include "BlueprintFieldNodeSpawner.h"
+#include "RigVMBlueprint.h"
+#include "EdGraph/RigVMEdGraph.h"
+#include "EdGraph/RigVMEdGraphNode.h"
+#include "RigVMModel/Nodes/RigVMTemplateNode.h"
+#include "RigVMEdGraphTemplateNodeSpawner.generated.h"
+
+class URigVMEdGraphNode;
+
+UCLASS(Transient)
+class RIGVMEDITOR_API URigVMEdGraphTemplateNodeSpawner : public UBlueprintNodeSpawner
+{
+	GENERATED_BODY()
+
+public:
+	/**
+	 * Creates a new URigVMEdGraphVariableNodeSpawner, charged with spawning 
+	 * a new member-variable node
+	 * 
+	 * @return A newly allocated instance of this class.
+	 */
+	static URigVMEdGraphTemplateNodeSpawner* CreateFromNotation(const FName& InNotation, const FText& InMenuDesc, const FText& InCategory, const FText& InTooltip);
+
+	// UBlueprintNodeSpawner interface
+	virtual void Prime() override;
+	virtual FBlueprintNodeSignature GetSpawnerSignature() const override;
+	bool IsTemplateNodeFilteredOut(FBlueprintActionFilter const& Filter) const override;
+	virtual FBlueprintActionUiSpec GetUiSpec(FBlueprintActionContext const& Context, FBindingSet const& Bindings) const override;
+	virtual UEdGraphNode* Invoke(UEdGraph* ParentGraph, FBindingSet const& Bindings, FVector2D const Location) const override;
+	// End UBlueprintNodeSpawner interface
+
+private:
+	/** The unit type we will spawn */
+	UPROPERTY(Transient)
+	FName TemplateNotation;
+
+	static URigVMEdGraphNode* SpawnNode(UEdGraph* ParentGraph, UBlueprint* Blueprint, const FName& InNotation, FVector2D const Location);
+
+	friend class UEngineTestControlRig;
+	friend class FControlRigEditor;
+};
diff --git a/Engine/Plugins/Runtime/RigVM/Source/RigVMEditor/Public/EdGraph/NodeSpawners/RigVMEdGraphUnitNodeSpawner.h b/Engine/Plugins/Runtime/RigVM/Source/RigVMEditor/Public/EdGraph/NodeSpawners/RigVMEdGraphUnitNodeSpawner.h
new file mode 100644
--- /dev/null
+++ b/Engine/Plugins/Runtime/RigVM/Source/RigVMEditor/Public/EdGraph/NodeSpawners/RigVMEdGraphUnitNodeSpawner.h
@@ -0,0 +1,57 @@
+// Copyright Epic Games, Inc. All Rights Reserved.
+
+#pragma once
+
+#include "CoreMinimal.h"
+#include "UObject/ObjectMacros.h"
+#include "UObject/Object.h"
+#include "Templates/SubclassOf.h"
+#include "UObject/UnrealType.h"
+#include "EdGraph/EdGraphPin.h"
+#include "Engine/Blueprint.h"
+#include "BlueprintActionFilter.h"
+#include "BlueprintNodeSignature.h"
+#include "BlueprintFieldNodeSpawner.h"
+#include "RigVMBlueprint.h"
+#include "EdGraph/RigVMEdGraph.h"
+#include "EdGraph/RigVMEdGraphNode.h"
+#include "RigVMModel/Nodes/RigVMUnitNode.h"
+#include "RigVMEdGraphUnitNodeSpawner.generated.h"
+
+class URigVMEdGraphNode;
+
+UCLASS(Transient)
+class RIGVMEDITOR_API URigVMEdGraphUnitNodeSpawner : public UBlueprintNodeSpawner
+{
+	GENERATED_BODY()
+
+public:
+	/**
+	 * Creates a new URigVMEdGraphVariableNodeSpawner, charged with spawning 
+	 * a new member-variable node
+	 * 
+	 * @return A newly allocated instance of this class.
+	 */
+	static URigVMEdGraphUnitNodeSpawner* CreateFromStruct(UScriptStruct* InStruct, const FName& InMethodName, const FText& InMenuDesc, const FText& InCategory, const FText& InTooltip);
+
+	// UBlueprintNodeSpawner interface
+	virtual void Prime() override;
+	virtual FBlueprintNodeSignature GetSpawnerSignature() const override;
+	virtual FBlueprintActionUiSpec GetUiSpec(FBlueprintActionContext const& Context, FBindingSet const& Bindings) const override;
+	virtual UEdGraphNode* Invoke(UEdGraph* ParentGraph, FBindingSet const& Bindings, FVector2D const Location) const override;
+	static void HookupMutableNode(URigVMNode* InModelNode, URigVMBlueprint* InRigBlueprint);
+	// End UBlueprintNodeSpawner interface
+
+private:
+	/** The unit type we will spawn */
+	UPROPERTY(Transient)
+	TObjectPtr<UScriptStruct> StructTemplate;
+
+	UPROPERTY(Transient)
+	FName MethodName;
+
+	static URigVMEdGraphNode* SpawnNode(UEdGraph* ParentGraph, UBlueprint* Blueprint, UScriptStruct* StructTemplate, const FName& InMethodName, FVector2D const Location);
+
+	friend class UEngineTestControlRig;
+	friend class FControlRigEditor;
+};
diff --git a/Engine/Plugins/Runtime/RigVM/Source/RigVMEditor/Public/EdGraph/NodeSpawners/RigVMEdGraphVariableNodeSpawner.h b/Engine/Plugins/Runtime/RigVM/Source/RigVMEditor/Public/EdGraph/NodeSpawners/RigVMEdGraphVariableNodeSpawner.h
new file mode 100644
--- /dev/null
+++ b/Engine/Plugins/Runtime/RigVM/Source/RigVMEditor/Public/EdGraph/NodeSpawners/RigVMEdGraphVariableNodeSpawner.h
@@ -0,0 +1,56 @@
+// Copyright Epic Games, Inc. All Rights Reserved.
+
+#pragma once
+
+#include "CoreMinimal.h"
+#include "UObject/ObjectMacros.h"
+#include "UObject/Object.h"
+#include "Templates/SubclassOf.h"
+#include "UObject/UnrealType.h"
+#include "EdGraph/EdGraphPin.h"
+#include "Engine/Blueprint.h"
+#include "BlueprintActionFilter.h"
+#include "BlueprintNodeSignature.h"
+#include "BlueprintFieldNodeSpawner.h"
+#include "RigVMCore/RigVMExternalVariable.h"
+#include "RigVMModel/RigVMVariableDescription.h"
+#include "RigVMEdGraphVariableNodeSpawner.generated.h"
+
+class URigVMEdGraphNode;
+class URigVMBlueprint;
+
+UCLASS(Transient)
+class RIGVMEDITOR_API URigVMEdGraphVariableNodeSpawner : public UBlueprintNodeSpawner
+{
+	GENERATED_BODY()
+
+public:
+
+	/**
+	 * Creates a new URigVMEdGraphVariableNodeSpawner, charged with spawning 
+	 * a new member-variable node
+	 * 
+	 * @return A newly allocated instance of this class.
+	 */
+	static URigVMEdGraphVariableNodeSpawner* CreateFromExternalVariable(URigVMBlueprint* InBlueprint, const FRigVMExternalVariable& InExternalVariable, bool bInIsGetter, const FText& InMenuDesc, const FText& InCategory, const FText& InTooltip);
+	static URigVMEdGraphVariableNodeSpawner* CreateFromLocalVariable(URigVMBlueprint* InBlueprint, URigVMGraph* InGraphOwner, const FRigVMGraphVariableDescription& InLocalVariable, bool bInIsGetter, const FText& InMenuDesc, const FText& InCategory, const FText& InTooltip);
+
+	// UBlueprintNodeSpawner interface
+	virtual void Prime() override;
+	virtual FBlueprintNodeSignature GetSpawnerSignature() const override;
+	virtual FBlueprintActionUiSpec GetUiSpec(FBlueprintActionContext const& Context, FBindingSet const& Bindings) const override;
+	virtual UEdGraphNode* Invoke(UEdGraph* ParentGraph, FBindingSet const& Bindings, FVector2D const Location) const override;
+	virtual bool IsTemplateNodeFilteredOut(FBlueprintActionFilter const& Filter) const override;
+	// End UBlueprintNodeSpawner interface
+
+private:
+
+	/** The pin type we will spawn */
+	TWeakObjectPtr<URigVMBlueprint> Blueprint;
+	TWeakObjectPtr<URigVMGraph> GraphOwner;
+	FRigVMExternalVariable ExternalVariable;
+	bool bIsGetter;
+	bool bIsLocalVariable;	
+
+	friend class UEngineTestControlRig;
+};
diff --git a/Engine/Plugins/Runtime/RigVM/Source/RigVMEditor/RigVMEditor.Build.cs b/Engine/Plugins/Runtime/RigVM/Source/RigVMEditor/RigVMEditor.Build.cs
--- a/Engine/Plugins/Runtime/RigVM/Source/RigVMEditor/RigVMEditor.Build.cs
+++ b/Engine/Plugins/Runtime/RigVM/Source/RigVMEditor/RigVMEditor.Build.cs
@@ -12,7 +12,11 @@ public class RigVMEditor : ModuleRules
                 "CoreUObject",
                 "Engine",
                 "RigVM",
-                "RigVMDeveloper"
+                "RigVMDeveloper",
+                "SlateCore",
+                "BlueprintGraph",
+                "GraphEditor",
+                "UnrealEd",
             }
         );
     }
