Control Rig: Move Variable mapping widget to RigVMEditor

diff --git a/Engine/Plugins/Animation/ControlRig/Source/ControlRigDeveloper/ControlRigDeveloper.Build.cs b/Engine/Plugins/Animation/ControlRig/Source/ControlRigDeveloper/ControlRigDeveloper.Build.cs
--- a/Engine/Plugins/Animation/ControlRig/Source/ControlRigDeveloper/ControlRigDeveloper.Build.cs
+++ b/Engine/Plugins/Animation/ControlRig/Source/ControlRigDeveloper/ControlRigDeveloper.Build.cs
@@ -56,6 +56,7 @@ namespace UnrealBuildTool.Rules
                         "RigVMDeveloper",
 						"GraphEditor",
 						"ApplicationCore",
+						"RigVMEditor",
 					}
                 );
 
diff --git a/Engine/Plugins/Animation/ControlRig/Source/ControlRigDeveloper/Private/AnimGraphNode_ControlRig.cpp b/Engine/Plugins/Animation/ControlRig/Source/ControlRigDeveloper/Private/AnimGraphNode_ControlRig.cpp
--- a/Engine/Plugins/Animation/ControlRig/Source/ControlRigDeveloper/Private/AnimGraphNode_ControlRig.cpp
+++ b/Engine/Plugins/Animation/ControlRig/Source/ControlRigDeveloper/Private/AnimGraphNode_ControlRig.cpp
@@ -7,7 +7,9 @@
 #include "Widgets/Text/STextBlock.h"
 #include "Widgets/Input/SCheckBox.h"
 #include "DetailWidgetRow.h"
-#include "SVariableMappingWidget.h"
+#if WITH_EDITOR
+#include "Widgets/SRigVMVariableMappingWidget.h"
+#endif
 #include "ScopedTransaction.h"
 #include "Kismet2/BlueprintEditorUtils.h"
 #include "AnimationGraphSchema.h"
@@ -502,33 +504,37 @@ void UAnimGraphNode_ControlRig::CustomizeDetails(IDetailLayoutBuilder& DetailBui
 	IDetailCategoryBuilder& InputCategoryBuilder = DetailBuilder.EditCategory(FName(TEXT("Input")));
 
 	FDetailWidgetRow& InputWidgetRow = InputCategoryBuilder.AddCustomRow(FText::FromString(TEXT("Input")));
+#if WITH_EDITOR
 	InputWidgetRow.WholeRowContent()
 	[
-		SNew(SVariableMappingWidget)
-		.OnVariableMappingChanged(FOnVariableMappingChanged::CreateUObject(this, &UAnimGraphNode_ControlRig::OnVariableMappingChanged, true))
-		.OnGetVariableMapping(FOnGetVariableMapping::CreateUObject(this, &UAnimGraphNode_ControlRig::GetVariableMapping, true))
-		.OnGetAvailableMapping(FOnGetAvailableMapping::CreateUObject(this, &UAnimGraphNode_ControlRig::GetAvailableMapping, true))
-		.OnCreateVariableMapping(FOnCreateVariableMapping::CreateUObject(this, &UAnimGraphNode_ControlRig::CreateVariableMapping, true))
-		.OnVariableOptionAvailable(FOnVarOptionAvailable::CreateUObject(this, &UAnimGraphNode_ControlRig::IsAvailableToMapToCurve, true))
-		.OnPinGetCheckState(FOnPinGetCheckState::CreateUObject(this, &UAnimGraphNode_ControlRig::IsPropertyExposed))
-		.OnPinCheckStateChanged(FOnPinCheckStateChanged::CreateUObject(this, &UAnimGraphNode_ControlRig::OnPropertyExposeCheckboxChanged))
-		.OnPinIsEnabledCheckState(FOnPinIsCheckEnabled::CreateUObject(this, &UAnimGraphNode_ControlRig::IsPropertyExposeEnabled))
+		SNew(SRigVMVariableMappingWidget)
+		.OnVariableMappingChanged(FOnRigVMVariableMappingChanged::CreateUObject(this, &UAnimGraphNode_ControlRig::OnVariableMappingChanged, true))
+		.OnGetVariableMapping(FOnRigVMGetVariableMapping::CreateUObject(this, &UAnimGraphNode_ControlRig::GetVariableMapping, true))
+		.OnGetAvailableMapping(FOnRigVMGetAvailableMapping::CreateUObject(this, &UAnimGraphNode_ControlRig::GetAvailableMapping, true))
+		.OnCreateVariableMapping(FOnRigVMCreateVariableMapping::CreateUObject(this, &UAnimGraphNode_ControlRig::CreateVariableMapping, true))
+		.OnVariableOptionAvailable(FOnRigVMVarOptionAvailable::CreateUObject(this, &UAnimGraphNode_ControlRig::IsAvailableToMapToCurve, true))
+		.OnPinGetCheckState(FOnRigVMPinGetCheckState::CreateUObject(this, &UAnimGraphNode_ControlRig::IsPropertyExposed))
+		.OnPinCheckStateChanged(FOnRigVMPinCheckStateChanged::CreateUObject(this, &UAnimGraphNode_ControlRig::OnPropertyExposeCheckboxChanged))
+		.OnPinIsEnabledCheckState(FOnRigVMPinIsCheckEnabled::CreateUObject(this, &UAnimGraphNode_ControlRig::IsPropertyExposeEnabled))
 	];
+#endif
 
 	IDetailCategoryBuilder& OutputCategoryBuilder = DetailBuilder.EditCategory(FName(TEXT("Output")));
 	FDetailWidgetRow& OutputWidgetRow = OutputCategoryBuilder.AddCustomRow(FText::FromString(TEXT("Output")));
+#if WITH_EDITOR
 	OutputWidgetRow.WholeRowContent()
 	[
-		SNew(SVariableMappingWidget)
-		.OnVariableMappingChanged(FOnVariableMappingChanged::CreateUObject(this, &UAnimGraphNode_ControlRig::OnVariableMappingChanged, false))
-		.OnGetVariableMapping(FOnGetVariableMapping::CreateUObject(this, &UAnimGraphNode_ControlRig::GetVariableMapping, false))
-		.OnGetAvailableMapping(FOnGetAvailableMapping::CreateUObject(this, &UAnimGraphNode_ControlRig::GetAvailableMapping, false))
-		.OnCreateVariableMapping(FOnCreateVariableMapping::CreateUObject(this, &UAnimGraphNode_ControlRig::CreateVariableMapping, false))
-		.OnVariableOptionAvailable(FOnVarOptionAvailable::CreateUObject(this, &UAnimGraphNode_ControlRig::IsAvailableToMapToCurve, false))
-		.OnPinGetCheckState(FOnPinGetCheckState::CreateUObject(this, &UAnimGraphNode_ControlRig::IsPropertyExposed))
-		.OnPinCheckStateChanged(FOnPinCheckStateChanged::CreateUObject(this, &UAnimGraphNode_ControlRig::OnPropertyExposeCheckboxChanged))
-		.OnPinIsEnabledCheckState(FOnPinIsCheckEnabled::CreateUObject(this, &UAnimGraphNode_ControlRig::IsPropertyExposeEnabled))
+		SNew(SRigVMVariableMappingWidget)
+		.OnVariableMappingChanged(FOnRigVMVariableMappingChanged::CreateUObject(this, &UAnimGraphNode_ControlRig::OnVariableMappingChanged, false))
+		.OnGetVariableMapping(FOnRigVMGetVariableMapping::CreateUObject(this, &UAnimGraphNode_ControlRig::GetVariableMapping, false))
+		.OnGetAvailableMapping(FOnRigVMGetAvailableMapping::CreateUObject(this, &UAnimGraphNode_ControlRig::GetAvailableMapping, false))
+		.OnCreateVariableMapping(FOnRigVMCreateVariableMapping::CreateUObject(this, &UAnimGraphNode_ControlRig::CreateVariableMapping, false))
+		.OnVariableOptionAvailable(FOnRigVMVarOptionAvailable::CreateUObject(this, &UAnimGraphNode_ControlRig::IsAvailableToMapToCurve, false))
+		.OnPinGetCheckState(FOnRigVMPinGetCheckState::CreateUObject(this, &UAnimGraphNode_ControlRig::IsPropertyExposed))
+		.OnPinCheckStateChanged(FOnRigVMPinCheckStateChanged::CreateUObject(this, &UAnimGraphNode_ControlRig::OnPropertyExposeCheckboxChanged))
+		.OnPinIsEnabledCheckState(FOnRigVMPinIsCheckEnabled::CreateUObject(this, &UAnimGraphNode_ControlRig::IsPropertyExposeEnabled))
 	];
+#endif
 
 	TSharedRef<IPropertyHandle> ClassHandle = DetailBuilder.GetProperty(TEXT("Node.ControlRigClass"), GetClass());
 	if (ClassHandle->IsValidHandle())
@@ -593,6 +599,8 @@ void UAnimGraphNode_ControlRig::GetVariables(bool bInput, TMap<FName, FRigVMExte
 	}
 }
 
+#if WITH_EDITOR
+
 void UAnimGraphNode_ControlRig::OnVariableMappingChanged(const FName& PathName, const FName& Curve, bool bInput)
 {
 	DECLARE_SCOPE_HIERARCHICAL_COUNTER_FUNC()
@@ -668,7 +676,7 @@ void UAnimGraphNode_ControlRig::GetAvailableMapping(const FName& PathName, TArra
 	}
 }
 
-void UAnimGraphNode_ControlRig::CreateVariableMapping(const FString& FilteredText, TArray< TSharedPtr<FVariableMappingInfo> >& OutArray, bool bInput)
+void UAnimGraphNode_ControlRig::CreateVariableMapping(const FString& FilteredText, TArray< TSharedPtr<FRigVMVariableMappingInfo> >& OutArray, bool bInput)
 {
 	DECLARE_SCOPE_HIERARCHICAL_COUNTER_FUNC()
 
@@ -687,9 +695,9 @@ void UAnimGraphNode_ControlRig::CreateVariableMapping(const FString& FilteredTex
 		if (!bDoFiltering || 
 			(DisplayName.Contains(FilteredText) || MappedName.Contains(FilteredText)))
 		{
-			TSharedRef<FVariableMappingInfo> Item = FVariableMappingInfo::Make(Iter.Key());
+			TSharedRef<FRigVMVariableMappingInfo> Item = FRigVMVariableMappingInfo::Make(Iter.Key());
 
-			FVariableMappingInfo& ItemRaw = Item.Get();
+			FRigVMVariableMappingInfo& ItemRaw = Item.Get();
 			FString PathString = ItemRaw.GetPathName().ToString();
 			FString DisplayString = ItemRaw.GetDisplayName();
 
@@ -715,7 +723,7 @@ void UAnimGraphNode_ControlRig::CreateVariableMapping(const FString& FilteredTex
 
 							if (!bDoFiltering || DisplayName.Contains(FilteredText))
 							{
-								TSharedRef<FVariableMappingInfo> Item = FVariableMappingInfo::Make(ControlName);
+								TSharedRef<FRigVMVariableMappingInfo> Item = FRigVMVariableMappingInfo::Make(ControlName);
 								OutArray.Add(Item);
 							}
 						}
@@ -727,6 +735,8 @@ void UAnimGraphNode_ControlRig::CreateVariableMapping(const FString& FilteredTex
 	}
 }
 
+#endif
+
 void UAnimGraphNode_ControlRig::PostEditChangeProperty(FPropertyChangedEvent& PropertyChangedEvent)
 {
 	DECLARE_SCOPE_HIERARCHICAL_COUNTER_FUNC()
diff --git a/Engine/Plugins/Animation/ControlRig/Source/ControlRigDeveloper/Private/SVariableMappingWidget.cpp b/Engine/Plugins/Animation/ControlRig/Source/ControlRigDeveloper/Private/SVariableMappingWidget.cpp
deleted file mode 100644
diff --git a/Engine/Plugins/Animation/ControlRig/Source/ControlRigDeveloper/Private/SVariableMappingWidget.h b/Engine/Plugins/Animation/ControlRig/Source/ControlRigDeveloper/Private/SVariableMappingWidget.h
deleted file mode 100644
diff --git a/Engine/Plugins/Animation/ControlRig/Source/ControlRigDeveloper/Public/AnimGraphNode_ControlRig.h b/Engine/Plugins/Animation/ControlRig/Source/ControlRigDeveloper/Public/AnimGraphNode_ControlRig.h
--- a/Engine/Plugins/Animation/ControlRig/Source/ControlRigDeveloper/Public/AnimGraphNode_ControlRig.h
+++ b/Engine/Plugins/Animation/ControlRig/Source/ControlRigDeveloper/Public/AnimGraphNode_ControlRig.h
@@ -7,7 +7,7 @@
 #include "SSearchableComboBox.h"
 #include "AnimGraphNode_ControlRig.generated.h"
 
-struct FVariableMappingInfo;
+struct FRigVMVariableMappingInfo;
 
 UCLASS(MinimalAPI)
 class UAnimGraphNode_ControlRig : public UAnimGraphNode_CustomProperty
@@ -45,11 +45,13 @@ private:
 	ECheckBoxState IsPropertyExposed(FName PropertyName) const;
 	void OnPropertyExposeCheckboxChanged(ECheckBoxState NewState, FName PropertyName);
 
-	// SVariableMappingWidget related
+#if WITH_EDITOR
+	// SRigVMVariableMappingWidget related
 	void OnVariableMappingChanged(const FName& PathName, const FName& Curve, bool bInput);
 	FName GetVariableMapping(const FName& PathName, bool bInput);
 	void GetAvailableMapping(const FName& PathName, TArray<FName>& OutArray, bool bInput);
-	void CreateVariableMapping(const FString& FilteredText, TArray< TSharedPtr<FVariableMappingInfo> >& OutArray, bool bInput);
+	void CreateVariableMapping(const FString& FilteredText, TArray< TSharedPtr<FRigVMVariableMappingInfo> >& OutArray, bool bInput);
+#endif
 
 	bool IsAvailableToMapToCurve(const FName& PropertyName, bool bInput) const;
 	bool IsInputProperty(const FName& PropertyName) const;
diff --git a/Engine/Plugins/Runtime/RigVM/Source/RigVMEditor/Private/Widgets/SRigVMVariableMappingWidget.cpp b/Engine/Plugins/Runtime/RigVM/Source/RigVMEditor/Private/Widgets/SRigVMVariableMappingWidget.cpp
new file mode 100644
--- /dev/null
+++ b/Engine/Plugins/Runtime/RigVM/Source/RigVMEditor/Private/Widgets/SRigVMVariableMappingWidget.cpp
@@ -0,0 +1,381 @@
+// Copyright Epic Games, Inc. All Rights Reserved.
+
+
+#include "Widgets/SRigVMVariableMappingWidget.h"
+#include "Widgets/Input/SCheckBox.h"
+#include "Widgets/Input/SButton.h"
+#include "Widgets/Text/SInlineEditableTextBlock.h"
+
+#define LOCTEXT_NAMESPACE "SRigVMVariableMappingWidget"
+
+static const FName ColumnId_VarLabel("Variable");
+static const FName ColumnID_MappingLabel("Mapping");
+
+void SRigVMVariableMappingTreeRow::Construct(const FArguments& InArgs, const TSharedRef<STableViewBase>& InOwnerTableView)
+{
+	Item = InArgs._Item;
+	OnVariableMappingChanged = InArgs._OnVariableMappingChanged;
+	OnGetVariableMapping = InArgs._OnGetVariableMapping;
+	OnGetAvailableMapping = InArgs._OnGetAvailableMapping;
+	OnGetFilteredText = InArgs._OnGetFilteredText;
+	OnVariableOptionAvailable = InArgs._OnVariableOptionAvailable;
+	OnPinCheckStateChanged = InArgs._OnPinCheckStateChanged;
+	OnPinGetCheckState = InArgs._OnPinGetCheckState;
+	OnPinIsEnabledCheckState = InArgs._OnPinIsEnabledCheckState;
+
+	check(Item.IsValid());
+
+	VariableOptionList.Reset();
+	if (OnGetAvailableMapping.IsBound())
+	{
+		TArray<FName> ListOfMappingVars;
+
+		OnGetAvailableMapping.Execute(Item->GetPathName(), ListOfMappingVars);
+		for (auto& MappingVar : ListOfMappingVars)
+		{
+			VariableOptionList.Add(MakeShareable(new FString(MappingVar.ToString())));
+		}
+	}
+	//SMultiColumnTableRow< FRigVMVariableMappingInfoPtr >::Construct(FSuperRowType::FArguments(), InOwnerTableView);
+	STableRow<TSharedPtr<FRigVMVariableMappingInfoPtr>>::Construct(
+		STableRow<TSharedPtr<FRigVMVariableMappingInfoPtr>>::FArguments()
+		.Padding(FMargin(3.0f, 2.0f))
+		.Content()
+		[
+			SNew(SHorizontalBox)
+
+			+ SHorizontalBox::Slot()
+			.FillWidth(0.5f)
+			.Padding(0.0f, 4.0f)
+			.HAlign(HAlign_Left)
+			[
+				SNew(STextBlock)
+				.Text(FText::FromString(Item->GetDisplayName()))
+				//.TextStyle(&FCoreStyle::Get().GetWidgetStyle<FTextBlockStyle>("BoldText"))
+				.HighlightText(this, &SRigVMVariableMappingTreeRow::GetFilterText)
+				//.IsReadOnly(true)
+				//.IsSelected(this, &SMultiColumnTableRow< FRigVMVariableMappingInfoPtr >::IsSelectedExclusively)
+			]
+// 			+ SHorizontalBox::Slot()
+// 			.Padding(2.0f, 4.0f)
+// 			.HAlign(HAlign_Right)
+// 			[
+// 				SNew(SHorizontalBox)
+// // 				+ SHorizontalBox::Slot()
+// // 				.FillWidth(0.5f)
+// // 				[
+// // 					SNew(STextBlock)
+// // 					.Text(LOCTEXT("EmptyString", ""))
+// // 				]
+				+ SHorizontalBox::Slot()
+				.Padding(0.0f, 4.0f)
+				.HAlign(HAlign_Right)
+				.AutoWidth()
+				[
+					SNew(STextBlock)
+					.Text(LOCTEXT("ExposeAsPinPropertyValue", "Use Pin "))
+				]
+				+ SHorizontalBox::Slot()
+				.Padding(0.0f, 4.0f)
+				.AutoWidth()
+				[
+					SNew(SCheckBox)
+					.IsChecked(this, &SRigVMVariableMappingTreeRow::IsPinChecked)
+					.OnCheckStateChanged(this, &SRigVMVariableMappingTreeRow::OnPinCheckStatusChanged)
+					.IsEnabled(this, &SRigVMVariableMappingTreeRow::IsPinEnabled)
+				]
+				+ SHorizontalBox::Slot()
+				.Padding(0.0f, 4.0f)
+				.HAlign(HAlign_Right)
+				.AutoWidth()
+				[
+					SNew(STextBlock)
+					.Text(LOCTEXT("Optional", " OR Use Curve "))
+				]
+				+ SHorizontalBox::Slot()
+				.FillWidth(0.5f)
+				[
+					SAssignNew(VarOptionComboBox, SSearchableComboBox)
+					.OptionsSource(&VariableOptionList)
+					.OnGenerateWidget(this, &SRigVMVariableMappingTreeRow::MakeVarOptionComboWidget)
+					.OnSelectionChanged(this, &SRigVMVariableMappingTreeRow::OnVarOptionSourceChanged)
+					.OnComboBoxOpening(this, &SRigVMVariableMappingTreeRow::OnVarOptionComboOpening)
+					.IsEnabled(this, &SRigVMVariableMappingTreeRow::IsVarOptionEnabled)
+					.ContentPadding(2)
+					.Content()
+					[
+						SNew(STextBlock)
+						.Text(this, &SRigVMVariableMappingTreeRow::GetVarOptionComboBoxContent)
+						//.Font(FAppStyle::GetFontStyle(TEXT("PropertyWindow.NormalFont")))
+						.ToolTipText(this, &SRigVMVariableMappingTreeRow::GetVarOptionComboBoxToolTip)
+					]
+				]
+				+ SHorizontalBox::Slot()
+				.AutoWidth()
+				[
+					SNew(SButton)
+					.OnClicked(FOnClicked::CreateSP(this, &SRigVMVariableMappingTreeRow::OnClearButtonClicked))
+					.Text(FText::FromString(TEXT("x")))
+				]
+
+			//]
+		],
+		InOwnerTableView);
+}
+
+bool SRigVMVariableMappingTreeRow::IsVarOptionEnabled() const
+{
+	// if pin is not checked
+	if (IsPinChecked()==ECheckBoxState::Unchecked && 
+		OnVariableOptionAvailable.IsBound())
+	{
+		return OnVariableOptionAvailable.Execute(Item->GetPathName());
+	}
+
+	return false;
+}
+
+ECheckBoxState SRigVMVariableMappingTreeRow::IsPinChecked() const
+{
+	if (OnPinGetCheckState.IsBound())
+	{
+		return OnPinGetCheckState.Execute(Item->GetPathName());
+	}
+
+	return ECheckBoxState::Unchecked;
+}
+
+void SRigVMVariableMappingTreeRow::OnPinCheckStatusChanged(ECheckBoxState NewState)
+{
+	if (OnPinCheckStateChanged.IsBound())
+	{
+		OnPinCheckStateChanged.Execute(NewState, Item->GetPathName());
+	}
+}
+
+bool SRigVMVariableMappingTreeRow::IsPinEnabled() const
+{
+	if (OnPinIsEnabledCheckState.IsBound())
+	{
+		return OnPinIsEnabledCheckState.Execute(Item->GetPathName());
+	}
+
+	return false;
+}
+
+
+FReply SRigVMVariableMappingTreeRow::OnClearButtonClicked()
+{
+	if (OnVariableMappingChanged.IsBound())
+	{
+		OnVariableMappingChanged.Execute(Item->GetPathName(), NAME_None);
+	}
+
+	return FReply::Handled();
+}
+
+FText SRigVMVariableMappingTreeRow::GetFilterText() const
+{
+	if (OnGetFilteredText.IsBound())
+	{
+		return OnGetFilteredText.Execute();
+	}
+
+	return FText::GetEmpty();
+}
+
+void SRigVMVariableMappingTreeRow::OnVarOptionComboOpening()
+{
+	VariableOptionList.Reset();
+	if (OnGetAvailableMapping.IsBound())
+	{
+		TArray<FName> ListOfMappingVars;
+
+		OnGetAvailableMapping.Execute(Item->GetPathName(), ListOfMappingVars);
+		for (auto& MappingVar : ListOfMappingVars)
+		{
+			VariableOptionList.Add(MakeShareable(new FString(MappingVar.ToString())));
+		}
+	}
+
+	if (OnGetVariableMapping.IsBound())
+	{
+		FName MappedVar = OnGetVariableMapping.Execute(Item->GetPathName());
+		if (MappedVar!= NAME_None)
+		{
+			TSharedPtr<FString> ComboStringPtr = GetVarOptionString(MappedVar);
+			if (ComboStringPtr.IsValid())
+			{
+				VarOptionComboBox->SetSelectedItem(ComboStringPtr);
+			}
+		}
+		else
+		{
+			VarOptionComboBox->ClearSelection();
+		}
+	}
+}
+
+void SRigVMVariableMappingTreeRow::OnVarOptionSourceChanged(TSharedPtr<FString> NewSelection, ESelectInfo::Type SelectInfo)
+{
+	// if it's set from code, we did that on purpose
+	if (SelectInfo != ESelectInfo::Direct)
+	{
+		FString NewValue = *NewSelection.Get();
+		FName NewName = FName(*NewValue);
+
+		OnVariableMappingChanged.ExecuteIfBound(Item->GetPathName(), NewName);
+	}
+}
+
+FText SRigVMVariableMappingTreeRow::GetVarOptionComboBoxContent() const
+{
+	if (OnGetVariableMapping.IsBound())
+	{
+		return FText::FromName(OnGetVariableMapping.Execute(Item->GetPathName()));
+	}
+
+	return FText::FromString(TEXT("Invalid"));
+}
+
+FText SRigVMVariableMappingTreeRow::GetVarOptionComboBoxToolTip() const
+{
+	return LOCTEXT("VarOptionComboToolTip", "Map input/output variable to available options.");
+}
+
+TSharedPtr<FString> SRigVMVariableMappingTreeRow::GetVarOptionString(FName VarOptionName) const
+{
+	return TSharedPtr<FString>();
+	FString VarOptionString = VarOptionName.ToString();
+
+	// go through profile and see if it has mine
+	for (int32 Index = 1; Index < VariableOptionList.Num(); ++Index)
+	{
+		if (VarOptionString == *VariableOptionList[Index])
+		{
+			return VariableOptionList[Index];
+		}
+	}
+
+	return TSharedPtr<FString>();
+}
+
+TSharedRef<SWidget> SRigVMVariableMappingTreeRow::MakeVarOptionComboWidget(TSharedPtr<FString> InItem)
+{
+	return SNew(STextBlock).Text(FText::FromString(*InItem));// .Font(FAppStyle::GetFontStyle(TEXT("PropertyWindow.NormalFont")));
+}
+//////////////////////////////////////////////////////////////////////////
+// SRigVMVariableMappingWidget
+
+void SRigVMVariableMappingWidget::Construct(const FArguments& InArgs/*, FSimpleMulticastDelegate& InOnPostUndo*/)
+{
+	OnGetAvailableMappingDelegate = InArgs._OnGetAvailableMapping;
+	OnGetVariableMappingDelegate = InArgs._OnGetVariableMapping;
+	OnVariableMappingChangedDelegate = InArgs._OnVariableMappingChanged;
+	OnCreateVariableMappingDelegate = InArgs._OnCreateVariableMapping;
+	OnVariableOptionAvailableDelegate = InArgs._OnVariableOptionAvailable;
+	OnPinCheckStateChangedDelegate = InArgs._OnPinCheckStateChanged;
+	OnPinGetCheckStateDelegate = InArgs._OnPinGetCheckState;
+	OnPinIsEnabledCheckStateDelegate = InArgs._OnPinIsEnabledCheckState;
+
+	//InOnPostUndo.Add(FSimpleDelegate::CreateSP(this, &SRigVMVariableMappingWidget::PostUndo));
+
+	ChildSlot
+		[
+			SNew(SVerticalBox)
+
+			// now show bone mapping
+			+ SVerticalBox::Slot()
+			.AutoHeight()
+			.Padding(0, 2)
+			[
+				SNew(SHorizontalBox)
+				// Filter entry
+				+ SHorizontalBox::Slot()
+				.FillWidth(1)
+				[
+					SAssignNew(NameFilterBox, SSearchBox)
+					.SelectAllTextWhenFocused(true)
+					.OnTextChanged(this, &SRigVMVariableMappingWidget::OnFilterTextChanged)
+					.OnTextCommitted(this, &SRigVMVariableMappingWidget::OnFilterTextCommitted)
+				]
+			]
+
+			+ SVerticalBox::Slot()
+			.FillHeight(1.0f)		// This is required to make the scrollbar work, as content overflows Slate containers by default
+			[
+				SAssignNew(VariableMappingTreeView, SRigVMVariableMappingTreeView)
+				.TreeItemsSource(&VariableMappingList)
+				.OnGenerateRow(this, &SRigVMVariableMappingWidget::MakeTreeRowWidget)
+				.OnGetChildren(this, &SRigVMVariableMappingWidget::GetChildrenForInfo)
+				.ItemHeight(22.0f)
+			]
+		];
+
+	RefreshVariableMappingList();
+}
+
+TSharedRef<ITableRow> SRigVMVariableMappingWidget::MakeTreeRowWidget(TSharedPtr<FRigVMVariableMappingInfo> InInfo, const TSharedRef<STableViewBase>& OwnerTable)
+{
+	return SNew(SRigVMVariableMappingTreeRow, OwnerTable)
+		.Item(InInfo)
+		.OnVariableMappingChanged(OnVariableMappingChangedDelegate)
+		.OnGetVariableMapping(OnGetVariableMappingDelegate)
+		.OnGetAvailableMapping(OnGetAvailableMappingDelegate)
+		.OnGetFilteredText(this, &SRigVMVariableMappingWidget::GetFilterText)
+		.OnVariableOptionAvailable(OnVariableOptionAvailableDelegate)
+		.OnPinCheckStateChanged(OnPinCheckStateChangedDelegate)
+		.OnPinGetCheckState(OnPinGetCheckStateDelegate)
+		.OnPinIsEnabledCheckState(OnPinIsEnabledCheckStateDelegate);
+}
+
+void SRigVMVariableMappingWidget::GetChildrenForInfo(TSharedPtr<FRigVMVariableMappingInfo> InInfo, TArray< TSharedPtr<FRigVMVariableMappingInfo> >& OutChildren)
+{
+	OutChildren = InInfo->Children;
+}
+
+void SRigVMVariableMappingWidget::OnFilterTextChanged(const FText& SearchText)
+{
+	// need to make sure not to have the same text go
+	// otherwise, the widget gets recreated multiple times causing 
+	// other issue
+	if (FilterText.CompareToCaseIgnored(SearchText) != 0)
+	{
+		FilterText = SearchText;
+		RefreshVariableMappingList();
+	}
+}
+
+void SRigVMVariableMappingWidget::OnFilterTextCommitted(const FText& SearchText, ETextCommit::Type CommitInfo)
+{
+	// Just do the same as if the user typed in the box
+	OnFilterTextChanged(SearchText);
+}
+
+TSharedRef<ITableRow> SRigVMVariableMappingWidget::GenerateVariableMappingRow(TSharedPtr<FRigVMVariableMappingInfo> InInfo, const TSharedRef<STableViewBase>& OwnerTable)
+{
+	check(InInfo.IsValid());
+
+	return
+		SNew(SRigVMVariableMappingTreeRow, OwnerTable)
+		.Item(InInfo)
+		.OnVariableMappingChanged(OnVariableMappingChangedDelegate)
+		.OnGetVariableMapping(OnGetVariableMappingDelegate)
+		.OnGetAvailableMapping(OnGetAvailableMappingDelegate)
+		.OnGetFilteredText(this, &SRigVMVariableMappingWidget::GetFilterText);
+}
+
+void SRigVMVariableMappingWidget::RefreshVariableMappingList()
+{
+	OnCreateVariableMappingDelegate.ExecuteIfBound(FilterText.ToString(), VariableMappingList);
+
+	VariableMappingTreeView->RequestListRefresh();
+}
+
+// void SRigVMVariableMappingWidget::PostUndo()
+// {
+// 	RefreshVariableMappingList();
+// }
+
+#undef LOCTEXT_NAMESPACE
+
diff --git a/Engine/Plugins/Runtime/RigVM/Source/RigVMEditor/Public/Widgets/SRigVMVariableMappingWidget.h b/Engine/Plugins/Runtime/RigVM/Source/RigVMEditor/Public/Widgets/SRigVMVariableMappingWidget.h
new file mode 100644
--- /dev/null
+++ b/Engine/Plugins/Runtime/RigVM/Source/RigVMEditor/Public/Widgets/SRigVMVariableMappingWidget.h
@@ -0,0 +1,235 @@
+// Copyright Epic Games, Inc. All Rights Reserved.
+
+#pragma once
+
+#include "CoreMinimal.h"
+#include "SSearchableComboBox.h"
+#include "Widgets/Input/SSearchBox.h"
+#include "Widgets/Input/SButton.h"
+#include "Widgets/Views/STreeView.h"
+#include "Widgets/Input/SSearchBox.h"
+
+//=====================================================================================
+// UI for displaying property info
+//=====================================================================================
+
+struct RIGVMEDITOR_API FRigVMVariableMappingInfo
+{
+public:
+	// This is the property that is the most shallow type
+	// It will be Transform.Translation.X
+	FName PropertyName; 
+	// Display Name
+	FString DisplayName;
+	// List of Children
+	// in theory, this actually shouldn't be active if you have children
+	// but it represent each raw nonetheless
+	// this is to map curve which is always float
+	TArray<TSharedPtr<FRigVMVariableMappingInfo>> Children;
+
+	/** Static function for creating a new item, but ensures that you can only have a TSharedRef to one */
+	static TSharedRef<FRigVMVariableMappingInfo> Make(const FName InPathName)
+	{
+		return MakeShareable(new FRigVMVariableMappingInfo(InPathName));
+	}
+
+	FName GetPathName() const
+	{
+		return PropertyName;
+	}
+
+	FString GetDisplayName() const
+	{
+		return DisplayName;
+	}
+
+protected:
+	/** Hidden constructor, always use Make above */
+	FRigVMVariableMappingInfo(const FName InPathName)
+		: PropertyName(InPathName)
+	{
+		FString PathString = InPathName.ToString();
+		int32 Found = PathString.Find(TEXT("."), ESearchCase::IgnoreCase, ESearchDir::FromEnd);
+		DisplayName = (Found != INDEX_NONE) ? PathString.RightChop(Found + 1) : PathString;
+	}
+
+	/** Hidden constructor, always use Make above */
+	FRigVMVariableMappingInfo() {}
+};
+
+//////////////////////////////////////////////////////////////////////////
+// SRigVMVariableMappingListRow
+
+typedef TSharedPtr< FRigVMVariableMappingInfo > FRigVMVariableMappingInfoPtr;
+
+DECLARE_DELEGATE_TwoParams(FOnRigVMVariableMappingChanged, const FName& /** PathName */, const FName&/** Curve Name**/)
+DECLARE_DELEGATE_RetVal_OneParam(FName, FOnRigVMGetVariableMapping, const FName& /** PathName **/)
+// type verified, and make sure it's not double booked
+DECLARE_DELEGATE_TwoParams(FOnRigVMGetAvailableMapping, const FName& /** PathName **/, TArray<FName>& /** List of available mappings**/)
+// to highlight item
+DECLARE_DELEGATE_RetVal(FText&, FOnRigVMGetFilteredText)
+// caller to create list
+DECLARE_DELEGATE_TwoParams(FOnRigVMCreateVariableMapping, const FString&, TArray< TSharedPtr<FRigVMVariableMappingInfo> >&)
+DECLARE_DELEGATE_RetVal_OneParam(bool, FOnRigVMVarOptionAvailable, const FName& /** Property Name */);
+
+// pin exposure option
+DECLARE_DELEGATE_TwoParams(FOnRigVMPinCheckStateChanged, ECheckBoxState, FName /** Property Name */);
+DECLARE_DELEGATE_RetVal_OneParam(ECheckBoxState, FOnRigVMPinGetCheckState, FName /** Property Name */);
+DECLARE_DELEGATE_RetVal_OneParam(bool, FOnRigVMPinIsCheckEnabled, FName /** Property Name */);
+
+class RIGVMEDITOR_API SRigVMVariableMappingTreeRow
+	: public STableRow< TSharedPtr<FRigVMVariableMappingInfoPtr> >
+{
+public:
+	
+	SLATE_BEGIN_ARGS(SRigVMVariableMappingTreeRow) {}
+
+	/** The item for this row **/
+	SLATE_ARGUMENT(FRigVMVariableMappingInfoPtr, Item)
+	SLATE_EVENT(FOnRigVMVariableMappingChanged, OnVariableMappingChanged)
+	SLATE_EVENT(FOnRigVMGetVariableMapping, OnGetVariableMapping)
+	SLATE_EVENT(FOnRigVMGetAvailableMapping, OnGetAvailableMapping)
+	SLATE_EVENT(FOnRigVMGetFilteredText, OnGetFilteredText)
+	SLATE_EVENT(FOnRigVMVarOptionAvailable, OnVariableOptionAvailable);
+	SLATE_EVENT(FOnRigVMPinCheckStateChanged, OnPinCheckStateChanged)
+	SLATE_EVENT(FOnRigVMPinGetCheckState, OnPinGetCheckState)
+	SLATE_EVENT(FOnRigVMPinIsCheckEnabled, OnPinIsEnabledCheckState)
+	SLATE_END_ARGS()
+
+	void Construct(const FArguments& InArgs, const TSharedRef<STableViewBase>& OwnerTableView);
+
+private:
+
+	/** Widget used to display the list of variable option list*/
+	TSharedPtr<SSearchableComboBox>	VarOptionComboBox;
+	TArray< TSharedPtr< FString > >	VariableOptionList;
+
+	/** The name and weight of the variable option*/
+	FRigVMVariableMappingInfoPtr	Item;
+
+	// Curve combo box options
+	FReply OnClearButtonClicked();
+	FText GetFilterText() const;
+
+	TSharedRef<SWidget> MakeVarOptionComboWidget(TSharedPtr<FString> InItem);
+	void OnVarOptionSourceChanged(TSharedPtr<FString> NewSelection, ESelectInfo::Type SelectInfo);
+	FText GetVarOptionComboBoxContent() const;
+	FText GetVarOptionComboBoxToolTip() const;
+	void OnVarOptionComboOpening();
+	bool IsVarOptionEnabled() const;
+	TSharedPtr<FString> GetVarOptionString(FName VarOptionName) const;
+
+	ECheckBoxState IsPinChecked() const;
+	void OnPinCheckStatusChanged(ECheckBoxState NewState);
+	bool IsPinEnabled() const;
+
+	FOnRigVMVariableMappingChanged	OnVariableMappingChanged;
+	FOnRigVMGetVariableMapping		OnGetVariableMapping;
+	FOnRigVMGetAvailableMapping		OnGetAvailableMapping;
+	FOnRigVMGetFilteredText			OnGetFilteredText;
+	FOnRigVMVarOptionAvailable		OnVariableOptionAvailable;
+	FOnRigVMPinCheckStateChanged		OnPinCheckStateChanged;
+	FOnRigVMPinGetCheckState			OnPinGetCheckState;
+	FOnRigVMPinIsCheckEnabled		OnPinIsEnabledCheckState;
+};
+
+typedef STreeView< TSharedPtr<FRigVMVariableMappingInfo> > SRigVMVariableMappingTreeView;
+
+class RIGVMEDITOR_API SRigVMVariableMappingWidget : public SCompoundWidget
+{
+public:
+	SLATE_BEGIN_ARGS(SRigVMVariableMappingWidget)
+	{}
+
+	SLATE_EVENT(FOnRigVMVariableMappingChanged, OnVariableMappingChanged)
+	SLATE_EVENT(FOnRigVMGetVariableMapping, OnGetVariableMapping)
+	SLATE_EVENT(FOnRigVMGetAvailableMapping, OnGetAvailableMapping)
+	SLATE_EVENT(FOnRigVMCreateVariableMapping, OnCreateVariableMapping)
+	SLATE_EVENT(FOnRigVMGetFilteredText, OnGetFilteredText)
+	SLATE_EVENT(FOnRigVMVarOptionAvailable, OnVariableOptionAvailable);
+	SLATE_EVENT(FOnRigVMPinCheckStateChanged, OnPinCheckStateChanged)
+	SLATE_EVENT(FOnRigVMPinGetCheckState, OnPinGetCheckState)
+	SLATE_EVENT(FOnRigVMPinIsCheckEnabled, OnPinIsEnabledCheckState)
+	SLATE_END_ARGS()
+
+	/**
+	* Slate construction function
+	*
+	* @param InArgs - Arguments passed from Slate
+	*
+	*/
+	void Construct(const FArguments& InArgs/*, FSimpleMulticastDelegate& InOnPostUndo*/);
+
+	/**
+	* Filters the SListView when the user changes the search text box (NameFilterBox)
+	*
+	* @param SearchText - The text the user has typed
+	*
+	*/
+	void OnFilterTextChanged(const FText& SearchText);
+
+	/**
+	* Filters the SListView when the user hits enter or clears the search box
+	* Simply calls OnFilterTextChanged
+	*
+	* @param SearchText - The text the user has typed
+	* @param CommitInfo - Not used
+	*
+	*/
+	void OnFilterTextCommitted(const FText& SearchText, ETextCommit::Type CommitInfo);
+
+	/**
+	* Create a widget for an entry in the tree from an info
+	*
+	* @param InInfo - Shared pointer to the morph target we're generating a row for
+	* @param OwnerTable - The table that owns this row
+	*
+	* @return A new Slate widget, containing the UI for this row
+	*/
+	TSharedRef<ITableRow> GenerateVariableMappingRow(TSharedPtr<FRigVMVariableMappingInfo> InInfo, const TSharedRef<STableViewBase>& OwnerTable);
+
+	/**
+	* Handler for the delete of retarget source
+	*/
+	void RefreshVariableMappingList();
+
+private:
+
+	/**
+	* Accessor so our rows can grab the filtertext for highlighting
+	*
+	*/
+	FText& GetFilterText() { return FilterText; }
+
+	/** Box to filter to a specific morph target name */
+	TSharedPtr<SSearchBox>	NameFilterBox;
+
+	/** Widget used to display the list of retarget sources */
+	TSharedPtr<SRigVMVariableMappingTreeView> VariableMappingTreeView;
+
+	/** A list of variable mapping list. Used by the VariableMappingListView. */
+	TArray< TSharedPtr<FRigVMVariableMappingInfo> > VariableMappingList;
+
+	/** Current text typed into NameFilterBox */
+	FText FilterText;
+
+	/** Delegate for undo/redo transaction **/
+	void PostUndo();
+
+	// Make a single tree row widget
+	TSharedRef<ITableRow> MakeTreeRowWidget(TSharedPtr<FRigVMVariableMappingInfo> InInfo, const TSharedRef<STableViewBase>& OwnerTable);
+
+	// Get the children for the provided bone info
+	void GetChildrenForInfo(TSharedPtr<FRigVMVariableMappingInfo> InInfo, TArray< TSharedPtr<FRigVMVariableMappingInfo> >& OutChildren);
+
+	FOnRigVMGetAvailableMapping		OnGetAvailableMappingDelegate;
+	FOnRigVMVariableMappingChanged	OnVariableMappingChangedDelegate;
+	FOnRigVMGetVariableMapping		OnGetVariableMappingDelegate;
+	FOnRigVMCreateVariableMapping	OnCreateVariableMappingDelegate;
+	FOnRigVMVarOptionAvailable		OnVariableOptionAvailableDelegate;
+	FOnRigVMPinCheckStateChanged		OnPinCheckStateChangedDelegate;
+	FOnRigVMPinGetCheckState			OnPinGetCheckStateDelegate;
+	FOnRigVMPinIsCheckEnabled		OnPinIsEnabledCheckStateDelegate;
+
+};
+
diff --git a/Engine/Plugins/Runtime/RigVM/Source/RigVMEditor/RigVMEditor.Build.cs b/Engine/Plugins/Runtime/RigVM/Source/RigVMEditor/RigVMEditor.Build.cs
--- a/Engine/Plugins/Runtime/RigVM/Source/RigVMEditor/RigVMEditor.Build.cs
+++ b/Engine/Plugins/Runtime/RigVM/Source/RigVMEditor/RigVMEditor.Build.cs
@@ -17,6 +17,14 @@ public class RigVMEditor : ModuleRules
                 "BlueprintGraph",
                 "GraphEditor",
                 "UnrealEd",
+                // "EditorFramework",
+                //"Kismet",
+                //"PropertyEditor",
+                //"ApplicationCore",
+                "Slate",
+                "SlateCore",
+                "InputCore",
+                "ToolWidgets",
             }
         );
     }
