RigVM: Revamp the use of slice indices

diff --git a/Engine/Plugins/Runtime/RigVM/Source/RigVM/Private/RigVMCore/RigVMMemoryStorage.cpp b/Engine/Plugins/Runtime/RigVM/Source/RigVM/Private/RigVMCore/RigVMMemoryStorage.cpp
--- a/Engine/Plugins/Runtime/RigVM/Source/RigVM/Private/RigVMCore/RigVMMemoryStorage.cpp
+++ b/Engine/Plugins/Runtime/RigVM/Source/RigVM/Private/RigVMCore/RigVMMemoryStorage.cpp
@@ -124,7 +124,7 @@ const uint8* TRigVMLazyValueBase::GetData(const FRigVMExecuteContext& Context) c
 			check(Context.ExtendedExecuteContext != nullptr);
 			MemoryHandle->ComputeLazyValueIfNecessary(*Context.ExtendedExecuteContext, SliceIndex);
 		}
-		return MemoryHandle->GetData(bFollowPropertyPath, SliceIndex);
+		return MemoryHandle->GetData(bFollowPropertyPath, INDEX_NONE);
 	}
 	return nullptr;
 }
diff --git a/Engine/Plugins/Runtime/RigVM/Source/RigVM/Public/RigVMCore/RigVMExecuteContext.h b/Engine/Plugins/Runtime/RigVM/Source/RigVM/Public/RigVMCore/RigVMExecuteContext.h
--- a/Engine/Plugins/Runtime/RigVM/Source/RigVM/Public/RigVMCore/RigVMExecuteContext.h
+++ b/Engine/Plugins/Runtime/RigVM/Source/RigVM/Public/RigVMCore/RigVMExecuteContext.h
@@ -47,6 +47,8 @@ public:
 	: LowerBound(0)
 	, UpperBound(0)
 	, Index(INDEX_NONE)
+	, Offset(0)
+	, InstructionIndex(INDEX_NONE)
 	{
 		Reset();
 	}
@@ -55,14 +57,8 @@ public:
 		: LowerBound(0)
 		, UpperBound(InCount - 1)
 		, Index(INDEX_NONE)
-	{
-		Reset();
-	}
-
-	FRigVMSlice(int32 InCount, const FRigVMSlice& InParent)
-		: LowerBound(InParent.GetIndex() * InCount)
-		, UpperBound((InParent.GetIndex() + 1) * InCount - 1)
-		, Index(INDEX_NONE)
+		, Offset(0)
+		, InstructionIndex(INDEX_NONE)
 	{
 		Reset();
 	}
@@ -79,12 +75,7 @@ public:
 
 	int32 GetIndex() const
 	{
-		return Index;
-	}
-
-	void SetIndex(int32 InIndex)
-	{
-		Index = InIndex;
+		return Index + Offset;
 	}
 
 	int32 GetRelativeIndex() const
@@ -102,6 +93,26 @@ public:
 		return float(GetRelativeIndex()) / float(FMath::Max<int32>(1, Num() - 1));
 	}
 
+	int32 GetOffset() const
+	{
+		return Offset;
+	}
+
+	void SetOffset(int32 InSliceOffset)
+	{
+		Offset = InSliceOffset;
+	}
+
+	int32 GetInstructionIndex() const
+	{
+		return InstructionIndex;
+	}
+
+	void SetInstructionIndex(int32 InInstructionIndex)
+	{
+		InstructionIndex = InInstructionIndex;
+	}
+
 	int32 Num() const
 	{
 		return 1 + UpperBound - LowerBound;
@@ -109,7 +120,7 @@ public:
 
 	int32 TotalNum() const
 	{
-		return UpperBound + 1;
+		return UpperBound + 1 + Offset;
 	}
 
 	operator bool() const
@@ -173,6 +184,8 @@ private:
 	int32 LowerBound;
 	int32 UpperBound;
 	int32 Index;
+	int32 Offset;
+	int32 InstructionIndex;
 };
 
 USTRUCT()
@@ -576,26 +589,35 @@ struct RIGVM_API FRigVMExtendedExecuteContext
 
 	const FRigVMSlice& GetSlice() const
 	{
-		const int32 SliceOffset = (int32)SliceOffsets[GetPublicData<>().InstructionIndex];
-		if (SliceOffset == 0)
-		{
-			return Slices.Last();
-		}
-		const int32 UpperBound = Slices.Num() - 1;
-		return Slices[FMath::Clamp<int32>(UpperBound - SliceOffset, 0, UpperBound)];
+		return Slices.Last();
 	}
 
 	void BeginSlice(int32 InCount, int32 InRelativeIndex = 0)
 	{
 		ensure(!IsSliceComplete());
-		Slices.Add(FRigVMSlice(InCount, Slices.Last()));
+		Slices.Add(FRigVMSlice(InCount));
 		Slices.Last().SetRelativeIndex(InRelativeIndex);
+		Slices.Last().SetInstructionIndex(GetPublicData<>().InstructionIndex);
+		Slices.Last().SetOffset(SliceOffsets[GetPublicData<>().InstructionIndex]);
 	}
 
 	void EndSlice()
 	{
 		ensure(Slices.Num() > 1);
-		Slices.Pop();
+
+		// if this slice has reached its upperbound
+		// we want to make sure to increment the offset
+		// for the instruction the slice originated from,
+		// so that next time around the slice indices are not
+		// reused.
+		const FRigVMSlice PoppedSlice = Slices.Pop();
+		if(PoppedSlice.GetRelativeIndex() == PoppedSlice.Num() - 1)
+		{
+			if(SliceOffsets.IsValidIndex(PoppedSlice.GetInstructionIndex()))
+			{
+				SliceOffsets[PoppedSlice.GetInstructionIndex()] += PoppedSlice.Num();
+			}
+		}
 	}
 
 	void IncrementSlice()
