Control Rig / RigVM: Moving graph related widgets to RigVMEditor

diff --git a/Engine/Plugins/Runtime/RigVM/Source/RigVMEditor/Public/Widgets/SRigVMGraphPinNameListValueWidget.h b/Engine/Plugins/Runtime/RigVM/Source/RigVMEditor/Public/Widgets/SRigVMGraphPinNameListValueWidget.h
new file mode 100644
--- /dev/null
+++ b/Engine/Plugins/Runtime/RigVM/Source/RigVMEditor/Public/Widgets/SRigVMGraphPinNameListValueWidget.h
@@ -0,0 +1,154 @@
+// Copyright Epic Games, Inc. All Rights Reservekd.
+
+#pragma once
+
+#include "CoreMinimal.h"
+#include "Input/Reply.h"
+#include "Widgets/SWidget.h"
+#include "Widgets/Input/SComboButton.h"
+#include "Widgets/Input/SEditableTextBox.h"
+#include "Widgets/Views/SListView.h"
+#include "Widgets/Input/SComboBox.h"
+
+/**
+ * A searchable text combo box
+ */
+class RIGVMEDITOR_API SRigVMGraphPinNameListValueWidget : public SComboButton
+{
+public:
+	/** Type of list used for showing menu options. */
+	typedef SListView< TSharedPtr<FString> > SComboListType;
+	/** Delegate type used to generate widgets that represent Options */
+	typedef typename TSlateDelegates< TSharedPtr<FString> >::FOnGenerateWidget FOnGenerateWidget;
+	typedef typename TSlateDelegates< TSharedPtr<FString> >::FOnSelectionChanged FOnSelectionChanged;
+
+	SLATE_BEGIN_ARGS(SRigVMGraphPinNameListValueWidget)
+		: _Content()
+		, _ContentPadding(FMargin(3.0, 3.0))
+		, _OptionsSource()
+		, _OnSelectionChanged()
+		, _OnGenerateWidget()
+		, _InitiallySelectedItem(nullptr)
+		, _Method()
+		, _MaxListHeight(450.0f)
+		, _HasDownArrow(true)
+		, _SearchHintText(NSLOCTEXT("SRigVMGraphPinNameListValueWidget", "Search", "Search"))
+		, _AllowUserProvidedText(false)
+	{}
+
+	/** Slot for this button's content (optional) */
+	SLATE_DEFAULT_SLOT(FArguments, Content)
+
+		SLATE_ATTRIBUTE(FMargin, ContentPadding)
+
+		SLATE_ARGUMENT(const TArray< TSharedPtr<FString> >*, OptionsSource)
+		SLATE_EVENT(FOnSelectionChanged, OnSelectionChanged)
+		SLATE_EVENT(FOnGenerateWidget, OnGenerateWidget)
+
+		/** Called when combo box is opened, before list is actually created */
+		SLATE_EVENT(FOnComboBoxOpening, OnComboBoxOpening)
+
+		/** The custom scrollbar to use in the ListView */
+		SLATE_ARGUMENT(TSharedPtr<SScrollBar>, CustomScrollbar)
+
+		/** The option that should be selected when the combo box is first created */
+		SLATE_ARGUMENT(TSharedPtr<FString>, InitiallySelectedItem)
+
+		SLATE_ARGUMENT(TOptional<EPopupMethod>, Method)
+
+		/** The max height of the combo box menu */
+		SLATE_ARGUMENT(float, MaxListHeight)
+
+		/**
+		 * When false, the down arrow is not generated and it is up to the API consumer
+		 * to make their own visual hint that this is a drop down.
+		 */
+		SLATE_ARGUMENT(bool, HasDownArrow)
+
+		/*
+		 * The visible text in the search / editable text field
+		 */
+		SLATE_ARGUMENT(FText, SearchHintText)
+
+		/*
+		 * If set to true the user is allowed to enter custom text here
+		 */
+		SLATE_ARGUMENT(bool, AllowUserProvidedText)
+
+	SLATE_END_ARGS()
+
+	/**
+	 * Construct the widget from a declaration
+	 *
+	 * @param InArgs   Declaration from which to construct the combo box
+	 */
+	void Construct(const FArguments& InArgs);
+
+	void ClearSelection();
+
+	void SetSelectedItem(TSharedPtr<FString> InSelectedItem);
+	void SetOptionsSource(const TArray< TSharedPtr<FString> >* InOptionsSource);
+
+	/** @return the item currently selected by the combo box. */
+	TSharedPtr<FString> GetSelectedItem();
+
+	/**
+	 * Requests a list refresh after updating options
+	 * Call SetSelectedItem to update the selected item if required
+	 * @see SetSelectedItem
+	 */
+	void RefreshOptions();
+
+private:
+
+	/** Generate a row for the InItem in the combo box's list (passed in as OwnerTable). Do this by calling the user-specified OnGenerateWidget */
+	TSharedRef<ITableRow> GenerateMenuItemRow(TSharedPtr<FString> InItem, const TSharedRef<STableViewBase>& OwnerTable);
+
+	//** Called if the menu is closed
+	void OnMenuOpenChanged(bool bOpen);
+
+	/** Invoked when the selection in the list changes */
+	void OnSelectionChanged_Internal(TSharedPtr<FString> ProposedSelection, ESelectInfo::Type SelectInfo, bool bForce = false);
+
+	/** Invoked when the search text changes */
+	void OnSearchTextChanged(const FText& ChangedText);
+
+	/** Invoked when the search is committed*/
+	void OnSearchTextCommitted(const FText& ChangedText, ETextCommit::Type CommitType);
+
+	/** Special case handling for search box key commands */
+	FReply OnSearchTextKeyDown(const FGeometry& MyGeometry, const FKeyEvent& InKeyEvent);
+
+	/** Special case handling for combo list key commands */
+	FReply OnComboListKeyDown(const FGeometry& MyGeometry, const FKeyEvent& InKeyEvent);
+
+	/** Handle clicking on the content menu */
+	virtual FReply OnButtonClicked() override;
+
+	/** The item style to use. */
+	const FTableRowStyle* ItemStyle;
+
+private:
+	/** Delegate that is invoked when the selected item in the combo box changes */
+	FOnSelectionChanged OnSelectionChanged;
+	/** The item currently selected in the combo box */
+	TSharedPtr<FString> SelectedItem;
+	/** The search field used for the combox box's contents */
+	TSharedPtr< SEditableTextBox > SearchField;
+	/** The ListView that we pop up; visualized the available options. */
+	TSharedPtr< SComboListType > ComboListView;
+	/** The Scrollbar used in the ListView. */
+	TSharedPtr< SScrollBar > CustomScrollbar;
+	/** Delegate to invoke before the combo box is opening. */
+	FOnComboBoxOpening OnComboBoxOpening;
+	/** Delegate to invoke when we need to visualize an option as a widget. */
+	FOnGenerateWidget OnGenerateWidget;
+
+	const TArray< TSharedPtr<FString> >* OptionsSource;
+	bool AllowUserProvidedText;
+
+	/** Used to focus the name box immediately following construction */
+	EActiveTimerReturnType SetFocusPostConstruct(double InCurrentTime, float InDeltaTime);
+};
+
+
diff --git a/Engine/Plugins/Runtime/RigVM/Source/RigVMEditor/Public/Widgets/SRigVMGraphPinUserDataNameSpace.h b/Engine/Plugins/Runtime/RigVM/Source/RigVMEditor/Public/Widgets/SRigVMGraphPinUserDataNameSpace.h
new file mode 100644
--- /dev/null
+++ b/Engine/Plugins/Runtime/RigVM/Source/RigVMEditor/Public/Widgets/SRigVMGraphPinUserDataNameSpace.h
@@ -0,0 +1,35 @@
+// Copyright Epic Games, Inc. All Rights Reserved.
+
+#pragma once
+
+#include "CoreMinimal.h"
+#include "Widgets/DeclarativeSyntaxSupport.h"
+#include "Widgets/SWidget.h"
+#include "SGraphPin.h"
+#include "Widgets/SRigVMGraphPinEditableNameValueWidget.h"
+
+class RIGVMEDITOR_API SRigVMGraphPinUserDataNameSpace : public SGraphPin
+{
+public:
+	SLATE_BEGIN_ARGS(SRigVMGraphPinUserDataNameSpace) {}
+	SLATE_END_ARGS()
+
+	void Construct(const FArguments& InArgs, UEdGraphPin* InGraphPinObj);
+
+protected:
+	//~ Begin SGraphPin Interface
+	virtual TSharedRef<SWidget>	GetDefaultValueWidget() override;
+	//~ End SGraphPin Interface
+
+	FText GetNameSpaceText() const;
+	virtual void SetNameSpaceText(const FText& NewTypeInValue, ETextCommit::Type CommitInfo);
+
+	TSharedRef<SWidget> MakeNameSpaceItemWidget(TSharedPtr<FString> InItem);
+	void OnNameSpaceChanged(TSharedPtr<FString> NewSelection, ESelectInfo::Type SelectInfo);
+	void OnNameSpaceComboBox();
+	TArray<TSharedPtr<FString>>& GetNameSpaces();
+
+	TArray<TSharedPtr<FString>> NameSpaces;
+	TSharedPtr<SRigVMGraphPinEditableNameValueWidget> NameComboBox;
+
+};
diff --git a/Engine/Plugins/Runtime/RigVM/Source/RigVMEditor/Public/Widgets/SRigVMGraphPinUserDataPath.h b/Engine/Plugins/Runtime/RigVM/Source/RigVMEditor/Public/Widgets/SRigVMGraphPinUserDataPath.h
new file mode 100644
--- /dev/null
+++ b/Engine/Plugins/Runtime/RigVM/Source/RigVMEditor/Public/Widgets/SRigVMGraphPinUserDataPath.h
@@ -0,0 +1,73 @@
+// Copyright Epic Games, Inc. All Rights Reserved.
+
+#pragma once
+
+#include "CoreMinimal.h"
+#include "Widgets/DeclarativeSyntaxSupport.h"
+#include "Widgets/SWidget.h"
+#include "Widgets/Input/SMenuAnchor.h"
+#include "Widgets/Input/SButton.h"
+#include "SGraphPin.h"
+#include "RigVMModel/RigVMPin.h"
+#include "RigVMBlueprint.h"
+#include "RigVMCore/RigVMAssetUserData.h"
+#include "IPropertyAccessEditor.h"
+
+class RIGVMEDITOR_API SRigVMUserDataPath : public SCompoundWidget
+{
+public:
+
+	SLATE_BEGIN_ARGS(SRigVMUserDataPath)
+	: _ModelPins()
+	{}
+
+	SLATE_ARGUMENT(TArray<URigVMPin*>, ModelPins)
+
+	SLATE_END_ARGS()
+
+	void Construct(const FArguments& InArgs);
+
+protected:
+
+	FString GetUserDataPath(URigVMPin* ModelPin) const;
+	FString GetUserDataPath() const;
+	FText GetUserDataPathText() const { return FText::FromString(GetUserDataPath()); }
+	const FSlateBrush* GetUserDataIcon() const;
+	FLinearColor GetUserDataColor() const;
+	const FSlateBrush* GetUserDataIcon(const UNameSpacedUserData::FUserData* InUserData) const;
+	FLinearColor GetUserDataColor(const UNameSpacedUserData::FUserData* InUserData) const;
+	TSharedRef<SWidget> GetTopLevelMenuContent();
+	void FillUserDataPathMenu( FMenuBuilder& InMenuBuilder, FString InParentPath );
+	void HandleSetUserDataPath(FString InUserDataPath);
+
+	URigVMBlueprint* GetBlueprint() const;
+	FString GetUserDataNameSpace(URigVMPin* ModelPin) const;
+	FString GetUserDataNameSpace() const;
+	const UNameSpacedUserData* GetUserDataObject() const;
+	const UNameSpacedUserData::FUserData* GetUserData() const;
+
+	TArray<URigVMPin*> ModelPins;
+	TSharedPtr<SMenuAnchor> MenuAnchor;
+	bool bAllowUObjects = false;
+};
+
+class RIGVMEDITOR_API  SRigVMGraphPinUserDataPath : public SGraphPin
+{
+public:
+
+	SLATE_BEGIN_ARGS(SRigVMGraphPinUserDataPath){}
+
+		SLATE_ARGUMENT(TArray<URigVMPin*>, ModelPins)
+
+	SLATE_END_ARGS()
+
+	void Construct(const FArguments& InArgs, UEdGraphPin* InGraphPinObj);
+
+protected:
+
+	//~ Begin SGraphPin Interface
+	virtual TSharedRef<SWidget>	GetDefaultValueWidget() override;
+	//~ End SGraphPin Interface
+
+	TArray<URigVMPin*> ModelPins;
+};
diff --git a/Engine/Plugins/Runtime/RigVM/Source/RigVMEditor/Public/Widgets/SRigVMGraphPinVariableBinding.h b/Engine/Plugins/Runtime/RigVM/Source/RigVMEditor/Public/Widgets/SRigVMGraphPinVariableBinding.h
new file mode 100644
--- /dev/null
+++ b/Engine/Plugins/Runtime/RigVM/Source/RigVMEditor/Public/Widgets/SRigVMGraphPinVariableBinding.h
@@ -0,0 +1,79 @@
+// Copyright Epic Games, Inc. All Rights Reserved.
+
+#pragma once
+
+#include "CoreMinimal.h"
+#include "Widgets/DeclarativeSyntaxSupport.h"
+#include "Widgets/SWidget.h"
+#include "SGraphPin.h"
+#include "RigVMModel/RigVMPin.h"
+#include "RigVMModel/Nodes/RigVMFunctionReferenceNode.h"
+#include "IPropertyAccessEditor.h"
+#include "RigVMBlueprint.h"
+
+class RIGVMEDITOR_API SRigVMGraphVariableBinding : public SCompoundWidget
+{
+public:
+
+	SLATE_BEGIN_ARGS(SRigVMGraphVariableBinding)
+	: _ModelPins()
+    , _FunctionReferenceNode(nullptr)
+    , _InnerVariableName(NAME_None)
+    , _Blueprint(nullptr)
+	, _CanRemoveBinding(true)
+	{}
+
+		SLATE_ARGUMENT(TArray<URigVMPin*>, ModelPins)
+		SLATE_ARGUMENT(URigVMFunctionReferenceNode*, FunctionReferenceNode)
+		SLATE_ARGUMENT(FName, InnerVariableName)
+		SLATE_ARGUMENT(URigVMBlueprint*, Blueprint)
+		SLATE_ARGUMENT(bool, CanRemoveBinding)
+
+	SLATE_END_ARGS()
+
+	void Construct(const FArguments& InArgs);
+
+protected:
+
+	FText GetBindingText(URigVMPin* ModelPin) const;
+	FText GetBindingText() const;
+	const FSlateBrush* GetBindingImage() const;
+	FLinearColor GetBindingColor() const;
+	bool OnCanBindProperty(FProperty* InProperty) const;
+	bool OnCanBindToClass(UClass* InClass) const;
+	void OnAddBinding(FName InPropertyName, const TArray<FBindingChainElement>& InBindingChain);
+	bool OnCanRemoveBinding(FName InPropertyName);
+	void OnRemoveBinding(FName InPropertyName);
+	void FillLocalVariableMenu( FMenuBuilder& MenuBuilder );
+	void HandleBindToLocalVariable(FRigVMGraphVariableDescription InLocalVariable);
+
+	TArray<URigVMPin*> ModelPins;
+	URigVMFunctionReferenceNode* FunctionReferenceNode;
+	FName InnerVariableName;
+	URigVMBlueprint* Blueprint;
+	FPropertyBindingWidgetArgs BindingArgs;
+	bool bCanRemoveBinding;
+};
+
+class RIGVMEDITOR_API SRigVMGraphPinVariableBinding : public SGraphPin
+{
+public:
+
+	SLATE_BEGIN_ARGS(SRigVMGraphPinVariableBinding){}
+
+		SLATE_ARGUMENT(TArray<URigVMPin*>, ModelPins)
+		SLATE_ARGUMENT(URigVMBlueprint*, Blueprint)
+
+	SLATE_END_ARGS()
+
+	void Construct(const FArguments& InArgs, UEdGraphPin* InGraphPinObj);
+
+protected:
+
+	//~ Begin SGraphPin Interface
+	virtual TSharedRef<SWidget>	GetDefaultValueWidget() override;
+	//~ End SGraphPin Interface
+
+	TArray<URigVMPin*> ModelPins;
+	URigVMBlueprint* Blueprint;
+};
diff --git a/Engine/Plugins/Runtime/RigVM/Source/RigVMEditor/Public/Widgets/SRigVMGraphPinVariableName.h b/Engine/Plugins/Runtime/RigVM/Source/RigVMEditor/Public/Widgets/SRigVMGraphPinVariableName.h
new file mode 100644
--- /dev/null
+++ b/Engine/Plugins/Runtime/RigVM/Source/RigVMEditor/Public/Widgets/SRigVMGraphPinVariableName.h
@@ -0,0 +1,35 @@
+// Copyright Epic Games, Inc. All Rights Reserved.
+
+#pragma once
+
+#include "CoreMinimal.h"
+#include "Widgets/DeclarativeSyntaxSupport.h"
+#include "Widgets/SWidget.h"
+#include "SGraphPin.h"
+#include "Widgets/SRigVMGraphPinEditableNameValueWidget.h"
+
+class RIGVMEDITOR_API SRigVMGraphPinVariableName : public SGraphPin
+{
+public:
+	SLATE_BEGIN_ARGS(SRigVMGraphPinVariableName) {}
+	SLATE_END_ARGS()
+
+		void Construct(const FArguments& InArgs, UEdGraphPin* InGraphPinObj);
+
+protected:
+	//~ Begin SGraphPin Interface
+	virtual TSharedRef<SWidget>	GetDefaultValueWidget() override;
+	//~ End SGraphPin Interface
+
+	FText GetVariableNameText() const;
+	virtual void SetVariableNameText(const FText& NewTypeInValue, ETextCommit::Type CommitInfo);
+
+	TSharedRef<SWidget> MakeVariableNameItemWidget(TSharedPtr<FString> InItem);
+	void OnVariableNameChanged(TSharedPtr<FString> NewSelection, ESelectInfo::Type SelectInfo);
+	void OnVariableNameComboBox();
+	TArray<TSharedPtr<FString>>& GetVariableNames();
+
+	TArray<TSharedPtr<FString>> VariableNames;
+	TSharedPtr<SRigVMGraphPinEditableNameValueWidget> NameComboBox;
+
+};
diff --git a/Engine/Plugins/Runtime/RigVM/Source/RigVMEditor/RigVMEditor.Build.cs b/Engine/Plugins/Runtime/RigVM/Source/RigVMEditor/RigVMEditor.Build.cs
--- a/Engine/Plugins/Runtime/RigVM/Source/RigVMEditor/RigVMEditor.Build.cs
+++ b/Engine/Plugins/Runtime/RigVM/Source/RigVMEditor/RigVMEditor.Build.cs
@@ -17,14 +17,11 @@ public class RigVMEditor : ModuleRules
                 "BlueprintGraph",
                 "GraphEditor",
                 "UnrealEd",
-                // "EditorFramework",
-                //"Kismet",
-                //"PropertyEditor",
-                //"ApplicationCore",
                 "Slate",
                 "SlateCore",
                 "InputCore",
                 "ToolWidgets",
+                "KismetWidgets",
             }
         );
     }
