ControlRig / RigVM: Moving remaining generic widgets to RigVMEditor

diff --git a/Engine/Plugins/Animation/ControlRig/Source/ControlRigEditor/Private/ControlRigEditorModule.cpp b/Engine/Plugins/Animation/ControlRig/Source/ControlRigEditor/Private/ControlRigEditorModule.cpp
--- a/Engine/Plugins/Animation/ControlRig/Source/ControlRigEditor/Private/ControlRigEditorModule.cpp
+++ b/Engine/Plugins/Animation/ControlRig/Source/ControlRigEditor/Private/ControlRigEditorModule.cpp
@@ -61,7 +61,7 @@
 #include "RigVMBlueprintUtils.h"
 #include "ControlRigBlueprintCommands.h"
 #include "ControlRigHierarchyCommands.h"
-#include "ControlRigStackCommands.h"
+#include "Editor/RigVMExecutionStackCommands.h"
 #include "Animation/AnimSequence.h"
 #include "Editor/ControlRigEditorEditMode.h"
 #include "ControlRigElementDetails.h"
@@ -146,7 +146,7 @@ void FControlRigEditorModule::StartupModule()
 	FControlRigEditModeCommands::Register();
 	FControlRigBlueprintCommands::Register();
 	FControlRigHierarchyCommands::Register();
-	FControlRigStackCommands::Register();
+	FRigVMExecutionStackCommands::Register();
 	FControlRigEditorStyle::Get();
 
 	MenuExtensibilityManager = MakeShareable(new FExtensibilityManager);
@@ -2707,7 +2707,7 @@ void FControlRigEditorModule::PreChange(const UUserDefinedStruct* Changed,
 			{
 				return true;
 			}
-			else if (UDetailsViewWrapperObject::IsValidClass(InStructProperty->GetOwnerClass()))
+			else if (URigVMDetailsViewWrapperObject::IsValidClass(InStructProperty->GetOwnerClass()))
 			{
 				return true;
 			}
@@ -2722,9 +2722,9 @@ void FControlRigEditorModule::PreChange(const UUserDefinedStruct* Changed,
 				GeneratorClass->RefreshLinkedProperties();
 				GeneratorClass->RefreshPropertyPaths();	
 			}
-			else if (InStruct->IsChildOf(UDetailsViewWrapperObject::StaticClass()))
+			else if (InStruct->IsChildOf(URigVMDetailsViewWrapperObject::StaticClass()))
 			{
-				UDetailsViewWrapperObject::MarkOutdatedClass(Cast<UClass>(InStruct));
+				URigVMDetailsViewWrapperObject::MarkOutdatedClass(Cast<UClass>(InStruct));
 			}
 		});
 	
diff --git a/Engine/Plugins/Animation/ControlRig/Source/ControlRigEditor/Private/ControlRigElementDetails.cpp b/Engine/Plugins/Animation/ControlRig/Source/ControlRigEditor/Private/ControlRigElementDetails.cpp
--- a/Engine/Plugins/Animation/ControlRig/Source/ControlRigEditor/Private/ControlRigElementDetails.cpp
+++ b/Engine/Plugins/Animation/ControlRig/Source/ControlRigEditor/Private/ControlRigElementDetails.cpp
@@ -537,7 +537,7 @@ void FRigControlTransformChannelDetails::CustomizeHeader(TSharedRef<IPropertyHan
 	InStructPropertyHandle->GetOuterObjects(Objects);
 	for (UObject* Object : Objects)
 	{
-		if (const UDetailsViewWrapperObject* WrapperObject = Cast<UDetailsViewWrapperObject>(Object))
+		if (const URigVMDetailsViewWrapperObject* WrapperObject = Cast<URigVMDetailsViewWrapperObject>(Object))
 		{
 			if(WrapperObject->GetWrappedStruct() == FRigControlElement::StaticStruct())
 			{
@@ -672,7 +672,7 @@ void FRigBaseElementDetails::CustomizeDetails(IDetailLayoutBuilder& DetailBuilde
 	DetailBuilder.GetObjectsBeingCustomized(DetailObjects);
 	for(TWeakObjectPtr<UObject> DetailObject : DetailObjects)
 	{
-		UDetailsViewWrapperObject* WrapperObject = CastChecked<UDetailsViewWrapperObject>(DetailObject.Get());
+		URigVMDetailsViewWrapperObject* WrapperObject = CastChecked<URigVMDetailsViewWrapperObject>(DetailObject.Get());
 
 		const FRigElementKey Key = WrapperObject->GetContent<FRigBaseElement>().GetKey();
 
@@ -1121,9 +1121,10 @@ bool FRigBaseElementDetails::ContainsElementByPredicate(const TFunction<bool(con
 
 void FRigBaseElementDetails::RegisterSectionMappings(FPropertyEditorModule& PropertyEditorModule)
 {
-	FRigBoneElementDetails().RegisterSectionMappings(PropertyEditorModule, UDetailsViewWrapperObject::GetClassForStruct(FRigBoneElement::StaticStruct()));
-	FRigNullElementDetails().RegisterSectionMappings(PropertyEditorModule, UDetailsViewWrapperObject::GetClassForStruct(FRigNullElement::StaticStruct()));
-	FRigControlElementDetails().RegisterSectionMappings(PropertyEditorModule, UDetailsViewWrapperObject::GetClassForStruct(FRigControlElement::StaticStruct()));
+	const URigVMDetailsViewWrapperObject* CDOWrapper = CastChecked<URigVMDetailsViewWrapperObject>(UControlRigWrapperObject::StaticClass()->GetDefaultObject());
+	FRigBoneElementDetails().RegisterSectionMappings(PropertyEditorModule, CDOWrapper->GetClassForStruct(FRigBoneElement::StaticStruct()));
+	FRigNullElementDetails().RegisterSectionMappings(PropertyEditorModule, CDOWrapper->GetClassForStruct(FRigNullElement::StaticStruct()));
+	FRigControlElementDetails().RegisterSectionMappings(PropertyEditorModule, CDOWrapper->GetClassForStruct(FRigControlElement::StaticStruct()));
 }
 
 void FRigBaseElementDetails::RegisterSectionMappings(FPropertyEditorModule& PropertyEditorModule, UClass* InClass)
diff --git a/Engine/Plugins/Animation/ControlRig/Source/ControlRigEditor/Private/ControlRigElementDetails.h b/Engine/Plugins/Animation/ControlRig/Source/ControlRigEditor/Private/ControlRigElementDetails.h
--- a/Engine/Plugins/Animation/ControlRig/Source/ControlRigEditor/Private/ControlRigElementDetails.h
+++ b/Engine/Plugins/Animation/ControlRig/Source/ControlRigEditor/Private/ControlRigElementDetails.h
@@ -9,7 +9,7 @@
 #include "Rigs/RigHierarchy.h"
 #include "ControlRig.h"
 #include "ControlRigBlueprint.h"
-#include "Editor/DetailsViewWrapperObject.h"
+#include "Editor/ControlRigWrapperObject.h"
 #include "Widgets/SRigVMGraphPinNameListValueWidget.h"
 #include "Editor/SControlRigGizmoNameList.h"
 #include "Styling/SlateTypes.h"
@@ -219,7 +219,7 @@ public:
 			return GetElement<T>();
 		}
 
-		TWeakObjectPtr<UDetailsViewWrapperObject> WrapperObject;
+		TWeakObjectPtr<URigVMDetailsViewWrapperObject> WrapperObject;
 		FRigElementHandle Element;
 		FRigElementHandle DefaultElement;
 	};
diff --git a/Engine/Plugins/Animation/ControlRig/Source/ControlRigEditor/Private/ControlRigGraphDetails.cpp b/Engine/Plugins/Animation/ControlRig/Source/ControlRigEditor/Private/ControlRigGraphDetails.cpp
--- a/Engine/Plugins/Animation/ControlRig/Source/ControlRigEditor/Private/ControlRigGraphDetails.cpp
+++ b/Engine/Plugins/Animation/ControlRig/Source/ControlRigEditor/Private/ControlRigGraphDetails.cpp
@@ -1297,7 +1297,7 @@ void FControlRigWrappedNodeDetails::CustomizeDetails(IDetailLayoutBuilder& Detai
 
 	for(TWeakObjectPtr<UObject> DetailObject : DetailObjects)
 	{
-		UDetailsViewWrapperObject* WrapperObject = CastChecked<UDetailsViewWrapperObject>(DetailObject.Get());
+		URigVMDetailsViewWrapperObject* WrapperObject = CastChecked<URigVMDetailsViewWrapperObject>(DetailObject.Get());
 		if(BlueprintBeingCustomized == nullptr)
 		{
 			BlueprintBeingCustomized = WrapperObject->GetTypedOuter<UControlRigBlueprint>();
@@ -1513,7 +1513,7 @@ TSharedRef<SWidget> FControlRigWrappedNodeDetails::MakeNameListItemWidget(TShare
 FText FControlRigWrappedNodeDetails::GetNameListText(FNameProperty* InProperty) const
 {
 	FText FirstText;
-	for(TWeakObjectPtr<UDetailsViewWrapperObject> ObjectBeingCustomized : ObjectsBeingCustomized)
+	for(TWeakObjectPtr<URigVMDetailsViewWrapperObject> ObjectBeingCustomized : ObjectsBeingCustomized)
 	{
 		if (FName* Value = InProperty->ContainerPtrToValuePtr<FName>(ObjectBeingCustomized.Get()))
 		{
@@ -1598,7 +1598,7 @@ void FControlRigWrappedNodeDetails::CustomizeLiveValues(IDetailLayoutBuilder& De
 		return;
 	}
 
-	UDetailsViewWrapperObject* FirstWrapper = ObjectsBeingCustomized[0].Get();
+	URigVMDetailsViewWrapperObject* FirstWrapper = ObjectsBeingCustomized[0].Get();
 	URigVMNode* FirstNode = NodesBeingCustomized[0].Get();
 	if(FirstNode->GetTypedOuter<URigVMFunctionLibrary>())
 	{
diff --git a/Engine/Plugins/Animation/ControlRig/Source/ControlRigEditor/Private/ControlRigGraphDetails.h b/Engine/Plugins/Animation/ControlRig/Source/ControlRigEditor/Private/ControlRigGraphDetails.h
--- a/Engine/Plugins/Animation/ControlRig/Source/ControlRigEditor/Private/ControlRigGraphDetails.h
+++ b/Engine/Plugins/Animation/ControlRig/Source/ControlRigEditor/Private/ControlRigGraphDetails.h
@@ -12,7 +12,7 @@
 #include "SGraphPin.h"
 #include "Widgets/SRigVMGraphNode.h"
 #include "Widgets/Colors/SColorBlock.h"
-#include "Editor/DetailsViewWrapperObject.h"
+#include "Editor/RigVMDetailsViewWrapperObject.h"
 #include "IDetailPropertyExtensionHandler.h"
 #include "Graph/ControlRigGraphSchema.h"
 #include "SAdvancedTransformInputBox.h"
@@ -232,7 +232,7 @@ public:
 	void CustomizeLiveValues(IDetailLayoutBuilder& DetailLayout);
 
 	UControlRigBlueprint* BlueprintBeingCustomized;
-	TArray<TWeakObjectPtr<UDetailsViewWrapperObject>> ObjectsBeingCustomized;
+	TArray<TWeakObjectPtr<URigVMDetailsViewWrapperObject>> ObjectsBeingCustomized;
 	TArray<TWeakObjectPtr<URigVMNode>> NodesBeingCustomized;
 	TMap<FName, TSharedPtr<SRigVMGraphPinNameListValueWidget>> NameListWidgets;
 };
diff --git a/Engine/Plugins/Animation/ControlRig/Source/ControlRigEditor/Private/ControlRigLocalVariableDetails.cpp b/Engine/Plugins/Animation/ControlRig/Source/ControlRigEditor/Private/ControlRigLocalVariableDetails.cpp
--- a/Engine/Plugins/Animation/ControlRig/Source/ControlRigEditor/Private/ControlRigLocalVariableDetails.cpp
+++ b/Engine/Plugins/Animation/ControlRig/Source/ControlRigEditor/Private/ControlRigLocalVariableDetails.cpp
@@ -24,7 +24,7 @@ void FRigVMLocalVariableDetails::CustomizeDetails(IDetailLayoutBuilder& DetailBu
 	DetailBuilder.GetObjectsBeingCustomized(DetailObjects);
 	for(TWeakObjectPtr<UObject> DetailObject : DetailObjects)
 	{
-		UDetailsViewWrapperObject* WrapperObject = CastChecked<UDetailsViewWrapperObject>(DetailObject.Get());
+		URigVMDetailsViewWrapperObject* WrapperObject = CastChecked<URigVMDetailsViewWrapperObject>(DetailObject.Get());
 		ObjectsBeingCustomized.Add(WrapperObject);
 	}
 
diff --git a/Engine/Plugins/Animation/ControlRig/Source/ControlRigEditor/Private/ControlRigLocalVariableDetails.h b/Engine/Plugins/Animation/ControlRig/Source/ControlRigEditor/Private/ControlRigLocalVariableDetails.h
--- a/Engine/Plugins/Animation/ControlRig/Source/ControlRigEditor/Private/ControlRigLocalVariableDetails.h
+++ b/Engine/Plugins/Animation/ControlRig/Source/ControlRigEditor/Private/ControlRigLocalVariableDetails.h
@@ -7,7 +7,7 @@
 #include "IPropertyTypeCustomization.h"
 #include "ControlRig.h"
 #include "ControlRigBlueprint.h"
-#include "Editor/DetailsViewWrapperObject.h"
+#include "Editor/RigVMDetailsViewWrapperObject.h"
 #include "Graph/ControlRigGraph.h"
 #include "Widgets/SRigVMGraphPinNameListValueWidget.h"
 #include "Styling/SlateTypes.h"
@@ -41,7 +41,7 @@ private:
 	URigVMGraph* GraphBeingCustomized;
 	UControlRigBlueprint* BlueprintBeingCustomized;
 	FRigVMGraphVariableDescription VariableDescription;
-	TArray<TWeakObjectPtr<UDetailsViewWrapperObject>> ObjectsBeingCustomized;
+	TArray<TWeakObjectPtr<URigVMDetailsViewWrapperObject>> ObjectsBeingCustomized;
 
 	TSharedPtr<IPropertyHandle> NameHandle;
 	TSharedPtr<IPropertyHandle> TypeHandle;
diff --git a/Engine/Plugins/Animation/ControlRig/Source/ControlRigEditor/Private/ControlRigStackCommands.cpp b/Engine/Plugins/Animation/ControlRig/Source/ControlRigEditor/Private/ControlRigStackCommands.cpp
deleted file mode 100644
diff --git a/Engine/Plugins/Animation/ControlRig/Source/ControlRigEditor/Private/ControlRigStackCommands.h b/Engine/Plugins/Animation/ControlRig/Source/ControlRigEditor/Private/ControlRigStackCommands.h
deleted file mode 100644
diff --git a/Engine/Plugins/Animation/ControlRig/Source/ControlRigEditor/Private/EditMode/ControlRigEditModeToolkit.cpp b/Engine/Plugins/Animation/ControlRig/Source/ControlRigEditor/Private/EditMode/ControlRigEditModeToolkit.cpp
--- a/Engine/Plugins/Animation/ControlRig/Source/ControlRigEditor/Private/EditMode/ControlRigEditModeToolkit.cpp
+++ b/Engine/Plugins/Animation/ControlRig/Source/ControlRigEditor/Private/EditMode/ControlRigEditModeToolkit.cpp
@@ -16,7 +16,6 @@
 #include "EditMode/SControlRigTweenWidget.h"
 #include "EditMode/SControlRigSnapper.h"
 #include "Tools/SMotionTrailOptions.h"
-#include "Editor/SControlRigProfilingView.h"
 #include "Toolkits/AssetEditorModeUILayer.h"
 #include "Widgets/Docking/SDockTab.h"
 #include "EditMode/ControlRigEditModeSettings.h"
@@ -181,16 +180,6 @@ TSharedRef<SDockTab> SpawnMotionTrailTab(const FSpawnTabArgs& Args)
 		];
 }
 
-TSharedRef<SDockTab> SpawnRigProfiler(const FSpawnTabArgs& Args)
-{
-	return SNew(SDockTab)
-		.TabRole(ETabRole::NomadTab)
-		[
-			SNew(SControlRigProfilingView)
-		];
-}
-
-
 TSharedRef<SDockTab> SpawnOutlinerTab(const FSpawnTabArgs& Args, FControlRigEditMode* InEditorMode)
 {
 	return SNew(SDockTab)
@@ -398,13 +387,6 @@ void FControlRigEditModeToolkit::RequestModeUITabs()
 			.SetIcon(FSlateIcon(TEXT("ControlRigEditorStyle"), TEXT("ControlRig.EditableMotionTrails")));
 		ModeUILayerPtr->GetTabManager()->RegisterDefaultTabWindowSize(MotionTrailTabName, FVector2D(425, 575));
 
-		ModeUILayerPtr->GetTabManager()->UnregisterTabSpawner("HierarchicalProfiler");
-		ModeUILayerPtr->GetTabManager()->RegisterTabSpawner("HierarchicalProfiler", FOnSpawnTab::CreateStatic(&SpawnRigProfiler))
-			.SetDisplayName(LOCTEXT("HierarchicalProfilerTab", "Hierarchical Profiler"))
-			.SetTooltipText(LOCTEXT("HierarchicalProfilerTooltip", "Open the Hierarchical Profiler tab."))
-			.SetGroup(MenuGroup)
-			.SetIcon(FSlateIcon(TEXT("ControlRigEditorStyle"), TEXT("HierarchicalProfiler.TabIcon")));
-
 		ModeUILayer.Pin()->ToolkitHostShutdownUI().BindSP(this, &FControlRigEditModeToolkit::UnregisterAndRemoveFloatingTabs);
 
 
@@ -432,12 +414,6 @@ void FControlRigEditModeToolkit::UnregisterAndRemoveFloatingTabs()
 		if (ModeUILayer.IsValid())
 		{
 			TSharedPtr<FAssetEditorModeUILayer> ModeUILayerPtr = ModeUILayer.Pin();
-			TSharedPtr<SDockTab> ProfilerTab = ModeUILayerPtr->GetTabManager()->FindExistingLiveTab(FTabId("HierarchicalProfiler"));
-			if (ProfilerTab)
-			{
-				ProfilerTab->RequestCloseTab();
-			}
-			ModeUILayerPtr->GetTabManager()->UnregisterTabSpawner("HierarchicalProfiler");
 
 			TSharedPtr<SDockTab> MotionTrailTab = ModeUILayerPtr->GetTabManager()->FindExistingLiveTab(FTabId(MotionTrailTabName));
 			if (MotionTrailTab)
diff --git a/Engine/Plugins/Animation/ControlRig/Source/ControlRigEditor/Private/Editor/ControlRigEditor.cpp b/Engine/Plugins/Animation/ControlRig/Source/ControlRigEditor/Private/Editor/ControlRigEditor.cpp
--- a/Engine/Plugins/Animation/ControlRig/Source/ControlRigEditor/Private/Editor/ControlRigEditor.cpp
+++ b/Engine/Plugins/Animation/ControlRig/Source/ControlRigEditor/Private/Editor/ControlRigEditor.cpp
@@ -84,9 +84,9 @@
 #include "RigVMModel/Nodes/RigVMFunctionReturnNode.h"
 #include "BlueprintActionDatabaseRegistrar.h"
 #include "IMessageLogListing.h"
-#include "Editor/SControlRigFunctionLocalizationWidget.h"
-#include "Editor/SControlRigFunctionBulkEditWidget.h"
-#include "Editor/SControlRigBreakLinksWidget.h"
+#include "Widgets/SRigVMGraphFunctionLocalizationWidget.h"
+#include "Widgets/SRigVMGraphFunctionBulkEditWidget.h"
+#include "Widgets/SRigVMGraphBreakLinksWidget.h"
 #include "Widgets/SRigVMGraphChangePinType.h"
 #include "SGraphPanel.h"
 #include "Engine/StaticMesh.h"
@@ -217,7 +217,8 @@ const TArray<FName> FControlRigEditor::ConstructionEventQueue = {FRigUnit_Prepar
 const TArray<FName> FControlRigEditor::BackwardsAndForwardsSolveEventQueue = {FRigUnit_InverseExecution::EventName, FRigUnit_BeginExecution::EventName};
 
 FControlRigEditor::FControlRigEditor()
-	: ControlRig(nullptr)
+	: IControlRigEditor()
+	, PreviewInstance(nullptr)
 	, ActiveController(nullptr)
 	, bControlRigEditorInitialized(false)
 	, bIsSettingObjectBeingDebugged(false)
@@ -296,7 +297,12 @@ FControlRigEditor::~FControlRigEditor()
 
 UControlRigBlueprint* FControlRigEditor::GetControlRigBlueprint() const
 {
-	return Cast<UControlRigBlueprint>(GetBlueprintObj());
+	return Cast<UControlRigBlueprint>(GetRigVMBlueprint());
+}
+
+UControlRig* FControlRigEditor::GetControlRig() const
+{
+	return Cast<UControlRig>(GetRigVMHost());
 }
 
 URigHierarchy* FControlRigEditor::GetHierarchyBeingDebugged() const
@@ -721,7 +727,7 @@ void FControlRigEditor::BindCommands()
 
 void FControlRigEditor::ToggleExecuteGraph()
 {
-	if (ControlRig)
+	if (GetControlRig())
 	{
 		bExecutionControlRig = !bExecutionControlRig;
 
@@ -994,7 +1000,7 @@ TSharedRef<SWidget> FControlRigEditor::GenerateTestAssetRecordMenuContent()
 
 bool FControlRigEditor::RecordTestData(double InRecordingDuration)
 {
-	if(ControlRig == nullptr)
+	if(GetControlRig() == nullptr)
 	{
 		return false;
 	}
@@ -1013,7 +1019,7 @@ bool FControlRigEditor::RecordTestData(double InRecordingDuration)
 	
 	if(UControlRigTestData* TestData = TestDataStrongPtr.Get())
 	{
-		TestData->Record(ControlRig, InRecordingDuration);
+		TestData->Record(GetControlRig(), InRecordingDuration);
 	}
 	return true;
 }
@@ -1031,12 +1037,12 @@ void FControlRigEditor::ToggleTestData()
 			}
 			case EControlRigTestDataPlaybackMode::GroundTruth:
 			{
-				TestDataStrongPtr->SetupReplay(ControlRig, false);
+				TestDataStrongPtr->SetupReplay(GetControlRig(), false);
 				break;
 			}
 			default:
 			{
-				TestDataStrongPtr->SetupReplay(ControlRig, true);
+				TestDataStrongPtr->SetupReplay(GetControlRig(), true);
 				break;
 			}
 		}
@@ -1279,7 +1285,7 @@ void FControlRigEditor::FillToolbar(FToolBarBuilder& ToolbarBuilder)
 
 TArray<FName> FControlRigEditor::GetEventQueue() const
 {
-	if (ControlRig)
+	if (const UControlRig* ControlRig = GetControlRig())
 	{
 		return ControlRig->GetEventQueue();
 	}
@@ -1312,7 +1318,7 @@ void FControlRigEditor::SetEventQueue(TArray<FName> InEventQueue, bool bCompile)
 			{
 				RigBlueprint->RequestAutoVMRecompilation();
 			}
-			RigBlueprint->Validator->SetControlRig(ControlRig);
+			RigBlueprint->Validator->SetControlRig(GetControlRig());
 		}
 		
 		// need to clear selection before remove transient control
@@ -1321,13 +1327,13 @@ void FControlRigEditor::SetEventQueue(TArray<FName> InEventQueue, bool bCompile)
 		RigBlueprint->GetHierarchyController()->ClearSelection();
 		
 		// need to copy here since the removal changes the iterator
-		if (ControlRig)
+		if (GetControlRig())
 		{
 			RigBlueprint->ClearTransientControls();
 		}
 	}
 
-	if (ControlRig)
+	if (UControlRig* ControlRig = GetControlRig())
 	{
 		if (InEventQueue.Num() > 0)
 		{
@@ -1463,7 +1469,7 @@ void FControlRigEditor::SetExecutionMode(const EControlRigExecutionModeType InEx
 	GetControlRigBlueprint()->SetDebugMode(InExecutionMode == EControlRigExecutionModeType_Debug);
 	Compile();
 	
-	if (ControlRig)
+	if (UControlRig* ControlRig = GetControlRig())
 	{
 		ControlRig->SetIsInDebugMode(InExecutionMode == EControlRigExecutionModeType_Debug);
 	}
@@ -1509,7 +1515,7 @@ void FControlRigEditor::GetCustomDebugObjects(TArray<FCustomDebugObject>& DebugL
 		return;
 	}
 
-	if (ControlRig)
+	if (UControlRig* ControlRig = GetControlRig())
 	{
 		FCustomDebugObject DebugObject;
 		DebugObject.Object = ControlRig;
@@ -1565,7 +1571,7 @@ void FControlRigEditor::GetCustomDebugObjects(TArray<FCustomDebugObject>& DebugL
 			for (UObject* Instance : ArchetypeInstances)
 			{
 				UControlRig* InstanceControlRig = Cast<UControlRig>(Instance);
-				if (InstanceControlRig && InstanceControlRig != ControlRig)
+				if (InstanceControlRig && InstanceControlRig != GetControlRig())
 				{
 					if (InstanceControlRig->GetOuter() == nullptr)
 					{
@@ -1609,10 +1615,10 @@ void FControlRigEditor::HandleSetObjectBeingDebugged(UObject* InObject)
 	if (DebuggedControlRig == nullptr)
 	{
 		// fall back to our default control rig (which still can be nullptr)
-		if (ControlRig != nullptr && GetBlueprintObj() && !bIsSettingObjectBeingDebugged)
+		if (GetControlRig() != nullptr && GetBlueprintObj() && !bIsSettingObjectBeingDebugged)
 		{
 			TGuardValue<bool> GuardSettingObjectBeingDebugged(bIsSettingObjectBeingDebugged, true);
-			GetBlueprintObj()->SetObjectBeingDebugged(ControlRig);
+			GetBlueprintObj()->SetObjectBeingDebugged(GetControlRig());
 			return;
 		}
 	}
@@ -1644,7 +1650,7 @@ void FControlRigEditor::HandleSetObjectBeingDebugged(UObject* InObject)
 
 	if (DebuggedControlRig)
 	{
-		bool bIsExternalControlRig = DebuggedControlRig != ControlRig;
+		bool bIsExternalControlRig = DebuggedControlRig != GetControlRig();
 		bool bShouldExecute = (!bIsExternalControlRig) && bExecutionControlRig;
 		DebuggedControlRig->RigVMLog = &ControlRigLog;
 		GetControlRigBlueprint()->Hierarchy->HierarchyForSelectionPtr = DebuggedControlRig->DynamicHierarchy;
@@ -1714,7 +1720,7 @@ FString FControlRigEditor::GetCustomDebugObjectLabel(UObject* ObjectBeingDebugge
 		return FString();
 	}
 
-	if (ObjectBeingDebugged == ControlRig)
+	if (ObjectBeingDebugged == GetControlRig())
 	{
 		return TEXT("Control Rig Editor Preview");
 	}
@@ -1875,17 +1881,18 @@ void FControlRigEditor::SetDetailObjects(const TArray<UObject*>& InObjects, bool
 		if (URigVMNode* ModelNode = Cast<URigVMNode>(InObject))
 		{
 			// check if we know the dynamic class already
-			bool bClassCreated = UDetailsViewWrapperObject::GetClassForNodes(ModelNodes, false) == nullptr;
+			const URigVMDetailsViewWrapperObject* CDOWrapper = CastChecked<URigVMDetailsViewWrapperObject>(GetDetailWrapperClass()->GetDefaultObject());
+			bool bClassCreated = CDOWrapper->GetClassForNodes(ModelNodes, false) == nullptr;
 
 			// create the wrapper object
-			if(UDetailsViewWrapperObject* WrapperObject = UDetailsViewWrapperObject::MakeInstance(GetBlueprintObj(), ModelNodes, ModelNode))
+			if(URigVMDetailsViewWrapperObject* WrapperObject = URigVMDetailsViewWrapperObject::MakeInstance(GetDetailWrapperClass(), GetBlueprintObj(), ModelNodes, ModelNode))
 			{
 				WrapperObject->GetWrappedPropertyChangedChainEvent().AddSP(this, &FControlRigEditor::OnWrappedPropertyChangedChainEvent);
 				WrapperObject->AddToRoot();
 
 				if(bClassCreated)
 				{
-					UClass* WrapperClass = UDetailsViewWrapperObject::GetClassForNodes(ModelNodes, false);
+					UClass* WrapperClass = WrapperObject->GetClassForNodes(ModelNodes, false);
 					check(WrapperClass);
 
 					FPropertyEditorModule& EditModule = FModuleManager::Get().GetModuleChecked<FPropertyEditorModule>("PropertyEditor");
@@ -1906,9 +1913,9 @@ void FControlRigEditor::SetDetailObjects(const TArray<UObject*>& InObjects, bool
 
 	for(UObject* FilteredObject : FilteredObjects)
 	{
-		if(UDetailsViewWrapperObject* WrapperObject = Cast<UDetailsViewWrapperObject>(FilteredObject))
+		if(URigVMDetailsViewWrapperObject* WrapperObject = Cast<URigVMDetailsViewWrapperObject>(FilteredObject))
 		{
-			WrapperObjects.Add(TStrongObjectPtr<UDetailsViewWrapperObject>(WrapperObject));
+			WrapperObjects.Add(TStrongObjectPtr<URigVMDetailsViewWrapperObject>(WrapperObject));
 		}
 	}
 
@@ -1939,7 +1946,7 @@ void FControlRigEditor::SetDetailViewForRigElements()
 			continue;
 		}
 
-		UDetailsViewWrapperObject* WrapperObject = UDetailsViewWrapperObject::MakeInstance(GetBlueprintObj(), Element->GetElementStruct(), (uint8*)Element, HierarchyBeingDebugged);
+		URigVMDetailsViewWrapperObject* WrapperObject = URigVMDetailsViewWrapperObject::MakeInstance(GetDetailWrapperClass(), GetBlueprintObj(), Element->GetElementStruct(), (uint8*)Element, HierarchyBeingDebugged);
 		WrapperObject->GetWrappedPropertyChangedChainEvent().AddSP(this, &FControlRigEditor::OnWrappedPropertyChangedChainEvent);
 		WrapperObject->AddToRoot();
 
@@ -1997,7 +2004,7 @@ void FControlRigEditor::SetDetailViewForLocalVariable()
 	{
 		if (SelectedObject.IsValid())
 		{
-			if(UDetailsViewWrapperObject* WrapperObject = Cast<UDetailsViewWrapperObject>(SelectedObject.Get()))
+			if(URigVMDetailsViewWrapperObject* WrapperObject = Cast<URigVMDetailsViewWrapperObject>(SelectedObject.Get()))
 			{
 				VariableName = WrapperObject->GetContent<FRigVMGraphVariableDescription>().Name;
 				break;
@@ -2049,7 +2056,7 @@ bool FControlRigEditor::DetailViewShowsAnyRigUnit() const
 	{
 		if (SelectedObject.IsValid())
 		{
-			if(UDetailsViewWrapperObject* WrapperObject = Cast<UDetailsViewWrapperObject>(SelectedObject.Get()))
+			if(URigVMDetailsViewWrapperObject* WrapperObject = Cast<URigVMDetailsViewWrapperObject>(SelectedObject.Get()))
 			{
 				const FString Notation = WrapperObject->GetWrappedNodeNotation();
 				if(!Notation.IsEmpty())
@@ -2075,7 +2082,7 @@ bool FControlRigEditor::DetailViewShowsStruct(UScriptStruct* InStruct) const
 	{
 		if (SelectedObject.IsValid())
 		{
-			if(UDetailsViewWrapperObject* WrapperObject = Cast<UDetailsViewWrapperObject>(SelectedObject.Get()))
+			if(URigVMDetailsViewWrapperObject* WrapperObject = Cast<URigVMDetailsViewWrapperObject>(SelectedObject.Get()))
 			{
 				if(const UScriptStruct* WrappedStruct = WrapperObject->GetWrappedStruct())
 				{
@@ -2097,7 +2104,7 @@ bool FControlRigEditor::DetailViewShowsRigElement(FRigElementKey InKey) const
 	{
 		if (SelectedObject.IsValid())
 		{
-			if(UDetailsViewWrapperObject* WrapperObject = Cast<UDetailsViewWrapperObject>(SelectedObject.Get()))
+			if(URigVMDetailsViewWrapperObject* WrapperObject = Cast<URigVMDetailsViewWrapperObject>(SelectedObject.Get()))
 			{
 				if (const UScriptStruct* WrappedStruct = WrapperObject->GetWrappedStruct())
 				{
@@ -2136,11 +2143,11 @@ void FControlRigEditor::ClearDetailObject(bool bChangeUISelectionState)
 
 void FControlRigEditor::ClearDetailsViewWrapperObjects()
 {
-	for(const TStrongObjectPtr<UDetailsViewWrapperObject>& WrapperObjectPtr : WrapperObjects)
+	for(const TStrongObjectPtr<URigVMDetailsViewWrapperObject>& WrapperObjectPtr : WrapperObjects)
 	{
 		if(WrapperObjectPtr.IsValid())
 		{
-			UDetailsViewWrapperObject* WrapperObject = WrapperObjectPtr.Get();
+			URigVMDetailsViewWrapperObject* WrapperObject = WrapperObjectPtr.Get();
 			WrapperObject->RemoveFromRoot();
 			WrapperObject->Rename(nullptr, GetTransientPackage(), REN_ForceNoResetLoaders | REN_DoNotDirty | REN_DontCreateRedirectors | REN_NonTransactional);
 			WrapperObject->MarkAsGarbage();
@@ -2207,7 +2214,7 @@ void FControlRigEditor::Compile()
 
 		TArray< TWeakObjectPtr<UObject> > SelectedObjects = Inspector->GetSelectedObjects();
 
-		if (ControlRig)
+		if (UControlRig* ControlRig = GetControlRig())
 		{
 			ControlRig->OnInitialized_AnyThread().Clear();
 			ControlRig->OnExecuted_AnyThread().Clear();
@@ -2246,7 +2253,7 @@ void FControlRigEditor::Compile()
 			}
 		}
 		
-		if (ControlRig)
+		if (UControlRig* ControlRig = GetControlRig())
 		{
 			ControlRigLog.Reset();
 			ControlRig->RigVMLog = &ControlRigLog;
@@ -2330,7 +2337,7 @@ void FControlRigEditor::Compile()
 				const FTransform Transform = RigBlueprint->Hierarchy->GetInitialLocalTransform(ControlElement->GetIndex());
 
 				/*/
-				if (ControlRig)
+				if (UControlRig* ControlRig = GetControlRig())
 				{
 					ControlRig->Modify();
 					ControlRig->GetControlHierarchy().SetLocalTransform(Control.Index, Transform);
@@ -2366,9 +2373,9 @@ void FControlRigEditor::SaveAsset_Execute()
 
 	// Save the new state of the hierarchy in the default object, so that it has the correct values on load
 	UControlRigBlueprint* RigBlueprint = Cast<UControlRigBlueprint>(GetBlueprintObj());
-	if(ControlRig)
+	if(const UControlRig* ControlRig = GetControlRig())
 	{
-		UControlRig* CDO = ControlRig->GetClass()->GetDefaultObject<UControlRig>();
+		const UControlRig* CDO = ControlRig->GetClass()->GetDefaultObject<UControlRig>();
 		CDO->DynamicHierarchy->CopyHierarchy(RigBlueprint->Hierarchy);
 	}
 
@@ -2384,9 +2391,9 @@ void FControlRigEditor::SaveAssetAs_Execute()
 
 	// Save the new state of the hierarchy in the default object, so that it has the correct values on load
 	UControlRigBlueprint* RigBlueprint = Cast<UControlRigBlueprint>(GetBlueprintObj());
-	if(ControlRig)
+	if(const UControlRig* ControlRig = GetControlRig())
 	{
-		UControlRig* CDO = ControlRig->GetClass()->GetDefaultObject<UControlRig>();
+		const UControlRig* CDO = ControlRig->GetClass()->GetDefaultObject<UControlRig>();
 		CDO->DynamicHierarchy->CopyHierarchy(RigBlueprint->Hierarchy);
 	}
 
@@ -2805,7 +2812,7 @@ void FControlRigEditor::EnsureValidRigElementsInDetailPanel()
 	{
 		if (SelectedObject.IsValid())
 		{
-			if(UDetailsViewWrapperObject* WrapperObject = Cast<UDetailsViewWrapperObject>(SelectedObject.Get()))
+			if(URigVMDetailsViewWrapperObject* WrapperObject = Cast<URigVMDetailsViewWrapperObject>(SelectedObject.Get()))
 			{
 				if(const UScriptStruct* WrappedStruct = WrapperObject->GetWrappedStruct())
 				{
@@ -2899,12 +2906,12 @@ void FControlRigEditor::OnAnimInitialized()
 		EditorSkelComp->bRequiredBonesUpToDateDuringTick = 0;
 
 		UControlRigLayerInstance* AnimInstance = Cast<UControlRigLayerInstance>(EditorSkelComp->GetAnimInstance());
-		if (AnimInstance && ControlRig)
+		if (AnimInstance && GetControlRig())
 		{
 			// update control rig data to anim instance since animation system has been reinitialized
 			FInputBlendPose Filter;
 			AnimInstance->ResetControlRigTracks();
-			AnimInstance->AddControlRigTrack(0, ControlRig);
+			AnimInstance->AddControlRigTrack(0, GetControlRig());
 			AnimInstance->UpdateControlRigTrack(0, 1.0f, FControlRigIOSettings::MakeEnabled(), bExecutionControlRig);
 		}
 	}
@@ -2998,12 +3005,15 @@ FGraphAppearanceInfo FControlRigEditor::GetGraphAppearance(UEdGraph* InGraph) co
 	{
 		AppearanceInfo.CornerText = LOCTEXT("AppearanceCornerText_ControlRig", "RIG");
 
-		if(ControlRig && ControlRig->VM && ControlRig->VM->IsNativized())
+		if(const UControlRig* ControlRig = GetControlRig())
 		{
-			AppearanceInfo.InstructionFade = 1;
-			AppearanceInfo.InstructionText = FText::FromString(
-				FString::Printf(TEXT("This graph runs a nativized VM (U%s)."), *ControlRig->VM->GetNativizedClass()->GetName())
-			);
+			if(ControlRig->VM && ControlRig->VM->IsNativized())
+			{
+				AppearanceInfo.InstructionFade = 1;
+				AppearanceInfo.InstructionText = FText::FromString(
+					FString::Printf(TEXT("This graph runs a nativized VM (U%s)."), *ControlRig->VM->GetNativizedClass()->GetName())
+				);
+			}
 		}
 	}
 
@@ -3070,7 +3080,7 @@ void FControlRigEditor::HandleModifiedEvent(ERigVMGraphNotifType InNotifType, UR
 					{
 						if (SelectedObject.IsValid())
 						{
-							if(UDetailsViewWrapperObject* WrapperObject = Cast<UDetailsViewWrapperObject>(SelectedObject.Get()))
+							if(URigVMDetailsViewWrapperObject* WrapperObject = Cast<URigVMDetailsViewWrapperObject>(SelectedObject.Get()))
 							{
 								if(WrapperObject->GetOuter() == Pin->GetNode())
 								{
@@ -3193,7 +3203,7 @@ void FControlRigEditor::HandleControlRigExecutedEvent(URigVMHost* InControlRig,
 		UControlRig* DebuggedControlRig = Cast<UControlRig>(ControlRigBP->GetObjectBeingDebugged());
 		if (DebuggedControlRig == nullptr)
 		{
-			DebuggedControlRig = ControlRig;
+			DebuggedControlRig = GetControlRig();
 		}
 
 		URigHierarchy* Hierarchy = GetHierarchyBeingDebugged(); 
@@ -3203,7 +3213,7 @@ void FControlRigEditor::HandleControlRigExecutedEvent(URigVMHost* InControlRig,
 		{
 			if (SelectedObject.IsValid())
 			{
-				if(UDetailsViewWrapperObject* WrapperObject = Cast<UDetailsViewWrapperObject>(SelectedObject.Get()))
+				if(URigVMDetailsViewWrapperObject* WrapperObject = Cast<URigVMDetailsViewWrapperObject>(SelectedObject.Get()))
 				{
 					if (const UScriptStruct* Struct = WrapperObject->GetWrappedStruct())
 					{
@@ -3387,7 +3397,7 @@ void FControlRigEditor::HandleControlRigExecutionHalted(const int32 InstructionI
 	}
 	else 
 	{
-		if (InEntryName == ControlRig->GetEventQueue().Last())
+		if (InEntryName == GetControlRig()->GetEventQueue().Last())
 		{
 			SetHaltedNode(nullptr);
 		}
@@ -3421,6 +3431,7 @@ void FControlRigEditor::Tick(float DeltaTime)
 	// tick the control rig in case we don't have skeletal mesh
 	if (UControlRigBlueprint* Blueprint = GetControlRigBlueprint())
 	{
+		UControlRig* ControlRig = GetControlRig();
 		if (Blueprint->GetPreviewMesh() == nullptr && 
 			ControlRig != nullptr && 
 			bExecutionControlRig)
@@ -3501,6 +3512,7 @@ bool FControlRigEditor::IsEditable(UEdGraph* InGraph) const
 			}
 		}
 
+		const UControlRig* ControlRig = GetControlRig();
 		if(ControlRig && ControlRig->VM)
 		{
 			const bool bIsReadOnly = ControlRig->VM->IsNativized();
@@ -4047,7 +4059,7 @@ void FControlRigEditor::OnToolbarDrawAxesOnSelectionChanged(ECheckBoxState InNew
 
 bool FControlRigEditor::IsToolbarDrawNullsEnabled() const
 {
-	if (ControlRig)
+	if (const UControlRig* ControlRig = GetControlRig())
 	{
 		if (!ControlRig->IsConstructionModeEnabled())
 		{
@@ -4076,7 +4088,7 @@ void FControlRigEditor::OnToolbarDrawNullsChanged(ECheckBoxState InNewValue)
 
 bool FControlRigEditor::IsPinControlNameListEnabled() const
 {
-	if (ControlRig)
+	if (UControlRig* ControlRig = GetControlRig())
 	{
 		TArray<FRigControlElement*>	TransientControls = ControlRig->GetHierarchy()->GetTransientControls();
 		if (TransientControls.Num() > 0)
@@ -4098,10 +4110,10 @@ TSharedRef<SWidget> FControlRigEditor::MakePinControlNameListItemWidget(TSharedP
 
 FText FControlRigEditor::GetPinControlNameListText() const
 {
-	if (ControlRig)
+	if (UControlRig* ControlRig = GetControlRig())
 	{
 		FText Result;
-		ControlRig->GetHierarchy()->ForEach<FRigControlElement>([this, &Result](FRigControlElement* ControlElement) -> bool
+		ControlRig->GetHierarchy()->ForEach<FRigControlElement>([this, &Result, ControlRig](FRigControlElement* ControlElement) -> bool
         {
 			if (ControlElement->Settings.bIsTransientControl)
 			{
@@ -4136,9 +4148,9 @@ TSharedPtr<FString> FControlRigEditor::GetPinControlCurrentlySelectedItem(const
 
 void FControlRigEditor::SetPinControlNameListText(const FText& NewTypeInValue, ETextCommit::Type /*CommitInfo*/)
 {
-	if (ControlRig)
+	if (UControlRig* ControlRig = GetControlRig())
 	{
-		ControlRig->GetHierarchy()->ForEach<FRigControlElement>([this, NewTypeInValue](FRigControlElement* ControlElement) -> bool
+		ControlRig->GetHierarchy()->ForEach<FRigControlElement>([this, NewTypeInValue, ControlRig](FRigControlElement* ControlElement) -> bool
         {
             if (ControlElement->Settings.bIsTransientControl)
 			{
@@ -4283,7 +4295,7 @@ void FControlRigEditor::UpdateControlRig()
 
 		if (AnimInstance)
 		{
-			if (ControlRig)
+			if (const URigVMHost* CurrentHost = GetRigVMHost())
 			{
 				if(TestDataStrongPtr.IsValid())
 				{
@@ -4293,15 +4305,18 @@ void FControlRigEditor::UpdateControlRig()
 				// if this control rig is from a temporary step,
 				// for example the reinstancing class, clear it 
 				// and create a new one!
-				if (ControlRig->GetClass() != Class)
+				if (CurrentHost->GetClass() != Class)
 				{
-					ControlRig = nullptr;
+					SetHost(nullptr);
 				}
 			}
 
+			UControlRig* ControlRig = GetControlRig();
 			if (ControlRig == nullptr)
 			{
 				ControlRig = NewObject<UControlRig>(EditorSkelComp, Class);
+				SetHost(ControlRig);
+				
 				// this is editing time rig
 				ControlRig->ExecutionType = ERigExecutionType::Editing;
 				ControlRig->RigVMLog = &ControlRigLog;
@@ -4394,6 +4409,7 @@ void FControlRigEditor::AddReferencedObjects( FReferenceCollector& Collector )
 {
 	FBlueprintEditor::AddReferencedObjects(Collector);
 
+	TWeakObjectPtr<UControlRig> ControlRig(GetControlRig());
 	Collector.AddReferencedObject(ControlRig);
 }
 
@@ -4543,7 +4559,7 @@ FTransform FControlRigEditor::GetRigElementTransform(const FRigElementKey& InEle
 		UControlRig* DebuggedControlRig = Cast<UControlRig>(GetBlueprintObj()->GetObjectBeingDebugged());
 		if (DebuggedControlRig == nullptr)
 		{
-			DebuggedControlRig = ControlRig;
+			DebuggedControlRig = GetControlRig();
 		}
 
 		if (DebuggedControlRig)
@@ -4665,7 +4681,7 @@ void FControlRigEditor::NotifyPostChange(const FPropertyChangedEvent& PropertyCh
 	// changes on a lower level property like transform under a user defined struct
 	// only go through this.
 	UControlRigBlueprint* ControlRigBP = GetControlRigBlueprint();
-	if(ControlRig && ControlRigBP)
+	if(GetControlRig() && ControlRigBP)
 	{
 		bool bUseCDO = false; 
 		if(PropertyChangedEvent.GetNumObjectsBeingEdited() == 1)
@@ -4735,7 +4751,7 @@ void FControlRigEditor::OnPropertyChanged(UObject* InObject, FPropertyChangedEve
 	}
 }
 
-void FControlRigEditor::OnWrappedPropertyChangedChainEvent(UDetailsViewWrapperObject* InWrapperObject, const FString& InPropertyPath, FPropertyChangedChainEvent& InPropertyChangedChainEvent)
+void FControlRigEditor::OnWrappedPropertyChangedChainEvent(URigVMDetailsViewWrapperObject* InWrapperObject, const FString& InPropertyPath, FPropertyChangedChainEvent& InPropertyChangedChainEvent)
 {
 	check(InWrapperObject);
 	check(!WrapperObjects.IsEmpty());
@@ -4986,7 +5002,7 @@ void FControlRigEditor::OnRequestLocalizeFunctionDialog(FRigVMGraphFunctionIdent
 		{
 			if (bForce || bIsPublic)
 			{
-				TSharedRef<SControlRigFunctionLocalizationDialog> LocalizationDialog = SNew(SControlRigFunctionLocalizationDialog)
+				TSharedRef<SRigVMGraphFunctionLocalizationDialog> LocalizationDialog = SNew(SRigVMGraphFunctionLocalizationDialog)
 							.Function(InFunction)
 							.TargetBlueprint(InTargetBlueprint);
 
@@ -5016,7 +5032,7 @@ FRigVMController_BulkEditResult FControlRigEditor::OnRequestBulkEditDialog(URigV
 		return FRigVMController_BulkEditResult();
 	}
 	
-	TSharedRef<SControlRigFunctionBulkEditDialog> BulkEditDialog = SNew(SControlRigFunctionBulkEditDialog)
+	TSharedRef<SRigVMGraphFunctionBulkEditDialog> BulkEditDialog = SNew(SRigVMGraphFunctionBulkEditDialog)
 	.Blueprint(InBlueprint)
 	.Controller(InController)
 	.Function(InFunction)
@@ -5041,9 +5057,9 @@ bool FControlRigEditor::OnRequestBreakLinksDialog(TArray<URigVMLink*> InLinks)
 		return true;
 	}
 
-	TSharedRef<SControlRigBreakLinksDialog> BreakLinksDialog = SNew(SControlRigBreakLinksDialog)
+	TSharedRef<SRigVMGraphBreakLinksDialog> BreakLinksDialog = SNew(SRigVMGraphBreakLinksDialog)
 	.Links(InLinks)
-	.OnFocusOnLink(FControlRigOnFocusOnLinkRequestedDelegate::CreateLambda([&](URigVMLink* InLink)
+	.OnFocusOnLink(FRigVMOnFocusOnLinkRequestedDelegate::CreateLambda([&](URigVMLink* InLink)
 	{
 		HandleJumpToHyperlink(InLink);
 	}));
@@ -5259,8 +5275,8 @@ bool FControlRigEditor::SelectLocalVariable(const UEdGraph* Graph, const FName&
 			{
 				if (Variable.Name == VariableName)
 				{
-					UDetailsViewWrapperObject* WrapperObject = UDetailsViewWrapperObject::MakeInstance(
-						GetBlueprintObj(), Variable.StaticStruct(), (uint8*)&Variable, RigVMGraph);
+					URigVMDetailsViewWrapperObject* WrapperObject = URigVMDetailsViewWrapperObject::MakeInstance(
+						GetDetailWrapperClass(), GetBlueprintObj(), Variable.StaticStruct(), (uint8*)&Variable, RigVMGraph);
 					WrapperObject->GetWrappedPropertyChangedChainEvent().AddSP(this, &FControlRigEditor::OnWrappedPropertyChangedChainEvent);
 					WrapperObject->AddToRoot();
 
@@ -5682,7 +5698,7 @@ void FControlRigEditor::SynchronizeViewportBoneSelection()
 
 void FControlRigEditor::UpdateBoneModification(FName BoneName, const FTransform& LocalTransform)
 {
-	if (ControlRig)
+	if (UControlRig* ControlRig = GetControlRig())
 	{ 
 		if (PreviewInstance)
 		{ 
@@ -5710,7 +5726,7 @@ void FControlRigEditor::UpdateBoneModification(FName BoneName, const FTransform&
 
 void FControlRigEditor::RemoveBoneModification(FName BoneName)
 {
-	if (ControlRig)
+	if (UControlRig* ControlRig = GetControlRig())
 	{
 		if (PreviewInstance)
 		{
@@ -5729,7 +5745,7 @@ void FControlRigEditor::RemoveBoneModification(FName BoneName)
 
 void FControlRigEditor::ResetAllBoneModification()
 {
-	if (ControlRig)
+	if (UControlRig* ControlRig = GetControlRig())
 	{
 		if (IsValid(PreviewInstance))
 		{
@@ -6607,7 +6623,7 @@ void FControlRigEditor::HandleMakeElementGetterSetter(ERigElementGetterSetterTyp
 
 void FControlRigEditor::HandleOnControlModified(UControlRig* Subject, FRigControlElement* ControlElement, const FRigControlModifiedContext& Context)
 {
-	if (Subject != ControlRig)
+	if (Subject != GetControlRig())
 	{
 		return;
 	}
@@ -6669,7 +6685,7 @@ void FControlRigEditor::HandleOnControlModified(UControlRig* Subject, FRigContro
 			{
 				bool bRequiresRecompile = true;
 
-				if(ControlRig)
+				if(UControlRig* ControlRig = GetControlRig())
 				{
 					if(TSharedPtr<FRigVMParserAST> AST = Pin->GetGraph()->GetDiagnosticsAST())
 					{
@@ -6723,7 +6739,7 @@ void FControlRigEditor::HandleOnControlModified(UControlRig* Subject, FRigContro
 			}
 			else if (ElementKey.Type == ERigElementType::Null)
 			{
-				const FTransform GlobalTransform = ControlRig->GetControlGlobalTransform(ControlElement->GetName());
+				const FTransform GlobalTransform = GetControlRig()->GetControlGlobalTransform(ControlElement->GetName());
 				Blueprint->Hierarchy->SetGlobalTransform(ElementKey, GlobalTransform);
 				Hierarchy->SetGlobalTransform(ElementKey, GlobalTransform);
 				if (IsConstructionModeEnabled())
@@ -6917,7 +6933,7 @@ void FControlRigEditor::HandleShowCurrentStatement()
 
 void FControlRigEditor::HandleBreakpointActionRequested(const ERigVMBreakpointAction BreakpointAction)
 {
-	if (ControlRig)
+	if (UControlRig* ControlRig = GetControlRig())
 	{
 		ControlRig->ExecuteBreakpointAction(BreakpointAction);
 	}
@@ -6945,6 +6961,7 @@ void FControlRigEditor::UpdateGraphCompilerErrors()
 	DECLARE_SCOPE_HIERARCHICAL_COUNTER_FUNC()
 
 	UControlRigBlueprint* Blueprint = Cast<UControlRigBlueprint>(GetBlueprintObj());
+	UControlRig* ControlRig = GetControlRig();
 	if (Blueprint && ControlRig && ControlRig->VM)
 	{
 		if (ControlRigLog.Entries.Num() == 0 && !bAnyErrorsLeft)
diff --git a/Engine/Plugins/Animation/ControlRig/Source/ControlRigEditor/Private/Editor/ControlRigEditor.h b/Engine/Plugins/Animation/ControlRig/Source/ControlRigEditor/Private/Editor/ControlRigEditor.h
--- a/Engine/Plugins/Animation/ControlRig/Source/ControlRigEditor/Private/Editor/ControlRigEditor.h
+++ b/Engine/Plugins/Animation/ControlRig/Source/ControlRigEditor/Private/Editor/ControlRigEditor.h
@@ -19,7 +19,7 @@
 #include "ScopedTransaction.h"
 #include "Graph/ControlRigGraphNode.h"
 #include "RigVMModel/RigVMController.h"
-#include "Editor/DetailsViewWrapperObject.h"
+#include "Editor/RigVMDetailsViewWrapperObject.h"
 #include "ControlRigTestData.h"
 #include "RigVMHost.h"
 
@@ -150,6 +150,9 @@ public:
 	// Gets the Control Rig Blueprint being edited/viewed
 	UControlRigBlueprint* GetControlRigBlueprint() const;
 
+	// returns the hierarchy being debugged
+	UControlRig* GetControlRig() const;
+
 	// returns the hierarchy being debugged
 	URigHierarchy* GetHierarchyBeingDebugged() const;
 
@@ -182,7 +185,7 @@ public:
 	FControlRigEditorEditMode* GetEditMode() const;
 
 	// this changes everytime you compile, so don't cache it expecting it will last. 
-	UControlRig* GetInstanceRig() const { return ControlRig;  }
+	UControlRig* GetInstanceRig() const { return GetControlRig();  }
 
 	void OnCurveContainerChanged();
 
@@ -197,9 +200,6 @@ public:
 	FOnGetContextMenu& OnGetViewportContextMenu() { return OnGetViewportContextMenuDelegate; }
 	FNewMenuCommandsDelegate& OnViewportContextMenuCommands() { return OnViewportContextMenuCommandsDelegate; }
 
-	DECLARE_EVENT_OneParam(FControlRigEditor, FPreviewControlRigUpdated, FControlRigEditor*);
-	FPreviewControlRigUpdated& OnPreviewControlRigUpdated() { return PreviewControlRigUpdated;  }
-
 	virtual bool SelectLocalVariable(const UEdGraph* Graph, const FName& VariableName) override;
 
 private:
@@ -245,7 +245,7 @@ protected:
 	virtual void SetupGraphEditorEvents(UEdGraph* InGraph, SGraphEditor::FGraphEditorEvents& InEvents) override;
 	virtual void FocusInspectorOnGraphSelection(const TSet<class UObject*>& NewSelection, bool bForceRefresh = false) override;
 
-	void HandleModifiedEvent(ERigVMGraphNotifType InNotifType, URigVMGraph* InGraph, UObject* InSubject);
+	virtual void HandleModifiedEvent(ERigVMGraphNotifType InNotifType, URigVMGraph* InGraph, UObject* InSubject) override;
 	void HandleVMCompiledEvent(UObject* InCompiledObject, URigVM* InVM);
 	void HandleControlRigExecutedEvent(URigVMHost* InControlRig, const FName& InEventName);
 	void HandleControlRigExecutionHalted(const int32 InstructionIndex, UObject* InNode, const FName& InEntryName);
@@ -407,9 +407,6 @@ protected:
 	/** Preview instance inspector widget */
 	TSharedPtr<SWidget> PreviewEditor;
 
-	/** Our currently running control rig instance */
-	TObjectPtr<UControlRig> ControlRig;
-
 	/** preview scene */
 	TSharedPtr<IPersonaPreviewScene> PreviewScene;
 
@@ -435,7 +432,7 @@ protected:
 	/** delegate for changing property */
 	virtual void OnFinishedChangingProperties(const FPropertyChangedEvent& PropertyChangedEvent) override;
 	void OnPropertyChanged(UObject* InObject, FPropertyChangedEvent& InEvent);
-	void OnWrappedPropertyChangedChainEvent(UDetailsViewWrapperObject* InWrapperObject, const FString& InPropertyPath, FPropertyChangedChainEvent& InPropertyChangedChainEvent);
+	void OnWrappedPropertyChangedChainEvent(URigVMDetailsViewWrapperObject* InWrapperObject, const FString& InPropertyPath, FPropertyChangedChainEvent& InPropertyChangedChainEvent);
 	void OnRequestLocalizeFunctionDialog(FRigVMGraphFunctionIdentifier& InFunction, URigVMBlueprint* InTargetBlueprint, bool bForce);
 	FRigVMController_BulkEditResult OnRequestBulkEditDialog(URigVMBlueprint* InBlueprint, URigVMController* InController, URigVMLibraryNode* InFunction, ERigVMControllerBulkEditType InEditType);
 	bool OnRequestBreakLinksDialog(TArray<URigVMLink*> InLinks);
@@ -464,8 +461,6 @@ protected:
 	/** Are we currently compiling through the user interface */
 	bool bIsCompilingThroughUI;
 
-	FPreviewControlRigUpdated PreviewControlRigUpdated;
-
 	TSharedPtr<SRigVMGraphPinNameListValueWidget> PinControlNameList;
 	bool IsPinControlNameListEnabled() const;
 	TSharedRef<SWidget> MakePinControlNameListItemWidget(TSharedPtr<FString> InItem);
@@ -490,7 +485,7 @@ protected:
 
 	bool bAllowBulkEdits;
 
-	TArray<TStrongObjectPtr<UDetailsViewWrapperObject>> WrapperObjects;
+	TArray<TStrongObjectPtr<URigVMDetailsViewWrapperObject>> WrapperObjects;
 	TWeakObjectPtr<AStaticMeshActor> WeakGroundActorPtr;
 
 	FDelegateHandle PropertyChangedHandle;
diff --git a/Engine/Plugins/Animation/ControlRig/Source/ControlRigEditor/Private/Editor/ControlRigEditorMode.cpp b/Engine/Plugins/Animation/ControlRig/Source/ControlRigEditor/Private/Editor/ControlRigEditorMode.cpp
--- a/Engine/Plugins/Animation/ControlRig/Source/ControlRigEditor/Private/Editor/ControlRigEditorMode.cpp
+++ b/Engine/Plugins/Animation/ControlRig/Source/ControlRigEditor/Private/Editor/ControlRigEditorMode.cpp
@@ -7,9 +7,8 @@
 #include "IPersonaToolkit.h"
 #include "PersonaTabs.h"
 #include "Editor/RigHierarchyTabSummoner.h"
-#include "Editor/RigStackTabSummoner.h"
+#include "Editor/RigVMExecutionStackTabSummoner.h"
 #include "Editor/RigCurveContainerTabSummoner.h"
-#include "Editor/RigInfluenceMapTabSummoner.h"
 #include "Editor/RigValidationTabSummoner.h"
 #include "Editor/RigAnimAttributeTabSummoner.h"
 #include "ToolMenus.h"
@@ -21,9 +20,8 @@ FControlRigEditorMode::FControlRigEditorMode(const TSharedRef<FControlRigEditor>
 	ControlRigBlueprintPtr = CastChecked<UControlRigBlueprint>(InControlRigEditor->GetBlueprintObj());
 
 	TabFactories.RegisterFactory(MakeShared<FRigHierarchyTabSummoner>(InControlRigEditor));
-	TabFactories.RegisterFactory(MakeShared<FRigStackTabSummoner>(InControlRigEditor));
+	TabFactories.RegisterFactory(MakeShared<FRigVMExecutionStackTabSummoner>(InControlRigEditor));
 	TabFactories.RegisterFactory(MakeShared<FRigCurveContainerTabSummoner>(InControlRigEditor));
-	//TabFactories.RegisterFactory(MakeShared<FRigInfluenceMapTabSummoner>(InControlRigEditor));
 	TabFactories.RegisterFactory(MakeShared<FRigValidationTabSummoner>(InControlRigEditor));
 	TabFactories.RegisterFactory(MakeShared<FRigAnimAttributeTabSummoner>(InControlRigEditor));
 
@@ -71,7 +69,7 @@ FControlRigEditorMode::FControlRigEditorMode(const TSharedRef<FControlRigEditor>
 						FTabManager::NewStack()
 						->SetSizeCoefficient(0.5f)
 						->AddTab(FRigHierarchyTabSummoner::TabID, ETabState::OpenedTab)
-						->AddTab(FRigStackTabSummoner::TabID, ETabState::OpenedTab)
+						->AddTab(FRigVMExecutionStackTabSummoner::TabID, ETabState::OpenedTab)
 						->AddTab(FRigCurveContainerTabSummoner::TabID, ETabState::OpenedTab)
 						->AddTab(FBlueprintEditorTabs::MyBlueprintID, ETabState::OpenedTab)
 					)
diff --git a/Engine/Plugins/Animation/ControlRig/Source/ControlRigEditor/Private/Editor/ControlRigWrapperObject.cpp b/Engine/Plugins/Animation/ControlRig/Source/ControlRigEditor/Private/Editor/ControlRigWrapperObject.cpp
new file mode 100644
--- /dev/null
+++ b/Engine/Plugins/Animation/ControlRig/Source/ControlRigEditor/Private/Editor/ControlRigWrapperObject.cpp
@@ -0,0 +1,53 @@
+// Copyright Epic Games, Inc. All Rights Reserved.
+
+#include "Editor/ControlRigWrapperObject.h"
+
+#if WITH_EDITOR
+#include "ControlRigElementDetails.h"
+#include "ControlRigLocalVariableDetails.h"
+#include "PropertyEditorModule.h"
+#include "Modules/ModuleManager.h"
+#endif
+
+UClass* UControlRigWrapperObject::GetClassForStruct(UScriptStruct* InStruct, bool bCreateIfNeeded) const
+{
+	UClass* Class = Super::GetClassForStruct(InStruct, bCreateIfNeeded);
+	if(Class == nullptr)
+	{
+		return nullptr;
+	}
+	
+	const FName WrapperClassName = Class->GetFName();
+
+#if WITH_EDITOR
+	if(InStruct->IsChildOf(FRigBaseElement::StaticStruct()))
+	{
+		FPropertyEditorModule& PropertyEditorModule = FModuleManager::LoadModuleChecked<FPropertyEditorModule>("PropertyEditor");
+		if (!PropertyEditorModule.GetClassNameToDetailLayoutNameMap().Contains(WrapperClassName))
+		{
+			if(InStruct == FRigBoneElement::StaticStruct())
+			{
+				PropertyEditorModule.RegisterCustomClassLayout(WrapperClassName, FOnGetDetailCustomizationInstance::CreateStatic(&FRigBoneElementDetails::MakeInstance));
+			}
+			else if(InStruct == FRigNullElement::StaticStruct())
+			{
+				PropertyEditorModule.RegisterCustomClassLayout(WrapperClassName, FOnGetDetailCustomizationInstance::CreateStatic(&FRigNullElementDetails::MakeInstance));
+			}
+			else if(InStruct == FRigControlElement::StaticStruct())
+			{
+				PropertyEditorModule.RegisterCustomClassLayout(WrapperClassName, FOnGetDetailCustomizationInstance::CreateStatic(&FRigControlElementDetails::MakeInstance));
+			}
+		}
+	}
+	else if(InStruct->IsChildOf(FRigVMGraphVariableDescription::StaticStruct()))
+	{
+		FPropertyEditorModule& PropertyEditorModule = FModuleManager::LoadModuleChecked<FPropertyEditorModule>("PropertyEditor");
+		if (!PropertyEditorModule.GetClassNameToDetailLayoutNameMap().Contains(WrapperClassName))
+		{
+			PropertyEditorModule.RegisterCustomClassLayout(WrapperClassName, FOnGetDetailCustomizationInstance::CreateStatic(&FRigVMLocalVariableDetails::MakeInstance));
+		}
+	}
+#endif
+
+	return Class;
+}
diff --git a/Engine/Plugins/Animation/ControlRig/Source/ControlRigEditor/Private/Editor/ControlRigWrapperObject.h b/Engine/Plugins/Animation/ControlRig/Source/ControlRigEditor/Private/Editor/ControlRigWrapperObject.h
new file mode 100644
--- /dev/null
+++ b/Engine/Plugins/Animation/ControlRig/Source/ControlRigEditor/Private/Editor/ControlRigWrapperObject.h
@@ -0,0 +1,15 @@
+// Copyright Epic Games, Inc. All Rights Reserved.
+
+#pragma once
+
+#include "Editor/RigVMDetailsViewWrapperObject.h"
+#include "ControlRigWrapperObject.generated.h"
+
+UCLASS()
+class CONTROLRIGEDITOR_API UControlRigWrapperObject : public URigVMDetailsViewWrapperObject
+{
+public:
+	GENERATED_BODY()
+
+	virtual UClass* GetClassForStruct(UScriptStruct* InStruct, bool bCreateIfNeeded) const override;
+};
diff --git a/Engine/Plugins/Animation/ControlRig/Source/ControlRigEditor/Private/Editor/DetailsViewWrapperObject.cpp b/Engine/Plugins/Animation/ControlRig/Source/ControlRigEditor/Private/Editor/DetailsViewWrapperObject.cpp
deleted file mode 100644
diff --git a/Engine/Plugins/Animation/ControlRig/Source/ControlRigEditor/Private/Editor/DetailsViewWrapperObject.h b/Engine/Plugins/Animation/ControlRig/Source/ControlRigEditor/Private/Editor/DetailsViewWrapperObject.h
deleted file mode 100644
diff --git a/Engine/Plugins/Animation/ControlRig/Source/ControlRigEditor/Private/Editor/RigStackTabSummoner.cpp b/Engine/Plugins/Animation/ControlRig/Source/ControlRigEditor/Private/Editor/RigStackTabSummoner.cpp
deleted file mode 100644
diff --git a/Engine/Plugins/Animation/ControlRig/Source/ControlRigEditor/Private/Editor/RigStackTabSummoner.h b/Engine/Plugins/Animation/ControlRig/Source/ControlRigEditor/Private/Editor/RigStackTabSummoner.h
deleted file mode 100644
diff --git a/Engine/Plugins/Animation/ControlRig/Source/ControlRigEditor/Private/Editor/SControlRigBreakLinksWidget.cpp b/Engine/Plugins/Animation/ControlRig/Source/ControlRigEditor/Private/Editor/SControlRigBreakLinksWidget.cpp
deleted file mode 100644
diff --git a/Engine/Plugins/Animation/ControlRig/Source/ControlRigEditor/Private/Editor/SControlRigBreakLinksWidget.h b/Engine/Plugins/Animation/ControlRig/Source/ControlRigEditor/Private/Editor/SControlRigBreakLinksWidget.h
deleted file mode 100644
diff --git a/Engine/Plugins/Animation/ControlRig/Source/ControlRigEditor/Private/Editor/SControlRigFunctionBulkEditWidget.cpp b/Engine/Plugins/Animation/ControlRig/Source/ControlRigEditor/Private/Editor/SControlRigFunctionBulkEditWidget.cpp
deleted file mode 100644
diff --git a/Engine/Plugins/Animation/ControlRig/Source/ControlRigEditor/Private/Editor/SControlRigFunctionBulkEditWidget.h b/Engine/Plugins/Animation/ControlRig/Source/ControlRigEditor/Private/Editor/SControlRigFunctionBulkEditWidget.h
deleted file mode 100644
diff --git a/Engine/Plugins/Animation/ControlRig/Source/ControlRigEditor/Private/Editor/SControlRigFunctionLocalizationWidget.cpp b/Engine/Plugins/Animation/ControlRig/Source/ControlRigEditor/Private/Editor/SControlRigFunctionLocalizationWidget.cpp
deleted file mode 100644
diff --git a/Engine/Plugins/Animation/ControlRig/Source/ControlRigEditor/Private/Editor/SControlRigFunctionLocalizationWidget.h b/Engine/Plugins/Animation/ControlRig/Source/ControlRigEditor/Private/Editor/SControlRigFunctionLocalizationWidget.h
deleted file mode 100644
diff --git a/Engine/Plugins/Animation/ControlRig/Source/ControlRigEditor/Private/Editor/SControlRigProfilingView.cpp b/Engine/Plugins/Animation/ControlRig/Source/ControlRigEditor/Private/Editor/SControlRigProfilingView.cpp
deleted file mode 100644
diff --git a/Engine/Plugins/Animation/ControlRig/Source/ControlRigEditor/Private/Editor/SControlRigProfilingView.h b/Engine/Plugins/Animation/ControlRig/Source/ControlRigEditor/Private/Editor/SControlRigProfilingView.h
deleted file mode 100644
diff --git a/Engine/Plugins/Animation/ControlRig/Source/ControlRigEditor/Private/Editor/SControlRigStackView.cpp b/Engine/Plugins/Animation/ControlRig/Source/ControlRigEditor/Private/Editor/SControlRigStackView.cpp
deleted file mode 100644
diff --git a/Engine/Plugins/Animation/ControlRig/Source/ControlRigEditor/Private/Editor/SControlRigStackView.h b/Engine/Plugins/Animation/ControlRig/Source/ControlRigEditor/Private/Editor/SControlRigStackView.h
deleted file mode 100644
diff --git a/Engine/Plugins/Animation/ControlRig/Source/ControlRigEditor/Private/Editor/SControlRigUnitCombo.cpp b/Engine/Plugins/Animation/ControlRig/Source/ControlRigEditor/Private/Editor/SControlRigUnitCombo.cpp
deleted file mode 100644
diff --git a/Engine/Plugins/Animation/ControlRig/Source/ControlRigEditor/Private/Editor/SControlRigUnitCombo.h b/Engine/Plugins/Animation/ControlRig/Source/ControlRigEditor/Private/Editor/SControlRigUnitCombo.h
deleted file mode 100644
diff --git a/Engine/Plugins/Animation/ControlRig/Source/ControlRigEditor/Private/Editor/SRigHierarchy.cpp b/Engine/Plugins/Animation/ControlRig/Source/ControlRigEditor/Private/Editor/SRigHierarchy.cpp
--- a/Engine/Plugins/Animation/ControlRig/Source/ControlRigEditor/Private/Editor/SRigHierarchy.cpp
+++ b/Engine/Plugins/Animation/ControlRig/Source/ControlRigEditor/Private/Editor/SRigHierarchy.cpp
@@ -1401,7 +1401,7 @@ void SRigHierarchy::RefreshHierarchy(const FAssetData& InAssetData)
 		// we do this to avoid the editmode / viewport shapes to refresh recursively,
 		// which can add an extreme slowdown depending on the number of bones (n^(n-1))
 		bool bSelectBones = true;
-		if (UControlRig* CurrentRig = StrongEditor->ControlRig)
+		if (UControlRig* CurrentRig = StrongEditor->GetControlRig())
 		{
 			bSelectBones = !CurrentRig->IsConstructionModeEnabled();
 		}
@@ -1503,7 +1503,7 @@ void SRigHierarchy::ImportHierarchy(const FAssetData& InAssetData)
 		// we do this to avoid the editmode / viewport shapes to refresh recursively,
 		// which can add an extreme slowdown depending on the number of bones (n^(n-1))
 		bool bSelectBones = true;
-		if (const UControlRig* CurrentRig = EditorSharedPtr->ControlRig)
+		if (const UControlRig* CurrentRig = EditorSharedPtr->GetControlRig())
 		{
 			bSelectBones = !CurrentRig->IsConstructionModeEnabled();
 		}
@@ -2156,7 +2156,7 @@ URigHierarchy* SRigHierarchy::GetHierarchy() const
 	}
 	if (ControlRigEditor.IsValid())
 	{
-		if (UControlRig* CurrentRig = ControlRigEditor.Pin()->ControlRig)
+		if (UControlRig* CurrentRig = ControlRigEditor.Pin()->GetControlRig())
 		{
 			return CurrentRig->GetHierarchy();
 		}
diff --git a/Engine/Plugins/Animation/ControlRig/Source/ControlRigEditor/Public/IControlRigEditor.h b/Engine/Plugins/Animation/ControlRig/Source/ControlRigEditor/Public/IControlRigEditor.h
--- a/Engine/Plugins/Animation/ControlRig/Source/ControlRigEditor/Public/IControlRigEditor.h
+++ b/Engine/Plugins/Animation/ControlRig/Source/ControlRigEditor/Public/IControlRigEditor.h
@@ -3,11 +3,11 @@
 #pragma once
 
 #include "CoreMinimal.h"
-#include "BlueprintEditor.h"
+#include "Editor/RigVMEditor.h"
 
 class UControlRigBlueprint;
 class FBlueprintActionDatabaseRegistrar;
 
-class IControlRigEditor : public FBlueprintEditor
+class IControlRigEditor : public FRigVMEditor
 {
 };
diff --git a/Engine/Plugins/Runtime/RigVM/Source/RigVM/Public/RigVMHost.h b/Engine/Plugins/Runtime/RigVM/Source/RigVM/Public/RigVMHost.h
--- a/Engine/Plugins/Runtime/RigVM/Source/RigVM/Public/RigVMHost.h
+++ b/Engine/Plugins/Runtime/RigVM/Source/RigVM/Public/RigVMHost.h
@@ -209,6 +209,10 @@ public:
 	UFUNCTION(BlueprintPure, Category = "RigVM")
 	URigVM* GetVM();
 
+#if WITH_EDITOR
+	const FRigVMLog* GetLog() const { return RigVMLog; }
+#endif
+
 	DECLARE_EVENT_TwoParams(URigVM, FRigVMExecutedEvent, class URigVMHost*, const FName&);
 	FRigVMExecutedEvent& OnInitialized_AnyThread() { return InitializedEvent; }
 	FRigVMExecutedEvent& OnExecuted_AnyThread() { return ExecutedEvent; }
diff --git a/Engine/Plugins/Runtime/RigVM/Source/RigVMDeveloper/Public/RigVMDeveloperModule.h b/Engine/Plugins/Runtime/RigVM/Source/RigVMDeveloper/Public/RigVMDeveloperModule.h
--- a/Engine/Plugins/Runtime/RigVM/Source/RigVMDeveloper/Public/RigVMDeveloperModule.h
+++ b/Engine/Plugins/Runtime/RigVM/Source/RigVMDeveloper/Public/RigVMDeveloperModule.h
@@ -9,5 +9,5 @@
 #include "CoreMinimal.h"
 #include "Logging/LogMacros.h"
 
-DECLARE_LOG_CATEGORY_EXTERN(LogRigVMDeveloper, Log, All);
+RIGVMDEVELOPER_API DECLARE_LOG_CATEGORY_EXTERN(LogRigVMDeveloper, Log, All);
 
diff --git a/Engine/Plugins/Runtime/RigVM/Source/RigVMEditor/Private/Editor/RigVMDetailsViewWrapperObject.cpp b/Engine/Plugins/Runtime/RigVM/Source/RigVMEditor/Private/Editor/RigVMDetailsViewWrapperObject.cpp
new file mode 100644
--- /dev/null
+++ b/Engine/Plugins/Runtime/RigVM/Source/RigVMEditor/Private/Editor/RigVMDetailsViewWrapperObject.cpp
@@ -0,0 +1,681 @@
+// Copyright Epic Games, Inc. All Rights Reserved.
+
+#include "Editor/RigVMDetailsViewWrapperObject.h"
+#include "RigVMCore/RigVMStruct.h"
+#include "RigVMModel/Nodes/RigVMUnitNode.h"
+#include "Modules/ModuleManager.h"
+#include "Algo/Sort.h"
+
+#include UE_INLINE_GENERATED_CPP_BY_NAME(RigVMDetailsViewWrapperObject)
+
+#if WITH_EDITOR
+#include "PropertyEditorModule.h"
+#endif
+
+TMap<URigVMDetailsViewWrapperObject::FPerClassInfo, UClass*> URigVMDetailsViewWrapperObject::InfoToClass;
+TMap<UClass*, URigVMDetailsViewWrapperObject::FPerClassInfo> URigVMDetailsViewWrapperObject::ClassToInfo;
+TSet<UClass*> URigVMDetailsViewWrapperObject::OutdatedClassToRecreate;
+
+static const TCHAR DiscardedWrapperClassTemplateName[] = TEXT("DiscardedWrapperClassTemplate");
+
+URigVMDetailsViewWrapperObject::URigVMDetailsViewWrapperObject()
+	: bIsSettingValue(false)
+{
+}
+
+UClass* URigVMDetailsViewWrapperObject::GetClassForStruct(UScriptStruct* InStruct, bool bCreateIfNeeded) const
+{
+	check(InStruct != nullptr);
+
+	if(UClass** ExistingClass = InfoToClass.Find(InStruct))
+	{
+		// if properties in a class has changed (due to user defined struct change)
+		// we have to recreate the class
+		if (OutdatedClassToRecreate.Contains(*ExistingClass))
+		{
+			OutdatedClassToRecreate.Remove(*ExistingClass);
+
+			InfoToClass.Remove(InStruct);
+			ClassToInfo.Remove(*ExistingClass);
+
+			FString DiscardedWrapperClassName;
+			static int32 DiscardedWrapperClassIndex = 0;
+			do
+			{
+				DiscardedWrapperClassName = FString::Printf(TEXT("%s_%d"), DiscardedWrapperClassTemplateName, DiscardedWrapperClassIndex++);
+				if(StaticFindObjectFast(nullptr, GetTransientPackage(), *DiscardedWrapperClassName) == nullptr)
+				{
+					break;
+				}
+			}
+			while (DiscardedWrapperClassIndex < INT_MAX);
+
+			(*ExistingClass)->Rename(*DiscardedWrapperClassName, GetTransientPackage(), REN_ForceNoResetLoaders | REN_DoNotDirty | REN_DontCreateRedirectors | REN_NonTransactional);
+			
+			(*ExistingClass)->RemoveFromRoot();
+		}
+		else
+		{
+			return *ExistingClass;
+		}
+	}
+
+	if(!bCreateIfNeeded)
+	{
+		return nullptr;
+	}
+
+	UClass* SuperClass = URigVMDetailsViewWrapperObject::StaticClass();
+	const FName WrapperClassName(FString::Printf(TEXT("%s_WrapperObject"), *InStruct->GetStructCPPName()));
+
+	UClass* WrapperClass = NewObject<UClass>(
+		GetTransientPackage(),
+		WrapperClassName,
+		RF_Public | RF_Transient
+ 		);
+
+	// make sure this doesn't get garbage collected
+	WrapperClass->AddToRoot();
+
+	// Eviscerate the class.
+	WrapperClass->PurgeClass(false);
+	WrapperClass->PropertyLink = SuperClass->PropertyLink;
+
+	WrapperClass->SetSuperStruct(SuperClass);
+	WrapperClass->ClassWithin = UObject::StaticClass();
+	WrapperClass->ClassConfigName = SuperClass->ClassConfigName;
+	WrapperClass->ClassFlags |= CLASS_NotPlaceable | CLASS_Hidden;
+	WrapperClass->SetMetaData(TEXT("DisplayName"), *InStruct->GetDisplayNameText().ToString());
+
+	struct Local
+	{
+		static bool IsStructHashable(const UScriptStruct* InStructType)
+		{
+			if (InStructType->IsNative())
+			{
+				return InStructType->GetCppStructOps() && InStructType->GetCppStructOps()->HasGetTypeHash();
+			}
+			else
+			{
+				for (TFieldIterator<FProperty> It(InStructType); It; ++It)
+				{
+					if (CastField<FBoolProperty>(*It))
+					{
+						continue;
+					}
+					else if (!It->HasAllPropertyFlags(CPF_HasGetValueTypeHash))
+					{
+						return false;
+					}
+				}
+				return true;
+			}
+		}
+	};
+
+	// duplicate all properties from the struct to the wrapper object
+	FField** LinkToProperty = &WrapperClass->ChildProperties;
+
+	for (TFieldIterator<FProperty> PropertyIt(InStruct); PropertyIt; ++PropertyIt)
+	{
+		const FProperty* InProperty = *PropertyIt;
+		FProperty* NewProperty = CastFieldChecked<FProperty>(FField::Duplicate(InProperty, WrapperClass, InProperty->GetFName()));
+		check(NewProperty);
+		FField::CopyMetaData(InProperty, NewProperty);
+
+		if (NewProperty->HasMetaData(TEXT("Input")) || NewProperty->HasMetaData(TEXT("Visible")))
+		{
+			// filter out execute pins
+			bool bIsEditable = true;
+			if (FStructProperty* StructProperty = CastField<FStructProperty>(NewProperty))
+			{
+				if (StructProperty->Struct->IsChildOf(FRigVMExecuteContext::StaticStruct()))
+				{
+					bIsEditable = false;
+				}
+			}
+
+			if(bIsEditable)
+			{
+				NewProperty->SetPropertyFlags(NewProperty->GetPropertyFlags() | CPF_Edit);
+			}
+		}
+
+		*LinkToProperty = NewProperty;
+		LinkToProperty = &(*LinkToProperty)->Next;
+	}
+
+	// Update the class
+	WrapperClass->Bind();
+	WrapperClass->StaticLink(true);
+	
+	// Similar to FConfigPropertyHelperDetails::CustomizeDetails, this is required for GC to work properly
+	WrapperClass->AssembleReferenceTokenStream();
+
+	InfoToClass.Add(InStruct, WrapperClass);
+	ClassToInfo.Add(WrapperClass, InStruct);
+
+	UObject* CDO = WrapperClass->GetDefaultObject(true);
+	CDO->AddToRoot();
+
+	// import the defaults from the struct onto the class
+	TSharedPtr<FStructOnScope> DefaultStruct = MakeShareable(new FStructOnScope(InStruct));
+	CopyPropertiesForUnrelatedStructs((uint8*)CDO, WrapperClass, DefaultStruct->GetStructMemory(), DefaultStruct->GetStruct());
+
+	return WrapperClass;
+}
+
+URigVMDetailsViewWrapperObject* URigVMDetailsViewWrapperObject::MakeInstance(UClass* InWrapperObjectClass, UObject* InOuter, UScriptStruct* InStruct, uint8* InStructMemory, UObject* InSubject)
+{
+	check(InWrapperObjectClass != nullptr);
+	check(InStruct != nullptr);
+
+	InOuter = InOuter == nullptr ? GetTransientPackage() : InOuter;
+
+	const URigVMDetailsViewWrapperObject* CDO = CastChecked<URigVMDetailsViewWrapperObject>(InWrapperObjectClass->GetDefaultObject());
+	const UClass* WrapperClass = CDO->GetClassForStruct(InStruct);
+	if(WrapperClass == nullptr)
+	{
+		return nullptr;
+	}
+
+	URigVMDetailsViewWrapperObject* Instance = NewObject<URigVMDetailsViewWrapperObject>(InOuter, WrapperClass, NAME_None, RF_Public | RF_Transient | RF_TextExportTransient | RF_DuplicateTransient);
+	Instance->SetContent(InStructMemory, InStruct);
+	Instance->SetSubject(InSubject);
+
+	return Instance;
+}
+
+UScriptStruct* URigVMDetailsViewWrapperObject::GetWrappedStruct() const
+{
+	return ClassToInfo.FindChecked(GetClass()).ScriptStruct;
+}
+
+UClass* URigVMDetailsViewWrapperObject::GetClassForNodes(TArray<URigVMNode*> InNodes, bool bCreateIfNeeded) const
+{
+	if(InNodes.IsEmpty())
+	{
+		return nullptr;
+	}
+
+	static TArray<ERigVMPinDirection> PinDirectionsToCheck = {
+		ERigVMPinDirection::IO,
+		ERigVMPinDirection::Input,
+		ERigVMPinDirection::Visible
+	};
+
+	// determine if all nodes are unit nodes and match their script struct
+	TArray<UScriptStruct*> UnitStructs;
+	for(URigVMNode* Node : InNodes)
+	{
+		if(URigVMUnitNode* UnitNode = Cast<URigVMUnitNode>(Node))
+		{
+			if(UScriptStruct* ScriptStruct = UnitNode->GetScriptStruct())
+			{
+				UnitStructs.AddUnique(ScriptStruct);
+			}
+		}
+	}
+
+	// find the matching pins on all passed in nodes
+	TArray<URigVMPin*> PinsToInspect;
+
+	for(const ERigVMPinDirection PinDirectionToCheck : PinDirectionsToCheck)
+	{
+		for(URigVMPin* Pin : InNodes[0]->GetPins())
+		{
+			if(Pin->IsExecuteContext())
+			{
+				continue;
+			}
+			
+			if(Pin->GetDirection() != PinDirectionToCheck)
+			{
+				continue;
+			}
+			
+			bool bInspectPin = true;
+			
+			for(int32 NodeIndex=1; NodeIndex<InNodes.Num(); NodeIndex++)
+			{
+				URigVMPin* OtherPin = InNodes[NodeIndex]->FindPin(Pin->GetName());
+				if(OtherPin == nullptr)
+				{
+					bInspectPin = false;
+					break;
+				}
+
+				if(Pin->GetCPPType() != OtherPin->GetCPPType())
+				{
+					bInspectPin = false;
+					break;
+				}
+			}
+
+			if(bInspectPin)
+			{
+				PinsToInspect.Add(Pin);
+			}
+		}
+	}
+
+	TArray<FString> PinNotations;
+	for(URigVMPin* Pin : PinsToInspect)
+	{
+		PinNotations.Add(FString::Printf(TEXT("%s %s"), *Pin->GetCPPType(), *Pin->GetName()));
+	}
+
+	FString Notation = FString::Join(PinNotations, TEXT(", "));
+	if(UnitStructs.Num() > 0)
+	{
+		// sort the structs to ensure we get the same notation each time
+		Algo::Sort(UnitStructs, [](UScriptStruct* A, UScriptStruct* B) -> bool
+		{
+			return A->GetStructCPPName() > B->GetStructCPPName();
+		});
+
+		TArray<FString> StructNames;
+		for(UScriptStruct* UnitStruct : UnitStructs)
+		{
+			StructNames.Add(UnitStruct->GetStructCPPName());
+		}
+		Notation = FString::Printf(TEXT("%s(%s)"), *FString::Join(StructNames, TEXT("|")), *Notation);
+	}
+
+	const FPerClassInfo PerClassInfo(Notation);
+	if(UClass** ExistingClass = InfoToClass.Find(PerClassInfo))
+	{
+		// if properties in a class has changed (due to user defined struct change)
+		// we have to recreate the class
+		if (OutdatedClassToRecreate.Contains(*ExistingClass))
+		{
+			OutdatedClassToRecreate.Remove(*ExistingClass);
+
+			InfoToClass.Remove(PerClassInfo);
+			ClassToInfo.Remove(*ExistingClass);
+
+			FString DiscardedWrapperClassName;
+			static int32 DiscardedWrapperClassIndex = 0;
+			do
+			{
+				DiscardedWrapperClassName = FString::Printf(TEXT("%s_%d"), DiscardedWrapperClassTemplateName, DiscardedWrapperClassIndex++);
+				if(StaticFindObjectFast(nullptr, GetTransientPackage(), *DiscardedWrapperClassName) == nullptr)
+				{
+					break;
+				}
+			}
+			while (DiscardedWrapperClassIndex < INT_MAX);
+
+			(*ExistingClass)->Rename(*DiscardedWrapperClassName, GetTransientPackage(), REN_ForceNoResetLoaders | REN_DoNotDirty | REN_DontCreateRedirectors | REN_NonTransactional);
+			
+			(*ExistingClass)->RemoveFromRoot();
+		}
+		else
+		{
+			return *ExistingClass;
+		}
+	}
+
+	if(!bCreateIfNeeded)
+	{
+		return nullptr;
+	}
+
+	UClass* SuperClass = URigVMDetailsViewWrapperObject::StaticClass();
+	const int32 HashForNotation = (int32)GetTypeHash(PerClassInfo);
+	const FName WrapperClassName(FString::Printf(TEXT("RigNode%d_WrapperObject"), HashForNotation));
+
+	UClass* WrapperClass = NewObject<UClass>(
+		GetTransientPackage(),
+		WrapperClassName,
+		RF_Public | RF_Transient
+		 );
+
+	// make sure this doesn't get garbage collected
+	WrapperClass->AddToRoot();
+
+	// Eviscerate the class.
+	WrapperClass->PurgeClass(false);
+	WrapperClass->PropertyLink = SuperClass->PropertyLink;
+
+	WrapperClass->SetSuperStruct(SuperClass);
+	WrapperClass->ClassWithin = UObject::StaticClass();
+	WrapperClass->ClassConfigName = SuperClass->ClassConfigName;
+	WrapperClass->ClassFlags |= CLASS_NotPlaceable | CLASS_Hidden;
+
+	WrapperClass->SetMetaData(TEXT("DisplayName"), TEXT("Rig Node"));
+
+	// create properties - one for each pin to inspect
+	FField** LinkToProperty = &WrapperClass->ChildProperties;
+
+	for(URigVMPin* Pin : PinsToInspect)
+	{
+		static FString BoolString = TEXT("bool");
+		static FString Int32String = TEXT("int32");
+		static FString IntString = TEXT("int");
+		static FString FloatString = TEXT("float");
+		static FString DoubleString = TEXT("double");
+		static FString StringString = TEXT("FString");
+		static FString NameString = TEXT("FName");
+
+		FProperty* Property = nullptr;
+		FProperty** ElementProperty = &Property;
+		FFieldVariant PropertyOwner = WrapperClass;
+
+		const FString CPPType = Pin->GetCPPType();
+		FString BaseCPPType = CPPType;
+		if(RigVMTypeUtils::IsArrayType(CPPType))
+		{
+			BaseCPPType = RigVMTypeUtils::BaseTypeFromArrayType(CPPType);
+			
+			// create an array property as a container for the tail
+			FArrayProperty* ArrayProperty = new FArrayProperty(PropertyOwner, Pin->GetFName(), RF_Public);
+			Property = ArrayProperty;
+			ElementProperty = &ArrayProperty->Inner;
+			PropertyOwner = ArrayProperty;
+		}
+
+		if(BaseCPPType.Equals(BoolString, ESearchCase::IgnoreCase))
+		{
+			(*ElementProperty) = new FBoolProperty(PropertyOwner, Pin->GetFName(), RF_Public);;
+		}
+		else if(BaseCPPType.Equals(Int32String, ESearchCase::IgnoreCase) ||
+			BaseCPPType.Equals(IntString, ESearchCase::IgnoreCase))
+		{
+			(*ElementProperty) = new FIntProperty(PropertyOwner, Pin->GetFName(), RF_Public);;
+		}
+		else if(BaseCPPType.Equals(FloatString, ESearchCase::IgnoreCase))
+		{
+			(*ElementProperty) = new FFloatProperty(PropertyOwner, Pin->GetFName(), RF_Public);;
+		}
+		else if(BaseCPPType.Equals(DoubleString, ESearchCase::IgnoreCase))
+		{
+			(*ElementProperty) = new FDoubleProperty(PropertyOwner, Pin->GetFName(), RF_Public);;
+		}
+		else if(BaseCPPType.Equals(StringString, ESearchCase::IgnoreCase))
+		{
+			(*ElementProperty) = new FStrProperty(PropertyOwner, Pin->GetFName(), RF_Public);;
+		}
+		else if(BaseCPPType.Equals(NameString, ESearchCase::IgnoreCase))
+		{
+			(*ElementProperty) = new FNameProperty(PropertyOwner, Pin->GetFName(), RF_Public);;
+		}
+		else if(Pin->GetCPPTypeObject())
+		{
+			if(UEnum* Enum = Cast<UEnum>(Pin->GetCPPTypeObject()))
+			{
+				FByteProperty* ByteProperty = new FByteProperty(PropertyOwner, Pin->GetFName(), RF_Public);
+				ByteProperty->Enum = Enum;
+				(*ElementProperty) = ByteProperty;
+			}
+			else if(UClass* Class = Cast<UClass>(Pin->GetCPPTypeObject()))
+			{
+				FObjectProperty* ObjectProperty = new FObjectProperty(PropertyOwner, Pin->GetFName(), RF_Public);
+				ObjectProperty->PropertyClass = Class; 
+				(*ElementProperty) = ObjectProperty;
+			}
+			else if(UScriptStruct* ScriptStruct = Cast<UScriptStruct>(Pin->GetCPPTypeObject()))
+			{
+				FStructProperty* StructProperty = new FStructProperty(PropertyOwner, Pin->GetFName(), RF_Public);
+				StructProperty->Struct = ScriptStruct;
+				(*ElementProperty) = StructProperty;
+			}
+		}
+
+		if(Property == nullptr)
+		{
+			return nullptr;
+		}
+		Property->SetPropertyFlags(CPF_Edit);
+
+		for(UScriptStruct* UnitStruct : UnitStructs)
+		{
+			if(const FProperty* UnitProperty = UnitStruct->FindPropertyByName(Property->GetFName()))
+			{
+				if(Property->SameType(UnitProperty))
+				{
+					FField::CopyMetaData(UnitProperty, Property);
+					break;
+				}
+			}
+		}
+
+		Property->SetMetaData(TEXT("Category"), TEXT("Nodes"));
+
+		*LinkToProperty = Property;
+		LinkToProperty = &(*LinkToProperty)->Next;
+	}
+
+	// Update the class
+	WrapperClass->Bind();
+	WrapperClass->StaticLink(true);
+	
+	// Similar to FConfigPropertyHelperDetails::CustomizeDetails, this is required for GC to work properly
+	WrapperClass->AssembleReferenceTokenStream();
+
+	InfoToClass.Add(PerClassInfo, WrapperClass);
+	ClassToInfo.Add(WrapperClass, PerClassInfo);
+
+	URigVMDetailsViewWrapperObject* CDO = Cast<URigVMDetailsViewWrapperObject>(WrapperClass->GetDefaultObject(true));
+	CDO->AddToRoot();
+
+	// setup the CDO's defaults to match the units
+	if(UnitStructs.Num() > 0)
+	{
+		UScriptStruct* UnitStruct = UnitStructs[0];
+		FStructOnScope DefaultStructScope(UnitStruct);
+		const FRigVMStruct* DefaultStruct = (const FRigVMStruct*)DefaultStructScope.GetStructMemory();
+
+		for (TFieldIterator<FProperty> UnitProperty(UnitStruct); UnitProperty; ++UnitProperty)
+		{
+			if(FProperty* Property = WrapperClass->FindPropertyByName(UnitProperty->GetFName()))
+			{
+				if(Property->SameType(*UnitProperty))
+				{
+					const uint8* SourceMemory = UnitProperty->ContainerPtrToValuePtr<uint8>(DefaultStruct);
+					uint8* TargetMemory = Property->ContainerPtrToValuePtr<uint8>(CDO);
+					Property->CopyCompleteValue(TargetMemory, SourceMemory);
+				}
+			}
+		}
+	}
+
+	return WrapperClass;	
+}
+
+URigVMDetailsViewWrapperObject* URigVMDetailsViewWrapperObject::MakeInstance(UClass* InWrapperObjectClass, UObject* InOuter, TArray<URigVMNode*> InNodes, URigVMNode* InSubject)
+{
+	InOuter = InOuter == nullptr ? GetTransientPackage() : InOuter;
+
+	const URigVMDetailsViewWrapperObject* CDO = CastChecked<URigVMDetailsViewWrapperObject>(InWrapperObjectClass->GetDefaultObject());
+	const UClass* WrapperClass = CDO->GetClassForNodes(InNodes);
+	if(WrapperClass == nullptr)
+	{
+		return nullptr;
+	}
+
+	URigVMDetailsViewWrapperObject* Instance = NewObject<URigVMDetailsViewWrapperObject>(InOuter, WrapperClass, NAME_None, RF_Public | RF_Transient | RF_TextExportTransient | RF_DuplicateTransient);
+	Instance->SetContent(InSubject);
+	Instance->SetSubject(InSubject);
+
+	InSubject->GetGraph()->OnModified().AddUObject(Instance, &URigVMDetailsViewWrapperObject::HandleModifiedEvent);
+	
+	return Instance;
+
+}
+
+void URigVMDetailsViewWrapperObject::MarkOutdatedClass(UClass* InClass)
+{
+	if (InClass)
+	{
+		OutdatedClassToRecreate.Add(InClass);
+	}
+}
+
+bool URigVMDetailsViewWrapperObject::IsValidClass(UClass* InClass)
+{
+	if (!InClass)
+	{
+		return false;
+	}
+	
+	if (!InClass->IsChildOf(StaticClass()))
+	{
+		return false;
+	}
+		
+	if (InClass->GetName().Contains(DiscardedWrapperClassTemplateName))
+	{
+		return false;
+	}
+
+	return true;
+}
+
+FString URigVMDetailsViewWrapperObject::GetWrappedNodeNotation() const
+{
+	return ClassToInfo.FindChecked(GetClass()).Notation;
+}
+
+void URigVMDetailsViewWrapperObject::SetContent(const uint8* InStructMemory, const UStruct* InStruct)
+{
+	CopyPropertiesForUnrelatedStructs((uint8*)this, GetClass(), InStructMemory, InStruct);
+}
+
+void URigVMDetailsViewWrapperObject::GetContent(uint8* OutStructMemory, const UStruct* InStruct) const
+{
+	CopyPropertiesForUnrelatedStructs(OutStructMemory, InStruct, (const uint8*)this, GetClass());
+}
+
+void URigVMDetailsViewWrapperObject::SetContent(URigVMNode* InNode)
+{
+	for(URigVMPin* Pin : InNode->GetPins())
+	{
+		SetContentForPin(Pin);
+	}
+}
+
+UObject* URigVMDetailsViewWrapperObject::GetSubject() const
+{
+	if(SubjectPtr.IsValid())
+	{
+		return SubjectPtr.Get();
+	}
+	return nullptr;
+}
+
+void URigVMDetailsViewWrapperObject::SetSubject(UObject* InSubject)
+{
+	SubjectPtr = TWeakObjectPtr<UObject>(InSubject);
+}
+
+void URigVMDetailsViewWrapperObject::PostEditChangeChainProperty(FPropertyChangedChainEvent& PropertyChangedEvent)
+{
+	Super::PostEditChangeChainProperty(PropertyChangedEvent);
+
+	FString PropertyPath;
+
+	const TDoubleLinkedList<FProperty*>::TDoubleLinkedListNode* PropertyNode = PropertyChangedEvent.PropertyChain.GetHead();
+	do
+	{
+		const FString PropertyName = PropertyNode->GetValue()->GetNameCPP();
+		if(PropertyPath.IsEmpty())
+		{
+			PropertyPath = PropertyName;
+		}
+		else
+		{
+			PropertyPath = FString::Printf(TEXT("%s->%s"), *PropertyPath, *PropertyName);
+		}
+		PropertyNode = PropertyNode->GetNextNode();
+	}
+	while(PropertyNode);
+
+	TGuardValue<bool> ValueSettingGuard(bIsSettingValue, true);
+	WrappedPropertyChangedChainEvent.Broadcast(this, PropertyPath, PropertyChangedEvent);
+}
+
+void URigVMDetailsViewWrapperObject::CopyPropertiesForUnrelatedStructs(uint8* InTargetMemory,
+	const UStruct* InTargetStruct, const uint8* InSourceMemory, const UStruct* InSourceStruct)
+{
+	check(InTargetMemory);
+	check(InTargetStruct);
+	check(InSourceMemory);
+	check(InSourceStruct);
+	
+	for (TFieldIterator<FProperty> PropertyIt(InTargetStruct); PropertyIt; ++PropertyIt)
+	{
+		const FProperty* TargetProperty = *PropertyIt;
+		const FProperty* SourceProperty = InSourceStruct->FindPropertyByName(TargetProperty->GetFName());
+		if(SourceProperty == nullptr)
+		{
+			continue;
+		}
+		check(TargetProperty->SameType(SourceProperty));
+
+		uint8* TargetPropertyMemory = TargetProperty->ContainerPtrToValuePtr<uint8>(InTargetMemory);
+		const uint8* SourcePropertyMemory = SourceProperty->ContainerPtrToValuePtr<uint8>(InSourceMemory);
+		TargetProperty->CopyCompleteValue(TargetPropertyMemory, SourcePropertyMemory);
+	}
+}
+
+void URigVMDetailsViewWrapperObject::HandleModifiedEvent(ERigVMGraphNotifType InNotifType, URigVMGraph* InGraph,
+	UObject* InSubject)
+{
+	if(InNotifType != ERigVMGraphNotifType::PinDefaultValueChanged)
+	{
+		return;
+	}
+	
+	URigVMPin* Pin = CastChecked<URigVMPin>(InSubject);
+	if(Pin->GetNode() != GetSubject())
+	{
+		return;
+	}
+
+	SetContentForPin(Pin->GetRootPin());
+}
+
+void URigVMDetailsViewWrapperObject::SetContentForPin(URigVMPin* InPin)
+{
+	if(bIsSettingValue)
+	{
+		return;
+	}
+	TGuardValue<bool> SetValueGuard(bIsSettingValue, true);
+	
+	ensure(InPin->IsRootPin());
+
+	const FString DefaultValue = InPin->GetDefaultValue();
+	if(DefaultValue.IsEmpty())
+	{
+		return;
+	}
+
+	class FRigVMDetailsViewWrapperObjectImportErrorContext : public FOutputDevice
+	{
+	public:
+
+		int32 NumErrors;
+
+		FRigVMDetailsViewWrapperObjectImportErrorContext()
+			: FOutputDevice()
+			, NumErrors(0)
+		{
+		}
+
+		virtual void Serialize(const TCHAR* V, ELogVerbosity::Type Verbosity, const class FName& Category) override
+		{
+			UE_LOG(LogRigVMDeveloper, Error, TEXT("Error Importing To Hierarchy: %s"), V);
+			NumErrors++;
+		}
+	};
+
+	// it's possible that the pin has been filtered out when creating the class
+	if(const FProperty* Property = GetClass()->FindPropertyByName(InPin->GetFName()))
+	{
+		FRigVMDetailsViewWrapperObjectImportErrorContext ErrorPipe;
+		Property->ImportText_InContainer(*DefaultValue, this, nullptr, EPropertyPortFlags::PPF_None, &ErrorPipe);
+	}
+}
+
diff --git a/Engine/Plugins/Runtime/RigVM/Source/RigVMEditor/Private/Editor/RigVMEditor.cpp b/Engine/Plugins/Runtime/RigVM/Source/RigVMEditor/Private/Editor/RigVMEditor.cpp
new file mode 100644
--- /dev/null
+++ b/Engine/Plugins/Runtime/RigVM/Source/RigVMEditor/Private/Editor/RigVMEditor.cpp
@@ -0,0 +1,41 @@
+// Copyright Epic Games, Inc. All Rights Reserved.
+
+#include "Editor/RigVMEditor.h"
+#include "Editor/RigVMDetailsViewWrapperObject.h"
+
+#define LOCTEXT_NAMESPACE "RigVMEditor"
+
+FRigVMEditor::FRigVMEditor()
+	: Host(nullptr)
+{
+}
+
+URigVMBlueprint* FRigVMEditor::GetRigVMBlueprint() const
+{
+	return Cast<URigVMBlueprint>(GetBlueprintObj());
+}
+
+URigVMHost* FRigVMEditor::GetRigVMHost() const
+{
+	return Host;
+}
+
+UClass* FRigVMEditor::GetDetailWrapperClass() const
+{
+	return URigVMDetailsViewWrapperObject::StaticClass();
+}
+
+void FRigVMEditor::HandleModifiedEvent(ERigVMGraphNotifType InNotifType, URigVMGraph* InGraph, UObject* InSubject)
+{
+}
+
+void FRigVMEditor::SetHost(URigVMHost* InHost)
+{
+	Host = InHost;
+	if(Host)
+	{
+		OnPreviewHostUpdated().Broadcast(this);
+	}
+}
+
+#undef LOCTEXT_NAMESPACE 
diff --git a/Engine/Plugins/Runtime/RigVM/Source/RigVMEditor/Private/Editor/RigVMExecutionStackCommands.cpp b/Engine/Plugins/Runtime/RigVM/Source/RigVMEditor/Private/Editor/RigVMExecutionStackCommands.cpp
new file mode 100644
--- /dev/null
+++ b/Engine/Plugins/Runtime/RigVM/Source/RigVMEditor/Private/Editor/RigVMExecutionStackCommands.cpp
@@ -0,0 +1,13 @@
+// Copyright Epic Games, Inc. All Rights Reserved.
+
+#include "Editor/RigVMExecutionStackCommands.h"
+
+#define LOCTEXT_NAMESPACE "RigVMExecutionStackCommands"
+
+void FRigVMExecutionStackCommands::RegisterCommands()
+{
+	UI_COMMAND(FocusOnSelection, "Focus On Selection", "Finds the selected operator's node in the graph.", EUserInterfaceActionType::Button, FInputChord());
+	UI_COMMAND(GoToInstruction, "Go To Instruction", "Looks for a specific instruction by index and brings it into focus.", EUserInterfaceActionType::Button, FInputChord(EKeys::G, EModifierKey::Control));
+}
+
+#undef LOCTEXT_NAMESPACE
diff --git a/Engine/Plugins/Runtime/RigVM/Source/RigVMEditor/Private/Editor/RigVMExecutionStackTabSummoner.cpp b/Engine/Plugins/Runtime/RigVM/Source/RigVMEditor/Private/Editor/RigVMExecutionStackTabSummoner.cpp
new file mode 100644
--- /dev/null
+++ b/Engine/Plugins/Runtime/RigVM/Source/RigVMEditor/Private/Editor/RigVMExecutionStackTabSummoner.cpp
@@ -0,0 +1,29 @@
+// Copyright Epic Games, Inc. All Rights Reserved.
+
+#include "Editor/RigVMExecutionStackTabSummoner.h"
+#include "Widgets/SRigVMExecutionStackView.h"
+#include "Editor/RigVMEditor.h"
+
+#define LOCTEXT_NAMESPACE "RigVMExecutionStackTabSummoner"
+
+const FName FRigVMExecutionStackTabSummoner::TabID(TEXT("Execution Stack"));
+
+FRigVMExecutionStackTabSummoner::FRigVMExecutionStackTabSummoner(const TSharedRef<FRigVMEditor>& InRigVMEditor)
+	: FWorkflowTabFactory(TabID, InRigVMEditor)
+	, RigVMEditor(InRigVMEditor)
+{
+	TabLabel = LOCTEXT("RigVMExecutionStackTabLabel", "Execution Stack");
+	TabIcon = FSlateIcon(TEXT("ControlRigEditorStyle"), "ExecutionStack.TabIcon");
+
+	bIsSingleton = true;
+
+	ViewMenuDescription = LOCTEXT("RigVMExecutionStack_ViewMenu_Desc", "Execution Stack");
+	ViewMenuTooltip = LOCTEXT("RigVMExecutionStack_ViewMenu_ToolTip", "Show the Execution Stack tab");
+}
+
+TSharedRef<SWidget> FRigVMExecutionStackTabSummoner::CreateTabBody(const FWorkflowTabSpawnInfo& Info) const
+{
+	return SNew(SRigVMExecutionStackView, RigVMEditor.Pin().ToSharedRef());
+}
+
+#undef LOCTEXT_NAMESPACE 
diff --git a/Engine/Plugins/Runtime/RigVM/Source/RigVMEditor/Private/Widgets/SRigVMExecutionStackView.cpp b/Engine/Plugins/Runtime/RigVM/Source/RigVMEditor/Private/Widgets/SRigVMExecutionStackView.cpp
new file mode 100644
--- /dev/null
+++ b/Engine/Plugins/Runtime/RigVM/Source/RigVMEditor/Private/Widgets/SRigVMExecutionStackView.cpp
@@ -0,0 +1,1053 @@
+// Copyright Epic Games, Inc. All Rights Reserved.
+
+#include "Widgets/SRigVMExecutionStackView.h"
+#include "Widgets/Layout/SBorder.h"
+#include "Widgets/SBoxPanel.h"
+#include "Widgets/Layout/SScrollBar.h"
+#include "Widgets/Layout/SSpacer.h"
+#include "Widgets/Input/SSearchBox.h"
+#include "Widgets/Input/SNumericEntryBox.h"
+#include "Editor/RigVMEditor.h"
+#include "Editor/RigVMExecutionStackCommands.h"
+#include "RigVMHost.h"
+#include "RigVMBlueprintGeneratedClass.h"
+#include "Framework/MultiBox/MultiBoxBuilder.h"
+#include "Styling/AppStyle.h"
+#include "DetailLayoutBuilder.h"
+#include "RigVMModel/Nodes/RigVMAggregateNode.h"
+#include "Widgets/Input/SSearchBox.h"
+#include "Dialog/SCustomDialog.h"
+
+#define LOCTEXT_NAMESPACE "SRigVMExecutionStackView"
+
+TAutoConsoleVariable<bool> CVarRigVMExecutionStackDetailedLabels(TEXT("RigVM.StackDetailedLabels"), false, TEXT("Set to true to turn on detailed labels for the execution stack widget"));
+
+//////////////////////////////////////////////////////////////
+/// FRigStackEntry
+///////////////////////////////////////////////////////////
+FRigStackEntry::FRigStackEntry(int32 InEntryIndex, ERigStackEntry::Type InEntryType, int32 InInstructionIndex, ERigVMOpCode InOpCode, const FString& InLabel, const FRigVMASTProxy& InProxy)
+	: EntryIndex(InEntryIndex)
+	, EntryType(InEntryType)
+	, InstructionIndex(InInstructionIndex)
+	, CallPath(InProxy.GetCallstack().GetCallPath())
+	, Callstack(InProxy.GetCallstack())
+	, OpCode(InOpCode)
+	, Label(InLabel)
+{
+
+}
+
+TSharedRef<ITableRow> FRigStackEntry::MakeTreeRowWidget(const TSharedRef<STableViewBase>& InOwnerTable, TSharedRef<FRigStackEntry> InEntry, TSharedRef<FUICommandList> InCommandList, TSharedPtr<SRigVMExecutionStackView> InStackView, TWeakObjectPtr<URigVMBlueprint> InBlueprint)
+{
+	return SNew(SRigStackItem, InOwnerTable, InEntry, InCommandList, InBlueprint);
+}
+
+//////////////////////////////////////////////////////////////
+/// SRigStackItem
+///////////////////////////////////////////////////////////
+void SRigStackItem::Construct(const FArguments& InArgs, const TSharedRef<STableViewBase>& OwnerTable, TSharedRef<FRigStackEntry> InStackEntry, TSharedRef<FUICommandList> InCommandList, TWeakObjectPtr<URigVMBlueprint> InBlueprint)
+{
+	WeakStackEntry = InStackEntry;
+	WeakBlueprint = InBlueprint;
+	WeakCommandList = InCommandList;
+
+	TSharedPtr< STextBlock > NumberWidget;
+	TSharedPtr< STextBlock > TextWidget;
+
+	const FSlateBrush* Icon = nullptr;
+	switch (InStackEntry->EntryType)
+	{
+		case ERigStackEntry::Operator:
+		{
+			Icon = FSlateIcon(TEXT("RigVMEditor"), "ControlRig.RigUnit").GetIcon();
+			break;
+		}
+		case ERigStackEntry::Info:
+		{
+			Icon = FAppStyle::GetBrush("Icons.Info");
+			break;
+		}
+		case ERigStackEntry::Warning:
+		{
+			Icon = FAppStyle::GetBrush("Icons.Warning");
+			break;
+		}
+		case ERigStackEntry::Error:
+		{
+			Icon = FAppStyle::GetBrush("Icons.Error");
+			break;
+		}
+		default:
+		{
+			break;
+		}
+	}
+
+	STableRow<TSharedPtr<FRigStackEntry>>::Construct(
+		STableRow<TSharedPtr<FRigStackEntry>>::FArguments()
+		.Content()
+		[
+			SNew(SHorizontalBox)
+			+ SHorizontalBox::Slot()
+			.AutoWidth()
+			.VAlign(VAlign_Center)
+			.HAlign(HAlign_Left)
+			[
+				SNew(SHorizontalBox)
+				+ SHorizontalBox::Slot()
+				.MaxWidth(35.f)
+				.AutoWidth()
+				.VAlign(VAlign_Center)
+				.HAlign(HAlign_Right)
+				[
+					SAssignNew(NumberWidget, STextBlock)
+					.Text(this, &SRigStackItem::GetIndexText)
+					.Font(IDetailLayoutBuilder::GetDetailFont())
+				]
+				+ SHorizontalBox::Slot()
+				.MaxWidth(22.f)
+				.VAlign(VAlign_Center)
+				.HAlign(HAlign_Center)
+				[
+					SNew(SImage)
+					.Image(Icon)
+				]
+				
+				+ SHorizontalBox::Slot()
+				.FillWidth(1.f)
+				.AutoWidth()
+				.VAlign(VAlign_Center)
+				.HAlign(HAlign_Fill)
+				[
+					SAssignNew(TextWidget, STextBlock)
+					.Text(this, &SRigStackItem::GetLabelText)
+					.Font(this, &SRigStackItem::GetLabelFont)
+					.ToolTipText(this, &SRigStackItem::GetTooltip)
+					.Justification(ETextJustify::Left)
+				]
+			]
+
+			+ SHorizontalBox::Slot()
+			.FillWidth(1.f)
+			[
+				SNew(SSpacer)
+			]
+
+			+ SHorizontalBox::Slot()
+			.AutoWidth()
+			.VAlign(VAlign_Center)
+			.HAlign(HAlign_Right)
+			[
+				SNew(SHorizontalBox)
+				+ SHorizontalBox::Slot()
+				.AutoWidth()
+				.VAlign(VAlign_Center)
+	            .HAlign(HAlign_Left)
+	            [
+	                SNew(STextBlock)
+	                .Text(this, &SRigStackItem::GetVisitedCountText)
+	                .Font(IDetailLayoutBuilder::GetDetailFont())
+	            ]
+	            
+	            + SHorizontalBox::Slot()
+	            .Padding(20, 0, 0, 0)
+				.AutoWidth()
+	            .VAlign(VAlign_Center)
+	            .HAlign(HAlign_Left)
+	            [
+	                SNew(STextBlock)
+	                .Text(this, &SRigStackItem::GetDurationText)
+	                .Font(IDetailLayoutBuilder::GetDetailFont())
+	            ]
+	        ]
+        ], OwnerTable);
+}
+
+FText SRigStackItem::GetIndexText() const
+{
+	const FString IndexStr = FString::FromInt(WeakStackEntry.Pin()->EntryIndex) + TEXT(".");
+	return FText::FromString(IndexStr);
+}
+
+FText SRigStackItem::GetLabelText() const
+{
+	return (FText::FromString(WeakStackEntry.Pin()->Label));
+}
+
+FSlateFontInfo SRigStackItem::GetLabelFont() const
+{
+	if(GetVisitedCountText().IsEmpty())
+	{
+		return IDetailLayoutBuilder::GetDetailFont();
+	}
+	return IDetailLayoutBuilder::GetDetailFontBold();
+}
+
+FText SRigStackItem::GetTooltip() const
+{
+	return FText::FromString(WeakStackEntry.Pin()->Callstack.GetCallPath(true));
+}
+
+FText SRigStackItem::GetVisitedCountText() const
+{
+	if(WeakStackEntry.IsValid() && WeakBlueprint.IsValid())
+	{
+		if(WeakBlueprint->RigGraphDisplaySettings.bShowNodeRunCounts)
+		{
+			if(WeakStackEntry.Pin()->EntryType == ERigStackEntry::Operator)
+			{
+				if(URigVMHost* RigVMHost = Cast<URigVMHost>(WeakBlueprint->GetObjectBeingDebugged()))
+				{
+					if(URigVM* VM = RigVMHost->GetVM())
+					{
+						const int32 Count = VM->GetInstructionVisitedCount(RigVMHost->GetExtendedExecuteContext(), WeakStackEntry.Pin()->InstructionIndex);
+						if(Count > 0)
+						{
+							return FText::FromString(FString::FromInt(Count));
+						}
+					}
+				}
+			}
+		}
+	}
+	return FText();
+}
+
+FText SRigStackItem::GetDurationText() const
+{
+	if(WeakStackEntry.IsValid() && WeakBlueprint.IsValid())
+	{
+		if(WeakBlueprint->VMRuntimeSettings.bEnableProfiling)
+		{
+			if(WeakStackEntry.Pin()->EntryType == ERigStackEntry::Operator)
+			{
+				if(URigVMHost* RigVMHost = Cast<URigVMHost>(WeakBlueprint->GetObjectBeingDebugged()))
+				{
+					if(URigVM* VM = RigVMHost->GetVM())
+					{
+						const double MicroSeconds = VM->GetInstructionMicroSeconds(RigVMHost->GetExtendedExecuteContext(), WeakStackEntry.Pin()->InstructionIndex);
+						if(MicroSeconds > 0.0)
+						{
+							return FText::FromString(FString::Printf(TEXT("%d µs"), (int32)MicroSeconds));
+						}
+					}
+				}
+			}
+		}
+	}
+	return FText();
+}
+
+//////////////////////////////////////////////////////////////
+/// SRigVMExecutionStackView
+///////////////////////////////////////////////////////////
+
+SRigVMExecutionStackView::~SRigVMExecutionStackView()
+{
+	if (RigVMEditor.IsValid())
+	{
+		if (OnModelModified.IsValid() && RigVMEditor.Pin()->GetRigVMBlueprint())
+		{
+			RigVMEditor.Pin()->GetRigVMBlueprint()->OnModified().Remove(OnModelModified);
+		}
+		if (OnHostInitializedHandle.IsValid())
+		{
+			if(RigVMEditor.Pin()->GetRigVMHost())
+			{
+				RigVMEditor.Pin()->GetRigVMHost()->OnInitialized_AnyThread().Remove(OnHostInitializedHandle);
+			}
+		}
+		if (OnPreviewHostUpdatedHandle.IsValid())
+		{
+			RigVMEditor.Pin()->OnPreviewHostUpdated().Remove(OnPreviewHostUpdatedHandle);
+		}
+	}
+	if (RigVMBlueprint.IsValid() && OnVMCompiledHandle.IsValid())
+	{
+		RigVMBlueprint->OnVMCompiled().Remove(OnVMCompiledHandle);
+	}
+}
+
+void SRigVMExecutionStackView::Construct( const FArguments& InArgs, TSharedRef<FRigVMEditor> InRigVMEditor)
+{
+	RigVMEditor = InRigVMEditor;
+	RigVMBlueprint = RigVMEditor.Pin()->GetRigVMBlueprint();
+	CommandList = MakeShared<FUICommandList>();
+	bSuspendModelNotifications = false;
+	bSuspendControllerSelection = false;
+	HaltedAtInstruction = INDEX_NONE;
+
+	BindCommands();
+
+	ChildSlot
+	[
+		SNew(SVerticalBox)
+		+SVerticalBox::Slot()
+		.AutoHeight()
+		.VAlign(VAlign_Top)
+		.Padding(0.0f)
+		[
+			SNew(SBorder)
+			.Padding(0.0f)
+			.BorderImage(FAppStyle::GetBrush("DetailsView.CategoryTop"))
+			.BorderBackgroundColor(FLinearColor(0.6f, 0.6f, 0.6f, 1.0f))
+			[
+				SNew(SVerticalBox)
+				+SVerticalBox::Slot()
+				.AutoHeight()
+				.VAlign(VAlign_Top)
+				[
+					SNew(SHorizontalBox)
+					//.Visibility(this, &SRigHierarchy::IsSearchbarVisible)
+					+SHorizontalBox::Slot()
+					.AutoWidth()
+					+SHorizontalBox::Slot()
+					.VAlign(VAlign_Center)
+					.Padding(3.0f, 1.0f)
+					[
+						SAssignNew(FilterBox, SSearchBox)
+						.OnTextChanged(this, &SRigVMExecutionStackView::OnFilterTextChanged)
+					]
+				]
+			]
+		]
+		+SVerticalBox::Slot()
+		.Padding(0.0f, 0.0f)
+		[
+			SNew(SBorder)
+			.Padding(2.0f)
+			.BorderImage(FAppStyle::GetBrush("SCSEditor.TreePanel"))
+			[
+				SAssignNew(TreeView, STreeView<TSharedPtr<FRigStackEntry>>)
+				.TreeItemsSource(&Operators)
+				.SelectionMode(ESelectionMode::Multi)
+				.OnGenerateRow(this, &SRigVMExecutionStackView::MakeTableRowWidget, RigVMBlueprint)
+				.OnGetChildren(this, &SRigVMExecutionStackView::HandleGetChildrenForTree)
+				.OnSelectionChanged(this, &SRigVMExecutionStackView::OnSelectionChanged)
+				.OnContextMenuOpening(this, &SRigVMExecutionStackView::CreateContextMenu)
+				.OnMouseButtonDoubleClick(this, &SRigVMExecutionStackView::HandleItemMouseDoubleClick)
+				.ItemHeight(28)
+			]
+		]
+	];
+
+	RefreshTreeView(nullptr);
+
+	if (RigVMBlueprint.IsValid())
+	{
+		if (OnVMCompiledHandle.IsValid())
+		{
+			RigVMBlueprint->OnVMCompiled().Remove(OnVMCompiledHandle);
+		}
+		if (OnModelModified.IsValid())
+		{
+			RigVMBlueprint->OnModified().Remove(OnModelModified);
+		}
+		OnVMCompiledHandle = RigVMBlueprint->OnVMCompiled().AddSP(this, &SRigVMExecutionStackView::OnVMCompiled);
+		OnModelModified = RigVMBlueprint->OnModified().AddSP(this, &SRigVMExecutionStackView::HandleModifiedEvent);
+
+		if (RigVMEditor.IsValid())
+		{
+			if (URigVMHost* RigVMHost = RigVMEditor.Pin()->GetRigVMHost())
+			{
+				OnVMCompiled(RigVMBlueprint.Get(), RigVMHost->GetVM());
+			}
+		}
+	}
+
+	if (RigVMEditor.IsValid())
+	{
+		OnPreviewHostUpdatedHandle = RigVMEditor.Pin()->OnPreviewHostUpdated().AddSP(this, &SRigVMExecutionStackView::HandlePreviewHostUpdated);
+	}
+}
+
+void SRigVMExecutionStackView::OnSelectionChanged(TSharedPtr<FRigStackEntry> Selection, ESelectInfo::Type SelectInfo)
+{
+	if (bSuspendModelNotifications || bSuspendControllerSelection)
+	{
+		return;
+	}
+	TGuardValue<bool> SuspendNotifs(bSuspendModelNotifications, true);
+
+	TArray<TSharedPtr<FRigStackEntry>> SelectedItems = TreeView->GetSelectedItems();
+	if (SelectedItems.Num() > 0)
+	{
+		if (!RigVMBlueprint.IsValid())
+		{
+			return;
+		}
+
+		URigVMBlueprintGeneratedClass* GeneratedClass = RigVMBlueprint->GetRigVMBlueprintGeneratedClass();
+		if (GeneratedClass == nullptr)
+		{
+			return;
+		}
+
+		URigVMHost* Host = RigVMEditor.Pin()->GetRigVMHost();
+		if (Host == nullptr || Host->GetVM() == nullptr)
+		{
+			return;
+		}
+
+		const FRigVMByteCode& ByteCode = Host->GetVM()->GetByteCode();
+
+		TMap<URigVMGraph*, TArray<FName>> SelectedNodesPerGraph;
+		for (TSharedPtr<FRigStackEntry>& Entry : SelectedItems)
+		{
+			for(int32 StackIndex = 0; StackIndex < Entry->Callstack.Num(); StackIndex++)
+			{
+				const UObject* Subject = Entry->Callstack[StackIndex];
+				URigVMGraph* SubjectGraph = nullptr;
+
+				FName NodeName = NAME_None;
+				if (const URigVMNode* Node = Cast<URigVMNode>(Subject))
+				{
+					NodeName = Node->GetFName();
+					SubjectGraph = Node->GetGraph();
+				}
+				else if (const URigVMPin* Pin = Cast<URigVMPin>(Subject))
+				{
+					NodeName = Pin->GetNode()->GetFName();
+					SubjectGraph = Pin->GetGraph();
+				}
+
+				if (NodeName.IsNone() || SubjectGraph == nullptr)
+				{
+					continue;
+				}
+
+				SelectedNodesPerGraph.FindOrAdd(SubjectGraph).AddUnique(NodeName);
+			}
+		}
+
+		for (const TPair< URigVMGraph*, TArray<FName> >& Pair : SelectedNodesPerGraph)
+		{
+			RigVMBlueprint->GetOrCreateController(Pair.Key)->SetNodeSelection(Pair.Value);
+		}
+	}
+}
+
+void SRigVMExecutionStackView::BindCommands()
+{
+	// create new command
+	const FRigVMExecutionStackCommands& Commands = FRigVMExecutionStackCommands::Get();
+	CommandList->MapAction(Commands.FocusOnSelection, FExecuteAction::CreateSP(this, &SRigVMExecutionStackView::HandleFocusOnSelectedGraphNode));
+	CommandList->MapAction(Commands.GoToInstruction, FExecuteAction::CreateSP(this, &SRigVMExecutionStackView::HandleGoToInstruction));
+}
+
+TSharedRef<ITableRow> SRigVMExecutionStackView::MakeTableRowWidget(TSharedPtr<FRigStackEntry> InItem, const TSharedRef<STableViewBase>& OwnerTable, TWeakObjectPtr<URigVMBlueprint> InBlueprint)
+{
+	return InItem->MakeTreeRowWidget(OwnerTable, InItem.ToSharedRef(), CommandList.ToSharedRef(), SharedThis(this), InBlueprint);
+}
+
+void SRigVMExecutionStackView::HandleGetChildrenForTree(TSharedPtr<FRigStackEntry> InItem, TArray<TSharedPtr<FRigStackEntry>>& OutChildren)
+{
+	OutChildren = InItem->Children;
+}
+
+void SRigVMExecutionStackView::PopulateStackView(URigVM* InVM)
+{
+	if (InVM)
+	{
+		URigVMBlueprint* Blueprint = RigVMEditor.Pin()->GetRigVMBlueprint();
+
+		const FRigVMInstructionArray Instructions = InVM->GetInstructions();
+		const FRigVMByteCode& ByteCode = InVM->GetByteCode();
+
+		TArray<URigVMGraph*> RootGraphs;
+		RootGraphs.AddZeroed(Instructions.Num());
+
+		const bool bUseSimpleLabels = !CVarRigVMExecutionStackDetailedLabels.GetValueOnAnyThread();
+		if(bUseSimpleLabels)
+		{
+			for (int32 InstructionIndex = 0; InstructionIndex < Instructions.Num(); InstructionIndex++)
+			{
+				URigVMNode* Node = Cast<URigVMNode>(ByteCode.GetSubjectForInstruction(InstructionIndex));
+
+				FString DisplayName;
+
+				if(Node)
+				{
+					DisplayName = Node->GetName();
+					RootGraphs[InstructionIndex] = Node->GetRootGraph();
+					
+					// only unit nodes among all nodes has StaticExecute() that generates actual instructions
+					if (URigVMUnitNode* UnitNode = Cast<URigVMUnitNode>(Node))
+					{
+						DisplayName = UnitNode->GetNodeTitle();
+	#if WITH_EDITOR
+						UScriptStruct* Struct = UnitNode->GetScriptStruct();
+						FString MenuDescSuffixMetadata;
+						if (Struct)
+						{
+							Struct->GetStringMetaDataHierarchical(FRigVMStruct::MenuDescSuffixMetaName, &MenuDescSuffixMetadata);
+						}
+						if (!MenuDescSuffixMetadata.IsEmpty())
+						{
+							DisplayName = FString::Printf(TEXT("%s %s"), *UnitNode->GetNodeTitle(), *MenuDescSuffixMetadata);
+						}
+
+						if(UnitNode->IsEvent())
+						{
+							DisplayName = Node->GetEventName().ToString();
+							if(!DisplayName.EndsWith(TEXT("Event")))
+							{
+								DisplayName += TEXT(" Event");
+							}
+						}
+	#endif
+					}
+				}
+
+				FString Label;
+
+				switch (Instructions[InstructionIndex].OpCode)
+				{
+					case ERigVMOpCode::Execute:
+					{
+						const FRigVMExecuteOp& Op = ByteCode.GetOpAt<FRigVMExecuteOp>(Instructions[InstructionIndex]);
+
+						if(Node)
+						{
+							Label = DisplayName;
+						}
+						else
+						{
+							Label = InVM->GetRigVMFunctionName(Op.FunctionIndex);
+						}
+						break;
+					}
+					case ERigVMOpCode::Copy:
+					case ERigVMOpCode::Zero:
+					case ERigVMOpCode::BoolFalse:
+					case ERigVMOpCode::BoolTrue:
+					case ERigVMOpCode::Increment:
+					case ERigVMOpCode::Decrement:
+					case ERigVMOpCode::Equals:
+					case ERigVMOpCode::NotEquals:
+					case ERigVMOpCode::JumpAbsolute:
+					case ERigVMOpCode::JumpForward:
+					case ERigVMOpCode::JumpBackward:
+					case ERigVMOpCode::JumpAbsoluteIf:
+					case ERigVMOpCode::JumpForwardIf:
+					case ERigVMOpCode::JumpBackwardIf:
+					case ERigVMOpCode::BeginBlock:
+					case ERigVMOpCode::EndBlock:
+					case ERigVMOpCode::ChangeType:
+					case ERigVMOpCode::ArrayReset:
+					case ERigVMOpCode::ArrayGetNum: 
+					case ERigVMOpCode::ArraySetNum:
+					case ERigVMOpCode::ArrayGetAtIndex:  
+					case ERigVMOpCode::ArraySetAtIndex:
+					case ERigVMOpCode::ArrayAdd:
+					case ERigVMOpCode::ArrayInsert:
+					case ERigVMOpCode::ArrayRemove:
+					case ERigVMOpCode::ArrayFind:
+					case ERigVMOpCode::ArrayAppend:
+					case ERigVMOpCode::ArrayClone:
+					case ERigVMOpCode::ArrayIterator:
+					case ERigVMOpCode::ArrayUnion:
+					case ERigVMOpCode::ArrayDifference:
+					case ERigVMOpCode::ArrayIntersection:
+					case ERigVMOpCode::ArrayReverse:
+					{
+						const FText OpCodeText = StaticEnum<ERigVMOpCode>()->GetDisplayNameTextByValue((int32)Instructions[InstructionIndex].OpCode);
+						if(Node)
+						{
+							Label = DisplayName + TEXT(" - ") + OpCodeText.ToString();
+						}
+						else
+						{
+							Label = OpCodeText.ToString();
+						}
+						break;
+					}
+					case ERigVMOpCode::InvokeEntry:
+					{
+						const FRigVMInvokeEntryOp& Op = ByteCode.GetOpAt<FRigVMInvokeEntryOp>(Instructions[InstructionIndex]);
+						Label = FString::Printf(TEXT("Run %s Event"), *Op.EntryName.ToString());
+						break;
+					}
+					case ERigVMOpCode::JumpToBranch:
+					{
+						const FRigVMJumpToBranchOp& Op = ByteCode.GetOpAt<FRigVMJumpToBranchOp>(Instructions[InstructionIndex]);
+						Label = TEXT("Jump To Branch");
+						break;
+					}
+					case ERigVMOpCode::Exit:
+					{
+						Label = TEXT("Exit");
+						break;
+					}
+					default:
+					{
+						ensure(false);
+						break;
+					}
+				}
+				
+				// add the entry with the new label to the stack view
+				const FRigVMInstruction& Instruction = Instructions[InstructionIndex];
+				if (FilterText.IsEmpty() || Label.Contains(FilterText.ToString()))
+				{
+					FRigVMASTProxy Proxy;
+					if(const TArray<UObject*>* Callstack = ByteCode.GetCallstackForInstruction(InstructionIndex))
+					{
+						Proxy = FRigVMASTProxy::MakeFromCallstack(Callstack);
+					}
+					TSharedPtr<FRigStackEntry> NewEntry = MakeShared<FRigStackEntry>(InstructionIndex, ERigStackEntry::Operator, InstructionIndex, Instruction.OpCode, Label, Proxy);
+					Operators.Add(NewEntry);
+				}
+			}
+			return;
+		}
+
+		// 1. cache information about instructions/nodes, which will be used later 
+		TMap<FString, FString> OperandFormatMap;
+		for (int32 InstructionIndex = 0; InstructionIndex < Instructions.Num(); InstructionIndex++)
+		{
+			const FRigVMASTProxy Proxy = FRigVMASTProxy::MakeFromCallPath(ByteCode.GetCallPathForInstruction(InstructionIndex), RootGraphs[InstructionIndex]);
+			if(URigVMNode* Node = Proxy.GetSubject<URigVMNode>())
+			{
+				FString DisplayName = Node->GetName();
+
+				// only unit nodes among all nodes has StaticExecute() that generates actual instructions
+				if (URigVMUnitNode* UnitNode = Cast<URigVMUnitNode>(Node))
+				{
+					DisplayName = UnitNode->GetNodeTitle();
+#if WITH_EDITOR
+					UScriptStruct* Struct = UnitNode->GetScriptStruct();
+					FString MenuDescSuffixMetadata;
+					if (Struct)
+					{
+						Struct->GetStringMetaDataHierarchical(FRigVMStruct::MenuDescSuffixMetaName, &MenuDescSuffixMetadata);
+					}
+					if (!MenuDescSuffixMetadata.IsEmpty())
+					{
+						DisplayName = FString::Printf(TEXT("%s %s"), *UnitNode->GetNodeTitle(), *MenuDescSuffixMetadata);
+					}
+#endif
+				}
+
+				// this is needed for name replacement later
+				OperandFormatMap.Add(Node->GetName(), DisplayName);
+			}
+		}
+
+		// 2. replace raw operand names with NodeTitle.PinName/PropertyName.OffsetName
+		TArray<FString> Labels = InVM->DumpByteCodeAsTextArray(TArray<int32>(), false, [OperandFormatMap](const FString& RegisterName, const FString& RegisterOffsetName)
+		{
+			FString NewRegisterName = RegisterName;
+			FString NodeName;
+			FString PinName;
+			if (RegisterName.Split(TEXT("."), &NodeName, &PinName))
+			{
+				const FString* NodeTitle = OperandFormatMap.Find(NodeName);
+				NewRegisterName = FString::Printf(TEXT("%s.%s"), NodeTitle ? **NodeTitle : *NodeName, *PinName);
+			}
+			FString OperandLabel;
+			OperandLabel = NewRegisterName;
+			if (!RegisterOffsetName.IsEmpty())
+			{
+				OperandLabel = FString::Printf(TEXT("%s.%s"), *OperandLabel, *RegisterOffsetName);
+			}
+			return OperandLabel;
+		});
+
+		ensure(Labels.Num() == Instructions.Num());
+		
+		// 3. replace instruction names with node titles
+		for (int32 InstructionIndex = 0; InstructionIndex < Labels.Num(); InstructionIndex++)
+		{
+			FString Label = Labels[InstructionIndex];
+			const FRigVMASTProxy Proxy = FRigVMASTProxy::MakeFromCallPath(ByteCode.GetCallPathForInstruction(InstructionIndex), RootGraphs[InstructionIndex]);
+
+			if(URigVMNode* Node = Proxy.GetSubject<URigVMNode>())
+			{
+				FString Suffix;
+				switch(Instructions[InstructionIndex].OpCode)
+				{
+					case ERigVMOpCode::Copy:
+					case ERigVMOpCode::Zero:
+	                case ERigVMOpCode::BoolFalse:
+	                case ERigVMOpCode::BoolTrue:
+	                case ERigVMOpCode::Increment:
+	                case ERigVMOpCode::Decrement:
+					case ERigVMOpCode::Equals:
+                	case ERigVMOpCode::NotEquals:
+					case ERigVMOpCode::JumpAbsolute:
+                	case ERigVMOpCode::JumpForward:
+                	case ERigVMOpCode::JumpBackward:
+					case ERigVMOpCode::JumpAbsoluteIf:
+                	case ERigVMOpCode::JumpForwardIf:
+                	case ERigVMOpCode::JumpBackwardIf:
+					case ERigVMOpCode::BeginBlock:
+					case ERigVMOpCode::EndBlock:
+					{
+						const FText OpCodeText = StaticEnum<ERigVMOpCode>()->GetDisplayNameTextByValue((int32)Instructions[InstructionIndex].OpCode);
+						Suffix = FString::Printf(TEXT(" - %s"), *OpCodeText.ToString());
+						break;
+					}
+					default:
+					{
+						break;
+					}
+				}
+				
+				Label =  OperandFormatMap[Node->GetName()] + Suffix;
+				// Label = Proxy.GetCallstack().GetCallPath() + Suffix;
+			}
+
+			// add the entry with the new label to the stack view
+			const FRigVMInstruction& Instruction = Instructions[InstructionIndex];
+			if (FilterText.IsEmpty() || Label.Contains(FilterText.ToString()))
+			{
+				TSharedPtr<FRigStackEntry> NewEntry = MakeShared<FRigStackEntry>(InstructionIndex, ERigStackEntry::Operator, InstructionIndex, Instruction.OpCode, Label, Proxy);
+				Operators.Add(NewEntry);
+			}
+		}
+	}
+}
+
+void SRigVMExecutionStackView::RefreshTreeView(URigVM* InVM)
+{
+	Operators.Reset();
+
+	// populate the stack with node names/instruction names
+	PopulateStackView(InVM);
+	
+	if (InVM)
+	{
+		// fill the children from the log
+		if (RigVMEditor.IsValid())
+		{
+			URigVMHost* Host = RigVMEditor.Pin()->GetRigVMHost();
+			if(Host && Host->GetLog())
+			{
+				const TArray<FRigVMLog::FLogEntry>& LogEntries = Host->GetLog()->Entries;
+				for (const FRigVMLog::FLogEntry& LogEntry : LogEntries)
+				{
+					if (Operators.Num() <= LogEntry.InstructionIndex)
+					{
+						continue;
+					}
+					int32 ChildIndex = Operators[LogEntry.InstructionIndex]->Children.Num();
+					switch (LogEntry.Severity)
+					{
+						case EMessageSeverity::Info:
+						{
+							Operators[LogEntry.InstructionIndex]->Children.Add(MakeShared<FRigStackEntry>(ChildIndex, ERigStackEntry::Info, LogEntry.InstructionIndex, ERigVMOpCode::Invalid, LogEntry.Message, FRigVMASTProxy()));
+							break;
+						}
+						case EMessageSeverity::Warning:
+						case EMessageSeverity::PerformanceWarning:
+						{
+							Operators[LogEntry.InstructionIndex]->Children.Add(MakeShared<FRigStackEntry>(ChildIndex, ERigStackEntry::Warning, LogEntry.InstructionIndex, ERigVMOpCode::Invalid, LogEntry.Message, FRigVMASTProxy()));
+							break;
+						}
+						case EMessageSeverity::Error:
+						{
+							Operators[LogEntry.InstructionIndex]->Children.Add(MakeShared<FRigStackEntry>(ChildIndex, ERigStackEntry::Error, LogEntry.InstructionIndex, ERigVMOpCode::Invalid, LogEntry.Message, FRigVMASTProxy()));
+							break;
+						}
+						default:
+						{
+							break;
+						}
+					}
+				}
+
+				Host->GetExtendedExecuteContext().ExecutionHalted().AddSP(this, &SRigVMExecutionStackView::HandleExecutionHalted);
+			}
+		}
+	}
+
+	TreeView->RequestTreeRefresh();
+}
+
+TSharedPtr< SWidget > SRigVMExecutionStackView::CreateContextMenu()
+{
+	TArray<TSharedPtr<FRigStackEntry>> SelectedItems = TreeView->GetSelectedItems();
+	if(SelectedItems.Num() == 0)
+	{
+		return nullptr;
+	}
+
+	const FRigVMExecutionStackCommands& Actions = FRigVMExecutionStackCommands::Get();
+
+	const bool CloseAfterSelection = true;
+	FMenuBuilder MenuBuilder(CloseAfterSelection, CommandList);
+	{
+		MenuBuilder.BeginSection("RigStackToolsAction", LOCTEXT("ToolsAction", "Tools"));
+		MenuBuilder.AddMenuEntry(Actions.FocusOnSelection);
+		MenuBuilder.AddMenuEntry(Actions.GoToInstruction);
+		MenuBuilder.EndSection();
+	}
+
+	return MenuBuilder.MakeWidget();
+}
+
+void SRigVMExecutionStackView::HandleFocusOnSelectedGraphNode()
+{
+	OnSelectionChanged(TSharedPtr<FRigStackEntry>(), ESelectInfo::Direct);
+
+	TArray<TSharedPtr<FRigStackEntry>> SelectedItems = TreeView->GetSelectedItems();
+	if (SelectedItems.Num() > 0)
+	{
+		URigVMHost* Host = RigVMEditor.Pin()->GetRigVMHost();
+		if (Host == nullptr || Host->GetVM() == nullptr)
+		{
+			return;
+		}
+
+		const FRigVMByteCode& ByteCode = Host->GetVM()->GetByteCode();
+		UObject* Subject = ByteCode.GetSubjectForInstruction(SelectedItems[0]->InstructionIndex);
+		if (URigVMNode* SelectedNode = Cast<URigVMNode>(Subject))
+		{
+			URigVMGraph* GraphToFocus = SelectedNode->GetGraph();
+			if (GraphToFocus && GraphToFocus->GetTypedOuter<URigVMAggregateNode>())
+			{
+				if(URigVMGraph* ParentGraph = GraphToFocus->GetParentGraph())
+				{
+					GraphToFocus = ParentGraph;
+				} 
+			}
+			
+			if (UEdGraph* EdGraph = RigVMBlueprint->GetEdGraph(GraphToFocus))
+			{
+				RigVMEditor.Pin()->OpenGraphAndBringToFront(EdGraph, true);
+				RigVMEditor.Pin()->ZoomToSelection_Clicked();
+				RigVMEditor.Pin()->HandleModifiedEvent(ERigVMGraphNotifType::NodeSelected, GraphToFocus, SelectedNode);
+			}
+		}
+	}
+}
+
+void SRigVMExecutionStackView::HandleGoToInstruction()
+{
+	// figure out the current instruction's index
+	int32 Index = 0;
+	TArray<TSharedPtr<FRigStackEntry>> SelectedItems = TreeView->GetSelectedItems();
+	if (SelectedItems.Num() > 0)
+	{
+		Index = SelectedItems[0]->InstructionIndex;
+	}
+	
+	const FVector2D MouseCursorLocation = FSlateApplication::Get().GetCursorPos();
+	
+	SWindow::FArguments WindowArguments;
+	WindowArguments.ScreenPosition(MouseCursorLocation);
+	WindowArguments.AutoCenter(EAutoCenter::None);
+	WindowArguments.FocusWhenFirstShown(true);
+
+	TSharedPtr<SNumericEntryBox<int32>> NumericBox;
+	const TSharedRef<SCustomDialog> OptionsDialog = SNew(SCustomDialog)
+	.Title(FText(LOCTEXT("GoToInstructionDialog", "Go to...")))
+	.WindowArguments(WindowArguments)
+	.Content()
+	[
+		SNew(SHorizontalBox)
+		+ SHorizontalBox::Slot()
+		.FillWidth(1)
+		.VAlign(VAlign_Center)
+		.HAlign(HAlign_Left)
+		[
+			SAssignNew(NumericBox, SNumericEntryBox<int32>)
+			.Font(IDetailLayoutBuilder::GetDetailFont())
+			.Value(Index)
+			.MinDesiredValueWidth(250)
+			.MinValue(0)
+			.MaxValue(Operators.Num() - 1)
+			.OnValueChanged_Lambda([&Index](const int32& InIndex)
+			{
+				Index = InIndex;
+			})
+			.IsEnabled(true)
+		]
+	]
+	.Buttons({
+		SCustomDialog::FButton(LOCTEXT("OK", "OK")),
+		SCustomDialog::FButton(LOCTEXT("Cancel", "Cancel"))
+	});
+
+	if(OptionsDialog->ShowModal() == 0)
+	{
+		TreeView->SetSelection(Operators[Index]);
+		TreeView->SetScrollOffset(FMath::Max(Index-5, 0));
+	}
+}
+
+void SRigVMExecutionStackView::OnVMCompiled(UObject* InCompiledObject, URigVM* InCompiledVM)
+{
+	RefreshTreeView(InCompiledVM);
+
+	if (RigVMEditor.IsValid() && !OnHostInitializedHandle.IsValid())
+	{
+		if(URigVMHost* Host = RigVMEditor.Pin()->GetRigVMHost())
+		{
+			OnHostInitializedHandle = Host->OnInitialized_AnyThread().AddSP(this, &SRigVMExecutionStackView::HandleHostInitializedEvent);
+		}
+	}
+}
+
+void SRigVMExecutionStackView::HandleExecutionHalted(const int32 InHaltedAtInstruction, UObject* InNode, const FName& InEntryName)
+{
+	if (HaltedAtInstruction == InHaltedAtInstruction)
+	{
+		return;
+	}
+
+	if (InHaltedAtInstruction == INDEX_NONE && InEntryName == RigVMEditor.Pin()->GetRigVMHost()->GetEventQueue().Last())
+	{
+		HaltedAtInstruction = InHaltedAtInstruction;
+		return;
+	}
+	
+	if (InHaltedAtInstruction != INDEX_NONE && Operators.Num() > InHaltedAtInstruction)
+	{
+		HaltedAtInstruction = InHaltedAtInstruction;
+		TreeView->SetSelection(Operators[InHaltedAtInstruction]);
+		TreeView->SetScrollOffset(FMath::Max(InHaltedAtInstruction-5, 0));
+	}	
+}
+
+void SRigVMExecutionStackView::OnFilterTextChanged(const FText& SearchText)
+{
+	FilterText = SearchText;
+	RefreshTreeView(RigVMEditor.Pin()->GetRigVMHost()->GetVM());
+}
+
+void SRigVMExecutionStackView::HandleModifiedEvent(ERigVMGraphNotifType InNotifType, URigVMGraph* InGraph, UObject* InSubject)
+{
+	if (bSuspendModelNotifications)
+	{
+		return;
+	}
+
+	URigVMHost* Host = RigVMEditor.Pin()->GetRigVMHost();
+	if (Host == nullptr || Host->GetVM() == nullptr)
+	{
+		return;
+	}
+
+	const FRigVMByteCode& ByteCode = Host->GetVM()->GetByteCode();
+
+	switch (InNotifType)
+	{
+		case ERigVMGraphNotifType::NodeSelectionChanged:
+		{
+			const TGuardValue<bool> SuspendNotifs(bSuspendModelNotifications, true);
+			TreeView->ClearSelection();
+			TArray<const URigVMNode*> SelectedNodes;
+			const TArray<FName>& SelectedNames = InGraph->GetSelectNodes();
+			for (const FName& Selected : SelectedNames)
+			{
+				if (const URigVMNode* Node = InGraph->FindNodeByName(Selected))
+				{
+					SelectedNodes.Add(Node);
+				}
+			}		
+
+			TArray<TSharedPtr<FRigStackEntry>> SelectedItems;
+			for (TSharedPtr<FRigStackEntry>& Operator : Operators)
+			{
+				for (const URigVMNode* Node : SelectedNodes)
+				{
+					if(Operator->Callstack.Contains(Node))
+					{
+						SelectedItems.Add(Operator);
+						break;
+					}
+				}
+			}
+			TreeView->SetItemSelection(SelectedItems, true, ESelectInfo::Direct);
+			break;
+		}
+		default:
+		{
+			break;
+		}
+	}
+}
+
+void SRigVMExecutionStackView::HandleHostInitializedEvent(URigVMHost* InHost, const FName& InEventName)
+{
+	TGuardValue<bool> SuspendControllerSelection(bSuspendControllerSelection, true);
+
+	RefreshTreeView(InHost->GetVM());
+	OnSelectionChanged(TSharedPtr<FRigStackEntry>(), ESelectInfo::Direct);
+
+	for (TSharedPtr<FRigStackEntry>& Operator : Operators)
+	{
+		for (TSharedPtr<FRigStackEntry>& Child : Operator->Children)
+		{
+			if (Child->EntryType == ERigStackEntry::Warning || Child->EntryType == ERigStackEntry::Error)
+			{
+				TreeView->SetItemExpansion(Operator, true);
+				break;
+			}
+		}
+	}
+	
+	if (RigVMEditor.IsValid())
+	{
+		InHost->OnInitialized_AnyThread().Remove(OnHostInitializedHandle);
+		OnHostInitializedHandle.Reset();
+
+		if (URigVMBlueprint* RigBlueprint = RigVMEditor.Pin()->GetRigVMBlueprint())
+		{
+			TArray<URigVMGraph*> Models = RigBlueprint->GetAllModels();
+			for (URigVMGraph* Model : Models)
+			{
+				for (URigVMNode* ModelNode : Model->GetNodes())
+				{
+					if (ModelNode->IsSelected())
+					{
+						HandleModifiedEvent(ERigVMGraphNotifType::NodeSelected, Model, ModelNode);
+					}
+				}
+			}
+		}
+	}
+}
+
+void SRigVMExecutionStackView::HandlePreviewHostUpdated(FRigVMEditor* InEditor)
+{
+}
+
+void SRigVMExecutionStackView::HandleItemMouseDoubleClick(TSharedPtr<FRigStackEntry> InItem)
+{
+	if (!RigVMEditor.IsValid() || !RigVMBlueprint.IsValid())
+	{
+		return;
+	}
+	
+	URigVMHost* Host = Cast<URigVMHost>(RigVMBlueprint->GetObjectBeingDebugged());
+	if (!Host || !Host->GetVM())
+	{
+		return;
+	}
+
+	const FRigVMByteCode& ByteCode = Host->GetVM()->GetByteCode();
+	if (URigVMNode* Subject = Cast<URigVMNode>(ByteCode.GetSubjectForInstruction(InItem->InstructionIndex)))
+	{
+		URigVMGraph* GraphToFocus = Subject->GetGraph();
+		if (GraphToFocus && GraphToFocus->GetTypedOuter<URigVMAggregateNode>())
+		{
+			if(URigVMGraph* ParentGraph = GraphToFocus->GetParentGraph())
+			{
+				GraphToFocus = ParentGraph;
+			} 
+		}
+		
+		if(URigVMEdGraph* EdGraph = Cast<URigVMEdGraph>(RigVMBlueprint->GetEdGraph(GraphToFocus)))
+		{
+			if(const UEdGraphNode* Node = EdGraph->FindNodeForModelNodeName(Subject->GetFName()))
+			{
+				RigVMEditor.Pin()->JumpToHyperlink(Node, false);
+			}
+		}
+	}
+}
+
+#undef LOCTEXT_NAMESPACE
diff --git a/Engine/Plugins/Runtime/RigVM/Source/RigVMEditor/Private/Widgets/SRigVMGraphBreakLinksWidget.cpp b/Engine/Plugins/Runtime/RigVM/Source/RigVMEditor/Private/Widgets/SRigVMGraphBreakLinksWidget.cpp
new file mode 100644
--- /dev/null
+++ b/Engine/Plugins/Runtime/RigVM/Source/RigVMEditor/Private/Widgets/SRigVMGraphBreakLinksWidget.cpp
@@ -0,0 +1,208 @@
+﻿// Copyright Epic Games, Inc. All Rights Reserved.
+
+#include "Widgets/SRigVMGraphBreakLinksWidget.h"
+#include "Widgets/Layout/SBorder.h"
+#include "Widgets/SBoxPanel.h"
+#include "Widgets/Input/SButton.h"
+#include "Widgets/Input/SCheckBox.h"
+#include "Widgets/Layout/SUniformGridPanel.h"
+#include "DetailLayoutBuilder.h"
+#include "Editor.h"
+#include "Styling/AppStyle.h"
+#include "HAL/ConsoleManager.h"
+#include "RigVMModel/RigVMController.h"
+#include "Editor/RigVMEditor.h"
+#include "ContentBrowserModule.h"
+#include "IContentBrowserSingleton.h"
+#include "Misc/ScopedSlowTask.h"
+
+#define LOCTEXT_NAMESPACE "SRigVMGraphBreakLinksWidget"
+
+class SRigVMGraphBreakLinksWidgetRow : public SMultiColumnTableRow< URigVMLink* >
+{
+	SLATE_BEGIN_ARGS(SRigVMGraphBreakLinksWidgetRow) { }
+	SLATE_ARGUMENT(URigVMLink*, Item)
+	SLATE_END_ARGS()
+	
+	void Construct( const FArguments& InArgs, const TSharedRef<STableViewBase>& InOwnerTableView)
+	{
+		Link = InArgs._Item;
+
+		SMultiColumnTableRow<URigVMLink*>::Construct(FSuperRowType::FArguments(), InOwnerTableView);
+	}
+
+	virtual TSharedRef<SWidget> GenerateWidgetForColumn(const FName& InColumnName) override
+	{
+		if(InColumnName == TEXT("Source"))
+		{
+			return SNew(SBox)
+				.Padding(FMargin(4.0f, 0.0f))
+				.VAlign(VAlign_Center)
+				[
+					SNew(STextBlock)
+					.Text(FText::FromString(Link->GetSourcePin()->GetPinPath()))
+				];
+		}
+		else if (InColumnName == TEXT("Target") )
+		{
+			return SNew(SBox)
+				.Padding(FMargin(4.0f, 0.0f))
+				.VAlign(VAlign_Center)
+				[
+					SNew(STextBlock)
+					.Text(FText::FromString(Link->GetTargetPin()->GetPinPath()))
+				];
+		}
+		else
+		{
+			ensure(false);
+			return SNew(SBorder);
+		}
+	}
+
+private:
+	URigVMLink* Link;
+};
+
+//////////////////////////////////////////////////////////////
+/// SRigVMGraphBreakLinksWidget
+///////////////////////////////////////////////////////////
+
+void SRigVMGraphBreakLinksWidget::Construct(const FArguments& InArgs, TArray<URigVMLink*> InLinks)
+{
+	Links = InLinks;
+	OnFocusOnLink = InArgs._OnFocusOnLink;
+
+	ChildSlot
+	[
+		SNew(SVerticalBox)
+
+		+ SVerticalBox::Slot()
+		.Padding(4.0f, 4.0f, 0.f, 0.f)
+		.AutoHeight()
+		.MaxHeight(600.f)
+		[
+			SNew(SBorder)
+			.Visibility(EVisibility::Visible)
+			.BorderImage(FAppStyle::GetBrush("Menu.Background"))
+			[			
+			   SNew( SListView< URigVMLink* > )
+			     .ItemHeight(24)
+			     .ListItemsSource( &Links )
+			     .ListViewStyle( FAppStyle::Get(), "SimpleListView" )
+			     .OnMouseButtonDoubleClick(this, &SRigVMGraphBreakLinksWidget::HandleItemMouseDoubleClick)
+			     .HeaderRow(SNew(SHeaderRow)
+			     	+SHeaderRow::Column(FName(TEXT("Source")))
+			     	.DefaultLabel(LOCTEXT("SourceColumnHeader", "Source"))
+			     	.FillWidth(0.5f)
+			     	+SHeaderRow::Column(FName(TEXT("Target")))
+			     	.DefaultLabel(LOCTEXT("TargetColumnHeader", "Target"))
+			     	.FillWidth(0.5f)
+			     )
+			     .OnGenerateRow(this, &SRigVMGraphBreakLinksWidget::GenerateItemRow)
+			]
+		]
+	];
+}
+
+TSharedRef<ITableRow> SRigVMGraphBreakLinksWidget::GenerateItemRow(URigVMLink* Item, const TSharedRef<STableViewBase>& OwnerTable)
+{
+	return SNew(SRigVMGraphBreakLinksWidgetRow, OwnerTable)
+			.Item(Item);
+}
+
+void SRigVMGraphBreakLinksWidget::HandleItemMouseDoubleClick(URigVMLink* InItem)
+{
+	OnFocusOnLink.ExecuteIfBound(InItem);		
+}
+
+void SRigVMGraphBreakLinksDialog::Construct(const FArguments& InArgs)
+{
+	UserResponse = EAppReturnType::Cancel;
+
+	static const FString DialogText = TEXT(
+        "This action will break links." 
+    );
+
+	const FString DialogTitle = TEXT("Break Links"); 
+	
+	SWindow::Construct(SWindow::FArguments()
+        .Title(FText::FromString(DialogTitle))
+        .SizingRule( ESizingRule::Autosized )
+        //.ClientSize(FVector2D(450, 450))
+        [
+	        SNew(SHorizontalBox)
+
+	        +SHorizontalBox::Slot()
+	        .FillWidth(1.f)
+	        .HAlign(HAlign_Left)
+	        .VAlign(VAlign_Top)
+	        //.Padding(2)
+	        [
+				SNew(SVerticalBox)
+
+                +SVerticalBox::Slot()
+                .AutoHeight()
+                .HAlign(HAlign_Left)
+                //.Padding(8)
+                [
+
+                    SNew(STextBlock)
+                    .Text(FText::FromString(DialogText))
+                    .Font(IDetailLayoutBuilder::GetDetailFont())
+                ]
+				
+                +SVerticalBox::Slot()
+                .AutoHeight()
+                .HAlign(HAlign_Left)
+                .Padding(2)
+                [
+                    SAssignNew(BreakLinksWidget, SRigVMGraphBreakLinksWidget, InArgs._Links)
+                    .OnFocusOnLink(InArgs._OnFocusOnLink)
+                ]
+
+                +SVerticalBox::Slot()
+                .AutoHeight()
+                .HAlign(HAlign_Right)
+                .Padding(5)
+                [
+                    SNew(SUniformGridPanel)
+                    .SlotPadding(FAppStyle::GetMargin("StandardDialog.SlotPadding"))
+                    .MinDesiredSlotWidth(FAppStyle::GetFloat("StandardDialog.MinDesiredSlotWidth"))
+                    .MinDesiredSlotHeight(FAppStyle::GetFloat("StandardDialog.MinDesiredSlotHeight"))
+                    +SUniformGridPanel::Slot(0, 0)
+                    [
+                        SNew(SButton)
+                        .HAlign(HAlign_Center)
+                        .ContentPadding(FAppStyle::GetMargin("StandardDialog.ContentPadding"))
+                        .Text(LOCTEXT("OK", "OK"))
+                        .OnClicked(this, &SRigVMGraphBreakLinksDialog::OnButtonClick, EAppReturnType::Ok)
+                    ]
+                    +SUniformGridPanel::Slot(1, 0)
+                    [
+                        SNew(SButton)
+                        .HAlign(HAlign_Center)
+                        .ContentPadding(FAppStyle::GetMargin("StandardDialog.ContentPadding"))
+                        .Text(LOCTEXT("Cancel", "Cancel"))
+                        .OnClicked(this, &SRigVMGraphBreakLinksDialog::OnButtonClick, EAppReturnType::Cancel)
+                    ]
+                ]
+			]
+		]);
+}
+
+FReply SRigVMGraphBreakLinksDialog::OnButtonClick(EAppReturnType::Type ButtonID)
+{
+	UserResponse = ButtonID;
+
+	RequestDestroyWindow();
+	return FReply::Handled();
+}
+
+EAppReturnType::Type SRigVMGraphBreakLinksDialog::ShowModal()
+{
+	GEditor->EditorAddModalWindow(SharedThis(this));
+	return UserResponse;
+}
+
+#undef LOCTEXT_NAMESPACE
diff --git a/Engine/Plugins/Runtime/RigVM/Source/RigVMEditor/Private/Widgets/SRigVMGraphFunctionBulkEditWidget.cpp b/Engine/Plugins/Runtime/RigVM/Source/RigVMEditor/Private/Widgets/SRigVMGraphFunctionBulkEditWidget.cpp
new file mode 100644
--- /dev/null
+++ b/Engine/Plugins/Runtime/RigVM/Source/RigVMEditor/Private/Widgets/SRigVMGraphFunctionBulkEditWidget.cpp
@@ -0,0 +1,223 @@
+// Copyright Epic Games, Inc. All Rights Reserved.
+
+#include "Widgets/SRigVMGraphFunctionBulkEditWidget.h"
+#include "Widgets/Layout/SBorder.h"
+#include "Widgets/SBoxPanel.h"
+#include "Widgets/Input/SButton.h"
+#include "Widgets/Input/SCheckBox.h"
+#include "Widgets/Layout/SUniformGridPanel.h"
+#include "Widgets/Text/STextBlock.h"
+#include "DetailLayoutBuilder.h"
+#include "Editor.h"
+#include "Styling/AppStyle.h"
+#include "HAL/ConsoleManager.h"
+#include "RigVMModel/Nodes/RigVMFunctionReferenceNode.h"
+#include "RigVMModel/RigVMController.h"
+#include "ContentBrowserModule.h"
+#include "IContentBrowserSingleton.h"
+#include "Misc/ScopedSlowTask.h"
+
+#define LOCTEXT_NAMESPACE "SRigVMGraphFunctionBulkEditWidget"
+
+//////////////////////////////////////////////////////////////
+/// SRigVMGraphFunctionBulkEditWidget
+///////////////////////////////////////////////////////////
+
+void SRigVMGraphFunctionBulkEditWidget::Construct(const FArguments& InArgs, URigVMBlueprint* InBlueprint, URigVMController* InController, URigVMLibraryNode* InFunction, ERigVMControllerBulkEditType InEditType)
+{
+	Blueprint = InBlueprint;
+	Controller = InController;
+	Function = InFunction;
+	EditType = InEditType;
+
+	ChildSlot
+	[
+		SNew(SVerticalBox)
+
+		+ SVerticalBox::Slot()
+		.Padding(4.0f, 4.0f, 0.f, 0.f)
+		.AutoHeight()
+		.MaxHeight(600.f)
+		[
+			SNew(SBorder)
+			.Visibility(EVisibility::Visible)
+			.BorderImage(FAppStyle::GetBrush("Menu.Background"))
+			[
+				MakeAssetViewForReferencedAssets()
+			]
+		]
+	];
+}
+
+TSharedPtr<SWidget> SRigVMGraphFunctionBulkEditWidget::OnGetAssetContextMenu(const TArray<FAssetData>& SelectedAssets)
+{
+	return nullptr;
+}
+
+void SRigVMGraphFunctionBulkEditWidget::OnAssetsActivated(const TArray<FAssetData>& ActivatedAssets,
+	EAssetTypeActivationMethod::Type ActivationMethod)
+{
+	// nothing to do here
+}
+
+void SRigVMGraphFunctionBulkEditWidget::LoadAffectedAssets()
+{
+	TArray<FAssetData> FirstLevelReferenceAssets = Controller->GetAffectedAssets(EditType, false);
+
+	{
+		FScopedSlowTask SlowTask((float)FirstLevelReferenceAssets.Num(), LOCTEXT("LoadingAffectedAssets", "Loading Affected Assets"));
+		Controller->OnBulkEditProgressDelegate.BindLambda([&SlowTask](
+            TSoftObjectPtr<URigVMFunctionReferenceNode> InReference,
+            ERigVMControllerBulkEditType InEditType,
+            ERigVMControllerBulkEditProgress InProgress,
+            int32 InIndex,
+            int32 InCount)
+        {
+            if(!InReference.IsValid())
+            {
+                return;
+            }
+
+			const FString PackageName = InReference.ToSoftObjectPath().GetLongPackageName();
+			const FText Message = FText::FromString(FString::Printf(TEXT("Loading '%s' ..."), *PackageName));
+
+			if (InIndex == 0)
+			{
+				SlowTask.TotalAmountOfWork += (float)InCount;
+			}
+			SlowTask.EnterProgressFrame(1.f, Message);
+        });
+		
+		SlowTask.MakeDialog();
+
+		Controller->GetAffectedAssets(EditType, true);
+	}
+
+	Controller->OnBulkEditProgressDelegate.Unbind();
+}
+
+TSharedRef<SWidget> SRigVMGraphFunctionBulkEditWidget::MakeAssetViewForReferencedAssets()
+{
+	FContentBrowserModule& ContentBrowserModule = FModuleManager::Get().LoadModuleChecked<FContentBrowserModule>(TEXT("ContentBrowser"));
+	FAssetPickerConfig AssetPickerConfig;
+
+	AssetPickerConfig.bAllowDragging = false;
+	AssetPickerConfig.bCanShowClasses = false;
+	AssetPickerConfig.bAllowNullSelection = false;
+	AssetPickerConfig.bShowBottomToolbar = false;
+	AssetPickerConfig.bAutohideSearchBar = true;
+	AssetPickerConfig.Filter.ClassPaths.Add(FTopLevelAssetPath(TEXT("/Script/Engine"), TEXT("InvalidClass")));
+	AssetPickerConfig.Filter.bIncludeOnlyOnDiskAssets = false;
+	AssetPickerConfig.OnAssetsActivated = FOnAssetsActivated::CreateSP(this, &SRigVMGraphFunctionBulkEditWidget::OnAssetsActivated);
+	AssetPickerConfig.OnGetAssetContextMenu = FOnGetAssetContextMenu::CreateSP( this, &SRigVMGraphFunctionBulkEditWidget::OnGetAssetContextMenu );
+
+	AssetPickerConfig.InitialAssetViewType = EAssetViewType::List;
+	AssetPickerConfig.OnGetCustomSourceAssets.BindLambda([this](const FARFilter& SourceFilter, TArray<FAssetData>& AddedAssets)
+	{
+		AddedAssets.Append(Controller->GetAffectedAssets(EditType, false));
+	});
+	
+	return ContentBrowserModule.Get().CreateAssetPicker(AssetPickerConfig);
+}
+
+
+void SRigVMGraphFunctionBulkEditDialog::Construct(const FArguments& InArgs)
+{
+	UserResponse = EAppReturnType::Cancel;
+
+	static const FString DialogText = TEXT(
+        "This action will affect other assets.\n\n" 
+        "Note that this action is not undoable.\n\n"
+        "All affected assets will be marked as modified\n"
+        "and require re-saving."
+    );
+
+	const FString DialogTitle = FString::Printf(TEXT("Preparing Bulk-Edit - %s"), *StaticEnum<ERigVMControllerBulkEditType>()->GetDisplayNameTextByValue((int64)InArgs._EditType).ToString()); 
+	
+	SWindow::Construct(SWindow::FArguments()
+        .Title(FText::FromString(DialogTitle))
+        .SupportsMinimize(false)
+        .SupportsMaximize(false)
+        .SizingRule( ESizingRule::Autosized )
+        .ClientSize(FVector2D(450, 450))
+        [
+	        SNew(SHorizontalBox)
+
+	        +SHorizontalBox::Slot()
+	        .AutoWidth()
+	        .HAlign(HAlign_Left)
+	        .VAlign(VAlign_Top)
+	        .Padding(2)
+	        [
+				SNew(SVerticalBox)
+
+                +SVerticalBox::Slot()
+                .AutoHeight()
+                .HAlign(HAlign_Left)
+                .Padding(8)
+                [
+
+                    SNew(STextBlock)
+                    .Text(FText::FromString(DialogText))
+                    .Font(IDetailLayoutBuilder::GetDetailFont())
+                ]
+				
+                +SVerticalBox::Slot()
+                .AutoHeight()
+                .HAlign(HAlign_Left)
+                .Padding(2)
+                [
+                    SAssignNew(BulkEditWidget, SRigVMGraphFunctionBulkEditWidget, InArgs._Blueprint, InArgs._Controller, InArgs._Function, InArgs._EditType)
+                ]
+
+                +SVerticalBox::Slot()
+                .AutoHeight()
+                .HAlign(HAlign_Right)
+                .Padding(5)
+                [
+                    SNew(SUniformGridPanel)
+                    .SlotPadding(FAppStyle::GetMargin("StandardDialog.SlotPadding"))
+                    .MinDesiredSlotWidth(FAppStyle::GetFloat("StandardDialog.MinDesiredSlotWidth"))
+                    .MinDesiredSlotHeight(FAppStyle::GetFloat("StandardDialog.MinDesiredSlotHeight"))
+                    +SUniformGridPanel::Slot(0, 0)
+                    [
+                        SNew(SButton)
+                        .HAlign(HAlign_Center)
+                        .ContentPadding(FAppStyle::GetMargin("StandardDialog.ContentPadding"))
+                        .Text(LOCTEXT("OK", "OK"))
+                        .OnClicked(this, &SRigVMGraphFunctionBulkEditDialog::OnButtonClick, EAppReturnType::Ok)
+                    ]
+                    +SUniformGridPanel::Slot(1, 0)
+                    [
+                        SNew(SButton)
+                        .HAlign(HAlign_Center)
+                        .ContentPadding(FAppStyle::GetMargin("StandardDialog.ContentPadding"))
+                        .Text(LOCTEXT("Cancel", "Cancel"))
+                        .OnClicked(this, &SRigVMGraphFunctionBulkEditDialog::OnButtonClick, EAppReturnType::Cancel)
+                    ]
+                ]
+			]
+		]);
+}
+
+FReply SRigVMGraphFunctionBulkEditDialog::OnButtonClick(EAppReturnType::Type ButtonID)
+{
+	UserResponse = ButtonID;
+
+	if(UserResponse == EAppReturnType::Ok)
+	{
+		BulkEditWidget->LoadAffectedAssets();
+	}
+	
+	RequestDestroyWindow();
+
+	return FReply::Handled();
+}
+
+EAppReturnType::Type SRigVMGraphFunctionBulkEditDialog::ShowModal()
+{
+	GEditor->EditorAddModalWindow(SharedThis(this));
+	return UserResponse;
+}
+
+#undef LOCTEXT_NAMESPACE
diff --git a/Engine/Plugins/Runtime/RigVM/Source/RigVMEditor/Private/Widgets/SRigVMGraphFunctionLocalizationWidget.cpp b/Engine/Plugins/Runtime/RigVM/Source/RigVMEditor/Private/Widgets/SRigVMGraphFunctionLocalizationWidget.cpp
new file mode 100644
--- /dev/null
+++ b/Engine/Plugins/Runtime/RigVM/Source/RigVMEditor/Private/Widgets/SRigVMGraphFunctionLocalizationWidget.cpp
@@ -0,0 +1,306 @@
+// Copyright Epic Games, Inc. All Rights Reserved.
+
+#include "Widgets/SRigVMGraphFunctionLocalizationWidget.h"
+#include "Widgets/Layout/SBorder.h"
+#include "Widgets/SBoxPanel.h"
+#include "Widgets/Input/SButton.h"
+#include "Widgets/Input/SCheckBox.h"
+#include "Widgets/Layout/SUniformGridPanel.h"
+#include "DetailLayoutBuilder.h"
+#include "Editor.h"
+#include "Styling/AppStyle.h"
+#include "HAL/ConsoleManager.h"
+#include "RigVMModel/Nodes/RigVMFunctionReferenceNode.h"
+
+#define LOCTEXT_NAMESPACE "SRigVMGraphFunctionLocalizationWidget"
+
+//////////////////////////////////////////////////////////////
+/// SRigVMGraphFunctionLocalizationItem
+///////////////////////////////////////////////////////////
+
+SRigVMGraphFunctionLocalizationItem::SRigVMGraphFunctionLocalizationItem(const FRigVMGraphFunctionIdentifier& InFunction)
+    : Function(InFunction)
+{
+	FString OuterName;
+	FString FunctionName;
+	bool bIsPublic = false;
+
+	if (FRigVMGraphFunctionData* FunctionData = FRigVMGraphFunctionData::FindFunctionData(InFunction, &bIsPublic))
+	{
+		OuterName = InFunction.LibraryNode.GetAssetPathString();
+		FunctionName = FunctionData->Header.Name.ToString();
+		DisplayText = FText::FromString(FString::Printf(TEXT("%s :: %s"), *OuterName, *FunctionName));
+		ToolTipText = bIsPublic ?
+			FText::FromString(FString::Printf(TEXT("%s :: %s is public. Check this to create a local copy."), *OuterName, *FunctionName)) :
+			FText::FromString(FString::Printf(TEXT("%s :: %s is private. A local copy has to be created. To avoid this you can turn it public in the source Control Rig."), *OuterName, *FunctionName));
+	}
+}
+
+//////////////////////////////////////////////////////////////
+/// SRigVMGraphFunctionLocalizationTableRow
+///////////////////////////////////////////////////////////
+
+void SRigVMGraphFunctionLocalizationTableRow::Construct(const FArguments& InArgs, const TSharedRef<STableViewBase>& OwnerTable, SRigVMGraphFunctionLocalizationWidget* InLocalizationWidget, TSharedRef<SRigVMGraphFunctionLocalizationItem> InFunctionItem)
+{
+	STableRow<TSharedPtr<SRigVMGraphFunctionLocalizationItem>>::Construct(
+		STableRow<TSharedPtr<SRigVMGraphFunctionLocalizationItem>>::FArguments()
+		.Content()
+		[
+			SNew(SVerticalBox)
+
+			+ SVerticalBox::Slot()
+			.Padding(4.0f, 4.0f, 0.f, 0.f)
+			.AutoHeight()
+			[
+				SNew(SHorizontalBox)
+
+				+ SHorizontalBox::Slot()
+				.AutoWidth()
+				.HAlign(HAlign_Left)
+				.VAlign(VAlign_Center)
+				.Padding(0.0f, 0.0f, 12.f, 0.f)
+				[
+					SNew(STextBlock)
+					.Text(InFunctionItem->DisplayText)
+					.ToolTipText(InFunctionItem->ToolTipText)
+				]
+
+				+ SHorizontalBox::Slot()
+				.AutoWidth()
+				.HAlign(HAlign_Left)
+				.VAlign(VAlign_Center)
+				[
+					SNew(SCheckBox)
+					.IsChecked(InLocalizationWidget, &SRigVMGraphFunctionLocalizationWidget::IsFunctionEnabled, InFunctionItem->Function)
+					.OnCheckStateChanged(InLocalizationWidget, &SRigVMGraphFunctionLocalizationWidget::SetFunctionEnabled, InFunctionItem->Function)
+					.IsEnabled(InLocalizationWidget, &SRigVMGraphFunctionLocalizationWidget::IsFunctionPublic, InFunctionItem->Function)
+					.ToolTipText(InFunctionItem->ToolTipText)
+				]
+			]
+		]
+		, OwnerTable
+	);
+}
+
+//////////////////////////////////////////////////////////////
+/// SRigVMGraphFunctionLocalizationWidget
+///////////////////////////////////////////////////////////
+
+void SRigVMGraphFunctionLocalizationWidget::Construct(const FArguments& InArgs, const FRigVMGraphFunctionIdentifier& InFunctionToLocalize, URigVMBlueprint* InTargetBlueprint)
+{
+	FunctionsToLocalize.Add(InFunctionToLocalize);
+	FunctionItems.Reset();
+
+	TArray<FRigVMGraphFunctionIdentifier> NodesToVisit;
+	TArray<FRigVMGraphFunctionData*> FunctionsForTable;
+
+	{
+		FRigVMGraphFunctionData* FunctionData = FRigVMGraphFunctionData::FindFunctionData(InFunctionToLocalize);
+		FunctionsForTable.Add(FunctionData);
+		NodesToVisit.Add(InFunctionToLocalize);
+	}
+
+	IRigVMGraphFunctionHost* TargetFunctionHost = InTargetBlueprint->GetRigVMGraphFunctionHost();
+	for(int32 NodeToVisitIndex=0; NodeToVisitIndex<NodesToVisit.Num(); NodeToVisitIndex++)
+	{
+		const FRigVMGraphFunctionIdentifier& NodeToVisit = NodesToVisit[NodeToVisitIndex];
+
+		bool bIsPublic;
+		FRigVMGraphFunctionData* FunctionData = FRigVMGraphFunctionData::FindFunctionData(NodeToVisit, &bIsPublic);
+		if (!FunctionData)
+		{
+			continue;
+		}
+
+		for (TPair<FRigVMGraphFunctionIdentifier, uint32>& Pair : FunctionData->Header.Dependencies)
+		{
+			NodesToVisit.AddUnique(Pair.Key);
+		}
+		
+		if (NodeToVisit.HostObject != Cast<UObject>(TargetFunctionHost))
+		{
+			FunctionsForTable.AddUnique(FunctionData);
+			if (bIsPublic)
+			{
+				FunctionsToLocalize.AddUnique(NodeToVisit);
+			}
+		}
+	} 
+
+	// sort the functions to localize based on their nesting
+	Algo::Sort(FunctionsForTable, [](FRigVMGraphFunctionData* A, FRigVMGraphFunctionData* B) -> bool
+    {
+		check(A);
+		check(B);
+		return A->Header.Dependencies.Contains(B->Header.LibraryPointer);
+    });
+
+	for (FRigVMGraphFunctionData* FunctionForTable : FunctionsForTable)
+	{
+		FunctionItems.Add(MakeShared<SRigVMGraphFunctionLocalizationItem>(FunctionForTable->Header.LibraryPointer));
+	}
+
+	ChildSlot
+		[
+			SNew(SVerticalBox)
+
+			+ SVerticalBox::Slot()
+			.Padding(4.0f, 4.0f, 0.f, 0.f)
+			.AutoHeight()
+			[
+				SNew(SBorder)
+				.Visibility(EVisibility::Visible)
+				.BorderImage(FAppStyle::GetBrush("Menu.Background"))
+				[
+					SNew(SListView<TSharedPtr<SRigVMGraphFunctionLocalizationItem>>)
+					.ListItemsSource(&FunctionItems)
+					.OnGenerateRow(this, &SRigVMGraphFunctionLocalizationWidget::GenerateFunctionListRow)
+					.SelectionMode(ESelectionMode::None)
+				]
+			]
+		];
+}
+
+TSharedRef<ITableRow> SRigVMGraphFunctionLocalizationWidget::GenerateFunctionListRow(TSharedPtr<SRigVMGraphFunctionLocalizationItem> InItem, const TSharedRef<STableViewBase>& InOwningTable)
+{
+	TSharedRef<SRigVMGraphFunctionLocalizationTableRow> TableRow = SNew(SRigVMGraphFunctionLocalizationTableRow, InOwningTable, this, InItem.ToSharedRef());
+	TableRows.Add(InItem->Function, TableRow);
+	return TableRow;
+}
+
+ECheckBoxState SRigVMGraphFunctionLocalizationWidget::IsFunctionEnabled(const FRigVMGraphFunctionIdentifier InFunction) const
+{
+	return FunctionsToLocalize.Contains(InFunction) ? ECheckBoxState::Checked : ECheckBoxState::Unchecked;
+}
+
+void SRigVMGraphFunctionLocalizationWidget::SetFunctionEnabled(ECheckBoxState NewState, const FRigVMGraphFunctionIdentifier InFunction)
+{
+	if(NewState == ECheckBoxState::Checked)
+	{
+		FunctionsToLocalize.AddUnique(InFunction);
+	}
+	else
+	{
+		FunctionsToLocalize.Remove(InFunction);
+	}
+}
+
+bool SRigVMGraphFunctionLocalizationWidget::IsFunctionPublic(const FRigVMGraphFunctionIdentifier InFunction) const
+{
+	IRigVMGraphFunctionHost* FunctionHost = nullptr;
+	if (UObject* FunctionHostObj = InFunction.HostObject.TryLoad())
+	{
+		FunctionHost = Cast<IRigVMGraphFunctionHost>(FunctionHostObj);									
+	}
+	if (FunctionHost)
+	{
+		return FunctionHost->GetRigVMGraphFunctionStore()->IsFunctionPublic(InFunction);
+	}
+			
+	return false;	
+}
+
+void SRigVMGraphFunctionLocalizationDialog::Construct(const FArguments& InArgs)
+{
+	UserResponse = EAppReturnType::Cancel;
+
+	static const FString DialogText = TEXT(
+        "This action requires a function to be localized.\n\n" 
+        "Check the functions you want to create copies of below.\n\n"
+        "Note that private functions have to be copied,\n"
+        "while public functions can be referenced.\n\n"
+        "Checking all functions removes the dependencies completely."
+    );
+	
+	SWindow::Construct(SWindow::FArguments()
+        .Title(LOCTEXT("SRigVMGraphFunctionLocalizationDialog_Title", "Select function(s) to localize"))
+        .SupportsMinimize(false)
+        .SupportsMaximize(false)
+        .SizingRule( ESizingRule::Autosized )
+        .ClientSize(FVector2D(450, 450))
+        [
+	        SNew(SHorizontalBox)
+
+	        +SHorizontalBox::Slot()
+	        .AutoWidth()
+	        .HAlign(HAlign_Left)
+	        .VAlign(VAlign_Top)
+	        .Padding(2)
+	        [
+				SNew(SVerticalBox)
+
+                +SVerticalBox::Slot()
+                .AutoHeight()
+                .HAlign(HAlign_Left)
+                .Padding(8)
+                [
+
+                    SNew(STextBlock)
+                    .Text(FText::FromString(DialogText))
+                    .Font(IDetailLayoutBuilder::GetDetailFont())
+                ]
+				
+                +SVerticalBox::Slot()
+                .AutoHeight()
+                .HAlign(HAlign_Left)
+                .Padding(2)
+                [
+                    SAssignNew(FunctionsWidget, SRigVMGraphFunctionLocalizationWidget, InArgs._Function, InArgs._TargetBlueprint)
+                ]
+
+                +SVerticalBox::Slot()
+                .AutoHeight()
+                .HAlign(HAlign_Right)
+                .Padding(5)
+                [
+                    SNew(SUniformGridPanel)
+                    .SlotPadding(FAppStyle::GetMargin("StandardDialog.SlotPadding"))
+                    .MinDesiredSlotWidth(FAppStyle::GetFloat("StandardDialog.MinDesiredSlotWidth"))
+                    .MinDesiredSlotHeight(FAppStyle::GetFloat("StandardDialog.MinDesiredSlotHeight"))
+                    +SUniformGridPanel::Slot(0, 0)
+                    [
+                        SNew(SButton)
+                        .HAlign(HAlign_Center)
+                        .ContentPadding(FAppStyle::GetMargin("StandardDialog.ContentPadding"))
+                        .Text(LOCTEXT("OK", "OK"))
+                        .OnClicked(this, &SRigVMGraphFunctionLocalizationDialog::OnButtonClick, EAppReturnType::Ok)
+                        .IsEnabled(this, &SRigVMGraphFunctionLocalizationDialog::IsOkButtonEnabled)
+                    ]
+                    +SUniformGridPanel::Slot(1, 0)
+                    [
+                        SNew(SButton)
+                        .HAlign(HAlign_Center)
+                        .ContentPadding(FAppStyle::GetMargin("StandardDialog.ContentPadding"))
+                        .Text(LOCTEXT("Cancel", "Cancel"))
+                        .OnClicked(this, &SRigVMGraphFunctionLocalizationDialog::OnButtonClick, EAppReturnType::Cancel)
+                    ]
+                ]
+			]
+		]);
+}
+
+FReply SRigVMGraphFunctionLocalizationDialog::OnButtonClick(EAppReturnType::Type ButtonID)
+{
+	UserResponse = ButtonID;
+	RequestDestroyWindow();
+
+	return FReply::Handled();
+}
+
+bool SRigVMGraphFunctionLocalizationDialog::IsOkButtonEnabled() const
+{
+	return FunctionsWidget->FunctionsToLocalize.Num() > 0;
+}
+
+EAppReturnType::Type SRigVMGraphFunctionLocalizationDialog::ShowModal()
+{
+	GEditor->EditorAddModalWindow(SharedThis(this));
+	return UserResponse;
+}
+
+TArray<FRigVMGraphFunctionIdentifier>& SRigVMGraphFunctionLocalizationDialog::GetFunctionsToLocalize()
+{
+	return FunctionsWidget->FunctionsToLocalize;
+}
+
+#undef LOCTEXT_NAMESPACE
diff --git a/Engine/Plugins/Runtime/RigVM/Source/RigVMEditor/Public/Editor/RigVMDetailsViewWrapperObject.h b/Engine/Plugins/Runtime/RigVM/Source/RigVMEditor/Public/Editor/RigVMDetailsViewWrapperObject.h
new file mode 100644
--- /dev/null
+++ b/Engine/Plugins/Runtime/RigVM/Source/RigVMEditor/Public/Editor/RigVMDetailsViewWrapperObject.h
@@ -0,0 +1,127 @@
+// Copyright Epic Games, Inc. All Rights Reserved.
+
+#pragma once
+
+#include "CoreMinimal.h"
+#include "RigVMModel/RigVMGraph.h"
+#include "RigVMModel/RigVMNode.h"
+#include "RigVMDetailsViewWrapperObject.generated.h"
+
+class URigVMDetailsViewWrapperObject;
+
+DECLARE_EVENT_ThreeParams(URigVMDetailsViewWrapperObject, FWrappedPropertyChangedChainEvent, URigVMDetailsViewWrapperObject*, const FString&, FPropertyChangedChainEvent&);
+
+UCLASS()
+class RIGVMEDITOR_API URigVMDetailsViewWrapperObject : public UObject
+{
+public:
+	GENERATED_BODY()
+
+	URigVMDetailsViewWrapperObject();
+
+	// Creating wrappers from a given struct
+	virtual UClass* GetClassForStruct(UScriptStruct* InStruct, bool bCreateIfNeeded = true) const;
+	static URigVMDetailsViewWrapperObject* MakeInstance(UClass* InWrapperObjectClass, UObject* InOuter, UScriptStruct* InStruct, uint8* InStructMemory, UObject* InSubject = nullptr);
+	UScriptStruct* GetWrappedStruct() const;
+
+	// Creating wrappers from a RigVMNode
+	virtual UClass* GetClassForNodes(TArray<URigVMNode*> InNodes, bool bCreateIfNeeded = true) const;
+	static URigVMDetailsViewWrapperObject* MakeInstance(UClass* InWrapperObjectClass, UObject* InOuter, TArray<URigVMNode*> InNodes, URigVMNode* InSubject);
+	
+	static void MarkOutdatedClass(UClass* InClass);
+	static bool IsValidClass(UClass* InClass);
+	
+	FString GetWrappedNodeNotation() const;
+	
+	bool IsChildOf(const UStruct* InStruct) const
+	{
+		const UScriptStruct* WrappedStruct = GetWrappedStruct();
+		return WrappedStruct && WrappedStruct->IsChildOf(InStruct);
+	}
+
+	template<typename T>
+	bool IsChildOf() const
+	{
+		return IsChildOf(T::StaticStruct());
+	}
+
+	void SetContent(const uint8* InStructMemory, const UStruct* InStruct);
+	void GetContent(uint8* OutStructMemory, const UStruct* InStruct) const;
+	void SetContent(URigVMNode* InNode);
+
+	template<typename T>
+	T GetContent() const
+	{
+		check(IsChildOf<T>());
+		
+		T Result;
+		GetContent((uint8*)&Result, T::StaticStruct());
+		return Result;
+	}
+
+	template<typename T>
+	void SetContent(const T& InValue)
+	{
+		check(IsChildOf<T>());
+
+		SetContent((const uint8*)&InValue, T::StaticStruct());
+	}
+
+	UObject* GetSubject() const;
+	void SetSubject(UObject* InSubject);
+
+	FWrappedPropertyChangedChainEvent& GetWrappedPropertyChangedChainEvent() { return WrappedPropertyChangedChainEvent; }
+	virtual void PostEditChangeChainProperty(FPropertyChangedChainEvent& PropertyChangedEvent) override;
+
+private:
+
+	static void CopyPropertiesForUnrelatedStructs(uint8* InTargetMemory, const UStruct* InTargetStruct, const uint8* InSourceMemory, const UStruct* InSourceStruct);
+	void HandleModifiedEvent(ERigVMGraphNotifType InNotifType, URigVMGraph* InGraph, UObject* InSubject);
+	void SetContentForPin(URigVMPin* InPin);
+
+	struct FPerClassInfo
+	{
+		FString Notation;
+		UScriptStruct* ScriptStruct;
+		
+		FPerClassInfo()
+			: Notation()
+			, ScriptStruct(nullptr)
+		{}
+
+		FPerClassInfo(UScriptStruct* InScriptStruct)
+		: Notation()
+		, ScriptStruct(InScriptStruct)
+		{}
+
+		FPerClassInfo(const FString& InNotation)
+		: Notation(InNotation)
+		, ScriptStruct(nullptr)
+		{}
+
+		friend uint32 GetTypeHash(const FPerClassInfo& Info)
+		{
+			return HashCombine(GetTypeHash(Info.Notation), GetTypeHash(Info.ScriptStruct));
+		}
+
+		bool operator ==(const FPerClassInfo& Other) const
+		{
+			return Notation == Other.Notation && ScriptStruct == Other.ScriptStruct;
+		}
+
+		bool operator !=(const FPerClassInfo& Other) const
+		{
+			return Notation != Other.Notation || ScriptStruct != Other.ScriptStruct;
+		}
+	};
+	
+	static TMap<FPerClassInfo, UClass*> InfoToClass;
+	static TMap<UClass*, FPerClassInfo> ClassToInfo;
+	static TSet<UClass*> OutdatedClassToRecreate;
+	bool bIsSettingValue;
+
+	UPROPERTY()
+	TWeakObjectPtr<UObject> SubjectPtr;
+	
+	FWrappedPropertyChangedChainEvent WrappedPropertyChangedChainEvent;
+};
diff --git a/Engine/Plugins/Runtime/RigVM/Source/RigVMEditor/Public/Editor/RigVMEditor.h b/Engine/Plugins/Runtime/RigVM/Source/RigVMEditor/Public/Editor/RigVMEditor.h
new file mode 100644
--- /dev/null
+++ b/Engine/Plugins/Runtime/RigVM/Source/RigVMEditor/Public/Editor/RigVMEditor.h
@@ -0,0 +1,42 @@
+// Copyright Epic Games, Inc. All Rights Reserved.
+
+#pragma once
+
+#include "CoreMinimal.h"
+#include "BlueprintEditor.h"
+#include "RigVMHost.h"
+#include "RigVMBlueprint.h"
+
+class RIGVMEDITOR_API FRigVMEditor : public FBlueprintEditor
+{
+public:
+
+	FRigVMEditor();
+
+	// returns the blueprint being edited
+	URigVMBlueprint* GetRigVMBlueprint() const;
+
+	// returns the currently debugged / viewed host
+	URigVMHost* GetRigVMHost() const;
+
+	// returns the class to use for detail wrapper objects (UI shim layer)
+	virtual UClass* GetDetailWrapperClass() const;
+
+	virtual void HandleModifiedEvent(ERigVMGraphNotifType InNotifType, URigVMGraph* InGraph, UObject* InSubject);
+
+	DECLARE_EVENT_OneParam(FRigVMEditor, FPreviewHostUpdated, FRigVMEditor*);
+	FPreviewHostUpdated& OnPreviewHostUpdated() { return PreviewHostUpdated;  }
+
+protected:
+
+	void SetHost(URigVMHost* InHost);
+
+private:
+
+	/** Our currently running control rig instance */
+	TObjectPtr<URigVMHost> Host;
+
+	FPreviewHostUpdated PreviewHostUpdated;
+
+	friend class SRigVMExecutionStackView;
+};
diff --git a/Engine/Plugins/Runtime/RigVM/Source/RigVMEditor/Public/Editor/RigVMExecutionStackCommands.h b/Engine/Plugins/Runtime/RigVM/Source/RigVMEditor/Public/Editor/RigVMExecutionStackCommands.h
new file mode 100644
--- /dev/null
+++ b/Engine/Plugins/Runtime/RigVM/Source/RigVMEditor/Public/Editor/RigVMExecutionStackCommands.h
@@ -0,0 +1,30 @@
+// Copyright Epic Games, Inc. All Rights Reserved.
+
+#pragma once
+
+#include "CoreMinimal.h"
+#include "Framework/Commands/Commands.h"
+
+class RIGVMEDITOR_API FRigVMExecutionStackCommands : public TCommands<FRigVMExecutionStackCommands>
+{
+public:
+	FRigVMExecutionStackCommands() : TCommands<FRigVMExecutionStackCommands>
+	(
+		"ControlRigStack",
+		NSLOCTEXT("Contexts", "RigStack", "Execution Stack"),
+		NAME_None, // "MainFrame" // @todo Fix this crash
+		"ControlRigEditorStyle" // Icon Style Set
+	)
+	{}
+	
+	/** Focuses on a selected operator's node */
+	TSharedPtr< FUICommandInfo > FocusOnSelection;
+
+	/** Looks for a specific instruction by index and brings it into focus */
+	TSharedPtr< FUICommandInfo > GoToInstruction;
+
+	/**
+	 * Initialize commands
+	 */
+	virtual void RegisterCommands() override;
+};
diff --git a/Engine/Plugins/Runtime/RigVM/Source/RigVMEditor/Public/Editor/RigVMExecutionStackTabSummoner.h b/Engine/Plugins/Runtime/RigVM/Source/RigVMEditor/Public/Editor/RigVMExecutionStackTabSummoner.h
new file mode 100644
--- /dev/null
+++ b/Engine/Plugins/Runtime/RigVM/Source/RigVMEditor/Public/Editor/RigVMExecutionStackTabSummoner.h
@@ -0,0 +1,22 @@
+// Copyright Epic Games, Inc. All Rights Reserved.
+
+#pragma once
+
+#include "CoreMinimal.h"
+#include "WorkflowOrientedApp/WorkflowTabFactory.h"
+
+class FRigVMEditor;
+
+struct RIGVMEDITOR_API FRigVMExecutionStackTabSummoner : public FWorkflowTabFactory
+{
+public:
+	static const FName TabID;
+	
+public:
+	FRigVMExecutionStackTabSummoner(const TSharedRef<FRigVMEditor>& InRigVMEditor);
+	
+	virtual TSharedRef<SWidget> CreateTabBody(const FWorkflowTabSpawnInfo& Info) const override;
+	
+protected:
+	TWeakPtr<FRigVMEditor> RigVMEditor;
+};
diff --git a/Engine/Plugins/Runtime/RigVM/Source/RigVMEditor/Public/Widgets/SRigVMExecutionStackView.h b/Engine/Plugins/Runtime/RigVM/Source/RigVMEditor/Public/Widgets/SRigVMExecutionStackView.h
new file mode 100644
--- /dev/null
+++ b/Engine/Plugins/Runtime/RigVM/Source/RigVMEditor/Public/Widgets/SRigVMExecutionStackView.h
@@ -0,0 +1,141 @@
+// Copyright Epic Games, Inc. All Rights Reserved.
+
+#pragma once
+
+#include "CoreMinimal.h"
+#include "EdGraph/RigVMEdGraph.h"
+#include "RigVMHost.h"
+#include "RigVMBlueprint.h"
+#include "Widgets/Views/STreeView.h"
+#include "RigVMCore/RigVM.h"
+#include "RigVMBlueprint.h"
+
+class FRigVMEditor;
+class SRigVMExecutionStackView;
+class FUICommandList;
+class SSearchBox;
+
+namespace ERigStackEntry
+{
+	enum Type
+	{
+		Operator,
+		Info,
+		Warning,
+		Error,
+	};
+}
+
+/** An item in the stack */
+class RIGVMEDITOR_API FRigStackEntry : public TSharedFromThis<FRigStackEntry>
+{
+public:
+	FRigStackEntry(int32 InEntryIndex, ERigStackEntry::Type InEntryType, int32 InInstructionIndex, ERigVMOpCode InOpCode, const FString& InLabel, const FRigVMASTProxy& InProxy);
+
+	TSharedRef<ITableRow> MakeTreeRowWidget(const TSharedRef<STableViewBase>& InOwnerTable, TSharedRef<FRigStackEntry> InEntry, TSharedRef<FUICommandList> InCommandList, TSharedPtr<SRigVMExecutionStackView> InStackView, TWeakObjectPtr<URigVMBlueprint> InBlueprint);
+
+	int32 EntryIndex;
+	ERigStackEntry::Type EntryType;
+	int32 InstructionIndex;
+	FString CallPath;
+	FRigVMCallstack Callstack;
+	ERigVMOpCode OpCode;
+	FString Label;
+	TArray<TSharedPtr<FRigStackEntry>> Children;
+};
+
+class SRigStackItem : public STableRow<TSharedPtr<FRigStackEntry>>
+{
+	SLATE_BEGIN_ARGS(SRigStackItem) {}
+	SLATE_END_ARGS()
+
+	void Construct(const FArguments& InArgs, const TSharedRef<STableViewBase>& OwnerTable, TSharedRef<FRigStackEntry> InStackEntry, TSharedRef<FUICommandList> InCommandList, TWeakObjectPtr<URigVMBlueprint> InBlueprint);
+
+private:
+	TWeakPtr<FRigStackEntry> WeakStackEntry;
+	TWeakObjectPtr<URigVMBlueprint> WeakBlueprint;
+	TWeakPtr<FUICommandList> WeakCommandList;
+
+	FText GetIndexText() const;
+	FText GetLabelText() const;
+	FSlateFontInfo GetLabelFont() const;
+	FText GetTooltip() const;
+	FText GetVisitedCountText() const;
+	FText GetDurationText() const;
+};
+
+class RIGVMEDITOR_API SRigVMExecutionStackView : public SCompoundWidget
+{
+public:
+	SLATE_BEGIN_ARGS(SRigVMExecutionStackView) {}
+	SLATE_END_ARGS()
+
+	~SRigVMExecutionStackView();
+
+	void Construct(const FArguments& InArgs, TSharedRef<FRigVMEditor> InRigVMEditor);
+
+	/** Set Selection Changed */
+	void OnSelectionChanged(TSharedPtr<FRigStackEntry> Selection, ESelectInfo::Type SelectInfo);
+
+	TSharedPtr< SWidget > CreateContextMenu();
+
+protected:
+
+	/** Rebuild the tree view */
+	void RefreshTreeView(URigVM* InVM);
+
+private:
+
+	/** Bind commands that this widget handles */
+	void BindCommands();
+
+	/** Make a row widget for the table */
+	TSharedRef<ITableRow> MakeTableRowWidget(TSharedPtr<FRigStackEntry> InItem, const TSharedRef<STableViewBase>& OwnerTable, TWeakObjectPtr<URigVMBlueprint> InBlueprint);
+
+	/** Get children for the tree */
+	void HandleGetChildrenForTree(TSharedPtr<FRigStackEntry> InItem, TArray<TSharedPtr<FRigStackEntry>>& OutChildren);
+
+	/** Focus on the selected operator in the graph*/
+	void HandleFocusOnSelectedGraphNode();
+
+	/** Offers a dialog to move to a specific instruction */
+	void HandleGoToInstruction();
+
+	void OnVMCompiled(UObject* InCompiledObject, URigVM* InCompiledVM);
+
+	//* Focus on the instruction when the execution is halted */
+	void HandleExecutionHalted(const int32 HaltedAtInstruction, UObject* InNode, const FName& InEntryName);
+
+	/** Search box widget */
+	TSharedPtr<SSearchBox> FilterBox;
+	FText FilterText;
+	void OnFilterTextChanged(const FText& SearchText);
+
+
+	bool bSuspendModelNotifications;
+	bool bSuspendControllerSelection;
+	void HandleModifiedEvent(ERigVMGraphNotifType InNotifType, URigVMGraph* InGraph, UObject* InSubject);
+	void HandleHostInitializedEvent(URigVMHost* InHost, const FName& InEventName);
+	void HandlePreviewHostUpdated(FRigVMEditor* InEditor);
+	void HandleItemMouseDoubleClick(TSharedPtr<FRigStackEntry> InItem);
+
+	/** Populate the execution stack with descriptive names for each instruction */
+	void PopulateStackView(URigVM* InVM);
+
+	TSharedPtr<STreeView<TSharedPtr<FRigStackEntry>>> TreeView;
+
+	/** Command list we bind to */
+	TSharedPtr<FUICommandList> CommandList;
+
+	TWeakPtr<FRigVMEditor> RigVMEditor;
+	TWeakObjectPtr<URigVMBlueprint> RigVMBlueprint;
+
+	TArray<TSharedPtr<FRigStackEntry>> Operators;
+
+	int32 HaltedAtInstruction;
+
+	FDelegateHandle OnModelModified;
+	FDelegateHandle OnHostInitializedHandle;
+	FDelegateHandle OnVMCompiledHandle;
+	FDelegateHandle OnPreviewHostUpdatedHandle;
+};
diff --git a/Engine/Plugins/Runtime/RigVM/Source/RigVMEditor/Public/Widgets/SRigVMGraphBreakLinksWidget.h b/Engine/Plugins/Runtime/RigVM/Source/RigVMEditor/Public/Widgets/SRigVMGraphBreakLinksWidget.h
new file mode 100644
--- /dev/null
+++ b/Engine/Plugins/Runtime/RigVM/Source/RigVMEditor/Public/Widgets/SRigVMGraphBreakLinksWidget.h
@@ -0,0 +1,62 @@
+﻿// Copyright Epic Games, Inc. All Rights Reserved.
+
+#pragma once
+
+#include "CoreMinimal.h"
+#include "RigVMBlueprint.h"
+#include "Dialogs/Dialogs.h"
+#include "IAssetTypeActions.h"
+
+class ITableRow;
+class STableViewBase;
+
+DECLARE_DELEGATE_OneParam(FRigVMOnFocusOnLinkRequestedDelegate, URigVMLink*);
+
+class RIGVMEDITOR_API SRigVMGraphBreakLinksWidget : public SCompoundWidget
+{
+public:
+	SLATE_BEGIN_ARGS(SRigVMGraphBreakLinksWidget) {}
+	SLATE_ARGUMENT(FRigVMOnFocusOnLinkRequestedDelegate, OnFocusOnLink)
+	SLATE_END_ARGS()
+
+	void Construct(const FArguments& InArgs, TArray<URigVMLink*> InLinks);
+
+private:
+
+	/** Handler for when the user double clicks, presses enter, or presses space on an asset */
+	void OnLinkDoubleClicked(const URigVMLink* InLink, EAssetTypeActivationMethod::Type ActivationMethod);
+
+	TSharedRef<ITableRow> GenerateItemRow(URigVMLink* Item, const TSharedRef<STableViewBase>& OwnerTable);
+ 
+	TArray<URigVMLink*> Links;
+	FRigVMOnFocusOnLinkRequestedDelegate OnFocusOnLink;
+
+	void HandleItemMouseDoubleClick(URigVMLink* InItem);
+
+	friend class SRigVMGraphBreakLinksDialog;
+};
+
+class RIGVMEDITOR_API SRigVMGraphBreakLinksDialog : public SWindow
+{
+public:
+	
+	SLATE_BEGIN_ARGS(SRigVMGraphBreakLinksDialog)
+	{
+	}
+
+	SLATE_ARGUMENT(TArray<URigVMLink*>, Links)
+	SLATE_ARGUMENT(FRigVMOnFocusOnLinkRequestedDelegate, OnFocusOnLink)
+	SLATE_END_ARGS()
+
+	void Construct(const FArguments& InArgs);
+
+	EAppReturnType::Type ShowModal();
+
+protected:
+
+	TSharedPtr<SRigVMGraphBreakLinksWidget> BreakLinksWidget;
+	
+	FReply OnButtonClick(EAppReturnType::Type ButtonID);
+	bool IsOkButtonEnabled() const;
+	EAppReturnType::Type UserResponse;
+};
\ No newline at end of file
diff --git a/Engine/Plugins/Runtime/RigVM/Source/RigVMEditor/Public/Widgets/SRigVMGraphFunctionBulkEditWidget.h b/Engine/Plugins/Runtime/RigVM/Source/RigVMEditor/Public/Widgets/SRigVMGraphFunctionBulkEditWidget.h
new file mode 100644
--- /dev/null
+++ b/Engine/Plugins/Runtime/RigVM/Source/RigVMEditor/Public/Widgets/SRigVMGraphFunctionBulkEditWidget.h
@@ -0,0 +1,67 @@
+// Copyright Epic Games, Inc. All Rights Reserved.
+
+#pragma once
+
+#include "CoreMinimal.h"
+#include "RigVMBlueprint.h"
+#include "RigVMModel/Nodes/RigVMLibraryNode.h"
+#include "Dialogs/Dialogs.h"
+#include "IAssetTypeActions.h"
+
+class SRigVMGraphFunctionBulkEditWidget;
+
+class RIGVMEDITOR_API SRigVMGraphFunctionBulkEditWidget : public SCompoundWidget
+{
+public:
+	SLATE_BEGIN_ARGS(SRigVMGraphFunctionBulkEditWidget) {}
+	SLATE_END_ARGS()
+
+	void Construct(const FArguments& InArgs, URigVMBlueprint* InBlueprint, URigVMController* InController, URigVMLibraryNode* InFunction, ERigVMControllerBulkEditType InEditType);
+
+private:
+
+	/** Handler for when an asset context menu has been requested. */
+	TSharedPtr<SWidget> OnGetAssetContextMenu( const TArray<FAssetData>& SelectedAssets );
+
+	/** Handler for when the user double clicks, presses enter, or presses space on an asset */
+	void OnAssetsActivated(const TArray<FAssetData>& ActivatedAssets, EAssetTypeActivationMethod::Type ActivationMethod);
+
+	/** Show a progress bar and load all of the assets */
+	void LoadAffectedAssets();
+
+	URigVMBlueprint* Blueprint;
+	URigVMController* Controller;
+	URigVMLibraryNode* Function;
+	ERigVMControllerBulkEditType EditType;
+
+	TSharedRef<SWidget> MakeAssetViewForReferencedAssets();
+
+	friend class SRigVMGraphFunctionBulkEditDialog;
+};
+
+class RIGVMEDITOR_API SRigVMGraphFunctionBulkEditDialog : public SWindow
+{
+public:
+	
+	SLATE_BEGIN_ARGS(SRigVMGraphFunctionBulkEditDialog)
+	{
+	}
+
+	SLATE_ARGUMENT(URigVMBlueprint*, Blueprint)
+	SLATE_ARGUMENT(URigVMController*, Controller)
+	SLATE_ARGUMENT(URigVMLibraryNode*, Function)
+	SLATE_ARGUMENT(ERigVMControllerBulkEditType, EditType)
+	SLATE_END_ARGS()
+
+	void Construct(const FArguments& InArgs);
+
+	EAppReturnType::Type ShowModal();
+
+protected:
+
+	TSharedPtr<SRigVMGraphFunctionBulkEditWidget> BulkEditWidget;
+	
+	FReply OnButtonClick(EAppReturnType::Type ButtonID);
+	bool IsOkButtonEnabled() const;
+	EAppReturnType::Type UserResponse;
+};
\ No newline at end of file
diff --git a/Engine/Plugins/Runtime/RigVM/Source/RigVMEditor/Public/Widgets/SRigVMGraphFunctionLocalizationWidget.h b/Engine/Plugins/Runtime/RigVM/Source/RigVMEditor/Public/Widgets/SRigVMGraphFunctionLocalizationWidget.h
new file mode 100644
--- /dev/null
+++ b/Engine/Plugins/Runtime/RigVM/Source/RigVMEditor/Public/Widgets/SRigVMGraphFunctionLocalizationWidget.h
@@ -0,0 +1,78 @@
+// Copyright Epic Games, Inc. All Rights Reserved.
+
+#pragma once
+
+#include "CoreMinimal.h"
+#include "RigVMBlueprint.h"
+#include "RigVMModel/Nodes/RigVMLibraryNode.h"
+#include "Widgets/Views/SListView.h"
+#include "Dialogs/Dialogs.h"
+
+class SRigVMGraphFunctionLocalizationWidget;
+
+class RIGVMEDITOR_API SRigVMGraphFunctionLocalizationItem : public TSharedFromThis<SRigVMGraphFunctionLocalizationItem>
+{
+public:
+	SRigVMGraphFunctionLocalizationItem(const FRigVMGraphFunctionIdentifier& InFunction);
+
+	FText DisplayText;
+	FText ToolTipText;
+	const FRigVMGraphFunctionIdentifier Function;
+};
+
+class RIGVMEDITOR_API SRigVMGraphFunctionLocalizationTableRow : public STableRow<TSharedPtr<SRigVMGraphFunctionLocalizationItem>>
+{
+	SLATE_BEGIN_ARGS(SRigVMGraphFunctionLocalizationTableRow) {}
+	SLATE_END_ARGS()
+
+	void Construct(const FArguments& InArgs, const TSharedRef<STableViewBase>& OwnerTable, SRigVMGraphFunctionLocalizationWidget* InLocalizationWidget, TSharedRef<SRigVMGraphFunctionLocalizationItem> InFunctionItem);
+};
+
+class RIGVMEDITOR_API SRigVMGraphFunctionLocalizationWidget : public SCompoundWidget
+{
+public:
+	SLATE_BEGIN_ARGS(SRigVMGraphFunctionLocalizationWidget) {}
+	SLATE_END_ARGS()
+
+	void Construct(const FArguments& InArgs, const FRigVMGraphFunctionIdentifier& InFunctionToLocalize, URigVMBlueprint* InTargetBlueprint);
+
+	TSharedRef<ITableRow> GenerateFunctionListRow(TSharedPtr<SRigVMGraphFunctionLocalizationItem> InItem, const TSharedRef<STableViewBase>& InOwningTable);
+	ECheckBoxState IsFunctionEnabled(const FRigVMGraphFunctionIdentifier InFunction) const;
+	void SetFunctionEnabled(ECheckBoxState NewState, const FRigVMGraphFunctionIdentifier InFunction);
+	bool IsFunctionPublic(const FRigVMGraphFunctionIdentifier InFunction) const;
+
+private:
+
+	TArray<FRigVMGraphFunctionIdentifier> FunctionsToLocalize;
+	TArray<TSharedPtr<SRigVMGraphFunctionLocalizationItem>> FunctionItems;
+	TMap<FRigVMGraphFunctionIdentifier, TSharedRef<SRigVMGraphFunctionLocalizationTableRow>> TableRows;
+
+	friend class SRigVMGraphFunctionLocalizationTableRow;
+	friend class SRigVMGraphFunctionLocalizationDialog;
+};
+
+class RIGVMEDITOR_API SRigVMGraphFunctionLocalizationDialog : public SWindow
+{
+public:
+	
+	SLATE_BEGIN_ARGS(SRigVMGraphFunctionLocalizationDialog)
+	{
+	}
+
+	SLATE_ARGUMENT(FRigVMGraphFunctionIdentifier, Function)
+	SLATE_ARGUMENT(URigVMBlueprint*, TargetBlueprint)
+	SLATE_END_ARGS()
+
+	void Construct(const FArguments& InArgs);
+
+	EAppReturnType::Type ShowModal();
+	TArray<FRigVMGraphFunctionIdentifier>& GetFunctionsToLocalize();
+
+protected:
+
+	TSharedPtr<SRigVMGraphFunctionLocalizationWidget> FunctionsWidget;
+	
+	FReply OnButtonClick(EAppReturnType::Type ButtonID);
+	bool IsOkButtonEnabled() const;
+	EAppReturnType::Type UserResponse;
+};
\ No newline at end of file
diff --git a/Engine/Plugins/Runtime/RigVM/Source/RigVMEditor/RigVMEditor.Build.cs b/Engine/Plugins/Runtime/RigVM/Source/RigVMEditor/RigVMEditor.Build.cs
--- a/Engine/Plugins/Runtime/RigVM/Source/RigVMEditor/RigVMEditor.Build.cs
+++ b/Engine/Plugins/Runtime/RigVM/Source/RigVMEditor/RigVMEditor.Build.cs
@@ -21,6 +21,7 @@ public class RigVMEditor : ModuleRules
                 "SlateCore",
                 "InputCore",
                 "ToolWidgets",
+                "Kismet",
                 "KismetWidgets",
             }
         );
