RigVM: Fix non-unity issue in RigVMBlueprintCompiler

diff --git a/Engine/Plugins/Runtime/RigVM/Source/RigVMDeveloper/Private/RigVMBlueprintCompiler.cpp b/Engine/Plugins/Runtime/RigVM/Source/RigVMDeveloper/Private/RigVMBlueprintCompiler.cpp
--- a/Engine/Plugins/Runtime/RigVM/Source/RigVMDeveloper/Private/RigVMBlueprintCompiler.cpp
+++ b/Engine/Plugins/Runtime/RigVM/Source/RigVMDeveloper/Private/RigVMBlueprintCompiler.cpp
@@ -12,6 +12,7 @@
 #include "Kismet2/KismetReinstanceUtilities.h"
 #include "Widgets/Notifications/SNotificationList.h"
 #include "Framework/Notifications/NotificationManager.h"
+#include "Stats/StatsHierarchical.h"
 
 DEFINE_LOG_CATEGORY_STATIC(LogRigVMCompiler, Log, All);
 #define LOCTEXT_NAMESPACE "RigVMBlueprintCompiler"
