ControlRig / RigVM: Move editor style / slate content to RigVMEditor

diff --git a/Engine/Plugins/Animation/ControlRig/Source/ControlRigEditor/Private/ControlRigBlueprintActions.cpp b/Engine/Plugins/Animation/ControlRig/Source/ControlRigEditor/Private/ControlRigBlueprintActions.cpp
--- a/Engine/Plugins/Animation/ControlRig/Source/ControlRigEditor/Private/ControlRigBlueprintActions.cpp
+++ b/Engine/Plugins/Animation/ControlRig/Source/ControlRigEditor/Private/ControlRigBlueprintActions.cpp
@@ -3,7 +3,7 @@
 #include "ControlRigBlueprintActions.h"
 #include "ControlRigBlueprintFactory.h"
 #include "ControlRigBlueprint.h"
-#include "ControlRigEditorStyle.h"
+#include "Editor/RigVMEditorStyle.h"
 #include "IControlRigEditorModule.h"
 
 #include "Styling/SlateIconFinder.h"
@@ -160,7 +160,7 @@ void FControlRigBlueprintActions::ExtendSketalMeshToolMenu()
 						"CreateControlRig",
 						LOCTEXT("CreateControlRig", "Control Rig"),
 						LOCTEXT("CreateControlRig_ToolTip", "Creates a control rig and preconfigures it for this asset"),
-						FSlateIcon(FControlRigEditorStyle::Get().GetStyleSetName(), "ControlRig", "ControlRig.RigUnit"),
+						FSlateIcon(FRigVMEditorStyle::Get().GetStyleSetName(), "RigVM", "RigVM.Unit"),
 						FExecuteAction::CreateLambda([SelectedObjects]()
 						{
 							for (UObject* SelectedObject : SelectedObjects)
diff --git a/Engine/Plugins/Animation/ControlRig/Source/ControlRigEditor/Private/ControlRigEditorModule.cpp b/Engine/Plugins/Animation/ControlRig/Source/ControlRigEditor/Private/ControlRigEditorModule.cpp
--- a/Engine/Plugins/Animation/ControlRig/Source/ControlRigEditor/Private/ControlRigEditorModule.cpp
+++ b/Engine/Plugins/Animation/ControlRig/Source/ControlRigEditor/Private/ControlRigEditorModule.cpp
@@ -14,6 +14,7 @@
 #include "ISequencerModule.h"
 #include "IAssetTools.h"
 #include "ControlRigEditorStyle.h"
+#include "Editor/RigVMEditorStyle.h"
 #include "Framework/Docking/LayoutExtender.h"
 #include "Framework/Application/SlateApplication.h"
 #include "LevelEditor.h"
@@ -61,7 +62,6 @@
 #include "RigVMBlueprintUtils.h"
 #include "ControlRigBlueprintCommands.h"
 #include "ControlRigHierarchyCommands.h"
-#include "Editor/RigVMExecutionStackCommands.h"
 #include "Animation/AnimSequence.h"
 #include "Editor/ControlRigEditorEditMode.h"
 #include "ControlRigElementDetails.h"
@@ -146,7 +146,6 @@ void FControlRigEditorModule::StartupModule()
 	FControlRigEditModeCommands::Register();
 	FControlRigBlueprintCommands::Register();
 	FControlRigHierarchyCommands::Register();
-	FRigVMExecutionStackCommands::Register();
 	FControlRigEditorStyle::Get();
 
 	MenuExtensibilityManager = MakeShareable(new FExtensibilityManager);
@@ -223,14 +222,14 @@ void FControlRigEditorModule::StartupModule()
 	FEditorModeRegistry::Get().RegisterMode<FControlRigEditMode>(
 		FControlRigEditMode::ModeName,
 		NSLOCTEXT("AnimationModeToolkit", "DisplayName", "Animation"),
-		FSlateIcon(FControlRigEditorStyle::Get().GetStyleSetName(), "ControlRigEditMode", "ControlRigEditMode.Small"),
+		FSlateIcon(FRigVMEditorStyle::Get().GetStyleSetName(), "RigVMEditMode", "RigVMEditMode.Small"),
 		true,
 		8000);
 
 	FEditorModeRegistry::Get().RegisterMode<FControlRigEditorEditMode>(
 		FControlRigEditorEditMode::ModeName,
 		NSLOCTEXT("RiggingModeToolkit", "DisplayName", "Rigging"),
-		FSlateIcon(FControlRigEditorStyle::Get().GetStyleSetName(), "ControlRigEditMode", "ControlRigEditMode.Small"),
+		FSlateIcon(FRigVMEditorStyle::Get().GetStyleSetName(), "RigVMEditMode", "RigVMEditMode.Small"),
 		false,
 		8500);
 
diff --git a/Engine/Plugins/Animation/ControlRig/Source/ControlRigEditor/Private/ControlRigEditorStyle.h b/Engine/Plugins/Animation/ControlRig/Source/ControlRigEditor/Private/ControlRigEditorStyle.h
--- a/Engine/Plugins/Animation/ControlRig/Source/ControlRigEditor/Private/ControlRigEditorStyle.h
+++ b/Engine/Plugins/Animation/ControlRig/Source/ControlRigEditor/Private/ControlRigEditorStyle.h
@@ -58,10 +58,6 @@ public:
 			Set("ClassIcon.ControlRigPose", new IMAGE_BRUSH("Slate/ControlRigPose_16", Icon16x16));
 		}
 
-		// Edit mode styles
-		Set("ControlRigEditMode", new IMAGE_BRUSH_SVG("Slate/animation", Icon40x40));
-		Set("ControlRigEditMode.Small", new IMAGE_BRUSH_SVG("Slate/animation", Icon20x20));
-
 		// Sequencer styles
 		{
 			Set("ControlRig.ExportAnimSequence", new IMAGE_BRUSH("Slate/ExportAnimSequence_24x", Icon24x24));
@@ -88,20 +84,11 @@ public:
 		// Control Rig Editor styles
 		{
 			// tab icons
-			Set("ControlRig.TabIcon", new IMAGE_BRUSH("Slate/ControlRigTab_16x", Icon16x16));
-			Set("ExecutionStack.TabIcon", new IMAGE_BRUSH_SVG("Slate/ExecutionStack", Icon16x16));
 			Set("RigHierarchy.TabIcon", new IMAGE_BRUSH_SVG("Slate/RigHierarchy", Icon16x16));
 			Set("RigValidation.TabIcon", new IMAGE_BRUSH_SVG("Slate/RigValidation", Icon16x16));
 			Set("CurveContainer.TabIcon", new IMAGE_BRUSH_SVG("Slate/CurveContainer", Icon16x16));
 			Set("HierarchicalProfiler.TabIcon", new IMAGE_BRUSH_SVG("Slate/HierarchicalProfiler", Icon16x16));
 			
-			// icons for control units
-			Set("ControlRig.RigUnit", new IMAGE_BRUSH("Slate/ControlRigUnit_16x", Icon16x16));
-			Set("ControlRig.Template", new IMAGE_BRUSH("Slate/Template_16x", Icon16x16));
-
-			Set("ControlRig.AutoCompileGraph", new IMAGE_BRUSH("Slate/AutoCompile", Icon20x20));
-			Set("ControlRig.AutoCompileGraph.Small", new IMAGE_BRUSH("Slate/AutoCompile", Icon20x20));
-
 			{
 				FContentRootBracket Bracket(this, EngineEditorSlateDir);
 				Set("ControlRig.ConstructionMode", new IMAGE_BRUSH_SVG("Starship/Common/Adjust", Icon40x40));
@@ -111,15 +98,6 @@ public:
 				Set("ControlRig.BackwardsAndForwardsSolveEvent", new IMAGE_BRUSH("Icons/Loop_40x", Icon40x40));
 			}
 
-			Set("ControlRig.Bug.Dot", new IMAGE_BRUSH("Slate/ControlRig_BugDot_32x", Icon16x16));
-			Set("ControlRig.Bug.Normal", new IMAGE_BRUSH("Slate/ControlRig_Bug_28x", Icon14x14));
-			Set("ControlRig.Bug.Open", new IMAGE_BRUSH("Slate/ControlRig_BugOpen_28x", Icon14x14));
-			Set("ControlRig.Bug.Solid", new IMAGE_BRUSH("Slate/ControlRig_BugSolid_28x", Icon14x14));
-
-			Set("ControlRig.ResumeExecution", new IMAGE_BRUSH_SVG("Slate/simulate", Icon40x40));
-			Set("ControlRig.ReleaseMode", new IMAGE_BRUSH_SVG("Slate/animation", Icon40x40));
-			Set("ControlRig.DebugMode", new CORE_IMAGE_BRUSH_SVG("Starship/Common/Debug", Icon40x40));
-
 			{
 				FContentRootBracket Bracket(this, EngineEditorSlateDir);
 				// similar style to "LevelViewport.StartingPlayInEditorBorder"
@@ -129,21 +107,6 @@ public:
 			}
 		}
 
-		// Graph styles
-		{
-			Set("ControlRig.Node.PinTree.Arrow_Collapsed_Left", new IMAGE_BRUSH("Slate/TreeArrow_Collapsed_Left", Icon10x10, DefaultForeground));
-			Set("ControlRig.Node.PinTree.Arrow_Collapsed_Hovered_Left", new IMAGE_BRUSH("Slate/TreeArrow_Collapsed_Hovered_Left", Icon10x10, DefaultForeground));
-
-			Set("ControlRig.Node.PinTree.Arrow_Expanded_Left", new IMAGE_BRUSH("Slate/TreeArrow_Expanded_Left", Icon10x10, DefaultForeground));
-			Set("ControlRig.Node.PinTree.Arrow_Expanded_Hovered_Left", new IMAGE_BRUSH("Slate/TreeArrow_Expanded_Hovered_Left", Icon10x10, DefaultForeground));
-
-			Set("ControlRig.Node.PinTree.Arrow_Collapsed_Right", new IMAGE_BRUSH("Slate/TreeArrow_Collapsed_Right", Icon10x10, DefaultForeground));
-			Set("ControlRig.Node.PinTree.Arrow_Collapsed_Hovered_Right", new IMAGE_BRUSH("Slate/TreeArrow_Collapsed_Hovered_Right", Icon10x10, DefaultForeground));
-
-			Set("ControlRig.Node.PinTree.Arrow_Expanded_Right", new IMAGE_BRUSH("Slate/TreeArrow_Expanded_Right", Icon10x10, DefaultForeground));
-			Set("ControlRig.Node.PinTree.Arrow_Expanded_Hovered_Right", new IMAGE_BRUSH("Slate/TreeArrow_Expanded_Hovered_Right", Icon10x10, DefaultForeground));
-		}
-
 		// Tree styles
 		{
 			Set("ControlRig.Tree.BoneUser", new IMAGE_BRUSH("Slate/BoneNonWeighted_16x", Icon16x16));
@@ -188,8 +151,5 @@ public:
 	FSlateColor SpacePickerSelectColor;
 };
 
-#undef IMAGE_BRUSH
-#undef BOX_BRUSH
-#undef BORDER_BRUSH
 #undef TTF_FONT
 #undef OTF_FONT
diff --git a/Engine/Plugins/Animation/ControlRig/Source/ControlRigEditor/Private/Editor/ControlRigEditor.cpp b/Engine/Plugins/Animation/ControlRig/Source/ControlRigEditor/Private/Editor/ControlRigEditor.cpp
--- a/Engine/Plugins/Animation/ControlRig/Source/ControlRigEditor/Private/Editor/ControlRigEditor.cpp
+++ b/Engine/Plugins/Animation/ControlRig/Source/ControlRigEditor/Private/Editor/ControlRigEditor.cpp
@@ -36,6 +36,7 @@
 #include "Misc/MessageDialog.h"
 #include "Framework/MultiBox/MultiBoxBuilder.h"
 #include "ControlRigEditorStyle.h"
+#include "Editor/RigVMEditorStyle.h"
 #include "EditorFontGlyphs.h"
 #include "Widgets/Input/SButton.h"
 #include "Widgets/Input/SNumericEntryBox.h"
@@ -1091,7 +1092,7 @@ void FControlRigEditor::FillToolbar(FToolBarBuilder& ToolbarBuilder)
 	{
 		/*
 		ToolbarBuilder.AddToolBarButton(FControlRigBlueprintCommands::Get().ExecuteGraph,
-			NAME_None, TAttribute<FText>(), TAttribute<FText>(), FSlateIcon(FControlRigEditorStyle::Get().GetStyleSetName(), "ControlRig.ExecuteGraph"));
+			NAME_None, TAttribute<FText>(), TAttribute<FText>(), FSlateIcon(FRigVMEditorStyle::Get().GetStyleSetName(), "RigVM.ExecuteGraph"));
 		*/
 
 		ToolbarBuilder.AddToolBarButton(
@@ -1112,7 +1113,7 @@ void FControlRigEditor::FillToolbar(FToolBarBuilder& ToolbarBuilder)
 			true);
 
 		ToolbarBuilder.AddToolBarButton(FControlRigBlueprintCommands::Get().AutoCompileGraph,
-			NAME_None, TAttribute<FText>(), TAttribute<FText>(), FSlateIcon(FControlRigEditorStyle::Get().GetStyleSetName(), "ControlRig.AutoCompileGraph"));
+			NAME_None, TAttribute<FText>(), TAttribute<FText>(), FSlateIcon(FRigVMEditorStyle::Get().GetStyleSetName(), "RigVM.AutoCompileGraph"));
 
 		ToolbarBuilder.AddWidget(SNew(SBlueprintEditorSelectedDebugObjectWidget, SharedThis(this)));
 
@@ -1499,7 +1500,7 @@ FSlateIcon FControlRigEditor::GetExecutionModeIcon(EControlRigExecutionModeType
 	{
 		return FSlateIcon(FAppStyle::GetAppStyleSetName(), "LevelEditor.Tabs.Debug");
 	}
-	return FSlateIcon(FControlRigEditorStyle::Get().GetStyleSetName(), "ControlRig.ReleaseMode");
+	return FSlateIcon(FRigVMEditorStyle::Get().GetStyleSetName(), "RigVM.ReleaseMode");
 }
 
 FSlateIcon FControlRigEditor::GetExecutionModeIcon() const
diff --git a/Engine/Plugins/Animation/ControlRig/Source/ControlRigEditor/Private/Editor/RigInfluenceMapTabSummoner.cpp b/Engine/Plugins/Animation/ControlRig/Source/ControlRigEditor/Private/Editor/RigInfluenceMapTabSummoner.cpp
--- a/Engine/Plugins/Animation/ControlRig/Source/ControlRigEditor/Private/Editor/RigInfluenceMapTabSummoner.cpp
+++ b/Engine/Plugins/Animation/ControlRig/Source/ControlRigEditor/Private/Editor/RigInfluenceMapTabSummoner.cpp
@@ -2,7 +2,7 @@
 
 #include "Editor/RigInfluenceMapTabSummoner.h"
 #include "Editor/SRigHierarchy.h"
-#include "ControlRigEditorStyle.h"
+#include "Editor/RigVMEditorStyle.h"
 #include "Editor/ControlRigEditor.h"
 #include "SKismetInspector.h"
 #include "ControlRigBlueprint.h"
@@ -16,7 +16,7 @@ FRigInfluenceMapTabSummoner::FRigInfluenceMapTabSummoner(const TSharedRef<FContr
 	, ControlRigEditor(InControlRigEditor)
 {
 	TabLabel = LOCTEXT("RigInfluenceMapTabLabel", "Rig Influence Map");
-	TabIcon = FSlateIcon(FControlRigEditorStyle::Get().GetStyleSetName(), "ControlRig.TabIcon");
+	TabIcon = FSlateIcon(FRigVMEditorStyle::Get().GetStyleSetName(), "RigVM.TabIcon");
 
 	bIsSingleton = true;
 
diff --git a/Engine/Plugins/Animation/ControlRig/Source/ControlRigEditor/Private/Editor/SControlRigValidationWidget.cpp b/Engine/Plugins/Animation/ControlRig/Source/ControlRigEditor/Private/Editor/SControlRigValidationWidget.cpp
--- a/Engine/Plugins/Animation/ControlRig/Source/ControlRigEditor/Private/Editor/SControlRigValidationWidget.cpp
+++ b/Engine/Plugins/Animation/ControlRig/Source/ControlRigEditor/Private/Editor/SControlRigValidationWidget.cpp
@@ -10,7 +10,6 @@
 #include "Widgets/Input/SCheckBox.h"
 #include "Widgets/Input/SComboBox.h"
 #include "Widgets/Input/SEditableTextBox.h"
-#include "ControlRigEditorStyle.h"
 #include "Styling/AppStyle.h"
 #include "SlateOptMacros.h"
 #include "HAL/ConsoleManager.h"
diff --git a/Engine/Plugins/Animation/ControlRig/Source/ControlRigEditor/Private/Editor/SRigHierarchy.cpp b/Engine/Plugins/Animation/ControlRig/Source/ControlRigEditor/Private/Editor/SRigHierarchy.cpp
--- a/Engine/Plugins/Animation/ControlRig/Source/ControlRigEditor/Private/Editor/SRigHierarchy.cpp
+++ b/Engine/Plugins/Animation/ControlRig/Source/ControlRigEditor/Private/Editor/SRigHierarchy.cpp
@@ -29,7 +29,6 @@
 #include "HelperUtil.h"
 #include "Widgets/Text/SInlineEditableTextBlock.h"
 #include "ControlRig.h"
-#include "ControlRigEditorStyle.h"
 #include "HAL/PlatformApplicationMisc.h"
 #include "HAL/PlatformTime.h"
 #include "Dialogs/Dialogs.h"
diff --git a/Engine/Plugins/Runtime/RigVM/RigVM.uplugin b/Engine/Plugins/Runtime/RigVM/RigVM.uplugin
--- a/Engine/Plugins/Runtime/RigVM/RigVM.uplugin
+++ b/Engine/Plugins/Runtime/RigVM/RigVM.uplugin
@@ -11,7 +11,7 @@
 	"MarketplaceURL": "",
 	"SupportURL": "",
 	"EnabledByDefault": true,
-	"CanContainContent": false,
+	"CanContainContent": true,
 	"IsBetaVersion": false,
 	"Installed": false,
 	"Modules" :
diff --git a/Engine/Plugins/Runtime/RigVM/Source/RigVMDeveloper/Private/EdGraph/RigVMEdGraphNode.cpp b/Engine/Plugins/Runtime/RigVM/Source/RigVMDeveloper/Private/EdGraph/RigVMEdGraphNode.cpp
--- a/Engine/Plugins/Runtime/RigVM/Source/RigVMDeveloper/Private/EdGraph/RigVMEdGraphNode.cpp
+++ b/Engine/Plugins/Runtime/RigVM/Source/RigVMDeveloper/Private/EdGraph/RigVMEdGraphNode.cpp
@@ -1285,7 +1285,7 @@ FSlateIcon URigVMEdGraphNode::GetIconAndTint(FLinearColor& OutColor) const
 	static FSlateIcon EventIcon(FAppStyle::GetAppStyleSetName(), "GraphEditor.Event_16x");
 	static FSlateIcon EntryReturnIcon(FAppStyle::GetAppStyleSetName(), "GraphEditor.Default_16x");
 	static FSlateIcon CollapsedNodeIcon(FAppStyle::GetAppStyleSetName(), "GraphEditor.SubGraph_16x");
-	static FSlateIcon TemplateNodeIcon("ControlRigEditorStyle", "ControlRig.Template");
+	static FSlateIcon TemplateNodeIcon("RigVMEditorStyle", "RigVM.Template");
 
 	if (URigVMNode* ModelNode = GetModelNode())
 	{
diff --git a/Engine/Plugins/Runtime/RigVM/Source/RigVMEditor/Private/EdGraph/NodeSpawners/RigVMEdGraphEnumNodeSpawner.cpp b/Engine/Plugins/Runtime/RigVM/Source/RigVMEditor/Private/EdGraph/NodeSpawners/RigVMEdGraphEnumNodeSpawner.cpp
--- a/Engine/Plugins/Runtime/RigVM/Source/RigVMEditor/Private/EdGraph/NodeSpawners/RigVMEdGraphEnumNodeSpawner.cpp
+++ b/Engine/Plugins/Runtime/RigVM/Source/RigVMEditor/Private/EdGraph/NodeSpawners/RigVMEdGraphEnumNodeSpawner.cpp
@@ -32,7 +32,7 @@ URigVMEdGraphEnumNodeSpawner* URigVMEdGraphEnumNodeSpawner::CreateForEnum(UEnum*
 	MenuSignature.Tooltip  = InTooltip;
 	MenuSignature.Category = InCategory;
 	MenuSignature.Keywords = FText::FromString(TEXT("Enum"));
-	MenuSignature.Icon = FSlateIcon(TEXT("ControlRigEditorStyle"), TEXT("ControlRig.RigUnit"));
+	MenuSignature.Icon = FSlateIcon(TEXT("RigVMEditorStyle"), TEXT("RigVM.Unit"));
 
 	return NodeSpawner;
 }
diff --git a/Engine/Plugins/Runtime/RigVM/Source/RigVMEditor/Private/EdGraph/NodeSpawners/RigVMEdGraphTemplateNodeSpawner.cpp b/Engine/Plugins/Runtime/RigVM/Source/RigVMEditor/Private/EdGraph/NodeSpawners/RigVMEdGraphTemplateNodeSpawner.cpp
--- a/Engine/Plugins/Runtime/RigVM/Source/RigVMEditor/Private/EdGraph/NodeSpawners/RigVMEdGraphTemplateNodeSpawner.cpp
+++ b/Engine/Plugins/Runtime/RigVM/Source/RigVMEditor/Private/EdGraph/NodeSpawners/RigVMEdGraphTemplateNodeSpawner.cpp
@@ -56,7 +56,7 @@ URigVMEdGraphTemplateNodeSpawner* URigVMEdGraphTemplateNodeSpawner::CreateFromNo
 	}
 
 	// @TODO: should use details customization-like extensibility system to provide editor only data like this
-	MenuSignature.Icon = FSlateIcon(TEXT("ControlRigEditorStyle"), TEXT("ControlRig.RigUnit"));
+	MenuSignature.Icon = FSlateIcon(TEXT("RigVMEditorStyle"), TEXT("RigVM.Unit"));
 
 	return NodeSpawner;
 }
diff --git a/Engine/Plugins/Runtime/RigVM/Source/RigVMEditor/Private/EdGraph/NodeSpawners/RigVMEdGraphUnitNodeSpawner.cpp b/Engine/Plugins/Runtime/RigVM/Source/RigVMEditor/Private/EdGraph/NodeSpawners/RigVMEdGraphUnitNodeSpawner.cpp
--- a/Engine/Plugins/Runtime/RigVM/Source/RigVMEditor/Private/EdGraph/NodeSpawners/RigVMEdGraphUnitNodeSpawner.cpp
+++ b/Engine/Plugins/Runtime/RigVM/Source/RigVMEditor/Private/EdGraph/NodeSpawners/RigVMEdGraphUnitNodeSpawner.cpp
@@ -103,7 +103,7 @@ URigVMEdGraphUnitNodeSpawner* URigVMEdGraphUnitNodeSpawner::CreateFromStruct(USc
 		}
 		else
 		{
-			MenuSignature.Icon = FSlateIcon(TEXT("ControlRigEditorStyle"), TEXT("ControlRig.RigUnit"));
+			MenuSignature.Icon = FSlateIcon(TEXT("RigVMEditorStyle"), TEXT("RigVM.Unit"));
 		}
 	}
 	else
diff --git a/Engine/Plugins/Runtime/RigVM/Source/RigVMEditor/Private/Editor/RigVMExecutionStackTabSummoner.cpp b/Engine/Plugins/Runtime/RigVM/Source/RigVMEditor/Private/Editor/RigVMExecutionStackTabSummoner.cpp
--- a/Engine/Plugins/Runtime/RigVM/Source/RigVMEditor/Private/Editor/RigVMExecutionStackTabSummoner.cpp
+++ b/Engine/Plugins/Runtime/RigVM/Source/RigVMEditor/Private/Editor/RigVMExecutionStackTabSummoner.cpp
@@ -13,7 +13,7 @@ FRigVMExecutionStackTabSummoner::FRigVMExecutionStackTabSummoner(const TSharedRe
 	, RigVMEditor(InRigVMEditor)
 {
 	TabLabel = LOCTEXT("RigVMExecutionStackTabLabel", "Execution Stack");
-	TabIcon = FSlateIcon(TEXT("ControlRigEditorStyle"), "ExecutionStack.TabIcon");
+	TabIcon = FSlateIcon(TEXT("RigVMEditorStyle"), "ExecutionStack.TabIcon");
 
 	bIsSingleton = true;
 
diff --git a/Engine/Plugins/Runtime/RigVM/Source/RigVMEditor/Private/RigVMEditorModule.cpp b/Engine/Plugins/Runtime/RigVM/Source/RigVMEditor/Private/RigVMEditorModule.cpp
--- a/Engine/Plugins/Runtime/RigVM/Source/RigVMEditor/Private/RigVMEditorModule.cpp
+++ b/Engine/Plugins/Runtime/RigVM/Source/RigVMEditor/Private/RigVMEditorModule.cpp
@@ -6,6 +6,8 @@
 
 #include "RigVMEditorModule.h"
 #include "Modules/ModuleManager.h"
+#include "Editor/RigVMExecutionStackCommands.h"
+#include "Editor/RigVMEditorStyle.h"
 
 DEFINE_LOG_CATEGORY(LogRigVMEditor);
 
@@ -18,6 +20,8 @@ FRigVMEditorModule& FRigVMEditorModule::Get()
 
 void FRigVMEditorModule::StartupModule()
 {
+	FRigVMExecutionStackCommands::Register();
+	FRigVMEditorStyle::Get();
 }
 
 void FRigVMEditorModule::ShutdownModule()
diff --git a/Engine/Plugins/Runtime/RigVM/Source/RigVMEditor/Private/Widgets/SRigVMExecutionStackView.cpp b/Engine/Plugins/Runtime/RigVM/Source/RigVMEditor/Private/Widgets/SRigVMExecutionStackView.cpp
--- a/Engine/Plugins/Runtime/RigVM/Source/RigVMEditor/Private/Widgets/SRigVMExecutionStackView.cpp
+++ b/Engine/Plugins/Runtime/RigVM/Source/RigVMEditor/Private/Widgets/SRigVMExecutionStackView.cpp
@@ -59,7 +59,7 @@ void SRigStackItem::Construct(const FArguments& InArgs, const TSharedRef<STableV
 	{
 		case ERigStackEntry::Operator:
 		{
-			Icon = FSlateIcon(TEXT("RigVMEditor"), "ControlRig.RigUnit").GetIcon();
+			Icon = FSlateIcon(TEXT("RigVMEditor"), "RigVM.Unit").GetIcon();
 			break;
 		}
 		case ERigStackEntry::Info:
diff --git a/Engine/Plugins/Runtime/RigVM/Source/RigVMEditor/Private/Widgets/SRigVMGraphNode.cpp b/Engine/Plugins/Runtime/RigVM/Source/RigVMEditor/Private/Widgets/SRigVMGraphNode.cpp
--- a/Engine/Plugins/Runtime/RigVM/Source/RigVMEditor/Private/Widgets/SRigVMGraphNode.cpp
+++ b/Engine/Plugins/Runtime/RigVM/Source/RigVMEditor/Private/Widgets/SRigVMGraphNode.cpp
@@ -87,7 +87,7 @@ void SRigVMGraphNode::Construct( const FArguments& InArgs )
 
 	UpdatePinTreeView();
 
-	const FSlateBrush* ImageBrush = FSlateIcon(TEXT("ControlRigEditorStyle"), TEXT("ControlRig.Bug.Dot")).GetIcon();
+	const FSlateBrush* ImageBrush = FSlateIcon(TEXT("RigVMEditorStyle"), TEXT("RigVM.Bug.Dot")).GetIcon();
 
 	VisualDebugIndicatorWidget =
 		SNew(SImage)
@@ -1365,23 +1365,23 @@ const FSlateBrush* SRigVMGraphNode::GetExpanderImage(int32 InPinInfoIndex, bool
 	
 	if(ExpandedHoveredLeftBrush == nullptr)
 	{
-		static const TCHAR* ExpandedHoveredLeftName = TEXT("ControlRig.Node.PinTree.Arrow_Expanded_Hovered_Left");
-		static const TCHAR* ExpandedHoveredRightName = TEXT("ControlRig.Node.PinTree.Arrow_Expanded_Hovered_Right");
-		static const TCHAR* ExpandedLeftName = TEXT("ControlRig.Node.PinTree.Arrow_Expanded_Left");
-		static const TCHAR* ExpandedRightName = TEXT("ControlRig.Node.PinTree.Arrow_Expanded_Right");
-		static const TCHAR* CollapsedHoveredLeftName = TEXT("ControlRig.Node.PinTree.Arrow_Collapsed_Hovered_Left");
-		static const TCHAR* CollapsedHoveredRightName = TEXT("ControlRig.Node.PinTree.Arrow_Collapsed_Hovered_Right");
-		static const TCHAR* CollapsedLeftName = TEXT("ControlRig.Node.PinTree.Arrow_Collapsed_Left");
-		static const TCHAR* CollapsedRightName = TEXT("ControlRig.Node.PinTree.Arrow_Collapsed_Right");
+		static const TCHAR* ExpandedHoveredLeftName = TEXT("RigVM.Node.PinTree.Arrow_Expanded_Hovered_Left");
+		static const TCHAR* ExpandedHoveredRightName = TEXT("RigVM.Node.PinTree.Arrow_Expanded_Hovered_Right");
+		static const TCHAR* ExpandedLeftName = TEXT("RigVM.Node.PinTree.Arrow_Expanded_Left");
+		static const TCHAR* ExpandedRightName = TEXT("RigVM.Node.PinTree.Arrow_Expanded_Right");
+		static const TCHAR* CollapsedHoveredLeftName = TEXT("RigVM.Node.PinTree.Arrow_Collapsed_Hovered_Left");
+		static const TCHAR* CollapsedHoveredRightName = TEXT("RigVM.Node.PinTree.Arrow_Collapsed_Hovered_Right");
+		static const TCHAR* CollapsedLeftName = TEXT("RigVM.Node.PinTree.Arrow_Collapsed_Left");
+		static const TCHAR* CollapsedRightName = TEXT("RigVM.Node.PinTree.Arrow_Collapsed_Right");
 		
-		ExpandedHoveredLeftBrush = FSlateIcon(TEXT("ControlRigEditorStyle"), ExpandedHoveredLeftName).GetIcon();
-		ExpandedHoveredRightBrush = FSlateIcon(TEXT("ControlRigEditorStyle"), ExpandedHoveredRightName).GetIcon();
-		ExpandedLeftBrush = FSlateIcon(TEXT("ControlRigEditorStyle"), ExpandedLeftName).GetIcon();
-		ExpandedRightBrush = FSlateIcon(TEXT("ControlRigEditorStyle"), ExpandedRightName).GetIcon();
-		CollapsedHoveredLeftBrush = FSlateIcon(TEXT("ControlRigEditorStyle"), CollapsedHoveredLeftName).GetIcon();
-		CollapsedHoveredRightBrush = FSlateIcon(TEXT("ControlRigEditorStyle"), CollapsedHoveredRightName).GetIcon();
-		CollapsedLeftBrush = FSlateIcon(TEXT("ControlRigEditorStyle"), CollapsedLeftName).GetIcon();
-		CollapsedRightBrush = FSlateIcon(TEXT("ControlRigEditorStyle"), CollapsedRightName).GetIcon();
+		ExpandedHoveredLeftBrush = FSlateIcon(TEXT("RigVMEditorStyle"), ExpandedHoveredLeftName).GetIcon();
+		ExpandedHoveredRightBrush = FSlateIcon(TEXT("RigVMEditorStyle"), ExpandedHoveredRightName).GetIcon();
+		ExpandedLeftBrush = FSlateIcon(TEXT("RigVMEditorStyle"), ExpandedLeftName).GetIcon();
+		ExpandedRightBrush = FSlateIcon(TEXT("RigVMEditorStyle"), ExpandedRightName).GetIcon();
+		CollapsedHoveredLeftBrush = FSlateIcon(TEXT("RigVMEditorStyle"), CollapsedHoveredLeftName).GetIcon();
+		CollapsedHoveredRightBrush = FSlateIcon(TEXT("RigVMEditorStyle"), CollapsedHoveredRightName).GetIcon();
+		CollapsedLeftBrush = FSlateIcon(TEXT("RigVMEditorStyle"), CollapsedLeftName).GetIcon();
+		CollapsedRightBrush = FSlateIcon(TEXT("RigVMEditorStyle"), CollapsedRightName).GetIcon();
 	}
 
 	if(PinInfos.IsValidIndex(InPinInfoIndex))
diff --git a/Engine/Plugins/Runtime/RigVM/Source/RigVMEditor/Public/Editor/RigVMEditorStyle.h b/Engine/Plugins/Runtime/RigVM/Source/RigVMEditor/Public/Editor/RigVMEditorStyle.h
new file mode 100644
--- /dev/null
+++ b/Engine/Plugins/Runtime/RigVM/Source/RigVMEditor/Public/Editor/RigVMEditorStyle.h
@@ -0,0 +1,108 @@
+// Copyright Epic Games, Inc. All Rights Reserved.
+
+#pragma once
+
+#include "Styling/SlateStyle.h"
+#include "Styling/SlateStyleRegistry.h"
+#include "Styling/SlateStyleMacros.h"
+#include "Styling/StyleColors.h"
+#include "Brushes/SlateRoundedBoxBrush.h"
+
+class RIGVMEDITOR_API FRigVMEditorStyle final
+	: public FSlateStyleSet
+{
+	class FContentRootBracket
+	{
+	public:
+		FContentRootBracket(FRigVMEditorStyle* InStyle, const FString& NewContentRoot)
+			: Style(InStyle)
+			, PreviousContentRoot(InStyle->GetContentRootDir())
+		{
+			Style->SetContentRoot(NewContentRoot);
+		}
+
+		~FContentRootBracket()
+		{
+			Style->SetContentRoot(PreviousContentRoot);
+		}
+	private:
+		FRigVMEditorStyle* Style;
+		FString PreviousContentRoot;
+	};
+	
+public:
+	FRigVMEditorStyle()
+		: FSlateStyleSet("RigVMEditorStyle")
+	{
+		const FVector2D Icon10x10(10.0f, 10.0f);
+		const FVector2D Icon14x14(14.0f, 14.0f);
+		const FVector2D Icon16x16(16.0f, 16.0f);
+		const FVector2D Icon20x20(20.0f, 20.0f);
+		const FVector2D Icon24x24(24.0f, 24.0f);
+		const FVector2D Icon32x32(32.0f, 32.0f);
+		const FVector2D Icon40x40(40.0f, 40.0f);
+		const FString RigVMPluginContentDir = FPaths::EnginePluginsDir() / TEXT("Runtime/RigVM/Content");
+		const FString EngineEditorSlateDir = FPaths::EngineContentDir() / TEXT("Editor/Slate");
+		SetContentRoot(RigVMPluginContentDir);
+		SetCoreContentRoot(EngineEditorSlateDir);
+
+		const FSlateColor DefaultForeground(FLinearColor(0.72f, 0.72f, 0.72f, 1.f));
+
+		// Edit mode styles
+		Set("RigVMEditMode", new IMAGE_BRUSH_SVG("Slate/RigVM_Release", Icon40x40));
+		Set("RigVMEditMode.Small", new IMAGE_BRUSH_SVG("Slate/RigVM_Release", Icon20x20));
+
+		// RigVM Editor styles
+		{
+			// tab icons
+			Set("RigVM.TabIcon", new IMAGE_BRUSH("Slate/RigVM_Tab_16x", Icon16x16));
+			Set("ExecutionStack.TabIcon", new IMAGE_BRUSH_SVG("Slate/RigVM_ExecutionStack", Icon16x16));
+			
+			// icons for control units
+			Set("RigVM.Unit", new IMAGE_BRUSH("Slate/RigVM_Unit_16x", Icon16x16));
+			Set("RigVM.Template", new IMAGE_BRUSH("Slate/RigVM_Template_16x", Icon16x16));
+
+			Set("RigVM.AutoCompileGraph", new IMAGE_BRUSH("Slate/RigVM_AutoCompile", Icon20x20));
+			Set("RigVM.AutoCompileGraph.Small", new IMAGE_BRUSH("Slate/RigVM_AutoCompile", Icon20x20));
+
+			Set("RigVM.Bug.Dot", new IMAGE_BRUSH("Slate/RigVM_BugDot_32x", Icon16x16));
+			Set("RigVM.Bug.Normal", new IMAGE_BRUSH("Slate/RigVM_Bug_28x", Icon14x14));
+			Set("RigVM.Bug.Open", new IMAGE_BRUSH("Slate/RigVM_BugOpen_28x", Icon14x14));
+			Set("RigVM.Bug.Solid", new IMAGE_BRUSH("Slate/RigVM_BugSolid_28x", Icon14x14));
+
+			Set("RigVM.ResumeExecution", new IMAGE_BRUSH_SVG("Slate/RigVM_Resume", Icon40x40));
+			Set("RigVM.ReleaseMode", new IMAGE_BRUSH_SVG("Slate/RigVM_Release", Icon40x40));
+			Set("RigVM.DebugMode", new CORE_IMAGE_BRUSH_SVG("Starship/Common/Debug", Icon40x40));
+		}
+
+		// Graph styles
+		{
+			Set("RigVM.Node.PinTree.Arrow_Collapsed_Left", new IMAGE_BRUSH("Slate/RigVM_TreeArrow_Collapsed_Left", Icon10x10, DefaultForeground));
+			Set("RigVM.Node.PinTree.Arrow_Collapsed_Hovered_Left", new IMAGE_BRUSH("Slate/RigVM_TreeArrow_Collapsed_Hovered_Left", Icon10x10, DefaultForeground));
+
+			Set("RigVM.Node.PinTree.Arrow_Expanded_Left", new IMAGE_BRUSH("Slate/RigVM_TreeArrow_Expanded_Left", Icon10x10, DefaultForeground));
+			Set("RigVM.Node.PinTree.Arrow_Expanded_Hovered_Left", new IMAGE_BRUSH("Slate/RigVM_TreeArrow_Expanded_Hovered_Left", Icon10x10, DefaultForeground));
+
+			Set("RigVM.Node.PinTree.Arrow_Collapsed_Right", new IMAGE_BRUSH("Slate/RigVM_TreeArrow_Collapsed_Right", Icon10x10, DefaultForeground));
+			Set("RigVM.Node.PinTree.Arrow_Collapsed_Hovered_Right", new IMAGE_BRUSH("Slate/RigVM_TreeArrow_Collapsed_Hovered_Right", Icon10x10, DefaultForeground));
+
+			Set("RigVM.Node.PinTree.Arrow_Expanded_Right", new IMAGE_BRUSH("Slate/RigVM_TreeArrow_Expanded_Right", Icon10x10, DefaultForeground));
+			Set("RigVM.Node.PinTree.Arrow_Expanded_Hovered_Right", new IMAGE_BRUSH("Slate/RigVM_TreeArrow_Expanded_Hovered_Right", Icon10x10, DefaultForeground));
+		}
+
+		FSlateStyleRegistry::RegisterSlateStyle(*this);
+	}
+
+	static FRigVMEditorStyle& Get()
+	{
+		static FRigVMEditorStyle Inst;
+		return Inst;
+	}
+	
+	~FRigVMEditorStyle()
+	{
+		FSlateStyleRegistry::UnRegisterSlateStyle(*this);
+	}
+
+	FSlateColor SpacePickerSelectColor;
+};
diff --git a/Engine/Plugins/Runtime/RigVM/Source/RigVMEditor/Public/Editor/RigVMExecutionStackCommands.h b/Engine/Plugins/Runtime/RigVM/Source/RigVMEditor/Public/Editor/RigVMExecutionStackCommands.h
--- a/Engine/Plugins/Runtime/RigVM/Source/RigVMEditor/Public/Editor/RigVMExecutionStackCommands.h
+++ b/Engine/Plugins/Runtime/RigVM/Source/RigVMEditor/Public/Editor/RigVMExecutionStackCommands.h
@@ -13,7 +13,7 @@ public:
 		"ControlRigStack",
 		NSLOCTEXT("Contexts", "RigStack", "Execution Stack"),
 		NAME_None, // "MainFrame" // @todo Fix this crash
-		"ControlRigEditorStyle" // Icon Style Set
+		"RigVMEditorStyle" // Icon Style Set
 	)
 	{}
 	
