Control Rig: Ensure we only register the singleton instance.

#rnx

diff --git a/Engine/Plugins/Runtime/RigVM/Source/RigVMEditor/Private/RigVMEditorModule.cpp b/Engine/Plugins/Runtime/RigVM/Source/RigVMEditor/Private/RigVMEditorModule.cpp
--- a/Engine/Plugins/Runtime/RigVM/Source/RigVMEditor/Private/RigVMEditorModule.cpp
+++ b/Engine/Plugins/Runtime/RigVM/Source/RigVMEditor/Private/RigVMEditorModule.cpp
@@ -21,9 +21,10 @@ FRigVMEditorModule& FRigVMEditorModule::Get()
 void FRigVMEditorModule::StartupModule()
 {
 	FRigVMExecutionStackCommands::Register();
-	FRigVMEditorStyle::Get();
+	FRigVMEditorStyle::Register();
 }
 
 void FRigVMEditorModule::ShutdownModule()
 {
+	FRigVMEditorStyle::Unregister();
 }
diff --git a/Engine/Plugins/Runtime/RigVM/Source/RigVMEditor/Public/Editor/RigVMEditorStyle.h b/Engine/Plugins/Runtime/RigVM/Source/RigVMEditor/Public/Editor/RigVMEditorStyle.h
--- a/Engine/Plugins/Runtime/RigVM/Source/RigVMEditor/Public/Editor/RigVMEditorStyle.h
+++ b/Engine/Plugins/Runtime/RigVM/Source/RigVMEditor/Public/Editor/RigVMEditorStyle.h
@@ -30,7 +30,6 @@ class RIGVMEDITOR_API FRigVMEditorStyle final
 		FString PreviousContentRoot;
 	};
 	
-public:
 	FRigVMEditorStyle()
 		: FSlateStyleSet("RigVMEditorStyle")
 	{
@@ -89,20 +88,23 @@ public:
 			Set("RigVM.Node.PinTree.Arrow_Expanded_Right", new IMAGE_BRUSH("Slate/RigVM_TreeArrow_Expanded_Right", Icon10x10, DefaultForeground));
 			Set("RigVM.Node.PinTree.Arrow_Expanded_Hovered_Right", new IMAGE_BRUSH("Slate/RigVM_TreeArrow_Expanded_Hovered_Right", Icon10x10, DefaultForeground));
 		}
-
-		FSlateStyleRegistry::RegisterSlateStyle(*this);
 	}
+public:
 
 	static FRigVMEditorStyle& Get()
 	{
 		static FRigVMEditorStyle Inst;
+		
 		return Inst;
 	}
-	
-	~FRigVMEditorStyle()
+
+	static void Register()
 	{
-		FSlateStyleRegistry::UnRegisterSlateStyle(*this);
+		FSlateStyleRegistry::RegisterSlateStyle(Get());
 	}
 
-	FSlateColor SpacePickerSelectColor;
+	static void Unregister()
+	{
+		FSlateStyleRegistry::UnRegisterSlateStyle(Get());
+	}
 };
