RigVMEditor: Refactoring editor module, blueprint actions etc

diff --git a/Engine/Plugins/Animation/ControlRig/Source/ControlRigEditor/Private/ControlRigEditorModule.h b/Engine/Plugins/Animation/ControlRig/Source/ControlRigEditor/Private/ControlRigEditorModule.h
--- a/Engine/Plugins/Animation/ControlRig/Source/ControlRigEditor/Private/ControlRigEditorModule.h
+++ b/Engine/Plugins/Animation/ControlRig/Source/ControlRigEditor/Private/ControlRigEditorModule.h
@@ -29,15 +29,13 @@ public:
 	virtual void StartupModule() override;
 	virtual void ShutdownModule() override;
 
+	/** FRigVMEditorModule interface */
+	virtual UClass* GetRigVMBlueprintClass() const override;
+	virtual void GetNodeContextMenuActions(URigVMBlueprint* RigVMBlueprint, const URigVMEdGraphNode* EdGraphNode, URigVMNode* ModelNode, UToolMenu* Menu) const override;
+	virtual void GetPinContextMenuActions(URigVMBlueprint* RigVMBlueprint, const UEdGraphPin* EdGraphPin, URigVMPin* ModelPin, UToolMenu* Menu) const override;
+
 	/** IControlRigEditorModule interface */
 	virtual TSharedRef<IControlRigEditor> CreateControlRigEditor(const EToolkitMode::Type Mode, const TSharedPtr<IToolkitHost>& InitToolkitHost, UControlRigBlueprint* Blueprint) override;
-	DECLARE_DELEGATE_RetVal_TwoParams(TSharedRef<FExtender>, FControlRigEditorToolbarExtender, const TSharedRef<FUICommandList> /*InCommandList*/, TSharedRef<IControlRigEditor> /*InControlRigEditor*/);
-	virtual TArray<FControlRigEditorToolbarExtender>& GetAllControlRigEditorToolbarExtenders() override { return ControlRigEditorToolbarExtenders; }
-	/** IHasMenuExtensibility interface */
-	virtual TSharedPtr<FExtensibilityManager> GetMenuExtensibilityManager() override { return MenuExtensibilityManager; }
-
-	/** IHasToolBarExtensibility interface */
-	virtual TSharedPtr<FExtensibilityManager> GetToolBarExtensibilityManager() override { return ToolBarExtensibilityManager; }
 	
 	/** Animation Toolbar Extender*/
 	void AddControlRigExtenderToToolMenu(FName InToolMenuName);
@@ -50,17 +48,7 @@ public:
 	static void UnLinkLevelSequence(UAnimSequence* AnimSequence);
 	void ExtendAnimSequenceMenu();
 
-	virtual void GetTypeActions(UControlRigBlueprint* CRB, FBlueprintActionDatabaseRegistrar& ActionRegistrar) override;
-	virtual void GetInstanceActions(UControlRigBlueprint* CRB, FBlueprintActionDatabaseRegistrar& ActionRegistrar) override;
-	virtual FConnectionDrawingPolicy* CreateConnectionDrawingPolicy(int32 InBackLayerID, int32 InFrontLayerID, float InZoomFactor, const FSlateRect& InClippingRect, class FSlateWindowElementList& InDrawElements, class UEdGraph* InGraphObj) override;
-	virtual void GetContextMenuActions(const UControlRigGraphSchema* Schema, class UToolMenu* Menu, class UGraphNodeContextMenuContext* Context) const override;
-
-	/* FStructureEditorUtils::INotifyOnStructChanged Interface, used to respond to changes to user defined structs */
-	virtual void PreChange(const UUserDefinedStruct* Changed, FStructureEditorUtils::EStructureEditorChangeInfo ChangedType) override;
-	virtual void PostChange(const UUserDefinedStruct* Changed, FStructureEditorUtils::EStructureEditorChangeInfo ChangedType) override;
 private:
-	/** Handle a new animation controller blueprint being created */
-	void HandleNewBlueprintCreated(UBlueprint* InBlueprint);
 
 	/** Handle for our sequencer control rig parameter track editor */
 	FDelegateHandle ControlRigParameterTrackCreateEditorHandle;
@@ -70,23 +58,7 @@ private:
 	/** StaticClass is not safe on shutdown, so we cache the name, and use this to unregister on shut down */
 	TArray<FName> ClassesToUnregisterOnShutdown;
 	TArray<FName> PropertiesToUnregisterOnShutdown;
-
-	/** Extensibility managers */
-	TSharedPtr<FExtensibilityManager> MenuExtensibilityManager;
-	TSharedPtr<FExtensibilityManager> ToolBarExtensibilityManager;
-	TArray<FControlRigEditorToolbarExtender> ControlRigEditorToolbarExtenders;
-
-	/** Node factory for the control rig graph */
-	TSharedPtr<FRigVMEdGraphPanelNodeFactory> ControlRigGraphPanelNodeFactory;
-
-	/** Pin factory for the control rig graph */
-	TSharedPtr<FControlRigGraphPanelPinFactory> ControlRigGraphPanelPinFactory;
-
-	/** Delegate handles for blueprint utils */
-	FDelegateHandle RefreshAllNodesDelegateHandle;
-	FDelegateHandle ReconstructAllNodesDelegateHandle;
-	FDelegateHandle BlueprintVariableCustomizationHandle;
-
+	
 	/** Param to hold Filter Result to pass to Filter*/
 	bool bFilterAssetBySkeleton;
 
diff --git a/Engine/Plugins/Animation/ControlRig/Source/ControlRigEditor/Private/Editor/ControlRigEditor.cpp b/Engine/Plugins/Animation/ControlRig/Source/ControlRigEditor/Private/Editor/ControlRigEditor.cpp
--- a/Engine/Plugins/Animation/ControlRig/Source/ControlRigEditor/Private/Editor/ControlRigEditor.cpp
+++ b/Engine/Plugins/Animation/ControlRig/Source/ControlRigEditor/Private/Editor/ControlRigEditor.cpp
@@ -96,7 +96,6 @@
 #include "Units/Execution/RigUnit_SequenceExecution.h"
 #include "Editor/ControlRigContextMenuContext.h"
 #include "Types/ISlateMetaData.h"
-#include "ControlRigGraphDetails.h"
 #include "Kismet2/KismetDebugUtilities.h"
 #include "Kismet2/WatchedPin.h"
 #include "ToolMenus.h"
@@ -107,6 +106,7 @@
 #include "RigVMModel/Nodes/RigVMAggregateNode.h"
 #include "AnimationEditorViewportClient.h"
 #include "RigVMCore/RigVMExecuteContext.h"
+#include "Editor/RigVMGraphDetailCustomization.h"
 
 #define LOCTEXT_NAMESPACE "ControlRigEditor"
 
@@ -296,6 +296,11 @@ FControlRigEditor::~FControlRigEditor()
 	}
 }
 
+UClass* FControlRigEditor::GetDetailWrapperClass() const
+{
+	return UControlRigWrapperObject::StaticClass();
+}
+
 UControlRigBlueprint* FControlRigEditor::GetControlRigBlueprint() const
 {
 	return Cast<UControlRigBlueprint>(GetRigVMBlueprint());
@@ -630,7 +635,7 @@ void FControlRigEditor::InitControlRigEditor(const EToolkitMode::Type Mode, cons
 		Inspector->GetPropertyView()->RegisterInstancedCustomPropertyTypeLayout(StructToCustomize->GetFName(),
 			FOnGetPropertyTypeCustomizationInstance::CreateLambda([=]()
 			{
-				return FControlRigGraphMathTypeDetails::MakeInstance();
+				return FRigVMGraphMathTypeDetailCustomization::MakeInstance();
 			}));
 	}
 
@@ -1068,7 +1073,7 @@ void FControlRigEditor::ExtendToolbar()
 	FControlRigEditorModule& ControlRigEditorModule = FModuleManager::LoadModuleChecked<FControlRigEditorModule>("ControlRigEditor");
 	AddToolbarExtender(ControlRigEditorModule.GetToolBarExtensibilityManager()->GetAllExtenders(GetToolkitCommands(), GetEditingObjects()));
 
-	TArray<IControlRigEditorModule::FControlRigEditorToolbarExtender> ToolbarExtenderDelegates = ControlRigEditorModule.GetAllControlRigEditorToolbarExtenders();
+	TArray<FRigVMEditorModule::FRigVMEditorToolbarExtender> ToolbarExtenderDelegates = ControlRigEditorModule.GetAllRigVMEditorToolbarExtenders();
 
 	for (auto& ToolbarExtenderDelegate : ToolbarExtenderDelegates)
 	{
@@ -1891,15 +1896,6 @@ void FControlRigEditor::SetDetailObjects(const TArray<UObject*>& InObjects, bool
 				WrapperObject->GetWrappedPropertyChangedChainEvent().AddSP(this, &FControlRigEditor::OnWrappedPropertyChangedChainEvent);
 				WrapperObject->AddToRoot();
 
-				if(bClassCreated)
-				{
-					UClass* WrapperClass = WrapperObject->GetClassForNodes(ModelNodes, false);
-					check(WrapperClass);
-
-					FPropertyEditorModule& EditModule = FModuleManager::Get().GetModuleChecked<FPropertyEditorModule>("PropertyEditor");
-					EditModule.RegisterCustomClassLayout(WrapperClass->GetFName(), FOnGetDetailCustomizationInstance::CreateStatic(&FControlRigWrappedNodeDetails::MakeInstance));
-				}
-
 				// todo: use transform widget for transforms
 				// todo: use rotation widget for rotations
 				
diff --git a/Engine/Plugins/Animation/ControlRig/Source/ControlRigEditor/Private/Editor/ControlRigEditor.h b/Engine/Plugins/Animation/ControlRig/Source/ControlRigEditor/Private/Editor/ControlRigEditor.h
--- a/Engine/Plugins/Animation/ControlRig/Source/ControlRigEditor/Private/Editor/ControlRigEditor.h
+++ b/Engine/Plugins/Animation/ControlRig/Source/ControlRigEditor/Private/Editor/ControlRigEditor.h
@@ -72,6 +72,9 @@ public:
 	FControlRigEditor();
 	virtual ~FControlRigEditor();
 
+	// FRigVMEditor interface
+	virtual UClass* GetDetailWrapperClass() const;
+
 	// FBlueprintEditor interface
 	virtual UBlueprint* GetBlueprintObj() const override;
 	virtual TSubclassOf<UEdGraphSchema> GetDefaultSchemaClass() const override;
diff --git a/Engine/Plugins/Animation/ControlRig/Source/ControlRigEditor/Private/Editor/ControlRigWrapperObject.cpp b/Engine/Plugins/Animation/ControlRig/Source/ControlRigEditor/Private/Editor/ControlRigWrapperObject.cpp
--- a/Engine/Plugins/Animation/ControlRig/Source/ControlRigEditor/Private/Editor/ControlRigWrapperObject.cpp
+++ b/Engine/Plugins/Animation/ControlRig/Source/ControlRigEditor/Private/Editor/ControlRigWrapperObject.cpp
@@ -4,7 +4,6 @@
 
 #if WITH_EDITOR
 #include "ControlRigElementDetails.h"
-#include "ControlRigLocalVariableDetails.h"
 #include "PropertyEditorModule.h"
 #include "Modules/ModuleManager.h"
 #endif
@@ -39,14 +38,6 @@ UClass* UControlRigWrapperObject::GetClassForStruct(UScriptStruct* InStruct, boo
 			}
 		}
 	}
-	else if(InStruct->IsChildOf(FRigVMGraphVariableDescription::StaticStruct()))
-	{
-		FPropertyEditorModule& PropertyEditorModule = FModuleManager::LoadModuleChecked<FPropertyEditorModule>("PropertyEditor");
-		if (!PropertyEditorModule.GetClassNameToDetailLayoutNameMap().Contains(WrapperClassName))
-		{
-			PropertyEditorModule.RegisterCustomClassLayout(WrapperClassName, FOnGetDetailCustomizationInstance::CreateStatic(&FRigVMLocalVariableDetails::MakeInstance));
-		}
-	}
 #endif
 
 	return Class;
diff --git a/Engine/Plugins/Animation/ControlRig/Source/ControlRigEditor/Public/IControlRigEditorModule.h b/Engine/Plugins/Animation/ControlRig/Source/ControlRigEditor/Public/IControlRigEditorModule.h
--- a/Engine/Plugins/Animation/ControlRig/Source/ControlRigEditor/Public/IControlRigEditorModule.h
+++ b/Engine/Plugins/Animation/ControlRig/Source/ControlRigEditor/Public/IControlRigEditorModule.h
@@ -3,12 +3,8 @@
 #pragma once
 
 #include "CoreMinimal.h"
-#include "Modules/ModuleInterface.h"
-#include "Framework/Commands/UICommandList.h"
-#include "Framework/MultiBox/MultiBoxExtender.h"
-#include "Toolkits/AssetEditorToolkit.h"
 #include "IControlRigEditor.h"
-#include "Kismet2/StructureEditorUtils.h"
+#include "RigVMEditorModule.h"
 
 DECLARE_LOG_CATEGORY_EXTERN(LogControlRigEditor, Log, All);
 
@@ -18,7 +14,7 @@ class UControlRigGraphNode;
 class UControlRigGraphSchema;
 class FConnectionDrawingPolicy;
 
-class IControlRigEditorModule : public IModuleInterface, public IHasMenuExtensibility, public IHasToolBarExtensibility, public FStructureEditorUtils::INotifyOnStructChanged
+class IControlRigEditorModule : public FRigVMEditorModule
 {
 public:
 
@@ -38,13 +34,4 @@ public:
 	 */
 	virtual TSharedRef<IControlRigEditor> CreateControlRigEditor(const EToolkitMode::Type Mode, const TSharedPtr<IToolkitHost>& InitToolkitHost, UControlRigBlueprint* Blueprint) = 0;
 
-	/** Get all toolbar extenders */
-	DECLARE_DELEGATE_RetVal_TwoParams(TSharedRef<FExtender>, FControlRigEditorToolbarExtender, const TSharedRef<FUICommandList> /*InCommandList*/, TSharedRef<IControlRigEditor> /*InControlRigEditor*/);
-	virtual TArray<FControlRigEditorToolbarExtender>& GetAllControlRigEditorToolbarExtenders() = 0;
-
-	virtual void GetTypeActions(UControlRigBlueprint* CRB, FBlueprintActionDatabaseRegistrar& ActionRegistrar) = 0;
-	virtual void GetInstanceActions(UControlRigBlueprint* CRB, FBlueprintActionDatabaseRegistrar& ActionRegistrar) = 0;
-	virtual FConnectionDrawingPolicy* CreateConnectionDrawingPolicy(int32 InBackLayerID, int32 InFrontLayerID, float InZoomFactor, const FSlateRect& InClippingRect, class FSlateWindowElementList& InDrawElements, class UEdGraph* InGraphObj) = 0;
-	virtual void GetContextMenuActions(const UControlRigGraphSchema* Schema, class UToolMenu* Menu, class UGraphNodeContextMenuContext* Context) const = 0;
-
 };
diff --git a/Engine/Plugins/Runtime/RigVM/Source/RigVMDeveloper/Private/RigVMBlueprint.cpp b/Engine/Plugins/Runtime/RigVM/Source/RigVMDeveloper/Private/RigVMBlueprint.cpp
--- a/Engine/Plugins/Runtime/RigVM/Source/RigVMDeveloper/Private/RigVMBlueprint.cpp
+++ b/Engine/Plugins/Runtime/RigVM/Source/RigVMDeveloper/Private/RigVMBlueprint.cpp
@@ -2503,7 +2503,7 @@ FOnRigVMCompiledEvent& URigVMBlueprint::OnVMCompiled()
 	return VMCompiledEvent;
 }
 
-UClass* URigVMBlueprint::GetRigVMHostClass()
+UClass* URigVMBlueprint::GetRigVMHostClass() const
 {
 	return GeneratedClass;
 }
diff --git a/Engine/Plugins/Runtime/RigVM/Source/RigVMDeveloper/Public/EdGraph/RigVMEdGraph.h b/Engine/Plugins/Runtime/RigVM/Source/RigVMDeveloper/Public/EdGraph/RigVMEdGraph.h
--- a/Engine/Plugins/Runtime/RigVM/Source/RigVMDeveloper/Public/EdGraph/RigVMEdGraph.h
+++ b/Engine/Plugins/Runtime/RigVM/Source/RigVMDeveloper/Public/EdGraph/RigVMEdGraph.h
@@ -60,6 +60,8 @@ public:
 	void CacheEntryNameList();
 	const TArray<TSharedPtr<FString>>* GetEntryNameList(URigVMPin* InPin = nullptr) const;
 
+	virtual const TArray<TSharedPtr<FString>>* GetNameListForWidget(const FString& InWidgetName) const { return nullptr; }
+
 	UPROPERTY()
 	FString ModelNodePath;
 
diff --git a/Engine/Plugins/Runtime/RigVM/Source/RigVMDeveloper/Public/EdGraph/RigVMEdGraphNode.h b/Engine/Plugins/Runtime/RigVM/Source/RigVMDeveloper/Public/EdGraph/RigVMEdGraphNode.h
--- a/Engine/Plugins/Runtime/RigVM/Source/RigVMDeveloper/Public/EdGraph/RigVMEdGraphNode.h
+++ b/Engine/Plugins/Runtime/RigVM/Source/RigVMDeveloper/Public/EdGraph/RigVMEdGraphNode.h
@@ -271,8 +271,8 @@ private:
 	mutable TOptional<bool> DrawAsCompactNodeCache;
 
 	friend class SRigVMGraphNode;
-	friend class FRigVMArgumentLayout;
-	friend class FRigVMEdGraphDetails;
+	friend class FRigVMFunctionArgumentLayout;
+	friend class FRigVMEdGraphDetailCustomization;
 	friend class URigVMEdGraphTemplateNodeSpawner;
 	friend class URigVMEdGraphArrayNodeSpawner;
 	friend class URigVMEdGraphIfNodeSpawner;
diff --git a/Engine/Plugins/Runtime/RigVM/Source/RigVMDeveloper/Public/EdGraph/RigVMEdGraphSchema.h b/Engine/Plugins/Runtime/RigVM/Source/RigVMDeveloper/Public/EdGraph/RigVMEdGraphSchema.h
--- a/Engine/Plugins/Runtime/RigVM/Source/RigVMDeveloper/Public/EdGraph/RigVMEdGraphSchema.h
+++ b/Engine/Plugins/Runtime/RigVM/Source/RigVMDeveloper/Public/EdGraph/RigVMEdGraphSchema.h
@@ -252,7 +252,7 @@ public:
 public:
 	URigVMEdGraphSchema();
 
-	virtual const FName& GetGraphNameForUI() const { return GraphName_RigVM; }
+	virtual const FName& GetRootGraphName() const { return GraphName_RigVM; }
 
 	// UEdGraphSchema interface
 	virtual void GetGraphContextActions(FGraphContextMenuBuilder& ContextMenuBuilder) const override;
diff --git a/Engine/Plugins/Runtime/RigVM/Source/RigVMDeveloper/Public/RigVMBlueprint.h b/Engine/Plugins/Runtime/RigVM/Source/RigVMDeveloper/Public/RigVMBlueprint.h
--- a/Engine/Plugins/Runtime/RigVM/Source/RigVMDeveloper/Public/RigVMBlueprint.h
+++ b/Engine/Plugins/Runtime/RigVM/Source/RigVMDeveloper/Public/RigVMBlueprint.h
@@ -3,6 +3,7 @@
 #pragma once
 
 #include "CoreMinimal.h"
+#include "RigVMBlueprintGeneratedClass.h"
 #include "Engine/Blueprint.h"
 #include "RigVMCore/RigVM.h"
 #include "RigVMHost.h"
@@ -208,6 +209,9 @@ public:
 	/** Get the (skeleton) generated class for this rigvm blueprint */
 	URigVMBlueprintGeneratedClass* GetRigVMBlueprintSkeletonClass() const;
 
+	/** Returns the class used as the super class for all generated classes */
+	virtual UClass* GetRigVMBlueprintGeneratedClassPrototype() const { return URigVMBlueprintGeneratedClass::StaticClass(); }
+
 	/** Returns the expected schema class to use for this blueprint */
 	virtual UClass* GetRigVMSchemaClass() const { return URigVMSchema::StaticClass(); }
 
@@ -467,7 +471,7 @@ public:
 	FOnRigVMCompiledEvent& OnVMCompiled();
 
 	UFUNCTION(BlueprintCallable, Category = "VM")
-	UClass* GetRigVMHostClass();
+	UClass* GetRigVMHostClass() const;
 
 	UFUNCTION(BlueprintCallable, Category = "VM")
 	URigVMHost* CreateRigVMHost();
diff --git a/Engine/Plugins/Runtime/RigVM/Source/RigVMEditor/Private/Editor/RigVMDetailsViewWrapperObject.cpp b/Engine/Plugins/Runtime/RigVM/Source/RigVMEditor/Private/Editor/RigVMDetailsViewWrapperObject.cpp
--- a/Engine/Plugins/Runtime/RigVM/Source/RigVMEditor/Private/Editor/RigVMDetailsViewWrapperObject.cpp
+++ b/Engine/Plugins/Runtime/RigVM/Source/RigVMEditor/Private/Editor/RigVMDetailsViewWrapperObject.cpp
@@ -10,6 +10,8 @@
 
 #if WITH_EDITOR
 #include "PropertyEditorModule.h"
+#include "Editor/RigVMGraphDetailCustomization.h"
+#include "Editor/RigVMLocalVariableDetailCustomization.h"
 #endif
 
 TMap<URigVMDetailsViewWrapperObject::FPerClassInfo, UClass*> URigVMDetailsViewWrapperObject::InfoToClass;
@@ -65,8 +67,8 @@ UClass* URigVMDetailsViewWrapperObject::GetClassForStruct(UScriptStruct* InStruc
 		return nullptr;
 	}
 
-	UClass* SuperClass = URigVMDetailsViewWrapperObject::StaticClass();
-	const FName WrapperClassName(FString::Printf(TEXT("%s_WrapperObject"), *InStruct->GetStructCPPName()));
+	UClass* SuperClass = GetClass();
+	const FName WrapperClassName(FString::Printf(TEXT("%s_%s"), *SuperClass->GetName(), *InStruct->GetStructCPPName()));
 
 	UClass* WrapperClass = NewObject<UClass>(
 		GetTransientPackage(),
@@ -162,6 +164,17 @@ UClass* URigVMDetailsViewWrapperObject::GetClassForStruct(UScriptStruct* InStruc
 	TSharedPtr<FStructOnScope> DefaultStruct = MakeShareable(new FStructOnScope(InStruct));
 	CopyPropertiesForUnrelatedStructs((uint8*)CDO, WrapperClass, DefaultStruct->GetStructMemory(), DefaultStruct->GetStruct());
 
+#if WITH_EDITOR
+	if(InStruct->IsChildOf(FRigVMGraphVariableDescription::StaticStruct()))
+	{
+		FPropertyEditorModule& PropertyEditorModule = FModuleManager::LoadModuleChecked<FPropertyEditorModule>("PropertyEditor");
+		if (!PropertyEditorModule.GetClassNameToDetailLayoutNameMap().Contains(WrapperClassName))
+		{
+			PropertyEditorModule.RegisterCustomClassLayout(WrapperClassName, FOnGetDetailCustomizationInstance::CreateStatic(&FRigVMLocalVariableDetailCustomization::MakeInstance));
+		}
+	}
+#endif
+
 	return WrapperClass;
 }
 
@@ -321,9 +334,9 @@ UClass* URigVMDetailsViewWrapperObject::GetClassForNodes(TArray<URigVMNode*> InN
 		return nullptr;
 	}
 
-	UClass* SuperClass = URigVMDetailsViewWrapperObject::StaticClass();
+	UClass* SuperClass = GetClass();
 	const int32 HashForNotation = (int32)GetTypeHash(PerClassInfo);
-	const FName WrapperClassName(FString::Printf(TEXT("RigNode%d_WrapperObject"), HashForNotation));
+	const FName WrapperClassName(FString::Printf(TEXT("%s_%d"), *SuperClass->GetName(), HashForNotation));
 
 	UClass* WrapperClass = NewObject<UClass>(
 		GetTransientPackage(),
@@ -343,7 +356,7 @@ UClass* URigVMDetailsViewWrapperObject::GetClassForNodes(TArray<URigVMNode*> InN
 	WrapperClass->ClassConfigName = SuperClass->ClassConfigName;
 	WrapperClass->ClassFlags |= CLASS_NotPlaceable | CLASS_Hidden;
 
-	WrapperClass->SetMetaData(TEXT("DisplayName"), TEXT("Rig Node"));
+	WrapperClass->SetMetaData(TEXT("DisplayName"), TEXT("RigVM Node"));
 
 	// create properties - one for each pin to inspect
 	FField** LinkToProperty = &WrapperClass->ChildProperties;
@@ -480,6 +493,12 @@ UClass* URigVMDetailsViewWrapperObject::GetClassForNodes(TArray<URigVMNode*> InN
 		}
 	}
 
+#if WITH_EDITOR
+	FPropertyEditorModule& EditModule = FModuleManager::Get().GetModuleChecked<FPropertyEditorModule>("PropertyEditor");
+	EditModule.RegisterCustomClassLayout(WrapperClass->GetFName(), FOnGetDetailCustomizationInstance::CreateStatic(&FRigVMWrappedNodeDetailCustomization::MakeInstance));
+#endif
+
+
 	return WrapperClass;	
 }
 
diff --git a/Engine/Plugins/Runtime/RigVM/Source/RigVMEditor/Private/Editor/RigVMGraphDetailCustomization.cpp b/Engine/Plugins/Runtime/RigVM/Source/RigVMEditor/Private/Editor/RigVMGraphDetailCustomization.cpp
new file mode 100644
--- /dev/null
+++ b/Engine/Plugins/Runtime/RigVM/Source/RigVMEditor/Private/Editor/RigVMGraphDetailCustomization.cpp
@@ -0,0 +1,1857 @@
+// Copyright Epic Games, Inc. All Rights Reserved.
+
+#include "Editor/RigVMGraphDetailCustomization.h"
+#include "Widgets/SWidget.h"
+#include "DetailLayoutBuilder.h"
+#include "DetailCategoryBuilder.h"
+#include "DetailWidgetRow.h"
+#include "IDetailChildrenBuilder.h"
+#include "Styling/AppStyle.h"
+#include "SPinTypeSelector.h"
+#include "Widgets/Input/SButton.h"
+#include "Widgets/Input/SCheckBox.h"
+#include "Widgets/Text/STextBlock.h"
+#include "Widgets/Colors/SColorPicker.h"
+#include "Widgets/Text/SInlineEditableTextBlock.h"
+#include "Widgets/Text/SMultiLineEditableText.h"
+#include "PropertyCustomizationHelpers.h"
+#include "NodeFactory.h"
+#include "EdGraph/RigVMEdGraphNode.h"
+#include "RigVMHost.h"
+#include "RigVMCore/RigVMExternalVariable.h"
+#include "Kismet2/BlueprintEditorUtils.h"
+#include "EdGraph/RigVMEdGraphSchema.h"
+#include "EditorCategoryUtils.h"
+#include "IPropertyUtilities.h"
+#include "Widgets/SRigVMGraphPinVariableBinding.h"
+
+#define LOCTEXT_NAMESPACE "RigVMGraphDetailCustomization"
+
+static const FText RigVMGraphDetailCustomizationMultipleValues = LOCTEXT("MultipleValues", "Multiple Values");
+
+FRigVMFunctionArgumentGroupLayout::FRigVMFunctionArgumentGroupLayout(
+	URigVMGraph* InGraph, 
+	URigVMBlueprint* InBlueprint, 
+	TWeakPtr<FRigVMEditor> InEditor,
+	bool bInputs)
+	: GraphPtr(InGraph)
+	, RigVMBlueprintPtr(InBlueprint)
+	, RigVMEditorPtr(InEditor)
+	, bIsInputGroup(bInputs)
+{
+	if (RigVMBlueprintPtr.IsValid())
+	{
+		RigVMBlueprintPtr.Get()->OnModified().AddRaw(this, &FRigVMFunctionArgumentGroupLayout::HandleModifiedEvent);
+	}
+}
+
+FRigVMFunctionArgumentGroupLayout::~FRigVMFunctionArgumentGroupLayout()
+{
+	if (RigVMBlueprintPtr.IsValid())
+	{
+		RigVMBlueprintPtr.Get()->OnModified().RemoveAll(this);
+	}
+}
+
+void FRigVMFunctionArgumentGroupLayout::GenerateChildContent(IDetailChildrenBuilder& ChildrenBuilder)
+{
+	bool WasContentAdded = false;
+	if (GraphPtr.IsValid())
+	{
+		URigVMGraph* Graph = GraphPtr.Get();
+		if (URigVMLibraryNode* LibraryNode = Cast<URigVMLibraryNode>(Graph->GetOuter()))
+		{
+			for (URigVMPin* Pin : LibraryNode->GetPins())
+			{
+				if ((bIsInputGroup && (Pin->GetDirection() == ERigVMPinDirection::Input || Pin->GetDirection() == ERigVMPinDirection::IO)) ||
+					(!bIsInputGroup && (Pin->GetDirection() == ERigVMPinDirection::Output || Pin->GetDirection() == ERigVMPinDirection::IO)))
+				{
+					TSharedRef<class FRigVMFunctionArgumentLayout> ArgumentLayout = MakeShareable(new FRigVMFunctionArgumentLayout(
+						Pin,
+						Graph,
+						RigVMBlueprintPtr.Get(),
+						RigVMEditorPtr
+					));
+					ChildrenBuilder.AddCustomBuilder(ArgumentLayout);
+					WasContentAdded = true;
+				}
+			}
+		}
+	}
+	if (!WasContentAdded)
+	{
+		// Add a text widget to let the user know to hit the + icon to add parameters.
+		ChildrenBuilder.AddCustomRow(FText::GetEmpty()).WholeRowContent()
+		.MaxDesiredWidth(980.f)
+		[
+			SNew(SHorizontalBox)
+			+ SHorizontalBox::Slot()
+			.VAlign(VAlign_Center)
+			.Padding(0.0f, 0.0f, 4.0f, 0.0f)
+			.AutoWidth()
+			[
+				SNew(STextBlock)
+				.Text(LOCTEXT("NoArgumentsAddedForRigVMHost", "Please press the + icon above to add parameters"))
+				.Font(IDetailLayoutBuilder::GetDetailFont())
+			]
+		];
+	}
+}
+
+void FRigVMFunctionArgumentGroupLayout::HandleModifiedEvent(ERigVMGraphNotifType InNotifType, URigVMGraph* InGraph, UObject* InSubject)
+{
+	if (!GraphPtr.IsValid())
+	{
+		return;
+	}
+	
+	URigVMGraph* Graph = GraphPtr.Get();
+	URigVMLibraryNode* LibraryNode = Cast<URigVMLibraryNode>(Graph->GetOuter());
+	if (LibraryNode == nullptr)
+	{
+		return;
+	}
+
+	switch (InNotifType)
+	{
+		case ERigVMGraphNotifType::PinAdded:
+		case ERigVMGraphNotifType::PinRenamed:
+		case ERigVMGraphNotifType::PinRemoved:
+		case ERigVMGraphNotifType::PinIndexChanged:
+		case ERigVMGraphNotifType::PinTypeChanged:
+		{
+			URigVMPin* Pin = CastChecked<URigVMPin>(InSubject);
+			if (Pin->GetNode() == LibraryNode ||
+				(Pin->GetNode()->IsA<URigVMFunctionEntryNode>() && Pin->GetNode()->GetOuter() == Graph) ||
+				(Pin->GetNode()->IsA<URigVMFunctionReturnNode>() && Pin->GetNode()->GetOuter() == Graph))
+			{
+				OnRebuildChildren.ExecuteIfBound();
+			}
+			break;
+		}
+		default:
+		{
+			break;
+		}
+	}
+}
+
+class FRigVMFunctionArgumentPinTypeSelectorFilter : public IPinTypeSelectorFilter
+{
+public:
+	FRigVMFunctionArgumentPinTypeSelectorFilter(TWeakPtr<FRigVMEditor> InRigVMEditor, TWeakObjectPtr<URigVMGraph> InGraph)
+		: RigVMEditorPtr(InRigVMEditor), GraphPtr(InGraph)
+	{
+	}
+	
+	virtual bool ShouldShowPinTypeTreeItem(FPinTypeTreeItem InItem) const override
+	{
+		if (!InItem.IsValid())
+		{
+			return false;
+		}
+
+		// Only allow an execute context pin if the graph doesnt have one already
+		FString CPPType;
+		UObject* CPPTypeObject = nullptr;
+		RigVMTypeUtils::CPPTypeFromPinType(InItem.Get()->GetPinType(false), CPPType, &CPPTypeObject);
+		if (UScriptStruct* ScriptStruct = Cast<UScriptStruct>(CPPTypeObject))
+		{
+			if (ScriptStruct->IsChildOf(FRigVMExecuteContext::StaticStruct()))
+			{
+				if (GraphPtr.IsValid())
+				{
+					if (URigVMFunctionEntryNode* EntryNode = GraphPtr.Get()->GetEntryNode())
+					{
+						for (URigVMPin* Pin : EntryNode->GetPins())
+						{
+							if (Pin->IsExecuteContext())
+							{
+								return false;
+							}
+						}
+					}
+				}
+			}
+		}
+		
+		if (RigVMEditorPtr.IsValid())
+		{
+			TArray<TSharedPtr<IPinTypeSelectorFilter>> Filters;
+			RigVMEditorPtr.Pin()->GetPinTypeSelectorFilters(Filters);
+			for(const TSharedPtr<IPinTypeSelectorFilter>& Filter : Filters)
+			{
+				if(!Filter->ShouldShowPinTypeTreeItem(InItem))
+				{
+					return false;
+				}
+			}
+			return true;
+		}
+
+		return false;
+	}
+
+private:
+
+	TWeakPtr<FRigVMEditor> RigVMEditorPtr;
+	
+	TWeakObjectPtr<URigVMGraph> GraphPtr;
+};
+
+void FRigVMFunctionArgumentLayout::GenerateHeaderRowContent(FDetailWidgetRow& NodeRow)
+{
+	const UEdGraphSchema* Schema = GetDefault<URigVMEdGraphSchema>();
+
+	ETypeTreeFilter TypeTreeFilter = ETypeTreeFilter::None;
+	TypeTreeFilter |= ETypeTreeFilter::AllowExec;
+
+	TArray<TSharedPtr<IPinTypeSelectorFilter>> CustomPinTypeFilters;
+	if (RigVMEditorPtr.IsValid())
+	{
+		CustomPinTypeFilters.Add(MakeShared<FRigVMFunctionArgumentPinTypeSelectorFilter>(RigVMEditorPtr, GraphPtr));
+	}
+	
+	NodeRow
+		.NameContent()
+		[
+			SNew(SHorizontalBox)
+			+ SHorizontalBox::Slot()
+			.FillWidth(1)
+			.VAlign(VAlign_Center)
+			[
+				SAssignNew(ArgumentNameWidget, SEditableTextBox)
+				.Text(this, &FRigVMFunctionArgumentLayout::OnGetArgNameText)
+				.OnTextCommitted(this, &FRigVMFunctionArgumentLayout::OnArgNameTextCommitted)
+				.ToolTipText(this, &FRigVMFunctionArgumentLayout::OnGetArgToolTipText)
+				.Font(IDetailLayoutBuilder::GetDetailFont())
+				.IsEnabled(!ShouldPinBeReadOnly())
+				.OnVerifyTextChanged_Lambda([&](const FText& InNewText, FText& OutErrorMessage) -> bool
+				{
+					if (InNewText.IsEmpty())
+					{
+						OutErrorMessage = LOCTEXT("ArgumentNameEmpty", "Cannot have an argument with an emtpy string name.");
+						return false;
+					}
+					else if (InNewText.ToString().Len() >= NAME_SIZE)
+					{
+						OutErrorMessage = LOCTEXT("ArgumentNameTooLong", "Name of argument is too long.");
+						return false;
+					}
+
+					EValidatorResult Result = NameValidator.IsValid(InNewText.ToString(), false);
+					OutErrorMessage = INameValidatorInterface::GetErrorText(InNewText.ToString(), Result);	
+
+					return Result == EValidatorResult::Ok || Result == EValidatorResult::ExistingName;
+				})
+			]
+		]
+		.ValueContent()
+		.MaxDesiredWidth(980.f)
+		[
+			SNew(SHorizontalBox)
+			+ SHorizontalBox::Slot()
+			.VAlign(VAlign_Center)
+			.Padding(0.0f, 0.0f, 4.0f, 0.0f)
+			.AutoWidth()
+			[
+				SNew(SPinTypeSelector, FGetPinTypeTree::CreateUObject(GetDefault<UEdGraphSchema_K2>(), &UEdGraphSchema_K2::GetVariableTypeTree))
+				.TargetPinType(this, &FRigVMFunctionArgumentLayout::OnGetPinInfo)
+				.OnPinTypePreChanged(this, &FRigVMFunctionArgumentLayout::OnPrePinInfoChange)
+				.OnPinTypeChanged(this, &FRigVMFunctionArgumentLayout::PinInfoChanged)
+				.Schema(Schema)
+				.TypeTreeFilter(TypeTreeFilter)
+				.bAllowArrays(!ShouldPinBeReadOnly())
+				.IsEnabled(!ShouldPinBeReadOnly(true))
+				.CustomFilters(CustomPinTypeFilters)
+				.Font(IDetailLayoutBuilder::GetDetailFont())
+			]
+			+ SHorizontalBox::Slot()
+			.AutoWidth()
+			[
+				SNew(SButton)
+				.ButtonStyle(FAppStyle::Get(), TEXT("SimpleButton"))
+				.ContentPadding(0)
+				.IsEnabled_Raw(this, &FRigVMFunctionArgumentLayout::CanArgumentBeMoved, true)
+				.OnClicked(this, &FRigVMFunctionArgumentLayout::OnArgMoveUp)
+				.ToolTipText(LOCTEXT("FunctionArgDetailsArgMoveUpTooltip", "Move this parameter up in the list."))
+				[
+					SNew(SImage)
+					.Image(FAppStyle::GetBrush("Icons.ChevronUp"))
+				.ColorAndOpacity(FSlateColor::UseForeground())
+				]
+			]
+			+ SHorizontalBox::Slot()
+			.AutoWidth()
+			.Padding(2, 0)
+			[
+				SNew(SButton)
+				.ButtonStyle(FAppStyle::Get(), TEXT("SimpleButton"))
+				.ContentPadding(0)
+				.IsEnabled_Raw(this, &FRigVMFunctionArgumentLayout::CanArgumentBeMoved, false)
+				.OnClicked(this, &FRigVMFunctionArgumentLayout::OnArgMoveDown)
+				.ToolTipText(LOCTEXT("FunctionArgDetailsArgMoveDownTooltip", "Move this parameter down in the list."))
+				[
+					SNew(SImage)
+					.Image(FAppStyle::GetBrush("Icons.ChevronDown"))
+					.ColorAndOpacity(FSlateColor::UseForeground())
+				]
+			]
+			+ SHorizontalBox::Slot()
+			.HAlign(HAlign_Right)
+			.VAlign(VAlign_Center)
+			.Padding(10, 0, 0, 0)
+			.AutoWidth()
+			[
+				PropertyCustomizationHelpers::MakeClearButton(FSimpleDelegate::CreateSP(this, &FRigVMFunctionArgumentLayout::OnRemoveClicked), LOCTEXT("FunctionArgDetailsClearTooltip", "Remove this parameter."), !IsPinEditingReadOnly())
+			]
+		];
+}
+
+void FRigVMFunctionArgumentLayout::GenerateChildContent(IDetailChildrenBuilder& ChildrenBuilder)
+{
+	// we don't show defaults here - we rely on a SRigVMGraphNode widget in the top of the details
+}
+
+void FRigVMFunctionArgumentLayout::OnRemoveClicked()
+{
+	if (PinPtr.IsValid() && RigVMBlueprintPtr.IsValid())
+	{
+		URigVMPin* Pin = PinPtr.Get();
+		URigVMBlueprint* Blueprint = RigVMBlueprintPtr.Get();
+		if (URigVMLibraryNode* LibraryNode = Cast<URigVMLibraryNode>(Pin->GetNode()))
+		{
+			if (URigVMController* Controller = Blueprint->GetController(LibraryNode->GetContainedGraph()))
+			{
+				Controller->RemoveExposedPin(Pin->GetFName(), true, true);
+			}
+		}
+	}
+}
+
+FReply FRigVMFunctionArgumentLayout::OnArgMoveUp()
+{
+	if (PinPtr.IsValid() && RigVMBlueprintPtr.IsValid())
+	{
+		URigVMPin* Pin = PinPtr.Get();
+		URigVMBlueprint* Blueprint = RigVMBlueprintPtr.Get();
+		if (URigVMLibraryNode* LibraryNode = Cast<URigVMLibraryNode>(Pin->GetNode()))
+		{
+			if (URigVMController* Controller = Blueprint->GetController(LibraryNode->GetContainedGraph()))
+			{
+				bool bIsInput = Pin->GetDirection() == ERigVMPinDirection::Input || Pin->GetDirection() == ERigVMPinDirection::IO;
+				
+				int32 NewPinIndex = Pin->GetPinIndex() - 1;
+				while (NewPinIndex != INDEX_NONE)
+				{
+					URigVMPin* OtherPin = LibraryNode->GetPins()[NewPinIndex];
+					if (bIsInput)
+					{
+						if (OtherPin->GetDirection() == ERigVMPinDirection::Input || OtherPin->GetDirection() == ERigVMPinDirection::IO)
+						{
+							break;
+						}
+					}	
+					else
+					{
+						if (OtherPin->GetDirection() == ERigVMPinDirection::Output || OtherPin->GetDirection() == ERigVMPinDirection::IO)
+						{
+							break;
+						}
+					}
+					--NewPinIndex;
+				}
+				if (NewPinIndex != INDEX_NONE)
+				{
+					Controller->SetExposedPinIndex(Pin->GetFName(), NewPinIndex, true, true);
+				}
+				return FReply::Handled();
+			}
+		}
+	}
+	return FReply::Unhandled();
+}
+
+FReply FRigVMFunctionArgumentLayout::OnArgMoveDown()
+{
+	if (PinPtr.IsValid() && RigVMBlueprintPtr.IsValid())
+	{
+		URigVMPin* Pin = PinPtr.Get();
+		URigVMBlueprint* Blueprint = RigVMBlueprintPtr.Get();
+		if (URigVMLibraryNode* LibraryNode = Cast<URigVMLibraryNode>(Pin->GetNode()))
+		{
+			if (URigVMController* Controller = Blueprint->GetController(LibraryNode->GetContainedGraph()))
+			{
+				bool bIsInput = Pin->GetDirection() == ERigVMPinDirection::Input || Pin->GetDirection() == ERigVMPinDirection::IO;
+				
+				int32 NewPinIndex = Pin->GetPinIndex() + 1;
+				while (NewPinIndex < LibraryNode->GetPins().Num())
+				{
+					URigVMPin* OtherPin = LibraryNode->GetPins()[NewPinIndex];
+					if (bIsInput)
+					{
+						if (OtherPin->GetDirection() == ERigVMPinDirection::Input || OtherPin->GetDirection() == ERigVMPinDirection::IO)
+						{
+							break;
+						}
+					}	
+					else
+					{
+						if (OtherPin->GetDirection() == ERigVMPinDirection::Output || OtherPin->GetDirection() == ERigVMPinDirection::IO)
+						{
+							break;
+						}
+					}
+					++NewPinIndex;
+				}
+				if (NewPinIndex < LibraryNode->GetPins().Num())
+				{
+					Controller->SetExposedPinIndex(Pin->GetFName(), NewPinIndex, true, true);
+				}
+				return FReply::Handled();
+			}
+		}
+	}
+	return FReply::Unhandled();
+}
+
+bool FRigVMFunctionArgumentLayout::ShouldPinBeReadOnly(bool bIsEditingPinType/* = false*/) const
+{
+	return IsPinEditingReadOnly(bIsEditingPinType);
+}
+
+bool FRigVMFunctionArgumentLayout::IsPinEditingReadOnly(bool bIsEditingPinType/* = false*/) const
+{
+	return false;
+}
+
+bool FRigVMFunctionArgumentLayout::CanArgumentBeMoved(bool bMoveUp) const
+{
+	if(IsPinEditingReadOnly())
+	{
+		return false;
+	}
+	if (PinPtr.IsValid())
+	{
+		URigVMPin* Pin = PinPtr.Get();
+		if(Pin->IsExecuteContext())
+		{
+			return false;
+		}
+
+		if(URigVMNode* Node = Pin->GetNode())
+		{
+			auto IsInput = [](URigVMPin* InPin) -> bool
+			{
+				return InPin->GetDirection() == ERigVMPinDirection::Input ||
+					InPin->GetDirection() == ERigVMPinDirection::Visible;
+			};
+
+			const bool bLookForInput = IsInput(Pin);
+
+			if(bMoveUp)
+			{
+				// if this is the first pin of its type
+				for(int32 Index = 0; Index < Node->GetPins().Num(); Index++)
+				{
+					URigVMPin* OtherPin = Node->GetPins()[Index];
+					if(IsInput(OtherPin) == bLookForInput)
+					{
+						return OtherPin != Pin;
+					}
+				}
+			}
+			else if(!bMoveUp)
+			{
+				// if this is the last pin of its type
+				for(int32 Index = Node->GetPins().Num() - 1; Index >= 0; Index--)
+				{
+					URigVMPin* OtherPin = Node->GetPins()[Index];
+					if(IsInput(OtherPin) == bLookForInput)
+					{
+						return OtherPin != Pin;
+					}
+				}
+			}
+		}
+	}
+	return true;
+}
+
+FText FRigVMFunctionArgumentLayout::OnGetArgNameText() const
+{
+	if (PinPtr.IsValid())
+	{
+		return FText::FromName(PinPtr.Get()->GetFName());
+	}
+	return FText();
+}
+
+FText FRigVMFunctionArgumentLayout::OnGetArgToolTipText() const
+{
+	return OnGetArgNameText(); // for now since we don't have tooltips
+}
+
+void FRigVMFunctionArgumentLayout::OnArgNameTextCommitted(const FText& NewText, ETextCommit::Type InTextCommit)
+{
+	if (InTextCommit == ETextCommit::OnEnter)
+	{
+		if (!NewText.IsEmpty() && PinPtr.IsValid() && RigVMBlueprintPtr.IsValid() && !ShouldPinBeReadOnly())
+		{
+			URigVMPin* Pin = PinPtr.Get();
+			URigVMBlueprint* Blueprint = RigVMBlueprintPtr.Get();
+			if (URigVMLibraryNode* LibraryNode = Cast<URigVMLibraryNode>(Pin->GetNode()))
+			{
+				if (URigVMController* Controller = Blueprint->GetController(LibraryNode->GetContainedGraph()))
+				{
+					const FString& NewName = NewText.ToString();
+					Controller->RenameExposedPin(Pin->GetFName(), *NewName, true, true);
+				}
+			}
+		}
+	}
+}
+
+FEdGraphPinType FRigVMFunctionArgumentLayout::OnGetPinInfo() const
+{
+	if (PinPtr.IsValid())
+	{
+		return URigVMEdGraphNode::GetPinTypeForModelPin(PinPtr.Get());
+	}
+	return FEdGraphPinType();
+}
+
+void FRigVMFunctionArgumentLayout::PinInfoChanged(const FEdGraphPinType& PinType)
+{
+	if (PinPtr.IsValid() && RigVMBlueprintPtr.IsValid() && FBlueprintEditorUtils::IsPinTypeValid(PinType))
+	{
+		URigVMPin* Pin = PinPtr.Get();
+		URigVMBlueprint* Blueprint = RigVMBlueprintPtr.Get();
+		if (URigVMLibraryNode* LibraryNode = Cast<URigVMLibraryNode>(Pin->GetNode()))
+		{
+			if (URigVMController* Controller = Blueprint->GetController(LibraryNode->GetContainedGraph()))
+			{
+				FString CPPType;
+				FName CPPTypeObjectName = NAME_None;
+				RigVMTypeUtils::CPPTypeFromPinType(PinType, CPPType, CPPTypeObjectName);
+				
+				bool bSetupUndoRedo = true;
+				Controller->ChangeExposedPinType(Pin->GetFName(), CPPType, CPPTypeObjectName, bSetupUndoRedo, false, true);
+
+				// If the controller has identified this as a bulk change, it has not added the actions to the action stack
+				// We need to disable the transaction from the UI as well to keep them synced
+				if (!bSetupUndoRedo)
+				{
+					GEditor->CancelTransaction(0);
+				}
+			}
+		}
+	}
+}
+
+void FRigVMFunctionArgumentLayout::OnPrePinInfoChange(const FEdGraphPinType& PinType)
+{
+	// not needed for rig vm
+}
+
+FRigVMFunctionArgumentDefaultNode::FRigVMFunctionArgumentDefaultNode(
+	URigVMGraph* InGraph,
+	URigVMBlueprint* InBlueprint
+)
+	: GraphPtr(InGraph)
+	, RigVMBlueprintPtr(InBlueprint)
+{
+	if (GraphPtr.IsValid() && RigVMBlueprintPtr.IsValid())
+	{
+		RigVMBlueprintPtr.Get()->OnModified().AddRaw(this, &FRigVMFunctionArgumentDefaultNode::HandleModifiedEvent);
+
+		if (URigVMLibraryNode* LibraryNode = Cast<URigVMLibraryNode>(GraphPtr->GetOuter()))
+		{
+			if (URigVMEdGraph* RigGraph = Cast<URigVMEdGraph>(RigVMBlueprintPtr->GetEdGraph(LibraryNode->GetGraph())))
+			{
+				EdGraphOuterPtr = RigGraph;
+				GraphChangedDelegateHandle = RigGraph->AddOnGraphChangedHandler(
+					FOnGraphChanged::FDelegate::CreateRaw(this, &FRigVMFunctionArgumentDefaultNode::OnGraphChanged)
+				);
+			}
+		}	
+
+	}
+}
+
+FRigVMFunctionArgumentDefaultNode::~FRigVMFunctionArgumentDefaultNode()
+{
+	if (RigVMBlueprintPtr.IsValid())
+	{
+		RigVMBlueprintPtr.Get()->OnModified().RemoveAll(this);
+	}
+	
+	if (EdGraphOuterPtr.IsValid())
+	{
+		if (GraphChangedDelegateHandle.IsValid())
+		{
+			EdGraphOuterPtr->RemoveOnGraphChangedHandler(GraphChangedDelegateHandle);
+		}		
+	}
+}
+
+void FRigVMFunctionArgumentDefaultNode::GenerateChildContent(IDetailChildrenBuilder& ChildrenBuilder)
+{
+	if (!GraphPtr.IsValid() || !RigVMBlueprintPtr.IsValid())
+	{
+		return;
+	}
+
+	URigVMBlueprint* Blueprint = RigVMBlueprintPtr.Get();
+	URigVMGraph* Graph = GraphPtr.Get();
+	URigVMEdGraphNode* RigVMEdGraphNode = nullptr;
+	if (URigVMLibraryNode* LibraryNode = Cast<URigVMLibraryNode>(Graph->GetOuter()))
+	{
+		if (URigVMEdGraph* RigGraph = Cast<URigVMEdGraph>(Blueprint->GetEdGraph(LibraryNode->GetGraph())))
+		{
+			RigVMEdGraphNode = Cast<URigVMEdGraphNode>(RigGraph->FindNodeForModelNodeName(LibraryNode->GetFName()));
+		}
+	}
+
+	if (RigVMEdGraphNode == nullptr)
+	{
+		return;
+	}
+
+	ChildrenBuilder.AddCustomRow(FText::GetEmpty())
+	.WholeRowContent()
+	.MaxDesiredWidth(980.f)
+	[
+		SAssignNew(OwnedNodeWidget, SRigVMGraphNode).GraphNodeObj(RigVMEdGraphNode)
+	];
+
+	OwnedNodeWidget->SetIsEditable(true);
+	TArray< TSharedRef<SWidget> > Pins;
+	OwnedNodeWidget->GetPins(Pins);
+	for (TSharedRef<SWidget> Pin : Pins)
+	{
+		TSharedRef<SGraphPin> SPin = StaticCastSharedRef<SGraphPin>(Pin);
+		SPin->EnableDragAndDrop(false);
+	}
+}
+
+void FRigVMFunctionArgumentDefaultNode::OnGraphChanged(const FEdGraphEditAction& InAction)
+{
+	if (GraphPtr.IsValid() && RigVMBlueprintPtr.IsValid())
+	{
+		OnRebuildChildren.ExecuteIfBound();
+	}
+}
+
+void FRigVMFunctionArgumentDefaultNode::HandleModifiedEvent(ERigVMGraphNotifType InNotifType, URigVMGraph* InGraph, UObject* InSubject)
+{
+	if (!GraphPtr.IsValid())
+	{
+		return;
+	}
+
+	URigVMGraph* Graph = GraphPtr.Get();
+	URigVMLibraryNode* LibraryNode = Cast<URigVMLibraryNode>(Graph->GetOuter());
+	if (LibraryNode == nullptr)
+	{
+		return;
+	}
+	if (LibraryNode->GetGraph() != InGraph)
+	{
+		return;
+	}
+
+	switch (InNotifType)
+	{
+		case ERigVMGraphNotifType::PinAdded:
+		case ERigVMGraphNotifType::PinRemoved:
+		case ERigVMGraphNotifType::PinTypeChanged:
+		case ERigVMGraphNotifType::PinIndexChanged:
+		case ERigVMGraphNotifType::PinRenamed:
+		{
+			URigVMPin* Pin = CastChecked<URigVMPin>(InSubject);
+			if (Pin->GetNode() == LibraryNode)
+			{
+				OnRebuildChildren.ExecuteIfBound();
+			}
+			break;
+		}
+		case ERigVMGraphNotifType::NodeRenamed:
+		case ERigVMGraphNotifType::NodeColorChanged:
+		{
+			URigVMNode* Node = CastChecked<URigVMNode>(InSubject);
+			if (Node == LibraryNode)
+			{
+				OnRebuildChildren.ExecuteIfBound();
+			}
+			break;
+		}
+		default:
+		{
+			break;
+		}
+	}
+}
+
+TSharedPtr<IDetailCustomization> FRigVMGraphDetailCustomization::MakeInstance(TSharedPtr<IBlueprintEditor> InBlueprintEditor, const UClass* InExpectedBlueprintClass)
+{
+	const TArray<UObject*>* Objects = (InBlueprintEditor.IsValid() ? InBlueprintEditor->GetObjectsCurrentlyBeingEdited() : nullptr);
+	if (Objects && Objects->Num() == 1)
+	{
+		if (URigVMBlueprint* RigVMBlueprint = Cast<URigVMBlueprint>((*Objects)[0]))
+		{
+			if(RigVMBlueprint->GetClass() == InExpectedBlueprintClass)
+			{
+				return MakeShareable(new FRigVMGraphDetailCustomization(StaticCastSharedPtr<FRigVMEditor>(InBlueprintEditor), RigVMBlueprint));
+			}
+		}
+	}
+
+	return nullptr;
+}
+
+void FRigVMGraphDetailCustomization::CustomizeDetails(IDetailLayoutBuilder& DetailLayout)
+{
+	bIsPickingColor = false;
+
+	TArray<TWeakObjectPtr<UObject>> Objects;
+	DetailLayout.GetObjectsBeingCustomized(Objects);
+
+	GraphPtr = CastChecked<URigVMEdGraph>(Objects[0].Get());
+	URigVMEdGraph* Graph = GraphPtr.Get();
+
+	URigVMBlueprint* Blueprint = RigVMBlueprintPtr.Get();
+	URigVMGraph* Model = nullptr;
+	URigVMController* Controller = nullptr;
+
+	if (Blueprint)
+	{
+		Model = Blueprint->GetModel(Graph);
+		Controller = Blueprint->GetController(Model);
+	}
+
+	if (Blueprint == nullptr || Model == nullptr || Controller == nullptr)
+	{
+		IDetailCategoryBuilder& Category = DetailLayout.EditCategory("Graph", LOCTEXT("FunctionDetailsGraph", "Graph"));
+		Category.AddCustomRow(FText::GetEmpty())
+		[
+			SNew(STextBlock)
+			.Text(LOCTEXT("GraphPresentButNotEditable", "Graph is not editable."))
+		];
+		return;
+	}
+
+	if (Model->IsTopLevelGraph())
+	{
+		IDetailCategoryBuilder& Category = DetailLayout.EditCategory("Graph", LOCTEXT("FunctionDetailsGraph", "Graph"));
+		Category.AddCustomRow(FText::GetEmpty())
+			[
+				SNew(STextBlock)
+				.Text(LOCTEXT("GraphIsTopLevelGraph", "Top-level Graphs are not editable."))
+			];
+		return;
+	}
+
+	IDetailCategoryBuilder& InputsCategory = DetailLayout.EditCategory("Inputs", LOCTEXT("FunctionDetailsInputs", "Inputs"));
+	TSharedRef<FRigVMFunctionArgumentGroupLayout> InputArgumentGroup = MakeShareable(new FRigVMFunctionArgumentGroupLayout(
+		Model, 
+		Blueprint,
+		RigVMEditorPtr,
+		true));
+	InputsCategory.AddCustomBuilder(InputArgumentGroup);
+
+	TSharedRef<SHorizontalBox> InputsHeaderContentWidget = SNew(SHorizontalBox);
+
+	InputsHeaderContentWidget->AddSlot()
+	.HAlign(HAlign_Right)
+	[
+		SNew(SButton)
+		.ButtonStyle(FAppStyle::Get(), "SimpleButton")
+		.ContentPadding(FMargin(1, 0))
+		.OnClicked(this, &FRigVMGraphDetailCustomization::OnAddNewInputClicked)
+		.Visibility(this, &FRigVMGraphDetailCustomization::GetAddNewInputOutputVisibility)
+		.HAlign(HAlign_Right)
+		.ToolTipText(LOCTEXT("FunctionNewInputArgTooltip", "Create a new input argument"))
+		.VAlign(VAlign_Center)
+		.AddMetaData<FTagMetaData>(FTagMetaData(TEXT("FunctionNewInputArg")))
+		.IsEnabled(this, &FRigVMGraphDetailCustomization::IsAddNewInputOutputEnabled)
+		[
+			SNew(SImage)
+			.Image(FAppStyle::Get().GetBrush("Icons.PlusCircle"))
+			.ColorAndOpacity(FSlateColor::UseForeground())
+		]
+	];
+	InputsCategory.HeaderContent(InputsHeaderContentWidget);
+
+	IDetailCategoryBuilder& OutputsCategory = DetailLayout.EditCategory("Outputs", LOCTEXT("FunctionDetailsOutputs", "Outputs"));
+	TSharedRef<FRigVMFunctionArgumentGroupLayout> OutputArgumentGroup = MakeShareable(new FRigVMFunctionArgumentGroupLayout(
+		Model, 
+		Blueprint,
+		RigVMEditorPtr,
+		false));
+	OutputsCategory.AddCustomBuilder(OutputArgumentGroup);
+
+	TSharedRef<SHorizontalBox> OutputsHeaderContentWidget = SNew(SHorizontalBox);
+
+	OutputsHeaderContentWidget->AddSlot()
+	.HAlign(HAlign_Right)
+	[
+		SNew(SButton)
+		.ButtonStyle(FAppStyle::Get(), "SimpleButton")
+		.ContentPadding(FMargin(1, 0))
+		.OnClicked(this, &FRigVMGraphDetailCustomization::OnAddNewOutputClicked)
+		.Visibility(this, &FRigVMGraphDetailCustomization::GetAddNewInputOutputVisibility)
+		.HAlign(HAlign_Right)
+		.ToolTipText(LOCTEXT("FunctionNewOutputArgTooltip", "Create a new output argument"))
+		.VAlign(VAlign_Center)
+		.AddMetaData<FTagMetaData>(FTagMetaData(TEXT("FunctionNewOutputArg")))
+		.IsEnabled(this, &FRigVMGraphDetailCustomization::IsAddNewInputOutputEnabled)
+		[
+			SNew(SImage)
+			.Image(FAppStyle::Get().GetBrush("Icons.PlusCircle"))
+			.ColorAndOpacity(FSlateColor::UseForeground())
+		]
+	];
+	OutputsCategory.HeaderContent(OutputsHeaderContentWidget);
+
+	IDetailCategoryBuilder& SettingsCategory = DetailLayout.EditCategory("NodeSettings", LOCTEXT("FunctionDetailsNodeSettings", "Node Settings"));
+
+	bool bIsFunction = false;
+	if (Model)
+	{
+		if (URigVMLibraryNode* LibraryNode = Cast<URigVMLibraryNode>(Model->GetOuter()))
+		{
+			bIsFunction = LibraryNode->GetGraph()->IsA<URigVMFunctionLibrary>();
+		}
+	}
+
+	if(bIsFunction)
+	{
+		// node category
+		SettingsCategory.AddCustomRow(FText::GetEmpty())
+		.NameContent()
+		[
+			SNew(STextBlock)
+			.Text(FText::FromString(TEXT("Category")))
+			.Font(IDetailLayoutBuilder::GetDetailFont())
+		]
+		.ValueContent()
+		[
+			SNew(SEditableTextBox)
+			.Font(IDetailLayoutBuilder::GetDetailFont())
+			.Text(this, &FRigVMGraphDetailCustomization::GetNodeCategory)
+			.OnTextCommitted(this, &FRigVMGraphDetailCustomization::SetNodeCategory)
+			.OnVerifyTextChanged_Lambda([&](const FText& InNewText, FText& OutErrorMessage) -> bool
+			{
+				const FText NewText = FEditorCategoryUtils::GetCategoryDisplayString(InNewText);
+				if (NewText.ToString().Len() >= NAME_SIZE)
+				{
+					OutErrorMessage = LOCTEXT("CategoryTooLong", "Name of category is too long.");
+					return false;
+				}
+				
+				if (RigVMBlueprintPtr.IsValid())
+				{
+					if (NewText.EqualTo(FText::FromString(RigVMBlueprintPtr.Get()->GetName())))
+					{
+						OutErrorMessage = LOCTEXT("CategoryEqualsBlueprintName", "Cannot add a category with the same name as the blueprint.");
+						return false;
+					}
+				}
+				return true;
+			})
+		];
+
+		// node keywords
+		SettingsCategory.AddCustomRow(FText::GetEmpty())
+		.NameContent()
+		[
+			SNew(STextBlock)
+			.Text(FText::FromString(TEXT("Keywords")))
+			.Font(IDetailLayoutBuilder::GetDetailFont())
+		]
+		.ValueContent()
+		[
+			SNew(SEditableTextBox)
+			.Font(IDetailLayoutBuilder::GetDetailFont())
+			.Text(this, &FRigVMGraphDetailCustomization::GetNodeKeywords)
+			.OnTextCommitted(this, &FRigVMGraphDetailCustomization::SetNodeKeywords)
+		];
+
+		// description
+		SettingsCategory.AddCustomRow(FText::GetEmpty())
+		.NameContent()
+		[
+			SNew(STextBlock)
+			.Text(FText::FromString(TEXT("Description")))
+			.Font(IDetailLayoutBuilder::GetDetailFont())
+		]
+		.ValueContent()
+		[
+			SNew(SMultiLineEditableText)
+			.Font(IDetailLayoutBuilder::GetDetailFont())
+			.Text(this, &FRigVMGraphDetailCustomization::GetNodeDescription)
+			.OnTextCommitted(this, &FRigVMGraphDetailCustomization::SetNodeDescription)
+		];
+
+		if(AccessSpecifierStrings.IsEmpty())
+		{
+			AccessSpecifierStrings.Add(TSharedPtr<FString>(new FString(TEXT("Public"))));
+			AccessSpecifierStrings.Add(TSharedPtr<FString>(new FString(TEXT("Private"))));
+		}
+
+		// access specifier
+		SettingsCategory.AddCustomRow( LOCTEXT( "AccessSpecifier", "Access Specifier" ) )
+        .NameContent()
+        [
+            SNew(STextBlock)
+                .Text( LOCTEXT( "AccessSpecifier", "Access Specifier" ) )
+                .Font( IDetailLayoutBuilder::GetDetailFont() )
+        ]
+        .ValueContent()
+        [
+            SNew(SComboButton)
+            .ContentPadding(0)
+            .ButtonContent()
+            [
+                SNew(STextBlock)
+                    .Text(this, &FRigVMGraphDetailCustomization::GetCurrentAccessSpecifierName)
+                    .Font( IDetailLayoutBuilder::GetDetailFont() )
+            ]
+            .MenuContent()
+            [
+                SNew(SListView<TSharedPtr<FString> >)
+                    .ListItemsSource( &AccessSpecifierStrings )
+                    .OnGenerateRow(this, &FRigVMGraphDetailCustomization::HandleGenerateRowAccessSpecifier)
+                    .OnSelectionChanged(this, &FRigVMGraphDetailCustomization::OnAccessSpecifierSelected)
+            ]
+        ];
+	}
+
+	// node color
+	SettingsCategory.AddCustomRow(FText::GetEmpty())
+	.NameContent()
+	[
+		SNew(STextBlock)
+		.Text(FText::FromString(TEXT("Color")))
+		.Font(IDetailLayoutBuilder::GetDetailFont())
+	]
+	.ValueContent()
+	[
+		SNew(SButton)
+		.ButtonStyle(FAppStyle::Get(), "Menu.Button")
+		.OnClicked(this, &FRigVMGraphDetailCustomization::OnNodeColorClicked)
+		[
+			SAssignNew(ColorBlock, SColorBlock)
+			.Color(this, &FRigVMGraphDetailCustomization::GetNodeColor)
+			.Size(FVector2D(77, 16))
+		]
+	];
+
+	IDetailCategoryBuilder& DefaultsCategory = DetailLayout.EditCategory("NodeDefaults", LOCTEXT("FunctionDetailsNodeDefaults", "Node Defaults"));
+	TSharedRef<FRigVMFunctionArgumentDefaultNode> DefaultsArgumentNode = MakeShareable(new FRigVMFunctionArgumentDefaultNode(
+		Model,
+		Blueprint));
+	DefaultsCategory.AddCustomBuilder(DefaultsArgumentNode);
+
+}
+
+bool FRigVMGraphDetailCustomization::IsAddNewInputOutputEnabled() const
+{
+	return true;
+}
+
+EVisibility FRigVMGraphDetailCustomization::GetAddNewInputOutputVisibility() const
+{
+	return EVisibility::Visible;
+}
+
+FReply FRigVMGraphDetailCustomization::OnAddNewInputClicked()
+{
+	if (GraphPtr.IsValid() && RigVMBlueprintPtr.IsValid())
+	{
+		URigVMBlueprint* Blueprint = RigVMBlueprintPtr.Get();
+		URigVMGraph* Model = Blueprint->GetModel(GraphPtr.Get());
+		if (URigVMController* Controller = Blueprint->GetController(Model))
+		{
+			FName ArgumentName = TEXT("Argument");
+			FString CPPType = TEXT("bool");
+			FName CPPTypeObjectPath = NAME_None;
+			FString DefaultValue = TEXT("False");
+
+			if (URigVMLibraryNode* LibraryNode = Cast<URigVMLibraryNode>(Model->GetOuter()))
+			{
+				if (LibraryNode->GetPins().Num() > 0)
+				{
+					URigVMPin* LastPin = LibraryNode->GetPins().Last();
+					if (!LastPin->IsExecuteContext())
+					{
+						// strip off any tailing number from for example Argument_2
+						FString StrippedArgumentName = LastPin->GetName();
+						FString LastChars = StrippedArgumentName.Right(1);
+						StrippedArgumentName.LeftChopInline(1);
+						while(LastChars.IsNumeric() && !StrippedArgumentName.IsEmpty())
+						{
+							LastChars = StrippedArgumentName.Right(1);
+							StrippedArgumentName.LeftChopInline(1);
+
+							if(LastChars.StartsWith(TEXT("_")))
+							{
+								LastChars.Reset();
+								break;
+							}
+						}
+
+						StrippedArgumentName = StrippedArgumentName + LastChars;
+						if(!StrippedArgumentName.IsEmpty())
+						{
+							ArgumentName = *StrippedArgumentName;
+						}
+
+						RigVMTypeUtils::CPPTypeFromPin(LastPin, CPPType, CPPTypeObjectPath);						
+						DefaultValue = LastPin->GetDefaultValue();
+					}
+				}
+			}
+
+			Controller->AddExposedPin(ArgumentName, ERigVMPinDirection::Input, CPPType, CPPTypeObjectPath, DefaultValue, true, true);
+		}
+	}
+	return FReply::Unhandled();
+}
+
+FReply FRigVMGraphDetailCustomization::OnAddNewOutputClicked()
+{
+	if (GraphPtr.IsValid() && RigVMBlueprintPtr.IsValid())
+	{
+		URigVMBlueprint* Blueprint = RigVMBlueprintPtr.Get();
+		URigVMGraph* Model = Blueprint->GetModel(GraphPtr.Get());
+		if (URigVMController* Controller = Blueprint->GetController(Model))
+		{
+			FName ArgumentName = TEXT("Argument");
+			FString CPPType = TEXT("bool");
+			FName CPPTypeObjectPath = NAME_None;
+			FString DefaultValue = TEXT("False");
+			// todo: base decisions on types on last argument
+
+			Controller->AddExposedPin(ArgumentName, ERigVMPinDirection::Output, CPPType, CPPTypeObjectPath, DefaultValue, true, true);
+		}
+	}
+	return FReply::Unhandled();
+}
+
+FText FRigVMGraphDetailCustomization::GetNodeCategory() const
+{
+	if (GraphPtr.IsValid() && RigVMBlueprintPtr.IsValid())
+	{
+		URigVMBlueprint* Blueprint = RigVMBlueprintPtr.Get();
+		if (URigVMGraph* Model = Blueprint->GetModel(GraphPtr.Get()))
+		{
+			if (URigVMCollapseNode* OuterNode = Cast<URigVMCollapseNode>(Model->GetOuter()))
+			{
+				return FText::FromString(OuterNode->GetNodeCategory());
+			}
+		}
+	}
+
+	return FText();
+}
+
+void FRigVMGraphDetailCustomization::SetNodeCategory(const FText& InNewText, ETextCommit::Type InCommitType)
+{
+	if(InCommitType == ETextCommit::OnCleared)
+	{
+		return;
+	}
+
+	if (GraphPtr.IsValid() && RigVMBlueprintPtr.IsValid())
+	{
+		URigVMBlueprint* Blueprint = RigVMBlueprintPtr.Get();
+		if (URigVMGraph* Model = Blueprint->GetModel(GraphPtr.Get()))
+		{
+			if (URigVMCollapseNode* OuterNode = Cast<URigVMCollapseNode>(Model->GetOuter()))
+			{
+				if (URigVMController* Controller = Blueprint->GetOrCreateController(OuterNode->GetGraph()))
+				{
+					Controller->SetNodeCategory(OuterNode, InNewText.ToString(), true, false, true);
+				}
+			}
+		}
+	}
+}
+
+FText FRigVMGraphDetailCustomization::GetNodeKeywords() const
+{
+	if (GraphPtr.IsValid() && RigVMBlueprintPtr.IsValid())
+	{
+		URigVMBlueprint* Blueprint = RigVMBlueprintPtr.Get();
+		if (URigVMGraph* Model = Blueprint->GetModel(GraphPtr.Get()))
+		{
+			if (URigVMCollapseNode* OuterNode = Cast<URigVMCollapseNode>(Model->GetOuter()))
+			{
+				return FText::FromString(OuterNode->GetNodeKeywords());
+			}
+		}
+	}
+
+	return FText();
+}
+
+void FRigVMGraphDetailCustomization::SetNodeKeywords(const FText& InNewText, ETextCommit::Type InCommitType)
+{
+	if(InCommitType == ETextCommit::OnCleared)
+	{
+		return;
+	}
+
+	if (GraphPtr.IsValid() && RigVMBlueprintPtr.IsValid())
+	{
+		URigVMBlueprint* Blueprint = RigVMBlueprintPtr.Get();
+		if (URigVMGraph* Model = Blueprint->GetModel(GraphPtr.Get()))
+		{
+			if (URigVMCollapseNode* OuterNode = Cast<URigVMCollapseNode>(Model->GetOuter()))
+			{
+				if (URigVMController* Controller = Blueprint->GetOrCreateController(OuterNode->GetGraph()))
+				{
+					Controller->SetNodeKeywords(OuterNode, InNewText.ToString(), true, false, true);
+				}
+			}
+		}
+	}
+}
+
+FText FRigVMGraphDetailCustomization::GetNodeDescription() const
+{
+	if (GraphPtr.IsValid() && RigVMBlueprintPtr.IsValid())
+	{
+		URigVMBlueprint* Blueprint = RigVMBlueprintPtr.Get();
+		if (URigVMGraph* Model = Blueprint->GetModel(GraphPtr.Get()))
+		{
+			if (URigVMCollapseNode* OuterNode = Cast<URigVMCollapseNode>(Model->GetOuter()))
+			{
+				return FText::FromString(OuterNode->GetNodeDescription());
+			}
+		}
+	}
+
+	return FText();
+}
+
+void FRigVMGraphDetailCustomization::SetNodeDescription(const FText& InNewText, ETextCommit::Type InCommitType)
+{
+	if(InCommitType == ETextCommit::OnCleared)
+	{
+		return;
+	}
+
+	if (GraphPtr.IsValid() && RigVMBlueprintPtr.IsValid())
+	{
+		URigVMBlueprint* Blueprint = RigVMBlueprintPtr.Get();
+		if (URigVMGraph* Model = Blueprint->GetModel(GraphPtr.Get()))
+		{
+			if (URigVMCollapseNode* OuterNode = Cast<URigVMCollapseNode>(Model->GetOuter()))
+			{
+				if (URigVMController* Controller = Blueprint->GetOrCreateController(OuterNode->GetGraph()))
+				{
+					Controller->SetNodeDescription(OuterNode, InNewText.ToString(), true, false, true);
+				}
+			}
+		}
+	}
+}
+
+FLinearColor FRigVMGraphDetailCustomization::GetNodeColor() const
+{
+	if (GraphPtr.IsValid() && RigVMBlueprintPtr.IsValid())
+	{
+		URigVMBlueprint* Blueprint = RigVMBlueprintPtr.Get();
+		if (URigVMGraph* Model = Blueprint->GetModel(GraphPtr.Get()))
+		{
+			if (URigVMCollapseNode* OuterNode = Cast<URigVMCollapseNode>(Model->GetOuter()))
+			{
+				return OuterNode->GetNodeColor();
+			}
+		}
+	}
+	return FLinearColor::White;
+}
+
+void FRigVMGraphDetailCustomization::SetNodeColor(FLinearColor InColor, bool bSetupUndoRedo)
+{
+	if (GraphPtr.IsValid() && RigVMBlueprintPtr.IsValid())
+	{
+		URigVMBlueprint* Blueprint = RigVMBlueprintPtr.Get();
+		if (URigVMGraph* Model = Blueprint->GetModel(GraphPtr.Get()))
+		{
+			if (URigVMCollapseNode* OuterNode = Cast<URigVMCollapseNode>(Model->GetOuter()))
+			{
+				if (URigVMController* Controller = Blueprint->GetOrCreateController(OuterNode->GetGraph()))
+				{
+					Controller->SetNodeColor(OuterNode, InColor, bSetupUndoRedo, bIsPickingColor, true);
+				}
+			}
+		}
+	}
+}
+
+void FRigVMGraphDetailCustomization::OnNodeColorBegin()
+{
+	bIsPickingColor = true;
+}
+void FRigVMGraphDetailCustomization::OnNodeColorEnd()
+{ 
+	bIsPickingColor = false; 
+}
+
+void FRigVMGraphDetailCustomization::OnNodeColorCancelled(FLinearColor OriginalColor)
+{
+	SetNodeColor(OriginalColor, true);
+}
+
+FReply FRigVMGraphDetailCustomization::OnNodeColorClicked()
+{
+	FColorPickerArgs PickerArgs = FColorPickerArgs(GetNodeColor(), FOnLinearColorValueChanged::CreateSP(this, &FRigVMGraphDetailCustomization::SetNodeColor, true));
+	PickerArgs.ParentWidget = ColorBlock;
+	PickerArgs.bUseAlpha = false;
+	PickerArgs.DisplayGamma = false;
+	PickerArgs.OnInteractivePickBegin = FSimpleDelegate::CreateSP(this, &FRigVMGraphDetailCustomization::OnNodeColorBegin);
+	PickerArgs.OnInteractivePickEnd = FSimpleDelegate::CreateSP(this, &FRigVMGraphDetailCustomization::OnNodeColorEnd);
+	PickerArgs.OnColorPickerCancelled = FOnColorPickerCancelled::CreateSP(this, &FRigVMGraphDetailCustomization::OnNodeColorCancelled);
+	OpenColorPicker(PickerArgs);
+	return FReply::Handled();
+}
+
+TArray<TSharedPtr<FString>> FRigVMGraphDetailCustomization::AccessSpecifierStrings;
+
+FText FRigVMGraphDetailCustomization::GetCurrentAccessSpecifierName() const
+{
+	if(RigVMBlueprintPtr.IsValid() && GraphPtr.IsValid())
+	{
+		URigVMEdGraph* Graph = GraphPtr.Get();
+		URigVMBlueprint* RigVMBlueprint = RigVMBlueprintPtr.Get();
+
+		if (URigVMGraph* RigVMGraph = Graph->GetModel())
+		{
+			if (URigVMLibraryNode* LibraryNode = RigVMGraph->GetTypedOuter<URigVMLibraryNode>())
+			{
+				if(RigVMBlueprint->IsFunctionPublic(LibraryNode->GetFName()))
+				{
+					return FText::FromString(*AccessSpecifierStrings[0].Get()); // public
+				}
+			}
+		}
+	}
+
+	return FText::FromString(*AccessSpecifierStrings[1].Get()); // private
+}
+
+void FRigVMGraphDetailCustomization::OnAccessSpecifierSelected( TSharedPtr<FString> SpecifierName, ESelectInfo::Type SelectInfo )
+{
+	if(RigVMBlueprintPtr.IsValid() && GraphPtr.IsValid())
+	{
+		URigVMEdGraph* Graph = GraphPtr.Get();
+		URigVMBlueprint* RigVMBlueprint = RigVMBlueprintPtr.Get();
+
+		if (URigVMGraph* RigVMGraph = Graph->GetModel())
+		{
+			if (URigVMLibraryNode* LibraryNode = RigVMGraph->GetTypedOuter<URigVMLibraryNode>())
+			{
+				if(SpecifierName->Equals(TEXT("Private")))
+				{
+					RigVMBlueprint->MarkFunctionPublic(LibraryNode->GetFName(), false);
+				}
+				else
+				{
+					RigVMBlueprint->MarkFunctionPublic(LibraryNode->GetFName(), true);
+				}
+			}
+		}
+	}
+}
+
+TSharedRef<ITableRow> FRigVMGraphDetailCustomization::HandleGenerateRowAccessSpecifier( TSharedPtr<FString> SpecifierName, const TSharedRef<STableViewBase>& OwnerTable )
+{
+	return SNew(STableRow< TSharedPtr<FString> >, OwnerTable)
+        .Content()
+        [
+            SNew( STextBlock ) 
+                .Text(FText::FromString(*SpecifierName.Get()) )
+        ];
+}
+
+FRigVMWrappedNodeDetailCustomization::FRigVMWrappedNodeDetailCustomization()
+: BlueprintBeingCustomized(nullptr)
+{
+}
+
+TSharedRef<IDetailCustomization> FRigVMWrappedNodeDetailCustomization::MakeInstance()
+{
+	return MakeShareable(new FRigVMWrappedNodeDetailCustomization);
+}
+
+void FRigVMWrappedNodeDetailCustomization::CustomizeDetails(IDetailLayoutBuilder& DetailLayout)
+{
+	TArray<TWeakObjectPtr<UObject>> DetailObjects;
+	DetailLayout.GetObjectsBeingCustomized(DetailObjects);
+	if (DetailObjects.Num() == 0)
+	{
+		return;
+	}
+
+	for(TWeakObjectPtr<UObject> DetailObject : DetailObjects)
+	{
+		URigVMDetailsViewWrapperObject* WrapperObject = CastChecked<URigVMDetailsViewWrapperObject>(DetailObject.Get());
+		if(BlueprintBeingCustomized == nullptr)
+		{
+			BlueprintBeingCustomized = WrapperObject->GetTypedOuter<URigVMBlueprint>();
+		}
+
+		ObjectsBeingCustomized.Add(WrapperObject);
+		NodesBeingCustomized.Add(CastChecked<URigVMNode>(WrapperObject->GetSubject()));
+	}
+
+	if (BlueprintBeingCustomized == nullptr || ObjectsBeingCustomized.IsEmpty() || NodesBeingCustomized.IsEmpty())
+	{
+		return;
+	}
+
+	UClass* WrapperClass = ObjectsBeingCustomized[0]->GetClass();
+
+	// now loop over all of the properties and display them
+	TArray<TSharedPtr<IPropertyHandle>> PropertiesToVisit;
+	for (TFieldIterator<FProperty> PropertyIt(WrapperClass); PropertyIt; ++PropertyIt)
+	{
+		FProperty* Property = *PropertyIt;
+		TSharedPtr<IPropertyHandle> PropertyHandle = DetailLayout.GetProperty(Property->GetFName(), WrapperClass);
+		if (!PropertyHandle->IsValidHandle())
+		{
+			continue;
+		}
+		PropertiesToVisit.Add(PropertyHandle);
+
+		// check if any / all pins are bound to a variable
+		int32 PinsBoundToVariable = 0;
+		TArray<URigVMPin*> ModelPins;
+		for(TWeakObjectPtr<URigVMNode> Node : NodesBeingCustomized)
+		{
+			if(URigVMPin* ModelPin = Node->FindPin(Property->GetName()))
+			{
+				ModelPins.Add(ModelPin);
+				PinsBoundToVariable += ModelPin->IsBoundToVariable() ? 1 : 0;
+			}
+		}
+
+		if(PinsBoundToVariable > 0)
+		{
+			if(PinsBoundToVariable == ModelPins.Num())
+			{
+				if(IDetailPropertyRow* Row = DetailLayout.EditDefaultProperty(PropertyHandle))
+				{
+					Row->CustomWidget()
+					.NameContent()
+					[
+						PropertyHandle->CreatePropertyNameWidget()
+					]
+					.ValueContent()
+					[
+					SNew(SRigVMGraphVariableBinding)
+						.ModelPins(ModelPins)
+						.Blueprint(BlueprintBeingCustomized)
+					];
+				}
+			}
+			else // in this case some pins are bound, and some are not - we'll hide the input value widget
+			{
+				if(IDetailPropertyRow* Row = DetailLayout.EditDefaultProperty(PropertyHandle))
+				{
+					Row->CustomWidget()
+					.NameContent()
+					[
+						PropertyHandle->CreatePropertyNameWidget()
+					];
+				}
+			}
+
+			continue;
+		}
+		
+		if (FNameProperty* NameProperty = CastField<FNameProperty>(Property))
+        {
+        	FString CustomWidgetName = NameProperty->GetMetaData(TEXT("CustomWidget"));
+        	if (!CustomWidgetName.IsEmpty())
+        	{
+        		URigVMEdGraph* GraphBeingCustomized = Cast<URigVMEdGraph>(
+        			BlueprintBeingCustomized->GetEdGraph(NodesBeingCustomized[0]->GetGraph()));
+        		ensure(GraphBeingCustomized);
+        		
+        		const TArray<TSharedPtr<FString>>* NameList = GraphBeingCustomized->GetNameListForWidget(CustomWidgetName);
+        		if (NameList)
+        		{
+        			TSharedPtr<SRigVMGraphPinNameListValueWidget> NameListWidget;
+
+        			if(IDetailPropertyRow* Row = DetailLayout.EditDefaultProperty(PropertyHandle))
+        			{
+        				Row->CustomWidget()
+						.NameContent()
+						[
+							PropertyHandle->CreatePropertyNameWidget()
+						]
+						.ValueContent()
+						[
+							SAssignNew(NameListWidget, SRigVMGraphPinNameListValueWidget)
+							.OptionsSource(NameList)
+							.OnGenerateWidget(this, &FRigVMWrappedNodeDetailCustomization::MakeNameListItemWidget)
+							.OnSelectionChanged(this, &FRigVMWrappedNodeDetailCustomization::OnNameListChanged, NameProperty, DetailLayout.GetPropertyUtilities())
+							.OnComboBoxOpening(this, &FRigVMWrappedNodeDetailCustomization::OnNameListComboBox, NameProperty, NameList)
+							.InitiallySelectedItem(GetCurrentlySelectedItem(NameProperty, NameList))
+							.Content()
+							[
+								SNew(STextBlock)
+								.Text(this, &FRigVMWrappedNodeDetailCustomization::GetNameListText, NameProperty)
+								.ColorAndOpacity_Lambda([this, NameProperty]() -> FSlateColor
+								{
+									static FText NoneText = LOCTEXT("None", "None"); 
+									if(GetNameListText(NameProperty).EqualToCaseIgnored(NoneText))
+									{
+										return FSlateColor(FLinearColor::Red);
+									}
+									return FSlateColor::UseForeground();
+								})
+							]
+        				];
+        			}        			
+        			NameListWidgets.Add(Property->GetFName(), NameListWidget);
+        		}
+        		else
+        		{
+        			if(IDetailPropertyRow* Row = DetailLayout.EditDefaultProperty(PropertyHandle))
+        			{
+        				Row->CustomWidget()
+						.NameContent()
+						[
+							PropertyHandle->CreatePropertyNameWidget()
+						];
+        			}
+        		}
+        		continue;
+        	}
+        }
+	}
+
+	// now loop over all handles and determine expansion states of the corresponding pins
+	for (int32 Index = 0; Index < PropertiesToVisit.Num(); Index++)
+	{
+		TSharedPtr<IPropertyHandle> PropertyHandle = PropertiesToVisit[Index];
+		FProperty* Property = PropertyHandle->GetProperty();
+
+		// certain properties we don't look at for expansion states
+		if(FStructProperty* StructProperty = CastField<FStructProperty>(Property))
+		{
+			if(StructProperty->Struct == TBaseStructure<FVector>::Get() ||
+				StructProperty->Struct == TBaseStructure<FVector2D>::Get() ||
+				StructProperty->Struct == TBaseStructure<FRotator>::Get() ||
+				StructProperty->Struct == TBaseStructure<FQuat>::Get())
+			{
+				continue;
+			}
+		}
+
+		bool bFound = false;
+		const FString PinPath = PropertyHandle->GeneratePathToProperty();
+		for(TWeakObjectPtr<URigVMNode> Node : NodesBeingCustomized)
+		{
+			if(URigVMPin* Pin = Node->FindPin(PinPath))
+			{
+				bFound = true;
+				
+				if(Pin->IsExpanded())
+				{
+					if(IDetailPropertyRow* Row = DetailLayout.EditDefaultProperty(PropertyHandle))
+					{
+						Row->ShouldAutoExpand(true);
+					}
+					break;
+				}
+			}
+		}
+
+		if(!bFound)
+		{
+			continue;
+		}
+
+		uint32 NumChildren = 0;
+		PropertyHandle->GetNumChildren(NumChildren);
+		for(uint32 ChildIndex = 0; ChildIndex < NumChildren; ChildIndex++)
+		{
+			PropertiesToVisit.Add(PropertyHandle->GetChildHandle(ChildIndex));
+		}
+	}
+
+	CustomizeLiveValues(DetailLayout);
+}
+
+TSharedRef<SWidget> FRigVMWrappedNodeDetailCustomization::MakeNameListItemWidget(TSharedPtr<FString> InItem)
+{
+	return 	SNew(STextBlock).Text(FText::FromString(*InItem));// .Font(FAppStyle::GetFontStyle(TEXT("PropertyWindow.NormalFont")));
+}
+
+FText FRigVMWrappedNodeDetailCustomization::GetNameListText(FNameProperty* InProperty) const
+{
+	FText FirstText;
+	for(TWeakObjectPtr<URigVMDetailsViewWrapperObject> ObjectBeingCustomized : ObjectsBeingCustomized)
+	{
+		if (FName* Value = InProperty->ContainerPtrToValuePtr<FName>(ObjectBeingCustomized.Get()))
+		{
+			FText Text = FText::FromName(*Value);
+			if(FirstText.IsEmpty())
+			{
+				FirstText = Text;
+			}
+			else if(!FirstText.EqualTo(Text))
+			{
+				return RigVMGraphDetailCustomizationMultipleValues;
+			}
+		}
+	}
+	return FirstText;
+}
+
+TSharedPtr<FString> FRigVMWrappedNodeDetailCustomization::GetCurrentlySelectedItem(FNameProperty* InProperty, const TArray<TSharedPtr<FString>>* InNameList) const
+{
+	FString CurrentItem = GetNameListText(InProperty).ToString();
+	for (const TSharedPtr<FString>& Item : *InNameList)
+	{
+		if (Item->Equals(CurrentItem))
+		{
+			return Item;
+		}
+	}
+	return TSharedPtr<FString>();
+}
+
+
+void FRigVMWrappedNodeDetailCustomization::SetNameListText(const FText& NewTypeInValue, ETextCommit::Type, FNameProperty* InProperty, TSharedRef<IPropertyUtilities> PropertyUtilities)
+{
+	URigVMGraph* Graph = NodesBeingCustomized[0]->GetGraph();
+	URigVMController* Controller = BlueprintBeingCustomized->GetController(Graph);
+
+	Controller->OpenUndoBracket(FString::Printf(TEXT("Set %s"), *InProperty->GetName()));
+	
+	for(TWeakObjectPtr<URigVMNode> Node : NodesBeingCustomized)
+	{
+		if(URigVMPin* Pin = Node->FindPin(InProperty->GetName()))
+		{
+			Controller->SetPinDefaultValue(Pin->GetPinPath(), NewTypeInValue.ToString(), false, true, false, true);
+		}
+	}
+
+	Controller->CloseUndoBracket();
+}
+
+void FRigVMWrappedNodeDetailCustomization::OnNameListChanged(TSharedPtr<FString> NewSelection, ESelectInfo::Type SelectInfo, FNameProperty* InProperty, TSharedRef<IPropertyUtilities> PropertyUtilities)
+{
+	if (SelectInfo != ESelectInfo::Direct)
+	{
+		FString NewValue = *NewSelection.Get();
+		SetNameListText(FText::FromString(NewValue), ETextCommit::OnEnter, InProperty, PropertyUtilities);
+	}
+}
+
+void FRigVMWrappedNodeDetailCustomization::OnNameListComboBox(FNameProperty* InProperty, const TArray<TSharedPtr<FString>>* InNameList)
+{
+	TSharedPtr<SRigVMGraphPinNameListValueWidget> Widget = NameListWidgets.FindChecked(InProperty->GetFName());
+	const TSharedPtr<FString> CurrentlySelected = GetCurrentlySelectedItem(InProperty, InNameList);
+	Widget->SetSelectedItem(CurrentlySelected);
+}
+
+void FRigVMWrappedNodeDetailCustomization::CustomizeLiveValues(IDetailLayoutBuilder& DetailLayout)
+{
+	if(ObjectsBeingCustomized.Num() != 1)
+	{
+		return;
+	}
+	
+	URigVMHost* DebuggedHost = Cast<URigVMHost>(BlueprintBeingCustomized->GetObjectBeingDebugged());
+	if(DebuggedHost == nullptr)
+	{
+		return;
+	}
+
+	URigVM* VM = DebuggedHost->GetVM();
+	if(VM == nullptr)
+	{
+		return;
+	}
+
+	URigVMDetailsViewWrapperObject* FirstWrapper = ObjectsBeingCustomized[0].Get();
+	URigVMNode* FirstNode = NodesBeingCustomized[0].Get();
+	if(FirstNode->GetTypedOuter<URigVMFunctionLibrary>())
+	{
+		return;
+	}
+
+	TSharedPtr<FRigVMParserAST> AST = FirstNode->GetGraph()->GetRuntimeAST(BlueprintBeingCustomized->VMCompileSettings.ASTSettings, false);
+	if(!AST.IsValid())
+	{
+		return;
+	}
+
+	FRigVMByteCode& ByteCode = VM->GetByteCode();
+	if(ByteCode.GetFirstInstructionIndexForSubject(FirstNode) == INDEX_NONE)
+	{
+		return;
+	}
+	
+	IDetailCategoryBuilder& DebugCategory = DetailLayout.EditCategory("DebugLiveValues", LOCTEXT("DebugLiveValues", "Inspect Live Values"), ECategoryPriority::Uncommon);
+	DebugCategory.InitiallyCollapsed(true);
+
+	for(URigVMPin* Pin : FirstNode->GetPins())
+	{
+		// only show hidden pins in debug mode
+		if(Pin->GetDirection() == ERigVMPinDirection::Hidden)
+		{
+			if(!DebuggedHost->IsInDebugMode())
+			{
+				continue;
+			}
+		}
+		
+		URigVMPin* SourcePin = Pin;
+		if(BlueprintBeingCustomized->VMCompileSettings.ASTSettings.bFoldAssignments)
+		{
+			do
+			{
+				TArray<URigVMPin*> SourcePins = SourcePin->GetLinkedSourcePins(false);
+				if(SourcePins.Num() > 0)
+				{
+					SourcePin = SourcePins[0];
+				}
+				else
+				{
+					break;
+				}
+			}
+			while(SourcePin->GetNode()->IsA<URigVMRerouteNode>());
+		}
+		
+		TArray<const FRigVMExprAST*> Expressions = AST->GetExpressionsForSubject(SourcePin);
+		if(Expressions.Num() == 0 && SourcePin != Pin)
+		{
+			SourcePin = Pin;
+			Expressions = AST->GetExpressionsForSubject(Pin);
+		}
+
+		bool bHasVar = false;
+		for(const FRigVMExprAST* Expression : Expressions)
+		{
+			if(Expression->IsA(FRigVMExprAST::EType::Literal))
+			{
+				continue;
+			}
+			else if(Expression->IsA(FRigVMExprAST::EType::Var))
+			{
+				bHasVar = true;
+				break;
+			}
+		}
+
+		TArray<const FRigVMExprAST*> FilteredExpressions;
+		for(const FRigVMExprAST* Expression : Expressions)
+		{
+			if(Expression->IsA(FRigVMExprAST::EType::Literal))
+			{
+				if(bHasVar)
+				{
+					continue;
+				}
+				FilteredExpressions.Add(Expression);
+			}
+			else if(Expression->IsA(FRigVMExprAST::EType::Var))
+			{
+				FilteredExpressions.Add(Expression);
+			}
+			else if(Expression->IsA(FRigVMExprAST::EType::CachedValue))
+			{
+				const FRigVMCachedValueExprAST* CachedValueExpr = Expression->To<FRigVMCachedValueExprAST>();
+				FilteredExpressions.Add(CachedValueExpr->GetVarExpr());
+			}
+		}
+
+		bool bAddedProperty = false;
+		int32 SuffixIndex = 1;
+		FString NameSuffix;
+
+		TArray<FRigVMOperand> KnownOperands; 
+		for(const FRigVMExprAST* Expression : FilteredExpressions)
+		{
+			const FRigVMVarExprAST* VarExpr = Expression->To<FRigVMVarExprAST>();
+
+			FString PinHash = URigVMCompiler::GetPinHash(SourcePin, VarExpr, false);
+			const FRigVMOperand* Operand = BlueprintBeingCustomized->PinToOperandMap.Find(PinHash);
+			if(Operand)
+			{
+				if(Operand->GetRegisterOffset() != INDEX_NONE)
+				{
+					continue;
+				}
+				if(KnownOperands.Contains(*Operand))
+				{
+					continue;
+				}
+
+				const FProperty* Property = nullptr;
+				TArray<UObject*> ExternalObjects;
+
+				if(Operand->GetMemoryType() == ERigVMMemoryType::External)
+				{
+					if(!VM->GetExternalVariableDefs().IsValidIndex(Operand->GetRegisterIndex()))
+					{
+						continue;
+					}
+					ExternalObjects.Add(DebuggedHost);
+					Property = VM->GetExternalVariableDefs()[Operand->GetRegisterIndex()].Property; 
+				}
+				else
+				{
+					URigVMMemoryStorage* Memory = VM->GetMemoryByType(Operand->GetMemoryType());
+					if(Memory == nullptr)
+					{
+						continue;
+					}
+
+					if(Memory->GetOuter() == GetTransientPackage())
+					{
+						continue;
+					}
+
+					// the UClass must be alive for the details view to access it
+					// this ensure can fail if VM memory is not updated immediately after compile
+					// because of deferred copy
+					if(!ensure(IsValidChecked(Memory->GetClass())))
+					{
+						continue;
+					}
+
+					if(Memory->GetClass()->GetOuter() == GetTransientPackage())
+					{
+						continue;
+					}
+
+					if(!Memory->IsValidIndex(Operand->GetRegisterIndex()))
+					{
+						continue;
+					}
+					
+					Property = Memory->GetProperty(Operand->GetRegisterIndex());
+					if(Property == nullptr)
+					{
+						continue;
+					}
+
+					ExternalObjects.Add(Memory);
+				}
+
+				check(ExternalObjects.Num() > 0);
+				check(Property);
+
+				IDetailPropertyRow* PropertyRow = DebugCategory.AddExternalObjectProperty(ExternalObjects, Property->GetFName(), EPropertyLocation::Default, FAddPropertyParams().ForceShowProperty());
+				if(PropertyRow)
+				{
+					PropertyRow->DisplayName(FText::FromString(FString::Printf(TEXT("%s%s"), *Pin->GetName(), *NameSuffix)));
+					PropertyRow->IsEnabled(false);
+
+					SuffixIndex++;
+					bAddedProperty = true;
+					NameSuffix = FString::Printf(TEXT("_%d"), SuffixIndex);
+				}
+
+				KnownOperands.Add(*Operand);
+			}
+		}
+
+		if(!bAddedProperty)
+		{
+			TSharedPtr<IPropertyHandle> PinHandle = DetailLayout.GetProperty(Pin->GetFName());
+			if(PinHandle.IsValid())
+			{
+				DebugCategory.AddProperty(PinHandle)
+				.DisplayName(FText::FromName(Pin->GetDisplayName()))
+				.IsEnabled(false);
+			}
+		}
+	}
+}
+
+FRigVMGraphMathTypeDetailCustomization::FRigVMGraphMathTypeDetailCustomization()
+: ScriptStruct(nullptr)
+, BlueprintBeingCustomized(nullptr)
+, GraphBeingCustomized(nullptr)
+{
+}
+
+void FRigVMGraphMathTypeDetailCustomization::CustomizeHeader(TSharedRef<IPropertyHandle> InPropertyHandle,
+                                                           FDetailWidgetRow& HeaderRow, IPropertyTypeCustomizationUtils& StructCustomizationUtils)
+{
+	TArray<UObject*> Objects;
+	InPropertyHandle->GetOuterObjects(Objects);
+
+	for (UObject* Object : Objects)
+	{
+		ObjectsBeingCustomized.Add(Object);
+
+		if(BlueprintBeingCustomized == nullptr)
+		{
+			BlueprintBeingCustomized = Object->GetTypedOuter<URigVMBlueprint>();
+		}
+
+		if(GraphBeingCustomized == nullptr)
+		{
+			GraphBeingCustomized = Object->GetTypedOuter<URigVMGraph>();
+		}
+	}
+
+	FProperty* Property = InPropertyHandle->GetProperty();
+	const FStructProperty* StructProperty = CastField<FStructProperty>(Property);
+	ScriptStruct = StructProperty->Struct;
+}
+
+void FRigVMGraphMathTypeDetailCustomization::CustomizeChildren(TSharedRef<IPropertyHandle> InPropertyHandle,
+	IDetailChildrenBuilder& StructBuilder, IPropertyTypeCustomizationUtils& StructCustomizationUtils)
+{
+	if(!InPropertyHandle->IsValidHandle())
+	{
+		return;
+	}
+
+	if(ScriptStruct == TBaseStructure<FVector>::Get())
+	{
+		CustomizeVector<FVector, 3>(InPropertyHandle, StructBuilder, StructCustomizationUtils);
+	}
+	else if(ScriptStruct == TBaseStructure<FVector2D>::Get())
+	{
+		CustomizeVector<FVector2D, 2>(InPropertyHandle, StructBuilder, StructCustomizationUtils);
+	}
+	else if(ScriptStruct == TBaseStructure<FVector4>::Get())
+	{
+		CustomizeVector<FVector4, 4>(InPropertyHandle, StructBuilder, StructCustomizationUtils);
+	}
+	else if(ScriptStruct == TBaseStructure<FRotator>::Get())
+	{
+		CustomizeRotation<FRotator>(InPropertyHandle, StructBuilder, StructCustomizationUtils);
+	}
+	else if(ScriptStruct == TBaseStructure<FQuat>::Get())
+	{
+		CustomizeRotation<FQuat>(InPropertyHandle, StructBuilder, StructCustomizationUtils);
+	}
+	else if(ScriptStruct == TBaseStructure<FTransform>::Get())
+	{
+		CustomizeTransform<FTransform>(InPropertyHandle, StructBuilder, StructCustomizationUtils);
+	}
+	else if(ScriptStruct == TBaseStructure<FEulerTransform>::Get())
+	{
+		CustomizeTransform<FEulerTransform>(InPropertyHandle, StructBuilder, StructCustomizationUtils);
+	}
+}
+
+#undef LOCTEXT_NAMESPACE
diff --git a/Engine/Plugins/Runtime/RigVM/Source/RigVMEditor/Private/Editor/RigVMVariableDetailCustomization.cpp b/Engine/Plugins/Runtime/RigVM/Source/RigVMEditor/Private/Editor/RigVMVariableDetailCustomization.cpp
new file mode 100644
--- /dev/null
+++ b/Engine/Plugins/Runtime/RigVM/Source/RigVMEditor/Private/Editor/RigVMVariableDetailCustomization.cpp
@@ -0,0 +1,114 @@
+// Copyright Epic Games, Inc. All Rights Reserved.
+
+#include "Editor/RigVMVariableDetailCustomization.h"
+#include "Kismet2/BlueprintEditorUtils.h"
+#include "RigVMHost.h"
+#include "DetailLayoutBuilder.h"
+#include "BlueprintEditorModule.h"
+#include "DetailCategoryBuilder.h"
+#include "DetailWidgetRow.h"
+#include "Widgets/Text/STextBlock.h"
+#include "Widgets/Input/SCheckBox.h"
+#include "Engine/BlueprintGeneratedClass.h"
+#include "RigVMBlueprint.h"
+#include "RigVMModel/RigVMController.h"
+
+#define LOCTEXT_NAMESPACE "RigVMVariableDetailCustomization"
+
+TSharedPtr<IDetailCustomization> FRigVMVariableDetailCustomization::MakeInstance(TSharedPtr<IBlueprintEditor> InBlueprintEditor)
+{
+	const TArray<UObject*>* Objects = (InBlueprintEditor.IsValid() ? InBlueprintEditor->GetObjectsCurrentlyBeingEdited() : nullptr);
+	if (Objects && Objects->Num() == 1)
+	{
+		if (UBlueprint* Blueprint = Cast<UBlueprint>((*Objects)[0]))
+		{
+			if (Blueprint->ParentClass && Blueprint->ParentClass->IsChildOf(URigVMHost::StaticClass()))
+			{
+				return MakeShareable(new FRigVMVariableDetailCustomization(InBlueprintEditor, Blueprint));
+			}
+		}
+	}
+
+	return nullptr;
+}
+
+void FRigVMVariableDetailCustomization::CustomizeDetails(IDetailLayoutBuilder& DetailLayout)
+{
+	TArray<TWeakObjectPtr<UObject>> ObjectsBeingCustomized;
+	DetailLayout.GetObjectsBeingCustomized(ObjectsBeingCustomized);
+	if (ObjectsBeingCustomized.Num() > 0)
+	{
+		UPropertyWrapper* PropertyWrapper = Cast<UPropertyWrapper>(ObjectsBeingCustomized[0].Get());
+		TWeakFieldPtr<FProperty> PropertyBeingCustomized = PropertyWrapper ? PropertyWrapper->GetProperty() : nullptr;
+		if (PropertyBeingCustomized.IsValid())
+		{
+			const FText AnimationInputText = LOCTEXT("AnimationInput", "Animation Input");
+			const FText AnimationOutputText = LOCTEXT("AnimationOutput", "Animation Output");
+			const FText AnimationInputTooltipText = LOCTEXT("AnimationInputTooltip", "Whether this variable acts as an input to this animation controller.\nSelecting this allow it to be exposed as an input pin on Evaluation nodes.");
+			const FText AnimationOutputTooltipText = LOCTEXT("AnimationOutputTooltip", "Whether this variable acts as an output from this animation controller.\nSelecting this will add a pin to the Animation Output node.");
+			
+			DetailLayout.EditCategory("Variable")
+				.AddCustomRow(AnimationInputText)
+				.NameContent()
+				[
+					SNew(STextBlock)
+					.Visibility(IsAnimationFlagEnabled(PropertyBeingCustomized) ? EVisibility::Visible : EVisibility::Hidden)
+					.Font(DetailLayout.GetDetailFont())
+					.Text(AnimationOutputText)
+					.ToolTipText(AnimationOutputTooltipText)
+				]
+				.ValueContent()
+				[
+					SNew(SCheckBox)
+					.Visibility(IsAnimationFlagEnabled(PropertyBeingCustomized) ? EVisibility::Visible : EVisibility::Hidden)
+					.IsChecked_Raw(this, &FRigVMVariableDetailCustomization::IsAnimationOutputChecked, PropertyBeingCustomized)
+					.OnCheckStateChanged_Raw(this, &FRigVMVariableDetailCustomization::HandleAnimationOutputCheckStateChanged, PropertyBeingCustomized)
+					.ToolTipText(AnimationOutputTooltipText)
+				];
+
+			DetailLayout.EditCategory("Variable")
+				.AddCustomRow(AnimationOutputText)
+				.NameContent()
+				[
+					SNew(STextBlock)
+					.Visibility(IsAnimationFlagEnabled(PropertyBeingCustomized) ? EVisibility::Visible : EVisibility::Hidden)
+					.Font(DetailLayout.GetDetailFont())
+					.Text(AnimationInputText)
+					.ToolTipText(AnimationInputTooltipText)
+				]
+				.ValueContent()
+				[
+					SNew(SCheckBox)
+					.Visibility(IsAnimationFlagEnabled(PropertyBeingCustomized) ? EVisibility::Visible : EVisibility::Hidden)
+					.IsChecked_Raw(this, &FRigVMVariableDetailCustomization::IsAnimationInputChecked, PropertyBeingCustomized)
+					.OnCheckStateChanged_Raw(this, &FRigVMVariableDetailCustomization::HandleAnimationInputCheckStateChanged, PropertyBeingCustomized)
+					.ToolTipText(AnimationInputTooltipText)
+				];
+		}
+	}
+}
+
+bool FRigVMVariableDetailCustomization::IsAnimationFlagEnabled(TWeakFieldPtr<FProperty> PropertyBeingCustomized) const
+{
+	return false;
+}
+
+ECheckBoxState FRigVMVariableDetailCustomization::IsAnimationOutputChecked(TWeakFieldPtr<FProperty> PropertyBeingCustomized) const
+{
+	return ECheckBoxState::Unchecked;
+}
+
+void FRigVMVariableDetailCustomization::HandleAnimationOutputCheckStateChanged(ECheckBoxState CheckBoxState, TWeakFieldPtr<FProperty> PropertyBeingCustomized)
+{
+}
+
+ECheckBoxState FRigVMVariableDetailCustomization::IsAnimationInputChecked(TWeakFieldPtr<FProperty> PropertyBeingCustomized) const
+{
+	return ECheckBoxState::Unchecked;
+}
+
+void FRigVMVariableDetailCustomization::HandleAnimationInputCheckStateChanged(ECheckBoxState CheckBoxState, TWeakFieldPtr<FProperty> PropertyBeingCustomized)
+{
+}
+
+#undef LOCTEXT_NAMESPACE
