RigVMEditor: Refactoring editor module, blueprint actions etc

diff --git a/Engine/Plugins/Runtime/RigVM/Source/RigVMEditor/Private/RigVMEditorModule.cpp b/Engine/Plugins/Runtime/RigVM/Source/RigVMEditor/Private/RigVMEditorModule.cpp
--- a/Engine/Plugins/Runtime/RigVM/Source/RigVMEditor/Private/RigVMEditorModule.cpp
+++ b/Engine/Plugins/Runtime/RigVM/Source/RigVMEditor/Private/RigVMEditorModule.cpp
@@ -6,11 +6,42 @@
 
 #include "RigVMEditorModule.h"
 #include "Modules/ModuleManager.h"
+#include "AssetRegistry/AssetRegistryModule.h"
 #include "Editor/RigVMExecutionStackCommands.h"
 #include "Editor/RigVMEditorStyle.h"
+#include "Editor/RigVMGraphDetailCustomization.h"
+#include "Editor/RigVMVariableDetailCustomization.h"
+#include "Kismet2/BlueprintEditorUtils.h"
+#include "Kismet2/KismetEditorUtilities.h"
+#include "Kismet2/KismetDebugUtilities.h"
+#include "GraphEditorActions.h"
+#include "RigVMBlueprintUtils.h"
+#include "UserDefinedStructureCompilerUtils.h"
+#include "EdGraph/RigVMEdGraphConnectionDrawingPolicy.h"
+#include "BlueprintActionDatabaseRegistrar.h"
+#include "EdGraph/NodeSpawners/RigVMEdGraphUnitNodeSpawner.h"
+#include "EdGraph/NodeSpawners/RigVMEdGraphVariableNodeSpawner.h"
+#include "EdGraph/NodeSpawners/RigVMEdGraphTemplateNodeSpawner.h"
+#include "EdGraph/NodeSpawners/RigVMEdGraphEnumNodeSpawner.h"
+#include "EdGraph/NodeSpawners/RigVMEdGraphFunctionRefNodeSpawner.h"
+#include "EdGraph/NodeSpawners/RigVMEdGraphInvokeEntryNodeSpawner.h"
+#include "RigVMModel/Nodes/RigVMAggregateNode.h"
+#include "RigVMModel/Nodes/RigVMDispatchNode.h"
+#include "Dialog/SCustomDialog.h"
+#include "Widgets/SRigVMGraphChangePinType.h"
+#include "Widgets/SRigVMGraphPinVariableBinding.h"
+#include "RigVMFunctions/RigVMDispatch_Array.h"
+#include "RigVMFunctions/RigVMDispatch_MakeStruct.h"
+#include "RigVMFunctions/RigVMDispatch_Constant.h"
+#include "RigVMFunctions/Simulation/RigVMFunction_AlphaInterp.h"
+#include "RigVMFunctions/Debug/RigVMFunction_VisualDebug.h"
+#include "ScopedTransaction.h"
+#include "UObject/UObjectIterator.h"
 
 DEFINE_LOG_CATEGORY(LogRigVMEditor);
 
+#define LOCTEXT_NAMESPACE "RigVMEditorModule"
+
 IMPLEMENT_MODULE(FRigVMEditorModule, RigVMEditor)
 
 FRigVMEditorModule& FRigVMEditorModule::Get()
@@ -20,11 +51,1767 @@ FRigVMEditorModule& FRigVMEditorModule::Get()
 
 void FRigVMEditorModule::StartupModule()
 {
-	FRigVMExecutionStackCommands::Register();
-	FRigVMEditorStyle::Register();
+	if(IsRigVMEditorModuleBase())
+	{
+		FRigVMExecutionStackCommands::Register();
+		FRigVMEditorStyle::Register();
+
+		EdGraphPanelNodeFactory = MakeShared<FRigVMEdGraphPanelNodeFactory>();
+		EdGraphPanelPinFactory = MakeShared<FRigVMEdGraphPanelPinFactory>();
+		FEdGraphUtilities::RegisterVisualNodeFactory(EdGraphPanelNodeFactory);
+		FEdGraphUtilities::RegisterVisualPinFactory(EdGraphPanelPinFactory);
+
+		ReconstructAllNodesDelegateHandle = FBlueprintEditorUtils::OnReconstructAllNodesEvent.AddStatic(&FRigVMBlueprintUtils::HandleReconstructAllNodes);
+		RefreshAllNodesDelegateHandle = FBlueprintEditorUtils::OnRefreshAllNodesEvent.AddStatic(&FRigVMBlueprintUtils::HandleRefreshAllNodes);
+
+		// Register Blueprint editor variable customization
+		FBlueprintEditorModule& BlueprintEditorModule = FModuleManager::LoadModuleChecked<FBlueprintEditorModule>("Kismet");
+		BlueprintVariableCustomizationHandle = BlueprintEditorModule.RegisterVariableCustomization(FProperty::StaticClass(), FOnGetVariableCustomizationInstance::CreateStatic(&FRigVMVariableDetailCustomization::MakeInstance));
+	}
+
+	StartupModuleCommon();
+}
+
+void FRigVMEditorModule::StartupModuleCommon()
+{
+	MenuExtensibilityManager = MakeShareable(new FExtensibilityManager);
+    ToolBarExtensibilityManager = MakeShareable(new FExtensibilityManager);
+
+	const UClass* BlueprintClass = GetRigVMBlueprintClass();
+	const URigVMBlueprint* BlueprintCDO = GetRigVMBlueprintCDO();
+	const UClass* EdGraphSchemaClass = BlueprintCDO->GetRigVMEdGraphSchemaClass();
+	const URigVMEdGraphSchema* SchemaCDO = CastChecked<URigVMEdGraphSchema>(EdGraphSchemaClass->GetDefaultObject());
+	FBlueprintEditorModule& BlueprintEditorModule = FModuleManager::LoadModuleChecked<FBlueprintEditorModule>("Kismet");
+	BlueprintEditorModule.RegisterGraphCustomization(SchemaCDO, FOnGetGraphCustomizationInstance::CreateStatic(&FRigVMGraphDetailCustomization::MakeInstance, BlueprintClass));
+
+	// Register to fixup newly created BPs
+	FKismetEditorUtilities::RegisterOnBlueprintCreatedCallback(this, BlueprintCDO->GetRigVMHostClass(), FKismetEditorUtilities::FOnBlueprintCreated::CreateRaw(this, &FRigVMEditorModule::HandleNewBlueprintCreated));
 }
 
 void FRigVMEditorModule::ShutdownModule()
 {
-	FRigVMEditorStyle::Unregister();
+	if (!IsEngineExitRequested())
+	{
+		if(IsRigVMEditorModuleBase())
+		{
+			FRigVMEditorStyle::Unregister();
+			FEdGraphUtilities::UnregisterVisualNodeFactory(EdGraphPanelNodeFactory);
+			FEdGraphUtilities::UnregisterVisualPinFactory(EdGraphPanelPinFactory);
+
+			FBlueprintEditorUtils::OnRefreshAllNodesEvent.Remove(RefreshAllNodesDelegateHandle);
+			FBlueprintEditorUtils::OnReconstructAllNodesEvent.Remove(ReconstructAllNodesDelegateHandle);
+
+			// Unregister Blueprint editor variable customization
+			FBlueprintEditorModule& BlueprintEditorModule = FModuleManager::LoadModuleChecked<FBlueprintEditorModule>("Kismet");
+			BlueprintEditorModule.UnregisterVariableCustomization(FProperty::StaticClass(), BlueprintVariableCustomizationHandle);
+		}
+	}
+
+	ShutdownModuleCommon();
 }
+
+void FRigVMEditorModule::ShutdownModuleCommon()
+{
+	if (!IsEngineExitRequested())
+	{
+		const URigVMBlueprint* BlueprintCDO = GetRigVMBlueprintCDO();
+		const UClass* EdGraphSchemaClass = BlueprintCDO->GetRigVMEdGraphSchemaClass();
+		const URigVMEdGraphSchema* SchemaCDO = CastChecked<URigVMEdGraphSchema>(EdGraphSchemaClass->GetDefaultObject());
+		FBlueprintEditorModule& BlueprintEditorModule = FModuleManager::LoadModuleChecked<FBlueprintEditorModule>("Kismet");
+		BlueprintEditorModule.UnregisterGraphCustomization(SchemaCDO);
+	}
+
+	// Unregister to fixup newly created BPs
+	FKismetEditorUtilities::UnregisterAutoBlueprintNodeCreation(this);
+}
+
+UClass* FRigVMEditorModule::GetRigVMBlueprintClass() const
+{
+	return URigVMBlueprint::StaticClass();
+}
+
+const URigVMBlueprint* FRigVMEditorModule::GetRigVMBlueprintCDO() const
+{
+	const UClass* RigVMBlueprintClass = GetRigVMBlueprintClass();
+	return CastChecked<URigVMBlueprint>(RigVMBlueprintClass->GetDefaultObject());
+}
+
+void FRigVMEditorModule::GetTypeActions(URigVMBlueprint* RigVMBlueprint, FBlueprintActionDatabaseRegistrar& ActionRegistrar)
+{
+	FAssetRegistryModule& AssetRegistryModule = FModuleManager::LoadModuleChecked<FAssetRegistryModule>("AssetRegistry");
+
+	// actions get registered under specific object-keys; the idea is that 
+	// actions might have to be updated (or deleted) if their object-key is  
+	// mutated (or removed)... here we use the class (so if the class 
+	// type disappears, then the action should go with it)
+	UClass* ActionKey = RigVMBlueprint->GetClass();
+	// to keep from needlessly instantiating a UBlueprintNodeSpawner, first   
+	// check to make sure that the registrar is looking for actions of this type
+	// (could be regenerating actions for a specific asset, and therefore the 
+	// registrar would only accept actions corresponding to that asset)
+	if (!ActionRegistrar.IsOpenForRegistration(ActionKey))
+	{
+		return;
+	}
+
+	FRigVMRegistry& Registry = FRigVMRegistry::Get();
+	Registry.RefreshEngineTypes();
+
+	const URigVMBlueprint* BlueprintCDO = GetRigVMBlueprintCDO();
+
+	for (const FRigVMTemplate& Template : Registry.GetTemplates())
+	{
+		// factories are registered below
+		if(Template.UsesDispatch())
+		{
+			continue;
+		}
+
+		// ignore templates that have only one permutation
+		if (Template.NumPermutations() <= 1)
+		{
+			continue;
+		}
+
+		// ignore templates which don't have a function backing it up
+		if(Template.GetPermutation(0) == nullptr)
+		{
+			continue;
+		}
+
+		if(!Template.SupportsExecuteContextStruct(BlueprintCDO->GetRigVMExecuteContextStruct()))
+		{
+			continue;
+		}
+
+		FText NodeCategory = FText::FromString(Template.GetCategory());
+		FText MenuDesc = FText::FromName(Template.GetName());
+		FText ToolTip = Template.GetTooltipText();
+
+		UBlueprintNodeSpawner* NodeSpawner = URigVMEdGraphTemplateNodeSpawner::CreateFromNotation(Template.GetNotation(), MenuDesc, NodeCategory, ToolTip);
+		check(NodeSpawner != nullptr);
+		ActionRegistrar.AddBlueprintAction(ActionKey, NodeSpawner);
+	};
+
+	for (const FRigVMDispatchFactory* Factory : Registry.GetFactories())
+	{
+		if(!Factory->SupportsExecuteContextStruct(BlueprintCDO->GetRigVMExecuteContextStruct()))
+		{
+			continue;
+		}
+
+		const FRigVMTemplate* Template = Factory->GetTemplate();
+		if(Template == nullptr)
+		{
+			continue;
+		}
+
+		FText NodeCategory = FText::FromString(Factory->GetCategory());
+		FText MenuDesc = FText::FromString(Factory->GetNodeTitle(FRigVMTemplateTypeMap()));
+		FText ToolTip = Factory->GetNodeTooltip(FRigVMTemplateTypeMap());
+
+		UBlueprintNodeSpawner* NodeSpawner = URigVMEdGraphTemplateNodeSpawner::CreateFromNotation(Template->GetNotation(), MenuDesc, NodeCategory, ToolTip);
+		check(NodeSpawner != nullptr);
+		ActionRegistrar.AddBlueprintAction(ActionKey, NodeSpawner);
+	};
+
+	// Add all rig units
+	for(const FRigVMFunction& Function : Registry.GetFunctions())
+	{
+		UScriptStruct* Struct = Function.Struct;
+		if (Struct == nullptr || !Struct->IsChildOf(FRigVMStruct::StaticStruct()))
+		{
+			continue;
+		}
+
+		if(!Function.SupportsExecuteContextStruct(BlueprintCDO->GetRigVMExecuteContextStruct()))
+		{
+			continue;
+		}
+
+		// skip rig units which have a template
+		if (Function.GetTemplate())
+		{
+			continue;
+		}
+
+		// skip deprecated units
+		if(Function.Struct->HasMetaData(FRigVMStruct::DeprecatedMetaName))
+		{
+			continue;
+		}
+
+		FString CategoryMetadata, DisplayNameMetadata, MenuDescSuffixMetadata;
+		Struct->GetStringMetaDataHierarchical(FRigVMStruct::CategoryMetaName, &CategoryMetadata);
+		Struct->GetStringMetaDataHierarchical(FRigVMStruct::DisplayNameMetaName, &DisplayNameMetadata);
+		Struct->GetStringMetaDataHierarchical(FRigVMStruct::MenuDescSuffixMetaName, &MenuDescSuffixMetadata);
+
+		if(DisplayNameMetadata.IsEmpty())
+		{
+			DisplayNameMetadata = Struct->GetDisplayNameText().ToString();
+		}
+		if (!MenuDescSuffixMetadata.IsEmpty())
+		{
+			MenuDescSuffixMetadata = TEXT(" ") + MenuDescSuffixMetadata;
+		}
+		FText NodeCategory = FText::FromString(CategoryMetadata);
+		FText MenuDesc = FText::FromString(DisplayNameMetadata + MenuDescSuffixMetadata);
+		FText ToolTip = Struct->GetToolTipText();
+
+		UBlueprintNodeSpawner* NodeSpawner = URigVMEdGraphUnitNodeSpawner::CreateFromStruct(Struct, Function.GetMethodName(), MenuDesc, NodeCategory, ToolTip);
+		check(NodeSpawner != nullptr);
+		ActionRegistrar.AddBlueprintAction(ActionKey, NodeSpawner);
+	};
+
+	for (TObjectIterator<UEnum> EnumIt; EnumIt; ++EnumIt)
+	{
+		UEnum* EnumToConsider = (*EnumIt);
+
+		if (EnumToConsider->HasMetaData(TEXT("Hidden")))
+		{
+			continue;
+		}
+
+		if (EnumToConsider->IsEditorOnly())
+		{
+			continue;
+		}
+
+		if(EnumToConsider->IsNative())
+		{
+			continue;
+		}
+
+		FText NodeCategory = FText::FromString(TEXT("Enum"));
+		FText MenuDesc = FText::FromString(FString::Printf(TEXT("Enum %s"), *EnumToConsider->GetName()));
+		FText ToolTip = MenuDesc;
+
+		UBlueprintNodeSpawner* NodeSpawner = URigVMEdGraphEnumNodeSpawner::CreateForEnum(EnumToConsider, MenuDesc, NodeCategory, ToolTip);
+		check(NodeSpawner != nullptr);
+		ActionRegistrar.AddBlueprintAction(ActionKey, NodeSpawner);
+	}
+
+	FArrayProperty* PublicGraphFunctionsProperty = CastField<FArrayProperty>(URigVMBlueprint::StaticClass()->FindPropertyByName(TEXT("PublicGraphFunctions")));
+	FArrayProperty* PublicFunctionsProperty = CastField<FArrayProperty>(URigVMBlueprint::StaticClass()->FindPropertyByName(TEXT("PublicFunctions")));
+	if(PublicGraphFunctionsProperty || PublicFunctionsProperty)
+	{
+		// find all control rigs in the project
+		TArray<FAssetData> ControlRigAssetDatas;
+		
+		FARFilter ControlRigAssetFilter;
+
+		auto AddClassToFilter = [&ControlRigAssetFilter](const UClass* InClass)
+		{
+			const UClass* Class = InClass;
+			while(Class)
+			{
+				ControlRigAssetFilter.ClassPaths.Add(Class->GetClassPathName());
+				Class = Class->GetSuperClass();
+				if(Class == nullptr)
+				{
+					break;
+				}
+				if(Class == UBlueprint::StaticClass() ||
+					Class == UBlueprintGeneratedClass::StaticClass() ||
+					Class == UObject::StaticClass())
+				{
+					break;
+				}
+			}
+		};
+		AddClassToFilter(BlueprintCDO->GetClass());
+		AddClassToFilter(BlueprintCDO->GetRigVMBlueprintGeneratedClassPrototype());
+
+		AssetRegistryModule.Get().GetAssets(ControlRigAssetFilter, ControlRigAssetDatas);
+
+		// loop over all control rigs in the project
+		TSet<FName> PackagesProcessed;
+		for(const FAssetData& ControlRigAssetData : ControlRigAssetDatas)
+		{
+			// Avoid duplication of spawners
+			if (PackagesProcessed.Contains(ControlRigAssetData.PackageName))
+			{
+				continue;
+			}
+			PackagesProcessed.Add(ControlRigAssetData.PackageName);
+
+			FString PublicGraphFunctionsString;
+			FString PublicFunctionsString;
+			if (PublicGraphFunctionsProperty)
+			{
+				PublicGraphFunctionsString = ControlRigAssetData.GetTagValueRef<FString>(PublicGraphFunctionsProperty->GetFName());
+			}
+			// Only look at the deprecated public functions if the PublicGraphFunctionsString is empty
+			if (PublicGraphFunctionsString.IsEmpty() && PublicFunctionsProperty)
+			{
+				PublicFunctionsString = ControlRigAssetData.GetTagValueRef<FString>(PublicFunctionsProperty->GetFName());
+			}
+			
+			if(PublicFunctionsString.IsEmpty() && PublicGraphFunctionsString.IsEmpty())
+			{
+				continue;
+			}
+
+			if (PublicFunctionsProperty && !PublicFunctionsString.IsEmpty())
+			{
+				TArray<FRigVMOldPublicFunctionData> PublicFunctions;
+				PublicFunctionsProperty->ImportText_Direct(*PublicFunctionsString, &PublicFunctions, nullptr, EPropertyPortFlags::PPF_None);
+				for(const FRigVMOldPublicFunctionData& PublicFunction : PublicFunctions)
+				{
+					UBlueprintNodeSpawner* NodeSpawner = URigVMEdGraphFunctionRefNodeSpawner::CreateFromAssetData(ControlRigAssetData, PublicFunction);
+					check(NodeSpawner != nullptr);
+					ActionRegistrar.AddBlueprintAction(ActionKey, NodeSpawner);
+				}
+			}
+
+			if (PublicGraphFunctionsProperty && !PublicGraphFunctionsString.IsEmpty())
+			{
+				TArray<FRigVMGraphFunctionHeader> PublicFunctions;
+				PublicGraphFunctionsProperty->ImportText_Direct(*PublicGraphFunctionsString, &PublicFunctions, nullptr, EPropertyPortFlags::PPF_None);
+				for(const FRigVMGraphFunctionHeader& PublicFunction : PublicFunctions)
+				{
+					UBlueprintNodeSpawner* NodeSpawner = URigVMEdGraphFunctionRefNodeSpawner::CreateFromAssetData(ControlRigAssetData, PublicFunction);
+					check(NodeSpawner != nullptr);
+					ActionRegistrar.AddBlueprintAction(ActionKey, NodeSpawner);
+				}
+			}
+		}
+	}
+}
+
+void FRigVMEditorModule::GetInstanceActions(URigVMBlueprint* RigVMBlueprint, FBlueprintActionDatabaseRegistrar& ActionRegistrar)
+{
+	if (URigVMBlueprintGeneratedClass* GeneratedClass = RigVMBlueprint->GetRigVMBlueprintGeneratedClass())
+	{
+		if (URigVMHost* CDO = Cast<URigVMHost>(GeneratedClass->GetDefaultObject()))
+		{
+			static const FText NodeCategory = LOCTEXT("Variables", "Variables");
+
+			TArray<FRigVMExternalVariable> ExternalVariables = CDO->GetExternalVariables();
+			for (const FRigVMExternalVariable& ExternalVariable : ExternalVariables)
+			{
+				FText MenuDesc = FText::FromName(ExternalVariable.Name);
+				FText ToolTip = FText::FromString(FString::Printf(TEXT("Get the value of variable %s"), *ExternalVariable.Name.ToString()));
+				ActionRegistrar.AddBlueprintAction(GeneratedClass, URigVMEdGraphVariableNodeSpawner::CreateFromExternalVariable(RigVMBlueprint, ExternalVariable, true, MenuDesc, NodeCategory, ToolTip));
+
+				ToolTip = FText::FromString(FString::Printf(TEXT("Set the value of variable %s"), *ExternalVariable.Name.ToString()));
+				ActionRegistrar.AddBlueprintAction(GeneratedClass, URigVMEdGraphVariableNodeSpawner::CreateFromExternalVariable(RigVMBlueprint, ExternalVariable, false, MenuDesc, NodeCategory, ToolTip));
+			}
+		}
+
+		if (URigVMFunctionLibrary* LocalFunctionLibrary = RigVMBlueprint->GetLocalFunctionLibrary())
+		{
+			TArray<URigVMLibraryNode*> Functions = LocalFunctionLibrary->GetFunctions();
+			for (URigVMLibraryNode* Function : Functions)
+			{
+				UBlueprintNodeSpawner* NodeSpawner = URigVMEdGraphFunctionRefNodeSpawner::CreateFromFunction(Function);
+				check(NodeSpawner != nullptr);
+				ActionRegistrar.AddBlueprintAction(GeneratedClass, NodeSpawner);
+			}
+
+			static const FText NodeCategory = LOCTEXT("LocalVariables", "Local Variables");
+			for (URigVMLibraryNode* Function : Functions)
+			{
+				for (const FRigVMGraphVariableDescription& LocalVariable : Function->GetContainedGraph()->GetLocalVariables())
+				{
+					FText MenuDesc = FText::FromName(LocalVariable.Name);
+					FText ToolTip = FText::FromString(FString::Printf(TEXT("Get the value of variable %s"), *LocalVariable.Name.ToString()));
+					ActionRegistrar.AddBlueprintAction(GeneratedClass, URigVMEdGraphVariableNodeSpawner::CreateFromLocalVariable(RigVMBlueprint, Function->GetContainedGraph(), LocalVariable, true, MenuDesc, NodeCategory, ToolTip));
+
+					ToolTip = FText::FromString(FString::Printf(TEXT("Set the value of variable %s"), *LocalVariable.Name.ToString()));
+					ActionRegistrar.AddBlueprintAction(GeneratedClass, URigVMEdGraphVariableNodeSpawner::CreateFromLocalVariable(RigVMBlueprint, Function->GetContainedGraph(), LocalVariable, false, MenuDesc, NodeCategory, ToolTip));
+				}
+			}
+		}
+
+		for (URigVMGraph* Graph : RigVMBlueprint->GetAllModels())
+		{
+			if (Graph->GetEntryNode())
+			{
+				static const FText NodeCategory = LOCTEXT("InputArguments", "Input Arguments");
+				for (const FRigVMGraphVariableDescription& InputArgument : Graph->GetInputArguments())
+				{
+					FText MenuDesc = FText::FromName(InputArgument.Name);
+					FText ToolTip = FText::FromString(FString::Printf(TEXT("Get the value of input %s"), *InputArgument.Name.ToString()));
+					ActionRegistrar.AddBlueprintAction(GeneratedClass, URigVMEdGraphVariableNodeSpawner::CreateFromLocalVariable(RigVMBlueprint, Graph, InputArgument, true, MenuDesc, NodeCategory, ToolTip));
+				}			
+			}
+		}
+
+		const TArray<FName> EntryNames = RigVMBlueprint->GetRigVMClient()->GetEntryNames();
+		if(!EntryNames.IsEmpty())
+		{
+			static const FText NodeCategory = LOCTEXT("Events", "Events");
+			for (const FName& EntryName : EntryNames)
+			{
+				static constexpr TCHAR EventStr[] = TEXT("Event");
+				static const FString EventSuffix = FString::Printf(TEXT(" %s"), EventStr);
+				FString Suffix = EntryName.ToString().EndsWith(EventStr) ? FString() : EventSuffix;
+				FText MenuDesc = FText::FromString(FString::Printf(TEXT("Run %s%s"), *EntryName.ToString(), *Suffix));
+				FText ToolTip = FText::FromString(FString::Printf(TEXT("Runs the %s%s"), *EntryName.ToString(), *Suffix));
+				ActionRegistrar.AddBlueprintAction(GeneratedClass, URigVMEdGraphInvokeEntryNodeSpawner::CreateForEntry(RigVMBlueprint, EntryName, MenuDesc, NodeCategory, ToolTip));
+			}
+		}
+	}
+}
+
+void FRigVMEditorModule::GetNodeContextMenuActions(URigVMBlueprint* RigVMBlueprint, const URigVMEdGraphNode* EdGraphNode, URigVMNode* ModelNode, UToolMenu* Menu) const
+{
+	GetNodeWorkflowContextMenuActions(RigVMBlueprint, EdGraphNode, ModelNode, Menu);
+	GetNodeEventsContextMenuActions(RigVMBlueprint, EdGraphNode, ModelNode, Menu);
+	GetNodeConversionContextMenuActions(RigVMBlueprint, EdGraphNode, ModelNode, Menu);
+	GetNodeDebugContextMenuActions(RigVMBlueprint, EdGraphNode, ModelNode, Menu);
+	GetNodeVariablesContextMenuActions(RigVMBlueprint, EdGraphNode, ModelNode, Menu);
+	GetNodeTemplatesContextMenuActions(RigVMBlueprint, EdGraphNode, ModelNode, Menu);
+	GetNodeOrganizationContextMenuActions(RigVMBlueprint, EdGraphNode, ModelNode, Menu);
+	GetNodeVersioningContextMenuActions(RigVMBlueprint, EdGraphNode, ModelNode, Menu);
+	GetNodeTestContextMenuActions(RigVMBlueprint, EdGraphNode, ModelNode, Menu);
+}
+
+void FRigVMEditorModule::GetNodeWorkflowContextMenuActions(URigVMBlueprint* RigVMBlueprint, const URigVMEdGraphNode* EdGraphNode, URigVMNode* ModelNode, UToolMenu* Menu) const
+{
+	const TArray<FRigVMUserWorkflow> Workflows = ModelNode->GetSupportedWorkflows(ERigVMUserWorkflowType::NodeContext, ModelNode);
+	if(!Workflows.IsEmpty())
+	{
+		FToolMenuSection& SettingsSection = Menu->AddSection("RigVMEditorContextMenuWorkflow", LOCTEXT("WorkflowHeader", "Workflow"));
+		URigVMController* Controller = RigVMBlueprint->GetController(ModelNode->GetGraph());
+
+		for(const FRigVMUserWorkflow& Workflow : Workflows)
+		{
+			SettingsSection.AddMenuEntry(
+				*Workflow.GetTitle(),
+				FText::FromString(Workflow.GetTitle()),
+				FText::FromString(Workflow.GetTooltip()),
+				FSlateIcon(),
+				FUIAction(FExecuteAction::CreateLambda([this, Controller, Workflow, ModelNode]()
+				{
+					URigVMUserWorkflowOptions* Options = Controller->MakeOptionsForWorkflow(ModelNode, Workflow);
+
+					bool bPerform = true;
+					if(Options->RequiresDialog())
+					{
+						bPerform = ShowWorkflowOptionsDialog(Options);
+					}
+					if(bPerform)
+					{
+						Controller->PerformUserWorkflow(Workflow, Options);
+					}
+				}))
+			);
+		}
+	}
+}
+
+void FRigVMEditorModule::GetNodeEventsContextMenuActions(URigVMBlueprint* RigVMBlueprint, const URigVMEdGraphNode* EdGraphNode, URigVMNode* ModelNode, UToolMenu* Menu) const
+{
+	if(ModelNode->IsEvent())
+	{
+		const FName& EventName = ModelNode->GetEventName();
+		const bool bCanRunOnce = !CastChecked<URigVMEdGraphSchema>(EdGraphNode->GetSchema())->IsRigVMDefaultEvent(EventName);
+
+		FToolMenuSection& EventsSection = Menu->AddSection("RigVMEditorContextMenuEvents", LOCTEXT("EventsHeader", "Events"));
+
+		EventsSection.AddMenuEntry(
+			"Switch to Event",
+			LOCTEXT("SwitchToEvent", "Switch to Event"),
+			LOCTEXT("SwitchToEvent_Tooltip", "Switches the Control Rig Editor to run this event permanently."),
+			FSlateIcon(),
+			FUIAction(FExecuteAction::CreateLambda([RigVMBlueprint, EventName]()
+			{
+				if(URigVMHost* Host = Cast<URigVMHost>(RigVMBlueprint->GetObjectBeingDebugged()))
+				{
+					Host->SetEventQueue({EventName});
+				}
+			}))
+		);
+
+		EventsSection.AddMenuEntry(
+			"Run Event Once",
+			LOCTEXT("RuntEventOnce", "Run Event Once"),
+			LOCTEXT("RuntEventOnce_Tooltip", "Runs the event once (for testing)"),
+			FSlateIcon(),
+			FUIAction(FExecuteAction::CreateLambda([RigVMBlueprint, EventName]()
+			{
+				if(URigVMHost* Host = Cast<URigVMHost>(RigVMBlueprint->GetObjectBeingDebugged()))
+				{
+					const TArray<FName> PreviousEventQueue = Host->GetEventQueue();
+					TArray<FName> NewEventQueue = PreviousEventQueue;
+					NewEventQueue.Add(EventName);
+					Host->SetEventQueue(NewEventQueue);
+
+					TSharedPtr<FDelegateHandle> RunOnceHandle = MakeShareable(new FDelegateHandle);
+					*(RunOnceHandle.Get()) = Host->OnExecuted_AnyThread().AddLambda(
+						[RunOnceHandle, EventName, PreviousEventQueue](URigVMHost* InRig, const FName& InEventName)
+						{
+							if(InEventName == EventName)
+							{
+								InRig->SetEventQueue(PreviousEventQueue);
+								InRig->OnExecuted_AnyThread().Remove(*RunOnceHandle.Get());
+							}
+						}
+					);
+				}
+			}),
+			FCanExecuteAction::CreateLambda([bCanRunOnce](){ return bCanRunOnce; }))
+		);
+	}
+}
+
+void FRigVMEditorModule::GetNodeConversionContextMenuActions(URigVMBlueprint* RigVMBlueprint, const URigVMEdGraphNode* EdGraphNode, URigVMNode* ModelNode, UToolMenu* Menu) const
+{
+	if(URigVMDispatchNode* DispatchNode = Cast<URigVMDispatchNode>(ModelNode))
+	{
+		if(DispatchNode->GetFactory()->GetFactoryName() == FRigVMDispatch_Constant().GetFactoryName())
+		{
+			if(const URigVMPin* ValuePin = DispatchNode->FindPin(TEXT("Value")))
+			{
+				// if the value pin has only links on the root pin
+				if((ValuePin->GetSourceLinks(false).Num() > 0) && (ValuePin->GetTargetLinks(false).Num() > 0))
+				{
+					URigVMController* Controller = RigVMBlueprint->GetController(ModelNode->GetGraph());
+
+					FToolMenuSection& ConversionSection = Menu->AddSection("RigVMEditorContextMenuConversion", LOCTEXT("ConversionHeader", "Conversion"));
+					ConversionSection.AddMenuEntry(
+						"Convert Constant to Reroute",
+						LOCTEXT("ConvertConstantToReroute", "Convert Constant to Reroute"),
+						LOCTEXT("ConvertConstantToReroute_Tooltip", "Converts the Constant node to a to Reroute node and sustains the value"),
+						FSlateIcon(),
+						FUIAction(FExecuteAction::CreateLambda([Controller, ValuePin, DispatchNode]()
+						{
+							Controller->OpenUndoBracket(TEXT("Replace Constant with Reroute"));
+							TArray<URigVMController::FLinkedPath> LinkedPaths = Controller->GetLinkedPaths(DispatchNode);
+							Controller->FastBreakLinkedPaths(LinkedPaths, true);
+							const FVector2D Position = DispatchNode->GetPosition();
+							const FString NodeName = DispatchNode->GetName();
+							const FString CPPType = ValuePin->GetCPPType();
+							FName CPPTypeObjectPath = NAME_None;
+							if(const UObject* CPPTypeObject = ValuePin->GetCPPTypeObject())
+							{
+								CPPTypeObjectPath = *CPPTypeObject->GetPathName();
+							}
+							
+							Controller->RemoveNode(DispatchNode, true, true);
+							Controller->AddFreeRerouteNode(CPPType, CPPTypeObjectPath, false, NAME_None, FString(), Position, NodeName, true);
+							Controller->RestoreLinkedPaths(LinkedPaths, URigVMController::FRestoreLinkedPathSettings(), true);
+							Controller->CloseUndoBracket();
+						}
+					)));
+				}
+			}
+		}
+	}
+}
+
+void FRigVMEditorModule::GetNodeDebugContextMenuActions(URigVMBlueprint* RigVMBlueprint, const URigVMEdGraphNode* EdGraphNode, URigVMNode* ModelNode, UToolMenu* Menu) const
+{
+	FToolMenuSection& DebugSection = Menu->AddSection("RigVMEditorContextMenuDebug", LOCTEXT("DebugHeader", "Debug"));
+	bool bNoneHasBreakpoint = true;
+
+	const URigVMGraph* Model = ModelNode->GetGraph();
+	URigVMController* Controller = RigVMBlueprint->GetController(Model);
+
+	TArray<URigVMNode*> SelectedNodes;
+	TArray<FName> SelectedNodeNames = Model->GetSelectNodes();
+	SelectedNodeNames.AddUnique(ModelNode->GetFName());
+
+	for (FName SelectedNodeName : SelectedNodeNames)
+	{
+		if (URigVMNode* FoundNode = Model->FindNodeByName(SelectedNodeName))
+		{
+			SelectedNodes.Add(FoundNode);
+			if (FoundNode->HasBreakpoint())
+			{
+				bNoneHasBreakpoint = false;
+			}
+		}
+	}
+	
+	if (bNoneHasBreakpoint)
+	{
+		DebugSection.AddMenuEntry(
+		"Add Breakpoint",
+		LOCTEXT("AddBreakpoint", "Add Breakpoint"),
+		LOCTEXT("AddBreakpoint_Tooltip", "Adds a breakpoint to the graph at this node"),
+		FSlateIcon(),
+		FUIAction(FExecuteAction::CreateLambda([Controller, SelectedNodes, RigVMBlueprint]()
+		{
+			for (URigVMNode* SelectedNode : SelectedNodes)
+			{
+				if (RigVMBlueprint->AddBreakpoint(SelectedNode))
+				{
+					SelectedNode->SetHasBreakpoint(true);
+				}
+			}
+		})));
+	}
+	else
+	{						
+		DebugSection.AddMenuEntry(
+		"Remove Breakpoint",
+		LOCTEXT("RemoveBreakpoint", "Remove Breakpoint"),
+		LOCTEXT("RemoveBreakpoint_Tooltip", "Removes a breakpoint to the graph at this node"),
+		FSlateIcon(),
+		FUIAction(FExecuteAction::CreateLambda([Controller, SelectedNodes, RigVMBlueprint]()
+		{
+			for (URigVMNode* SelectedNode : SelectedNodes)
+			{
+				if (SelectedNode->HasBreakpoint())
+				{
+					if (RigVMBlueprint->RemoveBreakpoint(SelectedNode))
+					{
+						SelectedNode->SetHasBreakpoint(false);
+					}
+				}
+			}                            
+		})));
+	}
+}
+
+void FRigVMEditorModule::GetNodeVariablesContextMenuActions(URigVMBlueprint* RigVMBlueprint, const URigVMEdGraphNode* EdGraphNode, URigVMNode* ModelNode, UToolMenu* Menu) const
+{
+	const URigVMGraph* Model = ModelNode->GetGraph();
+	URigVMController* Controller = RigVMBlueprint->GetController(Model);
+
+	if (URigVMVariableNode* VariableNode = Cast<URigVMVariableNode>(EdGraphNode->GetModelNode()))
+	{
+		FToolMenuSection& VariablesSection = Menu->AddSection("RigVMEditorContextMenuVariables", LOCTEXT("VariablesSettingsHeader", "Variables"));
+		VariablesSection.AddMenuEntry(
+			"MakePindingsFromVariableNode",
+			LOCTEXT("MakeBindingsFromVariableNode", "Make Bindings From Node"),
+			LOCTEXT("MakeBindingsFromVariableNode_Tooltip", "Turns the variable node into one ore more variable bindings on the pin(s)"),
+			FSlateIcon(),
+			FUIAction(FExecuteAction::CreateLambda([RigVMBlueprint, Controller, VariableNode]() {
+				Controller->MakeBindingsFromVariableNode(VariableNode->GetFName());
+			})
+		));
+	}
+
+	if (URigVMFunctionReferenceNode* FunctionReferenceNode = Cast<URigVMFunctionReferenceNode>(EdGraphNode->GetModelNode()))
+	{
+		TSoftObjectPtr<URigVMFunctionReferenceNode> RefPtr(FunctionReferenceNode);
+		if(RefPtr.GetLongPackageName() != FunctionReferenceNode->GetReferencedFunctionHeader().LibraryPointer.LibraryNode.GetLongPackageName())
+		{
+			if(!FunctionReferenceNode->IsFullyRemapped() && FunctionReferenceNode->GetReferencedFunctionHeader().LibraryPointer.LibraryNode.ResolveObject())
+			{
+				FToolMenuSection& VariablesSection = Menu->AddSection("RigVMEditorContextMenuVariables", LOCTEXT("Variables", "Variables"));
+				VariablesSection.AddMenuEntry(
+                    "MakeVariablesFromFunctionReferenceNode",
+                    LOCTEXT("MakeVariablesFromFunctionReferenceNode", "Create required variables"),
+                    LOCTEXT("MakeVariablesFromFunctionReferenceNode_Tooltip", "Creates all required variables for this function and binds them"),
+                    FSlateIcon(),
+                    FUIAction(FExecuteAction::CreateLambda([Controller, FunctionReferenceNode, RigVMBlueprint]() {
+
+                        const TArray<FRigVMExternalVariable> ExternalVariables = FunctionReferenceNode->GetExternalVariables(false);
+                        if(!ExternalVariables.IsEmpty())
+                        {
+							FScopedTransaction Transaction(LOCTEXT("MakeVariablesFromFunctionReferenceNode", "Create required variables"));
+                            RigVMBlueprint->Modify();
+
+                            if (URigVMLibraryNode* LibraryNode = FunctionReferenceNode->LoadReferencedNode())
+                            {
+                                URigVMBlueprint* ReferencedBlueprint = LibraryNode->GetTypedOuter<URigVMBlueprint>();
+							   // ReferencedBlueprint != RigVMBlueprint - since only FunctionReferenceNodes from other assets have the potential to be unmapped
+                            
+							   for(const FRigVMExternalVariable& ExternalVariable : ExternalVariables)
+							   {
+								   FString DefaultValue;
+								   if(ReferencedBlueprint)
+								   {
+									   for(const FBPVariableDescription& NewVariable : ReferencedBlueprint->NewVariables)
+									   {
+										   if(NewVariable.VarName == ExternalVariable.Name)
+										   {
+											   DefaultValue = NewVariable.DefaultValue;
+											   break;
+										   }
+									   }
+								   }
+                                
+								   FName NewVariableName = RigVMBlueprint->AddHostMemberVariableFromExternal(ExternalVariable, DefaultValue);
+								   if(!NewVariableName.IsNone())
+								   {
+									   Controller->SetRemappedVariable(FunctionReferenceNode, ExternalVariable.Name, NewVariableName);
+								   }
+							   }
+                            }
+
+                            FBlueprintEditorUtils::MarkBlueprintAsModified(RigVMBlueprint);
+                        }
+                        
+                    })
+                ));
+			}
+		}
+	}
+}
+
+void FRigVMEditorModule::GetNodeTemplatesContextMenuActions(URigVMBlueprint* RigVMBlueprint, const URigVMEdGraphNode* EdGraphNode, URigVMNode* ModelNode, UToolMenu* Menu) const
+{
+	if (URigVMTemplateNode* TemplateNode = Cast<URigVMTemplateNode>(EdGraphNode->GetModelNode()))
+	{
+		if (!TemplateNode->IsSingleton())
+		{
+			const URigVMGraph* Model = ModelNode->GetGraph();
+			URigVMController* Controller = RigVMBlueprint->GetController(Model);
+			
+			FToolMenuSection& TemplatesSection = Menu->AddSection("RigVMEditorContextMenuTemplates", LOCTEXT("TemplatesHeader", "Templates"));
+			TemplatesSection.AddMenuEntry(
+				"Unresolve Template Node",
+				LOCTEXT("UnresolveTemplateNode", "Unresolve Template Node"),
+				LOCTEXT("UnresolveTemplateNode_Tooltip", "Removes any type information from the template node"),
+				FSlateIcon(),
+				FUIAction(FExecuteAction::CreateLambda([Controller, Model]() {
+					const TArray<FName> Nodes = Model->GetSelectNodes();
+					Controller->UnresolveTemplateNodes(Nodes, true, true);
+				})
+			));
+		}
+	}
+}
+
+void FRigVMEditorModule::GetNodeOrganizationContextMenuActions(URigVMBlueprint* RigVMBlueprint, const URigVMEdGraphNode* EdGraphNode, URigVMNode* ModelNode, UToolMenu* Menu) const
+{
+	const URigVMGraph* Model = ModelNode->GetGraph();
+	URigVMController* Controller = RigVMBlueprint->GetController(Model);
+
+	FToolMenuSection& OrganizationSection = Menu->AddSection("RigVMEditorContextMenuOrganization", LOCTEXT("OrganizationHeader", "Organization"));
+	OrganizationSection.AddMenuEntry(
+		"Collapse Nodes",
+		LOCTEXT("CollapseNodes", "Collapse Nodes"),
+		LOCTEXT("CollapseNodes_Tooltip", "Turns the selected nodes into a single Collapse node"),
+		FSlateIcon(),
+		FUIAction(FExecuteAction::CreateLambda([Model, Controller]() {
+			TArray<FName> Nodes = Model->GetSelectNodes();
+			Controller->CollapseNodes(Nodes, FString(), true, true);
+		})
+	));
+	OrganizationSection.AddMenuEntry(
+		"Collapse to Function",
+		LOCTEXT("CollapseNodesToFunction", "Collapse to Function"),
+		LOCTEXT("CollapseNodesToFunction_Tooltip", "Turns the selected nodes into a new Function"),
+		FSlateIcon(),
+		FUIAction(FExecuteAction::CreateLambda([Model, Controller]() {
+			TArray<FName> Nodes = Model->GetSelectNodes();
+			Controller->OpenUndoBracket(TEXT("Collapse to Function"));
+			URigVMCollapseNode* CollapseNode = Controller->CollapseNodes(Nodes, TEXT("New Function"), true, true);
+			if(CollapseNode)
+			{
+				if (Controller->PromoteCollapseNodeToFunctionReferenceNode(CollapseNode->GetFName(), true, true).IsNone())
+				{
+					Controller->CancelUndoBracket();
+				}
+				else
+				{
+					Controller->CloseUndoBracket();
+				}
+			}
+			else
+			{
+				Controller->CancelUndoBracket();
+			}
+		})
+	));
+
+	if (URigVMCollapseNode* CollapseNode = Cast<URigVMCollapseNode>(EdGraphNode->GetModelNode()))
+	{
+		OrganizationSection.AddMenuEntry(
+			"Promote To Function",
+			LOCTEXT("PromoteToFunction", "Promote To Function"),
+			LOCTEXT("PromoteToFunction_Tooltip", "Turns the Collapse Node into a Function"),
+			FSlateIcon(),
+			FUIAction(FExecuteAction::CreateLambda([Controller, CollapseNode]() {
+				Controller->PromoteCollapseNodeToFunctionReferenceNode(CollapseNode->GetFName(), true, true);
+			})
+		));
+	}
+
+	if (URigVMFunctionReferenceNode* FunctionReferenceNode = Cast<URigVMFunctionReferenceNode>(EdGraphNode->GetModelNode()))
+	{
+		TSoftObjectPtr<URigVMFunctionReferenceNode> RefPtr(FunctionReferenceNode);
+		if(RefPtr.GetLongPackageName() != FunctionReferenceNode->GetReferencedFunctionHeader().LibraryPointer.LibraryNode.GetLongPackageName())
+		{
+			OrganizationSection.AddMenuEntry(
+			   "Localize Function",
+			   LOCTEXT("LocalizeFunction", "Localize Function"),
+			   LOCTEXT("LocalizeFunction_Tooltip", "Creates a local copy of the function backing the node."),
+			   FSlateIcon(),
+			   FUIAction(FExecuteAction::CreateLambda([RigVMBlueprint, FunctionReferenceNode]() {
+				   RigVMBlueprint->BroadcastRequestLocalizeFunctionDialog(FunctionReferenceNode->GetFunctionIdentifier(), true);
+			   })
+		    ));
+		}
+	}
+
+	if (URigVMFunctionReferenceNode* FunctionRefNode = Cast<URigVMFunctionReferenceNode>(EdGraphNode->GetModelNode()))
+	{
+		OrganizationSection.AddMenuEntry(
+			"Promote To Collapse Node",
+			LOCTEXT("PromoteToCollapseNode", "Promote To Collapse Node"),
+			LOCTEXT("PromoteToCollapseNode_Tooltip", "Turns the Function Ref Node into a Collapse Node"),
+			FSlateIcon(),
+			FUIAction(FExecuteAction::CreateLambda([Controller, FunctionRefNode]() {
+				Controller->PromoteFunctionReferenceNodeToCollapseNode(FunctionRefNode->GetFName());
+				})
+			));
+	}
+
+	if (URigVMLibraryNode* LibraryNode = Cast<URigVMLibraryNode>(EdGraphNode->GetModelNode()))
+	{
+		OrganizationSection.AddMenuEntry(
+			"Expand Node",
+			LOCTEXT("ExpandNode", "Expand Node"),
+			LOCTEXT("ExpandNode_Tooltip", "Expands the contents of the node into this graph"),
+			FSlateIcon(),
+			FUIAction(FExecuteAction::CreateLambda([Controller, LibraryNode]() {
+				Controller->OpenUndoBracket(TEXT("Expand node"));
+				TArray<URigVMNode*> ExpandedNodes = Controller->ExpandLibraryNode(LibraryNode->GetFName(), true, true);
+				if (ExpandedNodes.Num() > 0)
+				{
+					TArray<FName> ExpandedNodeNames;
+					for (URigVMNode* ExpandedNode : ExpandedNodes)
+					{
+						ExpandedNodeNames.Add(ExpandedNode->GetFName());
+					}
+					Controller->SetNodeSelection(ExpandedNodeNames);
+				}
+				Controller->CloseUndoBracket();
+			})
+		));
+	}
+
+	OrganizationSection.AddSubMenu("Alignment", LOCTEXT("AlignmentHeader", "Alignment"), FText(), FNewToolMenuDelegate::CreateLambda([](UToolMenu* AlignmentMenu)
+	{
+		{
+			FToolMenuSection& InSection = AlignmentMenu->AddSection("RigVMEditorContextMenuAlignment", LOCTEXT("AlignHeader", "Align"));
+			InSection.AddMenuEntry(FGraphEditorCommands::Get().AlignNodesTop);
+			InSection.AddMenuEntry(FGraphEditorCommands::Get().AlignNodesMiddle);
+			InSection.AddMenuEntry(FGraphEditorCommands::Get().AlignNodesBottom);
+			InSection.AddMenuEntry(FGraphEditorCommands::Get().AlignNodesLeft);
+			InSection.AddMenuEntry(FGraphEditorCommands::Get().AlignNodesCenter);
+			InSection.AddMenuEntry(FGraphEditorCommands::Get().AlignNodesRight);
+			InSection.AddMenuEntry(FGraphEditorCommands::Get().StraightenConnections);
+		}
+
+		{
+			FToolMenuSection& InSection = AlignmentMenu->AddSection("RigVMEditorContextMenuDistribution", LOCTEXT("DistributionHeader", "Distribution"));
+			InSection.AddMenuEntry(FGraphEditorCommands::Get().DistributeNodesHorizontally);
+			InSection.AddMenuEntry(FGraphEditorCommands::Get().DistributeNodesVertically);
+		}
+	}));
+}
+
+void FRigVMEditorModule::GetNodeVersioningContextMenuActions(URigVMBlueprint* RigVMBlueprint, const URigVMEdGraphNode* EdGraphNode, URigVMNode* ModelNode, UToolMenu* Menu) const
+{
+	const URigVMGraph* Model = ModelNode->GetGraph();
+	URigVMController* Controller = RigVMBlueprint->GetController(Model);
+
+	bool bCanNodeBeUpgraded = false;
+	TArray<FName> SelectedNodeNames = Model->GetSelectNodes();
+	SelectedNodeNames.AddUnique(ModelNode->GetFName());
+
+	for(const FName& SelectedNodeName : SelectedNodeNames)
+	{
+		if (URigVMNode* FoundNode = Model->FindNodeByName(SelectedNodeName))
+		{
+			bCanNodeBeUpgraded = bCanNodeBeUpgraded || FoundNode->CanBeUpgraded();
+		}
+	}
+	
+	if(bCanNodeBeUpgraded)
+	{
+		FToolMenuSection& VersioningSection = Menu->AddSection("RigVMEditorContextMenuVersioning", LOCTEXT("VersioningHeader", "Versioning"));
+		VersioningSection.AddMenuEntry(
+			"Upgrade Nodes",
+			LOCTEXT("UpgradeNodes", "Upgrade Nodes"),
+			LOCTEXT("UpgradeNodes_Tooltip", "Upgrades deprecated nodes to their current implementation"),
+			FSlateIcon(),
+			FUIAction(FExecuteAction::CreateLambda([Model, Controller]() {
+				TArray<FName> Nodes = Model->GetSelectNodes();
+				Controller->UpgradeNodes(Nodes, true, true);
+			})
+		));
+	}
+}
+
+void FRigVMEditorModule::GetNodeTestContextMenuActions(URigVMBlueprint* RigVMBlueprint, const URigVMEdGraphNode* EdGraphNode, URigVMNode* ModelNode, UToolMenu* Menu) const
+{
+	// this struct is only available in EngineTest for now
+	static const FString DecoratorObjectPath = TEXT("/Script/EngineTestEditor.EngineTestRigVM_SimpleDecorator");
+	const UScriptStruct* SimpleDecoratorStruct = Cast<UScriptStruct>(RigVMTypeUtils::FindObjectFromCPPTypeObjectPath(DecoratorObjectPath));
+	if(SimpleDecoratorStruct == nullptr)
+	{
+		return;
+	}
+
+	const URigVMGraph* Model = ModelNode->GetGraph();
+	URigVMController* Controller = RigVMBlueprint->GetController(Model);
+
+	FToolMenuSection& EngineTestSection = Menu->AddSection("RigVMEditorContextMenuEngineTest", LOCTEXT("EngineTestHeader", "EngineTest"));
+	EngineTestSection.AddMenuEntry(
+		"Add simple decorator",
+		LOCTEXT("AddSimpleDecorator", "Add simple decorator"),
+		LOCTEXT("AddSimpleDecorator_Tooltip", "Adds a simple test decorator to the node"),
+		FSlateIcon(),
+		FUIAction(FExecuteAction::CreateLambda([Controller, ModelNode, SimpleDecoratorStruct]()
+		{
+			(void)Controller->AddDecorator(
+				ModelNode->GetFName(),
+				*SimpleDecoratorStruct->GetPathName(),
+				TEXT("Decorator"),
+				FString(), INDEX_NONE, true, true);
+		}))
+	);
+}
+
+void FRigVMEditorModule::GetPinContextMenuActions(URigVMBlueprint* RigVMBlueprint, const UEdGraphPin* EdGraphPin, URigVMPin* ModelPin, UToolMenu* Menu) const
+{
+	GetPinWorkflowContextMenuActions(RigVMBlueprint, EdGraphPin, ModelPin, Menu);
+	GetPinDebugContextMenuActions(RigVMBlueprint, EdGraphPin, ModelPin, Menu);
+	GetPinArrayContextMenuActions(RigVMBlueprint, EdGraphPin, ModelPin, Menu);
+	GetPinAggregateContextMenuActions(RigVMBlueprint, EdGraphPin, ModelPin, Menu);
+	GetPinTemplateContextMenuActions(RigVMBlueprint, EdGraphPin, ModelPin, Menu);
+	GetPinConversionContextMenuActions(RigVMBlueprint, EdGraphPin, ModelPin, Menu);
+	GetPinVariableContextMenuActions(RigVMBlueprint, EdGraphPin, ModelPin, Menu);
+	GetPinResetDefaultContextMenuActions(RigVMBlueprint, EdGraphPin, ModelPin, Menu);
+	GetPinInjectedNodesContextMenuActions(RigVMBlueprint, EdGraphPin, ModelPin, Menu);
+}
+
+void FRigVMEditorModule::GetPinWorkflowContextMenuActions(URigVMBlueprint* RigVMBlueprint, const UEdGraphPin* EdGraphPin, URigVMPin* ModelPin, UToolMenu* Menu) const
+{
+	const TArray<FRigVMUserWorkflow> Workflows = ModelPin->GetNode()->GetSupportedWorkflows(ERigVMUserWorkflowType::PinContext, ModelPin);
+	if(!Workflows.IsEmpty())
+	{
+		URigVMController* Controller = RigVMBlueprint->GetController(ModelPin->GetGraph());
+
+		FToolMenuSection& SettingsSection = Menu->AddSection("RigVMEditorContextMenuWorkflow", LOCTEXT("WorkflowHeader", "Workflow"));
+
+		for(const FRigVMUserWorkflow& Workflow : Workflows)
+		{
+			SettingsSection.AddMenuEntry(
+				*Workflow.GetTitle(),
+				FText::FromString(Workflow.GetTitle()),
+				FText::FromString(Workflow.GetTooltip()),
+				FSlateIcon(),
+				FUIAction(FExecuteAction::CreateLambda([this, Controller, Workflow, ModelPin]()
+				{
+					URigVMUserWorkflowOptions* Options = Controller->MakeOptionsForWorkflow(ModelPin, Workflow);
+
+					bool bPerform = true;
+					if(Options->RequiresDialog())
+					{
+						bPerform = ShowWorkflowOptionsDialog(Options);
+					}
+					if(bPerform)
+					{
+						Controller->PerformUserWorkflow(Workflow, Options);
+					}
+				}))
+			);
+		}
+	}
+}
+
+void FRigVMEditorModule::GetPinDebugContextMenuActions(URigVMBlueprint* RigVMBlueprint, const UEdGraphPin* EdGraphPin, URigVMPin* ModelPin, UToolMenu* Menu) const
+{
+	const bool bIsEditablePin = !ModelPin->IsExecuteContext() && !ModelPin->IsWildCard();
+
+	if(Cast<URigVMEdGraphNode>(EdGraphPin->GetOwningNode()))
+	{
+		if(bIsEditablePin)
+		{
+			// Add the watch pin / unwatch pin menu items
+			FToolMenuSection& Section = Menu->AddSection("RigVMEditorContextMenuWatches", LOCTEXT("WatchesHeader", "Watches"));
+			if (FKismetDebugUtilities::IsPinBeingWatched(RigVMBlueprint, EdGraphPin))
+			{
+				Section.AddMenuEntry(FGraphEditorCommands::Get().StopWatchingPin);
+			}
+			else
+			{
+				Section.AddMenuEntry(FGraphEditorCommands::Get().StartWatchingPin);
+			}
+		}
+	}
+}
+
+void FRigVMEditorModule::GetPinArrayContextMenuActions(URigVMBlueprint* RigVMBlueprint, const UEdGraphPin* EdGraphPin, URigVMPin* ModelPin, UToolMenu* Menu) const
+{
+	URigVMController* Controller = RigVMBlueprint->GetController(ModelPin->GetGraph());
+	if (ModelPin->IsArray() && !ModelPin->IsExecuteContext())
+	{
+		FToolMenuSection& Section = Menu->AddSection("RigVMEditorContextMenuPinArrays", LOCTEXT("PinArrays", "Arrays"));
+		Section.AddMenuEntry(
+			"ClearPinArray",
+			LOCTEXT("ClearPinArray", "Clear Array"),
+			LOCTEXT("ClearPinArray_Tooltip", "Removes all elements of the array."),
+			FSlateIcon(),
+			FUIAction(FExecuteAction::CreateLambda([Controller, ModelPin]() {
+				Controller->ClearArrayPin(ModelPin->GetPinPath());
+			})
+		));
+	}
+	
+	if(ModelPin->IsArrayElement())
+	{
+		FToolMenuSection& Section = Menu->AddSection("RigVMEditorContextMenuPinArrays", LOCTEXT("PinArrays", "Arrays"));
+		Section.AddMenuEntry(
+			"RemoveArrayPin",
+			LOCTEXT("RemoveArrayPin", "Remove Array Element"),
+			LOCTEXT("RemoveArrayPin_Tooltip", "Removes the selected element from the array"),
+			FSlateIcon(),
+			FUIAction(FExecuteAction::CreateLambda([Controller, ModelPin]() {
+				Controller->RemoveArrayPin(ModelPin->GetPinPath(), true, true);
+			})
+		));
+		Section.AddMenuEntry(
+			"DuplicateArrayPin",
+			LOCTEXT("DuplicateArrayPin", "Duplicate Array Element"),
+			LOCTEXT("DuplicateArrayPin_Tooltip", "Duplicates the selected element"),
+			FSlateIcon(),
+			FUIAction(FExecuteAction::CreateLambda([Controller, ModelPin]() {
+				Controller->DuplicateArrayPin(ModelPin->GetPinPath(), true, true);
+			})
+		));
+	}
+}
+
+void FRigVMEditorModule::GetPinAggregateContextMenuActions(URigVMBlueprint* RigVMBlueprint, const UEdGraphPin* EdGraphPin, URigVMPin* ModelPin, UToolMenu* Menu) const
+{
+	URigVMController* Controller = RigVMBlueprint->GetController(ModelPin->GetGraph());
+	if (Cast<URigVMAggregateNode>(ModelPin->GetNode()))
+	{
+		FToolMenuSection& Section = Menu->AddSection("RigVMEditorContextMenuAggregatePin", LOCTEXT("AggregatePin", "Aggregates"));
+		Section.AddMenuEntry(
+			"RemoveAggregatePin",
+			LOCTEXT("RemoveAggregatePin", "Remove Aggregate Element"),
+			LOCTEXT("RemoveAggregatePin_Tooltip", "Removes the selected element from the aggregate"),
+			FSlateIcon(),
+			FUIAction(FExecuteAction::CreateLambda([Controller, ModelPin]() {
+				Controller->RemoveAggregatePin(ModelPin->GetPinPath(), true, true);
+			})
+		));
+	}
+}
+
+void FRigVMEditorModule::GetPinTemplateContextMenuActions(URigVMBlueprint* RigVMBlueprint, const UEdGraphPin* EdGraphPin, URigVMPin* ModelPin, UToolMenu* Menu) const
+{
+	if (URigVMTemplateNode* TemplateNode = Cast<URigVMTemplateNode>(ModelPin->GetNode()))
+	{
+		if (!TemplateNode->IsSingleton())
+		{
+			URigVMController* Controller = RigVMBlueprint->GetController(ModelPin->GetGraph());
+			FToolMenuSection& TemplatesSection = Menu->AddSection("RigVMEditorContextMenuTemplates", LOCTEXT("TemplatesHeader", "Templates"));
+
+			if(const FRigVMTemplate* Template = TemplateNode->GetTemplate())
+			{
+				if(!ModelPin->IsExecuteContext())
+				{
+					if(const FRigVMTemplateArgument* Argument = Template->FindArgument(ModelPin->GetRootPin()->GetFName()))
+					{
+						if(!Argument->IsSingleton())
+						{
+							TArray<TRigVMTypeIndex> ResolvedTypeIndices = Argument->GetSupportedTypeIndices(TemplateNode->GetResolvedPermutationIndices(true));
+							TSharedRef<SRigVMGraphChangePinType> ChangePinTypeWidget =
+							SNew(SRigVMGraphChangePinType)
+							.Blueprint(RigVMBlueprint)
+							.Types(ResolvedTypeIndices)
+							.OnTypeSelected_Lambda([RigVMBlueprint, ModelPin](const TRigVMTypeIndex& TypeSelected)
+							{
+								if (URigVMController* Controller = RigVMBlueprint->GetController(ModelPin->GetGraph()))
+								{
+									Controller->ResolveWildCardPin(ModelPin, TypeSelected, true, true);
+								}
+							});
+
+							TemplatesSection.AddEntry(FToolMenuEntry::InitWidget("ChangePinTypeWidget", ChangePinTypeWidget, FText(), true));
+						}
+					}
+				}
+							
+				TemplatesSection.AddMenuEntry(
+					"Unresolve Template Node",
+					LOCTEXT("UnresolveTemplateNode", "Unresolve Template Node"),
+					LOCTEXT("UnresolveTemplateNode_Tooltip", "Removes any type information from the template node"),
+					FSlateIcon(),
+					FUIAction(FExecuteAction::CreateLambda([Controller, ModelPin]() {
+						const TArray<FName> Nodes = ModelPin->GetGraph()->GetSelectNodes();
+						Controller->UnresolveTemplateNodes(Nodes, true, true);
+					})
+				));
+			}
+		}
+	}
+}
+
+void FRigVMEditorModule::GetPinConversionContextMenuActions(URigVMBlueprint* RigVMBlueprint, const UEdGraphPin* EdGraphPin, URigVMPin* ModelPin, UToolMenu* Menu) const
+{
+	if (URigVMRerouteNode* RerouteNode = Cast<URigVMRerouteNode>(ModelPin->GetNode()))
+	{
+		if(RerouteNode->GetLinkedSourceNodes().Num() == 0)
+		{
+			if(const URigVMPin* ValuePin = RerouteNode->FindPin(TEXT("Value")))
+			{
+				if(!ValuePin->IsExecuteContext())
+				{
+					URigVMController* Controller = RigVMBlueprint->GetController(ModelPin->GetGraph());
+					FToolMenuSection& ConversionSection = Menu->AddSection("RigVMEditorContextMenuConversion", LOCTEXT("ConversionHeader", "Conversion"));
+
+					if(ValuePin->IsArray())
+					{
+						ConversionSection.AddMenuEntry(
+							"Convert Reroute to Make Array",
+							LOCTEXT("ConvertReroutetoMakeArray", "Convert Reroute to Make Array"),
+							LOCTEXT("ConvertReroutetoMakeArray_Tooltip", "Converts the Reroute node to a to Make Array node and sustains the value"),
+							FSlateIcon(),
+							FUIAction(FExecuteAction::CreateLambda([Controller, RerouteNode]()
+							{
+								Controller->ConvertRerouteNodeToDispatch(RerouteNode, FRigVMDispatch_ArrayMake().GetTemplateNotation(), true, true);
+							}
+						)));
+					}
+					else if(ValuePin->IsStruct())
+					{
+						ConversionSection.AddMenuEntry(
+							"Convert Reroute to Make Struct",
+							LOCTEXT("ConvertReroutetoMakeStruct", "Convert Reroute to Make Struct"),
+							LOCTEXT("ConvertReroutetoMakeStruct_Tooltip", "Converts the Reroute node to a to Make Struct node and sustains the value"),
+							FSlateIcon(),
+							FUIAction(FExecuteAction::CreateLambda([Controller, RerouteNode]()
+							{
+								Controller->ConvertRerouteNodeToDispatch(RerouteNode, FRigVMDispatch_MakeStruct().GetTemplateNotation(), true, true);
+							}
+						)));
+					}
+					else
+					{
+						ConversionSection.AddMenuEntry(
+							"Convert Reroute to Constant",
+							LOCTEXT("ConvertReroutetoConstant", "Convert Reroute to Constant"),
+							LOCTEXT("ConvertReroutetoConstant_Tooltip", "Converts the Reroute node to a to Constant node and sustains the value"),
+							FSlateIcon(),
+							FUIAction(FExecuteAction::CreateLambda([Controller, RerouteNode]()
+							{
+								Controller->ConvertRerouteNodeToDispatch(RerouteNode, FRigVMDispatch_Constant().GetTemplateNotation(), true, true);
+							}
+						)));
+					}
+				}
+			}
+		}
+	}
+}
+
+void FRigVMEditorModule::GetPinVariableContextMenuActions(URigVMBlueprint* RigVMBlueprint, const UEdGraphPin* EdGraphPin, URigVMPin* ModelPin, UToolMenu* Menu) const
+{
+	const bool bIsEditablePin = !ModelPin->IsExecuteContext() && !ModelPin->IsWildCard();
+
+	if (ModelPin->GetDirection() == ERigVMPinDirection::Input && bIsEditablePin)
+	{
+		const UEdGraphNode* EdGraphNode = EdGraphPin->GetOwningNode();
+		URigVMController* Controller = RigVMBlueprint->GetController(ModelPin->GetGraph());
+
+		if (ModelPin->IsBoundToVariable())
+		{
+			FVector2D NodePosition = FVector2D(EdGraphNode->NodePosX - 200.f, EdGraphNode->NodePosY);
+
+			FToolMenuSection& VariablesSection = Menu->AddSection("RigVMEditorContextMenuVariables", LOCTEXT("Variables", "Variables"));
+			VariablesSection.AddMenuEntry(
+				"MakeVariableNodeFromBinding",
+				LOCTEXT("MakeVariableNodeFromBinding", "Make Variable Node"),
+				LOCTEXT("MakeVariableNodeFromBinding_Tooltip", "Turns the variable binding on the pin to a variable node"),
+				FSlateIcon(),
+				FUIAction(FExecuteAction::CreateLambda([Controller, ModelPin, NodePosition]() {
+					Controller->MakeVariableNodeFromBinding(ModelPin->GetPinPath(), NodePosition, true, true);
+				})
+			));
+		}
+		else
+		{
+			FVector2D NodePosition = FVector2D(EdGraphNode->NodePosX - 200.f, EdGraphNode->NodePosY);
+
+			FToolMenuSection& VariablesSection = Menu->AddSection("RigVMEditorContextMenuVariables", LOCTEXT("Variables", "Variables"));
+			VariablesSection.AddMenuEntry(
+				"PromotePinToVariable",
+				LOCTEXT("PromotePinToVariable", "Promote Pin To Variable"),
+				LOCTEXT("PromotePinToVariable_Tooltip", "Turns the pin into a variable"),
+				FSlateIcon(),
+				FUIAction(FExecuteAction::CreateLambda([Controller, ModelPin, NodePosition]() {
+
+					FModifierKeysState KeyState = FSlateApplication::Get().GetModifierKeys();
+					bool bCreateVariableNode = !KeyState.IsAltDown();
+
+					Controller->PromotePinToVariable(ModelPin->GetPinPath(), bCreateVariableNode, NodePosition, true, true);
+				})
+			));
+		}
+	}
+
+	if (Cast<URigVMUnitNode>(ModelPin->GetNode()) != nullptr || 
+		Cast<URigVMDispatchNode>(ModelPin->GetNode()) != nullptr || 
+		Cast<URigVMLibraryNode>(ModelPin->GetNode()) != nullptr)
+	{
+		if (ModelPin->GetDirection() == ERigVMPinDirection::Input &&
+			ModelPin->IsRootPin() &&
+			bIsEditablePin)
+		{
+			if (!ModelPin->IsBoundToVariable())
+			{
+				FToolMenuSection& VariablesSection = Menu->FindOrAddSection(TEXT("Variables"));
+
+				TSharedRef<SRigVMGraphVariableBinding> VariableBindingWidget =
+					SNew(SRigVMGraphVariableBinding)
+					.Blueprint(RigVMBlueprint)
+					.ModelPins({ModelPin})
+					.CanRemoveBinding(false);
+
+				VariablesSection.AddEntry(FToolMenuEntry::InitWidget("BindPinToVariableWidget", VariableBindingWidget, FText(), true));
+			}
+		}
+	}
+}
+
+void FRigVMEditorModule::GetPinResetDefaultContextMenuActions(URigVMBlueprint* RigVMBlueprint, const UEdGraphPin* EdGraphPin, URigVMPin* ModelPin, UToolMenu* Menu) const
+{
+	if (Cast<URigVMUnitNode>(ModelPin->GetNode()) != nullptr || 
+		Cast<URigVMDispatchNode>(ModelPin->GetNode()) != nullptr || 
+		Cast<URigVMLibraryNode>(ModelPin->GetNode()) != nullptr)
+	{
+		const bool bIsEditablePin = !ModelPin->IsExecuteContext() && !ModelPin->IsWildCard();
+		
+		if (ModelPin->GetDirection() == ERigVMPinDirection::Input &&
+			ModelPin->IsRootPin() &&
+			bIsEditablePin)
+		{
+			FToolMenuSection& Section = Menu->AddSection("RigVMEditorContextMenuPinDefaults", LOCTEXT("PinDefaults", "Pin Defaults"));
+			URigVMController* Controller = RigVMBlueprint->GetController(ModelPin->GetGraph());
+
+			Section.AddMenuEntry(
+				"ResetPinDefaultValue",
+				LOCTEXT("ResetPinDefaultValue", "Reset Pin Value"),
+				LOCTEXT("ResetPinDefaultValue_Tooltip", "Resets the pin's value to its default."),
+				FSlateIcon(),
+				FUIAction(FExecuteAction::CreateLambda([Controller, ModelPin]() {
+					Controller->ResetPinDefaultValue(ModelPin->GetPinPath());
+				})
+			));
+		}
+	}
+}
+
+void FRigVMEditorModule::GetPinInjectedNodesContextMenuActions(URigVMBlueprint* RigVMBlueprint, const UEdGraphPin* EdGraphPin, URigVMPin* ModelPin, UToolMenu* Menu) const
+{
+	if (ModelPin->GetRootPin() == ModelPin && (
+	Cast<URigVMUnitNode>(ModelPin->GetNode()) != nullptr ||
+	Cast<URigVMLibraryNode>(ModelPin->GetNode()) != nullptr))
+	{
+		URigVMController* Controller = RigVMBlueprint->GetController(ModelPin->GetGraph());
+
+		if (ModelPin->HasInjectedNodes())
+		{
+			FToolMenuSection& Section = Menu->AddSection("RigVMEditorContextMenuNodeEjectionInterp", LOCTEXT("NodeEjectionInterp", "Eject"));
+
+			Section.AddMenuEntry(
+				"EjectLastNode",
+				LOCTEXT("EjectLastNode", "Eject Last Node"),
+				LOCTEXT("EjectLastNode_Tooltip", "Eject the last injected node"),
+				FSlateIcon(),
+				FUIAction(FExecuteAction::CreateLambda([Controller, ModelPin]() {
+					Controller->OpenUndoBracket(TEXT("Eject node from pin"));
+					URigVMNode* Node = Controller->EjectNodeFromPin(ModelPin->GetPinPath(), true, true);
+					Controller->SelectNode(Node, true, true, true);
+					Controller->CloseUndoBracket();
+				})
+			));
+		}
+
+		if (ModelPin->GetCPPType() == TEXT("float") ||
+			ModelPin->GetCPPType() == TEXT("double") ||
+			ModelPin->GetCPPType() == TEXT("FVector"))
+		{
+			FToolMenuSection& Section = Menu->AddSection("RigVMEditorContextMenuNodeInjectionInterp", LOCTEXT("NodeInjectionInterp", "Interpolate"));
+			URigVMNode* InterpNode = nullptr;
+			bool bBoundToVariable = false;
+			for (URigVMInjectionInfo* Injection : ModelPin->GetInjectedNodes())
+			{
+				FString TemplateName;
+				if (URigVMUnitNode* UnitNode = Cast<URigVMUnitNode>(Injection->Node))
+				{
+					if (UnitNode->GetScriptStruct()->GetStringMetaDataHierarchical(FRigVMStruct::TemplateNameMetaName, &TemplateName))
+					{
+						if (TemplateName == TEXT("AlphaInterp"))
+						{
+							InterpNode = Injection->Node;
+							break;
+						}
+					}
+				}
+				else if (URigVMVariableNode* VariableNode = Cast<URigVMVariableNode>(Injection->Node))
+				{
+					bBoundToVariable = true;
+					break;
+				}
+			}
+
+			if(!bBoundToVariable)
+			{
+				if (InterpNode == nullptr)
+				{
+					UScriptStruct* ScriptStruct = nullptr;
+
+					if ((ModelPin->GetCPPType() == TEXT("float")) || (ModelPin->GetCPPType() == TEXT("double")))
+					{
+						ScriptStruct = FRigVMFunction_AlphaInterp::StaticStruct();
+					}
+					else if (ModelPin->GetCPPType() == TEXT("FVector"))
+					{
+						ScriptStruct = FRigVMFunction_AlphaInterpVector::StaticStruct();
+					}
+					else
+					{
+						checkNoEntry();
+					}
+
+					Section.AddMenuEntry(
+						"AddAlphaInterp",
+						LOCTEXT("AddAlphaInterp", "Add Interpolate"),
+						LOCTEXT("AddAlphaInterp_Tooltip", "Injects an interpolate node"),
+						FSlateIcon(),
+						FUIAction(FExecuteAction::CreateLambda([Controller, EdGraphPin, ModelPin, ScriptStruct]() {
+							Controller->OpenUndoBracket(TEXT("Add injected node"));
+							URigVMInjectionInfo* Injection = Controller->AddInjectedNode(ModelPin->GetPinPath(), ModelPin->GetDirection() != ERigVMPinDirection::Output, ScriptStruct, FRigVMStruct::ExecuteName, TEXT("Value"), TEXT("Result"), FString(), true, true);
+							if (Injection)
+							{
+								TArray<FName> NodeNames;
+								NodeNames.Add(Injection->Node->GetFName());
+								Controller->SetNodeSelection(NodeNames);
+							}
+							Controller->CloseUndoBracket();
+						})
+					));
+				}
+				else
+				{
+					Section.AddMenuEntry(
+						"EditAlphaInterp",
+						LOCTEXT("EditAlphaInterp", "Edit Interpolate"),
+						LOCTEXT("EditAlphaInterp_Tooltip", "Edit the interpolate node"),
+						FSlateIcon(),
+						FUIAction(FExecuteAction::CreateLambda([RigVMBlueprint, InterpNode]() {
+						TArray<FName> NodeNames;
+						NodeNames.Add(InterpNode->GetFName());
+						RigVMBlueprint->GetController(InterpNode->GetGraph())->SetNodeSelection(NodeNames);
+					})
+						));
+					Section.AddMenuEntry(
+						"RemoveAlphaInterp",
+						LOCTEXT("RemoveAlphaInterp", "Remove Interpolate"),
+						LOCTEXT("RemoveAlphaInterp_Tooltip", "Removes the interpolate node"),
+						FSlateIcon(),
+						FUIAction(FExecuteAction::CreateLambda([Controller, EdGraphPin, ModelPin, InterpNode]() {
+							Controller->RemoveInjectedNode(ModelPin->GetPinPath(), ModelPin->GetDirection() != ERigVMPinDirection::Output, true);
+						})
+					));
+				}
+			}
+		}
+
+		if (ModelPin->GetCPPType() == TEXT("FVector") ||
+			ModelPin->GetCPPType() == TEXT("FQuat") ||
+			ModelPin->GetCPPType() == TEXT("FTransform"))
+		{
+			FToolMenuSection& Section = Menu->AddSection("RigVMEditorContextMenuNodeInjectionVisualDebug", LOCTEXT("NodeInjectionVisualDebug", "Visual Debug"));
+
+			URigVMNode* VisualDebugNode = nullptr;
+			bool bBoundToVariable = false;
+			for (URigVMInjectionInfo* Injection : ModelPin->GetInjectedNodes())
+			{
+				FString TemplateName;
+				if (URigVMUnitNode* UnitNode = Cast<URigVMUnitNode>(Injection->Node))
+				{
+					if (UnitNode->GetScriptStruct()->GetStringMetaDataHierarchical(FRigVMStruct::TemplateNameMetaName, &TemplateName))
+					{
+						if (TemplateName == TEXT("VisualDebug"))
+						{
+							VisualDebugNode = Injection->Node;
+							break;
+						}
+					}
+				}
+				else if (URigVMVariableNode* VariableNode = Cast<URigVMVariableNode>(Injection->Node))
+				{
+					bBoundToVariable = true;
+					break;
+				}
+			}
+
+			if (!bBoundToVariable)
+			{
+				if (VisualDebugNode == nullptr)
+				{
+					UScriptStruct* ScriptStruct = nullptr;
+
+					if (ModelPin->GetCPPType() == TEXT("FVector"))
+					{
+						ScriptStruct = FRigVMFunction_VisualDebugVectorNoSpace::StaticStruct();
+					}
+					else if (ModelPin->GetCPPType() == TEXT("FQuat"))
+					{
+						ScriptStruct = FRigVMFunction_VisualDebugQuatNoSpace::StaticStruct();
+					}
+					else if (ModelPin->GetCPPType() == TEXT("FTransform"))
+					{
+						ScriptStruct = FRigVMFunction_VisualDebugTransformNoSpace::StaticStruct();
+					}
+					else
+					{
+						checkNoEntry();
+					}
+
+					Section.AddMenuEntry(
+						"AddVisualDebug",
+						LOCTEXT("AddVisualDebug", "Add Visual Debug"),
+						LOCTEXT("AddVisualDebug_Tooltip", "Injects a visual debugging node"),
+						FSlateIcon(),
+						FUIAction(FExecuteAction::CreateLambda([RigVMBlueprint, Controller, EdGraphPin, ModelPin, ScriptStruct]() {
+							URigVMInjectionInfo* Injection = Controller->AddInjectedNode(ModelPin->GetPinPath(), ModelPin->GetDirection() != ERigVMPinDirection::Output, ScriptStruct, FRigVMStruct::ExecuteName, TEXT("Value"), TEXT("Value"), FString(), true, true);
+							if (Injection)
+							{
+								TArray<FName> NodeNames;
+								NodeNames.Add(Injection->Node->GetFName());
+								Controller->SetNodeSelection(NodeNames);
+
+								/*
+								 * todo
+								if (URigVMUnitNode* UnitNode = Cast<URigVMUnitNode>(ModelPin->GetNode()))
+								{
+									if (TSharedPtr<FStructOnScope> DefaultStructScope = UnitNode->ConstructStructInstance())
+									{
+										if(DefaultStructScope->GetStruct()->IsChildOf(FRigUnit::StaticStruct()))
+										{
+											FRigUnit* DefaultStruct = (FRigUnit*)DefaultStructScope->GetStructMemory();
+
+											FString PinPath = ModelPin->GetPinPath();
+											FString Left, Right;
+
+											FRigElementKey SpaceKey;
+											if (URigVMPin::SplitPinPathAtStart(PinPath, Left, Right))
+											{
+												SpaceKey = DefaultStruct->DetermineSpaceForPin(Right, RigVMBlueprint->Hierarchy);
+											}
+
+											if (SpaceKey.IsValid())
+											{
+												if (URigVMPin* SpacePin = Injection->Node->FindPin(TEXT("Space")))
+												{
+													if(URigVMPin* SpaceTypePin = SpacePin->FindSubPin(TEXT("Type")))
+													{
+														FString SpaceTypeStr = StaticEnum<ERigElementType>()->GetDisplayNameTextByValue((int64)SpaceKey.Type).ToString();
+														Controller->SetPinDefaultValue(SpaceTypePin->GetPinPath(), SpaceTypeStr, true, true, false, true);
+													}
+													if(URigVMPin* SpaceNamePin = SpacePin->FindSubPin(TEXT("Name")))
+													{
+														Controller->SetPinDefaultValue(SpaceNamePin->GetPinPath(), SpaceKey.Name.ToString(), true, true, false, true);
+													}
+												}
+											}
+										}
+									}
+								}
+								*/
+							}
+						})
+					));
+				}
+				else
+				{
+					Section.AddMenuEntry(
+						"EditVisualDebug",
+						LOCTEXT("EditVisualDebug", "Edit Visual Debug"),
+						LOCTEXT("EditVisualDebug_Tooltip", "Edit the visual debugging node"),
+						FSlateIcon(),
+						FUIAction(FExecuteAction::CreateLambda([Controller, VisualDebugNode]() {
+							TArray<FName> NodeNames;
+							NodeNames.Add(VisualDebugNode->GetFName());
+							Controller->SetNodeSelection(NodeNames);
+						})
+					));
+					Section.AddMenuEntry(
+						"ToggleVisualDebug",
+						LOCTEXT("ToggleVisualDebug", "Toggle Visual Debug"),
+						LOCTEXT("ToggleVisualDebug_Tooltip", "Toggle the visibility the visual debugging"),
+						FSlateIcon(),
+						FUIAction(FExecuteAction::CreateLambda([Controller, VisualDebugNode]() {
+							URigVMPin* EnabledPin = VisualDebugNode->FindPin(TEXT("bEnabled"));
+							check(EnabledPin);
+							Controller->SetPinDefaultValue(EnabledPin->GetPinPath(), EnabledPin->GetDefaultValue() == TEXT("True") ? TEXT("False") : TEXT("True"), false, true, false, true);
+						})
+					));
+					Section.AddMenuEntry(
+						"RemoveVisualDebug",
+						LOCTEXT("RemoveVisualDebug", "Remove Visual Debug"),
+						LOCTEXT("RemoveVisualDebug_Tooltip", "Removes the visual debugging node"),
+						FSlateIcon(),
+							FUIAction(FExecuteAction::CreateLambda([Controller, EdGraphPin, ModelPin, VisualDebugNode]() {
+							Controller->RemoveNodeByName(VisualDebugNode->GetFName(), true, false);
+						})
+					));
+				}
+			}
+		}
+	}
+}
+
+void FRigVMEditorModule::GetContextMenuActions(const URigVMEdGraphSchema* Schema, UToolMenu* Menu, UGraphNodeContextMenuContext* Context) const
+{
+	if (Menu && Context)
+	{
+		Schema->UEdGraphSchema::GetContextMenuActions(Menu, Context);
+
+		if (const UEdGraphPin* InGraphPin = (UEdGraphPin* )Context->Pin)
+		{
+			if(URigVMBlueprint* RigVMBlueprint = Cast<URigVMBlueprint>((UBlueprint*)Context->Blueprint))
+			{
+				const UEdGraph* Graph = InGraphPin->GetOwningNode()->GetGraph();
+				if (URigVMPin* ModelPin = RigVMBlueprint->GetModel(Graph)->FindPin(InGraphPin->GetName()))
+				{
+					GetPinContextMenuActions(RigVMBlueprint, InGraphPin, ModelPin, Menu);
+				}
+			}
+		}
+		else if(const URigVMEdGraphNode* EdGraphNode = Cast<URigVMEdGraphNode>(Context->Node))
+		{
+			if (URigVMBlueprint* RigVMBlueprint = Cast<URigVMBlueprint>((UBlueprint*)Context->Blueprint))
+			{
+				if(URigVMNode* ModelNode = EdGraphNode->GetModelNode())
+				{
+					GetNodeContextMenuActions(RigVMBlueprint, EdGraphNode, ModelNode, Menu);
+				}
+			}
+		}
+	}
+}
+
+void FRigVMEditorModule::PreChange(const UUserDefinedStruct* Changed, FStructureEditorUtils::EStructureEditorChangeInfo ChangedType)
+{
+	if(!IsRigVMEditorModuleBase())
+	{
+		return;
+	}
+	
+	// the following is similar to
+	// FUserDefinedStructureCompilerInner::ReplaceStructWithTempDuplicate()
+	// it is necessary since existing rigs need to be kept valid until after PreBPCompile
+	// there are other systems, such as sequencer, that might need to evaluate the rig
+	// for one last time during PreBPCompile
+	// Overall sequence of events
+	// PreStructChange --1--> PostStructChange
+	//                              --2--> PreBPCompile --3--> PostBPCompile
+	
+	UUserDefinedStruct* StructureToReinstance = (UUserDefinedStruct*)Changed;
+
+	FUserDefinedStructureCompilerUtils::ReplaceStructWithTempDuplicateByPredicate(
+		StructureToReinstance,
+		[](FStructProperty* InStructProperty)
+		{
+			// make sure variable properties on the BP is patched
+			// since active rig instance still references it
+			if (URigVMBlueprintGeneratedClass* BPClass = Cast<URigVMBlueprintGeneratedClass>(InStructProperty->GetOwnerClass()))
+			{
+				if (BPClass->ClassGeneratedBy->IsA<URigVMBlueprint>())
+				{
+					return true;
+				}
+			}
+			// similar story, VM instructions reference properties on the GeneratorClass
+			else if ((InStructProperty->GetOwnerStruct())->IsA<URigVMMemoryStorageGeneratorClass>())
+			{
+				return true;
+			}
+			else if (URigVMDetailsViewWrapperObject::IsValidClass(InStructProperty->GetOwnerClass()))
+			{
+				return true;
+			}
+			
+			return false;
+		},
+		[](UStruct* InStruct)
+		{
+			// refresh these since VM caching references them
+			if (URigVMMemoryStorageGeneratorClass* GeneratorClass = Cast<URigVMMemoryStorageGeneratorClass>(InStruct))
+			{
+				GeneratorClass->RefreshLinkedProperties();
+				GeneratorClass->RefreshPropertyPaths();	
+			}
+			else if (InStruct->IsChildOf(URigVMDetailsViewWrapperObject::StaticClass()))
+			{
+				URigVMDetailsViewWrapperObject::MarkOutdatedClass(Cast<UClass>(InStruct));
+			}
+		});
+	
+	// in the future we could only invalidate caches on affected rig instances, it shouldn't make too much of a difference though
+	for (TObjectIterator<URigVMHost> It(RF_Transient | RF_ClassDefaultObject, /** bIncludeDerivedClasses */ true, /** InternalExcludeFlags */ EInternalObjectFlags::Garbage); It; ++It)
+	{
+		URigVMHost* Host = *It;
+		// rebuild property list and property path list
+		Host->InvalidateCachedMemory();
+	}
+}
+
+void FRigVMEditorModule::PostChange(const UUserDefinedStruct* Changed, FStructureEditorUtils::EStructureEditorChangeInfo ChangedType)
+{
+	if(!IsRigVMEditorModuleBase())
+	{
+		return;
+	}
+
+	TArray<URigVMBlueprint*> BlueprintsToRefresh;
+	for (TObjectIterator<URigVMPin> It(RF_Transient | RF_ClassDefaultObject, /** bIncludeDerivedClasses */ true, /** InternalExcludeFlags */ EInternalObjectFlags::Garbage); It; ++It)
+	{
+		const URigVMPin* Pin = *It;
+		// GetCPPTypeObject also makes sure the pin's type information is update to date
+		if (Pin && Pin->GetCPPTypeObject() == Changed)
+		{
+			if (URigVMBlueprint* RigVMBlueprint = Pin->GetTypedOuter<URigVMBlueprint>())
+			{
+				BlueprintsToRefresh.AddUnique(RigVMBlueprint);
+				
+				// this pin is part of a function definition
+				// update all BP that uses this function
+				if (Pin->GetGraph() == RigVMBlueprint->GetLocalFunctionLibrary())
+				{
+					TArray< TSoftObjectPtr<URigVMFunctionReferenceNode> > References;
+					References = RigVMBlueprint->RigVMClient.GetOrCreateFunctionLibrary(false)->GetReferencesForFunction(Pin->GetNode()->GetFName());
+					
+					for (const TSoftObjectPtr<URigVMFunctionReferenceNode>& Reference : References)
+					{
+						const URigVMFunctionReferenceNode* RefNode = Reference.LoadSynchronous();
+						if (!RefNode)
+						{
+							continue;
+						}
+						
+						if (URigVMBlueprint* FunctionUserBlueprint = RefNode->GetTypedOuter<URigVMBlueprint>())
+						{
+							BlueprintsToRefresh.AddUnique(FunctionUserBlueprint);
+						}
+					}	
+				}
+			}
+		}
+	}
+
+	for (URigVMBlueprint* RigVMBlueprint : BlueprintsToRefresh)
+	{
+		RigVMBlueprint->OnRigVMRegistryChanged();
+		(void)RigVMBlueprint->MarkPackageDirty();
+	}
+	
+	for (URigVMBlueprint* RigVMBlueprint : BlueprintsToRefresh)
+	{
+		// this should make sure variables in BP are updated with the latest struct object
+		// otherwise RigVMCompiler validation would complain about variable type - pin type mismatch
+		FCompilerResultsLog	ResultsLog;
+		FKismetEditorUtilities::CompileBlueprint(RigVMBlueprint, EBlueprintCompileOptions::None, &ResultsLog);
+		
+		// BP compiler always initialize the new CDO by copying from the old CDO,
+		// however, in case that a BP variable type has changed, the data old CDO would be invalid because
+		// while the old memory container still references the temp duplicated struct we created during PreChange()
+		// registers that reference the BP variable would be referencing the new struct as a result of
+		// FKismetCompilerContext::CompileClassLayout, so type mismatch would invalidate relevant copy operations
+		// so to simplify things, here we just reset all rigs upon error
+		if (ResultsLog.NumErrors > 0)
+		{
+			URigVMBlueprintGeneratedClass* RigClass = RigVMBlueprint->GetRigVMBlueprintGeneratedClass();
+			URigVMHost* CDO = Cast<URigVMHost>(RigClass->GetDefaultObject(true /* create if needed */));
+			if (CDO->GetVM() != nullptr)
+			{
+				CDO->GetVM()->Reset();
+			}
+			TArray<UObject*> ArchetypeInstances;
+			CDO->GetArchetypeInstances(ArchetypeInstances);
+			for (UObject* Instance : ArchetypeInstances)
+			{
+				if (URigVMHost* InstanceHost = Cast<URigVMHost>(Instance))
+				{
+					InstanceHost->GetVM()->Reset();
+				}
+			}
+		}
+	}
+}
+
+TArray<FRigVMEditorModule::FRigVMEditorToolbarExtender>& FRigVMEditorModule::GetAllRigVMEditorToolbarExtenders()
+{
+	return RigVMEditorToolbarExtenders;
+}
+
+void FRigVMEditorModule::CreateRootGraphIfRequired(URigVMBlueprint* InBlueprint) const
+{
+	if(InBlueprint == nullptr)
+	{
+		return;
+	}
+
+	const URigVMBlueprint* BlueprintCDO = GetRigVMBlueprintCDO();
+	UClass* EdGraphClass = BlueprintCDO->GetRigVMEdGraphClass();
+
+	for(const UEdGraph* EdGraph : InBlueprint->UbergraphPages)
+	{
+		if(EdGraph->IsA(EdGraphClass))
+		{
+			return;
+		}
+	}
+
+	UClass* EdGraphSchemaClass = BlueprintCDO->GetRigVMEdGraphSchemaClass();
+	const URigVMEdGraphSchema* SchemaCDO = CastChecked<URigVMEdGraphSchema>(EdGraphSchemaClass->GetDefaultObject());
+	
+	// add an initial graph for us to work in
+	UEdGraph* ControlRigGraph = FBlueprintEditorUtils::CreateNewGraph(InBlueprint, SchemaCDO->GetRootGraphName(), EdGraphClass, EdGraphSchemaClass);
+	ControlRigGraph->bAllowDeletion = false;
+	FBlueprintEditorUtils::AddUbergraphPage(InBlueprint, ControlRigGraph);
+	InBlueprint->LastEditedDocuments.AddUnique(ControlRigGraph);
+	InBlueprint->PostLoad();
+}
+
+void FRigVMEditorModule::HandleNewBlueprintCreated(UBlueprint* InBlueprint)
+{
+	CreateRootGraphIfRequired(Cast<URigVMBlueprint>(InBlueprint));
+}
+
+bool FRigVMEditorModule::ShowWorkflowOptionsDialog(URigVMUserWorkflowOptions* InOptions) const
+{
+	FPropertyEditorModule& EditModule = FModuleManager::Get().GetModuleChecked<FPropertyEditorModule>("PropertyEditor");
+
+	FDetailsViewArgs DetailsViewArgs;
+	DetailsViewArgs.NameAreaSettings = FDetailsViewArgs::HideNameArea;
+	DetailsViewArgs.bHideSelectionTip = true;
+	DetailsViewArgs.DefaultsOnlyVisibility = EEditDefaultsOnlyNodeVisibility::Hide;
+	DetailsViewArgs.bAllowSearch = false;
+
+	const TSharedRef<class IDetailsView> PropertyView = EditModule.CreateDetailView( DetailsViewArgs );
+	PropertyView->SetObject(InOptions);
+
+	TSharedRef<SCustomDialog> OptionsDialog = SNew(SCustomDialog)
+		.Title(FText(LOCTEXT("ControlRigWorkflowOptions", "Options")))
+		.Content()
+		[
+			PropertyView
+		]
+		.Buttons({
+			SCustomDialog::FButton(LOCTEXT("OK", "OK")),
+			SCustomDialog::FButton(LOCTEXT("Cancel", "Cancel"))
+	});
+	return OptionsDialog->ShowModal() == 0;
+}
+
+FConnectionDrawingPolicy* FRigVMEditorModule::CreateConnectionDrawingPolicy(int32 InBackLayerID, int32 InFrontLayerID, float InZoomFactor, const FSlateRect& InClippingRect, class FSlateWindowElementList& InDrawElements, class UEdGraph* InGraphObj) const
+{
+	return new FRigVMEdGraphConnectionDrawingPolicy(InBackLayerID, InFrontLayerID, InZoomFactor, InClippingRect, InDrawElements, InGraphObj);
+}
+
+bool FRigVMEditorModule::IsRigVMEditorModuleBase() const
+{
+	return GetRigVMBlueprintClass() == URigVMBlueprint::StaticClass();
+}
+
+#undef LOCTEXT_NAMESPACE
diff --git a/Engine/Plugins/Runtime/RigVM/Source/RigVMEditor/Private/Widgets/SRigVMGraphFunctionLocalizationWidget.cpp b/Engine/Plugins/Runtime/RigVM/Source/RigVMEditor/Private/Widgets/SRigVMGraphFunctionLocalizationWidget.cpp
--- a/Engine/Plugins/Runtime/RigVM/Source/RigVMEditor/Private/Widgets/SRigVMGraphFunctionLocalizationWidget.cpp
+++ b/Engine/Plugins/Runtime/RigVM/Source/RigVMEditor/Private/Widgets/SRigVMGraphFunctionLocalizationWidget.cpp
@@ -32,7 +32,7 @@ SRigVMGraphFunctionLocalizationItem::SRigVMGraphFunctionLocalizationItem(const F
 		DisplayText = FText::FromString(FString::Printf(TEXT("%s :: %s"), *OuterName, *FunctionName));
 		ToolTipText = bIsPublic ?
 			FText::FromString(FString::Printf(TEXT("%s :: %s is public. Check this to create a local copy."), *OuterName, *FunctionName)) :
-			FText::FromString(FString::Printf(TEXT("%s :: %s is private. A local copy has to be created. To avoid this you can turn it public in the source Control Rig."), *OuterName, *FunctionName));
+			FText::FromString(FString::Printf(TEXT("%s :: %s is private. A local copy has to be created. To avoid this you can turn it public in the source rig vm."), *OuterName, *FunctionName));
 	}
 }
 
diff --git a/Engine/Plugins/Runtime/RigVM/Source/RigVMEditor/Public/Editor/RigVMEditor.h b/Engine/Plugins/Runtime/RigVM/Source/RigVMEditor/Public/Editor/RigVMEditor.h
--- a/Engine/Plugins/Runtime/RigVM/Source/RigVMEditor/Public/Editor/RigVMEditor.h
+++ b/Engine/Plugins/Runtime/RigVM/Source/RigVMEditor/Public/Editor/RigVMEditor.h
@@ -33,7 +33,7 @@ protected:
 
 private:
 
-	/** Our currently running control rig instance */
+	/** Our currently running rig vm instance */
 	TObjectPtr<URigVMHost> Host;
 
 	FPreviewHostUpdated PreviewHostUpdated;
