RigVMEditor: Refactoring editor module, blueprint actions etc

diff --git a/Engine/Plugins/Runtime/RigVM/Source/RigVMEditor/Public/Editor/RigVMGraphDetailCustomization.h b/Engine/Plugins/Runtime/RigVM/Source/RigVMEditor/Public/Editor/RigVMGraphDetailCustomization.h
new file mode 100644
--- /dev/null
+++ b/Engine/Plugins/Runtime/RigVM/Source/RigVMEditor/Public/Editor/RigVMGraphDetailCustomization.h
@@ -0,0 +1,1043 @@
+// Copyright Epic Games, Inc. All Rights Reserved.
+
+#pragma once
+
+#include "CoreMinimal.h"
+#include "IDetailCustomization.h"
+#include "IDetailCustomNodeBuilder.h"
+#include "UObject/WeakObjectPtr.h"
+#include "EdGraph/RigVMEdGraph.h"
+#include "RigVMBlueprint.h"
+#include "Editor/RigVMEditor.h"
+#include "SGraphPin.h"
+#include "Widgets/SRigVMGraphNode.h"
+#include "Widgets/Colors/SColorBlock.h"
+#include "Editor/RigVMDetailsViewWrapperObject.h"
+#include "IDetailPropertyExtensionHandler.h"
+#include "EdGraph/RigVMEdGraphSchema.h"
+#include "SAdvancedTransformInputBox.h"
+#include "Widgets/SRigVMGraphPinNameListValueWidget.h"
+
+class IDetailLayoutBuilder;
+
+class RIGVMEDITOR_API FRigVMFunctionArgumentGroupLayout : public IDetailCustomNodeBuilder, public TSharedFromThis<FRigVMFunctionArgumentGroupLayout>
+{
+public:
+	FRigVMFunctionArgumentGroupLayout(
+		URigVMGraph* InGraph, 
+		URigVMBlueprint* InBlueprint,
+		TWeakPtr<FRigVMEditor> InEditor,
+		bool bInputs);
+	virtual ~FRigVMFunctionArgumentGroupLayout();
+
+private:
+	/** IDetailCustomNodeBuilder Interface*/
+	virtual void SetOnRebuildChildren(FSimpleDelegate InOnRegenerateChildren) override { OnRebuildChildren = InOnRegenerateChildren; }
+	virtual void GenerateHeaderRowContent(FDetailWidgetRow& NodeRow) override {}
+	virtual void GenerateChildContent(IDetailChildrenBuilder& ChildrenBuilder) override;
+	virtual void Tick(float DeltaTime) override {}
+	virtual bool RequiresTick() const override { return false; }
+	virtual FName GetName() const override { return NAME_None; }
+	virtual bool InitiallyCollapsed() const override { return false; }
+
+private:
+
+	void HandleModifiedEvent(ERigVMGraphNotifType InNotifType, URigVMGraph* InGraph, UObject* InSubject);
+
+	TWeakObjectPtr<URigVMGraph> GraphPtr;
+	TWeakObjectPtr<URigVMBlueprint> RigVMBlueprintPtr;
+	TWeakPtr<FRigVMEditor> RigVMEditorPtr;
+	bool bIsInputGroup;
+	FSimpleDelegate OnRebuildChildren;
+};
+
+class RIGVMEDITOR_API FRigVMFunctionArgumentLayout : public IDetailCustomNodeBuilder, public TSharedFromThis<FRigVMFunctionArgumentLayout>
+{
+public:
+
+	FRigVMFunctionArgumentLayout(
+		URigVMPin* InPin, 
+		URigVMGraph* InGraph, 
+		URigVMBlueprint* InBlueprint,
+		TWeakPtr<FRigVMEditor> InEditor)
+		: PinPtr(InPin)
+		, GraphPtr(InGraph)
+		, RigVMBlueprintPtr(InBlueprint)
+		, RigVMEditorPtr(InEditor)
+		, NameValidator(InBlueprint, InGraph, InPin->GetFName())
+	{}
+
+private:
+
+	/** IDetailCustomNodeBuilder Interface*/
+	virtual void GenerateHeaderRowContent(FDetailWidgetRow& NodeRow) override;
+	virtual void GenerateChildContent(IDetailChildrenBuilder& ChildrenBuilder) override;
+	virtual void Tick(float DeltaTime) override {}
+	virtual bool RequiresTick() const override { return false; }
+	virtual FName GetName() const override { return PinPtr.Get()->GetFName(); }
+	virtual bool InitiallyCollapsed() const override { return true; }
+
+private:
+
+	/** Determines if this pin should not be editable */
+	bool ShouldPinBeReadOnly(bool bIsEditingPinType = false) const;
+
+	/** Determines if editing the pins on the node should be read only */
+	bool IsPinEditingReadOnly(bool bIsEditingPinType = false) const;
+
+	/** Determines if an argument can be moved up or down */
+	bool CanArgumentBeMoved(bool bMoveUp) const;
+
+	/** Callbacks for all the functionality for modifying arguments */
+	void OnRemoveClicked();
+	FReply OnArgMoveUp();
+	FReply OnArgMoveDown();
+
+	FText OnGetArgNameText() const;
+	FText OnGetArgToolTipText() const;
+	void OnArgNameTextCommitted(const FText& NewText, ETextCommit::Type InTextCommit);
+
+	FEdGraphPinType OnGetPinInfo() const;
+	void PinInfoChanged(const FEdGraphPinType& PinType);
+	void OnPrePinInfoChange(const FEdGraphPinType& PinType);
+
+private:
+
+	/** The argument pin that this layout reflects */
+	TWeakObjectPtr<URigVMPin> PinPtr;
+	
+	/** The target graph that this argument is on */
+	TWeakObjectPtr<URigVMGraph> GraphPtr;
+
+	/** The blueprint we are editing */
+	TWeakObjectPtr<URigVMBlueprint> RigVMBlueprintPtr;
+
+	/** The editor we are editing */
+	TWeakPtr<FRigVMEditor> RigVMEditorPtr;
+
+	/** Holds a weak pointer to the argument name widget, used for error notifications */
+	TWeakPtr<SEditableTextBox> ArgumentNameWidget;
+
+	/** The validator to check if a name for an argument is valid */
+	FRigVMLocalVariableNameValidator NameValidator;
+};
+
+class RIGVMEDITOR_API FRigVMFunctionArgumentDefaultNode : public IDetailCustomNodeBuilder, public TSharedFromThis<FRigVMFunctionArgumentDefaultNode>
+{
+public:
+	FRigVMFunctionArgumentDefaultNode(
+		URigVMGraph* InGraph,
+		URigVMBlueprint* InBlueprint
+	);
+	virtual ~FRigVMFunctionArgumentDefaultNode();
+
+private:
+	/** IDetailCustomNodeBuilder Interface*/
+	virtual void SetOnRebuildChildren(FSimpleDelegate InOnRegenerateChildren) override { OnRebuildChildren = InOnRegenerateChildren; }
+	virtual void GenerateHeaderRowContent(FDetailWidgetRow& NodeRow) override {}
+	virtual void GenerateChildContent(IDetailChildrenBuilder& ChildrenBuilder) override;
+	virtual void Tick(float DeltaTime) override {}
+	virtual bool RequiresTick() const override { return false; }
+	virtual FName GetName() const override { return NAME_None; }
+	virtual bool InitiallyCollapsed() const override { return false; }
+
+private:
+
+	void OnGraphChanged(const FEdGraphEditAction& InAction);
+	void HandleModifiedEvent(ERigVMGraphNotifType InNotifType, URigVMGraph* InGraph, UObject* InSubject);
+
+	TWeakObjectPtr<URigVMGraph> GraphPtr;
+	TWeakObjectPtr<URigVMEdGraph> EdGraphOuterPtr;
+	TWeakObjectPtr<URigVMBlueprint> RigVMBlueprintPtr;
+	FSimpleDelegate OnRebuildChildren;
+	TSharedPtr<SRigVMGraphNode> OwnedNodeWidget;
+	FDelegateHandle GraphChangedDelegateHandle;
+};
+
+
+/** Customization for editing rig vm graphs */
+class RIGVMEDITOR_API FRigVMGraphDetailCustomization : public IDetailCustomization
+{
+public:
+
+	/** Makes a new instance of this detail layout class for a specific detail view requesting it */
+	static TSharedPtr<IDetailCustomization> MakeInstance(TSharedPtr<IBlueprintEditor> InBlueprintEditor, const UClass* InExpectedBlueprintClass);
+
+	FRigVMGraphDetailCustomization(TSharedPtr<FRigVMEditor> RigVMigEditor, URigVMBlueprint* RigVMBlueprint)
+		: RigVMEditorPtr(RigVMigEditor)
+		, RigVMBlueprintPtr(RigVMBlueprint)
+		, bIsPickingColor(false)
+	{}
+
+	// IDetailCustomization interface
+	virtual void CustomizeDetails(IDetailLayoutBuilder& DetailLayout) override;
+
+	bool IsAddNewInputOutputEnabled() const;
+	EVisibility GetAddNewInputOutputVisibility() const;
+	FReply OnAddNewInputClicked();
+	FReply OnAddNewOutputClicked();
+	FText GetNodeCategory() const;
+	void SetNodeCategory(const FText& InNewText, ETextCommit::Type InCommitType);
+	FText GetNodeKeywords() const;
+	void SetNodeKeywords(const FText& InNewText, ETextCommit::Type InCommitType);
+	FText GetNodeDescription() const;
+	void SetNodeDescription(const FText& InNewText, ETextCommit::Type InCommitType);
+	FLinearColor GetNodeColor() const;
+	void SetNodeColor(FLinearColor InColor, bool bSetupUndoRedo);
+	void OnNodeColorBegin();
+	void OnNodeColorEnd();
+	void OnNodeColorCancelled(FLinearColor OriginalColor);
+	FReply OnNodeColorClicked();
+	FText GetCurrentAccessSpecifierName() const;
+	void OnAccessSpecifierSelected( TSharedPtr<FString> SpecifierName, ESelectInfo::Type SelectInfo );
+	TSharedRef<ITableRow> HandleGenerateRowAccessSpecifier( TSharedPtr<FString> SpecifierName, const TSharedRef<STableViewBase>& OwnerTable );
+
+private:
+
+	/** The Blueprint editor we are embedded in */
+	TWeakPtr<FRigVMEditor> RigVMEditorPtr;
+
+	/** The graph we are editing */
+	TWeakObjectPtr<URigVMEdGraph> GraphPtr;
+
+	/** The blueprint we are editing */
+	TWeakObjectPtr<URigVMBlueprint> RigVMBlueprintPtr;
+
+	/** The color block widget */
+	TSharedPtr<SColorBlock> ColorBlock;
+
+	/** Set to true if the UI is currently picking a color */
+	bool bIsPickingColor;
+
+	static TArray<TSharedPtr<FString>> AccessSpecifierStrings;
+};
+
+/** Customization for editing a rig vm node */
+class RIGVMEDITOR_API FRigVMWrappedNodeDetailCustomization : public IDetailCustomization
+{
+public:
+	
+	FRigVMWrappedNodeDetailCustomization();
+
+	/** Makes a new instance of this detail layout class for a specific detail view requesting it */
+	static TSharedRef<IDetailCustomization> MakeInstance();
+
+	// IDetailCustomization interface
+	virtual void CustomizeDetails(IDetailLayoutBuilder& DetailLayout) override;
+
+	TSharedRef<SWidget> MakeNameListItemWidget(TSharedPtr<FString> InItem);
+	FText GetNameListText(FNameProperty* InProperty) const;
+	TSharedPtr<FString> GetCurrentlySelectedItem(FNameProperty* InProperty, const TArray<TSharedPtr<FString>>* InNameList) const;
+	void SetNameListText(const FText& NewTypeInValue, ETextCommit::Type, FNameProperty* InProperty, TSharedRef<IPropertyUtilities> PropertyUtilities);
+	void OnNameListChanged(TSharedPtr<FString> NewSelection, ESelectInfo::Type SelectInfo, FNameProperty* InProperty, TSharedRef<IPropertyUtilities> PropertyUtilities);
+	void OnNameListComboBox(FNameProperty* InProperty, const TArray<TSharedPtr<FString>>* InNameList);
+	void CustomizeLiveValues(IDetailLayoutBuilder& DetailLayout);
+
+	URigVMBlueprint* BlueprintBeingCustomized;
+	TArray<TWeakObjectPtr<URigVMDetailsViewWrapperObject>> ObjectsBeingCustomized;
+	TArray<TWeakObjectPtr<URigVMNode>> NodesBeingCustomized;
+	TMap<FName, TSharedPtr<SRigVMGraphPinNameListValueWidget>> NameListWidgets;
+};
+
+/** Customization for editing a rig vm node */
+class RIGVMEDITOR_API FRigVMGraphMathTypeDetailCustomization : public IPropertyTypeCustomization
+{
+public:
+
+	FRigVMGraphMathTypeDetailCustomization();
+	
+	static TSharedRef<IPropertyTypeCustomization> MakeInstance()
+	{
+		return MakeShareable(new FRigVMGraphMathTypeDetailCustomization);
+	}
+
+	/** IPropertyTypeCustomization interface */
+	virtual void CustomizeHeader(TSharedRef<class IPropertyHandle> InPropertyHandle, class FDetailWidgetRow& HeaderRow, IPropertyTypeCustomizationUtils& StructCustomizationUtils) override;
+	virtual void CustomizeChildren(TSharedRef<class IPropertyHandle> InPropertyHandle, class IDetailChildrenBuilder& StructBuilder, IPropertyTypeCustomizationUtils& StructCustomizationUtils) override;
+
+protected:
+
+	bool GetPropertyChain(TSharedRef<class IPropertyHandle> InPropertyHandle, FEditPropertyChain& OutPropertyChain, TArray<int32> &OutPropertyArrayIndices, bool& bOutEnabled)
+	{
+		OutPropertyChain.Empty();
+		OutPropertyArrayIndices.Reset();
+		bOutEnabled = false;
+		if (!ObjectsBeingCustomized.IsEmpty())
+		{
+			if (ObjectsBeingCustomized[0].Get())
+			{
+				TSharedPtr<class IPropertyHandle> ChainHandle = InPropertyHandle;
+				while (ChainHandle.IsValid() && ChainHandle->GetProperty() != nullptr)
+				{
+					OutPropertyChain.AddHead(ChainHandle->GetProperty());
+					OutPropertyArrayIndices.Insert(ChainHandle->GetIndexInArray(), 0);
+					ChainHandle = ChainHandle->GetParentHandle();					
+				}
+
+				if (OutPropertyChain.GetHead() != nullptr)
+				{
+					OutPropertyChain.SetActiveMemberPropertyNode(OutPropertyChain.GetTail()->GetValue());
+					bOutEnabled = !OutPropertyChain.GetHead()->GetValue()->HasAnyPropertyFlags(CPF_EditConst);
+					return true;
+				}
+			}
+		}
+		return false;
+	}
+
+	// extracts the value for a nested property (for Example Settings.WorldTransform) from an outer owner
+	template<typename ValueType>
+	ValueType& ContainerUObjectToValueRef(UObject* InOwner, ValueType& InDefault, FEditPropertyChain& InPropertyChain, TArray<int32> &InPropertyArrayIndices) const
+	{
+		if (InPropertyChain.GetHead() == nullptr)
+		{
+			return InDefault;
+		}
+		
+		FEditPropertyChain::TDoubleLinkedListNode* PropertyNode = InPropertyChain.GetHead();
+		uint8* MemoryPtr = (uint8*)InOwner;
+		int32 ChainIndex = 0;
+		do
+		{
+			const FProperty* Property = PropertyNode->GetValue();
+			MemoryPtr = Property->ContainerPtrToValuePtr<uint8>(MemoryPtr);
+
+			PropertyNode = PropertyNode->GetNextNode();
+			ChainIndex++;
+			
+			if(InPropertyArrayIndices.IsValidIndex(ChainIndex))
+			{
+				const int32 ArrayIndex = InPropertyArrayIndices[ChainIndex];
+				if(ArrayIndex != INDEX_NONE)
+				{
+					const FArrayProperty* ArrayProperty = CastField<FArrayProperty>(Property->GetOwnerProperty());
+					check(ArrayProperty);
+					
+					FScriptArrayHelper ArrayHelper(ArrayProperty, MemoryPtr);
+					if(!ArrayHelper.IsValidIndex(ArrayIndex))
+					{
+						return InDefault;
+					}
+					MemoryPtr = ArrayHelper.GetRawPtr(ArrayIndex);
+
+					// skip to the next property node already
+					PropertyNode = PropertyNode->GetNextNode();
+					ChainIndex++;
+				}
+			}
+
+		}
+		while (PropertyNode);
+
+		return *(ValueType*)MemoryPtr;
+	}
+
+	// specializations for FEulerTransform and FRotator at the end of this file
+	template<typename ValueType>
+	static bool IsQuaternionBasedRotation() { return true; }
+
+	// returns the numeric value of a vector component (or empty optional for varying values)
+	template<typename VectorType, typename NumericType>
+	TOptional<NumericType> GetVectorComponent(TSharedRef<class IPropertyHandle> InPropertyHandle, int32 InComponent) 
+	{
+		TOptional<NumericType> Result;
+		FEditPropertyChain PropertyChain;
+		TArray<int32> PropertyArrayIndices;
+		bool bEnabled;
+		if (!GetPropertyChain(InPropertyHandle, PropertyChain, PropertyArrayIndices, bEnabled))
+		{
+			return Result;
+		}
+	
+		for(const TWeakObjectPtr<UObject>& Object : ObjectsBeingCustomized)
+		{
+			if(Object.Get() && InPropertyHandle->IsValidHandle())
+			{
+				static VectorType ZeroVector = VectorType();
+				const VectorType& Vector = ContainerUObjectToValueRef<VectorType>(Object.Get(), ZeroVector, PropertyChain, PropertyArrayIndices);
+				NumericType Component = Vector[InComponent];
+				if(Result.IsSet())
+				{
+					if(!FMath::IsNearlyEqual(Result.GetValue(), Component))
+					{
+						return TOptional<NumericType>();
+					}
+				}
+				else
+				{
+					Result = Component;
+				}
+			}
+		}
+		return Result;
+	};
+
+	// called when a numeric value of a vector component is changed
+	template<typename VectorType, typename NumericType>
+	void OnVectorComponentChanged(TSharedRef<class IPropertyHandle> InPropertyHandle, int32 InComponent, NumericType InValue, bool bIsCommit, ETextCommit::Type InCommitType = ETextCommit::Default)
+	{
+		FEditPropertyChain PropertyChain;
+		TArray<int32> PropertyArrayIndices;
+		bool bEnabled;
+		if (!GetPropertyChain(InPropertyHandle, PropertyChain, PropertyArrayIndices, bEnabled))
+		{
+			return;
+		}
+	
+		TArray<UObject*> ObjectsView;
+		for(int32 Index = 0; Index < ObjectsBeingCustomized.Num(); Index++)
+		{
+			const TWeakObjectPtr<UObject>& Object = ObjectsBeingCustomized[Index];
+			if (Object.Get())
+			{
+				ObjectsView.Add(Object.Get());
+			}
+		}
+		FPropertyChangedEvent PropertyChangedEvent(InPropertyHandle->GetProperty(), bIsCommit ? EPropertyChangeType::ValueSet : EPropertyChangeType::Interactive, ObjectsView);
+		FPropertyChangedChainEvent PropertyChangedChainEvent(PropertyChain, PropertyChangedEvent);
+
+		URigVMController* Controller = nullptr;
+		if(BlueprintBeingCustomized && GraphBeingCustomized)
+		{
+			Controller = BlueprintBeingCustomized->GetController(GraphBeingCustomized);
+			if(bIsCommit)
+			{
+				Controller->OpenUndoBracket(FString::Printf(TEXT("Set %s"), *InPropertyHandle->GetProperty()->GetName()));
+			}
+		}
+
+		for(int32 Index = 0; Index < ObjectsBeingCustomized.Num(); Index++)
+		{
+			const TWeakObjectPtr<UObject>& Object = ObjectsBeingCustomized[Index];
+			if(Object.Get() && InPropertyHandle->IsValidHandle())
+			{
+				static VectorType ZeroVector = VectorType();
+				VectorType& Vector = ContainerUObjectToValueRef<VectorType>(Object.Get(), ZeroVector, PropertyChain, PropertyArrayIndices);
+				VectorType PreviousVector = Vector;
+				Vector[InComponent] = InValue;
+					
+				if(!PreviousVector.Equals(Vector))
+				{
+					Object->PostEditChangeChainProperty(PropertyChangedChainEvent);
+					InPropertyHandle->NotifyPostChange(PropertyChangedEvent.ChangeType);
+				}
+			}
+		}
+
+		if(Controller && bIsCommit)
+		{
+			Controller->CloseUndoBracket();
+		}
+	};
+
+	// specializations for FVector and FVector4 at the end of this file
+	template<typename VectorType>
+	void ExtendVectorArgs(TSharedRef<class IPropertyHandle> InPropertyHandle, void* ArgumentsPtr) {}
+
+	template<typename VectorType, int32 NumberOfComponents>
+	void CustomizeVector(TSharedRef<class IPropertyHandle> InPropertyHandle, class IDetailChildrenBuilder& StructBuilder, IPropertyTypeCustomizationUtils& StructCustomizationUtils)
+	{
+		typedef typename VectorType::FReal NumericType;
+		typedef SNumericVectorInputBox<NumericType, VectorType, NumberOfComponents> SLocalVectorInputBox;
+
+		FEditPropertyChain PropertyChain;
+        TArray<int32> PropertyArrayIndices;
+        bool bEnabled;
+        if (!GetPropertyChain(InPropertyHandle, PropertyChain, PropertyArrayIndices, bEnabled))
+        {
+        	return;
+        }
+
+		typename SLocalVectorInputBox::FArguments Args;
+		Args.Font(IDetailLayoutBuilder::GetDetailFont());
+		Args.IsEnabled(bEnabled);
+		Args.AllowSpin(true);
+		Args.SpinDelta(0.01f);
+		Args.bColorAxisLabels(true);
+		Args.X_Lambda([this, InPropertyHandle]()
+		{
+			return GetVectorComponent<VectorType, NumericType>(InPropertyHandle, 0);
+		});
+		Args.OnXChanged_Lambda([this, InPropertyHandle](NumericType Value)
+		{
+			OnVectorComponentChanged<VectorType, NumericType>(InPropertyHandle, 0, Value, false);
+		});
+		Args.OnXCommitted_Lambda([this, InPropertyHandle](NumericType Value, ETextCommit::Type CommitType)
+		{
+			OnVectorComponentChanged<VectorType, NumericType>(InPropertyHandle, 0, Value, true, CommitType);
+		});
+		Args.Y_Lambda([this, InPropertyHandle]()
+		{
+			return GetVectorComponent<VectorType, NumericType>(InPropertyHandle, 1);
+		});
+		Args.OnYChanged_Lambda([this, InPropertyHandle](NumericType Value)
+		{
+			OnVectorComponentChanged<VectorType, NumericType>(InPropertyHandle, 1, Value, false);
+		});
+		Args.OnYCommitted_Lambda([this, InPropertyHandle](NumericType Value, ETextCommit::Type CommitType)
+		{
+			OnVectorComponentChanged<VectorType, NumericType>(InPropertyHandle, 1, Value, true, CommitType);
+		});
+
+		ExtendVectorArgs<VectorType>(InPropertyHandle, &Args);
+
+		StructBuilder.AddProperty(InPropertyHandle).CustomWidget()
+		.IsEnabled(bEnabled)
+		.NameContent()
+		[
+			InPropertyHandle->CreatePropertyNameWidget()
+		]
+		.ValueContent()
+		.MinDesiredWidth(375.f)
+		.MaxDesiredWidth(375.f)
+		.HAlign(HAlign_Left)
+		[
+			SArgumentNew(Args, SLocalVectorInputBox)
+		];
+	}
+
+	// returns the rotation for rotator or quaternions (or empty optional for varying values)
+	template<typename RotationType>
+	TOptional<RotationType> GetRotation(TSharedRef<class IPropertyHandle> InPropertyHandle)
+	{
+		TOptional<RotationType> Result;
+		FEditPropertyChain PropertyChain;
+		TArray<int32> PropertyArrayIndices;
+		bool bEnabled;
+		if (!GetPropertyChain(InPropertyHandle, PropertyChain, PropertyArrayIndices, bEnabled))
+		{
+			return Result;
+		}
+		
+		for(const TWeakObjectPtr<UObject>& Object : ObjectsBeingCustomized)
+		{
+			if(Object.Get() && InPropertyHandle->IsValidHandle())
+			{
+				static RotationType ZeroRotation = RotationType();
+				const RotationType& Rotation = ContainerUObjectToValueRef<RotationType>(Object.Get(), ZeroRotation, PropertyChain, PropertyArrayIndices);
+				if(Result.IsSet())
+				{
+					if(!Rotation.Equals(Result.GetValue()))
+					{
+						return TOptional<RotationType>();
+					}
+				}
+				else
+				{
+					Result = Rotation;
+				}
+			}
+		}
+		return Result;
+	};
+
+	// called when a rotation value is changed / committed
+	template<typename RotationType>
+	void OnRotationChanged(TSharedRef<class IPropertyHandle> InPropertyHandle, RotationType InValue, bool bIsCommit, ETextCommit::Type InCommitType = ETextCommit::Default)
+	{
+		FEditPropertyChain PropertyChain;
+        TArray<int32> PropertyArrayIndices;
+        bool bEnabled;
+        if (!GetPropertyChain(InPropertyHandle, PropertyChain, PropertyArrayIndices, bEnabled))
+        {
+        	return;
+        }
+	
+		TArray<UObject*> ObjectsView;
+		for(int32 Index = 0; Index < ObjectsBeingCustomized.Num(); Index++)
+		{
+			const TWeakObjectPtr<UObject>& Object = ObjectsBeingCustomized[Index];
+			if (Object.Get())
+			{
+				ObjectsView.Add(Object.Get());
+			}
+		}
+		FPropertyChangedEvent PropertyChangedEvent(InPropertyHandle->GetProperty(), bIsCommit ? EPropertyChangeType::ValueSet : EPropertyChangeType::Interactive, ObjectsView);
+		FPropertyChangedChainEvent PropertyChangedChainEvent(PropertyChain, PropertyChangedEvent);
+
+		URigVMController* Controller = nullptr;
+		if(BlueprintBeingCustomized && GraphBeingCustomized)
+		{
+			Controller = BlueprintBeingCustomized->GetController(GraphBeingCustomized);
+			if(bIsCommit)
+			{
+				Controller->OpenUndoBracket(FString::Printf(TEXT("Set %s"), *InPropertyHandle->GetProperty()->GetName()));
+			}
+		}
+
+		for(int32 Index = 0; Index < ObjectsBeingCustomized.Num(); Index++)
+		{
+			const TWeakObjectPtr<UObject>& Object = ObjectsBeingCustomized[Index];
+			if(Object.Get() && InPropertyHandle->IsValidHandle())
+			{
+				static RotationType ZeroRotation = RotationType();
+				RotationType& Rotation = ContainerUObjectToValueRef<RotationType>(Object.Get(), ZeroRotation, PropertyChain, PropertyArrayIndices);
+				RotationType PreviousRotation = Rotation;
+				Rotation = InValue;
+					
+				if(!PreviousRotation.Equals(Rotation))
+				{
+					Object->PostEditChangeChainProperty(PropertyChangedChainEvent);
+					InPropertyHandle->NotifyPostChange(PropertyChangedEvent.ChangeType);
+				}
+			}
+		}
+
+		if(Controller && bIsCommit)
+		{
+			Controller->CloseUndoBracket();
+		}
+	};
+
+	// specializations for FRotator and FQuat at the end of this file
+	template<typename RotationType>
+	void ExtendRotationArgs(TSharedRef<class IPropertyHandle> InPropertyHandle, void* ArgumentsPtr) {}
+
+	// add the widget for a rotation (rotator or quat)
+	template<typename RotationType>
+	void CustomizeRotation(TSharedRef<class IPropertyHandle> InPropertyHandle, class IDetailChildrenBuilder& StructBuilder, IPropertyTypeCustomizationUtils& StructCustomizationUtils)
+	{
+		FEditPropertyChain PropertyChain;
+		TArray<int32> PropertyArrayIndices;
+		bool bEnabled;
+		if (!GetPropertyChain(InPropertyHandle, PropertyChain, PropertyArrayIndices, bEnabled))
+		{
+			return;
+		}
+		
+		typedef typename RotationType::FReal NumericType;
+		typedef SAdvancedRotationInputBox<NumericType> SLocalRotationInputBox;
+		typename SLocalRotationInputBox::FArguments Args;
+		Args.Font(IDetailLayoutBuilder::GetDetailFont());
+		Args.IsEnabled(bEnabled);
+		Args.AllowSpin(true);
+		Args.bColorAxisLabels(true);
+
+		ExtendRotationArgs<RotationType>(InPropertyHandle, &Args);
+		
+		StructBuilder.AddProperty(InPropertyHandle).CustomWidget()
+		.IsEnabled(bEnabled)
+		.NameContent()
+		[
+			
+			InPropertyHandle->CreatePropertyNameWidget()
+		]
+		.ValueContent()
+		.MinDesiredWidth(375.f)
+		.MaxDesiredWidth(375.f)
+		.HAlign(HAlign_Left)
+		[
+			SArgumentNew(Args, SLocalRotationInputBox)
+		];
+	}
+
+	// add the widget for a transform / euler transform
+	template<typename TransformType>
+	void CustomizeTransform(TSharedRef<class IPropertyHandle> InPropertyHandle, class IDetailChildrenBuilder& StructBuilder, IPropertyTypeCustomizationUtils& StructCustomizationUtils)
+	{
+		FEditPropertyChain PropertyChain;
+		TArray<int32> PropertyArrayIndices;
+		bool bEnabled;
+		if (!GetPropertyChain(InPropertyHandle, PropertyChain, PropertyArrayIndices, bEnabled))
+		{
+			return;
+		}
+		
+		typedef typename TransformType::FReal FReal;
+		typename SAdvancedTransformInputBox<TransformType>::FArguments WidgetArgs;
+		WidgetArgs.IsEnabled(bEnabled);
+		WidgetArgs.AllowEditRotationRepresentation(true);
+		WidgetArgs.UseQuaternionForRotation(IsQuaternionBasedRotation<TransformType>());
+
+		static TransformType Identity = TransformType::Identity;
+		TransformType DefaultValue = ContainerUObjectToValueRef<TransformType>(ObjectsBeingCustomized[0]->GetClass()->GetDefaultObject(), Identity, PropertyChain, PropertyArrayIndices);
+
+		WidgetArgs.DiffersFromDefault_Lambda([this, InPropertyHandle, DefaultValue](ESlateTransformComponent::Type InTransformComponent) -> bool
+		{
+			FEditPropertyChain PropertyChain;
+			TArray<int32> PropertyArrayIndices;
+			bool bEnabled;
+			if (!GetPropertyChain(InPropertyHandle, PropertyChain, PropertyArrayIndices, bEnabled))
+			{
+				return false;
+			}
+			
+			for(const TWeakObjectPtr<UObject>& Object : ObjectsBeingCustomized)
+			{
+				if(Object.Get() && InPropertyHandle->IsValidHandle())
+				{
+					const TransformType& Transform = ContainerUObjectToValueRef<TransformType>(Object.Get(), Identity, PropertyChain, PropertyArrayIndices);
+
+					switch(InTransformComponent)
+					{
+						case ESlateTransformComponent::Location:
+						{
+							if(!Transform.GetLocation().Equals(DefaultValue.GetLocation()))
+							{
+								return true;
+							}
+							break;
+						}
+						case ESlateTransformComponent::Rotation:
+						{
+							if(!Transform.Rotator().Equals(DefaultValue.Rotator()))
+							{
+								return true;
+							}
+							break;
+						}
+						case ESlateTransformComponent::Scale:
+						{
+							if(!Transform.GetScale3D().Equals(DefaultValue.GetScale3D()))
+							{
+								return true;
+							}
+							break;
+						}
+						default:
+						{
+							break;
+						}
+					}
+				}
+			}
+			return false;
+		});
+
+		WidgetArgs.OnGetNumericValue_Lambda([this, InPropertyHandle](
+			ESlateTransformComponent::Type InTransformComponent,
+			ESlateRotationRepresentation::Type InRotationRepresentation,
+			ESlateTransformSubComponent::Type InTransformSubComponent) -> TOptional<FReal>
+		{
+			TOptional<FReal> Result;
+			FEditPropertyChain PropertyChain;
+			TArray<int32> PropertyArrayIndices;
+			bool bEnabled;
+			if (!GetPropertyChain(InPropertyHandle, PropertyChain, PropertyArrayIndices, bEnabled))
+			{
+				return Result;
+			}
+			
+			for(const TWeakObjectPtr<UObject>& Object : ObjectsBeingCustomized)
+			{
+				if(Object.Get() && InPropertyHandle->IsValidHandle())
+				{
+					const TransformType& Transform = ContainerUObjectToValueRef<TransformType>(Object.Get(), Identity, PropertyChain, PropertyArrayIndices);
+					
+					TOptional<FReal> Value = SAdvancedTransformInputBox<TransformType>::GetNumericValueFromTransform(
+						Transform,
+						InTransformComponent,
+						InRotationRepresentation,
+						InTransformSubComponent
+						);
+
+					if(Value.IsSet())
+					{
+						if(Result.IsSet())
+						{
+							if(!FMath::IsNearlyEqual(Result.GetValue(), Value.GetValue()))
+							{
+								return TOptional<FReal>();
+							}
+						}
+						else
+						{
+							Result = Value;
+						}
+					}
+				}
+			}
+			return Result;
+		});
+		
+
+		auto OnNumericValueChanged = [this, InPropertyHandle](
+			ESlateTransformComponent::Type InTransformComponent, 
+			ESlateRotationRepresentation::Type InRotationRepresentation, 
+			ESlateTransformSubComponent::Type InSubComponent,
+			FReal InValue,
+			bool bIsCommit,
+			ETextCommit::Type InCommitType = ETextCommit::Default)
+		{
+			FEditPropertyChain PropertyChain;
+			TArray<int32> PropertyArrayIndices;
+			bool bEnabled;
+			if (!GetPropertyChain(InPropertyHandle, PropertyChain, PropertyArrayIndices, bEnabled))
+			{
+				return;
+			}
+			
+			TArray<UObject*> ObjectsView;
+			for(int32 Index = 0; Index < ObjectsBeingCustomized.Num(); Index++)
+			{
+				const TWeakObjectPtr<UObject>& Object = ObjectsBeingCustomized[Index];
+				if (Object.Get())
+				{
+					ObjectsView.Add(Object.Get());
+				}
+			}
+			FPropertyChangedEvent PropertyChangedEvent(InPropertyHandle->GetProperty(), bIsCommit ? EPropertyChangeType::ValueSet : EPropertyChangeType::Interactive, ObjectsView);
+			FPropertyChangedChainEvent PropertyChangedChainEvent(PropertyChain, PropertyChangedEvent);
+
+			URigVMController* Controller = nullptr;
+			if(BlueprintBeingCustomized && GraphBeingCustomized)
+			{
+				Controller = BlueprintBeingCustomized->GetController(GraphBeingCustomized);
+				if(bIsCommit)
+				{
+					Controller->OpenUndoBracket(FString::Printf(TEXT("Set %s"), *InPropertyHandle->GetProperty()->GetName()));
+				}
+			}
+
+			for(int32 Index = 0; Index < ObjectsBeingCustomized.Num(); Index++)
+			{
+				const TWeakObjectPtr<UObject>& Object = ObjectsBeingCustomized[Index];
+				if(Object.Get() && InPropertyHandle->IsValidHandle())
+				{
+					TransformType& Transform = ContainerUObjectToValueRef<TransformType>(Object.Get(), Identity, PropertyChain, PropertyArrayIndices);
+					TransformType PreviousTransform = Transform;
+					
+					SAdvancedTransformInputBox<TransformType>::ApplyNumericValueChange(
+						Transform,
+						InValue,
+						InTransformComponent,
+						InRotationRepresentation,
+						InSubComponent);
+
+					if(!PreviousTransform.Equals(Transform))
+					{
+						Object->PostEditChangeChainProperty(PropertyChangedChainEvent);
+						InPropertyHandle->NotifyPostChange(PropertyChangedEvent.ChangeType);
+					}
+				}
+			}
+
+			if(Controller && bIsCommit)
+			{
+				Controller->CloseUndoBracket();
+			}
+		};
+
+		WidgetArgs.OnNumericValueChanged_Lambda([OnNumericValueChanged](
+			ESlateTransformComponent::Type InTransformComponent, 
+			ESlateRotationRepresentation::Type InRotationRepresentation, 
+			ESlateTransformSubComponent::Type InSubComponent,
+			FReal InValue)
+		{
+			OnNumericValueChanged(InTransformComponent, InRotationRepresentation, InSubComponent, InValue, false);
+		});
+
+		WidgetArgs.OnNumericValueCommitted_Lambda([OnNumericValueChanged](
+			ESlateTransformComponent::Type InTransformComponent, 
+			ESlateRotationRepresentation::Type InRotationRepresentation, 
+			ESlateTransformSubComponent::Type InSubComponent,
+			FReal InValue, 
+			ETextCommit::Type InCommitType)
+		{
+			OnNumericValueChanged(InTransformComponent, InRotationRepresentation, InSubComponent, InValue, true, InCommitType);
+		});
+
+		WidgetArgs.OnResetToDefault_Lambda([this, DefaultValue, InPropertyHandle](ESlateTransformComponent::Type InTransformComponent)
+		{
+			FEditPropertyChain PropertyChain;
+			TArray<int32> PropertyArrayIndices;
+			bool bEnabled;
+			if (!GetPropertyChain(InPropertyHandle, PropertyChain, PropertyArrayIndices, bEnabled))
+			{
+				return;
+			}
+			
+			URigVMController* Controller = nullptr;
+			if(BlueprintBeingCustomized && GraphBeingCustomized)
+			{
+				Controller = BlueprintBeingCustomized->GetController(GraphBeingCustomized);
+				if(Controller)
+				{
+					Controller->OpenUndoBracket(FString::Printf(TEXT("Reset %s to Default"), *InPropertyHandle->GetProperty()->GetName()));
+				}
+			}
+
+			TArray<UObject*> ObjectsView;
+			for(int32 Index = 0; Index < ObjectsBeingCustomized.Num(); Index++)
+			{
+				const TWeakObjectPtr<UObject>& Object = ObjectsBeingCustomized[Index];
+				if (Object.Get())
+				{
+					ObjectsView.Add(Object.Get());
+				}
+			}
+			FPropertyChangedEvent PropertyChangedEvent(InPropertyHandle->GetProperty(), EPropertyChangeType::ValueSet, ObjectsView);
+			FPropertyChangedChainEvent PropertyChangedChainEvent(PropertyChain, PropertyChangedEvent);
+			
+			for(int32 Index = 0; Index < ObjectsBeingCustomized.Num(); Index++)
+			{
+				const TWeakObjectPtr<UObject>& Object = ObjectsBeingCustomized[Index];
+				if(Object.Get() && InPropertyHandle->IsValidHandle())
+				{
+					TransformType& Transform = ContainerUObjectToValueRef<TransformType>(Object.Get(), Identity, PropertyChain, PropertyArrayIndices);
+					TransformType PreviousTransform = Transform;
+
+					switch(InTransformComponent)
+					{
+						case ESlateTransformComponent::Location:
+						{
+							Transform.SetLocation(DefaultValue.GetLocation());
+							break;
+						}
+						case ESlateTransformComponent::Rotation:
+						{
+							Transform.SetRotation(DefaultValue.GetRotation());
+							break;
+						}
+						case ESlateTransformComponent::Scale:
+						{
+							Transform.SetScale3D(DefaultValue.GetScale3D());
+							break;
+						}
+						case ESlateTransformComponent::Max:
+						default:
+						{
+							Transform.SetLocation(DefaultValue.GetLocation());
+							break;
+						}
+					}
+
+					
+					if(!PreviousTransform.Equals(Transform))
+					{
+						Object->PostEditChangeChainProperty(PropertyChangedChainEvent);
+						InPropertyHandle->NotifyPostChange(PropertyChangedEvent.ChangeType);
+					}
+				}
+			}
+
+			if(Controller)
+			{
+				Controller->CloseUndoBracket();
+			}
+		});
+
+		SAdvancedTransformInputBox<TransformType>::ConstructGroupedTransformRows(
+			StructBuilder,
+			InPropertyHandle->GetPropertyDisplayName(),
+			InPropertyHandle->GetToolTipText(),
+			WidgetArgs);
+	}
+		
+	UScriptStruct* ScriptStruct;
+	URigVMBlueprint* BlueprintBeingCustomized;
+	URigVMGraph* GraphBeingCustomized;
+	TArray<TWeakObjectPtr<UObject>> ObjectsBeingCustomized; 
+};
+
+template<>
+inline bool FRigVMGraphMathTypeDetailCustomization::IsQuaternionBasedRotation<FEulerTransform>() { return false; }
+
+template<>
+inline bool FRigVMGraphMathTypeDetailCustomization::IsQuaternionBasedRotation<FRotator>() { return false; }
+
+template<>
+inline void FRigVMGraphMathTypeDetailCustomization::ExtendVectorArgs<FVector>(TSharedRef<class IPropertyHandle> InPropertyHandle, void* ArgumentsPtr)
+{
+	using VectorType = FVector;
+	typedef typename VectorType::FReal NumericType;
+	typedef SNumericVectorInputBox<NumericType, VectorType, 3> SLocalVectorInputBox;
+
+	typename SLocalVectorInputBox::FArguments& Args = *(typename SLocalVectorInputBox::FArguments*)ArgumentsPtr; 
+	Args
+	.Z_Lambda([this, InPropertyHandle]()
+	{
+		return GetVectorComponent<VectorType, NumericType>(InPropertyHandle, 2);
+	})
+	.OnZChanged_Lambda([this, InPropertyHandle](NumericType Value)
+	{
+		OnVectorComponentChanged<VectorType, NumericType>(InPropertyHandle, 2, Value, false);
+	})
+	.OnZCommitted_Lambda([this, InPropertyHandle](NumericType Value, ETextCommit::Type CommitType)
+	{
+		OnVectorComponentChanged<VectorType, NumericType>(InPropertyHandle, 2, Value, true, CommitType);
+	});
+}
+
+template<>
+inline void FRigVMGraphMathTypeDetailCustomization::ExtendVectorArgs<FVector4>(TSharedRef<class IPropertyHandle> InPropertyHandle, void* ArgumentsPtr)
+{
+	using VectorType = FVector4;
+	typedef typename VectorType::FReal NumericType;
+	typedef SNumericVectorInputBox<NumericType, VectorType, 4> SLocalVectorInputBox;
+
+	typename SLocalVectorInputBox::FArguments& Args = *(typename SLocalVectorInputBox::FArguments*)ArgumentsPtr; 
+	Args
+	.Z_Lambda([this, InPropertyHandle]()
+	{
+		return GetVectorComponent<VectorType, NumericType>(InPropertyHandle, 2);
+	})
+	.OnZChanged_Lambda([this, InPropertyHandle](NumericType Value)
+	{
+		OnVectorComponentChanged<VectorType, NumericType>(InPropertyHandle, 2, Value, false);
+	})
+	.OnZCommitted_Lambda([this, InPropertyHandle](NumericType Value, ETextCommit::Type CommitType)
+	{
+		OnVectorComponentChanged<VectorType, NumericType>(InPropertyHandle, 2, Value, true, CommitType);
+	})
+	.W_Lambda([this, InPropertyHandle]()
+	{
+		return GetVectorComponent<VectorType, NumericType>(InPropertyHandle, 3);
+	})
+	.OnWChanged_Lambda([this, InPropertyHandle](NumericType Value)
+	{
+		OnVectorComponentChanged<VectorType, NumericType>(InPropertyHandle, 3, Value, false);
+	})
+	.OnWCommitted_Lambda([this, InPropertyHandle](NumericType Value, ETextCommit::Type CommitType)
+	{
+		OnVectorComponentChanged<VectorType, NumericType>(InPropertyHandle, 3, Value, true, CommitType);
+	});
+}
+
+template<>
+inline void FRigVMGraphMathTypeDetailCustomization::ExtendRotationArgs<FQuat>(TSharedRef<class IPropertyHandle> InPropertyHandle, void* ArgumentsPtr)
+{
+	using RotationType = FQuat;
+	typedef typename RotationType::FReal NumericType;
+	typedef SAdvancedRotationInputBox<NumericType> SLocalRotationInputBox;
+	typename SLocalRotationInputBox::FArguments& Args = *(typename SLocalRotationInputBox::FArguments*)ArgumentsPtr; 
+
+	Args.Quaternion_Lambda([this, InPropertyHandle]() -> TOptional<RotationType>
+	{
+		return GetRotation<RotationType>(InPropertyHandle);
+	});
+
+	Args.OnQuaternionChanged_Lambda([this, InPropertyHandle](RotationType InValue)
+	{
+		OnRotationChanged<RotationType>(InPropertyHandle, InValue, false);
+	});
+
+	Args.OnQuaternionCommitted_Lambda([this, InPropertyHandle](RotationType InValue, ETextCommit::Type InCommitType)
+	{
+		OnRotationChanged<RotationType>(InPropertyHandle, InValue, true, InCommitType);
+	});
+}
+
+template<>
+inline void FRigVMGraphMathTypeDetailCustomization::ExtendRotationArgs<FRotator>(TSharedRef<class IPropertyHandle> InPropertyHandle, void* ArgumentsPtr)
+{
+	using RotationType = FRotator;
+	typedef typename RotationType::FReal NumericType;
+	typedef SAdvancedRotationInputBox<NumericType> SLocalRotationInputBox;
+	typename SLocalRotationInputBox::FArguments& Args = *(typename SLocalRotationInputBox::FArguments*)ArgumentsPtr; 
+
+	Args.Rotator_Lambda([this, InPropertyHandle]() -> TOptional<RotationType>
+	{
+		return GetRotation<RotationType>(InPropertyHandle);
+	});
+
+	Args.OnRotatorChanged_Lambda([this, InPropertyHandle](RotationType InValue)
+	{
+		OnRotationChanged<RotationType>(InPropertyHandle, InValue, false);
+	});
+
+	Args.OnRotatorCommitted_Lambda([this, InPropertyHandle](RotationType InValue, ETextCommit::Type InCommitType)
+	{
+		OnRotationChanged<RotationType>(InPropertyHandle, InValue, true, InCommitType);
+	});
+}
diff --git a/Engine/Plugins/Runtime/RigVM/Source/RigVMEditor/Public/Editor/RigVMLocalVariableDetailCustomization.h b/Engine/Plugins/Runtime/RigVM/Source/RigVMEditor/Public/Editor/RigVMLocalVariableDetailCustomization.h
new file mode 100644
--- /dev/null
+++ b/Engine/Plugins/Runtime/RigVM/Source/RigVMEditor/Public/Editor/RigVMLocalVariableDetailCustomization.h
@@ -0,0 +1,66 @@
+﻿// Copyright Epic Games, Inc. All Rights Reserved.
+
+#pragma once
+
+#include "CoreMinimal.h"
+#include "IDetailCustomization.h"
+#include "IPropertyTypeCustomization.h"
+#include "RigVMHost.h"
+#include "RigVMBlueprint.h"
+#include "Editor/RigVMDetailsViewWrapperObject.h"
+#include "EdGraph/RigVMEdGraph.h"
+#include "Widgets/SRigVMGraphPinNameListValueWidget.h"
+#include "Styling/SlateTypes.h"
+#include "IPropertyUtilities.h"
+#include "EdGraph/RigVMEdGraphSchema.h"
+
+class IPropertyHandle;
+
+class RIGVMEDITOR_API FRigVMLocalVariableDetailCustomization : public IDetailCustomization
+{
+	FRigVMLocalVariableDetailCustomization()
+	: GraphBeingCustomized(nullptr)
+	, BlueprintBeingCustomized(nullptr)
+	, NameValidator(nullptr, nullptr, NAME_None)
+	{}
+
+	
+public:
+
+	// Makes a new instance of this detail layout class for a specific detail view requesting it
+	static TSharedRef<IDetailCustomization> MakeInstance()
+	{
+		return MakeShareable(new FRigVMLocalVariableDetailCustomization);
+	}
+
+	/** IDetailCustomization interface */
+	virtual void CustomizeDetails(IDetailLayoutBuilder& DetailBuilder) override;
+
+private:
+
+	URigVMGraph* GraphBeingCustomized;
+	URigVMBlueprint* BlueprintBeingCustomized;
+	FRigVMGraphVariableDescription VariableDescription;
+	TArray<TWeakObjectPtr<URigVMDetailsViewWrapperObject>> ObjectsBeingCustomized;
+
+	TSharedPtr<IPropertyHandle> NameHandle;
+	TSharedPtr<IPropertyHandle> TypeHandle;
+	TSharedPtr<IPropertyHandle> TypeObjectHandle;
+	TSharedPtr<IPropertyHandle> DefaultValueHandle;
+
+	FRigVMLocalVariableNameValidator NameValidator;
+	TArray<TSharedPtr<FString>> EnumOptions;
+
+	FText GetName() const;
+	void SetName(const FText& InNewText, ETextCommit::Type InCommitType);
+	bool OnVerifyNameChanged(const FText& InText, FText& OutErrorMessage);
+	
+	FEdGraphPinType OnGetPinInfo() const;
+	void HandlePinInfoChanged(const FEdGraphPinType& PinType);
+
+	ECheckBoxState HandleBoolDefaultValueIsChecked( ) const;
+	void OnBoolDefaultValueChanged(ECheckBoxState InCheckBoxState);
+};
+
+
+
diff --git a/Engine/Plugins/Runtime/RigVM/Source/RigVMEditor/Public/Editor/RigVMVariableDetailCustomization.h b/Engine/Plugins/Runtime/RigVM/Source/RigVMEditor/Public/Editor/RigVMVariableDetailCustomization.h
new file mode 100644
--- /dev/null
+++ b/Engine/Plugins/Runtime/RigVM/Source/RigVMEditor/Public/Editor/RigVMVariableDetailCustomization.h
@@ -0,0 +1,46 @@
+// Copyright Epic Games, Inc. All Rights Reserved.
+
+#pragma once
+
+#include "CoreMinimal.h"
+#include "Engine/Blueprint.h"
+#include "IDetailCustomization.h"
+#include "Styling/SlateTypes.h"
+#include "UObject/WeakFieldPtr.h"
+
+class IDetailLayoutBuilder;
+class IBlueprintEditor;
+class UBlueprint;
+
+class RIGVMEDITOR_API FRigVMVariableDetailCustomization : public IDetailCustomization
+{
+public:
+	/** Makes a new instance of this detail layout class for a specific detail view requesting it */
+	static TSharedPtr<IDetailCustomization> MakeInstance(TSharedPtr<IBlueprintEditor> InBlueprintEditor);
+
+	FRigVMVariableDetailCustomization(TSharedPtr<IBlueprintEditor> InBlueprintEditor, UBlueprint* Blueprint)
+		: BlueprintEditorPtr(InBlueprintEditor)
+		, BlueprintPtr(Blueprint)
+	{}
+
+	// IDetailCustomization interface
+	virtual void CustomizeDetails(IDetailLayoutBuilder& DetailLayout) override;
+
+private:
+	bool IsAnimationFlagEnabled(TWeakFieldPtr<FProperty> PropertyBeingCustomized) const;
+
+	ECheckBoxState IsAnimationOutputChecked(TWeakFieldPtr<FProperty> PropertyBeingCustomized) const;
+
+	void HandleAnimationOutputCheckStateChanged(ECheckBoxState CheckBoxState, TWeakFieldPtr<FProperty> PropertyBeingCustomized);
+
+	ECheckBoxState IsAnimationInputChecked(TWeakFieldPtr<FProperty> PropertyBeingCustomized) const;
+
+	void HandleAnimationInputCheckStateChanged(ECheckBoxState CheckBoxState, TWeakFieldPtr<FProperty> PropertyBeingCustomized);
+
+private:
+	/** The Blueprint editor we are embedded in */
+	TWeakPtr<IBlueprintEditor> BlueprintEditorPtr;
+
+	/** The blueprint we are editing */
+	TWeakObjectPtr<UBlueprint> BlueprintPtr;
+};
diff --git a/Engine/Plugins/Runtime/RigVM/Source/RigVMEditor/Public/RigVMEditorModule.h b/Engine/Plugins/Runtime/RigVM/Source/RigVMEditor/Public/RigVMEditorModule.h
--- a/Engine/Plugins/Runtime/RigVM/Source/RigVMEditor/Public/RigVMEditorModule.h
+++ b/Engine/Plugins/Runtime/RigVM/Source/RigVMEditor/Public/RigVMEditorModule.h
@@ -9,10 +9,28 @@
 #include "CoreMinimal.h"
 #include "Modules/ModuleInterface.h"
 #include "Logging/LogMacros.h"
+#include "Editor/RigVMEditor.h"
+#include "EdGraph/RigVMEdGraphPanelNodeFactory.h"
+#include "EdGraph/RigVMEdGraphPanelPinFactory.h"
+#include "Toolkits/AssetEditorToolkit.h"
+#include "Kismet2/StructureEditorUtils.h"
 
 DECLARE_LOG_CATEGORY_EXTERN(LogRigVMEditor, Log, All);
 
-class FRigVMEditorModule : public IModuleInterface
+// shallow interface declaration for use within RigVMDeveloper 
+class IRigVMEditorModule : public IModuleInterface, public IHasMenuExtensibility, public IHasToolBarExtensibility, public FStructureEditorUtils::INotifyOnStructChanged
+{
+public:
+
+	static IRigVMEditorModule& Get()
+	{
+		return FModuleManager::LoadModuleChecked< IRigVMEditorModule >(TEXT("RigVMEditor"));
+	}
+
+	virtual void GetContextMenuActions(const URigVMEdGraphSchema* Schema, class UToolMenu* Menu, class UGraphNodeContextMenuContext* Context) const = 0;
+};
+
+class RIGVMEDITOR_API FRigVMEditorModule : public IRigVMEditorModule
 {
 public:
 
@@ -21,4 +39,84 @@ public:
 	/** IModuleInterface implementation */
 	virtual void StartupModule() override;
 	virtual void ShutdownModule() override;
+
+	void StartupModuleCommon();
+	void ShutdownModuleCommon();
+
+	/** RigVMEditorModule interface */
+	virtual UClass* GetRigVMBlueprintClass() const;
+	const URigVMBlueprint* GetRigVMBlueprintCDO() const;
+	virtual void GetTypeActions(URigVMBlueprint* RigVMBlueprint, FBlueprintActionDatabaseRegistrar& ActionRegistrar);
+	virtual void GetInstanceActions(URigVMBlueprint* RigVMBlueprint, FBlueprintActionDatabaseRegistrar& ActionRegistrar);
+	virtual void GetNodeContextMenuActions(URigVMBlueprint* RigVMBlueprint, const URigVMEdGraphNode* EdGraphNode, URigVMNode* ModelNode, class UToolMenu* Menu) const;
+	virtual void GetPinContextMenuActions(URigVMBlueprint* RigVMBlueprint, const UEdGraphPin* EdGraphPin, URigVMPin* ModelPin, class UToolMenu* Menu) const;
+
+	/** IHasMenuExtensibility interface */
+	virtual TSharedPtr<FExtensibilityManager> GetMenuExtensibilityManager() override { return MenuExtensibilityManager; }
+
+	/** IHasToolBarExtensibility interface */
+	virtual TSharedPtr<FExtensibilityManager> GetToolBarExtensibilityManager() override { return ToolBarExtensibilityManager; }
+
+	/* FStructureEditorUtils::INotifyOnStructChanged Interface, used to respond to changes to user defined structs */
+	virtual void PreChange(const UUserDefinedStruct* Changed, FStructureEditorUtils::EStructureEditorChangeInfo ChangedType) override;
+	virtual void PostChange(const UUserDefinedStruct* Changed, FStructureEditorUtils::EStructureEditorChangeInfo ChangedType) override;
+
+	/** Get all toolbar extenders */
+	DECLARE_DELEGATE_RetVal_TwoParams(TSharedRef<FExtender>, FRigVMEditorToolbarExtender, const TSharedRef<FUICommandList> /*InCommandList*/, TSharedRef<FRigVMEditor> /*InRigVMEditor*/);
+	virtual TArray<FRigVMEditorToolbarExtender>& GetAllRigVMEditorToolbarExtenders();
+
+	virtual void GetContextMenuActions(const URigVMEdGraphSchema* Schema, class UToolMenu* Menu, class UGraphNodeContextMenuContext* Context) const override;
+
+	/** Make sure to create the root graph for a given blueprint */
+	void CreateRootGraphIfRequired(URigVMBlueprint* InBlueprint) const;
+
+	FConnectionDrawingPolicy* CreateConnectionDrawingPolicy(int32 InBackLayerID, int32 InFrontLayerID, float InZoomFactor, const FSlateRect& InClippingRect, class FSlateWindowElementList& InDrawElements, class UEdGraph* InGraphObj) const;
+
+protected:
+
+	bool IsRigVMEditorModuleBase() const;
+
+	/** Specific section callbacks for the context menu */
+	virtual void GetNodeWorkflowContextMenuActions(URigVMBlueprint* RigVMBlueprint, const URigVMEdGraphNode* EdGraphNode, URigVMNode* ModelNode, UToolMenu* Menu) const;
+	virtual void GetNodeEventsContextMenuActions(URigVMBlueprint* RigVMBlueprint, const URigVMEdGraphNode* EdGraphNode, URigVMNode* ModelNode, UToolMenu* Menu) const;
+	virtual void GetNodeConversionContextMenuActions(URigVMBlueprint* RigVMBlueprint, const URigVMEdGraphNode* EdGraphNode, URigVMNode* ModelNode, UToolMenu* Menu) const;
+	virtual void GetNodeDebugContextMenuActions(URigVMBlueprint* RigVMBlueprint, const URigVMEdGraphNode* EdGraphNode, URigVMNode* ModelNode, UToolMenu* Menu) const;
+	virtual void GetNodeVariablesContextMenuActions(URigVMBlueprint* RigVMBlueprint, const URigVMEdGraphNode* EdGraphNode, URigVMNode* ModelNode, UToolMenu* Menu) const;
+	virtual void GetNodeTemplatesContextMenuActions(URigVMBlueprint* RigVMBlueprint, const URigVMEdGraphNode* EdGraphNode, URigVMNode* ModelNode, UToolMenu* Menu) const;
+	virtual void GetNodeOrganizationContextMenuActions(URigVMBlueprint* RigVMBlueprint, const URigVMEdGraphNode* EdGraphNode, URigVMNode* ModelNode, UToolMenu* Menu) const;
+	virtual void GetNodeVersioningContextMenuActions(URigVMBlueprint* RigVMBlueprint, const URigVMEdGraphNode* EdGraphNode, URigVMNode* ModelNode, UToolMenu* Menu) const;
+	virtual void GetNodeTestContextMenuActions(URigVMBlueprint* RigVMBlueprint, const URigVMEdGraphNode* EdGraphNode, URigVMNode* ModelNode, UToolMenu* Menu) const;
+	virtual void GetPinWorkflowContextMenuActions(URigVMBlueprint* RigVMBlueprint, const UEdGraphPin* EdGraphPin, URigVMPin* ModelPin, UToolMenu* Menu) const;
+	virtual void GetPinDebugContextMenuActions(URigVMBlueprint* RigVMBlueprint, const UEdGraphPin* EdGraphPin, URigVMPin* ModelPin, UToolMenu* Menu) const;
+	virtual void GetPinArrayContextMenuActions(URigVMBlueprint* RigVMBlueprint, const UEdGraphPin* EdGraphPin, URigVMPin* ModelPin, UToolMenu* Menu) const;
+	virtual void GetPinAggregateContextMenuActions(URigVMBlueprint* RigVMBlueprint, const UEdGraphPin* EdGraphPin, URigVMPin* ModelPin, UToolMenu* Menu) const;
+	virtual void GetPinTemplateContextMenuActions(URigVMBlueprint* RigVMBlueprint, const UEdGraphPin* EdGraphPin, URigVMPin* ModelPin, UToolMenu* Menu) const;
+	virtual void GetPinConversionContextMenuActions(URigVMBlueprint* RigVMBlueprint, const UEdGraphPin* EdGraphPin, URigVMPin* ModelPin, UToolMenu* Menu) const;
+	virtual void GetPinVariableContextMenuActions(URigVMBlueprint* RigVMBlueprint, const UEdGraphPin* EdGraphPin, URigVMPin* ModelPin, UToolMenu* Menu) const;
+	virtual void GetPinResetDefaultContextMenuActions(URigVMBlueprint* RigVMBlueprint, const UEdGraphPin* EdGraphPin, URigVMPin* ModelPin, UToolMenu* Menu) const;
+	virtual void GetPinInjectedNodesContextMenuActions(URigVMBlueprint* RigVMBlueprint, const UEdGraphPin* EdGraphPin, URigVMPin* ModelPin, UToolMenu* Menu) const;
+
+	/** Extensibility managers */
+	TSharedPtr<FExtensibilityManager> MenuExtensibilityManager;
+	TSharedPtr<FExtensibilityManager> ToolBarExtensibilityManager;
+	TArray<FRigVMEditorToolbarExtender> RigVMEditorToolbarExtenders;
+	
+	/** Node factory for the rigvm graph */
+	TSharedPtr<FRigVMEdGraphPanelNodeFactory> EdGraphPanelNodeFactory;
+
+	/** Pin factory for the rigvm graph */
+	TSharedPtr<FRigVMEdGraphPanelPinFactory> EdGraphPanelPinFactory;
+
+	/** Delegate handles for blueprint utils */
+	FDelegateHandle RefreshAllNodesDelegateHandle;
+	FDelegateHandle ReconstructAllNodesDelegateHandle;
+	FDelegateHandle BlueprintVariableCustomizationHandle;
+
+private:
+
+	void HandleNewBlueprintCreated(UBlueprint* InBlueprint);
+
+	bool ShowWorkflowOptionsDialog(URigVMUserWorkflowOptions* InOptions) const;
+
+	friend class URigVMEdGraphNode;
 };
\ No newline at end of file
diff --git a/Engine/Plugins/Runtime/RigVM/Source/RigVMEditor/RigVMEditor.Build.cs b/Engine/Plugins/Runtime/RigVM/Source/RigVMEditor/RigVMEditor.Build.cs
--- a/Engine/Plugins/Runtime/RigVM/Source/RigVMEditor/RigVMEditor.Build.cs
+++ b/Engine/Plugins/Runtime/RigVM/Source/RigVMEditor/RigVMEditor.Build.cs
@@ -22,7 +22,13 @@ public class RigVMEditor : ModuleRules
                 "InputCore",
                 "ToolWidgets",
                 "Kismet",
+                "KismetCompiler",
                 "KismetWidgets",
+                "AnimationWidgets",
+                "AppFramework",
+                "AnimationCore",
+                "PropertyEditor",
+                "ToolMenus",
             }
         );
     }
