RigVM: Fix pin-resolution behavior with already-resolved float and double pins

This fixes resolving a wildcard pin that is already set to double switching to float unexpectedly when UpdateTemplateNodePinTypes is called post-resolve

diff --git a/Engine/Plugins/Runtime/RigVM/Source/RigVMDeveloper/Private/RigVMModel/RigVMController.cpp b/Engine/Plugins/Runtime/RigVM/Source/RigVMDeveloper/Private/RigVMModel/RigVMController.cpp
--- a/Engine/Plugins/Runtime/RigVM/Source/RigVMDeveloper/Private/RigVMModel/RigVMController.cpp
+++ b/Engine/Plugins/Runtime/RigVM/Source/RigVMDeveloper/Private/RigVMModel/RigVMController.cpp
@@ -16907,7 +16907,7 @@ bool URigVMController::UpdateTemplateNodePinTypes(URigVMTemplateNode* InNode, bo
 			Template->Resolve(OutTypeMap, OutPermutations, false);
 			if (TRigVMTypeIndex* TypeIndex = OutTypeMap.Find(Pin->GetFName()))
 			{
-				if (*TypeIndex != INDEX_NONE)
+				if (*TypeIndex != INDEX_NONE && !Registry.IsWildCardType(*TypeIndex))
 				{
 					PreferredType = *TypeIndex;
 					TypesFoundInReduced = 1;
