Enable object pins in RigVM

diff --git a/Engine/Plugins/Runtime/RigVM/Source/RigVM/Private/RigVMCore/RigVMRegistry.cpp b/Engine/Plugins/Runtime/RigVM/Source/RigVM/Private/RigVMCore/RigVMRegistry.cpp
--- a/Engine/Plugins/Runtime/RigVM/Source/RigVM/Private/RigVMCore/RigVMRegistry.cpp
+++ b/Engine/Plugins/Runtime/RigVM/Source/RigVM/Private/RigVMCore/RigVMRegistry.cpp
@@ -352,6 +352,15 @@ void FRigVMRegistry::RefreshEngineTypes()
 		}
 	}
 	
+	for (TObjectIterator<UClass> ClassIt; ClassIt; ++ClassIt)
+	{
+		UClass* Class = *ClassIt;
+		if (IsAllowedType(Class))
+		{
+			FindOrAddType(FRigVMTemplateArgumentType(Class));
+		}
+	}
+
 	// Register all dispatch factories only after all other types have been registered.
 	for (UScriptStruct* DispatchFactoryStruct: DispatchFactoriesToRegister)
 	{
@@ -1047,7 +1056,7 @@ TRigVMTypeIndex FRigVMRegistry::GetBaseTypeFromArrayTypeIndex(TRigVMTypeIndex In
 	return INDEX_NONE;
 }
 
-bool FRigVMRegistry::IsAllowedType(const FProperty* InProperty)
+bool FRigVMRegistry::IsAllowedType(const FProperty* InProperty) const
 {
 	if(InProperty->IsA<FBoolProperty>() ||
 		InProperty->IsA<FUInt32Property>() ||
@@ -1091,12 +1100,12 @@ bool FRigVMRegistry::IsAllowedType(const FProperty* InProperty)
 	return false;
 }
 
-bool FRigVMRegistry::IsAllowedType(const UEnum* InEnum)
+bool FRigVMRegistry::IsAllowedType(const UEnum* InEnum) const
 {
 	return !InEnum->HasAnyFlags(DisallowedFlags()) && InEnum->HasAllFlags(NeededFlags());
 }
 
-bool FRigVMRegistry::IsAllowedType(const UStruct* InStruct)
+bool FRigVMRegistry::IsAllowedType(const UStruct* InStruct) const
 {
 	if(InStruct->HasAnyFlags(DisallowedFlags()) || !InStruct->HasAllFlags(NeededFlags()))
 	{
@@ -1128,18 +1137,21 @@ bool FRigVMRegistry::IsAllowedType(const UStruct* InStruct)
 	return true;
 }
 
-bool FRigVMRegistry::IsAllowedType(const UClass* InClass)
+bool FRigVMRegistry::IsAllowedType(const UClass* InClass) const
 {
 	if(!InClass || InClass->HasAnyClassFlags(CLASS_Hidden | CLASS_Abstract))
 	{
 		return false;
 	}
 
-	// note: currently we don't allow UObjects
-	return false;
-	//return IsAllowedType(Cast<UStruct>(InClass));
-}
+	// Only allow native object types
+	if (!InClass->HasAnyClassFlags(CLASS_Native))
+	{
+		return false;
+	}
 
+	return AllowedClasses.Contains(InClass);
+}
 
 void FRigVMRegistry::Register(const TCHAR* InName, FRigVMFunctionPtr InFunctionPtr, UScriptStruct* InStruct, const TArray<FRigVMFunctionArgument>& InArguments)
 {
@@ -1279,6 +1291,24 @@ void FRigVMRegistry::RegisterPredicate(UScriptStruct* InStruct, const TCHAR* InN
 	Predicates.Add(Function);
 }
 
+void FRigVMRegistry::RegisterObjectTypes(TConstArrayView<UClass*> InClasses)
+{
+	for (UClass* Class : InClasses)
+	{
+		// Only allow native object types
+		if (Class->HasAnyClassFlags(CLASS_Native))
+		{
+			// Add all child classes
+			TArray<UClass*> DerivedClasses({ Class });
+			GetDerivedClasses(Class, DerivedClasses, /*bRecursive=*/true);
+			for(UClass* DerivedClass : DerivedClasses)
+			{
+				AllowedClasses.Add(DerivedClass);
+			}
+		}
+	}
+}
+
 const FRigVMFunction* FRigVMRegistry::FindFunction(const TCHAR* InName) const
 {
 	FScopeLock FindFunctionScopeLock(&FindFunctionMutex);
diff --git a/Engine/Plugins/Runtime/RigVM/Source/RigVM/Public/RigVMCore/RigVMRegistry.h b/Engine/Plugins/Runtime/RigVM/Source/RigVM/Public/RigVMCore/RigVMRegistry.h
--- a/Engine/Plugins/Runtime/RigVM/Source/RigVM/Public/RigVMCore/RigVMRegistry.h
+++ b/Engine/Plugins/Runtime/RigVM/Source/RigVM/Public/RigVMCore/RigVMRegistry.h
@@ -62,6 +62,9 @@ public:
 	// Register a predicate contained in the input struct
 	void RegisterPredicate(UScriptStruct* InStruct, const TCHAR* InName, const TArray<FRigVMFunctionArgument>& InArguments);
 
+	// Register a set of allowed object types
+	void RegisterObjectTypes(TConstArrayView<UClass*> InClasses);
+
 	// Initializes the registry by storing the defaults
 	void InitializeIfNeeded();
 	
@@ -320,10 +323,10 @@ private:
 		return RF_Public;
 	}
 
-	static bool IsAllowedType(const FProperty* InProperty);
-	static bool IsAllowedType(const UEnum* InEnum);
-	static bool IsAllowedType(const UStruct* InStruct);
-	static bool IsAllowedType(const UClass* InClass);
+	bool IsAllowedType(const FProperty* InProperty) const;
+	bool IsAllowedType(const UEnum* InEnum) const;
+	bool IsAllowedType(const UStruct* InStruct) const;
+	bool IsAllowedType(const UClass* InClass) const;
 
 	void RegisterTypeInCategory(FRigVMTemplateArgument::ETypeCategory InCategory, TRigVMTypeIndex InTypeIndex);
 	
@@ -377,6 +380,9 @@ private:
 	// When that happens, it won't be safe to reload deleted assets so only type names are reliable
 	TMap<FSoftObjectPath, TRigVMTypeIndex> UserDefinedTypeToIndex;
 	
+	// All allowed classes
+	TSet<TObjectPtr<const UClass>> AllowedClasses;
+
 	// Notifies other system that types have been added/removed, and template permutations have been updated
 	FOnRigVMRegistryChanged OnRigVMRegistryChangedDelegate;
 
diff --git a/Engine/Plugins/Runtime/RigVM/Source/RigVM/Public/RigVMDefines.h b/Engine/Plugins/Runtime/RigVM/Source/RigVM/Public/RigVMDefines.h
--- a/Engine/Plugins/Runtime/RigVM/Source/RigVM/Public/RigVMDefines.h
+++ b/Engine/Plugins/Runtime/RigVM/Source/RigVM/Public/RigVMDefines.h
@@ -11,7 +11,7 @@
 #endif
 
 #ifndef UE_RIGVM_UOBJECT_PROPERTIES_ENABLED
-#define UE_RIGVM_UOBJECT_PROPERTIES_ENABLED 0
+#define UE_RIGVM_UOBJECT_PROPERTIES_ENABLED 1
 #endif
 
 #ifndef UE_RIGVM_UINTERFACE_PROPERTIES_ENABLED
diff --git a/Engine/Plugins/Runtime/RigVM/Source/RigVMDeveloper/Private/RigVMDeveloperTypeUtils.cpp b/Engine/Plugins/Runtime/RigVM/Source/RigVMDeveloper/Private/RigVMDeveloperTypeUtils.cpp
--- a/Engine/Plugins/Runtime/RigVM/Source/RigVMDeveloper/Private/RigVMDeveloperTypeUtils.cpp
+++ b/Engine/Plugins/Runtime/RigVM/Source/RigVMDeveloper/Private/RigVMDeveloperTypeUtils.cpp
@@ -100,6 +100,11 @@ FEdGraphPinType RigVMTypeUtils::PinTypeFromCPPType(const FName& InCPPType, UObje
 		PinType.PinCategory = UEdGraphSchema_K2::PC_Byte;
 		PinType.PinSubCategoryObject = InCPPTypeObject;
 	}
+	else if (Cast<UClass>(InCPPTypeObject))
+	{
+		PinType.PinCategory = UEdGraphSchema_K2::PC_Object;
+		PinType.PinSubCategoryObject = InCPPTypeObject;
+	}
 
 	return PinType;
 }
@@ -146,6 +151,11 @@ FEdGraphPinType RigVMTypeUtils::PinTypeFromExternalVariable(const FRigVMExternal
 		PinType.PinCategory = UEdGraphSchema_K2::PC_Byte;
 		PinType.PinSubCategoryObject = InExternalVariable.TypeObject;
 	}
+	else if (Cast<UClass>(InExternalVariable.TypeObject))
+	{
+		PinType.PinCategory = UEdGraphSchema_K2::PC_Object;
+		PinType.PinSubCategoryObject = InExternalVariable.TypeObject;
+	}
 
 	if (InExternalVariable.bIsArray)
 	{
@@ -497,7 +507,14 @@ bool RigVMTypeUtils::CPPTypeFromExternalVariable(const FRigVMExternalVariable& I
 	{
 		OutCPPType = Prefix + Enum->GetFName().ToString() + Suffix;
 		*OutCPPTypeObject = Enum;
-	}	
+	}
+	else if (UClass* Class = Cast<UClass>(InExternalVariable.TypeObject))
+	{
+		Prefix += TEXT("TObjectPtr<U");
+		Suffix += TEXT(">");
+		OutCPPType = Prefix + Class->GetFName().ToString() + Suffix;
+		*OutCPPTypeObject = Class;
+	}
 	else
 	{
 		check(false);
diff --git a/Engine/Source/Programs/Shared/EpicGames.UHT/Types/Properties/UhtObjectProperty.cs b/Engine/Source/Programs/Shared/EpicGames.UHT/Types/Properties/UhtObjectProperty.cs
--- a/Engine/Source/Programs/Shared/EpicGames.UHT/Types/Properties/UhtObjectProperty.cs
+++ b/Engine/Source/Programs/Shared/EpicGames.UHT/Types/Properties/UhtObjectProperty.cs
@@ -31,7 +31,7 @@ namespace EpicGames.UHT.Types
 		{
 			PropertyFlags |= extraFlags;
 			PropertyCaps |= UhtPropertyCaps.RequiresNullConstructorArg | UhtPropertyCaps.CanBeInstanced | UhtPropertyCaps.CanExposeOnSpawn |
-				UhtPropertyCaps.IsParameterSupportedByBlueprint | UhtPropertyCaps.IsMemberSupportedByBlueprint;
+				UhtPropertyCaps.IsParameterSupportedByBlueprint | UhtPropertyCaps.IsMemberSupportedByBlueprint | UhtPropertyCaps.SupportsRigVM;
 		}
 
 		/// <inheritdoc/>
