RigVM: Avoid folding for function arguments

diff --git a/Engine/Plugins/Runtime/RigVM/Source/RigVM/Private/RigVMCore/RigVMGraphFunctionDefinition.cpp b/Engine/Plugins/Runtime/RigVM/Source/RigVM/Private/RigVMCore/RigVMGraphFunctionDefinition.cpp
--- a/Engine/Plugins/Runtime/RigVM/Source/RigVM/Private/RigVMCore/RigVMGraphFunctionDefinition.cpp
+++ b/Engine/Plugins/Runtime/RigVM/Source/RigVM/Private/RigVMCore/RigVMGraphFunctionDefinition.cpp
@@ -2,6 +2,7 @@
 
 #include "RigVMCore/RigVMGraphFunctionDefinition.h"
 #include "RigVMCore/RigVMGraphFunctionHost.h"
+#include "RigVMStringUtils.h"
 
 #include UE_INLINE_GENERATED_CPP_BY_NAME(RigVMGraphFunctionDefinition)
 
@@ -31,6 +32,21 @@ bool FRigVMGraphFunctionArgument::IsCPPTypeObjectValid() const
 	return CPPTypeObject.IsValid();
 }
 
+bool FRigVMGraphFunctionArgument::IsExecuteContext() const
+{
+	if(IsCPPTypeObjectValid())
+	{
+		if(const UScriptStruct* ScriptStruct = Cast<UScriptStruct>(CPPTypeObject.Get()))
+		{
+			if(ScriptStruct->IsChildOf(FRigVMExecuteContext::StaticStruct()))
+			{
+				return true;
+			}
+		}
+	}
+	return false;
+}
+
 bool FRigVMGraphFunctionHeader::IsMutable() const
 {
 	for(const FRigVMGraphFunctionArgument& Arg : Arguments)
@@ -132,3 +148,215 @@ FRigVMGraphFunctionData* FRigVMGraphFunctionData::FindFunctionData(const FRigVMG
 	return FunctionStore->FindFunction(InIdentifier, bOutIsPublic);
 }
 
+
+FString FRigVMGraphFunctionData::GetArgumentNameFromPinHash(const FString& InPinHash)
+{
+	FString Left, PinPath, PinName;
+	if(RigVMStringUtils::SplitNodePathAtEnd(InPinHash, Left, PinPath))
+	{
+		Left.Reset();
+		if(RigVMStringUtils::SplitPinPathAtStart(PinPath, Left, PinName))
+		{
+			static FString EntryString = TEXT("Entry");
+			static FString ReturnString = TEXT("Return");
+			if(Left.Equals(EntryString, ESearchCase::CaseSensitive) ||
+				Left.Equals(ReturnString, ESearchCase::CaseSensitive))
+			{
+				return PinName;
+			}
+		}
+	}
+	return FString();
+}
+
+FRigVMOperand FRigVMGraphFunctionData::GetOperandForArgument(const FName& InArgumentName) const
+{
+	const FString InArgumentNameString = InArgumentName.ToString();
+	for(const TPair<FString, FRigVMOperand>& Pair : CompilationData.Operands)
+	{
+		const FString ArgumentName = GetArgumentNameFromPinHash(Pair.Key);
+		if(!ArgumentName.IsEmpty())
+		{
+			if(ArgumentName.Equals(InArgumentNameString, ESearchCase::CaseSensitive))
+			{
+				return Pair.Value;
+			}
+		}
+	}
+	return FRigVMOperand();
+}
+
+bool FRigVMGraphFunctionData::IsAnyOperandSharedAcrossArguments() const
+{
+	TSet<FRigVMOperand> UsedOperands;
+	UsedOperands.Reserve(Header.Arguments.Num());
+	for(const FRigVMGraphFunctionArgument& Argument : Header.Arguments)
+	{
+		if(Argument.IsExecuteContext())
+		{
+			continue;
+		}
+		
+		const FRigVMOperand Operand = GetOperandForArgument(Argument.Name);
+		check(Operand.IsValid());
+		if(UsedOperands.Contains(Operand))
+		{
+			return true;
+		}
+		UsedOperands.Add(Operand);
+	}
+	return false;
+}
+
+bool FRigVMGraphFunctionData::PatchSharedArgumentOperandsIfRequired()
+{
+	// we are doing this to avoid output arguments of a function to share
+	// memory. each output argument needs its own memory for the node
+	// referencing the function to rely on.
+	if(!IsAnyOperandSharedAcrossArguments())
+	{
+		return false;
+	}
+
+	// we'll keep doing this until there is no work left since
+	// we need to shift all operand indices every time we change anything.
+	bool bWorkLeft = true;
+	while(bWorkLeft)
+	{
+		bWorkLeft = false;
+		
+		// create a list of arguments / operands to update
+		TArray<TTuple<FName, FRigVMOperand>> ArgumentOperands;
+		TMap<FRigVMOperand, TArray<FName>> OperandToArguments;
+		ArgumentOperands.Reserve(Header.Arguments.Num());
+		for(const FRigVMGraphFunctionArgument& Argument : Header.Arguments)
+		{
+			if(Argument.IsExecuteContext())
+			{
+				continue;
+			}
+
+			const FRigVMOperand Operand = GetOperandForArgument(Argument.Name);
+			ArgumentOperands.Emplace(Argument.Name, Operand);
+			OperandToArguments.FindOrAdd(Operand).Add(Argument.Name);
+		}
+
+		// step 1: inject the properties and operands necessary to reflect the expected layout
+		FRigVMOperand SourceOperand, TargetOperand;
+		FString SourcePinPath, TargetPinPath;
+
+		int32 ArgumentIndex = -1;
+		for(int32 Index = 0; Index < Header.Arguments.Num(); Index++)
+		{
+			const FRigVMGraphFunctionArgument& Argument = Header.Arguments[Index];
+			if(Argument.IsExecuteContext())
+			{
+				continue;
+			}
+			ArgumentIndex++;
+			
+			SourceOperand = GetOperandForArgument(Argument.Name);
+			const TArray<FName>& ArgumentsSharingOperand = OperandToArguments.FindChecked(SourceOperand);
+			if(ArgumentsSharingOperand.Num() == 1 || ArgumentsSharingOperand[0].IsEqual(Argument.Name, ENameCase::CaseSensitive))
+			{
+				continue;
+			}
+
+			check(SourceOperand.GetMemoryType() == ERigVMMemoryType::Work);
+
+			// clone the property
+			FRigVMFunctionCompilationPropertyDescription PropertyDescription = CompilationData.WorkPropertyDescriptions[SourceOperand.GetRegisterIndex()];
+
+			for(const TPair<FString,FRigVMOperand>& Pair : CompilationData.Operands)
+			{
+				if(Pair.Value == SourceOperand)
+				{
+					SourcePinPath = Pair.Key;
+					break;
+				}
+			}
+			check(!SourcePinPath.IsEmpty());
+			FString CompleteNodePath, NodePathPrefix, NodeName, PinName;
+			verify(RigVMStringUtils::SplitPinPathAtEnd(SourcePinPath, CompleteNodePath, PinName));
+			verify(RigVMStringUtils::SplitNodePathAtEnd(CompleteNodePath, NodePathPrefix, NodeName));
+			if(Argument.Direction == ERigVMPinDirection::Input || Argument.Direction == ERigVMPinDirection::IO)
+			{
+				static const FString NewNodeName = TEXT("Entry");
+				CompleteNodePath = RigVMStringUtils::JoinNodePath(NodePathPrefix, NewNodeName);
+			}
+			else
+			{
+				static const FString NewNodeName = TEXT("Return");
+				CompleteNodePath = RigVMStringUtils::JoinNodePath(NodePathPrefix, NewNodeName);
+			}
+			TargetPinPath = RigVMStringUtils::JoinPinPath(CompleteNodePath, Argument.Name.ToString());
+			PropertyDescription.Name = FRigVMPropertyDescription::SanitizeName(*TargetPinPath);
+
+			int32 TargetIndex = ArgumentIndex;
+			if(CompilationData.WorkPropertyDescriptions.IsValidIndex(ArgumentIndex))
+			{
+				TargetIndex = CompilationData.WorkPropertyDescriptions.Insert(PropertyDescription, ArgumentIndex);
+			}
+			else
+			{
+				TargetIndex = CompilationData.WorkPropertyDescriptions.Add(PropertyDescription);
+			}
+
+			TargetOperand = FRigVMOperand(SourceOperand.GetMemoryType(), TargetIndex, SourceOperand.GetRegisterOffset());
+			bWorkLeft = true;
+			break;
+		}
+
+		if(!bWorkLeft)
+		{
+			return true;
+		}
+
+		auto UpdateOperand = [TargetOperand](FRigVMOperand& Operand)
+		{
+			if(Operand.GetMemoryType() == TargetOperand.GetMemoryType())
+			{
+				if(Operand.GetRegisterIndex() >= TargetOperand.GetRegisterIndex())
+				{
+					Operand = FRigVMOperand(Operand.GetMemoryType(), Operand.GetRegisterIndex() + 1, Operand.GetRegisterOffset());
+				}
+			}
+		};
+
+		// step 2: update the property paths
+		for(FRigVMFunctionCompilationPropertyPath& PropertyPath : CompilationData.WorkPropertyPathDescriptions)
+		{
+			if(PropertyPath.PropertyIndex != INDEX_NONE)
+			{
+				if(PropertyPath.PropertyIndex >= TargetOperand.GetRegisterIndex())
+				{
+					PropertyPath.PropertyIndex++;
+				}
+			}
+		}
+
+		// step 3: update the operands map
+		for(TPair<FString, FRigVMOperand>& Pair : CompilationData.Operands)
+		{
+			UpdateOperand(Pair.Value);
+		}
+		CompilationData.Operands.FindOrAdd(TargetPinPath) = TargetOperand;
+
+		// step 4: update the operands in the bytecode itself
+		const FRigVMInstructionArray Instructions = CompilationData.ByteCode.GetInstructions();
+		for(const FRigVMInstruction& Instruction : Instructions)
+		{
+			FRigVMOperandArray Operands = CompilationData.ByteCode.GetOperandsForOp(Instruction);
+			for(int32 Index = 0; Index < Operands.Num(); Index++)
+			{
+				FRigVMOperand* Operand = const_cast<FRigVMOperand*>(&Operands[Index]);
+				UpdateOperand(*Operand);
+			}
+		}
+
+		// step 5: add copy operations at the end of the byte code
+		CompilationData.ByteCode.AddCopyOp(SourceOperand, TargetOperand);
+	}
+	
+	return true;
+}
diff --git a/Engine/Plugins/Runtime/RigVM/Source/RigVM/Public/RigVMCore/RigVMGraphFunctionDefinition.h b/Engine/Plugins/Runtime/RigVM/Source/RigVM/Public/RigVMCore/RigVMGraphFunctionDefinition.h
--- a/Engine/Plugins/Runtime/RigVM/Source/RigVM/Public/RigVMCore/RigVMGraphFunctionDefinition.h
+++ b/Engine/Plugins/Runtime/RigVM/Source/RigVM/Public/RigVMCore/RigVMGraphFunctionDefinition.h
@@ -255,7 +255,10 @@ struct RIGVM_API FRigVMGraphFunctionArgument
 	FRigVMExternalVariable GetExternalVariable() const;
 
 	// validates and potentially loads the CPP Type Object
-	bool IsCPPTypeObjectValid() const; 
+	bool IsCPPTypeObjectValid() const;
+
+	// returns true if this argument is an execute context
+	bool IsExecuteContext() const;
 
 	friend uint32 GetTypeHash(const FRigVMGraphFunctionArgument& Argument)
 	{
@@ -456,4 +459,12 @@ struct RIGVM_API FRigVMGraphFunctionData
 	void PostDuplicateHost(const FString& InOldPathName, const FString& InNewPathName);
 
 	static FRigVMGraphFunctionData* FindFunctionData(const FRigVMGraphFunctionIdentifier& InIdentifier, bool* bOutIsPublic = nullptr);	
+
+	static FString GetArgumentNameFromPinHash(const FString& InPinHash);
+	
+	FRigVMOperand GetOperandForArgument(const FName& InArgumentName) const;
+
+	bool IsAnyOperandSharedAcrossArguments() const;
+
+	bool PatchSharedArgumentOperandsIfRequired();
 };
diff --git a/Engine/Plugins/Runtime/RigVM/Source/RigVM/Public/RigVMCore/RigVMGraphFunctionHost.h b/Engine/Plugins/Runtime/RigVM/Source/RigVM/Public/RigVMCore/RigVMGraphFunctionHost.h
--- a/Engine/Plugins/Runtime/RigVM/Source/RigVM/Public/RigVMCore/RigVMGraphFunctionHost.h
+++ b/Engine/Plugins/Runtime/RigVM/Source/RigVM/Public/RigVMCore/RigVMGraphFunctionHost.h
@@ -80,6 +80,15 @@ public:
 		{
 			Ar << Host.PrivateFunctions;
 		}
+
+		if(Ar.IsLoading())
+		{
+			for(FRigVMGraphFunctionData& Function : Host.PublicFunctions)
+			{
+				Function.PatchSharedArgumentOperandsIfRequired();
+			}
+		}
+		
 		return Ar;
 	}
 
diff --git a/Engine/Plugins/Runtime/RigVM/Source/RigVMDeveloper/Private/RigVMCompiler/RigVMCompiler.cpp b/Engine/Plugins/Runtime/RigVM/Source/RigVMDeveloper/Private/RigVMCompiler/RigVMCompiler.cpp
--- a/Engine/Plugins/Runtime/RigVM/Source/RigVMDeveloper/Private/RigVMCompiler/RigVMCompiler.cpp
+++ b/Engine/Plugins/Runtime/RigVM/Source/RigVMDeveloper/Private/RigVMCompiler/RigVMCompiler.cpp
@@ -401,6 +401,7 @@ bool URigVMCompiler::Compile(TArray<URigVMGraph*> InGraphs, URigVMController* In
 							}
 							if (bSuccessfullCompilation || CompilationData->IsValid())
 							{
+								FunctionData->PatchSharedArgumentOperandsIfRequired();
 								CurrentCompiledFunctions.Add(FunctionHash, CompilationData);
 							}
 							else
@@ -3373,7 +3374,6 @@ const FRigVMCompilerWorkData::FRigVMASTProxyArray& URigVMCompiler::FindProxiesWi
 
 	const FString CPPType = InProxy.GetSubjectChecked<URigVMPin>()->GetCPPType();
 
-	bool bContainsInterfacePin = false;
 	for(int32 ProxyIndex = 0; ProxyIndex < PinProxiesToProcess.Num(); ProxyIndex++)
 	{
 		if (PinProxiesToProcess[ProxyIndex].IsValid())
@@ -3399,22 +3399,6 @@ const FRigVMCompilerWorkData::FRigVMASTProxyArray& URigVMCompiler::FindProxiesWi
 				{
 					continue;
 				}
-
-				// when compiling a function, avoid folding pins from entry and return nodes
-				if (CurrentCompilationFunction)
-				{
-					if (URigVMNode* Node = Pin->GetNode())
-					{
-						if (Node->IsA<URigVMFunctionEntryNode>() || Node->IsA<URigVMFunctionReturnNode>())
-						{
-							if (bContainsInterfacePin)
-							{
-								continue;
-							}
-							bContainsInterfacePin = true;
-						}
-					}
-				}
 			}
 			PinProxies.Add(PinProxiesToProcess[ProxyIndex]);
 		}
