RigVM: Fail softly with FRigVMGraphFunctionData::PatchSharedArgumentOperandsIfRequired

diff --git a/Engine/Plugins/Runtime/RigVM/Source/RigVM/Private/RigVMCore/RigVMGraphFunctionDefinition.cpp b/Engine/Plugins/Runtime/RigVM/Source/RigVM/Private/RigVMCore/RigVMGraphFunctionDefinition.cpp
--- a/Engine/Plugins/Runtime/RigVM/Source/RigVM/Private/RigVMCore/RigVMGraphFunctionDefinition.cpp
+++ b/Engine/Plugins/Runtime/RigVM/Source/RigVM/Private/RigVMCore/RigVMGraphFunctionDefinition.cpp
@@ -198,7 +198,11 @@ bool FRigVMGraphFunctionData::IsAnyOperandSharedAcrossArguments() const
 		}
 		
 		const FRigVMOperand Operand = GetOperandForArgument(Argument.Name);
-		check(Operand.IsValid());
+		if(!Operand.IsValid())
+		{
+			continue;
+		}
+		
 		if(UsedOperands.Contains(Operand))
 		{
 			return true;
@@ -237,8 +241,11 @@ bool FRigVMGraphFunctionData::PatchSharedArgumentOperandsIfRequired()
 			}
 
 			const FRigVMOperand Operand = GetOperandForArgument(Argument.Name);
-			ArgumentOperands.Emplace(Argument.Name, Operand);
-			OperandToArguments.FindOrAdd(Operand).Add(Argument.Name);
+			if(Operand.IsValid())
+			{
+				ArgumentOperands.Emplace(Argument.Name, Operand);
+				OperandToArguments.FindOrAdd(Operand).Add(Argument.Name);
+			}
 		}
 
 		// step 1: inject the properties and operands necessary to reflect the expected layout
@@ -256,6 +263,12 @@ bool FRigVMGraphFunctionData::PatchSharedArgumentOperandsIfRequired()
 			ArgumentIndex++;
 			
 			SourceOperand = GetOperandForArgument(Argument.Name);
+			if(!SourceOperand.IsValid())
+			{
+				SourceOperand = FRigVMOperand();
+				continue;
+			}
+			
 			const TArray<FName>& ArgumentsSharingOperand = OperandToArguments.FindChecked(SourceOperand);
 			if(ArgumentsSharingOperand.Num() == 1 || ArgumentsSharingOperand[0].IsEqual(Argument.Name, ENameCase::CaseSensitive))
 			{
