RigVM Compiler : Fix for invalid remapping of inline function parameters

diff --git a/Engine/Plugins/Runtime/RigVM/Source/RigVM/Private/RigVMCore/RigVMGraphFunctionDefinition.cpp b/Engine/Plugins/Runtime/RigVM/Source/RigVM/Private/RigVMCore/RigVMGraphFunctionDefinition.cpp
--- a/Engine/Plugins/Runtime/RigVM/Source/RigVM/Private/RigVMCore/RigVMGraphFunctionDefinition.cpp
+++ b/Engine/Plugins/Runtime/RigVM/Source/RigVM/Private/RigVMCore/RigVMGraphFunctionDefinition.cpp
@@ -6,6 +6,9 @@
 
 #include UE_INLINE_GENERATED_CPP_BY_NAME(RigVMGraphFunctionDefinition)
 
+const FString FRigVMGraphFunctionData::EntryString = TEXT("Entry");
+const FString FRigVMGraphFunctionData::ReturnString = TEXT("Return");
+
 FRigVMExternalVariable FRigVMGraphFunctionArgument::GetExternalVariable() const
 {
 	FRigVMExternalVariable Variable;
@@ -157,8 +160,6 @@ FString FRigVMGraphFunctionData::GetArgumentNameFromPinHash(const FString& InPin
 		Left.Reset();
 		if(RigVMStringUtils::SplitPinPathAtStart(PinPath, Left, PinName))
 		{
-			static FString EntryString = TEXT("Entry");
-			static FString ReturnString = TEXT("Return");
 			if(Left.Equals(EntryString, ESearchCase::CaseSensitive) ||
 				Left.Equals(ReturnString, ESearchCase::CaseSensitive))
 			{
diff --git a/Engine/Plugins/Runtime/RigVM/Source/RigVM/Public/RigVMCore/RigVMGraphFunctionDefinition.h b/Engine/Plugins/Runtime/RigVM/Source/RigVM/Public/RigVMCore/RigVMGraphFunctionDefinition.h
--- a/Engine/Plugins/Runtime/RigVM/Source/RigVM/Public/RigVMCore/RigVMGraphFunctionDefinition.h
+++ b/Engine/Plugins/Runtime/RigVM/Source/RigVM/Public/RigVMCore/RigVMGraphFunctionDefinition.h
@@ -467,4 +467,7 @@ struct RIGVM_API FRigVMGraphFunctionData
 	bool IsAnyOperandSharedAcrossArguments() const;
 
 	bool PatchSharedArgumentOperandsIfRequired();
+
+	static const FString EntryString;
+	static const FString ReturnString;
 };
diff --git a/Engine/Plugins/Runtime/RigVM/Source/RigVMDeveloper/Private/RigVMCompiler/RigVMCompiler.cpp b/Engine/Plugins/Runtime/RigVM/Source/RigVMDeveloper/Private/RigVMCompiler/RigVMCompiler.cpp
--- a/Engine/Plugins/Runtime/RigVM/Source/RigVMDeveloper/Private/RigVMCompiler/RigVMCompiler.cpp
+++ b/Engine/Plugins/Runtime/RigVM/Source/RigVMDeveloper/Private/RigVMCompiler/RigVMCompiler.cpp
@@ -2061,7 +2061,8 @@ int32 URigVMCompiler::TraverseInlineFunction(const FRigVMInlineFunctionExprAST*
 	}
 	else
 	{
-		TArray<FRigVMOperand> Operands;		
+		TArray<FRigVMOperand> Operands;
+		TArray<FString> OperandsPinNames;
 		for(const URigVMPin* Pin : FunctionReferenceNode->GetPins())
 		{
 			const FRigVMExprAST* ChildExpr = InExpr->FindExprWithPinName(Pin->GetFName());
@@ -2072,6 +2073,7 @@ int32 URigVMCompiler::TraverseInlineFunction(const FRigVMInlineFunctionExprAST*
 				if(!SourceVarExpr->IsExecuteContext())
 				{
 					Operands.Add(WorkData.ExprToOperand.FindChecked(SourceVarExpr));
+					OperandsPinNames.Add(GetPinNameWithDirectionPrefix(Pin));
 				}
 			}
 			else if (ChildExpr->IsA(FRigVMExprAST::EType::Var))
@@ -2080,6 +2082,7 @@ int32 URigVMCompiler::TraverseInlineFunction(const FRigVMInlineFunctionExprAST*
 				if(!SourceVarExpr->IsExecuteContext())
 				{
 					Operands.Add(WorkData.ExprToOperand.FindChecked(SourceVarExpr));
+					OperandsPinNames.Add(GetPinNameWithDirectionPrefix(Pin));
 				}
 			}
 			else
@@ -2143,12 +2146,12 @@ int32 URigVMCompiler::TraverseInlineFunction(const FRigVMInlineFunctionExprAST*
 					const FRigVMOperand& OuterOperand = WorkData.PinPathToOperand->FindChecked(FString::Printf(TEXT("Variable::%s"), *OuterVariableName.ToString()));
 					Operand->RegisterIndex = OuterOperand.RegisterIndex;
 				}
-				// Operand is an interface pin: replace the index and memory type
-				else if (Operand->GetMemoryType() == ERigVMMemoryType::Work &&
-					Operands.IsValidIndex(Operand->RegisterIndex))
+				// If Operand is an interface pin: replace the index and memory type
+				else if (const int32 FunctionInterfaceParameterIndex = GetOperandFunctionInterfaceParameterIndex(OperandsPinNames, FunctionCompilationData, *Operand); FunctionInterfaceParameterIndex != INDEX_NONE)
 				{
-					Operand->MemoryType = Operands[Operand->GetRegisterIndex()].MemoryType;
-					Operand->RegisterIndex = Operands[Operand->GetRegisterIndex()].RegisterIndex;
+					const FRigVMOperand& InterfaceOperand = Operands[FunctionInterfaceParameterIndex];
+					Operand->MemoryType = InterfaceOperand.GetMemoryType();
+					Operand->RegisterIndex = InterfaceOperand.GetRegisterIndex();
 				}
 				else
 				{
@@ -3447,6 +3450,37 @@ const FRigVMCompilerWorkData::FRigVMASTProxyArray& URigVMCompiler::FindProxiesWi
 	return WorkData.CachedProxiesWithSharedOperand.Add(InProxy, PinProxies);
 }
 
+FString URigVMCompiler::GetPinNameWithDirectionPrefix(const URigVMPin* Pin)
+{
+	const FString & Prefix = (Pin->GetDirection() == ERigVMPinDirection::Input) ? FRigVMGraphFunctionData::EntryString : (Pin->GetDirection() == ERigVMPinDirection::Output) ? FRigVMGraphFunctionData::ReturnString : "";
+	const FString NameWithDirectionPrefix = Prefix + "_" + FRigVMPropertyDescription::SanitizeName(FName(Pin->GetName())).ToString();
+	return NameWithDirectionPrefix;
+}
+
+int32 URigVMCompiler::GetOperandFunctionInterfaceParameterIndex(const TArray<FString>& OperandsPinNames, const FRigVMFunctionCompilationData* FunctionCompilationData, const FRigVMOperand& Operand)
+{
+	const FRigVMFunctionCompilationPropertyDescription& CompilationPropertyDescription = (Operand.GetMemoryType() == ERigVMMemoryType::Work)
+		? FunctionCompilationData->WorkPropertyDescriptions[Operand.GetRegisterIndex()]
+		: FunctionCompilationData->LiteralPropertyDescriptions[Operand.GetRegisterIndex()];
+
+	const FString PropertyName = FRigVMPropertyDescription::SanitizeName(CompilationPropertyDescription.Name).ToString();
+	const int32 NumParams = OperandsPinNames.Num();
+	for (int32 ParamIndex = 0; ParamIndex < NumParams; ParamIndex++)
+	{
+		const FString& InterfacePinName = OperandsPinNames[ParamIndex];
+		if (const int32 SubStrStart = PropertyName.Find(InterfacePinName, ESearchCase::CaseSensitive, ESearchDir::FromEnd); SubStrStart != -1)
+		{
+			const FString Name = PropertyName.RightChop(SubStrStart);  // with this, we make sure that we don't return Argument_1 when searching for Argument
+			if (Name.Equals(InterfacePinName, ESearchCase::CaseSensitive))
+			{
+				return ParamIndex;
+			}
+		}
+	}
+
+	return INDEX_NONE;
+}
+
 bool URigVMCompiler::ValidateNode(URigVMNode* InNode, bool bCheck)
 {
 	if(bCheck)
diff --git a/Engine/Plugins/Runtime/RigVM/Source/RigVMDeveloper/Public/RigVMCompiler/RigVMCompiler.h b/Engine/Plugins/Runtime/RigVM/Source/RigVMDeveloper/Public/RigVMCompiler/RigVMCompiler.h
--- a/Engine/Plugins/Runtime/RigVM/Source/RigVMDeveloper/Public/RigVMCompiler/RigVMCompiler.h
+++ b/Engine/Plugins/Runtime/RigVM/Source/RigVMDeveloper/Public/RigVMCompiler/RigVMCompiler.h
@@ -242,6 +242,9 @@ private:
 	FRigVMOperand FindOrAddRegister(const FRigVMVarExprAST* InVarExpr, FRigVMCompilerWorkData& WorkData, bool bIsDebugValue = false);
 	const FRigVMCompilerWorkData::FRigVMASTProxyArray& FindProxiesWithSharedOperand(const FRigVMVarExprAST* InVarExpr, FRigVMCompilerWorkData& WorkData);
 
+	static FString GetPinNameWithDirectionPrefix(const URigVMPin* Pin);
+	static int32 GetOperandFunctionInterfaceParameterIndex(const TArray<FString>& OperandsPinNames, const FRigVMFunctionCompilationData* FunctionCompilationData, const FRigVMOperand& Operand);
+
 	bool ValidateNode(URigVMNode* InNode, bool bCheck = true);
 	
 	void ReportInfo(const FString& InMessage);
