RigVM: Move most of the functionality from ControlRigEditor to RigVMEditor

diff --git a/Engine/Plugins/Runtime/RigVM/Source/RigVMEditor/Private/Editor/RigVMEditor.cpp b/Engine/Plugins/Runtime/RigVM/Source/RigVMEditor/Private/Editor/RigVMEditor.cpp
--- a/Engine/Plugins/Runtime/RigVM/Source/RigVMEditor/Private/Editor/RigVMEditor.cpp
+++ b/Engine/Plugins/Runtime/RigVM/Source/RigVMEditor/Private/Editor/RigVMEditor.cpp
@@ -1,41 +1,3548 @@
 // Copyright Epic Games, Inc. All Rights Reserved.
 
 #include "Editor/RigVMEditor.h"
-#include "Editor/RigVMDetailsViewWrapperObject.h"
+#include "Editor/RigVMEditorMenuContext.h"
+#include "HAL/PlatformApplicationMisc.h"
+#include "SMyBlueprint.h"
+#include "EdGraphNode_Comment.h"
+#include "RigVMFunctions/RigVMFunction_ControlFlow.h"
+#include "Editor.h"
+#include "Editor/Transactor.h"
+#include "BlueprintEditorTabs.h"
+#include "Kismet2/BlueprintEditorUtils.h"
+#include "Widgets/Notifications/SNotificationList.h"
+#include "Framework/Notifications/NotificationManager.h"
+#include "Editor/RigVMEditorCommands.h"
+#include "BlueprintCompilationManager.h"
+#include "SBlueprintEditorToolbar.h"
+#include "BlueprintEditorModes.h"
+#include "AssetEditorModeManager.h"
+#include "RigVMBlueprintUtils.h"
+#include "RigVMPythonUtils.h"
+#include "EulerTransform.h"
+#include "RigVMEditorModule.h"
+#include "Editor/RigVMGraphDetailCustomization.h"
+#include "Kismet2/KismetDebugUtilities.h"
+#include "Kismet2/WatchedPin.h"
+#include "SBlueprintEditorSelectedDebugObjectWidget.h"
+#include "IMessageLogListing.h"
+#include "SGraphPanel.h"
+#include "BlueprintActionDatabase.h"
+#include "RigVMModel/Nodes/RigVMAggregateNode.h"
+#include "Widgets/SRigVMGraphFunctionLocalizationWidget.h"
+#include "Widgets/SRigVMGraphFunctionBulkEditWidget.h"
+#include "Widgets/SRigVMGraphBreakLinksWidget.h"
+#include "Widgets/SRigVMGraphChangePinType.h"
+#include "Framework/Commands/GenericCommands.h"
+#include "Stats/StatsHierarchical.h"
+#include "Widgets/Layout/SBorder.h"
+#include "Styling/AppStyle.h"
+#include "Widgets/Docking/SDockTab.h"
+#include "Widgets/Layout/SScrollBox.h"
+#include "Framework/Docking/TabManager.h"
+#include "ScopedTransaction.h"
 
 #define LOCTEXT_NAMESPACE "RigVMEditor"
 
 FRigVMEditor::FRigVMEditor()
 	: Host(nullptr)
+	, bAnyErrorsLeft(false)
+	, KnownInstructionLimitWarnings()
+	, HaltedAtNode(nullptr)
+	, LastDebuggedHost()
+	, bSuspendDetailsPanelRefresh(false)
+	, bAllowBulkEdits(false)
+	, bIsSettingObjectBeingDebugged(false)
+	, bRigVMEditorInitialized(false)
+	, bIsCompilingThroughUI(false)
+	, WrapperObjects()
+	, ExecutionMode(ERigVMEditorExecutionModeType_Release)
+	, LastEventQueue()
+
+{
+}
+
+FRigVMEditor::~FRigVMEditor()
+{
+	URigVMBlueprint* RigVMBlueprint = GetRigVMBlueprint();
+	RigVMEditorClosedDelegate.Broadcast(this, RigVMBlueprint);
+
+	ClearDetailObject();
+
+	if(PropertyChangedHandle.IsValid())
+	{
+		FCoreUObjectDelegates::OnObjectPropertyChanged.Remove(PropertyChangedHandle);
+	}
+
+	if (RigVMBlueprint)
+	{
+		// clear editor related data from the debugged control rig instance 
+		RigVMBlueprint->SetObjectBeingDebugged(nullptr);
+
+		URigVMBlueprint::sCurrentlyOpenedRigVMBlueprints.Remove(RigVMBlueprint);
+
+		RigVMBlueprint->OnRefreshEditor().RemoveAll(this);
+		RigVMBlueprint->OnVariableDropped().RemoveAll(this);
+		RigVMBlueprint->OnBreakpointAdded().RemoveAll(this);
+		RigVMBlueprint->OnNodeDoubleClicked().RemoveAll(this);
+		RigVMBlueprint->OnGraphImported().RemoveAll(this);
+		RigVMBlueprint->OnRequestLocalizeFunctionDialog().RemoveAll(this);
+		RigVMBlueprint->OnRequestBulkEditDialog().Unbind();
+		RigVMBlueprint->OnRequestBreakLinksDialog().Unbind();
+		RigVMBlueprint->OnRequestPinTypeSelectionDialog().Unbind();
+		RigVMBlueprint->OnRequestJumpToHyperlink().Unbind();
+		RigVMBlueprint->OnReportCompilerMessage().RemoveAll(this);
+
+#if WITH_EDITOR
+		RigVMBlueprint->SetDebugMode(false);
+		RigVMBlueprint->ClearBreakpoints();
+		SetHaltedNode(nullptr);
+		RigVMBlueprint->OnGetFocusedGraph().Unbind();
+#endif
+	}
+}
+
+void FRigVMEditor::InitRigVMEditor(const EToolkitMode::Type Mode, const TSharedPtr< class IToolkitHost >& InitToolkitHost, class URigVMBlueprint* InRigVMBlueprint)
+{
+	DECLARE_SCOPE_HIERARCHICAL_COUNTER_FUNC()
+
+	check(InRigVMBlueprint);
+
+	FBlueprintCompilationManager::FlushCompilationQueue(nullptr);
+
+	Toolbox = SNew(SBorder)
+		.BorderImage(FAppStyle::GetBrush("ToolPanel.GroupBorder"))
+		.Padding(0.f);
+
+	if (!Toolbar.IsValid())
+	{
+		Toolbar = MakeShareable(new FBlueprintEditorToolbar(SharedThis(this)));
+	}
+
+	// Build up a list of objects being edited in this asset editor
+	TArray<UObject*> ObjectsBeingEdited;
+	ObjectsBeingEdited.Add(InRigVMBlueprint);
+
+	// Initialize the asset editor and spawn tabs
+	const bool bCreateDefaultStandaloneMenu = true;
+	const bool bCreateDefaultToolbar = true;
+	InitAssetEditor(Mode, InitToolkitHost, GetEditorAppName(), FTabManager::FLayout::NullLayout, bCreateDefaultStandaloneMenu, bCreateDefaultToolbar, ObjectsBeingEdited);
+
+	CreateDefaultCommands();
+
+	TArray<UBlueprint*> Blueprints;
+	Blueprints.Add(InRigVMBlueprint);
+	InRigVMBlueprint->InitializeModelIfRequired();
+
+	CommonInitialization(Blueprints, false);
+	
+	// user-defined-struct can change even after load
+	// refresh the models such that pins are updated to match
+	// the latest struct member layout
+	InRigVMBlueprint->RefreshAllModels(ERigVMBlueprintLoadType::CheckUserDefinedStructs);
+
+	{
+		TArray<UEdGraph*> EdGraphs;
+		InRigVMBlueprint->GetAllGraphs(EdGraphs);
+
+		for (UEdGraph* Graph : EdGraphs)
+		{
+			URigVMEdGraph* RigVMEdGraph = Cast<URigVMEdGraph>(Graph);
+			if (RigVMEdGraph == nullptr)
+			{
+				continue;
+			}
+
+			RigVMEdGraph->InitializeFromBlueprint(InRigVMBlueprint);
+		}
+
+	}
+
+	URigVMBlueprint::sCurrentlyOpenedRigVMBlueprints.AddUnique(InRigVMBlueprint);
+
+	InRigVMBlueprint->OnModified().AddSP(this, &FRigVMEditor::HandleModifiedEvent);
+	InRigVMBlueprint->OnVMCompiled().AddSP(this, &FRigVMEditor::HandleVMCompiledEvent);
+	InRigVMBlueprint->OnRequestInspectObject().AddSP(this, &FRigVMEditor::SetDetailObjects);
+
+	BindCommands();
+
+	TSharedPtr<FApplicationMode> ApplicationMode = CreateEditorMode();
+	if(ApplicationMode.IsValid())
+	{
+		AddApplicationMode(
+			GetEditorModeName(),
+			ApplicationMode.ToSharedRef()
+		);
+	}
+
+	ExtendMenu();
+	ExtendToolbar();
+	RegenerateMenusAndToolbars();
+
+	if(ApplicationMode.IsValid())
+	{
+		// Activate the initial mode (which will populate with a real layout)
+		SetCurrentMode(GetEditorModeName());
+
+		// Activate our edit mode
+		GetEditorModeManager().SetDefaultMode(GetEditorModeName());
+		GetEditorModeManager().ActivateMode(GetEditorModeName());
+	}
+
+	UpdateRigVMHost();
+
+	// Post-layout initialization
+	PostLayoutBlueprintEditorInitialization();
+
+	// tabs opened before reload
+	FString ActiveTabNodePath;
+	TArray<FString> OpenedTabNodePaths;
+
+	if (Blueprints.Num() > 0)
+	{
+		bool bBroughtGraphToFront = false;
+		for(UEdGraph* Graph : Blueprints[0]->UbergraphPages)
+		{
+			if (URigVMEdGraph* RigGraph = Cast<URigVMEdGraph>(Graph))
+			{
+				if (!bBroughtGraphToFront)
+				{
+					OpenGraphAndBringToFront(Graph, false);
+					bBroughtGraphToFront = true;
+				}
+
+				RigGraph->OnGraphNodeClicked.AddSP(this, &FRigVMEditor::OnGraphNodeClicked);
+				ActiveTabNodePath = RigGraph->ModelNodePath;
+			}
+		}
+	}
+
+	{
+		if (URigVMGraph* Model = InRigVMBlueprint->GetDefaultModel())
+		{
+			if (Model->GetNodes().Num() == 0)
+			{
+				CreateEmptyGraphContent(InRigVMBlueprint->GetController());
+			}
+			else
+			{
+				// remember all ed graphs which were visible as tabs
+				TArray<UEdGraph*> EdGraphs;
+				InRigVMBlueprint->GetAllGraphs(EdGraphs);
+
+				for (UEdGraph* EdGraph : EdGraphs)
+				{
+					if (URigVMEdGraph* RigVMEdGraph = Cast<URigVMEdGraph>(EdGraph))
+					{
+						TArray<TSharedPtr<SDockTab>> TabsForEdGraph;
+						FindOpenTabsContainingDocument(EdGraph, TabsForEdGraph);
+
+						if (TabsForEdGraph.Num() > 0)
+						{
+							OpenedTabNodePaths.Add(RigVMEdGraph->ModelNodePath);
+
+							if(RigVMEdGraph->bIsFunctionDefinition)
+							{
+								CloseDocumentTab(RigVMEdGraph);
+							}
+						}
+					}
+				}
+
+				InRigVMBlueprint->RebuildGraphFromModel();
+
+				// selection state does not need to be persistent, even though it is saved in the RigVM.
+				for (URigVMGraph* Graph : InRigVMBlueprint->GetAllModels())
+				{
+					InRigVMBlueprint->GetController(Graph)->ClearNodeSelection(false);
+				}
+
+				if (UPackage* Package = InRigVMBlueprint->GetOutermost())
+				{
+					Package->SetDirtyFlag(InRigVMBlueprint->IsMarkedDirtyDuringLoad());
+				}
+			}
+		}
+
+		InRigVMBlueprint->OnRefreshEditor().AddSP(this, &FRigVMEditor::HandleRefreshEditorFromBlueprint);
+		InRigVMBlueprint->OnVariableDropped().AddSP(this, &FRigVMEditor::HandleVariableDroppedFromBlueprint);
+		InRigVMBlueprint->OnBreakpointAdded().AddSP(this, &FRigVMEditor::HandleBreakpointAdded);
+
+		InRigVMBlueprint->OnNodeDoubleClicked().AddSP(this, &FRigVMEditor::OnNodeDoubleClicked);
+		InRigVMBlueprint->OnGraphImported().AddSP(this, &FRigVMEditor::OnGraphImported);
+		InRigVMBlueprint->OnRequestLocalizeFunctionDialog().AddSP(this, &FRigVMEditor::OnRequestLocalizeFunctionDialog);
+		InRigVMBlueprint->OnRequestBulkEditDialog().BindSP(this, &FRigVMEditor::OnRequestBulkEditDialog);
+		InRigVMBlueprint->OnRequestBreakLinksDialog().BindSP(this, &FRigVMEditor::OnRequestBreakLinksDialog);
+		InRigVMBlueprint->OnRequestPinTypeSelectionDialog().BindSP(this, &FRigVMEditor::OnRequestPinTypeSelectionDialog);
+		InRigVMBlueprint->OnRequestJumpToHyperlink().BindSP(this, &FRigVMEditor::HandleJumpToHyperlink);
+#if WITH_EDITOR
+		InRigVMBlueprint->OnGetFocusedGraph().BindSP(this, &FRigVMEditor::GetFocusedModel);
+#endif
+	}
+
+	for (const FString& OpenedTabNodePath : OpenedTabNodePaths)
+	{
+		if (UEdGraph* EdGraph = InRigVMBlueprint->GetEdGraph(OpenedTabNodePath))
+		{
+			OpenDocument(EdGraph, FDocumentTracker::RestorePreviousDocument);
+		}
+	}
+
+	if (UEdGraph* ActiveGraph = InRigVMBlueprint->GetEdGraph(ActiveTabNodePath))
+	{
+		OpenGraphAndBringToFront(ActiveGraph, true);
+	}
+
+	FRigVMBlueprintUtils::HandleRefreshAllNodes(InRigVMBlueprint);
+
+	if (Blueprints.Num() > 0)
+	{
+		if(Blueprints[0]->Status == BS_Error)
+		{
+			Compile();
+		}
+	}
+	
+	FFunctionGraphTask::CreateAndDispatchWhenReady([this, InRigVMBlueprint]()
+	{
+		// no need to do anything if the the CR is not opened anymore
+		// (i.e. destructor has been called before that task actually got a chance to start)
+		if (!URigVMBlueprint::sCurrentlyOpenedRigVMBlueprints.Contains(InRigVMBlueprint))
+		{
+			return;		
+		}
+		
+		TSharedPtr<FTabManager> TabManager = GetTabManager();
+		if (!TabManager.IsValid())
+		{
+			return;
+		}
+		
+		// Always show the myblueprint tab
+		static const FTabId MyBlueprintTabId(FBlueprintEditorTabs::MyBlueprintID);
+		if (!TabManager->FindExistingLiveTab(MyBlueprintTabId).IsValid())
+		{
+			TabManager->TryInvokeTab(MyBlueprintTabId);
+		}
+		
+	}, TStatId(), nullptr, ENamedThreads::GameThread);
+
+	bRigVMEditorInitialized = true;
+	UpdateStaleWatchedPins();
+
+#if WITH_EDITOR
+	FString BlueprintName = InRigVMBlueprint->GetPathName();
+	RigVMPythonUtils::PrintPythonContext(BlueprintName);
+#endif
+
+	TArray<UScriptStruct*> StructsToCustomize = {
+		TBaseStructure<FVector>::Get(),
+		TBaseStructure<FVector2D>::Get(),
+		TBaseStructure<FVector4>::Get(),
+		TBaseStructure<FRotator>::Get(),
+		TBaseStructure<FQuat>::Get(),
+		TBaseStructure<FTransform>::Get(),
+		TBaseStructure<FEulerTransform>::Get(),
+	};
+
+	for(UScriptStruct* StructToCustomize : StructsToCustomize)
+	{
+		Inspector->GetPropertyView()->RegisterInstancedCustomPropertyTypeLayout(StructToCustomize->GetFName(),
+			FOnGetPropertyTypeCustomizationInstance::CreateLambda([=]()
+			{
+				return FRigVMGraphMathTypeDetailCustomization::MakeInstance();
+			}));
+	}
+
+	PropertyChangedHandle = FCoreUObjectDelegates::OnObjectPropertyChanged.AddSP(this, &FRigVMEditor::OnPropertyChanged);
+}
+
+const FName FRigVMEditor::GetEditorAppName() const
+{
+	static const FName AppName(TEXT("RigVMEditorApp"));
+	return AppName;
+}
+
+const FName FRigVMEditor::GetEditorModeName() const
+{
+	const FName EditorMode("VisualProgramming");
+	return EditorMode;
+}
+
+TSharedPtr<FApplicationMode> FRigVMEditor::CreateEditorMode()
+{
+	static TSharedPtr<FApplicationMode> EmptyMode;
+	return EmptyMode;
+}
+
+UBlueprint* FRigVMEditor::GetBlueprintObj() const
+{
+	const TArray<UObject*>& EditingObjs = GetEditingObjects();
+	for (UObject* Obj : EditingObjs)
+	{
+		if (Obj->IsA<URigVMBlueprint>()) 
+		{
+			return (UBlueprint*)Obj;
+		}
+	}
+	return nullptr;
+}
+
+TSubclassOf<UEdGraphSchema> FRigVMEditor::GetDefaultSchemaClass() const
+{
+	if(const URigVMBlueprint* RigVMBlueprint = GetRigVMBlueprint())
+	{
+		return RigVMBlueprint->GetRigVMEdGraphSchemaClass();
+	}
+	return URigVMEdGraphSchema::StaticClass();
+}
+
+void FRigVMEditor::Tick(float DeltaTime)
+{
+	FBlueprintEditor::Tick(DeltaTime);
+
+	// tick the  rigvm host
+	if (const URigVMBlueprint* Blueprint = GetRigVMBlueprint())
+	{
+		if (LastDebuggedHost != GetCustomDebugObjectLabel(Blueprint->GetObjectBeingDebugged()))
+		{
+			TArray<FCustomDebugObject> DebugList;
+			GetCustomDebugObjects(DebugList);
+
+			for (const FCustomDebugObject& DebugObject : DebugList)
+			{
+				if (DebugObject.NameOverride == LastDebuggedHost)
+				{
+					GetBlueprintObj()->SetObjectBeingDebugged(DebugObject.Object);
+					break;
+				}
+			}
+		}
+	}
+}
+
+FName FRigVMEditor::GetToolkitFName() const
+{
+	return FName("RigVMEditor");
+}
+
+FName FRigVMEditor::GetToolkitContextFName() const
+{
+	return GetToolkitFName();
+}
+
+FText FRigVMEditor::GetBaseToolkitName() const
+{
+	return LOCTEXT("AppLabel", "RigVM Editor");
+}
+
+FText FRigVMEditor::GetToolkitToolTipText() const
+{
+	return FAssetEditorToolkit::GetToolTipTextForObject(GetBlueprintObj());
+}
+
+FString FRigVMEditor::GetWorldCentricTabPrefix() const
+{
+	return LOCTEXT("WorldCentricTabPrefix", "RigVM Editor ").ToString();
+}
+
+FLinearColor FRigVMEditor::GetWorldCentricTabColorScale() const
+{
+	return FLinearColor( 0.5f, 0.25f, 0.35f, 0.5f );
+}
+
+void FRigVMEditor::InitToolMenuContext(FToolMenuContext& MenuContext)
+{
+	FBlueprintEditor::InitToolMenuContext(MenuContext);
+
+	if (URigVMBlueprint* RigVMBlueprint = GetRigVMBlueprint())
+	{
+		URigVMGraph* Model = nullptr;
+		URigVMNode* Node = nullptr;
+		URigVMPin* Pin = nullptr;
+		
+		if (UGraphNodeContextMenuContext* GraphNodeContext = MenuContext.FindContext<UGraphNodeContextMenuContext>())
+		{
+			if (GraphNodeContext->Node)
+			{
+				Model = RigVMBlueprint->GetModel(GraphNodeContext->Graph);
+				if(Model)
+				{
+					Node = Model->FindNodeByName(GraphNodeContext->Node->GetFName());
+				}
+			}
+		
+			if (GraphNodeContext->Pin && Node)
+			{
+				Pin = Model->FindPin(GraphNodeContext->Pin->GetName());
+			}
+		}
+		
+		URigVMEditorMenuContext* RigVMEditorMenuContext = NewObject<URigVMEditorMenuContext>();
+		const FRigVMEditorGraphMenuContext GraphMenuContext = FRigVMEditorGraphMenuContext(Model, Node, Pin);
+		RigVMEditorMenuContext->Init(SharedThis(this), GraphMenuContext);
+
+		MenuContext.AddObject(RigVMEditorMenuContext);
+	}
+}
+
+bool FRigVMEditor::TransactionObjectAffectsBlueprint(UObject* InTransactedObject)
+{
+	URigVMBlueprint* RigVMBlueprint = GetRigVMBlueprint();
+	if (RigVMBlueprint == nullptr)
+	{
+		return false;
+	}
+
+	if (InTransactedObject->GetOuter() == GetFocusedController())
+	{
+		return false;
+	}
+	return FBlueprintEditor::TransactionObjectAffectsBlueprint(InTransactedObject);
+}
+
+bool FRigVMEditor::CanAddNewLocalVariable() const
+{
+	const URigVMGraph* Graph = GetFocusedModel();
+	const URigVMGraph* ParentGraph = Graph->GetParentGraph();				
+	if (ParentGraph && ParentGraph->IsA<URigVMFunctionLibrary>())
+	{
+		return true;
+	}
+	return false;
+}
+
+void FRigVMEditor::OnAddNewLocalVariable()
+{
+	if (!CanAddNewLocalVariable())
+	{
+		return;
+	}
+
+	FRigVMGraphVariableDescription LastTypeVar;
+	LastTypeVar.ChangeType(MyBlueprintWidget->GetLastPinTypeUsed());
+	FRigVMGraphVariableDescription NewVar = GetFocusedController()->AddLocalVariable(TEXT("NewLocalVar"), LastTypeVar.CPPType, LastTypeVar.CPPTypeObject, LastTypeVar.DefaultValue, true, true);
+	if(NewVar.Name.IsNone())
+	{
+		LogSimpleMessage( LOCTEXT("AddLocalVariable_Error", "Adding new local variable failed.") );
+	}
+	else
+	{
+		RenameNewlyAddedAction(NewVar.Name);
+	}
+}
+
+void FRigVMEditor::OnPasteNewLocalVariable(const FBPVariableDescription& VariableDescription)
+{
+	if (!CanAddNewLocalVariable())
+	{
+		return;
+	}
+
+	FRigVMGraphVariableDescription TypeVar;
+	TypeVar.ChangeType(VariableDescription.VarType);
+	FRigVMGraphVariableDescription NewVar = GetFocusedController()->AddLocalVariable(VariableDescription.VarName, TypeVar.CPPType, TypeVar.CPPTypeObject, VariableDescription.DefaultValue, true, true);
+	if(NewVar.Name.IsNone())
+	{
+		LogSimpleMessage( LOCTEXT("PasteLocalVariable_Error", "Pasting new local variable failed.") );
+	}
+}
+
+void FRigVMEditor::DeleteSelectedNodes()
+{
+	DECLARE_SCOPE_HIERARCHICAL_COUNTER_FUNC()
+
+	URigVMBlueprint* RigVMBlueprint = GetRigVMBlueprint();
+	if (RigVMBlueprint == nullptr)
+	{
+		return;
+	}
+
+	const FGraphPanelSelectionSet SelectedNodes = GetSelectedNodes();
+	SetUISelectionState(NAME_None);
+
+	bool bRelinkPins = false;
+	TArray<URigVMNode*> NodesToRemove;
+
+	for (FGraphPanelSelectionSet::TConstIterator NodeIt(SelectedNodes); NodeIt; ++NodeIt)
+	{
+		if (UEdGraphNode* Node = Cast<UEdGraphNode>(*NodeIt))
+		{
+			if (Node->CanUserDeleteNode())
+			{
+				AnalyticsTrackNodeEvent(GetBlueprintObj(), Node, true);
+				if (const URigVMEdGraphNode* RigVMEdGraphNode = Cast<URigVMEdGraphNode>(Node))
+				{
+					bRelinkPins = bRelinkPins || FSlateApplication::Get().GetModifierKeys().IsShiftDown();
+
+					if(URigVMNode* ModelNode = GetFocusedController()->GetGraph()->FindNodeByName(*RigVMEdGraphNode->ModelNodePath))
+					{
+						NodesToRemove.Add(ModelNode);
+					}
+				}
+				else if (const UEdGraphNode_Comment* CommentNode = Cast<UEdGraphNode_Comment>(Node))
+				{
+					if(URigVMNode* ModelNode = GetFocusedController()->GetGraph()->FindNodeByName(CommentNode->GetFName()))
+					{
+						NodesToRemove.Add(ModelNode);
+					}
+				}
+				else
+				{
+					Node->GetGraph()->RemoveNode(Node);
+				}
+			}
+		}
+	}
+
+	if(NodesToRemove.IsEmpty())
+	{
+		return;
+	}
+
+	GetFocusedController()->OpenUndoBracket(TEXT("Delete selected nodes"));
+	if(bRelinkPins && NodesToRemove.Num() == 1)
+	{
+		GetFocusedController()->RelinkSourceAndTargetPins(NodesToRemove[0], true);;
+	}
+	GetFocusedController()->RemoveNodes(NodesToRemove, true);
+	GetFocusedController()->CloseUndoBracket();
+}
+
+bool FRigVMEditor::CanDeleteNodes() const
+{
+	return true;
+}
+
+void FRigVMEditor::CopySelectedNodes()
+{
+	FString ExportedText = GetFocusedController()->ExportSelectedNodesToText();
+	FPlatformApplicationMisc::ClipboardCopy(*ExportedText);
+}
+
+bool FRigVMEditor::CanCopyNodes() const
+{
+	return GetFocusedModel()->GetSelectNodes().Num() > 0;
+}
+
+bool FRigVMEditor::CanPasteNodes() const
+{
+	FString TextToImport;
+	FPlatformApplicationMisc::ClipboardPaste(TextToImport);
+	return GetFocusedController()->CanImportNodesFromText(TextToImport);
+}
+
+FReply FRigVMEditor::OnSpawnGraphNodeByShortcut(FInputChord InChord, const FVector2D& InPosition,
+                                                     UEdGraph* InGraph)
+{
+	if(!InChord.HasAnyModifierKeys())
+	{
+		if(URigVMEdGraph* RigVMEdGraph = Cast<URigVMEdGraph>(InGraph))
+		{
+			if(URigVMController* Controller = RigVMEdGraph->GetController())
+			{
+				if(InChord.Key == EKeys::B)
+				{
+					Controller->AddUnitNode(FRigVMFunction_ControlFlowBranch::StaticStruct(), FRigVMStruct::ExecuteName, InPosition, FString(), true, true);
+				}
+			}
+		}
+	}
+
+	return FReply::Unhandled();
+}
+
+void FRigVMEditor::JumpToHyperlink(const UObject* ObjectReference, bool bRequestRename)
+{
+	if(const URigVMEdGraph* Graph = Cast<URigVMEdGraph>(ObjectReference))
+	{
+		OpenGraphAndBringToFront((UEdGraph*)Graph, true);
+		return;
+	}
+	FBlueprintEditor::JumpToHyperlink(ObjectReference, bRequestRename);
+}
+
+void FRigVMEditor::PostUndo(bool bSuccess)
+{
+	const FTransaction* Transaction = GEditor->Trans->GetTransaction(GEditor->Trans->GetQueueLength() - GEditor->Trans->GetUndoCount());
+	FBlueprintEditor::PostUndo(bSuccess);
+	PostTransaction(bSuccess, Transaction, false);
+}
+
+void FRigVMEditor::PostRedo(bool bSuccess)
+{
+	const FTransaction* Transaction = GEditor->Trans->GetTransaction(GEditor->Trans->GetQueueLength() - GEditor->Trans->GetUndoCount() - 1);
+	FBlueprintEditor::PostRedo(bSuccess);
+	PostTransaction(bSuccess, Transaction, true);
+}
+
+void FRigVMEditor::OnStartWatchingPin()
+{
+	if (UEdGraphPin* Pin = GetCurrentlySelectedPin())
+	{
+		GetFocusedController()->SetPinIsWatched(Pin->GetName(), true);
+	}
+}
+
+bool FRigVMEditor::CanStartWatchingPin() const
+{
+	if (UEdGraphPin* Pin = GetCurrentlySelectedPin())
+	{
+		if (URigVMPin* ModelPin = GetFocusedModel()->FindPin(Pin->GetName()))
+		{
+			return ModelPin->GetParentPin() == nullptr &&
+					!ModelPin->RequiresWatch();
+		}
+	}
+	return false;
+}
+
+void FRigVMEditor::OnStopWatchingPin()
+{
+	if (UEdGraphPin* Pin = GetCurrentlySelectedPin())
+	{
+		GetFocusedController()->SetPinIsWatched(Pin->GetName(), false);
+	}
+}
+
+bool FRigVMEditor::CanStopWatchingPin() const
+{
+	if (UEdGraphPin* Pin = GetCurrentlySelectedPin())
+	{
+		if (URigVMPin* ModelPin = GetFocusedModel()->FindPin(Pin->GetName()))
+		{
+			return ModelPin->RequiresWatch();
+		}
+	}
+	return false;
+}
+
+void FRigVMEditor::OnToolkitHostingStarted(const TSharedRef<class IToolkit>& Toolkit)
+{
+	TSharedPtr<SWidget> InlineContent = Toolkit->GetInlineContent();
+	if (InlineContent.IsValid())
+	{
+		Toolbox->SetContent(InlineContent.ToSharedRef());
+	}
+}
+
+void FRigVMEditor::OnToolkitHostingFinished(const TSharedRef<class IToolkit>& Toolkit)
+{
+	Toolbox->SetContent(SNullWidget::NullWidget);
+}
+
+TStatId FRigVMEditor::GetStatId() const
+{
+	RETURN_QUICK_DECLARE_CYCLE_STAT(FRigVMEditor, STATGROUP_Tickables);
+}
+
+void FRigVMEditor::PasteNodes()
+{
+	DECLARE_SCOPE_HIERARCHICAL_COUNTER_FUNC()
+
+	GetFocusedController()->OpenUndoBracket(TEXT("Pasted Nodes."));
+
+	FVector2D PasteLocation = FSlateApplication::Get().GetCursorPos();
+
+	TSharedPtr<SDockTab> ActiveTab = DocumentManager->GetActiveTab();
+	if (ActiveTab.IsValid())
+	{
+		TSharedPtr<SGraphEditor> GraphEditor = StaticCastSharedRef<SGraphEditor>(ActiveTab->GetContent());
+		if (GraphEditor.IsValid())
+		{
+			PasteLocation = GraphEditor->GetPasteLocation();
+
+		}
+	}
+
+	FString TextToImport;
+	FPlatformApplicationMisc::ClipboardPaste(TextToImport);
+
+	TGuardValue<FRigVMController_RequestLocalizeFunctionDelegate> RequestLocalizeDelegateGuard(
+		GetFocusedController()->RequestLocalizeFunctionDelegate,
+		FRigVMController_RequestLocalizeFunctionDelegate::CreateLambda([this](FRigVMGraphFunctionIdentifier& InFunctionToLocalize)
+		{
+			OnRequestLocalizeFunctionDialog(InFunctionToLocalize, GetRigVMBlueprint(), true);
+
+		   const URigVMLibraryNode* LocalizedFunctionNode = GetRigVMBlueprint()->GetLocalFunctionLibrary()->FindPreviouslyLocalizedFunction(InFunctionToLocalize);
+		   return LocalizedFunctionNode != nullptr;
+		
+		})
+	);
+	
+	TArray<FName> NodeNames = GetFocusedController()->ImportNodesFromText(TextToImport, true, true);
+
+	if (NodeNames.Num() > 0)
+	{
+		FBox2D Bounds;
+		Bounds.bIsValid = false;
+
+		TArray<FName> NodesToSelect;
+		for (const FName& NodeName : NodeNames)
+		{
+			const URigVMNode* Node = GetFocusedModel()->FindNodeByName(NodeName);
+			check(Node);
+
+			if (Node->IsInjected())
+			{
+				continue;
+			}
+			NodesToSelect.Add(NodeName);
+
+			FVector2D Position = Node->GetPosition();
+			FVector2D Size = Node->GetSize();
+
+			if (!Bounds.bIsValid)
+			{
+				Bounds.Min = Bounds.Max = Position;
+				Bounds.bIsValid = true;
+			}
+			Bounds += Position;
+			Bounds += Position + Size;
+		}
+
+		for (const FName& NodeName : NodesToSelect)
+		{
+			const URigVMNode* Node = GetFocusedModel()->FindNodeByName(NodeName);
+			check(Node);
+
+			FVector2D Position = Node->GetPosition();
+			GetFocusedController()->SetNodePositionByName(NodeName, PasteLocation + Position - Bounds.GetCenter(), true, false, true);
+		}
+
+		GetFocusedController()->SetNodeSelection(NodesToSelect);
+		GetFocusedController()->CloseUndoBracket();
+	}
+	else
+	{
+		GetFocusedController()->CancelUndoBracket();
+	}
+}
+
+URigVMBlueprint* FRigVMEditor::GetRigVMBlueprint() const
+{
+	return Cast<URigVMBlueprint>(GetBlueprintObj());
+}
+
+URigVMHost* FRigVMEditor::GetRigVMHost() const
+{
+	if(IsValid(Host))
+	{
+		return Host;
+	}
+	return nullptr;
+}
+
+UObject* FRigVMEditor::GetOuterForHost() const
+{
+	return GetRigVMBlueprint();
+}
+
+UClass* FRigVMEditor::GetDetailWrapperClass() const
+{
+	return URigVMDetailsViewWrapperObject::StaticClass();
+}
+
+bool FRigVMEditor::SelectLocalVariable(const UEdGraph* Graph, const FName& VariableName)
+{
+	if (const URigVMEdGraph* RigVMEdGraph = Cast<URigVMEdGraph>(Graph))
+	{
+		if (URigVMGraph* RigVMGraph = RigVMEdGraph->GetModel())
+		{
+			for (FRigVMGraphVariableDescription& Variable : RigVMGraph->GetLocalVariables())
+			{
+				if (Variable.Name == VariableName)
+				{
+					URigVMDetailsViewWrapperObject* WrapperObject = URigVMDetailsViewWrapperObject::MakeInstance(
+						GetDetailWrapperClass(), GetBlueprintObj(), Variable.StaticStruct(), (uint8*)&Variable, RigVMGraph);
+					WrapperObject->GetWrappedPropertyChangedChainEvent().AddSP(this, &FRigVMEditor::OnWrappedPropertyChangedChainEvent);
+					WrapperObject->AddToRoot();
+
+					TArray<UObject*> Objects = {WrapperObject};
+					SetDetailObjects(Objects, false);
+					return true;
+				}
+			}
+		}
+	}
+	return false;
+}
+
+void FRigVMEditor::CreateDefaultCommands()
+{
+	if (GetBlueprintObj())
+	{
+		FBlueprintEditor::CreateDefaultCommands();
+	}
+	else
+	{
+		ToolkitCommands->MapAction( FGenericCommands::Get().Undo, 
+			FExecuteAction::CreateSP( this, &FRigVMEditor::UndoAction ));
+		ToolkitCommands->MapAction( FGenericCommands::Get().Redo, 
+			FExecuteAction::CreateSP( this, &FRigVMEditor::RedoAction ));
+	}
+}
+
+void FRigVMEditor::OnCreateGraphEditorCommands(TSharedPtr<FUICommandList> GraphEditorCommandsList)
+{
+}
+
+struct FRigVMEditorZoomLevelsContainer : public FZoomLevelsContainer
+{
+	struct FRigVMEditorZoomLevelEntry
+	{
+	public:
+		FRigVMEditorZoomLevelEntry(float InZoomAmount, const FText& InDisplayText, EGraphRenderingLOD::Type InLOD)
+			: DisplayText(FText::Format(NSLOCTEXT("GraphEditor", "Zoom", "Zoom {0}"), InDisplayText))
+		, ZoomAmount(InZoomAmount)
+		, LOD(InLOD)
+		{
+		}
+
+	public:
+		FText DisplayText;
+		float ZoomAmount;
+		EGraphRenderingLOD::Type LOD;
+	};
+	
+	FRigVMEditorZoomLevelsContainer()
+	{
+		ZoomLevels.Reserve(22);
+		ZoomLevels.Add(FRigVMEditorZoomLevelEntry(0.025f, FText::FromString(TEXT("-14")), EGraphRenderingLOD::LowestDetail));
+		ZoomLevels.Add(FRigVMEditorZoomLevelEntry(0.070f, FText::FromString(TEXT("-13")), EGraphRenderingLOD::LowestDetail));
+		ZoomLevels.Add(FRigVMEditorZoomLevelEntry(0.100f, FText::FromString(TEXT("-12")), EGraphRenderingLOD::LowestDetail));
+		ZoomLevels.Add(FRigVMEditorZoomLevelEntry(0.125f, FText::FromString(TEXT("-11")), EGraphRenderingLOD::LowestDetail));
+		ZoomLevels.Add(FRigVMEditorZoomLevelEntry(0.150f, FText::FromString(TEXT("-10")), EGraphRenderingLOD::LowestDetail));
+		ZoomLevels.Add(FRigVMEditorZoomLevelEntry(0.175f, FText::FromString(TEXT("-9")), EGraphRenderingLOD::LowestDetail));
+		ZoomLevels.Add(FRigVMEditorZoomLevelEntry(0.200f, FText::FromString(TEXT("-8")), EGraphRenderingLOD::LowestDetail));
+		ZoomLevels.Add(FRigVMEditorZoomLevelEntry(0.225f, FText::FromString(TEXT("-7")), EGraphRenderingLOD::LowDetail));
+		ZoomLevels.Add(FRigVMEditorZoomLevelEntry(0.250f, FText::FromString(TEXT("-6")), EGraphRenderingLOD::LowDetail));
+		ZoomLevels.Add(FRigVMEditorZoomLevelEntry(0.375f, FText::FromString(TEXT("-5")), EGraphRenderingLOD::MediumDetail));
+		ZoomLevels.Add(FRigVMEditorZoomLevelEntry(0.500f, FText::FromString(TEXT("-4")), EGraphRenderingLOD::MediumDetail));
+		ZoomLevels.Add(FRigVMEditorZoomLevelEntry(0.675f, FText::FromString(TEXT("-3")), EGraphRenderingLOD::MediumDetail));
+		ZoomLevels.Add(FRigVMEditorZoomLevelEntry(0.750f, FText::FromString(TEXT("-2")), EGraphRenderingLOD::DefaultDetail));
+		ZoomLevels.Add(FRigVMEditorZoomLevelEntry(0.875f, FText::FromString(TEXT("-1")), EGraphRenderingLOD::DefaultDetail));
+		ZoomLevels.Add(FRigVMEditorZoomLevelEntry(1.000f, FText::FromString(TEXT("1:1")), EGraphRenderingLOD::DefaultDetail));
+		ZoomLevels.Add(FRigVMEditorZoomLevelEntry(1.250f, FText::FromString(TEXT("+1")), EGraphRenderingLOD::DefaultDetail));
+		ZoomLevels.Add(FRigVMEditorZoomLevelEntry(1.375f, FText::FromString(TEXT("+2")), EGraphRenderingLOD::DefaultDetail));
+		ZoomLevels.Add(FRigVMEditorZoomLevelEntry(1.500f, FText::FromString(TEXT("+3")), EGraphRenderingLOD::FullyZoomedIn));
+		ZoomLevels.Add(FRigVMEditorZoomLevelEntry(1.675f, FText::FromString(TEXT("+4")), EGraphRenderingLOD::FullyZoomedIn));
+		ZoomLevels.Add(FRigVMEditorZoomLevelEntry(1.750f, FText::FromString(TEXT("+5")), EGraphRenderingLOD::FullyZoomedIn));
+		ZoomLevels.Add(FRigVMEditorZoomLevelEntry(1.875f, FText::FromString(TEXT("+6")), EGraphRenderingLOD::FullyZoomedIn));
+		ZoomLevels.Add(FRigVMEditorZoomLevelEntry(2.000f, FText::FromString(TEXT("+7")), EGraphRenderingLOD::FullyZoomedIn));
+	}
+
+	float GetZoomAmount(int32 InZoomLevel) const override
+	{
+		checkSlow(ZoomLevels.IsValidIndex(InZoomLevel));
+		return ZoomLevels[InZoomLevel].ZoomAmount;
+	}
+
+	int32 GetNearestZoomLevel(float InZoomAmount) const override
+	{
+		for (int32 ZoomLevelIndex=0; ZoomLevelIndex < GetNumZoomLevels(); ++ZoomLevelIndex)
+		{
+			if (InZoomAmount <= GetZoomAmount(ZoomLevelIndex))
+			{
+				return ZoomLevelIndex;
+			}
+		}
+
+		return GetDefaultZoomLevel();
+	}
+	
+	FText GetZoomText(int32 InZoomLevel) const override
+	{
+		checkSlow(ZoomLevels.IsValidIndex(InZoomLevel));
+		return ZoomLevels[InZoomLevel].DisplayText;
+	}
+	
+	int32 GetNumZoomLevels() const override
+	{
+		return ZoomLevels.Num();
+	}
+	
+	int32 GetDefaultZoomLevel() const override
+	{
+		return 14;
+	}
+
+	EGraphRenderingLOD::Type GetLOD(int32 InZoomLevel) const override
+	{
+		checkSlow(ZoomLevels.IsValidIndex(InZoomLevel));
+		return ZoomLevels[InZoomLevel].LOD;
+	}
+
+	TArray<FRigVMEditorZoomLevelEntry> ZoomLevels;
+};
+
+TSharedRef<SGraphEditor> FRigVMEditor::CreateGraphEditorWidget(TSharedRef<FTabInfo> InTabInfo, UEdGraph* InGraph)
+{
+	TSharedRef<SGraphEditor> GraphEditor = FBlueprintEditor::CreateGraphEditorWidget(InTabInfo, InGraph);
+	GraphEditor->GetGraphPanel()->SetZoomLevelsContainer<FRigVMEditorZoomLevelsContainer>();
+	return GraphEditor;
+}
+
+void FRigVMEditor::Compile()
+{
+	{
+		DECLARE_SCOPE_HIERARCHICAL_COUNTER_FUNC()
+
+		URigVMBlueprint* RigVMBlueprint = GetRigVMBlueprint();
+		if (RigVMBlueprint == nullptr)
+		{
+			return;
+		}
+
+		// force to disable the supended notif brackets
+		RigVMBlueprint->bSuspendModelNotificationsForOthers = false;
+		RigVMBlueprint->bSuspendModelNotificationsForSelf = false;
+
+		RigVMBlueprint->GetCompileLog().Messages.Reset();
+
+		FString LastDebuggedObjectName = GetCustomDebugObjectLabel(RigVMBlueprint->GetObjectBeingDebugged());
+		RigVMBlueprint->SetObjectBeingDebugged(nullptr);
+
+		TArray< TWeakObjectPtr<UObject> > SelectedObjects = Inspector->GetSelectedObjects();
+
+		if (URigVMHost* RigVMHost = GetRigVMHost())
+		{
+			RigVMHost->OnInitialized_AnyThread().Clear();
+			RigVMHost->OnExecuted_AnyThread().Clear();
+			RigVMHost->GetExtendedExecuteContext().ExecutionHalted().RemoveAll(this);
+		}
+
+		{
+			TGuardValue<bool> GuardCompileReEntry(bIsCompilingThroughUI, true);
+			FBlueprintEditor::Compile();
+		}
+
+		if (URigVMHost* RigVMHost = GetRigVMHost())
+		{
+			RigVMLog.Reset();
+			RigVMHost->SetLog(&RigVMLog);
+
+			URigVMBlueprintGeneratedClass* GeneratedClass = Cast<URigVMBlueprintGeneratedClass>(RigVMHost->GetClass());
+			if (GeneratedClass)
+			{
+				URigVMHost* CDO = Cast<URigVMHost>(GeneratedClass->GetDefaultObject(true /* create if needed */));
+				FRigVMInstructionArray Instructions = CDO->GetVM()->GetInstructions();
+
+				if (Instructions.Num() <= 1) // just the "done" operator
+				{
+					FNotificationInfo Info(LOCTEXT("ControlRigBlueprintCompilerEmptyRigMessage", "The asset you compiled doesn't do anything. Did you forget to add a Begin_Execution node?"));
+					Info.bFireAndForget = true;
+					Info.FadeOutDuration = 5.0f;
+					Info.ExpireDuration = 5.0f;
+					TSharedPtr<SNotificationItem> NotificationPtr = FSlateNotificationManager::Get().AddNotification(Info);
+					NotificationPtr->SetCompletionState(SNotificationItem::CS_Success);
+				}
+			}
+		}
+
+		TArray<FCustomDebugObject> DebugList;
+		GetCustomDebugObjects(DebugList);
+
+		for (const FCustomDebugObject& DebugObject : DebugList)
+		{
+			if (DebugObject.NameOverride == LastDebuggedObjectName)
+			{
+				RigVMBlueprint->SetObjectBeingDebugged(DebugObject.Object);
+			}
+		}
+
+		// invalidate all node titles
+		TArray<UEdGraph*> EdGraphs;
+		RigVMBlueprint->GetAllGraphs(EdGraphs);
+		for (UEdGraph* EdGraph : EdGraphs)
+		{
+			URigVMEdGraph* RigVMEdGraph = Cast<URigVMEdGraph>(EdGraph);
+			if (RigVMEdGraph == nullptr)
+			{
+				continue;
+			}
+
+			for (UEdGraphNode* EdNode : RigVMEdGraph->Nodes)
+			{
+				if (URigVMEdGraphNode* RigVMEdGraphNode = Cast<URigVMEdGraphNode>(EdNode))
+				{
+					RigVMEdGraphNode->InvalidateNodeTitle();
+				}
+			}
+		}
+
+		// store the defaults from the CDO back on the new variables list
+		bool bAnyVariableValueChanged = false;
+		for(FBPVariableDescription& NewVariable : RigVMBlueprint->NewVariables)
+		{
+			bAnyVariableValueChanged |= UpdateDefaultValueForVariable(NewVariable, true);
+		}
+		if (bAnyVariableValueChanged)
+		{
+			// Go over all the instances to update the default values from CDO
+			for (const FCustomDebugObject& DebugObject : DebugList)
+			{
+				if (URigVMHost* DebuggedHost = Cast<URigVMHost>(DebugObject.Object))
+				{
+					DebuggedHost->CopyExternalVariableDefaultValuesFromCDO();
+				}
+			}
+		}
+	}
+
+	// enable this for creating a new unit test
+	// DumpUnitTestCode();
+
+	// FStatsHierarchical::EndMeasurements();
+	// FMessageLog LogForMeasurements("ControlRigLog");
+	// FStatsHierarchical::DumpMeasurements(LogForMeasurements);
+}
+
+void FRigVMEditor::SaveAsset_Execute()
+{
+	LastDebuggedHost = GetCustomDebugObjectLabel(GetBlueprintObj()->GetObjectBeingDebugged());
+	FBlueprintEditor::SaveAsset_Execute();
+
+	FBlueprintActionDatabase& ActionDatabase = FBlueprintActionDatabase::Get();
+	ActionDatabase.ClearAssetActions(URigVMBlueprint::StaticClass());
+	ActionDatabase.RefreshClassActions(URigVMBlueprint::StaticClass());
+}
+
+void FRigVMEditor::SaveAssetAs_Execute()
+{
+	LastDebuggedHost = GetCustomDebugObjectLabel(GetBlueprintObj()->GetObjectBeingDebugged());
+	FBlueprintEditor::SaveAssetAs_Execute();
+
+	FBlueprintActionDatabase& ActionDatabase = FBlueprintActionDatabase::Get();
+	ActionDatabase.ClearAssetActions(URigVMBlueprint::StaticClass());
+	ActionDatabase.RefreshClassActions(URigVMBlueprint::StaticClass());
+}
+
+bool FRigVMEditor::IsEditable(UEdGraph* InGraph) const
+{
+	if(!IsGraphInCurrentBlueprint(InGraph))
+	{
+		return false;
+	}
+	
+	if(URigVMBlueprint* RigVMBlueprint = GetRigVMBlueprint())
+	{
+		// aggregate graphs are always read only
+		if(const URigVMGraph* Model = RigVMBlueprint->GetModel(InGraph))
+		{
+			if(Model->GetOuter()->IsA<URigVMAggregateNode>())
+			{
+				return false;
+			}
+		}
+
+		URigVMHost* RigVMHost = GetRigVMHost();
+		if(RigVMHost && RigVMHost->GetVM())
+		{
+			const bool bIsReadOnly = RigVMHost->GetVM()->IsNativized();
+			const bool bIsEditable = !bIsReadOnly;
+			InGraph->bEditable = bIsEditable ? 1 : 0;
+			return bIsEditable;
+		}
+	}
+
+	return FBlueprintEditor::IsEditable(InGraph);
+}
+
+bool FRigVMEditor::IsCompilingEnabled() const
+{
+	return true;
+}
+
+FText FRigVMEditor::GetGraphDecorationString(UEdGraph* InGraph) const
+{
+	return FText::GetEmpty();
+}
+
+void FRigVMEditor::OnSelectedNodesChangedImpl(const TSet<UObject*>& NewSelection)
+{
+	DECLARE_SCOPE_HIERARCHICAL_COUNTER_FUNC()
+
+	URigVMEdGraph* RigVMEdGraph = Cast<URigVMEdGraph>(GetFocusedGraph());
+	if (RigVMEdGraph == nullptr)
+	{
+		return;
+	}
+
+	if (RigVMEdGraph->bIsSelecting || GIsTransacting)
+	{
+		return;
+	}
+
+	TGuardValue<bool> SelectGuard(RigVMEdGraph->bIsSelecting, true);
+
+	URigVMBlueprint* RigVMBlueprint = GetRigVMBlueprint();
+	if (RigVMBlueprint)
+	{
+		TArray<FName> NodeNamesToSelect;
+		for (UObject* Object : NewSelection)
+		{
+			if (URigVMEdGraphNode* RigVMEdGraphNode = Cast<URigVMEdGraphNode>(Object))
+			{
+				NodeNamesToSelect.Add(RigVMEdGraphNode->GetModelNodeName());
+			}
+			else if(UEdGraphNode* Node = Cast<UEdGraphNode>(Object))
+			{
+				NodeNamesToSelect.Add(Node->GetFName());
+			}
+		}
+		GetFocusedController()->SetNodeSelection(NodeNamesToSelect, true, true);
+	}
+}
+
+void FRigVMEditor::OnBlueprintChangedImpl(UBlueprint* InBlueprint, bool bIsJustBeingCompiled)
+{
+	DECLARE_SCOPE_HIERARCHICAL_COUNTER_FUNC()
+
+	if (!bRigVMEditorInitialized)
+	{
+		return;
+	}
+
+	FBlueprintEditor::OnBlueprintChangedImpl(InBlueprint, bIsJustBeingCompiled);
+
+	if(InBlueprint == GetBlueprintObj())
+	{
+		if(bIsJustBeingCompiled)
+		{
+			UpdateRigVMHost();
+
+			if (!LastDebuggedHost.IsEmpty())
+			{
+				TArray<FCustomDebugObject> DebugList;
+				GetCustomDebugObjects(DebugList);
+
+				for (const FCustomDebugObject& DebugObject : DebugList)
+				{
+					if (DebugObject.NameOverride == LastDebuggedHost)
+					{
+						GetBlueprintObj()->SetObjectBeingDebugged(DebugObject.Object);
+						LastDebuggedHost.Empty();
+						break;
+					}
+				}
+			}
+		}
+	}}
+
+void FRigVMEditor::RefreshEditors(ERefreshBlueprintEditorReason::Type Reason)
+{
+	if(Reason == ERefreshBlueprintEditorReason::UnknownReason)
+	{
+		// we mark the reason as just compiled since we don't want to
+		// update the graph(s) all the time during compilation
+		Reason = ERefreshBlueprintEditorReason::BlueprintCompiled;
+	}
+	FBlueprintEditor::RefreshEditors(Reason);
+}
+
+void FRigVMEditor::SetupGraphEditorEvents(UEdGraph* InGraph, SGraphEditor::FGraphEditorEvents& InEvents)
+{
+	DECLARE_SCOPE_HIERARCHICAL_COUNTER_FUNC()
+
+	FBlueprintEditor::SetupGraphEditorEvents(InGraph, InEvents);
+
+	InEvents.OnCreateActionMenu = SGraphEditor::FOnCreateActionMenu::CreateSP(this, &FRigVMEditor::HandleCreateGraphActionMenu);
+	InEvents.OnTextCommitted = FOnNodeTextCommitted::CreateSP(this, &FRigVMEditor::OnNodeTitleCommitted);
+}
+
+FActionMenuContent FRigVMEditor::HandleCreateGraphActionMenu(UEdGraph* InGraph, const FVector2D& InNodePosition, const TArray<UEdGraphPin*>& InDraggedPins, bool bAutoExpand, SGraphEditor::FActionMenuClosed InOnMenuClosed)
+{
+	return FBlueprintEditor::OnCreateGraphActionMenu(InGraph, InNodePosition, InDraggedPins, bAutoExpand, InOnMenuClosed);
+}
+
+void FRigVMEditor::OnNodeTitleCommitted(const FText& NewText, ETextCommit::Type CommitInfo, UEdGraphNode* NodeBeingChanged)
+{
+	DECLARE_SCOPE_HIERARCHICAL_COUNTER_FUNC()
+
+	if (const UEdGraphNode_Comment* CommentBeingChanged = Cast<UEdGraphNode_Comment>(NodeBeingChanged))
+	{
+		if (GetRigVMBlueprint())
+		{
+			GetFocusedController()->SetCommentTextByName(CommentBeingChanged->GetFName(), NewText.ToString(), CommentBeingChanged->FontSize, CommentBeingChanged->bCommentBubbleVisible, CommentBeingChanged->bColorCommentBubble, true, true);
+		}
+	}
+}
+
+void FRigVMEditor::FocusInspectorOnGraphSelection(const TSet<UObject*>& NewSelection, bool bForceRefresh)
+{
+	// nothing to do here
+}
+
+void FRigVMEditor::AddReferencedObjects(FReferenceCollector& Collector)
+{
+	FBlueprintEditor::AddReferencedObjects(Collector);
+
+	TWeakObjectPtr<URigVMHost> RigVMHost(GetRigVMHost());
+	Collector.AddReferencedObject(RigVMHost);
+}
+
+void FRigVMEditor::BindCommands()
+{
+	DECLARE_SCOPE_HIERARCHICAL_COUNTER_FUNC()
+
+	GetToolkitCommands()->MapAction(
+		FRigVMEditorCommands::Get().AutoCompileGraph,
+		FExecuteAction::CreateSP(this, &FRigVMEditor::ToggleAutoCompileGraph), 
+		FIsActionChecked::CreateSP(this, &FRigVMEditor::CanAutoCompileGraph),
+		FIsActionChecked::CreateSP(this, &FRigVMEditor::IsAutoCompileGraphOn));
+
+	GetToolkitCommands()->MapAction(
+		FRigVMEditorCommands::Get().ToggleEventQueue,
+		FExecuteAction::CreateSP(this, &FRigVMEditor::ToggleEventQueue),
+		FCanExecuteAction());
+
+	GetToolkitCommands()->MapAction(
+		FRigVMEditorCommands::Get().ToggleExecutionMode,
+		FExecuteAction::CreateSP(this, &FRigVMEditor::ToggleExecutionMode),
+		FCanExecuteAction());
+
+	GetToolkitCommands()->MapAction(
+		FRigVMEditorCommands::Get().ReleaseMode,
+		FExecuteAction::CreateSP(this, &FRigVMEditor::SetExecutionMode, ERigVMEditorExecutionModeType_Release),
+		FCanExecuteAction());
+
+	GetToolkitCommands()->MapAction(
+		FRigVMEditorCommands::Get().DebugMode,
+		FExecuteAction::CreateSP(this, &FRigVMEditor::SetExecutionMode, ERigVMEditorExecutionModeType_Debug),
+		FCanExecuteAction());
+
+	GetToolkitCommands()->MapAction(
+		FRigVMEditorCommands::Get().ResumeExecution,
+		FExecuteAction::CreateSP(this, &FRigVMEditor::HandleBreakpointActionRequested, ERigVMBreakpointAction::Resume),			
+		FIsActionChecked::CreateSP(this, &FRigVMEditor::IsHaltedAtBreakpoint));
+
+	GetToolkitCommands()->MapAction(
+		FRigVMEditorCommands::Get().ShowCurrentStatement,
+		FExecuteAction::CreateSP(this, &FRigVMEditor::HandleShowCurrentStatement),
+		FIsActionChecked::CreateSP(this, &FRigVMEditor::IsHaltedAtBreakpoint));
+
+	GetToolkitCommands()->MapAction(
+		FRigVMEditorCommands::Get().StepOver,
+		FExecuteAction::CreateSP(this, &FRigVMEditor::HandleBreakpointActionRequested, ERigVMBreakpointAction::StepOver),
+		FIsActionChecked::CreateSP(this, &FRigVMEditor::IsHaltedAtBreakpoint));
+
+	GetToolkitCommands()->MapAction(
+		FRigVMEditorCommands::Get().StepInto,
+		FExecuteAction::CreateSP(this, &FRigVMEditor::HandleBreakpointActionRequested, ERigVMBreakpointAction::StepInto),
+		FIsActionChecked::CreateSP(this, &FRigVMEditor::IsHaltedAtBreakpoint));
+
+	GetToolkitCommands()->MapAction(
+		FRigVMEditorCommands::Get().StepOut,
+		FExecuteAction::CreateSP(this, &FRigVMEditor::HandleBreakpointActionRequested, ERigVMBreakpointAction::StepOut),
+		FIsActionChecked::CreateSP(this, &FRigVMEditor::IsHaltedAtBreakpoint));
+
+	GetToolkitCommands()->MapAction(
+		FRigVMEditorCommands::Get().FrameSelection,
+		FExecuteAction::CreateSP(this, &FRigVMEditor::FrameSelection),
+		FCanExecuteAction());
+}
+
+void FRigVMEditor::ToggleAutoCompileGraph()
+{
+	if (URigVMBlueprint* RigVMBlueprint = GetRigVMBlueprint())
+	{
+		RigVMBlueprint->SetAutoVMRecompile(!RigVMBlueprint->GetAutoVMRecompile());
+		if (RigVMBlueprint->GetAutoVMRecompile())
+		{
+			RigVMBlueprint->RequestAutoVMRecompilation();
+		}
+	}
+}
+
+bool FRigVMEditor::IsAutoCompileGraphOn() const
+{
+	if (URigVMBlueprint* RigVMBlueprint = GetRigVMBlueprint())
+	{
+		return RigVMBlueprint->GetAutoVMRecompile();
+	}
+	return false;
+}
+
+void FRigVMEditor::ToggleEventQueue()
+{
+	SetEventQueue(LastEventQueue);
+}
+
+void FRigVMEditor::ToggleExecutionMode()
+{
+	SetExecutionMode((ExecutionMode == ERigVMEditorExecutionModeType_Debug) ?
+		ERigVMEditorExecutionModeType_Release
+		: ERigVMEditorExecutionModeType_Debug);
+}
+
+TSharedRef<SWidget> FRigVMEditor::GenerateEventQueueMenuContent()
+{
+	FMenuBuilder MenuBuilder(true, GetToolkitCommands());
+	GenerateEventQueueMenuContent(MenuBuilder);
+	return MenuBuilder.MakeWidget();
+}
+
+void FRigVMEditor::GenerateEventQueueMenuContent(FMenuBuilder& MenuBuilder)
+{
+}
+
+TSharedRef<SWidget> FRigVMEditor::GenerateExecutionModeMenuContent()
+{
+	FMenuBuilder MenuBuilder(true, GetToolkitCommands());
+	MenuBuilder.BeginSection(TEXT("Events"));
+	MenuBuilder.AddMenuEntry(FRigVMEditorCommands::Get().ReleaseMode, TEXT("Release"), TAttribute<FText>(), TAttribute<FText>(), GetExecutionModeIcon(ERigVMEditorExecutionModeType_Release));
+	MenuBuilder.AddMenuEntry(FRigVMEditorCommands::Get().DebugMode, TEXT("Debug"), TAttribute<FText>(), TAttribute<FText>(), GetExecutionModeIcon(ERigVMEditorExecutionModeType_Debug));
+	MenuBuilder.EndSection();
+	return MenuBuilder.MakeWidget();
+}
+
+void FRigVMEditor::OnActiveTabChanged( TSharedPtr<SDockTab> PreviouslyActive, TSharedPtr<SDockTab> NewlyActivated )
+{
+	if (!NewlyActivated.IsValid())
+	{
+		TArray<UObject*> ObjArray;
+		Inspector->ShowDetailsForObjects(ObjArray);
+	}
+	else 
+	{
+		FBlueprintEditor::OnActiveTabChanged(PreviouslyActive, NewlyActivated);
+	}
+}
+
+void FRigVMEditor::UndoAction()
+{
+	GEditor->UndoTransaction();
+}
+
+void FRigVMEditor::RedoAction()
+{
+	GEditor->RedoTransaction();
+}
+
+void FRigVMEditor::CreateDefaultTabContents(const TArray<UBlueprint*>& InBlueprints)
+{
+	FBlueprintEditor::CreateDefaultTabContents(InBlueprints);
+}
+
+void FRigVMEditor::NewDocument_OnClicked(ECreatedDocumentType GraphType)
+{
+	if (GraphType == FBlueprintEditor::CGT_NewFunctionGraph)
+	{
+		if (URigVMBlueprint* RigVMBlueprint = GetRigVMBlueprint())
+		{
+			if (URigVMController* Controller = RigVMBlueprint->GetOrCreateController(RigVMBlueprint->GetLocalFunctionLibrary()))
+			{
+				if (const URigVMLibraryNode* FunctionNode = Controller->AddFunctionToLibrary(TEXT("New Function"), true, FVector2D::ZeroVector, true, true))
+				{
+					if (const UEdGraph* NewGraph = RigVMBlueprint->GetEdGraph(FunctionNode->GetContainedGraph()))
+					{
+						OpenDocument(NewGraph, FDocumentTracker::OpenNewDocument);
+						RenameNewlyAddedAction(FunctionNode->GetFName());
+					}
+
+				}
+			}
+		}
+	}
+	else if(GraphType == FBlueprintEditor::CGT_NewEventGraph)
+	{
+		if (URigVMBlueprint* RigVMBlueprint = GetRigVMBlueprint())
+		{
+			const UClass* EdGraphSchemaClass = RigVMBlueprint->GetRigVMEdGraphSchemaClass();
+			const URigVMEdGraphSchema* SchemaCDO = CastChecked<URigVMEdGraphSchema>(EdGraphSchemaClass->GetDefaultObject());
+
+			if(URigVMGraph* Model = RigVMBlueprint->AddModel(SchemaCDO->GetRootGraphName().ToString()))
+			{
+				if (const UEdGraph* NewGraph = RigVMBlueprint->GetEdGraph(Model))
+				{
+					OpenDocument(NewGraph, FDocumentTracker::OpenNewDocument);
+					RenameNewlyAddedAction(NewGraph->GetFName());
+				}
+			}
+		}
+	}
+}
+
+bool FRigVMEditor::IsSectionVisible(NodeSectionID::Type InSectionID) const
+{
+	switch (InSectionID)
+	{
+		case NodeSectionID::GRAPH:
+		case NodeSectionID::VARIABLE:
+		case NodeSectionID::FUNCTION:
+		{
+			return true;
+		}
+		case NodeSectionID::LOCAL_VARIABLE:
+		{
+			if(const URigVMGraph* Graph = GetFocusedModel())
+			{
+				const URigVMGraph* ParentGraph = Graph->GetParentGraph();				
+				if (ParentGraph && ParentGraph->IsA<URigVMFunctionLibrary>())
+				{
+					return true;
+				}
+			}
+		}
+		default:
+		{
+			break;
+		}
+	}
+	return false;
+}
+
+FGraphAppearanceInfo FRigVMEditor::GetGraphAppearance(UEdGraph* InGraph) const
+{
+	FGraphAppearanceInfo AppearanceInfo = FBlueprintEditor::GetGraphAppearance(InGraph);
+
+	if (GetBlueprintObj()->IsA(URigVMBlueprint::StaticClass()))
+	{
+		AppearanceInfo.CornerText = LOCTEXT("AppearanceCornerText_RigVMEditor", "RigVM");
+
+		if(URigVMHost* RigVMHost = GetRigVMHost())
+		{
+			if(RigVMHost->GetVM() && RigVMHost->GetVM()->IsNativized())
+			{
+				AppearanceInfo.InstructionFade = 1;
+				AppearanceInfo.InstructionText = FText::FromString(
+					FString::Printf(TEXT("This graph runs a nativized VM (U%s)."), *RigVMHost->GetVM()->GetNativizedClass()->GetName())
+				);
+			}
+		}
+	}
+
+	return AppearanceInfo;
+}
+
+void FRigVMEditor::HandleModifiedEvent(ERigVMGraphNotifType InNotifType, URigVMGraph* InGraph, UObject* InSubject)
+{
+	DECLARE_SCOPE_HIERARCHICAL_COUNTER_FUNC()
+
+	URigVMBlueprint* RigVMBlueprint = GetRigVMBlueprint();
+	if (RigVMBlueprint == nullptr)
+	{
+		return;
+	}
+
+	switch (InNotifType)
+	{
+		case ERigVMGraphNotifType::NodeSelectionChanged:
+		case ERigVMGraphNotifType::NodeSelected:
+		case ERigVMGraphNotifType::NodeDeselected:
+		{
+			if (URigVMEdGraph* RigVMEdGraph = Cast<URigVMEdGraph>(RigVMBlueprint->GetEdGraph(InGraph)))
+			{
+				TSharedPtr<SGraphEditor> GraphEd = GetGraphEditor(RigVMEdGraph);
+				URigVMNode* Node = Cast<URigVMNode>(InSubject);
+				if (InNotifType == ERigVMGraphNotifType::NodeSelectionChanged)
+				{
+					const TArray<FName> SelectedNodes = InGraph->GetSelectNodes();
+					if (!SelectedNodes.IsEmpty())
+					{
+						Node = Cast<URigVMNode>(InGraph->FindNodeByName(SelectedNodes.Last()));	
+					}
+				}
+
+				if (GraphEd.IsValid() && Node != nullptr)
+				{
+					SetDetailViewForGraph(Node->GetGraph());
+
+					if (!RigVMEdGraph->bIsSelecting)
+					{
+						TGuardValue<bool> SelectingGuard(RigVMEdGraph->bIsSelecting, true);
+						if (UEdGraphNode* EdNode = RigVMEdGraph->FindNodeForModelNodeName(Node->GetFName()))
+						{
+							GraphEd->SetNodeSelection(EdNode, InNotifType == ERigVMGraphNotifType::NodeSelected);
+						}
+					}
+				}
+			}
+			break;
+		}
+		case ERigVMGraphNotifType::PinDefaultValueChanged:
+		{
+			URigVMPin* Pin = Cast<URigVMPin>(InSubject);
+
+			if(URigVMPin* RootPin = Pin->GetRootPin())
+			{
+				const FString DefaultValue = RootPin->GetDefaultValue();
+				if(!DefaultValue.IsEmpty())
+				{
+					// sync the value change with the unit(s) displayed 
+					TArray< TWeakObjectPtr<UObject> > SelectedObjects = Inspector->GetSelectedObjects();
+					for (TWeakObjectPtr<UObject> SelectedObject : SelectedObjects)
+					{
+						if (SelectedObject.IsValid())
+						{
+							if(URigVMDetailsViewWrapperObject* WrapperObject = Cast<URigVMDetailsViewWrapperObject>(SelectedObject.Get()))
+							{
+								if(WrapperObject->GetOuter() == Pin->GetNode())
+								{
+									const FProperty* TargetProperty = WrapperObject->GetClass()->FindPropertyByName(RootPin->GetFName());
+									if(TargetProperty)
+									{
+										uint8* PropertyStorage = TargetProperty->ContainerPtrToValuePtr<uint8>(WrapperObject);
+
+										// we are ok with not reacting to errors here
+										FRigVMPinDefaultValueImportErrorContext ErrorPipe;										
+										TargetProperty->ImportText_Direct(*DefaultValue, PropertyStorage, nullptr, PPF_None, &ErrorPipe);
+									}
+								}
+							}
+						}
+					}
+				}
+
+				if(const URigVMUnitNode* UnitNode = Cast<URigVMUnitNode>(RootPin->GetNode()))
+				{
+					if(UnitNode->IsEvent())
+					{
+						FBlueprintEditorUtils::MarkBlueprintAsStructurallyModified(GetBlueprintObj());
+						CacheNameLists();
+					}
+				}
+			}
+	
+			break;
+		}
+		case ERigVMGraphNotifType::PinArraySizeChanged:
+		case ERigVMGraphNotifType::PinBoundVariableChanged:
+		case ERigVMGraphNotifType::PinTypeChanged:
+		{
+			URigVMPin* Pin = Cast<URigVMPin>(InSubject);
+
+			if(Pin->GetNode()->IsSelected())
+			{
+				TArray<UObject*> Objects;
+				Objects.Add(Pin->GetNode());
+				SetDetailObjects(Objects);
+			}
+			break;
+		}
+		case ERigVMGraphNotifType::NodeRemoved:
+		{
+			if (URigVMCollapseNode* CollapseNode = Cast<URigVMCollapseNode>(InSubject))
+			{
+				if (UEdGraph* EdGraph = RigVMBlueprint->GetEdGraph(CollapseNode->GetContainedGraph()))
+				{
+					CloseDocumentTab(EdGraph);
+					ClearDetailObject();
+				}
+			}
+			else if(URigVMFunctionReferenceNode* FunctionRefNode = Cast<URigVMFunctionReferenceNode>(InSubject))
+			{
+				ClearDetailObject();
+			}
+				
+			// fall through next case since we want to refresh the name lists
+			// both for removing or adding an event
+		}
+		case ERigVMGraphNotifType::NodeAdded:
+		{
+			if(URigVMUnitNode* UnitNode = Cast<URigVMUnitNode>(InSubject))
+			{
+				if(UnitNode->IsEvent())
+				{
+					FBlueprintEditorUtils::MarkBlueprintAsStructurallyModified(GetBlueprintObj());
+					CacheNameLists();
+				}
+			}
+			break;
+		}
+		default:
+		{
+			break;
+		}
+	}
+}
+
+class FMemoryTypeMetaData : public ISlateMetaData
+{
+public:
+	SLATE_METADATA_TYPE(FMemoryTypeMetaData, ISlateMetaData)
+
+		FMemoryTypeMetaData(ERigVMMemoryType InMemoryType)
+		: MemoryType(InMemoryType)
+	{
+	}
+	ERigVMMemoryType MemoryType;
+};
+
+void FRigVMEditor::HandleVMCompiledEvent(UObject* InCompiledObject, URigVM* InVM)
+{
+	if(URigVMBlueprint* RigVMBlueprint = Cast<URigVMBlueprint>(InCompiledObject))
+	{
+		CompilerResultsListing->ClearMessages();
+		CompilerResultsListing->AddMessages(RigVMBlueprint->GetCompileLog().Messages);
+		RigVMBlueprint->GetCompileLog().Messages.Reset();
+		RigVMBlueprint->GetCompileLog().NumErrors = RigVMBlueprint->GetCompileLog().NumWarnings = 0;
+	}
+
+	RefreshDetailView();
+	
+	TArray<FName> TabIds;
+	TabIds.Add(*FString::Printf(TEXT("RigVMMemoryDetails_%d"), (int32)ERigVMMemoryType::Literal));
+	TabIds.Add(*FString::Printf(TEXT("RigVMMemoryDetails_%d"), (int32)ERigVMMemoryType::Work));
+	TabIds.Add(*FString::Printf(TEXT("RigVMMemoryDetails_%d"), (int32)ERigVMMemoryType::Debug));
+
+	for (const FName& TabId : TabIds)
+	{
+		TSharedPtr<SDockTab> ActiveTab = GetTabManager()->FindExistingLiveTab(TabId);
+		if(ActiveTab)
+		{
+			if(ActiveTab->GetMetaData<FMemoryTypeMetaData>().IsValid())
+			{
+				ERigVMMemoryType MemoryType = ActiveTab->GetMetaData<FMemoryTypeMetaData>()->MemoryType;			
+				TSharedRef<IDetailsView> DetailsView = StaticCastSharedRef<IDetailsView>(ActiveTab->GetContent());
+				DetailsView->SetObject(InVM->GetMemoryByType(MemoryType));
+			}
+		}
+	}
+
+	UpdateGraphCompilerErrors();
+}
+
+void FRigVMEditor::HandleVMExecutedEvent(URigVMHost* InHost, const FName& InEventName)
+{
+	if (URigVMBlueprint* RigVMBlueprint = GetRigVMBlueprint())
+	{
+		URigVMHost* DebuggedHost = Cast<URigVMHost>(RigVMBlueprint->GetObjectBeingDebugged());
+		if (DebuggedHost == nullptr)
+		{
+			DebuggedHost = GetRigVMHost();
+		}
+
+		if(RigVMBlueprint->RigGraphDisplaySettings.NodeRunLimit > 1)
+		{
+			if(DebuggedHost)
+			{
+				if(URigVM* VM = DebuggedHost->GetVM())
+				{
+					bool bFoundLimitWarnings = false;
+					
+					const FRigVMByteCode& ByteCode = VM->GetByteCode();
+					for(int32 InstructionIndex = 0; InstructionIndex < ByteCode.GetNumInstructions(); InstructionIndex++)
+					{
+						const int32 Count = VM->GetInstructionVisitedCount(DebuggedHost->GetExtendedExecuteContext(), InstructionIndex);
+						if(Count > RigVMBlueprint->RigGraphDisplaySettings.NodeRunLimit)
+						{
+							bFoundLimitWarnings = true;
+
+							const FString CallPath = VM->GetByteCode().GetCallPathForInstruction(InstructionIndex); 
+							if(!KnownInstructionLimitWarnings.Contains(CallPath))
+							{
+								const FString Message = FString::Printf(
+                                    TEXT("Instruction has hit the NodeRunLimit\n(ran %d times, limit is %d)\n\nYou can increase the limit in the class settings."),
+                                    Count,
+                                    RigVMBlueprint->RigGraphDisplaySettings.NodeRunLimit
+                                );
+
+								if(DebuggedHost->GetLog())
+								{
+									DebuggedHost->GetLog()->Entries.Add(
+										FRigVMLog::FLogEntry(EMessageSeverity::Warning, InEventName, InstructionIndex, Message
+									));
+								}
+
+								if(URigVMNode* Subject = Cast<URigVMNode>(VM->GetByteCode().GetSubjectForInstruction(InstructionIndex)))
+								{
+									FNotificationInfo Info(FText::FromString(Message));
+									Info.bFireAndForget = true;
+									Info.FadeOutDuration = 1.0f;
+									Info.ExpireDuration = 5.0f;
+
+									if(URigVMEdGraph* EdGraph = Cast<URigVMEdGraph>(RigVMBlueprint->GetEdGraph(Subject->GetGraph())))
+									{
+										if(UEdGraphNode* Node = EdGraph->FindNodeForModelNodeName(Subject->GetFName()))
+										{
+											Info.Hyperlink = FSimpleDelegate::CreateLambda([this, Node] ()
+	                                        {
+	                                            JumpToHyperlink(Node, false);
+	                                        });
+									
+											Info.HyperlinkText = FText::FromString(Subject->GetName());
+										}
+									}
+
+									TSharedPtr<SNotificationItem> NotificationPtr = FSlateNotificationManager::Get().AddNotification(Info);
+									NotificationPtr->SetCompletionState(SNotificationItem::CS_Fail);
+								}
+
+								KnownInstructionLimitWarnings.Add(CallPath, Message);
+							}
+						}
+					}
+
+					if(!bFoundLimitWarnings)
+					{
+						KnownInstructionLimitWarnings.Reset();
+					}
+				}
+			}
+		}
+
+		if(RigVMBlueprint->VMRuntimeSettings.bEnableProfiling)
+		{
+			if(DebuggedHost)
+			{
+				RigVMBlueprint->RigGraphDisplaySettings.TotalMicroSeconds = DebuggedHost->GetExtendedExecuteContext().LastExecutionMicroSeconds;
+			}
+
+			if(RigVMBlueprint->RigGraphDisplaySettings.bAutoDetermineRange)
+			{
+				if(RigVMBlueprint->RigGraphDisplaySettings.LastMaxMicroSeconds < 0.0)
+				{
+					RigVMBlueprint->RigGraphDisplaySettings.LastMinMicroSeconds = RigVMBlueprint->RigGraphDisplaySettings.MinMicroSeconds; 
+					RigVMBlueprint->RigGraphDisplaySettings.LastMaxMicroSeconds = RigVMBlueprint->RigGraphDisplaySettings.MaxMicroSeconds;
+				}
+				else if(RigVMBlueprint->RigGraphDisplaySettings.MaxMicroSeconds >= 0.0)
+				{
+					const double T = 0.05;
+					RigVMBlueprint->RigGraphDisplaySettings.LastMinMicroSeconds = FMath::Lerp<double>(RigVMBlueprint->RigGraphDisplaySettings.LastMinMicroSeconds, RigVMBlueprint->RigGraphDisplaySettings.MinMicroSeconds, T); 
+					RigVMBlueprint->RigGraphDisplaySettings.LastMaxMicroSeconds = FMath::Lerp<double>(RigVMBlueprint->RigGraphDisplaySettings.LastMaxMicroSeconds, RigVMBlueprint->RigGraphDisplaySettings.MaxMicroSeconds, T); 
+				}
+
+				RigVMBlueprint->RigGraphDisplaySettings.MinMicroSeconds = DBL_MAX; 
+				RigVMBlueprint->RigGraphDisplaySettings.MaxMicroSeconds = (double)INDEX_NONE;
+			}
+			else
+			{
+				RigVMBlueprint->RigGraphDisplaySettings.LastMinMicroSeconds = RigVMBlueprint->RigGraphDisplaySettings.MinMicroSeconds; 
+				RigVMBlueprint->RigGraphDisplaySettings.LastMaxMicroSeconds = RigVMBlueprint->RigGraphDisplaySettings.MaxMicroSeconds;
+			}
+		}
+	}
+
+	UpdateGraphCompilerErrors();
+}
+
+void FRigVMEditor::HandleVMExecutionHalted(const int32 InstructionIndex, UObject* InNodeObject, const FName& InEntryName)
+{
+	if (HaltedAtNode == InNodeObject)
+	{
+		return;
+	}
+		
+	if (URigVMNode* InNode = Cast<URigVMNode>(InNodeObject))
+	{
+		SetHaltedNode(InNode);
+		
+		if (URigVMBlueprint* RigVMBlueprint = GetRigVMBlueprint())
+		{
+			if (RigVMBlueprint->GetAllModels().Contains(InNode->GetGraph()))
+			{
+				if(URigVMEdGraph* EdGraph = Cast<URigVMEdGraph>(RigVMBlueprint->GetEdGraph(InNode->GetGraph())))
+				{
+					if(UEdGraphNode* EdNode = EdGraph->FindNodeForModelNodeName(InNode->GetFName()))
+					{
+						JumpToHyperlink(EdNode, false);
+					}
+				}
+			}
+		}
+	}
+	else 
+	{
+		if (InEntryName == GetRigVMHost()->GetEventQueue().Last())
+		{
+			SetHaltedNode(nullptr);
+		}
+	}
+}
+
+void FRigVMEditor::SetHaltedNode(URigVMNode* Node)
+{
+	if (HaltedAtNode)
+	{
+		HaltedAtNode->SetExecutionIsHaltedAtThisNode(false);
+	}
+	HaltedAtNode = Node;
+	if (HaltedAtNode)
+	{
+		HaltedAtNode->SetExecutionIsHaltedAtThisNode(true);
+	}
+}
+
+void FRigVMEditor::NotifyPreChange(FProperty* PropertyAboutToChange)
+{
+	FBlueprintEditor::NotifyPreChange(PropertyAboutToChange);
+
+	if (URigVMBlueprint* RigVMBlueprint = GetRigVMBlueprint())
+	{
+		RigVMBlueprint->Modify();
+	}
+}
+
+void FRigVMEditor::NotifyPostChange(const FPropertyChangedEvent& PropertyChangedEvent, FProperty* PropertyThatChanged)
+{
+	FBlueprintEditor::NotifyPostChange(PropertyChangedEvent, PropertyThatChanged);
+
+	// we need to listen to changes for variables on the blueprint here since
+	// OnFinishedChangingProperties is called only for top level property changes.
+	// changes on a lower level property like transform under a user defined struct
+	// only go through this.
+	URigVMBlueprint* RigVMBlueprint = GetRigVMBlueprint();
+	if(GetRigVMHost() && RigVMBlueprint)
+	{
+		bool bUseCDO = false; 
+		if(PropertyChangedEvent.GetNumObjectsBeingEdited() == 1)
+		{
+			bUseCDO = PropertyChangedEvent.GetObjectBeingEdited(0)->HasAnyFlags(RF_ClassDefaultObject);
+		}
+		
+		const FName VarName = PropertyChangedEvent.MemberProperty->GetFName();
+		for(FBPVariableDescription& NewVariable : RigVMBlueprint->NewVariables)
+		{
+			if(NewVariable.VarName == VarName)
+			{
+				UpdateDefaultValueForVariable(NewVariable, bUseCDO);
+				break;
+			}
+		}
+	}
+}
+
+void FRigVMEditor::OnFinishedChangingProperties(const FPropertyChangedEvent& PropertyChangedEvent)
+{
+	DECLARE_SCOPE_HIERARCHICAL_COUNTER_FUNC()
+
+	URigVMBlueprint* RigVMBlueprint = GetRigVMBlueprint();
+
+	if (RigVMBlueprint)
+	{
+		if (PropertyChangedEvent.MemberProperty->GetNameCPP() == GET_MEMBER_NAME_STRING_CHECKED(URigVMBlueprint, VMCompileSettings))
+		{
+			RigVMBlueprint->RecompileVM();
+		}
+
+		else if (PropertyChangedEvent.MemberProperty->GetNameCPP() == GET_MEMBER_NAME_STRING_CHECKED(URigVMBlueprint, VMRuntimeSettings))
+		{
+			RigVMBlueprint->VMRuntimeSettings.Validate();
+			RigVMBlueprint->PropagateRuntimeSettingsFromBPToInstances();
+		}
+	}
+}
+
+void FRigVMEditor::OnPropertyChanged(UObject* InObject, FPropertyChangedEvent& InEvent)
+{
+	URigVMBlueprint* RigVMBlueprint = GetRigVMBlueprint();
+
+	if (RigVMBlueprint && InObject == RigVMBlueprint)
+	{
+		// if the models have changed - we may need to close a document
+		if(InEvent.MemberProperty == RigVMBlueprint->GetClass()->FindPropertyByName(GET_MEMBER_NAME_CHECKED(URigVMBlueprint, RigVMClient)) ||
+			InEvent.MemberProperty == RigVMBlueprint->GetClass()->FindPropertyByName(GET_MEMBER_NAME_CHECKED(URigVMBlueprint, UbergraphPages)))
+		{
+			DocumentManager->CleanInvalidTabs();
+		}
+	}
+}
+
+void FRigVMEditor::OnWrappedPropertyChangedChainEvent(URigVMDetailsViewWrapperObject* InWrapperObject,
+	const FString& InPropertyPath, FPropertyChangedChainEvent& InPropertyChangedChainEvent)
+{
+	check(InWrapperObject);
+	check(!WrapperObjects.IsEmpty());
+
+	TGuardValue<bool> SuspendDetailsPanelRefresh(bSuspendDetailsPanelRefresh, true);
+
+	URigVMBlueprint* RigVMBlueprint = GetRigVMBlueprint();
+
+	FString PropertyPath = InPropertyPath;
+	if(UScriptStruct* WrappedStruct = InWrapperObject->GetWrappedStruct())
+	{
+		if(WrappedStruct->IsChildOf(FRigVMGraphVariableDescription::StaticStruct()))
+		{
+			check(WrappedStruct == WrapperObjects[0]->GetWrappedStruct());
+			
+			const FRigVMGraphVariableDescription VariableDescription = InWrapperObject->GetContent<FRigVMGraphVariableDescription>();
+			URigVMGraph* Graph = CastChecked<URigVMGraph>(InWrapperObject->GetSubject());
+			URigVMController* Controller = RigVMBlueprint->GetController(Graph);
+			if (PropertyPath == TEXT("Name") && MyBlueprintWidget.IsValid())
+			{
+				if (FEdGraphSchemaAction_BlueprintVariableBase* VariableAcion = MyBlueprintWidget->SelectionAsBlueprintVariable())
+				{
+					const FName OldVariableName = VariableAcion->GetVariableName();
+					if (!OldVariableName.IsNone())
+					{
+						for (FRigVMGraphVariableDescription& Variable : Graph->GetLocalVariables())
+						{
+							if (Variable.Name == OldVariableName)
+							{
+								Controller->RenameLocalVariable(OldVariableName, VariableDescription.Name);
+								break;
+							}
+						}
+					}
+				}
+				RefreshMyBlueprint();
+				GetRigVMBlueprint()->RequestAutoVMRecompilation();
+			}
+			else if (PropertyPath == TEXT("CPPType") || PropertyPath == TEXT("CPPTypeObject"))
+			{			
+				for (FRigVMGraphVariableDescription& Variable : Graph->GetLocalVariables())
+				{
+					if (Variable.Name == VariableDescription.Name)
+					{
+						Controller->SetLocalVariableType(Variable.Name, VariableDescription.CPPType, VariableDescription.CPPTypeObject);
+						break;
+					}
+				}
+				GetRigVMBlueprint()->RequestAutoVMRecompilation();
+			}
+			else if (PropertyPath == TEXT("DefaultValue"))
+			{
+				FRigVMControllerNotifGuard NotifGuard(Controller, true);
+				for (FRigVMGraphVariableDescription& Variable : Graph->GetLocalVariables())
+				{
+					if (Variable.Name == VariableDescription.Name)
+					{
+						Controller->SetLocalVariableDefaultValue(Variable.Name, VariableDescription.DefaultValue, true, true);
+						break;
+					}
+				}
+
+				// Do not recompile now! That destroys the object that is currently being displayed (the literal memory storage), and can cause a crash.
+				// The user has to manually trigger the recompilation.
+			}		
+		}
+	}
+	else if(!InWrapperObject->GetWrappedNodeNotation().IsEmpty())
+	{
+		FName RootPinName = InPropertyChangedChainEvent.PropertyChain.GetHead()->GetValue()->GetFName();
+		FProperty* TargetProperty = WrapperObjects[0]->GetClass()->FindPropertyByName(RootPinName);
+		uint8* FirstPropertyStorage = TargetProperty->ContainerPtrToValuePtr<uint8>(WrapperObjects[0].Get());
+
+		URigVMNode* Node = CastChecked<URigVMNode>(InWrapperObject->GetSubject());
+
+		FString DefaultValue = FRigVMStruct::ExportToFullyQualifiedText(TargetProperty, FirstPropertyStorage);
+
+		if(TargetProperty->IsA<FStrProperty>() || TargetProperty->IsA<FNameProperty>())
+		{
+			DefaultValue.TrimCharInline(TEXT('\"'), nullptr);
+		}
+		
+		URigVMController* Controller = GetRigVMBlueprint()->GetController(Node->GetGraph());
+
+		if (!DefaultValue.IsEmpty())
+		{
+			FString PinPath = FString::Printf(TEXT("%s.%s"), *Node->GetName(), *RootPinName.ToString());
+			const bool bInteractive = InPropertyChangedChainEvent.ChangeType == EPropertyChangeType::Interactive;
+			Controller->SetPinDefaultValue(PinPath, DefaultValue, true, !bInteractive, true, !bInteractive);
+		}
+	}
+}
+
+void FRigVMEditor::OnRequestLocalizeFunctionDialog(FRigVMGraphFunctionIdentifier& InFunction,
+	URigVMBlueprint* InTargetBlueprint, bool bForce)
+{
+	check(InTargetBlueprint);
+
+	if(InTargetBlueprint != GetRigVMBlueprint())
+	{
+		return;
+	}
+	
+	if(URigVMController* TargetController = InTargetBlueprint->GetController(InTargetBlueprint->GetDefaultModel()))
+	{
+		bool bIsPublic;
+		if (FRigVMGraphFunctionData::FindFunctionData(InFunction, &bIsPublic))
+		{
+			if (bForce || bIsPublic)
+			{
+				TSharedRef<SRigVMGraphFunctionLocalizationDialog> LocalizationDialog = SNew(SRigVMGraphFunctionLocalizationDialog)
+							.Function(InFunction)
+							.TargetBlueprint(InTargetBlueprint);
+
+				if (LocalizationDialog->ShowModal() != EAppReturnType::Cancel)
+				{
+					TargetController->LocalizeFunctions(LocalizationDialog->GetFunctionsToLocalize(), true, true, true);
+				}
+			}
+		}
+	}
+}
+
+FRigVMController_BulkEditResult FRigVMEditor::OnRequestBulkEditDialog(URigVMBlueprint* InBlueprint,
+	URigVMController* InController, URigVMLibraryNode* InFunction, ERigVMControllerBulkEditType InEditType)
+{
+	if (bAllowBulkEdits)
+	{
+		FRigVMController_BulkEditResult Result;
+		Result.bCanceled = false; 
+		Result.bSetupUndoRedo = false;
+		return Result;
+	}
+	
+	const TArray<FAssetData> FirstLevelReferenceAssets = InController->GetAffectedAssets(InEditType, false);
+	if(FirstLevelReferenceAssets.Num() == 0)
+	{
+		return FRigVMController_BulkEditResult();
+	}
+	
+	TSharedRef<SRigVMGraphFunctionBulkEditDialog> BulkEditDialog = SNew(SRigVMGraphFunctionBulkEditDialog)
+	.Blueprint(InBlueprint)
+	.Controller(InController)
+	.Function(InFunction)
+	.EditType(InEditType);
+
+	FRigVMController_BulkEditResult Result;
+	Result.bCanceled = BulkEditDialog->ShowModal() == EAppReturnType::Cancel; 
+	Result.bSetupUndoRedo = false;
+
+	if (!Result.bCanceled)
+	{
+		bAllowBulkEdits = true;
+	}
+	
+	return Result;
+}
+
+bool FRigVMEditor::OnRequestBreakLinksDialog(TArray<URigVMLink*> InLinks)
+{
+	if(InLinks.Num() == 0)
+	{
+		return true;
+	}
+
+	TSharedRef<SRigVMGraphBreakLinksDialog> BreakLinksDialog = SNew(SRigVMGraphBreakLinksDialog)
+	.Links(InLinks)
+	.OnFocusOnLink(FRigVMOnFocusOnLinkRequestedDelegate::CreateLambda([&](URigVMLink* InLink)
+	{
+		HandleJumpToHyperlink(InLink);
+	}));
+
+	return BreakLinksDialog->ShowModal() == EAppReturnType::Ok; 
+}
+
+TRigVMTypeIndex FRigVMEditor::OnRequestPinTypeSelectionDialog(const TArray<TRigVMTypeIndex>& InTypes)
+{
+		if(InTypes.Num() == 0)
+	{
+		return true;
+	}
+
+	TRigVMTypeIndex Answer = INDEX_NONE;
+
+	const FRigVMRegistry& Registry = FRigVMRegistry::Get();
+
+	TArray<TSharedPtr<FName>> TypeNames;
+	TMap<FName, uint8> TypeNameToIndex;
+	TypeNames.Reserve(InTypes.Num());
+	for (int32 i=0; i<InTypes.Num(); ++i)
+	{
+		const TRigVMTypeIndex& TypeIndex = InTypes[i];
+		TRigVMTypeIndex FinalType = TypeIndex;
+		if (FinalType == RigVMTypeUtils::TypeIndex::Float)
+		{
+			FinalType = RigVMTypeUtils::TypeIndex::Double;
+		}
+		if (FinalType == RigVMTypeUtils::TypeIndex::FloatArray)
+		{
+			FinalType = RigVMTypeUtils::TypeIndex::DoubleArray;
+		}
+
+		const FRigVMTemplateArgumentType& ArgumentType = Registry.GetType(FinalType);
+		if (!TypeNames.ContainsByPredicate([&ArgumentType](const TSharedPtr<FName>& InName)
+		{
+			return *InName.Get() == ArgumentType.CPPType;
+		}))
+		{
+			TypeNames.AddUnique(MakeShared<FName>(ArgumentType.CPPType));
+			TypeNameToIndex.Add(ArgumentType.CPPType, i);
+		}
+	}
+	TSharedPtr< SWindow > Window = SNew(SWindow)
+		.Title(LOCTEXT("SelectPinType", "Select Pin Type"))
+		.ScreenPosition(FSlateApplication::Get().GetCursorPos())
+		.SizingRule(ESizingRule::Autosized)
+		.AutoCenter(EAutoCenter::None)
+		.SupportsMaximize(false)
+		.SupportsMinimize(false)
+		[
+			SNew( SBorder )
+			.Padding( 4.f )
+			.BorderImage( FAppStyle::GetBrush( "ToolPanel.GroupBorder" ) )
+			[
+				SNew(SBox)
+				.MaxDesiredHeight(static_cast<float>(300))
+				[
+					SNew(SVerticalBox)
+					+SVerticalBox::Slot()
+					[
+						SNew(SBox)
+						.MaxDesiredHeight(static_cast<float>(300))
+						[
+							SNew(SScrollBox)
+							+SScrollBox::Slot()
+							[
+								SNew(SListView<TSharedPtr<FName>>)
+									.ListItemsSource(&TypeNames)
+									.OnGenerateRow_Lambda([this, &Registry, &TypeNameToIndex, &InTypes](const TSharedPtr<FName> InItem, const TSharedRef<STableViewBase>& Owner)
+									{
+										TRigVMTypeIndex TypeIndex = InTypes[TypeNameToIndex.FindChecked(*InItem.Get())];
+										const FRigVMTemplateArgumentType Type = FRigVMRegistry::Get().GetType(TypeIndex);
+										const bool bIsArray = Type.IsArray();
+										static const FName TypeIcon(TEXT("Kismet.VariableList.TypeIcon"));
+										static const FName ArrayTypeIcon(TEXT("Kismet.VariableList.ArrayTypeIcon"));
+
+										const FEdGraphPinType PinType = RigVMTypeUtils::PinTypeFromTypeIndex(TypeIndex);
+										const URigVMEdGraphSchema* Schema = CastChecked<URigVMEdGraphSchema>(GetRigVMBlueprint()->GetRigVMEdGraphSchemaClass()->GetDefaultObject());
+										const FLinearColor Color = Schema->GetPinTypeColor(PinType);
+										
+										return SNew(STableRow<TSharedPtr<FString>>, Owner)
+												.Padding(FMargin(16, 4, 16, 4))
+												[
+													
+													SNew(SHorizontalBox)
+													+ SHorizontalBox::Slot()
+													.AutoWidth()
+													.VAlign(VAlign_Center)
+													[
+														SNew(SBox)
+														.HeightOverride(16.0f)
+														[
+															SNew(SImage)
+															.Image(bIsArray ? FAppStyle::GetBrush(ArrayTypeIcon) : FAppStyle::GetBrush(TypeIcon))
+															.ColorAndOpacity(Color)
+														]
+													]
+
+													+ SHorizontalBox::Slot()
+													[
+														SNew(STextBlock).Text(FText::FromName(*InItem.Get()))
+													]
+												];
+									})
+									.OnSelectionChanged_Lambda([&Answer, &TypeNames, &TypeNameToIndex, &InTypes](const TSharedPtr<FName> InName, ESelectInfo::Type)
+									{
+										Answer = InTypes[TypeNameToIndex.FindChecked(*InName.Get())];
+										FSlateApplication::Get().GetActiveModalWindow()->RequestDestroyWindow();
+									})
+							]
+						]
+					]
+				]
+			]
+		];
+
+	GEditor->EditorAddModalWindow(Window.ToSharedRef());
+	return Answer;
+}
+
+void FRigVMEditor::HandleJumpToHyperlink(const UObject* InSubject)
+{
+	URigVMBlueprint* RigBlueprint = GetRigVMBlueprint();
+	if(RigBlueprint == nullptr)
+	{
+		return;
+	}
+
+	const URigVMGraph* GraphToJumpTo = nullptr;
+	const URigVMNode* NodeToJumpTo = nullptr;
+	const URigVMPin* PinToJumpTo = nullptr;
+	if(const URigVMNode* Node = Cast<URigVMNode>(InSubject))
+	{
+		GraphToJumpTo = Node->GetGraph();
+		NodeToJumpTo = Node;
+	}
+	else if(const URigVMPin* Pin = Cast<URigVMPin>(InSubject))
+	{
+		GraphToJumpTo = Pin->GetGraph();
+		NodeToJumpTo = Pin->GetNode();
+		PinToJumpTo = Pin;
+	}
+	else if(const URigVMLink* Link = Cast<URigVMLink>(InSubject))
+	{
+		GraphToJumpTo = Link->GetGraph();
+		if(const URigVMPin* TargetPin = ((URigVMLink*)Link)->GetTargetPin())
+		{
+			NodeToJumpTo = TargetPin->GetNode();
+			PinToJumpTo = TargetPin;
+		}
+	}
+
+	if (GraphToJumpTo && NodeToJumpTo)
+	{
+		if(URigVMEdGraph* EdGraph = Cast<URigVMEdGraph>(RigBlueprint->GetEdGraph(NodeToJumpTo->GetGraph())))
+		{
+			if(const URigVMEdGraphNode* EdGraphNode = Cast<URigVMEdGraphNode>(EdGraph->FindNodeForModelNodeName(NodeToJumpTo->GetFName())))
+			{
+				if(PinToJumpTo)
+				{
+					if(const UEdGraphPin* EdGraphPin = EdGraphNode->FindPin(PinToJumpTo->GetSegmentPath(true)))
+					{
+						JumpToPin(EdGraphPin);
+						return;
+					}
+				}
+				
+				JumpToNode(EdGraphNode);
+				return;
+			}
+			
+			JumpToHyperlink(EdGraph);
+		}
+	}
+}
+
+bool FRigVMEditor::UpdateDefaultValueForVariable(FBPVariableDescription& InVariable, bool bUseCDO)
+{
+	bool bAnyValueChanged = false;
+	if (URigVMBlueprint* RigVMBlueprint = GetRigVMBlueprint())
+	{
+		UClass* GeneratedClass = RigVMBlueprint->GeneratedClass;
+		UObject* ObjectContainer = bUseCDO ? GeneratedClass->GetDefaultObject() : RigVMBlueprint->GetObjectBeingDebugged();
+		if(ObjectContainer)
+		{
+			FProperty* TargetProperty = FindFProperty<FProperty>(GeneratedClass, InVariable.VarName);
+
+			if (TargetProperty)
+			{
+				FString NewDefaultValue;
+				const uint8* Container = (const uint8*)ObjectContainer;
+				FBlueprintEditorUtils::PropertyValueToString(TargetProperty, Container, NewDefaultValue, nullptr);
+				if (InVariable.DefaultValue != NewDefaultValue)
+				{
+					InVariable.DefaultValue = NewDefaultValue;
+					bAnyValueChanged = true;
+				}
+			}
+		}
+	}
+	return bAnyValueChanged;
+}
+
+void FRigVMEditor::UpdateRigVMHost()
+{
+	DECLARE_SCOPE_HIERARCHICAL_COUNTER_FUNC()
+
+	URigVMBlueprint* Blueprint = GetRigVMBlueprint();
+	if(UClass* Class = Blueprint->GeneratedClass)
+	{
+		if (URigVMHost* CurrentHost = GetRigVMHost())
+		{
+			UpdateRigVMHost_PreClearOldHost(CurrentHost);
+
+			if(!IsValid(CurrentHost))
+			{
+				SetHost(nullptr);
+			}
+			
+			// if this control rig is from a temporary step,
+			// for example the reinstancing class, clear it 
+			// and create a new one!
+			if (CurrentHost->GetClass() != Class)
+			{
+				SetHost(nullptr);
+			}
+		}
+
+		URigVMHost* RigVMHost = GetRigVMHost();
+		if (RigVMHost == nullptr)
+		{
+			RigVMHost = NewObject<URigVMHost>(GetOuterForHost(), Class);
+			SetHost(RigVMHost);
+			
+			// this is editing time rig
+			RigVMHost->SetLog(&RigVMLog);
+
+			RigVMHost->Initialize(true);
+ 		}
+
+#if WITH_EDITOR
+		RigVMHost->SetIsInDebugMode(ExecutionMode == ERigVMEditorExecutionModeType_Debug);
+#endif
+
+		CacheNameLists();
+
+		// Make sure the object being debugged is the preview instance
+		GetBlueprintObj()->SetObjectBeingDebugged(RigVMHost);
+
+		if(!bIsCompilingThroughUI)
+		{
+			Blueprint->SetFlags(RF_Transient);
+			Blueprint->RecompileVM();
+			Blueprint->ClearFlags(RF_Transient);
+		}
+
+		RigVMHost->OnInitialized_AnyThread().AddSP(this, &FRigVMEditor::HandleVMExecutedEvent);
+		RigVMHost->OnExecuted_AnyThread().AddSP(this, &FRigVMEditor::HandleVMExecutedEvent);
+		RigVMHost->RequestInit();
+		RigVMHost->GetExtendedExecuteContext().ExecutionHalted().AddSP(this, &FRigVMEditor::HandleVMExecutionHalted);
+	}
+}
+
+void FRigVMEditor::CacheNameLists()
+{
+	if (URigVMBlueprint* RigVMBlueprint = GetRigVMBlueprint())
+	{
+		TArray<UEdGraph*> EdGraphs;
+		RigVMBlueprint->GetAllGraphs(EdGraphs);
+
+		for (UEdGraph* Graph : EdGraphs)
+		{
+			URigVMEdGraph* RigVMEdGraph = Cast<URigVMEdGraph>(Graph);
+			if (RigVMEdGraph == nullptr)
+			{
+				continue;
+			}
+			RigVMEdGraph->CacheEntryNameList();
+		}
+	}
+}
+
+void FRigVMEditor::OnCreateComment()
+{
+	TSharedPtr<SGraphEditor> GraphEditor = FocusedGraphEdPtr.Pin();
+	if (GraphEditor.IsValid())
+	{
+		if (UEdGraph* Graph = GraphEditor->GetCurrentGraph())
+		{
+			FEdGraphSchemaAction_K2AddComment CommentAction;
+			CommentAction.PerformAction(Graph, NULL, GraphEditor->GetPasteLocation());
+		}
+	}
+}
+
+void FRigVMEditor::SetDetailObjects(const TArray<UObject*>& InObjects)
+{
+	SetDetailObjects(InObjects, true);
+}
+
+void FRigVMEditor::SetDetailObjects(const TArray<UObject*>& InObjects, bool bChangeUISelectionState)
+{
+	if(bSuspendDetailsPanelRefresh)
+	{
+		return;
+	}
+
+	if(InObjects.Num() == 1)
+	{
+		if(URigVMMemoryStorage* Memory = Cast<URigVMMemoryStorage>(InObjects[0]))
+		{
+			FPropertyEditorModule& EditModule = FModuleManager::Get().GetModuleChecked<FPropertyEditorModule>("PropertyEditor");
+
+			FDetailsViewArgs DetailsViewArgs;
+			DetailsViewArgs.NameAreaSettings = FDetailsViewArgs::HideNameArea;
+			DetailsViewArgs.bHideSelectionTip = true;
+
+			TSharedRef<IDetailsView> DetailsView = EditModule.CreateDetailView( DetailsViewArgs );
+			TSharedRef<SDockTab> DockTab = SNew(SDockTab)
+			.Label( LOCTEXT("ControlRigMemoryDetails", "Control Rig Memory Details") )
+			.AddMetaData<FMemoryTypeMetaData>(FMemoryTypeMetaData(Memory->GetMemoryType()))
+			.TabRole(ETabRole::NomadTab)
+			[
+				DetailsView
+			];
+
+			FName TabId = *FString::Printf(TEXT("ControlRigMemoryDetails_%d"), (int32)Memory->GetMemoryType());
+			if(TSharedPtr<SDockTab> ActiveTab = GetTabManager()->FindExistingLiveTab(TabId))
+			{
+				ActiveTab->RequestCloseTab();
+			}
+
+			GetTabManager()->InsertNewDocumentTab(
+				FBlueprintEditorTabs::DetailsID,
+				TabId,
+				FTabManager::FLastMajorOrNomadTab(TEXT("ControlRigMemoryDetails")),
+				DockTab
+			);
+
+			FFunctionGraphTask::CreateAndDispatchWhenReady([DetailsView, InObjects]()
+			{
+				
+				DetailsView->SetObject(InObjects[0]);
+				
+			}, TStatId(), NULL, ENamedThreads::GameThread);
+
+			return;
+		}
+	}
+
+	ClearDetailObject(bChangeUISelectionState);
+
+	if (InObjects.Num() == 1)
+	{
+		if (InObjects[0]->GetClass()->GetDefaultObject() == InObjects[0])
+		{
+			EditClassDefaults_Clicked();
+			return;
+		}
+		else if (InObjects[0] == GetBlueprintObj())
+		{
+			EditGlobalOptions_Clicked();
+			return;
+		}
+	}
+
+	TArray<UObject*> FilteredObjects;
+
+	TArray<URigVMNode*> ModelNodes;
+	for(UObject* InObject : InObjects)
+	{
+		if (URigVMNode* ModelNode = Cast<URigVMNode>(InObject))
+		{
+			ModelNodes.Add(ModelNode);
+		}
+	}
+	
+	for(UObject* InObject : InObjects)
+	{
+		if (URigVMLibraryNode* LibraryNode = Cast<URigVMLibraryNode>(InObject))
+		{
+			if(!LibraryNode->IsA<URigVMFunctionReferenceNode>())
+			{
+				if (UEdGraph* EdGraph = GetRigVMBlueprint()->GetEdGraph(LibraryNode->GetContainedGraph()))
+				{
+					FilteredObjects.AddUnique(EdGraph);
+					continue;
+				}
+			}
+		}
+		else if (Cast<URigVMFunctionEntryNode>(InObject) || Cast<URigVMFunctionReturnNode>(InObject))
+		{
+			if (UEdGraph* EdGraph = GetRigVMBlueprint()->GetEdGraph(CastChecked<URigVMNode>(InObject)->GetGraph()))
+			{
+				FilteredObjects.AddUnique(EdGraph);
+				continue;
+			}
+		}
+		else if (URigVMCommentNode* CommentNode = Cast<URigVMCommentNode>(InObject))
+		{
+			if (URigVMEdGraph* EdGraph = Cast<URigVMEdGraph>(GetRigVMBlueprint()->GetEdGraph(CastChecked<URigVMNode>(InObject)->GetGraph())))
+			{
+				if(UEdGraphNode* EdGraphNode = EdGraph->FindNodeForModelNodeName(CommentNode->GetFName()))
+				{
+					FilteredObjects.AddUnique(EdGraphNode);
+					continue;
+				}
+			}
+		}
+
+		if (URigVMNode* ModelNode = Cast<URigVMNode>(InObject))
+		{
+			// check if we know the dynamic class already
+			const URigVMDetailsViewWrapperObject* CDOWrapper = CastChecked<URigVMDetailsViewWrapperObject>(GetDetailWrapperClass()->GetDefaultObject());
+			(void)CDOWrapper->GetClassForNodes(ModelNodes, false);
+
+			// create the wrapper object
+			if(URigVMDetailsViewWrapperObject* WrapperObject = URigVMDetailsViewWrapperObject::MakeInstance(GetDetailWrapperClass(), GetBlueprintObj(), ModelNodes, ModelNode))
+			{
+				WrapperObject->GetWrappedPropertyChangedChainEvent().AddSP(this, &FRigVMEditor::OnWrappedPropertyChangedChainEvent);
+				WrapperObject->AddToRoot();
+
+				// todo: use transform widget for transforms
+				// todo: use rotation widget for rotations
+				
+				FilteredObjects.Add(WrapperObject);
+				continue;
+			}
+		}
+
+
+		FilteredObjects.Add(InObject);
+	}
+
+	for(UObject* FilteredObject : FilteredObjects)
+	{
+		if(URigVMDetailsViewWrapperObject* WrapperObject = Cast<URigVMDetailsViewWrapperObject>(FilteredObject))
+		{
+			WrapperObjects.Add(TStrongObjectPtr<URigVMDetailsViewWrapperObject>(WrapperObject));
+		}
+	}
+
+	SKismetInspector::FShowDetailsOptions Options;
+	Options.bForceRefresh = true;
+	Inspector->ShowDetailsForObjects(FilteredObjects, Options);
+}
+
+void FRigVMEditor::SetDetailViewForGraph(URigVMGraph* InGraph)
+{
+	check(InGraph);
+	
+	if(bSuspendDetailsPanelRefresh)
+	{
+		return;
+	}
+
+	ClearDetailObject();
+
+	TArray<UObject*> SelectedNodes;
+	TArray<FName> SelectedNodeNames = InGraph->GetSelectNodes();
+	for(FName SelectedNodeName : SelectedNodeNames)
+	{
+		if(URigVMNode* Node = InGraph->FindNodeByName(SelectedNodeName))
+		{
+			SelectedNodes.Add(Node);
+		}
+	}
+
+	SetDetailObjects(SelectedNodes);
+}
+
+void FRigVMEditor::SetDetailViewForFocusedGraph()
+{
+	if(bSuspendDetailsPanelRefresh)
+	{
+		return;
+	}
+
+	URigVMGraph* Model = GetFocusedModel();
+	if(Model == nullptr)
+	{
+		return;
+	}
+
+	SetDetailViewForGraph(Model);
+}
+
+void FRigVMEditor::SetDetailViewForLocalVariable()
+{
+	FName VariableName;
+	TArray< TWeakObjectPtr<UObject> > SelectedObjects = Inspector->GetSelectedObjects();
+	for (TWeakObjectPtr<UObject> SelectedObject : SelectedObjects)
+	{
+		if (SelectedObject.IsValid())
+		{
+			if(URigVMDetailsViewWrapperObject* WrapperObject = Cast<URigVMDetailsViewWrapperObject>(SelectedObject.Get()))
+			{
+				VariableName = WrapperObject->GetContent<FRigVMGraphVariableDescription>().Name;
+				break;
+			}
+		}
+	}
+		
+	SelectLocalVariable(GetFocusedGraph(), VariableName);
+}
+
+void FRigVMEditor::RefreshDetailView()
+{
+	if(DetailViewShowsAnyRigUnit())
+	{
+		SetDetailViewForFocusedGraph();
+	}
+	else if(DetailViewShowsLocalVariable())
+	{
+		SetDetailViewForLocalVariable();	
+	}
+	else
+	{
+		// detail view is showing other stuff; could be a BP variable for example
+		// in this case wrapper objects are not in use, yet still rooted
+		// and preventing their outer objects from getting GCed after a Compile()
+		// so let's take the chance to manually clear them here.
+		ClearDetailsViewWrapperObjects();
+	}
+}
+
+bool FRigVMEditor::DetailViewShowsAnyRigUnit() const
+{
+	if (DetailViewShowsStruct(FRigVMStruct::StaticStruct()))
+	{
+		return true;
+	}
+
+	const TArray< TWeakObjectPtr<UObject> >& SelectedObjects = Inspector->GetSelectedObjects();
+	for (TWeakObjectPtr<UObject> SelectedObject : SelectedObjects)
+	{
+		if (SelectedObject.IsValid())
+		{
+			if(URigVMDetailsViewWrapperObject* WrapperObject = Cast<URigVMDetailsViewWrapperObject>(SelectedObject.Get()))
+			{
+				const FString Notation = WrapperObject->GetWrappedNodeNotation();
+				if(!Notation.IsEmpty())
+				{
+					return true;
+				}
+			}
+		}
+	}
+	
+	return false;
+}
+
+bool FRigVMEditor::DetailViewShowsLocalVariable() const
+{
+	return DetailViewShowsStruct(FRigVMGraphVariableDescription::StaticStruct());
+}
+
+bool FRigVMEditor::DetailViewShowsStruct(UScriptStruct* InStruct) const
+{
+	TArray< TWeakObjectPtr<UObject> > SelectedObjects = Inspector->GetSelectedObjects();
+	for (TWeakObjectPtr<UObject> SelectedObject : SelectedObjects)
+	{
+		if (SelectedObject.IsValid())
+		{
+			if(URigVMDetailsViewWrapperObject* WrapperObject = Cast<URigVMDetailsViewWrapperObject>(SelectedObject.Get()))
+			{
+				if(const UScriptStruct* WrappedStruct = WrapperObject->GetWrappedStruct())
+				{
+					if (WrappedStruct->IsChildOf(InStruct))
+					{
+						return true;
+					}
+				}
+			}
+		}
+	}
+	return false;
+}
+
+void FRigVMEditor::ClearDetailObject(bool bChangeUISelectionState)
+{
+	if(bSuspendDetailsPanelRefresh)
+	{
+		return;
+	}
+
+	ClearDetailsViewWrapperObjects();
+	
+	Inspector->GetPropertyView()->SetObjects(TArray<UObject*>(), true); // clear property view synchronously
+	Inspector->ShowDetailsForObjects(TArray<UObject*>());
+	Inspector->ShowSingleStruct(TSharedPtr<FStructOnScope>());
+
+	if (bChangeUISelectionState)
+	{
+		SetUISelectionState(FBlueprintEditor::SelectionState_Graph);
+	}
+}
+
+void FRigVMEditor::ClearDetailsViewWrapperObjects()
+{
+	for(const TStrongObjectPtr<URigVMDetailsViewWrapperObject>& WrapperObjectPtr : WrapperObjects)
+	{
+		if(WrapperObjectPtr.IsValid())
+		{
+			URigVMDetailsViewWrapperObject* WrapperObject = WrapperObjectPtr.Get();
+			WrapperObject->RemoveFromRoot();
+			WrapperObject->Rename(nullptr, GetTransientPackage(), REN_ForceNoResetLoaders | REN_DoNotDirty | REN_DontCreateRedirectors | REN_NonTransactional);
+			WrapperObject->MarkAsGarbage();
+		}
+	}
+	WrapperObjects.Reset();
+}
+
+void FRigVMEditor::SetHost(URigVMHost* InHost)
+{
+	Host = InHost;
+	if(Host && IsValid(Host))
+	{
+		OnPreviewHostUpdated().Broadcast(this);
+	}
+}
+
+URigVMGraph* FRigVMEditor::GetFocusedModel() const
+{
+	const URigVMBlueprint* Blueprint = GetRigVMBlueprint();
+	if (Blueprint == nullptr)
+	{
+		return nullptr;
+	}
+
+	const URigVMEdGraph* EdGraph = Cast<URigVMEdGraph>(GetFocusedGraph());
+	return Blueprint->GetModel(EdGraph);
+}
+
+URigVMController* FRigVMEditor::GetFocusedController() const
 {
+	URigVMBlueprint* Blueprint = GetRigVMBlueprint();
+	if (Blueprint == nullptr)
+	{
+		return nullptr;
+	}
+	return Blueprint->GetOrCreateController(GetFocusedModel());
 }
 
-URigVMBlueprint* FRigVMEditor::GetRigVMBlueprint() const
+TSharedPtr<SGraphEditor> FRigVMEditor::GetGraphEditor(UEdGraph* InEdGraph) const
 {
-	return Cast<URigVMBlueprint>(GetBlueprintObj());
+	TArray< TSharedPtr<SDockTab> > GraphEditorTabs;
+	DocumentManager->FindAllTabsForFactory(GraphEditorTabFactoryPtr, /*out*/ GraphEditorTabs);
+
+	for (const TSharedPtr<SDockTab>& GraphEditorTab : GraphEditorTabs)
+	{
+		TSharedRef<SGraphEditor> Editor = StaticCastSharedRef<SGraphEditor>((GraphEditorTab)->GetContent());
+		if (Editor->GetCurrentGraph() == InEdGraph)
+		{
+			return Editor;
+		}
+	}
+
+	return TSharedPtr<SGraphEditor>();
 }
 
-URigVMHost* FRigVMEditor::GetRigVMHost() const
+void FRigVMEditor::ExtendMenu()
 {
-	return Host;
+	if(MenuExtender.IsValid())
+	{
+		RemoveMenuExtender(MenuExtender);
+		MenuExtender.Reset();
+	}
+
+	MenuExtender = MakeShareable(new FExtender);
+
+	AddMenuExtender(MenuExtender);
+
+	// add extensible menu if exists
+	FRigVMEditorModule& RigVMEditorModule = FModuleManager::LoadModuleChecked<FRigVMEditorModule>("RigVMEditor");
+	AddMenuExtender(RigVMEditorModule.GetMenuExtensibilityManager()->GetAllExtenders(GetToolkitCommands(), GetEditingObjects()));
 }
 
-UClass* FRigVMEditor::GetDetailWrapperClass() const
+void FRigVMEditor::ExtendToolbar()
 {
-	return URigVMDetailsViewWrapperObject::StaticClass();
+	DECLARE_SCOPE_HIERARCHICAL_COUNTER_FUNC()
+
+	// If the ToolbarExtender is valid, remove it before rebuilding it
+	if(ToolbarExtender.IsValid())
+	{
+		RemoveToolbarExtender(ToolbarExtender);
+		ToolbarExtender.Reset();
+	}
+
+	ToolbarExtender = MakeShareable(new FExtender);
+
+	AddToolbarExtender(ToolbarExtender);
+
+	FRigVMEditorModule& RigVMEditorModule = FModuleManager::LoadModuleChecked<FRigVMEditorModule>("RigVMEditor");
+	AddToolbarExtender(RigVMEditorModule.GetToolBarExtensibilityManager()->GetAllExtenders(GetToolkitCommands(), GetEditingObjects()));
+
+	TArray<FRigVMEditorModule::FRigVMEditorToolbarExtender> ToolbarExtenderDelegates = RigVMEditorModule.GetAllRigVMEditorToolbarExtenders();
+
+	for (auto& ToolbarExtenderDelegate : ToolbarExtenderDelegates)
+	{
+		if(ToolbarExtenderDelegate.IsBound())
+		{
+			AddToolbarExtender(ToolbarExtenderDelegate.Execute(GetToolkitCommands(), SharedThis(this)));
+		}
+	}
+
+	ToolbarExtender->AddToolBarExtension(
+		"Asset",
+		EExtensionHook::After,
+		GetToolkitCommands(),
+		FToolBarExtensionDelegate::CreateSP(this, &FRigVMEditor::FillToolbar, true)
+	);
 }
 
-void FRigVMEditor::HandleModifiedEvent(ERigVMGraphNotifType InNotifType, URigVMGraph* InGraph, UObject* InSubject)
+void FRigVMEditor::FillToolbar(FToolBarBuilder& ToolbarBuilder, bool bEndSection)
+{
+	ToolbarBuilder.BeginSection("Toolbar");
+	{
+		ToolbarBuilder.AddToolBarButton(
+			FRigVMEditorCommands::Get().ToggleEventQueue,
+			NAME_None, 
+			TAttribute<FText>::Create(TAttribute<FText>::FGetter::CreateSP(this, &FRigVMEditor::GetEventQueueLabel)),
+			TAttribute<FText>(), 
+			TAttribute<FSlateIcon>::Create(TAttribute<FSlateIcon>::FGetter::CreateSP(this, &FRigVMEditor::GetEventQueueIcon))
+		);
+
+		FUIAction DefaultAction;
+		ToolbarBuilder.AddComboButton(
+			DefaultAction,
+			FOnGetContent::CreateSP(this, &FRigVMEditor::GenerateEventQueueMenuContent),
+			LOCTEXT("EventQueue_Label", "Available Events"),
+			LOCTEXT("EventQueue_ToolTip", "Pick between different events / modes for testing the Control Rig"),
+			FSlateIcon(FAppStyle::GetAppStyleSetName(), "LevelEditor.Recompile"),
+			true);
+
+		ToolbarBuilder.AddToolBarButton(FRigVMEditorCommands::Get().AutoCompileGraph,
+			NAME_None, TAttribute<FText>(), TAttribute<FText>(), FSlateIcon(FRigVMEditorStyle::Get().GetStyleSetName(), "RigVM.AutoCompileGraph"));
+
+		ToolbarBuilder.AddWidget(SNew(SBlueprintEditorSelectedDebugObjectWidget, SharedThis(this)));
+
+		ToolbarBuilder.AddSeparator();
+		ToolbarBuilder.AddToolBarButton(
+			FRigVMEditorCommands::Get().ToggleExecutionMode,
+			NAME_None, 
+			TAttribute<FText>::Create(TAttribute<FText>::FGetter::CreateSP(this, &FRigVMEditor::GetExecutionModeLabel)),
+			TAttribute<FText>(), 
+			TAttribute<FSlateIcon>::Create(TAttribute<FSlateIcon>::FGetter::CreateSP(this, &FRigVMEditor::GetExecutionModeIcon))
+		);
+		
+		FUIAction DefaultExecutionMode;
+		ToolbarBuilder.AddComboButton(
+			DefaultExecutionMode,
+			FOnGetContent::CreateSP(this, &FRigVMEditor::GenerateExecutionModeMenuContent),
+			LOCTEXT("ExecutionMode_Label", "Execution Modes"),
+			LOCTEXT("ExecutionMode_ToolTip", "Pick between different execution modes for testing the Control Rig"),
+			FSlateIcon(FAppStyle::GetAppStyleSetName(), "LevelEditor.Recompile"),
+			true);
+
+		ToolbarBuilder.BeginStyleOverride(FName("Toolbar.BackplateLeftPlay"));
+		ToolbarBuilder.AddToolBarButton(FRigVMEditorCommands::Get().ResumeExecution,
+			NAME_None, TAttribute<FText>(), TAttribute<FText>(), FSlateIcon(FAppStyle::GetAppStyleSetName(), "PlayWorld.ResumePlaySession"));
+
+		ToolbarBuilder.BeginStyleOverride(FName("Toolbar.BackplateLeft"));
+		ToolbarBuilder.AddToolBarButton(FRigVMEditorCommands::Get().ShowCurrentStatement,
+			NAME_None, TAttribute<FText>(), TAttribute<FText>(), FSlateIcon(FAppStyle::GetAppStyleSetName(), "PlayWorld.ShowCurrentStatement"));
+
+		ToolbarBuilder.BeginStyleOverride(FName("Toolbar.BackplateCenter"));
+		ToolbarBuilder.AddToolBarButton(FRigVMEditorCommands::Get().StepOver,
+			NAME_None, TAttribute<FText>(), TAttribute<FText>(), FSlateIcon(FAppStyle::GetAppStyleSetName(), "PlayWorld.StepOver"));
+		ToolbarBuilder.AddToolBarButton(FRigVMEditorCommands::Get().StepInto,
+			NAME_None, TAttribute<FText>(), TAttribute<FText>(), FSlateIcon(FAppStyle::GetAppStyleSetName(), "PlayWorld.StepInto"));
+
+		ToolbarBuilder.BeginStyleOverride(FName("Toolbar.BackplateRight"));
+		ToolbarBuilder.AddToolBarButton(FRigVMEditorCommands::Get().StepOut,
+			NAME_None, TAttribute<FText>(), TAttribute<FText>(), FSlateIcon(FAppStyle::GetAppStyleSetName(), "PlayWorld.StepOut"));
+
+		ToolbarBuilder.EndStyleOverride();
+	}
+
+	if(bEndSection)
+	{
+		ToolbarBuilder.EndSection();
+	}
+}
+
+void FRigVMEditor::HandleHideItem()
+{
+	DECLARE_SCOPE_HIERARCHICAL_COUNTER_FUNC()
+
+	URigVMBlueprint* RigVMBlueprint = GetRigVMBlueprint();
+	if(RigVMBlueprint == nullptr)
+	{
+		return;
+	}
+
+	TSet<UObject*> SelectedNodes = GetSelectedNodes();
+	if(SelectedNodes.Num() > 0)
+	{
+		FScopedTransaction Transaction(LOCTEXT("HideRigItem", "Hide rig item"));
+
+		RigVMBlueprint->Modify();
+
+		for(UObject* SelectedNodeObject : SelectedNodes)
+		{
+			if(URigVMEdGraphNode* SelectedNode = Cast<URigVMEdGraphNode>(SelectedNodeObject))
+			{
+				FBlueprintEditorUtils::RemoveNode(RigVMBlueprint, SelectedNode, true);
+			}
+		}
+	}
+}
+
+bool FRigVMEditor::CanHideItem() const
 {
+	return GetNumberOfSelectedNodes() > 0;
 }
 
-void FRigVMEditor::SetHost(URigVMHost* InHost)
+void FRigVMEditor::UpdateStaleWatchedPins()
 {
-	Host = InHost;
-	if(Host && IsValid(Host))
+	URigVMBlueprint* RigVMBlueprint = GetRigVMBlueprint();
+	if (RigVMBlueprint == nullptr)
 	{
-		OnPreviewHostUpdated().Broadcast(this);
+		return;
+	}
+
+	TSet<UEdGraphPin*> AllPins;
+	uint16 WatchCount;
+
+	// Find all unique pins being watched
+	FKismetDebugUtilities::ForeachPinWatch(
+		RigVMBlueprint,
+		[&AllPins, &WatchCount](UEdGraphPin* Pin)
+		{
+			++WatchCount;
+			if (Pin == nullptr)
+			{
+				return; // ~continue
+			}
+
+			UEdGraphNode* OwningNode = Pin->GetOwningNode();
+			// during node reconstruction, dead pins get moved to the transient 
+			// package (so just in case this blueprint got saved with dead pin watches)
+			if (OwningNode == nullptr)
+			{
+				return; // ~continue
+			}
+
+			if (!OwningNode->Pins.Contains(Pin))
+			{
+				return; // ~continue
+			}
+
+			AllPins.Add(Pin);
+		}
+	);
+
+	// Refresh watched pins with unique pins (throw away null or duplicate watches)
+	if (WatchCount != AllPins.Num())
+	{
+		RigVMBlueprint->Status = BS_Dirty;
+	}
+
+	FKismetDebugUtilities::ClearPinWatches(RigVMBlueprint);
+
+	TArray<URigVMGraph*> Models = RigVMBlueprint->GetAllModels();
+	for (URigVMGraph* Model : Models)
+	{
+		for (URigVMNode* ModelNode : Model->GetNodes())
+		{
+			TArray<URigVMPin*> ModelPins = ModelNode->GetAllPinsRecursively();
+			for (URigVMPin* ModelPin : ModelPins)
+			{
+				if (ModelPin->RequiresWatch())
+				{
+					RigVMBlueprint->GetController(Model)->SetPinIsWatched(ModelPin->GetPinPath(), false, false);
+				}
+			}
+		}
+	}
+	for (UEdGraphPin* Pin : AllPins)
+	{
+		FKismetDebugUtilities::AddPinWatch(RigVMBlueprint, FBlueprintWatchedPin(Pin));
+		UEdGraph* EdGraph = Pin->GetOwningNode()->GetGraph();
+		RigVMBlueprint->GetController(EdGraph)->SetPinIsWatched(Pin->GetName(), true, false);
+	}
+}
+
+void FRigVMEditor::HandleRefreshEditorFromBlueprint(URigVMBlueprint* InBlueprint)
+{
+	Compile();
+}
+
+void FRigVMEditor::HandleVariableDroppedFromBlueprint(UObject* InSubject, FProperty* InVariableToDrop, const FVector2D& InDropPosition, const FVector2D& InScreenPosition)
+{
+	URigVMBlueprint* Blueprint = Cast<URigVMBlueprint>(GetBlueprintObj());
+	if (Blueprint == nullptr)
+	{
+		return;
+	}
+
+	URigVMController* Controller = GetFocusedController();
+	check(Controller);
+
+	FRigVMExternalVariable ExternalVariable = FRigVMExternalVariable::Make(InVariableToDrop, nullptr);
+	if (!ExternalVariable.IsValid(true /* allow null ptr */))
+	{
+		return;
+	}
+
+	FMenuBuilder MenuBuilder(true, NULL);
+	const FText SectionText = FText::FromString(FString::Printf(TEXT("Variable %s"), *ExternalVariable.Name.ToString()));
+
+	MenuBuilder.BeginSection("VariableDropped", SectionText);
+
+	MenuBuilder.AddMenuEntry(
+		FText::FromString(FString::Printf(TEXT("Get %s"), *ExternalVariable.Name.ToString())),
+		FText::FromString(FString::Printf(TEXT("Adds a getter node for variable %s"), *ExternalVariable.Name.ToString())),
+		FSlateIcon(),
+		FUIAction(
+			FExecuteAction::CreateLambda([ExternalVariable, Controller, InDropPosition] {
+
+				Controller->AddVariableNode(ExternalVariable.Name, ExternalVariable.TypeName.ToString(), ExternalVariable.TypeObject, true, FString(), InDropPosition, FString(), true, true);
+
+			}),
+			FCanExecuteAction()
+		)
+	);
+
+	MenuBuilder.AddMenuEntry(
+		FText::FromString(FString::Printf(TEXT("Set %s"), *ExternalVariable.Name.ToString())),
+		FText::FromString(FString::Printf(TEXT("Adds a setter node for variable %s"), *ExternalVariable.Name.ToString())),
+		FSlateIcon(),
+		FUIAction(
+			FExecuteAction::CreateLambda([ExternalVariable, Controller, InDropPosition] {
+
+				Controller->AddVariableNode(ExternalVariable.Name, ExternalVariable.TypeName.ToString(), ExternalVariable.TypeObject, false, FString(), InDropPosition, FString(), true, true);
+
+			}),
+			FCanExecuteAction()
+		)
+	);
+
+	MenuBuilder.EndSection();
+
+	TSharedRef<SWidget> GraphEditorPanel = FocusedGraphEdPtr.Pin().ToSharedRef();
+
+	// Show dialog to choose getter vs setter
+	FSlateApplication::Get().PushMenu(
+		GraphEditorPanel,
+		FWidgetPath(),
+		MenuBuilder.MakeWidget(),
+		InScreenPosition,
+		FPopupTransitionEffect(FPopupTransitionEffect::ContextMenu)
+	);
+}
+
+void FRigVMEditor::HandleBreakpointAdded()
+{
+	SetExecutionMode(ERigVMEditorExecutionModeType_Debug);
+}
+
+void FRigVMEditor::OnGraphNodeClicked(URigVMEdGraphNode* InNode)
+{
+	if (InNode)
+	{
+		if (InNode->IsSelectedInEditor())
+		{
+			SetDetailViewForGraph(InNode->GetModel());
+		}
+	}
+}
+
+void FRigVMEditor::OnNodeDoubleClicked(URigVMBlueprint* InBlueprint, URigVMNode* InNode)
+{
+	ensure(GetRigVMBlueprint() == InBlueprint);
+
+	if (URigVMLibraryNode* LibraryNode = Cast<URigVMLibraryNode>(InNode))
+	{
+		URigVMGraph* ContainedGraph = LibraryNode->GetContainedGraph();
+		if (URigVMFunctionReferenceNode* FunctionReferenceNode = Cast<URigVMFunctionReferenceNode>(LibraryNode))
+		{
+			if (URigVMLibraryNode* ReferencedNode = FunctionReferenceNode->LoadReferencedNode())
+			{
+				ContainedGraph = ReferencedNode->GetContainedGraph();
+			}
+		}
+		if(ContainedGraph)
+		{
+			if (UEdGraph* EdGraph = InBlueprint->GetEdGraph(ContainedGraph))
+			{
+				OpenGraphAndBringToFront(EdGraph, true);
+			}
+			else
+			{
+				if(URigVMCollapseNode* FunctionLibraryNode = Cast<URigVMCollapseNode>(ContainedGraph->GetOuter()))
+				{
+					if(URigVMFunctionLibrary* FunctionLibrary = FunctionLibraryNode->GetLibrary())
+					{
+						if(URigVMBlueprint* FunctionBlueprint = Cast<URigVMBlueprint>(FunctionLibrary->GetOuter()))
+						{
+							if (UEdGraph* FunctionEdGraph = FunctionBlueprint->GetEdGraph(ContainedGraph))
+							{
+								FKismetEditorUtilities::BringKismetToFocusAttentionOnObject(FunctionEdGraph);
+							}
+						}
+					}
+				}
+			}
+		}
+	}
+}
+
+void FRigVMEditor::OnGraphImported(UEdGraph* InEdGraph)
+{
+	check(InEdGraph);
+
+	OpenDocument(InEdGraph, FDocumentTracker::OpenNewDocument);
+	RenameNewlyAddedAction(InEdGraph->GetFName());
+}
+
+bool FRigVMEditor::OnActionMatchesName(FEdGraphSchemaAction* InAction, const FName& InName) const
+{
+	if (InAction->GetMenuDescription().ToString() == InName.ToString())
+	{
+		return true;
+	}
+	return false;
+}
+
+void FRigVMEditor::HandleShowCurrentStatement()
+{
+	if(HaltedAtNode)
+	{
+		if(URigVMBlueprint* Blueprint = Cast<URigVMBlueprint>(GetBlueprintObj()))
+		{
+			if(URigVMEdGraph* EdGraph = Cast<URigVMEdGraph>(Blueprint->GetEdGraph(HaltedAtNode->GetGraph())))
+			{
+				if(UEdGraphNode* EdNode = EdGraph->FindNodeForModelNodeName(HaltedAtNode->GetFName()))
+				{
+					JumpToHyperlink(EdNode, false);
+				}
+			}
+		}
+	}
+}
+
+void FRigVMEditor::HandleBreakpointActionRequested(const ERigVMBreakpointAction BreakpointAction)
+{
+	if (URigVMHost* RigVMHost = GetRigVMHost())
+	{
+		RigVMHost->ExecuteBreakpointAction(BreakpointAction);
+	}
+}
+
+bool FRigVMEditor::IsHaltedAtBreakpoint() const
+{
+	return HaltedAtNode != nullptr;
+}
+
+void FRigVMEditor::FrameSelection()
+{
+	if (SGraphEditor* GraphEd = FocusedGraphEdPtr.Pin().Get())
+	{
+		if(URigVMGraph* Model = GetFocusedModel())
+		{
+			const bool bFrameAll = Model->GetSelectNodes().Num() == 0;
+			GraphEd->ZoomToFit(!bFrameAll);
+		}
+	}
+}
+
+void FRigVMEditor::UpdateGraphCompilerErrors()
+{
+	DECLARE_SCOPE_HIERARCHICAL_COUNTER_FUNC()
+
+	URigVMBlueprint* Blueprint = Cast<URigVMBlueprint>(GetBlueprintObj());
+	URigVMHost* RigVMHost = GetRigVMHost();
+	if (Blueprint && RigVMHost && RigVMHost->GetVM())
+	{
+		if (RigVMLog.Entries.Num() == 0 && !bAnyErrorsLeft)
+		{
+			return;
+		}
+
+		URigVM* VM = RigVMHost->GetVM();
+		const FRigVMByteCode& ByteCode = VM->GetByteCode();
+
+		TArray<UEdGraph*> EdGraphs;
+		Blueprint->GetAllGraphs(EdGraphs);
+
+		for (UEdGraph* Graph : EdGraphs)
+		{
+			URigVMEdGraph* RigVMEdGraph = Cast<URigVMEdGraph>(Graph);
+			if (RigVMEdGraph == nullptr)
+			{
+				continue;
+			}
+
+			// reset all nodes and store them in the map
+			bool bFoundWarning = false;
+			bool bFoundError = false;
+			
+			for (UEdGraphNode* GraphNode : Graph->Nodes)
+			{
+				if (URigVMEdGraphNode* RigVMEdGraphNode = Cast<URigVMEdGraphNode>(GraphNode))
+				{
+					bFoundError = bFoundError || RigVMEdGraphNode->ErrorType <= (int32)EMessageSeverity::Error;
+					bFoundWarning = bFoundWarning || RigVMEdGraphNode->ErrorType <= (int32)EMessageSeverity::Warning;
+
+					if(RigVMEdGraphNode->ErrorType <= (int32)EMessageSeverity::Warning)
+					{
+						if(!VM->WasInstructionVisitedDuringLastRun(RigVMHost->GetExtendedExecuteContext(), RigVMEdGraphNode->GetInstructionIndex(true)) &&
+							!VM->WasInstructionVisitedDuringLastRun(RigVMHost->GetExtendedExecuteContext(), RigVMEdGraphNode->GetInstructionIndex(false)))
+						{
+							continue;
+						}
+					}
+				}
+
+				GraphNode->ErrorType = int32(EMessageSeverity::Info) + 1;
+			}
+
+			// update the nodes' error messages
+			for (const FRigVMLog::FLogEntry& Entry : RigVMLog.Entries)
+			{
+				URigVMNode* ModelNode = Cast<URigVMNode>(ByteCode.GetSubjectForInstruction(Entry.InstructionIndex));
+				if (ModelNode == nullptr)
+				{
+					continue;
+				}
+
+				UEdGraphNode* GraphNode = RigVMEdGraph->FindNodeForModelNodeName(ModelNode->GetFName());
+				if (GraphNode == nullptr)
+				{
+					continue;
+				}
+
+				bFoundError = bFoundError || Entry.Severity <= EMessageSeverity::Error;
+				bFoundWarning = bFoundWarning || Entry.Severity <= EMessageSeverity::Warning;
+
+				int32 ErrorType = (int32)Entry.Severity;
+				if (GraphNode->ErrorType < ErrorType)
+				{
+					continue;
+				}
+				else if (GraphNode->ErrorType == ErrorType)
+				{
+					GraphNode->ErrorMsg = FString::Printf(TEXT("%s\n%s"), *GraphNode->ErrorMsg, *Entry.Message);
+				}
+				else
+				{
+					GraphNode->ErrorMsg = Entry.Message;
+					GraphNode->ErrorType = ErrorType;
+				}
+			}
+
+			bAnyErrorsLeft = false;
+			for (UEdGraphNode* GraphNode : Graph->Nodes)
+			{
+				GraphNode->bHasCompilerMessage = GraphNode->ErrorType <= int32(EMessageSeverity::Info);
+				bAnyErrorsLeft = bAnyErrorsLeft || GraphNode->bHasCompilerMessage; 
+			}
+
+			if (bFoundError)
+			{
+				Blueprint->Status = BS_Error;
+				(void)Blueprint->MarkPackageDirty();
+			}
+		}
+	}
+
+}
+
+TArray<FName> FRigVMEditor::GetDefaultEventQueue() const
+{
+	return TArray<FName>();
+}
+
+TArray<FName> FRigVMEditor::GetEventQueue() const
+{
+	if (const URigVMHost* CurrentHost = GetRigVMHost())
+	{
+		return CurrentHost->GetEventQueue();
+	}
+
+	return GetDefaultEventQueue();
+}
+
+void FRigVMEditor::SetEventQueue(TArray<FName> InEventQueue)
+{
+	return SetEventQueue(InEventQueue, false);
+}
+
+void FRigVMEditor::SetEventQueue(TArray<FName> InEventQueue, bool bCompile)
+{
+	if (GetEventQueue() == InEventQueue)
+	{
+		return;
+	}
+
+	LastEventQueue = GetEventQueue();
+
+	SetHaltedNode(nullptr);
+
+	if (URigVMHost* CurrentHost = GetRigVMHost())
+	{
+		if (InEventQueue.Num() > 0)
+		{
+			CurrentHost->SetEventQueue(InEventQueue);
+		}
+	}
+}
+
+FSlateIcon FRigVMEditor::GetEventQueueIcon(const TArray<FName>& InEventQueue) const
+{
+	return FSlateIcon();
+}
+
+FSlateIcon FRigVMEditor::GetEventQueueIcon() const
+{
+	return GetEventQueueIcon(GetEventQueue());
+}
+
+void FRigVMEditor::SetExecutionMode(const ERigVMEditorExecutionModeType InExecutionMode)
+{
+	if (ExecutionMode == InExecutionMode)
+	{
+		return;
+	}
+
+	ExecutionMode = InExecutionMode;
+	GetRigVMBlueprint()->SetDebugMode(InExecutionMode == ERigVMEditorExecutionModeType_Debug);
+	Compile();
+	
+	if (URigVMHost* CurrentHost = GetRigVMHost())
+	{
+		CurrentHost->SetIsInDebugMode(InExecutionMode == ERigVMEditorExecutionModeType_Debug);
+	}
+
+	SetHaltedNode(nullptr);
+	
+	RefreshDetailView();
+}
+
+int32 FRigVMEditor::GetExecutionModeComboValue() const
+{
+	return (int32) ExecutionMode;
+}
+
+FText FRigVMEditor::GetExecutionModeLabel() const
+{
+	if (ExecutionMode == ERigVMEditorExecutionModeType_Debug)
+	{
+		return FText::FromString(TEXT("DebugMode"));
+	}
+	return FText::FromString(TEXT("ReleaseMode"));
+}
+
+FSlateIcon FRigVMEditor::GetExecutionModeIcon(ERigVMEditorExecutionModeType InExecutionMode)
+{
+	if (InExecutionMode == ERigVMEditorExecutionModeType_Debug)
+	{
+		return FSlateIcon(FAppStyle::GetAppStyleSetName(), "LevelEditor.Tabs.Debug");
+	}
+	return FSlateIcon(FRigVMEditorStyle::Get().GetStyleSetName(), "RigVM.ReleaseMode");
+}
+
+FSlateIcon FRigVMEditor::GetExecutionModeIcon() const
+{
+	return GetExecutionModeIcon(ExecutionMode);
+}
+
+void FRigVMEditor::GetCustomDebugObjects(TArray<FCustomDebugObject>& DebugList) const
+{
+	URigVMBlueprint* RigVMBlueprint = GetRigVMBlueprint();
+	if (RigVMBlueprint == nullptr)
+	{
+		return;
+	}
+
+	if (URigVMHost* CurrentHost = GetRigVMHost())
+	{
+		if(IsValid(CurrentHost))
+		{
+			FCustomDebugObject DebugObject;
+			DebugObject.Object = CurrentHost;
+			DebugObject.NameOverride = GetCustomDebugObjectLabel(CurrentHost);
+			DebugList.Add(DebugObject);
+		}
+	}
+
+	URigVMBlueprintGeneratedClass* GeneratedClass = RigVMBlueprint->GetRigVMBlueprintGeneratedClass();
+	if (GeneratedClass)
+	{
+		struct Local
+		{
+			static bool IsPendingKillOrUnreachableRecursive(UObject* InObject)
+			{
+				if (InObject != nullptr)
+				{
+					if (!IsValidChecked(InObject) || InObject->IsUnreachable())
+					{
+						return true;
+					}
+					return IsPendingKillOrUnreachableRecursive(InObject->GetOuter());
+				}
+				return false;
+			}
+
+			static bool OuterNameContainsRecursive(UObject* InObject, const FString& InStringToSearch)
+			{
+				if (InObject == nullptr)
+				{
+					return false;
+				}
+
+				UObject* InObjectOuter = InObject->GetOuter();
+				if (InObjectOuter == nullptr)
+				{
+					return false;
+				}
+
+				if (InObjectOuter->GetName().Contains(InStringToSearch))
+				{
+					return true;
+				}
+
+				return OuterNameContainsRecursive(InObjectOuter, InStringToSearch);
+			}
+		};
+
+		if (UObject* DefaultObject = GeneratedClass->GetDefaultObject(false))
+		{
+			TArray<UObject*> ArchetypeInstances;
+			DefaultObject->GetArchetypeInstances(ArchetypeInstances);
+
+			for (UObject* Instance : ArchetypeInstances)
+			{
+				URigVMHost* InstancedHost = Cast<URigVMHost>(Instance);
+				if (InstancedHost && IsValid(InstancedHost) && InstancedHost != GetRigVMHost())
+				{
+					if (InstancedHost->GetOuter() == nullptr)
+					{
+						continue;
+					}
+
+					UWorld* World = InstancedHost->GetWorld();
+					if (World == nullptr)
+					{
+						continue;
+					}
+
+					// ensure to only allow preview actors in preview worlds
+					if (World->IsPreviewWorld())
+					{
+						if (!Local::OuterNameContainsRecursive(InstancedHost, TEXT("Preview")))
+						{
+							continue;
+						}
+					}
+
+					if (Local::IsPendingKillOrUnreachableRecursive(InstancedHost))
+					{
+						continue;
+					}
+
+					FCustomDebugObject DebugObject;
+					DebugObject.Object = InstancedHost;
+					DebugObject.NameOverride = GetCustomDebugObjectLabel(InstancedHost);
+					DebugList.Add(DebugObject);
+				}
+			}
+		}
+	}
+}
+
+void FRigVMEditor::HandleSetObjectBeingDebugged(UObject* InObject)
+{
+	URigVMHost* DebuggedHost = Cast<URigVMHost>(InObject);
+
+	if (DebuggedHost == nullptr)
+	{
+		// fall back to our default control rig (which still can be nullptr)
+		if (GetRigVMBlueprint() != nullptr && GetBlueprintObj() && !bIsSettingObjectBeingDebugged)
+		{
+			TGuardValue<bool> GuardSettingObjectBeingDebugged(bIsSettingObjectBeingDebugged, true);
+			GetBlueprintObj()->SetObjectBeingDebugged(GetRigVMHost());
+			return;
+		}
+	}
+
+	if (URigVMBlueprint* RigBlueprint = GetRigVMBlueprint())
+	{
+		if (URigVMBlueprintGeneratedClass* GeneratedClass = RigBlueprint->GetRigVMBlueprintGeneratedClass())
+		{
+			URigVMHost* CDO = Cast<URigVMHost>(GeneratedClass->GetDefaultObject(true /* create if needed */));
+			if (CDO->GetVM()->GetInstructions().Num() <= 1 /* only exit */)
+			{
+				RigBlueprint->RecompileVM();
+				RigBlueprint->RequestRigVMInit();
+			}
+		}
 	}
+
+	if(DebuggedHost)
+	{
+		DebuggedHost->SetLog(&RigVMLog);
+	}
+
+	RefreshDetailView();
+	LastDebuggedHost = GetCustomDebugObjectLabel(DebuggedHost);
+}
+
+FString FRigVMEditor::GetCustomDebugObjectLabel(UObject* ObjectBeingDebugged) const
+{
+	if (ObjectBeingDebugged == nullptr)
+	{
+		return FString();
+	}
+
+	if (ObjectBeingDebugged == GetRigVMHost())
+	{
+		return TEXT("Editor Preview");
+	}
+
+	if (AActor* ParentActor = ObjectBeingDebugged->GetTypedOuter<AActor>())
+	{
+		return FString::Printf(TEXT("%s in %s"), *GetBlueprintObj()->GetName(), *ParentActor->GetName());
+	}
+
+	return GetBlueprintObj()->GetName();
 }
 
 #undef LOCTEXT_NAMESPACE 
