RigVM: Move most of the functionality from ControlRigEditor to RigVMEditor

diff --git a/Engine/Plugins/Runtime/RigVM/Source/RigVMEditor/Private/Editor/RigVMEditorCommands.cpp b/Engine/Plugins/Runtime/RigVM/Source/RigVMEditor/Private/Editor/RigVMEditorCommands.cpp
new file mode 100644
--- /dev/null
+++ b/Engine/Plugins/Runtime/RigVM/Source/RigVMEditor/Private/Editor/RigVMEditorCommands.cpp
@@ -0,0 +1,24 @@
+// Copyright Epic Games, Inc. All Rights Reserved.
+
+#include "RigVMEditorCommands.h"
+
+#define LOCTEXT_NAMESPACE "RigVMEditorCommands"
+
+void FRigVMEditorCommands::RegisterCommands()
+{
+	UI_COMMAND(DeleteItem, "Delete", "Deletes the selected items and removes their nodes from the graph.", EUserInterfaceActionType::Button, FInputChord(EKeys::Delete));
+	UI_COMMAND(ExecuteGraph, "Execute", "Execute the rig graph if On.", EUserInterfaceActionType::ToggleButton, FInputChord());
+	UI_COMMAND(AutoCompileGraph, "Auto Compile", "Auto-compile the rig graph if On.", EUserInterfaceActionType::ToggleButton, FInputChord());
+	UI_COMMAND(ToggleEventQueue, "Toggle Event", "Toggle between the current and last running event", EUserInterfaceActionType::Button, FInputChord());
+	UI_COMMAND(ToggleExecutionMode, "Toggle Execution Mode", "Toggle between Release and Debug execution mode", EUserInterfaceActionType::Button, FInputChord());
+	UI_COMMAND(ReleaseMode, "Release Mode", "Compiles and Executes the rig, ignoring debug data.", EUserInterfaceActionType::Button, FInputChord());
+	UI_COMMAND(DebugMode, "Debug Mode", "Compiles and Executes the unoptimized rig, stopping at breakpoints.", EUserInterfaceActionType::Button, FInputChord());
+	UI_COMMAND(ResumeExecution, "Resume", "Resumes execution after being halted at a breakpoint.", EUserInterfaceActionType::Button, FInputChord(EKeys::F5));
+	UI_COMMAND(ShowCurrentStatement, "Show Current Statement", "Focuses on the node currently being debugged.", EUserInterfaceActionType::Button, FInputChord());
+	UI_COMMAND(StepInto, "Step Into", "Steps into the collapsed/function node, when halted at a breakpoint.", EUserInterfaceActionType::Button, FInputChord(EKeys::F11));
+	UI_COMMAND(StepOut, "Step Out", "Steps out of the collapsed/function node, when halted at a breakpoint.", EUserInterfaceActionType::Button, FInputChord(EKeys::F11, EModifierKey::Shift));
+	UI_COMMAND(StepOver, "Step Over", "Steps over the node, when halted at a breakpoint.", EUserInterfaceActionType::Button, FInputChord(EKeys::F10));
+	UI_COMMAND(FrameSelection, "Frame Selection", "Frames the selected nodes in the Graph View.", EUserInterfaceActionType::Button, FInputChord(EKeys::F));
+}
+
+#undef LOCTEXT_NAMESPACE
diff --git a/Engine/Plugins/Runtime/RigVM/Source/RigVMEditor/Private/Editor/RigVMEditorCommands.h b/Engine/Plugins/Runtime/RigVM/Source/RigVMEditor/Private/Editor/RigVMEditorCommands.h
new file mode 100644
--- /dev/null
+++ b/Engine/Plugins/Runtime/RigVM/Source/RigVMEditor/Private/Editor/RigVMEditorCommands.h
@@ -0,0 +1,64 @@
+// Copyright Epic Games, Inc. All Rights Reserved.
+
+#pragma once
+
+#include "CoreMinimal.h"
+#include "Framework/Commands/Commands.h"
+#include "Editor/RigVMEditorStyle.h"
+
+class FRigVMEditorCommands : public TCommands<FRigVMEditorCommands>
+{
+public:
+	FRigVMEditorCommands() : TCommands<FRigVMEditorCommands>
+	(
+		"RigVMBlueprint",
+		NSLOCTEXT("Contexts", "RigVM", "RigVM Blueprint"),
+		NAME_None, // "MainFrame" // @todo Fix this crash
+		FRigVMEditorStyle::Get().GetStyleSetName() // Icon Style Set
+	)
+	{}
+	
+	/** Deletes the selected items and removes their nodes from the graph. */
+	TSharedPtr< FUICommandInfo > DeleteItem;
+
+	/** Toggle Execute the Graph */
+	TSharedPtr< FUICommandInfo > ExecuteGraph;
+
+	/** Toggle Auto Compilation in the Graph */
+	TSharedPtr< FUICommandInfo > AutoCompileGraph;
+
+	/** Toggle between this and the last event queue */
+	TSharedPtr< FUICommandInfo > ToggleEventQueue;
+
+	/** Toggle between Release and Debug execution mode */
+	TSharedPtr< FUICommandInfo > ToggleExecutionMode;
+
+	/** Compile and run the optimized rig, ignoring any debug data */
+	TSharedPtr< FUICommandInfo > ReleaseMode;
+
+	/** Compile and run the unoptimized rig, ignoring any debug data */
+	TSharedPtr< FUICommandInfo > DebugMode;
+	
+	/** Resume the execution of the graph when halted at a breakpoint */
+	TSharedPtr< FUICommandInfo > ResumeExecution;
+
+	/** Focuses on the node currently being debugged */
+	TSharedPtr< FUICommandInfo > ShowCurrentStatement;
+
+	/** Steps to the next node in the execution graph (at the same graph level) when halted at a breakpoint */
+	TSharedPtr< FUICommandInfo > StepOver;
+
+	/** Steps into the collapsed/function node, when halted at a breakpoint */
+	TSharedPtr< FUICommandInfo > StepInto;
+
+	/** Steps out of the collapsed/function node, when halted at a breakpoint */
+	TSharedPtr< FUICommandInfo > StepOut;
+
+	/** Frames the selected nodes */
+	TSharedPtr< FUICommandInfo > FrameSelection;
+
+	/**
+	 * Initialize commands
+	 */
+	virtual void RegisterCommands() override;
+};
diff --git a/Engine/Plugins/Runtime/RigVM/Source/RigVMEditor/Private/Editor/RigVMEditorMenuContext.cpp b/Engine/Plugins/Runtime/RigVM/Source/RigVMEditor/Private/Editor/RigVMEditorMenuContext.cpp
new file mode 100644
--- /dev/null
+++ b/Engine/Plugins/Runtime/RigVM/Source/RigVMEditor/Private/Editor/RigVMEditorMenuContext.cpp
@@ -0,0 +1,57 @@
+// Copyright Epic Games, Inc. All Rights Reserved.
+
+#include "Editor/RigVMEditorMenuContext.h"
+
+#include "RigVMBlueprint.h"
+#include "Editor/RigVMEditor.h"
+#include "Framework/Application/SlateApplication.h"
+
+#include UE_INLINE_GENERATED_CPP_BY_NAME(RigVMEditorMenuContext)
+
+void URigVMEditorMenuContext::Init(TWeakPtr<FRigVMEditor> InRigVMEditor, const FRigVMEditorGraphMenuContext& InGraphMenuContext)
+{
+	WeakRigVMEditor = InRigVMEditor;
+	GraphMenuContext = InGraphMenuContext;
+}
+
+URigVMBlueprint* URigVMEditorMenuContext::GetRigVMBlueprint() const
+{
+	if (const TSharedPtr<FRigVMEditor> Editor = WeakRigVMEditor.Pin())
+	{
+		return Editor->GetRigVMBlueprint();
+	}
+	
+	return nullptr;
+}
+
+URigVMHost* URigVMEditorMenuContext::GetRigVMHost() const
+{
+	if (URigVMBlueprint* RigBlueprint = GetRigVMBlueprint())
+	{
+		if (URigVMHost* Host = Cast<URigVMHost>(RigBlueprint->GetObjectBeingDebugged()))
+		{
+			return Host;
+		}
+	}
+	return nullptr;
+}
+
+bool URigVMEditorMenuContext::IsAltDown() const
+{
+	return FSlateApplication::Get().GetModifierKeys().IsAltDown();
+}
+
+FRigVMEditorGraphMenuContext URigVMEditorMenuContext::GetGraphMenuContext()
+{
+	return GraphMenuContext;
+}
+
+FRigVMEditor* URigVMEditorMenuContext::GetRigVMEditor() const
+{
+	if (const TSharedPtr<FRigVMEditor> Editor = WeakRigVMEditor.Pin())
+	{
+		return Editor.Get();
+	}
+	return nullptr;
+}
+
diff --git a/Engine/Plugins/Runtime/RigVM/Source/RigVMEditor/Private/RigVMEditorModule.cpp b/Engine/Plugins/Runtime/RigVM/Source/RigVMEditor/Private/RigVMEditorModule.cpp
--- a/Engine/Plugins/Runtime/RigVM/Source/RigVMEditor/Private/RigVMEditorModule.cpp
+++ b/Engine/Plugins/Runtime/RigVM/Source/RigVMEditor/Private/RigVMEditorModule.cpp
@@ -7,6 +7,7 @@
 #include "RigVMEditorModule.h"
 #include "Modules/ModuleManager.h"
 #include "AssetRegistry/AssetRegistryModule.h"
+#include "Editor/RigVMEditorCommands.h"
 #include "Editor/RigVMExecutionStackCommands.h"
 #include "Editor/RigVMEditorStyle.h"
 #include "Editor/RigVMGraphDetailCustomization.h"
@@ -54,6 +55,7 @@ void FRigVMEditorModule::StartupModule()
 	if(IsRigVMEditorModuleBase())
 	{
 		FRigVMExecutionStackCommands::Register();
+		FRigVMEditorCommands::Register();
 		FRigVMEditorStyle::Register();
 
 		EdGraphPanelNodeFactory = MakeShared<FRigVMEdGraphPanelNodeFactory>();
@@ -138,6 +140,13 @@ const URigVMBlueprint* FRigVMEditorModule::GetRigVMBlueprintCDO() const
 
 void FRigVMEditorModule::GetTypeActions(URigVMBlueprint* RigVMBlueprint, FBlueprintActionDatabaseRegistrar& ActionRegistrar)
 {
+	// Only register actions for ourselves
+	UClass* BlueprintClass = GetRigVMBlueprintClass();
+	if(RigVMBlueprint->GetClass() != BlueprintClass)
+	{
+		return;
+	}
+
 	FAssetRegistryModule& AssetRegistryModule = FModuleManager::LoadModuleChecked<FAssetRegistryModule>("AssetRegistry");
 
 	// actions get registered under specific object-keys; the idea is that 
@@ -188,8 +197,9 @@ void FRigVMEditorModule::GetTypeActions(URigVMBlueprint* RigVMBlueprint, FBluepr
 		FText MenuDesc = FText::FromName(Template.GetName());
 		FText ToolTip = Template.GetTooltipText();
 
-		UBlueprintNodeSpawner* NodeSpawner = URigVMEdGraphTemplateNodeSpawner::CreateFromNotation(Template.GetNotation(), MenuDesc, NodeCategory, ToolTip);
+		URigVMEdGraphNodeSpawner* NodeSpawner = URigVMEdGraphTemplateNodeSpawner::CreateFromNotation(Template.GetNotation(), MenuDesc, NodeCategory, ToolTip);
 		check(NodeSpawner != nullptr);
+		NodeSpawner->SetRelatedBlueprintClass(BlueprintClass);
 		ActionRegistrar.AddBlueprintAction(ActionKey, NodeSpawner);
 	};
 
@@ -210,8 +220,9 @@ void FRigVMEditorModule::GetTypeActions(URigVMBlueprint* RigVMBlueprint, FBluepr
 		FText MenuDesc = FText::FromString(Factory->GetNodeTitle(FRigVMTemplateTypeMap()));
 		FText ToolTip = Factory->GetNodeTooltip(FRigVMTemplateTypeMap());
 
-		UBlueprintNodeSpawner* NodeSpawner = URigVMEdGraphTemplateNodeSpawner::CreateFromNotation(Template->GetNotation(), MenuDesc, NodeCategory, ToolTip);
+		URigVMEdGraphNodeSpawner* NodeSpawner = URigVMEdGraphTemplateNodeSpawner::CreateFromNotation(Template->GetNotation(), MenuDesc, NodeCategory, ToolTip);
 		check(NodeSpawner != nullptr);
+		NodeSpawner->SetRelatedBlueprintClass(BlueprintClass);
 		ActionRegistrar.AddBlueprintAction(ActionKey, NodeSpawner);
 	};
 
@@ -258,8 +269,9 @@ void FRigVMEditorModule::GetTypeActions(URigVMBlueprint* RigVMBlueprint, FBluepr
 		FText MenuDesc = FText::FromString(DisplayNameMetadata + MenuDescSuffixMetadata);
 		FText ToolTip = Struct->GetToolTipText();
 
-		UBlueprintNodeSpawner* NodeSpawner = URigVMEdGraphUnitNodeSpawner::CreateFromStruct(Struct, Function.GetMethodName(), MenuDesc, NodeCategory, ToolTip);
+		URigVMEdGraphNodeSpawner* NodeSpawner = URigVMEdGraphUnitNodeSpawner::CreateFromStruct(Struct, Function.GetMethodName(), MenuDesc, NodeCategory, ToolTip);
 		check(NodeSpawner != nullptr);
+		NodeSpawner->SetRelatedBlueprintClass(BlueprintClass);
 		ActionRegistrar.AddBlueprintAction(ActionKey, NodeSpawner);
 	};
 
@@ -286,8 +298,9 @@ void FRigVMEditorModule::GetTypeActions(URigVMBlueprint* RigVMBlueprint, FBluepr
 		FText MenuDesc = FText::FromString(FString::Printf(TEXT("Enum %s"), *EnumToConsider->GetName()));
 		FText ToolTip = MenuDesc;
 
-		UBlueprintNodeSpawner* NodeSpawner = URigVMEdGraphEnumNodeSpawner::CreateForEnum(EnumToConsider, MenuDesc, NodeCategory, ToolTip);
+		URigVMEdGraphNodeSpawner* NodeSpawner = URigVMEdGraphEnumNodeSpawner::CreateForEnum(EnumToConsider, MenuDesc, NodeCategory, ToolTip);
 		check(NodeSpawner != nullptr);
+		NodeSpawner->SetRelatedBlueprintClass(BlueprintClass);
 		ActionRegistrar.AddBlueprintAction(ActionKey, NodeSpawner);
 	}
 
@@ -358,8 +371,9 @@ void FRigVMEditorModule::GetTypeActions(URigVMBlueprint* RigVMBlueprint, FBluepr
 				PublicFunctionsProperty->ImportText_Direct(*PublicFunctionsString, &PublicFunctions, nullptr, EPropertyPortFlags::PPF_None);
 				for(const FRigVMOldPublicFunctionData& PublicFunction : PublicFunctions)
 				{
-					UBlueprintNodeSpawner* NodeSpawner = URigVMEdGraphFunctionRefNodeSpawner::CreateFromAssetData(ControlRigAssetData, PublicFunction);
+					URigVMEdGraphNodeSpawner* NodeSpawner = URigVMEdGraphFunctionRefNodeSpawner::CreateFromAssetData(ControlRigAssetData, PublicFunction);
 					check(NodeSpawner != nullptr);
+					NodeSpawner->SetRelatedBlueprintClass(BlueprintClass);
 					ActionRegistrar.AddBlueprintAction(ActionKey, NodeSpawner);
 				}
 			}
@@ -370,8 +384,9 @@ void FRigVMEditorModule::GetTypeActions(URigVMBlueprint* RigVMBlueprint, FBluepr
 				PublicGraphFunctionsProperty->ImportText_Direct(*PublicGraphFunctionsString, &PublicFunctions, nullptr, EPropertyPortFlags::PPF_None);
 				for(const FRigVMGraphFunctionHeader& PublicFunction : PublicFunctions)
 				{
-					UBlueprintNodeSpawner* NodeSpawner = URigVMEdGraphFunctionRefNodeSpawner::CreateFromAssetData(ControlRigAssetData, PublicFunction);
+					URigVMEdGraphNodeSpawner* NodeSpawner = URigVMEdGraphFunctionRefNodeSpawner::CreateFromAssetData(ControlRigAssetData, PublicFunction);
 					check(NodeSpawner != nullptr);
+					NodeSpawner->SetRelatedBlueprintClass(BlueprintClass);
 					ActionRegistrar.AddBlueprintAction(ActionKey, NodeSpawner);
 				}
 			}
@@ -381,6 +396,13 @@ void FRigVMEditorModule::GetTypeActions(URigVMBlueprint* RigVMBlueprint, FBluepr
 
 void FRigVMEditorModule::GetInstanceActions(URigVMBlueprint* RigVMBlueprint, FBlueprintActionDatabaseRegistrar& ActionRegistrar)
 {
+	// Only register actions for ourselves
+	UClass* BlueprintClass = GetRigVMBlueprintClass();
+	if(RigVMBlueprint->GetClass() != BlueprintClass)
+	{
+		return;
+	}
+
 	if (URigVMBlueprintGeneratedClass* GeneratedClass = RigVMBlueprint->GetRigVMBlueprintGeneratedClass())
 	{
 		if (URigVMHost* CDO = Cast<URigVMHost>(GeneratedClass->GetDefaultObject()))
@@ -392,10 +414,14 @@ void FRigVMEditorModule::GetInstanceActions(URigVMBlueprint* RigVMBlueprint, FBl
 			{
 				FText MenuDesc = FText::FromName(ExternalVariable.Name);
 				FText ToolTip = FText::FromString(FString::Printf(TEXT("Get the value of variable %s"), *ExternalVariable.Name.ToString()));
-				ActionRegistrar.AddBlueprintAction(GeneratedClass, URigVMEdGraphVariableNodeSpawner::CreateFromExternalVariable(RigVMBlueprint, ExternalVariable, true, MenuDesc, NodeCategory, ToolTip));
+				URigVMEdGraphNodeSpawner* GetNodeSpawner = URigVMEdGraphVariableNodeSpawner::CreateFromExternalVariable(RigVMBlueprint, ExternalVariable, true, MenuDesc, NodeCategory, ToolTip);
+				GetNodeSpawner->SetRelatedBlueprintClass(BlueprintClass);
+				ActionRegistrar.AddBlueprintAction(GeneratedClass, GetNodeSpawner);
 
 				ToolTip = FText::FromString(FString::Printf(TEXT("Set the value of variable %s"), *ExternalVariable.Name.ToString()));
-				ActionRegistrar.AddBlueprintAction(GeneratedClass, URigVMEdGraphVariableNodeSpawner::CreateFromExternalVariable(RigVMBlueprint, ExternalVariable, false, MenuDesc, NodeCategory, ToolTip));
+				URigVMEdGraphNodeSpawner* SetNodeSpawner = URigVMEdGraphVariableNodeSpawner::CreateFromExternalVariable(RigVMBlueprint, ExternalVariable, false, MenuDesc, NodeCategory, ToolTip);
+				SetNodeSpawner->SetRelatedBlueprintClass(BlueprintClass);
+				ActionRegistrar.AddBlueprintAction(GeneratedClass, SetNodeSpawner);
 			}
 		}
 
@@ -404,8 +430,9 @@ void FRigVMEditorModule::GetInstanceActions(URigVMBlueprint* RigVMBlueprint, FBl
 			TArray<URigVMLibraryNode*> Functions = LocalFunctionLibrary->GetFunctions();
 			for (URigVMLibraryNode* Function : Functions)
 			{
-				UBlueprintNodeSpawner* NodeSpawner = URigVMEdGraphFunctionRefNodeSpawner::CreateFromFunction(Function);
+				URigVMEdGraphNodeSpawner* NodeSpawner = URigVMEdGraphFunctionRefNodeSpawner::CreateFromFunction(Function);
 				check(NodeSpawner != nullptr);
+				NodeSpawner->SetRelatedBlueprintClass(BlueprintClass);
 				ActionRegistrar.AddBlueprintAction(GeneratedClass, NodeSpawner);
 			}
 
@@ -416,10 +443,14 @@ void FRigVMEditorModule::GetInstanceActions(URigVMBlueprint* RigVMBlueprint, FBl
 				{
 					FText MenuDesc = FText::FromName(LocalVariable.Name);
 					FText ToolTip = FText::FromString(FString::Printf(TEXT("Get the value of variable %s"), *LocalVariable.Name.ToString()));
-					ActionRegistrar.AddBlueprintAction(GeneratedClass, URigVMEdGraphVariableNodeSpawner::CreateFromLocalVariable(RigVMBlueprint, Function->GetContainedGraph(), LocalVariable, true, MenuDesc, NodeCategory, ToolTip));
+					URigVMEdGraphNodeSpawner* GetNodeSpawner = URigVMEdGraphVariableNodeSpawner::CreateFromLocalVariable(RigVMBlueprint, Function->GetContainedGraph(), LocalVariable, true, MenuDesc, NodeCategory, ToolTip);
+					GetNodeSpawner->SetRelatedBlueprintClass(BlueprintClass);
+					ActionRegistrar.AddBlueprintAction(GeneratedClass, GetNodeSpawner);
 
 					ToolTip = FText::FromString(FString::Printf(TEXT("Set the value of variable %s"), *LocalVariable.Name.ToString()));
-					ActionRegistrar.AddBlueprintAction(GeneratedClass, URigVMEdGraphVariableNodeSpawner::CreateFromLocalVariable(RigVMBlueprint, Function->GetContainedGraph(), LocalVariable, false, MenuDesc, NodeCategory, ToolTip));
+					URigVMEdGraphNodeSpawner* SetNodeSpawner = URigVMEdGraphVariableNodeSpawner::CreateFromLocalVariable(RigVMBlueprint, Function->GetContainedGraph(), LocalVariable, false, MenuDesc, NodeCategory, ToolTip);
+					SetNodeSpawner->SetRelatedBlueprintClass(BlueprintClass);
+					ActionRegistrar.AddBlueprintAction(GeneratedClass, SetNodeSpawner);
 				}
 			}
 		}
@@ -433,7 +464,9 @@ void FRigVMEditorModule::GetInstanceActions(URigVMBlueprint* RigVMBlueprint, FBl
 				{
 					FText MenuDesc = FText::FromName(InputArgument.Name);
 					FText ToolTip = FText::FromString(FString::Printf(TEXT("Get the value of input %s"), *InputArgument.Name.ToString()));
-					ActionRegistrar.AddBlueprintAction(GeneratedClass, URigVMEdGraphVariableNodeSpawner::CreateFromLocalVariable(RigVMBlueprint, Graph, InputArgument, true, MenuDesc, NodeCategory, ToolTip));
+					URigVMEdGraphNodeSpawner* NodeSpawner = URigVMEdGraphVariableNodeSpawner::CreateFromLocalVariable(RigVMBlueprint, Graph, InputArgument, true, MenuDesc, NodeCategory, ToolTip);
+					NodeSpawner->SetRelatedBlueprintClass(BlueprintClass);
+					ActionRegistrar.AddBlueprintAction(GeneratedClass, NodeSpawner);
 				}			
 			}
 		}
@@ -449,7 +482,9 @@ void FRigVMEditorModule::GetInstanceActions(URigVMBlueprint* RigVMBlueprint, FBl
 				FString Suffix = EntryName.ToString().EndsWith(EventStr) ? FString() : EventSuffix;
 				FText MenuDesc = FText::FromString(FString::Printf(TEXT("Run %s%s"), *EntryName.ToString(), *Suffix));
 				FText ToolTip = FText::FromString(FString::Printf(TEXT("Runs the %s%s"), *EntryName.ToString(), *Suffix));
-				ActionRegistrar.AddBlueprintAction(GeneratedClass, URigVMEdGraphInvokeEntryNodeSpawner::CreateForEntry(RigVMBlueprint, EntryName, MenuDesc, NodeCategory, ToolTip));
+				URigVMEdGraphNodeSpawner* NodeSpawner = URigVMEdGraphInvokeEntryNodeSpawner::CreateForEntry(RigVMBlueprint, EntryName, MenuDesc, NodeCategory, ToolTip);
+				NodeSpawner->SetRelatedBlueprintClass(BlueprintClass);
+				ActionRegistrar.AddBlueprintAction(GeneratedClass, NodeSpawner);
 			}
 		}
 	}
@@ -457,6 +492,12 @@ void FRigVMEditorModule::GetInstanceActions(URigVMBlueprint* RigVMBlueprint, FBl
 
 void FRigVMEditorModule::GetNodeContextMenuActions(URigVMBlueprint* RigVMBlueprint, const URigVMEdGraphNode* EdGraphNode, URigVMNode* ModelNode, UToolMenu* Menu) const
 {
+	// Only register menu actions for ourselves
+	if(RigVMBlueprint->GetClass() != GetRigVMBlueprintClass())
+	{
+		return;
+	}
+
 	GetNodeWorkflowContextMenuActions(RigVMBlueprint, EdGraphNode, ModelNode, Menu);
 	GetNodeEventsContextMenuActions(RigVMBlueprint, EdGraphNode, ModelNode, Menu);
 	GetNodeConversionContextMenuActions(RigVMBlueprint, EdGraphNode, ModelNode, Menu);
diff --git a/Engine/Plugins/Runtime/RigVM/Source/RigVMEditor/Public/EdGraph/NodeSpawners/RigVMEdGraphEnumNodeSpawner.h b/Engine/Plugins/Runtime/RigVM/Source/RigVMEditor/Public/EdGraph/NodeSpawners/RigVMEdGraphEnumNodeSpawner.h
--- a/Engine/Plugins/Runtime/RigVM/Source/RigVMEditor/Public/EdGraph/NodeSpawners/RigVMEdGraphEnumNodeSpawner.h
+++ b/Engine/Plugins/Runtime/RigVM/Source/RigVMEditor/Public/EdGraph/NodeSpawners/RigVMEdGraphEnumNodeSpawner.h
@@ -9,15 +9,13 @@
 #include "UObject/UnrealType.h"
 #include "EdGraph/EdGraphPin.h"
 #include "Engine/Blueprint.h"
-#include "BlueprintActionFilter.h"
-#include "BlueprintNodeSignature.h"
-#include "BlueprintFieldNodeSpawner.h"
+#include "EdGraph/NodeSpawners/RigVMEdGraphNodeSpawner.h"
 #include "RigVMEdGraphEnumNodeSpawner.generated.h"
 
 class URigVMEdGraphNode;
 
 UCLASS(Transient)
-class RIGVMEDITOR_API URigVMEdGraphEnumNodeSpawner : public UBlueprintNodeSpawner
+class RIGVMEDITOR_API URigVMEdGraphEnumNodeSpawner : public URigVMEdGraphNodeSpawner
 {
 	GENERATED_BODY()
 
diff --git a/Engine/Plugins/Runtime/RigVM/Source/RigVMEditor/Public/EdGraph/NodeSpawners/RigVMEdGraphFunctionRefNodeSpawner.h b/Engine/Plugins/Runtime/RigVM/Source/RigVMEditor/Public/EdGraph/NodeSpawners/RigVMEdGraphFunctionRefNodeSpawner.h
--- a/Engine/Plugins/Runtime/RigVM/Source/RigVMEditor/Public/EdGraph/NodeSpawners/RigVMEdGraphFunctionRefNodeSpawner.h
+++ b/Engine/Plugins/Runtime/RigVM/Source/RigVMEditor/Public/EdGraph/NodeSpawners/RigVMEdGraphFunctionRefNodeSpawner.h
@@ -9,9 +9,7 @@
 #include "UObject/UnrealType.h"
 #include "EdGraph/EdGraphPin.h"
 #include "Engine/Blueprint.h"
-#include "BlueprintActionFilter.h"
-#include "BlueprintNodeSignature.h"
-#include "BlueprintFieldNodeSpawner.h"
+#include "EdGraph/NodeSpawners/RigVMEdGraphNodeSpawner.h"
 #include "RigVMBlueprint.h"
 #include "EdGraph/RigVMEdGraph.h"
 #include "EdGraph/RigVMEdGraphNode.h"
@@ -21,7 +19,7 @@
 class URigVMEdGraphNode;
 
 UCLASS(Transient)
-class RIGVMEDITOR_API URigVMEdGraphFunctionRefNodeSpawner : public UBlueprintNodeSpawner
+class RIGVMEDITOR_API URigVMEdGraphFunctionRefNodeSpawner : public URigVMEdGraphNodeSpawner
 {
 	GENERATED_BODY()
 
@@ -66,5 +64,5 @@ private:
 	static URigVMEdGraphNode* SpawnNode(UEdGraph* ParentGraph, UBlueprint* Blueprint, FRigVMGraphFunctionHeader& InFunction, FVector2D const Location);
 
 	friend class UEngineTestControlRig;
-	friend class FControlRigEditor;
+	friend class FRigVMEditor;
 };
diff --git a/Engine/Plugins/Runtime/RigVM/Source/RigVMEditor/Public/EdGraph/NodeSpawners/RigVMEdGraphInvokeEntryNodeSpawner.h b/Engine/Plugins/Runtime/RigVM/Source/RigVMEditor/Public/EdGraph/NodeSpawners/RigVMEdGraphInvokeEntryNodeSpawner.h
--- a/Engine/Plugins/Runtime/RigVM/Source/RigVMEditor/Public/EdGraph/NodeSpawners/RigVMEdGraphInvokeEntryNodeSpawner.h
+++ b/Engine/Plugins/Runtime/RigVM/Source/RigVMEditor/Public/EdGraph/NodeSpawners/RigVMEdGraphInvokeEntryNodeSpawner.h
@@ -3,9 +3,7 @@
 #pragma once
 
 #include "CoreMinimal.h"
-#include "BlueprintActionFilter.h"
-#include "BlueprintNodeSignature.h"
-#include "BlueprintFieldNodeSpawner.h"
+#include "EdGraph/NodeSpawners/RigVMEdGraphNodeSpawner.h"
 #include "RigVMModel/RigVMGraph.h"
 #include "RigVMEdGraphInvokeEntryNodeSpawner.generated.h"
 
@@ -13,7 +11,7 @@ class URigVMEdGraphNode;
 class URigVMBlueprint;
 
 UCLASS(Transient)
-class RIGVMEDITOR_API URigVMEdGraphInvokeEntryNodeSpawner : public UBlueprintNodeSpawner
+class RIGVMEDITOR_API URigVMEdGraphInvokeEntryNodeSpawner : public URigVMEdGraphNodeSpawner
 {
 	GENERATED_BODY()
 
diff --git a/Engine/Plugins/Runtime/RigVM/Source/RigVMEditor/Public/EdGraph/NodeSpawners/RigVMEdGraphNodeSpawner.h b/Engine/Plugins/Runtime/RigVM/Source/RigVMEditor/Public/EdGraph/NodeSpawners/RigVMEdGraphNodeSpawner.h
new file mode 100644
--- /dev/null
+++ b/Engine/Plugins/Runtime/RigVM/Source/RigVMEditor/Public/EdGraph/NodeSpawners/RigVMEdGraphNodeSpawner.h
@@ -0,0 +1,28 @@
+// Copyright Epic Games, Inc. All Rights Reserved.
+
+#pragma once
+
+#include "CoreMinimal.h"
+#include "RigVMBlueprint.h"
+#include "BlueprintNodeSpawner.h"
+#include "BlueprintNodeSignature.h"
+#include "BlueprintActionFilter.h"
+#include "RigVMEdGraphNodeSpawner.generated.h"
+
+class URigVMEdGraphNode;
+
+UCLASS(Transient)
+class RIGVMEDITOR_API URigVMEdGraphNodeSpawner : public UBlueprintNodeSpawner
+{
+	GENERATED_BODY()
+
+public:
+
+	virtual bool IsTemplateNodeFilteredOut(FBlueprintActionFilter const& Filter) const override;
+
+	void SetRelatedBlueprintClass(TSubclassOf<URigVMBlueprint> InClass);
+
+private:
+
+	TSubclassOf<URigVMBlueprint> RelatedBlueprintClass;
+};
diff --git a/Engine/Plugins/Runtime/RigVM/Source/RigVMEditor/Public/EdGraph/NodeSpawners/RigVMEdGraphTemplateNodeSpawner.h b/Engine/Plugins/Runtime/RigVM/Source/RigVMEditor/Public/EdGraph/NodeSpawners/RigVMEdGraphTemplateNodeSpawner.h
--- a/Engine/Plugins/Runtime/RigVM/Source/RigVMEditor/Public/EdGraph/NodeSpawners/RigVMEdGraphTemplateNodeSpawner.h
+++ b/Engine/Plugins/Runtime/RigVM/Source/RigVMEditor/Public/EdGraph/NodeSpawners/RigVMEdGraphTemplateNodeSpawner.h
@@ -9,10 +9,7 @@
 #include "UObject/UnrealType.h"
 #include "EdGraph/EdGraphPin.h"
 #include "Engine/Blueprint.h"
-#include "BlueprintActionFilter.h"
-#include "BlueprintNodeSignature.h"
-#include "BlueprintFieldNodeSpawner.h"
-#include "RigVMBlueprint.h"
+#include "EdGraph/NodeSpawners/RigVMEdGraphNodeSpawner.h"
 #include "EdGraph/RigVMEdGraph.h"
 #include "EdGraph/RigVMEdGraphNode.h"
 #include "RigVMModel/Nodes/RigVMTemplateNode.h"
@@ -21,7 +18,7 @@
 class URigVMEdGraphNode;
 
 UCLASS(Transient)
-class RIGVMEDITOR_API URigVMEdGraphTemplateNodeSpawner : public UBlueprintNodeSpawner
+class RIGVMEDITOR_API URigVMEdGraphTemplateNodeSpawner : public URigVMEdGraphNodeSpawner
 {
 	GENERATED_BODY()
 
@@ -50,5 +47,5 @@ private:
 	static URigVMEdGraphNode* SpawnNode(UEdGraph* ParentGraph, UBlueprint* Blueprint, const FName& InNotation, FVector2D const Location);
 
 	friend class UEngineTestControlRig;
-	friend class FControlRigEditor;
+	friend class FRigVMEditor;
 };
diff --git a/Engine/Plugins/Runtime/RigVM/Source/RigVMEditor/Public/EdGraph/NodeSpawners/RigVMEdGraphUnitNodeSpawner.h b/Engine/Plugins/Runtime/RigVM/Source/RigVMEditor/Public/EdGraph/NodeSpawners/RigVMEdGraphUnitNodeSpawner.h
--- a/Engine/Plugins/Runtime/RigVM/Source/RigVMEditor/Public/EdGraph/NodeSpawners/RigVMEdGraphUnitNodeSpawner.h
+++ b/Engine/Plugins/Runtime/RigVM/Source/RigVMEditor/Public/EdGraph/NodeSpawners/RigVMEdGraphUnitNodeSpawner.h
@@ -9,9 +9,7 @@
 #include "UObject/UnrealType.h"
 #include "EdGraph/EdGraphPin.h"
 #include "Engine/Blueprint.h"
-#include "BlueprintActionFilter.h"
-#include "BlueprintNodeSignature.h"
-#include "BlueprintFieldNodeSpawner.h"
+#include "EdGraph/NodeSpawners/RigVMEdGraphNodeSpawner.h"
 #include "RigVMBlueprint.h"
 #include "EdGraph/RigVMEdGraph.h"
 #include "EdGraph/RigVMEdGraphNode.h"
@@ -21,7 +19,7 @@
 class URigVMEdGraphNode;
 
 UCLASS(Transient)
-class RIGVMEDITOR_API URigVMEdGraphUnitNodeSpawner : public UBlueprintNodeSpawner
+class RIGVMEDITOR_API URigVMEdGraphUnitNodeSpawner : public URigVMEdGraphNodeSpawner
 {
 	GENERATED_BODY()
 
@@ -53,5 +51,5 @@ private:
 	static URigVMEdGraphNode* SpawnNode(UEdGraph* ParentGraph, UBlueprint* Blueprint, UScriptStruct* StructTemplate, const FName& InMethodName, FVector2D const Location);
 
 	friend class UEngineTestControlRig;
-	friend class FControlRigEditor;
+	friend class FRigVMEditor;
 };
diff --git a/Engine/Plugins/Runtime/RigVM/Source/RigVMEditor/Public/EdGraph/NodeSpawners/RigVMEdGraphVariableNodeSpawner.h b/Engine/Plugins/Runtime/RigVM/Source/RigVMEditor/Public/EdGraph/NodeSpawners/RigVMEdGraphVariableNodeSpawner.h
--- a/Engine/Plugins/Runtime/RigVM/Source/RigVMEditor/Public/EdGraph/NodeSpawners/RigVMEdGraphVariableNodeSpawner.h
+++ b/Engine/Plugins/Runtime/RigVM/Source/RigVMEditor/Public/EdGraph/NodeSpawners/RigVMEdGraphVariableNodeSpawner.h
@@ -9,9 +9,7 @@
 #include "UObject/UnrealType.h"
 #include "EdGraph/EdGraphPin.h"
 #include "Engine/Blueprint.h"
-#include "BlueprintActionFilter.h"
-#include "BlueprintNodeSignature.h"
-#include "BlueprintFieldNodeSpawner.h"
+#include "EdGraph/NodeSpawners/RigVMEdGraphNodeSpawner.h"
 #include "RigVMCore/RigVMExternalVariable.h"
 #include "RigVMModel/RigVMVariableDescription.h"
 #include "RigVMEdGraphVariableNodeSpawner.generated.h"
@@ -20,7 +18,7 @@ class URigVMEdGraphNode;
 class URigVMBlueprint;
 
 UCLASS(Transient)
-class RIGVMEDITOR_API URigVMEdGraphVariableNodeSpawner : public UBlueprintNodeSpawner
+class RIGVMEDITOR_API URigVMEdGraphVariableNodeSpawner : public URigVMEdGraphNodeSpawner
 {
 	GENERATED_BODY()
 
diff --git a/Engine/Plugins/Runtime/RigVM/Source/RigVMEditor/Public/Editor/RigVMEditor.h b/Engine/Plugins/Runtime/RigVM/Source/RigVMEditor/Public/Editor/RigVMEditor.h
--- a/Engine/Plugins/Runtime/RigVM/Source/RigVMEditor/Public/Editor/RigVMEditor.h
+++ b/Engine/Plugins/Runtime/RigVM/Source/RigVMEditor/Public/Editor/RigVMEditor.h
@@ -6,12 +6,83 @@
 #include "BlueprintEditor.h"
 #include "RigVMHost.h"
 #include "RigVMBlueprint.h"
+#include "Editor/RigVMDetailsViewWrapperObject.h"
+
+class FRigVMEditor;
+
+DECLARE_MULTICAST_DELEGATE_TwoParams(FRigVMEditorClosed, const FRigVMEditor*, URigVMBlueprint*);
 
 class RIGVMEDITOR_API FRigVMEditor : public FBlueprintEditor
 {
 public:
 
 	FRigVMEditor();
+	virtual ~FRigVMEditor();
+
+	/**
+	 * Edits the specified asset(s)
+	 *
+	 * @param	Mode					Mode that this editor should operate in
+	 * @param	InitToolkitHost			When Mode is WorldCentric, this is the level editor instance to spawn this editor within
+	 * @param	InRigVMBlueprint	The blueprint object to start editing.
+	 */
+	virtual void InitRigVMEditor(const EToolkitMode::Type Mode, const TSharedPtr< class IToolkitHost >& InitToolkitHost, class URigVMBlueprint* InRigVMBlueprint);
+
+	virtual const FName GetEditorAppName() const;
+	virtual const FName GetEditorModeName() const;
+	virtual TSharedPtr<FApplicationMode> CreateEditorMode();
+
+	// FBlueprintEditor interface
+	virtual UBlueprint* GetBlueprintObj() const override;
+	virtual TSubclassOf<UEdGraphSchema> GetDefaultSchemaClass() const override;
+
+	//  FTickableEditorObject Interface
+	virtual void Tick(float DeltaTime) override;
+
+	// IToolkit Interface
+	virtual FName GetToolkitFName() const override;
+	virtual FName GetToolkitContextFName() const override;
+	virtual FText GetBaseToolkitName() const override;
+	virtual FText GetToolkitToolTipText() const override;
+	virtual FString GetWorldCentricTabPrefix() const override;
+	virtual FLinearColor GetWorldCentricTabColorScale() const override;	
+	virtual void InitToolMenuContext(FToolMenuContext& MenuContext) override;
+
+	// BlueprintEditor interface
+	virtual bool TransactionObjectAffectsBlueprint(UObject* InTransactedObject) override;
+	virtual bool CanAddNewLocalVariable() const override;
+	virtual void OnAddNewLocalVariable() override;
+	virtual void OnPasteNewLocalVariable(const FBPVariableDescription& VariableDescription) override;
+	virtual void DeleteSelectedNodes() override;
+	virtual bool CanDeleteNodes() const override;
+	virtual void CopySelectedNodes() override;
+	virtual bool CanCopyNodes() const override;
+	virtual void PasteNodes() override;
+	virtual bool CanPasteNodes() const override;
+	virtual bool IsNativeParentClassCodeLinkEnabled() const override { return false; }
+	virtual bool ReparentBlueprint_IsVisible() const override { return false; }
+	virtual FReply OnSpawnGraphNodeByShortcut(FInputChord InChord, const FVector2D& InPosition, UEdGraph* InGraph) override;
+	virtual bool ShouldLoadBPLibrariesFromAssetRegistry() override { return false; }
+	virtual void JumpToHyperlink(const UObject* ObjectReference, bool bRequestRename = false) override;
+
+	// FEditorUndoClient Interface
+	virtual void PostUndo(bool bSuccess) override;
+	virtual void PostRedo(bool bSuccess) override;
+
+	virtual void PostTransaction(bool bSuccess, const FTransaction* Transaction, bool bIsRedo) {};
+
+	virtual void OnStartWatchingPin();
+	virtual bool CanStartWatchingPin() const;
+
+	virtual void OnStopWatchingPin();
+	virtual bool CanStopWatchingPin() const;
+
+	// IToolkitHost Interface
+	virtual void OnToolkitHostingStarted(const TSharedRef<class IToolkit>& Toolkit) override;
+	virtual void OnToolkitHostingFinished(const TSharedRef<class IToolkit>& Toolkit) override;
+
+	//  FTickableEditorObject Interface
+	virtual TStatId GetStatId() const override;
 
 	// returns the blueprint being edited
 	URigVMBlueprint* GetRigVMBlueprint() const;
@@ -19,18 +90,185 @@ public:
 	// returns the currently debugged / viewed host
 	URigVMHost* GetRigVMHost() const;
 
+	virtual UObject* GetOuterForHost() const;
+
 	// returns the class to use for detail wrapper objects (UI shim layer)
 	virtual UClass* GetDetailWrapperClass() const;
 
-	virtual void HandleModifiedEvent(ERigVMGraphNotifType InNotifType, URigVMGraph* InGraph, UObject* InSubject);
+	// allows the editor to fill an empty graph
+	virtual void CreateEmptyGraphContent(URigVMController* InController) {}
 
 	DECLARE_EVENT_OneParam(FRigVMEditor, FPreviewHostUpdated, FRigVMEditor*);
 	FPreviewHostUpdated& OnPreviewHostUpdated() { return PreviewHostUpdated;  }
 
+	FRigVMEditorClosed& OnEditorClosed() { return RigVMEditorClosedDelegate; }
+
+	/** Get the toolbox hosting widget */
+	TSharedRef<SBorder> GetToolbox() { return Toolbox.ToSharedRef(); }
+
+	virtual bool SelectLocalVariable(const UEdGraph* Graph, const FName& VariableName) override;
+
 protected:
 
+	enum ERigVMEditorExecutionModeType
+	{
+		ERigVMEditorExecutionModeType_Release,
+		ERigVMEditorExecutionModeType_Debug
+	};
+
+	// FBlueprintEditor Interface
+	virtual void CreateDefaultCommands() override;
+	virtual void OnCreateGraphEditorCommands(TSharedPtr<FUICommandList> GraphEditorCommandsList);
+	virtual TSharedRef<SGraphEditor> CreateGraphEditorWidget(TSharedRef<FTabInfo> InTabInfo, UEdGraph* InGraph) override;
+	virtual void Compile() override;
+	virtual void SaveAsset_Execute() override;
+	virtual void SaveAssetAs_Execute() override;
+	virtual bool IsInAScriptingMode() const override { return true; }
+	virtual void CreateDefaultTabContents(const TArray<UBlueprint*>& InBlueprints) override;
+	virtual void NewDocument_OnClicked(ECreatedDocumentType GraphType) override;
+	virtual bool IsSectionVisible(NodeSectionID::Type InSectionID) const override;
+	virtual FGraphAppearanceInfo GetGraphAppearance(class UEdGraph* InGraph) const override;
+	virtual bool IsEditable(UEdGraph* InGraph) const override;
+	virtual bool IsCompilingEnabled() const override;
+	virtual FText GetGraphDecorationString(UEdGraph* InGraph) const override;
+	virtual void OnActiveTabChanged( TSharedPtr<SDockTab> PreviouslyActive, TSharedPtr<SDockTab> NewlyActivated ) override;
+	virtual void OnSelectedNodesChangedImpl(const TSet<class UObject*>& NewSelection) override;
+	virtual void OnBlueprintChangedImpl(UBlueprint* InBlueprint, bool bIsJustBeingCompiled) override;
+	virtual void RefreshEditors(ERefreshBlueprintEditorReason::Type Reason = ERefreshBlueprintEditorReason::UnknownReason) override;
+	virtual void SetupGraphEditorEvents(UEdGraph* InGraph, SGraphEditor::FGraphEditorEvents& InEvents) override;
+	virtual void FocusInspectorOnGraphSelection(const TSet<class UObject*>& NewSelection, bool bForceRefresh = false) override;
+
+	virtual void HandleModifiedEvent(ERigVMGraphNotifType InNotifType, URigVMGraph* InGraph, UObject* InSubject);
+	virtual void HandleVMCompiledEvent(UObject* InCompiledObject, URigVM* InVM);
+	virtual void HandleVMExecutedEvent(URigVMHost* InHost, const FName& InEventName);
+	virtual void HandleVMExecutionHalted(const int32 InstructionIndex, UObject* InNode, const FName& InEntryName);
+	void SetHaltedNode(URigVMNode* Node);
+	
+	// FNotifyHook Interface
+	virtual void NotifyPreChange(FProperty* PropertyAboutToChange) override;
+	virtual void NotifyPostChange(const FPropertyChangedEvent& PropertyChangedEvent, FProperty* PropertyThatChanged) override;
+	/** delegate for changing property */
+	virtual void OnFinishedChangingProperties(const FPropertyChangedEvent& PropertyChangedEvent) override;
+	void OnPropertyChanged(UObject* InObject, FPropertyChangedEvent& InEvent);
+	virtual void OnWrappedPropertyChangedChainEvent(URigVMDetailsViewWrapperObject* InWrapperObject, const FString& InPropertyPath, FPropertyChangedChainEvent& InPropertyChangedChainEvent);
+	void OnRequestLocalizeFunctionDialog(FRigVMGraphFunctionIdentifier& InFunction, URigVMBlueprint* InTargetBlueprint, bool bForce);
+	FRigVMController_BulkEditResult OnRequestBulkEditDialog(URigVMBlueprint* InBlueprint, URigVMController* InController, URigVMLibraryNode* InFunction, ERigVMControllerBulkEditType InEditType);
+	bool OnRequestBreakLinksDialog(TArray<URigVMLink*> InLinks);
+	TRigVMTypeIndex OnRequestPinTypeSelectionDialog(const TArray<TRigVMTypeIndex>& InTypes);
+	void HandleJumpToHyperlink(const UObject* InSubject);
+	bool UpdateDefaultValueForVariable(FBPVariableDescription& InVariable, bool bUseCDO);
+
+	URigVMController* ActiveController;
+
+	/** Push a newly compiled/opened host to the editor */
+	virtual void UpdateRigVMHost();
+	virtual void UpdateRigVMHost_PreClearOldHost(URigVMHost* InPreviousHost) {};
+
+	/** Update the name lists for use in name combo boxes */
+	virtual void CacheNameLists();
+
+	// FGCObject Interface
+	virtual void AddReferencedObjects( FReferenceCollector& Collector ) override;
+
+	virtual void BindCommands();
+
+	void ToggleAutoCompileGraph();
+	bool IsAutoCompileGraphOn() const;
+	bool CanAutoCompileGraph() const { return true; }
+	void ToggleEventQueue();
+	void ToggleExecutionMode();
+	TSharedRef<SWidget> GenerateEventQueueMenuContent();
+	TSharedRef<SWidget> GenerateExecutionModeMenuContent();
+	virtual void GenerateEventQueueMenuContent(FMenuBuilder& MenuBuilder);
+
+	/** Wraps the normal blueprint editor's action menu creation callback */
+	FActionMenuContent HandleCreateGraphActionMenu(UEdGraph* InGraph, const FVector2D& InNodePosition, const TArray<UEdGraphPin*>& InDraggedPins, bool bAutoExpand, SGraphEditor::FActionMenuClosed InOnMenuClosed);
+	void OnNodeTitleCommitted(const FText& NewText, ETextCommit::Type CommitInfo, UEdGraphNode* NodeBeingChanged);
+
+	/** Undo Action**/
+	void UndoAction();
+
+	/** Redo Action **/
+	void RedoAction();
+	
+	virtual void OnCreateComment() override;
+
+	bool IsDetailsPanelRefreshSuspended() const { return bSuspendDetailsPanelRefresh; }
+	bool& GetSuspendDetailsPanelRefreshFlag() { return bSuspendDetailsPanelRefresh; }
+	virtual void SetDetailObjects(const TArray<UObject*>& InObjects);
+	virtual void SetDetailObjects(const TArray<UObject*>& InObjects, bool bChangeUISelectionState);
+	virtual void SetDetailViewForGraph(URigVMGraph* InGraph);
+	virtual void SetDetailViewForFocusedGraph();
+	virtual void SetDetailViewForLocalVariable();
+	virtual void RefreshDetailView();
+	virtual bool DetailViewShowsAnyRigUnit() const;
+	virtual bool DetailViewShowsLocalVariable() const;
+	virtual bool DetailViewShowsStruct(UScriptStruct* InStruct) const;
+	virtual void ClearDetailObject(bool bChangeUISelectionState = true);
+	virtual void ClearDetailsViewWrapperObjects();
+	const TArray<TStrongObjectPtr<URigVMDetailsViewWrapperObject>>& GetWrapperObjects() const { return WrapperObjects; }
+
 	void SetHost(URigVMHost* InHost);
 
+	URigVMGraph* GetFocusedModel() const;
+	URigVMController* GetFocusedController() const;
+	TSharedPtr<SGraphEditor> GetGraphEditor(UEdGraph* InEdGraph) const;
+
+	/** Extend menu */
+	void ExtendMenu();
+
+	/** Extend toolbar */
+	void ExtendToolbar();
+	
+	/** Fill the toolbar with content */
+	virtual void FillToolbar(FToolBarBuilder& ToolbarBuilder, bool bEndSection = true);
+
+public:
+	
+	virtual TArray<FName> GetDefaultEventQueue() const;
+	TArray<FName> GetEventQueue() const;
+	void SetEventQueue(TArray<FName> InEventQueue);
+	virtual void SetEventQueue(TArray<FName> InEventQueue, bool bCompile);
+	virtual int32 GetEventQueueComboValue() const { return INDEX_NONE; }
+	virtual FText GetEventQueueLabel() const { return FText(); }
+	virtual FSlateIcon GetEventQueueIcon(const TArray<FName>& InEventQueue) const;
+	FSlateIcon GetEventQueueIcon() const;
+
+protected:
+
+	void SetExecutionMode(const ERigVMEditorExecutionModeType InExecutionMode);
+	int32 GetExecutionModeComboValue() const;
+	FText GetExecutionModeLabel() const;
+	static FSlateIcon GetExecutionModeIcon(const ERigVMEditorExecutionModeType InExecutionMode);
+	FSlateIcon GetExecutionModeIcon() const;
+	
+	virtual void GetCustomDebugObjects(TArray<FCustomDebugObject>& DebugList) const override;
+	virtual bool OnlyShowCustomDebugObjects() const override { return true; }
+	virtual void HandleSetObjectBeingDebugged(UObject* InObject) override;
+	virtual FString GetCustomDebugObjectLabel(UObject* ObjectBeingDebugged) const override;
+
+	/** Handle hiding items in the graph */
+	void HandleHideItem();
+	bool CanHideItem() const;
+
+	/** Update stale watch pins */
+	void UpdateStaleWatchedPins();
+	
+	virtual void HandleRefreshEditorFromBlueprint(URigVMBlueprint* InBlueprint);
+	virtual void HandleVariableDroppedFromBlueprint(UObject* InSubject, FProperty* InVariableToDrop, const FVector2D& InDropPosition, const FVector2D& InScreenPosition);
+	virtual void HandleBreakpointAdded();
+	virtual void OnGraphNodeClicked(URigVMEdGraphNode* InNode);
+	virtual void OnNodeDoubleClicked(URigVMBlueprint* InBlueprint, URigVMNode* InNode);
+	virtual void OnGraphImported(UEdGraph* InEdGraph);
+	virtual bool OnActionMatchesName(FEdGraphSchemaAction* InAction, const FName& InName) const override;
+	virtual void HandleShowCurrentStatement();
+	virtual void HandleBreakpointActionRequested(const ERigVMBreakpointAction BreakpointAction);
+	virtual bool IsHaltedAtBreakpoint() const;
+	virtual void FrameSelection();
+
+	/** Once the log is collected update the graph */
+	void UpdateGraphCompilerErrors();
+
 private:
 
 	/** Our currently running rig vm instance */
@@ -38,5 +276,42 @@ private:
 
 	FPreviewHostUpdated PreviewHostUpdated;
 
+	FRigVMEditorClosed RigVMEditorClosedDelegate;
+
+	/** Toolbox hosting widget */
+	TSharedPtr<SBorder> Toolbox;
+
+	bool bAnyErrorsLeft;
+	TMap<FString, FString> KnownInstructionLimitWarnings;
+	URigVMNode* HaltedAtNode;
+	FString LastDebuggedHost;
+
+	bool bSuspendDetailsPanelRefresh;
+	bool bAllowBulkEdits;
+	bool bIsSettingObjectBeingDebugged;
+	bool bRigVMEditorInitialized;
+	
+	/** Are we currently compiling through the user interface */
+	bool bIsCompilingThroughUI;
+
+	TArray<TStrongObjectPtr<URigVMDetailsViewWrapperObject>> WrapperObjects;
+
+	ERigVMEditorExecutionModeType ExecutionMode;
+
+	/** The log to use for errors resulting from the init phase of the units */
+	FRigVMLog RigVMLog;
+	
+
+protected:
+	TArray<FName> LastEventQueue;
+
+	/** The extender to pass to the level editor to extend it's window menu */
+	TSharedPtr<FExtender> MenuExtender;
+
+	/** Toolbar extender */
+	TSharedPtr<FExtender> ToolbarExtender;
+
+	FDelegateHandle PropertyChangedHandle;
+
 	friend class SRigVMExecutionStackView;
 };
diff --git a/Engine/Plugins/Runtime/RigVM/Source/RigVMEditor/Public/Editor/RigVMExecutionStackCommands.h b/Engine/Plugins/Runtime/RigVM/Source/RigVMEditor/Public/Editor/RigVMExecutionStackCommands.h
--- a/Engine/Plugins/Runtime/RigVM/Source/RigVMEditor/Public/Editor/RigVMExecutionStackCommands.h
+++ b/Engine/Plugins/Runtime/RigVM/Source/RigVMEditor/Public/Editor/RigVMExecutionStackCommands.h
@@ -10,7 +10,7 @@ class RIGVMEDITOR_API FRigVMExecutionStackCommands : public TCommands<FRigVMExec
 public:
 	FRigVMExecutionStackCommands() : TCommands<FRigVMExecutionStackCommands>
 	(
-		"ControlRigStack",
+		"RigVMExecutionStack",
 		NSLOCTEXT("Contexts", "RigStack", "Execution Stack"),
 		NAME_None, // "MainFrame" // @todo Fix this crash
 		"RigVMEditorStyle" // Icon Style Set
diff --git a/Engine/Plugins/Runtime/RigVM/Source/RigVMEditor/Public/RigVMEditorModule.h b/Engine/Plugins/Runtime/RigVM/Source/RigVMEditor/Public/RigVMEditorModule.h
--- a/Engine/Plugins/Runtime/RigVM/Source/RigVMEditor/Public/RigVMEditorModule.h
+++ b/Engine/Plugins/Runtime/RigVM/Source/RigVMEditor/Public/RigVMEditorModule.h
@@ -28,6 +28,10 @@ public:
 	}
 
 	virtual void GetContextMenuActions(const URigVMEdGraphSchema* Schema, class UToolMenu* Menu, class UGraphNodeContextMenuContext* Context) const = 0;
+	virtual void GetTypeActions(URigVMBlueprint* RigVMBlueprint, FBlueprintActionDatabaseRegistrar& ActionRegistrar) = 0;
+	virtual void GetInstanceActions(URigVMBlueprint* RigVMBlueprint, FBlueprintActionDatabaseRegistrar& ActionRegistrar) = 0;
+	virtual void GetNodeContextMenuActions(URigVMBlueprint* RigVMBlueprint, const URigVMEdGraphNode* EdGraphNode, URigVMNode* ModelNode, class UToolMenu* Menu) const = 0;
+	virtual void GetPinContextMenuActions(URigVMBlueprint* RigVMBlueprint, const UEdGraphPin* EdGraphPin, URigVMPin* ModelPin, class UToolMenu* Menu) const = 0;
 	virtual FConnectionDrawingPolicy* CreateConnectionDrawingPolicy(int32 InBackLayerID, int32 InFrontLayerID, float InZoomFactor, const FSlateRect& InClippingRect, class FSlateWindowElementList& InDrawElements, class UEdGraph* InGraphObj) const = 0;
 };
 
diff --git a/Engine/Plugins/Runtime/RigVM/Source/RigVMEditor/RigVMEditor.Build.cs b/Engine/Plugins/Runtime/RigVM/Source/RigVMEditor/RigVMEditor.Build.cs
--- a/Engine/Plugins/Runtime/RigVM/Source/RigVMEditor/RigVMEditor.Build.cs
+++ b/Engine/Plugins/Runtime/RigVM/Source/RigVMEditor/RigVMEditor.Build.cs
@@ -25,10 +25,12 @@ public class RigVMEditor : ModuleRules
                 "KismetCompiler",
                 "KismetWidgets",
                 "AnimationWidgets",
+                "ApplicationCore",
                 "AppFramework",
                 "AnimationCore",
                 "PropertyEditor",
                 "ToolMenus",
+                "MessageLog",
             }
         );
     }
