RigVMExample plugin

diff --git a/Engine/Plugins/Animation/ControlRig/Source/ControlRigEditor/Private/Editor/ControlRigEditor.cpp b/Engine/Plugins/Animation/ControlRig/Source/ControlRigEditor/Private/Editor/ControlRigEditor.cpp
--- a/Engine/Plugins/Animation/ControlRig/Source/ControlRigEditor/Private/Editor/ControlRigEditor.cpp
+++ b/Engine/Plugins/Animation/ControlRig/Source/ControlRigEditor/Private/Editor/ControlRigEditor.cpp
@@ -112,13 +112,6 @@
 
 TAutoConsoleVariable<bool> CVarControlRigShowTestingToolbar(TEXT("ControlRig.Test.EnableTestingToolbar"), false, TEXT("When true we'll show the testing toolbar in Control Rig Editor."));
 
-namespace ControlRigEditorTabs
-{
-	const FName DetailsTab(TEXT("DetailsTab"));
-// 	const FName ViewportTab(TEXT("Viewport"));
-// 	const FName AdvancedPreviewTab(TEXT("AdvancedPreviewTab"));
-};
-
 const FName FControlRigEditorModes::ControlRigEditorMode = TEXT("Rigging");
 const TArray<FName> FControlRigEditor::ForwardsSolveEventQueue = {FRigUnit_BeginExecution::EventName};
 const TArray<FName> FControlRigEditor::BackwardsSolveEventQueue = {FRigUnit_InverseExecution::EventName};
@@ -236,9 +229,6 @@ void FControlRigEditor::InitRigVMEditor(const EToolkitMode::Type Mode, const TSh
 		PersonaToolkit->GetPreviewScene()->SetRemoveAttachedComponentFilter(FOnRemoveAttachedComponentFilter::CreateSP(EditMode, &FControlRigEditMode::CanRemoveFromPreviewScene));
 	}
 
-	// Post-layout initialization
-	PostLayoutBlueprintEditorInitialization();
-
 	{
 		// listening to the BP's event instead of BP's Hierarchy's Event ensure a propagation order of
 		// 1. Hierarchy change in BP
diff --git a/Engine/Plugins/Runtime/RigVM/Source/RigVM/Private/RigVMHost.cpp b/Engine/Plugins/Runtime/RigVM/Source/RigVM/Private/RigVMHost.cpp
--- a/Engine/Plugins/Runtime/RigVM/Source/RigVM/Private/RigVMHost.cpp
+++ b/Engine/Plugins/Runtime/RigVM/Source/RigVM/Private/RigVMHost.cpp
@@ -914,7 +914,7 @@ void URigVMHost::PostInitInstance(URigVMHost* InCDO)
 void URigVMHost::HandleExecutionReachedExit(const FName& InEventName)
 {
 #if WITH_EDITOR
-	if (EventQueueToRun.Last() == InEventName)
+	if (EventQueueToRun.IsEmpty() || EventQueueToRun.Last() == InEventName)
 	{
 		if(URigVM* SnapShotVM = GetSnapshotVM(false))
 		{
diff --git a/Engine/Plugins/Runtime/RigVM/Source/RigVMDeveloper/Private/EdGraph/RigVMEdGraphNode.cpp b/Engine/Plugins/Runtime/RigVM/Source/RigVMDeveloper/Private/EdGraph/RigVMEdGraphNode.cpp
--- a/Engine/Plugins/Runtime/RigVM/Source/RigVMDeveloper/Private/EdGraph/RigVMEdGraphNode.cpp
+++ b/Engine/Plugins/Runtime/RigVM/Source/RigVMDeveloper/Private/EdGraph/RigVMEdGraphNode.cpp
@@ -1395,6 +1395,10 @@ void URigVMEdGraphNode::GetNodeContextMenuActions(class UToolMenu* Menu, class U
 {
 #if WITH_EDITOR
 	const URigVMEdGraphSchema* Schema = Cast<URigVMEdGraphSchema>(GetSchema());
+	if(const URigVMBlueprint* Blueprint = GetBlueprint())
+	{
+		return Blueprint->GetEditorModule()->GetContextMenuActions(Schema, Menu, Context);
+	}
 	IRigVMEditorModule::Get().GetContextMenuActions(Schema, Menu, Context);
 #endif
 }
diff --git a/Engine/Plugins/Runtime/RigVM/Source/RigVMDeveloper/Private/EdGraph/RigVMEdGraphSchema.cpp b/Engine/Plugins/Runtime/RigVM/Source/RigVMDeveloper/Private/EdGraph/RigVMEdGraphSchema.cpp
--- a/Engine/Plugins/Runtime/RigVM/Source/RigVMDeveloper/Private/EdGraph/RigVMEdGraphSchema.cpp
+++ b/Engine/Plugins/Runtime/RigVM/Source/RigVMDeveloper/Private/EdGraph/RigVMEdGraphSchema.cpp
@@ -1252,6 +1252,10 @@ FReply URigVMEdGraphSchema::BeginGraphDragAction(TSharedPtr<FEdGraphSchemaAction
 FConnectionDrawingPolicy* URigVMEdGraphSchema::CreateConnectionDrawingPolicy(int32 InBackLayerID, int32 InFrontLayerID, float InZoomFactor, const FSlateRect& InClippingRect, class FSlateWindowElementList& InDrawElements, class UEdGraph* InGraphObj) const
 {
 #if WITH_EDITOR
+	if (const URigVMBlueprint* Blueprint = Cast<URigVMBlueprint>(FBlueprintEditorUtils::FindBlueprintForGraph(InGraphObj)))
+	{
+		return Blueprint->GetEditorModule()->CreateConnectionDrawingPolicy(InBackLayerID, InFrontLayerID, InZoomFactor, InClippingRect, InDrawElements, InGraphObj);
+	}	
 	return IRigVMEditorModule::Get().CreateConnectionDrawingPolicy(InBackLayerID, InFrontLayerID, InZoomFactor, InClippingRect, InDrawElements, InGraphObj);
 #else
 	check(0);
diff --git a/Engine/Plugins/Runtime/RigVM/Source/RigVMEditor/Private/Editor/RigVMEditor.cpp b/Engine/Plugins/Runtime/RigVM/Source/RigVMEditor/Private/Editor/RigVMEditor.cpp
--- a/Engine/Plugins/Runtime/RigVM/Source/RigVMEditor/Private/Editor/RigVMEditor.cpp
+++ b/Engine/Plugins/Runtime/RigVM/Source/RigVMEditor/Private/Editor/RigVMEditor.cpp
@@ -41,9 +41,12 @@
 #include "Widgets/Layout/SScrollBox.h"
 #include "Framework/Docking/TabManager.h"
 #include "ScopedTransaction.h"
+#include "Editor/RigVMEditorMode.h"
 
 #define LOCTEXT_NAMESPACE "RigVMEditor"
 
+const FName FRigVMEditorModes::RigVMEditorMode = TEXT("RigVM");
+
 FRigVMEditor::FRigVMEditor()
 	: Host(nullptr)
 	, bAnyErrorsLeft(false)
@@ -365,14 +368,12 @@ const FName FRigVMEditor::GetEditorAppName() const
 
 const FName FRigVMEditor::GetEditorModeName() const
 {
-	const FName EditorMode("VisualProgramming");
-	return EditorMode;
+	return FRigVMEditorModes::RigVMEditorMode;
 }
 
 TSharedPtr<FApplicationMode> FRigVMEditor::CreateEditorMode()
 {
-	static TSharedPtr<FApplicationMode> EmptyMode;
-	return EmptyMode;
+	return MakeShareable(new FRigVMEditorMode(SharedThis(this)));
 }
 
 UBlueprint* FRigVMEditor::GetBlueprintObj() const
diff --git a/Engine/Plugins/Runtime/RigVM/Source/RigVMEditor/Private/Editor/RigVMEditorMode.cpp b/Engine/Plugins/Runtime/RigVM/Source/RigVMEditor/Private/Editor/RigVMEditorMode.cpp
new file mode 100644
--- /dev/null
+++ b/Engine/Plugins/Runtime/RigVM/Source/RigVMEditor/Private/Editor/RigVMEditorMode.cpp
@@ -0,0 +1,93 @@
+// Copyright Epic Games, Inc. All Rights Reserved.
+
+#include "Editor/RigVMEditorMode.h" 
+#include "BlueprintEditorTabs.h"
+#include "SBlueprintEditorToolbar.h"
+#include "Editor/RigVMExecutionStackTabSummoner.h"
+#include "ToolMenus.h"
+
+FRigVMEditorMode::FRigVMEditorMode(const TSharedRef<FRigVMEditor>& InRigVMEditor)
+	: FBlueprintEditorApplicationMode(InRigVMEditor, FRigVMEditorModes::RigVMEditorMode, FRigVMEditorModes::GetLocalizedMode, false, false)
+{
+	RigVMBlueprintPtr = CastChecked<URigVMBlueprint>(InRigVMEditor->GetBlueprintObj());
+
+	TabFactories.RegisterFactory(MakeShared<FRigVMExecutionStackTabSummoner>(InRigVMEditor));
+
+	TabLayout = FTabManager::NewLayout("Standalone_RigVMEditMode_Layout_v1.5")
+		->AddArea
+		(
+			// Main application area
+			FTabManager::NewPrimaryArea()
+			->SetOrientation(Orient_Vertical)
+			->Split
+			(
+				FTabManager::NewSplitter()
+				->SetOrientation(Orient_Horizontal)
+				->Split
+				(
+					//	Left bottom - rig/hierarchy
+					FTabManager::NewStack()
+					->SetSizeCoefficient(0.5f)
+					->AddTab(FRigVMExecutionStackTabSummoner::TabID, ETabState::OpenedTab)
+					->AddTab(FBlueprintEditorTabs::MyBlueprintID, ETabState::OpenedTab)
+				)
+				->Split
+				(
+					// Middle 
+					FTabManager::NewSplitter()
+					->SetOrientation(Orient_Vertical)
+					->SetSizeCoefficient(0.6f)
+					->Split
+					(
+						// Middle top - document edit area
+						FTabManager::NewStack()
+						->SetSizeCoefficient(0.8f)
+						->AddTab("Document", ETabState::ClosedTab)
+					)
+					->Split
+					(
+						// Middle bottom - compiler results & find
+						FTabManager::NewStack()
+						->SetSizeCoefficient(0.2f)
+						->AddTab(FBlueprintEditorTabs::CompilerResultsID, ETabState::ClosedTab)
+						->AddTab(FBlueprintEditorTabs::FindResultsID, ETabState::ClosedTab)
+					)
+				)
+				->Split
+				(
+					// Right side
+					FTabManager::NewSplitter()
+					->SetOrientation(Orient_Vertical)
+					->SetSizeCoefficient(0.2f)
+					->Split
+					(
+						// Right top
+						FTabManager::NewStack()
+						->SetHideTabWell(false)
+						->SetSizeCoefficient(1.f)
+						->AddTab(FBlueprintEditorTabs::DetailsID, ETabState::OpenedTab)
+						->SetForegroundTab(FBlueprintEditorTabs::DetailsID)
+					)
+				)
+			)
+		);
+
+	if (UToolMenu* Toolbar = InRigVMEditor->RegisterModeToolbarIfUnregistered(GetModeName()))
+	{
+		InRigVMEditor->GetToolbarBuilder()->AddCompileToolbar(Toolbar);
+		InRigVMEditor->GetToolbarBuilder()->AddScriptingToolbar(Toolbar);
+		InRigVMEditor->GetToolbarBuilder()->AddBlueprintGlobalOptionsToolbar(Toolbar);
+	}
+}
+
+void FRigVMEditorMode::RegisterTabFactories(TSharedPtr<FTabManager> InTabManager)
+{
+	TSharedPtr<FBlueprintEditor> BP = MyBlueprintEditor.Pin();
+	
+	BP->RegisterToolbarTab(InTabManager.ToSharedRef());
+
+	// Mode-specific setup
+	BP->PushTabFactories(CoreTabFactories);
+	BP->PushTabFactories(BlueprintEditorTabFactories);
+	BP->PushTabFactories(TabFactories);
+}
diff --git a/Engine/Plugins/Runtime/RigVM/Source/RigVMEditor/Public/Editor/RigVMEditor.h b/Engine/Plugins/Runtime/RigVM/Source/RigVMEditor/Public/Editor/RigVMEditor.h
--- a/Engine/Plugins/Runtime/RigVM/Source/RigVMEditor/Public/Editor/RigVMEditor.h
+++ b/Engine/Plugins/Runtime/RigVM/Source/RigVMEditor/Public/Editor/RigVMEditor.h
@@ -12,6 +12,28 @@ class FRigVMEditor;
 
 DECLARE_MULTICAST_DELEGATE_TwoParams(FRigVMEditorClosed, const FRigVMEditor*, URigVMBlueprint*);
 
+struct FRigVMEditorModes
+{
+	// Mode constants
+	static const FName RigVMEditorMode;
+	static FText GetLocalizedMode(const FName InMode)
+	{
+		static TMap< FName, FText > LocModes;
+
+		if (LocModes.Num() == 0)
+		{
+			LocModes.Add(RigVMEditorMode, NSLOCTEXT("RigVMEditorModes", "RigVMEditorMode", "RigVM"));
+		}
+
+		check(InMode != NAME_None);
+		const FText* OutDesc = LocModes.Find(InMode);
+		check(OutDesc);
+		return *OutDesc;
+	}
+private:
+	FRigVMEditorModes() {}
+};
+
 class RIGVMEDITOR_API FRigVMEditor : public FBlueprintEditor
 {
 public:
@@ -289,8 +311,11 @@ private:
 	bool bSuspendDetailsPanelRefresh;
 	bool bAllowBulkEdits;
 	bool bIsSettingObjectBeingDebugged;
+
+protected:
 	bool bRigVMEditorInitialized;
-	
+
+private:
 	/** Are we currently compiling through the user interface */
 	bool bIsCompilingThroughUI;
 
diff --git a/Engine/Plugins/Runtime/RigVM/Source/RigVMEditor/Public/Editor/RigVMEditorMode.h b/Engine/Plugins/Runtime/RigVM/Source/RigVMEditor/Public/Editor/RigVMEditorMode.h
new file mode 100644
--- /dev/null
+++ b/Engine/Plugins/Runtime/RigVM/Source/RigVMEditor/Public/Editor/RigVMEditorMode.h
@@ -0,0 +1,27 @@
+// Copyright Epic Games, Inc. All Rights Reserved.
+
+#pragma once
+
+#include "CoreMinimal.h"
+#include "WorkflowOrientedApp/WorkflowTabManager.h"
+#include "Editor/RigVMEditor.h"
+#include "RigVMBlueprint.h"
+#include "BlueprintEditorModes.h"
+
+class URigVMBlueprint;
+
+class FRigVMEditorMode : public FBlueprintEditorApplicationMode
+{
+public:
+	FRigVMEditorMode(const TSharedRef<FRigVMEditor>& InRigVMEditor);
+
+	// FApplicationMode interface
+	virtual void RegisterTabFactories(TSharedPtr<FTabManager> InTabManager) override;
+
+protected:
+	// Set of spawnable tabs
+	FWorkflowAllowedTabSet TabFactories;
+
+private:
+	TWeakObjectPtr<URigVMBlueprint> RigVMBlueprintPtr;
+};
