RigVm: Fix crash when opening editor

diff --git a/Engine/Plugins/Runtime/RigVM/Source/RigVMEditor/Private/Editor/RigVMEditor.cpp b/Engine/Plugins/Runtime/RigVM/Source/RigVMEditor/Private/Editor/RigVMEditor.cpp
--- a/Engine/Plugins/Runtime/RigVM/Source/RigVMEditor/Private/Editor/RigVMEditor.cpp
+++ b/Engine/Plugins/Runtime/RigVM/Source/RigVMEditor/Private/Editor/RigVMEditor.cpp
@@ -827,7 +827,7 @@ URigVMBlueprint* FRigVMEditor::GetRigVMBlueprint() const
 
 URigVMHost* FRigVMEditor::GetRigVMHost() const
 {
-	if(IsValid(Host))
+	if(Host && IsValid(Host))
 	{
 		return Host;
 	}
@@ -1013,10 +1013,12 @@ void FRigVMEditor::Compile()
 			RigVMHost->GetExtendedExecuteContext().ExecutionHalted().RemoveAll(this);
 		}
 
+		SetHost(nullptr);
 		{
 			TGuardValue<bool> GuardCompileReEntry(bIsCompilingThroughUI, true);
 			FBlueprintEditor::Compile();
 		}
+		UpdateRigVMHost();
 
 		if (URigVMHost* RigVMHost = GetRigVMHost())
 		{
