RigVM: Crash when resolving NearlyZero to vector type

diff --git a/Engine/Plugins/Runtime/RigVM/Source/RigVMDeveloper/Private/RigVMModel/RigVMController.cpp b/Engine/Plugins/Runtime/RigVM/Source/RigVMDeveloper/Private/RigVMModel/RigVMController.cpp
--- a/Engine/Plugins/Runtime/RigVM/Source/RigVMDeveloper/Private/RigVMModel/RigVMController.cpp
+++ b/Engine/Plugins/Runtime/RigVM/Source/RigVMDeveloper/Private/RigVMModel/RigVMController.cpp
@@ -16297,6 +16297,14 @@ bool URigVMController::FullyResolveTemplateNode(URigVMTemplateNode* InNode, int3
 	if (InPermutationIndex == INDEX_NONE)
 	{
 		TArray<int32> Permutations = InNode->GetResolvedPermutationIndices(false);
+
+		// If some float/double pin needs to change type, permutations might be empty
+		// Try running again allowing that change
+		if (Permutations.IsEmpty())
+		{
+			Permutations = InNode->GetResolvedPermutationIndices(true);
+		}
+		
 		check(!Permutations.IsEmpty());
 		InNode->ResolvedPermutation = Permutations[0];
 
