RigVM / ControlRig: Improve aggregate node details panel

diff --git a/Engine/Plugins/Runtime/RigVM/Source/RigVMEditor/Private/Editor/RigVMGraphDetailCustomization.cpp b/Engine/Plugins/Runtime/RigVM/Source/RigVMEditor/Private/Editor/RigVMGraphDetailCustomization.cpp
--- a/Engine/Plugins/Runtime/RigVM/Source/RigVMEditor/Private/Editor/RigVMGraphDetailCustomization.cpp
+++ b/Engine/Plugins/Runtime/RigVM/Source/RigVMEditor/Private/Editor/RigVMGraphDetailCustomization.cpp
@@ -23,6 +23,7 @@
 #include "EdGraph/RigVMEdGraphSchema.h"
 #include "EditorCategoryUtils.h"
 #include "IPropertyUtilities.h"
+#include "RigVMModel/Nodes/RigVMAggregateNode.h"
 #include "Widgets/SRigVMGraphPinVariableBinding.h"
 
 #define LOCTEXT_NAMESPACE "RigVMGraphDetailCustomization"
@@ -422,6 +423,22 @@ bool FRigVMFunctionArgumentLayout::ShouldPinBeReadOnly(bool bIsEditingPinType/*
 
 bool FRigVMFunctionArgumentLayout::IsPinEditingReadOnly(bool bIsEditingPinType/* = false*/) const
 {
+	if (PinPtr.IsValid())
+	{
+		if(const URigVMPin* Pin = PinPtr.Get())
+		{
+			if(Pin->IsExecuteContext())
+			{
+				if(const URigVMNode* Node = Pin->GetNode())
+				{
+					if(Node->IsA<URigVMAggregateNode>())
+					{
+						return true;
+					}
+				}
+			}
+		}
+	}
 	return false;
 }
 
@@ -752,6 +769,17 @@ void FRigVMGraphDetailCustomization::CustomizeDetails(IDetailLayoutBuilder& Deta
 		return;
 	}
 
+	bool bIsFunction = false;
+	bool bIsAggregate = false;
+	if (Model)
+	{
+		if (URigVMLibraryNode* LibraryNode = Cast<URigVMLibraryNode>(Model->GetOuter()))
+		{
+			bIsFunction = LibraryNode->GetGraph()->IsA<URigVMFunctionLibrary>();
+			bIsAggregate = LibraryNode->IsA<URigVMAggregateNode>();
+		}
+	}
+
 	IDetailCategoryBuilder& InputsCategory = DetailLayout.EditCategory("Inputs", LOCTEXT("FunctionDetailsInputs", "Inputs"));
 	TSharedRef<FRigVMFunctionArgumentGroupLayout> InputArgumentGroup = MakeShareable(new FRigVMFunctionArgumentGroupLayout(
 		Model, 
@@ -760,29 +788,32 @@ void FRigVMGraphDetailCustomization::CustomizeDetails(IDetailLayoutBuilder& Deta
 		true));
 	InputsCategory.AddCustomBuilder(InputArgumentGroup);
 
-	TSharedRef<SHorizontalBox> InputsHeaderContentWidget = SNew(SHorizontalBox);
+	if(!bIsAggregate)
+	{
+		TSharedRef<SHorizontalBox> InputsHeaderContentWidget = SNew(SHorizontalBox);
 
-	InputsHeaderContentWidget->AddSlot()
-	.HAlign(HAlign_Right)
-	[
-		SNew(SButton)
-		.ButtonStyle(FAppStyle::Get(), "SimpleButton")
-		.ContentPadding(FMargin(1, 0))
-		.OnClicked(this, &FRigVMGraphDetailCustomization::OnAddNewInputClicked)
-		.Visibility(this, &FRigVMGraphDetailCustomization::GetAddNewInputOutputVisibility)
+		InputsHeaderContentWidget->AddSlot()
 		.HAlign(HAlign_Right)
-		.ToolTipText(LOCTEXT("FunctionNewInputArgTooltip", "Create a new input argument"))
-		.VAlign(VAlign_Center)
-		.AddMetaData<FTagMetaData>(FTagMetaData(TEXT("FunctionNewInputArg")))
-		.IsEnabled(this, &FRigVMGraphDetailCustomization::IsAddNewInputOutputEnabled)
 		[
-			SNew(SImage)
-			.Image(FAppStyle::Get().GetBrush("Icons.PlusCircle"))
-			.ColorAndOpacity(FSlateColor::UseForeground())
-		]
-	];
-	InputsCategory.HeaderContent(InputsHeaderContentWidget);
-
+			SNew(SButton)
+			.ButtonStyle(FAppStyle::Get(), "SimpleButton")
+			.ContentPadding(FMargin(1, 0))
+			.OnClicked(this, &FRigVMGraphDetailCustomization::OnAddNewInputClicked)
+			.Visibility(this, &FRigVMGraphDetailCustomization::GetAddNewInputOutputVisibility)
+			.HAlign(HAlign_Right)
+			.ToolTipText(LOCTEXT("FunctionNewInputArgTooltip", "Create a new input argument"))
+			.VAlign(VAlign_Center)
+			.AddMetaData<FTagMetaData>(FTagMetaData(TEXT("FunctionNewInputArg")))
+			.IsEnabled(this, &FRigVMGraphDetailCustomization::IsAddNewInputOutputEnabled)
+			[
+				SNew(SImage)
+				.Image(FAppStyle::Get().GetBrush("Icons.PlusCircle"))
+				.ColorAndOpacity(FSlateColor::UseForeground())
+			]
+		];
+		InputsCategory.HeaderContent(InputsHeaderContentWidget);
+	}
+	
 	IDetailCategoryBuilder& OutputsCategory = DetailLayout.EditCategory("Outputs", LOCTEXT("FunctionDetailsOutputs", "Outputs"));
 	TSharedRef<FRigVMFunctionArgumentGroupLayout> OutputArgumentGroup = MakeShareable(new FRigVMFunctionArgumentGroupLayout(
 		Model, 
@@ -791,39 +822,33 @@ void FRigVMGraphDetailCustomization::CustomizeDetails(IDetailLayoutBuilder& Deta
 		false));
 	OutputsCategory.AddCustomBuilder(OutputArgumentGroup);
 
-	TSharedRef<SHorizontalBox> OutputsHeaderContentWidget = SNew(SHorizontalBox);
+	if(!bIsAggregate)
+	{
+		TSharedRef<SHorizontalBox> OutputsHeaderContentWidget = SNew(SHorizontalBox);
 
-	OutputsHeaderContentWidget->AddSlot()
-	.HAlign(HAlign_Right)
-	[
-		SNew(SButton)
-		.ButtonStyle(FAppStyle::Get(), "SimpleButton")
-		.ContentPadding(FMargin(1, 0))
-		.OnClicked(this, &FRigVMGraphDetailCustomization::OnAddNewOutputClicked)
-		.Visibility(this, &FRigVMGraphDetailCustomization::GetAddNewInputOutputVisibility)
+		OutputsHeaderContentWidget->AddSlot()
 		.HAlign(HAlign_Right)
-		.ToolTipText(LOCTEXT("FunctionNewOutputArgTooltip", "Create a new output argument"))
-		.VAlign(VAlign_Center)
-		.AddMetaData<FTagMetaData>(FTagMetaData(TEXT("FunctionNewOutputArg")))
-		.IsEnabled(this, &FRigVMGraphDetailCustomization::IsAddNewInputOutputEnabled)
 		[
-			SNew(SImage)
-			.Image(FAppStyle::Get().GetBrush("Icons.PlusCircle"))
-			.ColorAndOpacity(FSlateColor::UseForeground())
-		]
-	];
-	OutputsCategory.HeaderContent(OutputsHeaderContentWidget);
-
-	IDetailCategoryBuilder& SettingsCategory = DetailLayout.EditCategory("NodeSettings", LOCTEXT("FunctionDetailsNodeSettings", "Node Settings"));
-
-	bool bIsFunction = false;
-	if (Model)
-	{
-		if (URigVMLibraryNode* LibraryNode = Cast<URigVMLibraryNode>(Model->GetOuter()))
-		{
-			bIsFunction = LibraryNode->GetGraph()->IsA<URigVMFunctionLibrary>();
-		}
+			SNew(SButton)
+			.ButtonStyle(FAppStyle::Get(), "SimpleButton")
+			.ContentPadding(FMargin(1, 0))
+			.OnClicked(this, &FRigVMGraphDetailCustomization::OnAddNewOutputClicked)
+			.Visibility(this, &FRigVMGraphDetailCustomization::GetAddNewInputOutputVisibility)
+			.HAlign(HAlign_Right)
+			.ToolTipText(LOCTEXT("FunctionNewOutputArgTooltip", "Create a new output argument"))
+			.VAlign(VAlign_Center)
+			.AddMetaData<FTagMetaData>(FTagMetaData(TEXT("FunctionNewOutputArg")))
+			.IsEnabled(this, &FRigVMGraphDetailCustomization::IsAddNewInputOutputEnabled)
+			[
+				SNew(SImage)
+				.Image(FAppStyle::Get().GetBrush("Icons.PlusCircle"))
+				.ColorAndOpacity(FSlateColor::UseForeground())
+			]
+		];
+		OutputsCategory.HeaderContent(OutputsHeaderContentWidget);
 	}
+	
+	IDetailCategoryBuilder& SettingsCategory = DetailLayout.EditCategory("NodeSettings", LOCTEXT("FunctionDetailsNodeSettings", "Node Settings"));
 
 	if(bIsFunction)
 	{
@@ -929,31 +954,33 @@ void FRigVMGraphDetailCustomization::CustomizeDetails(IDetailLayoutBuilder& Deta
 	}
 
 	// node color
-	SettingsCategory.AddCustomRow(FText::GetEmpty())
-	.NameContent()
-	[
-		SNew(STextBlock)
-		.Text(FText::FromString(TEXT("Color")))
-		.Font(IDetailLayoutBuilder::GetDetailFont())
-	]
-	.ValueContent()
-	[
-		SNew(SButton)
-		.ButtonStyle(FAppStyle::Get(), "Menu.Button")
-		.OnClicked(this, &FRigVMGraphDetailCustomization::OnNodeColorClicked)
+	if(!bIsAggregate)
+	{
+		SettingsCategory.AddCustomRow(FText::GetEmpty())
+		.NameContent()
 		[
-			SAssignNew(ColorBlock, SColorBlock)
-			.Color(this, &FRigVMGraphDetailCustomization::GetNodeColor)
-			.Size(FVector2D(77, 16))
+			SNew(STextBlock)
+			.Text(FText::FromString(TEXT("Color")))
+			.Font(IDetailLayoutBuilder::GetDetailFont())
 		]
-	];
+		.ValueContent()
+		[
+			SNew(SButton)
+			.ButtonStyle(FAppStyle::Get(), "Menu.Button")
+			.OnClicked(this, &FRigVMGraphDetailCustomization::OnNodeColorClicked)
+			[
+				SAssignNew(ColorBlock, SColorBlock)
+				.Color(this, &FRigVMGraphDetailCustomization::GetNodeColor)
+				.Size(FVector2D(77, 16))
+			]
+		];
+	}
 
 	IDetailCategoryBuilder& DefaultsCategory = DetailLayout.EditCategory("NodeDefaults", LOCTEXT("FunctionDetailsNodeDefaults", "Node Defaults"));
 	TSharedRef<FRigVMFunctionArgumentDefaultNode> DefaultsArgumentNode = MakeShareable(new FRigVMFunctionArgumentDefaultNode(
 		Model,
 		Blueprint));
 	DefaultsCategory.AddCustomBuilder(DefaultsArgumentNode);
-
 }
 
 bool FRigVMGraphDetailCustomization::IsAddNewInputOutputEnabled() const
