RigVM: Fix debug watches inside functions not working

diff --git a/Engine/Plugins/Runtime/RigVM/Source/RigVM/Private/RigVMCore/RigVM.cpp b/Engine/Plugins/Runtime/RigVM/Source/RigVM/Private/RigVMCore/RigVM.cpp
--- a/Engine/Plugins/Runtime/RigVM/Source/RigVM/Private/RigVMCore/RigVM.cpp
+++ b/Engine/Plugins/Runtime/RigVM/Source/RigVM/Private/RigVMCore/RigVM.cpp
@@ -145,6 +145,13 @@ void URigVM::Save(FArchive& Ar)
 		Ar << FunctionNamesStorage;
 		Ar << ByteCodeStorage;
 		Ar << Parameters;
+
+		if (Ar.CustomVer(FFortniteMainBranchObjectVersion::GUID) < FFortniteMainBranchObjectVersion::RigVMSaveDebugMapInGraphFunctionData)
+		{
+			return;
+		}
+
+		Ar << OperandToDebugRegisters;
 	}
 }
 
@@ -188,6 +195,11 @@ void URigVM::Load(FArchive& Ar)
 				Reset();
 				return;
 			}
+
+			if (Ar.CustomVer(FFortniteMainBranchObjectVersion::GUID) >= FFortniteMainBranchObjectVersion::RigVMSaveDebugMapInGraphFunctionData)
+			{
+				Ar << OperandToDebugRegisters;
+			}
 		}
 
 		// we only deal with virtual machines now that use the new memory infrastructure.
@@ -217,6 +229,11 @@ void URigVM::Load(FArchive& Ar)
 		Ar << FunctionNamesStorage;
 		Ar << ByteCodeStorage;
 		Ar << Parameters;
+
+		if (Ar.CustomVer(FFortniteMainBranchObjectVersion::GUID) >= FFortniteMainBranchObjectVersion::RigVMSaveDebugMapInGraphFunctionData)
+		{
+			Ar << OperandToDebugRegisters;
+		}
 	}
 
 	// ensure to load the required functions
@@ -498,6 +515,7 @@ void URigVM::Reset(bool IsIgnoringArchetypeRef)
 		Instructions.Reset();
 		Parameters.Reset();
 		ParametersNameMap.Reset();
+		OperandToDebugRegisters.Reset();
 	}
 
 	if(!IsIgnoringArchetypeRef)
@@ -508,8 +526,6 @@ void URigVM::Reset(bool IsIgnoringArchetypeRef)
 		ByteCodePtr = &ByteCodeStorage;
 	}
 
-	OperandToDebugRegisters.Reset();
-
 	ExternalPropertyPaths.Reset();
 	LazyBranches.Reset();
 
diff --git a/Engine/Plugins/Runtime/RigVM/Source/RigVM/Public/RigVMCore/RigVMGraphFunctionDefinition.h b/Engine/Plugins/Runtime/RigVM/Source/RigVM/Public/RigVMCore/RigVMGraphFunctionDefinition.h
--- a/Engine/Plugins/Runtime/RigVM/Source/RigVM/Public/RigVMCore/RigVMGraphFunctionDefinition.h
+++ b/Engine/Plugins/Runtime/RigVM/Source/RigVM/Public/RigVMCore/RigVMGraphFunctionDefinition.h
@@ -127,6 +127,8 @@ struct RIGVM_API FRigVMFunctionCompilationData
 	UPROPERTY()
 	uint32 Hash;
 
+	TMap<FRigVMOperand, TArray<FRigVMOperand>> OperandToDebugRegisters;
+
 	bool IsValid() const
 	{
 		return Hash != 0;
@@ -193,6 +195,8 @@ struct RIGVM_API FRigVMFunctionCompilationData
 
 	friend FArchive& operator<<(FArchive& Ar, FRigVMFunctionCompilationData& Data)
 	{
+		Ar.UsingCustomVersion(FFortniteMainBranchObjectVersion::GUID);
+		
 		Ar << Data.ByteCode;
 		Ar << Data.FunctionNames;
 		Ar << Data.WorkPropertyDescriptions;
@@ -206,6 +210,48 @@ struct RIGVM_API FRigVMFunctionCompilationData
 		Ar << Data.ExternalRegisterIndexToVariable;
 		Ar << Data.Operands;
 		Ar << Data.Hash;
+
+		if (Ar.CustomVer(FFortniteMainBranchObjectVersion::GUID) < FFortniteMainBranchObjectVersion::RigVMSaveDebugMapInGraphFunctionData)
+		{
+			return Ar;
+		}
+
+		// Serialize OperandToDebugRegisters
+		{
+			uint8 NumKeys = Data.OperandToDebugRegisters.Num();
+			Ar << NumKeys;
+
+			if (Ar.IsLoading())
+			{
+				for (int32 KeyIndex=0; KeyIndex<NumKeys; ++KeyIndex)
+				{
+					FRigVMOperand Key;
+					Ar << Key;
+					uint8 NumValues;
+					Ar << NumValues;
+					TArray<FRigVMOperand> Values;
+					Values.SetNumUninitialized(NumValues);
+					for (int32 ValueIndex=0; ValueIndex<NumValues; ++ValueIndex)
+					{
+						Ar << Values[ValueIndex];
+					}
+					Data.OperandToDebugRegisters.Add(Key, Values);
+				}
+			}
+			else
+			{
+				for (TPair<FRigVMOperand, TArray<FRigVMOperand>>& Pair : Data.OperandToDebugRegisters)
+				{
+					Ar << Pair.Key;
+					uint8 NumValues = Pair.Value.Num();
+					Ar << NumValues;
+					for (FRigVMOperand& Operand : Pair.Value)
+					{
+						Ar << Operand;
+					}
+				}
+			}
+		}
 		return Ar;
 	}
 };
diff --git a/Engine/Plugins/Runtime/RigVM/Source/RigVMDeveloper/Private/RigVMCompiler/RigVMCompiler.cpp b/Engine/Plugins/Runtime/RigVM/Source/RigVMDeveloper/Private/RigVMCompiler/RigVMCompiler.cpp
--- a/Engine/Plugins/Runtime/RigVM/Source/RigVMDeveloper/Private/RigVMCompiler/RigVMCompiler.cpp
+++ b/Engine/Plugins/Runtime/RigVM/Source/RigVMDeveloper/Private/RigVMCompiler/RigVMCompiler.cpp
@@ -1143,6 +1143,7 @@ bool URigVMCompiler::Compile(TArray<URigVMGraph*> InGraphs, URigVMController* In
 		}
 
 		OutFunctionCompilationData->Hash = GetTypeHash(OutFunctionCompilationData);
+		OutFunctionCompilationData->OperandToDebugRegisters = WorkData.VM->OperandToDebugRegisters;
 	}
 
 	// make sure all functions are known and resolved now.
@@ -2013,6 +2014,7 @@ int32 URigVMCompiler::TraverseInlineFunction(const FRigVMInlineFunctionExprAST*
 				}
 			}
 
+			// Find out how many properties we need to skip (function arguments and execute pins) 
 			int32 NumProperties = 0;
 			if (MemoryType == ERigVMMemoryType::Work)
 			{
@@ -2040,7 +2042,8 @@ int32 URigVMCompiler::TraverseInlineFunction(const FRigVMInlineFunctionExprAST*
 					NumProperties++;
 				}
 			}
-			
+
+			// Add internal properties
 			int32 StartIndex = MemoryType == ERigVMMemoryType::Work ? NumProperties : 0;
 			for (int32 PropertyIndex = StartIndex; PropertyIndex < Properties.Num(); ++PropertyIndex)
 			{
@@ -2057,12 +2060,13 @@ int32 URigVMCompiler::TraverseInlineFunction(const FRigVMInlineFunctionExprAST*
 
 				// @todo Try to reuse literal operands
 			}
-		}		
+		}
 	}
 	else
 	{
-		TArray<FRigVMOperand> Operands;
+		TArray<FRigVMOperand> InterfaceOperands;
 		TArray<FString> OperandsPinNames;
+		// Find operands related to the function's interface
 		for(const URigVMPin* Pin : FunctionReferenceNode->GetPins())
 		{
 			const FRigVMExprAST* ChildExpr = InExpr->FindExprWithPinName(Pin->GetFName());
@@ -2072,7 +2076,7 @@ int32 URigVMCompiler::TraverseInlineFunction(const FRigVMInlineFunctionExprAST*
 				const FRigVMVarExprAST* SourceVarExpr = GetSourceVarExpr(ChildExpr->To<FRigVMCachedValueExprAST>()->GetVarExpr());
 				if(!SourceVarExpr->IsExecuteContext())
 				{
-					Operands.Add(WorkData.ExprToOperand.FindChecked(SourceVarExpr));
+					InterfaceOperands.Add(WorkData.ExprToOperand.FindChecked(SourceVarExpr));
 					OperandsPinNames.Add(GetPinNameWithDirectionPrefix(Pin));
 				}
 			}
@@ -2081,7 +2085,7 @@ int32 URigVMCompiler::TraverseInlineFunction(const FRigVMInlineFunctionExprAST*
 				const FRigVMVarExprAST* SourceVarExpr = GetSourceVarExpr(ChildExpr->To<FRigVMVarExprAST>());
 				if(!SourceVarExpr->IsExecuteContext())
 				{
-					Operands.Add(WorkData.ExprToOperand.FindChecked(SourceVarExpr));
+					InterfaceOperands.Add(WorkData.ExprToOperand.FindChecked(SourceVarExpr));
 					OperandsPinNames.Add(GetPinNameWithDirectionPrefix(Pin));
 				}
 			}
@@ -2097,10 +2101,11 @@ int32 URigVMCompiler::TraverseInlineFunction(const FRigVMInlineFunctionExprAST*
 		FRigVMByteCode& ByteCode = WorkData.VM->GetByteCode();
 		InstructionIndexStart = ByteCode.GetNumInstructions();
 		FRigVMInstructionArray OldInstructions = ByteCode.GetInstructions();
-		ByteCode.InlineFunction(&FunctionByteCode, Operands);
+		ByteCode.InlineFunction(&FunctionByteCode, InterfaceOperands);
 		InstructionIndexEnd = ByteCode.GetNumInstructions() - 1;
 		FRigVMCallstack FuncRefCallstack = InExpr->GetProxy().GetCallstack();
 
+		// Import branch infos, with the proper instruction indices
 		for(const FRigVMBranchInfo& BranchInfo : FunctionByteCode.BranchInfos)
 		{
 			const int32 BranchInfoIndex = ByteCode.AddBranchInfo(BranchInfo);
@@ -2113,6 +2118,8 @@ int32 URigVMCompiler::TraverseInlineFunction(const FRigVMInlineFunctionExprAST*
 			ByteCode.BranchInfos[BranchInfoIndex].LastInstruction += (uint16)InstructionIndexStart;
 		}
 
+		TMap<FRigVMOperand, FRigVMOperand> FunctionOperandToNewOperand;
+
 		// For each instruction, substitute the operand for the one used in the current bytecode
 		const FRigVMInstructionArray FunctionInstructions = FunctionByteCode.GetInstructions();
 		FRigVMInstructionArray Instructions = ByteCode.GetInstructions();
@@ -2128,6 +2135,7 @@ int32 URigVMCompiler::TraverseInlineFunction(const FRigVMInlineFunctionExprAST*
 			for (int32 j=0; j<OperandArray.Num(); ++j)
 			{
 				FRigVMOperand* Operand = (FRigVMOperand*)(ByteCode.ByteCode.GetData() + OperandsIndex + j*sizeof(FRigVMOperand));
+				FRigVMOperand FunctionOperand = *Operand;
 				ERigVMMemoryType OriginalMemoryType = Operand->GetMemoryType();
 
 				// Remap the variable: find the operand index of the outer variable
@@ -2149,7 +2157,7 @@ int32 URigVMCompiler::TraverseInlineFunction(const FRigVMInlineFunctionExprAST*
 				// If Operand is an interface pin: replace the index and memory type
 				else if (const int32 FunctionInterfaceParameterIndex = GetOperandFunctionInterfaceParameterIndex(OperandsPinNames, FunctionCompilationData, *Operand); FunctionInterfaceParameterIndex != INDEX_NONE)
 				{
-					const FRigVMOperand& InterfaceOperand = Operands[FunctionInterfaceParameterIndex];
+					const FRigVMOperand& InterfaceOperand = InterfaceOperands[FunctionInterfaceParameterIndex];
 					Operand->MemoryType = InterfaceOperand.GetMemoryType();
 					Operand->RegisterIndex = InterfaceOperand.GetRegisterIndex();
 				}
@@ -2166,7 +2174,6 @@ int32 URigVMCompiler::TraverseInlineFunction(const FRigVMInlineFunctionExprAST*
 				// For all operands, check to see if we need to add a property path
 				if (Operand->GetRegisterOffset() != INDEX_NONE)
 				{
-					
 					FRigVMFunctionCompilationPropertyPath Description;
 					switch (OriginalMemoryType)
 					{
@@ -2193,6 +2200,23 @@ int32 URigVMCompiler::TraverseInlineFunction(const FRigVMInlineFunctionExprAST*
 					}
 					Operand->RegisterOffset = WorkData.FindOrAddPropertyPath(*Operand, Description.HeadCPPType, Description.SegmentPath);
 				}
+
+				// Store correspondence between Function operand and New operands
+				check(!FunctionOperandToNewOperand.Contains(FunctionOperand) || FunctionOperandToNewOperand.FindChecked(FunctionOperand) == *Operand);
+				FunctionOperandToNewOperand.Add(FunctionOperand, *Operand);
+
+				// Also find the function debug operands and store correspondence to new debug operands
+				if (const TArray<FRigVMOperand>* FunctionDebugRegisters = FunctionCompilationData->OperandToDebugRegisters.Find(FunctionOperand))
+				{
+					for (const FRigVMOperand& FunctionDebugRegister : *FunctionDebugRegisters)
+					{
+						FRigVMCompilerWorkData::FFunctionRegisterData Data = {FunctionReferenceNode, FunctionDebugRegister.GetMemoryType(), FunctionDebugRegister.GetRegisterIndex()};
+						if (FRigVMOperand* NewDebugOperand = WorkData.FunctionRegisterToOperand.Find(Data))
+						{
+							FunctionOperandToNewOperand.Add(FunctionDebugRegister, *NewDebugOperand);
+						}
+					}
+				}
 			}
 
 			if (Instruction.OpCode == ERigVMOpCode::Execute)
@@ -2221,6 +2245,46 @@ int32 URigVMCompiler::TraverseInlineFunction(const FRigVMInlineFunctionExprAST*
 				}
 			}
 		}
+
+		// Add all pin paths to operand from the function
+		TMap<FRigVMOperand, TArray<FString>> OperandToPinPath;
+		for (const TPair<FString, FRigVMOperand>& PinPathToOperand : FunctionCompilationData->Operands)
+		{
+			if (FRigVMOperand* NewOperand = FunctionOperandToNewOperand.Find(PinPathToOperand.Value))
+			{
+				WorkData.PinPathToOperand->Add(PinPathToOperand.Key, *NewOperand);
+				TArray<FString>& Paths = OperandToPinPath.Emplace(*NewOperand);
+				Paths.Add(PinPathToOperand.Key);
+			}
+		}
+
+		// Fill out the operand to debug registers with the correct operands from the function
+		for (const TPair<FRigVMOperand, TArray<FRigVMOperand>>& Pair : FunctionCompilationData->OperandToDebugRegisters)
+		{
+			const FRigVMOperand& FunctionOperand = Pair.Key;
+			const TArray<FRigVMOperand>& FunctionDebugOperands = Pair.Value;
+			if (FRigVMOperand* NewOperand = FunctionOperandToNewOperand.Find(FunctionOperand))
+			{
+				TArray<FRigVMOperand> NewValues;
+				NewValues.Reserve(Pair.Value.Num());
+				for (const FRigVMOperand& FunctionDebugOperand : FunctionDebugOperands)
+				{
+					if (FRigVMOperand* NewDebugOperand = FunctionOperandToNewOperand.Find(FunctionDebugOperand))
+					{
+						NewValues.Add(*NewDebugOperand);
+						if (const TArray<FString>* Paths = OperandToPinPath.Find(*NewDebugOperand))
+						{
+							for (const FString& Path : *Paths)
+							{
+								FString PinPath = FString::Printf(TEXT("DebugWatch:%s"), *Path);
+								WorkData.PinPathToOperand->Add(PinPath, *NewDebugOperand);
+							}
+						}
+					}
+				}
+				WorkData.VM->OperandToDebugRegisters.Add(*NewOperand, NewValues);
+			}
+		}
 		
 
 #if WITH_EDITORONLY_DATA
@@ -2234,7 +2298,7 @@ int32 URigVMCompiler::TraverseInlineFunction(const FRigVMInlineFunctionExprAST*
 				continue;
 			}
 
-			const FRigVMOperand& Operand = Operands[ArgumentIndex++];
+			const FRigVMOperand& Operand = InterfaceOperands[ArgumentIndex++];
 
 			if(InputPin->GetDirection() == ERigVMPinDirection::Output || InputPin->GetDirection() == ERigVMPinDirection::IO)
 			{
diff --git a/Engine/Plugins/Runtime/RigVM/Source/RigVMDeveloper/Private/RigVMModel/RigVMClient.cpp b/Engine/Plugins/Runtime/RigVM/Source/RigVMDeveloper/Private/RigVMModel/RigVMClient.cpp
--- a/Engine/Plugins/Runtime/RigVM/Source/RigVMDeveloper/Private/RigVMModel/RigVMClient.cpp
+++ b/Engine/Plugins/Runtime/RigVM/Source/RigVMDeveloper/Private/RigVMModel/RigVMClient.cpp
@@ -1081,6 +1081,7 @@ void FRigVMClient::HandleGraphModifiedEvent(ERigVMGraphNotifType InNotifType, UR
 		case ERigVMGraphNotifType::PinDirectionChanged: // A pin's direction has changed (Subject == URigVMPin)
 		case ERigVMGraphNotifType::PinTypeChanged: // A pin's data type has changed (Subject == URigVMPin)
 		case ERigVMGraphNotifType::PinIndexChanged: // A pin's index has changed (Subject == URigVMPin)
+		case ERigVMGraphNotifType::PinWatchedChanged: // A pin's watch state has changed (Subject == URigVMPin)
 		{
 			if (URigVMPin* Pin = Cast<URigVMPin>(InSubject))
 			{
diff --git a/Engine/Source/Runtime/Core/Public/UObject/FortniteMainBranchObjectVersion.h b/Engine/Source/Runtime/Core/Public/UObject/FortniteMainBranchObjectVersion.h
--- a/Engine/Source/Runtime/Core/Public/UObject/FortniteMainBranchObjectVersion.h
+++ b/Engine/Source/Runtime/Core/Public/UObject/FortniteMainBranchObjectVersion.h
@@ -351,6 +351,9 @@ struct FFortniteMainBranchObjectVersion
 		// Binding extensions for anim graph nodes
 		AnimGraphNodeBindingExtensions,
 
+		// Function data stores a map from work to debug operands
+		RigVMSaveDebugMapInGraphFunctionData,
+
 		// -----<new versions can be added above this line>-------------------------------------------------
 		VersionPlusOne,
 		LatestVersion = VersionPlusOne - 1
