RigVM: Fix debug watches inside functions not working

diff --git a/Engine/Plugins/Runtime/RigVM/Source/RigVM/Private/RigVMCore/RigVM.cpp b/Engine/Plugins/Runtime/RigVM/Source/RigVM/Private/RigVMCore/RigVM.cpp
--- a/Engine/Plugins/Runtime/RigVM/Source/RigVM/Private/RigVMCore/RigVM.cpp
+++ b/Engine/Plugins/Runtime/RigVM/Source/RigVM/Private/RigVMCore/RigVM.cpp
@@ -130,6 +130,9 @@ void URigVM::Serialize(FArchive& Ar)
 
 void URigVM::Save(FArchive& Ar)
 {
+	Ar.UsingCustomVersion(FUE5ReleaseStreamObjectVersion::GUID);
+	Ar.UsingCustomVersion(FFortniteMainBranchObjectVersion::GUID);
+		
 	CopyDeferredVMIfRequired();
 
 	// we rely on Ar.IsIgnoringArchetypeRef for determining if we are currently performing
@@ -146,7 +149,8 @@ void URigVM::Save(FArchive& Ar)
 		Ar << ByteCodeStorage;
 		Ar << Parameters;
 
-		if (Ar.CustomVer(FFortniteMainBranchObjectVersion::GUID) < FFortniteMainBranchObjectVersion::RigVMSaveDebugMapInGraphFunctionData)
+		if (Ar.CustomVer(FUE5ReleaseStreamObjectVersion::GUID) < FUE5ReleaseStreamObjectVersion::RigVMSaveDebugMapInGraphFunctionData &&
+		    Ar.CustomVer(FFortniteMainBranchObjectVersion::GUID) < FFortniteMainBranchObjectVersion::RigVMSaveDebugMapInGraphFunctionData)
 		{
 			return;
 		}
@@ -157,6 +161,9 @@ void URigVM::Save(FArchive& Ar)
 
 void URigVM::Load(FArchive& Ar)
 {
+	Ar.UsingCustomVersion(FUE5ReleaseStreamObjectVersion::GUID);
+	Ar.UsingCustomVersion(FFortniteMainBranchObjectVersion::GUID);
+	
 	// we rely on Ar.IsIgnoringArchetypeRef for determining if we are currently performing
 	// CPFUO (Copy Properties for unrelated objects). During a reinstance pass we don't
 	// want to overwrite the bytecode and some other properties - since that's handled already
@@ -196,7 +203,8 @@ void URigVM::Load(FArchive& Ar)
 				return;
 			}
 
-			if (Ar.CustomVer(FFortniteMainBranchObjectVersion::GUID) >= FFortniteMainBranchObjectVersion::RigVMSaveDebugMapInGraphFunctionData)
+			if (Ar.CustomVer(FUE5ReleaseStreamObjectVersion::GUID) >= FUE5ReleaseStreamObjectVersion::RigVMSaveDebugMapInGraphFunctionData ||
+			    Ar.CustomVer(FFortniteMainBranchObjectVersion::GUID) >= FFortniteMainBranchObjectVersion::RigVMSaveDebugMapInGraphFunctionData)
 			{
 				Ar << OperandToDebugRegisters;
 			}
@@ -230,7 +238,8 @@ void URigVM::Load(FArchive& Ar)
 		Ar << ByteCodeStorage;
 		Ar << Parameters;
 
-		if (Ar.CustomVer(FFortniteMainBranchObjectVersion::GUID) >= FFortniteMainBranchObjectVersion::RigVMSaveDebugMapInGraphFunctionData)
+		if (Ar.CustomVer(FUE5ReleaseStreamObjectVersion::GUID) >= FUE5ReleaseStreamObjectVersion::RigVMSaveDebugMapInGraphFunctionData ||
+		    Ar.CustomVer(FFortniteMainBranchObjectVersion::GUID) >= FFortniteMainBranchObjectVersion::RigVMSaveDebugMapInGraphFunctionData)
 		{
 			Ar << OperandToDebugRegisters;
 		}
diff --git a/Engine/Plugins/Runtime/RigVM/Source/RigVM/Public/RigVMCore/RigVMGraphFunctionDefinition.h b/Engine/Plugins/Runtime/RigVM/Source/RigVM/Public/RigVMCore/RigVMGraphFunctionDefinition.h
--- a/Engine/Plugins/Runtime/RigVM/Source/RigVM/Public/RigVMCore/RigVMGraphFunctionDefinition.h
+++ b/Engine/Plugins/Runtime/RigVM/Source/RigVM/Public/RigVMCore/RigVMGraphFunctionDefinition.h
@@ -6,6 +6,7 @@
 #include "RigVMCore/RigVMByteCode.h"
 #include "UObject/UE5MainStreamObjectVersion.h"
 #include "UObject/FortniteMainBranchObjectVersion.h"
+#include "UObject/UE5ReleaseStreamObjectVersion.h"
 #include "RigVMGraphFunctionDefinition.generated.h"
 
 class IRigVMGraphFunctionHost;
@@ -196,6 +197,7 @@ struct RIGVM_API FRigVMFunctionCompilationData
 
 	friend FArchive& operator<<(FArchive& Ar, FRigVMFunctionCompilationData& Data)
 	{
+		Ar.UsingCustomVersion(FUE5ReleaseStreamObjectVersion::GUID);
 		Ar.UsingCustomVersion(FFortniteMainBranchObjectVersion::GUID);
 		
 		Ar << Data.ByteCode;
@@ -212,7 +214,8 @@ struct RIGVM_API FRigVMFunctionCompilationData
 		Ar << Data.Operands;
 		Ar << Data.Hash;
 
-		if (Ar.CustomVer(FFortniteMainBranchObjectVersion::GUID) < FFortniteMainBranchObjectVersion::RigVMSaveDebugMapInGraphFunctionData)
+		if (Ar.CustomVer(FUE5ReleaseStreamObjectVersion::GUID) < FUE5ReleaseStreamObjectVersion::RigVMSaveDebugMapInGraphFunctionData &&
+		    Ar.CustomVer(FFortniteMainBranchObjectVersion::GUID) < FFortniteMainBranchObjectVersion::RigVMSaveDebugMapInGraphFunctionData)
 		{
 			return Ar;
 		}
diff --git a/Engine/Source/Runtime/Core/Public/UObject/UE5ReleaseStreamObjectVersion.h b/Engine/Source/Runtime/Core/Public/UObject/UE5ReleaseStreamObjectVersion.h
--- a/Engine/Source/Runtime/Core/Public/UObject/UE5ReleaseStreamObjectVersion.h
+++ b/Engine/Source/Runtime/Core/Public/UObject/UE5ReleaseStreamObjectVersion.h
@@ -141,6 +141,9 @@ struct FUE5ReleaseStreamObjectVersion
 		
 		// Added resize methods to media capture
 		MediaCaptureNewResizeMethods,
+
+		// Function data stores a map from work to debug operands
+		RigVMSaveDebugMapInGraphFunctionData,
 		
 		// -----<new versions can be added above this line>-------------------------------------------------
 		VersionPlusOne,
