Control Rig: Fix crash when opening certain assets

diff --git a/Engine/Plugins/Runtime/RigVM/Source/RigVM/Private/RigVMCore/RigVM.cpp b/Engine/Plugins/Runtime/RigVM/Source/RigVM/Private/RigVMCore/RigVM.cpp
--- a/Engine/Plugins/Runtime/RigVM/Source/RigVM/Private/RigVMCore/RigVM.cpp
+++ b/Engine/Plugins/Runtime/RigVM/Source/RigVM/Private/RigVMCore/RigVM.cpp
@@ -6,6 +6,7 @@
 #include "UObject/AnimObjectVersion.h"
 #include "RigVMObjectVersion.h"
 #include "UObject/UE5MainStreamObjectVersion.h"
+#include "UObject/UE5ReleaseStreamObjectVersion.h"
 #include "RigVMObjectVersion.h"
 #include "HAL/PlatformTLS.h"
 #include "Async/ParallelFor.h"
