Control Rig: Fix NonUnity compilation

diff --git a/Engine/Plugins/Runtime/RigVM/Source/RigVM/Private/RigVMCore/RigVM.cpp b/Engine/Plugins/Runtime/RigVM/Source/RigVM/Private/RigVMCore/RigVM.cpp
--- a/Engine/Plugins/Runtime/RigVM/Source/RigVM/Private/RigVMCore/RigVM.cpp
+++ b/Engine/Plugins/Runtime/RigVM/Source/RigVM/Private/RigVMCore/RigVM.cpp
@@ -7,6 +7,7 @@
 #include "RigVMObjectVersion.h"
 #include "UObject/UE5MainStreamObjectVersion.h"
 #include "UObject/UE5ReleaseStreamObjectVersion.h"
+#include "UObject/FortniteMainBranchObjectVersion.h"
 #include "RigVMObjectVersion.h"
 #include "HAL/PlatformTLS.h"
 #include "Async/ParallelFor.h"
