RigVM: Fix Variable Remapping Unit test

diff --git a/Engine/Plugins/Runtime/RigVM/Source/RigVMDeveloper/Private/RigVMCompiler/RigVMCompiler.cpp b/Engine/Plugins/Runtime/RigVM/Source/RigVMDeveloper/Private/RigVMCompiler/RigVMCompiler.cpp
--- a/Engine/Plugins/Runtime/RigVM/Source/RigVMDeveloper/Private/RigVMCompiler/RigVMCompiler.cpp
+++ b/Engine/Plugins/Runtime/RigVM/Source/RigVMDeveloper/Private/RigVMCompiler/RigVMCompiler.cpp
@@ -2252,6 +2252,15 @@ int32 URigVMCompiler::TraverseInlineFunction(const FRigVMInlineFunctionExprAST*
 		{
 			if (FRigVMOperand* NewOperand = FunctionOperandToNewOperand.Find(PinPathToOperand.Value))
 			{
+				// avoid copying the variable pin paths over to the outer VM since
+				// the variables from within the function asset no longer exist and
+				// will have been remapped to the outer VM's external variables.
+				static const FString VariablePrefix = TEXT("Variable::");
+				if(PinPathToOperand.Key.StartsWith(VariablePrefix, ESearchCase::CaseSensitive))
+				{
+					continue;
+				}
+				
 				WorkData.PinPathToOperand->Add(PinPathToOperand.Key, *NewOperand);
 				TArray<FString>& Paths = OperandToPinPath.Emplace(*NewOperand);
 				Paths.Add(PinPathToOperand.Key);
