RigVM: Fix tooltips on resolved template nodes

diff --git a/Engine/Plugins/Runtime/RigVM/Source/RigVMEditor/Private/Widgets/SRigVMGraphNode.cpp b/Engine/Plugins/Runtime/RigVM/Source/RigVMEditor/Private/Widgets/SRigVMGraphNode.cpp
--- a/Engine/Plugins/Runtime/RigVM/Source/RigVMEditor/Private/Widgets/SRigVMGraphNode.cpp
+++ b/Engine/Plugins/Runtime/RigVM/Source/RigVMEditor/Private/Widgets/SRigVMGraphNode.cpp
@@ -390,7 +390,21 @@ void SRigVMGraphNode::AddPin(const TSharedRef<SGraphPin>& PinToAdd)
 				{
 					PinToAdd->SetCustomPinIcon(CachedImg_CR_Pin_Connected, CachedImg_CR_Pin_Disconnected);
 				}
-				PinToAdd->SetToolTipText(ModelPin->GetToolTipText());
+
+				const FText ToolTipText = ModelPin->GetToolTipText();
+				TWeakObjectPtr<const URigVMPin> WeakPin = ModelPin;
+				PinToAdd->SetToolTip(
+					FSlateApplicationBase::Get().MakeToolTip(
+						TAttribute<FText>::CreateLambda([ToolTipText, WeakPin]() -> FText
+						{
+							if(WeakPin.IsValid())
+							{
+								return WeakPin->GetToolTipText();
+							}
+							return ToolTipText;
+						})
+					)
+				);
 
 				// If the pin belongs to a template node that does not own an argument for that pin, make it transparent
 				if (!ModelPin->IsExecuteContext())
