RigVMEditor: Improve Nodespawner performance

diff --git a/Engine/Plugins/Runtime/RigVM/Source/RigVMDeveloper/Private/EdGraph/RigVMEdGraph.cpp b/Engine/Plugins/Runtime/RigVM/Source/RigVMDeveloper/Private/EdGraph/RigVMEdGraph.cpp
--- a/Engine/Plugins/Runtime/RigVM/Source/RigVMDeveloper/Private/EdGraph/RigVMEdGraph.cpp
+++ b/Engine/Plugins/Runtime/RigVM/Source/RigVMDeveloper/Private/EdGraph/RigVMEdGraph.cpp
@@ -153,7 +153,7 @@ bool URigVMEdGraph::HandleModifiedEvent_Internal(ERigVMGraphNotifType InNotifTyp
 
 	// only make sure to receive notifs for this graph - unless
 	// we are on a template graph (used by node spawners)
-	if (GetModel() != InGraph && TemplateController == nullptr)
+	if (GetModel() != InGraph)
 	{
 		return false;
 	}
@@ -163,28 +163,6 @@ bool URigVMEdGraph::HandleModifiedEvent_Internal(ERigVMGraphNotifType InNotifTyp
 		EdGraphSchema->HandleModifiedEvent(InNotifType, InGraph, InSubject);
 	}
 
-	if(TemplateController)
-	{
-		switch(InNotifType)
-		{
-			case ERigVMGraphNotifType::NodeRemoved:
-			case ERigVMGraphNotifType::PinTypeChanged:
-			case ERigVMGraphNotifType::PinDefaultValueChanged:
-			case ERigVMGraphNotifType::PinRemoved:
-			case ERigVMGraphNotifType::PinRenamed:
-			case ERigVMGraphNotifType::VariableRemoved:
-			case ERigVMGraphNotifType::VariableRenamed:
-			case ERigVMGraphNotifType::GraphChanged:
-			{
-				return false;
-			}
-			default:
-			{
-				break;
-			}
-		}
-	}
-
 	// increment the node topology version for any interaction
 	// with a node.
 	{
@@ -1024,17 +1002,6 @@ void URigVMEdGraph::RemoveNode(UEdGraphNode* InNode)
 	Super::RemoveNode(InNode);
 }
 
-URigVMController* URigVMEdGraph::GetTemplateController()
-{
-	if (TemplateController == nullptr)
-	{
-		TemplateController = GetBlueprint()->GetTemplateController();
-		TemplateController->OnModified().RemoveAll(this);
-		TemplateController->OnModified().AddUObject(this, &URigVMEdGraph::HandleModifiedEvent);
-	}
-	return TemplateController;
-}
-
 void URigVMEdGraph::HandleVMCompiledEvent(UObject* InCompiledObject, URigVM* InVM)
 {
 	CachedInstructionIndices.Reset();
diff --git a/Engine/Plugins/Runtime/RigVM/Source/RigVMDeveloper/Private/EdGraph/RigVMEdGraphNode.cpp b/Engine/Plugins/Runtime/RigVM/Source/RigVMDeveloper/Private/EdGraph/RigVMEdGraphNode.cpp
--- a/Engine/Plugins/Runtime/RigVM/Source/RigVMDeveloper/Private/EdGraph/RigVMEdGraphNode.cpp
+++ b/Engine/Plugins/Runtime/RigVM/Source/RigVMDeveloper/Private/EdGraph/RigVMEdGraphNode.cpp
@@ -1528,16 +1528,6 @@ URigVMNode* URigVMEdGraphNode::GetModelNode() const
 
 	if (URigVMEdGraph* Graph = Cast<URigVMEdGraph>(GetOuter()))
 	{
-#if WITH_EDITOR
-
-		if (Graph->TemplateController != nullptr)
-		{
-			MutableThis->CachedModelNode = TWeakObjectPtr<URigVMNode>(Graph->TemplateController->GetGraph()->FindNode(ModelNodePath));
-			return MutableThis->CachedModelNode.Get();
-		}
-
-#endif
-
 		if (URigVMGraph* Model = GetModel())
 		{
 			MutableThis->CachedModelNode = TWeakObjectPtr<URigVMNode>(Model->FindNode(ModelNodePath));
diff --git a/Engine/Plugins/Runtime/RigVM/Source/RigVMDeveloper/Public/EdGraph/RigVMEdGraph.h b/Engine/Plugins/Runtime/RigVM/Source/RigVMDeveloper/Public/EdGraph/RigVMEdGraph.h
--- a/Engine/Plugins/Runtime/RigVM/Source/RigVMDeveloper/Public/EdGraph/RigVMEdGraph.h
+++ b/Engine/Plugins/Runtime/RigVM/Source/RigVMDeveloper/Public/EdGraph/RigVMEdGraph.h
@@ -82,17 +82,6 @@ private:
 
 	void RemoveNode(UEdGraphNode* InNode);
 
-#endif
-#if WITH_EDITORONLY_DATA
-
-	UPROPERTY(transient)
-	TObjectPtr<URigVMController> TemplateController;
-
-#endif
-#if WITH_EDITOR
-
-	URigVMController* GetTemplateController();
-
 protected:
 	void HandleVMCompiledEvent(UObject* InCompiledObject, URigVM* InVM);
 
diff --git a/Engine/Plugins/Runtime/RigVM/Source/RigVMDeveloper/Public/RigVMModel/Nodes/RigVMVariableNode.h b/Engine/Plugins/Runtime/RigVM/Source/RigVMDeveloper/Public/RigVMModel/Nodes/RigVMVariableNode.h
--- a/Engine/Plugins/Runtime/RigVM/Source/RigVMDeveloper/Public/RigVMModel/Nodes/RigVMVariableNode.h
+++ b/Engine/Plugins/Runtime/RigVM/Source/RigVMDeveloper/Public/RigVMModel/Nodes/RigVMVariableNode.h
@@ -83,5 +83,6 @@ private:
 	friend class FRigVMVarExprAST;
 	friend class FRigVMParserAST;
 	friend class URigVMSchema;
+	friend class URigVMEdGraphVariableNodeSpawner;
 };
 
diff --git a/Engine/Plugins/Runtime/RigVM/Source/RigVMEditor/Private/EdGraph/NodeSpawners/RigVMEdGraphEnumNodeSpawner.cpp b/Engine/Plugins/Runtime/RigVM/Source/RigVMEditor/Private/EdGraph/NodeSpawners/RigVMEdGraphEnumNodeSpawner.cpp
--- a/Engine/Plugins/Runtime/RigVM/Source/RigVMEditor/Private/EdGraph/NodeSpawners/RigVMEdGraphEnumNodeSpawner.cpp
+++ b/Engine/Plugins/Runtime/RigVM/Source/RigVMEditor/Private/EdGraph/NodeSpawners/RigVMEdGraphEnumNodeSpawner.cpp
@@ -56,24 +56,11 @@ UEdGraphNode* URigVMEdGraphEnumNodeSpawner::Invoke(UEdGraph* ParentGraph, FBindi
 	URigVMEdGraphNode* NewNode = nullptr;
 
 	bool const bIsTemplateNode = FBlueprintNodeTemplateCache::IsTemplateOuter(ParentGraph);
-	bool const bIsUserFacingNode = !bIsTemplateNode;
-
 	if (bIsTemplateNode)
 	{
-		NewNode = NewObject<URigVMEdGraphNode>(ParentGraph, TEXT("EnumNode"));
-		ParentGraph->AddNode(NewNode, false);
-
-		NewNode->CreateNewGuid();
-		NewNode->PostPlacedNewNode();
-
-		UEdGraphPin* OutputValuePin = UEdGraphPin::CreatePin(NewNode);
-		NewNode->Pins.Add(OutputValuePin);
-
-		OutputValuePin->PinType.PinCategory = TEXT("int32");
-		OutputValuePin->Direction = EGPD_Output;
-		NewNode->SetFlags(RF_Transactional);
-
-		return NewNode;
+		TArray<FPinInfo> Pins;
+		Pins.Emplace(*URigVMEnumNode::EnumValueName, ERigVMPinDirection::Output, RigVMTypeUtils::Int32TypeName, nullptr);
+		return SpawnTemplateNode(ParentGraph, Pins);
 	}
 
 	// First create a backing member for our node
@@ -83,46 +70,33 @@ UEdGraphNode* URigVMEdGraphEnumNodeSpawner::Invoke(UEdGraph* ParentGraph, FBindi
 	check(RigBlueprint);
 
 #if WITH_EDITOR
-	if (GEditor && !bIsTemplateNode)
+	if (GEditor)
 	{
 		GEditor->CancelTransaction(0);
 	}
 #endif
 
-	URigVMController* Controller = bIsTemplateNode ? RigGraph->GetTemplateController() : RigBlueprint->GetController(ParentGraph);
+	URigVMController* Controller = RigBlueprint->GetController(ParentGraph);
 
 	FName Name = *URigVMEnumNode::EnumName;
 
-	if (!bIsTemplateNode)
-	{
-		Controller->OpenUndoBracket(FString::Printf(TEXT("Add '%s' Node"), *Name.ToString()));
-	}
+	Controller->OpenUndoBracket(FString::Printf(TEXT("Add '%s' Node"), *Name.ToString()));
 
-	if (URigVMEnumNode* ModelNode = Controller->AddEnumNode(*Enum->GetPathName(), Location, Name.ToString(), bIsUserFacingNode, !bIsTemplateNode))
+	if (URigVMEnumNode* ModelNode = Controller->AddEnumNode(*Enum->GetPathName(), Location, Name.ToString(), true, true))
 	{
 		NewNode = Cast<URigVMEdGraphNode>(RigGraph->FindNodeForModelNodeName(ModelNode->GetFName()));
 
-		if (NewNode && bIsUserFacingNode)
+		if (NewNode)
 		{
 			Controller->ClearNodeSelection(true);
 			Controller->SelectNode(ModelNode, true, true);
 		}
 
-		if (bIsUserFacingNode)
-		{
-			Controller->CloseUndoBracket();
-		}
-		else
-		{
-			Controller->RemoveNode(ModelNode, false);
-		}
+		Controller->CloseUndoBracket();
 	}
 	else
 	{
-		if (bIsUserFacingNode)
-		{
-			Controller->CancelUndoBracket();
-		}
+		Controller->CancelUndoBracket();
 	}
 
 
diff --git a/Engine/Plugins/Runtime/RigVM/Source/RigVMEditor/Private/EdGraph/NodeSpawners/RigVMEdGraphFunctionRefNodeSpawner.cpp b/Engine/Plugins/Runtime/RigVM/Source/RigVMEditor/Private/EdGraph/NodeSpawners/RigVMEdGraphFunctionRefNodeSpawner.cpp
--- a/Engine/Plugins/Runtime/RigVM/Source/RigVMEditor/Private/EdGraph/NodeSpawners/RigVMEdGraphFunctionRefNodeSpawner.cpp
+++ b/Engine/Plugins/Runtime/RigVM/Source/RigVMEditor/Private/EdGraph/NodeSpawners/RigVMEdGraphFunctionRefNodeSpawner.cpp
@@ -266,8 +266,9 @@ UEdGraphNode* URigVMEdGraphFunctionRefNodeSpawner::Invoke(UEdGraph* ParentGraph,
 			}
 		}
 	}
-		
-	if(ReferencedPublicFunctionHeader.IsValid())
+
+	bool const bIsTemplateNode = FBlueprintNodeTemplateCache::IsTemplateOuter(ParentGraph);
+	if(ReferencedPublicFunctionHeader.IsValid() && !bIsTemplateNode)
 	{
 #if WITH_EDITOR
 		if (GEditor)
@@ -282,37 +283,14 @@ UEdGraphNode* URigVMEdGraphFunctionRefNodeSpawner::Invoke(UEdGraph* ParentGraph,
 	else
 	{
 		// we are only going to get here if we are spawning a template node
-		NewNode = NewObject<URigVMEdGraphNode>(ParentGraph);
-		ParentGraph->AddNode(NewNode, false);
-
-		NewNode->CreateNewGuid();
-		NewNode->PostPlacedNewNode();
-
+		TArray<FPinInfo> Pins;
 		for(const FRigVMGraphFunctionArgument& Arg : ReferencedPublicFunctionHeader.Arguments)
 		{
-			if(Arg.Direction ==  ERigVMPinDirection::Input ||
-				Arg.Direction ==  ERigVMPinDirection::IO)
-			{
-				UEdGraphPin* InputPin = UEdGraphPin::CreatePin(NewNode);
-				NewNode->Pins.Add(InputPin);
-
-				InputPin->Direction = EGPD_Input;
-				InputPin->PinType = RigVMTypeUtils::PinTypeFromCPPType(Arg.CPPType, Arg.CPPTypeObject.Get());
-			}
-
-			if(Arg.Direction ==  ERigVMPinDirection::Output ||
-				Arg.Direction ==  ERigVMPinDirection::IO)
-			{
-				UEdGraphPin* OutputPin = UEdGraphPin::CreatePin(NewNode);
-				NewNode->Pins.Add(OutputPin);
-
-				OutputPin->Direction = EGPD_Output;
-				OutputPin->PinType = RigVMTypeUtils::PinTypeFromCPPType(Arg.CPPType, Arg.CPPTypeObject.Get());
-			}
+			Pins.Emplace(Arg.Name, Arg.Direction, Arg.CPPType, Arg.CPPTypeObject.Get());
 		}
 
-		NewNode->SetFlags(RF_Transactional);
-}
+		NewNode = SpawnTemplateNode(ParentGraph, Pins);
+	}
 
 	return NewNode;
 }
@@ -325,24 +303,17 @@ URigVMEdGraphNode* URigVMEdGraphFunctionRefNodeSpawner::SpawnNode(UEdGraph* Pare
 
 	if (RigBlueprint != nullptr && RigGraph != nullptr)
 	{
-		bool const bIsTemplateNode = FBlueprintNodeTemplateCache::IsTemplateOuter(ParentGraph);
-		bool const bIsUserFacingNode = !bIsTemplateNode;
+		FName Name = FRigVMBlueprintUtils::ValidateName(RigBlueprint, InFunction.Name.ToString());
+		URigVMController* Controller = RigBlueprint->GetController(ParentGraph);
 
-		FName Name = bIsTemplateNode ? InFunction.Name : FRigVMBlueprintUtils::ValidateName(RigBlueprint, InFunction.Name.ToString());
-		URigVMController* Controller = bIsTemplateNode ? RigGraph->GetTemplateController() : RigBlueprint->GetController(ParentGraph);
+		Controller->OpenUndoBracket(FString::Printf(TEXT("Add '%s' Node"), *Name.ToString()));
 
-		if (!bIsTemplateNode)
-		{
-			Controller->OpenUndoBracket(FString::Printf(TEXT("Add '%s' Node"), *Name.ToString()));
-		}
-
-		TGuardValue<bool> AllowPrivateFunctionsOnTemplateController(Controller->bAllowPrivateFunctions, bIsTemplateNode);
-		if (URigVMFunctionReferenceNode* ModelNode = Controller->AddFunctionReferenceNodeFromDescription(InFunction, Location, Name.ToString(), bIsUserFacingNode, !bIsTemplateNode))
+		if (URigVMFunctionReferenceNode* ModelNode = Controller->AddFunctionReferenceNodeFromDescription(InFunction, Location, Name.ToString(), true, true))
 		{
 			NewNode = Cast<URigVMEdGraphNode>(RigGraph->FindNodeForModelNodeName(ModelNode->GetFName()));
 			check(NewNode);
 
-			if (NewNode && bIsUserFacingNode)
+			if (NewNode)
 			{
 				Controller->ClearNodeSelection(true);
 				Controller->SelectNode(ModelNode, true, true);
@@ -350,60 +321,38 @@ URigVMEdGraphNode* URigVMEdGraphFunctionRefNodeSpawner::SpawnNode(UEdGraph* Pare
 				URigVMEdGraphUnitNodeSpawner::HookupMutableNode(ModelNode, RigBlueprint);
 			}
 
-			if (!bIsTemplateNode)
+			for(URigVMNode* OtherModelNode : ModelNode->GetGraph()->GetNodes())
 			{
-				for(URigVMNode* OtherModelNode : ModelNode->GetGraph()->GetNodes())
+				if(OtherModelNode == ModelNode)
 				{
-					if(OtherModelNode == ModelNode)
-					{
-						continue;
-					}
-					
-					URigVMFunctionReferenceNode* ExistingFunctionReferenceNode = Cast<URigVMFunctionReferenceNode>(OtherModelNode);
-					if(ExistingFunctionReferenceNode == nullptr)
-					{
-						continue;
-					}
+					continue;
+				}
+				
+				URigVMFunctionReferenceNode* ExistingFunctionReferenceNode = Cast<URigVMFunctionReferenceNode>(OtherModelNode);
+				if(ExistingFunctionReferenceNode == nullptr)
+				{
+					continue;
+				}
 
-					if(!(ExistingFunctionReferenceNode->GetReferencedFunctionHeader().LibraryPointer == InFunction.LibraryPointer))
-					{
-						continue;
-					}
+				if(!(ExistingFunctionReferenceNode->GetReferencedFunctionHeader().LibraryPointer == InFunction.LibraryPointer))
+				{
+					continue;
+				}
 
-					for(TPair<FName, FName> MappedVariablePair : ExistingFunctionReferenceNode->GetVariableMap())
+				for(TPair<FName, FName> MappedVariablePair : ExistingFunctionReferenceNode->GetVariableMap())
+				{
+					if(!MappedVariablePair.Value.IsNone())
 					{
-						if(!MappedVariablePair.Value.IsNone())
-						{
-							Controller->SetRemappedVariable(ModelNode, MappedVariablePair.Key, MappedVariablePair.Value, bIsUserFacingNode);
-						}
+						Controller->SetRemappedVariable(ModelNode, MappedVariablePair.Key, MappedVariablePair.Value, true);
 					}
 				}
 			}
-			else
-			{
-				// If the package is a template, do not remove the EdGraphNode
-				// We might be spawning a node to populate the PROTO_ context menu for function declarations.
-				FRigVMControllerNotifGuard NotifGuard(Controller, true);
-				Controller->RemoveNode(ModelNode, false);
-			}
 
-			if (bIsUserFacingNode)
-			{
-				Controller->CloseUndoBracket();
-			}
-			else
-			{
-				Controller->EnableReporting(false);
-				Controller->RemoveNode(ModelNode, false);
-				Controller->EnableReporting(true);
-			}
+			Controller->CloseUndoBracket();
 		}
 		else
 		{
-			if (bIsUserFacingNode)
-			{
-				Controller->CancelUndoBracket();
-			}
+			Controller->CancelUndoBracket();
 		}
 	}
 	return NewNode;
diff --git a/Engine/Plugins/Runtime/RigVM/Source/RigVMEditor/Private/EdGraph/NodeSpawners/RigVMEdGraphInvokeEntryNodeSpawner.cpp b/Engine/Plugins/Runtime/RigVM/Source/RigVMEditor/Private/EdGraph/NodeSpawners/RigVMEdGraphInvokeEntryNodeSpawner.cpp
--- a/Engine/Plugins/Runtime/RigVM/Source/RigVMEditor/Private/EdGraph/NodeSpawners/RigVMEdGraphInvokeEntryNodeSpawner.cpp
+++ b/Engine/Plugins/Runtime/RigVM/Source/RigVMEditor/Private/EdGraph/NodeSpawners/RigVMEdGraphInvokeEntryNodeSpawner.cpp
@@ -65,7 +65,6 @@ UEdGraphNode* URigVMEdGraphInvokeEntryNodeSpawner::Invoke(UEdGraph* ParentGraph,
 	URigVMEdGraphNode* NewNode = nullptr;
 
 	bool const bIsTemplateNode = FBlueprintNodeTemplateCache::IsTemplateOuter(ParentGraph);
-	bool const bIsUserFacingNode = !bIsTemplateNode;
 
 	// First create a backing member for our node
 	URigVMEdGraph* RigGraph = Cast<URigVMEdGraph>(ParentGraph);
@@ -80,13 +79,6 @@ UEdGraphNode* URigVMEdGraphInvokeEntryNodeSpawner::Invoke(UEdGraph* ParentGraph,
 	}
 #endif
 
-	URigVMController* Controller = bIsTemplateNode ? RigGraph->GetTemplateController() : RigBlueprint->GetController(ParentGraph);
-
-	if (bIsUserFacingNode)
-	{
-		Controller->OpenUndoBracket(TEXT("Add Run Event"));
-	}
-
 	FString NodeName;
 	if(bIsTemplateNode)
 	{
@@ -94,9 +86,20 @@ UEdGraphNode* URigVMEdGraphInvokeEntryNodeSpawner::Invoke(UEdGraph* ParentGraph,
 		// we need to create a unique here.
 		static constexpr TCHAR InvokeEntryNodeNameFormat[] = TEXT("InvokeEntryNode_%s");
 		NodeName = FString::Printf(InvokeEntryNodeNameFormat, *EntryName.ToString());
+
+		TArray<FPinInfo> Pins;
+
+		static UScriptStruct* ExecuteScriptStruct = FRigVMExecuteContext::StaticStruct();
+		static const FName ExecuteStructName = *ExecuteScriptStruct->GetStructCPPName();
+		Pins.Emplace(FRigVMStruct::ExecuteName, ERigVMPinDirection::IO, ExecuteStructName, ExecuteScriptStruct);
+
+		return SpawnTemplateNode(ParentGraph, Pins, *NodeName);
 	}
 
-	if (URigVMNode* ModelNode = Controller->AddInvokeEntryNode(EntryName, Location, NodeName, !bIsTemplateNode, !bIsTemplateNode))
+	URigVMController* Controller = RigBlueprint->GetController(ParentGraph);
+	Controller->OpenUndoBracket(TEXT("Add Run Event"));
+
+	if (URigVMNode* ModelNode = Controller->AddInvokeEntryNode(EntryName, Location, NodeName, true, true))
 	{
 		for (UEdGraphNode* Node : ParentGraph->Nodes)
 		{
@@ -110,28 +113,16 @@ UEdGraphNode* URigVMEdGraphInvokeEntryNodeSpawner::Invoke(UEdGraph* ParentGraph,
 			}
 		}
 
-		if (bIsUserFacingNode)
+		if (NewNode)
 		{
-			if (NewNode)
-			{
-				Controller->ClearNodeSelection(true);
-				Controller->SelectNode(ModelNode, true, true);
-			}
-			Controller->CloseUndoBracket();
-		}
-		else
-		{
-			// similar to UBlueprintNodeSpawner::Invoke -> UBlueprintNodeSpawner::SpawnEdGraphNode
-			// we simply want the node, but not actually adding it to a graph
-			Controller->RemoveNode(ModelNode, false);
+			Controller->ClearNodeSelection(true);
+			Controller->SelectNode(ModelNode, true, true);
 		}
+		Controller->CloseUndoBracket();
 	}
 	else
 	{
-		if (!bIsTemplateNode)
-		{
-			Controller->CancelUndoBracket();
-		}
+		Controller->CancelUndoBracket();
 	}
 
 
diff --git a/Engine/Plugins/Runtime/RigVM/Source/RigVMEditor/Private/EdGraph/NodeSpawners/RigVMEdGraphNodeSpawner.cpp b/Engine/Plugins/Runtime/RigVM/Source/RigVMEditor/Private/EdGraph/NodeSpawners/RigVMEdGraphNodeSpawner.cpp
--- a/Engine/Plugins/Runtime/RigVM/Source/RigVMEditor/Private/EdGraph/NodeSpawners/RigVMEdGraphNodeSpawner.cpp
+++ b/Engine/Plugins/Runtime/RigVM/Source/RigVMEditor/Private/EdGraph/NodeSpawners/RigVMEdGraphNodeSpawner.cpp
@@ -1,6 +1,7 @@
 // Copyright Epic Games, Inc. All Rights Reserved.
 
 #include "EdGraph/NodeSpawners/RigVMEdGraphNodeSpawner.h"
+#include "RigVMStringUtils.h"
 
 #include UE_INLINE_GENERATED_CPP_BY_NAME(RigVMEdGraphNodeSpawner)
 
@@ -26,5 +27,45 @@ void URigVMEdGraphNodeSpawner::SetRelatedBlueprintClass(TSubclassOf<URigVMBluepr
 	RelatedBlueprintClass = InClass;
 }
 
+URigVMEdGraphNode* URigVMEdGraphNodeSpawner::SpawnTemplateNode(UEdGraph* InParentGraph, const TArray<FPinInfo>& InPins, const FName& InNodeName)
+{
+	URigVMEdGraphNode* NewNode = NewObject<URigVMEdGraphNode>(InParentGraph, InNodeName);
+	InParentGraph->AddNode(NewNode, false);
+
+	NewNode->CreateNewGuid();
+	NewNode->PostPlacedNewNode();
+
+	for(const FPinInfo& PinInfo : InPins)
+	{
+		const FName PinName = *RigVMStringUtils::JoinPinPath(NewNode->GetName(), PinInfo.Name.ToString());
+		
+		if(PinInfo.Direction ==  ERigVMPinDirection::Input ||
+			PinInfo.Direction ==  ERigVMPinDirection::IO)
+		{
+			UEdGraphPin* InputPin = UEdGraphPin::CreatePin(NewNode);
+			InputPin->PinName = PinName;
+			NewNode->Pins.Add(InputPin);
+
+			InputPin->Direction = EGPD_Input;
+			InputPin->PinType = RigVMTypeUtils::PinTypeFromCPPType(PinInfo.CPPType, PinInfo.CPPTypeObject);
+		}
+
+		if(PinInfo.Direction ==  ERigVMPinDirection::Output ||
+			PinInfo.Direction ==  ERigVMPinDirection::IO)
+		{
+			UEdGraphPin* OutputPin = UEdGraphPin::CreatePin(NewNode);
+			OutputPin->PinName = PinName;
+			NewNode->Pins.Add(OutputPin);
+
+			OutputPin->Direction = EGPD_Output;
+			OutputPin->PinType = RigVMTypeUtils::PinTypeFromCPPType(PinInfo.CPPType, PinInfo.CPPTypeObject);
+		}
+	}
+
+	NewNode->SetFlags(RF_Transactional);
+
+	return NewNode;
+}
+
 #undef LOCTEXT_NAMESPACE
 
diff --git a/Engine/Plugins/Runtime/RigVM/Source/RigVMEditor/Private/EdGraph/NodeSpawners/RigVMEdGraphTemplateNodeSpawner.cpp b/Engine/Plugins/Runtime/RigVM/Source/RigVMEditor/Private/EdGraph/NodeSpawners/RigVMEdGraphTemplateNodeSpawner.cpp
--- a/Engine/Plugins/Runtime/RigVM/Source/RigVMEditor/Private/EdGraph/NodeSpawners/RigVMEdGraphTemplateNodeSpawner.cpp
+++ b/Engine/Plugins/Runtime/RigVM/Source/RigVMEditor/Private/EdGraph/NodeSpawners/RigVMEdGraphTemplateNodeSpawner.cpp
@@ -29,6 +29,7 @@ URigVMEdGraphTemplateNodeSpawner* URigVMEdGraphTemplateNodeSpawner::CreateFromNo
 {
 	URigVMEdGraphTemplateNodeSpawner* NodeSpawner = NewObject<URigVMEdGraphTemplateNodeSpawner>(GetTransientPackage());
 	NodeSpawner->TemplateNotation = InNotation;
+	NodeSpawner->Template =	FRigVMRegistry::Get().FindTemplate(InNotation);
 	NodeSpawner->NodeClass = URigVMEdGraphNode::StaticClass();
 
 	FBlueprintActionUiSpec& MenuSignature = NodeSpawner->DefaultMenuSignature;
@@ -75,13 +76,23 @@ FBlueprintNodeSignature URigVMEdGraphTemplateNodeSpawner::GetSpawnerSignature()
 
 bool URigVMEdGraphTemplateNodeSpawner::IsTemplateNodeFilteredOut(FBlueprintActionFilter const& Filter) const
 {
+	UE_LOG(LogRigVM, Display, TEXT("Filtering notation '%s'"), *Template->GetNotation().ToString());
+
 	if(URigVMEdGraphNodeSpawner::IsTemplateNodeFilteredOut(Filter))
 	{
 		return true;
 	}
 
-	bool bFilteredOut = true;
-	for (UEdGraphPin* Pin : Filter.Context.Pins)
+	if(Template->GetNotation().ToString().Contains(TEXT("Add::Execute(")))
+	{
+		UE_LOG(LogRigVM, Display, TEXT("Hitting notation '%s'"), *Template->GetNotation().ToString());
+		UE_LOG(LogRigVM, Display, TEXT("Hitting notation '%s'"), *Template->GetNotation().ToString());
+		UE_LOG(LogRigVM, Display, TEXT("Hitting notation '%s'"), *Template->GetNotation().ToString());
+		UE_LOG(LogRigVM, Display, TEXT("Hitting notation '%s'"), *Template->GetNotation().ToString());
+		UE_LOG(LogRigVM, Display, TEXT("Hitting notation '%s'"), *Template->GetNotation().ToString());
+	}
+
+	for (const UEdGraphPin* Pin : Filter.Context.Pins)
 	{
 		for (UEdGraph* Graph : Filter.Context.Graphs)
 		{
@@ -90,48 +101,41 @@ bool URigVMEdGraphTemplateNodeSpawner::IsTemplateNodeFilteredOut(FBlueprintActio
 				continue;
 			}
 
-			if (URigVMEdGraph* EdGraph = Cast<URigVMEdGraph>(Graph))
+			if (const URigVMEdGraph* EdGraph = Cast<URigVMEdGraph>(Graph))
 			{
-				if (URigVMGraph* RigGraph = EdGraph->GetModel())
+				if (const URigVMGraph* RigGraph = EdGraph->GetModel())
 				{
-					if (URigVMPin* ModelPin = RigGraph->FindPin(Pin->GetName()))
+					if (const URigVMPin* ModelPin = RigGraph->FindPin(Pin->GetName()))
 					{
-						if (URigVMTemplateNode* TemplateNode = Cast<URigVMTemplateNode>(ModelPin->GetNode()))
+						// Only filter when the source pin is not a wildcard
+						if (ModelPin->GetCPPType() != RigVMTypeUtils::GetWildCardCPPType())
 						{
-							// Only filter when the sourece pin is not a wildcard
-							if (ModelPin->GetCPPType() != RigVMTypeUtils::GetWildCardCPPType())
+							if (Template)
 							{
-								if (const FRigVMTemplate* Template = FRigVMRegistry::Get().FindTemplate(TemplateNotation))
+								if(ModelPin->IsExecuteContext())
 								{
-									if(ModelPin->IsExecuteContext())
+									FRigVMDispatchContext DispatchContext;
+									if(const URigVMDispatchNode* DispatchNode = Cast<URigVMDispatchNode>(ModelPin->GetNode()))
 									{
-										FRigVMDispatchContext DispatchContext;
-										if(const URigVMDispatchNode* DispatchNode = Cast<URigVMDispatchNode>(ModelPin->GetNode()))
-										{
-											DispatchContext = DispatchNode->GetDispatchContext();
-										}
-
-										if(Template->NumExecuteArguments(DispatchContext) > 0)
-										{
-											return false;
-										}
+										DispatchContext = DispatchNode->GetDispatchContext();
 									}
-									
-									for (int32 i=0; i<Template->NumArguments(); ++i)
+
+									if(Template->NumExecuteArguments(DispatchContext) > 0)
 									{
-										const FRigVMTemplateArgument* Argument = Template->GetArgument(i);
-										if (Template->ArgumentSupportsTypeIndex(Argument->GetName(), ModelPin->GetTypeIndex()))
-										{
-											return false;
-										}
+										return false;
 									}
-									return true;
-								}									
-							}
-							else
-							{
-								// TODO: given the filtered types of the source pin, provide only the nodes that will not break any connections
-							}
+								}
+								
+								for (int32 i=0; i<Template->NumArguments(); ++i)
+								{
+									const FRigVMTemplateArgument* Argument = Template->GetArgument(i);
+									if (Template->ArgumentSupportsTypeIndex(Argument->GetName(), ModelPin->GetTypeIndex()))
+									{
+										return false;
+									}
+								}
+								return true;
+							}									
 						}
 					}					
 				}
@@ -155,7 +159,7 @@ UEdGraphNode* URigVMEdGraphTemplateNodeSpawner::Invoke(UEdGraph* ParentGraph, FB
 {
 	URigVMEdGraphNode* NewNode = nullptr;
 
-	if(!TemplateNotation.IsNone())
+	if(Template)
 	{
 #if WITH_EDITOR
 		if (GEditor)
@@ -164,14 +168,61 @@ UEdGraphNode* URigVMEdGraphTemplateNodeSpawner::Invoke(UEdGraph* ParentGraph, FB
 		}
 #endif
 
+		bool const bIsTemplateNode = FBlueprintNodeTemplateCache::IsTemplateOuter(ParentGraph);
+		if(bIsTemplateNode)
+		{
+			const int32 NotationHash = (int32)GetTypeHash(Template->GetNotation());
+			const FString TemplateName = TEXT("RigVMTemplate_") + FString::FromInt(NotationHash);
+			FName Name = *TemplateName;
+
+			const FRigVMRegistry& Registry = FRigVMRegistry::Get();
+
+			TArray<FPinInfo> Pins;
+
+			static const FRigVMDispatchContext Context;
+			for (int32 Index = 0; Index < Template->NumExecuteArguments(Context); Index++)
+			{
+				const FRigVMExecuteArgument* Argument = Template->GetExecuteArgument(Index, Context);
+				check(Argument);
+				static UScriptStruct* ExecuteScriptStruct = FRigVMExecuteContext::StaticStruct();
+				static const FName ExecuteStructName = *ExecuteScriptStruct->GetStructCPPName();
+				Pins.Emplace(Argument->Name, Argument->Direction, ExecuteStructName, ExecuteScriptStruct);
+			}
+
+			for (int32 Index = 0; Index < Template->NumArguments(); Index++)
+			{
+				const FRigVMTemplateArgument* Argument = Template->GetArgument(Index);
+				check(Argument);
+
+				FName CPPType = RigVMTypeUtils::GetWildCardCPPTypeName();
+				UObject* CPPTypeObject = RigVMTypeUtils::GetWildCardCPPTypeObject();
+
+				if(Argument->IsSingleton())
+				{
+					const TRigVMTypeIndex TypeIndex = Argument->GetTypeIndices()[0];
+					const FRigVMTemplateArgumentType& Type = Registry.GetType(TypeIndex);
+					CPPType = Type.CPPType;
+					CPPTypeObject = Type.CPPTypeObject;
+				}
+				
+				Pins.Emplace(Argument->GetName(), Argument->GetDirection(), CPPType, CPPTypeObject);
+			}
+			NewNode = SpawnTemplateNode(ParentGraph, Pins, Name);
+			if(NewNode)
+			{
+				NewNode->ModelNodePath = Template->GetNotation().ToString();
+			}
+			return NewNode;
+		}			
+
 		UBlueprint* Blueprint = FBlueprintEditorUtils::FindBlueprintForGraph(ParentGraph);
-		NewNode = SpawnNode(ParentGraph, Blueprint, TemplateNotation, Location);
+		NewNode = SpawnNode(ParentGraph, Blueprint, Template, Location);
 	}
 
 	return NewNode;
 }
 
-URigVMEdGraphNode* URigVMEdGraphTemplateNodeSpawner::SpawnNode(UEdGraph* ParentGraph, UBlueprint* Blueprint, const FName& InNotation, FVector2D const Location)
+URigVMEdGraphNode* URigVMEdGraphTemplateNodeSpawner::SpawnNode(UEdGraph* ParentGraph, UBlueprint* Blueprint, const FRigVMTemplate* Template, FVector2D const Location)
 {
 	URigVMEdGraphNode* NewNode = nullptr;
 	URigVMBlueprint* RigBlueprint = Cast<URigVMBlueprint>(Blueprint);
@@ -179,31 +230,16 @@ URigVMEdGraphNode* URigVMEdGraphTemplateNodeSpawner::SpawnNode(UEdGraph* ParentG
 
 	if (RigBlueprint != nullptr && RigGraph != nullptr)
 	{
-		bool const bIsTemplateNode = FBlueprintNodeTemplateCache::IsTemplateOuter(ParentGraph);
-		bool const bIsUserFacingNode = !bIsTemplateNode;
+		const FName Name = FRigVMBlueprintUtils::ValidateName(RigBlueprint, Template->GetName().ToString());
+		URigVMController* Controller = RigBlueprint->GetController(ParentGraph);
 
-		const FRigVMTemplate* Template = FRigVMRegistry::Get().FindTemplate(InNotation);
-		if (Template == nullptr)
-		{
-			return nullptr;
-		}
-
-		const int32 NotationHash = (int32)GetTypeHash(InNotation);
-		const FString TemplateName = TEXT("RigVMTemplate_") + FString::FromInt(NotationHash);
+		Controller->OpenUndoBracket(FString::Printf(TEXT("Add '%s' Node"), *Name.ToString()));
 
-		FName Name = bIsTemplateNode ? *TemplateName : FRigVMBlueprintUtils::ValidateName(RigBlueprint, Template->GetName().ToString());
-		URigVMController* Controller = bIsTemplateNode ? RigGraph->GetTemplateController() : RigBlueprint->GetController(ParentGraph);
-
-		if (!bIsTemplateNode)
-		{
-			Controller->OpenUndoBracket(FString::Printf(TEXT("Add '%s' Node"), *Name.ToString()));
-		}
-
-		if (URigVMTemplateNode* ModelNode = Controller->AddTemplateNode(InNotation, Location, Name.ToString(), bIsUserFacingNode, !bIsTemplateNode))
+		if (URigVMTemplateNode* ModelNode = Controller->AddTemplateNode(Template->GetNotation(), Location, Name.ToString(), true, true))
 		{
 			NewNode = Cast<URigVMEdGraphNode>(RigGraph->FindNodeForModelNodeName(ModelNode->GetFName()));
 
-			if (NewNode && bIsUserFacingNode)
+			if (NewNode)
 			{
 				Controller->ClearNodeSelection(true);
 				Controller->SelectNode(ModelNode, true, true);
@@ -211,25 +247,11 @@ URigVMEdGraphNode* URigVMEdGraphTemplateNodeSpawner::SpawnNode(UEdGraph* ParentG
 				URigVMEdGraphUnitNodeSpawner::HookupMutableNode(ModelNode, RigBlueprint);
 			}
 
-			if (bIsUserFacingNode)
-			{
-				Controller->CloseUndoBracket();
-			}
-			else
-			{
-				if(NewNode)
-				{
-					NewNode->ModelNodePath = Template->GetNotation().ToString();
-				}
-				Controller->RemoveNode(ModelNode, false);
-			}
+			Controller->CloseUndoBracket();
 		}
 		else
 		{
-			if (bIsUserFacingNode)
-			{
-				Controller->CancelUndoBracket();
-			}
+			Controller->CancelUndoBracket();
 		}
 	}
 	return NewNode;
diff --git a/Engine/Plugins/Runtime/RigVM/Source/RigVMEditor/Private/EdGraph/NodeSpawners/RigVMEdGraphUnitNodeSpawner.cpp b/Engine/Plugins/Runtime/RigVM/Source/RigVMEditor/Private/EdGraph/NodeSpawners/RigVMEdGraphUnitNodeSpawner.cpp
--- a/Engine/Plugins/Runtime/RigVM/Source/RigVMEditor/Private/EdGraph/NodeSpawners/RigVMEdGraphUnitNodeSpawner.cpp
+++ b/Engine/Plugins/Runtime/RigVM/Source/RigVMEditor/Private/EdGraph/NodeSpawners/RigVMEdGraphUnitNodeSpawner.cpp
@@ -12,6 +12,7 @@
 #include "EditorCategoryUtils.h"
 #include "K2Node_Variable.h"
 #include "BlueprintNodeTemplateCache.h"
+#include "FindInBlueprintManager.h"
 #include "RigVMBlueprintUtils.h"
 #include "ScopedTransaction.h"
 #include "RigVMFunctions/Execution/RigVMFunction_UserDefinedEvent.h"
@@ -163,34 +164,77 @@ URigVMEdGraphNode* URigVMEdGraphUnitNodeSpawner::SpawnNode(UEdGraph* ParentGraph
 	if (RigBlueprint != nullptr && RigGraph != nullptr)
 	{
 		bool const bIsTemplateNode = FBlueprintNodeTemplateCache::IsTemplateOuter(ParentGraph);
-		bool const bIsUserFacingNode = !bIsTemplateNode;
 
 		FName Name = bIsTemplateNode ? *StructTemplate->GetStructCPPName() : FRigVMBlueprintUtils::ValidateName(RigBlueprint, StructTemplate->GetFName().ToString());
 
+		if(bIsTemplateNode)
+		{
+			TArray<FPinInfo> Pins;
+			for (TFieldIterator<FProperty> It(StructTemplate); It; ++It)
+			{
+				FPinInfo Pin;
+				const FProperty* Property = *It;
+				Pin.Name = Property->GetFName();
+				FRigVMExternalVariableDef::GetTypeFromProperty(Property, Pin.CPPType, Pin.CPPTypeObject);
+				if(Pin.CPPType.IsNone())
+				{
+					continue;
+				}
+
+				if(UScriptStruct* ScriptStruct = Cast<UScriptStruct>(Pin.CPPTypeObject))
+				{
+					static UScriptStruct* ExecuteScriptStruct = FRigVMExecuteContext::StaticStruct();
+					static const FName ExecuteStructName = *ExecuteScriptStruct->GetStructCPPName();
+
+					if(ScriptStruct->IsChildOf(ExecuteScriptStruct))
+					{
+						Pin.CPPType = ExecuteStructName;
+						Pin.CPPTypeObject = ExecuteScriptStruct;
+					}
+				}
+
+				const bool bInput = Property->HasMetaData(FRigVMStruct::InputMetaName);
+				const bool bOutput = Property->HasMetaData(FRigVMStruct::OutputMetaName);
+
+				if(bInput && bOutput)
+				{
+					Pin.Direction = ERigVMPinDirection::IO;
+				}
+				else if(bInput)
+				{
+					Pin.Direction = ERigVMPinDirection::Input;
+				}
+				else if(bOutput)
+				{
+					Pin.Direction = ERigVMPinDirection::Output;
+				}
+
+				Pins.Add(Pin);
+			}
+			return SpawnTemplateNode(ParentGraph, Pins, Name);
+		}
+
 		FStructOnScope StructScope(StructTemplate);
 		const FRigVMStruct* StructInstance = (const FRigVMStruct*)StructScope.GetStructMemory();
 		const FName EventName = StructInstance->GetEventName();
 
 		// events can only exist once across all uber graphs
-		if(bIsUserFacingNode)
+		if(!EventName.IsNone())
 		{
-			if(!EventName.IsNone())
+			if(StructInstance->CanOnlyExistOnce() &&
+				(StructTemplate != FRigVMFunction_UserDefinedEvent::StaticStruct()))
 			{
-				if(StructInstance->CanOnlyExistOnce() &&
-					(StructTemplate != FRigVMFunction_UserDefinedEvent::StaticStruct()))
+				const TArray<URigVMGraph*> Models = RigBlueprint->GetAllModels();
+				for(URigVMGraph* Model : Models)
 				{
-					const TArray<URigVMGraph*> Models = RigBlueprint->GetAllModels();
-					for(URigVMGraph* Model : Models)
+					if(Model->IsRootGraph() && !Model->IsA<URigVMFunctionLibrary>())
 					{
-						if(Model->IsRootGraph() && !Model->IsA<URigVMFunctionLibrary>())
+						for(const URigVMNode* Node : Model->GetNodes())
 						{
-							for(const URigVMNode* Node : Model->GetNodes())
+							if(Node->GetEventName() == EventName)
 							{
-								if(Node->GetEventName() == EventName)
-								{
-									RigBlueprint->OnRequestJumpToHyperlink().Execute(Node);
-									return nullptr;
-								}
+								RigBlueprint->OnRequestJumpToHyperlink().Execute(Node);
+								return nullptr;
 							}
 						}
 					}
@@ -198,26 +242,23 @@ URigVMEdGraphNode* URigVMEdGraphUnitNodeSpawner::SpawnNode(UEdGraph* ParentGraph
 			}
 		}
 		
-		URigVMController* Controller = bIsTemplateNode ? RigGraph->GetTemplateController() : RigBlueprint->GetController(ParentGraph);
+		URigVMController* Controller = RigBlueprint->GetController(ParentGraph);
 		if(Controller == nullptr)
 		{
 			return nullptr;
 		}
 
-		if (!bIsTemplateNode)
-		{
-			Controller->OpenUndoBracket(FString::Printf(TEXT("Add '%s' Node"), *Name.ToString()));
-		}
+		Controller->OpenUndoBracket(FString::Printf(TEXT("Add '%s' Node"), *Name.ToString()));
 
 		FRigVMUnitNodeCreatedContext& UnitNodeCreatedContext = Controller->GetUnitNodeCreatedContext();
 		FRigVMUnitNodeCreatedContext::FScope ReasonScope(UnitNodeCreatedContext, ERigVMNodeCreatedReason::NodeSpawner);
 
-		if (URigVMUnitNode* ModelNode = Controller->AddUnitNode(StructTemplate, InMethodName, Location, Name.ToString(), bIsUserFacingNode, !bIsTemplateNode))
+		if (URigVMUnitNode* ModelNode = Controller->AddUnitNode(StructTemplate, InMethodName, Location, Name.ToString(), true, true))
 		{
 			NewNode = Cast<URigVMEdGraphNode>(RigGraph->FindNodeForModelNodeName(ModelNode->GetFName()));
 			check(NewNode);
 
-			if (NewNode && bIsUserFacingNode)
+			if (NewNode)
 			{
 				if(StructTemplate == FRigVMFunction_UserDefinedEvent::StaticStruct())
 				{
@@ -315,21 +356,11 @@ URigVMEdGraphNode* URigVMEdGraphUnitNodeSpawner::SpawnNode(UEdGraph* ParentGraph
 			*/
 #endif
 
-			if (bIsUserFacingNode)
-			{
-				Controller->CloseUndoBracket();
-			}
-			else
-			{
-				Controller->RemoveNode(ModelNode, false);
-			}
+			Controller->CloseUndoBracket();
 		}
 		else
 		{
-			if (bIsUserFacingNode)
-			{
-				Controller->CancelUndoBracket();
-			}
+			Controller->CancelUndoBracket();
 		}
 	}
 	return NewNode;
diff --git a/Engine/Plugins/Runtime/RigVM/Source/RigVMEditor/Private/EdGraph/NodeSpawners/RigVMEdGraphVariableNodeSpawner.cpp b/Engine/Plugins/Runtime/RigVM/Source/RigVMEditor/Private/EdGraph/NodeSpawners/RigVMEdGraphVariableNodeSpawner.cpp
--- a/Engine/Plugins/Runtime/RigVM/Source/RigVMEditor/Private/EdGraph/NodeSpawners/RigVMEdGraphVariableNodeSpawner.cpp
+++ b/Engine/Plugins/Runtime/RigVM/Source/RigVMEditor/Private/EdGraph/NodeSpawners/RigVMEdGraphVariableNodeSpawner.cpp
@@ -116,7 +116,6 @@ UEdGraphNode* URigVMEdGraphVariableNodeSpawner::Invoke(UEdGraph* ParentGraph, FB
 	URigVMEdGraphNode* NewNode = nullptr;
 
 	bool const bIsTemplateNode = FBlueprintNodeTemplateCache::IsTemplateOuter(ParentGraph);
-	bool const bIsUserFacingNode = !bIsTemplateNode;
 
 	// First create a backing member for our node
 	URigVMEdGraph* RigGraph = Cast<URigVMEdGraph>(ParentGraph);
@@ -133,13 +132,6 @@ UEdGraphNode* URigVMEdGraphVariableNodeSpawner::Invoke(UEdGraph* ParentGraph, FB
 	}
 #endif
 
-	URigVMController* Controller = bIsTemplateNode ? RigGraph->GetTemplateController() : RigBlueprint->GetController(ParentGraph);
-
-	if (bIsUserFacingNode)
-	{
-		Controller->OpenUndoBracket(TEXT("Add Variable"));
-	}
-
 	FString ObjectPath;
 	if (ExternalVariable.TypeObject)
 	{
@@ -164,22 +156,30 @@ UEdGraphNode* URigVMEdGraphVariableNodeSpawner::Invoke(UEdGraph* ParentGraph, FB
 		static const FString SetterPrefix = TEXT("Setter");
 		NodeName = FString::Printf(VariableNodeNameFormat, bIsGetter ? *GetterPrefix : *SetterPrefix, *ExternalVariable.Name.ToString());
 
-		NewNode = NewObject<URigVMEdGraphNode>(ParentGraph, *NodeName);
-		ParentGraph->AddNode(NewNode, false);
+		static const FName ValueName = *URigVMVariableNode::ValueName;
 
-		NewNode->CreateNewGuid();
-		NewNode->PostPlacedNewNode();
+		TArray<FPinInfo> Pins;
 
-		UEdGraphPin* ValuePin = UEdGraphPin::CreatePin(NewNode);
-		NewNode->Pins.Add(ValuePin);
-
-		ValuePin->PinType = RigVMTypeUtils::PinTypeFromExternalVariable(ExternalVariable);
-		ValuePin->Direction = bIsGetter ? EGPD_Output : EGPD_Input;
-		NewNode->SetFlags(RF_Transactional);
-		return NewNode;
+		if(!bIsGetter)
+		{
+			static UScriptStruct* ExecuteScriptStruct = FRigVMExecuteContext::StaticStruct();
+			static const FName ExecuteStructName = *ExecuteScriptStruct->GetStructCPPName();
+			Pins.Emplace(FRigVMStruct::ExecuteName, ERigVMPinDirection::IO, ExecuteStructName, ExecuteScriptStruct);
+		}
+		
+		Pins.Emplace(
+			ValueName,
+			bIsGetter ? ERigVMPinDirection::Output : ERigVMPinDirection::Input,
+			ExternalVariable.TypeName,
+			ExternalVariable.TypeObject);
+		
+		return SpawnTemplateNode(ParentGraph, Pins, *NodeName);
 	}
 
-	if (URigVMNode* ModelNode = Controller->AddVariableNodeFromObjectPath(ExternalVariable.Name, TypeName, ObjectPath, bIsGetter, FString(), Location, NodeName, !bIsTemplateNode, !bIsTemplateNode))
+	URigVMController* Controller = RigBlueprint->GetController(ParentGraph);
+	Controller->OpenUndoBracket(TEXT("Add Variable"));
+
+	if (URigVMNode* ModelNode = Controller->AddVariableNodeFromObjectPath(ExternalVariable.Name, TypeName, ObjectPath, bIsGetter, FString(), Location, NodeName, true, true))
 	{
 		for (UEdGraphNode* Node : ParentGraph->Nodes)
 		{
@@ -193,28 +193,16 @@ UEdGraphNode* URigVMEdGraphVariableNodeSpawner::Invoke(UEdGraph* ParentGraph, FB
 			}
 		}
 
-		if (bIsUserFacingNode)
-		{
-			if (NewNode)
-			{
-				Controller->ClearNodeSelection(true);
-				Controller->SelectNode(ModelNode, true, true);
-			}
-			Controller->CloseUndoBracket();
-		}
-		else
+		if (NewNode)
 		{
-			// similar to UBlueprintNodeSpawner::Invoke -> UBlueprintNodeSpawner::SpawnEdGraphNode
-			// we simply want the node, but not actually adding it to a graph
-			Controller->RemoveNode(ModelNode, false);
+			Controller->ClearNodeSelection(true);
+			Controller->SelectNode(ModelNode, true, true);
 		}
+		Controller->CloseUndoBracket();
 	}
 	else
 	{
-		if (!bIsTemplateNode)
-		{
-			Controller->CancelUndoBracket();
-		}
+		Controller->CancelUndoBracket();
 	}
 
 
diff --git a/Engine/Plugins/Runtime/RigVM/Source/RigVMEditor/Public/EdGraph/NodeSpawners/RigVMEdGraphNodeSpawner.h b/Engine/Plugins/Runtime/RigVM/Source/RigVMEditor/Public/EdGraph/NodeSpawners/RigVMEdGraphNodeSpawner.h
--- a/Engine/Plugins/Runtime/RigVM/Source/RigVMEditor/Public/EdGraph/NodeSpawners/RigVMEdGraphNodeSpawner.h
+++ b/Engine/Plugins/Runtime/RigVM/Source/RigVMEditor/Public/EdGraph/NodeSpawners/RigVMEdGraphNodeSpawner.h
@@ -22,6 +22,30 @@ public:
 
 	void SetRelatedBlueprintClass(TSubclassOf<URigVMBlueprint> InClass);
 
+	struct FPinInfo
+	{
+		FName Name;
+		ERigVMPinDirection Direction;
+		FName CPPType;
+		UObject* CPPTypeObject;
+
+		FPinInfo()
+			: Name(NAME_None)
+			, Direction(ERigVMPinDirection::Invalid)
+			, CPPType(NAME_None)
+			, CPPTypeObject(nullptr)
+		{}
+
+		FPinInfo(const FName& InName, ERigVMPinDirection InDirection, const FName& InCPPType, UObject* InCPPTypeObject)
+		: Name(InName)
+		, Direction(InDirection)
+		, CPPType(InCPPType)
+		, CPPTypeObject(InCPPTypeObject)
+		{}
+	};
+	
+	static URigVMEdGraphNode* SpawnTemplateNode(UEdGraph* InParentGraph, const TArray<FPinInfo>& InPins, const FName& InNodeName = NAME_None);
+	
 private:
 
 	TSubclassOf<URigVMBlueprint> RelatedBlueprintClass;
diff --git a/Engine/Plugins/Runtime/RigVM/Source/RigVMEditor/Public/EdGraph/NodeSpawners/RigVMEdGraphTemplateNodeSpawner.h b/Engine/Plugins/Runtime/RigVM/Source/RigVMEditor/Public/EdGraph/NodeSpawners/RigVMEdGraphTemplateNodeSpawner.h
--- a/Engine/Plugins/Runtime/RigVM/Source/RigVMEditor/Public/EdGraph/NodeSpawners/RigVMEdGraphTemplateNodeSpawner.h
+++ b/Engine/Plugins/Runtime/RigVM/Source/RigVMEditor/Public/EdGraph/NodeSpawners/RigVMEdGraphTemplateNodeSpawner.h
@@ -44,7 +44,9 @@ private:
 	UPROPERTY(Transient)
 	FName TemplateNotation;
 
-	static URigVMEdGraphNode* SpawnNode(UEdGraph* ParentGraph, UBlueprint* Blueprint, const FName& InNotation, FVector2D const Location);
+	const FRigVMTemplate* Template = nullptr;
+
+	static URigVMEdGraphNode* SpawnNode(UEdGraph* ParentGraph, UBlueprint* Blueprint, const FRigVMTemplate* Template, FVector2D const Location);
 
 	friend class UEngineTestControlRig;
 	friend class FRigVMEditor;
