RigVM: Fix RigVMBlueprint spamming the log with messages from the default blueprint object

diff --git a/Engine/Plugins/Runtime/RigVM/Source/RigVMDeveloper/Private/RigVMBlueprint.cpp b/Engine/Plugins/Runtime/RigVM/Source/RigVMDeveloper/Private/RigVMBlueprint.cpp
--- a/Engine/Plugins/Runtime/RigVM/Source/RigVMDeveloper/Private/RigVMBlueprint.cpp
+++ b/Engine/Plugins/Runtime/RigVM/Source/RigVMDeveloper/Private/RigVMBlueprint.cpp
@@ -136,6 +136,10 @@ URigVMBlueprint::URigVMBlueprint(const FObjectInitializer& ObjectInitializer)
 	VMCompileSettings.ASTSettings.ReportDelegate.BindUObject(this, &URigVMBlueprint::HandleReportFromCompiler);
 
 #if WITH_EDITOR
+	if (HasAnyFlags(RF_ClassDefaultObject))
+	{
+		CompileLog.bSilentMode = true;
+	}
 	CompileLog.SetSourcePath(GetPathName());
 	CompileLog.bLogDetailedResults = false;
 	CompileLog.EventDisplayThresholdMs = false;
