RigVM: Recover during load for missing user defined structs / enums

This implements several strategies to recover:
a) the user defined struct now has the Guid as asset registry searchable
b) we store a map from user defined struct guid to path name along the VM - in case the asset's guid has changed
c) we fall back to loading all user defined structs in the worst case to recover their Guids and try to find the dependency

diff --git a/Engine/Plugins/Runtime/RigVM/Source/RigVM/Private/RigVMCore/RigVM.cpp b/Engine/Plugins/Runtime/RigVM/Source/RigVM/Private/RigVMCore/RigVM.cpp
--- a/Engine/Plugins/Runtime/RigVM/Source/RigVM/Private/RigVMCore/RigVM.cpp
+++ b/Engine/Plugins/Runtime/RigVM/Source/RigVM/Private/RigVMCore/RigVM.cpp
@@ -11,10 +11,12 @@
 #include "RigVMObjectVersion.h"
 #include "HAL/PlatformTLS.h"
 #include "Async/ParallelFor.h"
+#include "Engine/UserDefinedEnum.h"
 #include "GenericPlatform/GenericPlatformSurvey.h"
 #include "RigVMCore/RigVMStruct.h"
 #include "UObject/UObjectHash.h"
 #include "UObject/UObjectIterator.h"
+#include "UObject/ObjectSaveContext.h"
 
 #include UE_INLINE_GENERATED_CPP_BY_NAME(RigVM)
 
@@ -158,6 +160,23 @@ void URigVM::Save(FArchive& Ar)
 		}
 
 		Ar << OperandToDebugRegisters;
+		Ar << UserDefinedStructGuidToPathName;
+	}
+
+	// advertise dependencies on user defined structs and user defined enums
+	// to make sure they are loaded prior to the VM.
+	if(Ar.IsObjectReferenceCollector())
+	{
+		const TArray<const UObject*> UserDefinedDependencies = GetUserDefinedDependencies();
+		for(const UObject* UserDefinedDependency : UserDefinedDependencies)
+		{
+			if(Cast<UUserDefinedStruct>(UserDefinedDependency) ||
+				Cast<UUserDefinedEnum>(UserDefinedDependency))
+			{
+				FSoftObjectPath PathToTypeObject(UserDefinedDependency);
+				PathToTypeObject.Serialize(Ar);
+			}
+		}
 	}
 }
 
@@ -210,6 +229,15 @@ void URigVM::Load(FArchive& Ar)
 			{
 				Ar << OperandToDebugRegisters;
 			}
+			
+			if (Ar.CustomVer(FRigVMObjectVersion::GUID) >= FRigVMObjectVersion::VMStoringUserDefinedStructMap)
+			{
+				Ar << UserDefinedStructGuidToPathName;
+			}
+			else
+			{
+				UserDefinedStructGuidToPathName.Reset();
+			}
 		}
 
 		// we only deal with virtual machines now that use the new memory infrastructure.
@@ -245,6 +273,15 @@ void URigVM::Load(FArchive& Ar)
 		{
 			Ar << OperandToDebugRegisters;
 		}
+		
+		if (Ar.CustomVer(FRigVMObjectVersion::GUID) >= FRigVMObjectVersion::VMStoringUserDefinedStructMap)
+		{
+			Ar << UserDefinedStructGuidToPathName;
+		}
+		else
+		{
+			UserDefinedStructGuidToPathName.Reset();
+		}
 	}
 
 	// ensure to load the required functions
@@ -306,6 +343,22 @@ void URigVM::PostLoad()
 	}
 }
 
+void URigVM::PreSave(FObjectPreSaveContext SaveContext)
+{
+	const TArray<const UObject*> UserDefinedDependencies = GetUserDefinedDependencies();
+	UserDefinedStructGuidToPathName.Reset();
+
+	for(const UObject* UserDefinedDependency : UserDefinedDependencies)
+	{
+		if(const UUserDefinedStruct* UserDefinedStruct = Cast<UUserDefinedStruct>(UserDefinedDependency))
+		{
+			const FString GuidBasedName = RigVMTypeUtils::GetUniqueStructTypeName(UserDefinedStruct);
+			UserDefinedStructGuidToPathName.Add(GuidBasedName, UserDefinedStruct);
+		}
+	}
+}
+
+
 #if WITH_EDITORONLY_DATA
 void URigVM::DeclareConstructClasses(TArray<FTopLevelAssetPath>& OutConstructClasses, const UClass* SpecificSubclass)
 {
@@ -527,6 +580,7 @@ void URigVM::Reset(bool IsIgnoringArchetypeRef)
 		Parameters.Reset();
 		ParametersNameMap.Reset();
 		OperandToDebugRegisters.Reset();
+		UserDefinedStructGuidToPathName.Reset();
 	}
 
 	if(!IsIgnoringArchetypeRef)
@@ -555,6 +609,7 @@ void URigVM::Empty(FRigVMExtendedExecuteContext& Context)
 	Instructions.Empty();
 	Parameters.Empty();
 	ParametersNameMap.Empty();
+	UserDefinedStructGuidToPathName.Empty();
 	ExternalVariables.Empty();
 
 	InvalidateCachedMemory();
@@ -662,6 +717,7 @@ void URigVM::CopyFrom(URigVM* InVM, bool bDeferCopy, bool bReferenceLiteralMemor
 	ParametersNameMap = InVM->ParametersNameMap;
 	
 	OperandToDebugRegisters = InVM->OperandToDebugRegisters;
+	UserDefinedStructGuidToPathName = InVM->UserDefinedStructGuidToPathName;
 
 	if (bCopyExternalVariables)
 	{
@@ -750,7 +806,10 @@ URigVMMemoryStorage* URigVM::GetMemoryByType(ERigVMMemoryType InMemoryType, bool
 					}
 				}
 			}
-			check(WorkMemoryStorageObject->GetOuter() == this);
+			if(WorkMemoryStorageObject)
+			{
+				check(WorkMemoryStorageObject->GetOuter() == this);
+			}
 			return WorkMemoryStorageObject;
 		}
 		case ERigVMMemoryType::Debug:
@@ -778,7 +837,10 @@ URigVMMemoryStorage* URigVM::GetMemoryByType(ERigVMMemoryType InMemoryType, bool
 					}
 				}
 			}
-			check(DebugMemoryStorageObject->GetOuter() == this);
+			if(DebugMemoryStorageObject)
+			{
+				check(DebugMemoryStorageObject->GetOuter() == this);
+			}	
 			return DebugMemoryStorageObject;
 		}
 		default:
@@ -970,11 +1032,14 @@ bool URigVM::ResolveFunctionsIfRequired()
 		GetFactories().Reset();
 		GetFactories().SetNumZeroed(GetFunctionNames().Num());
 
+		FRigVMTypeResolvalInfo ResolvalInfo;
+		ResolvalInfo.CPPTypeToObjectPath = UserDefinedStructGuidToPathName;
+
 		TArray<FName>& FunctionNames = GetFunctionNames();
 		for (int32 FunctionIndex = 0; FunctionIndex < FunctionNames.Num(); FunctionIndex++)
 		{
 			const FString FunctionNameString = FunctionNames[FunctionIndex].ToString();
-			if(const FRigVMFunction* Function = FRigVMRegistry::Get().FindFunction(*FunctionNameString))
+			if(const FRigVMFunction* Function = FRigVMRegistry::Get().FindFunction(*FunctionNameString, ResolvalInfo))
 			{
 				GetFunctions()[FunctionIndex] = Function;
 				GetFactories()[FunctionIndex] = Function->Factory;
@@ -2653,6 +2718,71 @@ void URigVM::RefreshExternalPropertyPaths()
 	}
 }
 
+TArray<const UObject*> URigVM::GetUserDefinedDependencies()
+{
+	TArray<const UObject*> Dependencies;
+	auto ProcessMemory = [&Dependencies](const URigVMMemoryStorage* Memory)
+	{
+		if(Memory == nullptr)
+		{
+			return;
+		}
+
+		const TArray<const FProperty*>& Properties = Memory->GetProperties();
+
+		for(const FProperty* Property : Properties)
+		{
+			const FProperty* PropertyToVisit = Property;
+			while(const FArrayProperty* ArrayProperty = CastField<FArrayProperty>(PropertyToVisit))
+			{
+				PropertyToVisit = ArrayProperty->Inner;
+			}
+			if(const FStructProperty* StructProperty = CastField<FStructProperty>(PropertyToVisit))
+			{
+				if(const UUserDefinedStruct* UserDefinedStruct = Cast<UUserDefinedStruct>(StructProperty->Struct))
+				{
+					Dependencies.AddUnique(UserDefinedStruct);
+				}
+			}
+			else if(const FEnumProperty* EnumProperty = CastField<FEnumProperty>(PropertyToVisit))
+			{
+				if(const UUserDefinedEnum* UserDefinedEnum = Cast<UUserDefinedEnum>(EnumProperty->GetEnum()))
+				{
+					Dependencies.AddUnique(UserDefinedEnum);
+				}
+			}
+			else if(const FByteProperty* ByteProperty = CastField<FByteProperty>(PropertyToVisit))
+			{
+				if(const UUserDefinedEnum* UserDefinedEnum = Cast<UUserDefinedEnum>(ByteProperty->Enum))
+				{
+					Dependencies.AddUnique(UserDefinedEnum);
+				}
+			}
+		}
+	};
+
+	ProcessMemory(GetLiteralMemory(false));
+	ProcessMemory(GetWorkMemory(false));
+
+	TArray<const FRigVMFunction*>& Functions = GetFunctions();
+	for(const FRigVMFunction* Function : Functions)
+	{
+		const FRigVMRegistry& Registry = FRigVMRegistry::Get();
+		const TArray<TRigVMTypeIndex>& TypeIndices = Function->GetArgumentTypeIndices();
+		for(const TRigVMTypeIndex& TypeIndex : TypeIndices)
+		{
+			const FRigVMTemplateArgumentType& Type = Registry.GetType(TypeIndex);
+			if(Cast<UUserDefinedStruct>(Type.CPPTypeObject) ||
+				Cast<UUserDefinedEnum>(Type.CPPTypeObject))
+			{
+				Dependencies.AddUnique(Type.CPPTypeObject);
+			}
+		}
+	}
+
+	return Dependencies;
+}
+
 void URigVM::SetupInstructionTracking(FRigVMExtendedExecuteContext& Context, int32 InInstructionCount)
 {
 #if WITH_EDITOR
diff --git a/Engine/Plugins/Runtime/RigVM/Source/RigVM/Private/RigVMCore/RigVMRegistry.cpp b/Engine/Plugins/Runtime/RigVM/Source/RigVM/Private/RigVMCore/RigVMRegistry.cpp
--- a/Engine/Plugins/Runtime/RigVM/Source/RigVM/Private/RigVMCore/RigVMRegistry.cpp
+++ b/Engine/Plugins/Runtime/RigVM/Source/RigVM/Private/RigVMCore/RigVMRegistry.cpp
@@ -842,7 +842,7 @@ const FRigVMTemplateArgumentType& FRigVMRegistry::GetType(TRigVMTypeIndex InType
 	return EmptyType;
 }
 
-const FRigVMTemplateArgumentType& FRigVMRegistry::FindTypeFromCPPType(const FString& InCPPType) const
+const FRigVMTemplateArgumentType& FRigVMRegistry::FindTypeFromCPPType(const FString& InCPPType, const FRigVMTypeResolvalInfo* InResolvalInfo) const
 {
 	const int32 TypeIndex = GetTypeIndexFromCPPType(InCPPType);
 	if(ensure(Types.IsValidIndex(TypeIndex)))
@@ -854,7 +854,7 @@ const FRigVMTemplateArgumentType& FRigVMRegistry::FindTypeFromCPPType(const FStr
 	return EmptyType;
 }
 
-TRigVMTypeIndex FRigVMRegistry::GetTypeIndexFromCPPType(const FString& InCPPType) const
+TRigVMTypeIndex FRigVMRegistry::GetTypeIndexFromCPPType(const FString& InCPPType, const FRigVMTypeResolvalInfo* InResolvalInfo) const
 {
 	TRigVMTypeIndex Result = INDEX_NONE;
 	if(ensure(!InCPPType.IsEmpty()))
@@ -1309,13 +1309,13 @@ void FRigVMRegistry::RegisterObjectTypes(TConstArrayView<UClass*> InClasses)
 	}
 }
 
-const FRigVMFunction* FRigVMRegistry::FindFunction(const TCHAR* InName) const
+const FRigVMFunction* FRigVMRegistry::FindFunction(const TCHAR* InName, const FRigVMTypeResolvalInfo& InResolvalInfo) const
 {
 	FScopeLock FindFunctionScopeLock(&FindFunctionMutex);
-	return FindFunction_NoLock(InName);
+	return FindFunction_NoLock(InName, InResolvalInfo);
 }
 
-const FRigVMFunction* FRigVMRegistry::FindFunction_NoLock(const TCHAR* InName) const
+const FRigVMFunction* FRigVMRegistry::FindFunction_NoLock(const TCHAR* InName, const FRigVMTypeResolvalInfo& InResolvalInfo) const
 {
 	// Check first if the function is provided by internally registered rig units. 
 	if(const int32* FunctionIndexPtr = FunctionNameToIndex.Find(InName))
@@ -1333,7 +1333,7 @@ const FRigVMFunction* FRigVMRegistry::FindFunction_NoLock(const TCHAR* InName) c
 		{
 			if(const FRigVMTemplate* Template = Factory->GetTemplate())
 			{
-				const FRigVMTemplateTypeMap ArgumentTypes = Template->GetArgumentTypesFromString(SuffixString);
+				const FRigVMTemplateTypeMap ArgumentTypes = Template->GetArgumentTypesFromString(SuffixString, &InResolvalInfo);
 				if(ArgumentTypes.Num() == Template->NumArguments())
 				{
 					const int32 PermutationIndex = Template->FindPermutation(ArgumentTypes);
@@ -1376,7 +1376,7 @@ const FRigVMFunction* FRigVMRegistry::FindFunction_NoLock(const TCHAR* InName) c
 				{
 					NewStructOrFactoryName = StructPrefix + NewStructOrFactoryName;
 				}
-				const FRigVMFunction* RedirectedFunction = FindFunction_NoLock(*(NewStructOrFactoryName + TEXT("::") + SuffixString));
+				const FRigVMFunction* RedirectedFunction = FindFunction_NoLock(*(NewStructOrFactoryName + TEXT("::") + SuffixString), InResolvalInfo);
 				if(RedirectedFunction)
 				{
 					FRigVMRegistry& MutableRegistry = FRigVMRegistry::Get();
@@ -1390,13 +1390,13 @@ const FRigVMFunction* FRigVMRegistry::FindFunction_NoLock(const TCHAR* InName) c
 	return nullptr;
 }
 
-const FRigVMFunction* FRigVMRegistry::FindFunction(UScriptStruct* InStruct, const TCHAR* InName) const
+const FRigVMFunction* FRigVMRegistry::FindFunction(UScriptStruct* InStruct, const TCHAR* InName, const FRigVMTypeResolvalInfo& InResolvalInfo) const
 {
 	check(InStruct);
 	check(InName);
 	
 	const FString FunctionName = FString::Printf(TEXT("%s::%s"), *InStruct->GetStructCPPName(), InName);
-	return FindFunction(*FunctionName);
+	return FindFunction(*FunctionName, InResolvalInfo);
 }
 
 const TChunkedArray<FRigVMFunction>& FRigVMRegistry::GetFunctions() const
diff --git a/Engine/Plugins/Runtime/RigVM/Source/RigVM/Private/RigVMCore/RigVMTemplate.cpp b/Engine/Plugins/Runtime/RigVM/Source/RigVM/Private/RigVMCore/RigVMTemplate.cpp
--- a/Engine/Plugins/Runtime/RigVM/Source/RigVM/Private/RigVMCore/RigVMTemplate.cpp
+++ b/Engine/Plugins/Runtime/RigVM/Source/RigVM/Private/RigVMCore/RigVMTemplate.cpp
@@ -654,7 +654,7 @@ TArray<UStruct*> FRigVMTemplate::GetSuperStructs(UStruct* InStruct, bool bInclud
 	return SuperStructs;
 }
 
-FRigVMTemplate::FTypeMap FRigVMTemplate::GetArgumentTypesFromString(const FString& InTypeString) const
+FRigVMTemplate::FTypeMap FRigVMTemplate::GetArgumentTypesFromString(const FString& InTypeString, const FRigVMTypeResolvalInfo* InResolvalInfo) const
 {
 	FTypeMap Types;
 	if(!InTypeString.IsEmpty())
@@ -675,12 +675,12 @@ FRigVMTemplate::FTypeMap FRigVMTemplate::GetArgumentTypesFromString(const FStrin
 			{
 				if(const FRigVMTemplateArgument* Argument = FindArgument(*ArgumentName))
 				{
-					TRigVMTypeIndex TypeIndex = Registry.GetTypeIndexFromCPPType(TypeName);
+					TRigVMTypeIndex TypeIndex = Registry.GetTypeIndexFromCPPType(TypeName, InResolvalInfo);
 
 					// If the type was not found, check if it's a user-defined type that hasn't been registered yet.
 					if (TypeIndex == INDEX_NONE && RigVMTypeUtils::RequiresCPPTypeObject(TypeName))
 					{
-						UObject* CPPTypeObject = RigVMTypeUtils::ObjectFromCPPType(TypeName, true);
+						UObject* CPPTypeObject = RigVMTypeUtils::ObjectFromCPPType(TypeName, true, InResolvalInfo);
 						
 						FRigVMTemplateArgumentType ArgType(*TypeName, CPPTypeObject);
 						TypeIndex = Registry.FindOrAddType(ArgType);
diff --git a/Engine/Plugins/Runtime/RigVM/Source/RigVM/Private/RigVMTypeUtils.cpp b/Engine/Plugins/Runtime/RigVM/Source/RigVM/Private/RigVMTypeUtils.cpp
--- a/Engine/Plugins/Runtime/RigVM/Source/RigVM/Private/RigVMTypeUtils.cpp
+++ b/Engine/Plugins/Runtime/RigVM/Source/RigVM/Private/RigVMTypeUtils.cpp
@@ -5,6 +5,9 @@
 =============================================================================*/
 
 #include "RigVMTypeUtils.h"
+#include "AssetRegistry/AssetRegistryModule.h"
+#include "Engine/UserDefinedEnum.h"
+#include "Engine/UserDefinedStruct.h"
 
 namespace RigVMTypeUtils
 {
@@ -28,4 +31,128 @@ namespace RigVMTypeUtils
 	TRigVMTypeIndex TypeIndex::FNameArray = INDEX_NONE;	
 	TRigVMTypeIndex TypeIndex::FStringArray = INDEX_NONE;	
 	TRigVMTypeIndex TypeIndex::WildCardArray = INDEX_NONE;	
-}
\ No newline at end of file
+}
+
+UObject* RigVMTypeUtils::UserDefinedTypeFromCPPType(FString& InOutCPPType, const FRigVMTypeResolvalInfo* InResolvalInfo)
+{
+	const FString OriginalTypeName = InOutCPPType;
+	UObject* CPPTypeObject = nullptr;
+	InOutCPPType.Reset();
+
+	// try to resolve the type name using a path name potentially
+	if(InOutCPPType.IsEmpty() && InResolvalInfo != nullptr)
+	{
+		if(!InResolvalInfo->CPPTypeToObjectPath.IsEmpty())
+		{
+			FString TypeNameToLookUp = OriginalTypeName;
+			while(IsArrayType(TypeNameToLookUp))
+			{
+				TypeNameToLookUp = BaseTypeFromArrayType(TypeNameToLookUp);
+			}
+
+			// this map contains a map of user defined struct keys based on its guid to their path,
+			// for example FUserDefinedStruct_23E408214EE9E6DA5BFADDA0F9F4F577 -> /Game/Animation/MyUserDefinedStruct.MyUserDefinedStruct
+			if(const FSoftObjectPath* ResolveObjectPath = InResolvalInfo->CPPTypeToObjectPath.Find(TypeNameToLookUp))
+			{
+				CPPTypeObject = ResolveObjectPath->ResolveObject();
+				if(CPPTypeObject == nullptr)
+				{
+					CPPTypeObject = ResolveObjectPath->TryLoad();
+				}
+				if(CPPTypeObject)
+				{
+					InOutCPPType = PostProcessCPPType(OriginalTypeName, CPPTypeObject);
+					return CPPTypeObject;
+				}
+			}
+		}
+	}
+
+#if WITH_EDITOR
+	
+	// potentially this type hasn't been loaded yet. Let's try again by visiting relevant assets
+	if(InOutCPPType.IsEmpty())
+	{
+		const IAssetRegistry& AssetRegistry = FModuleManager::LoadModuleChecked<FAssetRegistryModule>("AssetRegistry").Get();
+
+		InOutCPPType = OriginalTypeName;
+		while(IsArrayType(InOutCPPType))
+		{
+			InOutCPPType = BaseTypeFromArrayType(InOutCPPType);
+		}
+
+		if(OriginalTypeName.Contains(TEXT("FUserDefinedStruct_")))
+		{
+			static const FName GuidTag = GET_MEMBER_NAME_CHECKED(UUserDefinedStruct, Guid);
+
+			TArray<FAssetData> AssetDataList;
+			AssetRegistry.GetAssetsByClass(UUserDefinedStruct::StaticClass()->GetClassPathName(), AssetDataList, true);
+
+			// first pass - try to find it using the tag
+			if(CPPTypeObject == nullptr)
+			{
+				for(const FAssetData& AssetData : AssetDataList)
+				{
+					if(AssetData.FindTag(GuidTag))
+					{
+						const FString GuidBasedName = GetUniqueStructTypeName(AssetData.GetTagValueRef<FGuid>(GuidTag));
+						if(GuidBasedName == InOutCPPType)
+						{
+							if(UUserDefinedStruct* UserDefinedStruct = Cast<UUserDefinedStruct>(AssetData.GetAsset()))
+							{
+								CPPTypeObject = UserDefinedStruct;
+								InOutCPPType = PostProcessCPPType(OriginalTypeName, CPPTypeObject);
+								return CPPTypeObject;
+							}
+						}
+					}
+				}
+			}
+
+			// second pass - deal with the ones that don't have a tag and force load them
+			if(CPPTypeObject == nullptr)
+			{
+				for(const FAssetData& AssetData : AssetDataList)
+				{
+					if(AssetData.FindTag(GuidTag))
+					{
+						continue;
+					}
+
+					if(UUserDefinedStruct* UserDefinedStruct = Cast<UUserDefinedStruct>(AssetData.GetAsset()))
+					{
+						const FString GuidBasedName = GetUniqueStructTypeName(UserDefinedStruct);
+						if(GuidBasedName == InOutCPPType)
+						{
+							CPPTypeObject = UserDefinedStruct;
+							InOutCPPType = PostProcessCPPType(OriginalTypeName, CPPTypeObject);
+							return CPPTypeObject;
+						}
+					}
+				}
+			}
+		}
+		else
+		{
+			TArray<FAssetData> AssetDataList;
+			AssetRegistry.GetAssetsByClass(UUserDefinedEnum::StaticClass()->GetClassPathName(), AssetDataList, true);
+
+			for(const FAssetData& AssetData : AssetDataList)
+			{
+				if(UUserDefinedEnum* UserDefinedEnum = Cast<UUserDefinedEnum>(AssetData.GetAsset()))
+				{
+					const FString EnumCPPName = UserDefinedEnum->GetName();
+					if(EnumCPPName == InOutCPPType)
+					{
+						CPPTypeObject = UserDefinedEnum;
+						InOutCPPType = PostProcessCPPType(OriginalTypeName, CPPTypeObject);
+						return CPPTypeObject;
+					}
+				}
+			}
+		}
+	}
+#endif
+
+	return CPPTypeObject;
+}
diff --git a/Engine/Plugins/Runtime/RigVM/Source/RigVM/Public/RigVMCore/RigVM.h b/Engine/Plugins/Runtime/RigVM/Source/RigVM/Public/RigVMCore/RigVM.h
--- a/Engine/Plugins/Runtime/RigVM/Source/RigVM/Public/RigVMCore/RigVM.h
+++ b/Engine/Plugins/Runtime/RigVM/Source/RigVM/Public/RigVMCore/RigVM.h
@@ -167,6 +167,7 @@ public:
 	virtual void Save(FArchive& Ar);
 	virtual void Load(FArchive& Ar);
 	virtual void PostLoad() override;
+	virtual void PreSave(FObjectPreSaveContext ObjectSaveContext) override;
 #if WITH_EDITORONLY_DATA
 	static void DeclareConstructClasses(TArray<FTopLevelAssetPath>& OutConstructClasses, const UClass* SpecificSubclass);
 #endif
@@ -814,6 +815,8 @@ public:
 	void InvalidateCachedMemory();
 	void InvalidateCachedMemory(FRigVMExtendedExecuteContext& Context);
 	
+	const TMap<FString, FSoftObjectPath>& GetUserDefinedStructGuidToObjectPath() const { return UserDefinedStructGuidToPathName; }
+	
 private:
 	void CacheMemoryHandlesIfRequired(FRigVMExtendedExecuteContext& Context, TArrayView<URigVMMemoryStorage*> InMemory);
 	void RebuildByteCodeOnLoad();
@@ -864,6 +867,7 @@ private:
 
 	TArray<FRigVMExternalVariableDef> ExternalVariables;
 	TArray<FRigVMLazyBranch> LazyBranches;
+	TMap<FString, FSoftObjectPath> UserDefinedStructGuidToPathName;
 
 	// this function should be kept in sync with FRigVMOperand::GetContainerIndex()
 	static int32 GetContainerIndex(ERigVMMemoryType InType)
@@ -909,6 +913,8 @@ private:
 	
 	TMap<FRigVMOperand, TArray<FRigVMOperand>> OperandToDebugRegisters;
 
+	TArray<const UObject*> GetUserDefinedDependencies();
+
 protected:
 
 	struct FEntryExecuteGuard
diff --git a/Engine/Plugins/Runtime/RigVM/Source/RigVM/Public/RigVMCore/RigVMRegistry.h b/Engine/Plugins/Runtime/RigVM/Source/RigVM/Public/RigVMCore/RigVMRegistry.h
--- a/Engine/Plugins/Runtime/RigVM/Source/RigVM/Public/RigVMCore/RigVMRegistry.h
+++ b/Engine/Plugins/Runtime/RigVM/Source/RigVM/Public/RigVMCore/RigVMRegistry.h
@@ -172,10 +172,10 @@ public:
 	int32 NumTypes() const { return Types.Num(); }
 
 	// Returns the type given only its cpp type
-	const FRigVMTemplateArgumentType& FindTypeFromCPPType(const FString& InCPPType) const;
+	const FRigVMTemplateArgumentType& FindTypeFromCPPType(const FString& InCPPType, const FRigVMTypeResolvalInfo* InResolvalInfo = nullptr) const;
 
 	// Returns the type index given only its cpp type
-	TRigVMTypeIndex GetTypeIndexFromCPPType(const FString& InCPPType) const;
+	TRigVMTypeIndex GetTypeIndexFromCPPType(const FString& InCPPType, const FRigVMTypeResolvalInfo* InResolvalInfo = nullptr) const;
 
 	// Returns true if the type is an array
 	bool IsArrayType(TRigVMTypeIndex InTypeIndex) const;
@@ -208,10 +208,10 @@ public:
 	TRigVMTypeIndex GetBaseTypeFromArrayTypeIndex(TRigVMTypeIndex InTypeIndex) const;
 
 	// Returns the function given its name (or nullptr)
-	const FRigVMFunction* FindFunction(const TCHAR* InName) const;
+	const FRigVMFunction* FindFunction(const TCHAR* InName, const FRigVMTypeResolvalInfo& InResolvalInfo = FRigVMTypeResolvalInfo()) const;
 
 	// Returns the function given its backing up struct and method name
-	const FRigVMFunction* FindFunction(UScriptStruct* InStruct, const TCHAR* InName) const;
+	const FRigVMFunction* FindFunction(UScriptStruct* InStruct, const TCHAR* InName, const FRigVMTypeResolvalInfo& InResolvalInfo = FRigVMTypeResolvalInfo()) const;
 
 	// Returns all current rigvm functions
 	const TChunkedArray<FRigVMFunction>& GetFunctions() const;
@@ -332,7 +332,7 @@ private:
 	
 	void RemoveTypeInCategory(FRigVMTemplateArgument::ETypeCategory InCategory, TRigVMTypeIndex InTypeIndex);
 
-	const FRigVMFunction* FindFunction_NoLock(const TCHAR* InName) const;
+	const FRigVMFunction* FindFunction_NoLock(const TCHAR* InName, const FRigVMTypeResolvalInfo& InResolvalInfo = FRigVMTypeResolvalInfo()) const;
 	const FRigVMTemplate* FindTemplate_NoLock(const FName& InNotation, bool bIncludeDeprecated) const;
 	FRigVMDispatchFactory* FindDispatchFactory_NoLock(const FName& InFactoryName) const;
 
diff --git a/Engine/Plugins/Runtime/RigVM/Source/RigVM/Public/RigVMCore/RigVMTemplate.h b/Engine/Plugins/Runtime/RigVM/Source/RigVM/Public/RigVMCore/RigVMTemplate.h
--- a/Engine/Plugins/Runtime/RigVM/Source/RigVM/Public/RigVMCore/RigVMTemplate.h
+++ b/Engine/Plugins/Runtime/RigVM/Source/RigVM/Public/RigVMCore/RigVMTemplate.h
@@ -30,6 +30,7 @@ class FProperty;
 struct FRigVMDispatchFactory;
 struct FRigVMTemplate;
 struct FRigVMDispatchContext;
+struct FRigVMTypeResolvalInfo;
 
 typedef TMap<FName, TRigVMTypeIndex> FRigVMTemplateTypeMap;
 
@@ -373,7 +374,7 @@ public:
 	static TArray<UStruct*> GetSuperStructs(UStruct* InStruct, bool bIncludeLeaf = true);
 
 	// converts the types provided by a string (like "A:float,B:int32") into a type map
-	FTypeMap GetArgumentTypesFromString(const FString& InTypeString) const;
+	FTypeMap GetArgumentTypesFromString(const FString& InTypeString, const FRigVMTypeResolvalInfo* InResolvalInfo = nullptr) const;
 
 	// converts the types provided to a string (like "A:float,B:int32")
 	static FString GetStringFromArgumentTypes(const FTypeMap& InTypes); 
diff --git a/Engine/Plugins/Runtime/RigVM/Source/RigVM/Public/RigVMObjectVersion.h b/Engine/Plugins/Runtime/RigVM/Source/RigVM/Public/RigVMObjectVersion.h
--- a/Engine/Plugins/Runtime/RigVM/Source/RigVM/Public/RigVMObjectVersion.h
+++ b/Engine/Plugins/Runtime/RigVM/Source/RigVM/Public/RigVMObjectVersion.h
@@ -17,6 +17,9 @@ struct RIGVM_API FRigVMObjectVersion
 		// Predicates added to execute operations
 		PredicatesAddedToExecuteOps,
 
+		// Storing paths to user defined structs map
+		VMStoringUserDefinedStructMap,
+
 		// -----<new versions can be added above this line>-------------------------------------------------
 		VersionPlusOne,
 		LatestVersion = VersionPlusOne - 1
diff --git a/Engine/Plugins/Runtime/RigVM/Source/RigVM/Public/RigVMTypeUtils.h b/Engine/Plugins/Runtime/RigVM/Source/RigVM/Public/RigVMTypeUtils.h
--- a/Engine/Plugins/Runtime/RigVM/Source/RigVM/Public/RigVMTypeUtils.h
+++ b/Engine/Plugins/Runtime/RigVM/Source/RigVM/Public/RigVMTypeUtils.h
@@ -9,6 +9,12 @@
 #include "Engine/UserDefinedStruct.h"
 #include "UObject/CoreRedirects.h"
 #include "UObject/Package.h"
+#include "UObject/SoftObjectPath.h"
+
+struct RIGVM_API FRigVMTypeResolvalInfo
+{
+	TMap<FString, FSoftObjectPath> CPPTypeToObjectPath;
+};
 
 namespace RigVMTypeUtils
 {
@@ -94,11 +100,16 @@ namespace RigVMTypeUtils
 		return InCPPType.RightChop(7).LeftChop(1).TrimStartAndEnd();
 	}
 
+	inline FString GetUniqueStructTypeName(const FGuid& InStructGuid)
+	{
+		return FString::Printf(TEXT("FUserDefinedStruct_%s"), *InStructGuid.ToString());
+	}
+
 	inline FString GetUniqueStructTypeName(const UScriptStruct* InScriptStruct)
 	{
 		if (const UUserDefinedStruct* UserDefinedStruct = Cast<UUserDefinedStruct>(InScriptStruct))
 		{
-			return FString::Printf(TEXT("FUserDefinedStruct_%s"), *UserDefinedStruct->GetCustomGuid().ToString());
+			return GetUniqueStructTypeName(UserDefinedStruct->GetCustomGuid());
 		}
 
 		return InScriptStruct->GetStructCPPName();
@@ -299,8 +310,11 @@ namespace RigVMTypeUtils
 		return FString();
 	}
 
+	// Finds the CPPTypeObject from a CPP type of a potentially missing / unloaded user defined struct or enum
+	RIGVM_API UObject* UserDefinedTypeFromCPPType(FString& InOutCPPType, const FRigVMTypeResolvalInfo* InResolvalInfo = nullptr);
+
 	// Finds the CPPTypeObject from the CPPType. If not found, tries to use redirectors and modifies the InOutCPPType.
-	static UObject* ObjectFromCPPType(FString& InOutCPPType, bool bUseRedirector = true)
+	static UObject* ObjectFromCPPType(FString& InOutCPPType, bool bUseRedirector = true, const FRigVMTypeResolvalInfo* InResolvalInfo = nullptr)
 	{
 		if (!RequiresCPPTypeObject(InOutCPPType))
 		{
@@ -334,6 +348,12 @@ namespace RigVMTypeUtils
 			CPPTypeObject = RigVMTypeUtils::FindObjectGlobally(*CPPType, bUseRedirector);
 		}
 
+		if(CPPTypeObject == nullptr)
+		{
+			CPPType = BaseCPPType;
+			CPPTypeObject = UserDefinedTypeFromCPPType(CPPType, InResolvalInfo);
+		}
+
 		if(CPPTypeObject == nullptr)
 		{
 			InOutCPPType.Reset();
@@ -345,7 +365,7 @@ namespace RigVMTypeUtils
 		return CPPTypeObject;
 	}
 
-	static FString PostProcessCPPType(const FString& InCPPType, UObject* InCPPTypeObject = nullptr)
+	static FString PostProcessCPPType(const FString& InCPPType, UObject* InCPPTypeObject = nullptr, const FRigVMTypeResolvalInfo* InResolvalInfo = nullptr)
 	{
 		FString CPPType = InCPPType;
 		if (InCPPTypeObject)
@@ -364,7 +384,7 @@ namespace RigVMTypeUtils
 		else if (RequiresCPPTypeObject(CPPType))
 		{
 			// Uses redirectors and updates the CPPType if necessary
-			ObjectFromCPPType(CPPType, true);
+			ObjectFromCPPType(CPPType, true, InResolvalInfo);
 		}
 	
 		return CPPType;
diff --git a/Engine/Plugins/Runtime/RigVM/Source/RigVMDeveloper/Private/RigVMBlueprint.cpp b/Engine/Plugins/Runtime/RigVM/Source/RigVMDeveloper/Private/RigVMBlueprint.cpp
--- a/Engine/Plugins/Runtime/RigVM/Source/RigVMDeveloper/Private/RigVMBlueprint.cpp
+++ b/Engine/Plugins/Runtime/RigVM/Source/RigVMDeveloper/Private/RigVMBlueprint.cpp
@@ -789,6 +789,27 @@ void URigVMBlueprint::PreSave(FObjectPreSaveContext ObjectSaveContext)
 	IAssetRegistry::GetChecked().AssetTagsFinalized(*this);
 
 	CachedAssetTags.Reset();
+
+	// also store the user defined struct guid to path name on the blueprint itself
+	// to aid the controller when recovering from user defined struct name changes or
+	// guid changes.
+	UserDefinedStructGuidToPathName.Reset();
+	TArray<URigVMGraph*> AllModels = GetAllModels();
+	for(const URigVMGraph* Graph : AllModels)
+	{
+		for(const URigVMNode* Node : Graph->GetNodes())
+		{
+			const TArray<URigVMPin*> AllPins = Node->GetAllPinsRecursively();
+			for(const URigVMPin* Pin : AllPins)
+			{
+				if(const UUserDefinedStruct* UserDefinedStruct = Cast<UUserDefinedStruct>(Pin->GetCPPTypeObject()))
+				{
+					const FString GuidBasedName = RigVMTypeUtils::GetUniqueStructTypeName(UserDefinedStruct);
+					UserDefinedStructGuidToPathName.FindOrAdd(GuidBasedName) = FSoftObjectPath(UserDefinedStruct);
+				}
+			}
+		}
+	}
 }
 
 void URigVMBlueprint::PostLoad()
diff --git a/Engine/Plugins/Runtime/RigVM/Source/RigVMDeveloper/Private/RigVMCompiler/RigVMCompiler.cpp b/Engine/Plugins/Runtime/RigVM/Source/RigVMDeveloper/Private/RigVMCompiler/RigVMCompiler.cpp
--- a/Engine/Plugins/Runtime/RigVM/Source/RigVMDeveloper/Private/RigVMCompiler/RigVMCompiler.cpp
+++ b/Engine/Plugins/Runtime/RigVM/Source/RigVMDeveloper/Private/RigVMCompiler/RigVMCompiler.cpp
@@ -474,6 +474,8 @@ bool URigVMCompiler::Compile(TArray<URigVMGraph*> InGraphs, URigVMController* In
 			{
 				if(!URigVMController::EnsurePinValidity(Pin, true))
 				{
+					static const FString InvalidPinEncountered = TEXT("Pin @@ is not valid - potentially using an invalid type?");
+					Settings.ASTSettings.Report(EMessageSeverity::Error, Pin, InvalidPinEncountered);
 					return false;
 				}
 			}
diff --git a/Engine/Plugins/Runtime/RigVM/Source/RigVMDeveloper/Private/RigVMModel/RigVMClient.cpp b/Engine/Plugins/Runtime/RigVM/Source/RigVMDeveloper/Private/RigVMModel/RigVMClient.cpp
--- a/Engine/Plugins/Runtime/RigVM/Source/RigVMDeveloper/Private/RigVMModel/RigVMClient.cpp
+++ b/Engine/Plugins/Runtime/RigVM/Source/RigVMDeveloper/Private/RigVMModel/RigVMClient.cpp
@@ -14,6 +14,12 @@
 #include "ScopedTransaction.h"
 #endif
 
+const TMap<FString, FSoftObjectPath>& IRigVMClientHost::GetUserDefinedStructGuidToObjectPath() const
+{
+	static const TMap<FString, FSoftObjectPath> EmptyMap;
+	return EmptyMap;
+}
+
 void FRigVMClient::SetSchemaClass(TSubclassOf<URigVMSchema> InSchemaClass)
 {
 	check(InSchemaClass);
diff --git a/Engine/Plugins/Runtime/RigVM/Source/RigVMDeveloper/Private/RigVMModel/RigVMController.cpp b/Engine/Plugins/Runtime/RigVM/Source/RigVMDeveloper/Private/RigVMModel/RigVMController.cpp
--- a/Engine/Plugins/Runtime/RigVM/Source/RigVMDeveloper/Private/RigVMModel/RigVMController.cpp
+++ b/Engine/Plugins/Runtime/RigVM/Source/RigVMDeveloper/Private/RigVMModel/RigVMController.cpp
@@ -17816,8 +17816,18 @@ bool URigVMController::EnsurePinValidity(URigVMPin* InPin, bool bRecursive)
 		// without testing for redirector
 		if(InPin->GetCPPTypeObject() == nullptr)
 		{
+			FRigVMTypeResolvalInfo ResolveInfo;
+			if(const IRigVMClientHost* ClientHost = InPin->GetImplementingOuter<IRigVMClientHost>())
+			{
+				ResolveInfo.CPPTypeToObjectPath = ClientHost->GetUserDefinedStructGuidToObjectPath();
+			}
+			
 			FString CPPType = InPin->GetCPPType();
-			InPin->CPPTypeObject = RigVMTypeUtils::ObjectFromCPPType(CPPType);
+			InPin->CPPTypeObject = RigVMTypeUtils::ObjectFromCPPType(CPPType, true, &ResolveInfo);
+			if(CPPType.IsEmpty())
+			{
+				return false;
+			}
 			InPin->CPPType = CPPType;
 		}
 		else
diff --git a/Engine/Plugins/Runtime/RigVM/Source/RigVMDeveloper/Public/RigVMBlueprint.h b/Engine/Plugins/Runtime/RigVM/Source/RigVMDeveloper/Public/RigVMBlueprint.h
--- a/Engine/Plugins/Runtime/RigVM/Source/RigVMDeveloper/Public/RigVMBlueprint.h
+++ b/Engine/Plugins/Runtime/RigVM/Source/RigVMDeveloper/Public/RigVMBlueprint.h
@@ -299,6 +299,7 @@ public:
 	virtual void HandleRigVMGraphRemoved(const FRigVMClient* InClient, const FString& InNodePath) override;
 	virtual void HandleRigVMGraphRenamed(const FRigVMClient* InClient, const FString& InOldNodePath, const FString& InNewNodePath) override;
 	virtual void HandleConfigureRigVMController(const FRigVMClient* InClient, URigVMController* InControllerToConfigure) override;
+	virtual const TMap<FString, FSoftObjectPath>& GetUserDefinedStructGuidToObjectPath() const override { return UserDefinedStructGuidToPathName; }
 
 	FOnRigVMRequestInspectObject& OnRequestInspectObject() { return OnRequestInspectObjectEvent; }
 	void RequestInspectObject(const TArray<UObject*>& InObjects) { OnRequestInspectObjectEvent.Broadcast(InObjects); }
@@ -421,6 +422,9 @@ public:
 	UPROPERTY(EditAnywhere, Category = "Python Log Settings")
 	FRigVMPythonSettings PythonLogSettings;
 
+	UPROPERTY()
+	TMap<FString, FSoftObjectPath> UserDefinedStructGuidToPathName;
+
 protected:
 
 #if WITH_EDITORONLY_DATA
diff --git a/Engine/Plugins/Runtime/RigVM/Source/RigVMDeveloper/Public/RigVMModel/RigVMClient.h b/Engine/Plugins/Runtime/RigVM/Source/RigVMDeveloper/Public/RigVMModel/RigVMClient.h
--- a/Engine/Plugins/Runtime/RigVM/Source/RigVMDeveloper/Public/RigVMModel/RigVMClient.h
+++ b/Engine/Plugins/Runtime/RigVM/Source/RigVMDeveloper/Public/RigVMModel/RigVMClient.h
@@ -50,6 +50,8 @@ public:
 	// Reacts to a request to configure a controller
 	virtual void HandleConfigureRigVMController(const FRigVMClient* InClient, URigVMController* InControllerToConfigure) = 0;
 
+	// Returns a map from user defined struct to object path used to resolve types
+	virtual const TMap<FString, FSoftObjectPath>& GetUserDefinedStructGuidToObjectPath() const;
 };
 
 UINTERFACE()
diff --git a/Engine/Source/Runtime/Engine/Classes/Engine/UserDefinedStruct.h b/Engine/Source/Runtime/Engine/Classes/Engine/UserDefinedStruct.h
--- a/Engine/Source/Runtime/Engine/Classes/Engine/UserDefinedStruct.h
+++ b/Engine/Source/Runtime/Engine/Classes/Engine/UserDefinedStruct.h
@@ -74,7 +74,7 @@ public:
 	TEnumAsByte<enum EUserDefinedStructureStatus> Status;
 
 	/** Uniquely identifies this specific user struct */
-	UPROPERTY()
+	UPROPERTY(AssetRegistrySearchable)
 	FGuid Guid;
 
 protected:
