RigVM: Fix non unity error in RigVMTypeUtils.cpp

diff --git a/Engine/Plugins/Runtime/RigVM/Source/RigVM/Private/RigVMTypeUtils.cpp b/Engine/Plugins/Runtime/RigVM/Source/RigVM/Private/RigVMTypeUtils.cpp
--- a/Engine/Plugins/Runtime/RigVM/Source/RigVM/Private/RigVMTypeUtils.cpp
+++ b/Engine/Plugins/Runtime/RigVM/Source/RigVM/Private/RigVMTypeUtils.cpp
@@ -5,6 +5,7 @@
 =============================================================================*/
 
 #include "RigVMTypeUtils.h"
+#include "AssetRegistry/AssetData.h"
 #include "AssetRegistry/AssetRegistryModule.h"
 #include "Engine/UserDefinedEnum.h"
 #include "Engine/UserDefinedStruct.h"
