RigVM: Fix URigVM::PreSave calling Super::Presave

diff --git a/Engine/Plugins/Runtime/RigVM/Source/RigVM/Private/RigVMCore/RigVM.cpp b/Engine/Plugins/Runtime/RigVM/Source/RigVM/Private/RigVMCore/RigVM.cpp
--- a/Engine/Plugins/Runtime/RigVM/Source/RigVM/Private/RigVMCore/RigVM.cpp
+++ b/Engine/Plugins/Runtime/RigVM/Source/RigVM/Private/RigVMCore/RigVM.cpp
@@ -345,6 +345,8 @@ void URigVM::PostLoad()
 
 void URigVM::PreSave(FObjectPreSaveContext SaveContext)
 {
+	Super::PreSave(SaveContext);
+
 	const TArray<const UObject*> UserDefinedDependencies = GetUserDefinedDependencies();
 	UserDefinedStructGuidToPathName.Reset();
 
