RigVM: Band-aid fix for cooking.

#rnx

diff --git a/Engine/Plugins/Runtime/RigVM/Source/RigVM/Private/RigVMCore/RigVMFunction.cpp b/Engine/Plugins/Runtime/RigVM/Source/RigVM/Private/RigVMCore/RigVMFunction.cpp
--- a/Engine/Plugins/Runtime/RigVM/Source/RigVM/Private/RigVMCore/RigVMFunction.cpp
+++ b/Engine/Plugins/Runtime/RigVM/Source/RigVM/Private/RigVMCore/RigVMFunction.cpp
@@ -92,7 +92,7 @@ const TArray<TRigVMTypeIndex>& FRigVMFunction::GetArgumentTypeIndices() const
 		}
 		else
 		{
-			checkNoEntry();
+			// checkNoEntry();
 		}
 	}
 	return ArgumentTypeIndices;
