RigVM: Fix crash when unresolving singleton nodes
#rnx

diff --git a/Engine/Plugins/Runtime/RigVM/Source/RigVMDeveloper/Private/RigVMModel/RigVMController.cpp b/Engine/Plugins/Runtime/RigVM/Source/RigVMDeveloper/Private/RigVMModel/RigVMController.cpp
--- a/Engine/Plugins/Runtime/RigVM/Source/RigVMDeveloper/Private/RigVMModel/RigVMController.cpp
+++ b/Engine/Plugins/Runtime/RigVM/Source/RigVMDeveloper/Private/RigVMModel/RigVMController.cpp
@@ -2384,7 +2384,10 @@ bool URigVMController::UnresolveTemplateNodes(const TArray<URigVMNode*>& InNodes
 	{
 		if (const URigVMTemplateNode* TemplateNode = Cast<URigVMTemplateNode>(Node))
 		{
-			return !TemplateNode->IsFullyUnresolved();
+			if (!TemplateNode->IsSingleton())
+			{
+				return !TemplateNode->IsFullyUnresolved();
+			}
 		}
 		return false;
 	});
