RigVM: Fix crash when copy-pasting event node
#rnx

diff --git a/Engine/Plugins/Runtime/RigVM/Source/RigVMDeveloper/Private/RigVMModel/RigVMController.cpp b/Engine/Plugins/Runtime/RigVM/Source/RigVMDeveloper/Private/RigVMModel/RigVMController.cpp
--- a/Engine/Plugins/Runtime/RigVM/Source/RigVMDeveloper/Private/RigVMModel/RigVMController.cpp
+++ b/Engine/Plugins/Runtime/RigVM/Source/RigVMDeveloper/Private/RigVMModel/RigVMController.cpp
@@ -1411,6 +1411,12 @@ URigVMUnitNode* URigVMController::AddUnitNode(UScriptStruct* InScriptStruct, con
 			return nullptr;
 		}
 
+		if (Graph->GetEventNames().Contains(StructMemory->GetEventName()))
+		{
+			ReportAndNotifyErrorf(TEXT("Event %s already exists in the graph."), *StructMemory->GetEventName().ToString());
+			return nullptr;
+		}
+
 		if(StructMemory->CanOnlyExistOnce())
 		{
 			// don't allow several event nodes in the main graph
diff --git a/Engine/Plugins/Runtime/RigVM/Source/RigVMDeveloper/Private/RigVMModel/RigVMGraph.cpp b/Engine/Plugins/Runtime/RigVM/Source/RigVMDeveloper/Private/RigVMModel/RigVMGraph.cpp
--- a/Engine/Plugins/Runtime/RigVM/Source/RigVMDeveloper/Private/RigVMModel/RigVMGraph.cpp
+++ b/Engine/Plugins/Runtime/RigVM/Source/RigVMDeveloper/Private/RigVMModel/RigVMGraph.cpp
@@ -130,6 +130,31 @@ URigVMFunctionReturnNode* URigVMGraph::GetReturnNode() const
 	return nullptr;
 }
 
+TArray<FName> URigVMGraph::GetEventNames() const
+{
+	TArray<FName> EventNames;
+	
+	if (!IsTopLevelGraph())
+	{
+		return EventNames;
+	}
+	
+	if (IsA<URigVMFunctionLibrary>())
+	{
+		return EventNames;
+	}
+
+	for (const URigVMNode* Node : Nodes)
+	{
+		if (Node->IsEvent())
+		{
+			EventNames.Add(Node->GetEventName());
+		}
+	}
+
+	return EventNames;
+}
+
 TArray<FRigVMGraphVariableDescription> URigVMGraph::GetVariableDescriptions() const
 {
 	TArray<FRigVMGraphVariableDescription> Variables;
diff --git a/Engine/Plugins/Runtime/RigVM/Source/RigVMDeveloper/Private/RigVMModel/RigVMSchema.cpp b/Engine/Plugins/Runtime/RigVM/Source/RigVMDeveloper/Private/RigVMModel/RigVMSchema.cpp
--- a/Engine/Plugins/Runtime/RigVM/Source/RigVMDeveloper/Private/RigVMModel/RigVMSchema.cpp
+++ b/Engine/Plugins/Runtime/RigVM/Source/RigVMDeveloper/Private/RigVMModel/RigVMSchema.cpp
@@ -388,6 +388,13 @@ bool URigVMSchema::CanAddNode(URigVMController* InController, const URigVMNode*
 					return false;
 				}
 
+				if (Graph->GetEventNames().Contains(InUnitNode->GetEventName()))
+				{
+					static const FString Message = FString::Printf(TEXT("An event named %s already exists."), *InUnitNode->GetEventName().ToString()); 
+					InController->ReportError(Message);
+					return false;
+				}
+
 				const TObjectPtr<URigVMNode> EventNode = FindEventNode(InController, EventStruct);
 				const bool bHasEventNode = (EventNode != nullptr) && EventNode->CanOnlyExistOnce();
 				if (bHasEventNode)
diff --git a/Engine/Plugins/Runtime/RigVM/Source/RigVMDeveloper/Public/RigVMModel/RigVMGraph.h b/Engine/Plugins/Runtime/RigVM/Source/RigVMDeveloper/Public/RigVMModel/RigVMGraph.h
--- a/Engine/Plugins/Runtime/RigVM/Source/RigVMDeveloper/Public/RigVMModel/RigVMGraph.h
+++ b/Engine/Plugins/Runtime/RigVM/Source/RigVMDeveloper/Public/RigVMModel/RigVMGraph.h
@@ -76,6 +76,10 @@ public:
 	UFUNCTION(BlueprintCallable, Category = RigVMGraph)
 	URigVMFunctionReturnNode* GetReturnNode() const;
 
+	// Returns array of event names
+	UFUNCTION(BlueprintCallable, Category = RigVMGraph)
+	TArray<FName> GetEventNames() const;
+
 	// Returns a list of unique Variable descriptions within this Graph.
 	// Multiple Variable Nodes can share the same description.
 	UFUNCTION(BlueprintCallable, Category = RigVMGraph)
