Control Rig: Fix crash when opeining F_Med_Ephemeral_CtrlRig
#rnx

diff --git a/Engine/Plugins/Runtime/RigVM/Source/RigVMDeveloper/Private/EdGraph/RigVMEdGraph.cpp b/Engine/Plugins/Runtime/RigVM/Source/RigVMDeveloper/Private/EdGraph/RigVMEdGraph.cpp
--- a/Engine/Plugins/Runtime/RigVM/Source/RigVMDeveloper/Private/EdGraph/RigVMEdGraph.cpp
+++ b/Engine/Plugins/Runtime/RigVM/Source/RigVMDeveloper/Private/EdGraph/RigVMEdGraph.cpp
@@ -158,6 +158,18 @@ bool URigVMEdGraph::HandleModifiedEvent_Internal(ERigVMGraphNotifType InNotifTyp
 		return false;
 	}
 
+	// Make sure this EdGraph has a valid blueprint
+	URigVMBlueprint* RigVMBlueprint = nullptr;
+	if (UBlueprint* Blueprint = FBlueprintEditorUtils::FindBlueprintForGraph(this))
+	{
+		RigVMBlueprint = Cast<URigVMBlueprint>(Blueprint);
+	}
+	if (!RigVMBlueprint)
+	{
+		return false;
+	}
+
+
 	if(URigVMEdGraphSchema* EdGraphSchema = (URigVMEdGraphSchema*)GetRigVMEdGraphSchema())
 	{
 		EdGraphSchema->HandleModifiedEvent(InNotifType, InGraph, InSubject);
diff --git a/Engine/Plugins/Runtime/RigVM/Source/RigVMDeveloper/Private/EdGraph/RigVMEdGraphSchema.cpp b/Engine/Plugins/Runtime/RigVM/Source/RigVMDeveloper/Private/EdGraph/RigVMEdGraphSchema.cpp
--- a/Engine/Plugins/Runtime/RigVM/Source/RigVMDeveloper/Private/EdGraph/RigVMEdGraphSchema.cpp
+++ b/Engine/Plugins/Runtime/RigVM/Source/RigVMDeveloper/Private/EdGraph/RigVMEdGraphSchema.cpp
@@ -2401,9 +2401,14 @@ void URigVMEdGraphSchema::HandleModifiedEvent(ERigVMGraphNotifType InNotifType,
 
 TSubclassOf<URigVMEdGraphNode> URigVMEdGraphSchema::GetGraphNodeClass(const URigVMEdGraph* InGraph) const
 {
-	const UBlueprint* Blueprint = FBlueprintEditorUtils::FindBlueprintForGraph(InGraph);
-	const URigVMBlueprint* RigBlueprint = CastChecked<URigVMBlueprint>(Blueprint);
-	return RigBlueprint->GetRigVMEdGraphNodeClass();
+	if (const UBlueprint* Blueprint = FBlueprintEditorUtils::FindBlueprintForGraph(InGraph))
+	{
+		if (const URigVMBlueprint* RigBlueprint = CastChecked<URigVMBlueprint>(Blueprint))
+		{
+			return RigBlueprint->GetRigVMEdGraphNodeClass();
+		}
+	}
+	return nullptr;
 }
 
 bool URigVMEdGraphSchema::IsRigVMDefaultEvent(const FName& InEventName) const
