Control Rig: Fix context menu empty
#rnx

diff --git a/Engine/Plugins/Runtime/RigVM/Source/RigVMEditor/Private/RigVMEditorModule.cpp b/Engine/Plugins/Runtime/RigVM/Source/RigVMEditor/Private/RigVMEditorModule.cpp
--- a/Engine/Plugins/Runtime/RigVM/Source/RigVMEditor/Private/RigVMEditorModule.cpp
+++ b/Engine/Plugins/Runtime/RigVM/Source/RigVMEditor/Private/RigVMEditorModule.cpp
@@ -493,7 +493,9 @@ void FRigVMEditorModule::GetInstanceActions(URigVMBlueprint* RigVMBlueprint, FBl
 void FRigVMEditorModule::GetNodeContextMenuActions(URigVMBlueprint* RigVMBlueprint, const URigVMEdGraphNode* EdGraphNode, URigVMNode* ModelNode, UToolMenu* Menu) const
 {
 	// Only register menu actions for ourselves
-	if(RigVMBlueprint->GetClass() != GetRigVMBlueprintClass())
+	UClass* BPClass = RigVMBlueprint->GetClass();
+	UClass* BaseClass = GetRigVMBlueprintClass();
+	if(BPClass != GetRigVMBlueprintClass() && !BPClass->IsChildOf(BaseClass))
 	{
 		return;
 	}
