RigVM: Fix functions not erroring when some paths do not end in a return node

diff --git a/Engine/Plugins/Animation/ControlRig/Config/BaseControlRig.ini b/Engine/Plugins/Animation/ControlRig/Config/BaseControlRig.ini
--- a/Engine/Plugins/Animation/ControlRig/Config/BaseControlRig.ini
+++ b/Engine/Plugins/Animation/ControlRig/Config/BaseControlRig.ini
@@ -387,6 +387,7 @@
 +EnumRedirects=(OldName="ECRSimPointIntegrateType",NewName="/Script/RigVM.ERigVMSimPointIntegrateType")
 +StructRedirects=(OldName="/Script/ControlRig.CRSimPoint",NewName="/Script/RigVM.RigVMSimPoint")
 +StructRedirects=(OldName="/Script/ControlRig.RigUnit_ForLoopCount",NewName="/Script/RigVM.RigVMFunction_ForLoopCount")
++StructRedirects=(OldName="/Script/ControlRig.RigUnit_SequenceAggregate",NewName="/Script/RigVM.RigVMFunction_Sequence")
 +StructRedirects=(OldName="/Script/ControlRig.RigUnit_UserDefinedEvent",NewName="/Script/RigVM.RigVMFunction_UserDefinedEvent")
 +StructRedirects=(OldName="/Script/ControlRig.RigUnit_AccumulateBase",NewName="/Script/RigVM.RigVMFunction_AccumulateBase")
 +StructRedirects=(OldName="/Script/ControlRig.RigUnit_AccumulateFloatAdd",NewName="/Script/RigVM.RigVMFunction_AccumulateFloatAdd")
diff --git a/Engine/Plugins/Animation/ControlRig/Source/ControlRig/Private/Units/Execution/RigUnit_SequenceExecution.cpp b/Engine/Plugins/Animation/ControlRig/Source/ControlRig/Private/Units/Execution/RigUnit_SequenceExecution.cpp
--- a/Engine/Plugins/Animation/ControlRig/Source/ControlRig/Private/Units/Execution/RigUnit_SequenceExecution.cpp
+++ b/Engine/Plugins/Animation/ControlRig/Source/ControlRig/Private/Units/Execution/RigUnit_SequenceExecution.cpp
@@ -1,6 +1,7 @@
 // Copyright Epic Games, Inc. All Rights Reserved.
 
 #include "Units/Execution/RigUnit_SequenceExecution.h"
+#include "RigVMFunctions/Execution/RigVMFunction_Sequence.h"
 #include "Units/RigUnitContext.h"
 
 #include UE_INLINE_GENERATED_CPP_BY_NAME(RigUnit_SequenceExecution)
@@ -13,7 +14,7 @@ FRigUnit_SequenceExecution_Execute()
 
 FRigVMStructUpgradeInfo FRigUnit_SequenceExecution::GetUpgradeInfo() const
 {
-	FRigUnit_SequenceAggregate NewNode;
+	FRigVMFunction_Sequence NewNode;
 	
 	FRigVMStructUpgradeInfo Info(*this, NewNode);
 
@@ -24,9 +25,4 @@ FRigVMStructUpgradeInfo FRigUnit_SequenceExecution::GetUpgradeInfo() const
 	return Info;
 }
 
-FRigUnit_SequenceAggregate_Execute()
-{
-	// nothing to do here. the execute context is actually
-	// the same shared memory for all pins
-}
 
diff --git a/Engine/Plugins/Animation/ControlRig/Source/ControlRig/Private/Units/Execution/RigUnit_SequenceExecution.h b/Engine/Plugins/Animation/ControlRig/Source/ControlRig/Private/Units/Execution/RigUnit_SequenceExecution.h
--- a/Engine/Plugins/Animation/ControlRig/Source/ControlRig/Private/Units/Execution/RigUnit_SequenceExecution.h
+++ b/Engine/Plugins/Animation/ControlRig/Source/ControlRig/Private/Units/Execution/RigUnit_SequenceExecution.h
@@ -41,26 +41,3 @@ struct CONTROLRIG_API FRigUnit_SequenceExecution : public FRigUnit
 	virtual FRigVMStructUpgradeInfo GetUpgradeInfo() const override;
 };
 
-/**
- * Allows for a single execution pulse to trigger a series of events in order.
- */
-USTRUCT(meta=(DisplayName="Sequence", Category="Execution", TitleColor="1 0 0", NodeColor="1 1 1", Icon="EditorStyle|GraphEditor.Sequence_16x"))
-struct CONTROLRIG_API FRigUnit_SequenceAggregate : public FRigUnit
-{
-	GENERATED_BODY()
-
-	RIGVM_METHOD()
-	virtual void Execute() override;
-
-	// The execution input
-	UPROPERTY(EditAnywhere, Transient, DisplayName = "Execute", Category = "SequenceExecution", meta = (Input, Aggregate))
-	FControlRigExecuteContext ExecuteContext;
-
-	// The execution result A
-	UPROPERTY(EditAnywhere, Transient, Category = "SequenceExecution", meta = (Output, Aggregate))
-	FControlRigExecuteContext A;
-
-	// The execution result B
-	UPROPERTY(EditAnywhere, Transient, Category = "SequenceExecution", meta = (Output, Aggregate))
-	FControlRigExecuteContext B;
-};
diff --git a/Engine/Plugins/Animation/ControlRig/Source/ControlRigEditor/Private/Editor/ControlRigEditor.cpp b/Engine/Plugins/Animation/ControlRig/Source/ControlRigEditor/Private/Editor/ControlRigEditor.cpp
--- a/Engine/Plugins/Animation/ControlRig/Source/ControlRigEditor/Private/Editor/ControlRigEditor.cpp
+++ b/Engine/Plugins/Animation/ControlRig/Source/ControlRigEditor/Private/Editor/ControlRigEditor.cpp
@@ -93,7 +93,7 @@
 #include "Engine/StaticMesh.h"
 #include "Engine/StaticMeshActor.h"
 #include "Components/StaticMeshComponent.h"
-#include "Units/Execution/RigUnit_SequenceExecution.h"
+#include "RigVMFunctions/Execution/RigVMFunction_Sequence.h"
 #include "Editor/ControlRigContextMenuContext.h"
 #include "Types/ISlateMetaData.h"
 #include "Kismet2/KismetDebugUtilities.h"
@@ -1085,7 +1085,7 @@ FReply FControlRigEditor::OnSpawnGraphNodeByShortcut(FInputChord InChord, const
 			{
 				if(InChord.Key == EKeys::S)
 				{
-					Controller->AddUnitNode(FRigUnit_SequenceAggregate::StaticStruct(), FRigUnit::GetMethodName(), InPosition, FString(), true, true);
+					Controller->AddUnitNode(FRigVMFunction_Sequence::StaticStruct(), FRigUnit::GetMethodName(), InPosition, FString(), true, true);
 				}
 				else if(InChord.Key == EKeys::One)
 				{
diff --git a/Engine/Plugins/Runtime/RigVM/Source/RigVM/Private/RigVMFunctions/Execution/RigVMFunction_Sequence.cpp b/Engine/Plugins/Runtime/RigVM/Source/RigVM/Private/RigVMFunctions/Execution/RigVMFunction_Sequence.cpp
new file mode 100644
--- /dev/null
+++ b/Engine/Plugins/Runtime/RigVM/Source/RigVM/Private/RigVMFunctions/Execution/RigVMFunction_Sequence.cpp
@@ -0,0 +1,11 @@
+// Copyright Epic Games, Inc. All Rights Reserved.
+
+#include "RigVMFunctions/Execution/RigVMFunction_Sequence.h"
+
+FRigVMFunction_Sequence_Execute()
+{
+	// nothing to do here. the execute context is actually
+	// the same shared memory for all pins
+}
+
+
diff --git a/Engine/Plugins/Runtime/RigVM/Source/RigVM/Public/RigVMFunctions/Execution/RigVMFunction_Sequence.h b/Engine/Plugins/Runtime/RigVM/Source/RigVM/Public/RigVMFunctions/Execution/RigVMFunction_Sequence.h
new file mode 100644
--- /dev/null
+++ b/Engine/Plugins/Runtime/RigVM/Source/RigVM/Public/RigVMFunctions/Execution/RigVMFunction_Sequence.h
@@ -0,0 +1,30 @@
+// Copyright Epic Games, Inc. All Rights Reserved.
+
+#pragma once
+
+#include "RigVMCore/RigVMStruct.h"
+#include "RigVMFunction_Sequence.generated.h"
+
+/**
+ * Allows for a single execution pulse to trigger a series of events in order.
+ */
+USTRUCT(meta=(DisplayName="Sequence", Category="Execution", TitleColor="1 0 0", NodeColor="1 1 1", Icon="EditorStyle|GraphEditor.Sequence_16x"))
+struct RIGVM_API FRigVMFunction_Sequence : public FRigVMStruct
+{
+	GENERATED_BODY()
+
+	RIGVM_METHOD()
+	virtual void Execute() override;
+
+	// The execution input
+	UPROPERTY(EditAnywhere, Transient, DisplayName = "Execute", Category = "SequenceExecution", meta = (Input, Aggregate))
+	FRigVMExecuteContext ExecuteContext;
+
+	// The execution result A
+	UPROPERTY(EditAnywhere, Transient, Category = "SequenceExecution", meta = (Output, Aggregate))
+	FRigVMExecuteContext A;
+
+	// The execution result B
+	UPROPERTY(EditAnywhere, Transient, Category = "SequenceExecution", meta = (Output, Aggregate))
+	FRigVMExecuteContext B;
+};
\ No newline at end of file
diff --git a/Engine/Plugins/Runtime/RigVM/Source/RigVMDeveloper/Private/RigVMCompiler/RigVMCompiler.cpp b/Engine/Plugins/Runtime/RigVM/Source/RigVMDeveloper/Private/RigVMCompiler/RigVMCompiler.cpp
--- a/Engine/Plugins/Runtime/RigVM/Source/RigVMDeveloper/Private/RigVMCompiler/RigVMCompiler.cpp
+++ b/Engine/Plugins/Runtime/RigVM/Source/RigVMDeveloper/Private/RigVMCompiler/RigVMCompiler.cpp
@@ -825,6 +825,80 @@ bool URigVMCompiler::Compile(TArray<URigVMGraph*> InGraphs, URigVMController* In
 		}
 	}
 
+	// If compiling a function, check if all mutable paths lead to a return node
+	// - Return node path should go through the completed pins of control flow nodes
+	// - At least one path on sequence nodes should lead to a return node
+	if (CurrentCompilationFunction)
+	{
+		const URigVMFunctionEntryNode* EntryNode = CurrentCompilationFunction->GetEntryNode();
+		if (EntryNode && EntryNode->IsMutable())
+		{
+			URigVMFunctionReturnNode* ReturnNode = CurrentCompilationFunction->GetReturnNode();
+			if (!ReturnNode)
+			{
+				const FString FunctionCompilationErrorMessage = FString::Printf(TEXT("Mutable function graph %s does not contain a return node."), *CurrentCompilationFunction->GetName());
+				Settings.ASTSettings.Report(EMessageSeverity::Error, nullptr, FunctionCompilationErrorMessage);
+				bEncounteredGraphError = true;
+			}
+
+			bool bReturnNodeFound = false;
+			TArray<URigVMPin*> Stack;
+			URigVMPin* const* EntryExecutePin = EntryNode->GetPins().FindByPredicate([](URigVMPin* Pin) -> bool
+				{
+					return Pin->IsExecuteContext();
+				}
+			);
+			Stack.Push(*EntryExecutePin);
+			while (!Stack.IsEmpty())
+			{
+				URigVMPin* Pin = Stack.Pop();
+				URigVMNode* Node = Pin->GetNode();
+				if (Node == ReturnNode)
+				{
+					bReturnNodeFound = true;
+					break;
+				}
+
+				URigVMPin* OutputPin = Pin;
+				if (Node->IsControlFlowNode())
+				{
+					OutputPin = Node->FindPin(FRigVMStruct::ControlFlowCompletedName.ToString());
+				}
+
+				// If this is a sequence node, we need to go through all output pins
+				if (OutputPin->GetDirection() == ERigVMPinDirection::Input)
+				{
+					TArray<URigVMPin*> SequenceOutputPins = Node->GetPins().FilterByPredicate([](URigVMPin* Pin) -> bool
+					{
+						return Pin->IsExecuteContext() && Pin->GetDirection() == ERigVMPinDirection::Output;
+					});
+					for (URigVMPin* SequenceOutputPin : SequenceOutputPins)
+					{
+						Stack.Push(SequenceOutputPin);
+					}
+				}
+				else
+				{
+					check(OutputPin);
+					check(OutputPin->GetDirection() == ERigVMPinDirection::Output || OutputPin->GetDirection() == ERigVMPinDirection::IO);
+				
+					TArray<URigVMPin*> TargetPins = OutputPin->GetLinkedTargetPins();
+					if (!TargetPins.IsEmpty())
+					{
+						Stack.Push(TargetPins[0]);
+					}
+				}
+			}
+
+			if (!bReturnNodeFound)
+			{
+				const FString FunctionCompilationErrorMessage = FString::Printf(TEXT("Not all paths of the mutable function graph %s lead to a return node."), *CurrentCompilationFunction->GetName());
+				Settings.ASTSettings.Report(EMessageSeverity::Error, nullptr, FunctionCompilationErrorMessage);
+				bEncounteredGraphError = true;
+			}
+		}
+	}
+
 	if(bEncounteredGraphError)
 	{
 		return false;
