RigVM: Fix links not restored when undoing remove exposed pin in function
#rnx

diff --git a/Engine/Plugins/Runtime/RigVM/Source/RigVMDeveloper/Private/RigVMModel/RigVMController.cpp b/Engine/Plugins/Runtime/RigVM/Source/RigVMDeveloper/Private/RigVMModel/RigVMController.cpp
--- a/Engine/Plugins/Runtime/RigVM/Source/RigVMDeveloper/Private/RigVMModel/RigVMController.cpp
+++ b/Engine/Plugins/Runtime/RigVM/Source/RigVMDeveloper/Private/RigVMModel/RigVMController.cpp
@@ -6312,7 +6312,7 @@ void URigVMController::SetReferencedFunction(URigVMFunctionReferenceNode* InFunc
 	}
 }
 
-void URigVMController::RefreshFunctionPins(URigVMNode* InNode)
+void URigVMController::RefreshFunctionPins(URigVMNode* InNode, bool bSetupUndoRedo)
 {
 	if (InNode == nullptr)
 	{
@@ -6325,10 +6325,10 @@ void URigVMController::RefreshFunctionPins(URigVMNode* InNode)
 	if (EntryNode || ReturnNode)
 	{
 		const TArray<FLinkedPath> LinkedPaths = GetLinkedPaths(InNode);
-		FastBreakLinkedPaths(LinkedPaths);
+		FastBreakLinkedPaths(LinkedPaths, bSetupUndoRedo);
 		RepopulatePinsOnNode(InNode, false, false, true);
 		FRestoreLinkedPathSettings Settings;
-		RestoreLinkedPaths(LinkedPaths, Settings);
+		RestoreLinkedPaths(LinkedPaths, Settings, bSetupUndoRedo);
 	}
 }
 
@@ -10662,9 +10662,9 @@ bool URigVMController::RemoveExposedPin(const FName& InPinName, bool bSetupUndoR
 			RemoveNode(NodesToRemove[i], bSetupUndoRedo);
 		}
 
-		RefreshFunctionPins(Graph->GetEntryNode());
-		RefreshFunctionPins(Graph->GetReturnNode());
-		RefreshFunctionReferences(LibraryNode, false);
+		RefreshFunctionPins(Graph->GetEntryNode(), bSetupUndoRedo);
+		RefreshFunctionPins(Graph->GetReturnNode(), bSetupUndoRedo);
+		RefreshFunctionReferences(LibraryNode, bSetupUndoRedo);
 	}
 
 	if (bSetupUndoRedo)
@@ -17949,15 +17949,15 @@ void URigVMController::RefreshFunctionReferences(URigVMLibraryNode* InFunctionDe
 
 	if (const URigVMFunctionLibrary* FunctionLibrary = Cast<URigVMFunctionLibrary>(InFunctionDefinition->GetGraph()))
 	{
-		FunctionLibrary->ForEachReference(InFunctionDefinition->GetFName(), [this](URigVMFunctionReferenceNode* ReferenceNode)
+		FunctionLibrary->ForEachReference(InFunctionDefinition->GetFName(), [this, bSetupUndoRedo](URigVMFunctionReferenceNode* ReferenceNode)
 		{
 			if(URigVMController* ReferenceController = GetControllerForGraph(ReferenceNode->GetGraph()))
 			{
 				const TArray<FLinkedPath> LinkedPaths = GetLinkedPaths(ReferenceNode->GetLinks());
-				ReferenceController->FastBreakLinkedPaths(LinkedPaths);
+				ReferenceController->FastBreakLinkedPaths(LinkedPaths, bSetupUndoRedo);
 				ReferenceController->RepopulatePinsOnNode(ReferenceNode, false, false, true);
 				TGuardValue<bool> ReportGuard(ReferenceController->bReportWarningsAndErrors, false);
-				ReferenceController->RestoreLinkedPaths(LinkedPaths);
+				ReferenceController->RestoreLinkedPaths(LinkedPaths, FRestoreLinkedPathSettings(), bSetupUndoRedo);
 			}
 		});
 	}
diff --git a/Engine/Plugins/Runtime/RigVM/Source/RigVMDeveloper/Public/RigVMModel/RigVMController.h b/Engine/Plugins/Runtime/RigVM/Source/RigVMDeveloper/Public/RigVMModel/RigVMController.h
--- a/Engine/Plugins/Runtime/RigVM/Source/RigVMDeveloper/Public/RigVMModel/RigVMController.h
+++ b/Engine/Plugins/Runtime/RigVM/Source/RigVMDeveloper/Public/RigVMModel/RigVMController.h
@@ -1154,7 +1154,7 @@ private:
 	URigVMCollapseNode* PromoteFunctionReferenceNodeToCollapseNode(URigVMFunctionReferenceNode* InFunctionRefNode, bool bSetupUndoRedo, bool bRemoveFunctionDefinition);
 	void SetReferencedFunction(URigVMFunctionReferenceNode* InFunctionRefNode, URigVMLibraryNode* InNewReferencedNode, bool bSetupUndoRedo);
 
-	void RefreshFunctionPins(URigVMNode* InNode);
+	void RefreshFunctionPins(URigVMNode* InNode, bool bSetupUndoRedo = false);
 
 	void ReportRemovedLink(const FString& InSourcePinPath, const FString& InTargetPinPath, const FString& Reason = FString());
 
