RigVMEditor: Allow to identify node, pin and connection factories

diff --git a/Engine/Plugins/Animation/ControlRig/Source/ControlRigDeveloper/Private/ControlRigBlueprint.cpp b/Engine/Plugins/Animation/ControlRig/Source/ControlRigDeveloper/Private/ControlRigBlueprint.cpp
--- a/Engine/Plugins/Animation/ControlRig/Source/ControlRigDeveloper/Private/ControlRigBlueprint.cpp
+++ b/Engine/Plugins/Animation/ControlRig/Source/ControlRigDeveloper/Private/ControlRigBlueprint.cpp
@@ -39,6 +39,9 @@
 #define LOCTEXT_NAMESPACE "ControlRigBlueprint"
 
 TArray<UControlRigBlueprint*> UControlRigBlueprint::sCurrentlyOpenedRigBlueprints;
+#if WITH_EDITOR
+const FName UControlRigBlueprint::ControlRigPanelNodeFactoryName(TEXT("FControlRigGraphPanelPinFactory"));
+#endif
 
 UControlRigBlueprint::UControlRigBlueprint(const FObjectInitializer& ObjectInitializer)
 	: URigVMBlueprint(ObjectInitializer)
@@ -393,6 +396,11 @@ UClass* UControlRigBlueprint::GetRigVMEditorSettingsClass() const
 }
 
 #if WITH_EDITOR
+const FName& UControlRigBlueprint::GetPanelPinFactoryName() const
+{
+	return ControlRigPanelNodeFactoryName;
+}
+
 IRigVMEditorModule* UControlRigBlueprint::GetEditorModule() const
 {
 	return &IControlRigEditorModule::Get();
diff --git a/Engine/Plugins/Animation/ControlRig/Source/ControlRigDeveloper/Public/ControlRigBlueprint.h b/Engine/Plugins/Animation/ControlRig/Source/ControlRigDeveloper/Public/ControlRigBlueprint.h
--- a/Engine/Plugins/Animation/ControlRig/Source/ControlRigDeveloper/Public/ControlRigBlueprint.h
+++ b/Engine/Plugins/Animation/ControlRig/Source/ControlRigDeveloper/Public/ControlRigBlueprint.h
@@ -47,6 +47,8 @@ public:
 	virtual TArray<FString> GeneratePythonCommands(const FString InNewBlueprintName) override;
 	virtual UClass* GetRigVMEditorSettingsClass() const override;
 #if WITH_EDITOR
+	virtual const FName& GetPanelPinFactoryName() const override;
+	static const FName ControlRigPanelNodeFactoryName;
 	virtual IRigVMEditorModule* GetEditorModule() const override;
 #endif
 
diff --git a/Engine/Plugins/Animation/ControlRig/Source/ControlRigEditor/Private/Graph/ControlRigGraphPanelPinFactory.cpp b/Engine/Plugins/Animation/ControlRig/Source/ControlRigEditor/Private/Graph/ControlRigGraphPanelPinFactory.cpp
--- a/Engine/Plugins/Animation/ControlRig/Source/ControlRigEditor/Private/Graph/ControlRigGraphPanelPinFactory.cpp
+++ b/Engine/Plugins/Animation/ControlRig/Source/ControlRigEditor/Private/Graph/ControlRigGraphPanelPinFactory.cpp
@@ -107,6 +107,11 @@ TSharedPtr<SGraphPin> FControlRigGraphPanelPinFactory::CreatePin(UEdGraphPin* In
 	return K2PinWidget;
 }
 
+FName FControlRigGraphPanelPinFactory::GetFactoryName() const
+{
+	return UControlRigBlueprint::ControlRigPanelNodeFactoryName;
+}
+
 TSharedPtr<SGraphPin> FControlRigGraphPanelPinFactory::CreatePin_Internal(UEdGraphPin* InPin) const
 {
 	TSharedPtr<SGraphPin> SuperResult = FRigVMEdGraphPanelPinFactory::CreatePin_Internal(InPin);
diff --git a/Engine/Plugins/Animation/ControlRig/Source/ControlRigEditor/Private/Graph/ControlRigGraphPanelPinFactory.h b/Engine/Plugins/Animation/ControlRig/Source/ControlRigEditor/Private/Graph/ControlRigGraphPanelPinFactory.h
--- a/Engine/Plugins/Animation/ControlRig/Source/ControlRigEditor/Private/Graph/ControlRigGraphPanelPinFactory.h
+++ b/Engine/Plugins/Animation/ControlRig/Source/ControlRigEditor/Private/Graph/ControlRigGraphPanelPinFactory.h
@@ -11,5 +11,6 @@ public:
 	virtual TSharedPtr<class SGraphPin> CreatePin(class UEdGraphPin* InPin) const override;
 
 	// FRigVMEdGraphPanelPinFactory interface
+	virtual FName GetFactoryName() const override;
 	virtual TSharedPtr<class SGraphPin> CreatePin_Internal(class UEdGraphPin* InPin) const override;
 };
diff --git a/Engine/Plugins/Runtime/RigVM/Source/RigVMDeveloper/Private/RigVMBlueprint.cpp b/Engine/Plugins/Runtime/RigVM/Source/RigVMDeveloper/Private/RigVMBlueprint.cpp
--- a/Engine/Plugins/Runtime/RigVM/Source/RigVMDeveloper/Private/RigVMBlueprint.cpp
+++ b/Engine/Plugins/Runtime/RigVM/Source/RigVMDeveloper/Private/RigVMBlueprint.cpp
@@ -103,6 +103,10 @@ bool FRigVMOldPublicFunctionData::IsMutable() const
 FSoftObjectPath URigVMBlueprint::PreDuplicateAssetPath;
 FSoftObjectPath URigVMBlueprint::PreDuplicateHostPath;
 TArray<URigVMBlueprint*> URigVMBlueprint::sCurrentlyOpenedRigVMBlueprints;
+#if WITH_EDITOR
+const FName URigVMBlueprint::RigVMPanelNodeFactoryName(TEXT("FRigVMEdGraphPanelNodeFactory"));
+const FName URigVMBlueprint::RigVMPanelPinFactoryName(TEXT("FRigVMEdGraphPanelPinFactory"));
+#endif
 
 URigVMBlueprint::URigVMBlueprint()
 {
@@ -685,6 +689,16 @@ URigVMEditorSettings* URigVMBlueprint::GetRigVMEditorSettings() const
 }
 
 #if WITH_EDITOR
+const FName& URigVMBlueprint::GetPanelNodeFactoryName() const
+{
+	return RigVMPanelNodeFactoryName;
+}
+
+const FName& URigVMBlueprint::GetPanelPinFactoryName() const
+{
+	return RigVMPanelPinFactoryName;
+}
+
 IRigVMEditorModule* URigVMBlueprint::GetEditorModule() const
 {
 	return &IRigVMEditorModule::Get();
diff --git a/Engine/Plugins/Runtime/RigVM/Source/RigVMDeveloper/Public/RigVMBlueprint.h b/Engine/Plugins/Runtime/RigVM/Source/RigVMDeveloper/Public/RigVMBlueprint.h
--- a/Engine/Plugins/Runtime/RigVM/Source/RigVMDeveloper/Public/RigVMBlueprint.h
+++ b/Engine/Plugins/Runtime/RigVM/Source/RigVMDeveloper/Public/RigVMBlueprint.h
@@ -238,6 +238,15 @@ public:
 	URigVMEditorSettings* GetRigVMEditorSettings() const;
 
 #if WITH_EDITOR
+	/** Returns true if a given panel node factory is compatible this blueprint */
+	virtual const FName& GetPanelNodeFactoryName() const;
+
+	/** Returns true if a given panel pin factory is compatible this blueprint */
+	virtual const FName& GetPanelPinFactoryName() const;
+
+	static const FName RigVMPanelNodeFactoryName;
+	static const FName RigVMPanelPinFactoryName;
+
 	/** Returns the editor module to be used for this blueprint */
 	virtual IRigVMEditorModule* GetEditorModule() const;
 #endif
diff --git a/Engine/Plugins/Runtime/RigVM/Source/RigVMEditor/Private/EdGraph/RigVMEdGraphPanelNodeFactory.cpp b/Engine/Plugins/Runtime/RigVM/Source/RigVMEditor/Private/EdGraph/RigVMEdGraphPanelNodeFactory.cpp
--- a/Engine/Plugins/Runtime/RigVM/Source/RigVMEditor/Private/EdGraph/RigVMEdGraphPanelNodeFactory.cpp
+++ b/Engine/Plugins/Runtime/RigVM/Source/RigVMEditor/Private/EdGraph/RigVMEdGraphPanelNodeFactory.cpp
@@ -11,8 +11,22 @@
 #include "RigVMModel/RigVMNode.h"
 #include "RigVMModel/Nodes/RigVMRerouteNode.h"
 
+FName FRigVMEdGraphPanelNodeFactory::GetFactoryName() const
+{
+	return URigVMBlueprint::RigVMPanelNodeFactoryName;
+}
+
 TSharedPtr<SGraphNode> FRigVMEdGraphPanelNodeFactory::CreateNode(UEdGraphNode* Node) const
 {
+	// we need to check if this is the right factory for the implementation
+	if(const URigVMBlueprint* Blueprint = Node->GetTypedOuter<URigVMBlueprint>())
+	{
+		if(Blueprint->GetPanelNodeFactoryName() != GetFactoryName())
+		{
+			return nullptr;
+		}
+	}
+
 	if (URigVMEdGraphNode* RigVMEdGraphNode = Cast<URigVMEdGraphNode>(Node))
 	{
 		int32 InputPin = 0, OutputPin = 0;
diff --git a/Engine/Plugins/Runtime/RigVM/Source/RigVMEditor/Private/EdGraph/RigVMEdGraphPanelPinFactory.cpp b/Engine/Plugins/Runtime/RigVM/Source/RigVMEditor/Private/EdGraph/RigVMEdGraphPanelPinFactory.cpp
--- a/Engine/Plugins/Runtime/RigVM/Source/RigVMEditor/Private/EdGraph/RigVMEdGraphPanelPinFactory.cpp
+++ b/Engine/Plugins/Runtime/RigVM/Source/RigVMEditor/Private/EdGraph/RigVMEdGraphPanelPinFactory.cpp
@@ -19,13 +19,31 @@
 #include "RigVMModel/Nodes/RigVMUnitNode.h"
 #include "RigVMCore/RigVMExecuteContext.h"
 
+FName FRigVMEdGraphPanelPinFactory::GetFactoryName() const
+{
+	return URigVMBlueprint::RigVMPanelPinFactoryName;
+}
+
 TSharedPtr<SGraphPin> FRigVMEdGraphPanelPinFactory::CreatePin(UEdGraphPin* InPin) const
 {
+	// we need to check if this is the right factory for the implementation
+	if(const UEdGraphNode* EdGraphNode = InPin->GetOuter())
+	{
+		if(const URigVMBlueprint* Blueprint = EdGraphNode->GetTypedOuter<URigVMBlueprint>())
+		{
+			if(Blueprint->GetPanelPinFactoryName() != GetFactoryName())
+			{
+				return nullptr;
+			}
+		}
+	}
+
 	TSharedPtr<SGraphPin> InternalResult = CreatePin_Internal(InPin);
 	if(InternalResult.IsValid())
 	{
 		return InternalResult;
 	}
+
 	return FNodeFactory::CreateK2PinWidget(InPin);
 }
 
diff --git a/Engine/Plugins/Runtime/RigVM/Source/RigVMEditor/Public/EdGraph/RigVMEdGraphPanelNodeFactory.h b/Engine/Plugins/Runtime/RigVM/Source/RigVMEditor/Public/EdGraph/RigVMEdGraphPanelNodeFactory.h
--- a/Engine/Plugins/Runtime/RigVM/Source/RigVMEditor/Public/EdGraph/RigVMEdGraphPanelNodeFactory.h
+++ b/Engine/Plugins/Runtime/RigVM/Source/RigVMEditor/Public/EdGraph/RigVMEdGraphPanelNodeFactory.h
@@ -6,5 +6,8 @@
 
 class RIGVMEDITOR_API FRigVMEdGraphPanelNodeFactory : public FGraphPanelNodeFactory
 {
+public:
+	virtual FName GetFactoryName() const;
+private:
 	virtual TSharedPtr<class SGraphNode> CreateNode(UEdGraphNode* Node) const override;
 };
diff --git a/Engine/Plugins/Runtime/RigVM/Source/RigVMEditor/Public/EdGraph/RigVMEdGraphPanelPinFactory.h b/Engine/Plugins/Runtime/RigVM/Source/RigVMEditor/Public/EdGraph/RigVMEdGraphPanelPinFactory.h
--- a/Engine/Plugins/Runtime/RigVM/Source/RigVMEditor/Public/EdGraph/RigVMEdGraphPanelPinFactory.h
+++ b/Engine/Plugins/Runtime/RigVM/Source/RigVMEditor/Public/EdGraph/RigVMEdGraphPanelPinFactory.h
@@ -8,6 +8,8 @@
 class RIGVMEDITOR_API FRigVMEdGraphPanelPinFactory : public FGraphPanelPinFactory
 {
 public:
+	virtual FName GetFactoryName() const;
+	
 	// FGraphPanelPinFactory interface
 	virtual TSharedPtr<class SGraphPin> CreatePin(class UEdGraphPin* InPin) const override;
 
