RigVM: Always add a return node when collapsing nodes

diff --git a/Engine/Plugins/Runtime/RigVM/Source/RigVMDeveloper/Private/RigVMModel/RigVMController.cpp b/Engine/Plugins/Runtime/RigVM/Source/RigVMDeveloper/Private/RigVMModel/RigVMController.cpp
--- a/Engine/Plugins/Runtime/RigVM/Source/RigVMDeveloper/Private/RigVMModel/RigVMController.cpp
+++ b/Engine/Plugins/Runtime/RigVM/Source/RigVMDeveloper/Private/RigVMModel/RigVMController.cpp
@@ -4758,8 +4758,6 @@ URigVMCollapseNode* URigVMController::CollapseNodes(const TArray<URigVMNode*>& I
   	FVector2D Diagonal = Bounds.Max - Bounds.Min;
 	FVector2D Center = (Bounds.Min + Bounds.Max) * 0.5f;
 
-	bool bContainsOutputs = false;
-
 	TArray<URigVMPin*> PinsToCollapse;
 	TMap<URigVMPin*, URigVMPin*> CollapsedPins;
 	TArray<URigVMLink*> LinksToRewire;
@@ -4867,7 +4865,6 @@ URigVMCollapseNode* URigVMController::CollapseNodes(const TArray<URigVMNode*>& I
 	for (URigVMLink* Link : LinksToRewire)
 	{
 		bool bSourceToBeCollapsed = NodeToBeCollapsed(Link->GetSourcePin()->GetNode());
-		bContainsOutputs = bContainsOutputs || bSourceToBeCollapsed;
 
 		URigVMPin* PinToCollapse = bSourceToBeCollapsed ? Link->GetSourcePin() : Link->GetTargetPin();
 		if (CollapsedPins.Contains(PinToCollapse))
@@ -4922,11 +4919,7 @@ URigVMCollapseNode* URigVMController::CollapseNodes(const TArray<URigVMNode*>& I
 
 		if (CollapsedPin->GetDirection() == ERigVMPinDirection::IO)
 		{
-			if(CollapsedPin->IsExecuteContext())
-			{
-				bContainsOutputs = true;
-			}
-			else
+			if(!CollapsedPin->IsExecuteContext())
 			{
 				CollapsedPin->Direction = bSourceToBeCollapsed ? ERigVMPinDirection::Output : ERigVMPinDirection::Input;
 			}
@@ -4967,20 +4960,17 @@ URigVMCollapseNode* URigVMController::CollapseNodes(const TArray<URigVMNode*>& I
 			CollapseController->Notify(ERigVMGraphNotifType::NodeAdded, EntryNode);
 		}
 
-		if (bContainsOutputs)
-		{
-			ReturnNode = NewObject<URigVMFunctionReturnNode>(CollapseNode->ContainedGraph, TEXT("Return"));
+		ReturnNode = NewObject<URigVMFunctionReturnNode>(CollapseNode->ContainedGraph, TEXT("Return"));
 
-			if(CollapseController->AddGraphNode(ReturnNode, false))
+		if(CollapseController->AddGraphNode(ReturnNode, false))
+		{
+			ReturnNode->Position = FVector2D(Diagonal.X, -Diagonal.Y) * 0.5f + FVector2D(300.f, 0.f);
 			{
-				ReturnNode->Position = FVector2D(Diagonal.X, -Diagonal.Y) * 0.5f + FVector2D(300.f, 0.f);
-				{
-					TGuardValue<bool> SuspendNotifications(CollapseController->bSuspendNotifications, true);
-					CollapseController->RefreshFunctionPins(ReturnNode);
-				}
-
-				CollapseController->Notify(ERigVMGraphNotifType::NodeAdded, ReturnNode);
+				TGuardValue<bool> SuspendNotifications(CollapseController->bSuspendNotifications, true);
+				CollapseController->RefreshFunctionPins(ReturnNode);
 			}
+
+			CollapseController->Notify(ERigVMGraphNotifType::NodeAdded, ReturnNode);
 		}
 	}
 
