RigVM: Temporarily disable error when not all paths lead to a return node
#rnx

diff --git a/Engine/Plugins/Runtime/RigVM/Source/RigVMDeveloper/Private/RigVMCompiler/RigVMCompiler.cpp b/Engine/Plugins/Runtime/RigVM/Source/RigVMDeveloper/Private/RigVMCompiler/RigVMCompiler.cpp
--- a/Engine/Plugins/Runtime/RigVM/Source/RigVMDeveloper/Private/RigVMCompiler/RigVMCompiler.cpp
+++ b/Engine/Plugins/Runtime/RigVM/Source/RigVMDeveloper/Private/RigVMCompiler/RigVMCompiler.cpp
@@ -890,12 +890,12 @@ bool URigVMCompiler::Compile(TArray<URigVMGraph*> InGraphs, URigVMController* In
 				}
 			}
 
-			if (!bReturnNodeFound)
-			{
-				const FString FunctionCompilationErrorMessage = FString::Printf(TEXT("Not all paths of the mutable function graph %s lead to a return node."), *CurrentCompilationFunction->GetName());
-				Settings.ASTSettings.Report(EMessageSeverity::Error, nullptr, FunctionCompilationErrorMessage);
-				bEncounteredGraphError = true;
-			}
+			// if (!bReturnNodeFound)
+			// {
+			// 	const FString FunctionCompilationErrorMessage = FString::Printf(TEXT("Not all paths of the mutable function graph %s lead to a return node."), *CurrentCompilationFunction->GetName());
+			// 	Settings.ASTSettings.Report(EMessageSeverity::Error, nullptr, FunctionCompilationErrorMessage);
+			// 	bEncounteredGraphError = true;
+			// }
 		}
 	}
 
