RigVM: Wrapper objects not displayed correctly in details panel

diff --git a/Engine/Plugins/Runtime/RigVM/Source/RigVMEditor/Private/Editor/RigVMGraphDetailCustomization.cpp b/Engine/Plugins/Runtime/RigVM/Source/RigVMEditor/Private/Editor/RigVMGraphDetailCustomization.cpp
--- a/Engine/Plugins/Runtime/RigVM/Source/RigVMEditor/Private/Editor/RigVMGraphDetailCustomization.cpp
+++ b/Engine/Plugins/Runtime/RigVM/Source/RigVMEditor/Private/Editor/RigVMGraphDetailCustomization.cpp
@@ -1802,9 +1802,18 @@ void FRigVMWrappedNodeDetailCustomization::CustomizeLiveValues(IDetailLayoutBuil
 			TSharedPtr<IPropertyHandle> PinHandle = DetailLayout.GetProperty(Pin->GetFName());
 			if(PinHandle.IsValid())
 			{
-				DebugCategory.AddProperty(PinHandle)
-				.DisplayName(FText::FromName(Pin->GetDisplayName()))
-				.IsEnabled(false);
+				// we'll build a new custom row. adding the same property again
+				// causes the property to be marked customized - thus it won't
+				// show correctly in the default category.
+				DebugCategory.AddCustomRow(FText::FromName(Pin->GetFName()))
+				.NameContent()
+				[
+					PinHandle->CreatePropertyNameWidget()
+				]
+				.ValueContent()
+				[
+					PinHandle->CreatePropertyValueWidget()
+				];
 			}
 		}
 	}
