RigVM: Fix AnimNext graphs not responding to modifications in the model
#rnx

diff --git a/Engine/Plugins/Runtime/RigVM/Source/RigVMDeveloper/Private/EdGraph/RigVMEdGraph.cpp b/Engine/Plugins/Runtime/RigVM/Source/RigVMDeveloper/Private/EdGraph/RigVMEdGraph.cpp
--- a/Engine/Plugins/Runtime/RigVM/Source/RigVMDeveloper/Private/EdGraph/RigVMEdGraph.cpp
+++ b/Engine/Plugins/Runtime/RigVM/Source/RigVMDeveloper/Private/EdGraph/RigVMEdGraph.cpp
@@ -158,13 +158,9 @@ bool URigVMEdGraph::HandleModifiedEvent_Internal(ERigVMGraphNotifType InNotifTyp
 		return false;
 	}
 
-	// Make sure this EdGraph has a valid blueprint
-	URigVMBlueprint* RigVMBlueprint = nullptr;
-	if (UBlueprint* Blueprint = FBlueprintEditorUtils::FindBlueprintForGraph(this))
-	{
-		RigVMBlueprint = Cast<URigVMBlueprint>(Blueprint);
-	}
-	if (!RigVMBlueprint)
+	// Make sure this EdGraph has a valid rigvm host
+	const IRigVMClientHost* Host = GetImplementingOuter<IRigVMClientHost>();
+	if (!Host)
 	{
 		return false;
 	}
