RigVM: remove useless copies for GetArguments

#rnx

TODO: check if the following functions can return const reference instead of copies

virtual FRigVMTemplateTypeMap FRigVMDispatchFactory::OnNewArgumentType(const FName& InArgumentName, TRigVMTypeIndex InTypeIndex)
virtual TArray<FRigVMFunction> GetDispatchPredicatesImpl(const FRigVMTemplateTypeMap& InTypes)
virtual TArray<FRigVMExecuteArgument> GetExecuteArguments_Impl(const FRigVMDispatchContext& InContext)

and others?

diff --git a/Engine/Plugins/Animation/ControlRig/Source/ControlRig/Private/Units/Animation/RigUnit_AnimAttribute.cpp b/Engine/Plugins/Animation/ControlRig/Source/ControlRig/Private/Units/Animation/RigUnit_AnimAttribute.cpp
--- a/Engine/Plugins/Animation/ControlRig/Source/ControlRig/Private/Units/Animation/RigUnit_AnimAttribute.cpp
+++ b/Engine/Plugins/Animation/ControlRig/Source/ControlRig/Private/Units/Animation/RigUnit_AnimAttribute.cpp
@@ -321,7 +321,7 @@ FString FRigDispatch_AnimAttributeBase::GetNodeTitle(const FRigVMTemplateTypeMap
 }
 #endif
 
-TArray<FRigVMTemplateArgument> FRigDispatch_AnimAttributeBase::GetArguments() const
+const TArray<FRigVMTemplateArgument>& FRigDispatch_AnimAttributeBase::GetArguments() const
 {
 	if (Arguments.IsEmpty())
 	{
@@ -372,7 +372,7 @@ FText FRigDispatch_AnimAttributeBase::GetArgumentTooltip(const FName& InArgument
 
 
 
-TArray<FRigVMTemplateArgument> FRigDispatch_GetAnimAttribute::GetArguments() const
+const TArray<FRigVMTemplateArgument>& FRigDispatch_GetAnimAttribute::GetArguments() const
 {
 	if (ValueArgIndex == INDEX_NONE)
 	{
@@ -387,12 +387,10 @@ TArray<FRigVMTemplateArgument> FRigDispatch_GetAnimAttribute::GetArguments() con
 		FoundArgIndex = Arguments.Emplace(FoundArgName, ERigVMPinDirection::Output, RigVMTypeUtils::TypeIndex::Bool);
 	}
 
-	
 	return Arguments;
 }
 
-FRigVMTemplateTypeMap FRigDispatch_GetAnimAttribute::OnNewArgumentType(const FName& InArgumentName,
-	TRigVMTypeIndex InTypeIndex) const
+FRigVMTemplateTypeMap FRigDispatch_GetAnimAttribute::OnNewArgumentType(const FName& InArgumentName, TRigVMTypeIndex InTypeIndex) const
 {
 	FRigVMTemplateTypeMap Types;
 
@@ -455,7 +453,7 @@ FRigVMFunctionPtr FRigDispatch_GetAnimAttribute::GetDispatchFunctionImpl(const F
 	return nullptr;
 }
 
-TArray<FRigVMTemplateArgument> FRigDispatch_SetAnimAttribute::GetArguments() const
+const TArray<FRigVMTemplateArgument>& FRigDispatch_SetAnimAttribute::GetArguments() const
 {
 	if (ValueArgIndex == INDEX_NONE)
 	{
@@ -476,8 +474,7 @@ TArray<FRigVMExecuteArgument> FRigDispatch_SetAnimAttribute::GetExecuteArguments
 	return {{TEXT("ExecuteContext"), ERigVMPinDirection::IO}};
 }
 
-FRigVMTemplateTypeMap FRigDispatch_SetAnimAttribute::OnNewArgumentType(const FName& InArgumentName,
-                                                                       TRigVMTypeIndex InTypeIndex) const
+FRigVMTemplateTypeMap FRigDispatch_SetAnimAttribute::OnNewArgumentType(const FName& InArgumentName, TRigVMTypeIndex InTypeIndex) const
 {
 	FRigVMTemplateTypeMap Types;
 
diff --git a/Engine/Plugins/Animation/ControlRig/Source/ControlRig/Private/Units/Animation/RigUnit_AnimAttribute.h b/Engine/Plugins/Animation/ControlRig/Source/ControlRig/Private/Units/Animation/RigUnit_AnimAttribute.h
--- a/Engine/Plugins/Animation/ControlRig/Source/ControlRig/Private/Units/Animation/RigUnit_AnimAttribute.h
+++ b/Engine/Plugins/Animation/ControlRig/Source/ControlRig/Private/Units/Animation/RigUnit_AnimAttribute.h
@@ -134,7 +134,7 @@ struct CONTROLRIG_API FRigDispatch_AnimAttributeBase : public FRigDispatchFactor
 	virtual FString GetNodeTitle(const FRigVMTemplateTypeMap& InTypes) const override;;
 #endif
 	
-	virtual TArray<FRigVMTemplateArgument> GetArguments() const override;
+	virtual const TArray<FRigVMTemplateArgument>& GetArguments() const override;
 	virtual bool IsSet() const { return false; }
 
 #if WITH_EDITOR
@@ -187,7 +187,7 @@ struct CONTROLRIG_API FRigDispatch_GetAnimAttribute: public FRigDispatch_AnimAtt
 		FactoryScriptStruct = StaticStruct();
 	}
 
-	virtual TArray<FRigVMTemplateArgument> GetArguments() const override;
+	virtual const TArray<FRigVMTemplateArgument>& GetArguments() const override;
 	virtual FRigVMTemplateTypeMap OnNewArgumentType(const FName& InArgumentName, TRigVMTypeIndex InTypeIndex) const override;
 
 	
@@ -267,7 +267,7 @@ struct CONTROLRIG_API FRigDispatch_SetAnimAttribute: public FRigDispatch_AnimAtt
 	}
 
 	virtual bool IsSet() const override { return true; }
-	virtual TArray<FRigVMTemplateArgument> GetArguments() const override;
+	virtual const TArray<FRigVMTemplateArgument>& GetArguments() const override;
 	virtual TArray<FRigVMExecuteArgument> GetExecuteArguments_Impl(const FRigVMDispatchContext& InContext) const override;
 	virtual FRigVMTemplateTypeMap OnNewArgumentType(const FName& InArgumentName, TRigVMTypeIndex InTypeIndex) const override;
 	
diff --git a/Engine/Plugins/Animation/ControlRig/Source/ControlRig/Private/Units/Core/RigUnit_UserData.cpp b/Engine/Plugins/Animation/ControlRig/Source/ControlRig/Private/Units/Core/RigUnit_UserData.cpp
--- a/Engine/Plugins/Animation/ControlRig/Source/ControlRig/Private/Units/Core/RigUnit_UserData.cpp
+++ b/Engine/Plugins/Animation/ControlRig/Source/ControlRig/Private/Units/Core/RigUnit_UserData.cpp
@@ -28,23 +28,26 @@ FString FRigDispatch_GetUserData::GetArgumentMetaData(const FName& InArgumentNam
 
 #endif
 
-TArray<FRigVMTemplateArgument> FRigDispatch_GetUserData::GetArguments() const
+const TArray<FRigVMTemplateArgument>& FRigDispatch_GetUserData::GetArguments() const
 {
-	const TArray<FRigVMTemplateArgument::ETypeCategory> ValueCategories = {
-		FRigVMTemplateArgument::ETypeCategory_SingleAnyValue,
-		FRigVMTemplateArgument::ETypeCategory_ArrayAnyValue
-	};
-	return {
-		FRigVMTemplateArgument(ArgNameSpaceName, ERigVMPinDirection::Input, RigVMTypeUtils::TypeIndex::FString),
-		FRigVMTemplateArgument(ArgPathName, ERigVMPinDirection::Input, RigVMTypeUtils::TypeIndex::FString),
-		FRigVMTemplateArgument(ArgDefaultName, ERigVMPinDirection::Input, ValueCategories),
-		FRigVMTemplateArgument(ArgResultName, ERigVMPinDirection::Output, ValueCategories),
-		FRigVMTemplateArgument(ArgFoundName, ERigVMPinDirection::Output, RigVMTypeUtils::TypeIndex::Bool)
-	};
+	static TArray<FRigVMTemplateArgument> Arguments;
+	if (Arguments.IsEmpty())
+	{
+		const TArray<FRigVMTemplateArgument::ETypeCategory> ValueCategories = {
+			FRigVMTemplateArgument::ETypeCategory_SingleAnyValue,
+			FRigVMTemplateArgument::ETypeCategory_ArrayAnyValue
+		};
+		
+		Arguments.Emplace(ArgNameSpaceName, ERigVMPinDirection::Input, RigVMTypeUtils::TypeIndex::FString);
+		Arguments.Emplace(ArgPathName, ERigVMPinDirection::Input, RigVMTypeUtils::TypeIndex::FString);
+		Arguments.Emplace(ArgDefaultName, ERigVMPinDirection::Input, ValueCategories);
+		Arguments.Emplace(ArgResultName, ERigVMPinDirection::Output, ValueCategories);
+		Arguments.Emplace(ArgFoundName, ERigVMPinDirection::Output, RigVMTypeUtils::TypeIndex::Bool);
+	}
+	return Arguments;
 }
 
-FRigVMTemplateTypeMap FRigDispatch_GetUserData::OnNewArgumentType(const FName& InArgumentName,
-	TRigVMTypeIndex InTypeIndex) const
+FRigVMTemplateTypeMap FRigDispatch_GetUserData::OnNewArgumentType(const FName& InArgumentName, TRigVMTypeIndex InTypeIndex) const
 {
 	FRigVMTemplateTypeMap Types;
 	Types.Add(ArgNameSpaceName, RigVMTypeUtils::TypeIndex::FString);
diff --git a/Engine/Plugins/Animation/ControlRig/Source/ControlRig/Private/Units/Core/RigUnit_UserData.h b/Engine/Plugins/Animation/ControlRig/Source/ControlRig/Private/Units/Core/RigUnit_UserData.h
--- a/Engine/Plugins/Animation/ControlRig/Source/ControlRig/Private/Units/Core/RigUnit_UserData.h
+++ b/Engine/Plugins/Animation/ControlRig/Source/ControlRig/Private/Units/Core/RigUnit_UserData.h
@@ -22,7 +22,7 @@ public:
 #if WITH_EDITOR
 	virtual FString GetArgumentMetaData(const FName& InArgumentName, const FName& InMetaDataKey) const override;
 #endif
-	virtual TArray<FRigVMTemplateArgument> GetArguments() const override;
+	virtual const TArray<FRigVMTemplateArgument>& GetArguments() const override;
 	virtual FRigVMTemplateTypeMap OnNewArgumentType(const FName& InArgumentName, TRigVMTypeIndex InTypeIndex) const override;
 
 protected:
diff --git a/Engine/Plugins/Animation/ControlRig/Source/ControlRig/Private/Units/Hierarchy/RigUnit_Metadata.cpp b/Engine/Plugins/Animation/ControlRig/Source/ControlRig/Private/Units/Hierarchy/RigUnit_Metadata.cpp
--- a/Engine/Plugins/Animation/ControlRig/Source/ControlRig/Private/Units/Hierarchy/RigUnit_Metadata.cpp
+++ b/Engine/Plugins/Animation/ControlRig/Source/ControlRig/Private/Units/Hierarchy/RigUnit_Metadata.cpp
@@ -65,7 +65,7 @@ FString FRigDispatch_MetadataBase::GetNodeTitle(const FRigVMTemplateTypeMap& InT
 
 #endif
 
-TArray<FRigVMTemplateArgument> FRigDispatch_MetadataBase::GetArguments() const
+const TArray<FRigVMTemplateArgument>& FRigDispatch_MetadataBase::GetArguments() const
 {
 	if(Arguments.IsEmpty())
 	{
@@ -155,7 +155,7 @@ const TArray<TRigVMTypeIndex>& FRigDispatch_MetadataBase::GetValueTypes() const
 
 ////////////////////////////////////////////////////////////////////////////////////////////////
 
-TArray<FRigVMTemplateArgument> FRigDispatch_GetMetadata::GetArguments() const
+const TArray<FRigVMTemplateArgument>& FRigDispatch_GetMetadata::GetArguments() const
 {
 	if(ValueArgIndex == INDEX_NONE)
 	{
@@ -272,7 +272,7 @@ FRigVMFunctionPtr FRigDispatch_GetMetadata::GetDispatchFunctionImpl(const FRigVM
 
 ////////////////////////////////////////////////////////////////////////////////////////////////
 
-TArray<FRigVMTemplateArgument> FRigDispatch_SetMetadata::GetArguments() const
+const TArray<FRigVMTemplateArgument>& FRigDispatch_SetMetadata::GetArguments() const
 {
 	if(ValueArgIndex == INDEX_NONE)
 	{
diff --git a/Engine/Plugins/Animation/ControlRig/Source/ControlRig/Private/Units/Hierarchy/RigUnit_Metadata.h b/Engine/Plugins/Animation/ControlRig/Source/ControlRig/Private/Units/Hierarchy/RigUnit_Metadata.h
--- a/Engine/Plugins/Animation/ControlRig/Source/ControlRig/Private/Units/Hierarchy/RigUnit_Metadata.h
+++ b/Engine/Plugins/Animation/ControlRig/Source/ControlRig/Private/Units/Hierarchy/RigUnit_Metadata.h
@@ -18,7 +18,7 @@ struct CONTROLRIG_API FRigDispatch_MetadataBase : public FRigDispatchFactory
 #if WITH_EDITOR
 	virtual FString GetNodeTitle(const FRigVMTemplateTypeMap& InTypes) const override;;
 #endif
-	virtual TArray<FRigVMTemplateArgument> GetArguments() const override;
+	virtual const TArray<FRigVMTemplateArgument>& GetArguments() const override;
 	virtual bool IsSetMetadata() const { return false; }
 
 #if WITH_EDITOR
@@ -56,7 +56,7 @@ struct CONTROLRIG_API FRigDispatch_GetMetadata : public FRigDispatch_MetadataBas
 {
 	GENERATED_BODY()
 
-	virtual TArray<FRigVMTemplateArgument> GetArguments() const override;
+	virtual const TArray<FRigVMTemplateArgument>& GetArguments() const override;
 
 protected:
 
@@ -118,7 +118,7 @@ struct CONTROLRIG_API FRigDispatch_SetMetadata : public FRigDispatch_MetadataBas
 {
 	GENERATED_BODY()
 
-	virtual TArray<FRigVMTemplateArgument> GetArguments() const override;
+	virtual const TArray<FRigVMTemplateArgument>& GetArguments() const override;
 	virtual TArray<FRigVMExecuteArgument> GetExecuteArguments_Impl(const FRigVMDispatchContext& InContext) const override;
 	virtual bool IsSetMetadata() const override { return true; }
 
diff --git a/Engine/Plugins/Experimental/AnimNext/Source/AnimNext/Internal/Param/RigVMDispatch_GetParameter.h b/Engine/Plugins/Experimental/AnimNext/Source/AnimNext/Internal/Param/RigVMDispatch_GetParameter.h
--- a/Engine/Plugins/Experimental/AnimNext/Source/AnimNext/Internal/Param/RigVMDispatch_GetParameter.h
+++ b/Engine/Plugins/Experimental/AnimNext/Source/AnimNext/Internal/Param/RigVMDispatch_GetParameter.h
@@ -26,7 +26,7 @@ private:
 
 	virtual UScriptStruct* GetExecuteContextStruct() const { return FAnimNextExecuteContext::StaticStruct(); }
 	virtual FName GetArgumentNameForOperandIndex(int32 InOperandIndex, int32 InTotalOperands) const override;
-	virtual TArray<FRigVMTemplateArgument> GetArguments() const override;
+	virtual const TArray<FRigVMTemplateArgument>& GetArguments() const override;
 #if WITH_EDITOR
 	virtual FString GetArgumentMetaData(const FName& InArgumentName, const FName& InMetaDataKey) const override;
 	virtual FString GetArgumentDefaultValue(const FName& InArgumentName, TRigVMTypeIndex InTypeIndex) const override;
diff --git a/Engine/Plugins/Experimental/AnimNext/Source/AnimNext/Internal/Param/RigVMDispatch_SetParameter.h b/Engine/Plugins/Experimental/AnimNext/Source/AnimNext/Internal/Param/RigVMDispatch_SetParameter.h
--- a/Engine/Plugins/Experimental/AnimNext/Source/AnimNext/Internal/Param/RigVMDispatch_SetParameter.h
+++ b/Engine/Plugins/Experimental/AnimNext/Source/AnimNext/Internal/Param/RigVMDispatch_SetParameter.h
@@ -26,7 +26,7 @@ private:
 
 	virtual UScriptStruct* GetExecuteContextStruct() const { return FAnimNextExecuteContext::StaticStruct(); }
 	virtual FName GetArgumentNameForOperandIndex(int32 InOperandIndex, int32 InTotalOperands) const override;
-	virtual TArray<FRigVMTemplateArgument> GetArguments() const override;
+	virtual const TArray<FRigVMTemplateArgument>& GetArguments() const override;
 #if WITH_EDITOR
 	virtual FString GetArgumentMetaData(const FName& InArgumentName, const FName& InMetaDataKey) const override;
 	virtual FString GetArgumentDefaultValue(const FName& InArgumentName, TRigVMTypeIndex InTypeIndex) const override;
diff --git a/Engine/Plugins/Experimental/AnimNext/Source/AnimNext/Private/Param/RigVMDispatch_GetParameter.cpp b/Engine/Plugins/Experimental/AnimNext/Source/AnimNext/Private/Param/RigVMDispatch_GetParameter.cpp
--- a/Engine/Plugins/Experimental/AnimNext/Source/AnimNext/Private/Param/RigVMDispatch_GetParameter.cpp
+++ b/Engine/Plugins/Experimental/AnimNext/Source/AnimNext/Private/Param/RigVMDispatch_GetParameter.cpp
@@ -56,20 +56,23 @@ FString FRigVMDispatch_GetParameter::GetArgumentDefaultValue(const FName& InArgu
 }
 #endif
 
-TArray<FRigVMTemplateArgument> FRigVMDispatch_GetParameter::GetArguments() const
+const TArray<FRigVMTemplateArgument>& FRigVMDispatch_GetParameter::GetArguments() const
 {
-	const TArray<FRigVMTemplateArgument::ETypeCategory> ValueCategories =
+	static TArray<FRigVMTemplateArgument> Arguments;
+	if (Arguments.IsEmpty())
 	{
-		FRigVMTemplateArgument::ETypeCategory_SingleAnyValue,
-		FRigVMTemplateArgument::ETypeCategory_ArrayAnyValue
-	};
-	return
-	{
-		FRigVMTemplateArgument(ParameterName, ERigVMPinDirection::Input, RigVMTypeUtils::TypeIndex::FName),
-		FRigVMTemplateArgument(ValueName, ERigVMPinDirection::Output, ValueCategories),
-		FRigVMTemplateArgument(ParameterIdName, ERigVMPinDirection::Hidden, RigVMTypeUtils::TypeIndex::UInt32),
-		FRigVMTemplateArgument(TypeHandleName, ERigVMPinDirection::Hidden, RigVMTypeUtils::TypeIndex::UInt32),
-	};
+		const TArray<FRigVMTemplateArgument::ETypeCategory> ValueCategories =
+	   {
+			FRigVMTemplateArgument::ETypeCategory_SingleAnyValue,
+			FRigVMTemplateArgument::ETypeCategory_ArrayAnyValue
+		};
+		
+		Arguments.Emplace(ParameterName, ERigVMPinDirection::Input, RigVMTypeUtils::TypeIndex::FName);
+		Arguments.Emplace(ValueName, ERigVMPinDirection::Output, ValueCategories);
+		Arguments.Emplace(ParameterIdName, ERigVMPinDirection::Hidden, RigVMTypeUtils::TypeIndex::UInt32);
+		Arguments.Emplace(TypeHandleName, ERigVMPinDirection::Hidden, RigVMTypeUtils::TypeIndex::UInt32);
+	}
+	return Arguments;
 }
 
 FRigVMTemplateTypeMap FRigVMDispatch_GetParameter::OnNewArgumentType(const FName& InArgumentName, TRigVMTypeIndex InTypeIndex) const
diff --git a/Engine/Plugins/Experimental/AnimNext/Source/AnimNext/Private/Param/RigVMDispatch_SetParameter.cpp b/Engine/Plugins/Experimental/AnimNext/Source/AnimNext/Private/Param/RigVMDispatch_SetParameter.cpp
--- a/Engine/Plugins/Experimental/AnimNext/Source/AnimNext/Private/Param/RigVMDispatch_SetParameter.cpp
+++ b/Engine/Plugins/Experimental/AnimNext/Source/AnimNext/Private/Param/RigVMDispatch_SetParameter.cpp
@@ -57,20 +57,23 @@ FString FRigVMDispatch_SetParameter::GetArgumentDefaultValue(const FName& InArgu
 }
 #endif
 
-TArray<FRigVMTemplateArgument> FRigVMDispatch_SetParameter::GetArguments() const
+const TArray<FRigVMTemplateArgument>& FRigVMDispatch_SetParameter::GetArguments() const
 {
-	const TArray<FRigVMTemplateArgument::ETypeCategory> ValueCategories =
+	static TArray<FRigVMTemplateArgument> Arguments;
+	if (Arguments.IsEmpty())
 	{
-		FRigVMTemplateArgument::ETypeCategory_SingleAnyValue,
-		FRigVMTemplateArgument::ETypeCategory_ArrayAnyValue
-	};
-	return 
-	{ 
-		FRigVMTemplateArgument(ParameterName, ERigVMPinDirection::Input, RigVMTypeUtils::TypeIndex::FName),
-		FRigVMTemplateArgument(ValueName, ERigVMPinDirection::Input, ValueCategories),
-		FRigVMTemplateArgument(ParameterIdName, ERigVMPinDirection::Hidden, RigVMTypeUtils::TypeIndex::UInt32),
-		FRigVMTemplateArgument(TypeHandleName, ERigVMPinDirection::Hidden, RigVMTypeUtils::TypeIndex::UInt32),
-	};
+		const TArray<FRigVMTemplateArgument::ETypeCategory> ValueCategories =
+		{
+			FRigVMTemplateArgument::ETypeCategory_SingleAnyValue,
+			FRigVMTemplateArgument::ETypeCategory_ArrayAnyValue
+		};
+		 
+		Arguments.Emplace(ParameterName, ERigVMPinDirection::Input, RigVMTypeUtils::TypeIndex::FName);
+		Arguments.Emplace(ValueName, ERigVMPinDirection::Input, ValueCategories);
+		Arguments.Emplace(ParameterIdName, ERigVMPinDirection::Hidden, RigVMTypeUtils::TypeIndex::UInt32);
+		Arguments.Emplace(TypeHandleName, ERigVMPinDirection::Hidden, RigVMTypeUtils::TypeIndex::UInt32);
+	}
+	return Arguments;
 }
 
 TArray<FRigVMExecuteArgument> FRigVMDispatch_SetParameter::GetExecuteArguments_Impl(const FRigVMDispatchContext& InContext) const
diff --git a/Engine/Plugins/Runtime/RigVM/Source/RigVM/Private/RigVMFunctions/RigVMDispatch_Array.cpp b/Engine/Plugins/Runtime/RigVM/Source/RigVM/Private/RigVMFunctions/RigVMDispatch_Array.cpp
--- a/Engine/Plugins/Runtime/RigVM/Source/RigVM/Private/RigVMFunctions/RigVMDispatch_Array.cpp
+++ b/Engine/Plugins/Runtime/RigVM/Source/RigVM/Private/RigVMFunctions/RigVMDispatch_Array.cpp
@@ -204,12 +204,15 @@ TArray<FRigVMExecuteArgument> FRigVMDispatch_ArrayBaseMutable::GetExecuteArgumen
 
 ////////////////////////////////////////////////////////////////////////////////////////////////////////////////////////
 
-TArray<FRigVMTemplateArgument> FRigVMDispatch_ArrayMake::GetArguments() const
+const TArray<FRigVMTemplateArgument>& FRigVMDispatch_ArrayMake::GetArguments() const
 {
-	return {
-		CreateArgument(ValuesName, ERigVMPinDirection::Input),
-		CreateArgument(ArrayName, ERigVMPinDirection::Output)
+	static TArray<FRigVMTemplateArgument> Arguments;
+	if (Arguments.IsEmpty())
+	{
+		Arguments.Emplace(CreateArgument(ValuesName, ERigVMPinDirection::Input));
+		Arguments.Emplace(CreateArgument(ArrayName, ERigVMPinDirection::Output));
 	};
+	return Arguments;
 }
 
 FRigVMTemplateTypeMap FRigVMDispatch_ArrayMake::OnNewArgumentType(const FName& InArgumentName, TRigVMTypeIndex InTypeIndex) const
@@ -313,11 +316,14 @@ FName FRigVMDispatch_ArrayReset::GetArgumentNameForOperandIndex(int32 InOperandI
 	return ArrayName;
 }
 
-TArray<FRigVMTemplateArgument> FRigVMDispatch_ArrayReset::GetArguments() const
+const TArray<FRigVMTemplateArgument>& FRigVMDispatch_ArrayReset::GetArguments() const
 {
-	return {
-		CreateArgument(ArrayName, ERigVMPinDirection::IO)
+	static TArray<FRigVMTemplateArgument> Arguments;
+	if (Arguments.IsEmpty())
+	{
+		Arguments.Emplace(CreateArgument(ArrayName, ERigVMPinDirection::IO));
 	};
+	return Arguments;
 }
 
 FRigVMTemplateTypeMap FRigVMDispatch_ArrayReset::OnNewArgumentType(const FName& InArgumentName, TRigVMTypeIndex InTypeIndex) const
@@ -365,12 +371,15 @@ FName FRigVMDispatch_ArrayGetNum::GetArgumentNameForOperandIndex(int32 InOperand
 	return ArgumentNames[InOperandIndex];
 }
 
-TArray<FRigVMTemplateArgument> FRigVMDispatch_ArrayGetNum::GetArguments() const
+const TArray<FRigVMTemplateArgument>& FRigVMDispatch_ArrayGetNum::GetArguments() const
 {
-	return {
-		CreateArgument(ArrayName, ERigVMPinDirection::Input),
-		CreateArgument(NumName, ERigVMPinDirection::Output)
+	static TArray<FRigVMTemplateArgument> Arguments;
+	if (Arguments.IsEmpty())
+	{
+		Arguments.Emplace(CreateArgument(ArrayName, ERigVMPinDirection::Input));
+		Arguments.Emplace(CreateArgument(NumName, ERigVMPinDirection::Output));
 	};
+	return Arguments;
 }
 
 FRigVMTemplateTypeMap FRigVMDispatch_ArrayGetNum::OnNewArgumentType(const FName& InArgumentName, TRigVMTypeIndex InTypeIndex) const
@@ -424,12 +433,15 @@ FName FRigVMDispatch_ArraySetNum::GetArgumentNameForOperandIndex(int32 InOperand
 	return ArgumentNames[InOperandIndex];
 }
 
-TArray<FRigVMTemplateArgument> FRigVMDispatch_ArraySetNum::GetArguments() const
+const TArray<FRigVMTemplateArgument>& FRigVMDispatch_ArraySetNum::GetArguments() const
 {
-	return {
-		CreateArgument(ArrayName, ERigVMPinDirection::IO),
-		CreateArgument(NumName, ERigVMPinDirection::Input)
-	};
+	static TArray<FRigVMTemplateArgument> Arguments;
+	if (Arguments.IsEmpty())
+	{
+		Arguments.Emplace(CreateArgument(ArrayName, ERigVMPinDirection::IO));
+		Arguments.Emplace(CreateArgument(NumName, ERigVMPinDirection::Input));
+	}
+	return Arguments;
 }
 
 FRigVMTemplateTypeMap FRigVMDispatch_ArraySetNum::OnNewArgumentType(const FName& InArgumentName, TRigVMTypeIndex InTypeIndex) const
@@ -487,13 +499,16 @@ FName FRigVMDispatch_ArrayGetAtIndex::GetArgumentNameForOperandIndex(int32 InOpe
 	return ArgumentNames[InOperandIndex];
 }
 
-TArray<FRigVMTemplateArgument> FRigVMDispatch_ArrayGetAtIndex::GetArguments() const
+const TArray<FRigVMTemplateArgument>& FRigVMDispatch_ArrayGetAtIndex::GetArguments() const
 {
-	return {
-		CreateArgument(ArrayName, ERigVMPinDirection::Input),
-		CreateArgument(IndexName, ERigVMPinDirection::Input),
-		CreateArgument(ElementName, ERigVMPinDirection::Output)
-	};
+	static TArray<FRigVMTemplateArgument> Arguments;
+	if (Arguments.IsEmpty())
+	{
+		Arguments.Emplace(CreateArgument(ArrayName, ERigVMPinDirection::Input));
+		Arguments.Emplace(CreateArgument(IndexName, ERigVMPinDirection::Input));
+		Arguments.Emplace(CreateArgument(ElementName, ERigVMPinDirection::Output));
+	}
+	return Arguments;
 }
 
 FRigVMTemplateTypeMap FRigVMDispatch_ArrayGetAtIndex::OnNewArgumentType(const FName& InArgumentName, TRigVMTypeIndex InTypeIndex) const
@@ -569,13 +584,16 @@ FName FRigVMDispatch_ArraySetAtIndex::GetArgumentNameForOperandIndex(int32 InOpe
 	return ArgumentNames[InOperandIndex];
 }
 
-TArray<FRigVMTemplateArgument> FRigVMDispatch_ArraySetAtIndex::GetArguments() const
+const TArray<FRigVMTemplateArgument>& FRigVMDispatch_ArraySetAtIndex::GetArguments() const
 {
-	return {
-		CreateArgument(ArrayName, ERigVMPinDirection::IO),
-		CreateArgument(IndexName, ERigVMPinDirection::Input),
-		CreateArgument(ElementName, ERigVMPinDirection::Input)
-	};
+	static TArray<FRigVMTemplateArgument> Arguments;
+	if (Arguments.IsEmpty())
+	{
+		Arguments.Emplace(CreateArgument(ArrayName, ERigVMPinDirection::IO));
+		Arguments.Emplace(CreateArgument(IndexName, ERigVMPinDirection::Input));
+		Arguments.Emplace(CreateArgument(ElementName, ERigVMPinDirection::Input));
+	}
+	return Arguments;
 }
 
 FRigVMTemplateTypeMap FRigVMDispatch_ArraySetAtIndex::OnNewArgumentType(const FName& InArgumentName, TRigVMTypeIndex InTypeIndex) const
@@ -651,13 +669,16 @@ FName FRigVMDispatch_ArrayAdd::GetArgumentNameForOperandIndex(int32 InOperandInd
 	return ArgumentNames[InOperandIndex];
 }
 
-TArray<FRigVMTemplateArgument> FRigVMDispatch_ArrayAdd::GetArguments() const
+const TArray<FRigVMTemplateArgument>& FRigVMDispatch_ArrayAdd::GetArguments() const
 {
-	return {
-		CreateArgument(ArrayName, ERigVMPinDirection::IO),
-		CreateArgument(ElementName, ERigVMPinDirection::Input),
-		CreateArgument(IndexName, ERigVMPinDirection::Output)
-	};
+	static TArray<FRigVMTemplateArgument> Arguments;
+	if (Arguments.IsEmpty())
+	{
+		Arguments.Emplace(CreateArgument(ArrayName, ERigVMPinDirection::IO));
+		Arguments.Emplace(CreateArgument(ElementName, ERigVMPinDirection::Input));
+		Arguments.Emplace(CreateArgument(IndexName, ERigVMPinDirection::Output));
+	}
+	return Arguments;
 }
 
 void FRigVMDispatch_ArrayAdd::Execute(FRigVMExtendedExecuteContext& InContext, FRigVMMemoryHandleArray Handles, FRigVMPredicateBranchArray Predicates)
@@ -768,12 +789,15 @@ FName FRigVMDispatch_ArrayRemove::GetArgumentNameForOperandIndex(int32 InOperand
 	return ArgumentNames[InOperandIndex];
 }
 
-TArray<FRigVMTemplateArgument> FRigVMDispatch_ArrayRemove::GetArguments() const
+const TArray<FRigVMTemplateArgument>& FRigVMDispatch_ArrayRemove::GetArguments() const
 {
-	return {
-		CreateArgument(ArrayName, ERigVMPinDirection::IO),
-		CreateArgument(IndexName, ERigVMPinDirection::Input)
-	};
+	static TArray<FRigVMTemplateArgument> Arguments;
+	if (Arguments.IsEmpty())
+	{
+		Arguments.Emplace(CreateArgument(ArrayName, ERigVMPinDirection::IO));
+		Arguments.Emplace(CreateArgument(IndexName, ERigVMPinDirection::Input));
+	}
+	return Arguments;
 }
 
 FRigVMTemplateTypeMap FRigVMDispatch_ArrayRemove::OnNewArgumentType(const FName& InArgumentName, TRigVMTypeIndex InTypeIndex) const
@@ -864,14 +888,17 @@ FName FRigVMDispatch_ArrayFind::GetArgumentNameForOperandIndex(int32 InOperandIn
 	return ArgumentNames[InOperandIndex];
 }
 
-TArray<FRigVMTemplateArgument> FRigVMDispatch_ArrayFind::GetArguments() const
+const TArray<FRigVMTemplateArgument>& FRigVMDispatch_ArrayFind::GetArguments() const
 {
-	return {
-		CreateArgument(ArrayName, ERigVMPinDirection::Input),
-		CreateArgument(ElementName, ERigVMPinDirection::Input),
-		CreateArgument(IndexName, ERigVMPinDirection::Output),
-		CreateArgument(SuccessName, ERigVMPinDirection::Output)
-	};
+	static TArray<FRigVMTemplateArgument> Arguments;
+	if (Arguments.IsEmpty())
+	{
+		Arguments.Emplace(CreateArgument(ArrayName, ERigVMPinDirection::Input));
+		Arguments.Emplace(CreateArgument(ElementName, ERigVMPinDirection::Input));
+		Arguments.Emplace(CreateArgument(IndexName, ERigVMPinDirection::Output));
+		Arguments.Emplace(CreateArgument(SuccessName, ERigVMPinDirection::Output));
+	}
+	return Arguments;
 }
 
 FRigVMTemplateTypeMap FRigVMDispatch_ArrayFind::OnNewArgumentType(const FName& InArgumentName, TRigVMTypeIndex InTypeIndex) const
@@ -973,12 +1000,15 @@ FName FRigVMDispatch_ArrayAppend::GetArgumentNameForOperandIndex(int32 InOperand
 	return ArgumentNames[InOperandIndex];
 }
 
-TArray<FRigVMTemplateArgument> FRigVMDispatch_ArrayAppend::GetArguments() const
+const TArray<FRigVMTemplateArgument>& FRigVMDispatch_ArrayAppend::GetArguments() const
 {
-	return {
-		CreateArgument(ArrayName, ERigVMPinDirection::IO),
-		CreateArgument(OtherName, ERigVMPinDirection::Input)
-	};
+	static TArray<FRigVMTemplateArgument> Arguments;
+	if (Arguments.IsEmpty())
+	{
+		Arguments.Emplace(CreateArgument(ArrayName, ERigVMPinDirection::IO));
+		Arguments.Emplace(CreateArgument(OtherName, ERigVMPinDirection::Input));
+	}
+	return Arguments;
 }
 
 FRigVMTemplateTypeMap FRigVMDispatch_ArrayAppend::OnNewArgumentType(const FName& InArgumentName, TRigVMTypeIndex InTypeIndex) const
@@ -1052,12 +1082,15 @@ FName FRigVMDispatch_ArrayClone::GetArgumentNameForOperandIndex(int32 InOperandI
 	return ArgumentNames[InOperandIndex];
 }
 
-TArray<FRigVMTemplateArgument> FRigVMDispatch_ArrayClone::GetArguments() const
+const TArray<FRigVMTemplateArgument>& FRigVMDispatch_ArrayClone::GetArguments() const
 {
-	return {
-		CreateArgument(ArrayName, ERigVMPinDirection::Input),
-		CreateArgument(CloneName, ERigVMPinDirection::Output)
-	};
+	static TArray<FRigVMTemplateArgument> Arguments;
+	if (Arguments.IsEmpty())
+	{
+		Arguments.Emplace(CreateArgument(ArrayName, ERigVMPinDirection::Input));
+		Arguments.Emplace(CreateArgument(CloneName, ERigVMPinDirection::Output));
+	}
+	return Arguments;
 }
 
 FRigVMTemplateTypeMap FRigVMDispatch_ArrayClone::OnNewArgumentType(const FName& InArgumentName, TRigVMTypeIndex InTypeIndex) const
@@ -1232,13 +1265,16 @@ FName FRigVMDispatch_ArrayDifference::GetArgumentNameForOperandIndex(int32 InOpe
 	return ArgumentNames[InOperandIndex];
 }
 
-TArray<FRigVMTemplateArgument> FRigVMDispatch_ArrayDifference::GetArguments() const
+const TArray<FRigVMTemplateArgument>& FRigVMDispatch_ArrayDifference::GetArguments() const
 {
-	return {
-		CreateArgument(ArrayName, ERigVMPinDirection::Input),
-		CreateArgument(OtherName, ERigVMPinDirection::Input),
-		CreateArgument(ResultName, ERigVMPinDirection::Output)
-	};
+	static TArray<FRigVMTemplateArgument> Arguments;
+	if (Arguments.IsEmpty())
+	{
+		Arguments.Emplace(CreateArgument(ArrayName, ERigVMPinDirection::Input));
+		Arguments.Emplace(CreateArgument(OtherName, ERigVMPinDirection::Input));
+		Arguments.Emplace(CreateArgument(ResultName, ERigVMPinDirection::Output));
+	}
+	return Arguments;
 }
 
 FRigVMTemplateTypeMap FRigVMDispatch_ArrayDifference::OnNewArgumentType(const FName& InArgumentName, TRigVMTypeIndex InTypeIndex) const
@@ -1400,16 +1436,19 @@ FName FRigVMDispatch_ArrayIterator::GetArgumentNameForOperandIndex(int32 InOpera
 	return ArgumentNames[InOperandIndex];
 }
 
-TArray<FRigVMTemplateArgument> FRigVMDispatch_ArrayIterator::GetArguments() const
+const TArray<FRigVMTemplateArgument>& FRigVMDispatch_ArrayIterator::GetArguments() const
 {
-	return {
-		CreateArgument(ArrayName, ERigVMPinDirection::Input),
-		CreateArgument(ElementName, ERigVMPinDirection::Output),
-		CreateArgument(IndexName, ERigVMPinDirection::Output),
-		CreateArgument(CountName, ERigVMPinDirection::Output),
-		CreateArgument(RatioName, ERigVMPinDirection::Output),
-		{FRigVMStruct::ControlFlowBlockToRunName, ERigVMPinDirection::Hidden, RigVMTypeUtils::TypeIndex::FName}
-	};
+	static TArray<FRigVMTemplateArgument> Arguments;
+	if (Arguments.IsEmpty())
+	{
+		Arguments.Emplace(CreateArgument(ArrayName, ERigVMPinDirection::Input));
+		Arguments.Emplace(CreateArgument(ElementName, ERigVMPinDirection::Output));
+		Arguments.Emplace(CreateArgument(IndexName, ERigVMPinDirection::Output));
+		Arguments.Emplace(CreateArgument(CountName, ERigVMPinDirection::Output));
+		Arguments.Emplace(CreateArgument(RatioName, ERigVMPinDirection::Output));
+		Arguments.Emplace(FRigVMStruct::ControlFlowBlockToRunName, ERigVMPinDirection::Hidden, RigVMTypeUtils::TypeIndex::FName);
+	}
+	return Arguments;
 }
 
 TArray<FRigVMExecuteArgument> FRigVMDispatch_ArrayIterator::GetExecuteArguments_Impl(
@@ -1432,8 +1471,7 @@ const bool FRigVMDispatch_ArrayIterator::IsControlFlowBlockSliced(const FName& I
 	return InBlockName == ExecuteName;
 }
 
-FRigVMTemplateTypeMap FRigVMDispatch_ArrayIterator::OnNewArgumentType(const FName& InArgumentName,
-	TRigVMTypeIndex InTypeIndex) const
+FRigVMTemplateTypeMap FRigVMDispatch_ArrayIterator::OnNewArgumentType(const FName& InArgumentName, TRigVMTypeIndex InTypeIndex) const
 {
 	check(InArgumentName == ArrayName || InArgumentName == ElementName);
 
diff --git a/Engine/Plugins/Runtime/RigVM/Source/RigVM/Private/RigVMFunctions/RigVMDispatch_Constant.cpp b/Engine/Plugins/Runtime/RigVM/Source/RigVM/Private/RigVMFunctions/RigVMDispatch_Constant.cpp
--- a/Engine/Plugins/Runtime/RigVM/Source/RigVM/Private/RigVMFunctions/RigVMDispatch_Constant.cpp
+++ b/Engine/Plugins/Runtime/RigVM/Source/RigVM/Private/RigVMFunctions/RigVMDispatch_Constant.cpp
@@ -10,14 +10,17 @@
 
 const FName FRigVMDispatch_Constant::ValueName = TEXT("Value");
 
-TArray<FRigVMTemplateArgument> FRigVMDispatch_Constant::GetArguments() const
+const TArray<FRigVMTemplateArgument>& FRigVMDispatch_Constant::GetArguments() const
 {
-	const TArray<FRigVMTemplateArgument::ETypeCategory> Categories = {
-		FRigVMTemplateArgument::ETypeCategory_SingleAnyValue
-	};
-	return {
-		{ValueName, ERigVMPinDirection::IO, Categories}
-	};
+	static TArray<FRigVMTemplateArgument> Arguments;
+	if (Arguments.IsEmpty())
+	{
+		static const TArray<FRigVMTemplateArgument::ETypeCategory> Categories = {
+			FRigVMTemplateArgument::ETypeCategory_SingleAnyValue
+		};
+		Arguments.Emplace(ValueName, ERigVMPinDirection::IO, Categories);
+	}
+	return Arguments;
 }
 
 FRigVMTemplateTypeMap FRigVMDispatch_Constant::OnNewArgumentType(const FName& InArgumentName, TRigVMTypeIndex InTypeIndex) const
diff --git a/Engine/Plugins/Runtime/RigVM/Source/RigVM/Private/RigVMFunctions/RigVMDispatch_Core.cpp b/Engine/Plugins/Runtime/RigVM/Source/RigVM/Private/RigVMFunctions/RigVMDispatch_Core.cpp
--- a/Engine/Plugins/Runtime/RigVM/Source/RigVM/Private/RigVMFunctions/RigVMDispatch_Core.cpp
+++ b/Engine/Plugins/Runtime/RigVM/Source/RigVM/Private/RigVMFunctions/RigVMDispatch_Core.cpp
@@ -23,21 +23,24 @@ FName FRigVMDispatch_CoreEquals::GetArgumentNameForOperandIndex(int32 InOperandI
 	return ArgumentNames[InOperandIndex];
 }
 
-TArray<FRigVMTemplateArgument> FRigVMDispatch_CoreEquals::GetArguments() const
+const TArray<FRigVMTemplateArgument>& FRigVMDispatch_CoreEquals::GetArguments() const
 {
-	const TArray<FRigVMTemplateArgument::ETypeCategory> ValueCategories = {
-		FRigVMTemplateArgument::ETypeCategory_SingleAnyValue,
-		FRigVMTemplateArgument::ETypeCategory_ArrayAnyValue
-	};
-	return {
-		FRigVMTemplateArgument(AName, ERigVMPinDirection::Input, ValueCategories),
-		FRigVMTemplateArgument(BName, ERigVMPinDirection::Input, ValueCategories),
-		FRigVMTemplateArgument(ResultName, ERigVMPinDirection::Output, RigVMTypeUtils::TypeIndex::Bool)
-	};
+	static TArray<FRigVMTemplateArgument> Arguments;
+	if (Arguments.IsEmpty())
+	{
+		static const TArray<FRigVMTemplateArgument::ETypeCategory> ValueCategories = {
+			FRigVMTemplateArgument::ETypeCategory_SingleAnyValue,
+			FRigVMTemplateArgument::ETypeCategory_ArrayAnyValue
+		};
+		
+		Arguments.Emplace(AName, ERigVMPinDirection::Input, ValueCategories);
+		Arguments.Emplace(BName, ERigVMPinDirection::Input, ValueCategories);
+		Arguments.Emplace(ResultName, ERigVMPinDirection::Output, RigVMTypeUtils::TypeIndex::Bool);
+	}
+	return Arguments;
 }
 
-FRigVMTemplateTypeMap FRigVMDispatch_CoreEquals::OnNewArgumentType(const FName& InArgumentName,
-	TRigVMTypeIndex InTypeIndex) const
+FRigVMTemplateTypeMap FRigVMDispatch_CoreEquals::OnNewArgumentType(const FName& InArgumentName, TRigVMTypeIndex InTypeIndex) const
 {
 	FRigVMTemplateTypeMap Types;
 	Types.Add(AName, InTypeIndex);
diff --git a/Engine/Plugins/Runtime/RigVM/Source/RigVM/Private/RigVMFunctions/RigVMDispatch_If.cpp b/Engine/Plugins/Runtime/RigVM/Source/RigVM/Private/RigVMFunctions/RigVMDispatch_If.cpp
--- a/Engine/Plugins/Runtime/RigVM/Source/RigVM/Private/RigVMFunctions/RigVMDispatch_If.cpp
+++ b/Engine/Plugins/Runtime/RigVM/Source/RigVM/Private/RigVMFunctions/RigVMDispatch_If.cpp
@@ -23,28 +23,34 @@ FName FRigVMDispatch_If::GetArgumentNameForOperandIndex(int32 InOperandIndex, in
 	return ArgumentNames[InOperandIndex];
 }
 
-TArray<FRigVMTemplateArgument> FRigVMDispatch_If::GetArguments() const
+const TArray<FRigVMTemplateArgument>& FRigVMDispatch_If::GetArguments() const
 {
-	static const TArray<FRigVMTemplateArgument::ETypeCategory> ValueCategories = {
-		FRigVMTemplateArgument::ETypeCategory_SingleAnyValue,
-		FRigVMTemplateArgument::ETypeCategory_ArrayAnyValue
-	};
-	return {
-		FRigVMTemplateArgument(ConditionName, ERigVMPinDirection::Input, RigVMTypeUtils::TypeIndex::Bool),
-		FRigVMTemplateArgument(TrueName, ERigVMPinDirection::Input, ValueCategories),
-		FRigVMTemplateArgument(FalseName, ERigVMPinDirection::Input, ValueCategories),
-		FRigVMTemplateArgument(ResultName, ERigVMPinDirection::Output, ValueCategories)
-	};
+	static TArray<FRigVMTemplateArgument> Arguments;
+	if (Arguments.IsEmpty())
+	{
+		static const TArray<FRigVMTemplateArgument::ETypeCategory> ValueCategories = {
+			FRigVMTemplateArgument::ETypeCategory_SingleAnyValue,
+			FRigVMTemplateArgument::ETypeCategory_ArrayAnyValue
+		};
+
+		Arguments.Emplace(ConditionName, ERigVMPinDirection::Input, RigVMTypeUtils::TypeIndex::Bool);
+		Arguments.Emplace(TrueName, ERigVMPinDirection::Input, ValueCategories);
+		Arguments.Emplace(FalseName, ERigVMPinDirection::Input, ValueCategories);
+		Arguments.Emplace(ResultName, ERigVMPinDirection::Output, ValueCategories);
+	}
+	return Arguments;
 }
 
-FRigVMTemplateTypeMap FRigVMDispatch_If::OnNewArgumentType(const FName& InArgumentName,
-	TRigVMTypeIndex InTypeIndex) const
+FRigVMTemplateTypeMap FRigVMDispatch_If::OnNewArgumentType(const FName& InArgumentName, TRigVMTypeIndex InTypeIndex) const
 {
-	FRigVMTemplateTypeMap Types;
-	Types.Add(ConditionName, RigVMTypeUtils::TypeIndex::Bool);
-	Types.Add(TrueName, InTypeIndex);
-	Types.Add(FalseName, InTypeIndex);
-	Types.Add(ResultName, InTypeIndex);
+	static FRigVMTemplateTypeMap Types;
+	if (Types.IsEmpty())
+	{
+		Types.Add(ConditionName, RigVMTypeUtils::TypeIndex::Bool);
+		Types.Add(TrueName, InTypeIndex);
+		Types.Add(FalseName, InTypeIndex);
+		Types.Add(ResultName, InTypeIndex);
+	}
 	return Types;
 }
 
diff --git a/Engine/Plugins/Runtime/RigVM/Source/RigVM/Private/RigVMFunctions/RigVMDispatch_MakeStruct.cpp b/Engine/Plugins/Runtime/RigVM/Source/RigVM/Private/RigVMFunctions/RigVMDispatch_MakeStruct.cpp
--- a/Engine/Plugins/Runtime/RigVM/Source/RigVM/Private/RigVMFunctions/RigVMDispatch_MakeStruct.cpp
+++ b/Engine/Plugins/Runtime/RigVM/Source/RigVM/Private/RigVMFunctions/RigVMDispatch_MakeStruct.cpp
@@ -11,15 +11,19 @@
 const FName FRigVMDispatch_MakeStruct::ElementsName = TEXT("Elements");
 const FName FRigVMDispatch_MakeStruct::StructName = TEXT("Struct");
 
-TArray<FRigVMTemplateArgument> FRigVMDispatch_MakeStruct::GetArguments() const
+const TArray<FRigVMTemplateArgument>& FRigVMDispatch_MakeStruct::GetArguments() const
 {
-	const TArray<FRigVMTemplateArgument::ETypeCategory> Categories = {
-		FRigVMTemplateArgument::ETypeCategory_SingleScriptStructValue
-	};
-	return {
-		{ElementsName, ERigVMPinDirection::Input, Categories},
-		{StructName, ERigVMPinDirection::Output, Categories}
-	};
+	static TArray<FRigVMTemplateArgument> Arguments;
+	if (Arguments.IsEmpty())
+	{
+		const TArray<FRigVMTemplateArgument::ETypeCategory> Categories = {
+			FRigVMTemplateArgument::ETypeCategory_SingleScriptStructValue
+		};
+		
+		Arguments.Emplace(ElementsName, ERigVMPinDirection::Input, Categories);
+		Arguments.Emplace(StructName, ERigVMPinDirection::Output, Categories);
+	}
+	return Arguments;
 }
 
 FRigVMTemplateTypeMap FRigVMDispatch_MakeStruct::OnNewArgumentType(const FName& InArgumentName, TRigVMTypeIndex InTypeIndex) const
@@ -84,15 +88,19 @@ void FRigVMDispatch_MakeStruct::Execute(FRigVMExtendedExecuteContext& InContext,
 	URigVMMemoryStorage::CopyProperty(TargetProperty, TargetMemory, SourceProperty, SourceMemory);
 }
 
-TArray<FRigVMTemplateArgument> FRigVMDispatch_BreakStruct::GetArguments() const
+const TArray<FRigVMTemplateArgument>& FRigVMDispatch_BreakStruct::GetArguments() const
 {
-	const TArray<FRigVMTemplateArgument::ETypeCategory> Categories = {
-		FRigVMTemplateArgument::ETypeCategory_SingleScriptStructValue
-	};
-	return {
-		{StructName, ERigVMPinDirection::Input, Categories},
-		{ElementsName, ERigVMPinDirection::Output, Categories}
-	};
+	static TArray<FRigVMTemplateArgument> Arguments;
+	if (Arguments.IsEmpty())
+	{
+		const TArray<FRigVMTemplateArgument::ETypeCategory> Categories = {
+			FRigVMTemplateArgument::ETypeCategory_SingleScriptStructValue
+		};
+		
+		Arguments.Emplace(StructName, ERigVMPinDirection::Input, Categories);
+		Arguments.Emplace(ElementsName, ERigVMPinDirection::Output, Categories);
+	}
+	return Arguments;
 }
 
 #if WITH_EDITOR
diff --git a/Engine/Plugins/Runtime/RigVM/Source/RigVM/Private/RigVMFunctions/RigVMDispatch_Print.cpp b/Engine/Plugins/Runtime/RigVM/Source/RigVM/Private/RigVMFunctions/RigVMDispatch_Print.cpp
--- a/Engine/Plugins/Runtime/RigVM/Source/RigVM/Private/RigVMFunctions/RigVMDispatch_Print.cpp
+++ b/Engine/Plugins/Runtime/RigVM/Source/RigVM/Private/RigVMFunctions/RigVMDispatch_Print.cpp
@@ -28,19 +28,23 @@ FName FRigVMDispatch_Print::GetArgumentNameForOperandIndex(int32 InOperandIndex,
 	return ArgumentNames[InOperandIndex];
 }
 
-TArray<FRigVMTemplateArgument> FRigVMDispatch_Print::GetArguments() const
+const TArray<FRigVMTemplateArgument>& FRigVMDispatch_Print::GetArguments() const
 {
-	const TArray<FRigVMTemplateArgument::ETypeCategory> ValueCategories = {
-		FRigVMTemplateArgument::ETypeCategory_SingleAnyValue,
-		FRigVMTemplateArgument::ETypeCategory_ArrayAnyValue
-	};
-	return {
-		FRigVMTemplateArgument(PrefixName, ERigVMPinDirection::Input, RigVMTypeUtils::TypeIndex::FString),
-		FRigVMTemplateArgument(ValueName, ERigVMPinDirection::Input, ValueCategories),
-		FRigVMTemplateArgument(EnabledName, ERigVMPinDirection::Input, RigVMTypeUtils::TypeIndex::Bool),
-		FRigVMTemplateArgument(ScreenDurationName, ERigVMPinDirection::Input, RigVMTypeUtils::TypeIndex::Float),
-		FRigVMTemplateArgument(ScreenColorName, ERigVMPinDirection::Input, FRigVMRegistry::Get().GetTypeIndex<FLinearColor>())
-	};
+	static TArray<FRigVMTemplateArgument> Arguments;
+	if (Arguments.IsEmpty())
+	{
+		const TArray<FRigVMTemplateArgument::ETypeCategory> ValueCategories = {
+			FRigVMTemplateArgument::ETypeCategory_SingleAnyValue,
+			FRigVMTemplateArgument::ETypeCategory_ArrayAnyValue
+		};
+		
+		Arguments.Emplace(PrefixName, ERigVMPinDirection::Input, RigVMTypeUtils::TypeIndex::FString);
+		Arguments.Emplace(ValueName, ERigVMPinDirection::Input, ValueCategories);
+		Arguments.Emplace(EnabledName, ERigVMPinDirection::Input, RigVMTypeUtils::TypeIndex::Bool);
+		Arguments.Emplace(ScreenDurationName, ERigVMPinDirection::Input, RigVMTypeUtils::TypeIndex::Float);
+		Arguments.Emplace(ScreenColorName, ERigVMPinDirection::Input, FRigVMRegistry::Get().GetTypeIndex<FLinearColor>());
+	}
+	return Arguments;
 }
 
 TArray<FRigVMExecuteArgument> FRigVMDispatch_Print::GetExecuteArguments_Impl(const FRigVMDispatchContext& InContext) const
@@ -48,8 +52,7 @@ TArray<FRigVMExecuteArgument> FRigVMDispatch_Print::GetExecuteArguments_Impl(con
 	return {{TEXT("ExecuteContext"), ERigVMPinDirection::IO}};
 }
 
-FRigVMTemplateTypeMap FRigVMDispatch_Print::OnNewArgumentType(const FName& InArgumentName,
-                                                              TRigVMTypeIndex InTypeIndex) const
+FRigVMTemplateTypeMap FRigVMDispatch_Print::OnNewArgumentType(const FName& InArgumentName, TRigVMTypeIndex InTypeIndex) const
 {
 	FRigVMTemplateTypeMap Types;
 	Types.Add(PrefixName, RigVMTypeUtils::TypeIndex::FString);
diff --git a/Engine/Plugins/Runtime/RigVM/Source/RigVM/Private/RigVMFunctions/RigVMDispatch_Select.cpp b/Engine/Plugins/Runtime/RigVM/Source/RigVM/Private/RigVMFunctions/RigVMDispatch_Select.cpp
--- a/Engine/Plugins/Runtime/RigVM/Source/RigVM/Private/RigVMFunctions/RigVMDispatch_Select.cpp
+++ b/Engine/Plugins/Runtime/RigVM/Source/RigVM/Private/RigVMFunctions/RigVMDispatch_Select.cpp
@@ -23,26 +23,28 @@ FName FRigVMDispatch_SelectInt32::GetArgumentNameForOperandIndex(int32 InOperand
 }
 
 
-TArray<FRigVMTemplateArgument> FRigVMDispatch_SelectInt32::GetArguments() const
+const TArray<FRigVMTemplateArgument>& FRigVMDispatch_SelectInt32::GetArguments() const
 {
-	static const TArray<FRigVMTemplateArgument::ETypeCategory> ArrayCategories = {
-		FRigVMTemplateArgument::ETypeCategory_ArrayAnyValue,
-		FRigVMTemplateArgument::ETypeCategory_ArrayArrayAnyValue
-	};
-	static const TArray<FRigVMTemplateArgument::ETypeCategory> SingleCategories = {
-		FRigVMTemplateArgument::ETypeCategory_SingleAnyValue,
-		FRigVMTemplateArgument::ETypeCategory_ArrayAnyValue
-	};
+	static TArray<FRigVMTemplateArgument> Arguments;
+	if (Arguments.IsEmpty())
+	{
+		static const TArray<FRigVMTemplateArgument::ETypeCategory> ArrayCategories = {
+			FRigVMTemplateArgument::ETypeCategory_ArrayAnyValue,
+			FRigVMTemplateArgument::ETypeCategory_ArrayArrayAnyValue
+		};
+		static const TArray<FRigVMTemplateArgument::ETypeCategory> SingleCategories = {
+			FRigVMTemplateArgument::ETypeCategory_SingleAnyValue,
+			FRigVMTemplateArgument::ETypeCategory_ArrayAnyValue
+		};
 	
-	return {
-		FRigVMTemplateArgument(IndexName, ERigVMPinDirection::Input, RigVMTypeUtils::TypeIndex::Int32),
-		FRigVMTemplateArgument(ValuesName, ERigVMPinDirection::Input, ArrayCategories),
-		FRigVMTemplateArgument(ResultName, ERigVMPinDirection::Output, SingleCategories)
-	};
+		Arguments.Emplace(IndexName, ERigVMPinDirection::Input, RigVMTypeUtils::TypeIndex::Int32);
+		Arguments.Emplace(ValuesName, ERigVMPinDirection::Input, ArrayCategories);
+		Arguments.Emplace(ResultName, ERigVMPinDirection::Output, SingleCategories);
+	}
+	return Arguments;
 }
 
-FRigVMTemplateTypeMap FRigVMDispatch_SelectInt32::OnNewArgumentType(const FName& InArgumentName,
-	TRigVMTypeIndex InTypeIndex) const
+FRigVMTemplateTypeMap FRigVMDispatch_SelectInt32::OnNewArgumentType(const FName& InArgumentName, TRigVMTypeIndex InTypeIndex) const
 {
 	FRigVMTemplateTypeMap Types;
 	Types.Add(IndexName, RigVMTypeUtils::TypeIndex::Int32);
diff --git a/Engine/Plugins/Runtime/RigVM/Source/RigVM/Private/RigVMFunctions/RigVMDispatch_Switch.cpp b/Engine/Plugins/Runtime/RigVM/Source/RigVM/Private/RigVMFunctions/RigVMDispatch_Switch.cpp
--- a/Engine/Plugins/Runtime/RigVM/Source/RigVM/Private/RigVMFunctions/RigVMDispatch_Switch.cpp
+++ b/Engine/Plugins/Runtime/RigVM/Source/RigVM/Private/RigVMFunctions/RigVMDispatch_Switch.cpp
@@ -20,12 +20,15 @@ FName FRigVMDispatch_SwitchInt32::GetArgumentNameForOperandIndex(int32 InOperand
 	return ArgumentNames[InOperandIndex];
 }
 
-TArray<FRigVMTemplateArgument> FRigVMDispatch_SwitchInt32::GetArguments() const
+const TArray<FRigVMTemplateArgument>& FRigVMDispatch_SwitchInt32::GetArguments() const
 {
-	return {
-		FRigVMTemplateArgument(IndexName, ERigVMPinDirection::Input, RigVMTypeUtils::TypeIndex::Int32),
-		FRigVMTemplateArgument(FRigVMStruct::ControlFlowBlockToRunName, ERigVMPinDirection::Hidden, RigVMTypeUtils::TypeIndex::FName)
-	};
+	static TArray<FRigVMTemplateArgument> Arguments;
+	if (Arguments.IsEmpty())
+	{
+		Arguments.Emplace(IndexName, ERigVMPinDirection::Input, RigVMTypeUtils::TypeIndex::Int32);
+		Arguments.Emplace(FRigVMStruct::ControlFlowBlockToRunName, ERigVMPinDirection::Hidden, RigVMTypeUtils::TypeIndex::FName);
+	}
+	return Arguments;
 }
 
 TArray<FRigVMExecuteArgument> FRigVMDispatch_SwitchInt32::GetExecuteArguments_Impl(const FRigVMDispatchContext& InContext) const
diff --git a/Engine/Plugins/Runtime/RigVM/Source/RigVM/Private/RigVMFunctions/RigVMFunction_String.cpp b/Engine/Plugins/Runtime/RigVM/Source/RigVM/Private/RigVMFunctions/RigVMFunction_String.cpp
--- a/Engine/Plugins/Runtime/RigVM/Source/RigVM/Private/RigVMFunctions/RigVMFunction_String.cpp
+++ b/Engine/Plugins/Runtime/RigVM/Source/RigVM/Private/RigVMFunctions/RigVMFunction_String.cpp
@@ -252,20 +252,23 @@ FRigVMFunction_StringPadInteger_Execute()
 	}
 }
 
-TArray<FRigVMTemplateArgument> FRigDispatch_ToString::GetArguments() const
+const TArray<FRigVMTemplateArgument>& FRigDispatch_ToString::GetArguments() const
 {
-	const TArray<FRigVMTemplateArgument::ETypeCategory> ValueCategories = {
-		FRigVMTemplateArgument::ETypeCategory_SingleAnyValue,
-		FRigVMTemplateArgument::ETypeCategory_ArrayAnyValue
-	};
-	return {
-		FRigVMTemplateArgument(TEXT("Value"), ERigVMPinDirection::Input, ValueCategories),
-		FRigVMTemplateArgument(TEXT("Result"), ERigVMPinDirection::Output, RigVMTypeUtils::TypeIndex::FString)
-	};
+	static TArray<FRigVMTemplateArgument> Arguments;
+	if (Arguments.IsEmpty())
+	{
+		static const TArray<FRigVMTemplateArgument::ETypeCategory> ValueCategories = {
+			FRigVMTemplateArgument::ETypeCategory_SingleAnyValue,
+			FRigVMTemplateArgument::ETypeCategory_ArrayAnyValue
+		};
+
+		Arguments.Emplace(TEXT("Value"), ERigVMPinDirection::Input, ValueCategories);
+		Arguments.Emplace(TEXT("Result"), ERigVMPinDirection::Output, RigVMTypeUtils::TypeIndex::FString);
+	}
+	return Arguments;
 }
 
-FRigVMTemplateTypeMap FRigDispatch_ToString::OnNewArgumentType(const FName& InArgumentName,
-	TRigVMTypeIndex InTypeIndex) const
+FRigVMTemplateTypeMap FRigDispatch_ToString::OnNewArgumentType(const FName& InArgumentName, TRigVMTypeIndex InTypeIndex) const
 {
 	FRigVMTemplateTypeMap Types;
 	Types.Add(TEXT("Value"), InTypeIndex);
@@ -317,20 +320,23 @@ void FRigDispatch_ToString::Execute(FRigVMExtendedExecuteContext& InContext, FRi
 	ValueProperty->ExportText_Direct(Result, Value, Value, nullptr, PPF_None, nullptr);
 }
 
-TArray<FRigVMTemplateArgument> FRigDispatch_FromString::GetArguments() const
+const TArray<FRigVMTemplateArgument>& FRigDispatch_FromString::GetArguments() const
 {
-	const TArray<FRigVMTemplateArgument::ETypeCategory> ValueCategories = {
-		FRigVMTemplateArgument::ETypeCategory_SingleAnyValue,
-		FRigVMTemplateArgument::ETypeCategory_ArrayAnyValue
-	};
-	return {
-		FRigVMTemplateArgument(TEXT("String"), ERigVMPinDirection::Input, RigVMTypeUtils::TypeIndex::FString),
-		FRigVMTemplateArgument(TEXT("Result"), ERigVMPinDirection::Output, ValueCategories)
-	};
+	static TArray<FRigVMTemplateArgument> Arguments;
+	if (Arguments.IsEmpty())
+	{
+		const TArray<FRigVMTemplateArgument::ETypeCategory> ValueCategories = {
+			FRigVMTemplateArgument::ETypeCategory_SingleAnyValue,
+			FRigVMTemplateArgument::ETypeCategory_ArrayAnyValue
+		};
+	
+		Arguments.Emplace(TEXT("String"), ERigVMPinDirection::Input, RigVMTypeUtils::TypeIndex::FString);
+		Arguments.Emplace(TEXT("Result"), ERigVMPinDirection::Output, ValueCategories);
+	}
+	return Arguments;
 }
 
-FRigVMTemplateTypeMap FRigDispatch_FromString::OnNewArgumentType(const FName& InArgumentName,
-	TRigVMTypeIndex InTypeIndex) const
+FRigVMTemplateTypeMap FRigDispatch_FromString::OnNewArgumentType(const FName& InArgumentName, TRigVMTypeIndex InTypeIndex) const
 {
 	FRigVMTemplateTypeMap Types;
 	Types.Add(TEXT("String"), RigVMTypeUtils::TypeIndex::FString);
diff --git a/Engine/Plugins/Runtime/RigVM/Source/RigVM/Public/RigVMCore/RigVMDispatchFactory.h b/Engine/Plugins/Runtime/RigVM/Source/RigVM/Public/RigVMCore/RigVMDispatchFactory.h
--- a/Engine/Plugins/Runtime/RigVM/Source/RigVM/Public/RigVMCore/RigVMDispatchFactory.h
+++ b/Engine/Plugins/Runtime/RigVM/Source/RigVM/Public/RigVMCore/RigVMDispatchFactory.h
@@ -135,7 +135,11 @@ public:
 	virtual void RegisterDependencyTypes() const {}
 
 	// returns the arguments of the template
-	virtual TArray<FRigVMTemplateArgument> GetArguments() const { return TArray<FRigVMTemplateArgument>(); }
+	virtual const TArray<FRigVMTemplateArgument>& GetArguments() const
+	{
+		static TArray<FRigVMTemplateArgument> Arguments;
+		return Arguments;
+	}
 
 	// returns the execute arguments of the template
 	RIGVM_API TArray<FRigVMExecuteArgument> GetExecuteArguments(const FRigVMDispatchContext& InContext) const;
diff --git a/Engine/Plugins/Runtime/RigVM/Source/RigVM/Public/RigVMFunctions/RigVMDispatch_Array.h b/Engine/Plugins/Runtime/RigVM/Source/RigVM/Public/RigVMFunctions/RigVMDispatch_Array.h
--- a/Engine/Plugins/Runtime/RigVM/Source/RigVM/Public/RigVMFunctions/RigVMDispatch_Array.h
+++ b/Engine/Plugins/Runtime/RigVM/Source/RigVM/Public/RigVMFunctions/RigVMDispatch_Array.h
@@ -62,7 +62,7 @@ public:
 		FactoryScriptStruct = StaticStruct();
 	}
 
-	virtual TArray<FRigVMTemplateArgument> GetArguments() const override;
+	virtual const TArray<FRigVMTemplateArgument>& GetArguments() const override;
 	virtual FRigVMTemplateTypeMap OnNewArgumentType(const FName& InArgumentName, TRigVMTypeIndex InTypeIndex) const override;
 #if WITH_EDITOR
 	virtual FText GetNodeTooltip(const FRigVMTemplateTypeMap& InTypes) const override;
@@ -90,7 +90,7 @@ public:
 
 	virtual ERigVMOpCode GetOpCode() const override { return ERigVMOpCode::ArrayReset; }
 	virtual FName GetArgumentNameForOperandIndex(int32 InOperandIndex, int32 InTotalOperands) const override;
-	virtual TArray<FRigVMTemplateArgument> GetArguments() const override;
+	virtual const TArray<FRigVMTemplateArgument>& GetArguments() const override;
 	virtual FRigVMTemplateTypeMap OnNewArgumentType(const FName& InArgumentName, TRigVMTypeIndex InTypeIndex) const override;
 #if WITH_EDITOR
 	virtual FText GetNodeTooltip(const FRigVMTemplateTypeMap& InTypes) const override;
@@ -115,7 +115,7 @@ public:
 
 	virtual ERigVMOpCode GetOpCode() const override { return ERigVMOpCode::ArrayGetNum; }
 	virtual FName GetArgumentNameForOperandIndex(int32 InOperandIndex, int32 InTotalOperands) const override;
-	virtual TArray<FRigVMTemplateArgument> GetArguments() const override;
+	virtual const TArray<FRigVMTemplateArgument>& GetArguments() const override;
 	virtual FRigVMTemplateTypeMap OnNewArgumentType(const FName& InArgumentName, TRigVMTypeIndex InTypeIndex) const override;
 #if WITH_EDITOR
 	virtual FText GetNodeTooltip(const FRigVMTemplateTypeMap& InTypes) const override;
@@ -140,7 +140,7 @@ public:
 
 	virtual ERigVMOpCode GetOpCode() const override { return ERigVMOpCode::ArraySetNum; }
 	virtual FName GetArgumentNameForOperandIndex(int32 InOperandIndex, int32 InTotalOperands) const override;
-	virtual TArray<FRigVMTemplateArgument> GetArguments() const override;
+	virtual const TArray<FRigVMTemplateArgument>& GetArguments() const override;
 	virtual FRigVMTemplateTypeMap OnNewArgumentType(const FName& InArgumentName, TRigVMTypeIndex InTypeIndex) const override;
 #if WITH_EDITOR
 	virtual FText GetNodeTooltip(const FRigVMTemplateTypeMap& InTypes) const override;
@@ -165,7 +165,7 @@ public:
 
 	virtual ERigVMOpCode GetOpCode() const override { return ERigVMOpCode::ArrayGetAtIndex; }
 	virtual FName GetArgumentNameForOperandIndex(int32 InOperandIndex, int32 InTotalOperands) const override;
-	virtual TArray<FRigVMTemplateArgument> GetArguments() const override;
+	virtual const TArray<FRigVMTemplateArgument>& GetArguments() const override;
 	virtual FRigVMTemplateTypeMap OnNewArgumentType(const FName& InArgumentName, TRigVMTypeIndex InTypeIndex) const override;
 #if WITH_EDITOR
 	virtual FText GetNodeTooltip(const FRigVMTemplateTypeMap& InTypes) const override;
@@ -190,7 +190,7 @@ public:
 
 	virtual ERigVMOpCode GetOpCode() const override { return ERigVMOpCode::ArraySetAtIndex; }
 	virtual FName GetArgumentNameForOperandIndex(int32 InOperandIndex, int32 InTotalOperands) const override;
-	virtual TArray<FRigVMTemplateArgument> GetArguments() const override;
+	virtual const TArray<FRigVMTemplateArgument>& GetArguments() const override;
 	virtual FRigVMTemplateTypeMap OnNewArgumentType(const FName& InArgumentName, TRigVMTypeIndex InTypeIndex) const override;
 #if WITH_EDITOR
 	virtual FText GetNodeTooltip(const FRigVMTemplateTypeMap& InTypes) const override;
@@ -215,7 +215,7 @@ public:
 
 	virtual ERigVMOpCode GetOpCode() const override { return ERigVMOpCode::ArrayAdd; }
 	virtual FName GetArgumentNameForOperandIndex(int32 InOperandIndex, int32 InTotalOperands) const override;
-	virtual TArray<FRigVMTemplateArgument> GetArguments() const override;
+	virtual const TArray<FRigVMTemplateArgument>& GetArguments() const override;
 #if WITH_EDITOR
 	virtual FText GetNodeTooltip(const FRigVMTemplateTypeMap& InTypes) const override;
 	virtual FText GetArgumentTooltip(const FName& InArgumentName, TRigVMTypeIndex InTypeIndex) const override;
@@ -261,7 +261,7 @@ public:
 
 	virtual ERigVMOpCode GetOpCode() const override { return ERigVMOpCode::ArrayRemove; }
 	virtual FName GetArgumentNameForOperandIndex(int32 InOperandIndex, int32 InTotalOperands) const override;
-	virtual TArray<FRigVMTemplateArgument> GetArguments() const override;
+	virtual const TArray<FRigVMTemplateArgument>& GetArguments() const override;
 	virtual FRigVMTemplateTypeMap OnNewArgumentType(const FName& InArgumentName, TRigVMTypeIndex InTypeIndex) const override;
 #if WITH_EDITOR
    	virtual FText GetNodeTooltip(const FRigVMTemplateTypeMap& InTypes) const override;
@@ -309,7 +309,7 @@ public:
 
 	virtual ERigVMOpCode GetOpCode() const override { return ERigVMOpCode::ArrayFind; }
 	virtual FName GetArgumentNameForOperandIndex(int32 InOperandIndex, int32 InTotalOperands) const override;
-	virtual TArray<FRigVMTemplateArgument> GetArguments() const override;
+	virtual const TArray<FRigVMTemplateArgument>& GetArguments() const override;
 	virtual FRigVMTemplateTypeMap OnNewArgumentType(const FName& InArgumentName, TRigVMTypeIndex InTypeIndex) const override;
 #if WITH_EDITOR
    	virtual FText GetNodeTooltip(const FRigVMTemplateTypeMap& InTypes) const override;
@@ -334,7 +334,7 @@ public:
 
 	virtual ERigVMOpCode GetOpCode() const override { return ERigVMOpCode::ArrayAppend; }
 	virtual FName GetArgumentNameForOperandIndex(int32 InOperandIndex, int32 InTotalOperands) const override;
-	virtual TArray<FRigVMTemplateArgument> GetArguments() const override;
+	virtual const TArray<FRigVMTemplateArgument>& GetArguments() const override;
 	virtual FRigVMTemplateTypeMap OnNewArgumentType(const FName& InArgumentName, TRigVMTypeIndex InTypeIndex) const override;
 #if WITH_EDITOR
    	virtual FText GetNodeTooltip(const FRigVMTemplateTypeMap& InTypes) const override;
@@ -359,7 +359,7 @@ public:
 
 	virtual ERigVMOpCode GetOpCode() const override { return ERigVMOpCode::ArrayClone; }
 	virtual FName GetArgumentNameForOperandIndex(int32 InOperandIndex, int32 InTotalOperands) const override;
-	virtual TArray<FRigVMTemplateArgument> GetArguments() const override;
+	virtual const TArray<FRigVMTemplateArgument>& GetArguments() const override;
 	virtual FRigVMTemplateTypeMap OnNewArgumentType(const FName& InArgumentName, TRigVMTypeIndex InTypeIndex) const override;
 #if WITH_EDITOR
    	virtual FText GetNodeTooltip(const FRigVMTemplateTypeMap& InTypes) const override;
@@ -406,7 +406,7 @@ public:
 
 	virtual ERigVMOpCode GetOpCode() const override { return ERigVMOpCode::ArrayDifference; }
 	virtual FName GetArgumentNameForOperandIndex(int32 InOperandIndex, int32 InTotalOperands) const override;
-	virtual TArray<FRigVMTemplateArgument> GetArguments() const override;
+	virtual const TArray<FRigVMTemplateArgument>& GetArguments() const override;
 	virtual FRigVMTemplateTypeMap OnNewArgumentType(const FName& InArgumentName, TRigVMTypeIndex InTypeIndex) const override;
 #if WITH_EDITOR
    	virtual FText GetNodeTooltip(const FRigVMTemplateTypeMap& InTypes) const override;
@@ -453,7 +453,7 @@ public:
 
 	virtual ERigVMOpCode GetOpCode() const override { return ERigVMOpCode::ArrayIterator; }
 	virtual FName GetArgumentNameForOperandIndex(int32 InOperandIndex, int32 InTotalOperands) const override;
-	virtual TArray<FRigVMTemplateArgument> GetArguments() const override;
+	virtual const TArray<FRigVMTemplateArgument>& GetArguments() const override;
 	virtual TArray<FRigVMExecuteArgument> GetExecuteArguments_Impl(const FRigVMDispatchContext& InContext) const override;
 	virtual const TArray<FName>& GetControlFlowBlocks_Impl(const FRigVMDispatchContext& InContext) const override;
 	virtual const bool IsControlFlowBlockSliced(const FName& InBlockName) const override;
diff --git a/Engine/Plugins/Runtime/RigVM/Source/RigVM/Public/RigVMFunctions/RigVMDispatch_Constant.h b/Engine/Plugins/Runtime/RigVM/Source/RigVM/Public/RigVMFunctions/RigVMDispatch_Constant.h
--- a/Engine/Plugins/Runtime/RigVM/Source/RigVM/Public/RigVMFunctions/RigVMDispatch_Constant.h
+++ b/Engine/Plugins/Runtime/RigVM/Source/RigVM/Public/RigVMFunctions/RigVMDispatch_Constant.h
@@ -18,8 +18,9 @@ public:
 		FactoryScriptStruct = StaticStruct();
 	}
 
-	virtual TArray<FRigVMTemplateArgument> GetArguments() const override;
-	virtual FRigVMTemplateTypeMap OnNewArgumentType(const FName& InArgumentName, TRigVMTypeIndex InTypeIndex) const override;
+	virtual const TArray<FRigVMTemplateArgument>& GetArguments() const override;
+	virtual FRigVMTemplateTypeMap
+	OnNewArgumentType(const FName& InArgumentName, TRigVMTypeIndex InTypeIndex) const override;
 #if WITH_EDITOR
 	virtual FString GetNodeTitle(const FRigVMTemplateTypeMap& InTypes) const override;
 	virtual FText GetNodeTooltip(const FRigVMTemplateTypeMap& InTypes) const override;
diff --git a/Engine/Plugins/Runtime/RigVM/Source/RigVM/Public/RigVMFunctions/RigVMDispatch_Core.h b/Engine/Plugins/Runtime/RigVM/Source/RigVM/Public/RigVMFunctions/RigVMDispatch_Core.h
--- a/Engine/Plugins/Runtime/RigVM/Source/RigVM/Public/RigVMFunctions/RigVMDispatch_Core.h
+++ b/Engine/Plugins/Runtime/RigVM/Source/RigVM/Public/RigVMFunctions/RigVMDispatch_Core.h
@@ -27,7 +27,7 @@ public:
 	}
 
 	virtual FName GetArgumentNameForOperandIndex(int32 InOperandIndex, int32 InTotalOperands) const override;
-	virtual TArray<FRigVMTemplateArgument> GetArguments() const override;
+	virtual const TArray<FRigVMTemplateArgument>& GetArguments() const override;
 	virtual FRigVMTemplateTypeMap OnNewArgumentType(const FName& InArgumentName, TRigVMTypeIndex InTypeIndex) const override;
 
 protected:
diff --git a/Engine/Plugins/Runtime/RigVM/Source/RigVM/Public/RigVMFunctions/RigVMDispatch_If.h b/Engine/Plugins/Runtime/RigVM/Source/RigVM/Public/RigVMFunctions/RigVMDispatch_If.h
--- a/Engine/Plugins/Runtime/RigVM/Source/RigVM/Public/RigVMFunctions/RigVMDispatch_If.h
+++ b/Engine/Plugins/Runtime/RigVM/Source/RigVM/Public/RigVMFunctions/RigVMDispatch_If.h
@@ -22,7 +22,7 @@ public:
 	}
 
 	virtual FName GetArgumentNameForOperandIndex(int32 InOperandIndex, int32 InTotalOperands) const override;
-	virtual TArray<FRigVMTemplateArgument> GetArguments() const override;
+	virtual const TArray<FRigVMTemplateArgument>& GetArguments() const override;
 	virtual FRigVMTemplateTypeMap OnNewArgumentType(const FName& InArgumentName, TRigVMTypeIndex InTypeIndex) const override;
 	virtual bool IsSingleton() const override { return true; } 
 
diff --git a/Engine/Plugins/Runtime/RigVM/Source/RigVM/Public/RigVMFunctions/RigVMDispatch_MakeStruct.h b/Engine/Plugins/Runtime/RigVM/Source/RigVM/Public/RigVMFunctions/RigVMDispatch_MakeStruct.h
--- a/Engine/Plugins/Runtime/RigVM/Source/RigVM/Public/RigVMFunctions/RigVMDispatch_MakeStruct.h
+++ b/Engine/Plugins/Runtime/RigVM/Source/RigVM/Public/RigVMFunctions/RigVMDispatch_MakeStruct.h
@@ -18,7 +18,7 @@ public:
 		FactoryScriptStruct = StaticStruct();
 	}
 	
-	virtual TArray<FRigVMTemplateArgument> GetArguments() const override;
+	virtual const TArray<FRigVMTemplateArgument>& GetArguments() const override;
 	virtual FRigVMTemplateTypeMap OnNewArgumentType(const FName& InArgumentName, TRigVMTypeIndex InTypeIndex) const override;
 #if WITH_EDITOR
 	virtual FString GetNodeTitle(const FRigVMTemplateTypeMap& InTypes) const override;
@@ -48,7 +48,7 @@ public:
 		FactoryScriptStruct = StaticStruct();
 	}
 
-	virtual TArray<FRigVMTemplateArgument> GetArguments() const override;
+	virtual const TArray<FRigVMTemplateArgument>& GetArguments() const override;
 #if WITH_EDITOR
 	virtual FText GetNodeTooltip(const FRigVMTemplateTypeMap& InTypes) const override;
 #endif
diff --git a/Engine/Plugins/Runtime/RigVM/Source/RigVM/Public/RigVMFunctions/RigVMDispatch_Print.h b/Engine/Plugins/Runtime/RigVM/Source/RigVM/Public/RigVMFunctions/RigVMDispatch_Print.h
--- a/Engine/Plugins/Runtime/RigVM/Source/RigVM/Public/RigVMFunctions/RigVMDispatch_Print.h
+++ b/Engine/Plugins/Runtime/RigVM/Source/RigVM/Public/RigVMFunctions/RigVMDispatch_Print.h
@@ -21,7 +21,7 @@ public:
 	}
 
 	virtual FName GetArgumentNameForOperandIndex(int32 InOperandIndex, int32 InTotalOperands) const override;
-	virtual TArray<FRigVMTemplateArgument> GetArguments() const override;
+	virtual const TArray<FRigVMTemplateArgument>& GetArguments() const override;
 	virtual TArray<FRigVMExecuteArgument> GetExecuteArguments_Impl(const FRigVMDispatchContext& InContext) const override;
 	virtual FRigVMTemplateTypeMap OnNewArgumentType(const FName& InArgumentName, TRigVMTypeIndex InTypeIndex) const override;
 	virtual bool IsSingleton() const override { return true; } 
diff --git a/Engine/Plugins/Runtime/RigVM/Source/RigVM/Public/RigVMFunctions/RigVMDispatch_Select.h b/Engine/Plugins/Runtime/RigVM/Source/RigVM/Public/RigVMFunctions/RigVMDispatch_Select.h
--- a/Engine/Plugins/Runtime/RigVM/Source/RigVM/Public/RigVMFunctions/RigVMDispatch_Select.h
+++ b/Engine/Plugins/Runtime/RigVM/Source/RigVM/Public/RigVMFunctions/RigVMDispatch_Select.h
@@ -22,7 +22,7 @@ public:
 	}
 
 	virtual FName GetArgumentNameForOperandIndex(int32 InOperandIndex, int32 InTotalOperands) const override;
-	virtual TArray<FRigVMTemplateArgument> GetArguments() const override;
+	virtual const TArray<FRigVMTemplateArgument>& GetArguments() const override;
 	virtual FRigVMTemplateTypeMap OnNewArgumentType(const FName& InArgumentName, TRigVMTypeIndex InTypeIndex) const override;
 #if WITH_EDITOR
 	virtual FString GetArgumentMetaData(const FName& InArgumentName, const FName& InMetaDataKey) const override;
diff --git a/Engine/Plugins/Runtime/RigVM/Source/RigVM/Public/RigVMFunctions/RigVMDispatch_Switch.h b/Engine/Plugins/Runtime/RigVM/Source/RigVM/Public/RigVMFunctions/RigVMDispatch_Switch.h
--- a/Engine/Plugins/Runtime/RigVM/Source/RigVM/Public/RigVMFunctions/RigVMDispatch_Switch.h
+++ b/Engine/Plugins/Runtime/RigVM/Source/RigVM/Public/RigVMFunctions/RigVMDispatch_Switch.h
@@ -22,7 +22,7 @@ public:
 	}
 
 	virtual FName GetArgumentNameForOperandIndex(int32 InOperandIndex, int32 InTotalOperands) const override;
-	virtual TArray<FRigVMTemplateArgument> GetArguments() const override;
+	virtual const TArray<FRigVMTemplateArgument>& GetArguments() const override;
 	virtual TArray<FRigVMExecuteArgument> GetExecuteArguments_Impl(const FRigVMDispatchContext& InContext) const override;
 	virtual bool IsSingleton() const override { return true; } 
 
diff --git a/Engine/Plugins/Runtime/RigVM/Source/RigVM/Public/RigVMFunctions/RigVMFunction_String.h b/Engine/Plugins/Runtime/RigVM/Source/RigVM/Public/RigVMFunctions/RigVMFunction_String.h
--- a/Engine/Plugins/Runtime/RigVM/Source/RigVM/Public/RigVMFunctions/RigVMFunction_String.h
+++ b/Engine/Plugins/Runtime/RigVM/Source/RigVM/Public/RigVMFunctions/RigVMFunction_String.h
@@ -510,7 +510,7 @@ struct FRigDispatch_ToString : public FRigVMDispatchFactory
 
 public:
 
-	virtual TArray<FRigVMTemplateArgument> GetArguments() const override;
+	virtual const TArray<FRigVMTemplateArgument>& GetArguments() const override;
 	virtual FRigVMTemplateTypeMap OnNewArgumentType(const FName& InArgumentName, TRigVMTypeIndex InTypeIndex) const override;
 
 protected:
@@ -529,7 +529,7 @@ struct FRigDispatch_FromString : public FRigVMDispatchFactory
 
 public:
 
-	virtual TArray<FRigVMTemplateArgument> GetArguments() const override;
+	virtual const TArray<FRigVMTemplateArgument>& GetArguments() const override;
 	virtual FRigVMTemplateTypeMap OnNewArgumentType(const FName& InArgumentName, TRigVMTypeIndex InTypeIndex) const override;
 
 protected:
