RigVM: Patch functions which do not connect all paths to a return node
#rnx

diff --git a/Engine/Plugins/Runtime/RigVM/Source/RigVMDeveloper/Private/RigVMModel/RigVMClient.cpp b/Engine/Plugins/Runtime/RigVM/Source/RigVMDeveloper/Private/RigVMModel/RigVMClient.cpp
--- a/Engine/Plugins/Runtime/RigVM/Source/RigVMDeveloper/Private/RigVMModel/RigVMClient.cpp
+++ b/Engine/Plugins/Runtime/RigVM/Source/RigVMDeveloper/Private/RigVMModel/RigVMClient.cpp
@@ -883,6 +883,11 @@ FRigVMClientPatchResult FRigVMClient::PatchModelsOnLoad()
 		Result.Merge(Controller->PatchArrayNodesOnLoad());
 		Result.Merge(Controller->PatchReduceArrayFloatDoubleConvertsionsOnLoad());
 		Result.Merge(Controller->PatchInvalidLinksOnWildcards());
+
+		if (URigVMCollapseNode* CollapseNode = Model->GetTypedOuter<URigVMCollapseNode>())
+		{
+			Result.Merge(Controller->PatchFunctionsWithInvalidReturnPaths());	
+		}
 	}
 
 	return Result;
diff --git a/Engine/Plugins/Runtime/RigVM/Source/RigVMDeveloper/Private/RigVMModel/RigVMController.cpp b/Engine/Plugins/Runtime/RigVM/Source/RigVMDeveloper/Private/RigVMModel/RigVMController.cpp
--- a/Engine/Plugins/Runtime/RigVM/Source/RigVMDeveloper/Private/RigVMModel/RigVMController.cpp
+++ b/Engine/Plugins/Runtime/RigVM/Source/RigVMDeveloper/Private/RigVMModel/RigVMController.cpp
@@ -19557,6 +19557,69 @@ FRigVMClientPatchResult URigVMController::PatchInvalidLinksOnWildcards()
 	return Result;
 }
 
+FRigVMClientPatchResult URigVMController::PatchFunctionsWithInvalidReturnPaths()
+{
+	FRigVMClientPatchResult Result;
+
+	if (URigVMGraph* Graph = GetGraph())
+	{
+		if (!Graph->GetOuter()->IsA<URigVMCollapseNode>())
+		{
+			Result.bSucceeded = false;
+			return Result;
+		}
+
+		URigVMFunctionReturnNode* ReturnNode = Graph->GetReturnNode();
+		if (!ReturnNode || !ReturnNode->IsMutable())
+		{
+			return Result;
+		}
+
+		URigVMPin* ReturnExecutePin = *ReturnNode->GetPins().FindByPredicate([](const URigVMPin* Pin)
+		{
+			return Pin->IsExecuteContext();
+		});
+
+		if (!ReturnExecutePin->IsLinked())
+		{
+			return Result;
+		}
+
+		URigVMPin* ExecuteSourcePin = ReturnExecutePin->GetLinkedSourcePins()[0];
+		while(ExecuteSourcePin && !ExecuteSourcePin->GetNode()->IsA<URigVMFunctionEntryNode>())
+		{
+			URigVMPin* ExecuteTargetPin = ExecuteSourcePin;
+			if (ExecuteSourcePin->GetNode()->IsControlFlowNode())
+			{
+				if (ExecuteSourcePin->GetFName() != FRigVMStruct::ControlFlowCompletedName)
+				{
+					URigVMPin* CompletedPin = ExecuteSourcePin->GetNode()->FindPin(FRigVMStruct::ControlFlowCompletedName.ToString());
+					AddLink(CompletedPin, ReturnExecutePin);
+					ExecuteSourcePin = CompletedPin;
+					Result.bChangedContent = true;
+				}
+			}
+
+			if (ExecuteSourcePin->GetDirection() == ERigVMPinDirection::Output)
+			{
+				ExecuteTargetPin = *ExecuteSourcePin->GetNode()->GetPins().FindByPredicate([](URigVMPin* Pin)
+				{
+					return Pin->IsExecuteContext() && (Pin->GetDirection() == ERigVMPinDirection::Input || Pin->GetDirection() == ERigVMPinDirection::IO);
+				});
+			}
+
+			const TArray<URigVMPin*>& SourcePins = ExecuteTargetPin->GetLinkedSourcePins();
+			if (SourcePins.IsEmpty())
+			{
+				break;
+			}
+			ExecuteSourcePin = SourcePins[0];
+		}
+	}
+	
+	return Result;
+}
+
 void URigVMController::PostDuplicateHost(const FString& InOldPathName, const FString& InNewPathName)
 {
 	if (const URigVMGraph* Graph = GetGraph())
diff --git a/Engine/Plugins/Runtime/RigVM/Source/RigVMDeveloper/Public/RigVMModel/RigVMController.h b/Engine/Plugins/Runtime/RigVM/Source/RigVMDeveloper/Public/RigVMModel/RigVMController.h
--- a/Engine/Plugins/Runtime/RigVM/Source/RigVMDeveloper/Public/RigVMModel/RigVMController.h
+++ b/Engine/Plugins/Runtime/RigVM/Source/RigVMDeveloper/Public/RigVMModel/RigVMController.h
@@ -1302,6 +1302,7 @@ protected:
 	FRigVMClientPatchResult PatchArrayNodesOnLoad();
 	FRigVMClientPatchResult PatchReduceArrayFloatDoubleConvertsionsOnLoad();
 	FRigVMClientPatchResult PatchInvalidLinksOnWildcards();
+	FRigVMClientPatchResult PatchFunctionsWithInvalidReturnPaths();
 
 	// work to do after a duplication of the host asset
 	void PostDuplicateHost(const FString& InOldPathName, const FString& InNewPathName);
