RigVM: Patch functions with ExecuteContext as input or output, not IO
#rnx

diff --git a/Engine/Plugins/Runtime/RigVM/Source/RigVMDeveloper/Private/RigVMModel/RigVMController.cpp b/Engine/Plugins/Runtime/RigVM/Source/RigVMDeveloper/Private/RigVMModel/RigVMController.cpp
--- a/Engine/Plugins/Runtime/RigVM/Source/RigVMDeveloper/Private/RigVMModel/RigVMController.cpp
+++ b/Engine/Plugins/Runtime/RigVM/Source/RigVMDeveloper/Private/RigVMModel/RigVMController.cpp
@@ -14,6 +14,7 @@
 #include "RigVMFunctions/RigVMDispatch_Array.h"
 #include "RigVMFunctions/RigVMDispatch_Constant.h"
 #include "RigVMFunctions/RigVMDispatch_MakeStruct.h"
+#include "RigVMFunctions/Execution/RigVMFunction_Sequence.h"
 #include "RigVMCompiler/RigVMCompiler.h"
 #include "RigVMDeveloperModule.h"
 #include "UObject/PropertyPortFlags.h"
@@ -19560,6 +19561,27 @@ FRigVMClientPatchResult URigVMController::PatchInvalidLinksOnWildcards()
 FRigVMClientPatchResult URigVMController::PatchFunctionsWithInvalidReturnPaths()
 {
 	FRigVMClientPatchResult Result;
+	TGuardValue<bool> GuardReportWarningsAndErrors(bSuspendNotifications, true);
+
+	auto FindExecutePin = [](URigVMNode* Node, ERigVMPinDirection Direction = ERigVMPinDirection::Invalid) -> URigVMPin*
+	{
+		URigVMPin* const* Result = Node->GetPins().FindByPredicate([Direction](const URigVMPin* Pin)
+		{
+			if (Direction == ERigVMPinDirection::Invalid)
+			{
+				return Pin->IsExecuteContext();
+			}
+			else
+			{
+				return Pin->IsExecuteContext() && Pin->GetDirection() == Direction;
+			}
+		});
+		if (Result)
+		{
+			return *Result;
+		}
+		return nullptr;
+	};
 
 	if (URigVMGraph* Graph = GetGraph())
 	{
@@ -19569,22 +19591,81 @@ FRigVMClientPatchResult URigVMController::PatchFunctionsWithInvalidReturnPaths()
 			return Result;
 		}
 
+		URigVMFunctionEntryNode* EntryNode = Graph->GetEntryNode();
 		URigVMFunctionReturnNode* ReturnNode = Graph->GetReturnNode();
-		if (!ReturnNode || !ReturnNode->IsMutable())
+		bool bEntryIsMutable = false;
+		bool bReturnIsMutable = false;
+		if (EntryNode)
+		{
+			bEntryIsMutable = EntryNode->IsMutable();
+		}
+		if (ReturnNode)
+		{
+			bReturnIsMutable = ReturnNode->IsMutable();
+		}
+
+		if (!bEntryIsMutable && !bReturnIsMutable)
 		{
 			return Result;
 		}
 
-		URigVMPin* ReturnExecutePin = *ReturnNode->GetPins().FindByPredicate([](const URigVMPin* Pin)
+		bool bReturnExecuteIsLinked = false;
+		if (bReturnIsMutable)
 		{
-			return Pin->IsExecuteContext();
-		});
+			URigVMPin* ReturnExecutePin = FindExecutePin(ReturnNode);
+			bReturnExecuteIsLinked = ReturnExecutePin->IsLinked();
+		}
 
-		if (!ReturnExecutePin->IsLinked())
+		if (bEntryIsMutable && (!bReturnIsMutable || !bReturnExecuteIsLinked))
 		{
+			URigVMPin* EntryExecutePin = FindExecutePin(EntryNode);
+			
+			URigVMPin* OtherExecutePin = nullptr;
+			if (EntryExecutePin->IsLinked())
+			{
+				OtherExecutePin = EntryExecutePin->GetLinkedTargetPins()[0];
+			}
+
+			const FName ExecutePinName = EntryExecutePin->GetFName();
+			const FString ExecuteCPPType = EntryExecutePin->GetCPPType();
+			const FName ExecuteCPPTypeObjectPath = *EntryExecutePin->GetCPPTypeObject()->GetPathName();
+			RemoveExposedPin(ExecutePinName, false);
+			AddExposedPin(ExecutePinName, ERigVMPinDirection::IO, ExecuteCPPType, ExecuteCPPTypeObjectPath, FString(), false);
+			EntryExecutePin = FindExecutePin(EntryNode);
+			URigVMPin* ReturnExecutePin = FindExecutePin(ReturnNode);
+			
+			if (OtherExecutePin)
+			{
+				URigVMNode* SequenceNode = AddUnitNode(FRigVMFunction_Sequence::StaticStruct(),
+				   FRigVMStruct::ExecuteName,
+				   EntryNode->GetPosition() + FVector2D(200.f, 0),
+				   FString(), false);
+				AddLink(SequenceNode->Pins[1], OtherExecutePin, false);
+				AddLink(EntryExecutePin, SequenceNode->Pins[0], false);
+				AddLink(SequenceNode->Pins[2], ReturnExecutePin, false);
+			}
+			else
+			{
+				AddLink(EntryExecutePin, ReturnExecutePin, false);
+			}
+
+			Result.bChangedContent = true;
 			return Result;
 		}
 
+		if (bReturnIsMutable && !bEntryIsMutable)
+		{
+			URigVMPin* ReturnExecutePin = FindExecutePin(ReturnNode);
+			const FName ExecutePinName = ReturnExecutePin->GetFName();
+			const FString ExecuteCPPType = ReturnExecutePin->GetCPPType();
+			const FName ExecuteCPPTypeObjectPath = *ReturnExecutePin->GetCPPTypeObject()->GetPathName();
+			RemoveExposedPin(ExecutePinName, false);
+			AddExposedPin(ExecutePinName, ERigVMPinDirection::IO, ExecuteCPPType, ExecuteCPPTypeObjectPath, FString(), false);
+		}
+
+		URigVMPin* EntryExecutePin = FindExecutePin(EntryNode);
+		URigVMPin* ReturnExecutePin = FindExecutePin(ReturnNode);
+
 		URigVMPin* ExecuteSourcePin = ReturnExecutePin->GetLinkedSourcePins()[0];
 		while(ExecuteSourcePin && !ExecuteSourcePin->GetNode()->IsA<URigVMFunctionEntryNode>())
 		{
@@ -19611,6 +19692,7 @@ FRigVMClientPatchResult URigVMController::PatchFunctionsWithInvalidReturnPaths()
 			const TArray<URigVMPin*>& SourcePins = ExecuteTargetPin->GetLinkedSourcePins();
 			if (SourcePins.IsEmpty())
 			{
+				AddLink(EntryExecutePin, ExecuteTargetPin);
 				break;
 			}
 			ExecuteSourcePin = SourcePins[0];
