RigVM: Improve performance when retrieving argument names

diff --git a/Engine/Plugins/Animation/ControlRig/Source/ControlRig/Private/Units/Animation/RigUnit_AnimAttribute.cpp b/Engine/Plugins/Animation/ControlRig/Source/ControlRig/Private/Units/Animation/RigUnit_AnimAttribute.cpp
--- a/Engine/Plugins/Animation/ControlRig/Source/ControlRig/Private/Units/Animation/RigUnit_AnimAttribute.cpp
+++ b/Engine/Plugins/Animation/ControlRig/Source/ControlRig/Private/Units/Animation/RigUnit_AnimAttribute.cpp
@@ -321,7 +321,7 @@ FString FRigDispatch_AnimAttributeBase::GetNodeTitle(const FRigVMTemplateTypeMap
 }
 #endif
 
-TArray<FRigVMTemplateArgument> FRigDispatch_AnimAttributeBase::GetArguments() const
+const TArray<FRigVMTemplateArgument>& FRigDispatch_AnimAttributeBase::GetArguments() const
 {
 	if (Arguments.IsEmpty())
 	{
@@ -372,7 +372,7 @@ FText FRigDispatch_AnimAttributeBase::GetArgumentTooltip(const FName& InArgument
 
 
 
-TArray<FRigVMTemplateArgument> FRigDispatch_GetAnimAttribute::GetArguments() const
+const TArray<FRigVMTemplateArgument>& FRigDispatch_GetAnimAttribute::GetArguments() const
 {
 	if (ValueArgIndex == INDEX_NONE)
 	{
@@ -455,7 +455,7 @@ FRigVMFunctionPtr FRigDispatch_GetAnimAttribute::GetDispatchFunctionImpl(const F
 	return nullptr;
 }
 
-TArray<FRigVMTemplateArgument> FRigDispatch_SetAnimAttribute::GetArguments() const
+const TArray<FRigVMTemplateArgument>& FRigDispatch_SetAnimAttribute::GetArguments() const
 {
 	if (ValueArgIndex == INDEX_NONE)
 	{
@@ -471,9 +471,12 @@ TArray<FRigVMTemplateArgument> FRigDispatch_SetAnimAttribute::GetArguments() con
 	return Arguments;
 }
 
-TArray<FRigVMExecuteArgument> FRigDispatch_SetAnimAttribute::GetExecuteArguments_Impl(const FRigVMDispatchContext& InContext) const
+const TArray<FRigVMExecuteArgument>& FRigDispatch_SetAnimAttribute::GetExecuteArguments_Impl(const FRigVMDispatchContext& InContext) const
 {
-	return {{TEXT("ExecuteContext"), ERigVMPinDirection::IO}};
+	static const TArray<FRigVMExecuteArgument> ExecuteArguments = {
+		{TEXT("ExecuteContext"), ERigVMPinDirection::IO}
+	};
+	return ExecuteArguments;
 }
 
 FRigVMTemplateTypeMap FRigDispatch_SetAnimAttribute::OnNewArgumentType(const FName& InArgumentName,
diff --git a/Engine/Plugins/Animation/ControlRig/Source/ControlRig/Private/Units/Animation/RigUnit_AnimAttribute.h b/Engine/Plugins/Animation/ControlRig/Source/ControlRig/Private/Units/Animation/RigUnit_AnimAttribute.h
--- a/Engine/Plugins/Animation/ControlRig/Source/ControlRig/Private/Units/Animation/RigUnit_AnimAttribute.h
+++ b/Engine/Plugins/Animation/ControlRig/Source/ControlRig/Private/Units/Animation/RigUnit_AnimAttribute.h
@@ -134,7 +134,7 @@ struct CONTROLRIG_API FRigDispatch_AnimAttributeBase : public FRigDispatchFactor
 	virtual FString GetNodeTitle(const FRigVMTemplateTypeMap& InTypes) const override;;
 #endif
 	
-	virtual TArray<FRigVMTemplateArgument> GetArguments() const override;
+	virtual const TArray<FRigVMTemplateArgument>& GetArguments() const override;
 	virtual bool IsSet() const { return false; }
 
 #if WITH_EDITOR
@@ -187,7 +187,7 @@ struct CONTROLRIG_API FRigDispatch_GetAnimAttribute: public FRigDispatch_AnimAtt
 		FactoryScriptStruct = StaticStruct();
 	}
 
-	virtual TArray<FRigVMTemplateArgument> GetArguments() const override;
+	virtual const TArray<FRigVMTemplateArgument>& GetArguments() const override;
 	virtual FRigVMTemplateTypeMap OnNewArgumentType(const FName& InArgumentName, TRigVMTypeIndex InTypeIndex) const override;
 
 	
@@ -267,8 +267,8 @@ struct CONTROLRIG_API FRigDispatch_SetAnimAttribute: public FRigDispatch_AnimAtt
 	}
 
 	virtual bool IsSet() const override { return true; }
-	virtual TArray<FRigVMTemplateArgument> GetArguments() const override;
-	virtual TArray<FRigVMExecuteArgument> GetExecuteArguments_Impl(const FRigVMDispatchContext& InContext) const override;
+	virtual const TArray<FRigVMTemplateArgument>& GetArguments() const override;
+	virtual const TArray<FRigVMExecuteArgument>& GetExecuteArguments_Impl(const FRigVMDispatchContext& InContext) const override;
 	virtual FRigVMTemplateTypeMap OnNewArgumentType(const FName& InArgumentName, TRigVMTypeIndex InTypeIndex) const override;
 	
 protected:
diff --git a/Engine/Plugins/Animation/ControlRig/Source/ControlRig/Private/Units/Core/RigUnit_UserData.cpp b/Engine/Plugins/Animation/ControlRig/Source/ControlRig/Private/Units/Core/RigUnit_UserData.cpp
--- a/Engine/Plugins/Animation/ControlRig/Source/ControlRig/Private/Units/Core/RigUnit_UserData.cpp
+++ b/Engine/Plugins/Animation/ControlRig/Source/ControlRig/Private/Units/Core/RigUnit_UserData.cpp
@@ -28,19 +28,20 @@ FString FRigDispatch_GetUserData::GetArgumentMetaData(const FName& InArgumentNam
 
 #endif
 
-TArray<FRigVMTemplateArgument> FRigDispatch_GetUserData::GetArguments() const
+const TArray<FRigVMTemplateArgument>& FRigDispatch_GetUserData::GetArguments() const
 {
-	const TArray<FRigVMTemplateArgument::ETypeCategory> ValueCategories = {
+	static const TArray<FRigVMTemplateArgument::ETypeCategory> ValueCategories = {
 		FRigVMTemplateArgument::ETypeCategory_SingleAnyValue,
 		FRigVMTemplateArgument::ETypeCategory_ArrayAnyValue
 	};
-	return {
+	static const TArray<FRigVMTemplateArgument> Arguments = {
 		FRigVMTemplateArgument(ArgNameSpaceName, ERigVMPinDirection::Input, RigVMTypeUtils::TypeIndex::FString),
 		FRigVMTemplateArgument(ArgPathName, ERigVMPinDirection::Input, RigVMTypeUtils::TypeIndex::FString),
 		FRigVMTemplateArgument(ArgDefaultName, ERigVMPinDirection::Input, ValueCategories),
 		FRigVMTemplateArgument(ArgResultName, ERigVMPinDirection::Output, ValueCategories),
 		FRigVMTemplateArgument(ArgFoundName, ERigVMPinDirection::Output, RigVMTypeUtils::TypeIndex::Bool)
 	};
+	return Arguments;
 }
 
 FRigVMTemplateTypeMap FRigDispatch_GetUserData::OnNewArgumentType(const FName& InArgumentName,
diff --git a/Engine/Plugins/Animation/ControlRig/Source/ControlRig/Private/Units/Core/RigUnit_UserData.h b/Engine/Plugins/Animation/ControlRig/Source/ControlRig/Private/Units/Core/RigUnit_UserData.h
--- a/Engine/Plugins/Animation/ControlRig/Source/ControlRig/Private/Units/Core/RigUnit_UserData.h
+++ b/Engine/Plugins/Animation/ControlRig/Source/ControlRig/Private/Units/Core/RigUnit_UserData.h
@@ -22,7 +22,7 @@ public:
 #if WITH_EDITOR
 	virtual FString GetArgumentMetaData(const FName& InArgumentName, const FName& InMetaDataKey) const override;
 #endif
-	virtual TArray<FRigVMTemplateArgument> GetArguments() const override;
+	virtual const TArray<FRigVMTemplateArgument>& GetArguments() const override;
 	virtual FRigVMTemplateTypeMap OnNewArgumentType(const FName& InArgumentName, TRigVMTypeIndex InTypeIndex) const override;
 
 protected:
diff --git a/Engine/Plugins/Animation/ControlRig/Source/ControlRig/Private/Units/Hierarchy/RigUnit_Metadata.cpp b/Engine/Plugins/Animation/ControlRig/Source/ControlRig/Private/Units/Hierarchy/RigUnit_Metadata.cpp
--- a/Engine/Plugins/Animation/ControlRig/Source/ControlRig/Private/Units/Hierarchy/RigUnit_Metadata.cpp
+++ b/Engine/Plugins/Animation/ControlRig/Source/ControlRig/Private/Units/Hierarchy/RigUnit_Metadata.cpp
@@ -65,7 +65,7 @@ FString FRigDispatch_MetadataBase::GetNodeTitle(const FRigVMTemplateTypeMap& InT
 
 #endif
 
-TArray<FRigVMTemplateArgument> FRigDispatch_MetadataBase::GetArguments() const
+const TArray<FRigVMTemplateArgument>& FRigDispatch_MetadataBase::GetArguments() const
 {
 	if(Arguments.IsEmpty())
 	{
@@ -155,7 +155,7 @@ const TArray<TRigVMTypeIndex>& FRigDispatch_MetadataBase::GetValueTypes() const
 
 ////////////////////////////////////////////////////////////////////////////////////////////////
 
-TArray<FRigVMTemplateArgument> FRigDispatch_GetMetadata::GetArguments() const
+const TArray<FRigVMTemplateArgument>& FRigDispatch_GetMetadata::GetArguments() const
 {
 	if(ValueArgIndex == INDEX_NONE)
 	{
@@ -272,7 +272,7 @@ FRigVMFunctionPtr FRigDispatch_GetMetadata::GetDispatchFunctionImpl(const FRigVM
 
 ////////////////////////////////////////////////////////////////////////////////////////////////
 
-TArray<FRigVMTemplateArgument> FRigDispatch_SetMetadata::GetArguments() const
+const TArray<FRigVMTemplateArgument>& FRigDispatch_SetMetadata::GetArguments() const
 {
 	if(ValueArgIndex == INDEX_NONE)
 	{
@@ -283,9 +283,12 @@ TArray<FRigVMTemplateArgument> FRigDispatch_SetMetadata::GetArguments() const
 	return Arguments;
 }
 
-TArray<FRigVMExecuteArgument> FRigDispatch_SetMetadata::GetExecuteArguments_Impl(const FRigVMDispatchContext& InContext) const
+const TArray<FRigVMExecuteArgument>& FRigDispatch_SetMetadata::GetExecuteArguments_Impl(const FRigVMDispatchContext& InContext) const
 {
-	return {{TEXT("ExecuteContext"), ERigVMPinDirection::IO}};
+	static const TArray<FRigVMExecuteArgument> ExecuteArguments = {
+		{TEXT("ExecuteContext"), ERigVMPinDirection::IO}
+	};
+	return ExecuteArguments;
 }
 
 FRigBaseMetadata* FRigDispatch_SetMetadata::FindOrAddMetadata(FControlRigExecuteContext& InContext,
diff --git a/Engine/Plugins/Animation/ControlRig/Source/ControlRig/Private/Units/Hierarchy/RigUnit_Metadata.h b/Engine/Plugins/Animation/ControlRig/Source/ControlRig/Private/Units/Hierarchy/RigUnit_Metadata.h
--- a/Engine/Plugins/Animation/ControlRig/Source/ControlRig/Private/Units/Hierarchy/RigUnit_Metadata.h
+++ b/Engine/Plugins/Animation/ControlRig/Source/ControlRig/Private/Units/Hierarchy/RigUnit_Metadata.h
@@ -18,7 +18,7 @@ struct CONTROLRIG_API FRigDispatch_MetadataBase : public FRigDispatchFactory
 #if WITH_EDITOR
 	virtual FString GetNodeTitle(const FRigVMTemplateTypeMap& InTypes) const override;;
 #endif
-	virtual TArray<FRigVMTemplateArgument> GetArguments() const override;
+	virtual const TArray<FRigVMTemplateArgument>& GetArguments() const override;
 	virtual bool IsSetMetadata() const { return false; }
 
 #if WITH_EDITOR
@@ -56,7 +56,7 @@ struct CONTROLRIG_API FRigDispatch_GetMetadata : public FRigDispatch_MetadataBas
 {
 	GENERATED_BODY()
 
-	virtual TArray<FRigVMTemplateArgument> GetArguments() const override;
+	virtual const TArray<FRigVMTemplateArgument>& GetArguments() const override;
 
 protected:
 
@@ -118,8 +118,8 @@ struct CONTROLRIG_API FRigDispatch_SetMetadata : public FRigDispatch_MetadataBas
 {
 	GENERATED_BODY()
 
-	virtual TArray<FRigVMTemplateArgument> GetArguments() const override;
-	virtual TArray<FRigVMExecuteArgument> GetExecuteArguments_Impl(const FRigVMDispatchContext& InContext) const override;
+	virtual const TArray<FRigVMTemplateArgument>& GetArguments() const override;
+	virtual const TArray<FRigVMExecuteArgument>& GetExecuteArguments_Impl(const FRigVMDispatchContext& InContext) const override;
 	virtual bool IsSetMetadata() const override { return true; }
 
 protected:
diff --git a/Engine/Plugins/Experimental/AnimNext/Source/AnimNext/Internal/Param/RigVMDispatch_GetParameter.h b/Engine/Plugins/Experimental/AnimNext/Source/AnimNext/Internal/Param/RigVMDispatch_GetParameter.h
--- a/Engine/Plugins/Experimental/AnimNext/Source/AnimNext/Internal/Param/RigVMDispatch_GetParameter.h
+++ b/Engine/Plugins/Experimental/AnimNext/Source/AnimNext/Internal/Param/RigVMDispatch_GetParameter.h
@@ -26,7 +26,7 @@ private:
 
 	virtual UScriptStruct* GetExecuteContextStruct() const { return FAnimNextExecuteContext::StaticStruct(); }
 	virtual FName GetArgumentNameForOperandIndex(int32 InOperandIndex, int32 InTotalOperands) const override;
-	virtual TArray<FRigVMTemplateArgument> GetArguments() const override;
+	virtual const TArray<FRigVMTemplateArgument>& GetArguments() const override;
 #if WITH_EDITOR
 	virtual FString GetArgumentMetaData(const FName& InArgumentName, const FName& InMetaDataKey) const override;
 	virtual FString GetArgumentDefaultValue(const FName& InArgumentName, TRigVMTypeIndex InTypeIndex) const override;
diff --git a/Engine/Plugins/Experimental/AnimNext/Source/AnimNext/Internal/Param/RigVMDispatch_SetParameter.h b/Engine/Plugins/Experimental/AnimNext/Source/AnimNext/Internal/Param/RigVMDispatch_SetParameter.h
--- a/Engine/Plugins/Experimental/AnimNext/Source/AnimNext/Internal/Param/RigVMDispatch_SetParameter.h
+++ b/Engine/Plugins/Experimental/AnimNext/Source/AnimNext/Internal/Param/RigVMDispatch_SetParameter.h
@@ -26,12 +26,12 @@ private:
 
 	virtual UScriptStruct* GetExecuteContextStruct() const { return FAnimNextExecuteContext::StaticStruct(); }
 	virtual FName GetArgumentNameForOperandIndex(int32 InOperandIndex, int32 InTotalOperands) const override;
-	virtual TArray<FRigVMTemplateArgument> GetArguments() const override;
+	virtual const TArray<FRigVMTemplateArgument>& GetArguments() const override;
 #if WITH_EDITOR
 	virtual FString GetArgumentMetaData(const FName& InArgumentName, const FName& InMetaDataKey) const override;
 	virtual FString GetArgumentDefaultValue(const FName& InArgumentName, TRigVMTypeIndex InTypeIndex) const override;
 #endif
-	virtual TArray<FRigVMExecuteArgument> GetExecuteArguments_Impl(const FRigVMDispatchContext& InContext) const override;
+	virtual TArray<FRigVMExecuteArgument>& GetExecuteArguments_Impl(const FRigVMDispatchContext& InContext) const override;
 	virtual FRigVMTemplateTypeMap OnNewArgumentType(const FName& InArgumentName, TRigVMTypeIndex InTypeIndex) const override;
 	virtual bool IsSingleton() const override { return true; } 
 
diff --git a/Engine/Plugins/Experimental/AnimNext/Source/AnimNext/Private/Param/RigVMDispatch_GetParameter.cpp b/Engine/Plugins/Experimental/AnimNext/Source/AnimNext/Private/Param/RigVMDispatch_GetParameter.cpp
--- a/Engine/Plugins/Experimental/AnimNext/Source/AnimNext/Private/Param/RigVMDispatch_GetParameter.cpp
+++ b/Engine/Plugins/Experimental/AnimNext/Source/AnimNext/Private/Param/RigVMDispatch_GetParameter.cpp
@@ -56,20 +56,22 @@ FString FRigVMDispatch_GetParameter::GetArgumentDefaultValue(const FName& InArgu
 }
 #endif
 
-TArray<FRigVMTemplateArgument> FRigVMDispatch_GetParameter::GetArguments() const
+const TArray<FRigVMTemplateArgument>& FRigVMDispatch_GetParameter::GetArguments() const
 {
-	const TArray<FRigVMTemplateArgument::ETypeCategory> ValueCategories =
+	static const TArray<FRigVMTemplateArgument::ETypeCategory> ValueCategories =
 	{
 		FRigVMTemplateArgument::ETypeCategory_SingleAnyValue,
 		FRigVMTemplateArgument::ETypeCategory_ArrayAnyValue
 	};
-	return
-	{
+
+	static const TArray<FRigVMTemplateArgument> Arguments = {
 		FRigVMTemplateArgument(ParameterName, ERigVMPinDirection::Input, RigVMTypeUtils::TypeIndex::FName),
 		FRigVMTemplateArgument(ValueName, ERigVMPinDirection::Output, ValueCategories),
 		FRigVMTemplateArgument(ParameterIdName, ERigVMPinDirection::Hidden, RigVMTypeUtils::TypeIndex::UInt32),
 		FRigVMTemplateArgument(TypeHandleName, ERigVMPinDirection::Hidden, RigVMTypeUtils::TypeIndex::UInt32),
 	};
+
+	return Arguments;
 }
 
 FRigVMTemplateTypeMap FRigVMDispatch_GetParameter::OnNewArgumentType(const FName& InArgumentName, TRigVMTypeIndex InTypeIndex) const
diff --git a/Engine/Plugins/Experimental/AnimNext/Source/AnimNext/Private/Param/RigVMDispatch_SetParameter.cpp b/Engine/Plugins/Experimental/AnimNext/Source/AnimNext/Private/Param/RigVMDispatch_SetParameter.cpp
--- a/Engine/Plugins/Experimental/AnimNext/Source/AnimNext/Private/Param/RigVMDispatch_SetParameter.cpp
+++ b/Engine/Plugins/Experimental/AnimNext/Source/AnimNext/Private/Param/RigVMDispatch_SetParameter.cpp
@@ -57,25 +57,28 @@ FString FRigVMDispatch_SetParameter::GetArgumentDefaultValue(const FName& InArgu
 }
 #endif
 
-TArray<FRigVMTemplateArgument> FRigVMDispatch_SetParameter::GetArguments() const
+const TArray<FRigVMTemplateArgument>& FRigVMDispatch_SetParameter::GetArguments() const
 {
-	const TArray<FRigVMTemplateArgument::ETypeCategory> ValueCategories =
+	static const TArray<FRigVMTemplateArgument::ETypeCategory> ValueCategories =
 	{
 		FRigVMTemplateArgument::ETypeCategory_SingleAnyValue,
 		FRigVMTemplateArgument::ETypeCategory_ArrayAnyValue
 	};
-	return 
-	{ 
+	
+	static const TArray<FRigVMTemplateArgument> Arguments = { 
 		FRigVMTemplateArgument(ParameterName, ERigVMPinDirection::Input, RigVMTypeUtils::TypeIndex::FName),
 		FRigVMTemplateArgument(ValueName, ERigVMPinDirection::Input, ValueCategories),
 		FRigVMTemplateArgument(ParameterIdName, ERigVMPinDirection::Hidden, RigVMTypeUtils::TypeIndex::UInt32),
 		FRigVMTemplateArgument(TypeHandleName, ERigVMPinDirection::Hidden, RigVMTypeUtils::TypeIndex::UInt32),
 	};
+
+	return Arguments;
 }
 
-TArray<FRigVMExecuteArgument> FRigVMDispatch_SetParameter::GetExecuteArguments_Impl(const FRigVMDispatchContext& InContext) const
+TArray<FRigVMExecuteArgument>& FRigVMDispatch_SetParameter::GetExecuteArguments_Impl(const FRigVMDispatchContext& InContext) const
 {
-	return { { ExecuteContextName, ERigVMPinDirection::Input } };
+	static TArray<FRigVMExecuteArgument> Arguments = { { ExecuteContextName, ERigVMPinDirection::Input } };
+	return Arguments;
 }
 
 FRigVMTemplateTypeMap FRigVMDispatch_SetParameter::OnNewArgumentType(const FName& InArgumentName, TRigVMTypeIndex InTypeIndex) const
diff --git a/Engine/Plugins/Runtime/RigVM/Source/RigVM/Private/RigVMCore/RigVM.cpp b/Engine/Plugins/Runtime/RigVM/Source/RigVM/Private/RigVMCore/RigVM.cpp
--- a/Engine/Plugins/Runtime/RigVM/Source/RigVM/Private/RigVMCore/RigVM.cpp
+++ b/Engine/Plugins/Runtime/RigVM/Source/RigVM/Private/RigVMCore/RigVM.cpp
@@ -1583,6 +1583,24 @@ bool URigVM::Initialize(FRigVMExtendedExecuteContext& Context, TArrayView<URigVM
 			}
 		}
 	}
+
+	// make sure to update all argument name caches
+	TArray<const FRigVMFunction*>& Functions = GetFunctions();
+	for(const FRigVMInstruction& Instruction : Instructions)
+	{
+		if(Instruction.OpCode == ERigVMOpCode::Execute)
+		{
+			const FRigVMExecuteOp& Op = ByteCodeStorage.GetOpAt<FRigVMExecuteOp>(Instruction);
+			if(Functions.IsValidIndex(Op.FunctionIndex))
+			{
+				if(const FRigVMDispatchFactory* Factory = Functions[Op.FunctionIndex]->Factory)
+				{
+					FRigVMOperandArray Operands = ByteCodeStorage.GetOperandsForExecuteOp(Instruction);
+					(void)Factory->UpdateArgumentNameCache(Operands.Num());
+				}
+			}
+		}
+	}
 	
 	CacheMemoryHandlesIfRequired(Context, Memory);
 
diff --git a/Engine/Plugins/Runtime/RigVM/Source/RigVM/Private/RigVMCore/RigVMDispatchFactory.cpp b/Engine/Plugins/Runtime/RigVM/Source/RigVM/Private/RigVMCore/RigVMDispatchFactory.cpp
--- a/Engine/Plugins/Runtime/RigVM/Source/RigVM/Private/RigVMCore/RigVMDispatchFactory.cpp
+++ b/Engine/Plugins/Runtime/RigVM/Source/RigVM/Private/RigVMCore/RigVMDispatchFactory.cpp
@@ -179,6 +179,33 @@ const TArray<FName>& FRigVMDispatchFactory::GetControlFlowBlocks_Impl(const FRig
 	return EmptyArray;
 }
 
+const TArray<FName>* FRigVMDispatchFactory::UpdateArgumentNameCache(int32 InNumberOperands) const
+{
+	FScopeLock ArgumentNamesLock(ArgumentNamesMutex);
+	return UpdateArgumentNameCache_NoLock(InNumberOperands);
+}
+
+const TArray<FName>* FRigVMDispatchFactory::UpdateArgumentNameCache_NoLock(int32 InNumberOperands) const
+{
+	TSharedPtr<TArray<FName>>& ArgumentNames = ArgumentNamesMap.FindOrAdd(InNumberOperands);
+	if(!ArgumentNames.IsValid())
+	{
+		ArgumentNames = MakeShareable(new TArray<FName>()); 
+		ArgumentNames->Reserve(InNumberOperands);
+		for(int32 OperandIndex = 0; OperandIndex < InNumberOperands; OperandIndex++)
+		{
+			ArgumentNames->Add(GetArgumentNameForOperandIndex(OperandIndex, InNumberOperands));
+		}
+	}
+	return ArgumentNames.Get();
+}
+
+const TArray<FRigVMTemplateArgument>& FRigVMDispatchFactory::GetArguments() const
+{
+	static TArray<FRigVMTemplateArgument> EmptyArguments;
+	return EmptyArguments;
+}
+
 TArray<FRigVMExecuteArgument> FRigVMDispatchFactory::GetExecuteArguments(const FRigVMDispatchContext& InContext) const
 {
 	TArray<FRigVMExecuteArgument> Arguments = GetExecuteArguments_Impl(InContext);
@@ -196,6 +223,13 @@ TArray<FRigVMExecuteArgument> FRigVMDispatchFactory::GetExecuteArguments(const F
 	return Arguments;
 }
 
+const TArray<FRigVMExecuteArgument>& FRigVMDispatchFactory::GetExecuteArguments_Impl(
+	const FRigVMDispatchContext& InContext) const
+{
+	static TArray<FRigVMExecuteArgument> EmptyArguments;
+	return EmptyArguments;
+}
+
 FRigVMFunctionPtr FRigVMDispatchFactory::GetOrCreateDispatchFunction(const FRigVMTemplateTypeMap& InTypes) const
 {
 	FScopeLock GetTemplateScopeLock(&FRigVMRegistry::GetDispatchFunctionMutex);
@@ -234,7 +268,7 @@ TArray<FRigVMFunction> FRigVMDispatchFactory::CreateDispatchPredicates_NoLock(co
 FString FRigVMDispatchFactory::GetPermutationName(const FRigVMTemplateTypeMap& InTypes) const
 {
 #if WITH_EDITOR
-	const TArray<FRigVMTemplateArgument> Arguments = GetArguments();
+	const TArray<FRigVMTemplateArgument>& Arguments = GetArguments();
 	check(InTypes.Num() == Arguments.Num());
 	for(const FRigVMTemplateArgument& Argument : Arguments)
 	{
@@ -252,7 +286,7 @@ FString FRigVMDispatchFactory::GetPermutationNameImpl(const FRigVMTemplateTypeMa
 }
 
 bool FRigVMDispatchFactory::CopyProperty(const FProperty* InTargetProperty, uint8* InTargetPtr,
-	const FProperty* InSourceProperty, const uint8* InSourcePtr)
+                                         const FProperty* InSourceProperty, const uint8* InSourcePtr)
 {
 	return URigVMMemoryStorage::CopyProperty(InTargetProperty, InTargetPtr, InSourceProperty, InSourcePtr);
 }
@@ -275,7 +309,7 @@ const FRigVMTemplate* FRigVMDispatchFactory::GetTemplate() const
 		return CachedTemplate;
 	}
 
-	const TArray<FRigVMTemplateArgument> Arguments = GetArguments();
+	const TArray<FRigVMTemplateArgument>& Arguments = GetArguments();
 
 	// we don't allow execute types on arguments
 	for(const FRigVMTemplateArgument& Argument : Arguments)
diff --git a/Engine/Plugins/Runtime/RigVM/Source/RigVM/Private/RigVMCore/RigVMFunction.cpp b/Engine/Plugins/Runtime/RigVM/Source/RigVM/Private/RigVMCore/RigVMFunction.cpp
--- a/Engine/Plugins/Runtime/RigVM/Source/RigVM/Private/RigVMCore/RigVMFunction.cpp
+++ b/Engine/Plugins/Runtime/RigVM/Source/RigVM/Private/RigVMCore/RigVMFunction.cpp
@@ -7,6 +7,23 @@
 #include "RigVMCore/RigVMStruct.h"
 #include "UObject/Package.h"
 
+FRigVMFunction::FRigVMFunction(const TCHAR* InName, FRigVMFunctionPtr InFunctionPtr, UScriptStruct* InStruct, int32 InIndex, const TArray<FRigVMFunctionArgument>& InArguments)
+: Name(InName)
+, Struct(InStruct)
+, Factory(nullptr)
+, FunctionPtr(InFunctionPtr)
+, Index(InIndex)
+, TemplateIndex(INDEX_NONE)
+, Arguments(InArguments)
+, ArgumentNames()
+{
+	ArgumentNames.SetNumZeroed(Arguments.Num());
+	for(int32 ArgumentIndex = 0; ArgumentIndex < Arguments.Num(); ArgumentIndex++)
+	{
+		ArgumentNames[ArgumentIndex] = Arguments[ArgumentIndex].Name;
+	}
+}
+
 FString FRigVMFunction::GetName() const
 {
 	return Name;
@@ -160,14 +177,16 @@ bool FRigVMFunction::SupportsExecuteContextStruct(const UScriptStruct* InExecute
 	return InExecuteContextStruct->IsChildOf(GetExecuteContextStruct());
 }
 
-FName FRigVMFunction::GetArgumentNameForOperandIndex(int32 InOperandIndex, int32 InTotalOperands) const
+const FName& FRigVMFunction::GetArgumentNameForOperandIndex(int32 InOperandIndex, int32 InTotalOperands) const
 {
 	if(Factory)
 	{
-		return Factory->GetArgumentNameForOperandIndex(InOperandIndex, InTotalOperands);
+		const TArray<FName>* FactoryArgumentNames = Factory->UpdateArgumentNameCache_NoLock(InTotalOperands);
+		check(FactoryArgumentNames);
+		check(FactoryArgumentNames->IsValidIndex(InOperandIndex));
+		return (*FactoryArgumentNames)[InOperandIndex];
 	}
 
-	check(Arguments.IsValidIndex(InOperandIndex));
-	check(Arguments.Num() == InTotalOperands);
-	return Arguments[InOperandIndex].Name;
+	check(ArgumentNames.IsValidIndex(InOperandIndex));
+	return ArgumentNames[InOperandIndex];
 }
diff --git a/Engine/Plugins/Runtime/RigVM/Source/RigVM/Private/RigVMFunctions/RigVMDispatch_Array.cpp b/Engine/Plugins/Runtime/RigVM/Source/RigVM/Private/RigVMFunctions/RigVMDispatch_Array.cpp
--- a/Engine/Plugins/Runtime/RigVM/Source/RigVM/Private/RigVMFunctions/RigVMDispatch_Array.cpp
+++ b/Engine/Plugins/Runtime/RigVM/Source/RigVM/Private/RigVMFunctions/RigVMDispatch_Array.cpp
@@ -194,22 +194,24 @@ TMap<uint32, int32> FRigVMDispatch_ArrayBase::GetArrayHash(FScriptArrayHelper& I
 
 ////////////////////////////////////////////////////////////////////////////////////////////////////////////////////////
 
-TArray<FRigVMExecuteArgument> FRigVMDispatch_ArrayBaseMutable::GetExecuteArguments_Impl(
+const TArray<FRigVMExecuteArgument>& FRigVMDispatch_ArrayBaseMutable::GetExecuteArguments_Impl(
 	const FRigVMDispatchContext& InContext) const
 {
-	return {
+	static const TArray<FRigVMExecuteArgument> Arguments = {
 			{ExecuteName, ERigVMPinDirection::IO}
 	};
+	return Arguments;
 }
 
 ////////////////////////////////////////////////////////////////////////////////////////////////////////////////////////
 
-TArray<FRigVMTemplateArgument> FRigVMDispatch_ArrayMake::GetArguments() const
+const TArray<FRigVMTemplateArgument>& FRigVMDispatch_ArrayMake::GetArguments() const
 {
-	return {
+	static const TArray<FRigVMTemplateArgument> Arguments = {
 		CreateArgument(ValuesName, ERigVMPinDirection::Input),
 		CreateArgument(ArrayName, ERigVMPinDirection::Output)
 	};
+	return Arguments;
 }
 
 FRigVMTemplateTypeMap FRigVMDispatch_ArrayMake::OnNewArgumentType(const FName& InArgumentName, TRigVMTypeIndex InTypeIndex) const
@@ -313,11 +315,12 @@ FName FRigVMDispatch_ArrayReset::GetArgumentNameForOperandIndex(int32 InOperandI
 	return ArrayName;
 }
 
-TArray<FRigVMTemplateArgument> FRigVMDispatch_ArrayReset::GetArguments() const
+const TArray<FRigVMTemplateArgument>& FRigVMDispatch_ArrayReset::GetArguments() const
 {
-	return {
+	static const TArray<FRigVMTemplateArgument> Arguments = {
 		CreateArgument(ArrayName, ERigVMPinDirection::IO)
 	};
+	return Arguments;
 }
 
 FRigVMTemplateTypeMap FRigVMDispatch_ArrayReset::OnNewArgumentType(const FName& InArgumentName, TRigVMTypeIndex InTypeIndex) const
@@ -365,12 +368,13 @@ FName FRigVMDispatch_ArrayGetNum::GetArgumentNameForOperandIndex(int32 InOperand
 	return ArgumentNames[InOperandIndex];
 }
 
-TArray<FRigVMTemplateArgument> FRigVMDispatch_ArrayGetNum::GetArguments() const
+const TArray<FRigVMTemplateArgument>& FRigVMDispatch_ArrayGetNum::GetArguments() const
 {
-	return {
+	static const TArray<FRigVMTemplateArgument> Arguments = {
 		CreateArgument(ArrayName, ERigVMPinDirection::Input),
 		CreateArgument(NumName, ERigVMPinDirection::Output)
 	};
+	return Arguments;
 }
 
 FRigVMTemplateTypeMap FRigVMDispatch_ArrayGetNum::OnNewArgumentType(const FName& InArgumentName, TRigVMTypeIndex InTypeIndex) const
@@ -424,12 +428,13 @@ FName FRigVMDispatch_ArraySetNum::GetArgumentNameForOperandIndex(int32 InOperand
 	return ArgumentNames[InOperandIndex];
 }
 
-TArray<FRigVMTemplateArgument> FRigVMDispatch_ArraySetNum::GetArguments() const
+const TArray<FRigVMTemplateArgument>& FRigVMDispatch_ArraySetNum::GetArguments() const
 {
-	return {
+	static const TArray<FRigVMTemplateArgument> Arguments = {
 		CreateArgument(ArrayName, ERigVMPinDirection::IO),
 		CreateArgument(NumName, ERigVMPinDirection::Input)
 	};
+	return Arguments;
 }
 
 FRigVMTemplateTypeMap FRigVMDispatch_ArraySetNum::OnNewArgumentType(const FName& InArgumentName, TRigVMTypeIndex InTypeIndex) const
@@ -487,13 +492,14 @@ FName FRigVMDispatch_ArrayGetAtIndex::GetArgumentNameForOperandIndex(int32 InOpe
 	return ArgumentNames[InOperandIndex];
 }
 
-TArray<FRigVMTemplateArgument> FRigVMDispatch_ArrayGetAtIndex::GetArguments() const
+const TArray<FRigVMTemplateArgument>& FRigVMDispatch_ArrayGetAtIndex::GetArguments() const
 {
-	return {
+	static const TArray<FRigVMTemplateArgument> Arguments = {
 		CreateArgument(ArrayName, ERigVMPinDirection::Input),
 		CreateArgument(IndexName, ERigVMPinDirection::Input),
 		CreateArgument(ElementName, ERigVMPinDirection::Output)
 	};
+	return Arguments;
 }
 
 FRigVMTemplateTypeMap FRigVMDispatch_ArrayGetAtIndex::OnNewArgumentType(const FName& InArgumentName, TRigVMTypeIndex InTypeIndex) const
@@ -569,13 +575,14 @@ FName FRigVMDispatch_ArraySetAtIndex::GetArgumentNameForOperandIndex(int32 InOpe
 	return ArgumentNames[InOperandIndex];
 }
 
-TArray<FRigVMTemplateArgument> FRigVMDispatch_ArraySetAtIndex::GetArguments() const
+const TArray<FRigVMTemplateArgument>& FRigVMDispatch_ArraySetAtIndex::GetArguments() const
 {
-	return {
+	static const TArray<FRigVMTemplateArgument> Arguments = {
 		CreateArgument(ArrayName, ERigVMPinDirection::IO),
 		CreateArgument(IndexName, ERigVMPinDirection::Input),
 		CreateArgument(ElementName, ERigVMPinDirection::Input)
 	};
+	return Arguments;
 }
 
 FRigVMTemplateTypeMap FRigVMDispatch_ArraySetAtIndex::OnNewArgumentType(const FName& InArgumentName, TRigVMTypeIndex InTypeIndex) const
@@ -651,13 +658,14 @@ FName FRigVMDispatch_ArrayAdd::GetArgumentNameForOperandIndex(int32 InOperandInd
 	return ArgumentNames[InOperandIndex];
 }
 
-TArray<FRigVMTemplateArgument> FRigVMDispatch_ArrayAdd::GetArguments() const
+const TArray<FRigVMTemplateArgument>& FRigVMDispatch_ArrayAdd::GetArguments() const
 {
-	return {
+	static const TArray<FRigVMTemplateArgument> Arguments = {
 		CreateArgument(ArrayName, ERigVMPinDirection::IO),
 		CreateArgument(ElementName, ERigVMPinDirection::Input),
 		CreateArgument(IndexName, ERigVMPinDirection::Output)
 	};
+	return Arguments;
 }
 
 void FRigVMDispatch_ArrayAdd::Execute(FRigVMExtendedExecuteContext& InContext, FRigVMMemoryHandleArray Handles, FRigVMPredicateBranchArray Predicates)
@@ -768,12 +776,13 @@ FName FRigVMDispatch_ArrayRemove::GetArgumentNameForOperandIndex(int32 InOperand
 	return ArgumentNames[InOperandIndex];
 }
 
-TArray<FRigVMTemplateArgument> FRigVMDispatch_ArrayRemove::GetArguments() const
+const TArray<FRigVMTemplateArgument>& FRigVMDispatch_ArrayRemove::GetArguments() const
 {
-	return {
+	static const TArray<FRigVMTemplateArgument> Arguments = {
 		CreateArgument(ArrayName, ERigVMPinDirection::IO),
 		CreateArgument(IndexName, ERigVMPinDirection::Input)
 	};
+	return Arguments;
 }
 
 FRigVMTemplateTypeMap FRigVMDispatch_ArrayRemove::OnNewArgumentType(const FName& InArgumentName, TRigVMTypeIndex InTypeIndex) const
@@ -864,14 +873,15 @@ FName FRigVMDispatch_ArrayFind::GetArgumentNameForOperandIndex(int32 InOperandIn
 	return ArgumentNames[InOperandIndex];
 }
 
-TArray<FRigVMTemplateArgument> FRigVMDispatch_ArrayFind::GetArguments() const
+const TArray<FRigVMTemplateArgument>& FRigVMDispatch_ArrayFind::GetArguments() const
 {
-	return {
+	static const TArray<FRigVMTemplateArgument> Arguments = {
 		CreateArgument(ArrayName, ERigVMPinDirection::Input),
 		CreateArgument(ElementName, ERigVMPinDirection::Input),
 		CreateArgument(IndexName, ERigVMPinDirection::Output),
 		CreateArgument(SuccessName, ERigVMPinDirection::Output)
 	};
+	return Arguments;
 }
 
 FRigVMTemplateTypeMap FRigVMDispatch_ArrayFind::OnNewArgumentType(const FName& InArgumentName, TRigVMTypeIndex InTypeIndex) const
@@ -973,12 +983,13 @@ FName FRigVMDispatch_ArrayAppend::GetArgumentNameForOperandIndex(int32 InOperand
 	return ArgumentNames[InOperandIndex];
 }
 
-TArray<FRigVMTemplateArgument> FRigVMDispatch_ArrayAppend::GetArguments() const
+const TArray<FRigVMTemplateArgument>& FRigVMDispatch_ArrayAppend::GetArguments() const
 {
-	return {
+	static const TArray<FRigVMTemplateArgument> Arguments = {
 		CreateArgument(ArrayName, ERigVMPinDirection::IO),
 		CreateArgument(OtherName, ERigVMPinDirection::Input)
 	};
+	return Arguments;
 }
 
 FRigVMTemplateTypeMap FRigVMDispatch_ArrayAppend::OnNewArgumentType(const FName& InArgumentName, TRigVMTypeIndex InTypeIndex) const
@@ -1052,12 +1063,13 @@ FName FRigVMDispatch_ArrayClone::GetArgumentNameForOperandIndex(int32 InOperandI
 	return ArgumentNames[InOperandIndex];
 }
 
-TArray<FRigVMTemplateArgument> FRigVMDispatch_ArrayClone::GetArguments() const
+const TArray<FRigVMTemplateArgument>& FRigVMDispatch_ArrayClone::GetArguments() const
 {
-	return {
+	static const TArray<FRigVMTemplateArgument> Arguments = {
 		CreateArgument(ArrayName, ERigVMPinDirection::Input),
 		CreateArgument(CloneName, ERigVMPinDirection::Output)
 	};
+	return Arguments;
 }
 
 FRigVMTemplateTypeMap FRigVMDispatch_ArrayClone::OnNewArgumentType(const FName& InArgumentName, TRigVMTypeIndex InTypeIndex) const
@@ -1232,13 +1244,14 @@ FName FRigVMDispatch_ArrayDifference::GetArgumentNameForOperandIndex(int32 InOpe
 	return ArgumentNames[InOperandIndex];
 }
 
-TArray<FRigVMTemplateArgument> FRigVMDispatch_ArrayDifference::GetArguments() const
+const TArray<FRigVMTemplateArgument>& FRigVMDispatch_ArrayDifference::GetArguments() const
 {
-	return {
+	static const TArray<FRigVMTemplateArgument> Arguments = {
 		CreateArgument(ArrayName, ERigVMPinDirection::Input),
 		CreateArgument(OtherName, ERigVMPinDirection::Input),
 		CreateArgument(ResultName, ERigVMPinDirection::Output)
 	};
+	return Arguments;
 }
 
 FRigVMTemplateTypeMap FRigVMDispatch_ArrayDifference::OnNewArgumentType(const FName& InArgumentName, TRigVMTypeIndex InTypeIndex) const
@@ -1400,9 +1413,9 @@ FName FRigVMDispatch_ArrayIterator::GetArgumentNameForOperandIndex(int32 InOpera
 	return ArgumentNames[InOperandIndex];
 }
 
-TArray<FRigVMTemplateArgument> FRigVMDispatch_ArrayIterator::GetArguments() const
+const TArray<FRigVMTemplateArgument>& FRigVMDispatch_ArrayIterator::GetArguments() const
 {
-	return {
+	static const TArray<FRigVMTemplateArgument> Arguments = {
 		CreateArgument(ArrayName, ERigVMPinDirection::Input),
 		CreateArgument(ElementName, ERigVMPinDirection::Output),
 		CreateArgument(IndexName, ERigVMPinDirection::Output),
@@ -1410,13 +1423,18 @@ TArray<FRigVMTemplateArgument> FRigVMDispatch_ArrayIterator::GetArguments() cons
 		CreateArgument(RatioName, ERigVMPinDirection::Output),
 		{FRigVMStruct::ControlFlowBlockToRunName, ERigVMPinDirection::Hidden, RigVMTypeUtils::TypeIndex::FName}
 	};
+	return Arguments;
 }
 
-TArray<FRigVMExecuteArgument> FRigVMDispatch_ArrayIterator::GetExecuteArguments_Impl(
+const TArray<FRigVMExecuteArgument>& FRigVMDispatch_ArrayIterator::GetExecuteArguments_Impl(
 	const FRigVMDispatchContext& InContext) const
 {
-	TArray<FRigVMExecuteArgument> ExecuteArguments = FRigVMDispatch_ArrayBaseMutable::GetExecuteArguments_Impl(InContext);
-	ExecuteArguments.Emplace(CompletedName, ERigVMPinDirection::Output);
+	static TArray<FRigVMExecuteArgument> ExecuteArguments;
+	if(ExecuteArguments.IsEmpty())
+	{
+		ExecuteArguments = FRigVMDispatch_ArrayBaseMutable::GetExecuteArguments_Impl(InContext);
+		ExecuteArguments.Emplace(CompletedName, ERigVMPinDirection::Output);
+	}
 	return ExecuteArguments;
 }
 
diff --git a/Engine/Plugins/Runtime/RigVM/Source/RigVM/Private/RigVMFunctions/RigVMDispatch_Constant.cpp b/Engine/Plugins/Runtime/RigVM/Source/RigVM/Private/RigVMFunctions/RigVMDispatch_Constant.cpp
--- a/Engine/Plugins/Runtime/RigVM/Source/RigVM/Private/RigVMFunctions/RigVMDispatch_Constant.cpp
+++ b/Engine/Plugins/Runtime/RigVM/Source/RigVM/Private/RigVMFunctions/RigVMDispatch_Constant.cpp
@@ -10,14 +10,15 @@
 
 const FName FRigVMDispatch_Constant::ValueName = TEXT("Value");
 
-TArray<FRigVMTemplateArgument> FRigVMDispatch_Constant::GetArguments() const
+const TArray<FRigVMTemplateArgument>& FRigVMDispatch_Constant::GetArguments() const
 {
-	const TArray<FRigVMTemplateArgument::ETypeCategory> Categories = {
+	static const TArray<FRigVMTemplateArgument::ETypeCategory> Categories = {
 		FRigVMTemplateArgument::ETypeCategory_SingleAnyValue
 	};
-	return {
+	static const TArray<FRigVMTemplateArgument> Arguments = {
 		{ValueName, ERigVMPinDirection::IO, Categories}
 	};
+	return Arguments;
 }
 
 FRigVMTemplateTypeMap FRigVMDispatch_Constant::OnNewArgumentType(const FName& InArgumentName, TRigVMTypeIndex InTypeIndex) const
diff --git a/Engine/Plugins/Runtime/RigVM/Source/RigVM/Private/RigVMFunctions/RigVMDispatch_Core.cpp b/Engine/Plugins/Runtime/RigVM/Source/RigVM/Private/RigVMFunctions/RigVMDispatch_Core.cpp
--- a/Engine/Plugins/Runtime/RigVM/Source/RigVM/Private/RigVMFunctions/RigVMDispatch_Core.cpp
+++ b/Engine/Plugins/Runtime/RigVM/Source/RigVM/Private/RigVMFunctions/RigVMDispatch_Core.cpp
@@ -23,17 +23,18 @@ FName FRigVMDispatch_CoreEquals::GetArgumentNameForOperandIndex(int32 InOperandI
 	return ArgumentNames[InOperandIndex];
 }
 
-TArray<FRigVMTemplateArgument> FRigVMDispatch_CoreEquals::GetArguments() const
+const TArray<FRigVMTemplateArgument>& FRigVMDispatch_CoreEquals::GetArguments() const
 {
-	const TArray<FRigVMTemplateArgument::ETypeCategory> ValueCategories = {
+	static const TArray<FRigVMTemplateArgument::ETypeCategory> ValueCategories = {
 		FRigVMTemplateArgument::ETypeCategory_SingleAnyValue,
 		FRigVMTemplateArgument::ETypeCategory_ArrayAnyValue
 	};
-	return {
+	static const TArray<FRigVMTemplateArgument> Arguments = {
 		FRigVMTemplateArgument(AName, ERigVMPinDirection::Input, ValueCategories),
 		FRigVMTemplateArgument(BName, ERigVMPinDirection::Input, ValueCategories),
 		FRigVMTemplateArgument(ResultName, ERigVMPinDirection::Output, RigVMTypeUtils::TypeIndex::Bool)
 	};
+	return Arguments;
 }
 
 FRigVMTemplateTypeMap FRigVMDispatch_CoreEquals::OnNewArgumentType(const FName& InArgumentName,
diff --git a/Engine/Plugins/Runtime/RigVM/Source/RigVM/Private/RigVMFunctions/RigVMDispatch_If.cpp b/Engine/Plugins/Runtime/RigVM/Source/RigVM/Private/RigVMFunctions/RigVMDispatch_If.cpp
--- a/Engine/Plugins/Runtime/RigVM/Source/RigVM/Private/RigVMFunctions/RigVMDispatch_If.cpp
+++ b/Engine/Plugins/Runtime/RigVM/Source/RigVM/Private/RigVMFunctions/RigVMDispatch_If.cpp
@@ -23,18 +23,19 @@ FName FRigVMDispatch_If::GetArgumentNameForOperandIndex(int32 InOperandIndex, in
 	return ArgumentNames[InOperandIndex];
 }
 
-TArray<FRigVMTemplateArgument> FRigVMDispatch_If::GetArguments() const
+const TArray<FRigVMTemplateArgument>& FRigVMDispatch_If::GetArguments() const
 {
 	static const TArray<FRigVMTemplateArgument::ETypeCategory> ValueCategories = {
 		FRigVMTemplateArgument::ETypeCategory_SingleAnyValue,
 		FRigVMTemplateArgument::ETypeCategory_ArrayAnyValue
 	};
-	return {
+	static const TArray<FRigVMTemplateArgument> Arguments = {
 		FRigVMTemplateArgument(ConditionName, ERigVMPinDirection::Input, RigVMTypeUtils::TypeIndex::Bool),
 		FRigVMTemplateArgument(TrueName, ERigVMPinDirection::Input, ValueCategories),
 		FRigVMTemplateArgument(FalseName, ERigVMPinDirection::Input, ValueCategories),
 		FRigVMTemplateArgument(ResultName, ERigVMPinDirection::Output, ValueCategories)
 	};
+	return Arguments;
 }
 
 FRigVMTemplateTypeMap FRigVMDispatch_If::OnNewArgumentType(const FName& InArgumentName,
diff --git a/Engine/Plugins/Runtime/RigVM/Source/RigVM/Private/RigVMFunctions/RigVMDispatch_MakeStruct.cpp b/Engine/Plugins/Runtime/RigVM/Source/RigVM/Private/RigVMFunctions/RigVMDispatch_MakeStruct.cpp
--- a/Engine/Plugins/Runtime/RigVM/Source/RigVM/Private/RigVMFunctions/RigVMDispatch_MakeStruct.cpp
+++ b/Engine/Plugins/Runtime/RigVM/Source/RigVM/Private/RigVMFunctions/RigVMDispatch_MakeStruct.cpp
@@ -11,15 +11,16 @@
 const FName FRigVMDispatch_MakeStruct::ElementsName = TEXT("Elements");
 const FName FRigVMDispatch_MakeStruct::StructName = TEXT("Struct");
 
-TArray<FRigVMTemplateArgument> FRigVMDispatch_MakeStruct::GetArguments() const
+const TArray<FRigVMTemplateArgument>& FRigVMDispatch_MakeStruct::GetArguments() const
 {
-	const TArray<FRigVMTemplateArgument::ETypeCategory> Categories = {
+	static const TArray<FRigVMTemplateArgument::ETypeCategory> Categories = {
 		FRigVMTemplateArgument::ETypeCategory_SingleScriptStructValue
 	};
-	return {
+	static TArray<FRigVMTemplateArgument> Arguments = {
 		{ElementsName, ERigVMPinDirection::Input, Categories},
 		{StructName, ERigVMPinDirection::Output, Categories}
 	};
+	return Arguments;
 }
 
 FRigVMTemplateTypeMap FRigVMDispatch_MakeStruct::OnNewArgumentType(const FName& InArgumentName, TRigVMTypeIndex InTypeIndex) const
@@ -84,15 +85,16 @@ void FRigVMDispatch_MakeStruct::Execute(FRigVMExtendedExecuteContext& InContext,
 	URigVMMemoryStorage::CopyProperty(TargetProperty, TargetMemory, SourceProperty, SourceMemory);
 }
 
-TArray<FRigVMTemplateArgument> FRigVMDispatch_BreakStruct::GetArguments() const
+const TArray<FRigVMTemplateArgument>& FRigVMDispatch_BreakStruct::GetArguments() const
 {
-	const TArray<FRigVMTemplateArgument::ETypeCategory> Categories = {
+	static const TArray<FRigVMTemplateArgument::ETypeCategory> Categories = {
 		FRigVMTemplateArgument::ETypeCategory_SingleScriptStructValue
 	};
-	return {
+	static const TArray<FRigVMTemplateArgument> Arguments = {
 		{StructName, ERigVMPinDirection::Input, Categories},
 		{ElementsName, ERigVMPinDirection::Output, Categories}
 	};
+	return Arguments;
 }
 
 #if WITH_EDITOR
diff --git a/Engine/Plugins/Runtime/RigVM/Source/RigVM/Private/RigVMFunctions/RigVMDispatch_Print.cpp b/Engine/Plugins/Runtime/RigVM/Source/RigVM/Private/RigVMFunctions/RigVMDispatch_Print.cpp
--- a/Engine/Plugins/Runtime/RigVM/Source/RigVM/Private/RigVMFunctions/RigVMDispatch_Print.cpp
+++ b/Engine/Plugins/Runtime/RigVM/Source/RigVM/Private/RigVMFunctions/RigVMDispatch_Print.cpp
@@ -28,24 +28,28 @@ FName FRigVMDispatch_Print::GetArgumentNameForOperandIndex(int32 InOperandIndex,
 	return ArgumentNames[InOperandIndex];
 }
 
-TArray<FRigVMTemplateArgument> FRigVMDispatch_Print::GetArguments() const
+const TArray<FRigVMTemplateArgument>& FRigVMDispatch_Print::GetArguments() const
 {
-	const TArray<FRigVMTemplateArgument::ETypeCategory> ValueCategories = {
+	static const TArray<FRigVMTemplateArgument::ETypeCategory> ValueCategories = {
 		FRigVMTemplateArgument::ETypeCategory_SingleAnyValue,
 		FRigVMTemplateArgument::ETypeCategory_ArrayAnyValue
 	};
-	return {
+	static const TArray<FRigVMTemplateArgument> Arguments = {
 		FRigVMTemplateArgument(PrefixName, ERigVMPinDirection::Input, RigVMTypeUtils::TypeIndex::FString),
 		FRigVMTemplateArgument(ValueName, ERigVMPinDirection::Input, ValueCategories),
 		FRigVMTemplateArgument(EnabledName, ERigVMPinDirection::Input, RigVMTypeUtils::TypeIndex::Bool),
 		FRigVMTemplateArgument(ScreenDurationName, ERigVMPinDirection::Input, RigVMTypeUtils::TypeIndex::Float),
 		FRigVMTemplateArgument(ScreenColorName, ERigVMPinDirection::Input, FRigVMRegistry::Get().GetTypeIndex<FLinearColor>())
 	};
+	return Arguments;
 }
 
-TArray<FRigVMExecuteArgument> FRigVMDispatch_Print::GetExecuteArguments_Impl(const FRigVMDispatchContext& InContext) const
+const TArray<FRigVMExecuteArgument>& FRigVMDispatch_Print::GetExecuteArguments_Impl(const FRigVMDispatchContext& InContext) const
 {
-	return {{TEXT("ExecuteContext"), ERigVMPinDirection::IO}};
+	static const TArray<FRigVMExecuteArgument> Arguments = {
+		{TEXT("ExecuteContext"), ERigVMPinDirection::IO}
+	};
+	return Arguments;
 }
 
 FRigVMTemplateTypeMap FRigVMDispatch_Print::OnNewArgumentType(const FName& InArgumentName,
diff --git a/Engine/Plugins/Runtime/RigVM/Source/RigVM/Private/RigVMFunctions/RigVMDispatch_Select.cpp b/Engine/Plugins/Runtime/RigVM/Source/RigVM/Private/RigVMFunctions/RigVMDispatch_Select.cpp
--- a/Engine/Plugins/Runtime/RigVM/Source/RigVM/Private/RigVMFunctions/RigVMDispatch_Select.cpp
+++ b/Engine/Plugins/Runtime/RigVM/Source/RigVM/Private/RigVMFunctions/RigVMDispatch_Select.cpp
@@ -23,7 +23,7 @@ FName FRigVMDispatch_SelectInt32::GetArgumentNameForOperandIndex(int32 InOperand
 }
 
 
-TArray<FRigVMTemplateArgument> FRigVMDispatch_SelectInt32::GetArguments() const
+const TArray<FRigVMTemplateArgument>& FRigVMDispatch_SelectInt32::GetArguments() const
 {
 	static const TArray<FRigVMTemplateArgument::ETypeCategory> ArrayCategories = {
 		FRigVMTemplateArgument::ETypeCategory_ArrayAnyValue,
@@ -34,11 +34,13 @@ TArray<FRigVMTemplateArgument> FRigVMDispatch_SelectInt32::GetArguments() const
 		FRigVMTemplateArgument::ETypeCategory_ArrayAnyValue
 	};
 	
-	return {
+	static const TArray<FRigVMTemplateArgument> Arguments = {
 		FRigVMTemplateArgument(IndexName, ERigVMPinDirection::Input, RigVMTypeUtils::TypeIndex::Int32),
 		FRigVMTemplateArgument(ValuesName, ERigVMPinDirection::Input, ArrayCategories),
 		FRigVMTemplateArgument(ResultName, ERigVMPinDirection::Output, SingleCategories)
 	};
+
+	return Arguments;
 }
 
 FRigVMTemplateTypeMap FRigVMDispatch_SelectInt32::OnNewArgumentType(const FName& InArgumentName,
diff --git a/Engine/Plugins/Runtime/RigVM/Source/RigVM/Private/RigVMFunctions/RigVMDispatch_Switch.cpp b/Engine/Plugins/Runtime/RigVM/Source/RigVM/Private/RigVMFunctions/RigVMDispatch_Switch.cpp
--- a/Engine/Plugins/Runtime/RigVM/Source/RigVM/Private/RigVMFunctions/RigVMDispatch_Switch.cpp
+++ b/Engine/Plugins/Runtime/RigVM/Source/RigVM/Private/RigVMFunctions/RigVMDispatch_Switch.cpp
@@ -20,15 +20,16 @@ FName FRigVMDispatch_SwitchInt32::GetArgumentNameForOperandIndex(int32 InOperand
 	return ArgumentNames[InOperandIndex];
 }
 
-TArray<FRigVMTemplateArgument> FRigVMDispatch_SwitchInt32::GetArguments() const
+const TArray<FRigVMTemplateArgument>& FRigVMDispatch_SwitchInt32::GetArguments() const
 {
-	return {
+	static const TArray<FRigVMTemplateArgument> Arguments = {
 		FRigVMTemplateArgument(IndexName, ERigVMPinDirection::Input, RigVMTypeUtils::TypeIndex::Int32),
 		FRigVMTemplateArgument(FRigVMStruct::ControlFlowBlockToRunName, ERigVMPinDirection::Hidden, RigVMTypeUtils::TypeIndex::FName)
 	};
+	return Arguments;
 }
 
-TArray<FRigVMExecuteArgument> FRigVMDispatch_SwitchInt32::GetExecuteArguments_Impl(const FRigVMDispatchContext& InContext) const
+const TArray<FRigVMExecuteArgument>& FRigVMDispatch_SwitchInt32::GetExecuteArguments_Impl(const FRigVMDispatchContext& InContext) const
 {
 	static const TArray<FRigVMExecuteArgument> Arguments =
 	{
diff --git a/Engine/Plugins/Runtime/RigVM/Source/RigVM/Private/RigVMFunctions/RigVMFunction_String.cpp b/Engine/Plugins/Runtime/RigVM/Source/RigVM/Private/RigVMFunctions/RigVMFunction_String.cpp
--- a/Engine/Plugins/Runtime/RigVM/Source/RigVM/Private/RigVMFunctions/RigVMFunction_String.cpp
+++ b/Engine/Plugins/Runtime/RigVM/Source/RigVM/Private/RigVMFunctions/RigVMFunction_String.cpp
@@ -252,16 +252,17 @@ FRigVMFunction_StringPadInteger_Execute()
 	}
 }
 
-TArray<FRigVMTemplateArgument> FRigDispatch_ToString::GetArguments() const
+const TArray<FRigVMTemplateArgument>& FRigDispatch_ToString::GetArguments() const
 {
-	const TArray<FRigVMTemplateArgument::ETypeCategory> ValueCategories = {
+	static const TArray<FRigVMTemplateArgument::ETypeCategory> ValueCategories = {
 		FRigVMTemplateArgument::ETypeCategory_SingleAnyValue,
 		FRigVMTemplateArgument::ETypeCategory_ArrayAnyValue
 	};
-	return {
+	static const TArray<FRigVMTemplateArgument> Arguments = {
 		FRigVMTemplateArgument(TEXT("Value"), ERigVMPinDirection::Input, ValueCategories),
 		FRigVMTemplateArgument(TEXT("Result"), ERigVMPinDirection::Output, RigVMTypeUtils::TypeIndex::FString)
 	};
+	return Arguments;
 }
 
 FRigVMTemplateTypeMap FRigDispatch_ToString::OnNewArgumentType(const FName& InArgumentName,
@@ -317,16 +318,17 @@ void FRigDispatch_ToString::Execute(FRigVMExtendedExecuteContext& InContext, FRi
 	ValueProperty->ExportText_Direct(Result, Value, Value, nullptr, PPF_None, nullptr);
 }
 
-TArray<FRigVMTemplateArgument> FRigDispatch_FromString::GetArguments() const
+const TArray<FRigVMTemplateArgument>& FRigDispatch_FromString::GetArguments() const
 {
-	const TArray<FRigVMTemplateArgument::ETypeCategory> ValueCategories = {
+	static const TArray<FRigVMTemplateArgument::ETypeCategory> ValueCategories = {
 		FRigVMTemplateArgument::ETypeCategory_SingleAnyValue,
 		FRigVMTemplateArgument::ETypeCategory_ArrayAnyValue
 	};
-	return {
+	static const TArray<FRigVMTemplateArgument> Arguments = {
 		FRigVMTemplateArgument(TEXT("String"), ERigVMPinDirection::Input, RigVMTypeUtils::TypeIndex::FString),
 		FRigVMTemplateArgument(TEXT("Result"), ERigVMPinDirection::Output, ValueCategories)
 	};
+	return Arguments;
 }
 
 FRigVMTemplateTypeMap FRigDispatch_FromString::OnNewArgumentType(const FName& InArgumentName,
diff --git a/Engine/Plugins/Runtime/RigVM/Source/RigVM/Public/RigVMCore/RigVMDispatchFactory.h b/Engine/Plugins/Runtime/RigVM/Source/RigVM/Public/RigVMCore/RigVMDispatchFactory.h
--- a/Engine/Plugins/Runtime/RigVM/Source/RigVM/Public/RigVMCore/RigVMDispatchFactory.h
+++ b/Engine/Plugins/Runtime/RigVM/Source/RigVM/Public/RigVMCore/RigVMDispatchFactory.h
@@ -52,9 +52,14 @@ public:
 	FRigVMDispatchFactory()
 		: FactoryScriptStruct(nullptr)
 		, CachedTemplate(nullptr)
-	{}
+	{
+		ArgumentNamesMutex = new FCriticalSection();
+	}
 	
-	virtual ~FRigVMDispatchFactory() {}
+	virtual ~FRigVMDispatchFactory()
+	{
+		delete ArgumentNamesMutex;
+	}
 
 	// returns the name of the factory template
 	RIGVM_API FName GetFactoryName() const;
@@ -135,7 +140,7 @@ public:
 	virtual void RegisterDependencyTypes() const {}
 
 	// returns the arguments of the template
-	virtual TArray<FRigVMTemplateArgument> GetArguments() const { return TArray<FRigVMTemplateArgument>(); }
+	RIGVM_API virtual const TArray<FRigVMTemplateArgument>& GetArguments() const;
 
 	// returns the execute arguments of the template
 	RIGVM_API TArray<FRigVMExecuteArgument> GetExecuteArguments(const FRigVMDispatchContext& InContext) const;
@@ -177,7 +182,7 @@ protected:
 	
 	virtual TArray<FRigVMFunction> GetDispatchPredicatesImpl(const FRigVMTemplateTypeMap& InTypes) const { return TArray<FRigVMFunction>(); }
 
-	virtual TArray<FRigVMExecuteArgument> GetExecuteArguments_Impl(const FRigVMDispatchContext& InContext) const { return TArray<FRigVMExecuteArgument>(); }
+	RIGVM_API virtual const TArray<FRigVMExecuteArgument>& GetExecuteArguments_Impl(const FRigVMDispatchContext& InContext) const;
 
 	RIGVM_API static bool CopyProperty(
 		const FProperty* InTargetProperty,
@@ -229,9 +234,16 @@ protected:
 	RIGVM_API static const FString DispatchPrefix;
 	RIGVM_API static const FString TrueString;
 
+	const TArray<FName>* UpdateArgumentNameCache(int32 InNumberOperands) const;
+	const TArray<FName>* UpdateArgumentNameCache_NoLock(int32 InNumberOperands) const;
+
 	UScriptStruct* FactoryScriptStruct;
 	mutable const FRigVMTemplate* CachedTemplate;
 	static FCriticalSection GetTemplateMutex;
+	mutable TMap<int32, TSharedPtr<TArray<FName>>> ArgumentNamesMap;
+	FCriticalSection* ArgumentNamesMutex;
 	friend struct FRigVMTemplate;
 	friend struct FRigVMRegistry;
+	friend struct FRigVMFunction;
+	friend class URigVM;
 };
diff --git a/Engine/Plugins/Runtime/RigVM/Source/RigVM/Public/RigVMCore/RigVMFunction.h b/Engine/Plugins/Runtime/RigVM/Source/RigVM/Public/RigVMCore/RigVMFunction.h
--- a/Engine/Plugins/Runtime/RigVM/Source/RigVM/Public/RigVMCore/RigVMFunction.h
+++ b/Engine/Plugins/Runtime/RigVM/Source/RigVM/Public/RigVMCore/RigVMFunction.h
@@ -134,6 +134,7 @@ struct RIGVM_API FRigVMFunction
 	int32 TemplateIndex;
 	TArray<FRigVMFunctionArgument> Arguments;
 	mutable TArray<TRigVMTypeIndex> ArgumentTypeIndices;
+	mutable TArray<FName> ArgumentNames;
 
 	FRigVMFunction()
 		: Name()
@@ -144,19 +145,11 @@ struct RIGVM_API FRigVMFunction
 		, TemplateIndex(INDEX_NONE)
 		, Arguments()
 		, ArgumentTypeIndices()
+		, ArgumentNames()
 	{
 	}
 
-	FRigVMFunction(const TCHAR* InName, FRigVMFunctionPtr InFunctionPtr, UScriptStruct* InStruct = nullptr, int32 InIndex = INDEX_NONE, const TArray<FRigVMFunctionArgument>& InArguments = TArray<FRigVMFunctionArgument>())
-		: Name(InName)
-		, Struct(InStruct)
-		, Factory(nullptr)
-		, FunctionPtr(InFunctionPtr)
-		, Index(InIndex)
-		, TemplateIndex(INDEX_NONE)
-		, Arguments(InArguments)
-	{
-	}
+	FRigVMFunction(const TCHAR* InName, FRigVMFunctionPtr InFunctionPtr, UScriptStruct* InStruct = nullptr, int32 InIndex = INDEX_NONE, const TArray<FRigVMFunctionArgument>& InArguments = TArray<FRigVMFunctionArgument>());
 
 	FRigVMFunction(const FString& InName, FRigVMFunctionPtr InFunctionPtr, FRigVMDispatchFactory* InFactory, int32 InIndex = INDEX_NONE, const TArray<FRigVMFunctionArgument>& InArguments = TArray<FRigVMFunctionArgument>())
 	: Name(InName)
@@ -166,6 +159,7 @@ struct RIGVM_API FRigVMFunction
 	, Index(InIndex)
 	, TemplateIndex(INDEX_NONE)
 	, Arguments(InArguments)
+	, ArgumentNames()
 	{
 	}
 
@@ -179,5 +173,5 @@ struct RIGVM_API FRigVMFunction
 	const FRigVMTemplate* GetTemplate() const;
 	const UScriptStruct* GetExecuteContextStruct() const;
 	bool SupportsExecuteContextStruct(const UScriptStruct* InExecuteContextStruct) const;
-	FName GetArgumentNameForOperandIndex(int32 InOperandIndex, int32 InTotalOperands) const;
+	const FName& GetArgumentNameForOperandIndex(int32 InOperandIndex, int32 InTotalOperands) const;
 };
diff --git a/Engine/Plugins/Runtime/RigVM/Source/RigVM/Public/RigVMFunctions/RigVMDispatch_Array.h b/Engine/Plugins/Runtime/RigVM/Source/RigVM/Public/RigVMFunctions/RigVMDispatch_Array.h
--- a/Engine/Plugins/Runtime/RigVM/Source/RigVM/Public/RigVMFunctions/RigVMDispatch_Array.h
+++ b/Engine/Plugins/Runtime/RigVM/Source/RigVM/Public/RigVMFunctions/RigVMDispatch_Array.h
@@ -48,7 +48,7 @@ struct RIGVM_API FRigVMDispatch_ArrayBaseMutable : public FRigVMDispatch_ArrayBa
 	GENERATED_BODY()
 
 public:
-	virtual TArray<FRigVMExecuteArgument> GetExecuteArguments_Impl(const FRigVMDispatchContext& InContext) const override;
+	virtual const TArray<FRigVMExecuteArgument>& GetExecuteArguments_Impl(const FRigVMDispatchContext& InContext) const override;
 };
 
 USTRUCT(meta=(DisplayName = "Make Array", Keywords = "Make,MakeArray,Constant,Reroute"))
@@ -62,7 +62,7 @@ public:
 		FactoryScriptStruct = StaticStruct();
 	}
 
-	virtual TArray<FRigVMTemplateArgument> GetArguments() const override;
+	virtual const TArray<FRigVMTemplateArgument>& GetArguments() const override;
 	virtual FRigVMTemplateTypeMap OnNewArgumentType(const FName& InArgumentName, TRigVMTypeIndex InTypeIndex) const override;
 #if WITH_EDITOR
 	virtual FText GetNodeTooltip(const FRigVMTemplateTypeMap& InTypes) const override;
@@ -90,7 +90,7 @@ public:
 
 	virtual ERigVMOpCode GetOpCode() const override { return ERigVMOpCode::ArrayReset; }
 	virtual FName GetArgumentNameForOperandIndex(int32 InOperandIndex, int32 InTotalOperands) const override;
-	virtual TArray<FRigVMTemplateArgument> GetArguments() const override;
+	virtual const TArray<FRigVMTemplateArgument>& GetArguments() const override;
 	virtual FRigVMTemplateTypeMap OnNewArgumentType(const FName& InArgumentName, TRigVMTypeIndex InTypeIndex) const override;
 #if WITH_EDITOR
 	virtual FText GetNodeTooltip(const FRigVMTemplateTypeMap& InTypes) const override;
@@ -115,7 +115,7 @@ public:
 
 	virtual ERigVMOpCode GetOpCode() const override { return ERigVMOpCode::ArrayGetNum; }
 	virtual FName GetArgumentNameForOperandIndex(int32 InOperandIndex, int32 InTotalOperands) const override;
-	virtual TArray<FRigVMTemplateArgument> GetArguments() const override;
+	virtual const TArray<FRigVMTemplateArgument>& GetArguments() const override;
 	virtual FRigVMTemplateTypeMap OnNewArgumentType(const FName& InArgumentName, TRigVMTypeIndex InTypeIndex) const override;
 #if WITH_EDITOR
 	virtual FText GetNodeTooltip(const FRigVMTemplateTypeMap& InTypes) const override;
@@ -140,7 +140,7 @@ public:
 
 	virtual ERigVMOpCode GetOpCode() const override { return ERigVMOpCode::ArraySetNum; }
 	virtual FName GetArgumentNameForOperandIndex(int32 InOperandIndex, int32 InTotalOperands) const override;
-	virtual TArray<FRigVMTemplateArgument> GetArguments() const override;
+	virtual const TArray<FRigVMTemplateArgument>& GetArguments() const override;
 	virtual FRigVMTemplateTypeMap OnNewArgumentType(const FName& InArgumentName, TRigVMTypeIndex InTypeIndex) const override;
 #if WITH_EDITOR
 	virtual FText GetNodeTooltip(const FRigVMTemplateTypeMap& InTypes) const override;
@@ -165,7 +165,7 @@ public:
 
 	virtual ERigVMOpCode GetOpCode() const override { return ERigVMOpCode::ArrayGetAtIndex; }
 	virtual FName GetArgumentNameForOperandIndex(int32 InOperandIndex, int32 InTotalOperands) const override;
-	virtual TArray<FRigVMTemplateArgument> GetArguments() const override;
+	virtual const TArray<FRigVMTemplateArgument>& GetArguments() const override;
 	virtual FRigVMTemplateTypeMap OnNewArgumentType(const FName& InArgumentName, TRigVMTypeIndex InTypeIndex) const override;
 #if WITH_EDITOR
 	virtual FText GetNodeTooltip(const FRigVMTemplateTypeMap& InTypes) const override;
@@ -190,7 +190,7 @@ public:
 
 	virtual ERigVMOpCode GetOpCode() const override { return ERigVMOpCode::ArraySetAtIndex; }
 	virtual FName GetArgumentNameForOperandIndex(int32 InOperandIndex, int32 InTotalOperands) const override;
-	virtual TArray<FRigVMTemplateArgument> GetArguments() const override;
+	virtual const TArray<FRigVMTemplateArgument>& GetArguments() const override;
 	virtual FRigVMTemplateTypeMap OnNewArgumentType(const FName& InArgumentName, TRigVMTypeIndex InTypeIndex) const override;
 #if WITH_EDITOR
 	virtual FText GetNodeTooltip(const FRigVMTemplateTypeMap& InTypes) const override;
@@ -215,7 +215,7 @@ public:
 
 	virtual ERigVMOpCode GetOpCode() const override { return ERigVMOpCode::ArrayAdd; }
 	virtual FName GetArgumentNameForOperandIndex(int32 InOperandIndex, int32 InTotalOperands) const override;
-	virtual TArray<FRigVMTemplateArgument> GetArguments() const override;
+	virtual const TArray<FRigVMTemplateArgument>& GetArguments() const override;
 #if WITH_EDITOR
 	virtual FText GetNodeTooltip(const FRigVMTemplateTypeMap& InTypes) const override;
 	virtual FText GetArgumentTooltip(const FName& InArgumentName, TRigVMTypeIndex InTypeIndex) const override;
@@ -261,7 +261,7 @@ public:
 
 	virtual ERigVMOpCode GetOpCode() const override { return ERigVMOpCode::ArrayRemove; }
 	virtual FName GetArgumentNameForOperandIndex(int32 InOperandIndex, int32 InTotalOperands) const override;
-	virtual TArray<FRigVMTemplateArgument> GetArguments() const override;
+	virtual const TArray<FRigVMTemplateArgument>& GetArguments() const override;
 	virtual FRigVMTemplateTypeMap OnNewArgumentType(const FName& InArgumentName, TRigVMTypeIndex InTypeIndex) const override;
 #if WITH_EDITOR
    	virtual FText GetNodeTooltip(const FRigVMTemplateTypeMap& InTypes) const override;
@@ -309,7 +309,7 @@ public:
 
 	virtual ERigVMOpCode GetOpCode() const override { return ERigVMOpCode::ArrayFind; }
 	virtual FName GetArgumentNameForOperandIndex(int32 InOperandIndex, int32 InTotalOperands) const override;
-	virtual TArray<FRigVMTemplateArgument> GetArguments() const override;
+	virtual const TArray<FRigVMTemplateArgument>& GetArguments() const override;
 	virtual FRigVMTemplateTypeMap OnNewArgumentType(const FName& InArgumentName, TRigVMTypeIndex InTypeIndex) const override;
 #if WITH_EDITOR
    	virtual FText GetNodeTooltip(const FRigVMTemplateTypeMap& InTypes) const override;
@@ -334,7 +334,7 @@ public:
 
 	virtual ERigVMOpCode GetOpCode() const override { return ERigVMOpCode::ArrayAppend; }
 	virtual FName GetArgumentNameForOperandIndex(int32 InOperandIndex, int32 InTotalOperands) const override;
-	virtual TArray<FRigVMTemplateArgument> GetArguments() const override;
+	virtual const TArray<FRigVMTemplateArgument>& GetArguments() const override;
 	virtual FRigVMTemplateTypeMap OnNewArgumentType(const FName& InArgumentName, TRigVMTypeIndex InTypeIndex) const override;
 #if WITH_EDITOR
    	virtual FText GetNodeTooltip(const FRigVMTemplateTypeMap& InTypes) const override;
@@ -359,7 +359,7 @@ public:
 
 	virtual ERigVMOpCode GetOpCode() const override { return ERigVMOpCode::ArrayClone; }
 	virtual FName GetArgumentNameForOperandIndex(int32 InOperandIndex, int32 InTotalOperands) const override;
-	virtual TArray<FRigVMTemplateArgument> GetArguments() const override;
+	virtual const TArray<FRigVMTemplateArgument>& GetArguments() const override;
 	virtual FRigVMTemplateTypeMap OnNewArgumentType(const FName& InArgumentName, TRigVMTypeIndex InTypeIndex) const override;
 #if WITH_EDITOR
    	virtual FText GetNodeTooltip(const FRigVMTemplateTypeMap& InTypes) const override;
@@ -406,7 +406,7 @@ public:
 
 	virtual ERigVMOpCode GetOpCode() const override { return ERigVMOpCode::ArrayDifference; }
 	virtual FName GetArgumentNameForOperandIndex(int32 InOperandIndex, int32 InTotalOperands) const override;
-	virtual TArray<FRigVMTemplateArgument> GetArguments() const override;
+	virtual const TArray<FRigVMTemplateArgument>& GetArguments() const override;
 	virtual FRigVMTemplateTypeMap OnNewArgumentType(const FName& InArgumentName, TRigVMTypeIndex InTypeIndex) const override;
 #if WITH_EDITOR
    	virtual FText GetNodeTooltip(const FRigVMTemplateTypeMap& InTypes) const override;
@@ -453,8 +453,8 @@ public:
 
 	virtual ERigVMOpCode GetOpCode() const override { return ERigVMOpCode::ArrayIterator; }
 	virtual FName GetArgumentNameForOperandIndex(int32 InOperandIndex, int32 InTotalOperands) const override;
-	virtual TArray<FRigVMTemplateArgument> GetArguments() const override;
-	virtual TArray<FRigVMExecuteArgument> GetExecuteArguments_Impl(const FRigVMDispatchContext& InContext) const override;
+	virtual const TArray<FRigVMTemplateArgument>& GetArguments() const override;
+	virtual const TArray<FRigVMExecuteArgument>& GetExecuteArguments_Impl(const FRigVMDispatchContext& InContext) const override;
 	virtual const TArray<FName>& GetControlFlowBlocks_Impl(const FRigVMDispatchContext& InContext) const override;
 	virtual const bool IsControlFlowBlockSliced(const FName& InBlockName) const override;
 	virtual FRigVMTemplateTypeMap OnNewArgumentType(const FName& InArgumentName, TRigVMTypeIndex InTypeIndex) const override;
diff --git a/Engine/Plugins/Runtime/RigVM/Source/RigVM/Public/RigVMFunctions/RigVMDispatch_Constant.h b/Engine/Plugins/Runtime/RigVM/Source/RigVM/Public/RigVMFunctions/RigVMDispatch_Constant.h
--- a/Engine/Plugins/Runtime/RigVM/Source/RigVM/Public/RigVMFunctions/RigVMDispatch_Constant.h
+++ b/Engine/Plugins/Runtime/RigVM/Source/RigVM/Public/RigVMFunctions/RigVMDispatch_Constant.h
@@ -18,7 +18,7 @@ public:
 		FactoryScriptStruct = StaticStruct();
 	}
 
-	virtual TArray<FRigVMTemplateArgument> GetArguments() const override;
+	virtual const TArray<FRigVMTemplateArgument>& GetArguments() const override;
 	virtual FRigVMTemplateTypeMap OnNewArgumentType(const FName& InArgumentName, TRigVMTypeIndex InTypeIndex) const override;
 #if WITH_EDITOR
 	virtual FString GetNodeTitle(const FRigVMTemplateTypeMap& InTypes) const override;
diff --git a/Engine/Plugins/Runtime/RigVM/Source/RigVM/Public/RigVMFunctions/RigVMDispatch_Core.h b/Engine/Plugins/Runtime/RigVM/Source/RigVM/Public/RigVMFunctions/RigVMDispatch_Core.h
--- a/Engine/Plugins/Runtime/RigVM/Source/RigVM/Public/RigVMFunctions/RigVMDispatch_Core.h
+++ b/Engine/Plugins/Runtime/RigVM/Source/RigVM/Public/RigVMFunctions/RigVMDispatch_Core.h
@@ -27,7 +27,7 @@ public:
 	}
 
 	virtual FName GetArgumentNameForOperandIndex(int32 InOperandIndex, int32 InTotalOperands) const override;
-	virtual TArray<FRigVMTemplateArgument> GetArguments() const override;
+	virtual const TArray<FRigVMTemplateArgument>& GetArguments() const override;
 	virtual FRigVMTemplateTypeMap OnNewArgumentType(const FName& InArgumentName, TRigVMTypeIndex InTypeIndex) const override;
 
 protected:
diff --git a/Engine/Plugins/Runtime/RigVM/Source/RigVM/Public/RigVMFunctions/RigVMDispatch_If.h b/Engine/Plugins/Runtime/RigVM/Source/RigVM/Public/RigVMFunctions/RigVMDispatch_If.h
--- a/Engine/Plugins/Runtime/RigVM/Source/RigVM/Public/RigVMFunctions/RigVMDispatch_If.h
+++ b/Engine/Plugins/Runtime/RigVM/Source/RigVM/Public/RigVMFunctions/RigVMDispatch_If.h
@@ -22,7 +22,7 @@ public:
 	}
 
 	virtual FName GetArgumentNameForOperandIndex(int32 InOperandIndex, int32 InTotalOperands) const override;
-	virtual TArray<FRigVMTemplateArgument> GetArguments() const override;
+	virtual const TArray<FRigVMTemplateArgument>& GetArguments() const override;
 	virtual FRigVMTemplateTypeMap OnNewArgumentType(const FName& InArgumentName, TRigVMTypeIndex InTypeIndex) const override;
 	virtual bool IsSingleton() const override { return true; } 
 
diff --git a/Engine/Plugins/Runtime/RigVM/Source/RigVM/Public/RigVMFunctions/RigVMDispatch_MakeStruct.h b/Engine/Plugins/Runtime/RigVM/Source/RigVM/Public/RigVMFunctions/RigVMDispatch_MakeStruct.h
--- a/Engine/Plugins/Runtime/RigVM/Source/RigVM/Public/RigVMFunctions/RigVMDispatch_MakeStruct.h
+++ b/Engine/Plugins/Runtime/RigVM/Source/RigVM/Public/RigVMFunctions/RigVMDispatch_MakeStruct.h
@@ -18,7 +18,7 @@ public:
 		FactoryScriptStruct = StaticStruct();
 	}
 	
-	virtual TArray<FRigVMTemplateArgument> GetArguments() const override;
+	virtual const TArray<FRigVMTemplateArgument>& GetArguments() const override;
 	virtual FRigVMTemplateTypeMap OnNewArgumentType(const FName& InArgumentName, TRigVMTypeIndex InTypeIndex) const override;
 #if WITH_EDITOR
 	virtual FString GetNodeTitle(const FRigVMTemplateTypeMap& InTypes) const override;
@@ -48,7 +48,7 @@ public:
 		FactoryScriptStruct = StaticStruct();
 	}
 
-	virtual TArray<FRigVMTemplateArgument> GetArguments() const override;
+	virtual const TArray<FRigVMTemplateArgument>& GetArguments() const override;
 #if WITH_EDITOR
 	virtual FText GetNodeTooltip(const FRigVMTemplateTypeMap& InTypes) const override;
 #endif
diff --git a/Engine/Plugins/Runtime/RigVM/Source/RigVM/Public/RigVMFunctions/RigVMDispatch_Print.h b/Engine/Plugins/Runtime/RigVM/Source/RigVM/Public/RigVMFunctions/RigVMDispatch_Print.h
--- a/Engine/Plugins/Runtime/RigVM/Source/RigVM/Public/RigVMFunctions/RigVMDispatch_Print.h
+++ b/Engine/Plugins/Runtime/RigVM/Source/RigVM/Public/RigVMFunctions/RigVMDispatch_Print.h
@@ -21,8 +21,8 @@ public:
 	}
 
 	virtual FName GetArgumentNameForOperandIndex(int32 InOperandIndex, int32 InTotalOperands) const override;
-	virtual TArray<FRigVMTemplateArgument> GetArguments() const override;
-	virtual TArray<FRigVMExecuteArgument> GetExecuteArguments_Impl(const FRigVMDispatchContext& InContext) const override;
+	virtual const TArray<FRigVMTemplateArgument>& GetArguments() const override;
+	virtual const TArray<FRigVMExecuteArgument>& GetExecuteArguments_Impl(const FRigVMDispatchContext& InContext) const override;
 	virtual FRigVMTemplateTypeMap OnNewArgumentType(const FName& InArgumentName, TRigVMTypeIndex InTypeIndex) const override;
 	virtual bool IsSingleton() const override { return true; } 
 
diff --git a/Engine/Plugins/Runtime/RigVM/Source/RigVM/Public/RigVMFunctions/RigVMDispatch_Select.h b/Engine/Plugins/Runtime/RigVM/Source/RigVM/Public/RigVMFunctions/RigVMDispatch_Select.h
--- a/Engine/Plugins/Runtime/RigVM/Source/RigVM/Public/RigVMFunctions/RigVMDispatch_Select.h
+++ b/Engine/Plugins/Runtime/RigVM/Source/RigVM/Public/RigVMFunctions/RigVMDispatch_Select.h
@@ -22,7 +22,7 @@ public:
 	}
 
 	virtual FName GetArgumentNameForOperandIndex(int32 InOperandIndex, int32 InTotalOperands) const override;
-	virtual TArray<FRigVMTemplateArgument> GetArguments() const override;
+	virtual const TArray<FRigVMTemplateArgument>& GetArguments() const override;
 	virtual FRigVMTemplateTypeMap OnNewArgumentType(const FName& InArgumentName, TRigVMTypeIndex InTypeIndex) const override;
 #if WITH_EDITOR
 	virtual FString GetArgumentMetaData(const FName& InArgumentName, const FName& InMetaDataKey) const override;
diff --git a/Engine/Plugins/Runtime/RigVM/Source/RigVM/Public/RigVMFunctions/RigVMDispatch_Switch.h b/Engine/Plugins/Runtime/RigVM/Source/RigVM/Public/RigVMFunctions/RigVMDispatch_Switch.h
--- a/Engine/Plugins/Runtime/RigVM/Source/RigVM/Public/RigVMFunctions/RigVMDispatch_Switch.h
+++ b/Engine/Plugins/Runtime/RigVM/Source/RigVM/Public/RigVMFunctions/RigVMDispatch_Switch.h
@@ -22,8 +22,8 @@ public:
 	}
 
 	virtual FName GetArgumentNameForOperandIndex(int32 InOperandIndex, int32 InTotalOperands) const override;
-	virtual TArray<FRigVMTemplateArgument> GetArguments() const override;
-	virtual TArray<FRigVMExecuteArgument> GetExecuteArguments_Impl(const FRigVMDispatchContext& InContext) const override;
+	virtual const TArray<FRigVMTemplateArgument>& GetArguments() const override;
+	virtual const TArray<FRigVMExecuteArgument>& GetExecuteArguments_Impl(const FRigVMDispatchContext& InContext) const override;
 	virtual bool IsSingleton() const override { return true; } 
 
 #if WITH_EDITOR
diff --git a/Engine/Plugins/Runtime/RigVM/Source/RigVM/Public/RigVMFunctions/RigVMFunction_String.h b/Engine/Plugins/Runtime/RigVM/Source/RigVM/Public/RigVMFunctions/RigVMFunction_String.h
--- a/Engine/Plugins/Runtime/RigVM/Source/RigVM/Public/RigVMFunctions/RigVMFunction_String.h
+++ b/Engine/Plugins/Runtime/RigVM/Source/RigVM/Public/RigVMFunctions/RigVMFunction_String.h
@@ -510,7 +510,7 @@ struct FRigDispatch_ToString : public FRigVMDispatchFactory
 
 public:
 
-	virtual TArray<FRigVMTemplateArgument> GetArguments() const override;
+	virtual const TArray<FRigVMTemplateArgument>& GetArguments() const override;
 	virtual FRigVMTemplateTypeMap OnNewArgumentType(const FName& InArgumentName, TRigVMTypeIndex InTypeIndex) const override;
 
 protected:
@@ -529,7 +529,7 @@ struct FRigDispatch_FromString : public FRigVMDispatchFactory
 
 public:
 
-	virtual TArray<FRigVMTemplateArgument> GetArguments() const override;
+	virtual const TArray<FRigVMTemplateArgument>& GetArguments() const override;
 	virtual FRigVMTemplateTypeMap OnNewArgumentType(const FName& InArgumentName, TRigVMTypeIndex InTypeIndex) const override;
 
 protected:
