Fix race in URigVMPin::NameEquals when redirectors are used and ALT is active

#rnx

diff --git a/Engine/Plugins/Runtime/RigVM/Source/RigVMDeveloper/Private/RigVMModel/RigVMPin.cpp b/Engine/Plugins/Runtime/RigVM/Source/RigVMDeveloper/Private/RigVMModel/RigVMPin.cpp
--- a/Engine/Plugins/Runtime/RigVM/Source/RigVMDeveloper/Private/RigVMModel/RigVMPin.cpp
+++ b/Engine/Plugins/Runtime/RigVM/Source/RigVMDeveloper/Private/RigVMModel/RigVMPin.cpp
@@ -14,6 +14,7 @@
 #include "Misc/DefaultValueHelper.h"
 #include "Misc/PackageName.h"
 #include "Misc/OutputDevice.h"
+#include "Misc/ScopeRWLock.h"
 #include "Misc/StringBuilder.h"
 #include "Logging/LogScopedVerbosityOverride.h"
 #include "RigVMModel/Nodes/RigVMCollapseNode.h"
@@ -145,8 +146,10 @@ bool URigVMPin::NameEquals(const FString& InName, bool bFollowCoreRedirectors) c
 		{
 			typedef TPair<FName, FString> FRedirectPinPair;
 			const FRedirectPinPair Key(Struct->GetFName(), InName);
+			static FRWLock RedirectedPinNamesLock;
 			static TMap<FRedirectPinPair, FName> RedirectedPinNames;
 
+			FWriteScopeLock ScopeLock(RedirectedPinNamesLock);
 			if(const FName* RedirectedNamePtr = RedirectedPinNames.Find(Key))
 			{
 				if(RedirectedNamePtr->IsNone())
