RigVM Editor: Allow editing of quat pins

diff --git a/Engine/Plugins/Runtime/RigVM/Source/RigVMDeveloper/Private/EdGraph/RigVMEdGraphSchema.cpp b/Engine/Plugins/Runtime/RigVM/Source/RigVMDeveloper/Private/EdGraph/RigVMEdGraphSchema.cpp
--- a/Engine/Plugins/Runtime/RigVM/Source/RigVMDeveloper/Private/EdGraph/RigVMEdGraphSchema.cpp
+++ b/Engine/Plugins/Runtime/RigVM/Source/RigVMDeveloper/Private/EdGraph/RigVMEdGraphSchema.cpp
@@ -1324,6 +1324,10 @@ bool URigVMEdGraphSchema::MarkBlueprintDirtyFromNewNode(UBlueprint* InBlueprint,
 
 bool URigVMEdGraphSchema::IsStructEditable(UStruct* InStruct) const
 {
+	if (InStruct == TBaseStructure<FQuat>::Get())
+	{
+		return true;
+	}
 	if (InStruct == FRuntimeFloatCurve::StaticStruct())
 	{
 		return true;
diff --git a/Engine/Plugins/Runtime/RigVM/Source/RigVMDeveloper/Private/RigVMModel/RigVMController.cpp b/Engine/Plugins/Runtime/RigVM/Source/RigVMDeveloper/Private/RigVMModel/RigVMController.cpp
--- a/Engine/Plugins/Runtime/RigVM/Source/RigVMDeveloper/Private/RigVMModel/RigVMController.cpp
+++ b/Engine/Plugins/Runtime/RigVM/Source/RigVMDeveloper/Private/RigVMModel/RigVMController.cpp
@@ -7937,12 +7937,20 @@ bool URigVMController::SetPinDefaultValue(URigVMPin* InPin, const FString& InDef
 			{
 				Graph->MarkPackageDirty();
 			}
+			bSetPinDefaultValueSucceeded = true;
 		}
 	}
 
 	if (bSetupUndoRedo)
 	{
-		GetActionStack()->EndAction(Action, bMergeUndoAction);
+		if(bSetPinDefaultValueSucceeded)
+		{
+			GetActionStack()->EndAction(Action, bMergeUndoAction);
+		}
+		else
+		{
+			GetActionStack()->CancelAction(Action);
+		}
 	}
 
 	return true;
diff --git a/Engine/Plugins/Runtime/RigVM/Source/RigVMDeveloper/Private/RigVMModel/RigVMControllerActions.cpp b/Engine/Plugins/Runtime/RigVM/Source/RigVMDeveloper/Private/RigVMModel/RigVMControllerActions.cpp
--- a/Engine/Plugins/Runtime/RigVM/Source/RigVMDeveloper/Private/RigVMModel/RigVMControllerActions.cpp
+++ b/Engine/Plugins/Runtime/RigVM/Source/RigVMDeveloper/Private/RigVMModel/RigVMControllerActions.cpp
@@ -1075,7 +1075,7 @@ FRigVMSetPinDefaultValueAction::FRigVMSetPinDefaultValueAction()
 FRigVMSetPinDefaultValueAction::FRigVMSetPinDefaultValueAction(URigVMController* InController, URigVMPin* InPin, const FString& InNewDefaultValue)
 : FRigVMBaseAction(InController)
 , PinPath(InPin->GetPinPath())
-, OldDefaultValue(InPin->GetDefaultValue())
+, OldDefaultValue(InPin->GetDefaultValueStoredByUserInterface())
 , NewDefaultValue(InNewDefaultValue)
 {
 	/* Since for template we are chaning types - it is possible that the
diff --git a/Engine/Plugins/Runtime/RigVM/Source/RigVMDeveloper/Private/RigVMModel/RigVMPin.cpp b/Engine/Plugins/Runtime/RigVM/Source/RigVMDeveloper/Private/RigVMModel/RigVMPin.cpp
--- a/Engine/Plugins/Runtime/RigVM/Source/RigVMDeveloper/Private/RigVMModel/RigVMPin.cpp
+++ b/Engine/Plugins/Runtime/RigVM/Source/RigVMDeveloper/Private/RigVMModel/RigVMPin.cpp
@@ -763,10 +763,10 @@ bool URigVMPin::ShouldHideSubPins() const
 
 FString URigVMPin::GetDefaultValue() const
 {
-	return GetDefaultValue(EmptyPinOverride);
+	return GetDefaultValue(EmptyPinOverride, true);
 }
 
-FString URigVMPin::GetDefaultValue(const URigVMPin::FPinOverride& InOverride) const
+FString URigVMPin::GetDefaultValue(const URigVMPin::FPinOverride& InOverride, bool bAdaptValueForPinType) const
 {
 	if (FPinOverrideValue const* OverrideValuePtr = InOverride.Value.Find(InOverride.Key.GetSibling((URigVMPin*)this)))
 	{
@@ -791,7 +791,7 @@ FString URigVMPin::GetDefaultValue(const URigVMPin::FPinOverride& InOverride) co
 				{
 					return TEXT("()");
 				}
-				FString ElementDefaultValue = SubPin->GetDefaultValue(InOverride);
+				FString ElementDefaultValue = SubPin->GetDefaultValue(InOverride, bAdaptValueForPinType);
 				if (SubPin->IsStringType())
 				{
 					ElementDefaultValue = TEXT("\"") + ElementDefaultValue + TEXT("\"");
@@ -820,7 +820,7 @@ FString URigVMPin::GetDefaultValue(const URigVMPin::FPinOverride& InOverride) co
 			TArray<FString> MemberDefaultValues;
 			for (const URigVMPin* SubPin : SubPins)
 			{
-				FString MemberDefaultValue = SubPin->GetDefaultValue(InOverride);
+				FString MemberDefaultValue = SubPin->GetDefaultValue(InOverride, bAdaptValueForPinType);
 				if (SubPin->IsStringType() && !MemberDefaultValue.IsEmpty())
 				{
 					MemberDefaultValue = TEXT("\"") + MemberDefaultValue + TEXT("\"");
@@ -838,7 +838,29 @@ FString URigVMPin::GetDefaultValue(const URigVMPin::FPinOverride& InOverride) co
 			return FString::Printf(TEXT("(%s)"), *FString::Join(MemberDefaultValues, TEXT(",")));
 		}
 
-		return DefaultValue.IsEmpty() ? TEXT("()") : DefaultValue;
+		// special case certain pin types to adapt their values from
+		// alternative representations.
+		static const FString EmptyStructDefaultValue = TEXT("()");
+		if(bAdaptValueForPinType && !DefaultValue.IsEmpty() && DefaultValue != EmptyStructDefaultValue)
+		{
+			if(GetScriptStruct() == TBaseStructure<FQuat>::Get())
+			{
+				// quaternions also allow default values stored as rotators
+				FRigVMPinDefaultValueImportErrorContext ErrorPipe;
+				FRotator Rotator = FRotator::ZeroRotator;
+				LOG_SCOPE_VERBOSITY_OVERRIDE(LogExec, ELogVerbosity::Verbose); 
+				TBaseStructure<FRotator>::Get()->ImportText(*DefaultValue, &Rotator, nullptr, PPF_None, &ErrorPipe, TBaseStructure<FRotator>::Get()->GetName());
+				if(ErrorPipe.NumErrors == 0)
+				{
+					const FQuat Quat = FQuat::MakeFromRotator(Rotator);
+					FString AdaptedDefaultValue;	
+					TBaseStructure<FQuat>::Get()->ExportText(AdaptedDefaultValue, &Quat, &Quat, nullptr, PPF_None, nullptr);
+					return AdaptedDefaultValue;
+				}
+			}
+		}
+		
+		return DefaultValue.IsEmpty() ? EmptyStructDefaultValue : DefaultValue;
 	}
 	else if (IsArrayElement() && DefaultValue.IsEmpty())
 	{
@@ -856,6 +878,11 @@ FString URigVMPin::GetDefaultValue(const URigVMPin::FPinOverride& InOverride) co
 	return DefaultValue;
 }
 
+FString URigVMPin::GetDefaultValueStoredByUserInterface() const
+{
+	return GetDefaultValue(EmptyPinOverride, false);
+}
+
 template< typename Type>
 static FString ClampValue(const FString& InValueString, const FString& InMinValueString, const FString& InMaxValueString)
 {
@@ -941,11 +968,25 @@ bool URigVMPin::IsValidDefaultValue(const FString& InDefaultValue) const
 		} 
 		else if (UScriptStruct* ScriptStruct = Cast<UScriptStruct>(GetCPPTypeObject()))
 		{
-			FRigVMPinDefaultValueImportErrorContext ErrorPipe;
+			// special case alternative representations 
+			if(ScriptStruct == TBaseStructure<FQuat>::Get())
+			{
+				// quaternions also allow default values stored as rotators
+				FRigVMPinDefaultValueImportErrorContext ErrorPipe;
+				FRotator Rotator = FRotator::ZeroRotator;
+				LOG_SCOPE_VERBOSITY_OVERRIDE(LogExec, ELogVerbosity::Verbose); 
+				TBaseStructure<FRotator>::Get()->ImportText(*Value, &Rotator, nullptr, PPF_None, &ErrorPipe, TBaseStructure<FRotator>::Get()->GetName());
+				if(ErrorPipe.NumErrors == 0)
+				{
+					return true;
+				}
+			}
+			
 			TArray<uint8> TempStructBuffer;
 			TempStructBuffer.AddUninitialized(ScriptStruct->GetStructureSize());
 			ScriptStruct->InitializeDefaultValue(TempStructBuffer.GetData());
 
+			FRigVMPinDefaultValueImportErrorContext ErrorPipe;
 			{
 				// force logging to the error pipe for error detection
 				LOG_SCOPE_VERBOSITY_OVERRIDE(LogExec, ELogVerbosity::Verbose); 
diff --git a/Engine/Plugins/Runtime/RigVM/Source/RigVMDeveloper/Public/RigVMModel/RigVMPin.h b/Engine/Plugins/Runtime/RigVM/Source/RigVMDeveloper/Public/RigVMModel/RigVMPin.h
--- a/Engine/Plugins/Runtime/RigVM/Source/RigVMDeveloper/Public/RigVMModel/RigVMPin.h
+++ b/Engine/Plugins/Runtime/RigVM/Source/RigVMDeveloper/Public/RigVMModel/RigVMPin.h
@@ -294,7 +294,10 @@ public:
 	FString GetDefaultValue() const;
 
 	// Returns the default value with an additional override ma
-	FString GetDefaultValue(const FPinOverride& InOverride) const;
+	FString GetDefaultValue(const FPinOverride& InOverride, bool bAdaptValueForPinType = true) const;
+
+	// Returns the default value as stored by the user.
+	FString GetDefaultValueStoredByUserInterface() const;
 
 	// Returns true if the default value provided is valid
 	bool IsValidDefaultValue(const FString& InDefaultValue) const;
diff --git a/Engine/Plugins/Runtime/RigVM/Source/RigVMEditor/Private/EdGraph/RigVMEdGraphPanelPinFactory.cpp b/Engine/Plugins/Runtime/RigVM/Source/RigVMEditor/Private/EdGraph/RigVMEdGraphPanelPinFactory.cpp
--- a/Engine/Plugins/Runtime/RigVM/Source/RigVMEditor/Private/EdGraph/RigVMEdGraphPanelPinFactory.cpp
+++ b/Engine/Plugins/Runtime/RigVM/Source/RigVMEditor/Private/EdGraph/RigVMEdGraphPanelPinFactory.cpp
@@ -10,6 +10,7 @@
 #include "Widgets/SRigVMGraphPinVariableBinding.h"
 #include "Widgets/SRigVMGraphPinUserDataNameSpace.h"
 #include "Widgets/SRigVMGraphPinUserDataPath.h"
+#include "Widgets/SRigVMGraphPinQuat.h"
 #include "KismetPins/SGraphPinExec.h"
 #include "SGraphPinComboBox.h"
 #include "RigVMHost.h"
@@ -64,7 +65,8 @@ TSharedPtr<SGraphPin> FRigVMEdGraphPanelPinFactory::CreatePin_Internal(UEdGraphP
 		{
 			URigVMEdGraph* RigGraph = Cast<URigVMEdGraph>(RigNode->GetGraph());
 
-			if (URigVMPin* ModelPin = RigNode->GetModelPinFromPinPath(InPin->GetName()))
+			URigVMPin* ModelPin = RigNode->GetModelPinFromPinPath(InPin->GetName());
+			if (ModelPin)
 			{
 				if (ModelPin->IsBoundToVariable())
 				{
@@ -109,6 +111,10 @@ TSharedPtr<SGraphPin> FRigVMEdGraphPanelPinFactory::CreatePin_Internal(UEdGraphP
 				{
 					return SNew(SRigVMGraphPinCurveFloat, InPin);
 				}
+				if (ModelPin && (InPin->PinType.PinSubCategoryObject == TBaseStructure<FQuat>::Get()))
+				{
+					return SNew(SRigVMGraphPinQuat, InPin).ModelPin(ModelPin);
+				}
 			}
 		}
 	}
diff --git a/Engine/Plugins/Runtime/RigVM/Source/RigVMEditor/Private/Widgets/SRigVMGraphPinQuat.cpp b/Engine/Plugins/Runtime/RigVM/Source/RigVMEditor/Private/Widgets/SRigVMGraphPinQuat.cpp
new file mode 100644
--- /dev/null
+++ b/Engine/Plugins/Runtime/RigVM/Source/RigVMEditor/Private/Widgets/SRigVMGraphPinQuat.cpp
@@ -0,0 +1,135 @@
+// Copyright Epic Games, Inc. All Rights Reserved.
+
+
+#include "Widgets/SRigVMGraphPinQuat.h"
+#include "EdGraph/RigVMEdGraph.h"
+#include "EdGraph/RigVMEdGraphNode.h"
+#include "RigVMBlueprint.h"
+#include "RigVMModel/RigVMController.h"
+#include "Logging/LogScopedVerbosityOverride.h"
+
+void SRigVMGraphPinQuat::Construct(const FArguments& InArgs, UEdGraphPin* InGraphPinObj)
+{
+	ModelPin = InArgs._ModelPin;
+	SGraphPin::Construct(SGraphPin::FArguments(), InGraphPinObj);
+}
+
+TSharedRef<SWidget>	SRigVMGraphPinQuat::GetDefaultValueWidget()
+{
+	TSharedRef<SWidget> Widget = SNew(SBox)
+		.MinDesiredWidth(150)
+		.MaxDesiredWidth(400)
+		[
+			SNew(SRotatorInputBox)
+			.DisplayToggle(false)
+			.bColorAxisLabels(true)
+			.Pitch(this, &SRigVMGraphPinQuat::GetRotatorComponent, 0)
+			.Yaw(this, &SRigVMGraphPinQuat::GetRotatorComponent, 1)
+			.Roll(this, &SRigVMGraphPinQuat::GetRotatorComponent, 2)
+			.OnPitchChanged(this, &SRigVMGraphPinQuat::OnRotatorComponentChanged, 0)
+			.OnYawChanged(this, &SRigVMGraphPinQuat::OnRotatorComponentChanged, 1)
+			.OnRollChanged(this, &SRigVMGraphPinQuat::OnRotatorComponentChanged, 2)
+			.OnPitchCommitted(this, &SRigVMGraphPinQuat::OnRotatorComponentCommitted, 0, true)
+			.OnYawCommitted(this, &SRigVMGraphPinQuat::OnRotatorComponentCommitted, 1, true)
+			.OnRollCommitted(this, &SRigVMGraphPinQuat::OnRotatorComponentCommitted, 2, true)
+			.Visibility(this, &SGraphPin::GetDefaultValueVisibility)
+		];
+	return Widget;
+}
+
+TOptional<FRotator> SRigVMGraphPinQuat::GetRotator() const
+{
+	if(ModelPin)
+	{
+		const FString DefaultValue = ModelPin->GetDefaultValueStoredByUserInterface();
+		if(!DefaultValue.IsEmpty())
+		{
+			// try to import a quaternion
+			{
+				FRigVMPinDefaultValueImportErrorContext ErrorPipe;
+				FQuat Quat = FQuat::Identity;
+				LOG_SCOPE_VERBOSITY_OVERRIDE(LogExec, ELogVerbosity::Verbose); 
+				TBaseStructure<FQuat>::Get()->ImportText(*DefaultValue, &Quat, nullptr, PPF_None, &ErrorPipe, TBaseStructure<FQuat>::Get()->GetName(), true);
+				if(ErrorPipe.NumErrors == 0)
+				{
+					return Quat.Rotator();
+				}
+			}
+
+			// also try to import a rotator - in case that's how it got stored
+			{
+				FRigVMPinDefaultValueImportErrorContext ErrorPipe; 
+				FRotator Rotator = FRotator::ZeroRotator;
+				LOG_SCOPE_VERBOSITY_OVERRIDE(LogExec, ELogVerbosity::Verbose); 
+				TBaseStructure<FRotator>::Get()->ImportText(*DefaultValue, &Rotator, nullptr, PPF_None, &ErrorPipe, TBaseStructure<FQuat>::Get()->GetName(), true);
+				if(ErrorPipe.NumErrors == 0)
+				{
+					return Rotator;
+				}
+			}
+		}
+	}
+	return TOptional<FRotator>();
+}
+
+void SRigVMGraphPinQuat::OnRotatorCommitted(FRotator InRotator, ETextCommit::Type InCommitType, bool bUndoRedo)
+{
+	if(ModelPin)
+	{
+		FString NewDefaultValue;	
+		TBaseStructure<FRotator>::Get()->ExportText(NewDefaultValue, &InRotator, &InRotator, nullptr, PPF_None, nullptr);
+
+		if(URigVMBlueprint* Blueprint = ModelPin->GetTypedOuter<URigVMBlueprint>())
+		{
+			if(URigVMController* Controller = Blueprint->GetOrCreateController(ModelPin->GetGraph()))
+			{
+				Controller->SetPinDefaultValue(ModelPin->GetPinPath(), NewDefaultValue, true, bUndoRedo, false, bUndoRedo, true);
+			}
+		}
+	}
+}
+
+TOptional<float> SRigVMGraphPinQuat::GetRotatorComponent(int32 InComponent) const
+{
+	TOptional<FRotator> Rotator = GetRotator();
+	if(Rotator.IsSet())
+	{
+		if(InComponent == 0)
+		{
+			return Rotator.GetValue().Pitch;
+		}
+		if(InComponent == 1)
+		{
+			return Rotator.GetValue().Yaw;
+		}
+		return Rotator.GetValue().Roll;
+	}
+	return TOptional<float>();
+}
+
+void SRigVMGraphPinQuat::OnRotatorComponentChanged(float InValue, int32 InComponent)
+{
+	return OnRotatorComponentCommitted(InValue, ETextCommit::Default, InComponent, false);
+}
+
+void SRigVMGraphPinQuat::OnRotatorComponentCommitted(float InValue, ETextCommit::Type InCommitType, int32 InComponent, bool bUndoRedo)
+{
+	const TOptional<FRotator> OptionalRotator = GetRotator();
+	if(OptionalRotator.IsSet())
+	{
+		FRotator Rotator = OptionalRotator.GetValue();
+		if(InComponent == 0)
+		{
+			Rotator.Pitch = InValue;
+		}
+		else if(InComponent == 1)
+		{
+			Rotator.Yaw = InValue;
+		}
+		else
+		{
+			Rotator.Roll = InValue;
+		}
+		OnRotatorCommitted(Rotator, InCommitType, bUndoRedo);
+	}
+}
diff --git a/Engine/Plugins/Runtime/RigVM/Source/RigVMEditor/Public/Widgets/SRigVMGraphPinQuat.h b/Engine/Plugins/Runtime/RigVM/Source/RigVMEditor/Public/Widgets/SRigVMGraphPinQuat.h
new file mode 100644
--- /dev/null
+++ b/Engine/Plugins/Runtime/RigVM/Source/RigVMEditor/Public/Widgets/SRigVMGraphPinQuat.h
@@ -0,0 +1,37 @@
+// Copyright Epic Games, Inc. All Rights Reserved.
+
+#pragma once
+
+#include "CoreMinimal.h"
+#include "Widgets/DeclarativeSyntaxSupport.h"
+#include "Widgets/SWidget.h"
+#include "Widgets/Input/SRotatorInputBox.h"
+#include "RigVMModel/RigVMPin.h"
+#include "SGraphPin.h"
+
+class RIGVMEDITOR_API SRigVMGraphPinQuat : public SGraphPin
+{
+public:
+	SLATE_BEGIN_ARGS(SRigVMGraphPinQuat)
+		: _ModelPin(nullptr)
+	{}
+		SLATE_ARGUMENT(URigVMPin*, ModelPin)
+	SLATE_END_ARGS()
+
+
+	void Construct(const FArguments& InArgs, UEdGraphPin* InGraphPinObj);
+
+protected:
+
+	//~ Begin SGraphPin Interface
+	virtual TSharedRef<SWidget>	GetDefaultValueWidget() override;
+	//~ End SGraphPin Interface
+
+	TOptional<FRotator> GetRotator() const;
+	void OnRotatorCommitted(FRotator InRotator, ETextCommit::Type InCommitType, bool bUndoRedo);
+	TOptional<float> GetRotatorComponent(int32 InComponent) const;
+	void OnRotatorComponentChanged(float InValue, int32 InComponent);
+	void OnRotatorComponentCommitted(float InValue, ETextCommit::Type InCommitType, int32 InComponent, bool bUndoRedo);
+
+	URigVMPin* ModelPin;
+};
