Fix non unity issue in SRigVMGraphPinQuat.cpp

diff --git a/Engine/Plugins/Runtime/RigVM/Source/RigVMEditor/Private/Widgets/SRigVMGraphPinQuat.cpp b/Engine/Plugins/Runtime/RigVM/Source/RigVMEditor/Private/Widgets/SRigVMGraphPinQuat.cpp
--- a/Engine/Plugins/Runtime/RigVM/Source/RigVMEditor/Private/Widgets/SRigVMGraphPinQuat.cpp
+++ b/Engine/Plugins/Runtime/RigVM/Source/RigVMEditor/Private/Widgets/SRigVMGraphPinQuat.cpp
@@ -7,6 +7,7 @@
 #include "RigVMBlueprint.h"
 #include "RigVMModel/RigVMController.h"
 #include "Logging/LogScopedVerbosityOverride.h"
+#include "Widgets/Layout/SBox.h"
 
 void SRigVMGraphPinQuat::Construct(const FArguments& InArgs, UEdGraphPin* InGraphPinObj)
 {
