RigVM: Fix crash when double clicking a link

diff --git a/Engine/Plugins/Runtime/RigVM/Source/RigVMDeveloper/Private/RigVMModel/RigVMController.cpp b/Engine/Plugins/Runtime/RigVM/Source/RigVMDeveloper/Private/RigVMModel/RigVMController.cpp
--- a/Engine/Plugins/Runtime/RigVM/Source/RigVMDeveloper/Private/RigVMModel/RigVMController.cpp
+++ b/Engine/Plugins/Runtime/RigVM/Source/RigVMDeveloper/Private/RigVMModel/RigVMController.cpp
@@ -2932,7 +2932,7 @@ URigVMRerouteNode* URigVMController::AddRerouteNodeOnLink(URigVMLink* InLink, co
 		const FString NodeName = GetSchema()->GetSanitizedNodeName(Node->GetName());
 
 		RigVMPythonUtils::Print(GetSchema()->GetGraphOuterName(GetGraph()), 
-							FString::Printf(TEXT("blueprint.get_controller_by_name('%s').add_reroute_node_on_link_path('%s', %s, %s, '%s')"),
+							FString::Printf(TEXT("blueprint.get_controller_by_name('%s').add_reroute_node_on_link_path('%s', %s, '%s')"),
 											*GraphName,
 											*URigVMLink::GetPinPathRepresentation(SourcePin->GetPinPath(), TargetPin->GetPinPath()),
 											*RigVMPythonUtils::Vector2DToPythonString(Node->GetPosition()),
