RigVM: Allow editing of graphs and details also during PIE

diff --git a/Engine/Plugins/Runtime/RigVM/Source/RigVMEditor/Private/Editor/RigVMEditor.cpp b/Engine/Plugins/Runtime/RigVM/Source/RigVMEditor/Private/Editor/RigVMEditor.cpp
--- a/Engine/Plugins/Runtime/RigVM/Source/RigVMEditor/Private/Editor/RigVMEditor.cpp
+++ b/Engine/Plugins/Runtime/RigVM/Source/RigVMEditor/Private/Editor/RigVMEditor.cpp
@@ -398,6 +398,12 @@ TSubclassOf<UEdGraphSchema> FRigVMEditor::GetDefaultSchemaClass() const
 	return URigVMEdGraphSchema::StaticClass();
 }
 
+bool FRigVMEditor::InEditingMode() const
+{
+	// always allow editing - also during PIE.
+	return true;
+}
+
 void FRigVMEditor::Tick(float DeltaTime)
 {
 	FBlueprintEditor::Tick(DeltaTime);
diff --git a/Engine/Plugins/Runtime/RigVM/Source/RigVMEditor/Public/Editor/RigVMEditor.h b/Engine/Plugins/Runtime/RigVM/Source/RigVMEditor/Public/Editor/RigVMEditor.h
--- a/Engine/Plugins/Runtime/RigVM/Source/RigVMEditor/Public/Editor/RigVMEditor.h
+++ b/Engine/Plugins/Runtime/RigVM/Source/RigVMEditor/Public/Editor/RigVMEditor.h
@@ -57,6 +57,7 @@ public:
 	// FBlueprintEditor interface
 	virtual UBlueprint* GetBlueprintObj() const override;
 	virtual TSubclassOf<UEdGraphSchema> GetDefaultSchemaClass() const override;
+	virtual bool InEditingMode() const override;
 
 	//  FTickableEditorObject Interface
 	virtual void Tick(float DeltaTime) override;
diff --git a/Engine/Source/Editor/Kismet/Public/BlueprintEditor.h b/Engine/Source/Editor/Kismet/Public/BlueprintEditor.h
--- a/Engine/Source/Editor/Kismet/Public/BlueprintEditor.h
+++ b/Engine/Source/Editor/Kismet/Public/BlueprintEditor.h
@@ -429,7 +429,7 @@ public:
 	float GetInstructionTextOpacity(UEdGraph* InGraph) const;
 
 	/** Returns true if in editing mode */
-	bool InEditingMode() const;
+	virtual bool InEditingMode() const;
 
 	/** Returns true if able to compile */
 	virtual bool IsCompilingEnabled() const;
