RigVM: Fix crash when expanding node with injected nodes inside
#rnx

diff --git a/Engine/Plugins/Runtime/RigVM/Source/RigVMDeveloper/Private/RigVMModel/RigVMController.cpp b/Engine/Plugins/Runtime/RigVM/Source/RigVMDeveloper/Private/RigVMModel/RigVMController.cpp
--- a/Engine/Plugins/Runtime/RigVM/Source/RigVMDeveloper/Private/RigVMModel/RigVMController.cpp
+++ b/Engine/Plugins/Runtime/RigVM/Source/RigVMDeveloper/Private/RigVMModel/RigVMController.cpp
@@ -5271,6 +5271,7 @@ TArray<URigVMNode*> URigVMController::ExpandLibraryNode(URigVMLibraryNode* InNod
 	}
 
 	TArray<FName> NodeNames;
+	TArray<FName> InjectedNodeNames;
 	FBox2D Bounds = FBox2D(EForceInit::ForceInit);
 	{
 		TArray<URigVMNode*> FilteredNodes;
@@ -5292,6 +5293,7 @@ TArray<URigVMNode*> URigVMController::ExpandLibraryNode(URigVMLibraryNode* InNod
 
 			if(Node->IsInjected())
 			{
+				InjectedNodeNames.Add(Node->GetFName());
 				continue;
 			}
 			
@@ -5398,10 +5400,12 @@ TArray<URigVMNode*> URigVMController::ExpandLibraryNode(URigVMLibraryNode* InNod
 	check(ExpandedNodeNames.Num() >= NodeNames.Num());
 
 	TMap<FName, FName> NodeNameMap;
-	for (int32 NodeNameIndex = 0, ExpandedNodeNameIndex = 0; NodeNameIndex < NodeNames.Num(); ExpandedNodeNameIndex++)
+	for (int32 NodeNameIndex = 0, ExpandedNodeNameIndex = 0, InjectedNodeName = 0; ExpandedNodeNameIndex < ExpandedNodeNames.Num(); ExpandedNodeNameIndex++)
 	{
 		if (ExpandedNodes[ExpandedNodeNameIndex]->IsInjected())
 		{
+			NodeNameMap.Add(InjectedNodeNames[InjectedNodeName], ExpandedNodeNames[ExpandedNodeNameIndex]);
+			InjectedNodeName++;
 			continue;
 		}
 		NodeNameMap.Add(NodeNames[NodeNameIndex], ExpandedNodeNames[ExpandedNodeNameIndex]);
