Control Rig: Store references to user-defined enums on the RigVM to ensure they get properly loaded when the VM is loaded.

#rnx

diff --git a/Engine/Plugins/Runtime/RigVM/Source/RigVM/Private/RigVMCore/RigVM.cpp b/Engine/Plugins/Runtime/RigVM/Source/RigVM/Private/RigVMCore/RigVM.cpp
--- a/Engine/Plugins/Runtime/RigVM/Source/RigVM/Private/RigVMCore/RigVM.cpp
+++ b/Engine/Plugins/Runtime/RigVM/Source/RigVM/Private/RigVMCore/RigVM.cpp
@@ -161,6 +161,7 @@ void URigVM::Save(FArchive& Ar)
 
 		Ar << OperandToDebugRegisters;
 		Ar << UserDefinedStructGuidToPathName;
+		Ar << UserDefinedEnumToPathName;
 	}
 
 	// advertise dependencies on user defined structs and user defined enums
@@ -238,6 +239,14 @@ void URigVM::Load(FArchive& Ar)
 			{
 				UserDefinedStructGuidToPathName.Reset();
 			}
+			if (Ar.CustomVer(FRigVMObjectVersion::GUID) >= FRigVMObjectVersion::VMStoringUserDefinedEnumMap)
+			{
+				Ar << UserDefinedEnumToPathName;
+			}
+			else
+			{
+				UserDefinedEnumToPathName.Reset();
+			}
 		}
 
 		// we only deal with virtual machines now that use the new memory infrastructure.
@@ -282,6 +291,14 @@ void URigVM::Load(FArchive& Ar)
 		{
 			UserDefinedStructGuidToPathName.Reset();
 		}
+		if (Ar.CustomVer(FRigVMObjectVersion::GUID) >= FRigVMObjectVersion::VMStoringUserDefinedEnumMap)
+		{
+			Ar << UserDefinedEnumToPathName;
+		}
+		else
+		{
+			UserDefinedEnumToPathName.Reset();
+		}
 	}
 
 	// ensure to load the required functions
@@ -349,6 +366,7 @@ void URigVM::PreSave(FObjectPreSaveContext SaveContext)
 
 	const TArray<const UObject*> UserDefinedDependencies = GetUserDefinedDependencies();
 	UserDefinedStructGuidToPathName.Reset();
+	UserDefinedEnumToPathName.Reset();
 
 	for(const UObject* UserDefinedDependency : UserDefinedDependencies)
 	{
@@ -357,6 +375,11 @@ void URigVM::PreSave(FObjectPreSaveContext SaveContext)
 			const FString GuidBasedName = RigVMTypeUtils::GetUniqueStructTypeName(UserDefinedStruct);
 			UserDefinedStructGuidToPathName.Add(GuidBasedName, UserDefinedStruct);
 		}
+		else if (const UUserDefinedEnum* UserDefinedEnum = Cast<UUserDefinedEnum>(UserDefinedDependency))
+		{
+			const FString EnumName = RigVMTypeUtils::CPPTypeFromEnum(UserDefinedEnum);
+			UserDefinedEnumToPathName.Add(EnumName, UserDefinedEnum);
+		}
 	}
 }
 
@@ -583,6 +606,7 @@ void URigVM::Reset(bool IsIgnoringArchetypeRef)
 		ParametersNameMap.Reset();
 		OperandToDebugRegisters.Reset();
 		UserDefinedStructGuidToPathName.Reset();
+		UserDefinedEnumToPathName.Reset();
 	}
 
 	if(!IsIgnoringArchetypeRef)
@@ -612,6 +636,7 @@ void URigVM::Empty(FRigVMExtendedExecuteContext& Context)
 	Parameters.Empty();
 	ParametersNameMap.Empty();
 	UserDefinedStructGuidToPathName.Empty();
+	UserDefinedEnumToPathName.Reset();
 	ExternalVariables.Empty();
 
 	InvalidateCachedMemory();
@@ -720,6 +745,7 @@ void URigVM::CopyFrom(URigVM* InVM, bool bDeferCopy, bool bReferenceLiteralMemor
 	
 	OperandToDebugRegisters = InVM->OperandToDebugRegisters;
 	UserDefinedStructGuidToPathName = InVM->UserDefinedStructGuidToPathName;
+	UserDefinedEnumToPathName = InVM->UserDefinedEnumToPathName;
 
 	if (bCopyExternalVariables)
 	{
@@ -1036,6 +1062,7 @@ bool URigVM::ResolveFunctionsIfRequired()
 
 		FRigVMTypeResolvalInfo ResolvalInfo;
 		ResolvalInfo.CPPTypeToObjectPath = UserDefinedStructGuidToPathName;
+		ResolvalInfo.CPPTypeToObjectPath.Append(UserDefinedEnumToPathName);
 
 		TArray<FName>& FunctionNames = GetFunctionNames();
 		for (int32 FunctionIndex = 0; FunctionIndex < FunctionNames.Num(); FunctionIndex++)
diff --git a/Engine/Plugins/Runtime/RigVM/Source/RigVM/Public/RigVMCore/RigVM.h b/Engine/Plugins/Runtime/RigVM/Source/RigVM/Public/RigVMCore/RigVM.h
--- a/Engine/Plugins/Runtime/RigVM/Source/RigVM/Public/RigVMCore/RigVM.h
+++ b/Engine/Plugins/Runtime/RigVM/Source/RigVM/Public/RigVMCore/RigVM.h
@@ -868,6 +868,7 @@ private:
 	TArray<FRigVMExternalVariableDef> ExternalVariables;
 	TArray<FRigVMLazyBranch> LazyBranches;
 	TMap<FString, FSoftObjectPath> UserDefinedStructGuidToPathName;
+	TMap<FString, FSoftObjectPath> UserDefinedEnumToPathName;
 
 	// this function should be kept in sync with FRigVMOperand::GetContainerIndex()
 	static int32 GetContainerIndex(ERigVMMemoryType InType)
diff --git a/Engine/Plugins/Runtime/RigVM/Source/RigVM/Public/RigVMObjectVersion.h b/Engine/Plugins/Runtime/RigVM/Source/RigVM/Public/RigVMObjectVersion.h
--- a/Engine/Plugins/Runtime/RigVM/Source/RigVM/Public/RigVMObjectVersion.h
+++ b/Engine/Plugins/Runtime/RigVM/Source/RigVM/Public/RigVMObjectVersion.h
@@ -20,6 +20,9 @@ struct RIGVM_API FRigVMObjectVersion
 		// Storing paths to user defined structs map
 		VMStoringUserDefinedStructMap,
 
+		// Storing paths to user defined enums map
+		VMStoringUserDefinedEnumMap,
+		
 		// -----<new versions can be added above this line>-------------------------------------------------
 		VersionPlusOne,
 		LatestVersion = VersionPlusOne - 1
diff --git a/Engine/Plugins/Runtime/RigVM/Source/RigVM/Public/RigVMTypeUtils.h b/Engine/Plugins/Runtime/RigVM/Source/RigVM/Public/RigVMTypeUtils.h
--- a/Engine/Plugins/Runtime/RigVM/Source/RigVM/Public/RigVMTypeUtils.h
+++ b/Engine/Plugins/Runtime/RigVM/Source/RigVM/Public/RigVMTypeUtils.h
@@ -6,6 +6,7 @@
 #include "RigVMCore/RigVMTypeIndex.h"
 #include "RigVMCore/RigVMUnknownType.h"
 #include "UObject/Interface.h"
+#include "Engine/UserDefinedEnum.h"
 #include "Engine/UserDefinedStruct.h"
 #include "UObject/CoreRedirects.h"
 #include "UObject/Package.h"
@@ -117,12 +118,10 @@ namespace RigVMTypeUtils
 
 	inline FString CPPTypeFromEnum(const UEnum* InEnum)
 	{
-		check(InEnum);
-
 		FString CPPType = InEnum->CppType;
 		if(CPPType.IsEmpty()) // this might be a user defined enum
 		{
-			CPPType = InEnum->GetName();
+			CPPType = FString::Printf(TEXT("EUserDefinedEnum_%s_%08x"), *InEnum->GetName(), GetTypeHash(InEnum->GetPathName()));
 		}
 		return CPPType;
 	}
