Optimize editor startup by only doing the RigVM class lookup if the asset might actually have data

diff --git a/Engine/Plugins/Runtime/RigVM/Source/RigVMDeveloper/Private/RigVMModel/RigVMBuildData.cpp b/Engine/Plugins/Runtime/RigVM/Source/RigVMDeveloper/Private/RigVMModel/RigVMBuildData.cpp
--- a/Engine/Plugins/Runtime/RigVM/Source/RigVMDeveloper/Private/RigVMModel/RigVMBuildData.cpp
+++ b/Engine/Plugins/Runtime/RigVM/Source/RigVMDeveloper/Private/RigVMModel/RigVMBuildData.cpp
@@ -9,6 +9,8 @@
 
 #include UE_INLINE_GENERATED_CPP_BY_NAME(RigVMBuildData)
 
+static const FName FunctionReferenceNodeDataName = TEXT("FunctionReferenceNodeData");
+
 // When the object system has been completely loaded, collect all the references between RigVM graphs 
 static FDelayedAutoRegisterHelper GRigVMBuildDataSingletonHelper(EDelayedRegisterRunPhase::EndOfEngineInit, []() -> void
 {
@@ -95,7 +97,7 @@ void URigVMBuildData::InitializeIfNeeded()
 	for (UClass* Class : ImplementedClasses)
 	{
 		const FArrayProperty* ReferenceNodeDataProperty =
-		   CastField<FArrayProperty>(Class->FindPropertyByName(TEXT("FunctionReferenceNodeData")));
+		   CastField<FArrayProperty>(Class->FindPropertyByName(FunctionReferenceNodeDataName));
 		if(ReferenceNodeDataProperty)
 		{
 			const FAssetRegistryModule& AssetRegistryModule = FModuleManager::LoadModuleChecked<FAssetRegistryModule>("AssetRegistry");
@@ -124,35 +126,40 @@ void URigVMBuildData::RegisterReferencesFromAsset(const FAssetData& InAssetData)
 {
 	URigVMBuildData* BuildData = URigVMBuildData::Get();
 
-	if (UClass* Class = InAssetData.GetClass())
+	// It's faster to check for a key directly then trying to get the class
+	FAssetDataTagMapSharedView::FFindTagResult FoundValue = InAssetData.TagsAndValues.FindTag(FunctionReferenceNodeDataName);
+
+	if (FoundValue.IsSet())
 	{
-		const FArrayProperty* ReferenceNodeDataProperty =
-			  CastField<FArrayProperty>(Class->FindPropertyByName(TEXT("FunctionReferenceNodeData")));
-		if(ReferenceNodeDataProperty)
+		if (UClass* Class = InAssetData.GetClass())
 		{
-			const FString ReferenceNodeDataString =
-						   InAssetData.GetTagValueRef<FString>(ReferenceNodeDataProperty->GetFName());
-			if(ReferenceNodeDataString.IsEmpty())
-			{
-				return;
-			}
-
-			// See if it has reference node data, and register the references
-			TArray<FRigVMReferenceNodeData> ReferenceNodeDatas;
-			ReferenceNodeDataProperty->ImportText_Direct(*ReferenceNodeDataString, &ReferenceNodeDatas, nullptr, EPropertyPortFlags::PPF_None);	
-			for(FRigVMReferenceNodeData& ReferenceNodeData : ReferenceNodeDatas)
+			const FArrayProperty* ReferenceNodeDataProperty =
+				  CastField<FArrayProperty>(Class->FindPropertyByName(FunctionReferenceNodeDataName));
+			if(ReferenceNodeDataProperty)
 			{
-				if (ReferenceNodeData.ReferencedFunctionIdentifier.LibraryNode.IsValid())
+				const FString ReferenceNodeDataString = FoundValue.AsString();
+				if(ReferenceNodeDataString.IsEmpty())
 				{
-					BuildData->RegisterFunctionReference(ReferenceNodeData.ReferencedFunctionIdentifier, ReferenceNodeData.GetReferenceNodeObjectPath());
+					return;
 				}
-				else if (ReferenceNodeData.ReferencedHeader_DEPRECATED.IsValid())
-				{
-					BuildData->RegisterFunctionReference(ReferenceNodeData.ReferencedHeader_DEPRECATED.LibraryPointer, ReferenceNodeData.GetReferenceNodeObjectPath());
-				}
-				else if (!ReferenceNodeData.ReferencedFunctionPath_DEPRECATED.IsEmpty())
+
+				// See if it has reference node data, and register the references
+				TArray<FRigVMReferenceNodeData> ReferenceNodeDatas;
+				ReferenceNodeDataProperty->ImportText_Direct(*ReferenceNodeDataString, &ReferenceNodeDatas, nullptr, EPropertyPortFlags::PPF_None);	
+				for(FRigVMReferenceNodeData& ReferenceNodeData : ReferenceNodeDatas)
 				{
-					BuildData->RegisterFunctionReference(ReferenceNodeData);
+					if (ReferenceNodeData.ReferencedFunctionIdentifier.LibraryNode.IsValid())
+					{
+						BuildData->RegisterFunctionReference(ReferenceNodeData.ReferencedFunctionIdentifier, ReferenceNodeData.GetReferenceNodeObjectPath());
+					}
+					else if (ReferenceNodeData.ReferencedHeader_DEPRECATED.IsValid())
+					{
+						BuildData->RegisterFunctionReference(ReferenceNodeData.ReferencedHeader_DEPRECATED.LibraryPointer, ReferenceNodeData.GetReferenceNodeObjectPath());
+					}
+					else if (!ReferenceNodeData.ReferencedFunctionPath_DEPRECATED.IsEmpty())
+					{
+						BuildData->RegisterFunctionReference(ReferenceNodeData);
+					}
 				}
 			}
 		}
