RigVM: CoreNoEquals dispatch corrups GetArguments results for CoreEquals

diff --git a/Engine/Plugins/Runtime/RigVM/Source/RigVM/Private/RigVMFunctions/RigVMDispatch_Core.cpp b/Engine/Plugins/Runtime/RigVM/Source/RigVM/Private/RigVMFunctions/RigVMDispatch_Core.cpp
--- a/Engine/Plugins/Runtime/RigVM/Source/RigVM/Private/RigVMFunctions/RigVMDispatch_Core.cpp
+++ b/Engine/Plugins/Runtime/RigVM/Source/RigVM/Private/RigVMFunctions/RigVMDispatch_Core.cpp
@@ -126,3 +126,18 @@ void FRigVMDispatch_CoreEquals::Execute(FRigVMExtendedExecuteContext& InContext,
 	Result = AdaptResult(Result, InContext);
 }
 
+// duplicate the code here so that the FRigVMDispatch_CoreNotEquals has it's own static variables
+// to store the types are registration time.
+const TArray<FRigVMTemplateArgument>& FRigVMDispatch_CoreNotEquals::GetArguments() const
+{
+	static const TArray<FRigVMTemplateArgument::ETypeCategory> ValueCategories = {
+		FRigVMTemplateArgument::ETypeCategory_SingleAnyValue,
+		FRigVMTemplateArgument::ETypeCategory_ArrayAnyValue
+	};
+	static const TArray<FRigVMTemplateArgument> Arguments = {
+		FRigVMTemplateArgument(AName, ERigVMPinDirection::Input, ValueCategories),
+		FRigVMTemplateArgument(BName, ERigVMPinDirection::Input, ValueCategories),
+		FRigVMTemplateArgument(ResultName, ERigVMPinDirection::Output, RigVMTypeUtils::TypeIndex::Bool)
+	};
+	return Arguments;
+}
diff --git a/Engine/Plugins/Runtime/RigVM/Source/RigVM/Public/RigVMFunctions/RigVMDispatch_Core.h b/Engine/Plugins/Runtime/RigVM/Source/RigVM/Public/RigVMFunctions/RigVMDispatch_Core.h
--- a/Engine/Plugins/Runtime/RigVM/Source/RigVM/Public/RigVMFunctions/RigVMDispatch_Core.h
+++ b/Engine/Plugins/Runtime/RigVM/Source/RigVM/Public/RigVMFunctions/RigVMDispatch_Core.h
@@ -100,6 +100,8 @@ struct RIGVM_API FRigVMDispatch_CoreNotEquals : public FRigVMDispatch_CoreEquals
 		FactoryScriptStruct = StaticStruct();
 	}
 
+	virtual const TArray<FRigVMTemplateArgument>& GetArguments() const override;
+
 	// we are inheriting everything from the equals dispatch,
 	// and due to the check of the factory within FRigVMDispatch_CoreEquals::Execute we can
 	// rely on that completely. we only need this class for the displayname and
