Control Rig: Fix UI slowness
#rnx

diff --git a/Engine/Plugins/Runtime/RigVM/Source/RigVMDeveloper/Private/EdGraph/RigVMEdGraphNode.cpp b/Engine/Plugins/Runtime/RigVM/Source/RigVMDeveloper/Private/EdGraph/RigVMEdGraphNode.cpp
--- a/Engine/Plugins/Runtime/RigVM/Source/RigVMDeveloper/Private/EdGraph/RigVMEdGraphNode.cpp
+++ b/Engine/Plugins/Runtime/RigVM/Source/RigVMDeveloper/Private/EdGraph/RigVMEdGraphNode.cpp
@@ -213,8 +213,8 @@ void URigVMEdGraphNode::ReconstructNode_Internal(bool bForce)
 #endif
 
 	// Clear previously set messages
-	ErrorMsg.Reset();
-
+	ClearErrorInfo();
+	
 	// Move the existing pins to a saved array.
 	// This way we can reuse them later
 	LastEdGraphPins = Pins;
@@ -473,7 +473,42 @@ void URigVMEdGraphNode::ClearErrorInfo()
 	// SRigVMGraphNode only updates if the error types do not match so we have
 	// clear the error type as well, see SRigVMGraphNode::RefreshErrorInfo()
 	ErrorType = (int32)EMessageSeverity::Info + 1;
-	ErrorMsg = FString();	
+	ErrorMsg = FString();
+	ErrorMessageHashes.Reset();
+}
+
+void URigVMEdGraphNode::AddErrorInfo(const EMessageSeverity::Type& InSeverity, const FString& InMessage)
+{
+	if (ErrorType < InSeverity)
+	{
+		return;
+	}
+
+	const uint32 MessageHash = GetTypeHash(InMessage);
+	if (ErrorType == InSeverity)
+	{
+		if (ErrorMessageHashes.Contains(MessageHash))
+		{
+			return;
+		}
+		ErrorMessageHashes.Add(MessageHash);
+		ErrorMsg = FString::Printf(TEXT("%s\n%s"), *ErrorMsg, *InMessage);
+	}
+	else
+	{
+		ErrorMessageHashes.Reset();
+		ErrorMessageHashes.Add(MessageHash);
+		ErrorMsg = InMessage;
+		ErrorType = InSeverity;
+	}
+}
+
+void URigVMEdGraphNode::SetErrorInfo(const EMessageSeverity::Type& InSeverity, const FString& InMessage)
+{
+	ErrorMessageHashes.Reset();
+	ErrorMessageHashes.Add(GetTypeHash(InMessage));
+	ErrorMsg = InMessage;
+	ErrorType = InSeverity;
 }
 
 URigVMPin* URigVMEdGraphNode::FindModelPinFromGraphPin(const UEdGraphPin* InGraphPin) const
diff --git a/Engine/Plugins/Runtime/RigVM/Source/RigVMDeveloper/Private/RigVMBlueprint.cpp b/Engine/Plugins/Runtime/RigVM/Source/RigVMDeveloper/Private/RigVMBlueprint.cpp
--- a/Engine/Plugins/Runtime/RigVM/Source/RigVMDeveloper/Private/RigVMBlueprint.cpp
+++ b/Engine/Plugins/Runtime/RigVM/Source/RigVMDeveloper/Private/RigVMBlueprint.cpp
@@ -1481,8 +1481,7 @@ void URigVMBlueprint::HandleReportFromCompiler(EMessageSeverity::Type InSeverity
 
 	if (URigVMEdGraphNode* EdGraphNode = Cast<URigVMEdGraphNode>(SubjectForMessage))
 	{
-		EdGraphNode->ErrorType = (int32)InSeverity;
-		EdGraphNode->ErrorMsg = InMessage;
+		EdGraphNode->SetErrorInfo(InSeverity, InMessage);
 		EdGraphNode->bHasCompilerMessage = EdGraphNode->ErrorType <= int32(EMessageSeverity::Info);
 	}
 }
diff --git a/Engine/Plugins/Runtime/RigVM/Source/RigVMDeveloper/Public/EdGraph/RigVMEdGraphNode.h b/Engine/Plugins/Runtime/RigVM/Source/RigVMDeveloper/Public/EdGraph/RigVMEdGraphNode.h
--- a/Engine/Plugins/Runtime/RigVM/Source/RigVMDeveloper/Public/EdGraph/RigVMEdGraphNode.h
+++ b/Engine/Plugins/Runtime/RigVM/Source/RigVMDeveloper/Public/EdGraph/RigVMEdGraphNode.h
@@ -173,6 +173,10 @@ public:
 
 	void ClearErrorInfo();
 
+	void AddErrorInfo(const EMessageSeverity::Type& InSeverity, const FString& InMessage);
+
+	void SetErrorInfo(const EMessageSeverity::Type& InSeverity, const FString& InMessage);
+
 	URigVMPin* FindModelPinFromGraphPin(const UEdGraphPin* InGraphPin) const;
 	UEdGraphPin* FindGraphPinFromModelPin(const URigVMPin* InModelPin, bool bAsInput) const;
 
@@ -270,6 +274,8 @@ private:
 	FNodePinExpansionChanged NodePinExpansionChanged;
 	FNodeBeginRemoval NodeBeginRemoval;
 
+	TSet<uint32> ErrorMessageHashes;
+
 	mutable const FRigVMTemplate* CachedTemplate;
 	mutable TOptional<bool> DrawAsCompactNodeCache;
 
diff --git a/Engine/Plugins/Runtime/RigVM/Source/RigVMEditor/Private/Editor/RigVMEditor.cpp b/Engine/Plugins/Runtime/RigVM/Source/RigVMEditor/Private/Editor/RigVMEditor.cpp
--- a/Engine/Plugins/Runtime/RigVM/Source/RigVMEditor/Private/Editor/RigVMEditor.cpp
+++ b/Engine/Plugins/Runtime/RigVM/Source/RigVMEditor/Private/Editor/RigVMEditor.cpp
@@ -3261,19 +3261,9 @@ void FRigVMEditor::UpdateGraphCompilerErrors()
 				bFoundError = bFoundError || Entry.Severity <= EMessageSeverity::Error;
 				bFoundWarning = bFoundWarning || Entry.Severity <= EMessageSeverity::Warning;
 
-				int32 ErrorType = (int32)Entry.Severity;
-				if (GraphNode->ErrorType < ErrorType)
-				{
-					continue;
-				}
-				else if (GraphNode->ErrorType == ErrorType)
-				{
-					GraphNode->ErrorMsg = FString::Printf(TEXT("%s\n%s"), *GraphNode->ErrorMsg, *Entry.Message);
-				}
-				else
+				if (URigVMEdGraphNode* RigVMEdGraphNode = Cast<URigVMEdGraphNode>(GraphNode))
 				{
-					GraphNode->ErrorMsg = Entry.Message;
-					GraphNode->ErrorType = ErrorType;
+					RigVMEdGraphNode->AddErrorInfo(Entry.Severity, Entry.Message);
 				}
 			}
 
diff --git a/Engine/Plugins/Runtime/RigVM/Source/RigVMEditor/Private/Widgets/SRigVMGraphNode.cpp b/Engine/Plugins/Runtime/RigVM/Source/RigVMEditor/Private/Widgets/SRigVMGraphNode.cpp
--- a/Engine/Plugins/Runtime/RigVM/Source/RigVMEditor/Private/Widgets/SRigVMGraphNode.cpp
+++ b/Engine/Plugins/Runtime/RigVM/Source/RigVMEditor/Private/Widgets/SRigVMGraphNode.cpp
@@ -1053,9 +1053,12 @@ void SRigVMGraphNode::RefreshErrorInfo()
 							if(Pin->IsArray() && !Pin->IsFixedSizeArray())
 							{
 								GraphNode->bHasCompilerMessage = true;
-								GraphNode->ErrorType = int32(EMessageSeverity::Info);
-								static const FString ArrayWarning = TEXT("This node creates a copy of the array.\nThis may cause side effects.");
-								GraphNode->ErrorMsg = ArrayWarning;
+								if (URigVMEdGraphNode* RigVMEdGraphNode = Cast<URigVMEdGraphNode>(GraphNode))
+								{
+									static const FString ArrayWarning = TEXT("This node creates a copy of the array.\nThis may cause side effects.");
+									RigVMEdGraphNode->SetErrorInfo(EMessageSeverity::Info, ArrayWarning);
+								}
+								
 								break;
 							}
 						}
