RigVM: Improve tooltip message for resolved pins

diff --git a/Engine/Plugins/Runtime/RigVM/Source/RigVMDeveloper/Private/RigVMModel/RigVMPin.cpp b/Engine/Plugins/Runtime/RigVM/Source/RigVMDeveloper/Private/RigVMModel/RigVMPin.cpp
--- a/Engine/Plugins/Runtime/RigVM/Source/RigVMDeveloper/Private/RigVMModel/RigVMPin.cpp
+++ b/Engine/Plugins/Runtime/RigVM/Source/RigVMDeveloper/Private/RigVMModel/RigVMPin.cpp
@@ -1985,6 +1985,43 @@ bool URigVMPin::CanLink(const URigVMPin* InSourcePin, const URigVMPin* InTargetP
 			if (OutFailureReason)
 			{
 				*OutFailureReason = TEXT("Source and target pin types are not compatible.");
+
+				const URigVMPin* TemplatePinToCheck = nullptr;
+				switch(InUserLinkDirection)
+				{
+					case ERigVMPinDirection::Input:
+					{
+						TemplatePinToCheck = InSourcePin;
+						break;
+					}
+					case ERigVMPinDirection::Output:
+					{
+						TemplatePinToCheck = InTargetPin;
+						break;
+					}
+					default:
+					{
+						break;
+					}
+				}
+
+				if(TemplatePinToCheck)
+				{
+					if(const URigVMTemplateNode* TemplateNode = Cast<URigVMTemplateNode>(TemplatePinToCheck->GetNode()))
+					{
+						if(const FRigVMTemplate* Template = TemplateNode->GetTemplate())
+						{
+							if(const FRigVMTemplateArgument* Argument = Template->FindArgument(TemplatePinToCheck->GetFName()))
+							{
+								const URigVMPin* OtherPin = TemplatePinToCheck == InSourcePin ? InTargetPin : InSourcePin;
+								if(Argument->SupportsTypeIndex(OtherPin->GetTypeIndex()))
+								{
+									*OutFailureReason = TEXT("Link supported - please unresolve template node.");
+								}
+							}
+						}
+					}
+				}
 			}
 			return false;
 		}
