RigVM: Compiler error to deal with disconnected execute pins

diff --git a/Engine/Plugins/Experimental/AnimNext/Source/AnimNextUncookedOnly/Private/UncookedOnlyUtils.cpp b/Engine/Plugins/Experimental/AnimNext/Source/AnimNextUncookedOnly/Private/UncookedOnlyUtils.cpp
--- a/Engine/Plugins/Experimental/AnimNext/Source/AnimNextUncookedOnly/Private/UncookedOnlyUtils.cpp
+++ b/Engine/Plugins/Experimental/AnimNext/Source/AnimNextUncookedOnly/Private/UncookedOnlyUtils.cpp
@@ -334,14 +334,14 @@ void FUtils::Compile(UAnimNextGraph* InGraph)
 
 	URigVMCompiler* Compiler = URigVMCompiler::StaticClass()->GetDefaultObject<URigVMCompiler>();
 	EditorData->VMCompileSettings.SetExecuteContextStruct(EditorData->RigVMClient.GetExecuteContextStruct());
-	Compiler->Settings = (EditorData->bCompileInDebugMode) ? FRigVMCompileSettings::Fast(EditorData->VMCompileSettings.GetExecuteContextStruct()) : EditorData->VMCompileSettings;
-	Compiler->Compile(GraphModelsToCompile, TempController, InGraph->RigVM, InGraph->ExtendedExecuteContext, InGraph->GetRigVMExternalVariables(), &EditorData->PinToOperandMap);
+	const FRigVMCompileSettings Settings = (EditorData->bCompileInDebugMode) ? FRigVMCompileSettings::Fast(EditorData->VMCompileSettings.GetExecuteContextStruct()) : EditorData->VMCompileSettings;
+	Compiler->Compile(Settings, GraphModelsToCompile, TempController, InGraph->RigVM, InGraph->ExtendedExecuteContext, InGraph->GetRigVMExternalVariables(), &EditorData->PinToOperandMap);
 
 	if (EditorData->bErrorsDuringCompilation)
 	{
-		if(Compiler->Settings.SurpressErrors)
+		if(Settings.SurpressErrors)
 		{
-			Compiler->Settings.Reportf(EMessageSeverity::Info, InGraph,TEXT("Compilation Errors may be suppressed for AnimNext Interface Graph: %s. See VM Compile Settings for more Details"), *InGraph->GetName());
+			Settings.Reportf(EMessageSeverity::Info, InGraph,TEXT("Compilation Errors may be suppressed for AnimNext Interface Graph: %s. See VM Compile Settings for more Details"), *InGraph->GetName());
 		}
 	}
 
@@ -512,15 +512,15 @@ void FUtils::CompileVM(UAnimNextParameterBlock* InParameterBlock)
 
 	URigVMCompiler* Compiler = URigVMCompiler::StaticClass()->GetDefaultObject<URigVMCompiler>();
 	EditorData->VMCompileSettings.SetExecuteContextStruct(EditorData->RigVMClient.GetExecuteContextStruct());
-	Compiler->Settings = (EditorData->bCompileInDebugMode) ? FRigVMCompileSettings::Fast(EditorData->VMCompileSettings.GetExecuteContextStruct()) : EditorData->VMCompileSettings;
+	const FRigVMCompileSettings Settings = (EditorData->bCompileInDebugMode) ? FRigVMCompileSettings::Fast(EditorData->VMCompileSettings.GetExecuteContextStruct()) : EditorData->VMCompileSettings;
 	URigVMController* RootController = EditorData->GetRigVMClient()->GetOrCreateController(EditorData->GetRigVMClient()->GetDefaultModel());
-	Compiler->Compile(EditorData->GetRigVMClient()->GetAllModels(false, false), RootController, InParameterBlock->RigVM, InParameterBlock->ExtendedExecuteContext, InParameterBlock->GetRigVMExternalVariables(), &EditorData->PinToOperandMap);
+	Compiler->Compile(Settings, EditorData->GetRigVMClient()->GetAllModels(false, false), RootController, InParameterBlock->RigVM, InParameterBlock->ExtendedExecuteContext, InParameterBlock->GetRigVMExternalVariables(), &EditorData->PinToOperandMap);
 
 	if (EditorData->bErrorsDuringCompilation)
 	{
-		if(Compiler->Settings.SurpressErrors)
+		if(Settings.SurpressErrors)
 		{
-			Compiler->Settings.Reportf(EMessageSeverity::Info, InParameterBlock, TEXT("Compilation Errors may be suppressed for AnimNext Interface Graph: %s. See VM Compile Settings for more Details"), *InParameterBlock->GetName());
+			Settings.Reportf(EMessageSeverity::Info, InParameterBlock, TEXT("Compilation Errors may be suppressed for AnimNext Interface Graph: %s. See VM Compile Settings for more Details"), *InParameterBlock->GetName());
 		}
 	}
 
diff --git a/Engine/Plugins/Runtime/RigVM/Source/RigVM/Public/RigVMCore/RigVMGraphFunctionDefinition.h b/Engine/Plugins/Runtime/RigVM/Source/RigVM/Public/RigVMCore/RigVMGraphFunctionDefinition.h
--- a/Engine/Plugins/Runtime/RigVM/Source/RigVM/Public/RigVMCore/RigVMGraphFunctionDefinition.h
+++ b/Engine/Plugins/Runtime/RigVM/Source/RigVM/Public/RigVMCore/RigVMGraphFunctionDefinition.h
@@ -88,7 +88,10 @@ struct RIGVM_API FRigVMFunctionCompilationData
 {
 	GENERATED_BODY()
 
-	FRigVMFunctionCompilationData() : Hash(0) {}
+	FRigVMFunctionCompilationData()
+	: Hash(0)
+	, bEncounteredSurpressedErrors(false)
+	{}
 
 	UPROPERTY()
 	FRigVMByteCode ByteCode;
@@ -129,6 +132,9 @@ struct RIGVM_API FRigVMFunctionCompilationData
 	UPROPERTY()
 	uint32 Hash;
 
+	UPROPERTY(Transient)
+	bool bEncounteredSurpressedErrors;
+
 	TMap<FRigVMOperand, TArray<FRigVMOperand>> OperandToDebugRegisters;
 
 	bool IsValid() const
@@ -136,6 +142,11 @@ struct RIGVM_API FRigVMFunctionCompilationData
 		return Hash != 0;
 	}
 
+	bool RequiresRecompilation() const
+	{
+		return bEncounteredSurpressedErrors;
+	}
+
 	friend uint32 GetTypeHash(const FRigVMFunctionCompilationData& Data) 
 	{
 		uint32 DataHash = Data.ByteCode.GetByteCodeHash();
@@ -214,6 +225,11 @@ struct RIGVM_API FRigVMFunctionCompilationData
 		Ar << Data.Operands;
 		Ar << Data.Hash;
 
+		if(Ar.IsLoading())
+		{
+			Data.bEncounteredSurpressedErrors = false;
+		}
+
 		if (Ar.CustomVer(FUE5ReleaseStreamObjectVersion::GUID) < FUE5ReleaseStreamObjectVersion::RigVMSaveDebugMapInGraphFunctionData &&
 		    Ar.CustomVer(FFortniteMainBranchObjectVersion::GUID) < FFortniteMainBranchObjectVersion::RigVMSaveDebugMapInGraphFunctionData)
 		{
diff --git a/Engine/Plugins/Runtime/RigVM/Source/RigVMDeveloper/Private/RigVMBlueprint.cpp b/Engine/Plugins/Runtime/RigVM/Source/RigVMDeveloper/Private/RigVMBlueprint.cpp
--- a/Engine/Plugins/Runtime/RigVM/Source/RigVMDeveloper/Private/RigVMBlueprint.cpp
+++ b/Engine/Plugins/Runtime/RigVM/Source/RigVMDeveloper/Private/RigVMBlueprint.cpp
@@ -1073,7 +1073,10 @@ void URigVMBlueprint::HandlePackageDone()
 	}
 	
 	RemoveDeprecatedVMMemoryClass();
-	RecompileVM();
+	{
+		const FRigVMCompileSettingsDuringLoadGuard Guard(VMCompileSettings);
+		RecompileVM();
+	}
 	RequestRigVMInit();
 	BroadcastRigVMPackageDone();
 }
@@ -1172,14 +1175,15 @@ void URigVMBlueprint::RecompileVM()
 
 		URigVMCompiler* Compiler = URigVMCompiler::StaticClass()->GetDefaultObject<URigVMCompiler>();
 		VMCompileSettings.SetExecuteContextStruct(RigVMClient.GetExecuteContextStruct());
-		Compiler->Settings = (bCompileInDebugMode) ? FRigVMCompileSettings::Fast(VMCompileSettings.GetExecuteContextStruct()) : VMCompileSettings;
-		Compiler->Compile(RigVMClient.GetAllModels(false, false), GetOrCreateController(), CDO->VM, CDO->GetExtendedExecuteContext(), CDO->GetExternalVariablesImpl(false), &PinToOperandMap);
+
+		const FRigVMCompileSettings Settings = (bCompileInDebugMode) ? FRigVMCompileSettings::Fast(VMCompileSettings.GetExecuteContextStruct()) : VMCompileSettings;
+		Compiler->Compile(Settings, RigVMClient.GetAllModels(false, false), GetOrCreateController(), CDO->VM, CDO->GetExtendedExecuteContext(), CDO->GetExternalVariablesImpl(false), &PinToOperandMap);
 
 		if (bErrorsDuringCompilation)
 		{
-			if(Compiler->Settings.SurpressErrors)
+			if(Settings.SurpressErrors)
 			{
-				Compiler->Settings.Reportf(EMessageSeverity::Info, this,
+				Settings.Reportf(EMessageSeverity::Info, this,
 					TEXT("Compilation Errors may be suppressed for ControlRigBlueprint: %s. See VM Compile Setting in Class Settings for more Details"), *this->GetName());
 			}
 			bVMRecompilationRequired = false;
@@ -2436,7 +2440,7 @@ void URigVMBlueprint::PostTransacted(const FTransactionObjectEvent& TransactionE
  				return UberGraph == nullptr || !IsValid(UberGraph);
 			});
 			RigVMClient.PostTransacted(TransactionEvent);
-
+			
 			RecompileVM();
 			(void)MarkPackageDirty();			
 		}
@@ -3012,7 +3016,6 @@ void URigVMBlueprint::HandleModifiedEvent(ERigVMGraphNotifType InNotifType, URig
 				{
 					URigVMPin* Pin = CastChecked<URigVMPin>(InSubject)->GetRootPin(); 
 					URigVMCompiler* Compiler = URigVMCompiler::StaticClass()->GetDefaultObject<URigVMCompiler>();
-					Compiler->Settings = VMCompileSettings;
 
 					TSharedPtr<FRigVMParserAST> RuntimeAST = GetDefaultModel()->GetRuntimeAST();
 					
@@ -3033,13 +3036,13 @@ void URigVMBlueprint::HandleModifiedEvent(ERigVMGraphNotifType InNotifType, URig
 							}
 							else
 							{
-								Compiler->MarkDebugWatch(true, Pin, DebuggedHost->GetVM(), &PinToOperandMap, RuntimeAST);
+								Compiler->MarkDebugWatch(VMCompileSettings, true, Pin, DebuggedHost->GetVM(), &PinToOperandMap, RuntimeAST);
 							}
 						}
 					}
 					else
 					{
-						Compiler->MarkDebugWatch(false, Pin, DebuggedHost->GetVM(), &PinToOperandMap, RuntimeAST);
+						Compiler->MarkDebugWatch(VMCompileSettings, false, Pin, DebuggedHost->GetVM(), &PinToOperandMap, RuntimeAST);
 					}
 				}
 				// break; fall through
diff --git a/Engine/Plugins/Runtime/RigVM/Source/RigVMDeveloper/Private/RigVMCompiler/RigVMAST.cpp b/Engine/Plugins/Runtime/RigVM/Source/RigVMDeveloper/Private/RigVMCompiler/RigVMAST.cpp
--- a/Engine/Plugins/Runtime/RigVM/Source/RigVMDeveloper/Private/RigVMCompiler/RigVMAST.cpp
+++ b/Engine/Plugins/Runtime/RigVM/Source/RigVMDeveloper/Private/RigVMCompiler/RigVMAST.cpp
@@ -869,6 +869,76 @@ FRigVMExprAST* FRigVMParserAST::TraverseNode(const FRigVMASTProxy& InNodeProxy,
 		return NodeExpr;
 	}
 
+	// if we hit a mutable node here that hasn't been traversed yet,
+	// it means that the node is not wired up correctly.
+	if(Node->IsMutable())
+	{
+		struct Local
+		{
+			static bool IsNodeWiredToEvent(const FRigVMASTProxy& InNodeProxy)
+			{
+				const URigVMNode* Node = InNodeProxy.GetSubjectChecked<URigVMNode>();
+				if(Node->IsEvent())
+				{
+					return true;
+				}
+
+				for(const URigVMPin* Pin : Node->GetPins())
+				{
+					if(Pin->GetDirection() != ERigVMPinDirection::Input &&
+						Pin->GetDirection() != ERigVMPinDirection::IO)
+							
+					{
+						continue;
+					}
+					if(!Pin->IsExecuteContext())
+					{
+						continue;
+					}
+					
+					const TArray<URigVMPin*> SourcePins = Pin->GetLinkedSourcePins();
+					for(const URigVMPin* SourcePin : SourcePins)
+					{
+						const FRigVMASTProxy SourceNodeProxy = InNodeProxy.GetSibling(SourcePin->GetNode());
+						if(IsNodeWiredToEvent(SourceNodeProxy))
+						{
+							return true;
+						}
+					}
+				}
+
+				// if we hit an entry node - we need to continue the search a level up.
+				// events cannot be placed inside of a function / collapse node so
+				// there's no need to dive into library nodes.
+				if(Node->IsA<URigVMFunctionEntryNode>())
+				{
+					return IsNodeWiredToEvent(InNodeProxy.GetParent());
+				}
+				return false;
+			}
+		};
+
+		if(!Local::IsNodeWiredToEvent(InNodeProxy))
+		{
+			bool bIsInObsoleteBlock = false;
+			if(InParentExpr)
+			{
+				if(InParentExpr->GetBlock() == GetObsoleteBlock())
+				{
+					bIsInObsoleteBlock = true;
+				}
+			}
+
+			if(!bIsInObsoleteBlock)
+			{
+				Settings.Report(
+					EMessageSeverity::Error,
+					Node,
+					TEXT("Node @@ is not linked to execution."));
+			}
+		}
+	}
+
 	FRigVMExprAST* NodeExpr = CreateExpressionForNode(InNodeProxy, InParentExpr);
 	if (NodeExpr)
 	{
diff --git a/Engine/Plugins/Runtime/RigVM/Source/RigVMDeveloper/Private/RigVMCompiler/RigVMCompiler.cpp b/Engine/Plugins/Runtime/RigVM/Source/RigVMDeveloper/Private/RigVMCompiler/RigVMCompiler.cpp
--- a/Engine/Plugins/Runtime/RigVM/Source/RigVMDeveloper/Private/RigVMCompiler/RigVMCompiler.cpp
+++ b/Engine/Plugins/Runtime/RigVM/Source/RigVMDeveloper/Private/RigVMCompiler/RigVMCompiler.cpp
@@ -19,18 +19,24 @@
 #include "Algo/Count.h"
 #include "String/Join.h"
 
+#if WITH_EDITOR
+#include "Widgets/Notifications/SNotificationList.h"
+#include "Framework/Notifications/NotificationManager.h"
+#include "Styling/AppStyle.h"
+#endif
+
 #include UE_INLINE_GENERATED_CPP_BY_NAME(RigVMCompiler)
 
 class FRigVMCompilerImportErrorContext : public FOutputDevice
 {
 public:
 
-	URigVMCompiler* Compiler;
+	FRigVMCompilerWorkData* WorkData;
 	int32 NumErrors;
 
-	FRigVMCompilerImportErrorContext(URigVMCompiler* InCompiler)
+	FRigVMCompilerImportErrorContext(FRigVMCompilerWorkData* InWorkData)
 		: FOutputDevice()
-		, Compiler(InCompiler)
+		, WorkData(InWorkData)
 		, NumErrors(0)
 	{
 	}
@@ -42,17 +48,17 @@ public:
 		case ELogVerbosity::Error:
 		case ELogVerbosity::Fatal:
 		{
-			Compiler->ReportError(V);
+			WorkData->ReportError(V);
 			break;
 		}
 		case ELogVerbosity::Warning:
 		{
-			Compiler->ReportWarning(V);
+			WorkData->ReportWarning(V);
 			break;
 		}
 		default:
 		{
-			Compiler->ReportInfo(V);
+			WorkData->ReportInfo(V);
 			break;
 		}
 		}
@@ -68,6 +74,7 @@ FRigVMCompileSettings::FRigVMCompileSettings()
 	, IsPreprocessorPhase(false)
 	, ASTSettings(FRigVMParserASTSettings::Optimized())
 	, SetupNodeInstructionIndex(true)
+	, ASTErrorsAsNotifications(false)
 {
 }
 
@@ -81,6 +88,25 @@ FRigVMCompileSettings::FRigVMCompileSettings(UScriptStruct* InExecuteContextScri
 	}
 }
 
+void FRigVMCompileSettings::ReportInfo(const FString& InMessage) const
+{
+	if (SurpressInfoMessages)
+	{
+		return;
+	}
+	Report(EMessageSeverity::Info, nullptr, InMessage);
+}
+
+void FRigVMCompileSettings::ReportWarning(const FString& InMessage) const
+{
+	Report(EMessageSeverity::Warning, nullptr, InMessage);
+}
+
+void FRigVMCompileSettings::ReportError(const FString& InMessage) const
+{
+	Report(EMessageSeverity::Error, nullptr, InMessage);
+}
+
 FRigVMOperand FRigVMCompilerWorkData::AddProperty(
 	ERigVMMemoryType InMemoryType,
 	const FName& InName,
@@ -254,6 +280,75 @@ TRigVMTypeIndex FRigVMCompilerWorkData::GetTypeIndexForOperand(const FRigVMOpera
 	return FRigVMRegistry::Get().GetTypeIndex(CPPTypeName, CPPTypeObject);
 }
 
+void FRigVMCompilerWorkData::ReportInfo(const FString& InMessage) const
+{
+	Settings.ReportInfo(InMessage);
+}
+
+void FRigVMCompilerWorkData::ReportWarning(const FString& InMessage) const
+{
+	Settings.ReportWarning(InMessage);
+}
+
+void FRigVMCompilerWorkData::ReportError(const FString& InMessage) const
+{
+	Settings.ReportError(InMessage);
+}
+
+void FRigVMCompilerWorkData::OverrideReportDelegate(bool& bEncounteredASTError, bool& bSurpressedASTError)
+{
+	check(!OriginalReportDelegate.IsBound());
+	OriginalReportDelegate = Settings.ASTSettings.ReportDelegate;
+	
+	Settings.ASTSettings.ReportDelegate =
+		FRigVMReportDelegate::CreateLambda([this, &bEncounteredASTError, &bSurpressedASTError]
+			(EMessageSeverity::Type InSeverity, UObject* InSubject, const FString& InMessage)
+			{
+				FString Message = InMessage;
+				if(Settings.ASTErrorsAsNotifications &&
+					(InSeverity == EMessageSeverity::Error || InSeverity == EMessageSeverity::Warning))
+				{
+					const bool bIsError = InSeverity == EMessageSeverity::Error;
+					static constexpr TCHAR Warning[] = TEXT("Warning");
+					static constexpr TCHAR Error[] = TEXT("Error");
+					const TCHAR* SeverityLabel = bIsError ? Error : Warning;
+					static constexpr TCHAR Format[] = TEXT("%s: The %s '%s' has been surpressed to allow the content to load. Please fix the content since it may become a requirement in future versions.");
+					Message = FString::Printf(Format, *Graphs[0]->GetOutermost()->GetPathName(), SeverityLabel, *Message);
+#if WITH_EDITOR
+					if(InSubject)
+					{
+						Message.ReplaceInline(TEXT("@@"), *InSubject->GetName());
+					}
+						
+					FNotificationInfo Info(FText::FromString(Message));
+					Info.Image = bIsError ?
+						FAppStyle::GetBrush("Icons.ErrorWithColor") :
+						FAppStyle::GetBrush("Icons.WarningWithColor");
+					Info.bFireAndForget = true;
+					Info.FadeOutDuration = 1.0f;
+					Info.ExpireDuration = 7.0f;
+					Info.WidthOverride = 640.f;
+
+					(void)FSlateNotificationManager::Get().AddNotification(Info);
+#endif
+					InSeverity = EMessageSeverity::Info;
+					bSurpressedASTError = true;
+				}
+				(void)OriginalReportDelegate.ExecuteIfBound(InSeverity, InSubject, Message);
+				if(InSeverity == EMessageSeverity::Error)
+				{
+					bEncounteredASTError = true;
+				}
+			}
+		);
+}
+
+void FRigVMCompilerWorkData::RemoveOverrideReportDelegate()
+{
+	Settings.ASTSettings.ReportDelegate = OriginalReportDelegate;
+	OriginalReportDelegate = FRigVMReportDelegate();
+}
+
 URigVMCompiler::URigVMCompiler()
 	: CurrentCompilationFunction(nullptr)
 {
@@ -261,24 +356,39 @@ URigVMCompiler::URigVMCompiler()
 
 bool URigVMCompiler::Compile(TArray<URigVMGraph*> InGraphs, URigVMController* InController, URigVM* OutVM, FRigVMExtendedExecuteContext& OutVMContext, const TArray<FRigVMExternalVariable>& InExternalVariables, TMap<FString, FRigVMOperand>* OutOperands, TSharedPtr<FRigVMParserAST> InAST, FRigVMFunctionCompilationData* OutFunctionCompilationData)
 {
+	return Compile(Settings_DEPRECATED, InGraphs, InController, OutVM, OutVMContext,
+		InExternalVariables, OutOperands, InAST, OutFunctionCompilationData);
+}
+
+bool URigVMCompiler::Compile(const FRigVMCompileSettings& InSettings, TArray<URigVMGraph*> InGraphs,
+	URigVMController* InController, URigVM* OutVM, FRigVMExtendedExecuteContext& OutVMContext,
+	const TArray<FRigVMExternalVariable>& InExternalVariables, TMap<FString, FRigVMOperand>* OutOperands,
+	TSharedPtr<FRigVMParserAST> InAST, FRigVMFunctionCompilationData* OutFunctionCompilationData)
+{
+	FRigVMCompilerWorkData WorkData;
+	WorkData.Settings = InSettings;
+	WorkData.Graphs = InGraphs;
+
+	FRigVMCompileSettings& Settings = WorkData.Settings;
+	
 	double CompilationTime = 0;
 	FDurationTimer CompileTimer(CompilationTime);
 	
 	if (InGraphs.IsEmpty() || InGraphs.Contains(nullptr))
 	{
-		ReportError(TEXT("Provided graph is nullptr."));
+		WorkData.ReportError(TEXT("Provided graph is nullptr."));
 		return false;
 	}
 	
 	if (OutVM == nullptr)
 	{
-		ReportError(TEXT("Provided vm is nullptr."));
+		WorkData.ReportError(TEXT("Provided vm is nullptr."));
 		return false;
 	}
 
 	if (Settings.GetExecuteContextStruct() == nullptr)
 	{
-		ReportError(TEXT("Compiler settings don't provide the ExecuteContext to use. Cannot compile."));
+		WorkData.ReportError(TEXT("Compiler settings don't provide the ExecuteContext to use. Cannot compile."));
 		return false;;
 	}
 
@@ -298,7 +408,7 @@ bool URigVMCompiler::Compile(TArray<URigVMGraph*> InGraphs, URigVMController* In
 		{
 			if(!ValidExecuteContextStructs.Contains(Graph->GetExecuteContextStruct()))
 			{
-				ReportErrorf(
+				WorkData.ReportErrorf(
 					TEXT("Compiler settings' ExecuteContext (%s) is not compatible with '%s' graph's ExecuteContext (%s). Cannot compile."),
 					*Settings.GetExecuteContextStruct()->GetStructCPPName(),
 					*Graph->GetNodePath(),
@@ -313,14 +423,14 @@ bool URigVMCompiler::Compile(TArray<URigVMGraph*> InGraphs, URigVMController* In
 	{
 		if(InGraphs[0]->GetOuter() != InGraphs[Index]->GetOuter())
 		{
-			ReportError(TEXT("Provided graphs don't share a common outer / package."));
+			WorkData.ReportError(TEXT("Provided graphs don't share a common outer / package."));
 			return false;
 		}
 	}
 
 	if(OutVM->GetClass()->IsChildOf(URigVMNativized::StaticClass()))
 	{
-		ReportError(TEXT("Provided vm is nativized."));
+		WorkData.ReportError(TEXT("Provided vm is nativized."));
 		return false;
 	}
 
@@ -384,13 +494,13 @@ bool URigVMCompiler::Compile(TArray<URigVMGraph*> InGraphs, URigVMController* In
 						if (const FRigVMFunctionCompilationData* CompilationData = &FunctionData->CompilationData)
 						{
 							bool bSuccessfullCompilation = false;
-							if (!CompilationData->IsValid())
+							if (!CompilationData->IsValid() || CompilationData->RequiresRecompilation())
 							{
 								if (URigVMLibraryNode* LibraryNode = Cast<URigVMLibraryNode>(FunctionData->Header.LibraryPointer.LibraryNode.TryLoad()))
 								{
 									IRigVMClientHost* ClientHost = LibraryNode->GetImplementingOuter<IRigVMClientHost>();
 									URigVMController* FunctionController = ClientHost->GetRigVMClient()->GetController(LibraryNode->GetLibrary());
-									bSuccessfullCompilation = CompileFunction(LibraryNode, FunctionController, &FunctionData->CompilationData, OutVMContext);
+									bSuccessfullCompilation = CompileFunction(WorkData.Settings, LibraryNode, FunctionController, &FunctionData->CompilationData, OutVMContext);
 								}
 								else
 								{
@@ -699,7 +809,7 @@ bool URigVMCompiler::Compile(TArray<URigVMGraph*> InGraphs, URigVMController* In
 			auto ReportIncompatibleExecuteContextString = [&] (const FString InExecuteContextName)
 			{
 				static constexpr TCHAR Format[] = TEXT("ExecuteContext '%s' on node '%s' is not compatible with '%s' provided by the compiler settings."); 
-				ReportErrorf(
+				WorkData.ReportErrorf(
 					Format,
 					*InExecuteContextName,
 					*ModelNode->GetNodePath(),
@@ -914,12 +1024,15 @@ bool URigVMCompiler::Compile(TArray<URigVMGraph*> InGraphs, URigVMController* In
 		OutOperands->Add(Hash, Operand);
 	}
 
-	FRigVMCompilerWorkData WorkData;
-
+	bool bEncounteredASTError = false;
+	bool bSurpressedASTError = false;
 	WorkData.AST = InAST;
 	if (!WorkData.AST.IsValid())
 	{
+		WorkData.OverrideReportDelegate(bEncounteredASTError, bSurpressedASTError);
 		WorkData.AST = MakeShareable(new FRigVMParserAST(InGraphs, InController, Settings.ASTSettings, InExternalVariables));
+		WorkData.RemoveOverrideReportDelegate();
+
 		for(URigVMGraph* Graph : InGraphs)
 		{
 			Graph->RuntimeAST = WorkData.AST;
@@ -930,6 +1043,11 @@ bool URigVMCompiler::Compile(TArray<URigVMGraph*> InGraphs, URigVMController* In
 	}
 	ensure(WorkData.AST.IsValid());
 
+	if(bEncounteredASTError)
+	{
+		return false;
+	}
+
 	WorkData.VM = OutVM;
 	WorkData.ExecuteContextStruct = Settings.GetExecuteContextStruct();
 	WorkData.PinPathToOperand = OutOperands;
@@ -963,7 +1081,7 @@ bool URigVMCompiler::Compile(TArray<URigVMGraph*> InGraphs, URigVMController* In
 			{
 				if(EntryNode == nullptr)
 				{
-					ReportErrorf(TEXT("Corrupt library node '%s' - Missing entry node."), *CurrentCompilationFunction->GetPathName());
+					WorkData.ReportErrorf(TEXT("Corrupt library node '%s' - Missing entry node."), *CurrentCompilationFunction->GetPathName());
 					return false;
 				}
 				InterfacePin = EntryNode->FindPin(Pin->GetName());
@@ -972,7 +1090,7 @@ bool URigVMCompiler::Compile(TArray<URigVMGraph*> InGraphs, URigVMController* In
 			{
 				if(ReturnNode == nullptr)
 				{
-					ReportErrorf(TEXT("Corrupt library node '%s' - Missing return node."), *CurrentCompilationFunction->GetPathName());
+					WorkData.ReportErrorf(TEXT("Corrupt library node '%s' - Missing return node."), *CurrentCompilationFunction->GetPathName());
 					return false;
 				}
 				InterfacePin = ReturnNode->FindPin(Pin->GetName());
@@ -980,7 +1098,7 @@ bool URigVMCompiler::Compile(TArray<URigVMGraph*> InGraphs, URigVMController* In
 
 			if(InterfacePin == nullptr)
 			{
-				ReportErrorf(TEXT("Corrupt library node '%s' - Pin '%s' is not part of the entry / return node."), *CurrentCompilationFunction->GetPathName(), *Pin->GetPathName());
+				WorkData.ReportErrorf(TEXT("Corrupt library node '%s' - Pin '%s' is not part of the entry / return node."), *CurrentCompilationFunction->GetPathName(), *Pin->GetPathName());
 				return false;
 			}
 			
@@ -1085,7 +1203,7 @@ bool URigVMCompiler::Compile(TArray<URigVMGraph*> InGraphs, URigVMController* In
 	{
 		for(URigVMPin* WatchedPin : WorkData.WatchedPins)
 		{
-			MarkDebugWatch(true, WatchedPin, WorkData.VM, WorkData.PinPathToOperand, WorkData.AST);
+			MarkDebugWatch(WorkData.Settings, true, WatchedPin, WorkData.VM, WorkData.PinPathToOperand, WorkData.AST);
 		}
 	}
 
@@ -1122,6 +1240,7 @@ bool URigVMCompiler::Compile(TArray<URigVMGraph*> InGraphs, URigVMController* In
 		OutFunctionCompilationData->ByteCode = WorkData.VM->ByteCodeStorage;
 		OutFunctionCompilationData->FunctionNames = WorkData.VM->FunctionNamesStorage;
 		OutFunctionCompilationData->Operands = *OutOperands;
+		OutFunctionCompilationData->bEncounteredSurpressedErrors = bSurpressedASTError;
 
 		for (uint8 MemoryTypeIndex=0; MemoryTypeIndex<(uint8)ERigVMMemoryType::Invalid; ++MemoryTypeIndex)
 		{
@@ -1228,7 +1347,7 @@ bool URigVMCompiler::Compile(TArray<URigVMGraph*> InGraphs, URigVMController* In
 	if (!CurrentCompilationFunction)
 	{
 		CompileTimer.Stop();
-		ReportInfof(TEXT("Total Compilation time %f\n"), CompilationTime*1000);
+		WorkData.ReportInfof(TEXT("Total Compilation time %f\n"), CompilationTime*1000);
 	}
 
 	WorkData.VM->SetVMHash(WorkData.VM->ComputeVMHash());
@@ -1238,6 +1357,11 @@ bool URigVMCompiler::Compile(TArray<URigVMGraph*> InGraphs, URigVMController* In
 }
 
 bool URigVMCompiler::CompileFunction(const URigVMLibraryNode* InLibraryNode, URigVMController* InController, FRigVMFunctionCompilationData* OutFunctionCompilationData, FRigVMExtendedExecuteContext& OutVMContext)
+{
+	return CompileFunction(Settings_DEPRECATED, InLibraryNode, InController, OutFunctionCompilationData, OutVMContext);
+}
+
+bool URigVMCompiler::CompileFunction(const FRigVMCompileSettings& InSettings, const URigVMLibraryNode* InLibraryNode, URigVMController* InController, FRigVMFunctionCompilationData* OutFunctionCompilationData, FRigVMExtendedExecuteContext& OutVMContext)
 {
 	TGuardValue<const URigVMLibraryNode*> CompilationGuard(CurrentCompilationFunction, InLibraryNode);
 
@@ -1266,13 +1390,13 @@ bool URigVMCompiler::CompileFunction(const URigVMLibraryNode* InLibraryNode, URi
 	TMap<FString, FRigVMOperand> Operands;
 
 	URigVM* TempVM = NewObject<URigVM>(InLibraryNode->GetContainedGraph());
-	const bool bSuccess = Compile({InLibraryNode->GetContainedGraph()}, LibraryController, TempVM, OutVMContext, ExternalVariables, &Operands, nullptr, OutFunctionCompilationData);
+	const bool bSuccess = Compile(InSettings, {InLibraryNode->GetContainedGraph()}, LibraryController, TempVM, OutVMContext, ExternalVariables, &Operands, nullptr, OutFunctionCompilationData);
 	TempVM->ClearMemory();
 	TempVM->Rename(nullptr, GetTransientPackage(), REN_ForceNoResetLoaders | REN_DoNotDirty | REN_DontCreateRedirectors | REN_NonTransactional);
 	TempVM->MarkAsGarbage();
 
 	CompileTimer.Stop();
-	ReportInfof(TEXT("Compiled Function %s in %fms"), *InLibraryNode->GetName(), CompilationTime*1000);
+	InSettings.ReportInfof(TEXT("Compiled Function %s in %fms"), *InLibraryNode->GetName(), CompilationTime*1000);
 
 	// Update the compilation data of this library, and the hashes of the compilation data of its dependencies used for this compilation
 	if (IRigVMClientHost* ClientHost = InLibraryNode->GetImplementingOuter<IRigVMClientHost>())
@@ -1537,7 +1661,7 @@ void URigVMCompiler::TraverseEntry(const FRigVMEntryExprAST* InExpr, FRigVMCompi
 {
 	if (URigVMUnitNode* UnitNode = Cast<URigVMUnitNode>(InExpr->GetNode()))
 	{
-		if(!ValidateNode(UnitNode))
+		if(!ValidateNode(WorkData.Settings, UnitNode))
 		{
 			return;
 		}
@@ -1581,7 +1705,7 @@ void URigVMCompiler::TraverseEntry(const FRigVMEntryExprAST* InExpr, FRigVMCompi
 				WorkData.VM->GetByteCode().Entries.Add(Entry);
 			}
 
-			if (Settings.SetupNodeInstructionIndex)
+			if (WorkData.Settings.SetupNodeInstructionIndex)
 			{
 				const FRigVMCallstack Callstack = InExpr->GetProxy().GetCallstack();
 				WorkData.VM->GetByteCode().SetSubject(EntryInstructionIndex, Callstack.GetCallPath(), Callstack.GetStack());
@@ -1630,7 +1754,7 @@ int32 URigVMCompiler::TraverseCallExtern(const FRigVMCallExternExprAST* InExpr,
 	URigVMUnitNode* UnitNode = Cast<URigVMUnitNode>(Node);
 	URigVMDispatchNode* DispatchNode = Cast<URigVMDispatchNode>(Node);
 	const FRigVMRegistry& Registry = FRigVMRegistry::Get();
-	if(!ValidateNode(UnitNode, false) && !ValidateNode(DispatchNode, false))
+	if(!ValidateNode(WorkData.Settings, UnitNode, false) && !ValidateNode(WorkData.Settings, DispatchNode, false))
 	{
 		return INDEX_NONE;
 	}
@@ -1644,7 +1768,7 @@ int32 URigVMCompiler::TraverseCallExtern(const FRigVMCallExternExprAST* InExpr,
 			if(!SpecializedExecuteStruct->IsChildOf(ExecuteStruct))
 			{
 				static constexpr TCHAR UnknownExecuteContextMessage[] = TEXT("Node @@ uses an unexpected execute type '%s'. This graph uses '%s'.");
-				Settings.Report(EMessageSeverity::Error, Subject, FString::Printf(
+				WorkData.Settings.Report(EMessageSeverity::Error, Subject, FString::Printf(
 					UnknownExecuteContextMessage, *ExecuteStruct->GetStructCPPName(), *SpecializedExecuteStruct->GetStructCPPName()));
 				return false;
 			}
@@ -1658,7 +1782,7 @@ int32 URigVMCompiler::TraverseCallExtern(const FRigVMCallExternExprAST* InExpr,
 		if(ScriptStruct == nullptr)
 		{
 			static const FString UnresolvedMessage = TEXT("Node @@ is unresolved.");
-			Settings.Report(EMessageSeverity::Error, UnitNode, UnresolvedMessage);
+			WorkData.Settings.Report(EMessageSeverity::Error, UnitNode, UnresolvedMessage);
 			return INDEX_NONE;
 		}
 
@@ -1680,7 +1804,7 @@ int32 URigVMCompiler::TraverseCallExtern(const FRigVMCallExternExprAST* InExpr,
 		if(DispatchNode->GetFactory() == nullptr)
 		{
 			static const FString UnresolvedDispatchMessage = TEXT("Dispatch node @@ has no factory.");
-			Settings.Report(EMessageSeverity::Error, DispatchNode, UnresolvedDispatchMessage);
+			WorkData.Settings.Report(EMessageSeverity::Error, DispatchNode, UnresolvedDispatchMessage);
 			return INDEX_NONE;
 		}
 
@@ -1958,7 +2082,7 @@ int32 URigVMCompiler::TraverseCallExtern(const FRigVMCallExternExprAST* InExpr,
 
 #endif
 		
-		if (Settings.SetupNodeInstructionIndex)
+		if (WorkData.Settings.SetupNodeInstructionIndex)
 		{
 			WorkData.VM->GetByteCode().SetSubject(CallExternInstructionIndex, Callstack.GetCallPath(), Callstack.GetStack());
 		}
@@ -2040,7 +2164,7 @@ int32 URigVMCompiler::TraverseInlineFunction(const FRigVMInlineFunctionExprAST*
 {
 	URigVMNode* Node = InExpr->GetNode();
 	URigVMFunctionReferenceNode* FunctionReferenceNode = Cast<URigVMFunctionReferenceNode>(Node);
-	if(!ValidateNode(FunctionReferenceNode, false))
+	if(!ValidateNode(WorkData.Settings, FunctionReferenceNode, false))
 	{
 		return INDEX_NONE;
 	}
@@ -2308,7 +2432,7 @@ int32 URigVMCompiler::TraverseInlineFunction(const FRigVMInlineFunctionExprAST*
 				Op.FirstBranchInfoIndex = Op.FirstBranchInfoIndex + BranchIndexStart;
 			}
 
-			if (Settings.SetupNodeInstructionIndex)
+			if (WorkData.Settings.SetupNodeInstructionIndex)
 			{
 				if (const TArray<UObject*>* Callstack = FunctionByteCode.GetCallstackForInstruction(i-InstructionIndexStart))
 				{
@@ -2584,7 +2708,7 @@ void URigVMCompiler::TraverseExit(const FRigVMExitExprAST* InExpr, FRigVMCompile
 void URigVMCompiler::TraverseInvokeEntry(const FRigVMInvokeEntryExprAST* InExpr, FRigVMCompilerWorkData& WorkData)
 {
 	URigVMInvokeEntryNode* InvokeEntryNode = Cast<URigVMInvokeEntryNode>(InExpr->GetNode());
-	if(!ValidateNode(InvokeEntryNode))
+	if(!ValidateNode(WorkData.Settings, InvokeEntryNode))
 	{
 		return;
 	}
@@ -2598,7 +2722,7 @@ void URigVMCompiler::TraverseInvokeEntry(const FRigVMInvokeEntryExprAST* InExpr,
 		const int32 InstructionIndex = WorkData.VM->GetByteCode().GetNumInstructions();
 		WorkData.VM->GetByteCode().AddInvokeEntryOp(InvokeEntryNode->GetEntryName());
 
-		if (Settings.SetupNodeInstructionIndex)
+		if (WorkData.Settings.SetupNodeInstructionIndex)
 		{
 			const FRigVMCallstack Callstack = InExpr->GetProxy().GetSibling(InvokeEntryNode).GetCallstack();
 			WorkData.VM->GetByteCode().SetSubject(InstructionIndex, Callstack.GetCallPath(), Callstack.GetStack());
@@ -2691,7 +2815,7 @@ void URigVMCompiler::AddCopyOperator(const FRigVMCopyOp& InOp, const FRigVMAssig
 				static constexpr TCHAR MissingCastMessage[] = TEXT("Cast (%s to %s) for Node @@ not found.");
 				const FString& SourceCPPType = Registry.GetType(SourceTypeIndex).CPPType.ToString();
 				const FString& TargetCPPType = Registry.GetType(TargetTypeIndex).CPPType.ToString();
-				Settings.Report(EMessageSeverity::Error, InAssignExpr->GetTargetPin()->GetNode(),
+				WorkData.Settings.Report(EMessageSeverity::Error, InAssignExpr->GetTargetPin()->GetNode(),
 					FString::Printf(MissingCastMessage, *SourceCPPType, *TargetCPPType));
 				return;
 			}
@@ -2735,7 +2859,7 @@ void URigVMCompiler::AddCopyOperator(const FRigVMCopyOp& InOp, const FRigVMAssig
 	}
 
 	int32 InstructionIndex = WorkData.VM->GetByteCode().GetNumInstructions() - 1;
-	if (Settings.SetupNodeInstructionIndex)
+	if (WorkData.Settings.SetupNodeInstructionIndex)
 	{
 		bool bSetSubject = false;
 		if (URigVMPin* SourcePin = InAssignExpr->GetSourcePin())
@@ -3026,8 +3150,9 @@ const FRigVMVarExprAST* URigVMCompiler::GetSourceVarExpr(const FRigVMExprAST* In
 	return nullptr;
 }
 
-void URigVMCompiler::MarkDebugWatch(bool bRequired, URigVMPin* InPin, URigVM* OutVM,
-	TMap<FString, FRigVMOperand>* OutOperands, TSharedPtr<FRigVMParserAST> InRuntimeAST)
+void URigVMCompiler::MarkDebugWatch(const FRigVMCompileSettings& InSettings, bool bRequired,
+	URigVMPin* InPin, URigVM* OutVM, TMap<FString, FRigVMOperand>* OutOperands,
+	TSharedPtr<FRigVMParserAST> InRuntimeAST)
 {
 	check(InPin);
 	check(OutVM);
@@ -3036,7 +3161,7 @@ void URigVMCompiler::MarkDebugWatch(bool bRequired, URigVMPin* InPin, URigVM* Ou
 	
 	URigVMPin* Pin = InPin->GetRootPin();
 	URigVMPin* SourcePin = Pin;
-	if(Settings.ASTSettings.bFoldAssignments)
+	if(InSettings.ASTSettings.bFoldAssignments)
 	{
 		while(SourcePin->GetSourceLinks().Num() > 0)
 		{
@@ -3296,7 +3421,7 @@ FRigVMOperand URigVMCompiler::FindOrAddRegister(const FRigVMVarExprAST* InVarExp
 	FRigVMOperand const* ExistingOperandPtr = WorkData.PinPathToOperand->Find(Hash);
 	if (!ExistingOperandPtr)
 	{
-		if(Settings.ASTSettings.bFoldAssignments) 		
+		if(WorkData.Settings.ASTSettings.bFoldAssignments) 		
 		{
 			// Get all possible pins that lead to the same operand		
 			const FRigVMCompilerWorkData::FRigVMASTProxyArray PinProxies = FindProxiesWithSharedOperand(InVarExpr, WorkData);
@@ -3482,7 +3607,7 @@ FRigVMOperand URigVMCompiler::FindOrAddRegister(const FRigVMVarExprAST* InVarExp
 	ensure(Operand.IsValid());
 
 	// Get all possible pins that lead to the same operand
-	if(Settings.ASTSettings.bFoldAssignments)
+	if(WorkData.Settings.ASTSettings.bFoldAssignments)
 	{
 		// tbd: this functionality is only needed when there is a watch anywhere?
 		//if(!WorkData.WatchedPins.IsEmpty())
@@ -3630,7 +3755,7 @@ int32 URigVMCompiler::GetOperandFunctionInterfaceParameterIndex(const TArray<FSt
 	return INDEX_NONE;
 }
 
-bool URigVMCompiler::ValidateNode(URigVMNode* InNode, bool bCheck)
+bool URigVMCompiler::ValidateNode(const FRigVMCompileSettings& InSettings, URigVMNode* InNode, bool bCheck)
 {
 	if(bCheck)
 	{
@@ -3641,7 +3766,7 @@ bool URigVMCompiler::ValidateNode(URigVMNode* InNode, bool bCheck)
 		if(InNode->HasWildCardPin())
 		{
 			static const FString UnknownTypeMessage = TEXT("Node @@ has unresolved pins of wildcard type.");
-			Settings.Report(EMessageSeverity::Error, InNode, UnknownTypeMessage);
+			InSettings.Report(EMessageSeverity::Error, InNode, UnknownTypeMessage);
 			return false;
 		}
 		return true;
@@ -3649,21 +3774,21 @@ bool URigVMCompiler::ValidateNode(URigVMNode* InNode, bool bCheck)
 	return false;
 }
 
-void URigVMCompiler::ReportInfo(const FString& InMessage)
+void URigVMCompiler::ReportInfo(const FRigVMCompileSettings& InSettings, const FString& InMessage)
 {
-	if (Settings.SurpressInfoMessages)
+	if (InSettings.SurpressInfoMessages)
 	{
 		return;
 	}
-	Settings.Report(EMessageSeverity::Info, nullptr, InMessage);
+	InSettings.Report(EMessageSeverity::Info, nullptr, InMessage);
 }
 
-void URigVMCompiler::ReportWarning(const FString& InMessage)
+void URigVMCompiler::ReportWarning(const FRigVMCompileSettings& InSettings, const FString& InMessage)
 {
-	Settings.Report(EMessageSeverity::Warning, nullptr, InMessage);
+	InSettings.Report(EMessageSeverity::Warning, nullptr, InMessage);
 }
 
-void URigVMCompiler::ReportError(const FString& InMessage)
+void URigVMCompiler::ReportError(const FRigVMCompileSettings& InSettings, const FString& InMessage)
 {
-	Settings.Report(EMessageSeverity::Error, nullptr, InMessage);
+	InSettings.Report(EMessageSeverity::Error, nullptr, InMessage);
 }
diff --git a/Engine/Plugins/Runtime/RigVM/Source/RigVMDeveloper/Public/RigVMCompiler/RigVMCompiler.h b/Engine/Plugins/Runtime/RigVM/Source/RigVMDeveloper/Public/RigVMCompiler/RigVMCompiler.h
--- a/Engine/Plugins/Runtime/RigVM/Source/RigVMDeveloper/Public/RigVMCompiler/RigVMCompiler.h
+++ b/Engine/Plugins/Runtime/RigVM/Source/RigVMDeveloper/Public/RigVMCompiler/RigVMCompiler.h
@@ -56,6 +56,9 @@ public:
 
 	UPROPERTY()
 	bool SetupNodeInstructionIndex;
+	
+	UPROPERTY()
+	bool ASTErrorsAsNotifications;
 
 	static FRigVMCompileSettings Fast(UScriptStruct* InExecuteContextStruct = nullptr)
 	{
@@ -75,6 +78,28 @@ public:
 		return Settings;
 	}
 
+	void ReportInfo(const FString& InMessage) const;
+	void ReportWarning(const FString& InMessage) const;
+	void ReportError(const FString& InMessage) const;
+
+	template <typename FmtType, typename... Types>
+	void ReportInfof(const FmtType& Fmt, Types... Args) const
+	{
+		ReportInfo(FString::Printf(Fmt, Args...));
+	}
+
+	template <typename FmtType, typename... Types>
+	void ReportWarningf(const FmtType& Fmt, Types... Args) const
+	{
+		ReportWarning(FString::Printf(Fmt, Args...));
+	}
+
+	template <typename FmtType, typename... Types>
+	void ReportErrorf(const FmtType& Fmt, Types... Args) const
+	{
+		ReportError(FString::Printf(Fmt, Args...));
+	}
+
 	void Report(EMessageSeverity::Type InSeverity, UObject* InSubject, const FString& InMessage) const
 	{
 		ASTSettings.Report(InSeverity, InSubject, InMessage);
@@ -87,11 +112,26 @@ public:
 	}
 };
 
+class RIGVMDEVELOPER_API FRigVMCompileSettingsDuringLoadGuard
+{
+public:
+
+	FRigVMCompileSettingsDuringLoadGuard(FRigVMCompileSettings& InSettings)
+		: ASTErrorsAsNotifications(InSettings.ASTErrorsAsNotifications, true)
+	{}
+
+private:
+
+	TGuardValue<bool> ASTErrorsAsNotifications;
+};
+
 struct RIGVMDEVELOPER_API FRigVMCompilerWorkData
 {
 public:
+	FRigVMCompileSettings Settings;
 	bool bSetupMemory;
 	URigVM* VM;
+	TArray<URigVMGraph*> Graphs;
 	UScriptStruct* ExecuteContextStruct;
 	TMap<FString, FRigVMOperand>* PinPathToOperand;
 	TMap<const FRigVMVarExprAST*, FRigVMOperand> ExprToOperand;
@@ -153,6 +193,32 @@ public:
 
 	// operators that have been delayed for injection into the bytecode
 	TMap<FRigVMOperand, FCopyOpInfo> DeferredCopyOps;
+
+	void ReportInfo(const FString& InMessage) const;
+	void ReportWarning(const FString& InMessage) const;
+	void ReportError(const FString& InMessage) const;
+
+	template <typename FmtType, typename... Types>
+	void ReportInfof(const FmtType& Fmt, Types... Args) const
+	{
+		ReportInfo(FString::Printf(Fmt, Args...));
+	}
+
+	template <typename FmtType, typename... Types>
+	void ReportWarningf(const FmtType& Fmt, Types... Args) const
+	{
+		ReportWarning(FString::Printf(Fmt, Args...));
+	}
+
+	template <typename FmtType, typename... Types>
+	void ReportErrorf(const FmtType& Fmt, Types... Args) const
+	{
+		ReportError(FString::Printf(Fmt, Args...));
+	}
+
+	FRigVMReportDelegate OriginalReportDelegate;
+	void OverrideReportDelegate(bool& bEncounteredASTError, bool& bSurpressedASTError);
+	void RemoveOverrideReportDelegate();
 };
 
 DECLARE_DELEGATE_RetVal_OneParam(const FRigVMFunctionCompilationData*, FRigVMCompiler_GetFunctionCompilationData, const FRigVMGraphFunctionHeader& Header);
@@ -166,8 +232,8 @@ public:
 
 	URigVMCompiler();
 
-	UPROPERTY(BlueprintReadWrite, Category = FRigVMCompiler)
-	FRigVMCompileSettings Settings;
+	UPROPERTY()
+	FRigVMCompileSettings Settings_DEPRECATED;
 
 	UFUNCTION(BlueprintCallable, Category = FRigVMCompiler, meta=(DeprecatedFunction, DeprecationMessage="Compile is deprecated, use CompileVM with Context parameter."))
 	bool Compile(TArray<URigVMGraph*> InGraphs, URigVMController* InController, URigVM* OutVM) { return false; }
@@ -175,19 +241,25 @@ public:
 	UFUNCTION(BlueprintCallable, Category = FRigVMCompiler)
 	bool CompileVM(TArray<URigVMGraph*> InGraphs, URigVMController* InController, URigVM* OutVM, FRigVMExtendedExecuteContext& Context)
 	{
-		return Compile(InGraphs, InController, OutVM, Context, TArray<FRigVMExternalVariable>(), nullptr);
+		return Compile(Settings_DEPRECATED, InGraphs, InController, OutVM, Context, TArray<FRigVMExternalVariable>(), nullptr);
 	}
 
 	UE_DEPRECATED(5.3, "Please use Compile with Context param.")
 	bool Compile(TArray<URigVMGraph*> InGraphs, URigVMController* InController, URigVM* OutVM, const TArray<FRigVMExternalVariable>& InExternalVariables, TMap<FString, FRigVMOperand>* OutOperands, TSharedPtr<FRigVMParserAST> InAST = TSharedPtr<FRigVMParserAST>(), FRigVMFunctionCompilationData* OutFunctionCompilationData = nullptr) { return false; }
 
+	UE_DEPRECATED(5.4, "Please use CompileFunction with Settings param.")
 	bool Compile(TArray<URigVMGraph*> InGraphs, URigVMController* InController, URigVM* OutVM, FRigVMExtendedExecuteContext& OutVMContext, const TArray<FRigVMExternalVariable>& InExternalVariables, TMap<FString, FRigVMOperand>* OutOperands, TSharedPtr<FRigVMParserAST> InAST = TSharedPtr<FRigVMParserAST>(), FRigVMFunctionCompilationData* OutFunctionCompilationData = nullptr);
 
+	bool Compile(const FRigVMCompileSettings& InSettings, TArray<URigVMGraph*> InGraphs, URigVMController* InController, URigVM* OutVM, FRigVMExtendedExecuteContext& OutVMContext, const TArray<FRigVMExternalVariable>& InExternalVariables, TMap<FString, FRigVMOperand>* OutOperands, TSharedPtr<FRigVMParserAST> InAST = TSharedPtr<FRigVMParserAST>(), FRigVMFunctionCompilationData* OutFunctionCompilationData = nullptr);
+
 	UE_DEPRECATED(5.3, "Please use CompileFunction with Context param.")
 	bool CompileFunction(const URigVMLibraryNode* InLibraryNode, URigVMController* InController, FRigVMFunctionCompilationData* OutFunctionCompilationData) { return false; }
 
+	UE_DEPRECATED(5.4, "Please use CompileFunction with Settings param.")
 	bool CompileFunction(const URigVMLibraryNode* InLibraryNode, URigVMController* InController, FRigVMFunctionCompilationData* OutFunctionCompilationData, FRigVMExtendedExecuteContext& OutVMContext);
 
+	bool CompileFunction(const FRigVMCompileSettings& InSettings, const URigVMLibraryNode* InLibraryNode, URigVMController* InController, FRigVMFunctionCompilationData* OutFunctionCompilationData, FRigVMExtendedExecuteContext& OutVMContext);
+
 	FRigVMCompiler_GetFunctionCompilationData GetFunctionCompilationData;
 	TMap<FString, const FRigVMFunctionCompilationData*> CompiledFunctions;
 
@@ -199,7 +271,7 @@ public:
 	// this is currently only used for arrays.
 	static const FRigVMVarExprAST* GetSourceVarExpr(const FRigVMExprAST* InExpr);
 
-	void MarkDebugWatch(bool bRequired, URigVMPin* InPin, URigVM* OutVM, TMap<FString, FRigVMOperand>* OutOperands, TSharedPtr<FRigVMParserAST> InRuntimeAST);
+	void MarkDebugWatch(const FRigVMCompileSettings& InSettings, bool bRequired, URigVMPin* InPin, URigVM* OutVM, TMap<FString, FRigVMOperand>* OutOperands, TSharedPtr<FRigVMParserAST> InRuntimeAST);
 
 private:
 
@@ -245,28 +317,28 @@ private:
 	static FString GetPinNameWithDirectionPrefix(const URigVMPin* Pin);
 	static int32 GetOperandFunctionInterfaceParameterIndex(const TArray<FString>& OperandsPinNames, const FRigVMFunctionCompilationData* FunctionCompilationData, const FRigVMOperand& Operand);
 
-	bool ValidateNode(URigVMNode* InNode, bool bCheck = true);
+	bool ValidateNode(const FRigVMCompileSettings& InSettings, URigVMNode* InNode, bool bCheck = true);
 	
-	void ReportInfo(const FString& InMessage);
-	void ReportWarning(const FString& InMessage);
-	void ReportError(const FString& InMessage);
+	void ReportInfo(const FRigVMCompileSettings& InSettings, const FString& InMessage);
+	void ReportWarning(const FRigVMCompileSettings& InSettings, const FString& InMessage);
+	void ReportError(const FRigVMCompileSettings& InSettings, const FString& InMessage);
 
 	template <typename FmtType, typename... Types>
-	void ReportInfof(const FmtType& Fmt, Types... Args)
+	void ReportInfof(const FRigVMCompileSettings& InSettings, const FmtType& Fmt, Types... Args)
 	{
-		ReportInfo(FString::Printf(Fmt, Args...));
+		ReportInfo(InSettings, FString::Printf(Fmt, Args...));
 	}
 
 	template <typename FmtType, typename... Types>
-	void ReportWarningf(const FmtType& Fmt, Types... Args)
+	void ReportWarningf(const FRigVMCompileSettings& InSettings, const FmtType& Fmt, Types... Args)
 	{
-		ReportWarning(FString::Printf(Fmt, Args...));
+		ReportWarning(InSettings, FString::Printf(Fmt, Args...));
 	}
 
 	template <typename FmtType, typename... Types>
-	void ReportErrorf(const FmtType& Fmt, Types... Args)
+	void ReportErrorf(const FRigVMCompileSettings& InSettings, const FmtType& Fmt, Types... Args)
 	{
-		ReportError(FString::Printf(Fmt, Args...));
+		ReportError(InSettings, FString::Printf(Fmt, Args...));
 	}
 	
 	friend class FRigVMCompilerImportErrorContext;
