RigVM: Avoid redundant compilation

diff --git a/Engine/Plugins/Runtime/RigVM/Source/RigVMEditor/Private/Editor/RigVMEditor.cpp b/Engine/Plugins/Runtime/RigVM/Source/RigVMEditor/Private/Editor/RigVMEditor.cpp
--- a/Engine/Plugins/Runtime/RigVM/Source/RigVMEditor/Private/Editor/RigVMEditor.cpp
+++ b/Engine/Plugins/Runtime/RigVM/Source/RigVMEditor/Private/Editor/RigVMEditor.cpp
@@ -192,8 +192,11 @@ void FRigVMEditor::InitRigVMEditor(const EToolkitMode::Type Mode, const TSharedP
 		GetEditorModeManager().ActivateMode(GetEditorModeName());
 	}
 
-	UpdateRigVMHost();
-
+	{
+		TGuardValue<bool> GuardCompileReEntry(bIsCompilingThroughUI, true); // avoid redundant compilation, as it will be done at RebuildGraphFromModel
+		UpdateRigVMHost();
+	}
+	
 	// Post-layout initialization
 	PostLayoutBlueprintEditorInitialization();
 
@@ -1023,8 +1026,8 @@ void FRigVMEditor::Compile()
 		{
 			TGuardValue<bool> GuardCompileReEntry(bIsCompilingThroughUI, true);
 			FBlueprintEditor::Compile();
+			UpdateRigVMHost();
 		}
-		UpdateRigVMHost();
 
 		if (URigVMHost* RigVMHost = GetRigVMHost())
 		{
