Control Rig: Fix assets being loaded when they have function references during loading of the dependency asset
#rnx

diff --git a/Engine/Plugins/Runtime/RigVM/Source/RigVMDeveloper/Private/RigVMModel/RigVMClient.cpp b/Engine/Plugins/Runtime/RigVM/Source/RigVMDeveloper/Private/RigVMModel/RigVMClient.cpp
--- a/Engine/Plugins/Runtime/RigVM/Source/RigVMDeveloper/Private/RigVMModel/RigVMClient.cpp
+++ b/Engine/Plugins/Runtime/RigVM/Source/RigVMDeveloper/Private/RigVMModel/RigVMClient.cpp
@@ -1178,11 +1178,8 @@ bool FRigVMClient::UpdateFunctionReferences(const FRigVMGraphFunctionHeader& Hea
 		{
 			const TSoftObjectPtr<URigVMFunctionReferenceNode>& Reference = FunctionReferenceArray->FunctionReferences[i];
 
-			// Load reference package
-			if (!Reference.IsValid())
-			{
-				Reference.LoadSynchronous();
-			}
+			// Only update references that are loaded
+			// Other references will be updated when they are loaded
 			if (Reference.IsValid())
 			{
 				URigVMFunctionReferenceNode* Node = Reference.Get();
@@ -1194,7 +1191,7 @@ bool FRigVMClient::UpdateFunctionReferences(const FRigVMGraphFunctionHeader& Hea
 				{
 					if (URigVMLibraryNode* LibraryNode = Node->FindFunctionForNode())
 					{
-						IRigVMClientHost* OtherClientHost = LibraryNode->GetImplementingOuter<IRigVMClientHost>();							
+						IRigVMClientHost* OtherClientHost = LibraryNode->GetImplementingOuter<IRigVMClientHost>();
 						if (bUpdateDependencies)
 						{
 							OtherClientHost->GetRigVMClient()->UpdateDependenciesForFunction(LibraryNode);
