Remove impossible cast in RigVMEdGraphSchema
#rnx

diff --git a/Engine/Plugins/Runtime/RigVM/Source/RigVMDeveloper/Private/EdGraph/RigVMEdGraphSchema.cpp b/Engine/Plugins/Runtime/RigVM/Source/RigVMDeveloper/Private/EdGraph/RigVMEdGraphSchema.cpp
--- a/Engine/Plugins/Runtime/RigVM/Source/RigVMDeveloper/Private/EdGraph/RigVMEdGraphSchema.cpp
+++ b/Engine/Plugins/Runtime/RigVM/Source/RigVMDeveloper/Private/EdGraph/RigVMEdGraphSchema.cpp
@@ -960,6 +960,8 @@ TSharedPtr<INameValidatorInterface> URigVMEdGraphSchema::GetNameValidator(const
 {
 	if (ActionTypeId == FRigVMEdGraphSchemaAction_LocalVar::StaticGetTypeId())
 	{
+		// this cast will always fail, URigVMEdGraph is not a UStruct
+		/*
 		if (const URigVMEdGraph* EdGraph = Cast<URigVMEdGraph>(ValidationScope))
 		{
 			if (const URigVMGraph* Graph = EdGraph->GetModel())
@@ -967,6 +969,7 @@ TSharedPtr<INameValidatorInterface> URigVMEdGraphSchema::GetNameValidator(const
 				return MakeShareable(new FRigVMLocalVariableNameValidator(BlueprintObj, Graph, OriginalName));
 			}
 		}
+		*/
 	}
 
 	return MakeShareable(new FRigVMNameValidator(BlueprintObj, ValidationScope, OriginalName));		
