RigVM: Improve discoverability of Make and Break dispatches

diff --git a/Engine/Plugins/Runtime/RigVM/Source/RigVM/Private/RigVMFunctions/RigVMDispatch_MakeStruct.cpp b/Engine/Plugins/Runtime/RigVM/Source/RigVM/Private/RigVMFunctions/RigVMDispatch_MakeStruct.cpp
--- a/Engine/Plugins/Runtime/RigVM/Source/RigVM/Private/RigVMFunctions/RigVMDispatch_MakeStruct.cpp
+++ b/Engine/Plugins/Runtime/RigVM/Source/RigVM/Private/RigVMFunctions/RigVMDispatch_MakeStruct.cpp
@@ -73,6 +73,23 @@ FString FRigVMDispatch_MakeStruct::GetArgumentMetaData(const FName& InArgumentNa
 	return FRigVMDispatch_CoreBase::GetArgumentMetaData(InArgumentName, InMetaDataKey);
 }
 
+FString FRigVMDispatch_MakeStruct::GetKeywords() const
+{
+	const FRigVMRegistry& Registry = FRigVMRegistry::Get();
+	FString Keywords = FRigVMDispatch_CoreBase::GetKeywords();
+	const TArray<TRigVMTypeIndex>& StructTypeIndices = Registry.GetTypesForCategory(FRigVMTemplateArgument::ETypeCategory_SingleScriptStructValue);
+	for(const TRigVMTypeIndex& StructTypeIndex : StructTypeIndices)
+	{
+		const FRigVMTemplateArgumentType& Type = Registry.GetType(StructTypeIndex);
+		if(const UScriptStruct* ScriptStruct = Cast<UScriptStruct>(Type.CPPTypeObject))
+		{
+			static constexpr TCHAR Format[] = TEXT(",Make%s");
+			Keywords += FString::Printf(Format, *ScriptStruct->GetName());
+		}
+	}
+	return Keywords;
+}
+
 #endif
 
 void FRigVMDispatch_MakeStruct::Execute(FRigVMExtendedExecuteContext& InContext, FRigVMMemoryHandleArray Handles, FRigVMPredicateBranchArray RigVMBranches)
@@ -104,6 +121,23 @@ FText FRigVMDispatch_BreakStruct::GetNodeTooltip(const FRigVMTemplateTypeMap& In
 	return LOCTEXT("BreakStructToolTip", "Decomposes a struct into its elements");
 }
 
+FString FRigVMDispatch_BreakStruct::GetKeywords() const
+{
+	const FRigVMRegistry& Registry = FRigVMRegistry::Get();
+	FString Keywords = FRigVMDispatch_CoreBase::GetKeywords();
+	const TArray<TRigVMTypeIndex>& StructTypeIndices = Registry.GetTypesForCategory(FRigVMTemplateArgument::ETypeCategory_SingleScriptStructValue);
+	for(const TRigVMTypeIndex& StructTypeIndex : StructTypeIndices)
+	{
+		const FRigVMTemplateArgumentType& Type = Registry.GetType(StructTypeIndex);
+		if(const UScriptStruct* ScriptStruct = Cast<UScriptStruct>(Type.CPPTypeObject))
+		{
+			static constexpr TCHAR Format[] = TEXT(",Break%s");
+			Keywords += FString::Printf(Format, *ScriptStruct->GetName());
+		}
+	}
+	return Keywords;
+}
+
 #endif
 
 #undef LOCTEXT_NAMESPACE
diff --git a/Engine/Plugins/Runtime/RigVM/Source/RigVM/Public/RigVMFunctions/RigVMDispatch_MakeStruct.h b/Engine/Plugins/Runtime/RigVM/Source/RigVM/Public/RigVMFunctions/RigVMDispatch_MakeStruct.h
--- a/Engine/Plugins/Runtime/RigVM/Source/RigVM/Public/RigVMFunctions/RigVMDispatch_MakeStruct.h
+++ b/Engine/Plugins/Runtime/RigVM/Source/RigVM/Public/RigVMFunctions/RigVMDispatch_MakeStruct.h
@@ -24,6 +24,7 @@ public:
 	virtual FString GetNodeTitle(const FRigVMTemplateTypeMap& InTypes) const override;
 	virtual FText GetNodeTooltip(const FRigVMTemplateTypeMap& InTypes) const override;
 	virtual FString GetArgumentMetaData(const FName& InArgumentName, const FName& InMetaDataKey) const override;
+	virtual FString GetKeywords() const override;
 #endif
 
 protected:
@@ -51,6 +52,7 @@ public:
 	virtual const TArray<FRigVMTemplateArgument>& GetArguments() const override;
 #if WITH_EDITOR
 	virtual FText GetNodeTooltip(const FRigVMTemplateTypeMap& InTypes) const override;
+	virtual FString GetKeywords() const;
 #endif
 };
 
