RigVM Blueprint Library for finding broken assets

diff --git a/Engine/Plugins/Animation/ControlRig/Source/ControlRigEditor/Private/ControlRigBlueprintEditorLibrary.cpp b/Engine/Plugins/Animation/ControlRig/Source/ControlRigEditor/Private/ControlRigBlueprintEditorLibrary.cpp
--- a/Engine/Plugins/Animation/ControlRig/Source/ControlRigEditor/Private/ControlRigBlueprintEditorLibrary.cpp
+++ b/Engine/Plugins/Animation/ControlRig/Source/ControlRigEditor/Private/ControlRigBlueprintEditorLibrary.cpp
@@ -2,10 +2,8 @@
 
 #include "ControlRigBlueprintEditorLibrary.h"
 #include "Editor/SRigHierarchy.h"
-#include "AssetRegistry/AssetData.h"
 #include "AssetRegistry/AssetRegistryModule.h"
 #include "RigVMModel/RigVMPin.h"
-#include "Misc/ScopedSlowTask.h"
 
 #include UE_INLINE_GENERATED_CPP_BY_NAME(ControlRigBlueprintEditorLibrary)
 
@@ -15,7 +13,7 @@ FAutoConsoleCommand FCmdControlRigLoadAllAssets
 	TEXT("Loads all control rig assets."),
 	FConsoleCommandDelegate::CreateLambda([]()
 	{
-		UControlRigBlueprintEditorLibrary::LoadAllControlRigs();
+		UControlRigBlueprintEditorLibrary::LoadAssetsByClass(UControlRigBlueprint::StaticClass());
 	})
 );
 
@@ -48,33 +46,6 @@ USkeletalMesh* UControlRigBlueprintEditorLibrary::GetPreviewMesh(UControlRigBlue
 	return InRigBlueprint->GetPreviewMesh();
 }
 
-void UControlRigBlueprintEditorLibrary::RecompileVM(UControlRigBlueprint* InRigBlueprint)
-{
-	if(InRigBlueprint == nullptr)
-	{
-		return;
-	}
-	InRigBlueprint->RecompileVM();
-}
-
-void UControlRigBlueprintEditorLibrary::RecompileVMIfRequired(UControlRigBlueprint* InRigBlueprint)
-{
-	if(InRigBlueprint == nullptr)
-	{
-		return;
-	}
-	InRigBlueprint->RecompileVMIfRequired();
-}
-
-void UControlRigBlueprintEditorLibrary::RequestAutoVMRecompilation(UControlRigBlueprint* InRigBlueprint)
-{
-	if(InRigBlueprint == nullptr)
-	{
-		return;
-	}
-	InRigBlueprint->RequestAutoVMRecompilation();
-}
-
 void UControlRigBlueprintEditorLibrary::RequestControlRigInit(UControlRigBlueprint* InRigBlueprint)
 {
 	if(InRigBlueprint == nullptr)
@@ -84,24 +55,6 @@ void UControlRigBlueprintEditorLibrary::RequestControlRigInit(UControlRigBluepri
 	InRigBlueprint->RequestRigVMInit();
 }
 
-URigVMGraph* UControlRigBlueprintEditorLibrary::GetModel(UControlRigBlueprint* InRigBlueprint)
-{
-	if(InRigBlueprint == nullptr)
-	{
-		return nullptr;
-	}
-	return InRigBlueprint->GetDefaultModel();
-}
-
-URigVMController* UControlRigBlueprintEditorLibrary::GetController(UControlRigBlueprint* InRigBlueprint)
-{
-	if(InRigBlueprint == nullptr)
-	{
-		return nullptr;
-	}
-	return InRigBlueprint->GetController();
-}
-
 TArray<UControlRigBlueprint*> UControlRigBlueprintEditorLibrary::GetCurrentlyOpenRigBlueprints()
 {
 	return UControlRigBlueprint::GetCurrentlyOpenRigBlueprints();
@@ -174,35 +127,3 @@ TArray<FRigModuleDescription> UControlRigBlueprintEditorLibrary::GetAvailableRig
 	
 	return ModuleDescriptions;
 }
-
-void UControlRigBlueprintEditorLibrary::LoadAllControlRigs()
-{
-	const FAssetRegistryModule& AssetRegistryModule = FModuleManager::LoadModuleChecked<FAssetRegistryModule>(TEXT("AssetRegistry"));
-
-	// Collect a full list of assets with the specified class
-	TArray<FAssetData> AssetDataList;
-	AssetRegistryModule.Get().GetAssetsByClass(UControlRigBlueprint::StaticClass()->GetClassPathName(), AssetDataList, true);
-
-	const int32 NumAssets = AssetDataList.Num();
-
-	FScopedSlowTask LoadAllRigsTask(NumAssets, FText::FromString("Load all control rigs..."));
-	LoadAllRigsTask.MakeDialog(true);
-
-	for(int32 Index = 0; Index < NumAssets; Index++)
-	{
-		if (LoadAllRigsTask.ShouldCancel())
-		{
-			break;
-		}
-
-		const FAssetData& AssetData = AssetDataList[Index];
-		LoadAllRigsTask.EnterProgressFrame(1, FText::FromName(AssetData.PackageName));		
-
-		static constexpr TCHAR Format[] = TEXT("[%d/%d]: %s -> %s");
-		UE_LOG(LogControlRig, Display, Format, Index, NumAssets, *AssetData.AssetName.ToString(), *AssetData.PackageName.ToString())
-		(void)AssetData.GetAsset();
-
-		LoadAllRigsTask.ForceRefresh();
-	}
-}
-
diff --git a/Engine/Plugins/Animation/ControlRig/Source/ControlRigEditor/Public/ControlRigBlueprintEditorLibrary.h b/Engine/Plugins/Animation/ControlRig/Source/ControlRigEditor/Public/ControlRigBlueprintEditorLibrary.h
--- a/Engine/Plugins/Animation/ControlRig/Source/ControlRigEditor/Public/ControlRigBlueprintEditorLibrary.h
+++ b/Engine/Plugins/Animation/ControlRig/Source/ControlRigEditor/Public/ControlRigBlueprintEditorLibrary.h
@@ -10,7 +10,7 @@
 #include "CoreMinimal.h"
 #include "UObject/ObjectMacros.h"
 #include "ControlRigBlueprint.h"
-#include "Kismet/BlueprintFunctionLibrary.h"
+#include "RigVMEditorBlueprintLibrary.h"
 #include "ControlRigBlueprintEditorLibrary.generated.h"
 
 UENUM(BlueprintType)
@@ -21,7 +21,7 @@ enum class ECastToControlRigBlueprintCases : uint8
 };
 
 UCLASS(BlueprintType, meta=(ScriptName="ControlRigBlueprintLibrary"))
-class CONTROLRIGEDITOR_API UControlRigBlueprintEditorLibrary : public UBlueprintFunctionLibrary
+class CONTROLRIGEDITOR_API UControlRigBlueprintEditorLibrary : public URigVMEditorBlueprintLibrary
 {
 	GENERATED_BODY()
 
@@ -39,24 +39,9 @@ public:
 	UFUNCTION(BlueprintCallable, BlueprintPure, Category = "Control Rig Blueprint")
 	static USkeletalMesh* GetPreviewMesh(UControlRigBlueprint* InRigBlueprint);
 
-	UFUNCTION(BlueprintCallable, Category = "Control Rig Blueprint")
-	static void RecompileVM(UControlRigBlueprint* InRigBlueprint);
-
-	UFUNCTION(BlueprintCallable, Category = "Control Rig Blueprint")
-	static void RecompileVMIfRequired(UControlRigBlueprint* InRigBlueprint);
-	
-	UFUNCTION(BlueprintCallable, Category = "Control Rig Blueprint")
-	static void RequestAutoVMRecompilation(UControlRigBlueprint* InRigBlueprint);
-
 	UFUNCTION(BlueprintCallable, Category = "Control Rig Blueprint")
 	static void RequestControlRigInit(UControlRigBlueprint* InRigBlueprint);
 
-	UFUNCTION(BlueprintCallable, BlueprintPure, Category = "Control Rig Blueprint")
-	static URigVMGraph* GetModel(UControlRigBlueprint* InRigBlueprint);
-
-	UFUNCTION(BlueprintCallable, BlueprintPure, Category = "Control Rig Blueprint")
-	static URigVMController* GetController(UControlRigBlueprint* InRigBlueprint);
-
 	UFUNCTION(BlueprintCallable, BlueprintPure, Category = "Control Rig Blueprint")
 	static TArray<UControlRigBlueprint*> GetCurrentlyOpenRigBlueprints();
 
@@ -74,8 +59,5 @@ public:
 
 	UFUNCTION(BlueprintCallable, BlueprintPure, Category = "Control Rig Blueprint")
 	static TArray<FRigModuleDescription> GetAvailableRigModules();
-	
-	UFUNCTION(BlueprintCallable, Category = "Control Rig Blueprint")
-	static void LoadAllControlRigs();
 };
 
diff --git a/Engine/Plugins/Runtime/RigVM/Source/RigVM/Public/RigVMCore/RigVMGraphFunctionDefinition.h b/Engine/Plugins/Runtime/RigVM/Source/RigVM/Public/RigVMCore/RigVMGraphFunctionDefinition.h
--- a/Engine/Plugins/Runtime/RigVM/Source/RigVM/Public/RigVMCore/RigVMGraphFunctionDefinition.h
+++ b/Engine/Plugins/Runtime/RigVM/Source/RigVM/Public/RigVMCore/RigVMGraphFunctionDefinition.h
@@ -291,28 +291,28 @@ struct RIGVM_API FRigVMGraphFunctionArgument
 	, bIsConst(false)
 	{}
 
-	UPROPERTY()
+	UPROPERTY(BlueprintReadOnly, VisibleAnywhere, Category=FunctionArgument)
 	FName Name;
 
-	UPROPERTY()
+	UPROPERTY(BlueprintReadOnly, VisibleAnywhere, Category=FunctionArgument)
 	FName DisplayName;
 
-	UPROPERTY()
+	UPROPERTY(BlueprintReadOnly, VisibleAnywhere, Category=FunctionArgument)
 	FName CPPType;
 
-	UPROPERTY()
+	UPROPERTY(BlueprintReadOnly, VisibleAnywhere, Category=FunctionArgument)
 	TSoftObjectPtr<UObject> CPPTypeObject;
 
-	UPROPERTY()
+	UPROPERTY(BlueprintReadOnly, VisibleAnywhere, Category=FunctionArgument)
 	bool bIsArray;
 
-	UPROPERTY()
+	UPROPERTY(BlueprintReadOnly, VisibleAnywhere, Category=FunctionArgument)
 	ERigVMPinDirection Direction;
 
-	UPROPERTY()
+	UPROPERTY(BlueprintReadOnly, VisibleAnywhere, Category=FunctionArgument)
 	FString DefaultValue;
 
-	UPROPERTY()
+	UPROPERTY(BlueprintReadOnly, VisibleAnywhere, Category=FunctionArgument)
 	bool bIsConst;
 	
 	UPROPERTY()
@@ -368,11 +368,11 @@ struct RIGVM_API FRigVMGraphFunctionIdentifier
 {
 	GENERATED_BODY()
 
-	UPROPERTY()
+	UPROPERTY(BlueprintReadOnly, VisibleAnywhere, Category=FunctionIdentifier)
 	FSoftObjectPath LibraryNode;
 
 	// A path to the IRigVMGraphFunctionHost that stores the function information, and compilation data (e.g. RigVMBlueprintGeneratedClass)
-	UPROPERTY()
+	UPROPERTY(BlueprintReadOnly, VisibleAnywhere, Category=FunctionIdentifier)
 	FSoftObjectPath HostObject;
 
 	FRigVMGraphFunctionIdentifier()
@@ -409,28 +409,28 @@ struct RIGVM_API FRigVMGraphFunctionHeader
 		, Name(NAME_None)
 	{}
 
-	UPROPERTY()
+	UPROPERTY(BlueprintReadOnly, VisibleAnywhere, Category=FunctionHeader)
 	FRigVMGraphFunctionIdentifier LibraryPointer;
 
-	UPROPERTY()
+	UPROPERTY(BlueprintReadOnly, VisibleAnywhere, Category=FunctionHeader)
 	FName Name;
 
-	UPROPERTY()
+	UPROPERTY(BlueprintReadOnly, VisibleAnywhere, Category=FunctionHeader)
 	FString NodeTitle;
 	
-	UPROPERTY()
+	UPROPERTY(BlueprintReadOnly, VisibleAnywhere, Category=FunctionHeader)
 	FLinearColor NodeColor = FLinearColor::White;
 
-	UPROPERTY()
+	UPROPERTY(BlueprintReadOnly, VisibleAnywhere, Category=FunctionHeader)
 	FText Tooltip;
 	
-	UPROPERTY()
+	UPROPERTY(BlueprintReadOnly, VisibleAnywhere, Category=FunctionHeader)
 	FString Category;
 
-	UPROPERTY()
+	UPROPERTY(BlueprintReadOnly, VisibleAnywhere, Category=FunctionHeader)
 	FString Keywords;	
 
-	UPROPERTY()
+	UPROPERTY(BlueprintReadOnly, VisibleAnywhere, Category=FunctionHeader)
 	TArray<FRigVMGraphFunctionArgument> Arguments;
 
 	UPROPERTY()
diff --git a/Engine/Plugins/Runtime/RigVM/Source/RigVMDeveloper/Private/RigVMBlueprint.cpp b/Engine/Plugins/Runtime/RigVM/Source/RigVMDeveloper/Private/RigVMBlueprint.cpp
--- a/Engine/Plugins/Runtime/RigVM/Source/RigVMDeveloper/Private/RigVMBlueprint.cpp
+++ b/Engine/Plugins/Runtime/RigVM/Source/RigVMDeveloper/Private/RigVMBlueprint.cpp
@@ -106,6 +106,8 @@ TArray<URigVMBlueprint*> URigVMBlueprint::sCurrentlyOpenedRigVMBlueprints;
 #if WITH_EDITOR
 const FName URigVMBlueprint::RigVMPanelNodeFactoryName(TEXT("FRigVMEdGraphPanelNodeFactory"));
 const FName URigVMBlueprint::RigVMPanelPinFactoryName(TEXT("FRigVMEdGraphPanelPinFactory"));
+FCriticalSection URigVMBlueprint::QueuedCompilerMessageDelegatesMutex;
+TArray<FOnRigVMReportCompilerMessage::FDelegate> URigVMBlueprint::QueuedCompilerMessageDelegates;
 #endif
 
 URigVMBlueprint::URigVMBlueprint()
@@ -140,6 +142,17 @@ URigVMBlueprint::URigVMBlueprint(const FObjectInitializer& ObjectInitializer)
 	VMCompileSettings.ASTSettings.ReportDelegate.BindUObject(this, &URigVMBlueprint::HandleReportFromCompiler);
 
 #if WITH_EDITOR
+	TArray<FOnRigVMReportCompilerMessage::FDelegate> DelegatesForReportFromCompiler;
+	{
+		FScopeLock Lock(&QueuedCompilerMessageDelegatesMutex);
+		Swap(QueuedCompilerMessageDelegates, DelegatesForReportFromCompiler);
+	}
+
+	for(const FOnRigVMReportCompilerMessage::FDelegate& Delegate : DelegatesForReportFromCompiler)
+	{
+		ReportCompilerMessageEvent.Add(Delegate);
+	}
+
 	if (HasAnyFlags(RF_ClassDefaultObject))
 	{
 		CompileLog.bSilentMode = true;
@@ -1486,6 +1499,13 @@ void URigVMBlueprint::HandleReportFromCompiler(EMessageSeverity::Type InSeverity
 			Log->Note(*InMessage);
 		}
 
+		static const FString Error = TEXT("Error");
+		static const FString Warning = TEXT("Warning");
+		if(InMessage.Contains(Error, ESearchCase::IgnoreCase) ||
+			InMessage.Contains(Warning, ESearchCase::IgnoreCase))
+		{
+			BroadCastReportCompilerMessage(InSeverity, InSubject, InMessage);
+		}
 		UE_LOG(LogRigVMDeveloper, Display, TEXT("%s"), *InMessage);
 	}
 
@@ -4106,6 +4126,22 @@ bool URigVMBlueprint::RemoveEdGraphForCollapseNode(URigVMCollapseNode* InNode, b
 	return false;
 }
 
+#if WITH_EDITOR
+
+void URigVMBlueprint::QueueCompilerMessageDelegate(const FOnRigVMReportCompilerMessage::FDelegate& InDelegate)
+{
+	FScopeLock Lock(&QueuedCompilerMessageDelegatesMutex);
+	QueuedCompilerMessageDelegates.Add(InDelegate);
+}
+
+void URigVMBlueprint::ClearQueuedCompilerMessageDelegates()
+{
+	FScopeLock Lock(&QueuedCompilerMessageDelegatesMutex);
+	QueuedCompilerMessageDelegates.Reset();
+}
+
+#endif
+
 #undef LOCTEXT_NAMESPACE
 
 
diff --git a/Engine/Plugins/Runtime/RigVM/Source/RigVMDeveloper/Public/RigVMBlueprint.h b/Engine/Plugins/Runtime/RigVM/Source/RigVMDeveloper/Public/RigVMBlueprint.h
--- a/Engine/Plugins/Runtime/RigVM/Source/RigVMDeveloper/Public/RigVMBlueprint.h
+++ b/Engine/Plugins/Runtime/RigVM/Source/RigVMDeveloper/Public/RigVMBlueprint.h
@@ -13,6 +13,9 @@
 #include "EdGraph/RigVMEdGraph.h"
 #include "EdGraph/RigVMEdGraphSchema.h"
 #include "RigVMSettings.h"
+#if WITH_EDITOR
+#include "HAL/CriticalSection.h"
+#endif
 
 #include "RigVMBlueprint.generated.h"
 
@@ -743,6 +746,20 @@ private:
 	bool bSuspendPythonMessagesForRigVMClient;
 	bool bMarkBlueprintAsStructurallyModifiedPending;
 
+#if WITH_EDITOR
+
+public:
+
+	static void QueueCompilerMessageDelegate(const FOnRigVMReportCompilerMessage::FDelegate& InDelegate);
+	static void ClearQueuedCompilerMessageDelegates();
+	
+private:
+
+	static FCriticalSection QueuedCompilerMessageDelegatesMutex;
+	static TArray<FOnRigVMReportCompilerMessage::FDelegate> QueuedCompilerMessageDelegates;
+
+#endif
+
 	friend class FRigVMBlueprintCompilerContext;
 	friend class FRigVMEditor;
 	friend class FRigVMEditorModule;
diff --git a/Engine/Plugins/Runtime/RigVM/Source/RigVMDeveloper/Public/RigVMModel/Nodes/RigVMFunctionReferenceNode.h b/Engine/Plugins/Runtime/RigVM/Source/RigVMDeveloper/Public/RigVMModel/Nodes/RigVMFunctionReferenceNode.h
--- a/Engine/Plugins/Runtime/RigVM/Source/RigVMDeveloper/Public/RigVMModel/Nodes/RigVMFunctionReferenceNode.h
+++ b/Engine/Plugins/Runtime/RigVM/Source/RigVMDeveloper/Public/RigVMModel/Nodes/RigVMFunctionReferenceNode.h
@@ -47,6 +47,10 @@ public:
 	// end Variable remapping
 
 	virtual uint32 GetStructureHash() const override;
+
+	UFUNCTION(BlueprintCallable, Category = RigVMLibraryNode, meta = (DisplayName = "GetReferencedFunctionHeader", ScriptName = "GetReferencedFunctionHeader"))
+	FRigVMGraphFunctionHeader GetReferencedFunctionHeader_ForBlueprint() const { return GetReferencedFunctionHeader(); }
+
 	const FRigVMGraphFunctionHeader& GetReferencedFunctionHeader() const { return ReferencedFunctionHeader; }
 
 	const FRigVMGraphFunctionData* GetReferencedFunctionData(bool bLoadIfNecessary = true) const;
diff --git a/Engine/Plugins/Runtime/RigVM/Source/RigVMEditor/Private/RigVMEditorBlueprintLibrary.cpp b/Engine/Plugins/Runtime/RigVM/Source/RigVMEditor/Private/RigVMEditorBlueprintLibrary.cpp
new file mode 100644
--- /dev/null
+++ b/Engine/Plugins/Runtime/RigVM/Source/RigVMEditor/Private/RigVMEditorBlueprintLibrary.cpp
@@ -0,0 +1,327 @@
+// Copyright Epic Games, Inc. All Rights Reserved.
+
+#include "RigVMEditorBlueprintLibrary.h"
+#include "AssetRegistry/AssetRegistryModule.h"
+#include "Misc/ScopedSlowTask.h"
+
+#include UE_INLINE_GENERATED_CPP_BY_NAME(RigVMEditorBlueprintLibrary)
+
+#define LOCTEXT_NAMESPACE "RigVMEditorBlueprintLibrary"
+
+void URigVMEditorBlueprintLibrary::RecompileVM(URigVMBlueprint* InBlueprint)
+{
+	if(InBlueprint == nullptr)
+	{
+		return;
+	}
+	InBlueprint->RecompileVM();
+}
+
+void URigVMEditorBlueprintLibrary::RecompileVMIfRequired(URigVMBlueprint* InBlueprint)
+{
+	if(InBlueprint == nullptr)
+	{
+		return;
+	}
+	InBlueprint->RecompileVMIfRequired();
+}
+
+void URigVMEditorBlueprintLibrary::RequestAutoVMRecompilation(URigVMBlueprint* InBlueprint)
+{
+	if(InBlueprint == nullptr)
+	{
+		return;
+	}
+	InBlueprint->RequestAutoVMRecompilation();
+}
+
+URigVMGraph* URigVMEditorBlueprintLibrary::GetModel(URigVMBlueprint* InBlueprint)
+{
+	if(InBlueprint == nullptr)
+	{
+		return nullptr;
+	}
+	return InBlueprint->GetDefaultModel();
+}
+
+URigVMController* URigVMEditorBlueprintLibrary::GetController(URigVMBlueprint* InBlueprint)
+{
+	if(InBlueprint == nullptr)
+	{
+		return nullptr;
+	}
+	return InBlueprint->GetController();
+}
+
+TArray<URigVMBlueprint*> URigVMEditorBlueprintLibrary::LoadAssets()
+{
+	return LoadAssetsByClass(URigVMBlueprint::StaticClass());
+}
+
+TArray<URigVMBlueprint*> URigVMEditorBlueprintLibrary::LoadAssetsByClass(TSubclassOf<URigVMBlueprint> InClass)
+{
+	return LoadAssetsWithAssetDataAndBlueprintFilters(InClass, FRigVMAssetDataFilter(), FRigVMBlueprintFilter());
+}
+
+TArray<URigVMBlueprint*> URigVMEditorBlueprintLibrary::LoadAssetsWithBlueprintFilter_ForBlueprint(
+	TSubclassOf<URigVMBlueprint> InClass, FRigVMBlueprintFilterDynamic InBlueprintFilter)
+{
+	return LoadAssetsWithAssetDataAndBlueprintFilters(InClass, FRigVMAssetDataFilter(), FRigVMBlueprintFilter::CreateLambda(
+		[InBlueprintFilter](const URigVMBlueprint* Blueprint, const TArray<FRigVMBlueprintLoadLogEntry>& LogEntries) -> bool
+		{
+			return InBlueprintFilter.Execute(Blueprint, LogEntries);
+		})
+	);
+}
+
+TArray<URigVMBlueprint*> URigVMEditorBlueprintLibrary::LoadAssetsWithAssetDataFilter_ForBlueprint(
+	TSubclassOf<URigVMBlueprint> InClass, FRigVMAssetDataFilterDynamic InAssetDataFilter)
+{
+	return LoadAssetsWithAssetDataAndBlueprintFilters(InClass, FRigVMAssetDataFilter::CreateLambda(
+		[InAssetDataFilter](const FAssetData& AssetData) -> bool
+		{
+			return InAssetDataFilter.Execute(AssetData);
+		}),
+		FRigVMBlueprintFilter()
+	);
+}
+
+TArray<URigVMBlueprint*> URigVMEditorBlueprintLibrary::LoadAssetsWithNodeFilter_ForBlueprint(
+	TSubclassOf<URigVMBlueprint> InClass, FRigVMNodeFilterDynamic InNodeFilter)
+{
+	return LoadAssetsWithNodeFilter(InClass, FRigVMNodeFilter::CreateLambda(
+		[InNodeFilter](const URigVMBlueprint* Blueprint, const URigVMNode* Node) -> bool
+		{
+			return InNodeFilter.Execute(Blueprint, Node);
+		})
+	);
+}
+
+TArray<URigVMBlueprint*> URigVMEditorBlueprintLibrary::LoadAssetsWithNodeFilter(
+	TSubclassOf<URigVMBlueprint> InClass,
+	FRigVMNodeFilter InNodeFilter)
+{
+	return LoadAssetsWithAssetDataAndNodeFilters(InClass, FRigVMAssetDataFilter(), InNodeFilter);
+}
+
+TArray<URigVMBlueprint*> URigVMEditorBlueprintLibrary::LoadAssetsWithAssetDataAndBlueprintFilters_ForBlueprint(
+	TSubclassOf<URigVMBlueprint> InClass,
+	FRigVMAssetDataFilterDynamic InAssetDataFilter,
+	FRigVMBlueprintFilterDynamic InBlueprintFilter)
+{
+	return LoadAssetsWithAssetDataAndBlueprintFilters(InClass,
+		FRigVMAssetDataFilter::CreateLambda(
+		[InAssetDataFilter](const FAssetData& AssetData) -> bool
+			{
+				return InAssetDataFilter.Execute(AssetData);
+			}
+		),
+		FRigVMBlueprintFilter::CreateLambda(
+		[InBlueprintFilter](const URigVMBlueprint* Blueprint, const TArray<FRigVMBlueprintLoadLogEntry>& LogEntries) -> bool
+			{
+				return InBlueprintFilter.Execute(Blueprint, LogEntries);
+			}
+		)
+	);
+}
+
+TArray<URigVMBlueprint*> URigVMEditorBlueprintLibrary::LoadAssetsWithAssetDataAndBlueprintFilters(
+	TSubclassOf<URigVMBlueprint> InClass, 
+	FRigVMAssetDataFilter InAssetDataFilter, 
+	FRigVMBlueprintFilter InBlueprintFilter
+)
+{
+	TArray<FAssetData> AssetDataList = GetAssetsWithFilter(InClass, InAssetDataFilter);
+
+	const int32 NumAssets = AssetDataList.Num();
+
+	TArray<URigVMBlueprint*> LoadedAssets;
+	{
+		const FString Title = FString::Printf(TEXT("Load all %s assets..."), *InClass->GetName()); 
+		FScopedSlowTask LoadAssetsTask(NumAssets, FText::FromString(Title));
+		LoadAssetsTask.MakeDialog(true);
+
+		for(int32 Index = 0; Index < NumAssets; Index++)
+		{
+			if (LoadAssetsTask.ShouldCancel())
+			{
+				break;
+			}
+
+			const FAssetData& AssetData = AssetDataList[Index];
+			LoadAssetsTask.EnterProgressFrame(1, FText::FromName(AssetData.PackageName));		
+
+			static constexpr TCHAR Format[] = TEXT("[%d/%d]: %s -> %s");
+			UE_LOG(LogRigVM, Display, Format, Index, NumAssets, *AssetData.AssetName.ToString(), *AssetData.PackageName.ToString())
+
+			// completely ignore exceptions during this scope
+			FScopedScriptExceptionHandler ScopedScriptExceptionHandler([](ELogVerbosity::Type Verbosity, const TCHAR* ErrorMessage, const TCHAR* StackMessage)
+			{
+				FString Message;
+				if(ErrorMessage)
+				{
+					Message += ErrorMessage;
+				}
+				if(StackMessage)
+				{
+					if(!Message.IsEmpty())
+					{
+						Message += TEXT("\n");
+					}
+					Message += StackMessage;
+				}
+
+				if(Verbosity == ELogVerbosity::Error || Verbosity == ELogVerbosity::Fatal)
+				{
+					UE_LOG(LogRigVMDeveloper, Error, TEXT("%s"), *Message);
+				}
+				else if(Verbosity == ELogVerbosity::Warning)
+				{
+					UE_LOG(LogRigVMDeveloper, Warning, TEXT("%s"), *Message);
+				}
+				else
+				{
+					UE_LOG(LogRigVMDeveloper, Display, TEXT("%s"), *Message);
+				}
+			});
+
+			// set up a lambda to record all error messages
+			TArray<FRigVMBlueprintLoadLogEntry> LogEntries;
+
+			if(InBlueprintFilter.IsBound())
+			{
+				URigVMBlueprint::QueueCompilerMessageDelegate(FOnRigVMReportCompilerMessage::FDelegate::CreateLambda(
+					[&LogEntries](EMessageSeverity::Type InSeverity, UObject* InSubject, const FString& InMessage)
+					{
+						ERigVMBlueprintLoadLogSeverity Severity = ERigVMBlueprintLoadLogSeverity::Display;
+						if(InSeverity == EMessageSeverity::Warning)
+						{
+							Severity = ERigVMBlueprintLoadLogSeverity::Warning;
+						}
+						else if(InSeverity == EMessageSeverity::Error)
+						{
+							Severity = ERigVMBlueprintLoadLogSeverity::Error;
+						}
+						LogEntries.Emplace(Severity, InSubject, InMessage);
+					})
+				);
+			}
+			
+			if(URigVMBlueprint* Blueprint = Cast<URigVMBlueprint>(AssetData.GetAsset()))
+			{
+				if(InBlueprintFilter.IsBound())
+				{
+					if(InBlueprintFilter.Execute(Blueprint, LogEntries))
+					{
+						LoadedAssets.Add(Blueprint);
+					}
+				}
+				else
+				{
+					LoadedAssets.Add(Blueprint);
+				}
+
+				Blueprint->OnReportCompilerMessage().Clear();
+			}
+
+			if(InBlueprintFilter.IsBound())
+			{
+				URigVMBlueprint::ClearQueuedCompilerMessageDelegates();
+			}
+			LoadAssetsTask.ForceRefresh();
+		}
+	}
+
+	return LoadedAssets;
+}
+
+TArray<URigVMBlueprint*> URigVMEditorBlueprintLibrary::LoadAssetsWithAssetDataAndNodeFilters_ForBlueprint(
+	TSubclassOf<URigVMBlueprint> InClass, FRigVMAssetDataFilterDynamic InAssetDataFilter,
+	FRigVMNodeFilterDynamic InNodeFilter)
+{
+	return LoadAssetsWithAssetDataAndNodeFilters(InClass,
+		FRigVMAssetDataFilter::CreateLambda(
+		[InAssetDataFilter](const FAssetData& AssetData) -> bool
+		{
+			return InAssetDataFilter.Execute(AssetData);
+		}),
+		FRigVMNodeFilter::CreateLambda(
+		[InNodeFilter](const URigVMBlueprint* Blueprint, const URigVMNode* Node) -> bool
+		{
+			return InNodeFilter.Execute(Blueprint, Node);
+		})
+	);
+}
+
+TArray<URigVMBlueprint*> URigVMEditorBlueprintLibrary::LoadAssetsWithAssetDataAndNodeFilters(
+	TSubclassOf<URigVMBlueprint> InClass, FRigVMAssetDataFilter InAssetDataFilter, FRigVMNodeFilter InNodeFilter)
+{
+	return LoadAssetsWithAssetDataAndBlueprintFilters(InClass, InAssetDataFilter, FRigVMBlueprintFilter::CreateLambda(
+		[InNodeFilter](const URigVMBlueprint* Blueprint, const TArray<FRigVMBlueprintLoadLogEntry>& LogEntries) -> bool
+		{
+			TArray<URigVMGraph*> Models = Blueprint->GetAllModels();
+			int32 NumNodes = 0;
+			for(const URigVMGraph* Model : Models)
+			{
+				NumNodes += Model->GetNodes().Num();
+			}
+			FScopedSlowTask FilterNodesTask(NumNodes, LOCTEXT("FilteringNodes", "Filtering Nodes"));
+			for(const URigVMGraph* Model : Models)
+			{
+				for(const URigVMNode* Node : Model->GetNodes())
+				{
+					FilterNodesTask.EnterProgressFrame(1);		
+					if(InNodeFilter.Execute(Blueprint, Node))
+					{
+						return true;
+					}
+				}
+			}
+			return false;
+		})
+	);
+}
+
+TArray<FAssetData> URigVMEditorBlueprintLibrary::GetAssetsWithFilter_ForBlueprint(
+	TSubclassOf<URigVMBlueprint> InClass, FRigVMAssetDataFilterDynamic InAssetDataFilter)
+{
+	return GetAssetsWithFilter(InClass, FRigVMAssetDataFilter::CreateLambda(
+		[InAssetDataFilter](const FAssetData& AssetData) -> bool
+		{
+			return InAssetDataFilter.Execute(AssetData);
+		})
+	);
+}
+
+TArray<FAssetData> URigVMEditorBlueprintLibrary::GetAssetsWithFilter(
+		TSubclassOf<URigVMBlueprint> InClass,
+		FRigVMAssetDataFilter InAssetDataFilter)
+{
+	if(InClass == nullptr)
+	{
+		InClass = URigVMBlueprint::StaticClass();
+	}
+
+	const FAssetRegistryModule& AssetRegistryModule = FModuleManager::LoadModuleChecked<FAssetRegistryModule>(TEXT("AssetRegistry"));
+
+	// Collect a full list of assets with the specified class
+	TArray<FAssetData> AssetDataList;
+	AssetRegistryModule.Get().GetAssetsByClass(InClass->GetClassPathName(), AssetDataList, true);
+
+	if(InAssetDataFilter.IsBound())
+	{
+		TArray<FAssetData> FilteredAssetDataList;
+		for(const FAssetData& AssetData : AssetDataList)
+		{
+			if(InAssetDataFilter.Execute(AssetData))
+			{
+				FilteredAssetDataList.Add(AssetData);
+			}			
+		}
+		Swap(FilteredAssetDataList, AssetDataList);
+	}
+
+	return AssetDataList;
+}
+
+#undef LOCTEXT_NAMESPACE
diff --git a/Engine/Plugins/Runtime/RigVM/Source/RigVMEditor/Public/RigVMEditorBlueprintLibrary.h b/Engine/Plugins/Runtime/RigVM/Source/RigVMEditor/Public/RigVMEditorBlueprintLibrary.h
new file mode 100644
--- /dev/null
+++ b/Engine/Plugins/Runtime/RigVM/Source/RigVMEditor/Public/RigVMEditorBlueprintLibrary.h
@@ -0,0 +1,134 @@
+// Copyright Epic Games, Inc. All Rights Reserved.
+
+/**
+ *
+ * This thumbnail renderer displays a given Control Rig
+ */
+
+#pragma once
+
+#include "CoreMinimal.h"
+#include "RigVMBlueprint.h"
+#include "AssetRegistry/AssetData.h"
+#include "Kismet/BlueprintFunctionLibrary.h"
+#include "RigVMEditorBlueprintLibrary.generated.h"
+
+UENUM(BlueprintType)
+enum ERigVMBlueprintLoadLogSeverity
+{
+	Display,
+	Warning,
+	Error
+};
+
+USTRUCT(BlueprintType)
+struct RIGVMEDITOR_API FRigVMBlueprintLoadLogEntry
+{
+	GENERATED_BODY()
+
+	FRigVMBlueprintLoadLogEntry()
+		: Severity(ERigVMBlueprintLoadLogSeverity::Display)
+		, Subject(nullptr)
+	{}
+
+	FRigVMBlueprintLoadLogEntry(ERigVMBlueprintLoadLogSeverity InSeverity, UObject* InSubject, const FString& InMessage)
+		: Severity(InSeverity)
+		, Subject(InSubject)
+		, Message(InMessage)
+	{}
+	
+	UPROPERTY(BlueprintReadOnly, VisibleAnywhere, Category=Log)
+	TEnumAsByte<ERigVMBlueprintLoadLogSeverity> Severity;
+
+	UPROPERTY(BlueprintReadOnly, VisibleAnywhere, Category=Log)
+	TObjectPtr<UObject> Subject;
+
+	UPROPERTY(BlueprintReadOnly, VisibleAnywhere, Category=Log)
+	FString Message;
+};
+
+DECLARE_DELEGATE_RetVal_OneParam(bool, FRigVMAssetDataFilter, const FAssetData&);
+DECLARE_DELEGATE_RetVal_TwoParams(bool, FRigVMBlueprintFilter, const URigVMBlueprint*, const TArray<FRigVMBlueprintLoadLogEntry>&);
+DECLARE_DELEGATE_RetVal_TwoParams(bool, FRigVMNodeFilter, const URigVMBlueprint*, const URigVMNode*);
+DECLARE_DYNAMIC_DELEGATE_RetVal_OneParam(bool, FRigVMAssetDataFilterDynamic, FAssetData, AssetData);
+DECLARE_DYNAMIC_DELEGATE_RetVal_TwoParams(bool, FRigVMBlueprintFilterDynamic, const URigVMBlueprint*, Blueprint, TArray<FRigVMBlueprintLoadLogEntry>, LogDuringLoad);
+DECLARE_DYNAMIC_DELEGATE_RetVal_TwoParams(bool, FRigVMNodeFilterDynamic, const URigVMBlueprint*, Blueprint, const URigVMNode*, Node);
+
+UCLASS(BlueprintType, meta=(ScriptName="RigVMBlueprintLibrary"))
+class RIGVMEDITOR_API URigVMEditorBlueprintLibrary : public UBlueprintFunctionLibrary
+{
+	GENERATED_BODY()
+
+public:
+
+	UFUNCTION(BlueprintCallable, Category = "RigVM Blueprint")
+	static void RecompileVM(URigVMBlueprint* InBlueprint);
+
+	UFUNCTION(BlueprintCallable, Category = "RigVM Blueprint")
+	static void RecompileVMIfRequired(URigVMBlueprint* InBlueprint);
+	
+	UFUNCTION(BlueprintCallable, Category = "RigVM Blueprint")
+	static void RequestAutoVMRecompilation(URigVMBlueprint* InBlueprint);
+
+	UFUNCTION(BlueprintCallable, BlueprintPure, Category = "RigVM Blueprint")
+	static URigVMGraph* GetModel(URigVMBlueprint* InBlueprint);
+
+	UFUNCTION(BlueprintCallable, BlueprintPure, Category = "RigVM Blueprint")
+	static URigVMController* GetController(URigVMBlueprint* InBlueprint);
+
+	UFUNCTION(BlueprintCallable, Category = "RigVM Blueprint")
+	static TArray<URigVMBlueprint*> LoadAssets();
+
+	UFUNCTION(BlueprintCallable, Category = "RigVM Blueprint")
+	static TArray<URigVMBlueprint*> LoadAssetsByClass(TSubclassOf<URigVMBlueprint> InClass);
+
+	UFUNCTION(BlueprintCallable, Category = "RigVM Blueprint", meta = (DisplayName = "LoadAssetsWithBlueprintFilter", ScriptName = "LoadAssetsWithBlueprintFilter"))
+		static TArray<URigVMBlueprint*> LoadAssetsWithBlueprintFilter_ForBlueprint(
+		TSubclassOf<URigVMBlueprint> InClass,
+		FRigVMBlueprintFilterDynamic InBlueprintFilter);
+
+	UFUNCTION(BlueprintCallable, Category = "RigVM Blueprint", meta = (DisplayName = "LoadAssetsWithAssetDataFilter", ScriptName = "LoadAssetsWithAssetDataFilter"))
+		static TArray<URigVMBlueprint*> LoadAssetsWithAssetDataFilter_ForBlueprint(
+		TSubclassOf<URigVMBlueprint> InClass,
+		FRigVMAssetDataFilterDynamic InAssetDataFilter);
+
+	UFUNCTION(BlueprintCallable, Category = "RigVM Blueprint", meta = (DisplayName = "LoadAssetsWithNodeFilter", ScriptName = "LoadAssetsWithNodeFilter"))
+	static TArray<URigVMBlueprint*> LoadAssetsWithNodeFilter_ForBlueprint(
+		TSubclassOf<URigVMBlueprint> InClass,
+		FRigVMNodeFilterDynamic InNodeFilter);
+
+	static TArray<URigVMBlueprint*> LoadAssetsWithNodeFilter(
+		TSubclassOf<URigVMBlueprint> InClass,
+		FRigVMNodeFilter InNodeFilter);
+
+	UFUNCTION(BlueprintCallable, Category = "RigVM Blueprint", meta = (DisplayName = "LoadAssetsWithAssetDataAndBlueprintFilters", ScriptName = "LoadAssetsWithAssetDataAndBlueprintFilters"))
+	static TArray<URigVMBlueprint*> LoadAssetsWithAssetDataAndBlueprintFilters_ForBlueprint(
+		TSubclassOf<URigVMBlueprint> InClass,
+		FRigVMAssetDataFilterDynamic InAssetDataFilter,
+		FRigVMBlueprintFilterDynamic InBlueprintFilter);
+
+	static TArray<URigVMBlueprint*> LoadAssetsWithAssetDataAndBlueprintFilters(
+		TSubclassOf<URigVMBlueprint> InClass,
+		FRigVMAssetDataFilter InAssetDataFilter,
+		FRigVMBlueprintFilter InBlueprintFilter);
+
+	UFUNCTION(BlueprintCallable, Category = "RigVM Blueprint", meta = (DisplayName = "LoadAssetsWithAssetDataAndNodeFilters", ScriptName = "LoadAssetsWithAssetDataAndNodeFilters"))
+	static TArray<URigVMBlueprint*> LoadAssetsWithAssetDataAndNodeFilters_ForBlueprint(
+		TSubclassOf<URigVMBlueprint> InClass,
+		FRigVMAssetDataFilterDynamic InAssetDataFilter,
+		FRigVMNodeFilterDynamic InNodeFilter);
+
+	static TArray<URigVMBlueprint*> LoadAssetsWithAssetDataAndNodeFilters(
+		TSubclassOf<URigVMBlueprint> InClass,
+		FRigVMAssetDataFilter InAssetDataFilter,
+		FRigVMNodeFilter InNodeFilter);
+
+	UFUNCTION(BlueprintCallable, Category = "RigVM Blueprint", meta = (DisplayName = "GetAssetsWithFilter", ScriptName = "GetAssetsWithFilter"))
+	static TArray<FAssetData> GetAssetsWithFilter_ForBlueprint(
+		TSubclassOf<URigVMBlueprint> InClass,
+		FRigVMAssetDataFilterDynamic InAssetDataFilter);
+
+	static TArray<FAssetData> GetAssetsWithFilter(
+		TSubclassOf<URigVMBlueprint> InClass,
+		FRigVMAssetDataFilter InAssetDataFilter);
+};
