ControlRig : Fix for crash when loading a rig that has compile errors

diff --git a/Engine/Plugins/Runtime/RigVM/Source/RigVM/Private/RigVMHost.cpp b/Engine/Plugins/Runtime/RigVM/Source/RigVM/Private/RigVMHost.cpp
--- a/Engine/Plugins/Runtime/RigVM/Source/RigVM/Private/RigVMHost.cpp
+++ b/Engine/Plugins/Runtime/RigVM/Source/RigVM/Private/RigVMHost.cpp
@@ -253,8 +253,14 @@ void URigVMHost::Initialize(bool bRequestInit)
 
 bool URigVMHost::InitializeVM(const FName& InEventName)
 {
+	const TArray<FRigVMExternalVariable> ExternalVariables = GetExternalVariablesImpl(false);
+	if (VM->GetExternalVariableDefs().Num() != ExternalVariables.Num())
+	{
+		return false;	// The rig did compile with errors
+	}
+
 	// update the VM's external variables
-	VM->SetExternalVariablesInstanceData(GetExtendedExecuteContext(), GetExternalVariablesImpl(false));
+	VM->SetExternalVariablesInstanceData(GetExtendedExecuteContext(), ExternalVariables);
 
 	TArray<URigVMMemoryStorage*> LocalMemory = VM->GetLocalMemoryArray(ExtendedExecuteContext);
 	const bool bResult = VM->InitializeInstance(GetExtendedExecuteContext(), LocalMemory);
