RigVM: Improve performance of RigVM::ExecuteBranch

diff --git a/Engine/Plugins/Runtime/RigVM/Source/RigVM/Private/RigVMCore/RigVM.cpp b/Engine/Plugins/Runtime/RigVM/Source/RigVM/Private/RigVMCore/RigVM.cpp
--- a/Engine/Plugins/Runtime/RigVM/Source/RigVM/Private/RigVMCore/RigVM.cpp
+++ b/Engine/Plugins/Runtime/RigVM/Source/RigVM/Private/RigVMCore/RigVM.cpp
@@ -2088,8 +2088,20 @@ bool URigVM::Execute(FRigVMExtendedExecuteContext& Context, const FName& InEntry
 
 ERigVMExecuteResult URigVM::ExecuteBranch(FRigVMExtendedExecuteContext& Context, const FRigVMBranchInfo& InBranchToRun)
 {
-	// likely have to optimize this
-	TGuardValue<FRigVMExtendedExecuteContext> ContextGuard(Context, Context);
+	// Maintain all settings on the context - to be reset once the branch has executed. 
+	FRigVMExecuteContext& PublicContext = Context.GetPublicData<>();
+	TGuardValue<double> LastExecutionMicroSecondsGuard(Context.LastExecutionMicroSeconds, Context.LastExecutionMicroSeconds);
+	TGuardValue<uint32> NumExecutionsGuard(Context.NumExecutions, Context.NumExecutions);
+	TGuardValue<ERigVMExecuteResult> CurrentExecuteResultGuard(Context.CurrentExecuteResult, Context.CurrentExecuteResult);
+	TGuardValue<FName> CurrentEntryNameGuard(Context.CurrentEntryName, Context.CurrentEntryName);
+	TGuardValue<bool> bCurrentlyRunningRootEntryGuard(Context.bCurrentlyRunningRootEntry, Context.bCurrentlyRunningRootEntry);
+	TGuardValue<FName> EventNameGuard(PublicContext.EventName, PublicContext.EventName);
+	TGuardValue<FName> FunctionNameGuard(PublicContext.FunctionName, PublicContext.FunctionName);
+	TGuardValue<uint16> InstructionIndexGuard(PublicContext.InstructionIndex, PublicContext.InstructionIndex);
+	TGuardValue<double> DeltaTimeGuard(PublicContext.DeltaTime, PublicContext.DeltaTime);
+	TGuardValue<double> AbsoluteTimeGuard(PublicContext.AbsoluteTime, PublicContext.AbsoluteTime);
+	TGuardValue<double> FramesPerSecondGuard(PublicContext.FramesPerSecond, PublicContext.FramesPerSecond);
+
 	return ExecuteInstructions(Context, InBranchToRun.FirstInstruction, InBranchToRun.LastInstruction);
 }
 
