Fix for CIS crash on content workers

#rnx

diff --git a/Engine/Plugins/Runtime/RigVM/Source/RigVM/Private/RigVMHost.cpp b/Engine/Plugins/Runtime/RigVM/Source/RigVM/Private/RigVMHost.cpp
--- a/Engine/Plugins/Runtime/RigVM/Source/RigVM/Private/RigVMHost.cpp
+++ b/Engine/Plugins/Runtime/RigVM/Source/RigVM/Private/RigVMHost.cpp
@@ -1020,21 +1020,24 @@ void URigVMHost::PostInitInstance(URigVMHost* InCDO)
 
 void URigVMHost::GenerateUserDefinedDependenciesData(FRigVMExtendedExecuteContext& Context)
 {
-	const TArray<const UObject*> UserDefinedDependencies = GetUserDefinedDependencies({ GetLiteralMemory(), GetWorkMemory() });
-	UserDefinedStructGuidToPathName.Reset();
-	UserDefinedEnumToPathName.Reset();
-
-	for (const UObject* UserDefinedDependency : UserDefinedDependencies)
+	if (VM)
 	{
-		if (const UUserDefinedStruct* UserDefinedStruct = Cast<UUserDefinedStruct>(UserDefinedDependency))
-		{
-			const FString GuidBasedName = RigVMTypeUtils::GetUniqueStructTypeName(UserDefinedStruct);
-			UserDefinedStructGuidToPathName.Add(GuidBasedName, UserDefinedStruct);
-		}
-		else if (const UUserDefinedEnum* UserDefinedEnum = Cast<UUserDefinedEnum>(UserDefinedDependency))
+		const TArray<const UObject*> UserDefinedDependencies = GetUserDefinedDependencies({ GetLiteralMemory(), GetWorkMemory() });
+		UserDefinedStructGuidToPathName.Reset();
+		UserDefinedEnumToPathName.Reset();
+
+		for (const UObject* UserDefinedDependency : UserDefinedDependencies)
 		{
-			const FString EnumName = RigVMTypeUtils::CPPTypeFromEnum(UserDefinedEnum);
-			UserDefinedEnumToPathName.Add(EnumName, UserDefinedEnum);
+			if (const UUserDefinedStruct* UserDefinedStruct = Cast<UUserDefinedStruct>(UserDefinedDependency))
+			{
+				const FString GuidBasedName = RigVMTypeUtils::GetUniqueStructTypeName(UserDefinedStruct);
+				UserDefinedStructGuidToPathName.Add(GuidBasedName, UserDefinedStruct);
+			}
+			else if (const UUserDefinedEnum* UserDefinedEnum = Cast<UUserDefinedEnum>(UserDefinedDependency))
+			{
+				const FString EnumName = RigVMTypeUtils::CPPTypeFromEnum(UserDefinedEnum);
+				UserDefinedEnumToPathName.Add(EnumName, UserDefinedEnum);
+			}
 		}
 	}
 }
