RigVMCompiler: Address nullptr for controller during compile function

Relevant for 5.3.1

diff --git a/Engine/Plugins/Runtime/RigVM/Source/RigVMDeveloper/Private/RigVMCompiler/RigVMCompiler.cpp b/Engine/Plugins/Runtime/RigVM/Source/RigVMDeveloper/Private/RigVMCompiler/RigVMCompiler.cpp
--- a/Engine/Plugins/Runtime/RigVM/Source/RigVMDeveloper/Private/RigVMCompiler/RigVMCompiler.cpp
+++ b/Engine/Plugins/Runtime/RigVM/Source/RigVMDeveloper/Private/RigVMCompiler/RigVMCompiler.cpp
@@ -499,7 +499,7 @@ bool URigVMCompiler::Compile(const FRigVMCompileSettings& InSettings, TArray<URi
 								if (URigVMLibraryNode* LibraryNode = Cast<URigVMLibraryNode>(FunctionData->Header.LibraryPointer.LibraryNode.TryLoad()))
 								{
 									IRigVMClientHost* ClientHost = LibraryNode->GetImplementingOuter<IRigVMClientHost>();
-									URigVMController* FunctionController = ClientHost->GetRigVMClient()->GetController(LibraryNode->GetLibrary());
+									URigVMController* FunctionController = ClientHost->GetRigVMClient()->GetOrCreateController(LibraryNode->GetLibrary());
 									bSuccessfullCompilation = CompileFunction(WorkData.Settings, LibraryNode, FunctionController, &FunctionData->CompilationData, OutVMContext);
 								}
 								else
