RigVM: initialize all archetype instances after compilation (also after bp reinstantiate)

diff --git a/Engine/Plugins/Runtime/RigVM/Source/RigVMDeveloper/Private/RigVMBlueprint.cpp b/Engine/Plugins/Runtime/RigVM/Source/RigVMDeveloper/Private/RigVMBlueprint.cpp
--- a/Engine/Plugins/Runtime/RigVM/Source/RigVMDeveloper/Private/RigVMBlueprint.cpp
+++ b/Engine/Plugins/Runtime/RigVM/Source/RigVMDeveloper/Private/RigVMBlueprint.cpp
@@ -1213,22 +1213,7 @@ void URigVMBlueprint::RecompileVM()
 			return;
 		}
 
-		TArray<UObject*> ArchetypeInstances;
-		CDO->GetArchetypeInstances(ArchetypeInstances);
-		for (UObject* Instance : ArchetypeInstances)
-		{
-			if (URigVMHost* InstanceHost = Cast<URigVMHost>(Instance))
-			{
-				// No objects should be created during load, so PostInitInstanceIfRequired, which creates a new VM and
-				// DynamicHierarchy, should not be called during load
-				if (!InstanceHost->HasAllFlags(RF_NeedPostLoad))
-				{
-					InstanceHost->PostInitInstanceIfRequired();
-				}
-				InstanceHost->InstantiateVMFromCDO();
-				InstanceHost->CopyExternalVariableDefaultValuesFromCDO();
-			}
-		}
+		InitializeArchetypeInstances();
 
 		bVMRecompilationRequired = false;
 		VMCompiledEvent.Broadcast(this, CDO->GetVM(), CDO->GetExtendedExecuteContext());
@@ -3636,6 +3621,36 @@ void URigVMBlueprint::PropagateRuntimeSettingsFromBPToInstances()
 	}
 }
 
+void URigVMBlueprint::InitializeArchetypeInstances()
+{
+	URigVMBlueprintGeneratedClass* RigClass = GetRigVMBlueprintGeneratedClass();
+	if(RigClass == nullptr)
+	{
+		return;
+	}
+
+	URigVMHost* CDO = Cast<URigVMHost>(RigClass->GetDefaultObject(true /* create if needed */));
+	if (CDO && CDO->VM != nullptr)
+	{
+		TArray<UObject*> ArchetypeInstances;
+		CDO->GetArchetypeInstances(ArchetypeInstances);
+		for (UObject* Instance : ArchetypeInstances)
+		{
+			if (URigVMHost* InstanceHost = Cast<URigVMHost>(Instance))
+			{
+				// No objects should be created during load, so PostInitInstanceIfRequired, which creates a new VM and
+				// DynamicHierarchy, should not be called during load
+				if (!InstanceHost->HasAllFlags(RF_NeedPostLoad))
+				{
+					InstanceHost->PostInitInstanceIfRequired();
+				}
+				InstanceHost->InstantiateVMFromCDO();
+				InstanceHost->CopyExternalVariableDefaultValuesFromCDO();
+			}
+		}
+	}
+}
+
 #if WITH_EDITOR
 
 void URigVMBlueprint::OnPreVariableChange(UObject* InObject)
diff --git a/Engine/Plugins/Runtime/RigVM/Source/RigVMDeveloper/Public/RigVMBlueprint.h b/Engine/Plugins/Runtime/RigVM/Source/RigVMDeveloper/Public/RigVMBlueprint.h
--- a/Engine/Plugins/Runtime/RigVM/Source/RigVMDeveloper/Public/RigVMBlueprint.h
+++ b/Engine/Plugins/Runtime/RigVM/Source/RigVMDeveloper/Public/RigVMBlueprint.h
@@ -588,6 +588,7 @@ protected:
 
 public:
 	void PropagateRuntimeSettingsFromBPToInstances();
+	void InitializeArchetypeInstances();
 
 private:
 
diff --git a/Engine/Plugins/Runtime/RigVM/Source/RigVMEditor/Private/Editor/RigVMEditor.cpp b/Engine/Plugins/Runtime/RigVM/Source/RigVMEditor/Private/Editor/RigVMEditor.cpp
--- a/Engine/Plugins/Runtime/RigVM/Source/RigVMEditor/Private/Editor/RigVMEditor.cpp
+++ b/Engine/Plugins/Runtime/RigVM/Source/RigVMEditor/Private/Editor/RigVMEditor.cpp
@@ -1039,6 +1039,7 @@ void FRigVMEditor::Compile()
 		{
 			TGuardValue<bool> GuardCompileReEntry(bIsCompilingThroughUI, true);
 			FBlueprintEditor::Compile();
+			RigVMBlueprint->InitializeArchetypeInstances();
 			UpdateRigVMHost();
 		}
 
