RigVM: Refactor PDI drawing code to host

diff --git a/Engine/Plugins/Animation/ControlRig/Source/ControlRigEditor/Private/EditMode/ControlRigEditMode.cpp b/Engine/Plugins/Animation/ControlRig/Source/ControlRigEditor/Private/EditMode/ControlRigEditMode.cpp
--- a/Engine/Plugins/Animation/ControlRig/Source/ControlRigEditor/Private/EditMode/ControlRigEditMode.cpp
+++ b/Engine/Plugins/Animation/ControlRig/Source/ControlRigEditor/Private/EditMode/ControlRigEditMode.cpp
@@ -898,56 +898,8 @@ void FControlRigEditMode::Render(const FSceneView* View, FViewport* Viewport, FP
 						}
 					}
 				}
-				for (const FRigVMDrawInstruction& Instruction : ControlRig->DrawInterface)
-				{
-					if (!Instruction.IsValid())
-					{
-						continue;
-					}
-
-					FTransform InstructionTransform = Instruction.Transform * ComponentTransform;
-					switch (Instruction.PrimitiveType)
-					{
-					case ERigVMDrawSettings::Points:
-					{
-						for (const FVector& Point : Instruction.Positions)
-						{
-							PDI->DrawPoint(InstructionTransform.TransformPosition(Point), Instruction.Color, Instruction.Thickness, SDPG_Foreground);
-						}
-						break;
-					}
-					case ERigVMDrawSettings::Lines:
-					{
-						const TArray<FVector>& Points = Instruction.Positions;
-						PDI->AddReserveLines(SDPG_Foreground, Points.Num() / 2, false, Instruction.Thickness > SMALL_NUMBER);
-						for (int32 PointIndex = 0; PointIndex < Points.Num() - 1; PointIndex += 2)
-						{
-							PDI->DrawLine(InstructionTransform.TransformPosition(Points[PointIndex]), InstructionTransform.TransformPosition(Points[PointIndex + 1]), Instruction.Color, SDPG_Foreground, Instruction.Thickness);
-						}
-						break;
-					}
-					case ERigVMDrawSettings::LineStrip:
-					{
-						const TArray<FVector>& Points = Instruction.Positions;
-						PDI->AddReserveLines(SDPG_Foreground, Points.Num() - 1, false, Instruction.Thickness > SMALL_NUMBER);
-						for (int32 PointIndex = 0; PointIndex < Points.Num() - 1; PointIndex++)
-						{
-							PDI->DrawLine(InstructionTransform.TransformPosition(Points[PointIndex]), InstructionTransform.TransformPosition(Points[PointIndex + 1]), Instruction.Color, SDPG_Foreground, Instruction.Thickness);
-						}
-						break;
-					}
 
-					case ERigVMDrawSettings::DynamicMesh:
-					{
-						FDynamicMeshBuilder MeshBuilder(PDI->View->GetFeatureLevel());
-						MeshBuilder.AddVertices(Instruction.MeshVerts);
-						MeshBuilder.AddTriangles(Instruction.MeshIndices);
-						MeshBuilder.Draw(PDI, InstructionTransform.ToMatrixWithScale(), Instruction.MaterialRenderProxy, SDPG_World/*SDPG_Foreground*/);
-						break;
-					}
-
-					}
-				}
+				ControlRig->DrawIntoPDI(PDI, ComponentTransform);
 			}
 		}
 	}
diff --git a/Engine/Plugins/Runtime/RigVM/Source/RigVM/Private/RigVMHost.cpp b/Engine/Plugins/Runtime/RigVM/Source/RigVM/Private/RigVMHost.cpp
--- a/Engine/Plugins/Runtime/RigVM/Source/RigVM/Private/RigVMHost.cpp
+++ b/Engine/Plugins/Runtime/RigVM/Source/RigVM/Private/RigVMHost.cpp
@@ -9,6 +9,7 @@
 #include "UObject/UE5MainStreamObjectVersion.h"
 #include "UObject/UObjectIterator.h"
 #include "UObject/ObjectSaveContext.h"
+#include "SceneView.h"
 
 #if WITH_EDITOR
 #include "Kismet2/BlueprintEditorUtils.h"
@@ -823,6 +824,62 @@ const URigVMMemoryStorage* URigVMHost::GetMemoryByType(ERigVMMemoryType InMemory
 	return VM->GetMemoryByType(ExtendedExecuteContext, InMemoryType);
 }
 
+void URigVMHost::DrawIntoPDI(FPrimitiveDrawInterface* PDI, const FTransform& InTransform)
+{
+	for (const FRigVMDrawInstruction& Instruction : DrawInterface)
+	{
+		if (!Instruction.IsValid())
+		{
+			continue;
+		}
+
+		FTransform InstructionTransform = Instruction.Transform * InTransform;
+		switch (Instruction.PrimitiveType)
+		{
+			case ERigVMDrawSettings::Points:
+			{
+				for (const FVector& Point : Instruction.Positions)
+				{
+					PDI->DrawPoint(InstructionTransform.TransformPosition(Point), Instruction.Color, Instruction.Thickness, SDPG_Foreground);
+				}
+				break;
+			}
+			case ERigVMDrawSettings::Lines:
+			{
+				const TArray<FVector>& Points = Instruction.Positions;
+				PDI->AddReserveLines(SDPG_Foreground, Points.Num() / 2, false, Instruction.Thickness > SMALL_NUMBER);
+				for (int32 PointIndex = 0; PointIndex < Points.Num() - 1; PointIndex += 2)
+				{
+					PDI->DrawLine(InstructionTransform.TransformPosition(Points[PointIndex]), InstructionTransform.TransformPosition(Points[PointIndex + 1]), Instruction.Color, SDPG_Foreground, Instruction.Thickness);
+				}
+				break;
+			}
+			case ERigVMDrawSettings::LineStrip:
+			{
+				const TArray<FVector>& Points = Instruction.Positions;
+				PDI->AddReserveLines(SDPG_Foreground, Points.Num() - 1, false, Instruction.Thickness > SMALL_NUMBER);
+				for (int32 PointIndex = 0; PointIndex < Points.Num() - 1; PointIndex++)
+				{
+					PDI->DrawLine(InstructionTransform.TransformPosition(Points[PointIndex]), InstructionTransform.TransformPosition(Points[PointIndex + 1]), Instruction.Color, SDPG_Foreground, Instruction.Thickness);
+				}
+				break;
+			}
+			case ERigVMDrawSettings::DynamicMesh:
+			{
+				FDynamicMeshBuilder MeshBuilder(PDI->View->GetFeatureLevel());
+				MeshBuilder.AddVertices(Instruction.MeshVerts);
+				MeshBuilder.AddTriangles(Instruction.MeshIndices);
+				MeshBuilder.Draw(PDI, InstructionTransform.ToMatrixWithScale(), Instruction.MaterialRenderProxy, SDPG_World/*SDPG_Foreground*/);
+				break;
+			}
+			default:
+			{
+				break;
+			}
+		}
+	}
+}
+
 USceneComponent* URigVMHost::GetOwningSceneComponent()
 {
 	return GetTypedOuter<USceneComponent>();
diff --git a/Engine/Plugins/Runtime/RigVM/Source/RigVM/Public/RigVMHost.h b/Engine/Plugins/Runtime/RigVM/Source/RigVM/Public/RigVMHost.h
--- a/Engine/Plugins/Runtime/RigVM/Source/RigVM/Public/RigVMHost.h
+++ b/Engine/Plugins/Runtime/RigVM/Source/RigVM/Public/RigVMHost.h
@@ -9,6 +9,7 @@
 #include "RigVMCore/RigVMGraphFunctionHost.h"
 #include "Engine/AssetUserData.h"
 #include "Interfaces/Interface_AssetUserData.h"
+#include "SceneManagement.h"
 #include "RigVMHost.generated.h"
 
 // set this to something larger than 0 to profile N runs
@@ -250,6 +251,8 @@ public:
 	const FRigVMDrawContainer& GetDrawContainer() const { return DrawContainer; };
 	FRigVMDrawContainer& GetDrawContainer() { return DrawContainer; };
 
+	void DrawIntoPDI(FPrimitiveDrawInterface* PDI, const FTransform& InTransform);
+
 	virtual USceneComponent* GetOwningSceneComponent(); 
 
 	virtual void PostInitInstanceIfRequired() {};
