Control Rig: Move patching of shared arguments to PostLoad.

#rnx

diff --git a/Engine/Plugins/Runtime/RigVM/Source/RigVM/Private/RigVMBlueprintGeneratedClass.cpp b/Engine/Plugins/Runtime/RigVM/Source/RigVM/Private/RigVMBlueprintGeneratedClass.cpp
--- a/Engine/Plugins/Runtime/RigVM/Source/RigVM/Private/RigVMBlueprintGeneratedClass.cpp
+++ b/Engine/Plugins/Runtime/RigVM/Source/RigVM/Private/RigVMBlueprintGeneratedClass.cpp
@@ -93,3 +93,9 @@ void URigVMBlueprintGeneratedClass::Serialize(FArchive& Ar)
 	Ar << GraphFunctionStore;
 }
 
+void URigVMBlueprintGeneratedClass::PostLoad()
+{
+	Super::PostLoad();
+
+	GraphFunctionStore.PostLoad();
+}
diff --git a/Engine/Plugins/Runtime/RigVM/Source/RigVM/Public/RigVMBlueprintGeneratedClass.h b/Engine/Plugins/Runtime/RigVM/Source/RigVM/Public/RigVMBlueprintGeneratedClass.h
--- a/Engine/Plugins/Runtime/RigVM/Source/RigVM/Public/RigVMBlueprintGeneratedClass.h
+++ b/Engine/Plugins/Runtime/RigVM/Source/RigVM/Public/RigVMBlueprintGeneratedClass.h
@@ -21,7 +21,8 @@ public:
 	virtual void PostInitInstance(UObject* InObj, FObjectInstancingGraph* InstanceGraph) override;
 
 	// UObject interface
-	void Serialize(FArchive& Ar);
+	virtual void Serialize(FArchive& Ar) override;
+	virtual void PostLoad() override;
 
 	// IRigVMGraphFunctionHost interface
 	virtual const FRigVMGraphFunctionStore* GetRigVMGraphFunctionStore() const override { return &GraphFunctionStore; }
diff --git a/Engine/Plugins/Runtime/RigVM/Source/RigVM/Public/RigVMCore/RigVMGraphFunctionHost.h b/Engine/Plugins/Runtime/RigVM/Source/RigVM/Public/RigVMCore/RigVMGraphFunctionHost.h
--- a/Engine/Plugins/Runtime/RigVM/Source/RigVM/Public/RigVMCore/RigVMGraphFunctionHost.h
+++ b/Engine/Plugins/Runtime/RigVM/Source/RigVM/Public/RigVMCore/RigVMGraphFunctionHost.h
@@ -81,15 +81,15 @@ public:
 			Ar << Host.PrivateFunctions;
 		}
 
-		if(Ar.IsLoading())
+		return Ar;
+	}
+
+	void PostLoad()
+	{
+		for(FRigVMGraphFunctionData& Function: PublicFunctions)
 		{
-			for(FRigVMGraphFunctionData& Function : Host.PublicFunctions)
-			{
-				Function.PatchSharedArgumentOperandsIfRequired();
-			}
+			Function.PatchSharedArgumentOperandsIfRequired();
 		}
-		
-		return Ar;
 	}
 
 private:
