Control Rig: Store a temporary hard reference to resolved user-defined types on the VM to ensure they don't get GC'd.

#rnx

diff --git a/Engine/Plugins/Runtime/RigVM/Source/RigVM/Private/RigVMCore/RigVM.cpp b/Engine/Plugins/Runtime/RigVM/Source/RigVM/Private/RigVMCore/RigVM.cpp
--- a/Engine/Plugins/Runtime/RigVM/Source/RigVM/Private/RigVMCore/RigVM.cpp
+++ b/Engine/Plugins/Runtime/RigVM/Source/RigVM/Private/RigVMCore/RigVM.cpp
@@ -906,18 +906,20 @@ bool URigVM::ResolveFunctionsIfRequired()
 		GetFactories().Reset();
 		GetFactories().SetNumZeroed(GetFunctionNames().Num());
 
-		FRigVMTypeResolvalInfo ResolvalInfo;
+		FRigVMUserDefinedTypeResolver TypeResolver;
 		if (const URigVMHost* HostCDO = GetHostCDO())
 		{
-			ResolvalInfo.CPPTypeToObjectPath = HostCDO->GetUserDefinedStructGuidToPathName();
-			ResolvalInfo.CPPTypeToObjectPath.Append(HostCDO->GetUserDefinedEnumToPathName());
+			TypeResolver = FRigVMUserDefinedTypeResolver([HostCDO](const FString& InTypeName) -> UObject*
+			{
+				return HostCDO->ResolveUserDefinedTypeById(InTypeName);
+			});
 		}
 
 		TArray<FName>& FunctionNames = GetFunctionNames();
 		for (int32 FunctionIndex = 0; FunctionIndex < FunctionNames.Num(); FunctionIndex++)
 		{
 			const FString FunctionNameString = FunctionNames[FunctionIndex].ToString();
-			if(const FRigVMFunction* Function = FRigVMRegistry::Get().FindFunction(*FunctionNameString, ResolvalInfo))
+			if(const FRigVMFunction* Function = FRigVMRegistry::Get().FindFunction(*FunctionNameString, TypeResolver))
 			{
 				GetFunctions()[FunctionIndex] = Function;
 				GetFactories()[FunctionIndex] = Function->Factory;
diff --git a/Engine/Plugins/Runtime/RigVM/Source/RigVM/Private/RigVMCore/RigVMRegistry.cpp b/Engine/Plugins/Runtime/RigVM/Source/RigVM/Private/RigVMCore/RigVMRegistry.cpp
--- a/Engine/Plugins/Runtime/RigVM/Source/RigVM/Private/RigVMCore/RigVMRegistry.cpp
+++ b/Engine/Plugins/Runtime/RigVM/Source/RigVM/Private/RigVMCore/RigVMRegistry.cpp
@@ -842,7 +842,7 @@ const FRigVMTemplateArgumentType& FRigVMRegistry::GetType(TRigVMTypeIndex InType
 	return EmptyType;
 }
 
-const FRigVMTemplateArgumentType& FRigVMRegistry::FindTypeFromCPPType(const FString& InCPPType, const FRigVMTypeResolvalInfo* InResolvalInfo) const
+const FRigVMTemplateArgumentType& FRigVMRegistry::FindTypeFromCPPType(const FString& InCPPType) const
 {
 	const int32 TypeIndex = GetTypeIndexFromCPPType(InCPPType);
 	if(ensure(Types.IsValidIndex(TypeIndex)))
@@ -854,7 +854,7 @@ const FRigVMTemplateArgumentType& FRigVMRegistry::FindTypeFromCPPType(const FStr
 	return EmptyType;
 }
 
-TRigVMTypeIndex FRigVMRegistry::GetTypeIndexFromCPPType(const FString& InCPPType, const FRigVMTypeResolvalInfo* InResolvalInfo) const
+TRigVMTypeIndex FRigVMRegistry::GetTypeIndexFromCPPType(const FString& InCPPType) const
 {
 	TRigVMTypeIndex Result = INDEX_NONE;
 	if(ensure(!InCPPType.IsEmpty()))
@@ -1309,13 +1309,13 @@ void FRigVMRegistry::RegisterObjectTypes(TConstArrayView<UClass*> InClasses)
 	}
 }
 
-const FRigVMFunction* FRigVMRegistry::FindFunction(const TCHAR* InName, const FRigVMTypeResolvalInfo& InResolvalInfo) const
+const FRigVMFunction* FRigVMRegistry::FindFunction(const TCHAR* InName, const FRigVMUserDefinedTypeResolver& InTypeResolver) const
 {
 	FScopeLock FindFunctionScopeLock(&FindFunctionMutex);
-	return FindFunction_NoLock(InName, InResolvalInfo);
+	return FindFunction_NoLock(InName, InTypeResolver);
 }
 
-const FRigVMFunction* FRigVMRegistry::FindFunction_NoLock(const TCHAR* InName, const FRigVMTypeResolvalInfo& InResolvalInfo) const
+const FRigVMFunction* FRigVMRegistry::FindFunction_NoLock(const TCHAR* InName, const FRigVMUserDefinedTypeResolver& InTypeResolver) const
 {
 	// Check first if the function is provided by internally registered rig units. 
 	if(const int32* FunctionIndexPtr = FunctionNameToIndex.Find(InName))
@@ -1333,7 +1333,7 @@ const FRigVMFunction* FRigVMRegistry::FindFunction_NoLock(const TCHAR* InName, c
 		{
 			if(const FRigVMTemplate* Template = Factory->GetTemplate())
 			{
-				const FRigVMTemplateTypeMap ArgumentTypes = Template->GetArgumentTypesFromString(SuffixString, &InResolvalInfo);
+				const FRigVMTemplateTypeMap ArgumentTypes = Template->GetArgumentTypesFromString(SuffixString, &InTypeResolver);
 				if(ArgumentTypes.Num() == Template->NumArguments())
 				{
 					const int32 PermutationIndex = Template->FindPermutation(ArgumentTypes);
@@ -1376,7 +1376,7 @@ const FRigVMFunction* FRigVMRegistry::FindFunction_NoLock(const TCHAR* InName, c
 				{
 					NewStructOrFactoryName = StructPrefix + NewStructOrFactoryName;
 				}
-				const FRigVMFunction* RedirectedFunction = FindFunction_NoLock(*(NewStructOrFactoryName + TEXT("::") + SuffixString), InResolvalInfo);
+				const FRigVMFunction* RedirectedFunction = FindFunction_NoLock(*(NewStructOrFactoryName + TEXT("::") + SuffixString), InTypeResolver);
 				if(RedirectedFunction)
 				{
 					FRigVMRegistry& MutableRegistry = FRigVMRegistry::Get();
@@ -1390,7 +1390,7 @@ const FRigVMFunction* FRigVMRegistry::FindFunction_NoLock(const TCHAR* InName, c
 	return nullptr;
 }
 
-const FRigVMFunction* FRigVMRegistry::FindFunction(UScriptStruct* InStruct, const TCHAR* InName, const FRigVMTypeResolvalInfo& InResolvalInfo) const
+const FRigVMFunction* FRigVMRegistry::FindFunction(UScriptStruct* InStruct, const TCHAR* InName, const FRigVMUserDefinedTypeResolver& InResolvalInfo) const
 {
 	check(InStruct);
 	check(InName);
diff --git a/Engine/Plugins/Runtime/RigVM/Source/RigVM/Private/RigVMCore/RigVMTemplate.cpp b/Engine/Plugins/Runtime/RigVM/Source/RigVM/Private/RigVMCore/RigVMTemplate.cpp
--- a/Engine/Plugins/Runtime/RigVM/Source/RigVM/Private/RigVMCore/RigVMTemplate.cpp
+++ b/Engine/Plugins/Runtime/RigVM/Source/RigVM/Private/RigVMCore/RigVMTemplate.cpp
@@ -654,7 +654,7 @@ TArray<UStruct*> FRigVMTemplate::GetSuperStructs(UStruct* InStruct, bool bInclud
 	return SuperStructs;
 }
 
-FRigVMTemplate::FTypeMap FRigVMTemplate::GetArgumentTypesFromString(const FString& InTypeString, const FRigVMTypeResolvalInfo* InResolvalInfo) const
+FRigVMTemplate::FTypeMap FRigVMTemplate::GetArgumentTypesFromString(const FString& InTypeString, const FRigVMUserDefinedTypeResolver* InTypeResolver) const
 {
 	FTypeMap Types;
 	if(!InTypeString.IsEmpty())
@@ -675,12 +675,12 @@ FRigVMTemplate::FTypeMap FRigVMTemplate::GetArgumentTypesFromString(const FStrin
 			{
 				if(const FRigVMTemplateArgument* Argument = FindArgument(*ArgumentName))
 				{
-					TRigVMTypeIndex TypeIndex = Registry.GetTypeIndexFromCPPType(TypeName, InResolvalInfo);
+					TRigVMTypeIndex TypeIndex = Registry.GetTypeIndexFromCPPType(TypeName);
 
 					// If the type was not found, check if it's a user-defined type that hasn't been registered yet.
 					if (TypeIndex == INDEX_NONE && RigVMTypeUtils::RequiresCPPTypeObject(TypeName))
 					{
-						UObject* CPPTypeObject = RigVMTypeUtils::ObjectFromCPPType(TypeName, true, InResolvalInfo);
+						UObject* CPPTypeObject = RigVMTypeUtils::ObjectFromCPPType(TypeName, true, InTypeResolver);
 						
 						FRigVMTemplateArgumentType ArgType(*TypeName, CPPTypeObject);
 						TypeIndex = Registry.FindOrAddType(ArgType);
diff --git a/Engine/Plugins/Runtime/RigVM/Source/RigVM/Private/RigVMHost.cpp b/Engine/Plugins/Runtime/RigVM/Source/RigVM/Private/RigVMHost.cpp
--- a/Engine/Plugins/Runtime/RigVM/Source/RigVM/Private/RigVMHost.cpp
+++ b/Engine/Plugins/Runtime/RigVM/Source/RigVM/Private/RigVMHost.cpp
@@ -1001,6 +1001,30 @@ const FString URigVMHost::GetDebugExecutionString()
 #endif
 #endif
 
+
+UObject* URigVMHost::ResolveUserDefinedTypeById(const FString& InTypeName) const
+{
+	const FSoftObjectPath* ResultPathPtr = UserDefinedStructGuidToPathName.Find(InTypeName);
+	if (ResultPathPtr == nullptr)
+	{
+		ResultPathPtr = UserDefinedEnumToPathName.Find(InTypeName);
+	}
+
+	if (ResultPathPtr == nullptr)
+	{
+		return nullptr;
+	}
+
+	if (UObject* TypeObject = ResultPathPtr->TryLoad())
+	{
+		// Ensure we have a hold on this type so it doesn't get nixed on the next GC.
+		const_cast<URigVMHost*>(this)->UserDefinedTypesInUse.Add(TypeObject);
+		return TypeObject;
+	}
+
+	return nullptr;
+}
+
 void URigVMHost::PostInitInstance(URigVMHost* InCDO)
 {
 	const EObjectFlags SubObjectFlags =
@@ -1082,6 +1106,7 @@ void URigVMHost::GenerateUserDefinedDependenciesData(FRigVMExtendedExecuteContex
 		const TArray<const UObject*> UserDefinedDependencies = GetUserDefinedDependencies({ GetLiteralMemory(), GetWorkMemory() });
 		UserDefinedStructGuidToPathName.Reset();
 		UserDefinedEnumToPathName.Reset();
+		UserDefinedTypesInUse.Reset();
 
 		for (const UObject* UserDefinedDependency : UserDefinedDependencies)
 		{
diff --git a/Engine/Plugins/Runtime/RigVM/Source/RigVM/Private/RigVMTypeUtils.cpp b/Engine/Plugins/Runtime/RigVM/Source/RigVM/Private/RigVMTypeUtils.cpp
--- a/Engine/Plugins/Runtime/RigVM/Source/RigVM/Private/RigVMTypeUtils.cpp
+++ b/Engine/Plugins/Runtime/RigVM/Source/RigVM/Private/RigVMTypeUtils.cpp
@@ -34,38 +34,28 @@ namespace RigVMTypeUtils
 	TRigVMTypeIndex TypeIndex::WildCardArray = INDEX_NONE;	
 }
 
-UObject* RigVMTypeUtils::UserDefinedTypeFromCPPType(FString& InOutCPPType, const FRigVMTypeResolvalInfo* InResolvalInfo)
+UObject* RigVMTypeUtils::UserDefinedTypeFromCPPType(FString& InOutCPPType, const FRigVMUserDefinedTypeResolver* InTypeResolver)
 {
 	const FString OriginalTypeName = InOutCPPType;
 	UObject* CPPTypeObject = nullptr;
 	InOutCPPType.Reset();
 
 	// try to resolve the type name using a path name potentially
-	if(InOutCPPType.IsEmpty() && InResolvalInfo != nullptr)
+	if(InOutCPPType.IsEmpty() && InTypeResolver != nullptr && InTypeResolver->IsValid())
 	{
-		if(!InResolvalInfo->CPPTypeToObjectPath.IsEmpty())
+		FString TypeNameToLookUp = OriginalTypeName;
+		while(IsArrayType(TypeNameToLookUp))
 		{
-			FString TypeNameToLookUp = OriginalTypeName;
-			while(IsArrayType(TypeNameToLookUp))
-			{
-				TypeNameToLookUp = BaseTypeFromArrayType(TypeNameToLookUp);
-			}
+			TypeNameToLookUp = BaseTypeFromArrayType(TypeNameToLookUp);
+		}
 
-			// this map contains a map of user defined struct keys based on its guid to their path,
-			// for example FUserDefinedStruct_23E408214EE9E6DA5BFADDA0F9F4F577 -> /Game/Animation/MyUserDefinedStruct.MyUserDefinedStruct
-			if(const FSoftObjectPath* ResolveObjectPath = InResolvalInfo->CPPTypeToObjectPath.Find(TypeNameToLookUp))
-			{
-				CPPTypeObject = ResolveObjectPath->ResolveObject();
-				if(CPPTypeObject == nullptr)
-				{
-					CPPTypeObject = ResolveObjectPath->TryLoad();
-				}
-				if(CPPTypeObject)
-				{
-					InOutCPPType = PostProcessCPPType(OriginalTypeName, CPPTypeObject);
-					return CPPTypeObject;
-				}
-			}
+		// Ask the resolver to the name of the user-defined struct/enum to an object.
+		// For example FUserDefinedStruct_23E408214EE9E6DA5BFADDA0F9F4F577 -> /Game/Animation/MyUserDefinedStruct.MyUserDefinedStruct
+		CPPTypeObject = InTypeResolver->GetTypeObjectByName(TypeNameToLookUp); 
+		if(CPPTypeObject)
+		{
+			InOutCPPType = PostProcessCPPType(OriginalTypeName, CPPTypeObject);
+			return CPPTypeObject;
 		}
 	}
 
diff --git a/Engine/Plugins/Runtime/RigVM/Source/RigVM/Public/RigVMCore/RigVMRegistry.h b/Engine/Plugins/Runtime/RigVM/Source/RigVM/Public/RigVMCore/RigVMRegistry.h
--- a/Engine/Plugins/Runtime/RigVM/Source/RigVM/Public/RigVMCore/RigVMRegistry.h
+++ b/Engine/Plugins/Runtime/RigVM/Source/RigVM/Public/RigVMCore/RigVMRegistry.h
@@ -172,10 +172,10 @@ public:
 	int32 NumTypes() const { return Types.Num(); }
 
 	// Returns the type given only its cpp type
-	const FRigVMTemplateArgumentType& FindTypeFromCPPType(const FString& InCPPType, const FRigVMTypeResolvalInfo* InResolvalInfo = nullptr) const;
+	const FRigVMTemplateArgumentType& FindTypeFromCPPType(const FString& InCPPType) const;
 
 	// Returns the type index given only its cpp type
-	TRigVMTypeIndex GetTypeIndexFromCPPType(const FString& InCPPType, const FRigVMTypeResolvalInfo* InResolvalInfo = nullptr) const;
+	TRigVMTypeIndex GetTypeIndexFromCPPType(const FString& InCPPType) const;
 
 	// Returns true if the type is an array
 	bool IsArrayType(TRigVMTypeIndex InTypeIndex) const;
@@ -208,10 +208,10 @@ public:
 	TRigVMTypeIndex GetBaseTypeFromArrayTypeIndex(TRigVMTypeIndex InTypeIndex) const;
 
 	// Returns the function given its name (or nullptr)
-	const FRigVMFunction* FindFunction(const TCHAR* InName, const FRigVMTypeResolvalInfo& InResolvalInfo = FRigVMTypeResolvalInfo()) const;
+	const FRigVMFunction* FindFunction(const TCHAR* InName, const FRigVMUserDefinedTypeResolver& InTypeResolver = FRigVMUserDefinedTypeResolver()) const;
 
 	// Returns the function given its backing up struct and method name
-	const FRigVMFunction* FindFunction(UScriptStruct* InStruct, const TCHAR* InName, const FRigVMTypeResolvalInfo& InResolvalInfo = FRigVMTypeResolvalInfo()) const;
+	const FRigVMFunction* FindFunction(UScriptStruct* InStruct, const TCHAR* InName, const FRigVMUserDefinedTypeResolver& InResolvalInfo = FRigVMUserDefinedTypeResolver()) const;
 
 	// Returns all current rigvm functions
 	const TChunkedArray<FRigVMFunction>& GetFunctions() const;
@@ -332,7 +332,7 @@ private:
 	
 	void RemoveTypeInCategory(FRigVMTemplateArgument::ETypeCategory InCategory, TRigVMTypeIndex InTypeIndex);
 
-	const FRigVMFunction* FindFunction_NoLock(const TCHAR* InName, const FRigVMTypeResolvalInfo& InResolvalInfo = FRigVMTypeResolvalInfo()) const;
+	const FRigVMFunction* FindFunction_NoLock(const TCHAR* InName, const FRigVMUserDefinedTypeResolver& InTypeResolver = FRigVMUserDefinedTypeResolver()) const;
 	const FRigVMTemplate* FindTemplate_NoLock(const FName& InNotation, bool bIncludeDeprecated) const;
 	FRigVMDispatchFactory* FindDispatchFactory_NoLock(const FName& InFactoryName) const;
 
diff --git a/Engine/Plugins/Runtime/RigVM/Source/RigVM/Public/RigVMCore/RigVMTemplate.h b/Engine/Plugins/Runtime/RigVM/Source/RigVM/Public/RigVMCore/RigVMTemplate.h
--- a/Engine/Plugins/Runtime/RigVM/Source/RigVM/Public/RigVMCore/RigVMTemplate.h
+++ b/Engine/Plugins/Runtime/RigVM/Source/RigVM/Public/RigVMCore/RigVMTemplate.h
@@ -30,7 +30,7 @@ class FProperty;
 struct FRigVMDispatchFactory;
 struct FRigVMTemplate;
 struct FRigVMDispatchContext;
-struct FRigVMTypeResolvalInfo;
+struct FRigVMUserDefinedTypeResolver;
 
 typedef TMap<FName, TRigVMTypeIndex> FRigVMTemplateTypeMap;
 
@@ -374,7 +374,7 @@ public:
 	static TArray<UStruct*> GetSuperStructs(UStruct* InStruct, bool bIncludeLeaf = true);
 
 	// converts the types provided by a string (like "A:float,B:int32") into a type map
-	FTypeMap GetArgumentTypesFromString(const FString& InTypeString, const FRigVMTypeResolvalInfo* InResolvalInfo = nullptr) const;
+	FTypeMap GetArgumentTypesFromString(const FString& InTypeString, const FRigVMUserDefinedTypeResolver* InTypeResolver = nullptr) const;
 
 	// converts the types provided to a string (like "A:float,B:int32")
 	static FString GetStringFromArgumentTypes(const FTypeMap& InTypes); 
diff --git a/Engine/Plugins/Runtime/RigVM/Source/RigVM/Public/RigVMHost.h b/Engine/Plugins/Runtime/RigVM/Source/RigVM/Public/RigVMHost.h
--- a/Engine/Plugins/Runtime/RigVM/Source/RigVM/Public/RigVMHost.h
+++ b/Engine/Plugins/Runtime/RigVM/Source/RigVM/Public/RigVMHost.h
@@ -281,8 +281,7 @@ public:
 		return ExtendedExecuteContext;
 	};
 
-	const TMap<FString, FSoftObjectPath>& GetUserDefinedStructGuidToPathName() const { return UserDefinedStructGuidToPathName; }
-	const TMap<FString, FSoftObjectPath>& GetUserDefinedEnumToPathName() const { return  UserDefinedEnumToPathName; }
+	UObject* ResolveUserDefinedTypeById(const FString& InTypeName) const;
 
 protected:
 
@@ -317,6 +316,9 @@ protected:
 	TMap<FString, FSoftObjectPath> UserDefinedEnumToPathName;
 
 private:
+	UPROPERTY(transient)
+	TSet<TObjectPtr<UObject>> UserDefinedTypesInUse;
+	
 	UPROPERTY()
 	FRigVMExtendedExecuteContext ExtendedExecuteContext;
 
@@ -364,7 +366,8 @@ public:
 	/** Broadcasts a notification whenever the RigVMHost's memory is initialized. */
 	FRigVMExecutedEvent InitializedEvent;
 
-	/** Broadcasts a notification whenever the RigVMHost is executed / updated. */
+	/** Broadcasts a not
+	 * ification whenever the RigVMHost is executed / updated. */
 	FRigVMExecutedEvent ExecutedEvent;
 
 protected: 
diff --git a/Engine/Plugins/Runtime/RigVM/Source/RigVM/Public/RigVMTypeUtils.h b/Engine/Plugins/Runtime/RigVM/Source/RigVM/Public/RigVMTypeUtils.h
--- a/Engine/Plugins/Runtime/RigVM/Source/RigVM/Public/RigVMTypeUtils.h
+++ b/Engine/Plugins/Runtime/RigVM/Source/RigVM/Public/RigVMTypeUtils.h
@@ -12,9 +12,36 @@
 #include "UObject/Package.h"
 #include "UObject/SoftObjectPath.h"
 
-struct RIGVM_API FRigVMTypeResolvalInfo
+
+
+struct RIGVM_API FRigVMUserDefinedTypeResolver
 {
-	TMap<FString, FSoftObjectPath> CPPTypeToObjectPath;
+	FRigVMUserDefinedTypeResolver() = default;
+	explicit FRigVMUserDefinedTypeResolver(const TFunction<UObject*(const FString&)>& InResolver) : Resolver(InResolver) {}
+	explicit FRigVMUserDefinedTypeResolver(TMap<FString, FSoftObjectPath>&& InObjectMap) : ObjectMap(InObjectMap) {} 
+	
+	UObject* GetTypeObjectByName(const FString& InTypeName) const
+	{
+		if (Resolver)
+		{
+			return Resolver(InTypeName);
+		}
+		
+		if (const FSoftObjectPath *ObjectPath = ObjectMap.Find(InTypeName))
+		{
+			return ObjectPath->TryLoad();
+		}
+		return nullptr;
+	}
+
+	bool IsValid() const
+	{
+		return Resolver || !ObjectMap.IsEmpty();
+	}
+	
+private:
+	TFunction<UObject*(const FString&)> Resolver;
+	TMap<FString, FSoftObjectPath> ObjectMap;
 };
 
 namespace RigVMTypeUtils
@@ -310,10 +337,10 @@ namespace RigVMTypeUtils
 	}
 
 	// Finds the CPPTypeObject from a CPP type of a potentially missing / unloaded user defined struct or enum
-	RIGVM_API UObject* UserDefinedTypeFromCPPType(FString& InOutCPPType, const FRigVMTypeResolvalInfo* InResolvalInfo = nullptr);
+	RIGVM_API UObject* UserDefinedTypeFromCPPType(FString& InOutCPPType, const FRigVMUserDefinedTypeResolver* InTypeResolver = nullptr);
 
 	// Finds the CPPTypeObject from the CPPType. If not found, tries to use redirectors and modifies the InOutCPPType.
-	static UObject* ObjectFromCPPType(FString& InOutCPPType, bool bUseRedirector = true, const FRigVMTypeResolvalInfo* InResolvalInfo = nullptr)
+	static UObject* ObjectFromCPPType(FString& InOutCPPType, bool bUseRedirector = true, const FRigVMUserDefinedTypeResolver* InTypeResolver = nullptr)
 	{
 		if (!RequiresCPPTypeObject(InOutCPPType))
 		{
@@ -350,7 +377,7 @@ namespace RigVMTypeUtils
 		if(CPPTypeObject == nullptr)
 		{
 			CPPType = BaseCPPType;
-			CPPTypeObject = UserDefinedTypeFromCPPType(CPPType, InResolvalInfo);
+			CPPTypeObject = UserDefinedTypeFromCPPType(CPPType, InTypeResolver);
 		}
 
 		if(CPPTypeObject == nullptr)
@@ -364,7 +391,7 @@ namespace RigVMTypeUtils
 		return CPPTypeObject;
 	}
 
-	static FString PostProcessCPPType(const FString& InCPPType, UObject* InCPPTypeObject = nullptr, const FRigVMTypeResolvalInfo* InResolvalInfo = nullptr)
+	static FString PostProcessCPPType(const FString& InCPPType, UObject* InCPPTypeObject = nullptr, const FRigVMUserDefinedTypeResolver* InResolvalInfo = nullptr)
 	{
 		FString CPPType = InCPPType;
 		if (InCPPTypeObject)
diff --git a/Engine/Plugins/Runtime/RigVM/Source/RigVMDeveloper/Private/RigVMBlueprint.cpp b/Engine/Plugins/Runtime/RigVM/Source/RigVMDeveloper/Private/RigVMBlueprint.cpp
--- a/Engine/Plugins/Runtime/RigVM/Source/RigVMDeveloper/Private/RigVMBlueprint.cpp
+++ b/Engine/Plugins/Runtime/RigVM/Source/RigVMDeveloper/Private/RigVMBlueprint.cpp
@@ -651,6 +651,24 @@ void URigVMBlueprint::HandleConfigureRigVMController(const FRigVMClient* InClien
 #endif
 }
 
+UObject* URigVMBlueprint::ResolveUserDefinedTypeById(const FString& InTypeName) const
+{
+	const FSoftObjectPath* ResultPathPtr = UserDefinedStructGuidToPathName.Find(InTypeName);
+	if (ResultPathPtr == nullptr)
+	{
+		return nullptr;
+	}
+
+	if (UObject* TypeObject = ResultPathPtr->TryLoad())
+	{
+		// Ensure we have a hold on this type so it doesn't get nixed on the next GC.
+		const_cast<URigVMBlueprint*>(this)->UserDefinedTypesInUse.Add(TypeObject);
+		return TypeObject;
+	}
+
+	return nullptr;
+}
+
 bool URigVMBlueprint::TryImportGraphFromText(const FString& InClipboardText, UEdGraph** OutGraphPtr)
 {
 	if (OutGraphPtr)
@@ -821,6 +839,7 @@ void URigVMBlueprint::PreSave(FObjectPreSaveContext ObjectSaveContext)
 	// to aid the controller when recovering from user defined struct name changes or
 	// guid changes.
 	UserDefinedStructGuidToPathName.Reset();
+	UserDefinedTypesInUse.Reset();
 	TArray<URigVMGraph*> AllModels = GetAllModels();
 	for(const URigVMGraph* Graph : AllModels)
 	{
diff --git a/Engine/Plugins/Runtime/RigVM/Source/RigVMDeveloper/Private/RigVMModel/RigVMClient.cpp b/Engine/Plugins/Runtime/RigVM/Source/RigVMDeveloper/Private/RigVMModel/RigVMClient.cpp
--- a/Engine/Plugins/Runtime/RigVM/Source/RigVMDeveloper/Private/RigVMModel/RigVMClient.cpp
+++ b/Engine/Plugins/Runtime/RigVM/Source/RigVMDeveloper/Private/RigVMModel/RigVMClient.cpp
@@ -14,10 +14,9 @@
 #include "ScopedTransaction.h"
 #endif
 
-const TMap<FString, FSoftObjectPath>& IRigVMClientHost::GetUserDefinedStructGuidToObjectPath() const
+UObject* IRigVMClientHost::ResolveUserDefinedTypeById(const FString& InTypeName) const
 {
-	static const TMap<FString, FSoftObjectPath> EmptyMap;
-	return EmptyMap;
+	return nullptr;
 }
 
 void FRigVMClient::SetSchemaClass(TSubclassOf<URigVMSchema> InSchemaClass)
diff --git a/Engine/Plugins/Runtime/RigVM/Source/RigVMDeveloper/Private/RigVMModel/RigVMController.cpp b/Engine/Plugins/Runtime/RigVM/Source/RigVMDeveloper/Private/RigVMModel/RigVMController.cpp
--- a/Engine/Plugins/Runtime/RigVM/Source/RigVMDeveloper/Private/RigVMModel/RigVMController.cpp
+++ b/Engine/Plugins/Runtime/RigVM/Source/RigVMDeveloper/Private/RigVMModel/RigVMController.cpp
@@ -17911,14 +17911,17 @@ bool URigVMController::EnsurePinValidity(URigVMPin* InPin, bool bRecursive)
 		// without testing for redirector
 		if(InPin->GetCPPTypeObject() == nullptr)
 		{
-			FRigVMTypeResolvalInfo ResolveInfo;
+			FRigVMUserDefinedTypeResolver TypeResolver;
 			if(const IRigVMClientHost* ClientHost = InPin->GetImplementingOuter<IRigVMClientHost>())
 			{
-				ResolveInfo.CPPTypeToObjectPath = ClientHost->GetUserDefinedStructGuidToObjectPath();
+				TypeResolver = FRigVMUserDefinedTypeResolver([ClientHost](const FString& InTypeName) -> UObject*
+				{
+					return ClientHost->ResolveUserDefinedTypeById(InTypeName);
+				});
 			}
 			
 			FString CPPType = InPin->GetCPPType();
-			InPin->CPPTypeObject = RigVMTypeUtils::ObjectFromCPPType(CPPType, true, &ResolveInfo);
+			InPin->CPPTypeObject = RigVMTypeUtils::ObjectFromCPPType(CPPType, true, &TypeResolver);
 			if(CPPType.IsEmpty())
 			{
 				return false;
diff --git a/Engine/Plugins/Runtime/RigVM/Source/RigVMDeveloper/Public/RigVMBlueprint.h b/Engine/Plugins/Runtime/RigVM/Source/RigVMDeveloper/Public/RigVMBlueprint.h
--- a/Engine/Plugins/Runtime/RigVM/Source/RigVMDeveloper/Public/RigVMBlueprint.h
+++ b/Engine/Plugins/Runtime/RigVM/Source/RigVMDeveloper/Public/RigVMBlueprint.h
@@ -311,7 +311,7 @@ public:
 	virtual void HandleRigVMGraphRemoved(const FRigVMClient* InClient, const FString& InNodePath) override;
 	virtual void HandleRigVMGraphRenamed(const FRigVMClient* InClient, const FString& InOldNodePath, const FString& InNewNodePath) override;
 	virtual void HandleConfigureRigVMController(const FRigVMClient* InClient, URigVMController* InControllerToConfigure) override;
-	virtual const TMap<FString, FSoftObjectPath>& GetUserDefinedStructGuidToObjectPath() const override { return UserDefinedStructGuidToPathName; }
+	virtual UObject* ResolveUserDefinedTypeById(const FString& InTypeName) const override;
 
 	FOnRigVMRequestInspectObject& OnRequestInspectObject() { return OnRequestInspectObjectEvent; }
 	void RequestInspectObject(const TArray<UObject*>& InObjects) { OnRequestInspectObjectEvent.Broadcast(InObjects); }
@@ -437,6 +437,9 @@ public:
 	UPROPERTY()
 	TMap<FString, FSoftObjectPath> UserDefinedStructGuidToPathName;
 
+	UPROPERTY(transient)
+	TSet<TObjectPtr<UObject>> UserDefinedTypesInUse;
+
 protected:
 
 #if WITH_EDITORONLY_DATA
diff --git a/Engine/Plugins/Runtime/RigVM/Source/RigVMDeveloper/Public/RigVMModel/RigVMClient.h b/Engine/Plugins/Runtime/RigVM/Source/RigVMDeveloper/Public/RigVMModel/RigVMClient.h
--- a/Engine/Plugins/Runtime/RigVM/Source/RigVMDeveloper/Public/RigVMModel/RigVMClient.h
+++ b/Engine/Plugins/Runtime/RigVM/Source/RigVMDeveloper/Public/RigVMModel/RigVMClient.h
@@ -50,8 +50,9 @@ public:
 	// Reacts to a request to configure a controller
 	virtual void HandleConfigureRigVMController(const FRigVMClient* InClient, URigVMController* InControllerToConfigure) = 0;
 
-	// Returns a map from user defined struct to object path used to resolve types
-	virtual const TMap<FString, FSoftObjectPath>& GetUserDefinedStructGuidToObjectPath() const;
+	// Given a type name for a user-defined type, either struct or an enum, returns a pointer to the object definition
+	// or nullptr if the client host has no knowledge of such a type.
+	virtual UObject* ResolveUserDefinedTypeById(const FString& InTypeName) const;
 };
 
 UINTERFACE()
