RigVM: Ray Math Functions

diff --git a/Engine/Plugins/Runtime/RigVM/Source/RigVM/Private/RigVMFunctions/Math/RigVMFunction_MathRay.cpp b/Engine/Plugins/Runtime/RigVM/Source/RigVM/Private/RigVMFunctions/Math/RigVMFunction_MathRay.cpp
new file mode 100644
--- /dev/null
+++ b/Engine/Plugins/Runtime/RigVM/Source/RigVM/Private/RigVMFunctions/Math/RigVMFunction_MathRay.cpp
@@ -0,0 +1,52 @@
+// Copyright Epic Games, Inc. All Rights Reserved.
+
+#include "RigVMFunctions/Math/RigVMFunction_MathRay.h"
+#include "RigVMFunctions/Math/RigVMFunction_MathTransform.h"
+#include "RigVMFunctions/RigVMDispatch_Core.h"
+
+#include UE_INLINE_GENERATED_CPP_BY_NAME(RigVMFunction_MathRay)
+
+FRigVMFunction_MathRayIntersectRay_Execute()
+{
+	const FVector DirectionA = A.Direction * 10000.0;
+	const FVector DirectionB = B.Direction * 10000.0;
+	
+	FVector ClosestA, ClosestB;
+	FMath::SegmentDistToSegmentSafe(
+		A.Origin - DirectionA,
+		A.Origin + DirectionA,
+		B.Origin - DirectionB,
+		B.Origin + DirectionB,
+		ClosestA, ClosestB);
+
+	Result = (ClosestA + ClosestB) * 0.5;
+	Distance = 0;
+	if(!ClosestA.Equals(ClosestB))
+	{
+		Distance = (ClosestA - ClosestB).Size();
+	}
+	RatioA = A.GetParameter(ClosestA);	
+	RatioB = B.GetParameter(ClosestB);	
+}
+
+FRigVMFunction_MathRayIntersectPlane_Execute()
+{
+	Ratio = FMath::RayPlaneIntersectionParam(Ray.Origin, Ray.Direction, FPlane(PlanePoint, PlaneNormal));
+	Result = Ray.PointAt(Ratio);
+	Distance = 0;
+	if(!Ray.Origin.Equals(Result))
+	{
+		Distance = (Ray.Origin - Result).Size();
+	}
+}
+
+FRigVMFunction_MathRayGetAt_Execute()
+{
+	Result = Ray.PointAt(Ratio);
+}
+
+FRigVMFunction_MathRayTransform_Execute()
+{
+	Result.Origin = Transform.TransformPosition(Ray.Origin);
+	Result.Direction = Transform.TransformVector(Ray.Direction);
+}
diff --git a/Engine/Plugins/Runtime/RigVM/Source/RigVM/Public/RigVMFunctions/Math/RigVMFunction_MathRay.h b/Engine/Plugins/Runtime/RigVM/Source/RigVM/Public/RigVMFunctions/Math/RigVMFunction_MathRay.h
new file mode 100644
--- /dev/null
+++ b/Engine/Plugins/Runtime/RigVM/Source/RigVM/Public/RigVMFunctions/Math/RigVMFunction_MathRay.h
@@ -0,0 +1,145 @@
+// Copyright Epic Games, Inc. All Rights Reserved.
+
+#pragma once
+
+#include "Math/Ray.h"
+#include "Math/Plane.h"
+#include "RigVMFunction_MathBase.h"
+#include "RigVMFunction_MathRay.generated.h"
+
+USTRUCT(meta=(Abstract, Category="Math|Ray", MenuDescSuffix="(Ray)"))
+struct RIGVM_API FRigVMFunction_MathRayBase : public FRigVMFunction_MathBase
+{
+	GENERATED_BODY()
+};
+
+/**
+ * Returns the closest point intersection of a ray with another ray
+ */
+USTRUCT(meta = (DisplayName = "Intersect Ray", Keywords = "Closest,Ray,Cross"))
+struct RIGVM_API FRigVMFunction_MathRayIntersectRay : public FRigVMFunction_MathRayBase
+{
+	GENERATED_BODY()
+
+	FRigVMFunction_MathRayIntersectRay()
+	{
+		A = B = FRay();
+		Result = FVector::ZeroVector;
+		Distance = RatioA = RatioB = 0.f;
+	}
+
+	RIGVM_METHOD()
+	virtual void Execute() override;
+
+	UPROPERTY(meta = (Input))
+	FRay A;
+
+	UPROPERTY(meta = (Input))
+	FRay B;
+
+	UPROPERTY(meta = (Output))
+	FVector Result;
+
+	UPROPERTY(meta = (Output))
+	float Distance;
+
+	UPROPERTY(meta = (Output))
+	float RatioA;
+
+	UPROPERTY(meta = (Output))
+	float RatioB;
+};
+
+/**
+ * Returns the closest point intersection of a ray with a plane
+ */
+USTRUCT(meta = (DisplayName = "Intersect Plane", Keywords = "Closest,Ray,Cross"))
+struct RIGVM_API FRigVMFunction_MathRayIntersectPlane : public FRigVMFunction_MathRayBase
+{
+	GENERATED_BODY()
+
+	FRigVMFunction_MathRayIntersectPlane()
+	{
+		Ray = FRay();
+		PlanePoint = FVector::ZeroVector;
+		PlaneNormal = FVector::UpVector;
+		Result = FVector::ZeroVector;
+		Distance = Ratio = 0.f;
+	}
+
+	RIGVM_METHOD()
+	virtual void Execute() override;
+
+	UPROPERTY(meta = (Input))
+	FRay Ray;
+
+	UPROPERTY(meta = (Input))
+	FVector PlanePoint;
+
+	UPROPERTY(meta = (Input))
+	FVector PlaneNormal;
+
+	UPROPERTY(meta = (Output))
+	FVector Result;
+
+	UPROPERTY(meta = (Output))
+	float Distance;
+
+	UPROPERTY(meta = (Output))
+	float Ratio;
+};
+
+/**
+ * Returns the position on a ray
+ */
+USTRUCT(meta=(DisplayName="GetAt", Keywords="Ratio,Percentage"))
+struct RIGVM_API FRigVMFunction_MathRayGetAt : public FRigVMFunction_MathRayBase
+{
+	GENERATED_BODY()
+
+	FRigVMFunction_MathRayGetAt()
+	{
+		Ray = FRay();
+		Ratio = 0.f;
+		Result = FVector::ZeroVector;
+	}
+
+	RIGVM_METHOD()
+	virtual void Execute() override;
+
+	UPROPERTY(meta = (Input))
+	FRay Ray;
+
+	UPROPERTY(meta = (Input))
+	float Ratio;
+
+	UPROPERTY(meta = (Output))
+	FVector Result;
+};
+
+/**
+ * Returns the division of the two values
+ */
+USTRUCT(meta=(DisplayName="Transform Ray", Keywords="Multiply,Project"))
+struct RIGVM_API FRigVMFunction_MathRayTransform : public FRigVMFunction_MathRayBase
+{
+	GENERATED_BODY()
+
+	FRigVMFunction_MathRayTransform()
+	{
+		Ray = Result = FRay();
+		Transform = FTransform::Identity;
+	}
+
+	RIGVM_METHOD()
+	virtual void Execute() override;
+
+	UPROPERTY(meta = (Input))
+	FRay Ray;
+
+	UPROPERTY(meta = (Input))
+	FTransform Transform;
+
+	UPROPERTY(meta = (Output))
+	FRay Result;
+};
