RigVM Compiler : Avoid duplicated properties by reusing operands in inline functions

diff --git a/Engine/Plugins/Runtime/RigVM/Source/RigVM/Private/RigVMCore/RigVMMemoryStorage.cpp b/Engine/Plugins/Runtime/RigVM/Source/RigVM/Private/RigVMCore/RigVMMemoryStorage.cpp
--- a/Engine/Plugins/Runtime/RigVM/Source/RigVM/Private/RigVMCore/RigVMMemoryStorage.cpp
+++ b/Engine/Plugins/Runtime/RigVM/Source/RigVM/Private/RigVMCore/RigVMMemoryStorage.cpp
@@ -307,11 +307,23 @@ FRigVMPropertyDescription::FRigVMPropertyDescription(const FName& InName, const
 FName FRigVMPropertyDescription::SanitizeName(const FName& InName)
 {
 	FString NameString = InName.ToString();
+	SanitizeName(NameString);
 
+	if (NameString != InName.ToString())
+	{
+		return *NameString;
+	}
+
+	return InName;
+
+}
+
+void FRigVMPropertyDescription::SanitizeName(FString& InString)
+{
 	// Sanitize the name
-	for (int32 i = 0; i < NameString.Len(); ++i)
+	for (int32 i = 0; i < InString.Len(); ++i)
 	{
-		TCHAR& C = NameString[i];
+		TCHAR& C = InString[i];
 
 		const bool bGoodChar = FChar::IsAlpha(C) ||							// Any letter
 			(C == '_') || 													// _ anytime
@@ -322,13 +334,6 @@ FName FRigVMPropertyDescription::SanitizeName(const FName& InName)
 			C = '_';
 		}
 	}
-
-	if(NameString != InName.ToString())
-	{
-		return *NameString;
-	}
-
-	return InName;
 }
 
 void FRigVMPropertyDescription::SanitizeName()
diff --git a/Engine/Plugins/Runtime/RigVM/Source/RigVM/Public/RigVMCore/RigVMMemoryStorage.h b/Engine/Plugins/Runtime/RigVM/Source/RigVM/Public/RigVMCore/RigVMMemoryStorage.h
--- a/Engine/Plugins/Runtime/RigVM/Source/RigVM/Public/RigVMCore/RigVMMemoryStorage.h
+++ b/Engine/Plugins/Runtime/RigVM/Source/RigVM/Public/RigVMCore/RigVMMemoryStorage.h
@@ -1019,6 +1019,9 @@ public:
 	// Returns a sanitized, valid name to use for a new property 
 	static FName SanitizeName(const FName& InName);
 
+	// Sanitizes a name using a string ref, creating a valid name to use for a new property 
+	static void SanitizeName(FString& InString);
+
 	// Sanitize the name of this description in line
 	void SanitizeName();
 
diff --git a/Engine/Plugins/Runtime/RigVM/Source/RigVMDeveloper/Private/RigVMCompiler/RigVMCompiler.cpp b/Engine/Plugins/Runtime/RigVM/Source/RigVMDeveloper/Private/RigVMCompiler/RigVMCompiler.cpp
--- a/Engine/Plugins/Runtime/RigVM/Source/RigVMDeveloper/Private/RigVMCompiler/RigVMCompiler.cpp
+++ b/Engine/Plugins/Runtime/RigVM/Source/RigVMDeveloper/Private/RigVMCompiler/RigVMCompiler.cpp
@@ -130,6 +130,7 @@ FRigVMOperand FRigVMCompilerWorkData::AddProperty(
 	}
 	
 	FRigVMPropertyDescription Description(InName, ArgumentType.CPPType.ToString(), ArgumentType.CPPTypeObject, InDefaultValue);
+	ensure(FindProperty(InMemoryType, Description.Name).IsValid() == false); // Warning : this check can not be done before the SanitizeName inside FRigVMPropertyDescription constructor
 
 	TArray<FRigVMPropertyDescription>& PropertyArray = PropertyDescriptions.FindOrAdd(InMemoryType);
 	const int32 PropertyIndex = PropertyArray.Add(Description);
@@ -2257,12 +2258,17 @@ int32 URigVMCompiler::TraverseInlineFunction(const FRigVMInlineFunctionExprAST*
 				if (NewName.StartsWith(FunctionLibraryPrefix))
 				{
 					NewName = FString::Printf(TEXT("%s%s"), *FunctionReferenceNode->GetNodePath(), *NewName.RightChop(FunctionLibraryPrefix.Len()));
+					FRigVMPropertyDescription::SanitizeName(NewName);
+				}
+
+				// Reuse operands if already added to the WorkData
+				FRigVMOperand Operand = WorkData.FindProperty(MemoryType, *NewName);
+				if (!Operand.IsValid())
+				{
+					Operand = WorkData.AddProperty(MemoryType, *NewName, Description.CPPType, Description.CPPTypeObject.Get(), Description.DefaultValue);
 				}
-				FRigVMOperand Operand = WorkData.AddProperty(MemoryType, *NewName, Description.CPPType, Description.CPPTypeObject.Get(), Description.DefaultValue);
 				FRigVMCompilerWorkData::FFunctionRegisterData Data = {FunctionReferenceNode, MemoryType, PropertyIndex};
 				WorkData.FunctionRegisterToOperand.Add(Data, Operand);
-
-				// @todo Try to reuse literal operands
 			}
 		}
 	}
