RigVMEditor: Populate execution stack widget on load of editor

diff --git a/Engine/Plugins/Runtime/RigVM/Source/RigVMEditor/Private/Widgets/SRigVMExecutionStackView.cpp b/Engine/Plugins/Runtime/RigVM/Source/RigVMEditor/Private/Widgets/SRigVMExecutionStackView.cpp
--- a/Engine/Plugins/Runtime/RigVM/Source/RigVMEditor/Private/Widgets/SRigVMExecutionStackView.cpp
+++ b/Engine/Plugins/Runtime/RigVM/Source/RigVMEditor/Private/Widgets/SRigVMExecutionStackView.cpp
@@ -1020,6 +1020,10 @@ void SRigVMExecutionStackView::HandleHostInitializedEvent(URigVMHost* InHost, co
 
 void SRigVMExecutionStackView::HandlePreviewHostUpdated(FRigVMEditor* InEditor)
 {
+	if(URigVMHost* RigVMHost = InEditor->GetRigVMHost())
+	{
+		RefreshTreeView(RigVMHost->GetVM(), &RigVMHost->GetExtendedExecuteContext());
+	}
 }
 
 void SRigVMExecutionStackView::HandleItemMouseDoubleClick(TSharedPtr<FRigStackEntry> InItem)
