RigVM now uses a struct based storage, built on top of PropertyBag.

diff --git a/Engine/Plugins/Animation/ControlRig/Source/ControlRig/Private/ControlRig.cpp b/Engine/Plugins/Animation/ControlRig/Source/ControlRig/Private/ControlRig.cpp
--- a/Engine/Plugins/Animation/ControlRig/Source/ControlRig/Private/ControlRig.cpp
+++ b/Engine/Plugins/Animation/ControlRig/Source/ControlRig/Private/ControlRig.cpp
@@ -1228,13 +1228,16 @@ bool UControlRig::Execute_Internal(const FName& InEventName)
 		else
 		{
 			// sanity check the validity of the VM to ensure stability.
-			if(!VM->IsContextValidForExecution(Context) ||
-				!IsValidLowLevel() ||
-				!VM->IsValidLowLevel() ||
-				!GetLiteralMemory() ||
-				!GetLiteralMemory()->IsValidLowLevel() ||
-				!GetWorkMemory() ||
-				!GetWorkMemory()->IsValidLowLevel())
+			if(!VM->IsContextValidForExecution(Context)
+				|| !IsValidLowLevel()
+				|| !VM->IsValidLowLevel()
+#if !UE_RIGVM_PROPERTY_BAG_STORAGE_ENABLED
+				|| !GetLiteralMemory()
+				|| !GetLiteralMemory()->IsValidLowLevel()
+				|| !GetWorkMemory()
+				||!GetWorkMemory()->IsValidLowLevel()
+#endif // !UE_RIGVM_PROPERTY_BAG_STORAGE_ENABLED
+			)
 			{
 				UE_LOG(LogControlRig, Warning, InvalidatedVMFormat, *GetClass()->GetName());
 				return false;
@@ -1297,7 +1300,7 @@ bool UControlRig::Execute_Internal(const FName& InEventName)
 #endif
 		FRigHierarchyExecuteContextBracket HierarchyContextGuard(Hierarchy, &Context);
 
-		TArray<URigVMMemoryStorage*> LocalMemory = VM->GetLocalMemoryArray(Context);
+		TArray<TRigVMMemoryStorage*> LocalMemory = VM->GetLocalMemoryArray(Context);
 		const bool bSuccess = VM->Execute(Context, LocalMemory, InEventName) != ERigVMExecuteResult::Failed;
 
 #if UE_RIGVM_PROFILE_EXECUTE_UNITS_NUM
diff --git a/Engine/Plugins/Animation/ControlRig/Source/ControlRigEditor/Private/ControlRigCompilerDetails.cpp b/Engine/Plugins/Animation/ControlRig/Source/ControlRigEditor/Private/ControlRigCompilerDetails.cpp
--- a/Engine/Plugins/Animation/ControlRig/Source/ControlRigEditor/Private/ControlRigCompilerDetails.cpp
+++ b/Engine/Plugins/Animation/ControlRig/Source/ControlRigEditor/Private/ControlRigCompilerDetails.cpp
@@ -192,10 +192,15 @@ FReply FRigVMCompileSettingsDetails::OnInspectMemory(ERigVMMemoryType InMemoryTy
 	{
 		if(UControlRig* DebuggedRig = Cast<UControlRig>(BlueprintBeingCustomized->GetObjectBeingDebugged()))
 		{
-			if(URigVMMemoryStorage* MemoryStorage = DebuggedRig->GetMemoryByType(InMemoryType))
+			if(TRigVMMemoryStorage* MemoryStorage = DebuggedRig->GetMemoryByType(InMemoryType))
 			{
+#if UE_RIGVM_PROPERTY_BAG_STORAGE_ENABLED
+				TArray<TRigVMMemoryStorage*> InStructs = { MemoryStorage };
+				BlueprintBeingCustomized->RequestInspectMemoryStorage(InStructs);
+#else
 				TArray<UObject*> ObjectsToSelect = {MemoryStorage};
 				BlueprintBeingCustomized->RequestInspectObject(ObjectsToSelect);
+#endif // !UE_RIGVM_PROPERTY_BAG_STORAGE_ENABLED
 			}
 		}
 	}
diff --git a/Engine/Plugins/Experimental/AnimNext/Source/AnimNext/Private/Graph/AnimNextGraph.cpp b/Engine/Plugins/Experimental/AnimNext/Source/AnimNext/Private/Graph/AnimNextGraph.cpp
--- a/Engine/Plugins/Experimental/AnimNext/Source/AnimNext/Private/Graph/AnimNextGraph.cpp
+++ b/Engine/Plugins/Experimental/AnimNext/Source/AnimNext/Private/Graph/AnimNextGraph.cpp
@@ -54,7 +54,7 @@ void UAnimNextGraph::Run(const UE::AnimNext::FContext& Context, UE::AnimNext::FW
 		AnimNextContext.InitializeWithGraph(SharedDataBuffer, GraphInstancePtr);
 		AnimNextContext.SetSimulationSteps(SimulationSteps);
 
-		RigVM->Execute(RigVMExtendedExecuteContext, TArray<URigVMMemoryStorage*>(), FRigUnit_AnimNextGraphRoot::EventName);
+		RigVM->Execute(RigVMExtendedExecuteContext, TArray<TRigVMMemoryStorage*>(), FRigUnit_AnimNextGraphRoot::EventName);
 	}
 }
 
diff --git a/Engine/Plugins/Experimental/AnimNext/Source/AnimNext/Private/Param/AnimNextParameterBlock.cpp b/Engine/Plugins/Experimental/AnimNext/Source/AnimNext/Private/Param/AnimNextParameterBlock.cpp
--- a/Engine/Plugins/Experimental/AnimNext/Source/AnimNext/Private/Param/AnimNextParameterBlock.cpp
+++ b/Engine/Plugins/Experimental/AnimNext/Source/AnimNext/Private/Param/AnimNextParameterBlock.cpp
@@ -43,7 +43,7 @@ void UAnimNextParameterBlock::Run(const UE::AnimNext::FContext& Context) const
 		FRigVMExtendedExecuteContext RigVMExtendedExecuteContext;
 		FAnimNextExecuteContext& AnimNextContext = RigVMExtendedExecuteContext.GetPublicDataSafe<FAnimNextExecuteContext>();
 		AnimNextContext.SetContextData(Context);
-		RigVM->Execute(RigVMExtendedExecuteContext, TArray<URigVMMemoryStorage*>(), FRigUnit_AnimNextBeginExecution::EventName);
+		RigVM->Execute(RigVMExtendedExecuteContext, TArray<TRigVMMemoryStorage*>(), FRigUnit_AnimNextBeginExecution::EventName);
 	}
 }
 
diff --git a/Engine/Plugins/Runtime/RigVM/RigVM.uplugin b/Engine/Plugins/Runtime/RigVM/RigVM.uplugin
--- a/Engine/Plugins/Runtime/RigVM/RigVM.uplugin
+++ b/Engine/Plugins/Runtime/RigVM/RigVM.uplugin
@@ -31,5 +31,12 @@
 			"Type" : "Editor",
 			"LoadingPhase" : "PostEngineInit"
 		}
+	],
+	"Plugins" :
+	[
+		{
+			"Name": "StructUtils",
+			"Enabled": true
+		}
 	]
 }
\ No newline at end of file
diff --git a/Engine/Plugins/Runtime/RigVM/Source/RigVM/Private/RigVMCore/RigVM.cpp b/Engine/Plugins/Runtime/RigVM/Source/RigVM/Private/RigVMCore/RigVM.cpp
--- a/Engine/Plugins/Runtime/RigVM/Source/RigVM/Private/RigVMCore/RigVM.cpp
+++ b/Engine/Plugins/Runtime/RigVM/Source/RigVM/Private/RigVMCore/RigVM.cpp
@@ -114,8 +114,6 @@ void URigVM::Serialize(FArchive& Ar)
 	{
 		Super::Serialize(Ar);
 	}
-	
-	ensure(ActiveExecutions.load() == 0);  // we require that no active worker threads are running while we serialize
 
 	if (Ar.IsSaving() || Ar.IsObjectReferenceCollector() || Ar.IsCountingMemory())
 	{
@@ -142,6 +140,8 @@ void URigVM::Save(FArchive& Ar)
 	Ar << Parameters;
 
 	Ar << OperandToDebugRegisters;
+	
+	Ar << LiteralMemoryStorage;
 }
 
 void URigVM::Load(FArchive& Ar)
@@ -170,11 +170,11 @@ void URigVM::Load(FArchive& Ar)
 
 		if (RigVMUClassBasedStorageDefine == 1)
 		{
-			FRigVMMemoryContainer WorkMemoryStorage;
-			FRigVMMemoryContainer LiteralMemoryStorage;
+			FRigVMMemoryContainer TmpWorkMemoryStorage;
+			FRigVMMemoryContainer TmpLiteralMemoryStorage;
 
-			Ar << WorkMemoryStorage;
-			Ar << LiteralMemoryStorage;
+			Ar << TmpWorkMemoryStorage;
+			Ar << TmpLiteralMemoryStorage;
 			Ar << FunctionNamesStorage;
 			Ar << ByteCodeStorage;
 			Ar << Parameters;
@@ -245,6 +245,11 @@ void URigVM::Load(FArchive& Ar)
 		TMap<FString, FSoftObjectPath> UserDefinedEnumToPathName;
 		Ar << UserDefinedEnumToPathName;
 	}
+
+	if (Ar.CustomVer(FRigVMObjectVersion::GUID) >= FRigVMObjectVersion::VMMemoryStorageStructSerialized)
+	{
+		Ar << LiteralMemoryStorage;
+	}
 }
 
 void URigVM::CopyDataForSerialization(URigVM* InVM)
@@ -261,6 +266,8 @@ void URigVM::CopyDataForSerialization(URigVM* InVM)
 	Parameters = InVM->Parameters;
 
 	OperandToDebugRegisters = InVM->OperandToDebugRegisters;
+
+	LiteralMemoryStorage = InVM->LiteralMemoryStorage;
 }
 
 void URigVM::PostLoad()
@@ -355,14 +362,19 @@ uint32 URigVM::ComputeVMHash(const FRigVMExtendedExecuteContext& Context) const
 		Hash = HashCombine(Hash, GetTypeHash(ExternalVariable.TypeName.ToString()));
 	}
 
+#if UE_RIGVM_PROPERTY_BAG_STORAGE_ENABLED
+	Hash = HashCombine(Hash, LiteralMemoryStorage.GetMemoryHash());
+	Hash = HashCombine(Hash, Context.WorkMemoryStorage.GetMemoryHash());
+#else
 	if(LiteralMemoryStorageObject)
 	{
 		Hash = HashCombine(Hash, LiteralMemoryStorageObject->GetMemoryHash());
 	}
-	if(Context.WorkMemoryStorageObject)
+	if(GetWorkMemory(Context))
 	{
-		Hash = HashCombine(Hash, Context.WorkMemoryStorageObject->GetMemoryHash());
+		Hash = HashCombine(Hash, GetWorkMemory(Context)->GetMemoryHash());
 	}
+#endif
 	
 	return Hash;
 }
@@ -399,7 +411,7 @@ bool URigVM::ValidateAllOperandsDuringLoad(FRigVMExtendedExecuteContext& Context
 	// check all operands on all ops for validity
 	bool bAllOperandsValid = true;
 
-	TArray<URigVMMemoryStorage*> LocalMemory = { GetWorkMemory(Context), GetLiteralMemory(), GetDebugMemory(Context) };
+	TArray<TRigVMMemoryStorage*> LocalMemory = { GetWorkMemory(Context), GetLiteralMemory(), GetDebugMemory(Context) };
 	
 	auto CheckOperandValidity = [LocalMemory, &bAllOperandsValid, this](const FRigVMOperand& InOperand) -> bool
 	{
@@ -410,7 +422,7 @@ bool URigVM::ValidateAllOperandsDuringLoad(FRigVMExtendedExecuteContext& Context
 		}
 
 
-		const URigVMMemoryStorage* MemoryForOperand = LocalMemory[InOperand.GetContainerIndex()];
+		const TRigVMMemoryStorage* MemoryForOperand = LocalMemory[InOperand.GetContainerIndex()];
 		if(InOperand.GetMemoryType() != ERigVMMemoryType::External)
 		{
 			if(!MemoryForOperand->IsValidIndex(InOperand.GetRegisterIndex()))
@@ -523,7 +535,7 @@ bool URigVM::ValidateAllOperandsDuringLoad(FRigVMExtendedExecuteContext& Context
 
 const URigVMHost* URigVM::GetHostCDO() const
 {
-	check(HasAnyFlags(RF_ClassDefaultObject | RF_DefaultSubObject));
+	ensure(HasAnyFlags(RF_ClassDefaultObject | RF_DefaultSubObject));
 
 	return GetTypedOuter<const URigVMHost>();
 }
@@ -576,7 +588,6 @@ void URigVM::Empty(FRigVMExtendedExecuteContext& Context)
 
 	OperandToDebugRegisters.Empty();
 
-	Context.CachedMemory.Empty();
 	Context.CachedMemoryHandles.Empty();
 	Context.ExternalVariableRuntimeData.Reset();
 }
@@ -640,6 +651,7 @@ URigVMMemoryStorage* URigVM::CreateMemoryByType(FRigVMExtendedExecuteContext& Co
 {
 	URigVMMemoryStorage* MemoryStorage = nullptr;
 
+#if !UE_RIGVM_PROPERTY_BAG_STORAGE_ENABLED
 	switch(InMemoryType)
 	{
 		case ERigVMMemoryType::Literal:
@@ -656,7 +668,7 @@ URigVMMemoryStorage* URigVM::CreateMemoryByType(FRigVMExtendedExecuteContext& Co
 		case ERigVMMemoryType::Work:
 		{
 			CreateMemoryByType(GetOuter(), Context.WorkMemoryStorageObject, InMemoryType, RF_NoFlags, bForceCreation);
-			//check(Context.WorkMemoryStorageObject->GetOuter() == GetOuter());
+			//check(GetWorkMemory(Context)->GetOuter() == GetOuter());
 			MemoryStorage = Context.WorkMemoryStorageObject;
 			break;
 		}
@@ -676,13 +688,47 @@ URigVMMemoryStorage* URigVM::CreateMemoryByType(FRigVMExtendedExecuteContext& Co
 			break;
 		}
 	}
+#endif // UE_RIGVM_PROPERTY_BAG_STORAGE_ENABLED
 
 	return MemoryStorage;
 }
 
-URigVMMemoryStorage* URigVM::GetMemoryByType(FRigVMExtendedExecuteContext& Context, ERigVMMemoryType InMemoryType/*, bool bCreateIfNeeded*/)
+void URigVM::GenerateMemoryType(FRigVMExtendedExecuteContext& Context, ERigVMMemoryType InMemoryType, const TArray<FRigVMPropertyDescription>* InProperties, bool bForceCreation)
 {
-	URigVMMemoryStorage* MemoryStorage = nullptr;
+	FRigVMMemoryStorageStruct NewBag = (InProperties != nullptr) ? FRigVMMemoryStorageStruct(InMemoryType, *InProperties) : FRigVMMemoryStorageStruct(InMemoryType);
+
+	switch(InMemoryType)
+	{
+		case ERigVMMemoryType::Literal:
+		{
+			LiteralMemoryStorage = NewBag;
+			break;
+		}
+
+		case ERigVMMemoryType::Work:
+		{
+			Context.WorkMemoryStorage = NewBag;
+			break;
+		}
+
+		case ERigVMMemoryType::Debug:
+		{
+#if WITH_EDITOR
+			Context.DebugMemoryStorage = NewBag;
+#endif
+			break;
+		}
+
+		default:
+		{
+			break;
+		}
+	}
+}
+
+TRigVMMemoryStorage* URigVM::GetMemoryByType(FRigVMExtendedExecuteContext& Context, ERigVMMemoryType InMemoryType/*, bool bCreateIfNeeded*/)
+{
+	TRigVMMemoryStorage* MemoryStorage = nullptr;
 
 	switch(InMemoryType)
 	{
@@ -694,15 +740,23 @@ URigVMMemoryStorage* URigVM::GetMemoryByType(FRigVMExtendedExecuteContext& Conte
 
 		case ERigVMMemoryType::Work:
 		{
+#if UE_RIGVM_PROPERTY_BAG_STORAGE_ENABLED
+			MemoryStorage = &Context.WorkMemoryStorage;
+#else
 			//check(Context.WorkMemoryStorageObject->GetOuter() == GetOuter());
 			MemoryStorage = Context.WorkMemoryStorageObject;
+#endif
 			break;
 		}
 
 		case ERigVMMemoryType::Debug:
 		{
+#if UE_RIGVM_PROPERTY_BAG_STORAGE_ENABLED
+			MemoryStorage = &Context.DebugMemoryStorage;
+#else
 			//check(Context.DebugMemoryStorageObject->GetOuter() == GetOuter());
 			MemoryStorage = Context.DebugMemoryStorageObject;
+#endif
 			break;
 		}
 
@@ -715,7 +769,7 @@ URigVMMemoryStorage* URigVM::GetMemoryByType(FRigVMExtendedExecuteContext& Conte
 	return MemoryStorage;
 }
 
-const URigVMMemoryStorage* URigVM::GetMemoryByType(const FRigVMExtendedExecuteContext& Context, ERigVMMemoryType InMemoryType) const
+const TRigVMMemoryStorage* URigVM::GetMemoryByType(const FRigVMExtendedExecuteContext& Context, ERigVMMemoryType InMemoryType) const
 {
 	return const_cast<URigVM*>(this)->GetMemoryByType(const_cast<FRigVMExtendedExecuteContext&>(Context), InMemoryType/*, false*/);
 }
@@ -753,7 +807,11 @@ void URigVM::ClearMemory_Internal()
 	}
 #endif
 
+#if UE_RIGVM_PROPERTY_BAG_STORAGE_ENABLED
+	LiteralMemoryStorage = FRigVMMemoryStorageStruct(ERigVMMemoryType::Literal);
+#else
 	LiteralMemoryStorageObject = nullptr;
+#endif
 
 	InvalidateCachedMemory_Internal();
 }
@@ -762,15 +820,20 @@ void URigVM::ClearMemory(FRigVMExtendedExecuteContext& Context)
 {
 	ClearMemory_Internal();
 
-	if(Context.WorkMemoryStorageObject)
+#if UE_RIGVM_PROPERTY_BAG_STORAGE_ENABLED
+	Context.WorkMemoryStorage = FRigVMMemoryStorageStruct(ERigVMMemoryType::Work);
+	Context.DebugMemoryStorage = FRigVMMemoryStorageStruct(ERigVMMemoryType::Debug);
+#else
+	if(GetWorkMemory(Context))
 	{
 		Context.WorkMemoryStorageObject = nullptr;
 	}
 
-	if(Context.DebugMemoryStorageObject)
+	if(GetDebugMemory(Context))
 	{
 		Context.DebugMemoryStorageObject = nullptr;
 	}
+#endif
 
 	InvalidateCachedMemory(Context);
 }
@@ -855,7 +918,7 @@ bool URigVM::ResumeExecution(FRigVMExtendedExecuteContext& Context)
 	return false;
 }
 
-bool URigVM::ResumeExecution(FRigVMExtendedExecuteContext& Context, TArrayView<URigVMMemoryStorage*> Memory, const FName& InEntryName)
+bool URigVM::ResumeExecution(FRigVMExtendedExecuteContext& Context, TArrayView<TRigVMMemoryStorage*> Memory, const FName& InEntryName)
 {
 	ResumeExecution(Context);
 	return Execute(Context, Memory, InEntryName) != ERigVMExecuteResult::Failed;
@@ -1083,7 +1146,7 @@ void URigVM::InstructionOpEval(FRigVMExtendedExecuteContext& Context, int32 Inst
 	}
 }
 
-void URigVM::PrepareMemoryForExecution(FRigVMExtendedExecuteContext& Context, TArrayView<URigVMMemoryStorage*> InMemory)
+void URigVM::PrepareMemoryForExecution(FRigVMExtendedExecuteContext& Context, TArrayView<TRigVMMemoryStorage*> InMemory)
 {
 	ensureMsgf(Context.ExecutingThreadId == FPlatformTLS::GetCurrentThreadId(), TEXT("RigVM::CacheMemoryHandlesIfRequired from multiple threads (%d and %d)"), Context.ExecutingThreadId, (int32)FPlatformTLS::GetCurrentThreadId());
 
@@ -1125,7 +1188,7 @@ void URigVM::PrepareMemoryForExecution(FRigVMExtendedExecuteContext& Context, TA
 	Context.CachedMemoryHandles.Reset(MemoryHandleCount);
 }
 
-void URigVM::CacheMemoryHandlesIfRequired(FRigVMExtendedExecuteContext& Context, TArrayView<URigVMMemoryStorage*> InMemory)
+void URigVM::CacheMemoryHandlesIfRequired(FRigVMExtendedExecuteContext& Context, TArrayView<TRigVMMemoryStorage*> InMemory)
 {
 	if (Instructions.Num() == 0 || InMemory.Num() == 0)
 	{
@@ -1365,7 +1428,7 @@ bool URigVM::ShouldHaltAtInstruction(FRigVMExtendedExecuteContext& Context, cons
 }
 #endif
 
-bool URigVM::Initialize(FRigVMExtendedExecuteContext& Context, TArrayView<URigVMMemoryStorage*> Memory)
+bool URigVM::Initialize(FRigVMExtendedExecuteContext& Context, TArrayView<TRigVMMemoryStorage*> Memory)
 {
 	if (Context.ExecutingThreadId != INDEX_NONE)
 	{
@@ -1383,7 +1446,7 @@ bool URigVM::Initialize(FRigVMExtendedExecuteContext& Context, TArrayView<URigVM
 	}
 
 	// changes to the layout of memory array should be reflected in GetContainerIndex()
-	TArray<URigVMMemoryStorage*> LocalMemory;
+	TArray<TRigVMMemoryStorage*> LocalMemory;
 	if (Memory.Num() == 0)
 	{
 		LocalMemory = GetLocalMemoryArray(Context);
@@ -1397,7 +1460,7 @@ bool URigVM::Initialize(FRigVMExtendedExecuteContext& Context, TArrayView<URigVM
 	return true;
 }
 
-bool URigVM::InitializeInstance(FRigVMExtendedExecuteContext& Context, TArrayView<URigVMMemoryStorage*> Memory)
+bool URigVM::InitializeInstance(FRigVMExtendedExecuteContext& Context, TArrayView<TRigVMMemoryStorage*> Memory)
 {
 	check(!Memory.IsEmpty());
 
@@ -1409,8 +1472,9 @@ bool URigVM::InitializeInstance(FRigVMExtendedExecuteContext& Context, TArrayVie
 	Context.LazyBranchInstanceData.Reset(LazyBranchSize);
 	Context.LazyBranchInstanceData.SetNumZeroed(LazyBranchSize);
 
+#if !UE_RIGVM_PROPERTY_BAG_STORAGE_ENABLED
 	// re-initialize work memory from CDO
-	if (URigVMMemoryStorage* WorkMemory = Memory[(int32)ERigVMMemoryType::Work])
+	if (TRigVMMemoryStorage* WorkMemory = Memory[(int32)ERigVMMemoryType::Work])
 	{
 		if (ensure(!WorkMemory->HasAnyFlags(RF_ClassDefaultObject)))
 		{
@@ -1439,19 +1503,20 @@ bool URigVM::InitializeInstance(FRigVMExtendedExecuteContext& Context, TArrayVie
 							UE_LOG(LogRigVM, Warning, TEXT("Property %s defaults to '%s'."), *Property->GetName(), *DefaultValue);
 							UE_LOG(LogRigVM, Warning, TEXT("Property %s defaults to '%s'."), *Property->GetName(), *DefaultValue);
 						}
-#endif
+#endif // UE_RIGVM_DEBUG_EXECUTION
 					}
 				}
 			}
 		}
 	}
-	
+#endif // UE_RIGVM_PROPERTY_BAG_STORAGE_ENABLED
+
 	CacheMemoryHandlesIfRequired(Context, Memory);
 
 	return true;
 }
 
-ERigVMExecuteResult URigVM::Execute(FRigVMExtendedExecuteContext& Context, TArrayView<URigVMMemoryStorage*> Memory, const FName& InEntryName)
+ERigVMExecuteResult URigVM::Execute(FRigVMExtendedExecuteContext& Context, TArrayView<TRigVMMemoryStorage*> Memory, const FName& InEntryName)
 {
 	// if this the first entry being executed - get ready for execution
 	const bool bIsRootEntry = Context.EntriesBeingExecuted.IsEmpty();
@@ -1478,7 +1543,7 @@ ERigVMExecuteResult URigVM::Execute(FRigVMExtendedExecuteContext& Context, TArra
 		}
 
 		check(!Memory.IsEmpty());
-		Context.CurrentMemory = Memory;
+		Context.CurrentVMMemory = Memory;
 	}
 
 	FRigVMByteCode& ByteCode = GetByteCode();
@@ -1574,7 +1639,7 @@ ERigVMExecuteResult URigVM::Execute(FRigVMExtendedExecuteContext& Context, TArra
 	if (ContextPublicData.bDebugExecution)
 	{
 		ContextPublicData.InstanceOpCodeEnum = StaticEnum<ERigVMOpCode>();
-		URigVMMemoryStorage* LiteralMemory = GetLiteralMemory(false);
+		TRigVMMemoryStorage* LiteralMemory = GetLiteralMemory();
 		ContextPublicData.DebugMemoryString = FString("\n\nLiteral Memory\n\n");
 		for (int32 PropertyIndex=0; PropertyIndex<LiteralMemory->Num(); ++PropertyIndex)
 		{
@@ -1592,7 +1657,7 @@ ERigVMExecuteResult URigVM::Execute(FRigVMExtendedExecuteContext& Context, TArra
 #if WITH_EDITOR
 	if(bIsRootEntry)
 	{
-		Context.CurrentMemory = TArrayView<URigVMMemoryStorage*>();
+		Context.CurrentVMMemory = TArrayView<TRigVMMemoryStorage*>();
 		
 		if (Context.HaltedAtBreakpoint.IsValid() && Context.CurrentExecuteResult != ERigVMExecuteResult::Halted)
 		{
@@ -1609,7 +1674,7 @@ ERigVMExecuteResult URigVM::Execute(FRigVMExtendedExecuteContext& Context, TArra
 ERigVMExecuteResult URigVM::ExecuteInstructions(FRigVMExtendedExecuteContext& Context, int32 InFirstInstruction, int32 InLastInstruction)
 {
 	// make we are already executing this VM
-	check(!Context.CurrentMemory.IsEmpty());
+	check(!Context.CurrentVMMemory.IsEmpty());
 	
 	FRigVMExecuteContext& ContextPublicData = Context.GetPublicData<>();
 	TGuardValue<uint16> InstructionIndexGuard(ContextPublicData.InstructionIndex, (uint16)InFirstInstruction);
@@ -1771,7 +1836,7 @@ ERigVMExecuteResult URigVM::ExecuteInstructions(FRigVMExtendedExecuteContext& Co
 
 				FRigVMMemoryHandle& SourceHandle = Context.CachedMemoryHandles[FirstHandleForInstruction[ContextPublicData.InstructionIndex]];
 				FRigVMMemoryHandle& TargetHandle = Context.CachedMemoryHandles[FirstHandleForInstruction[ContextPublicData.InstructionIndex] + 1];
-				URigVMMemoryStorage::CopyProperty(TargetHandle, SourceHandle);
+				TRigVMMemoryStorage::CopyProperty(TargetHandle, SourceHandle);
 					
 #if WITH_EDITOR
 				if(GetDebugMemory(Context) && GetDebugMemory(Context)->Num() > 0)
@@ -1936,7 +2001,7 @@ ERigVMExecuteResult URigVM::ExecuteInstructions(FRigVMExtendedExecuteContext& Co
 				{
 					// this will restore the public data after invoking the entry
 					TGuardValue<FRigVMExecuteContext> PublicDataGuard(ContextPublicData, ContextPublicData);
-					const ERigVMExecuteResult ExecuteResult = Execute(Context, Context.CurrentMemory, Op.EntryName);
+					const ERigVMExecuteResult ExecuteResult = Execute(Context, Context.CurrentVMMemory, Op.EntryName);
 					if(ExecuteResult != ERigVMExecuteResult::Succeeded)
 					{
 						return Context.CurrentExecuteResult = ExecuteResult;
@@ -2036,7 +2101,7 @@ ERigVMExecuteResult URigVM::ExecuteInstructions(FRigVMExtendedExecuteContext& Co
 		if (ContextPublicData.bDebugExecution)
 		{
 			TArray<FString> CurrentWorkMemory;
-			URigVMMemoryStorage* WorkMemory = GetWorkMemory(false);
+			TRigVMMemoryStorage* WorkMemory = GetWorkMemory();
 			int32 LineIndex = 0;
 			for (int32 PropertyIndex=0; PropertyIndex<WorkMemory->Num(); ++PropertyIndex, ++LineIndex)
 			{
@@ -2136,7 +2201,7 @@ FRigVMExternalVariable URigVM::GetExternalVariableByName(const FRigVMExtendedExe
 
 void URigVM::SetPropertyValueFromString(FRigVMExtendedExecuteContext& Context, const FRigVMOperand& InOperand, const FString& InDefaultValue)
 {
-	URigVMMemoryStorage* Memory = GetMemoryByType(Context, InOperand.GetMemoryType());
+	TRigVMMemoryStorage* Memory = GetMemoryByType(Context, InOperand.GetMemoryType());
 	if(Memory == nullptr)
 	{
 		return;
@@ -2373,7 +2438,7 @@ FString URigVM::GetOperandLabel(FRigVMExtendedExecuteContext& Context, const FRi
 	}
 	else
 	{
-		URigVMMemoryStorage* Memory = GetMemoryByType(Context, InOperand.GetMemoryType());
+		TRigVMMemoryStorage* Memory = GetMemoryByType(Context, InOperand.GetMemoryType());
 		if(Memory == nullptr)
 		{
 			return FString();
@@ -2420,7 +2485,7 @@ void URigVM::ClearDebugMemory(FRigVMExtendedExecuteContext& Context)
 
 void URigVM::CacheSingleMemoryHandle(FRigVMExtendedExecuteContext& Context, int32 InHandleIndex, const FRigVMBranchInfoKey& InBranchInfoKey, const FRigVMOperand& InArg, bool bForExecute)
 {
-	URigVMMemoryStorage* Memory = GetMemoryByType(Context, InArg.GetMemoryType()/*, false*/);
+	TRigVMMemoryStorage* Memory = GetMemoryByType(Context, InArg.GetMemoryType()/*, false*/);
 
 	if (InArg.GetMemoryType() == ERigVMMemoryType::External)
 	{
@@ -2492,7 +2557,7 @@ void URigVM::CopyOperandForDebuggingImpl(FRigVMExtendedExecuteContext& Context,
 {
 #if WITH_EDITOR
 
-	URigVMMemoryStorage* TargetMemory = GetDebugMemory(Context);
+	TRigVMMemoryStorage* TargetMemory = GetDebugMemory(Context);
 	if(TargetMemory == nullptr)
 	{
 		return;
@@ -2532,12 +2597,12 @@ void URigVM::CopyOperandForDebuggingImpl(FRigVMExtendedExecuteContext& Context,
 			const FProperty* SourceProperty = ExternalVariable.Property;
 			FRigVMExternalVariableRuntimeData& ExternalVariableRuntimeData = Context.ExternalVariableRuntimeData[ExternalVariableIndex];
 			const uint8* SourcePtr = ExternalVariableRuntimeData.Memory;
-			URigVMMemoryStorage::CopyProperty(TargetProperty, TargetPtr, SourceProperty, SourcePtr);
+			TRigVMMemoryStorage::CopyProperty(TargetProperty, TargetPtr, SourceProperty, SourcePtr);
 		}
 		return;
 	}
 
-	URigVMMemoryStorage* SourceMemory = GetMemoryByType(Context, InArg.GetMemoryType());
+	TRigVMMemoryStorage* SourceMemory = GetMemoryByType(Context, InArg.GetMemoryType());
 	if(SourceMemory == nullptr)
 	{
 		return;
@@ -2545,7 +2610,7 @@ void URigVM::CopyOperandForDebuggingImpl(FRigVMExtendedExecuteContext& Context,
 	const FProperty* SourceProperty = SourceMemory->GetProperties()[InArg.GetRegisterIndex()];
 	const uint8* SourcePtr = SourceMemory->GetData<uint8>(InArg.GetRegisterIndex());
 
-	URigVMMemoryStorage::CopyProperty(TargetProperty, TargetPtr, SourceProperty, SourcePtr);
+	TRigVMMemoryStorage::CopyProperty(TargetProperty, TargetPtr, SourceProperty, SourcePtr);
 	
 #endif
 }
@@ -2576,73 +2641,6 @@ void URigVM::RefreshExternalPropertyPaths()
 	}
 }
 
-//TArray<const UObject*> URigVM::GetUserDefinedDependencies(const TArray<const URigVMMemoryStorage*> InMemory)
-//{
-//	TArray<const UObject*> Dependencies;
-//	auto ProcessMemory = [&Dependencies](const URigVMMemoryStorage* Memory)
-//	{
-//		if(Memory == nullptr)
-//		{
-//			return;
-//		}
-//
-//		const TArray<const FProperty*>& Properties = Memory->GetProperties();
-//
-//		for(const FProperty* Property : Properties)
-//		{
-//			const FProperty* PropertyToVisit = Property;
-//			while(const FArrayProperty* ArrayProperty = CastField<FArrayProperty>(PropertyToVisit))
-//			{
-//				PropertyToVisit = ArrayProperty->Inner;
-//			}
-//			if(const FStructProperty* StructProperty = CastField<FStructProperty>(PropertyToVisit))
-//			{
-//				if(const UUserDefinedStruct* UserDefinedStruct = Cast<UUserDefinedStruct>(StructProperty->Struct))
-//				{
-//					Dependencies.AddUnique(UserDefinedStruct);
-//				}
-//			}
-//			else if(const FEnumProperty* EnumProperty = CastField<FEnumProperty>(PropertyToVisit))
-//			{
-//				if(const UUserDefinedEnum* UserDefinedEnum = Cast<UUserDefinedEnum>(EnumProperty->GetEnum()))
-//				{
-//					Dependencies.AddUnique(UserDefinedEnum);
-//				}
-//			}
-//			else if(const FByteProperty* ByteProperty = CastField<FByteProperty>(PropertyToVisit))
-//			{
-//				if(const UUserDefinedEnum* UserDefinedEnum = Cast<UUserDefinedEnum>(ByteProperty->Enum))
-//				{
-//					Dependencies.AddUnique(UserDefinedEnum);
-//				}
-//			}
-//		}
-//	};
-//
-//	for (const URigVMMemoryStorage* MemoryStorage : InMemory)
-//	{
-//		ProcessMemory(MemoryStorage);
-//	}
-//
-//	TArray<const FRigVMFunction*>& Functions = GetFunctions();
-//	for(const FRigVMFunction* Function : Functions)
-//	{
-//		const FRigVMRegistry& Registry = FRigVMRegistry::Get();
-//		const TArray<TRigVMTypeIndex>& TypeIndices = Function->GetArgumentTypeIndices();
-//		for(const TRigVMTypeIndex& TypeIndex : TypeIndices)
-//		{
-//			const FRigVMTemplateArgumentType& Type = Registry.GetType(TypeIndex);
-//			if(Cast<UUserDefinedStruct>(Type.CPPTypeObject) ||
-//				Cast<UUserDefinedEnum>(Type.CPPTypeObject))
-//			{
-//				Dependencies.AddUnique(Type.CPPTypeObject);
-//			}
-//		}
-//	}
-//
-//	return Dependencies;
-//}
-
 void URigVM::SetupInstructionTracking(FRigVMExtendedExecuteContext& Context, int32 InInstructionCount)
 {
 #if WITH_EDITOR
diff --git a/Engine/Plugins/Runtime/RigVM/Source/RigVM/Private/RigVMCore/RigVMExecuteContext.cpp b/Engine/Plugins/Runtime/RigVM/Source/RigVM/Private/RigVMCore/RigVMExecuteContext.cpp
--- a/Engine/Plugins/Runtime/RigVM/Source/RigVM/Private/RigVMCore/RigVMExecuteContext.cpp
+++ b/Engine/Plugins/Runtime/RigVM/Source/RigVM/Private/RigVMCore/RigVMExecuteContext.cpp
@@ -51,6 +51,14 @@ bool FRigVMExecuteContext::SerializeFromMismatchedTag(const FPropertyTag& Tag, F
 	return false;
 }
 
+PRAGMA_DISABLE_DEPRECATION_WARNINGS
+FRigVMExtendedExecuteContext::~FRigVMExtendedExecuteContext()
+{
+	Reset();
+}
+PRAGMA_ENABLE_DEPRECATION_WARNINGS
+
+
 void FRigVMExtendedExecuteContext::Reset()
 {
 	VMHash = 0;
@@ -60,11 +68,13 @@ void FRigVMExtendedExecuteContext::Reset()
 	WorkMemoryStorageObject = nullptr;
 	DebugMemoryStorageObject = nullptr;
 
-	CurrentMemory = TArrayView<URigVMMemoryStorage*>();
+	WorkMemoryStorage = FRigVMMemoryStorageStruct();
+	DebugMemoryStorage = FRigVMMemoryStorageStruct();
+
+	CurrentVMMemory = TArrayView<TRigVMMemoryStorage*>();
 	ExecutionReachedExit().Clear();
 
 	CachedMemoryHandles.Reset();
-	CachedMemory.Reset();
 
 #if WITH_EDITOR
 	DebugInfo = nullptr;
@@ -120,8 +130,13 @@ void FRigVMExtendedExecuteContext::ResetExecutionState()
 
 void FRigVMExtendedExecuteContext::CopyMemoryStorage(const FRigVMExtendedExecuteContext& Other, UObject* Outer)
 {
+#if UE_RIGVM_PROPERTY_BAG_STORAGE_ENABLED
+	WorkMemoryStorage = Other.WorkMemoryStorage;
+	DebugMemoryStorage = Other.DebugMemoryStorage;
+#else
 	CopyMemoryStorage(WorkMemoryStorageObject, Other.WorkMemoryStorageObject, Outer);
 	CopyMemoryStorage(DebugMemoryStorageObject, Other.DebugMemoryStorageObject, Outer);
+#endif // UE_RIGVM_PROPERTY_BAG_STORAGE_ENABLED
 }
 
 void FRigVMExtendedExecuteContext::CopyMemoryStorage(TObjectPtr<URigVMMemoryStorage>& TargetMemory, const TObjectPtr <URigVMMemoryStorage>& SourceMemory, UObject* Outer)
diff --git a/Engine/Plugins/Runtime/RigVM/Source/RigVM/Private/RigVMCore/RigVMMemoryCommon.cpp b/Engine/Plugins/Runtime/RigVM/Source/RigVM/Private/RigVMCore/RigVMMemoryCommon.cpp
--- a/Engine/Plugins/Runtime/RigVM/Source/RigVM/Private/RigVMCore/RigVMMemoryCommon.cpp
+++ b/Engine/Plugins/Runtime/RigVM/Source/RigVM/Private/RigVMCore/RigVMMemoryCommon.cpp
@@ -1,6 +1,7 @@
 // Copyright Epic Games, Inc. All Rights Reserved.
 
 #include "RigVMCore/RigVMMemoryCommon.h"
+#include "RigVMModule.h"
 
 #include UE_INLINE_GENERATED_CPP_BY_NAME(RigVMMemoryCommon)
 
@@ -17,3 +18,15 @@ void FRigVMOperand::Serialize(FArchive& Ar)
 	Ar << RegisterOffset;
 }
 
+void FRigVMMemoryStorageImportErrorContext::Serialize(const TCHAR* V, ELogVerbosity::Type Verbosity, const class FName& Category)
+{
+	if (bLogErrors)
+	{
+#if WITH_EDITOR
+		UE_LOG(LogRigVM, Display, TEXT("Skipping Importing To MemoryStorage: %s"), V);
+#else
+		UE_LOG(LogRigVM, Error, TEXT("Error Importing To MemoryStorage: %s"), V);
+#endif
+	}
+	NumErrors++;
+}
diff --git a/Engine/Plugins/Runtime/RigVM/Source/RigVM/Private/RigVMCore/RigVMMemoryStorage.cpp b/Engine/Plugins/Runtime/RigVM/Source/RigVM/Private/RigVMCore/RigVMMemoryStorage.cpp
--- a/Engine/Plugins/Runtime/RigVM/Source/RigVM/Private/RigVMCore/RigVMMemoryStorage.cpp
+++ b/Engine/Plugins/Runtime/RigVM/Source/RigVM/Private/RigVMCore/RigVMMemoryStorage.cpp
@@ -376,34 +376,6 @@ FString FRigVMPropertyDescription::GetTailCPPType() const
 
 ////////////////////////////////////////////////////////////////////////////////////////////////////////////////////////
 
-class FRigVMMemoryStorageImportErrorContext : public FOutputDevice
-{
-public:
-
-	bool bLogErrors;
-	int32 NumErrors;
-
-	FRigVMMemoryStorageImportErrorContext(bool InLogErrors = true)
-		: FOutputDevice()
-		, bLogErrors(InLogErrors)
-		, NumErrors(0)
-	{
-	}
-
-	void Serialize(const TCHAR* V, ELogVerbosity::Type Verbosity, const class FName& Category) override
-	{
-		if(bLogErrors)
-		{
-#if WITH_EDITOR
-			UE_LOG(LogRigVM, Display, TEXT("Skipping Importing To MemoryStorage: %s"), V);
-#else
-			UE_LOG(LogRigVM, Error, TEXT("Error Importing To MemoryStorage: %s"), V);
-#endif
-		}
-		NumErrors++;
-	}
-};
-
 void URigVMMemoryStorageGeneratorClass::PurgeClass(bool bRecompilingOnLoad)
 {
 	Super::PurgeClass(bRecompilingOnLoad);
diff --git a/Engine/Plugins/Runtime/RigVM/Source/RigVM/Private/RigVMCore/RigVMMemoryStorageStruct.cpp b/Engine/Plugins/Runtime/RigVM/Source/RigVM/Private/RigVMCore/RigVMMemoryStorageStruct.cpp
new file mode 100644
--- /dev/null
+++ b/Engine/Plugins/Runtime/RigVM/Source/RigVM/Private/RigVMCore/RigVMMemoryStorageStruct.cpp
@@ -0,0 +1,651 @@
+// Copyright Epic Games, Inc. All Rights Reserved.
+
+#include "RigVMCore/RigVMMemoryStorageStruct.h"
+#include "RigVMTypeUtils.h"
+#include "RigVMModule.h"
+
+
+namespace 
+{
+	const FString EmptyBraces = TEXT("()");
+	constexpr TCHAR BraceFormat[] = TEXT("(%s)");
+}
+
+
+FRigVMMemoryStorageStruct::FRigVMMemoryStorageStruct()
+	: FInstancedPropertyBag()
+	, MemoryType(ERigVMMemoryType::Invalid)
+{
+}
+
+FRigVMMemoryStorageStruct::FRigVMMemoryStorageStruct(ERigVMMemoryType InMemoryType)
+	: FInstancedPropertyBag()
+	, MemoryType(InMemoryType)
+{
+}
+
+FRigVMMemoryStorageStruct::FRigVMMemoryStorageStruct(ERigVMMemoryType InMemoryType, const TArray<FRigVMPropertyDescription>& InPropertyDescriptions, const TArray<FRigVMPropertyPathDescription>& InPropertyPaths)
+	: FInstancedPropertyBag()
+	, MemoryType(InMemoryType)
+{
+	AddProperties(InPropertyDescriptions, InPropertyPaths);
+}
+
+FRigVMMemoryStorageStruct::~FRigVMMemoryStorageStruct()
+{
+}
+
+bool FRigVMMemoryStorageStruct::Serialize(FArchive& Ar)
+{
+	Super::Serialize(Ar);
+
+	Ar << MemoryType;
+	Ar << PropertyPathDescriptions;
+
+	if (Ar.IsLoading())
+	{
+		// rebuild property list and property path list
+		RefreshLinkedProperties();
+		RefreshPropertyPaths();
+
+		CachedMemoryHash = 0;
+	}
+
+	return true;
+}
+
+void FRigVMMemoryStorageStruct::AddStructReferencedObjects(FReferenceCollector& Collector)
+{
+	Super::AddStructReferencedObjects(Collector);
+}
+
+void FRigVMMemoryStorageStruct::AddProperties(const TArray<FRigVMPropertyDescription>& InPropertyDescriptions, const TArray<FRigVMPropertyPathDescription>& InPropertyPathDescriptions)
+{
+	TArray<FPropertyBagPropertyDesc> BagDescriptors;
+	BagDescriptors.Reserve(InPropertyDescriptions.Num());
+
+	// Generate PropertyBag descriptors
+	for (const FRigVMPropertyDescription& RigVMDescriptor : InPropertyDescriptions)
+	{
+		if (!ensure(BagDescriptors.FindByPredicate([&RigVMDescriptor](const FPropertyBagPropertyDesc& Desc) { return Desc.Name == RigVMDescriptor.Name; }) == nullptr))
+		{
+			UE_LOG(LogRigVM, Error, TEXT("Detected duplicated name in the generated properties : [%s]"), *RigVMDescriptor.Name.ToString());
+		}
+
+		const FPropertyBagPropertyDesc& PropertyBagDescriptor = GeneratePropertyBagDescriptor(RigVMDescriptor);
+		if (PropertyBagDescriptor.ValueType != EPropertyBagPropertyType::None)
+		{
+			BagDescriptors.Add(PropertyBagDescriptor);
+		}
+	}
+
+	// Generate the struct with the passed types
+	Super::AddProperties(BagDescriptors);
+
+	PropertyPathDescriptions = InPropertyPathDescriptions;
+
+	RefreshLinkedProperties();
+	RefreshPropertyPaths();
+
+	// And set the passed text DefaultValues
+	SetDefaultValues(InPropertyDescriptions);
+
+	CachedMemoryHash = 0;
+}
+
+int32 FRigVMMemoryStorageStruct::GetPropertyIndex(const FProperty* InProperty) const
+{
+	if (InProperty != nullptr)
+	{
+		const TArray<const FProperty*>& Properties = GetProperties();
+
+		return Properties.IndexOfByKey(InProperty);
+	}
+
+	return INDEX_NONE;
+}
+
+int32 FRigVMMemoryStorageStruct::GetPropertyIndexByName(const FName& InName) const
+{
+	const FProperty* Property = FindPropertyByName(InName);
+	return GetPropertyIndex(Property);
+}
+
+FProperty* FRigVMMemoryStorageStruct::FindPropertyByName(const FName& InName) const
+{
+	const FName SanitizedName = FRigVMPropertyDescription::SanitizeName(InName);
+	return Value.GetScriptStruct()->FindPropertyByName(SanitizedName);
+}
+
+FRigVMOperand FRigVMMemoryStorageStruct::GetOperand(int32 InPropertyIndex, int32 InPropertyPathIndex) const
+{
+	if (GetProperties().IsValidIndex(InPropertyIndex))
+	{
+		if (InPropertyPathIndex != INDEX_NONE)
+		{
+			check(GetPropertyPaths().IsValidIndex(InPropertyPathIndex));
+			return FRigVMOperand(GetMemoryType(), InPropertyIndex, InPropertyPathIndex);
+		}
+		return FRigVMOperand(GetMemoryType(), InPropertyIndex);
+	}
+	return FRigVMOperand();
+}
+
+FRigVMOperand FRigVMMemoryStorageStruct::GetOperandByName(const FName& InName, int32 InPropertyPathIndex) const
+{
+	const int32 PropertyIndex = GetPropertyIndexByName(InName);
+	return GetOperand(PropertyIndex, InPropertyPathIndex);
+}
+
+void* FRigVMMemoryStorageStruct::GetContainerPtr() const
+{
+	return (void*)GetValue().GetMemory();
+}
+
+uint32 FRigVMMemoryStorageStruct::GetMemoryHash() const
+{
+	if (CachedMemoryHash != 0)
+	{
+		return CachedMemoryHash;
+	}
+
+	for (const FProperty* Property : LinkedProperties)
+	{
+		CachedMemoryHash = HashCombine(CachedMemoryHash, GetTypeHash(Property->GetFName().ToString()));
+		CachedMemoryHash = HashCombine(CachedMemoryHash, GetTypeHash(Property->GetCPPType()));
+	}
+
+	// for literals we also hash the content / defaults for each property
+	if (GetMemoryType() == ERigVMMemoryType::Literal)
+	{
+		for (const FProperty* Property : LinkedProperties)
+		{
+			FString DefaultValue;
+			Property->ExportTextItem_InContainer(DefaultValue, GetContainerPtr(), nullptr, nullptr, PPF_None, nullptr);
+			CachedMemoryHash = HashCombine(CachedMemoryHash, GetTypeHash(DefaultValue));
+		}
+	}
+
+	return CachedMemoryHash;
+}
+
+//****************************************************************************
+
+const TArray<const FProperty*> FRigVMMemoryStorageStruct::EmptyProperties;
+const TArray<FRigVMPropertyPath> FRigVMMemoryStorageStruct::EmptyPropertyPaths;
+
+FString FRigVMMemoryStorageStruct::GetDataAsString(int32 InPropertyIndex, int32 PortFlags)
+{
+	check(IsValidIndex(InPropertyIndex));
+	const uint8* Data = GetData<uint8>(InPropertyIndex);
+
+	FString RetValue;
+	GetProperty(InPropertyIndex)->ExportTextItem_Direct(RetValue, Data, nullptr, nullptr, PortFlags);
+	return RetValue;
+}
+
+FString FRigVMMemoryStorageStruct::GetDataAsString(const FRigVMOperand& InOperand, int32 PortFlags)
+{
+	const int32 PropertyIndex = InOperand.GetRegisterIndex();
+	check(IsValidIndex(PropertyIndex));
+	return GetDataAsString(PropertyIndex, PortFlags);
+}
+
+FString FRigVMMemoryStorageStruct::GetDataAsStringSafe(int32 InPropertyIndex, int32 PortFlags)
+{
+	if (!IsValidIndex(InPropertyIndex))
+	{
+		return FString();
+	}
+	return GetDataAsString(InPropertyIndex, PortFlags);
+}
+
+FString FRigVMMemoryStorageStruct::GetDataAsStringSafe(const FRigVMOperand& InOperand, int32 PortFlags)
+{
+	const int32 PropertyIndex = InOperand.GetRegisterIndex();
+	return GetDataAsStringSafe(PropertyIndex, PortFlags);
+}
+
+bool FRigVMMemoryStorageStruct::SetDataFromString(int32 InPropertyIndex, const FString& InValue)
+{
+	check(IsValidIndex(InPropertyIndex));
+	uint8* Data = GetData<uint8>(InPropertyIndex);
+
+	const FProperty* Property = GetProperty(InPropertyIndex);
+
+	FRigVMMemoryStorageImportErrorContext ErrorPipe(false);
+	Property->ImportText_Direct(*InValue, Data, nullptr, EPropertyPortFlags::PPF_None, &ErrorPipe);
+
+	if (ErrorPipe.NumErrors > 0)
+	{
+		// check if the value was provided as a single element
+		if (const FArrayProperty* ArrayProperty = CastField<FArrayProperty>(Property))
+		{
+			const FString ValueWithBraces = FString::Printf(BraceFormat, *InValue);
+
+			ErrorPipe = FRigVMMemoryStorageImportErrorContext(false);
+			Property->ImportText_Direct(*ValueWithBraces, Data, nullptr, EPropertyPortFlags::PPF_None, &ErrorPipe);
+		}
+	}
+
+	return ErrorPipe.NumErrors == 0;
+}
+
+//****************************************************************************
+
+FRigVMMemoryHandle FRigVMMemoryStorageStruct::GetHandle(int32 InPropertyIndex, const FRigVMPropertyPath* InPropertyPath)
+{
+	check(IsValidIndex(InPropertyIndex));
+
+	const FProperty* Property = GetProperty(InPropertyIndex);
+	uint8* Data = GetData<uint8>(InPropertyIndex);
+
+	return FRigVMMemoryHandle(Data, Property, InPropertyPath);
+}
+
+//****************************************************************************
+
+bool FRigVMMemoryStorageStruct::CopyProperty(
+	const FProperty* InTargetProperty,
+	uint8* InTargetPtr,
+	const FProperty* InSourceProperty,
+	const uint8* InSourcePtr)
+{
+	check(InTargetProperty != nullptr);
+	check(InSourceProperty != nullptr);
+	check(InTargetPtr != nullptr);
+	check(InSourcePtr != nullptr);
+
+	// This block below is there to support Large World Coordinates (LWC).
+	// We allow to link float and double pins (single and arrays) so we need
+	// to support copying values between those as well.
+	if (!InTargetProperty->SameType(InSourceProperty))
+	{
+		if (const FFloatProperty* TargetFloatProperty = CastField<FFloatProperty>(InTargetProperty))
+		{
+			if (const FDoubleProperty* SourceDoubleProperty = CastField<FDoubleProperty>(InSourceProperty))
+			{
+				if (TargetFloatProperty->ArrayDim == SourceDoubleProperty->ArrayDim)
+				{
+					float* TargetFloats = (float*)InTargetPtr;
+					double* SourceDoubles = (double*)InSourcePtr;
+					for (int32 Index = 0; Index < TargetFloatProperty->ArrayDim; Index++)
+					{
+						TargetFloats[Index] = (float)SourceDoubles[Index];
+					}
+					return true;
+				}
+			}
+		}
+		else if (const FDoubleProperty* TargetDoubleProperty = CastField<FDoubleProperty>(InTargetProperty))
+		{
+			if (const FFloatProperty* SourceFloatProperty = CastField<FFloatProperty>(InSourceProperty))
+			{
+				if (TargetDoubleProperty->ArrayDim == SourceFloatProperty->ArrayDim)
+				{
+					double* TargetDoubles = (double*)InTargetPtr;
+					float* SourceFloats = (float*)InSourcePtr;
+					for (int32 Index = 0; Index < TargetDoubleProperty->ArrayDim; Index++)
+					{
+						TargetDoubles[Index] = (double)SourceFloats[Index];
+					}
+					return true;
+				}
+			}
+		}
+		else if (const FByteProperty* TargetByteProperty = CastField<FByteProperty>(InTargetProperty))
+		{
+			if (const FEnumProperty* SourceEnumProperty = CastField<FEnumProperty>(InSourceProperty))
+			{
+				if (TargetByteProperty->Enum == SourceEnumProperty->GetEnum())
+				{
+					InTargetProperty->CopyCompleteValue(InTargetPtr, InSourcePtr);
+					return true;
+				}
+			}
+		}
+		else if (const FEnumProperty* TargetEnumProperty = CastField<FEnumProperty>(InTargetProperty))
+		{
+			if (const FByteProperty* SourceByteProperty = CastField<FByteProperty>(InSourceProperty))
+			{
+				if (TargetEnumProperty->GetEnum() == SourceByteProperty->Enum)
+				{
+					InTargetProperty->CopyCompleteValue(InTargetPtr, InSourcePtr);
+					return true;
+				}
+			}
+		}
+		else if (const FArrayProperty* TargetArrayProperty = CastField<FArrayProperty>(InTargetProperty))
+		{
+			if (const FArrayProperty* SourceArrayProperty = CastField<FArrayProperty>(InSourceProperty))
+			{
+				FScriptArrayHelper TargetArray(TargetArrayProperty, InTargetPtr);
+				FScriptArrayHelper SourceArray(SourceArrayProperty, InSourcePtr);
+
+				if (TargetArrayProperty->Inner->IsA<FFloatProperty>())
+				{
+					if (SourceArrayProperty->Inner->IsA<FDoubleProperty>())
+					{
+						TargetArray.Resize(SourceArray.Num());
+						for (int32 Index = 0; Index < TargetArray.Num(); Index++)
+						{
+							float* TargetFloat = (float*)TargetArray.GetRawPtr(Index);
+							const double* SourceDouble = (const double*)SourceArray.GetRawPtr(Index);
+							*TargetFloat = (float)*SourceDouble;
+						}
+						return true;
+					}
+				}
+				else if (TargetArrayProperty->Inner->IsA<FDoubleProperty>())
+				{
+					if (SourceArrayProperty->Inner->IsA<FFloatProperty>())
+					{
+						TargetArray.Resize(SourceArray.Num());
+						for (int32 Index = 0; Index < TargetArray.Num(); Index++)
+						{
+							double* TargetDouble = (double*)TargetArray.GetRawPtr(Index);
+							const float* SourceFloat = (const float*)SourceArray.GetRawPtr(Index);
+							*TargetDouble = (double)*SourceFloat;
+						}
+						return true;
+					}
+				}
+				else if (FByteProperty* TargetArrayInnerByteProperty = CastField<FByteProperty>(TargetArrayProperty->Inner))
+				{
+					if (FEnumProperty* SourceArrayInnerEnumProperty = CastField<FEnumProperty>(SourceArrayProperty->Inner))
+					{
+						if (TargetArrayInnerByteProperty->Enum == SourceArrayInnerEnumProperty->GetEnum())
+						{
+							InTargetProperty->CopyCompleteValue(InTargetPtr, InSourcePtr);
+							return true;
+						}
+					}
+				}
+				else if (FEnumProperty* TargetArrayInnerEnumProperty = CastField<FEnumProperty>(TargetArrayProperty->Inner))
+				{
+					if (FByteProperty* SourceArrayInnerByteProperty = CastField<FByteProperty>(SourceArrayProperty->Inner))
+					{
+						if (TargetArrayInnerEnumProperty->GetEnum() == SourceArrayInnerByteProperty->Enum)
+						{
+							InTargetProperty->CopyCompleteValue(InTargetPtr, InSourcePtr);
+							return true;
+						}
+					}
+				}
+			}
+		}
+
+		// if we reach this we failed since we are trying to copy
+		// between two properties which are not compatible
+		if (!InTargetProperty->SameType(InSourceProperty))
+		{
+			// Only log the issue once, rather than spam.
+			static TSet<FString> ReportedErrors;
+
+			FString TargetType, SourceType, TargetExtendedType, SourceExtendedType;
+			TargetType = InTargetProperty->GetCPPType(&TargetExtendedType);
+			SourceType = InSourceProperty->GetCPPType(&SourceExtendedType);
+			TargetType += TargetExtendedType;
+			SourceType += SourceExtendedType;
+
+			UPackage* Package = InTargetProperty->GetOutermost();
+
+			FString Message = FString::Printf(TEXT("Failed to copy %s (%s) to %s (%s) in package %s"),
+				*InSourceProperty->GetName(),
+				*SourceType,
+				*InTargetProperty->GetName(),
+				*TargetType,
+				Package ? *Package->GetName() : TEXT("<Unknown Package>"));
+
+			if (!ReportedErrors.Contains(Message))
+			{
+				UE_LOG(LogRigVM, Error, TEXT("%s"), *Message);
+				ReportedErrors.Add(Message);
+			}
+
+			return false;
+		}
+	}
+
+	// rely on the core to copy the property contents
+	InTargetProperty->CopyCompleteValue(InTargetPtr, InSourcePtr);
+	return true;
+}
+
+bool FRigVMMemoryStorageStruct::CopyProperty(
+	const FProperty* InTargetProperty,
+	uint8* InTargetPtr,
+	const FRigVMPropertyPath& InTargetPropertyPath,
+	const FProperty* InSourceProperty,
+	const uint8* InSourcePtr,
+	const FRigVMPropertyPath& InSourcePropertyPath)
+{
+	check(InTargetProperty != nullptr);
+	check(InSourceProperty != nullptr);
+	check(InTargetPtr != nullptr);
+	check(InSourcePtr != nullptr);
+
+	auto TraversePropertyPath = [](const FProperty*& Property, uint8*& MemoryPtr, const FRigVMPropertyPath& PropertyPath)
+	{
+		if (PropertyPath.IsEmpty())
+		{
+			return;
+		}
+
+		MemoryPtr = PropertyPath.GetData<uint8>(MemoryPtr, Property);
+		Property = PropertyPath.GetTailProperty();
+	};
+
+	uint8* SourcePtr = (uint8*)InSourcePtr;
+	TraversePropertyPath(InTargetProperty, InTargetPtr, InTargetPropertyPath);
+	TraversePropertyPath(InSourceProperty, SourcePtr, InSourcePropertyPath);
+
+	if (InTargetPtr == nullptr)
+	{
+		check(!InTargetPropertyPath.IsEmpty());
+		return false;
+	}
+	if (InSourcePtr == nullptr)
+	{
+		check(!InSourcePropertyPath.IsEmpty());
+		return false;
+	}
+
+	return CopyProperty(InTargetProperty, InTargetPtr, InSourceProperty, SourcePtr);
+}
+
+bool FRigVMMemoryStorageStruct::CopyProperty(
+	FRigVMMemoryStorageStruct* InTargetStorage,
+	int32 InTargetPropertyIndex,
+	const FRigVMPropertyPath& InTargetPropertyPath,
+	FRigVMMemoryStorageStruct* InSourceStorage,
+	int32 InSourcePropertyIndex,
+	const FRigVMPropertyPath& InSourcePropertyPath)
+{
+	check(InTargetStorage != nullptr);
+	check(InSourceStorage != nullptr);
+
+	const FProperty* TargetProperty = InTargetStorage->GetProperty(InTargetPropertyIndex);
+	const FProperty* SourceProperty = InSourceStorage->GetProperty(InSourcePropertyIndex);
+	uint8* TargetPtr = InTargetStorage->GetData<uint8>(TargetProperty);
+	uint8* SourcePtr = InSourceStorage->GetData<uint8>(SourceProperty);
+
+	return CopyProperty(TargetProperty, TargetPtr, InTargetPropertyPath, SourceProperty, SourcePtr, InSourcePropertyPath);
+}
+
+bool FRigVMMemoryStorageStruct::CopyProperty(
+	FRigVMMemoryHandle& InTargetHandle,
+	FRigVMMemoryHandle& InSourceHandle)
+{
+	return CopyProperty(
+		InTargetHandle.GetProperty(),
+		InTargetHandle.GetData(false),
+		InTargetHandle.GetPropertyPathRef(),
+		InSourceHandle.GetProperty(),
+		InSourceHandle.GetData(false),
+		InSourceHandle.GetPropertyPathRef());
+}
+
+
+//****************************************************************************
+
+void FRigVMMemoryStorageStruct::RefreshLinkedProperties()
+{
+	LinkedProperties.Reset();
+
+	if (const UPropertyBag* CurrentBagStruct = GetPropertyBagStruct())
+	{
+		const FProperty* Property = CastField<FProperty>(CurrentBagStruct->ChildProperties);
+		while (Property)
+		{
+			LinkedProperties.Add(Property);
+			Property = CastField<FProperty>(Property->Next);
+		}
+	}
+}
+
+void FRigVMMemoryStorageStruct::RefreshPropertyPaths()
+{
+	// Update the property paths based on the descriptions
+	PropertyPaths.SetNumZeroed(PropertyPathDescriptions.Num());
+	for (int32 PropertyPathIndex = 0; PropertyPathIndex < PropertyPaths.Num(); PropertyPathIndex++)
+	{
+		PropertyPaths[PropertyPathIndex] = FRigVMPropertyPath();
+
+		const int32 PropertyIndex = PropertyPathDescriptions[PropertyPathIndex].PropertyIndex;
+		if (LinkedProperties.IsValidIndex(PropertyIndex))
+		{
+			PropertyPaths[PropertyPathIndex] = FRigVMPropertyPath(
+				LinkedProperties[PropertyIndex],
+				PropertyPathDescriptions[PropertyPathIndex].SegmentPath);
+		}
+	}
+}
+
+void FRigVMMemoryStorageStruct::SetDefaultValues(const TArray<FRigVMPropertyDescription>& InPropertyDescriptions)
+{
+	// and store default values.
+	const TArray<const FProperty*>& Properties = GetProperties();
+	
+	const int32 NumProperties = Properties.Num();
+	if (!ensure(NumProperties == InPropertyDescriptions.Num()))
+	{
+		UE_LOG(LogRigVM, Error, TEXT("Number of properties in storage is different from descriptors, unsupported type or duplicated name in the generated properties"));
+		return;
+	}
+
+	for (int32 PropertyIndex = 0; PropertyIndex < NumProperties; PropertyIndex++)
+	{
+		const FRigVMPropertyDescription& RigVMDesc = InPropertyDescriptions[PropertyIndex];
+		const FString& DefaultValue = RigVMDesc.DefaultValue;
+
+		if (DefaultValue.IsEmpty() || DefaultValue == EmptyBraces)
+		{
+			continue;
+		}
+
+		SetDataFromString(PropertyIndex, DefaultValue);
+	}
+}
+
+//****************************************************************************
+
+/*static*/ FPropertyBagPropertyDesc FRigVMMemoryStorageStruct::GeneratePropertyBagDescriptor(const FRigVMPropertyDescription& RigVMDescriptor)
+{
+	FPropertyBagPropertyDesc Result;
+
+	EPropertyBagPropertyType PropertyBagType = EPropertyBagPropertyType::None;
+	FPropertyBagContainerTypes PropertyBagContainerTypes;
+
+	if (GetPropertyTypeDataFromVMDescriptor(RigVMDescriptor, PropertyBagType, PropertyBagContainerTypes))
+	{
+		if (PropertyBagContainerTypes.Num() > 0)
+		{
+			Result = FPropertyBagPropertyDesc(RigVMDescriptor.Name, PropertyBagContainerTypes, PropertyBagType, RigVMDescriptor.CPPTypeObject);
+		}
+		else
+		{
+			Result = FPropertyBagPropertyDesc(RigVMDescriptor.Name, PropertyBagType, RigVMDescriptor.CPPTypeObject);
+		}
+	}
+
+	return Result;
+}
+
+/*static*/ bool FRigVMMemoryStorageStruct::GetPropertyTypeDataFromVMDescriptor(const FRigVMPropertyDescription& RigVMDescriptor, EPropertyBagPropertyType & OutBagPropertyType, FPropertyBagContainerTypes& OutBagContainerTypes)
+{
+	EPropertyBagPropertyType BagPropertyType = EPropertyBagPropertyType::None;
+
+	FString VMType = RigVMDescriptor.CPPType;
+
+	OutBagContainerTypes.Reset(); // in case it is reused by the caller
+
+	if (RigVMTypeUtils::IsArrayType(VMType))
+	{
+		const int32 NumContainers = RigVMDescriptor.Containers.Num();
+		for (int i= 0; i < NumContainers; ++i)
+		{
+			switch (RigVMDescriptor.Containers[i])
+			{
+			case EPinContainerType::Array:
+				OutBagContainerTypes.Add(EPropertyBagContainerType::Array);
+				break;
+			case EPinContainerType::Set:
+				ensureMsgf(false, TEXT("Unsuported Set type container : %s"), *VMType);
+				break;
+			case EPinContainerType::Map:
+				ensureMsgf(false, TEXT("Unsuported Map type container : %s"), *VMType);
+				break;
+			default:
+				break;
+			}
+		}
+
+		VMType = RigVMDescriptor.GetTailCPPType();
+	}
+
+	static const FName IntTypeName(TEXT("int")); // type used by some engine tests
+
+	if (VMType == RigVMTypeUtils::BoolTypeName)
+	{
+		OutBagPropertyType = EPropertyBagPropertyType::Bool;
+	}
+	else if (VMType == RigVMTypeUtils::Int32TypeName || VMType == IntTypeName)
+	{
+		OutBagPropertyType = EPropertyBagPropertyType::Int32;
+	}
+	else if (VMType == RigVMTypeUtils::FloatTypeName)
+	{
+		OutBagPropertyType = EPropertyBagPropertyType::Float;
+	}
+	else if (VMType == RigVMTypeUtils::DoubleTypeName)
+	{
+		OutBagPropertyType = EPropertyBagPropertyType::Double;
+	}
+	else if (VMType == RigVMTypeUtils::FNameTypeName)
+	{
+		OutBagPropertyType = EPropertyBagPropertyType::Name;
+	}
+	else if (VMType == RigVMTypeUtils::FStringTypeName)
+	{
+		OutBagPropertyType = EPropertyBagPropertyType::String;
+	}
+	else if (Cast<UScriptStruct>(RigVMDescriptor.CPPTypeObject))
+	{
+		OutBagPropertyType = EPropertyBagPropertyType::Struct;
+	}
+	else if (Cast<UEnum>(RigVMDescriptor.CPPTypeObject))
+	{
+		OutBagPropertyType = EPropertyBagPropertyType::Enum;
+	}
+	else if (Cast<UObject>(RigVMDescriptor.CPPTypeObject))
+	{
+		OutBagPropertyType = EPropertyBagPropertyType::Object;
+	}
+	else
+	{
+		OutBagPropertyType = EPropertyBagPropertyType::None;
+	}
+
+	return OutBagPropertyType != EPropertyBagPropertyType::None;
+}
diff --git a/Engine/Plugins/Runtime/RigVM/Source/RigVM/Private/RigVMCore/RigVMNativized.cpp b/Engine/Plugins/Runtime/RigVM/Source/RigVM/Private/RigVMCore/RigVMNativized.cpp
--- a/Engine/Plugins/Runtime/RigVM/Source/RigVM/Private/RigVMCore/RigVMNativized.cpp
+++ b/Engine/Plugins/Runtime/RigVM/Source/RigVM/Private/RigVMCore/RigVMNativized.cpp
@@ -27,13 +27,13 @@ void URigVMNativized::Reset(bool IsIgnoringArchetypeRef)
 	ByteCodeStorage.Reset();
 }
 
-bool URigVMNativized::Initialize(FRigVMExtendedExecuteContext& Context, TArrayView<URigVMMemoryStorage*> Memory)
+bool URigVMNativized::Initialize(FRigVMExtendedExecuteContext& Context, TArrayView<TRigVMMemoryStorage*> Memory)
 {
 	// nothing to do here 
 	return true;
 }
 
-ERigVMExecuteResult URigVMNativized::Execute(FRigVMExtendedExecuteContext& Context, TArrayView<URigVMMemoryStorage*> Memory, const FName& InEntryName)
+ERigVMExecuteResult URigVMNativized::Execute(FRigVMExtendedExecuteContext& Context, TArrayView<TRigVMMemoryStorage*> Memory, const FName& InEntryName)
 {
 	// to be implemented by the generated code
 	return ERigVMExecuteResult::Failed;
diff --git a/Engine/Plugins/Runtime/RigVM/Source/RigVM/Private/RigVMHost.cpp b/Engine/Plugins/Runtime/RigVM/Source/RigVM/Private/RigVMHost.cpp
--- a/Engine/Plugins/Runtime/RigVM/Source/RigVM/Private/RigVMHost.cpp
+++ b/Engine/Plugins/Runtime/RigVM/Source/RigVM/Private/RigVMHost.cpp
@@ -260,10 +260,20 @@ bool URigVMHost::InitializeVM(const FName& InEventName)
 		return false;	// The rig did compile with errors
 	}
 
+#if UE_RIGVM_PROPERTY_BAG_STORAGE_ENABLED
+	if (!HasAnyFlags(RF_ClassDefaultObject))
+	{
+		if (URigVMHost* CDO = GetClass()->GetDefaultObject<URigVMHost>())
+		{
+			ExtendedExecuteContext.WorkMemoryStorage = CDO->GetExtendedExecuteContext().WorkMemoryStorage; // initialize the memory to CDO default state
+		}
+	}
+#endif // UE_RIGVM_PROPERTY_BAG_STORAGE_ENABLED
+
 	// update the VM's external variables
 	VM->SetExternalVariablesInstanceData(GetExtendedExecuteContext(), ExternalVariables);
 
-	TArray<URigVMMemoryStorage*> LocalMemory = VM->GetLocalMemoryArray(ExtendedExecuteContext);
+	TArray<TRigVMMemoryStorage*> LocalMemory = VM->GetLocalMemoryArray(ExtendedExecuteContext);
 	const bool bResult = VM->InitializeInstance(GetExtendedExecuteContext(), LocalMemory);
 	if(bResult)
 	{
@@ -606,6 +616,7 @@ bool URigVMHost::InitializeCDOVM()
 	check(VM != nullptr);
 	check(VM->HasAnyFlags(RF_ClassDefaultObject | RF_DefaultSubObject));
 
+#if !UE_RIGVM_PROPERTY_BAG_STORAGE_ENABLED
 	// Create the memory
 	const TArray<ERigVMMemoryType> MemoryTypes = { ERigVMMemoryType::Literal, ERigVMMemoryType::Work, ERigVMMemoryType::Debug };
 	for (ERigVMMemoryType MemoryType : MemoryTypes)
@@ -626,6 +637,7 @@ bool URigVMHost::InitializeCDOVM()
 		// But create a memory object in any case (required by sequencer rigs)
 		VM->CreateMemoryByType(ExtendedExecuteContext, MemoryType);
 	}
+#endif // !UE_RIGVM_PROPERTY_BAG_STORAGE_ENABLED
 
 	// update the VM's external variables
 	VM->ClearExternalVariables(GetExtendedExecuteContext());
@@ -655,13 +667,16 @@ bool URigVMHost::Execute_Internal(const FName& InEventName)
 	else
 	{
 		// sanity check the validity of the VM to ensure stability.
-		if(!VM->IsContextValidForExecution(Context) ||
-			!IsValidLowLevel() ||
-			!VM->IsValidLowLevel() ||
-			!GetLiteralMemory() ||
-			!GetLiteralMemory()->IsValidLowLevel() ||
-			!GetWorkMemory() ||
-			!GetWorkMemory()->IsValidLowLevel())
+		if(!VM->IsContextValidForExecution(Context)
+			|| !IsValidLowLevel()
+			|| !VM->IsValidLowLevel()
+#if !UE_RIGVM_PROPERTY_BAG_STORAGE_ENABLED
+			|| !GetLiteralMemory()
+			|| !GetLiteralMemory()->IsValidLowLevel()
+			|| !GetWorkMemory()
+			|| !GetWorkMemory()->IsValidLowLevel()
+#endif // !UE_RIGVM_PROPERTY_BAG_STORAGE_ENABLED
+		)
 		{
 			UE_LOG(LogRigVM, Warning, InvalidatedVMFormat, *GetClass()->GetName());
 			return false;
@@ -679,7 +694,7 @@ bool URigVMHost::Execute_Internal(const FName& InEventName)
 	
 	const bool bUseDebuggingSnapshots = !VM->IsNativized();
 	
-	TArray<URigVMMemoryStorage*> LocalMemory = VM->GetLocalMemoryArray(ExtendedExecuteContext);
+	TArray<TRigVMMemoryStorage*> LocalMemory = VM->GetLocalMemoryArray(ExtendedExecuteContext);
 
 #if WITH_EDITOR
 	if(bUseDebuggingSnapshots)
@@ -812,13 +827,13 @@ URigVM* URigVMHost::GetVM()
 	return VM;
 }
 
-URigVMMemoryStorage* URigVMHost::GetMemoryByType(ERigVMMemoryType InMemoryType)
+TRigVMMemoryStorage* URigVMHost::GetMemoryByType(ERigVMMemoryType InMemoryType)
 {
 	check(VM);
 	return VM->GetMemoryByType(ExtendedExecuteContext, InMemoryType);
 }
 
-const URigVMMemoryStorage* URigVMHost::GetMemoryByType(ERigVMMemoryType InMemoryType) const
+const TRigVMMemoryStorage* URigVMHost::GetMemoryByType(ERigVMMemoryType InMemoryType) const
 {
 	check(VM);
 	return VM->GetMemoryByType(ExtendedExecuteContext, InMemoryType);
@@ -1124,10 +1139,10 @@ void URigVMHost::GenerateUserDefinedDependenciesData(FRigVMExtendedExecuteContex
 	}
 }
 
-TArray<const UObject*> URigVMHost::GetUserDefinedDependencies(const TArray<const URigVMMemoryStorage*> InMemory)
+TArray<const UObject*> URigVMHost::GetUserDefinedDependencies(const TArray<const TRigVMMemoryStorage*> InMemory)
 {
 	TArray<const UObject*> Dependencies;
-	auto ProcessMemory = [&Dependencies](const URigVMMemoryStorage* Memory)
+	auto ProcessMemory = [&Dependencies](const TRigVMMemoryStorage* Memory)
 	{
 		if (Memory == nullptr)
 		{
@@ -1167,7 +1182,7 @@ TArray<const UObject*> URigVMHost::GetUserDefinedDependencies(const TArray<const
 		}
 	};
 
-	for (const URigVMMemoryStorage* MemoryStorage : InMemory)
+	for (const TRigVMMemoryStorage* MemoryStorage : InMemory)
 	{
 		ProcessMemory(MemoryStorage);
 	}
diff --git a/Engine/Plugins/Runtime/RigVM/Source/RigVM/Private/RigVMModule.cpp b/Engine/Plugins/Runtime/RigVM/Source/RigVM/Private/RigVMModule.cpp
--- a/Engine/Plugins/Runtime/RigVM/Source/RigVM/Private/RigVMModule.cpp
+++ b/Engine/Plugins/Runtime/RigVM/Source/RigVM/Private/RigVMModule.cpp
@@ -9,8 +9,13 @@
 #include "HAL/IConsoleManager.h"
 #include "RigVMObjectVersion.h"
 #include "UObject/DevObjectVersion.h"
+#if WITH_EDITOR && UE_RIGVM_PROPERTY_BAG_STORAGE_ENABLED
+#include "PropertyBagDetails.h"
+#include "PropertyEditorModule.h"
+#endif
+
 
-IMPLEMENT_MODULE(FDefaultModuleImpl, RigVM);
+IMPLEMENT_MODULE(FRigVMModule, RigVM)
 
 DEFINE_LOG_CATEGORY(LogRigVM);
 
@@ -27,6 +32,28 @@ TAutoConsoleVariable<bool> CVarRigVMEnableUObjects(TEXT("RigVM.UObjectSupport"),
 TAutoConsoleVariable<bool> CVarRigVMEnableUInterfaces(TEXT("RigVM.UInterfaceSupport"), true, TEXT("When true the RigVMCompiler will allow UInterfaces."));
 #endif
 
+void FRigVMModule::StartupModule()
+{
+#if WITH_EDITOR && UE_RIGVM_PROPERTY_BAG_STORAGE_ENABLED
+	// Register the details customizer
+	FPropertyEditorModule& PropertyModule = FModuleManager::LoadModuleChecked<FPropertyEditorModule>("PropertyEditor");
+	PropertyModule.RegisterCustomPropertyTypeLayout("RigVMMemoryStorageStruct", FOnGetPropertyTypeCustomizationInstance::CreateStatic(&FPropertyBagDetails::MakeInstance));
+	PropertyModule.NotifyCustomizationModuleChanged();
+#endif
+}
+
+void FRigVMModule::ShutdownModule()
+{
+#if WITH_EDITOR && UE_RIGVM_PROPERTY_BAG_STORAGE_ENABLED
+	// Unregister the details customization
+	if (FPropertyEditorModule* PropertyModule = FModuleManager::Get().GetModulePtr<FPropertyEditorModule>("PropertyEditor"))
+	{
+		PropertyModule->UnregisterCustomPropertyTypeLayout("RigVMMemoryStorageStruct");
+		PropertyModule->NotifyCustomizationModuleChanged();
+	}
+#endif
+}
+
 bool RigVMCore::SupportsUObjects()
 {
 #if UE_RIGVM_UOBJECT_PROPERTIES_ENABLED
diff --git a/Engine/Plugins/Runtime/RigVM/Source/RigVM/Public/RigVMCore/RigVM.h b/Engine/Plugins/Runtime/RigVM/Source/RigVM/Public/RigVMCore/RigVM.h
--- a/Engine/Plugins/Runtime/RigVM/Source/RigVM/Public/RigVMCore/RigVM.h
+++ b/Engine/Plugins/Runtime/RigVM/Source/RigVM/Public/RigVMCore/RigVM.h
@@ -19,10 +19,11 @@
 #include "RigVMCore/RigVMFunction.h"
 #include "RigVMCore/RigVMMemoryCommon.h"
 #include "RigVMCore/RigVMPropertyPath.h"
+#include "RigVMCore/RigVMMemoryStorage.h"
+#include "RigVMCore/RigVMMemoryStorageStruct.h"
 #include "RigVMExecuteContext.h"
 #include "RigVMMemory.h"
 #include "RigVMMemoryDeprecated.h"
-#include "RigVMMemoryStorage.h"
 #include "RigVMRegistry.h"
 #include "RigVMStatistics.h"
 #include "Templates/Function.h"
@@ -221,19 +222,28 @@ public:
 	UE_DEPRECATED(5.3, "Please, use Initialize with Context param")
 	virtual bool Initialize(TArrayView<URigVMMemoryStorage*> Memory) { return false; }
 
+	UE_DEPRECATED(5.4, "Please, use Initialize with TRigVMMemoryStorage param")
+	virtual bool Initialize(FRigVMExtendedExecuteContext& Context, TArrayView<TRigVMMemoryStorageDeprecatedType*> Memory) { return false; }
+
 	// Prepares caches and memory for execution
-	virtual bool Initialize(FRigVMExtendedExecuteContext& Context, TArrayView<URigVMMemoryStorage*> Memory);
+	virtual bool Initialize(FRigVMExtendedExecuteContext& Context, TArrayView<TRigVMMemoryStorage*> Memory);
+
+	UE_DEPRECATED(5.4, "Please, use InitializeInstance with TRigVMMemoryStorage param")
+	virtual bool InitializeInstance(FRigVMExtendedExecuteContext& Context, TArrayView<TRigVMMemoryStorageDeprecatedType*> Memory) { return false; }
 
 	// Initializes cached memory handles and copies work memory from the CDO to the Context
-	virtual bool InitializeInstance(FRigVMExtendedExecuteContext& Context, TArrayView<URigVMMemoryStorage*> Memory);
+	virtual bool InitializeInstance(FRigVMExtendedExecuteContext& Context, TArrayView<TRigVMMemoryStorage*> Memory);
 
 	UE_DEPRECATED(5.3, "Please, use Execute with Context param")
 	virtual ERigVMExecuteResult Execute(TArrayView<URigVMMemoryStorage*> Memory, const FName& InEntryName = NAME_None) { return ERigVMExecuteResult::Failed; }
 
+	UE_DEPRECATED(5.4, "Please, use Execute with TRigVMMemoryStorage param")
+	virtual ERigVMExecuteResult Execute(FRigVMExtendedExecuteContext& Context, TArrayView<TRigVMMemoryStorageDeprecatedType*> Memory, const FName& InEntryName = NAME_None) { return ERigVMExecuteResult::Failed; }
+
 	// Executes the VM.
 	// You can optionally provide external memory to the execution
 	// and provide optional additional operands.
-	virtual ERigVMExecuteResult Execute(FRigVMExtendedExecuteContext& Context, TArrayView<URigVMMemoryStorage*> Memory, const FName& InEntryName = NAME_None);
+	virtual ERigVMExecuteResult Execute(FRigVMExtendedExecuteContext& Context, TArrayView<TRigVMMemoryStorage*> Memory, const FName& InEntryName = NAME_None);
 
 	UE_DEPRECATED(5.3, "Please, use Execute with Context param")
 	virtual bool Execute(const FName& InEntryName = NAME_None) { return false; }
@@ -266,64 +276,74 @@ public:
 	virtual FString GetRigVMFunctionName(int32 InFunctionIndex) const;
 
 	UE_DEPRECATED(5.4, "GetMemoryByType has been deprecated from the VM. Please, use GetWorkMemory from VMHost or the version with a Context param")
-	virtual URigVMMemoryStorage* GetMemoryByType(ERigVMMemoryType InMemoryType, bool bCreateIfNeeded = true) { return nullptr; }
+	virtual TRigVMMemoryStorageDeprecatedType* GetMemoryByType(ERigVMMemoryType InMemoryType, bool bCreateIfNeeded = true) { return nullptr; }
 
 	// Creates a new memory storage by type
 	void CreateMemoryByType(UObject* Outer, TObjectPtr<URigVMMemoryStorage>& MemoryStorage, ERigVMMemoryType InMemoryType, EObjectFlags ObjectFlags, bool bForceCreation);
 	virtual URigVMMemoryStorage* CreateMemoryByType(FRigVMExtendedExecuteContext& Context, ERigVMMemoryType InMemoryType, bool bForceCreation = false);
 
+	virtual void GenerateMemoryType(FRigVMExtendedExecuteContext& Context, ERigVMMemoryType InMemoryType, const TArray<FRigVMPropertyDescription>* InProperties, bool bForceCreation = false);
+	
 	// Returns a memory storage by type
-	virtual URigVMMemoryStorage* GetMemoryByType(FRigVMExtendedExecuteContext& Context, ERigVMMemoryType InMemoryType);
-	virtual const URigVMMemoryStorage* GetMemoryByType(const FRigVMExtendedExecuteContext& Context, ERigVMMemoryType InMemoryType) const;
+	virtual TRigVMMemoryStorage* GetMemoryByType(FRigVMExtendedExecuteContext& Context, ERigVMMemoryType InMemoryType);
+	virtual const TRigVMMemoryStorage* GetMemoryByType(const FRigVMExtendedExecuteContext& Context, ERigVMMemoryType InMemoryType) const;
 
-	UE_DEPRECATED(5.4, "This version of GetLiteralMemory has been deprecated from the VM. Please, use GetLiteralMemory without bool parameter")
-	URigVMMemoryStorage* GetLiteralMemory(bool bCreateIfNeeded) { return nullptr; }
+	UE_DEPRECATED(5.4, "This version of GetLiteralMemory has been deprecated from the VM. Please, use GetLiteralMemory from Host or version with no parameter.")
+	TRigVMMemoryStorageDeprecatedType* GetLiteralMemory(bool bCreateIfNeeded) { return nullptr; }
 
 	// The default const literal memory
-	URigVMMemoryStorage* GetLiteralMemory()
+	TRigVMMemoryStorage* GetLiteralMemory()
 	{
+#if UE_RIGVM_PROPERTY_BAG_STORAGE_ENABLED
+		return &LiteralMemoryStorage;
+#else
 		return LiteralMemoryStorageObject;
+#endif
 	}
-	const URigVMMemoryStorage* GetLiteralMemory() const
+	const TRigVMMemoryStorage* GetLiteralMemory() const
 	{
+#if UE_RIGVM_PROPERTY_BAG_STORAGE_ENABLED
+		return &LiteralMemoryStorage;
+#else
 		return LiteralMemoryStorageObject;
+#endif
 	}
 
 	UE_DEPRECATED(5.4, "GetWorkMemory has been deprecated from the VM. Please, use GetWorkMemory from VMHost or the version with a Context param")
-	URigVMMemoryStorage* GetWorkMemory(bool bCreateIfNeeded = true) { return nullptr; }
+	TRigVMMemoryStorageDeprecatedType* GetWorkMemory(bool bCreateIfNeeded = true) { return nullptr; }
 
 	// The default mutable work memory
-	URigVMMemoryStorage* GetWorkMemory(FRigVMExtendedExecuteContext& Context)
+	TRigVMMemoryStorage* GetWorkMemory(FRigVMExtendedExecuteContext& Context)
 	{
 		return GetMemoryByType(Context, ERigVMMemoryType::Work);
 	}
 
-	const URigVMMemoryStorage* GetWorkMemory(const FRigVMExtendedExecuteContext& Context) const
+	const TRigVMMemoryStorage* GetWorkMemory(const FRigVMExtendedExecuteContext& Context) const
 	{
 		return GetMemoryByType(Context, ERigVMMemoryType::Work);
 	}
 
 	UE_DEPRECATED(5.4, "GetDebugMemory has been deprecated from the VM. Please, use GetDebugMemory from VMHost or the version with a Context param")
-	URigVMMemoryStorage* GetDebugMemory(bool bCreateIfNeeded = true) { return nullptr; }
+	TRigVMMemoryStorageDeprecatedType* GetDebugMemory(bool bCreateIfNeeded = true) { return nullptr; }
 
 	// The default debug watch memory
-	URigVMMemoryStorage* GetDebugMemory(FRigVMExtendedExecuteContext& Context)
+	TRigVMMemoryStorage* GetDebugMemory(FRigVMExtendedExecuteContext& Context)
 	{
 		return GetMemoryByType(Context, ERigVMMemoryType::Debug);
 	}
 
-	const URigVMMemoryStorage* GetDebugMemory(const FRigVMExtendedExecuteContext& Context) const
+	const TRigVMMemoryStorage* GetDebugMemory(const FRigVMExtendedExecuteContext& Context) const
 	{
 		return GetMemoryByType(Context, ERigVMMemoryType::Debug);
 	}
 
 	UE_DEPRECATED(5.4, "Please, use GetLocalMemoryArray with a Context param")
-	TArray<URigVMMemoryStorage*> GetLocalMemoryArray() { return TArray<URigVMMemoryStorage*>(); }
+	TArray<TRigVMMemoryStorageDeprecatedType*> GetLocalMemoryArray() { return TArray<TRigVMMemoryStorageDeprecatedType*>(); }
 
 	// returns all memory storages as an array
-	TArray<URigVMMemoryStorage*> GetLocalMemoryArray(FRigVMExtendedExecuteContext& Context)
+	TArray<TRigVMMemoryStorage*> GetLocalMemoryArray(FRigVMExtendedExecuteContext& Context)
 	{
-		TArray<URigVMMemoryStorage*> LocalMemory;
+		TArray<TRigVMMemoryStorage*> LocalMemory;
 		LocalMemory.Add(GetWorkMemory(Context));
 		LocalMemory.Add(GetLiteralMemory());
 		LocalMemory.Add(GetDebugMemory(Context));
@@ -346,6 +366,9 @@ public:
 	UPROPERTY()
 	TObjectPtr<URigVMMemoryStorage> LiteralMemoryStorageObject;
 
+	UPROPERTY()
+	FRigVMMemoryStorageStruct LiteralMemoryStorage;
+
 #if WITH_EDITORONLY_DATA
 	// Deprecated 5.4, please use the DebugMemory in the ExtendedExecuteContext
 	UPROPERTY()
@@ -450,9 +473,11 @@ public:
 	bool ResumeExecution() { return false; }
 	UE_DEPRECATED(5.3, "Please, use ResumeExecution with Context param")
 	bool ResumeExecution(TArrayView<URigVMMemoryStorage*> Memory, const FName& InEntryName = NAME_None) { return false; }
+	UE_DEPRECATED(5.4, "Please, use ResumeExecution with TRigVMMemoryStorage param")
+	bool ResumeExecution(FRigVMExtendedExecuteContext& Context, TArrayView<TRigVMMemoryStorageDeprecatedType*> Memory, const FName& InEntryName = NAME_None) { return false; }
 
 	bool ResumeExecution(FRigVMExtendedExecuteContext& Context);
-	bool ResumeExecution(FRigVMExtendedExecuteContext& Context, TArrayView<URigVMMemoryStorage*> Memory, const FName& InEntryName = NAME_None);
+	bool ResumeExecution(FRigVMExtendedExecuteContext& Context, TArrayView<TRigVMMemoryStorage*> Memory, const FName& InEntryName = NAME_None);
 #endif
 
 	// Returns the parameters of the VM
@@ -466,15 +491,15 @@ public:
 
 	FRigVMParameter AddParameter(FRigVMExtendedExecuteContext& Context, ERigVMParameterType InType, const FName& InParameterName, const FName& InWorkMemoryPropertyName)
 	{
-		check(Context.WorkMemoryStorageObject);
+		check(GetWorkMemory(Context));
 
 		if(ParametersNameMap.Contains(InParameterName))
 		{
 			return FRigVMParameter();
 		}
 
-		const FProperty* Property = Context.WorkMemoryStorageObject->FindPropertyByName(InWorkMemoryPropertyName);
-		const int32 PropertyIndex = Context.WorkMemoryStorageObject->GetPropertyIndex(Property);
+		const FProperty* Property = GetWorkMemory(Context)->FindPropertyByName(InWorkMemoryPropertyName);
+		const int32 PropertyIndex = GetWorkMemory(Context)->GetPropertyIndex(Property);
 
 		UScriptStruct* Struct = nullptr;
 		if(const FStructProperty* StructProperty = CastField<FStructProperty>(Property))
@@ -494,11 +519,11 @@ public:
 	int32 GetParameterArraySize(FRigVMExtendedExecuteContext& Context, const FRigVMParameter& InParameter)
 	{
 		const int32 PropertyIndex = InParameter.GetRegisterIndex();
-		const FProperty* Property = Context.WorkMemoryStorageObject->GetProperties()[PropertyIndex];
+		const FProperty* Property = GetWorkMemory(Context)->GetProperties()[PropertyIndex];
 		const FArrayProperty* ArrayProperty = CastField<FArrayProperty>(Property);
 		if(ArrayProperty)
 		{
-			FScriptArrayHelper ArrayHelper(ArrayProperty, Context.WorkMemoryStorageObject->GetData<uint8>(PropertyIndex));
+			FScriptArrayHelper ArrayHelper(ArrayProperty, GetWorkMemory(Context)->GetData<uint8>(PropertyIndex));
 			return ArrayHelper.Num();
 		}
 		return 1;
@@ -532,9 +557,9 @@ public:
 	{
 		if (InParameter.GetRegisterIndex() != INDEX_NONE)
 		{
-			if(Context.WorkMemoryStorageObject->IsArray(InParameter.GetRegisterIndex()))
+			if(GetWorkMemory(Context)->IsArray(InParameter.GetRegisterIndex()))
 			{
-				TArray<T>& Storage = *Context.WorkMemoryStorageObject->GetData<TArray<T>>(InParameter.GetRegisterIndex());
+				TArray<T>& Storage = *GetWorkMemory(Context)->GetData<TArray<T>>(InParameter.GetRegisterIndex());
 				if(Storage.IsValidIndex(InArrayIndex))
 				{
 					return Storage[InArrayIndex];
@@ -542,10 +567,10 @@ public:
 			}
 			else
 			{
-				return *Context.WorkMemoryStorageObject->GetData<T>(InParameter.GetRegisterIndex());
+				return *GetWorkMemory(Context)->GetData<T>(InParameter.GetRegisterIndex());
 			}
 			
-			return *Context.WorkMemoryStorageObject->GetData<T>(InParameter.GetRegisterIndex());
+			return *GetWorkMemory(Context)->GetData<T>(InParameter.GetRegisterIndex());
 		}
 		return DefaultValue;
 	}
@@ -580,9 +605,9 @@ public:
 	{
 		if (InParameter.GetRegisterIndex() != INDEX_NONE)
 		{
-			if(Context.WorkMemoryStorageObject->IsArray(InParameter.GetRegisterIndex()))
+			if(GetWorkMemory(Context)->IsArray(InParameter.GetRegisterIndex()))
 			{
-				TArray<T>& Storage = *Context.WorkMemoryStorageObject->GetData<TArray<T>>(InParameter.GetRegisterIndex());
+				TArray<T>& Storage = *GetWorkMemory(Context)->GetData<TArray<T>>(InParameter.GetRegisterIndex());
 				if(Storage.IsValidIndex(InArrayIndex))
 				{
 					Storage[InArrayIndex] = InNewValue;
@@ -590,7 +615,7 @@ public:
 			}
 			else
 			{
-				T& Storage = *Context.WorkMemoryStorageObject->GetData<T>(InParameter.GetRegisterIndex());
+				T& Storage = *GetWorkMemory(Context)->GetData<T>(InParameter.GetRegisterIndex());
 				Storage = InNewValue;
 			}
 		}
@@ -781,13 +806,13 @@ public:
 	FRigVMStatistics GetStatistics() const
 	{
 		FRigVMStatistics Statistics;
-		if(LiteralMemoryStorageObject)
+		if(GetLiteralMemory())
 		{
-			Statistics.LiteralMemory = LiteralMemoryStorageObject->GetStatistics();
+			//Statistics.LiteralMemory = GetLiteralMemory()->GetStatistics();
 		}
-		//if(Context.WorkMemoryStorageObject)
+		//if(GetWorkMemory(Context))
 		//{
-		//	Statistics.WorkMemory = WorkMemoryStorageObject->GetStatistics();
+		//	Statistics.WorkMemory = GetWorkMemory(Context)->GetStatistics();
 		//}
 
 		Statistics.ByteCode = ByteCodePtr->GetStatistics();
@@ -901,13 +926,10 @@ public:
 	void InvalidateCachedMemory() {}
 	void InvalidateCachedMemory(FRigVMExtendedExecuteContext& Context);
 
-	//const TMap<FString, FSoftObjectPath>& GetUserDefinedStructGuidToObjectPath() const { return UserDefinedStructGuidToPathName; }
-	//TArray<const UObject*> GetUserDefinedDependencies(const TArray<const URigVMMemoryStorage*> InMemory);
-
 private:
 	void InstructionOpEval(FRigVMExtendedExecuteContext& Context, int32 InstructionIndex, int32 InHandleBaseIndex, const TFunctionRef<void(FRigVMExtendedExecuteContext& Context, int32 InHandleIndex, const FRigVMBranchInfoKey& InBranchInfoKey, const FRigVMOperand& InArg)>& InOpFunc);
-	void PrepareMemoryForExecution(FRigVMExtendedExecuteContext& Context, TArrayView<URigVMMemoryStorage*> InMemory);
-	void CacheMemoryHandlesIfRequired(FRigVMExtendedExecuteContext& Context, TArrayView<URigVMMemoryStorage*> InMemory);
+	void PrepareMemoryForExecution(FRigVMExtendedExecuteContext& Context, TArrayView<TRigVMMemoryStorage*> InMemory);
+	void CacheMemoryHandlesIfRequired(FRigVMExtendedExecuteContext& Context, TArrayView<TRigVMMemoryStorage*> InMemory);
 	void RebuildByteCodeOnLoad();
 
 	UPROPERTY(transient)
@@ -960,7 +982,6 @@ private:
 	UPROPERTY()
 	TArray<FRigVMParameter> Parameters;
 
-	UPROPERTY()
 	TMap<FName, int32> ParametersNameMap;
 
 	TArray<uint32> FirstHandleForInstruction;
@@ -1065,4 +1086,4 @@ protected:
 	friend class URigVMCompiler;
 	friend struct FRigVMCompilerWorkData;
 	friend struct FRigVMCodeGenerator;
-};
\ No newline at end of file
+};
diff --git a/Engine/Plugins/Runtime/RigVM/Source/RigVM/Public/RigVMCore/RigVMExecuteContext.h b/Engine/Plugins/Runtime/RigVM/Source/RigVM/Public/RigVMCore/RigVMExecuteContext.h
--- a/Engine/Plugins/Runtime/RigVM/Source/RigVM/Public/RigVMCore/RigVMExecuteContext.h
+++ b/Engine/Plugins/Runtime/RigVM/Source/RigVM/Public/RigVMCore/RigVMExecuteContext.h
@@ -6,6 +6,8 @@
 #include "Containers/ArrayView.h"
 #include "Containers/UnrealString.h"
 #include "CoreMinimal.h"
+#include "Components/SceneComponent.h"
+#include "GameFramework/Actor.h"
 #include "Logging/LogCategory.h"
 #include "Logging/LogMacros.h"
 #include "Logging/TokenizedMessage.h"
@@ -14,20 +16,19 @@
 #include "RigVMDefines.h"
 #include "RigVMExternalVariable.h"
 #include "RigVMModule.h"
+#include "RigVMCore/RigVMDebugInfo.h"
+#include "RigVMCore/RigVMNameCache.h"
+#include "RigVMCore/RigVMMemoryStorageStruct.h"
+#include "RigVMLog.h"
+#include "RigVMDrawInterface.h"
+#include "RigVMDrawContainer.h"
 #include "Templates/SharedPointer.h"
 #include "Trace/Detail/Channel.h"
 #include "UObject/NameTypes.h"
 #include "UObject/ObjectMacros.h"
+#include "UObject/StructOnScope.h"
 #include "UObject/UnrealNames.h"
 #include "UObject/UnrealType.h"
-#include "RigVMCore/RigVMDebugInfo.h"
-#include "RigVMCore/RigVMNameCache.h"
-#include "UObject/StructOnScope.h"
-#include "RigVMLog.h"
-#include "RigVMDrawInterface.h"
-#include "RigVMDrawContainer.h"
-#include "GameFramework/Actor.h"
-#include "Components/SceneComponent.h"
 
 #include "RigVMExecuteContext.generated.h"
 
@@ -540,10 +541,7 @@ struct RIGVM_API FRigVMExtendedExecuteContext
 		*this = InOther;
 	}
 
-	virtual ~FRigVMExtendedExecuteContext()
-	{
-		Reset();
-	}
+	virtual ~FRigVMExtendedExecuteContext();
 
 	// /** Full context reset */
 	void Reset();
@@ -693,7 +691,6 @@ struct RIGVM_API FRigVMExtendedExecuteContext
 
 	void InvalidateCachedMemory()
 	{
-		CachedMemory.Reset();
 		CachedMemoryHandles.Reset();
 		LazyBranchInstanceData.Reset();
 	}
@@ -716,6 +713,12 @@ struct RIGVM_API FRigVMExtendedExecuteContext
 	UPROPERTY(transient)
 	TObjectPtr<URigVMMemoryStorage> DebugMemoryStorageObject;
 
+	UPROPERTY()
+	FRigVMMemoryStorageStruct WorkMemoryStorage;
+
+	UPROPERTY()
+	FRigVMMemoryStorageStruct DebugMemoryStorage;
+
 	FStructOnScope PublicDataScope;
 	URigVM* VM = nullptr;
 	TArray<FRigVMSlice> Slices;
@@ -725,9 +728,7 @@ struct RIGVM_API FRigVMExtendedExecuteContext
 	FRigVMNameCache NameCache;
 
 	TArray<FRigVMMemoryHandle> CachedMemoryHandles;
-	// changes to the layout of cached memory array should be reflected in GetContainerIndex()
-	TArray<URigVMMemoryStorage*> CachedMemory;
-
+	
 	int32 ExecutingThreadId = INDEX_NONE;
 
 	TArray<int32> EntriesBeingExecuted;
@@ -735,12 +736,18 @@ struct RIGVM_API FRigVMExtendedExecuteContext
 	ERigVMExecuteResult CurrentExecuteResult = ERigVMExecuteResult::Failed;
 	FName CurrentEntryName = NAME_None;
 	bool bCurrentlyRunningRootEntry = false;
-	TArrayView<URigVMMemoryStorage*> CurrentMemory;
+
+	TArrayView<TRigVMMemoryStorage*> CurrentVMMemory;
 
 #if WITH_EDITORONLY_DATA
-	UE_DEPRECATED(5.4, "DeferredVMToCopy has been deprecated. Please update your code.")
+	// changes to the layout of cached memory array should be reflected in GetContainerIndex()
+	UE_DEPRECATED(5.4, "CachedMemory has been deprecated.")
+	TArray<URigVMMemoryStorage*> CachedMemory;
+	UE_DEPRECATED(5.4, "CurrentMemory has been deprecated, please use CurrentVMMemory with TRigVMMemoryStorage.")
+	TArrayView<TRigVMMemoryStorageDeprecatedType*> CurrentMemory;
+	UE_DEPRECATED(5.4, "DeferredVMToCopy has been deprecated.")
 	TObjectPtr<URigVM> DeferredVMToCopy = nullptr;
-	UE_DEPRECATED(5.4, "DeferredVMContextToCopy has been deprecated. Please update your code.")
+	UE_DEPRECATED(5.4, "DeferredVMContextToCopy has been deprecated.")
 	const FRigVMExtendedExecuteContext* DeferredVMContextToCopy = nullptr;
 #endif
 
diff --git a/Engine/Plugins/Runtime/RigVM/Source/RigVM/Public/RigVMCore/RigVMMemoryCommon.h b/Engine/Plugins/Runtime/RigVM/Source/RigVM/Public/RigVMCore/RigVMMemoryCommon.h
--- a/Engine/Plugins/Runtime/RigVM/Source/RigVM/Public/RigVMCore/RigVMMemoryCommon.h
+++ b/Engine/Plugins/Runtime/RigVM/Source/RigVM/Public/RigVMCore/RigVMMemoryCommon.h
@@ -9,6 +9,7 @@
 #include "RigVMTraits.h"
 #include "Templates/TypeHash.h"
 #include "UObject/ObjectMacros.h"
+#include "Misc/OutputDevice.h"
 
 #include "RigVMMemoryCommon.generated.h"
 
@@ -148,3 +149,24 @@ private:
 };
 
 typedef TArrayView<const FRigVMOperand> FRigVMOperandArray;
+
+
+/**
+ * Helper class to catch default value import errors on properties
+ */
+class RIGVM_API FRigVMMemoryStorageImportErrorContext : public FOutputDevice
+{
+public:
+
+	bool bLogErrors;
+	int32 NumErrors;
+
+	explicit FRigVMMemoryStorageImportErrorContext(bool InLogErrors = true)
+		: FOutputDevice()
+		, bLogErrors(InLogErrors)
+		, NumErrors(0)
+	{
+	}
+
+	void Serialize(const TCHAR* V, ELogVerbosity::Type Verbosity, const class FName& Category) override;
+};
diff --git a/Engine/Plugins/Runtime/RigVM/Source/RigVM/Public/RigVMCore/RigVMMemoryStorage.h b/Engine/Plugins/Runtime/RigVM/Source/RigVM/Public/RigVMCore/RigVMMemoryStorage.h
--- a/Engine/Plugins/Runtime/RigVM/Source/RigVM/Public/RigVMCore/RigVMMemoryStorage.h
+++ b/Engine/Plugins/Runtime/RigVM/Source/RigVM/Public/RigVMCore/RigVMMemoryStorage.h
@@ -791,15 +791,55 @@ private:
 	{
 		if(InSliceIndex != INDEX_NONE)
 		{
-			// sliced memory cannot be accessed
-			// using a property path.
-			// it refers to opaque memory only
+			// Sliced memory cannot be accessed using a property path.
+			// It refers to opaque memory only
 			check(PropertyPath == nullptr);
 			check(!bFollowPropertyPath);
 
 			const FArrayProperty* ArrayProperty = CastFieldChecked<FArrayProperty>(Property);
 			FScriptArrayHelper ArrayHelper(ArrayProperty, Ptr);
-			if(ArrayHelper.Num() <= InSliceIndex)
+
+#if UE_RIGVM_PROPERTY_BAG_STORAGE_ENABLED
+			if (InSliceIndex >= ArrayHelper.Num() - 1)
+			{
+				// For each slice we copy the default value to the next slice index
+				// The Index 0 has the initial default value, we carry it over the next slice on each call
+				const int32 NumValuesToAdd = FMath::Max(InSliceIndex + 1 - FMath::Max(ArrayHelper.Num() - 1, 0), 0);
+				if (NumValuesToAdd > 0)
+				{
+					const int32 FirstAddedIndex = ArrayHelper.AddValues(NumValuesToAdd);
+					if (FirstAddedIndex > 0)
+					{
+						if (const uint8* SourceElementMemory = ArrayHelper.GetRawPtr(FirstAddedIndex - 1))
+						{
+							const FProperty* ElementProperty = ArrayProperty->Inner;
+							for (int32 i = FirstAddedIndex; i < ArrayHelper.Num(); ++i)
+							{
+#if UE_RIGVM_DEBUG_EXECUTION
+								FString DefaultValue;
+								ElementProperty->ExportText_Direct(
+									DefaultValue,
+									SourceElementMemory,
+									SourceElementMemory,
+									nullptr,
+									PPF_None,
+									nullptr);
+
+								UE_LOG(LogRigVM, Display, TEXT("Adding slice %d for Property '%s', defaulting to '%s'."),
+									InSliceIndex,
+									*ArrayProperty->GetName(),
+									*DefaultValue
+								);
+#endif // UE_RIGVM_DEBUG_EXECUTION
+								uint8* DestMemory = ArrayHelper.GetRawPtr(i);
+								ElementProperty->CopyCompleteValue(DestMemory, SourceElementMemory);
+							}
+						}
+					}
+				}
+			}
+#else // !UE_RIGVM_PROPERTY_BAG_STORAGE_ENABLED
+			if (ArrayHelper.Num() <= InSliceIndex)
 			{
 				const int32 NumValuesToAdd = 1 + InSliceIndex - ArrayHelper.Num();
 				const int32 FirstAddedIndex = ArrayHelper.AddValues(NumValuesToAdd);
@@ -817,20 +857,20 @@ private:
 						for (int32 i=FirstAddedIndex; i<ArrayHelper.Num(); ++i)
 						{
 #if UE_RIGVM_DEBUG_EXECUTION
-							// FString DefaultValue;
-							// ElementProperty->ExportText_Direct(
-							// 	DefaultValue,
-							// 	DefaultElementMemory,
-							// 	DefaultElementMemory,
-							// 	nullptr,
-							// 	PPF_None,
-							// 	nullptr);
-							//
-							// UE_LOG(LogRigVM, Display, TEXT("Adding slice %d for Property '%s', defaulting to '%s'."),
-							// 	InSliceIndex,
-							// 	*ArrayProperty->GetName(),
-							// 	*DefaultValue
-							// );
+							 FString DefaultValue;
+							 ElementProperty->ExportText_Direct(
+							 	DefaultValue,
+							 	DefaultElementMemory,
+							 	DefaultElementMemory,
+							 	nullptr,
+							 	PPF_None,
+							 	nullptr);
+							
+							 UE_LOG(LogRigVM, Display, TEXT("Adding slice %d for Property '%s', defaulting to '%s'."),
+							 	InSliceIndex,
+							 	*ArrayProperty->GetName(),
+							 	*DefaultValue
+							 );
 #endif
 
 							uint8* DestMemory = ArrayHelper.GetRawPtr(i);
@@ -839,6 +879,7 @@ private:
 					}
 				}
 			}
+#endif // UE_RIGVM_PROPERTY_BAG_STORAGE_ENABLED
 
 #if UE_RIGVM_DEBUG_EXECUTION
 			// const FProperty* ElementProperty = ArrayProperty->Inner;
@@ -1128,6 +1169,16 @@ public:
 	void RefreshLinkedProperties();
 	void RefreshPropertyPaths();
 
+	bool IsValidPropertyPathDescriptionIndex(int32 Index) const
+	{
+		return PropertyPathDescriptions.IsValidIndex(Index);
+	}
+
+	const FRigVMPropertyPathDescription* GetPropertyPathDescriptionByIndex(int32 Index) const
+	{
+		return &PropertyPathDescriptions[Index];
+	}
+
 private:
 
 	// The type of memory of this class
diff --git a/Engine/Plugins/Runtime/RigVM/Source/RigVM/Public/RigVMCore/RigVMMemoryStorageStruct.h b/Engine/Plugins/Runtime/RigVM/Source/RigVM/Public/RigVMCore/RigVMMemoryStorageStruct.h
new file mode 100644
--- /dev/null
+++ b/Engine/Plugins/Runtime/RigVM/Source/RigVM/Public/RigVMCore/RigVMMemoryStorageStruct.h
@@ -0,0 +1,400 @@
+// Copyright Epic Games, Inc. All Rights Reserved.
+
+#pragma once
+
+#include "RigVMDefines.h"
+#include "RigVMCore/RigVMMemoryStorage.h"
+#include "RigVMCore/RigVMPropertyPath.h"
+#include "PropertyBag.h"
+#include "RigVMMemoryStorageStruct.generated.h"
+
+USTRUCT()
+struct RIGVM_API FRigVMMemoryStorageStruct : public FInstancedPropertyBag
+{
+	GENERATED_BODY()
+
+	FRigVMMemoryStorageStruct();
+	explicit FRigVMMemoryStorageStruct(ERigVMMemoryType InMemoryType);
+	FRigVMMemoryStorageStruct(ERigVMMemoryType InMemoryType, const TArray<FRigVMPropertyDescription>& InPropertyDescriptions, const TArray<FRigVMPropertyPathDescription>& InPropertyPaths = TArray<FRigVMPropertyPathDescription>());
+	~FRigVMMemoryStorageStruct();
+
+	bool Serialize(FArchive& Ar);
+	void AddStructReferencedObjects(FReferenceCollector& Collector);
+
+	friend FArchive& operator<<(FArchive& Ar, FRigVMMemoryStorageStruct& P)
+	{
+		P.Serialize(Ar);
+		return Ar;
+	}
+
+	static bool IsClassOf(const FRigVMMemoryStorageStruct* InElement)
+	{
+		return true;
+	}
+
+	//---------------------------------------------------------------------------
+
+	void SetMemoryType(ERigVMMemoryType InMemoryType)
+	{
+		MemoryType = InMemoryType;
+	}
+	ERigVMMemoryType GetMemoryType() const
+	{
+		return MemoryType;
+	}
+
+	//---------------------------------------------------------------------------
+
+	/**
+	 * Adds properties to the storage. If property of same name already exists, it will be replaced with the new type.
+	 * Numeric property values will be converted if possible, when a property's type changes.
+	 * @param Descriptors : Descriptors of new properties to add.
+	 */
+	void AddProperties(const TArray<FRigVMPropertyDescription>& InPropertyDescriptions, const TArray<FRigVMPropertyPathDescription>& InPropertyPathDescriptions = TArray<FRigVMPropertyPathDescription>());
+
+
+	// Returns the number of properties stored in this instance
+	int32 Num() const
+	{
+		const TArray<const FProperty*>& Properties = GetProperties();
+		return Properties.Num();
+	}
+
+	// Returns true if a provided property index is valid
+	bool IsValidIndex(int32 InIndex) const
+	{
+		const TArray<const FProperty*>& Properties = GetProperties();
+		return Properties.IsValidIndex(InIndex);
+	}
+
+
+	// Returns the properties provided by this instance
+	const TArray<const FProperty*>& GetProperties() const 
+	{
+		return LinkedProperties;
+	}
+
+	// Returns the property paths provided by this instance
+	const TArray<FRigVMPropertyPath>& GetPropertyPaths() const 
+	{
+		return PropertyPaths;
+	}
+
+	// Returns the index of a property given the property itself
+	int32 GetPropertyIndex(const FProperty* InProperty) const;
+
+	// Returns the index of a property given its name
+	int32 GetPropertyIndexByName(const FName& InName) const;
+
+	// Returns a property given its index
+	const FProperty* GetProperty(int32 InPropertyIndex) const
+	{
+		const TArray<const FProperty*>& Properties = GetProperties();
+		return Properties[InPropertyIndex];
+	}
+
+	// Returns a property given its name (or nullptr if the name wasn't found)
+	FProperty* FindPropertyByName(const FName& InName) const;
+
+	// Creates and returns a new operand for a property (and optionally a property path)
+	FRigVMOperand GetOperand(int32 InPropertyIndex, int32 InPropertyPathIndex = INDEX_NONE) const;
+
+	// Creates and returns a new operand for a property (and optionally a property path)
+	FRigVMOperand GetOperandByName(const FName& InName, int32 InPropertyPathIndex = INDEX_NONE) const;
+
+	// Returns the raw memory storage pointer
+	void* GetContainerPtr() const;
+
+	uint32 GetMemoryHash() const;
+
+	// Returns true if the property at a given index is a TArray
+	bool IsArray(int32 InPropertyIndex) const
+	{
+		return GetProperty(InPropertyIndex)->IsA<FArrayProperty>();
+	}
+
+	// Returns true if the property at a given index is a TMap
+	bool IsMap(int32 InPropertyIndex) const
+	{
+		return GetProperty(InPropertyIndex)->IsA<FMapProperty>();
+	}
+
+	//---------------------------------------------------------------------------
+
+	// Returns the memory for a property given its index
+	template<typename T>
+	T* GetData(int32 InPropertyIndex)
+	{
+		const TArray<const FProperty*>& Properties = GetProperties();
+		check(Properties.IsValidIndex(InPropertyIndex));
+		return Properties[InPropertyIndex]->ContainerPtrToValuePtr<T>(GetMutableValue().GetMemory());
+	}
+
+	// Returns the memory for a property given its name (or nullptr)
+	template<typename T>
+	T* GetDataByName(const FName& InName)
+	{
+		const int32 PropertyIndex = GetPropertyIndexByName(InName);
+		if (PropertyIndex == INDEX_NONE)
+		{
+			return nullptr;
+		}
+		return GetData<T>(PropertyIndex);
+	}
+
+	// Returns the mutable memory for a given property (or nullptr if does not belong to this storage)
+	template<typename T>
+	T* GetData(const FProperty* Property)
+	{
+		if (const int32 PropertyIndex = GetPropertyIndex(Property); PropertyIndex != INDEX_NONE)
+		{
+			return GetData<T>(PropertyIndex);
+		}
+
+		return nullptr;
+	}
+
+	// Returns the memory for a property given its index and a matching property path
+	template<typename T>
+	T* GetData(int32 InPropertyIndex, const FRigVMPropertyPath& InPropertyPath)
+	{
+		const FProperty* Property = GetProperty(InPropertyIndex);
+		return InPropertyPath.GetData<T>(GetData<uint8>(InPropertyIndex), Property);
+	}
+
+	// Returns the memory for a property given its name and a matching property path (or nullptr)
+	template<typename T>
+	T* GetDataByName(const FName& InName, const FRigVMPropertyPath& InPropertyPath)
+	{
+		const int32 PropertyIndex = GetPropertyIndexByName(InName);
+		if (PropertyIndex == INDEX_NONE)
+		{
+			return nullptr;
+		}
+		return GetData<T>(PropertyIndex, InPropertyPath);
+	}
+
+	// Returns the memory for a property (and optionally a property path) given an operand
+	template<typename T>
+	T* GetData(const FRigVMOperand& InOperand)
+	{
+		const int32 PropertyIndex = InOperand.GetRegisterIndex();
+		const int32 PropertyPathIndex = InOperand.GetRegisterOffset();
+
+		check(GetProperties().IsValidIndex(PropertyIndex));
+
+		if (PropertyPathIndex == INDEX_NONE)
+		{
+			return GetData<T>(PropertyIndex);
+		}
+
+		check(GetPropertyPaths().IsValidIndex(PropertyPathIndex));
+		return GetData<T>(PropertyIndex, GetPropertyPaths()[PropertyPathIndex]);
+	}
+
+	// Returns the ref of an element stored at a given property index
+	template<typename T>
+	T& GetRef(int32 InPropertyIndex)
+	{
+		return *GetData<T>(InPropertyIndex);
+	}
+
+	// Returns the ref of an element stored at a given property name (throws if name is invalid)
+	template<typename T>
+	T& GetRefByName(const FName& InName)
+	{
+		return *GetDataByName<T>(InName);
+	}
+
+	// Returns the ref of an element stored at a given property index and a property path
+	template<typename T>
+	T& GetRef(int32 InPropertyIndex, const FRigVMPropertyPath& InPropertyPath)
+	{
+		return *GetData<T>(InPropertyIndex, InPropertyPath);
+	}
+
+	// Returns the ref of an element stored at a given property name and a property path (throws if name is invalid)
+	template<typename T>
+	T& GetRefByName(const FName& InName, const FRigVMPropertyPath& InPropertyPath)
+	{
+		return *GetDataByName<T>(InName, InPropertyPath);
+	}
+
+	// Returns the ref of an element stored for a given operand
+	template<typename T>
+	T& GetRef(const FRigVMOperand& InOperand)
+	{
+		return *GetData<T>(InOperand);
+	}
+
+	//---------------------------------------------------------------------------
+
+	// Returns the exported text for a given property index
+	FString GetDataAsString(int32 InPropertyIndex, int32 PortFlags = PPF_None);
+
+	// Returns the exported text for given property name 
+	FString GetDataAsStringByName(const FName& InName, int32 PortFlags = PPF_None)
+	{
+		const int32 PropertyIndex = GetPropertyIndexByName(InName);
+		return GetDataAsString(PropertyIndex, PortFlags);
+	}
+
+	// Returns the exported text for a given operand
+	FString GetDataAsString(const FRigVMOperand& InOperand, int32 PortFlags = PPF_None);
+
+	// Returns the exported text for a given property index
+	FString GetDataAsStringSafe(int32 InPropertyIndex, int32 PortFlags = PPF_None);
+
+	// Returns the exported text for given property name 
+	FString GetDataAsStringByNameSafe(const FName& InName, int32 PortFlags = PPF_None)
+	{
+		const int32 PropertyIndex = GetPropertyIndexByName(InName);
+		return GetDataAsStringSafe(PropertyIndex, PortFlags);
+	}
+
+	// Returns the exported text for a given operand
+	FString GetDataAsStringSafe(const FRigVMOperand& InOperand, int32 PortFlags = PPF_None);
+
+	// Sets the content of a property by index given an exported string. Returns true if succeeded
+	bool SetDataFromString(int32 InPropertyIndex, const FString& InValue);
+
+	// Sets the content of a property by name given an exported string. Returns true if succeeded
+	bool SetDataFromStringByName(const FName& InName, const FString& InValue)
+	{
+		const int32 PropertyIndex = GetPropertyIndexByName(InName);
+		return SetDataFromString(PropertyIndex, InValue);
+	}
+
+	//---------------------------------------------------------------------------
+
+	// Returns the handle for a given property by index (and optionally property path)
+	FRigVMMemoryHandle GetHandle(int32 InPropertyIndex, const FRigVMPropertyPath* InPropertyPath = nullptr);
+
+	// Returns the handle for a given property by name (and optionally property path)
+	FRigVMMemoryHandle GetHandleByName(const FName& InName, const FRigVMPropertyPath* InPropertyPath = nullptr)
+	{
+		const int32 PropertyIndex = GetPropertyIndexByName(InName);
+		return GetHandle(PropertyIndex, InPropertyPath);
+	}
+
+	//---------------------------------------------------------------------------
+
+	void SetPropertyPathDescriptions(const TArray<FRigVMPropertyPathDescription>& InPropertyPathDescriptions)
+	{
+		PropertyPathDescriptions = InPropertyPathDescriptions;
+	}
+
+	void ResetPropertyPathDescriptions()
+	{
+		PropertyPathDescriptions.Reset();
+	}
+
+	void RefreshPropertyPaths();
+
+	bool IsValidPropertyPathDescriptionIndex(int32 Index) const
+	{
+		return PropertyPathDescriptions.IsValidIndex(Index);
+	}
+
+	const FRigVMPropertyPathDescription* GetPropertyPathDescriptionByIndex(int32 Index) const
+	{
+		return &PropertyPathDescriptions[Index];
+	}
+
+	/**
+	 * Copies the content of a source property and memory into the memory of a target property
+	 * @param InTargetProperty The target property to copy into
+	 * @param InTargetPtr The memory of the target value to copy into
+	 * @param InSourceProperty The source property to copy from
+	 * @param InSourcePtr The memory of the source value to copy from
+	 * @return true if the copy operations was successful
+	 */
+	static bool CopyProperty(
+		const FProperty* InTargetProperty,
+		uint8* InTargetPtr,
+		const FProperty* InSourceProperty,
+		const uint8* InSourcePtr);
+
+	/**
+	 * Copies the content of a source property and memory into the memory of a target property.
+	 * This variant supports property paths for both target and source.
+	 * @param InTargetProperty The target property to copy into
+	 * @param InTargetPtr The memory of the target value to copy into
+	 * @param InTargetPropertyPath The property path to use when traversing the target memory.
+	 * @param InSourceProperty The source property to copy from
+	 * @param InSourcePtr The memory of the source value to copy from
+	 * @param InSourcePropertyPath The property path to use when traversing the source memory.
+	 * @return true if the copy operations was successful
+	 */
+	static bool CopyProperty(
+		const FProperty* InTargetProperty,
+		uint8* InTargetPtr,
+		const FRigVMPropertyPath& InTargetPropertyPath,
+		const FProperty* InSourceProperty,
+		const uint8* InSourcePtr,
+		const FRigVMPropertyPath& InSourcePropertyPath);
+
+	/**
+	 * Copies the content of a source property and memory into the memory of a target property.
+	 * This variant retrieves the memory pointers given target and source storage objects.
+	 * @param InTargetStorage The target property to copy into
+	 * @param InTargetPropertyIndex The index of the value to use on the InTargetStorage
+	 * @param InTargetPropertyPath The property path to use when traversing the target memory.
+	 * @param InSourceStorage The source property to copy from
+	 * @param InSourcePropertyIndex The index of the value to use on the InSourceStorage
+	 * @param InSourcePropertyPath The property path to use when traversing the source memory.
+	 * @return true if the copy operations was successful
+	 */
+	static bool CopyProperty(
+		FRigVMMemoryStorageStruct* InTargetStorage,
+		int32 InTargetPropertyIndex,
+		const FRigVMPropertyPath& InTargetPropertyPath,
+		FRigVMMemoryStorageStruct* InSourceStorage,
+		int32 InSourcePropertyIndex,
+		const FRigVMPropertyPath& InSourcePropertyPath);
+
+	/**
+	 * Copies the content of a source property and memory into the memory of a target property.
+	 * This variant retrieves the memory pointers given target and source storage objects.
+	 * @param InTargetHandle The handle of the target memory to copy into
+	 * @param InSourceHandle The handle of the source memory to copy from
+	 * @return true if the copy operations was successful
+	 */
+	static bool CopyProperty(
+		FRigVMMemoryHandle& InTargetHandle,
+		FRigVMMemoryHandle& InSourceHandle);
+
+protected:
+	ERigVMMemoryType MemoryType = ERigVMMemoryType::Invalid;
+
+	TArray<FRigVMPropertyPathDescription> PropertyPathDescriptions;
+
+	// A cached list of all linked properties (created by RefreshLinkedProperties)
+	TArray<const FProperty*> LinkedProperties;
+
+	// A cached list of all property paths (created by RefreshPropertyPaths)
+	TArray<FRigVMPropertyPath> PropertyPaths;
+
+	// A list of decriptions for the property paths - used for serialization
+
+	mutable int32 CachedMemoryHash = 0;
+
+	static const TArray<const FProperty*> EmptyProperties;
+	static const TArray<FRigVMPropertyPath> EmptyPropertyPaths;
+
+	void RefreshLinkedProperties();
+
+	void SetDefaultValues(const TArray<FRigVMPropertyDescription>& InPropertyDescriptions);
+
+	static FPropertyBagPropertyDesc GeneratePropertyBagDescriptor(const FRigVMPropertyDescription& RigVMDescriptor);
+	static bool GetPropertyTypeDataFromVMDescriptor(const FRigVMPropertyDescription& RigVMDescriptor, EPropertyBagPropertyType& OutBagPropertyType, FPropertyBagContainerTypes& OutBagContainerTypes);
+};
+
+template<> struct TStructOpsTypeTraits<FRigVMMemoryStorageStruct> : public TStructOpsTypeTraitsBase2<FRigVMMemoryStorageStruct>
+{
+	enum
+	{
+		WithSerializer = true,
+		WithAddStructReferencedObjects = true,
+	};
+};
diff --git a/Engine/Plugins/Runtime/RigVM/Source/RigVM/Public/RigVMCore/RigVMNativized.h b/Engine/Plugins/Runtime/RigVM/Source/RigVM/Public/RigVMCore/RigVMNativized.h
--- a/Engine/Plugins/Runtime/RigVM/Source/RigVM/Public/RigVMCore/RigVMNativized.h
+++ b/Engine/Plugins/Runtime/RigVM/Source/RigVM/Public/RigVMCore/RigVMNativized.h
@@ -12,6 +12,7 @@
 #include "Math/UnrealMathSSE.h"
 #include "Math/Vector.h"
 #include "Misc/OutputDevice.h"
+#include "RigVMDefines.h"
 #include "RigVM.h"
 #include "RigVMCore/RigVMExecuteContext.h"
 #include "RigVMCore/RigVMExternalVariable.h"
@@ -30,7 +31,6 @@
 
 class FArchive;
 class UObject;
-class URigVMMemoryStorage;
 struct FRigVMInstructionArray;
 
 UCLASS(BlueprintType, Abstract)
@@ -50,11 +50,11 @@ public:
 	virtual bool IsNativized() const override { return true; }
 	virtual void Empty(FRigVMExtendedExecuteContext& Context) override { return; }
 	virtual void CopyFrom(URigVM* InVM, bool bDeferCopy = false, bool bReferenceLiteralMemory = false, bool bReferenceByteCode = false, bool bCopyExternalVariables = false, bool bCopyDynamicRegisters = false) override { return; }
-	virtual bool Initialize(FRigVMExtendedExecuteContext& Context, TArrayView<URigVMMemoryStorage*> Memory) override;
-	virtual ERigVMExecuteResult Execute(FRigVMExtendedExecuteContext& Context, TArrayView<URigVMMemoryStorage*> Memory, const FName& InEntryName = NAME_None) override;
+	virtual bool Initialize(FRigVMExtendedExecuteContext& Context, TArrayView<TRigVMMemoryStorage*> Memory) override;
+	virtual ERigVMExecuteResult Execute(FRigVMExtendedExecuteContext& Context, TArrayView<TRigVMMemoryStorage*> Memory, const FName& InEntryName = NAME_None) override;
 	virtual int32 AddRigVMFunction(UScriptStruct* InRigVMStruct, const FName& InMethodName) override { return INDEX_NONE; }
 	virtual FString GetRigVMFunctionName(int32 InFunctionIndex) const override { return FString(); }
-	virtual URigVMMemoryStorage* GetMemoryByType(ERigVMMemoryType InMemoryType, bool bCreateIfNeeded = false) override { return nullptr; }
+	virtual TRigVMMemoryStorage* GetMemoryByType(FRigVMExtendedExecuteContext& Context, ERigVMMemoryType InMemoryType) override { return nullptr; }
 	virtual void ClearMemory() override { return; }
 	virtual const FRigVMInstructionArray& GetInstructions() override;
 	virtual bool ContainsEntry(const FName& InEntryName) const override { return GetEntryNames().Contains(InEntryName); }
@@ -73,8 +73,14 @@ public:
 	virtual void Empty() override {}
 	UE_DEPRECATED(5.3, "Please, use Initialize with Context param")
 	virtual bool Initialize(TArrayView<URigVMMemoryStorage*> Memory) override { return false; }
+	UE_DEPRECATED(5.4, "Please, use Execute with TRigVMMemoryStorage param")
+	virtual bool Initialize(FRigVMExtendedExecuteContext& Context, TArrayView<TRigVMMemoryStorageDeprecatedType*> Memory) override { return false; }
+	UE_DEPRECATED(5.4, "GetMemoryByType has been deprecated from the VM. Please, use GetWorkMemory from VMHost or the new GetMemoryByType with a TRigVMMemoryStorage and Context parameter.")
+	virtual TRigVMMemoryStorageDeprecatedType* GetMemoryByType(ERigVMMemoryType InMemoryType, bool bCreateIfNeeded = false) override { return nullptr; }
 	UE_DEPRECATED(5.3, "Please, use Execute with Context param")
 	virtual ERigVMExecuteResult Execute(TArrayView<URigVMMemoryStorage*> Memory, const FName& InEntryName = NAME_None) override { return ERigVMExecuteResult::Failed; }
+	UE_DEPRECATED(5.4, "Please, use Execute with TRigVMMemoryStorage param")
+	virtual ERigVMExecuteResult Execute(FRigVMExtendedExecuteContext& Context, TArrayView<TRigVMMemoryStorageDeprecatedType*> Memory, const FName& InEntryName = NAME_None) override { return ERigVMExecuteResult::Failed; }
 	UE_DEPRECATED(5.3, "Please, use SetInstructionIndex with Context param")
 	virtual void SetInstructionIndex(uint16 InInstructionIndex) {}
 
diff --git a/Engine/Plugins/Runtime/RigVM/Source/RigVM/Public/RigVMDefines.h b/Engine/Plugins/Runtime/RigVM/Source/RigVM/Public/RigVMDefines.h
--- a/Engine/Plugins/Runtime/RigVM/Source/RigVM/Public/RigVMDefines.h
+++ b/Engine/Plugins/Runtime/RigVM/Source/RigVM/Public/RigVMDefines.h
@@ -35,5 +35,17 @@
 #endif
 
 #ifndef UE_RIGVM_PROPERTY_BAG_STORAGE_ENABLED
-#define UE_RIGVM_PROPERTY_BAG_STORAGE_ENABLED 0
+#define UE_RIGVM_PROPERTY_BAG_STORAGE_ENABLED 1
+#endif
+
+#if UE_RIGVM_PROPERTY_BAG_STORAGE_ENABLED
+	using TRigVMMemoryStorage = struct FRigVMMemoryStorageStruct;
+	using TRigVMMemoryStorageDeprecatedType = class URigVMMemoryStorage;
+#else
+	struct URigVMMemoryStorageDummy
+	{
+	};
+
+	using TRigVMMemoryStorage = class URigVMMemoryStorage;
+	using TRigVMMemoryStorageDeprecatedType = struct URigVMMemoryStorageDummy;
 #endif
diff --git a/Engine/Plugins/Runtime/RigVM/Source/RigVM/Public/RigVMHost.h b/Engine/Plugins/Runtime/RigVM/Source/RigVM/Public/RigVMHost.h
--- a/Engine/Plugins/Runtime/RigVM/Source/RigVM/Public/RigVMHost.h
+++ b/Engine/Plugins/Runtime/RigVM/Source/RigVM/Public/RigVMHost.h
@@ -226,20 +226,20 @@ public:
 #endif
 
 	// Returns a VM memory storage by type
-	virtual URigVMMemoryStorage* GetMemoryByType(ERigVMMemoryType InMemoryType);
-	virtual const URigVMMemoryStorage* GetMemoryByType(ERigVMMemoryType InMemoryType) const;
+	virtual TRigVMMemoryStorage* GetMemoryByType(ERigVMMemoryType InMemoryType);
+	virtual const TRigVMMemoryStorage* GetMemoryByType(ERigVMMemoryType InMemoryType) const;
 
 	// The default mutable work memory
-	URigVMMemoryStorage* GetWorkMemory() { return GetMemoryByType(ERigVMMemoryType::Work); }
-	const URigVMMemoryStorage* GetWorkMemory() const { return GetMemoryByType(ERigVMMemoryType::Work); }
+	TRigVMMemoryStorage* GetWorkMemory() { return GetMemoryByType(ERigVMMemoryType::Work); }
+	const TRigVMMemoryStorage* GetWorkMemory() const { return GetMemoryByType(ERigVMMemoryType::Work); }
 
 	// The default const literal memory
-	URigVMMemoryStorage* GetLiteralMemory() { return GetMemoryByType(ERigVMMemoryType::Literal); }
-	const URigVMMemoryStorage* GetLiteralMemory() const { return GetMemoryByType(ERigVMMemoryType::Literal); }
+	TRigVMMemoryStorage* GetLiteralMemory() { return GetMemoryByType(ERigVMMemoryType::Literal); }
+	const TRigVMMemoryStorage* GetLiteralMemory() const { return GetMemoryByType(ERigVMMemoryType::Literal); }
 
 	// The default debug watch memory
-	URigVMMemoryStorage* GetDebugMemory() { return GetMemoryByType(ERigVMMemoryType::Debug); }
-	const URigVMMemoryStorage* GetDebugMemory() const { return GetMemoryByType(ERigVMMemoryType::Debug); }
+	TRigVMMemoryStorage* GetDebugMemory() { return GetMemoryByType(ERigVMMemoryType::Debug); }
+	const TRigVMMemoryStorage* GetDebugMemory() const { return GetMemoryByType(ERigVMMemoryType::Debug); }
 
 	DECLARE_EVENT_TwoParams(URigVM, FRigVMExecutedEvent, class URigVMHost*, const FName&);
 	FRigVMExecutedEvent& OnInitialized_AnyThread() { return InitializedEvent; }
@@ -308,7 +308,7 @@ protected:
 #endif
 
 	void GenerateUserDefinedDependenciesData(FRigVMExtendedExecuteContext& Context);
-	TArray<const UObject*> GetUserDefinedDependencies(const TArray<const URigVMMemoryStorage*> InMemory);
+	TArray<const UObject*> GetUserDefinedDependencies(const TArray<const TRigVMMemoryStorage*> InMemory);
 
 	UPROPERTY()
 	TMap<FString, FSoftObjectPath> UserDefinedStructGuidToPathName;
diff --git a/Engine/Plugins/Runtime/RigVM/Source/RigVM/Public/RigVMModule.h b/Engine/Plugins/Runtime/RigVM/Source/RigVM/Public/RigVMModule.h
--- a/Engine/Plugins/Runtime/RigVM/Source/RigVM/Public/RigVMModule.h
+++ b/Engine/Plugins/Runtime/RigVM/Source/RigVM/Public/RigVMModule.h
@@ -8,10 +8,22 @@
 
 #include "CoreMinimal.h"
 #include "Logging/LogMacros.h"
+#include "Modules/ModuleInterface.h"
 #include "RigVMDefines.h"
 
 RIGVM_API DECLARE_LOG_CATEGORY_EXTERN(LogRigVM, Log, All);
 
+/**
+* The public interface to this module
+*/
+class RIGVM_API FRigVMModule : public IModuleInterface
+{
+public:
+	// Begin IModuleInterface
+	virtual void StartupModule() override;
+	virtual void ShutdownModule() override;
+};
+
 namespace RigVMCore
 {
 	RIGVM_API bool SupportsUObjects();
diff --git a/Engine/Plugins/Runtime/RigVM/Source/RigVM/Public/RigVMObjectVersion.h b/Engine/Plugins/Runtime/RigVM/Source/RigVM/Public/RigVMObjectVersion.h
--- a/Engine/Plugins/Runtime/RigVM/Source/RigVM/Public/RigVMObjectVersion.h
+++ b/Engine/Plugins/Runtime/RigVM/Source/RigVM/Public/RigVMObjectVersion.h
@@ -27,6 +27,9 @@ struct RIGVM_API FRigVMObjectVersion
 		// Storing paths to user defined enums map
 		HostStoringUserDefinedData,
 
+		// VM Memory Storage Struct serialized
+		VMMemoryStorageStructSerialized,
+
 		// -----<new versions can be added above this line>-------------------------------------------------
 		VersionPlusOne,
 		LatestVersion = VersionPlusOne - 1
diff --git a/Engine/Plugins/Runtime/RigVM/Source/RigVM/RigVM.Build.cs b/Engine/Plugins/Runtime/RigVM/Source/RigVM/RigVM.Build.cs
--- a/Engine/Plugins/Runtime/RigVM/Source/RigVM/RigVM.Build.cs
+++ b/Engine/Plugins/Runtime/RigVM/Source/RigVM/RigVM.Build.cs
@@ -17,8 +17,9 @@ public class RigVM : ModuleRules
                 "AnimationCore",
                 "AnimGraphRuntime",
                 "DeveloperSettings",
-            }
-        );
+				"StructUtils",
+			}
+		);
 
         if (Target.bBuildEditor == true)
         {
@@ -27,7 +28,8 @@ public class RigVM : ModuleRules
                 {
                     "UnrealEd",
                     "BlueprintGraph",
-                }
+					"StructUtilsEditor",
+				}
             );
         }
     }
diff --git a/Engine/Plugins/Runtime/RigVM/Source/RigVMDeveloper/Private/RigVMCompiler/RigVMCodeGenerator.cpp b/Engine/Plugins/Runtime/RigVM/Source/RigVMDeveloper/Private/RigVMCompiler/RigVMCodeGenerator.cpp
--- a/Engine/Plugins/Runtime/RigVM/Source/RigVMDeveloper/Private/RigVMCompiler/RigVMCodeGenerator.cpp
+++ b/Engine/Plugins/Runtime/RigVM/Source/RigVMDeveloper/Private/RigVMCompiler/RigVMCodeGenerator.cpp
@@ -98,12 +98,12 @@ static constexpr TCHAR RigVM_GetVMHashFormat[] = TEXT("\tvirtual uint32 GetVMHas
 static constexpr TCHAR RigVM_GetEntryNamesFormat[] = TEXT("\tvirtual const TArray<FName>& GetEntryNames() const override\r\n\t{\r\n\t\tstatic const TArray<FName> StaticEntryNames = { {0} };\r\n\t\treturn StaticEntryNames;\r\n\t}");
 static constexpr TCHAR RigVM_DeclareUpdateExternalVariablesFormat[] = TEXT("\tvirtual void UpdateExternalVariables(FRigVMExtendedExecuteContext& Context) override;");
 static constexpr TCHAR RigVM_DeclareInvokeEntryByNameFormat[] = TEXT("\tERigVMExecuteResult InvokeEntryByName(FRigVMExtendedExecuteContext& Context, const FName& InEntryName{0});");
-static constexpr TCHAR RigVM_DeclareInitializeFormat[] = TEXT("\tvirtual bool Initialize(FRigVMExtendedExecuteContext& Context, TArrayView<URigVMMemoryStorage*> Memory) override;");
-static constexpr TCHAR RigVM_DefineInitializeFormat[] = TEXT("bool U{0}::Initialize(FRigVMExtendedExecuteContext& Context, TArrayView<URigVMMemoryStorage*> Memory)\r\n{");
-static constexpr TCHAR RigVM_DeclareExecuteFormat[] = TEXT("\tvirtual ERigVMExecuteResult Execute(FRigVMExtendedExecuteContext& Context, TArrayView<URigVMMemoryStorage*> Memory, const FName& InEntryName) override;");
+static constexpr TCHAR RigVM_DeclareInitializeFormat[] = TEXT("\tvirtual bool Initialize(FRigVMExtendedExecuteContext& Context, TArrayView<TRigVMMemoryStorage*> Memory) override;");
+static constexpr TCHAR RigVM_DefineInitializeFormat[] = TEXT("bool U{0}::Initialize(FRigVMExtendedExecuteContext& Context, TArrayView<TRigVMMemoryStorage*> Memory)\r\n{");
+static constexpr TCHAR RigVM_DeclareExecuteFormat[] = TEXT("\tvirtual ERigVMExecuteResult Execute(FRigVMExtendedExecuteContext& Context, TArrayView<TRigVMMemoryStorage*> Memory, const FName& InEntryName) override;");
 static constexpr TCHAR RigVM_DefineUpdateExternalVariablesFormat[] = TEXT("void U{0}::UpdateExternalVariables(FRigVMExtendedExecuteContext& Context)\r\n{");
 static constexpr TCHAR RigVM_DefineInvokeEntryByNameFormat[] = TEXT("ERigVMExecuteResult U{0}::InvokeEntryByName(FRigVMExtendedExecuteContext& Context, const FName& InEntryName{1})\r\n{");
-static constexpr TCHAR RigVM_DefineExecuteFormat[] = TEXT("ERigVMExecuteResult U{0}::Execute(FRigVMExtendedExecuteContext& Context, TArrayView<URigVMMemoryStorage*> Memory, const FName& InEntryName)\r\n{");
+static constexpr TCHAR RigVM_DefineExecuteFormat[] = TEXT("ERigVMExecuteResult U{0}::Execute(FRigVMExtendedExecuteContext& Context, TArrayView<TRigVMMemoryStorage*> Memory, const FName& InEntryName)\r\n{");
 static constexpr TCHAR RigVM_DeclareExecuteEntryFormat[] = TEXT("\tERigVMExecuteResult ExecuteEntry_{0}(FRigVMExtendedExecuteContext& Context, {1});");
 static constexpr TCHAR RigVM_DefineExecuteEntryFormat[] = TEXT("ERigVMExecuteResult U{0}::ExecuteEntry_{1}(FRigVMExtendedExecuteContext& Context, {2})\r\n{");
 static constexpr TCHAR RigVM_DeclareExecuteGroupFormat[] = TEXT("\tERigVMExecuteResult ExecuteGroup_{0}_{1}({2});");
@@ -1241,13 +1241,19 @@ void FRigVMCodeGenerator::ParseRequiredUProperties(const FRigVMExtendedExecuteCo
 	}
 }
 
-void FRigVMCodeGenerator::ParseMemory(const FRigVMExtendedExecuteContext& Context, URigVMMemoryStorage* InMemory)
+void FRigVMCodeGenerator::ParseMemory(const FRigVMExtendedExecuteContext& Context, TRigVMMemoryStorage* InMemory)
 {
 	if (InMemory == nullptr)
 	{
 		return;
 	}
-	if (InMemory->GetClass() == URigVMMemoryStorage::StaticClass())
+#if UE_RIGVM_PROPERTY_BAG_STORAGE_ENABLED
+	for (const FProperty* Property : InMemory->GetProperties())
+	{
+		ParseProperty(Context, InMemory->GetMemoryType(), Property, InMemory);
+	}
+#else
+	if (InMemory->GetClass() == TRigVMMemoryStorage::StaticClass())
 	{
 		return;
 	}
@@ -1257,11 +1263,14 @@ void FRigVMCodeGenerator::ParseMemory(const FRigVMExtendedExecuteContext& Contex
 		const FProperty* Property = *PropertyIt;
 		ParseProperty(Context, InMemory->GetMemoryType(), Property, InMemory);
 	}
+#endif // UE_RIGVM_PROPERTY_BAG_STORAGE_ENABLED
 }
 
-void FRigVMCodeGenerator::ParseProperty(const FRigVMExtendedExecuteContext& Context, ERigVMMemoryType InMemoryType, const FProperty* InProperty, URigVMMemoryStorage* InMemory)
+void FRigVMCodeGenerator::ParseProperty(const FRigVMExtendedExecuteContext& Context, ERigVMMemoryType InMemoryType, const FProperty* InProperty, TRigVMMemoryStorage* InMemory)
 {
+#if !UE_RIGVM_PROPERTY_BAG_STORAGE_ENABLED
 	if (URigVMMemoryStorageGeneratorClass* MemoryClass = Cast<URigVMMemoryStorageGeneratorClass>(InMemory->GetClass()))
+#endif
 	{
 		const int32 PropertyIndex = InMemory->GetPropertyIndex(InProperty);
 		const FRigVMOperand Operand(InMemoryType, PropertyIndex);
@@ -1975,7 +1984,7 @@ FRigVMPropertyDescription FRigVMCodeGenerator::GetPropertyDescForOperand(const F
 	}
 	else
 	{
-		if (const URigVMMemoryStorage* MemoryStorage = VM->GetMemoryByType(Context, InOperand.GetMemoryType()))
+		if (const TRigVMMemoryStorage* MemoryStorage = VM->GetMemoryByType(Context, InOperand.GetMemoryType()))
 		{
 			Property = MemoryStorage->GetProperty(InOperand.GetRegisterIndex());
 		}
@@ -2000,12 +2009,16 @@ FRigVMPropertyDescription FRigVMCodeGenerator::GetPropertyForOperand(const FRigV
 	}
 	else
 	{
-		if (const URigVMMemoryStorage* MemoryStorage = VM->GetMemoryByType(Context, InOperand.GetMemoryType()))
+		if (const TRigVMMemoryStorage* MemoryStorage = VM->GetMemoryByType(Context, InOperand.GetMemoryType()))
 		{
 			Property = MemoryStorage->GetProperty(InOperand.GetRegisterIndex());
 			if (Property)
 			{
+#if UE_RIGVM_PROPERTY_BAG_STORAGE_ENABLED
+				Memory = Property->ContainerPtrToValuePtr<uint8>(MemoryStorage->GetContainerPtr());
+#else
 				Memory = Property->ContainerPtrToValuePtr<uint8>(MemoryStorage);
+#endif // UE_RIGVM_PROPERTY_BAG_STORAGE_ENABLED
 			}
 		}
 	}
@@ -2049,15 +2062,22 @@ const FRigVMPropertyPathDescription& FRigVMCodeGenerator::GetPropertyPathForOper
 		}
 		else
 		{
-			if (const URigVMMemoryStorage* MemoryStorage = VM->GetMemoryByType(Context, InOperand.GetMemoryType()))
+			if (const TRigVMMemoryStorage* MemoryStorage = VM->GetMemoryByType(Context, InOperand.GetMemoryType()))
 			{
+#if UE_RIGVM_PROPERTY_BAG_STORAGE_ENABLED
+				if (MemoryStorage->IsValidPropertyPathDescriptionIndex((RegisterOffsetIndex)))
+				{
+					return *MemoryStorage->GetPropertyPathDescriptionByIndex(RegisterOffsetIndex);
+				}
+#else
 				if (const URigVMMemoryStorageGeneratorClass* MemoryClass = Cast<URigVMMemoryStorageGeneratorClass>(MemoryStorage->GetClass()))
 				{
-					if (MemoryClass->PropertyPathDescriptions.IsValidIndex(RegisterOffsetIndex))
+					if (MemoryClass->IsValidPropertyPathDescriptionIndex((RegisterOffsetIndex)))
 					{
-						return MemoryClass->PropertyPathDescriptions[RegisterOffsetIndex];
+						return *MemoryClass->GetPropertyPathDescriptionByIndex(RegisterOffsetIndex);
 					}
 				}
+#endif // UE_RIGVM_PROPERTY_BAG_STORAGE_ENABLED
 			}
 		}
 	}
@@ -2113,15 +2133,17 @@ void FRigVMCodeGenerator::CheckOperand(const FRigVMExtendedExecuteContext& Conte
 	}
 	else
 	{
-		const URigVMMemoryStorage* Memory = VM->GetMemoryByType(Context, InOperand.GetMemoryType());
+		const TRigVMMemoryStorage* Memory = VM->GetMemoryByType(Context, InOperand.GetMemoryType());
 		check(Memory);
 
 		ensure(Memory->GetProperties().IsValidIndex(InOperand.GetRegisterIndex()));
 		if (InOperand.GetRegisterOffset() != INDEX_NONE)
 		{
 			ensure(Memory->GetPropertyPaths().IsValidIndex(InOperand.GetRegisterOffset()));
+#if !UE_RIGVM_PROPERTY_BAG_STORAGE_ENABLED
 			URigVMMemoryStorageGeneratorClass* Class = CastChecked<URigVMMemoryStorageGeneratorClass>(Memory->GetClass()); 
 			ensure(Class->PropertyPathDescriptions.IsValidIndex(InOperand.GetRegisterOffset()));
+#endif // UE_RIGVM_PROPERTY_BAG_STORAGE_ENABLED
 		}
 	}
 }
diff --git a/Engine/Plugins/Runtime/RigVM/Source/RigVMDeveloper/Private/RigVMCompiler/RigVMCompiler.cpp b/Engine/Plugins/Runtime/RigVM/Source/RigVMDeveloper/Private/RigVMCompiler/RigVMCompiler.cpp
--- a/Engine/Plugins/Runtime/RigVM/Source/RigVMDeveloper/Private/RigVMCompiler/RigVMCompiler.cpp
+++ b/Engine/Plugins/Runtime/RigVM/Source/RigVMDeveloper/Private/RigVMCompiler/RigVMCompiler.cpp
@@ -195,7 +195,7 @@ const FProperty* FRigVMCompilerWorkData::GetPropertyForOperand(const FRigVMOpera
 	
 	check(!bSetupMemory);
 
-	auto GetPropertyFromMemory = [this](const URigVMMemoryStorage* InMemory, const FRigVMOperand& InOperand)
+	auto GetPropertyFromMemory = [this](TRigVMMemoryStorage* InMemory, const FRigVMOperand& InOperand)
 	{
 		if(InOperand.GetRegisterOffset() == INDEX_NONE)
 		{
@@ -203,11 +203,16 @@ const FProperty* FRigVMCompilerWorkData::GetPropertyForOperand(const FRigVMOpera
 		}
 		if(!InMemory->GetPropertyPaths().IsValidIndex(InOperand.GetRegisterOffset()))
 		{
+#if UE_RIGVM_PROPERTY_BAG_STORAGE_ENABLED
+			InMemory->SetPropertyPathDescriptions(PropertyPathDescriptions.FindChecked(InOperand.GetMemoryType()));
+			InMemory->RefreshPropertyPaths();
+#else
 			if(URigVMMemoryStorageGeneratorClass* MemoryClass = Cast<URigVMMemoryStorageGeneratorClass>(InMemory->GetClass()))
 			{
 				MemoryClass->PropertyPathDescriptions = PropertyPathDescriptions.FindChecked(InOperand.GetMemoryType());;
 				MemoryClass->RefreshPropertyPaths();
 			}
+#endif
 		}
 		return InMemory->GetPropertyPaths()[InOperand.GetRegisterOffset()].GetTailProperty();
 	};
@@ -1165,16 +1170,20 @@ bool URigVMCompiler::Compile(const FRigVMCompileSettings& InSettings, TArray<URi
 
 	WorkData.VM->ClearMemory(*WorkData.Context);
 
-	TArray<ERigVMMemoryType> MemoryTypes;
-	MemoryTypes.Add(ERigVMMemoryType::Work);
-	MemoryTypes.Add(ERigVMMemoryType::Literal);
-	MemoryTypes.Add(ERigVMMemoryType::Debug);
+	const TArray<ERigVMMemoryType> MemoryTypes = { ERigVMMemoryType::Literal, ERigVMMemoryType::Work, ERigVMMemoryType::Debug };
 
+#if UE_RIGVM_PROPERTY_BAG_STORAGE_ENABLED
+	for (ERigVMMemoryType MemoryType : MemoryTypes)
+	{
+		const TArray<FRigVMPropertyDescription>* Properties = WorkData.PropertyDescriptions.Find(MemoryType);
+		WorkData.VM->GenerateMemoryType(*WorkData.Context, MemoryType, Properties);
+	}
+#else
 	for(ERigVMMemoryType MemoryType : MemoryTypes)
 	{
 		UPackage* Package = InGraphs[0]->GetOutermost();
 
-		TArray<FRigVMPropertyDescription>* Properties = WorkData.PropertyDescriptions.Find(MemoryType);
+		const TArray<FRigVMPropertyDescription>* Properties = WorkData.PropertyDescriptions.Find(MemoryType);
 		if(Properties == nullptr)
 		{
 			URigVMMemoryStorageGeneratorClass::RemoveStorageClass(Package, MemoryType);
@@ -1188,6 +1197,7 @@ bool URigVMCompiler::Compile(const FRigVMCompileSettings& InSettings, TArray<URi
 	{
 		WorkData.VM->CreateMemoryByType(*WorkData.Context, MemoryType);
 	}
+#endif
 
 	WorkData.bSetupMemory = false;
 	WorkData.ExprComplete.Reset();
@@ -1220,8 +1230,19 @@ bool URigVMCompiler::Compile(const FRigVMCompileSettings& InSettings, TArray<URi
 	for(ERigVMMemoryType MemoryType : MemoryTypes)
 	{
 		const TArray<FRigVMPropertyPathDescription>* Descriptions = WorkData.PropertyPathDescriptions.Find(MemoryType);
-		if(URigVMMemoryStorage* MemoryStorageObject = WorkData.VM->GetMemoryByType(*WorkData.Context, MemoryType))
+		if(TRigVMMemoryStorage* MemoryStorageObject = WorkData.VM->GetMemoryByType(*WorkData.Context, MemoryType))
 		{
+#if UE_RIGVM_PROPERTY_BAG_STORAGE_ENABLED
+			if (Descriptions)
+			{
+				MemoryStorageObject->SetPropertyPathDescriptions(*Descriptions);
+			}
+			else
+			{
+				MemoryStorageObject->ResetPropertyPathDescriptions();
+			}
+			MemoryStorageObject->RefreshPropertyPaths();
+#else
 			if(URigVMMemoryStorageGeneratorClass* Class = Cast<URigVMMemoryStorageGeneratorClass>(MemoryStorageObject->GetClass()))
 			{
 				if(Descriptions)
@@ -1234,6 +1255,7 @@ bool URigVMCompiler::Compile(const FRigVMCompileSettings& InSettings, TArray<URi
 				}
 				Class->RefreshPropertyPaths();
 			}
+#endif
 		}
 	}
 
diff --git a/Engine/Plugins/Runtime/RigVM/Source/RigVMDeveloper/Private/RigVMModel/Nodes/RigVMUnitNode.cpp b/Engine/Plugins/Runtime/RigVM/Source/RigVMDeveloper/Private/RigVMModel/Nodes/RigVMUnitNode.cpp
--- a/Engine/Plugins/Runtime/RigVM/Source/RigVMDeveloper/Private/RigVMModel/Nodes/RigVMUnitNode.cpp
+++ b/Engine/Plugins/Runtime/RigVM/Source/RigVMDeveloper/Private/RigVMModel/Nodes/RigVMUnitNode.cpp
@@ -455,11 +455,15 @@ void URigVMUnitNode::EnumeratePropertiesOnHostAndStructInstance(URigVMHost* InHo
 				}
 			case ERigVMMemoryType::Literal:
 				{
-					if(URigVMMemoryStorage* LiteralMemory = InHost->GetLiteralMemory())
+					if(TRigVMMemoryStorage* LiteralMemory = InHost->GetLiteralMemory())
 					{
 						if(const FProperty* Property = LiteralMemory->GetProperty(Operand.GetRegisterIndex()))
 						{
+#if UE_RIGVM_PROPERTY_BAG_STORAGE_ENABLED
+							HostMemory = Property->ContainerPtrToValuePtr<uint8>(LiteralMemory->GetContainerPtr());
+#else
 							HostMemory = Property->ContainerPtrToValuePtr<uint8>(LiteralMemory);
+#endif // UE_RIGVM_PROPERTY_BAG_STORAGE_ENABLED
 							HostProperty = Property;
 						}
 					}
@@ -467,11 +471,15 @@ void URigVMUnitNode::EnumeratePropertiesOnHostAndStructInstance(URigVMHost* InHo
 				}
 			case ERigVMMemoryType::Work:
 				{
-					if(URigVMMemoryStorage* WorkMemory = InHost->GetWorkMemory())
+					if(TRigVMMemoryStorage* WorkMemory = InHost->GetWorkMemory())
 					{
 						if(const FProperty* Property = WorkMemory->GetProperty(Operand.GetRegisterIndex()))
 						{
+#if UE_RIGVM_PROPERTY_BAG_STORAGE_ENABLED
+							HostMemory = Property->ContainerPtrToValuePtr<uint8>(WorkMemory->GetContainerPtr());
+#else
 							HostMemory = Property->ContainerPtrToValuePtr<uint8>(WorkMemory);
+#endif // UE_RIGVM_PROPERTY_BAG_STORAGE_ENABLED
 							HostProperty = Property;
 						}
 					}
@@ -570,7 +578,7 @@ bool URigVMUnitNode::UpdateHostFromStructInstance(URigVMHost* InHost, TSharedPtr
 		}
 
 		// copy the memory back to the host
-		URigVMMemoryStorage::CopyProperty(HostProperty, HostMemory, InstanceProperty, InstanceMemory);
+		TRigVMMemoryStorage::CopyProperty(HostProperty, HostMemory, InstanceProperty, InstanceMemory);
 		bChangedSomething = true;
 
 	}, InSliceIndex);
diff --git a/Engine/Plugins/Runtime/RigVM/Source/RigVMDeveloper/Public/RigVMBlueprint.h b/Engine/Plugins/Runtime/RigVM/Source/RigVMDeveloper/Public/RigVMBlueprint.h
--- a/Engine/Plugins/Runtime/RigVM/Source/RigVMDeveloper/Public/RigVMBlueprint.h
+++ b/Engine/Plugins/Runtime/RigVM/Source/RigVMDeveloper/Public/RigVMBlueprint.h
@@ -25,6 +25,7 @@ class URigVMBlueprintGeneratedClass;
 class IRigVMEditorModule;
 #endif
 struct FEndLoadPackageContext;
+struct FRigVMMemoryStorageStruct;
 
 DECLARE_EVENT_ThreeParams(URigVMBlueprint, FOnRigVMCompiledEvent, UObject*, URigVM*, FRigVMExtendedExecuteContext&);
 DECLARE_EVENT_OneParam(URigVMBlueprint, FOnRigVMRefreshEditorEvent, URigVMBlueprint*);
@@ -40,6 +41,7 @@ DECLARE_DELEGATE_RetVal_OneParam(bool, FRigVMOnBreakLinksDialogRequestedDelegate
 DECLARE_DELEGATE_RetVal_OneParam(TRigVMTypeIndex, FRigVMOnPinTypeSelectionRequestedDelegate, const TArray<TRigVMTypeIndex>&);
 DECLARE_EVENT(URigVMBlueprint, FOnRigVMBreakpointAdded);
 DECLARE_EVENT_OneParam(URigVMBlueprint, FOnRigVMRequestInspectObject, const TArray<UObject*>& );
+DECLARE_EVENT_OneParam(URigVMBlueprint, FOnRigVMRequestInspectMemoryStorage, const TArray<FRigVMMemoryStorageStruct*>&);
 DECLARE_DELEGATE_RetVal(URigVMGraph*, FRigVMBlueprintGetFocusedGraph);
 
 USTRUCT()
@@ -316,6 +318,9 @@ public:
 	FOnRigVMRequestInspectObject& OnRequestInspectObject() { return OnRequestInspectObjectEvent; }
 	void RequestInspectObject(const TArray<UObject*>& InObjects) { OnRequestInspectObjectEvent.Broadcast(InObjects); }
 
+	FOnRigVMRequestInspectMemoryStorage& OnRequestInspectMemoryStorage() { return OnRequestInspectMemoryStorageEvent; }
+	void RequestInspectMemoryStorage(const TArray<FRigVMMemoryStorageStruct*>& InMemoryStorageStructs) { OnRequestInspectMemoryStorageEvent.Broadcast(InMemoryStorageStructs); }
+
 #endif	// #if WITH_EDITOR
 
 	FRigVMBlueprintGetFocusedGraph& OnGetFocusedGraph() { return OnGetFocusedGraphDelegate; }
@@ -693,6 +698,7 @@ private:
 	TArray<URigVMNode*> RigVMBreakpointNodes;
 
 	FOnRigVMRequestInspectObject OnRequestInspectObjectEvent;
+	FOnRigVMRequestInspectMemoryStorage OnRequestInspectMemoryStorageEvent;
 	FRigVMBlueprintGetFocusedGraph OnGetFocusedGraphDelegate;
 		
 public:
diff --git a/Engine/Plugins/Runtime/RigVM/Source/RigVMDeveloper/Public/RigVMCompiler/RigVMCodeGenerator.h b/Engine/Plugins/Runtime/RigVM/Source/RigVMDeveloper/Public/RigVMCompiler/RigVMCodeGenerator.h
--- a/Engine/Plugins/Runtime/RigVM/Source/RigVMDeveloper/Public/RigVMCompiler/RigVMCodeGenerator.h
+++ b/Engine/Plugins/Runtime/RigVM/Source/RigVMDeveloper/Public/RigVMCompiler/RigVMCodeGenerator.h
@@ -105,8 +105,8 @@ private:
 		int32 InMaxInstructionsPerFunction);
 	void ParseInclude(UStruct* InDependency, const FName& InMethodName = NAME_None);
 	void ParseRequiredUProperties(const FRigVMExtendedExecuteContext& Context);
-	void ParseMemory(const FRigVMExtendedExecuteContext& Context, URigVMMemoryStorage* InMemory);
-	void ParseProperty(const FRigVMExtendedExecuteContext& Context, ERigVMMemoryType InMemoryType, const FProperty* InProperty, URigVMMemoryStorage* InMemory);
+	void ParseMemory(const FRigVMExtendedExecuteContext& Context, TRigVMMemoryStorage* InMemory);
+	void ParseProperty(const FRigVMExtendedExecuteContext& Context, ERigVMMemoryType InMemoryType, const FProperty* InProperty, TRigVMMemoryStorage* InMemory);
 	void ParseInstructionGroups(const FRigVMExtendedExecuteContext& Context);
 	FString DumpInstructions(const FRigVMExtendedExecuteContext& Context, const FString& InPrefix, int32 InFirstInstruction, int32 InLastInstruction, const FInstructionGroup& InGroup, bool bLog = false);
 	FString DumpInstructions(const FRigVMExtendedExecuteContext& Context, const FString& InPrefix, const TArray<int32> InInstructionIndices, const FInstructionGroup& InGroup, bool bLog = false);
diff --git a/Engine/Plugins/Runtime/RigVM/Source/RigVMEditor/Private/Editor/RigVMEditor.cpp b/Engine/Plugins/Runtime/RigVM/Source/RigVMEditor/Private/Editor/RigVMEditor.cpp
--- a/Engine/Plugins/Runtime/RigVM/Source/RigVMEditor/Private/Editor/RigVMEditor.cpp
+++ b/Engine/Plugins/Runtime/RigVM/Source/RigVMEditor/Private/Editor/RigVMEditor.cpp
@@ -17,6 +17,7 @@
 #include "SBlueprintEditorToolbar.h"
 #include "BlueprintEditorModes.h"
 #include "AssetEditorModeManager.h"
+#include "RigVMCore/RigVMMemoryStorageStruct.h"
 #include "RigVMBlueprintUtils.h"
 #include "RigVMPythonUtils.h"
 #include "EulerTransform.h"
@@ -172,6 +173,7 @@ void FRigVMEditor::InitRigVMEditor(const EToolkitMode::Type Mode, const TSharedP
 	InRigVMBlueprint->OnModified().AddSP(this, &FRigVMEditor::HandleModifiedEvent);
 	InRigVMBlueprint->OnVMCompiled().AddSP(this, &FRigVMEditor::HandleVMCompiledEvent);
 	InRigVMBlueprint->OnRequestInspectObject().AddSP(this, &FRigVMEditor::SetDetailObjects);
+	InRigVMBlueprint->OnRequestInspectMemoryStorage().AddSP(this, &FRigVMEditor::SetMemoryStorageDetails);
 
 	BindCommands();
 
@@ -1707,7 +1709,7 @@ void FRigVMEditor::HandleVMCompiledEvent(UObject* InCompiledObject, URigVM* InVM
 	}
 
 	RefreshDetailView();
-	
+
 	TArray<FName> TabIds;
 	TabIds.Add(*FString::Printf(TEXT("RigVMMemoryDetails_%d"), (int32)ERigVMMemoryType::Literal));
 	TabIds.Add(*FString::Printf(TEXT("RigVMMemoryDetails_%d"), (int32)ERigVMMemoryType::Work));
@@ -1720,9 +1722,26 @@ void FRigVMEditor::HandleVMCompiledEvent(UObject* InCompiledObject, URigVM* InVM
 		{
 			if(ActiveTab->GetMetaData<FMemoryTypeMetaData>().IsValid())
 			{
-				ERigVMMemoryType MemoryType = ActiveTab->GetMetaData<FMemoryTypeMetaData>()->MemoryType;			
+				ERigVMMemoryType MemoryType = ActiveTab->GetMetaData<FMemoryTypeMetaData>()->MemoryType;
+				// TODO zzz : UE-195014 - Fix memory tab losing values on VM recompile
+				TRigVMMemoryStorage* Memory = InVM->GetMemoryByType(InContext, MemoryType);
+
+#if UE_RIGVM_PROPERTY_BAG_STORAGE_ENABLED
+
+			#if 1
+				ActiveTab->RequestCloseTab();
+				const TArray<FRigVMMemoryStorageStruct*> MemoryStorage = { Memory };
+				SetMemoryStorageDetails(MemoryStorage);
+			#else
+				// TODO zzz : This does not compile, IStructureDetailsView does not inherit from SCompoundWidget, like IDetailsView
+				TSharedRef<IStructureDetailsView> StructDetailsView = StaticCastSharedRef<IStructureDetailsView>(ActiveTab->GetContent());
+				TSharedPtr<FStructOnScope> StructOnScope = MakeShareable(new FStructOnScope(Memory->GetPropertyBagStruct(), (uint8*)Memory->GetContainerPtr()));
+				StructDetailsView->SetStructureData(StructOnScope);
+			#endif
+#else
 				TSharedRef<IDetailsView> DetailsView = StaticCastSharedRef<IDetailsView>(ActiveTab->GetContent());
-				DetailsView->SetObject(InVM->GetMemoryByType(InContext, MemoryType));
+				DetailsView->SetObject(Memory);
+#endif // UE_RIGVM_PROPERTY_BAG_STORAGE_ENABLED
 			}
 		}
 	}
@@ -2457,14 +2476,14 @@ void FRigVMEditor::SetDetailObjects(const TArray<UObject*>& InObjects, bool bCha
 
 			TSharedRef<IDetailsView> DetailsView = EditModule.CreateDetailView( DetailsViewArgs );
 			TSharedRef<SDockTab> DockTab = SNew(SDockTab)
-			.Label( LOCTEXT("ControlRigMemoryDetails", "Control Rig Memory Details") )
+			.Label( LOCTEXT("RigVMMemoryDetails", "RigVM Memory Details") )
 			.AddMetaData<FMemoryTypeMetaData>(FMemoryTypeMetaData(Memory->GetMemoryType()))
 			.TabRole(ETabRole::NomadTab)
 			[
 				DetailsView
 			];
 
-			FName TabId = *FString::Printf(TEXT("ControlRigMemoryDetails_%d"), (int32)Memory->GetMemoryType());
+			FName TabId = *FString::Printf(TEXT("RigVMMemoryDetails_%d"), (int32)Memory->GetMemoryType());
 			if(TSharedPtr<SDockTab> ActiveTab = GetTabManager()->FindExistingLiveTab(TabId))
 			{
 				ActiveTab->RequestCloseTab();
@@ -2473,7 +2492,7 @@ void FRigVMEditor::SetDetailObjects(const TArray<UObject*>& InObjects, bool bCha
 			GetTabManager()->InsertNewDocumentTab(
 				FBlueprintEditorTabs::DetailsID,
 				TabId,
-				FTabManager::FLastMajorOrNomadTab(TEXT("ControlRigMemoryDetails")),
+				FTabManager::FLastMajorOrNomadTab(TEXT("RigVMMemoryDetails")),
 				DockTab
 			);
 
@@ -2585,6 +2604,54 @@ void FRigVMEditor::SetDetailObjects(const TArray<UObject*>& InObjects, bool bCha
 	Inspector->ShowDetailsForObjects(FilteredObjects, Options);
 }
 
+void FRigVMEditor::SetMemoryStorageDetails(const TArray<FRigVMMemoryStorageStruct*>& InStructs)
+{
+	if (bSuspendDetailsPanelRefresh)
+	{
+		return;
+	}
+
+	if (InStructs.Num() == 1)
+	{
+		if (const FRigVMMemoryStorageStruct* Memory = InStructs[0])
+		{
+			FPropertyEditorModule& EditModule = FModuleManager::Get().GetModuleChecked<FPropertyEditorModule>("PropertyEditor");
+
+			FDetailsViewArgs DetailsViewArgs;
+			DetailsViewArgs.NameAreaSettings = FDetailsViewArgs::HideNameArea;
+			DetailsViewArgs.bHideSelectionTip = true;
+
+			FStructureDetailsViewArgs StructureViewArgs;
+
+			TSharedRef<IStructureDetailsView> DetailsView = EditModule.CreateStructureDetailView(DetailsViewArgs, StructureViewArgs, TSharedPtr<FStructOnScope>());
+			TSharedPtr<FStructOnScope> StructOnScope = MakeShareable(new FStructOnScope(Memory->GetPropertyBagStruct(), (uint8*)Memory->GetContainerPtr()));
+			DetailsView->SetStructureData(StructOnScope);
+
+			TSharedRef<SDockTab> DockTab = SNew(SDockTab)
+				.Label(LOCTEXT("RigVMMemoryDetails", "RigVM Memory Details"))
+				.AddMetaData<FMemoryTypeMetaData>(FMemoryTypeMetaData(Memory->GetMemoryType()))
+				.TabRole(ETabRole::NomadTab)
+				[
+					DetailsView->GetWidget().ToSharedRef()
+				];
+
+			FName TabId = *FString::Printf(TEXT("RigVMMemoryDetails_%d"), (int32)Memory->GetMemoryType());
+			if (TSharedPtr<SDockTab> ActiveTab = GetTabManager()->FindExistingLiveTab(TabId))
+			{
+				ActiveTab->RequestCloseTab();
+			}
+
+			GetTabManager()->InsertNewDocumentTab(
+				FBlueprintEditorTabs::DetailsID,
+				TabId,
+				FTabManager::FLastMajorOrNomadTab(TEXT("RigVMMemoryDetails")),
+				DockTab
+			);
+			return;
+		}
+	}
+}
+
 void FRigVMEditor::SetDetailViewForGraph(URigVMGraph* InGraph)
 {
 	check(InGraph);
diff --git a/Engine/Plugins/Runtime/RigVM/Source/RigVMEditor/Private/Editor/RigVMGraphDetailCustomization.cpp b/Engine/Plugins/Runtime/RigVM/Source/RigVMEditor/Private/Editor/RigVMGraphDetailCustomization.cpp
--- a/Engine/Plugins/Runtime/RigVM/Source/RigVMEditor/Private/Editor/RigVMGraphDetailCustomization.cpp
+++ b/Engine/Plugins/Runtime/RigVM/Source/RigVMEditor/Private/Editor/RigVMGraphDetailCustomization.cpp
@@ -1728,6 +1728,8 @@ void FRigVMWrappedNodeDetailCustomization::CustomizeLiveValues(IDetailLayoutBuil
 				}
 
 				const FProperty* Property = nullptr;
+
+				TArray<TRigVMMemoryStorage*> ExternalStructs;
 				TArray<UObject*> ExternalObjects;
 
 				if(Operand->GetMemoryType() == ERigVMMemoryType::External)
@@ -1741,12 +1743,12 @@ void FRigVMWrappedNodeDetailCustomization::CustomizeLiveValues(IDetailLayoutBuil
 				}
 				else
 				{
-					URigVMMemoryStorage* Memory = DebuggedHost->GetMemoryByType(Operand->GetMemoryType());
+					TRigVMMemoryStorage* Memory = DebuggedHost->GetMemoryByType(Operand->GetMemoryType());
 					if(Memory == nullptr)
 					{
 						continue;
 					}
-
+#if !UE_RIGVM_PROPERTY_BAG_STORAGE_ENABLED
 					if(Memory->GetOuter() == GetTransientPackage())
 					{
 						continue;
@@ -1764,6 +1766,7 @@ void FRigVMWrappedNodeDetailCustomization::CustomizeLiveValues(IDetailLayoutBuil
 					{
 						continue;
 					}
+#endif // UE_RIGVM_PROPERTY_BAG_STORAGE_ENABLED
 
 					if(!Memory->IsValidIndex(Operand->GetRegisterIndex()))
 					{
@@ -1776,14 +1779,17 @@ void FRigVMWrappedNodeDetailCustomization::CustomizeLiveValues(IDetailLayoutBuil
 						continue;
 					}
 
+#if UE_RIGVM_PROPERTY_BAG_STORAGE_ENABLED
+					ExternalStructs.Add(Memory);
+#else
 					ExternalObjects.Add(Memory);
+#endif // UE_RIGVM_PROPERTY_BAG_STORAGE_ENABLED
 				}
 
-				check(ExternalObjects.Num() > 0);
+				check(ExternalObjects.Num() > 0 || ExternalStructs.Num() > 0);
 				check(Property);
 
-				IDetailPropertyRow* PropertyRow = DebugCategory.AddExternalObjectProperty(ExternalObjects, Property->GetFName(), EPropertyLocation::Default, FAddPropertyParams().ForceShowProperty());
-				if(PropertyRow)
+				if(IDetailPropertyRow* PropertyRow = DebugCategory.AddExternalObjectProperty(ExternalObjects, Property->GetFName(), EPropertyLocation::Default, FAddPropertyParams().ForceShowProperty()))
 				{
 					PropertyRow->DisplayName(FText::FromString(FString::Printf(TEXT("%s%s"), *Pin->GetName(), *NameSuffix)));
 					PropertyRow->IsEnabled(false);
@@ -1793,6 +1799,22 @@ void FRigVMWrappedNodeDetailCustomization::CustomizeLiveValues(IDetailLayoutBuil
 					NameSuffix = FString::Printf(TEXT("_%d"), SuffixIndex);
 				}
 
+#if UE_RIGVM_PROPERTY_BAG_STORAGE_ENABLED
+				for (TRigVMMemoryStorage* Memory : ExternalStructs)
+				{
+					TSharedPtr<FStructOnScope> StructOnScope = MakeShareable(new FStructOnScope(Memory->GetPropertyBagStruct(), (uint8*)Memory->GetContainerPtr()));
+					if (IDetailPropertyRow* PropertyRow = DebugCategory.AddExternalStructureProperty(StructOnScope, Property->GetFName(), EPropertyLocation::Default, FAddPropertyParams().ForceShowProperty()))
+					{
+						PropertyRow->DisplayName(FText::FromString(FString::Printf(TEXT("%s%s"), *Pin->GetName(), *NameSuffix)));
+						PropertyRow->IsEnabled(false);
+
+						SuffixIndex++;
+						bAddedProperty = true;
+						NameSuffix = FString::Printf(TEXT("_%d"), SuffixIndex);
+					}
+				}
+#endif // UE_RIGVM_PROPERTY_BAG_STORAGE_ENABLED
+
 				KnownOperands.Add(*Operand);
 			}
 		}
diff --git a/Engine/Plugins/Runtime/RigVM/Source/RigVMEditor/Private/Editor/RigVMLocalVariableDetailCustomization.cpp b/Engine/Plugins/Runtime/RigVM/Source/RigVMEditor/Private/Editor/RigVMLocalVariableDetailCustomization.cpp
--- a/Engine/Plugins/Runtime/RigVM/Source/RigVMEditor/Private/Editor/RigVMLocalVariableDetailCustomization.cpp
+++ b/Engine/Plugins/Runtime/RigVM/Source/RigVMEditor/Private/Editor/RigVMLocalVariableDetailCustomization.cpp
@@ -96,7 +96,7 @@ void FRigVMLocalVariableDetailCustomization::CustomizeDetails(IDetailLayoutBuild
 		if (CDO->GetVM() != nullptr)
 		{
 			FString SourcePath = FString::Printf(TEXT("LocalVariableDefault::%s|%s::Const"), *GraphBeingCustomized->GetGraphName(), *VariableDescription.Name.ToString());
-			URigVMMemoryStorage* LiteralMemory = CDO->GetVM()->GetLiteralMemory();
+			TRigVMMemoryStorage* LiteralMemory = CDO->GetVM()->GetLiteralMemory();
 			FProperty* Property = LiteralMemory->FindPropertyByName(*SourcePath);
 			if (Property)
 			{
@@ -104,8 +104,13 @@ void FRigVMLocalVariableDetailCustomization::CustomizeDetails(IDetailLayoutBuild
 				Property->ClearPropertyFlags(CPF_EditConst);
 			
 				const FName SanitizedName = FRigVMPropertyDescription::SanitizeName(*SourcePath);
-				TArray<UObject*> Objects = {LiteralMemory};
+#if UE_RIGVM_PROPERTY_BAG_STORAGE_ENABLED
+				TSharedPtr<FStructOnScope> StructOnScope = MakeShareable(new FStructOnScope(LiteralMemory->GetPropertyBagStruct(), (uint8*)LiteralMemory->GetContainerPtr()));
+				IDetailPropertyRow* Row = DefaultValueCategory.AddExternalStructureProperty(StructOnScope, SanitizedName);
+#else
+				TArray<UObject*> Objects = { LiteralMemory };
 				IDetailPropertyRow* Row = DefaultValueCategory.AddExternalObjectProperty(Objects, SanitizedName);
+#endif
 				Row->DisplayName(FText::FromName(VariableDescription.Name));
 
 				const FSimpleDelegate OnDefaultValueChanged = FSimpleDelegate::CreateLambda([this, Property, LiteralMemory]()
diff --git a/Engine/Plugins/Runtime/RigVM/Source/RigVMEditor/Private/Widgets/SRigVMGraphNode.cpp b/Engine/Plugins/Runtime/RigVM/Source/RigVMEditor/Private/Widgets/SRigVMGraphNode.cpp
--- a/Engine/Plugins/Runtime/RigVM/Source/RigVMEditor/Private/Widgets/SRigVMGraphNode.cpp
+++ b/Engine/Plugins/Runtime/RigVM/Source/RigVMEditor/Private/Widgets/SRigVMGraphNode.cpp
@@ -837,7 +837,7 @@ void SRigVMGraphNode::GetNodeInfoPopups(FNodeInfoContext* Context, TArray<FGraph
 							FString PinHash = URigVMCompiler::GetPinHash(ModelPin, nullptr, true);
 							if (const FRigVMOperand* WatchOperand = RigBlueprint->PinToOperandMap.Find(PinHash))
 							{
-								URigVMMemoryStorage* Memory = ActiveObject->GetDebugMemory();
+								TRigVMMemoryStorage* Memory = ActiveObject->GetDebugMemory();
 								// We mark PPF_ExternalEditor so that default values are also printed
 								const FString DebugValue = Memory->GetDataAsStringSafe(WatchOperand->GetRegisterIndex(), PPF_ExternalEditor | STRUCT_ExportTextItemNative);
 								if(!DebugValue.IsEmpty())
diff --git a/Engine/Plugins/Runtime/RigVM/Source/RigVMEditor/Public/Editor/RigVMEditor.h b/Engine/Plugins/Runtime/RigVM/Source/RigVMEditor/Public/Editor/RigVMEditor.h
--- a/Engine/Plugins/Runtime/RigVM/Source/RigVMEditor/Public/Editor/RigVMEditor.h
+++ b/Engine/Plugins/Runtime/RigVM/Source/RigVMEditor/Public/Editor/RigVMEditor.h
@@ -222,6 +222,7 @@ protected:
 	bool& GetSuspendDetailsPanelRefreshFlag() { return bSuspendDetailsPanelRefresh; }
 	virtual void SetDetailObjects(const TArray<UObject*>& InObjects);
 	virtual void SetDetailObjects(const TArray<UObject*>& InObjects, bool bChangeUISelectionState);
+	virtual void SetMemoryStorageDetails(const TArray<FRigVMMemoryStorageStruct*>& InStructs);
 	virtual void SetDetailViewForGraph(URigVMGraph* InGraph);
 	virtual void SetDetailViewForFocusedGraph();
 	virtual void SetDetailViewForLocalVariable();
diff --git a/Engine/Plugins/Runtime/RigVM/Source/RigVMEditor/RigVMEditor.Build.cs b/Engine/Plugins/Runtime/RigVM/Source/RigVMEditor/RigVMEditor.Build.cs
--- a/Engine/Plugins/Runtime/RigVM/Source/RigVMEditor/RigVMEditor.Build.cs
+++ b/Engine/Plugins/Runtime/RigVM/Source/RigVMEditor/RigVMEditor.Build.cs
@@ -31,7 +31,9 @@ public class RigVMEditor : ModuleRules
                 "PropertyEditor",
                 "ToolMenus",
                 "MessageLog",
-            }
-        );
-    }
+				"StructUtils",
+				"StructUtilsEditor",
+			}
+		);
+	}
 }
