RigVMEditor: Allow full compilation in PIE

diff --git a/Engine/Plugins/Runtime/RigVM/Source/RigVMEditor/Private/Editor/RigVMEditor.cpp b/Engine/Plugins/Runtime/RigVM/Source/RigVMEditor/Private/Editor/RigVMEditor.cpp
--- a/Engine/Plugins/Runtime/RigVM/Source/RigVMEditor/Private/Editor/RigVMEditor.cpp
+++ b/Engine/Plugins/Runtime/RigVM/Source/RigVMEditor/Private/Editor/RigVMEditor.cpp
@@ -1012,13 +1012,6 @@ void FRigVMEditor::Compile()
 			return;
 		}
 
-		// if we are running PIE - only compile the RigVM 
-		if(IsPIERunning())
-		{
-			RigVMBlueprint->RecompileVM();
-			return;
-		}
-
 		// force to disable the supended notif brackets
 		RigVMBlueprint->bSuspendModelNotificationsForOthers = false;
 		RigVMBlueprint->bSuspendModelNotificationsForSelf = false;
