RigVM: Improve debug drawing of clamp spatially node

diff --git a/Engine/Plugins/Runtime/RigVM/Source/RigVM/Private/RigVMCore/RigVMDrawInterface.cpp b/Engine/Plugins/Runtime/RigVM/Source/RigVM/Private/RigVMCore/RigVMDrawInterface.cpp
--- a/Engine/Plugins/Runtime/RigVM/Source/RigVM/Private/RigVMCore/RigVMDrawInterface.cpp
+++ b/Engine/Plugins/Runtime/RigVM/Source/RigVM/Private/RigVMCore/RigVMDrawInterface.cpp
@@ -274,6 +274,12 @@ void FRigVMDrawInterface::DrawArc(const FTransform& WorldOffset, const FTransfor
 	Instructions.Add(Instruction);
 }
 
+void FRigVMDrawInterface::DrawCircle(const FTransform& WorldOffset, const FTransform& Transform, float Radius,
+	const FLinearColor& Color, float Thickness, int32 Detail)
+{
+	DrawArc(WorldOffset, Transform, Radius, 0, TWO_PI, Color, Thickness, Detail);
+}
+
 void FRigVMDrawInterface::DrawCone(const FTransform& WorldOffset, const FTransform& ConeOffset, float Angle1, float Angle2, uint32 NumSides, bool bDrawSideLines, const FLinearColor& SideLineColor, FMaterialRenderProxy* const MaterialRenderProxy)
 {
 	if (!IsEnabled())
@@ -310,6 +316,67 @@ void FRigVMDrawInterface::DrawCone(const FTransform& WorldOffset, const FTransfo
 	}
 }
 
+void FRigVMDrawInterface::DrawArrow(const FTransform& WorldOffset, const FVector& Direction, const FVector& Side,
+	const FLinearColor& Color, float Thickness)
+{
+	if (!IsEnabled())
+	{
+		return;
+	}
+
+	FRigVMDrawInstruction Instruction(ERigVMDrawSettings::Lines, Color, Thickness, WorldOffset);
+
+	Instruction.Positions.SetNum(6);
+	Instruction.Positions[0] = FVector::ZeroVector;
+	Instruction.Positions[1] = Direction;
+	Instruction.Positions[2] = Instruction.Positions[1];
+	Instruction.Positions[3] = Direction - Direction.GetSafeNormal() * Side.Size() + Side;
+	Instruction.Positions[4] = Instruction.Positions[1];
+	Instruction.Positions[5] = Direction - Direction.GetSafeNormal() * Side.Size() - Side;
+
+	Instructions.Add(Instruction);
+}
+
+void FRigVMDrawInterface::DrawPlane(const FTransform& WorldOffset, const FVector2D& Scale, const FLinearColor& MeshColor,
+	bool bDrawLines, const FLinearColor& LineColor, FMaterialRenderProxy* const MaterialRenderProxy)
+{
+	if (!IsEnabled())
+	{
+		return;
+	}
+
+	const FVector Corner(Scale.X, Scale.Y, 0);
+
+	FRigVMDrawInstruction MeshInstruction;
+	MeshInstruction.PrimitiveType = ERigVMDrawSettings::DynamicMesh;
+	MeshInstruction.Transform = WorldOffset;
+	MeshInstruction.MeshVerts = {
+		(FVector3f)(Corner * FVector(-1, 1, 0)),
+		(FVector3f)(Corner * FVector(1, 1, 0)),
+		(FVector3f)(Corner * FVector(1, -1, 0)),
+		(FVector3f)(Corner * FVector(-1, -1, 0))
+	};
+	MeshInstruction.MeshIndices = { 0, 1, 2, 2, 3, 0 };
+	MeshInstruction.MaterialRenderProxy = MaterialRenderProxy;
+
+	Instructions.Add(MeshInstruction);
+
+	if (bDrawLines)
+	{
+		FRigVMDrawInstruction LinesInstruction(ERigVMDrawSettings::LineStrip, LineColor, 0.f, WorldOffset);
+
+		LinesInstruction.Positions.Reserve(5);
+		LinesInstruction.Positions.Add(FVector(MeshInstruction.MeshVerts[0].Position));
+		LinesInstruction.Positions.Add(FVector(MeshInstruction.MeshVerts[1].Position));
+		LinesInstruction.Positions.Add(FVector(MeshInstruction.MeshVerts[2].Position));
+		LinesInstruction.Positions.Add(FVector(MeshInstruction.MeshVerts[3].Position));
+		LinesInstruction.Positions.Add(FVector(MeshInstruction.MeshVerts[0].Position));
+		LinesInstruction.Positions.Add(FVector(MeshInstruction.MeshVerts[2].Position));
+
+		Instructions.Add(LinesInstruction);
+	}
+}
+
 bool FRigVMDrawInterface::IsEnabled() const
 {
 	bool bIsEnabled = true;
diff --git a/Engine/Plugins/Runtime/RigVM/Source/RigVM/Private/RigVMFunctions/Math/RigVMFunction_MathTransform.cpp b/Engine/Plugins/Runtime/RigVM/Source/RigVM/Private/RigVMFunctions/Math/RigVMFunction_MathTransform.cpp
--- a/Engine/Plugins/Runtime/RigVM/Source/RigVM/Private/RigVMFunctions/Math/RigVMFunction_MathTransform.cpp
+++ b/Engine/Plugins/Runtime/RigVM/Source/RigVM/Private/RigVMFunctions/Math/RigVMFunction_MathTransform.cpp
@@ -178,7 +178,7 @@ FRigVMFunction_MathTransformArrayToSRT_Execute()
 FRigVMFunction_MathTransformClampSpatially_Execute()
 {
 	FVector Position;
-	FRigVMFunction_MathVectorClampSpatially::StaticExecute(ExecuteContext, Value.GetTranslation(), Axis, Type, Minimum, Maximum, Space, bDrawDebug, DebugColor, DebugThickness, Position);
+	FRigVMFunction_MathVectorClampSpatially::StaticExecute(ExecuteContext, Value.GetTranslation(), Axis, Type, Minimum, Maximum, Space, bDrawDebug, DebugColor, DebugThickness, DebugScale, Position);
 	Result = Value;
 	Result.SetTranslation(Position);
 }
diff --git a/Engine/Plugins/Runtime/RigVM/Source/RigVM/Private/RigVMFunctions/Math/RigVMFunction_MathVector.cpp b/Engine/Plugins/Runtime/RigVM/Source/RigVM/Private/RigVMFunctions/Math/RigVMFunction_MathVector.cpp
--- a/Engine/Plugins/Runtime/RigVM/Source/RigVM/Private/RigVMFunctions/Math/RigVMFunction_MathVector.cpp
+++ b/Engine/Plugins/Runtime/RigVM/Source/RigVM/Private/RigVMFunctions/Math/RigVMFunction_MathVector.cpp
@@ -3,6 +3,8 @@
 #include "RigVMFunctions/Math/RigVMFunction_MathVector.h"
 #include "RigVMFunctions/Math/RigVMFunction_MathTransform.h"
 #include "RigVMFunctions/RigVMDispatch_Core.h"
+#include "Engine/Engine.h"
+#include "Materials/MaterialInstanceDynamic.h"
 
 #include UE_INLINE_GENERATED_CPP_BY_NAME(RigVMFunction_MathVector)
 
@@ -397,33 +399,46 @@ FRigVMFunction_MathVectorClampSpatially_Execute()
 		{
 			case ERigVMClampSpatialMode::Plane:
 			{
-				TArray<FVector> Points;
-				Points.SetNumUninitialized(2);
+				FVector A, B, Normal;
 					
 				switch (Axis)
 				{
 					case EAxis::X:
 					{
-						Points[0] = FVector(Minimum, 0.f, 0.f);
-						Points[1] = FVector(Maximum, 0.f, 0.f);
+						A = FVector::YAxisVector;
+						B = FVector::ZAxisVector;
+						Normal = FVector::XAxisVector;
 						break;
 					}
 					case EAxis::Y:
 					{
-						Points[0] = FVector(0.f, Minimum, 0.f);
-						Points[1] = FVector(0.f, Maximum, 0.f);
+						A = FVector::XAxisVector;
+						B = FVector::ZAxisVector;
+						Normal = FVector::YAxisVector;
 						break;
 					}
 					default:
 					{
-						Points[0] = FVector(0.f, 0.f, Minimum);
-						Points[1] = FVector(0.f, 0.f, Maximum);
+						A = FVector::XAxisVector;
+						B = FVector::YAxisVector;
+						Normal = FVector::ZAxisVector;
 						break;
 					}
 				}
 
-				ExecuteContext.GetDrawInterface()->DrawLine(Space, Points[0], Points[1], DebugColor, DebugThickness);
-				ExecuteContext.GetDrawInterface()->DrawPoints(Space, Points, DebugThickness * 8.f, DebugColor);
+				FTransform PlaneOffset = FTransform::Identity;
+				PlaneOffset.SetRotation(FMatrix(A, Normal.Cross(A), Normal, FVector::ZeroVector).ToQuat());
+
+				const FTransform MinimumWorldTransform = FTransform(Normal * Minimum) * Space;
+				ExecuteContext.GetDrawInterface()->DrawPlane(PlaneOffset * MinimumWorldTransform, FVector2D(1, 1) * DebugScale, FLinearColor::Yellow, true, FLinearColor::Yellow, GEngine->ConstraintLimitMaterialPrismatic->GetRenderProxy());
+				ExecuteContext.GetDrawInterface()->DrawArrow(MinimumWorldTransform, Normal * DebugScale * 0.25f, A * DebugScale * 0.08f, FLinearColor::Yellow, 0.f);
+
+				if (Maximum > SMALL_NUMBER)
+				{
+					const FTransform MaximumWorldTransform = FTransform(Normal * Maximum) * Space;
+					ExecuteContext.GetDrawInterface()->DrawPlane(PlaneOffset * MaximumWorldTransform, FVector2D(1, 1) * DebugScale, FLinearColor::Yellow, true, FLinearColor::Yellow, GEngine->ConstraintLimitMaterialPrismatic->GetRenderProxy());
+					ExecuteContext.GetDrawInterface()->DrawArrow(MaximumWorldTransform, -Normal * DebugScale * 0.25f, A * DebugScale * 0.08f, FLinearColor::Yellow, 0.f);
+				}
 
 				break;
 			}
@@ -447,11 +462,24 @@ FRigVMFunction_MathVectorClampSpatially_Execute()
 						break;
 					}
 				}
+
+				const FTransform ArrowTransform = CircleTransform * Space;
 				if (Minimum > SMALL_NUMBER)
 				{
-					ExecuteContext.GetDrawInterface()->DrawArc(Space, CircleTransform, Minimum, 0.f, PI * 2.f, DebugColor, DebugThickness, 16);
+					ExecuteContext.GetDrawInterface()->DrawCircle(Space, CircleTransform, Minimum, FLinearColor::Yellow, 0.f, 32);
+					ExecuteContext.GetDrawInterface()->DrawArrow(FTransform(FVector::XAxisVector * Minimum) * ArrowTransform, FVector::XAxisVector * Minimum * 0.25f, FVector::ZAxisVector * Minimum * 0.08f, FLinearColor::Yellow, 0.f);
+					ExecuteContext.GetDrawInterface()->DrawArrow(FTransform(-FVector::XAxisVector * Minimum) * ArrowTransform, -FVector::XAxisVector * Minimum * 0.25f, FVector::ZAxisVector * Minimum * 0.08f, FLinearColor::Yellow, 0.f);
+					ExecuteContext.GetDrawInterface()->DrawArrow(FTransform(FVector::YAxisVector * Minimum) * ArrowTransform, FVector::YAxisVector * Minimum * 0.25f, FVector::ZAxisVector * Minimum * 0.08f, FLinearColor::Yellow, 0.f);
+					ExecuteContext.GetDrawInterface()->DrawArrow(FTransform(-FVector::YAxisVector * Minimum) * ArrowTransform, -FVector::YAxisVector * Minimum * 0.25f, FVector::ZAxisVector * Minimum * 0.08f, FLinearColor::Yellow, 0.f);
+				}
+				if (Maximum > SMALL_NUMBER)
+				{
+					ExecuteContext.GetDrawInterface()->DrawCircle(Space, CircleTransform, Maximum, FLinearColor::Yellow, 0.f, 32);
+					ExecuteContext.GetDrawInterface()->DrawArrow(FTransform(-FVector::XAxisVector * Maximum) * ArrowTransform, FVector::XAxisVector * Maximum * 0.25f, FVector::ZAxisVector * Maximum * 0.08f, FLinearColor::Yellow, 0.f);
+					ExecuteContext.GetDrawInterface()->DrawArrow(FTransform(FVector::XAxisVector * Maximum) * ArrowTransform, -FVector::XAxisVector * Maximum * 0.25f, FVector::ZAxisVector * Maximum * 0.08f, FLinearColor::Yellow, 0.f);
+					ExecuteContext.GetDrawInterface()->DrawArrow(FTransform(-FVector::YAxisVector * Maximum) * ArrowTransform, FVector::YAxisVector * Maximum * 0.25f, FVector::ZAxisVector * Maximum * 0.08f, FLinearColor::Yellow, 0.f);
+					ExecuteContext.GetDrawInterface()->DrawArrow(FTransform(FVector::YAxisVector * Maximum) * ArrowTransform, -FVector::YAxisVector * Maximum * 0.25f, FVector::ZAxisVector * Maximum * 0.08f, FLinearColor::Yellow, 0.f);
 				}
-				ExecuteContext.GetDrawInterface()->DrawArc(Space, CircleTransform, Maximum, 0.f, PI * 2.f, DebugColor, DebugThickness, 16);
 				break;
 			}
 			default:
@@ -460,25 +488,64 @@ FRigVMFunction_MathVectorClampSpatially_Execute()
 				FTransform CircleTransform = FTransform::Identity;
 				if (Minimum > SMALL_NUMBER)
 				{
-					ExecuteContext.GetDrawInterface()->DrawArc(Space, CircleTransform, Minimum, 0.f, PI * 2.f, DebugColor, DebugThickness, 16);
+					ExecuteContext.GetDrawInterface()->DrawCircle(Space, CircleTransform, Minimum, FLinearColor::Yellow, 0.f, 32);
+					const FTransform ArrowTransform = CircleTransform * Space;
+					ExecuteContext.GetDrawInterface()->DrawArrow(FTransform(FVector::XAxisVector * Minimum) * ArrowTransform, FVector::XAxisVector * Minimum * 0.25f, FVector::ZAxisVector * Minimum * 0.08f, FLinearColor::Yellow, 0.f);
+					ExecuteContext.GetDrawInterface()->DrawArrow(FTransform(-FVector::XAxisVector * Minimum) * ArrowTransform, -FVector::XAxisVector * Minimum * 0.25f, FVector::ZAxisVector * Minimum * 0.08f, FLinearColor::Yellow, 0.f);
+					ExecuteContext.GetDrawInterface()->DrawArrow(FTransform(FVector::YAxisVector * Minimum) * ArrowTransform, FVector::YAxisVector * Minimum * 0.25f, FVector::ZAxisVector * Minimum * 0.08f, FLinearColor::Yellow, 0.f);
+					ExecuteContext.GetDrawInterface()->DrawArrow(FTransform(-FVector::YAxisVector * Minimum) * ArrowTransform, -FVector::YAxisVector * Minimum * 0.25f, FVector::ZAxisVector * Minimum * 0.08f, FLinearColor::Yellow, 0.f);
+				}
+				if(Maximum > SMALL_NUMBER)
+				{
+					ExecuteContext.GetDrawInterface()->DrawCircle(Space, CircleTransform, Maximum, FLinearColor::Yellow, 0.f, 32);
+					const FTransform ArrowTransform = CircleTransform * Space;
+					ExecuteContext.GetDrawInterface()->DrawArrow(FTransform(-FVector::XAxisVector * Maximum) * ArrowTransform, FVector::XAxisVector * Maximum * 0.25f, FVector::ZAxisVector * Maximum * 0.08f, FLinearColor::Yellow, 0.f);
+					ExecuteContext.GetDrawInterface()->DrawArrow(FTransform(FVector::XAxisVector * Maximum) * ArrowTransform, -FVector::XAxisVector * Maximum * 0.25f, FVector::ZAxisVector * Maximum * 0.08f, FLinearColor::Yellow, 0.f);
+					ExecuteContext.GetDrawInterface()->DrawArrow(FTransform(-FVector::YAxisVector * Maximum) * ArrowTransform, FVector::YAxisVector * Maximum * 0.25f, FVector::ZAxisVector * Maximum * 0.08f, FLinearColor::Yellow, 0.f);
+					ExecuteContext.GetDrawInterface()->DrawArrow(FTransform(FVector::YAxisVector * Maximum) * ArrowTransform, -FVector::YAxisVector * Maximum * 0.25f, FVector::ZAxisVector * Maximum * 0.08f, FLinearColor::Yellow, 0.f);
 				}
-				ExecuteContext.GetDrawInterface()->DrawArc(Space, CircleTransform, Maximum, 0.f, PI * 2.f, DebugColor, DebugThickness, 16);
 				CircleTransform.SetRotation(FQuat(FVector(0.f, 1.f, 0.f), PI * 0.5f));
 				if (Minimum > SMALL_NUMBER)
 				{
-					ExecuteContext.GetDrawInterface()->DrawArc(Space, CircleTransform, Minimum, 0.f, PI * 2.f, DebugColor, DebugThickness, 16);
+					ExecuteContext.GetDrawInterface()->DrawCircle(Space, CircleTransform, Minimum, FLinearColor::Yellow, 0.f, 32);
+					const FTransform ArrowTransform = CircleTransform * Space;
+					ExecuteContext.GetDrawInterface()->DrawArrow(FTransform(FVector::XAxisVector * Minimum) * ArrowTransform, FVector::XAxisVector * Minimum * 0.25f, FVector::ZAxisVector * Minimum * 0.08f, FLinearColor::Yellow, 0.f);
+					ExecuteContext.GetDrawInterface()->DrawArrow(FTransform(-FVector::XAxisVector * Minimum) * ArrowTransform, -FVector::XAxisVector * Minimum * 0.25f, FVector::ZAxisVector * Minimum * 0.08f, FLinearColor::Yellow, 0.f);
+					ExecuteContext.GetDrawInterface()->DrawArrow(FTransform(FVector::YAxisVector * Minimum) * ArrowTransform, FVector::YAxisVector * Minimum * 0.25f, FVector::ZAxisVector * Minimum * 0.08f, FLinearColor::Yellow, 0.f);
+					ExecuteContext.GetDrawInterface()->DrawArrow(FTransform(-FVector::YAxisVector * Minimum) * ArrowTransform, -FVector::YAxisVector * Minimum * 0.25f, FVector::ZAxisVector * Minimum * 0.08f, FLinearColor::Yellow, 0.f);
+				}
+				if(Maximum > SMALL_NUMBER)
+				{
+					ExecuteContext.GetDrawInterface()->DrawCircle(Space, CircleTransform, Maximum, FLinearColor::Yellow, 0.f, 32);
+					const FTransform ArrowTransform = CircleTransform * Space;
+					ExecuteContext.GetDrawInterface()->DrawArrow(FTransform(-FVector::XAxisVector * Maximum) * ArrowTransform, FVector::XAxisVector * Maximum * 0.25f, FVector::ZAxisVector * Maximum * 0.08f, FLinearColor::Yellow, 0.f);
+					ExecuteContext.GetDrawInterface()->DrawArrow(FTransform(FVector::XAxisVector * Maximum) * ArrowTransform, -FVector::XAxisVector * Maximum * 0.25f, FVector::ZAxisVector * Maximum * 0.08f, FLinearColor::Yellow, 0.f);
+					ExecuteContext.GetDrawInterface()->DrawArrow(FTransform(-FVector::YAxisVector * Maximum) * ArrowTransform, FVector::YAxisVector * Maximum * 0.25f, FVector::ZAxisVector * Maximum * 0.08f, FLinearColor::Yellow, 0.f);
+					ExecuteContext.GetDrawInterface()->DrawArrow(FTransform(FVector::YAxisVector * Maximum) * ArrowTransform, -FVector::YAxisVector * Maximum * 0.25f, FVector::ZAxisVector * Maximum * 0.08f, FLinearColor::Yellow, 0.f);
 				}
-				ExecuteContext.GetDrawInterface()->DrawArc(Space, CircleTransform, Maximum, 0.f, PI * 2.f, DebugColor, DebugThickness, 16);
 				CircleTransform.SetRotation(FQuat(FVector(1.f, 0.f, 0.f), PI * 0.5f));
 				if (Minimum > SMALL_NUMBER)
 				{
-					ExecuteContext.GetDrawInterface()->DrawArc(Space, CircleTransform, Minimum, 0.f, PI * 2.f, DebugColor, DebugThickness, 16);
+					ExecuteContext.GetDrawInterface()->DrawCircle(Space, CircleTransform, Minimum, FLinearColor::Yellow, 0.f, 32);
+					const FTransform ArrowTransform = CircleTransform * Space;
+					ExecuteContext.GetDrawInterface()->DrawArrow(FTransform(FVector::XAxisVector * Minimum) * ArrowTransform, FVector::XAxisVector * Minimum * 0.25f, FVector::ZAxisVector * Minimum * 0.08f, FLinearColor::Yellow, 0.f);
+					ExecuteContext.GetDrawInterface()->DrawArrow(FTransform(-FVector::XAxisVector * Minimum) * ArrowTransform, -FVector::XAxisVector * Minimum * 0.25f, FVector::ZAxisVector * Minimum * 0.08f, FLinearColor::Yellow, 0.f);
+					ExecuteContext.GetDrawInterface()->DrawArrow(FTransform(FVector::YAxisVector * Minimum) * ArrowTransform, FVector::YAxisVector * Minimum * 0.25f, FVector::ZAxisVector * Minimum * 0.08f, FLinearColor::Yellow, 0.f);
+					ExecuteContext.GetDrawInterface()->DrawArrow(FTransform(-FVector::YAxisVector * Minimum) * ArrowTransform, -FVector::YAxisVector * Minimum * 0.25f, FVector::ZAxisVector * Minimum * 0.08f, FLinearColor::Yellow, 0.f);
+				}
+				if(Maximum > SMALL_NUMBER)
+				{
+					ExecuteContext.GetDrawInterface()->DrawCircle(Space, CircleTransform, Maximum, FLinearColor::Yellow, 0.f, 32);
+					const FTransform ArrowTransform = CircleTransform * Space;
+					ExecuteContext.GetDrawInterface()->DrawArrow(FTransform(-FVector::XAxisVector * Maximum) * ArrowTransform, FVector::XAxisVector * Maximum * 0.25f, FVector::ZAxisVector * Maximum * 0.08f, FLinearColor::Yellow, 0.f);
+					ExecuteContext.GetDrawInterface()->DrawArrow(FTransform(FVector::XAxisVector * Maximum) * ArrowTransform, -FVector::XAxisVector * Maximum * 0.25f, FVector::ZAxisVector * Maximum * 0.08f, FLinearColor::Yellow, 0.f);
+					ExecuteContext.GetDrawInterface()->DrawArrow(FTransform(-FVector::YAxisVector * Maximum) * ArrowTransform, FVector::YAxisVector * Maximum * 0.25f, FVector::ZAxisVector * Maximum * 0.08f, FLinearColor::Yellow, 0.f);
+					ExecuteContext.GetDrawInterface()->DrawArrow(FTransform(FVector::YAxisVector * Maximum) * ArrowTransform, -FVector::YAxisVector * Maximum * 0.25f, FVector::ZAxisVector * Maximum * 0.08f, FLinearColor::Yellow, 0.f);
 				}
-				ExecuteContext.GetDrawInterface()->DrawArc(Space, CircleTransform, Maximum, 0.f, PI * 2.f, DebugColor, DebugThickness, 16);
 				break;
 			}
-
 		}
+		ExecuteContext.GetDrawInterface()->DrawPoint(FTransform::Identity, Result, DebugThickness * 8.f, DebugColor);
 	}
 }
 
diff --git a/Engine/Plugins/Runtime/RigVM/Source/RigVM/Private/RigVMFunctions/Math/RigVMMathLibrary.cpp b/Engine/Plugins/Runtime/RigVM/Source/RigVM/Private/RigVMFunctions/Math/RigVMMathLibrary.cpp
--- a/Engine/Plugins/Runtime/RigVM/Source/RigVM/Private/RigVMFunctions/Math/RigVMMathLibrary.cpp
+++ b/Engine/Plugins/Runtime/RigVM/Source/RigVM/Private/RigVMFunctions/Math/RigVMMathLibrary.cpp
@@ -263,6 +263,15 @@ FVector FRigVMMathLibrary::ClampSpatially(const FVector& Value, EAxis::Type Axis
 {
 	FVector Local = Space.InverseTransformPosition(Value);
 
+	auto Clamp = [](float InValue, float InMinimum, float InMaximum)
+	{
+		if(InMaximum <= InMinimum || InMaximum < SMALL_NUMBER)
+		{
+			return FMath::Max<float>(InValue, InMinimum);
+		}
+		return FMath::Clamp<float>(InValue, InMinimum, InMaximum);
+	};
+
 	switch (Type)
 	{
 		case ERigVMClampSpatialMode::Plane:
@@ -271,17 +280,17 @@ FVector FRigVMMathLibrary::ClampSpatially(const FVector& Value, EAxis::Type Axis
 			{
 				case EAxis::X:
 				{
-					Local.X = FMath::Clamp<float>(Local.X, Minimum, Maximum);
+					Local.X = Clamp((float)Local.X, Minimum, Maximum);
 					break;
 				}
 				case EAxis::Y:
 				{
-					Local.Y = FMath::Clamp<float>(Local.Y, Minimum, Maximum);
+					Local.Y = Clamp((float)Local.Y, Minimum, Maximum);
 					break;
 				}
 				default:
 				{
-					Local.Z = FMath::Clamp<float>(Local.Z, Minimum, Maximum);
+					Local.Z = Clamp((float)Local.Z, Minimum, Maximum);
 					break;
 				}
 			}
@@ -296,8 +305,8 @@ FVector FRigVMMathLibrary::ClampSpatially(const FVector& Value, EAxis::Type Axis
 					FVector OnPlane = Local * FVector(0.f, 1.f, 1.f);
 					if (!OnPlane.IsNearlyZero())
 					{
-						float Length = OnPlane.Size();
-						OnPlane = OnPlane * FMath::Clamp<float>(Length, Minimum, Maximum) / Length;
+						const float Length = (float)OnPlane.Size();
+						OnPlane = OnPlane * Clamp(Length, Minimum, Maximum) / Length;
 						Local.Y = OnPlane.Y;
 						Local.Z = OnPlane.Z;
 					}
@@ -308,8 +317,8 @@ FVector FRigVMMathLibrary::ClampSpatially(const FVector& Value, EAxis::Type Axis
 					FVector OnPlane = Local * FVector(1.f, 0.f, 1.f);
 					if (!OnPlane.IsNearlyZero())
 					{
-						float Length = OnPlane.Size();
-						OnPlane = OnPlane * FMath::Clamp<float>(Length, Minimum, Maximum) / Length;
+						const float Length = (float)OnPlane.Size();
+						OnPlane = OnPlane * Clamp(Length, Minimum, Maximum) / Length;
 						Local.X = OnPlane.X;
 						Local.Z = OnPlane.Z;
 					}
@@ -320,8 +329,8 @@ FVector FRigVMMathLibrary::ClampSpatially(const FVector& Value, EAxis::Type Axis
 					FVector OnPlane = Local * FVector(1.f, 1.f, 0.f);
 					if (!OnPlane.IsNearlyZero())
 					{
-						float Length = OnPlane.Size();
-						OnPlane = OnPlane * FMath::Clamp<float>(Length, Minimum, Maximum) / Length;
+						const float Length = (float)OnPlane.Size();
+						OnPlane = OnPlane * Clamp(Length, Minimum, Maximum) / Length;
 						Local.X = OnPlane.X;
 						Local.Y = OnPlane.Y;
 					}
@@ -335,8 +344,8 @@ FVector FRigVMMathLibrary::ClampSpatially(const FVector& Value, EAxis::Type Axis
 		{
 			if (!Local.IsNearlyZero())
 			{
-				float Length = Local.Size();
-				Local = Local * FMath::Clamp<float>(Length, Minimum, Maximum) / Length;
+				const float Length = (float)Local.Size();
+				Local = Local * Clamp(Length, Minimum, Maximum) / Length;
 			}
 			break;
 		}
diff --git a/Engine/Plugins/Runtime/RigVM/Source/RigVM/Public/RigVMCore/RigVMDrawInterface.h b/Engine/Plugins/Runtime/RigVM/Source/RigVM/Public/RigVMCore/RigVMDrawInterface.h
--- a/Engine/Plugins/Runtime/RigVM/Source/RigVM/Public/RigVMCore/RigVMDrawInterface.h
+++ b/Engine/Plugins/Runtime/RigVM/Source/RigVM/Public/RigVMCore/RigVMDrawInterface.h
@@ -26,7 +26,10 @@ public:
 	void DrawAxes(const FTransform& WorldOffset, const TArrayView<const FTransform>& Transforms, const FLinearColor& InColor, float Size, float Thickness = 0.f);
 	void DrawRectangle(const FTransform& WorldOffset, const FTransform& Transform, float Size, const FLinearColor& Color, float Thickness);
 	void DrawArc(const FTransform& WorldOffset, const FTransform& Transform, float Radius, float MinimumAngle, float MaximumAngle, const FLinearColor& Color, float Thickness, int32 Detail);
+	void DrawCircle(const FTransform& WorldOffset, const FTransform& Transform, float Radius, const FLinearColor& Color, float Thickness, int32 Detail);
 	void DrawCone(const FTransform& WorldOffset, const FTransform& ConeOffset, float Angle1, float Angle2, uint32 NumSides, bool bDrawSideLines, const FLinearColor& SideLineColor, FMaterialRenderProxy* const MaterialRenderProxy);
+	void DrawArrow(const FTransform& WorldOffset, const FVector& Direction, const FVector& Side, const FLinearColor& InColor, float Thickness);
+	void DrawPlane(const FTransform& WorldOffset, const FVector2D& Scale, const FLinearColor& MeshColor, bool bDrawLines, const FLinearColor& LineColor, FMaterialRenderProxy* const MaterialRenderProxy);
 
 	bool IsEnabled() const;
 };
diff --git a/Engine/Plugins/Runtime/RigVM/Source/RigVM/Public/RigVMFunctions/Math/RigVMFunction_MathTransform.h b/Engine/Plugins/Runtime/RigVM/Source/RigVM/Public/RigVMFunctions/Math/RigVMFunction_MathTransform.h
--- a/Engine/Plugins/Runtime/RigVM/Source/RigVM/Public/RigVMFunctions/Math/RigVMFunction_MathTransform.h
+++ b/Engine/Plugins/Runtime/RigVM/Source/RigVM/Public/RigVMFunctions/Math/RigVMFunction_MathTransform.h
@@ -504,6 +504,7 @@ struct RIGVM_API FRigVMFunction_MathTransformClampSpatially : public FRigVMFunct
 		bDrawDebug = false;
 		DebugColor = FLinearColor::Red;
 		DebugThickness = 1.f;
+		DebugScale = 50.f;
 	}
 
 	RIGVM_METHOD()
@@ -518,9 +519,14 @@ struct RIGVM_API FRigVMFunction_MathTransformClampSpatially : public FRigVMFunct
 	UPROPERTY(meta = (Input))
 	TEnumAsByte<ERigVMClampSpatialMode::Type> Type;
 
+	// The minimum allowed distance at which a collision occurs. 
+	// Disable by setting to 0.0.
 	UPROPERTY(meta = (Input))
 	float Minimum;
 
+	// This maximum allowed distance.
+	// A collision will occur towards the center at this wall.
+	// Disable by setting to 0.0.
 	UPROPERTY(meta = (Input))
 	float Maximum;
 
@@ -538,6 +544,10 @@ struct RIGVM_API FRigVMFunction_MathTransformClampSpatially : public FRigVMFunct
 	UPROPERTY(meta = (Input))
 	float DebugThickness;
 
+	// Size of debug plane to draw
+	UPROPERTY(meta = (Input))
+	float DebugScale;
+
 	UPROPERTY(meta = (Output))
 	FTransform Result;
 };
diff --git a/Engine/Plugins/Runtime/RigVM/Source/RigVM/Public/RigVMFunctions/Math/RigVMFunction_MathVector.h b/Engine/Plugins/Runtime/RigVM/Source/RigVM/Public/RigVMFunctions/Math/RigVMFunction_MathVector.h
--- a/Engine/Plugins/Runtime/RigVM/Source/RigVM/Public/RigVMFunctions/Math/RigVMFunction_MathVector.h
+++ b/Engine/Plugins/Runtime/RigVM/Source/RigVM/Public/RigVMFunctions/Math/RigVMFunction_MathVector.h
@@ -973,6 +973,8 @@ struct RIGVM_API FRigVMFunction_MathVectorMakeBezierFourPoint : public FRigVMFun
 
 /**
  * Clamps a position using a plane collision, cylindric collision or spherical collision.
+ * The collision happens both towards an inner envelope (minimum) and towards an outer envelope (maximum).
+ * You can disable the inner / outer envelope / collision by setting the minimum / maximum to 0.0.
  */
 USTRUCT(meta = (DisplayName = "Clamp Spatially", TemplateName="ClampSpatially", Keywords="Collide,Collision"))
 struct RIGVM_API FRigVMFunction_MathVectorClampSpatially: public FRigVMFunction_MathVectorBase
@@ -990,6 +992,7 @@ struct RIGVM_API FRigVMFunction_MathVectorClampSpatially: public FRigVMFunction_
 		bDrawDebug = false;
 		DebugColor = FLinearColor::Red;
 		DebugThickness = 1.f;
+		DebugScale = 50.f;
 	}
 
 	RIGVM_METHOD()
@@ -1004,9 +1007,14 @@ struct RIGVM_API FRigVMFunction_MathVectorClampSpatially: public FRigVMFunction_
 	UPROPERTY(meta = (Input))
 	TEnumAsByte<ERigVMClampSpatialMode::Type> Type;
 
+	// The minimum allowed distance at which a collision occurs. 
+	// Disable by setting to 0.0.
 	UPROPERTY(meta = (Input))
 	float Minimum;
 
+	// This maximum allowed distance.
+	// A collision will occur towards the center at this wall.
+	// Disable by setting to 0.0.
 	UPROPERTY(meta = (Input))
 	float Maximum;
 
@@ -1024,6 +1032,10 @@ struct RIGVM_API FRigVMFunction_MathVectorClampSpatially: public FRigVMFunction_
 	UPROPERTY(meta = (Input))
 	float DebugThickness;
 
+	// Size of debug plane to draw
+	UPROPERTY(meta = (Input))
+	float DebugScale;
+
 	UPROPERTY(meta = (Output))
 	FVector Result;
 };
