RigVMStorageStruct : Fixed CIS error in UnityCompilation due to name clash with other CPP

diff --git a/Engine/Plugins/Runtime/RigVM/Source/RigVM/Private/RigVMCore/RigVMMemoryStorageStruct.cpp b/Engine/Plugins/Runtime/RigVM/Source/RigVM/Private/RigVMCore/RigVMMemoryStorageStruct.cpp
--- a/Engine/Plugins/Runtime/RigVM/Source/RigVM/Private/RigVMCore/RigVMMemoryStorageStruct.cpp
+++ b/Engine/Plugins/Runtime/RigVM/Source/RigVM/Private/RigVMCore/RigVMMemoryStorageStruct.cpp
@@ -5,10 +5,10 @@
 #include "RigVMModule.h"
 
 
-namespace 
+namespace UE::RigVM::RigVMCore::Private
 {
-	const FString EmptyBraces = TEXT("()");
-	constexpr TCHAR BraceFormat[] = TEXT("(%s)");
+	static const FString EmptyBraces = TEXT("()");
+	static constexpr TCHAR BraceFormat[] = TEXT("(%s)");
 }
 
 
@@ -221,7 +221,7 @@ bool FRigVMMemoryStorageStruct::SetDataFromString(int32 InPropertyIndex, const F
 		// check if the value was provided as a single element
 		if (const FArrayProperty* ArrayProperty = CastField<FArrayProperty>(Property))
 		{
-			const FString ValueWithBraces = FString::Printf(BraceFormat, *InValue);
+			const FString ValueWithBraces = FString::Printf(UE::RigVM::RigVMCore::Private::BraceFormat, *InValue);
 
 			ErrorPipe = FRigVMMemoryStorageImportErrorContext(false);
 			Property->ImportText_Direct(*ValueWithBraces, Data, nullptr, EPropertyPortFlags::PPF_None, &ErrorPipe);
@@ -539,7 +539,7 @@ void FRigVMMemoryStorageStruct::SetDefaultValues(const TArray<FRigVMPropertyDesc
 		const FRigVMPropertyDescription& RigVMDesc = InPropertyDescriptions[PropertyIndex];
 		const FString& DefaultValue = RigVMDesc.DefaultValue;
 
-		if (DefaultValue.IsEmpty() || DefaultValue == EmptyBraces)
+		if (DefaultValue.IsEmpty() || DefaultValue == UE::RigVM::RigVMCore::Private::EmptyBraces)
 		{
 			continue;
 		}
