RigVM: Add capsule option to clamp spatially node

diff --git a/Engine/Plugins/Runtime/RigVM/Source/RigVM/Private/RigVMFunctions/Math/RigVMFunction_MathTransform.cpp b/Engine/Plugins/Runtime/RigVM/Source/RigVM/Private/RigVMFunctions/Math/RigVMFunction_MathTransform.cpp
--- a/Engine/Plugins/Runtime/RigVM/Source/RigVM/Private/RigVMFunctions/Math/RigVMFunction_MathTransform.cpp
+++ b/Engine/Plugins/Runtime/RigVM/Source/RigVM/Private/RigVMFunctions/Math/RigVMFunction_MathTransform.cpp
@@ -178,7 +178,7 @@ FRigVMFunction_MathTransformArrayToSRT_Execute()
 FRigVMFunction_MathTransformClampSpatially_Execute()
 {
 	FVector Position;
-	FRigVMFunction_MathVectorClampSpatially::StaticExecute(ExecuteContext, Value.GetTranslation(), Axis, Type, Minimum, Maximum, Space, bDrawDebug, DebugColor, DebugThickness, DebugScale, Position);
+	FRigVMFunction_MathVectorClampSpatially::StaticExecute(ExecuteContext, Value.GetTranslation(), Axis, Type, Minimum, Maximum, Space, bDrawDebug, DebugColor, DebugThickness, Position);
 	Result = Value;
 	Result.SetTranslation(Position);
 }
diff --git a/Engine/Plugins/Runtime/RigVM/Source/RigVM/Private/RigVMFunctions/Math/RigVMFunction_MathVector.cpp b/Engine/Plugins/Runtime/RigVM/Source/RigVM/Private/RigVMFunctions/Math/RigVMFunction_MathVector.cpp
--- a/Engine/Plugins/Runtime/RigVM/Source/RigVM/Private/RigVMFunctions/Math/RigVMFunction_MathVector.cpp
+++ b/Engine/Plugins/Runtime/RigVM/Source/RigVM/Private/RigVMFunctions/Math/RigVMFunction_MathVector.cpp
@@ -395,55 +395,224 @@ FRigVMFunction_MathVectorClampSpatially_Execute()
 
 	if (ExecuteContext.GetDrawInterface() != nullptr && bDrawDebug)
 	{
-		switch (Type)
+		ERigVMClampSpatialMode::Type Mode = Type;
+		float CurrentMaximum = Maximum;
+		FTransform CurrentSpace = Space;
+
+		auto DrawPlane = [ExecuteContext, Axis, Minimum, &CurrentMaximum, &CurrentSpace, Result]()
 		{
-			case ERigVMClampSpatialMode::Plane:
+			FVector A, B, Normal;
+				
+			switch (Axis)
 			{
-				FVector A, B, Normal;
-					
-				switch (Axis)
+				case EAxis::X:
 				{
-					case EAxis::X:
-					{
-						A = FVector::YAxisVector;
-						B = FVector::ZAxisVector;
-						Normal = FVector::XAxisVector;
-						break;
-					}
-					case EAxis::Y:
-					{
-						A = FVector::XAxisVector;
-						B = FVector::ZAxisVector;
-						Normal = FVector::YAxisVector;
-						break;
-					}
-					default:
-					{
-						A = FVector::XAxisVector;
-						B = FVector::YAxisVector;
-						Normal = FVector::ZAxisVector;
-						break;
-					}
+					A = FVector::YAxisVector;
+					B = FVector::ZAxisVector;
+					Normal = FVector::XAxisVector;
+					break;
 				}
+				case EAxis::Y:
+				{
+					A = FVector::XAxisVector;
+					B = FVector::ZAxisVector;
+					Normal = FVector::YAxisVector;
+					break;
+				}
+				default:
+				{
+					A = FVector::XAxisVector;
+					B = FVector::YAxisVector;
+					Normal = FVector::ZAxisVector;
+					break;
+				}
+			}
 
-				FTransform PlaneOffset = FTransform::Identity;
-				PlaneOffset.SetRotation(FMatrix(A, Normal.Cross(A), Normal, FVector::ZeroVector).ToQuat());
+			const FVector PointOnPlane = FVector::PointPlaneProject(Result, FPlane(CurrentSpace.GetTranslation(), Normal));
+			const float PlaneScale = FMath::Max<float>(10.f, FVector::Distance(CurrentSpace.GetTranslation(), PointOnPlane)) * 0.525f;
 
-				const FTransform MinimumWorldTransform = FTransform(Normal * Minimum) * Space;
-				ExecuteContext.GetDrawInterface()->DrawPlane(PlaneOffset * MinimumWorldTransform, FVector2D(1, 1) * DebugScale, FLinearColor::Yellow, true, FLinearColor::Yellow, GEngine->ConstraintLimitMaterialPrismatic->GetRenderProxy());
-				ExecuteContext.GetDrawInterface()->DrawArrow(MinimumWorldTransform, Normal * DebugScale * 0.25f, A * DebugScale * 0.08f, FLinearColor::Yellow, 0.f);
+			FTransform PlaneRotation = FTransform::Identity;
+			PlaneRotation.SetRotation(FMatrix(A, Normal.Cross(A), Normal, FVector::ZeroVector).ToQuat());
+			const FTransform PlaneTranslation = FTransform((PointOnPlane - CurrentSpace.GetTranslation()) * 0.5);
 
-				if (Maximum > SMALL_NUMBER)
+			const FTransform MinimumWorldTransform = FTransform(Normal * Minimum) * CurrentSpace * PlaneTranslation;
+			ExecuteContext.GetDrawInterface()->DrawPlane(PlaneRotation * MinimumWorldTransform, FVector2D(PlaneScale, PlaneScale), FLinearColor::Yellow, true, FLinearColor::Yellow, GEngine->ConstraintLimitMaterialPrismatic->GetRenderProxy());
+			ExecuteContext.GetDrawInterface()->DrawArrow(MinimumWorldTransform, Normal * PlaneScale * 0.25f, A * PlaneScale * 0.08f, FLinearColor::Yellow, 0.f);
+
+			if (CurrentMaximum > SMALL_NUMBER)
+			{
+				const FTransform MaximumWorldTransform = FTransform(Normal * CurrentMaximum) * CurrentSpace * PlaneTranslation;
+				ExecuteContext.GetDrawInterface()->DrawPlane(PlaneRotation * MaximumWorldTransform, FVector2D(PlaneScale, PlaneScale), FLinearColor::Yellow, true, FLinearColor::Yellow, GEngine->ConstraintLimitMaterialPrismatic->GetRenderProxy());
+				ExecuteContext.GetDrawInterface()->DrawArrow(MaximumWorldTransform, -Normal * PlaneScale * 0.25f, A * PlaneScale * 0.08f, FLinearColor::Yellow, 0.f);
+			}
+		};
+
+		auto DrawCylinder = [ExecuteContext, Axis, Minimum, &CurrentMaximum, &CurrentSpace, Result]()
+		{
+			FTransform CircleTransform = FTransform::Identity;
+			switch (Axis)
+			{
+				case EAxis::X:
+				{
+					CircleTransform.SetRotation(FQuat(FVector(0.f, 1.f, 0.f), PI * 0.5f));
+					break;
+				}
+				case EAxis::Y:
+				{
+					CircleTransform.SetRotation(FQuat(FVector(1.f, 0.f, 0.f), PI * 0.5f));
+					break;
+				}
+				default:
+				{
+					break;
+				}
+			}
+
+			const FTransform CombinedTransform = CircleTransform * CurrentSpace;
+			const FVector PointOnAxis = CombinedTransform.InverseTransformPosition(Result) * FVector(0, 0, 1);
+			const float Extent = FMath::Abs<float>((float)PointOnAxis.Z);
+			const FTransform CenterTransform = FTransform(PointOnAxis * 0.5) * CombinedTransform;
+			const TArray<FTransform> Transforms = {
+				FTransform(FVector(0, 0, Extent * 0.5)) * CenterTransform,
+				FTransform(FVector(0, 0, -Extent * 0.5)) * CenterTransform
+			};
+
+			if (Minimum > SMALL_NUMBER)
+			{
+				for(const FTransform& Transform : Transforms)
+				{
+					ExecuteContext.GetDrawInterface()->DrawCircle(Transform, FTransform::Identity, Minimum, FLinearColor::Yellow, 0.f, 32);
+					ExecuteContext.GetDrawInterface()->DrawArrow(FTransform(FVector::XAxisVector * Minimum) * Transform, FVector::XAxisVector * Minimum * 0.25f, FVector::YAxisVector * Minimum * 0.08f, FLinearColor::Yellow, 0.f);
+					ExecuteContext.GetDrawInterface()->DrawArrow(FTransform(-FVector::XAxisVector * Minimum) * Transform, -FVector::XAxisVector * Minimum * 0.25f, FVector::YAxisVector * Minimum * 0.08f, FLinearColor::Yellow, 0.f);
+					ExecuteContext.GetDrawInterface()->DrawArrow(FTransform(FVector::YAxisVector * Minimum) * Transform, FVector::YAxisVector * Minimum * 0.25f, FVector::XAxisVector * Minimum * 0.08f, FLinearColor::Yellow, 0.f);
+					ExecuteContext.GetDrawInterface()->DrawArrow(FTransform(-FVector::YAxisVector * Minimum) * Transform, -FVector::YAxisVector * Minimum * 0.25f, FVector::XAxisVector * Minimum * 0.08f, FLinearColor::Yellow, 0.f);
+				}
+				ExecuteContext.GetDrawInterface()->DrawLines(FTransform::Identity,
+					{
+						Transforms[0].TransformPosition(FVector::XAxisVector * Minimum),
+						Transforms[1].TransformPosition(FVector::XAxisVector * Minimum),
+						Transforms[0].TransformPosition(-FVector::XAxisVector * Minimum),
+						Transforms[1].TransformPosition(-FVector::XAxisVector * Minimum),
+						Transforms[0].TransformPosition(FVector::YAxisVector * Minimum),
+						Transforms[1].TransformPosition(FVector::YAxisVector * Minimum),
+						Transforms[0].TransformPosition(-FVector::YAxisVector * Minimum),
+						Transforms[1].TransformPosition(-FVector::YAxisVector * Minimum),
+					},
+					FLinearColor::Yellow,
+					0
+				);
+			}
+			if (CurrentMaximum > SMALL_NUMBER)
+			{
+				for(const FTransform& Transform : Transforms)
 				{
-					const FTransform MaximumWorldTransform = FTransform(Normal * Maximum) * Space;
-					ExecuteContext.GetDrawInterface()->DrawPlane(PlaneOffset * MaximumWorldTransform, FVector2D(1, 1) * DebugScale, FLinearColor::Yellow, true, FLinearColor::Yellow, GEngine->ConstraintLimitMaterialPrismatic->GetRenderProxy());
-					ExecuteContext.GetDrawInterface()->DrawArrow(MaximumWorldTransform, -Normal * DebugScale * 0.25f, A * DebugScale * 0.08f, FLinearColor::Yellow, 0.f);
+					TArray<FVector> LinesToDraw;
+					LinesToDraw.Reserve(8);
+					ExecuteContext.GetDrawInterface()->DrawCircle(Transform, FTransform::Identity, CurrentMaximum, FLinearColor::Yellow, 0.f, 32);
+					ExecuteContext.GetDrawInterface()->DrawArrow(FTransform(-FVector::XAxisVector * CurrentMaximum) * Transform, FVector::XAxisVector * CurrentMaximum * 0.25f, FVector::YAxisVector * CurrentMaximum * 0.08f, FLinearColor::Yellow, 0.f);
+					ExecuteContext.GetDrawInterface()->DrawArrow(FTransform(FVector::XAxisVector * CurrentMaximum) * Transform, -FVector::XAxisVector * CurrentMaximum * 0.25f, FVector::YAxisVector * CurrentMaximum * 0.08f, FLinearColor::Yellow, 0.f);
+					ExecuteContext.GetDrawInterface()->DrawArrow(FTransform(-FVector::YAxisVector * CurrentMaximum) * Transform, FVector::YAxisVector * CurrentMaximum * 0.25f, FVector::XAxisVector * CurrentMaximum * 0.08f, FLinearColor::Yellow, 0.f);
+					ExecuteContext.GetDrawInterface()->DrawArrow(FTransform(FVector::YAxisVector * CurrentMaximum) * Transform, -FVector::YAxisVector * CurrentMaximum * 0.25f, FVector::XAxisVector * CurrentMaximum * 0.08f, FLinearColor::Yellow, 0.f);
 				}
+				ExecuteContext.GetDrawInterface()->DrawLines(FTransform::Identity,
+                    {
+                    	Transforms[0].TransformPosition(FVector::XAxisVector * CurrentMaximum),
+                    	Transforms[1].TransformPosition(FVector::XAxisVector * CurrentMaximum),
+                    	Transforms[0].TransformPosition(-FVector::XAxisVector * CurrentMaximum),
+                    	Transforms[1].TransformPosition(-FVector::XAxisVector * CurrentMaximum),
+                    	Transforms[0].TransformPosition(FVector::YAxisVector * CurrentMaximum),
+                    	Transforms[1].TransformPosition(FVector::YAxisVector * CurrentMaximum),
+                    	Transforms[0].TransformPosition(-FVector::YAxisVector * CurrentMaximum),
+                    	Transforms[1].TransformPosition(-FVector::YAxisVector * CurrentMaximum),
+                    },
+                    FLinearColor::Yellow,
+                    0
+                );
+			}
+		};
 
+		auto DrawSphere = [ExecuteContext, Axis, Minimum, &CurrentMaximum, &CurrentSpace, Result]()
+		{
+			FTransform CircleTransform = FTransform::Identity;
+			if (Minimum > SMALL_NUMBER)
+			{
+				ExecuteContext.GetDrawInterface()->DrawCircle(CurrentSpace, CircleTransform, Minimum, FLinearColor::Yellow, 0.f, 32);
+				const FTransform ArrowTransform = CircleTransform * CurrentSpace;
+				ExecuteContext.GetDrawInterface()->DrawArrow(FTransform(FVector::XAxisVector * Minimum) * ArrowTransform, FVector::XAxisVector * Minimum * 0.25f, FVector::YAxisVector * Minimum * 0.08f, FLinearColor::Yellow, 0.f);
+				ExecuteContext.GetDrawInterface()->DrawArrow(FTransform(-FVector::XAxisVector * Minimum) * ArrowTransform, -FVector::XAxisVector * Minimum * 0.25f, FVector::YAxisVector * Minimum * 0.08f, FLinearColor::Yellow, 0.f);
+				ExecuteContext.GetDrawInterface()->DrawArrow(FTransform(FVector::YAxisVector * Minimum) * ArrowTransform, FVector::YAxisVector * Minimum * 0.25f, FVector::YAxisVector * Minimum * 0.08f, FLinearColor::Yellow, 0.f);
+				ExecuteContext.GetDrawInterface()->DrawArrow(FTransform(-FVector::YAxisVector * Minimum) * ArrowTransform, -FVector::YAxisVector * Minimum * 0.25f, FVector::YAxisVector * Minimum * 0.08f, FLinearColor::Yellow, 0.f);
+			}
+			if(CurrentMaximum > SMALL_NUMBER)
+			{
+				ExecuteContext.GetDrawInterface()->DrawCircle(CurrentSpace, CircleTransform, CurrentMaximum, FLinearColor::Yellow, 0.f, 32);
+				const FTransform ArrowTransform = CircleTransform * CurrentSpace;
+				ExecuteContext.GetDrawInterface()->DrawArrow(FTransform(-FVector::XAxisVector * CurrentMaximum) * ArrowTransform, FVector::XAxisVector * CurrentMaximum * 0.25f, FVector::ZAxisVector * CurrentMaximum * 0.08f, FLinearColor::Yellow, 0.f);
+				ExecuteContext.GetDrawInterface()->DrawArrow(FTransform(FVector::XAxisVector * CurrentMaximum) * ArrowTransform, -FVector::XAxisVector * CurrentMaximum * 0.25f, FVector::ZAxisVector * CurrentMaximum * 0.08f, FLinearColor::Yellow, 0.f);
+				ExecuteContext.GetDrawInterface()->DrawArrow(FTransform(-FVector::YAxisVector * CurrentMaximum) * ArrowTransform, FVector::YAxisVector * CurrentMaximum * 0.25f, FVector::ZAxisVector * CurrentMaximum * 0.08f, FLinearColor::Yellow, 0.f);
+				ExecuteContext.GetDrawInterface()->DrawArrow(FTransform(FVector::YAxisVector * CurrentMaximum) * ArrowTransform, -FVector::YAxisVector * CurrentMaximum * 0.25f, FVector::ZAxisVector * CurrentMaximum * 0.08f, FLinearColor::Yellow, 0.f);
+			}
+			CircleTransform.SetRotation(FQuat(FVector(0.f, 1.f, 0.f), PI * 0.5f));
+			if (Minimum > SMALL_NUMBER)
+			{
+				ExecuteContext.GetDrawInterface()->DrawCircle(CurrentSpace, CircleTransform, Minimum, FLinearColor::Yellow, 0.f, 32);
+				const FTransform ArrowTransform = CircleTransform * CurrentSpace;
+				ExecuteContext.GetDrawInterface()->DrawArrow(FTransform(FVector::XAxisVector * Minimum) * ArrowTransform, FVector::XAxisVector * Minimum * 0.25f, FVector::ZAxisVector * Minimum * 0.08f, FLinearColor::Yellow, 0.f);
+				ExecuteContext.GetDrawInterface()->DrawArrow(FTransform(-FVector::XAxisVector * Minimum) * ArrowTransform, -FVector::XAxisVector * Minimum * 0.25f, FVector::ZAxisVector * Minimum * 0.08f, FLinearColor::Yellow, 0.f);
+				ExecuteContext.GetDrawInterface()->DrawArrow(FTransform(FVector::YAxisVector * Minimum) * ArrowTransform, FVector::YAxisVector * Minimum * 0.25f, FVector::ZAxisVector * Minimum * 0.08f, FLinearColor::Yellow, 0.f);
+				ExecuteContext.GetDrawInterface()->DrawArrow(FTransform(-FVector::YAxisVector * Minimum) * ArrowTransform, -FVector::YAxisVector * Minimum * 0.25f, FVector::ZAxisVector * Minimum * 0.08f, FLinearColor::Yellow, 0.f);
+			}
+			if(CurrentMaximum > SMALL_NUMBER)
+			{
+				ExecuteContext.GetDrawInterface()->DrawCircle(CurrentSpace, CircleTransform, CurrentMaximum, FLinearColor::Yellow, 0.f, 32);
+				const FTransform ArrowTransform = CircleTransform * CurrentSpace;
+				ExecuteContext.GetDrawInterface()->DrawArrow(FTransform(-FVector::XAxisVector * CurrentMaximum) * ArrowTransform, FVector::XAxisVector * CurrentMaximum * 0.25f, FVector::ZAxisVector * CurrentMaximum * 0.08f, FLinearColor::Yellow, 0.f);
+				ExecuteContext.GetDrawInterface()->DrawArrow(FTransform(FVector::XAxisVector * CurrentMaximum) * ArrowTransform, -FVector::XAxisVector * CurrentMaximum * 0.25f, FVector::ZAxisVector * CurrentMaximum * 0.08f, FLinearColor::Yellow, 0.f);
+				ExecuteContext.GetDrawInterface()->DrawArrow(FTransform(-FVector::YAxisVector * CurrentMaximum) * ArrowTransform, FVector::YAxisVector * CurrentMaximum * 0.25f, FVector::ZAxisVector * CurrentMaximum * 0.08f, FLinearColor::Yellow, 0.f);
+				ExecuteContext.GetDrawInterface()->DrawArrow(FTransform(FVector::YAxisVector * CurrentMaximum) * ArrowTransform, -FVector::YAxisVector * CurrentMaximum * 0.25f, FVector::ZAxisVector * CurrentMaximum * 0.08f, FLinearColor::Yellow, 0.f);
+			}
+			CircleTransform.SetRotation(FQuat(FVector(1.f, 0.f, 0.f), PI * 0.5f));
+			if (Minimum > SMALL_NUMBER)
+			{
+				ExecuteContext.GetDrawInterface()->DrawCircle(CurrentSpace, CircleTransform, Minimum, FLinearColor::Yellow, 0.f, 32);
+				const FTransform ArrowTransform = CircleTransform * CurrentSpace;
+				ExecuteContext.GetDrawInterface()->DrawArrow(FTransform(FVector::XAxisVector * Minimum) * ArrowTransform, FVector::XAxisVector * Minimum * 0.25f, FVector::ZAxisVector * Minimum * 0.08f, FLinearColor::Yellow, 0.f);
+				ExecuteContext.GetDrawInterface()->DrawArrow(FTransform(-FVector::XAxisVector * Minimum) * ArrowTransform, -FVector::XAxisVector * Minimum * 0.25f, FVector::ZAxisVector * Minimum * 0.08f, FLinearColor::Yellow, 0.f);
+				ExecuteContext.GetDrawInterface()->DrawArrow(FTransform(FVector::YAxisVector * Minimum) * ArrowTransform, FVector::YAxisVector * Minimum * 0.25f, FVector::ZAxisVector * Minimum * 0.08f, FLinearColor::Yellow, 0.f);
+				ExecuteContext.GetDrawInterface()->DrawArrow(FTransform(-FVector::YAxisVector * Minimum) * ArrowTransform, -FVector::YAxisVector * Minimum * 0.25f, FVector::ZAxisVector * Minimum * 0.08f, FLinearColor::Yellow, 0.f);
+			}
+			if(CurrentMaximum > SMALL_NUMBER)
+			{
+				ExecuteContext.GetDrawInterface()->DrawCircle(CurrentSpace, CircleTransform, CurrentMaximum, FLinearColor::Yellow, 0.f, 32);
+				const FTransform ArrowTransform = CircleTransform * CurrentSpace;
+				ExecuteContext.GetDrawInterface()->DrawArrow(FTransform(-FVector::XAxisVector * CurrentMaximum) * ArrowTransform, FVector::XAxisVector * CurrentMaximum * 0.25f, FVector::ZAxisVector * CurrentMaximum * 0.08f, FLinearColor::Yellow, 0.f);
+				ExecuteContext.GetDrawInterface()->DrawArrow(FTransform(FVector::XAxisVector * CurrentMaximum) * ArrowTransform, -FVector::XAxisVector * CurrentMaximum * 0.25f, FVector::ZAxisVector * CurrentMaximum * 0.08f, FLinearColor::Yellow, 0.f);
+				ExecuteContext.GetDrawInterface()->DrawArrow(FTransform(-FVector::YAxisVector * CurrentMaximum) * ArrowTransform, FVector::YAxisVector * CurrentMaximum * 0.25f, FVector::ZAxisVector * CurrentMaximum * 0.08f, FLinearColor::Yellow, 0.f);
+				ExecuteContext.GetDrawInterface()->DrawArrow(FTransform(FVector::YAxisVector * CurrentMaximum) * ArrowTransform, -FVector::YAxisVector * CurrentMaximum * 0.25f, FVector::ZAxisVector * CurrentMaximum * 0.08f, FLinearColor::Yellow, 0.f);
+			}
+		};
+		
+		switch (Mode)
+		{
+			case ERigVMClampSpatialMode::Plane:
+			{
+				DrawPlane();
 				break;
 			}
-			case ERigVMClampSpatialMode::Cylinder:
+			case ERigVMClampSpatialMode::Capsule:
 			{
+				const float Radius = Minimum;
+				if(CurrentMaximum < Radius * 2.f)
+				{
+					CurrentMaximum = 0.f;
+					DrawSphere();
+					break;
+				}
+			
+				const float Length = FMath::Max(CurrentMaximum, Radius * 2);
+				const float HalfLength = Length * 0.5f;
+				const float HalfCylinderLength = HalfLength - Radius;
+					
 				FTransform CircleTransform = FTransform::Identity;
 				switch (Axis)
 				{
@@ -463,88 +632,48 @@ FRigVMFunction_MathVectorClampSpatially_Execute()
 					}
 				}
 
-				const FTransform ArrowTransform = CircleTransform * Space;
-				if (Minimum > SMALL_NUMBER)
+				CircleTransform  = CircleTransform * Space;
+					
+				const FVector CylinderDirection = FVector::ZAxisVector * HalfCylinderLength;
+				ExecuteContext.GetDrawInterface()->DrawLines(CircleTransform,
+					{
+						FVector::XAxisVector * Minimum - CylinderDirection,
+						FVector::XAxisVector * Minimum + CylinderDirection,
+						-FVector::XAxisVector * Minimum - CylinderDirection,
+						-FVector::XAxisVector * Minimum + CylinderDirection,
+						FVector::YAxisVector * Minimum - CylinderDirection,
+						FVector::YAxisVector * Minimum + CylinderDirection,
+						-FVector::YAxisVector * Minimum - CylinderDirection,
+						-FVector::YAxisVector * Minimum + CylinderDirection
+					},
+					FLinearColor::Yellow,
+					0
+				);
+
+				CurrentMaximum = 0.f;
 				{
-					ExecuteContext.GetDrawInterface()->DrawCircle(Space, CircleTransform, Minimum, FLinearColor::Yellow, 0.f, 32);
-					ExecuteContext.GetDrawInterface()->DrawArrow(FTransform(FVector::XAxisVector * Minimum) * ArrowTransform, FVector::XAxisVector * Minimum * 0.25f, FVector::ZAxisVector * Minimum * 0.08f, FLinearColor::Yellow, 0.f);
-					ExecuteContext.GetDrawInterface()->DrawArrow(FTransform(-FVector::XAxisVector * Minimum) * ArrowTransform, -FVector::XAxisVector * Minimum * 0.25f, FVector::ZAxisVector * Minimum * 0.08f, FLinearColor::Yellow, 0.f);
-					ExecuteContext.GetDrawInterface()->DrawArrow(FTransform(FVector::YAxisVector * Minimum) * ArrowTransform, FVector::YAxisVector * Minimum * 0.25f, FVector::ZAxisVector * Minimum * 0.08f, FLinearColor::Yellow, 0.f);
-					ExecuteContext.GetDrawInterface()->DrawArrow(FTransform(-FVector::YAxisVector * Minimum) * ArrowTransform, -FVector::YAxisVector * Minimum * 0.25f, FVector::ZAxisVector * Minimum * 0.08f, FLinearColor::Yellow, 0.f);
+					TGuardValue<FTransform> SpaceGuard(CurrentSpace, FTransform(CylinderDirection) * CircleTransform);
+					DrawSphere();
 				}
-				if (Maximum > SMALL_NUMBER)
 				{
-					ExecuteContext.GetDrawInterface()->DrawCircle(Space, CircleTransform, Maximum, FLinearColor::Yellow, 0.f, 32);
-					ExecuteContext.GetDrawInterface()->DrawArrow(FTransform(-FVector::XAxisVector * Maximum) * ArrowTransform, FVector::XAxisVector * Maximum * 0.25f, FVector::ZAxisVector * Maximum * 0.08f, FLinearColor::Yellow, 0.f);
-					ExecuteContext.GetDrawInterface()->DrawArrow(FTransform(FVector::XAxisVector * Maximum) * ArrowTransform, -FVector::XAxisVector * Maximum * 0.25f, FVector::ZAxisVector * Maximum * 0.08f, FLinearColor::Yellow, 0.f);
-					ExecuteContext.GetDrawInterface()->DrawArrow(FTransform(-FVector::YAxisVector * Maximum) * ArrowTransform, FVector::YAxisVector * Maximum * 0.25f, FVector::ZAxisVector * Maximum * 0.08f, FLinearColor::Yellow, 0.f);
-					ExecuteContext.GetDrawInterface()->DrawArrow(FTransform(FVector::YAxisVector * Maximum) * ArrowTransform, -FVector::YAxisVector * Maximum * 0.25f, FVector::ZAxisVector * Maximum * 0.08f, FLinearColor::Yellow, 0.f);
+					TGuardValue<FTransform> SpaceGuard(CurrentSpace, FTransform(-CylinderDirection) * CircleTransform);
+					DrawSphere();
 				}
 				break;
 			}
+			case ERigVMClampSpatialMode::Cylinder:
+			{
+				DrawCylinder();
+				break;
+			}
 			default:
 			case ERigVMClampSpatialMode::Sphere:
 			{
-				FTransform CircleTransform = FTransform::Identity;
-				if (Minimum > SMALL_NUMBER)
-				{
-					ExecuteContext.GetDrawInterface()->DrawCircle(Space, CircleTransform, Minimum, FLinearColor::Yellow, 0.f, 32);
-					const FTransform ArrowTransform = CircleTransform * Space;
-					ExecuteContext.GetDrawInterface()->DrawArrow(FTransform(FVector::XAxisVector * Minimum) * ArrowTransform, FVector::XAxisVector * Minimum * 0.25f, FVector::ZAxisVector * Minimum * 0.08f, FLinearColor::Yellow, 0.f);
-					ExecuteContext.GetDrawInterface()->DrawArrow(FTransform(-FVector::XAxisVector * Minimum) * ArrowTransform, -FVector::XAxisVector * Minimum * 0.25f, FVector::ZAxisVector * Minimum * 0.08f, FLinearColor::Yellow, 0.f);
-					ExecuteContext.GetDrawInterface()->DrawArrow(FTransform(FVector::YAxisVector * Minimum) * ArrowTransform, FVector::YAxisVector * Minimum * 0.25f, FVector::ZAxisVector * Minimum * 0.08f, FLinearColor::Yellow, 0.f);
-					ExecuteContext.GetDrawInterface()->DrawArrow(FTransform(-FVector::YAxisVector * Minimum) * ArrowTransform, -FVector::YAxisVector * Minimum * 0.25f, FVector::ZAxisVector * Minimum * 0.08f, FLinearColor::Yellow, 0.f);
-				}
-				if(Maximum > SMALL_NUMBER)
-				{
-					ExecuteContext.GetDrawInterface()->DrawCircle(Space, CircleTransform, Maximum, FLinearColor::Yellow, 0.f, 32);
-					const FTransform ArrowTransform = CircleTransform * Space;
-					ExecuteContext.GetDrawInterface()->DrawArrow(FTransform(-FVector::XAxisVector * Maximum) * ArrowTransform, FVector::XAxisVector * Maximum * 0.25f, FVector::ZAxisVector * Maximum * 0.08f, FLinearColor::Yellow, 0.f);
-					ExecuteContext.GetDrawInterface()->DrawArrow(FTransform(FVector::XAxisVector * Maximum) * ArrowTransform, -FVector::XAxisVector * Maximum * 0.25f, FVector::ZAxisVector * Maximum * 0.08f, FLinearColor::Yellow, 0.f);
-					ExecuteContext.GetDrawInterface()->DrawArrow(FTransform(-FVector::YAxisVector * Maximum) * ArrowTransform, FVector::YAxisVector * Maximum * 0.25f, FVector::ZAxisVector * Maximum * 0.08f, FLinearColor::Yellow, 0.f);
-					ExecuteContext.GetDrawInterface()->DrawArrow(FTransform(FVector::YAxisVector * Maximum) * ArrowTransform, -FVector::YAxisVector * Maximum * 0.25f, FVector::ZAxisVector * Maximum * 0.08f, FLinearColor::Yellow, 0.f);
-				}
-				CircleTransform.SetRotation(FQuat(FVector(0.f, 1.f, 0.f), PI * 0.5f));
-				if (Minimum > SMALL_NUMBER)
-				{
-					ExecuteContext.GetDrawInterface()->DrawCircle(Space, CircleTransform, Minimum, FLinearColor::Yellow, 0.f, 32);
-					const FTransform ArrowTransform = CircleTransform * Space;
-					ExecuteContext.GetDrawInterface()->DrawArrow(FTransform(FVector::XAxisVector * Minimum) * ArrowTransform, FVector::XAxisVector * Minimum * 0.25f, FVector::ZAxisVector * Minimum * 0.08f, FLinearColor::Yellow, 0.f);
-					ExecuteContext.GetDrawInterface()->DrawArrow(FTransform(-FVector::XAxisVector * Minimum) * ArrowTransform, -FVector::XAxisVector * Minimum * 0.25f, FVector::ZAxisVector * Minimum * 0.08f, FLinearColor::Yellow, 0.f);
-					ExecuteContext.GetDrawInterface()->DrawArrow(FTransform(FVector::YAxisVector * Minimum) * ArrowTransform, FVector::YAxisVector * Minimum * 0.25f, FVector::ZAxisVector * Minimum * 0.08f, FLinearColor::Yellow, 0.f);
-					ExecuteContext.GetDrawInterface()->DrawArrow(FTransform(-FVector::YAxisVector * Minimum) * ArrowTransform, -FVector::YAxisVector * Minimum * 0.25f, FVector::ZAxisVector * Minimum * 0.08f, FLinearColor::Yellow, 0.f);
-				}
-				if(Maximum > SMALL_NUMBER)
-				{
-					ExecuteContext.GetDrawInterface()->DrawCircle(Space, CircleTransform, Maximum, FLinearColor::Yellow, 0.f, 32);
-					const FTransform ArrowTransform = CircleTransform * Space;
-					ExecuteContext.GetDrawInterface()->DrawArrow(FTransform(-FVector::XAxisVector * Maximum) * ArrowTransform, FVector::XAxisVector * Maximum * 0.25f, FVector::ZAxisVector * Maximum * 0.08f, FLinearColor::Yellow, 0.f);
-					ExecuteContext.GetDrawInterface()->DrawArrow(FTransform(FVector::XAxisVector * Maximum) * ArrowTransform, -FVector::XAxisVector * Maximum * 0.25f, FVector::ZAxisVector * Maximum * 0.08f, FLinearColor::Yellow, 0.f);
-					ExecuteContext.GetDrawInterface()->DrawArrow(FTransform(-FVector::YAxisVector * Maximum) * ArrowTransform, FVector::YAxisVector * Maximum * 0.25f, FVector::ZAxisVector * Maximum * 0.08f, FLinearColor::Yellow, 0.f);
-					ExecuteContext.GetDrawInterface()->DrawArrow(FTransform(FVector::YAxisVector * Maximum) * ArrowTransform, -FVector::YAxisVector * Maximum * 0.25f, FVector::ZAxisVector * Maximum * 0.08f, FLinearColor::Yellow, 0.f);
-				}
-				CircleTransform.SetRotation(FQuat(FVector(1.f, 0.f, 0.f), PI * 0.5f));
-				if (Minimum > SMALL_NUMBER)
-				{
-					ExecuteContext.GetDrawInterface()->DrawCircle(Space, CircleTransform, Minimum, FLinearColor::Yellow, 0.f, 32);
-					const FTransform ArrowTransform = CircleTransform * Space;
-					ExecuteContext.GetDrawInterface()->DrawArrow(FTransform(FVector::XAxisVector * Minimum) * ArrowTransform, FVector::XAxisVector * Minimum * 0.25f, FVector::ZAxisVector * Minimum * 0.08f, FLinearColor::Yellow, 0.f);
-					ExecuteContext.GetDrawInterface()->DrawArrow(FTransform(-FVector::XAxisVector * Minimum) * ArrowTransform, -FVector::XAxisVector * Minimum * 0.25f, FVector::ZAxisVector * Minimum * 0.08f, FLinearColor::Yellow, 0.f);
-					ExecuteContext.GetDrawInterface()->DrawArrow(FTransform(FVector::YAxisVector * Minimum) * ArrowTransform, FVector::YAxisVector * Minimum * 0.25f, FVector::ZAxisVector * Minimum * 0.08f, FLinearColor::Yellow, 0.f);
-					ExecuteContext.GetDrawInterface()->DrawArrow(FTransform(-FVector::YAxisVector * Minimum) * ArrowTransform, -FVector::YAxisVector * Minimum * 0.25f, FVector::ZAxisVector * Minimum * 0.08f, FLinearColor::Yellow, 0.f);
-				}
-				if(Maximum > SMALL_NUMBER)
-				{
-					ExecuteContext.GetDrawInterface()->DrawCircle(Space, CircleTransform, Maximum, FLinearColor::Yellow, 0.f, 32);
-					const FTransform ArrowTransform = CircleTransform * Space;
-					ExecuteContext.GetDrawInterface()->DrawArrow(FTransform(-FVector::XAxisVector * Maximum) * ArrowTransform, FVector::XAxisVector * Maximum * 0.25f, FVector::ZAxisVector * Maximum * 0.08f, FLinearColor::Yellow, 0.f);
-					ExecuteContext.GetDrawInterface()->DrawArrow(FTransform(FVector::XAxisVector * Maximum) * ArrowTransform, -FVector::XAxisVector * Maximum * 0.25f, FVector::ZAxisVector * Maximum * 0.08f, FLinearColor::Yellow, 0.f);
-					ExecuteContext.GetDrawInterface()->DrawArrow(FTransform(-FVector::YAxisVector * Maximum) * ArrowTransform, FVector::YAxisVector * Maximum * 0.25f, FVector::ZAxisVector * Maximum * 0.08f, FLinearColor::Yellow, 0.f);
-					ExecuteContext.GetDrawInterface()->DrawArrow(FTransform(FVector::YAxisVector * Maximum) * ArrowTransform, -FVector::YAxisVector * Maximum * 0.25f, FVector::ZAxisVector * Maximum * 0.08f, FLinearColor::Yellow, 0.f);
-				}
+				DrawSphere();
 				break;
 			}
 		}
+		
 		ExecuteContext.GetDrawInterface()->DrawPoint(FTransform::Identity, Result, DebugThickness * 8.f, DebugColor);
 	}
 }
diff --git a/Engine/Plugins/Runtime/RigVM/Source/RigVM/Private/RigVMFunctions/Math/RigVMMathLibrary.cpp b/Engine/Plugins/Runtime/RigVM/Source/RigVM/Private/RigVMFunctions/Math/RigVMMathLibrary.cpp
--- a/Engine/Plugins/Runtime/RigVM/Source/RigVM/Private/RigVMFunctions/Math/RigVMMathLibrary.cpp
+++ b/Engine/Plugins/Runtime/RigVM/Source/RigVM/Private/RigVMFunctions/Math/RigVMMathLibrary.cpp
@@ -272,81 +272,146 @@ FVector FRigVMMathLibrary::ClampSpatially(const FVector& Value, EAxis::Type Axis
 		return FMath::Clamp<float>(InValue, InMinimum, InMaximum);
 	};
 
-	switch (Type)
+	auto CollidePlane = [Clamp, &Local, Axis, Minimum, &Maximum]
 	{
-		case ERigVMClampSpatialMode::Plane:
+		switch (Axis)
 		{
-			switch (Axis)
+		case EAxis::X:
 			{
-				case EAxis::X:
+				Local.X = Clamp((float)Local.X, Minimum, Maximum);
+				break;
+			}
+		case EAxis::Y:
+			{
+				Local.Y = Clamp((float)Local.Y, Minimum, Maximum);
+				break;
+			}
+		default:
+			{
+				Local.Z = Clamp((float)Local.Z, Minimum, Maximum);
+				break;
+			}
+		}
+	};
+
+	auto CollideCylinder = [Clamp, &Local, Axis, Minimum, &Maximum]
+	{
+		switch (Axis)
+		{
+			case EAxis::X:
+			{
+				FVector OnPlane = Local * FVector(0.f, 1.f, 1.f);
+				if (!OnPlane.IsNearlyZero())
 				{
-					Local.X = Clamp((float)Local.X, Minimum, Maximum);
-					break;
+					const float Length = (float)OnPlane.Size();
+					OnPlane = OnPlane * Clamp(Length, Minimum, Maximum) / Length;
+					Local.Y = OnPlane.Y;
+					Local.Z = OnPlane.Z;
 				}
-				case EAxis::Y:
+				break;
+			}
+			case EAxis::Y:
+			{
+				FVector OnPlane = Local * FVector(1.f, 0.f, 1.f);
+				if (!OnPlane.IsNearlyZero())
 				{
-					Local.Y = Clamp((float)Local.Y, Minimum, Maximum);
-					break;
+					const float Length = (float)OnPlane.Size();
+					OnPlane = OnPlane * Clamp(Length, Minimum, Maximum) / Length;
+					Local.X = OnPlane.X;
+					Local.Z = OnPlane.Z;
 				}
-				default:
+				break;
+			}
+			default:
+			{
+				FVector OnPlane = Local * FVector(1.f, 1.f, 0.f);
+				if (!OnPlane.IsNearlyZero())
 				{
-					Local.Z = Clamp((float)Local.Z, Minimum, Maximum);
-					break;
+					const float Length = (float)OnPlane.Size();
+					OnPlane = OnPlane * Clamp(Length, Minimum, Maximum) / Length;
+					Local.X = OnPlane.X;
+					Local.Y = OnPlane.Y;
 				}
+				break;
 			}
+		}
+	};
+
+	auto CollideSphere = [Clamp, &Local, Minimum, &Maximum]
+	{
+		if (!Local.IsNearlyZero())
+		{
+			const float Length = (float)Local.Size();
+			Local = Local * Clamp(Length, Minimum, Maximum) / Length;
+		}
+	};
+	
+	switch (Type)
+	{
+		case ERigVMClampSpatialMode::Plane:
+		{
+			CollidePlane();
 			break;
 		}
-		case ERigVMClampSpatialMode::Cylinder:
+		case ERigVMClampSpatialMode::Capsule:
 		{
+			const float Radius = Minimum;
+			if(Maximum < Radius * 2.f)
+			{
+				Maximum = 0.f;
+				CollideSphere();
+				break;
+			}
+				
+			const float Length = FMath::Max(Maximum, Radius * 2);
+			const float HalfLength = Length * 0.5f;
+			const float HalfCylinderLength = HalfLength - Radius;
+				
+			FVector Normal = FVector::XAxisVector;
 			switch (Axis)
 			{
 				case EAxis::X:
 				{
-					FVector OnPlane = Local * FVector(0.f, 1.f, 1.f);
-					if (!OnPlane.IsNearlyZero())
-					{
-						const float Length = (float)OnPlane.Size();
-						OnPlane = OnPlane * Clamp(Length, Minimum, Maximum) / Length;
-						Local.Y = OnPlane.Y;
-						Local.Z = OnPlane.Z;
-					}
 					break;
 				}
 				case EAxis::Y:
 				{
-					FVector OnPlane = Local * FVector(1.f, 0.f, 1.f);
-					if (!OnPlane.IsNearlyZero())
-					{
-						const float Length = (float)OnPlane.Size();
-						OnPlane = OnPlane * Clamp(Length, Minimum, Maximum) / Length;
-						Local.X = OnPlane.X;
-						Local.Z = OnPlane.Z;
-					}
+					Normal = FVector::YAxisVector;
 					break;
 				}
 				default:
 				{
-					FVector OnPlane = Local * FVector(1.f, 1.f, 0.f);
-					if (!OnPlane.IsNearlyZero())
-					{
-						const float Length = (float)OnPlane.Size();
-						OnPlane = OnPlane * Clamp(Length, Minimum, Maximum) / Length;
-						Local.X = OnPlane.X;
-						Local.Y = OnPlane.Y;
-					}
+					Normal = FVector::ZAxisVector;
 					break;
 				}
 			}
+
+			const float DotOnNormal = Normal.Dot(Local);
+
+			// if we are on the cylinder part of the collision, fall through 
+			if(FMath::Abs(DotOnNormal) < HalfCylinderLength)
+			{
+				Maximum = 0.f;
+				CollideCylinder();
+				break;
+			}
+
+			// since now we are going to be colliding with either the upper of lower half sphere
+			const FVector LocalSphereCenter = Normal * FMath::Sign(DotOnNormal) * HalfCylinderLength;
+			Space = FTransform(LocalSphereCenter) * Space;
+			Local = Space.InverseTransformPosition(Value);
+			CollideSphere();
+			break;
+		}
+		case ERigVMClampSpatialMode::Cylinder:
+		{
+			CollideCylinder();
 			break;
 		}
 		default:
 		case ERigVMClampSpatialMode::Sphere:
 		{
-			if (!Local.IsNearlyZero())
-			{
-				const float Length = (float)Local.Size();
-				Local = Local * Clamp(Length, Minimum, Maximum) / Length;
-			}
+			CollideSphere();
 			break;
 		}
 
diff --git a/Engine/Plugins/Runtime/RigVM/Source/RigVM/Public/RigVMFunctions/Math/RigVMFunction_MathTransform.h b/Engine/Plugins/Runtime/RigVM/Source/RigVM/Public/RigVMFunctions/Math/RigVMFunction_MathTransform.h
--- a/Engine/Plugins/Runtime/RigVM/Source/RigVM/Public/RigVMFunctions/Math/RigVMFunction_MathTransform.h
+++ b/Engine/Plugins/Runtime/RigVM/Source/RigVM/Public/RigVMFunctions/Math/RigVMFunction_MathTransform.h
@@ -504,7 +504,6 @@ struct RIGVM_API FRigVMFunction_MathTransformClampSpatially : public FRigVMFunct
 		bDrawDebug = false;
 		DebugColor = FLinearColor::Red;
 		DebugThickness = 1.f;
-		DebugScale = 50.f;
 	}
 
 	RIGVM_METHOD()
@@ -520,12 +519,14 @@ struct RIGVM_API FRigVMFunction_MathTransformClampSpatially : public FRigVMFunct
 	TEnumAsByte<ERigVMClampSpatialMode::Type> Type;
 
 	// The minimum allowed distance at which a collision occurs. 
+	// Note: For capsule this represents the radius.
 	// Disable by setting to 0.0.
 	UPROPERTY(meta = (Input))
 	float Minimum;
 
 	// This maximum allowed distance.
 	// A collision will occur towards the center at this wall.
+	// Note: For capsule this represents the length.
 	// Disable by setting to 0.0.
 	UPROPERTY(meta = (Input))
 	float Maximum;
@@ -544,10 +545,6 @@ struct RIGVM_API FRigVMFunction_MathTransformClampSpatially : public FRigVMFunct
 	UPROPERTY(meta = (Input))
 	float DebugThickness;
 
-	// Size of debug plane to draw
-	UPROPERTY(meta = (Input))
-	float DebugScale;
-
 	UPROPERTY(meta = (Output))
 	FTransform Result;
 };
diff --git a/Engine/Plugins/Runtime/RigVM/Source/RigVM/Public/RigVMFunctions/Math/RigVMFunction_MathVector.h b/Engine/Plugins/Runtime/RigVM/Source/RigVM/Public/RigVMFunctions/Math/RigVMFunction_MathVector.h
--- a/Engine/Plugins/Runtime/RigVM/Source/RigVM/Public/RigVMFunctions/Math/RigVMFunction_MathVector.h
+++ b/Engine/Plugins/Runtime/RigVM/Source/RigVM/Public/RigVMFunctions/Math/RigVMFunction_MathVector.h
@@ -992,7 +992,6 @@ struct RIGVM_API FRigVMFunction_MathVectorClampSpatially: public FRigVMFunction_
 		bDrawDebug = false;
 		DebugColor = FLinearColor::Red;
 		DebugThickness = 1.f;
-		DebugScale = 50.f;
 	}
 
 	RIGVM_METHOD()
@@ -1007,13 +1006,15 @@ struct RIGVM_API FRigVMFunction_MathVectorClampSpatially: public FRigVMFunction_
 	UPROPERTY(meta = (Input))
 	TEnumAsByte<ERigVMClampSpatialMode::Type> Type;
 
-	// The minimum allowed distance at which a collision occurs. 
+	// The minimum allowed distance at which a collision occurs.
+	// Note: For capsule this represents the radius.
 	// Disable by setting to 0.0.
 	UPROPERTY(meta = (Input))
 	float Minimum;
 
 	// This maximum allowed distance.
 	// A collision will occur towards the center at this wall.
+	// Note: For capsule this represents the length.
 	// Disable by setting to 0.0.
 	UPROPERTY(meta = (Input))
 	float Maximum;
@@ -1032,10 +1033,6 @@ struct RIGVM_API FRigVMFunction_MathVectorClampSpatially: public FRigVMFunction_
 	UPROPERTY(meta = (Input))
 	float DebugThickness;
 
-	// Size of debug plane to draw
-	UPROPERTY(meta = (Input))
-	float DebugScale;
-
 	UPROPERTY(meta = (Output))
 	FVector Result;
 };
diff --git a/Engine/Plugins/Runtime/RigVM/Source/RigVM/Public/RigVMFunctions/RigVMFunctionDefines.h b/Engine/Plugins/Runtime/RigVM/Source/RigVM/Public/RigVMFunctions/RigVMFunctionDefines.h
--- a/Engine/Plugins/Runtime/RigVM/Source/RigVM/Public/RigVMFunctions/RigVMFunctionDefines.h
+++ b/Engine/Plugins/Runtime/RigVM/Source/RigVM/Public/RigVMFunctions/RigVMFunctionDefines.h
@@ -27,7 +27,8 @@ namespace ERigVMClampSpatialMode
 	{
 		Plane,
 		Cylinder,
-		Sphere
+		Sphere,
+		Capsule
 	};
 }
 
