RigVM Compiler: Remap variables during function compilation

diff --git a/Engine/Plugins/Runtime/RigVM/Source/RigVMDeveloper/Private/RigVMCompiler/RigVMCompiler.cpp b/Engine/Plugins/Runtime/RigVM/Source/RigVMDeveloper/Private/RigVMCompiler/RigVMCompiler.cpp
--- a/Engine/Plugins/Runtime/RigVM/Source/RigVMDeveloper/Private/RigVMCompiler/RigVMCompiler.cpp
+++ b/Engine/Plugins/Runtime/RigVM/Source/RigVMDeveloper/Private/RigVMCompiler/RigVMCompiler.cpp
@@ -506,7 +506,46 @@ bool URigVMCompiler::Compile(const FRigVMCompileSettings& InSettings, TArray<URi
 								{
 									IRigVMClientHost* ClientHost = LibraryNode->GetImplementingOuter<IRigVMClientHost>();
 									URigVMController* FunctionController = ClientHost->GetRigVMClient()->GetOrCreateController(LibraryNode->GetLibrary());
-									bSuccessfullCompilation = CompileFunction(WorkData.Settings, LibraryNode, FunctionController, &FunctionData->CompilationData, OutVMContext);
+									
+									TArray<FRigVMExternalVariable> FunctionExternalVariables = InExternalVariables;
+									if(ReferenceNode->RequiresVariableRemapping())
+									{
+										if(FunctionController->GetExternalVariablesDelegate.IsBound())
+										{
+											FunctionExternalVariables = FunctionController->GetExternalVariablesDelegate.Execute(LibraryNode->GetContainedGraph());
+
+											for(FRigVMExternalVariable& FunctionExternalVariable : FunctionExternalVariables)
+											{
+												const FName OuterVariableName = ReferenceNode->GetOuterVariableName(FunctionExternalVariable.Name);
+												if(OuterVariableName.IsNone())
+												{
+													const FString VariableRemappingErrorMessage =
+														FString::Printf(TEXT("The function's variable '%s' is not remapped on function reference @@."),
+														*FunctionExternalVariable.Name.ToString());
+													Settings.ASTSettings.Report(EMessageSeverity::Error, ReferenceNode, VariableRemappingErrorMessage);
+													bEncounteredGraphError = true;
+												}
+												else
+												{
+													const FRigVMExternalVariable* OuterExternalVariable = InExternalVariables.FindByPredicate(
+														[OuterVariableName](const FRigVMExternalVariable& ExternalVariable) -> bool
+														{
+															return ExternalVariable.Name.IsEqual(OuterVariableName, ENameCase::CaseSensitive);
+														}
+													);
+
+													if(OuterExternalVariable)
+													{
+														check(FRigVMRegistry::Get().CanMatchTypes(OuterExternalVariable->GetTypeIndex(), FunctionExternalVariable.GetTypeIndex(), true));
+														FunctionExternalVariable.Property = OuterExternalVariable->Property;
+														FunctionExternalVariable.Memory = OuterExternalVariable->Memory;
+													}
+												}
+											}
+										}
+									}
+									
+									bSuccessfullCompilation = CompileFunction(WorkData.Settings, LibraryNode, FunctionController, FunctionExternalVariables, &FunctionData->CompilationData, OutVMContext);
 								}
 								else
 								{
@@ -1385,10 +1424,10 @@ bool URigVMCompiler::Compile(const FRigVMCompileSettings& InSettings, TArray<URi
 
 bool URigVMCompiler::CompileFunction(const URigVMLibraryNode* InLibraryNode, URigVMController* InController, FRigVMFunctionCompilationData* OutFunctionCompilationData, FRigVMExtendedExecuteContext& OutVMContext)
 {
-	return CompileFunction(Settings_DEPRECATED, InLibraryNode, InController, OutFunctionCompilationData, OutVMContext);
+	return CompileFunction(Settings_DEPRECATED, InLibraryNode, InController, TArray<FRigVMExternalVariable>(), OutFunctionCompilationData, OutVMContext);
 }
 
-bool URigVMCompiler::CompileFunction(const FRigVMCompileSettings& InSettings, const URigVMLibraryNode* InLibraryNode, URigVMController* InController, FRigVMFunctionCompilationData* OutFunctionCompilationData, FRigVMExtendedExecuteContext& OutVMContext)
+bool URigVMCompiler::CompileFunction(const FRigVMCompileSettings& InSettings, const URigVMLibraryNode* InLibraryNode, URigVMController* InController, const TArray<FRigVMExternalVariable>& InExternalVariables, FRigVMFunctionCompilationData* OutFunctionCompilationData, FRigVMExtendedExecuteContext& OutVMContext)
 {
 	TGuardValue<const URigVMLibraryNode*> CompilationGuard(CurrentCompilationFunction, InLibraryNode);
 
@@ -1409,8 +1448,8 @@ bool URigVMCompiler::CompileFunction(const FRigVMCompileSettings& InSettings, co
 	OutFunctionCompilationData->Hash = 0;
 	OutFunctionCompilationData->ByteCode.Reset();
 
-	TArray<FRigVMExternalVariable> ExternalVariables;
-	if (LibraryController->GetExternalVariablesDelegate.IsBound())
+	TArray<FRigVMExternalVariable> ExternalVariables = InExternalVariables;
+	if (InExternalVariables.IsEmpty() && LibraryController->GetExternalVariablesDelegate.IsBound())
 	{
 		ExternalVariables = LibraryController->GetExternalVariablesDelegate.Execute(InLibraryNode->GetContainedGraph());
 	}
diff --git a/Engine/Plugins/Runtime/RigVM/Source/RigVMDeveloper/Public/RigVMCompiler/RigVMCompiler.h b/Engine/Plugins/Runtime/RigVM/Source/RigVMDeveloper/Public/RigVMCompiler/RigVMCompiler.h
--- a/Engine/Plugins/Runtime/RigVM/Source/RigVMDeveloper/Public/RigVMCompiler/RigVMCompiler.h
+++ b/Engine/Plugins/Runtime/RigVM/Source/RigVMDeveloper/Public/RigVMCompiler/RigVMCompiler.h
@@ -259,7 +259,7 @@ public:
 	UE_DEPRECATED(5.4, "Please use CompileFunction with Settings param.")
 	bool CompileFunction(const URigVMLibraryNode* InLibraryNode, URigVMController* InController, FRigVMFunctionCompilationData* OutFunctionCompilationData, FRigVMExtendedExecuteContext& OutVMContext);
 
-	bool CompileFunction(const FRigVMCompileSettings& InSettings, const URigVMLibraryNode* InLibraryNode, URigVMController* InController, FRigVMFunctionCompilationData* OutFunctionCompilationData, FRigVMExtendedExecuteContext& OutVMContext);
+	bool CompileFunction(const FRigVMCompileSettings& InSettings, const URigVMLibraryNode* InLibraryNode, URigVMController* InController, const TArray<FRigVMExternalVariable>& InExternalVariables, FRigVMFunctionCompilationData* OutFunctionCompilationData, FRigVMExtendedExecuteContext& OutVMContext);
 
 	FRigVMCompiler_GetFunctionCompilationData GetFunctionCompilationData;
 	TMap<FString, const FRigVMFunctionCompilationData*> CompiledFunctions;
