Fixed a crash at RigVMMemoeyStorageStruct  trying to use a null pointer (empty storage)

diff --git a/Engine/Plugins/Runtime/RigVM/Source/RigVM/Private/RigVMCore/RigVMMemoryStorageStruct.cpp b/Engine/Plugins/Runtime/RigVM/Source/RigVM/Private/RigVMCore/RigVMMemoryStorageStruct.cpp
--- a/Engine/Plugins/Runtime/RigVM/Source/RigVM/Private/RigVMCore/RigVMMemoryStorageStruct.cpp
+++ b/Engine/Plugins/Runtime/RigVM/Source/RigVM/Private/RigVMCore/RigVMMemoryStorageStruct.cpp
@@ -114,7 +114,7 @@ int32 FRigVMMemoryStorageStruct::GetPropertyIndexByName(const FName& InName) con
 FProperty* FRigVMMemoryStorageStruct::FindPropertyByName(const FName& InName) const
 {
 	const FName SanitizedName = FRigVMPropertyDescription::SanitizeName(InName);
-	return Value.GetScriptStruct()->FindPropertyByName(SanitizedName);
+	return Value.GetScriptStruct() != nullptr ? Value.GetScriptStruct()->FindPropertyByName(SanitizedName) : nullptr;
 }
 
 FRigVMOperand FRigVMMemoryStorageStruct::GetOperand(int32 InPropertyIndex, int32 InPropertyPathIndex) const
