RIgVm: Allow dispatch factories to create more complex permutations

diff --git a/Engine/Plugins/Runtime/RigVM/Source/RigVM/Private/RigVMCore/RigVMDispatchFactory.cpp b/Engine/Plugins/Runtime/RigVM/Source/RigVM/Private/RigVMCore/RigVMDispatchFactory.cpp
--- a/Engine/Plugins/Runtime/RigVM/Source/RigVM/Private/RigVMCore/RigVMDispatchFactory.cpp
+++ b/Engine/Plugins/Runtime/RigVM/Source/RigVM/Private/RigVMCore/RigVMDispatchFactory.cpp
@@ -278,6 +278,60 @@ FString FRigVMDispatchFactory::GetPermutationName(const FRigVMTemplateTypeMap& I
 	return GetPermutationNameImpl(InTypes);
 }
 
+TArray<FRigVMTemplateArgument> FRigVMDispatchFactory::BuildArgumentListFromPrimaryArgument(const TArray<FRigVMTemplateArgument>& InArguments, const FName& InPrimaryArgumentName) const
+{
+	TArray<FRigVMTemplateArgument> NewArguments;
+	const FRigVMTemplateArgument* PrimaryArgument = InArguments.FindByPredicate([InPrimaryArgumentName](const FRigVMTemplateArgument& Arg)
+	{
+		return Arg.Name == InPrimaryArgumentName;
+	});
+	
+	if (!PrimaryArgument)
+	{
+		return NewArguments;
+	}
+
+	NewArguments.SetNum(InArguments.Num());
+	for (int32 Index=0; Index < InArguments.Num(); ++Index)
+	{
+		NewArguments[Index].Name = InArguments[Index].Name;
+		NewArguments[Index].Direction = InArguments[Index].Direction;
+		NewArguments[Index].TypeCategories = InArguments[Index].TypeCategories;
+	}
+	
+	TSet<TRigVMTypeIndex> ProcessedTypes;
+	for (const TRigVMTypeIndex& Type : PrimaryArgument->TypeIndices)
+	{
+		if (ProcessedTypes.Contains(Type))
+		{
+			continue;
+		}
+
+		const TArray<FRigVMTemplateTypeMap> Permutations = GetPermutationsFromArgumentType(InPrimaryArgumentName, Type);
+		for (const FRigVMTemplateTypeMap& Permutation : Permutations)
+		{
+			for (int32 Index=0; Index < InArguments.Num(); ++Index)
+			{
+				const TRigVMTypeIndex* PermutationArg = Permutation.Find(NewArguments[Index].Name);
+				if (!PermutationArg)
+				{
+					NewArguments.Reset();
+					return NewArguments;
+				}
+				NewArguments[Index].TypeIndices.Add(*PermutationArg);
+			}
+		}
+	}
+
+	for (FRigVMTemplateArgument& Argument : NewArguments)
+	{
+		Argument.EnsureValidExecuteType();
+		Argument.UpdateTypeToPermutations();
+	}
+	
+	return NewArguments;
+}
+
 FString FRigVMDispatchFactory::GetPermutationNameImpl(const FRigVMTemplateTypeMap& InTypes) const
 {
 	static constexpr TCHAR Format[] = TEXT("%s::%s");
@@ -333,6 +387,12 @@ const FRigVMTemplate* FRigVMDispatchFactory::GetTemplate() const
 			return OnNewArgumentType(InArgumentName, InTypeIndex);
 		});
 
+	Delegates.GetPermutationsFromArgumentTypeDelegate = FRigVMTemplate_GetPermutationsFromArgumentTypeDelegate::CreateLambda(
+		[this](const FRigVMTemplate*, const FName& InArgumentName, int32 InTypeIndex)
+		{
+			return GetPermutationsFromArgumentType(InArgumentName, InTypeIndex);
+		});
+
 	Delegates.GetDispatchFactoryDelegate = FRigVMTemplate_GetDispatchFactoryDelegate::CreateLambda(
 [FactoryName]()
 	{
diff --git a/Engine/Plugins/Runtime/RigVM/Source/RigVM/Private/RigVMCore/RigVMTemplate.cpp b/Engine/Plugins/Runtime/RigVM/Source/RigVM/Private/RigVMCore/RigVMTemplate.cpp
--- a/Engine/Plugins/Runtime/RigVM/Source/RigVM/Private/RigVMCore/RigVMTemplate.cpp
+++ b/Engine/Plugins/Runtime/RigVM/Source/RigVM/Private/RigVMCore/RigVMTemplate.cpp
@@ -36,6 +36,13 @@ FRigVMTemplateArgument::FRigVMTemplateArgument()
 {
 }
 
+FRigVMTemplateArgument::FRigVMTemplateArgument(const FName& InName, ERigVMPinDirection InDirection)
+: Index(INDEX_NONE)
+, Name(InName)
+, Direction(InDirection)
+{
+}
+
 FRigVMTemplateArgument::FRigVMTemplateArgument(FProperty* InProperty)
 	: Index(INDEX_NONE)
 	, Name(InProperty->GetFName())
@@ -1678,52 +1685,68 @@ FString FRigVMTemplate::GetKeywords() const
 bool FRigVMTemplate::AddTypeForArgument(const FName& InArgumentName, TRigVMTypeIndex InTypeIndex)
 {
 	InvalidateHash();
-	
-	if(OnNewArgumentType().IsBound())
+
+	TArray<FRigVMTemplateTypeMap> TypesArray;
+	if (OnGetPermutationsFromArgumentType().IsBound())
+	{
+		TypesArray = OnGetPermutationsFromArgumentType().Execute(this, InArgumentName, InTypeIndex);
+	}
+	else if(OnNewArgumentType().IsBound())
 	{
 		FRigVMTemplateTypeMap Types = OnNewArgumentType().Execute(this, InArgumentName, InTypeIndex);
-		if(Types.Num() == Arguments.Num())
+		TypesArray = {Types};
+	}
+
+	if (!TypesArray.IsEmpty())
+	{
+		for (FRigVMTemplateTypeMap& Types : TypesArray)
 		{
-			const FRigVMRegistry& Registry = FRigVMRegistry::Get();
-			for (TPair<FName, TRigVMTypeIndex>& ArgumentAndType : Types)
+			if(Types.Num() == Arguments.Num())
 			{
-				// similar to FRigVMTemplateArgument::EnsureValidExecuteType
-				Registry.ConvertExecuteContextToBaseType(ArgumentAndType.Value);
-			}
+				const FRigVMRegistry& Registry = FRigVMRegistry::Get();
+				for (TPair<FName, TRigVMTypeIndex>& ArgumentAndType : Types)
+				{
+					// similar to FRigVMTemplateArgument::EnsureValidExecuteType
+					Registry.ConvertExecuteContextToBaseType(ArgumentAndType.Value);
+				}
 			
-			for(FRigVMTemplateArgument& Argument : Arguments)
-			{
-				const TRigVMTypeIndex* TypeIndex = Types.Find(Argument.Name);
-				if(TypeIndex == nullptr)
+				for(FRigVMTemplateArgument& Argument : Arguments)
 				{
-					return false;
+					const TRigVMTypeIndex* TypeIndex = Types.Find(Argument.Name);
+					if(TypeIndex == nullptr)
+					{
+						return false;
+					}
+					if(*TypeIndex == INDEX_NONE)
+					{
+						return false;
+					}
 				}
-				if(*TypeIndex == INDEX_NONE)
+
+				// Find if these types were already registered
+				FRigVMTemplateTypeMap TestTypes = Types;
+				if (ContainsPermutation(TestTypes))
 				{
 					return false;
 				}
-			}
+			
+				for(FRigVMTemplateArgument& Argument : Arguments)
+				{
+					const TRigVMTypeIndex TypeIndex = Types.FindChecked(Argument.Name);
+					Argument.TypeIndices.Add(TypeIndex);
+					Argument.TypeToPermutations.FindOrAdd(TypeIndex).Add(Permutations.Num());
+				}
 
-			// Find if these types were already registered
-			FRigVMTemplateTypeMap TestTypes = Types;
-			if (ContainsPermutation(TestTypes))
-			{
-				return false;
+				Permutations.Add(INDEX_NONE);
+
+				UpdateTypesHashToPermutation(Permutations.Num()-1);
 			}
-			
-			for(FRigVMTemplateArgument& Argument : Arguments)
+			else
 			{
-				const TRigVMTypeIndex TypeIndex = Types.FindChecked(Argument.Name);
-				Argument.TypeIndices.Add(TypeIndex);
-				Argument.TypeToPermutations.FindOrAdd(TypeIndex).Add(Permutations.Num());
+				return false;
 			}
-
-			Permutations.Add(INDEX_NONE);
-
-			UpdateTypesHashToPermutation(Permutations.Num()-1);
-			
-			return true;
 		}
+		return true;
 	}
 	return false;
 }
diff --git a/Engine/Plugins/Runtime/RigVM/Source/RigVM/Public/RigVMCore/RigVMDispatchFactory.h b/Engine/Plugins/Runtime/RigVM/Source/RigVM/Public/RigVMCore/RigVMDispatchFactory.h
--- a/Engine/Plugins/Runtime/RigVM/Source/RigVM/Public/RigVMCore/RigVMDispatchFactory.h
+++ b/Engine/Plugins/Runtime/RigVM/Source/RigVM/Public/RigVMCore/RigVMDispatchFactory.h
@@ -145,10 +145,22 @@ public:
 	// returns the execute arguments of the template
 	RIGVM_API TArray<FRigVMExecuteArgument> GetExecuteArguments(const FRigVMDispatchContext& InContext) const;
 
-	// returns the delegate to react to new types being added to an argument.
-	// this happens if types are being loaded later after this factory has already been deployed
+	// this function is deprecated, please use GetPermutationsFromArgumentType
+	// returns the new permutation argument types after a new type is defined for one argument
+	// this happens if types are being loaded later after this factory has already been deployed (like UUserDefinedStruct)
 	virtual FRigVMTemplateTypeMap OnNewArgumentType(const FName& InArgumentName, TRigVMTypeIndex InTypeIndex) const { return FRigVMTemplateTypeMap(); }
 
+	// returns the new permutations argument types after a new type is defined for one argument
+	virtual TArray<FRigVMTemplateTypeMap> GetPermutationsFromArgumentType(const FName& InArgumentName, const TRigVMTypeIndex& InTypeIndex) const
+	{
+		FRigVMTemplateTypeMap Permutation = OnNewArgumentType(InArgumentName, InTypeIndex);
+		if (!Permutation.IsEmpty())
+		{
+			return {Permutation};
+		}
+		return {};
+	}
+
 	// returns the upgrade info to use for this factory
 	virtual FRigVMStructUpgradeInfo GetUpgradeInfo(const FRigVMTemplateTypeMap& InTypes, const FRigVMDispatchContext& InContext) const { return FRigVMStructUpgradeInfo(); }
 
@@ -169,6 +181,9 @@ public:
 
 protected:
 
+	// for each type defined in the primary argument, this function will call GetPermutationsFromArgumentType to construct an array of arguments with the appropiate permutations
+	RIGVM_API TArray<FRigVMTemplateArgument> BuildArgumentListFromPrimaryArgument(const TArray<FRigVMTemplateArgument>& InArguments, const FName& InPrimaryArgumentName) const;
+
 	// returns the name of the permutation for a given set of types
 	RIGVM_API FString GetPermutationNameImpl(const FRigVMTemplateTypeMap& InTypes) const;
 
diff --git a/Engine/Plugins/Runtime/RigVM/Source/RigVM/Public/RigVMCore/RigVMTemplate.h b/Engine/Plugins/Runtime/RigVM/Source/RigVM/Public/RigVMCore/RigVMTemplate.h
--- a/Engine/Plugins/Runtime/RigVM/Source/RigVM/Public/RigVMCore/RigVMTemplate.h
+++ b/Engine/Plugins/Runtime/RigVM/Source/RigVM/Public/RigVMCore/RigVMTemplate.h
@@ -34,7 +34,9 @@ struct FRigVMUserDefinedTypeResolver;
 
 typedef TMap<FName, TRigVMTypeIndex> FRigVMTemplateTypeMap;
 
+// FRigVMTemplate_NewArgumentTypeDelegate is deprecated, use FRigVMTemplate_GetPermutationsFromArgumentTypeDelegate
 DECLARE_DELEGATE_RetVal_ThreeParams(FRigVMTemplateTypeMap, FRigVMTemplate_NewArgumentTypeDelegate, const FRigVMTemplate* /* InTemplate */, const FName& /* InArgumentName */, TRigVMTypeIndex /* InTypeIndexToAdd */);
+DECLARE_DELEGATE_RetVal_ThreeParams(TArray<FRigVMTemplateTypeMap>, FRigVMTemplate_GetPermutationsFromArgumentTypeDelegate, const FRigVMTemplate* /* InTemplate */, const FName& /* InArgumentName */, TRigVMTypeIndex /* InTypeIndexToAdd */);
 DECLARE_DELEGATE_RetVal_TwoParams(FRigVMFunctionPtr, FRigVMTemplate_RequestDispatchFunctionDelegate, const FRigVMTemplate* /* InTemplate */,  const FRigVMTemplateTypeMap& /* InTypes */);
 DECLARE_DELEGATE_RetVal_TwoParams(TArray<FRigVMFunction>, FRigVMTemplate_RequestDispatchPredicatesDelegate, const FRigVMTemplate* /* InTemplate */,  const FRigVMTemplateTypeMap& /* InTypes */);
 DECLARE_DELEGATE_RetVal(FRigVMDispatchFactory*, FRigVMTemplate_GetDispatchFactoryDelegate);
@@ -42,6 +44,7 @@ DECLARE_DELEGATE_RetVal(FRigVMDispatchFactory*, FRigVMTemplate_GetDispatchFactor
 struct RIGVM_API FRigVMTemplateDelegates
 {
 	FRigVMTemplate_NewArgumentTypeDelegate NewArgumentTypeDelegate;
+	FRigVMTemplate_GetPermutationsFromArgumentTypeDelegate GetPermutationsFromArgumentTypeDelegate;
 	FRigVMTemplate_GetDispatchFactoryDelegate GetDispatchFactoryDelegate;
 	FRigVMTemplate_RequestDispatchFunctionDelegate RequestDispatchFunctionDelegate;
 	FRigVMTemplate_RequestDispatchPredicatesDelegate RequestDispatchPredicatesDelegate;
@@ -198,6 +201,7 @@ struct RIGVM_API FRigVMTemplateArgument
 	// default constructor
 	FRigVMTemplateArgument();
 
+	FRigVMTemplateArgument(const FName& InName, ERigVMPinDirection InDirection);
 	FRigVMTemplateArgument(const FName& InName, ERigVMPinDirection InDirection, TRigVMTypeIndex InType);
 	FRigVMTemplateArgument(const FName& InName, ERigVMPinDirection InDirection, const TArray<TRigVMTypeIndex>& InTypeIndices);
 	FRigVMTemplateArgument(const FName& InName, ERigVMPinDirection InDirection, const TArray<ETypeCategory>& InTypeCategories, const FTypeFilter& InTypeFilter = {});
@@ -253,6 +257,7 @@ protected:
 	void UpdateTypeToPermutations();
 
 	friend struct FRigVMTemplate;
+	friend struct FRigVMDispatchFactory;
 	friend class URigVMController;
 	friend struct FRigVMRegistry;
 	friend struct FRigVMStructUpgradeInfo;
@@ -406,8 +411,12 @@ public:
 	void HandleTypeRemoval(TRigVMTypeIndex InTypeIndex);
 	
 	// Returns the delegate to be able to react to type changes dynamically
+	// This delegate is deprecated
 	FRigVMTemplate_NewArgumentTypeDelegate& OnNewArgumentType() { return Delegates.NewArgumentTypeDelegate; }
 
+	// Returns the delegate to be able to react to type changes dynamically
+	FRigVMTemplate_GetPermutationsFromArgumentTypeDelegate& OnGetPermutationsFromArgumentType() { return Delegates.GetPermutationsFromArgumentTypeDelegate; }
+
 	// Returns the factory this template was created by
 	const FRigVMDispatchFactory* GetDispatchFactory() const
 	{
