RigVM: Remove unnecesary includes
#rnx

diff --git a/Engine/Plugins/Runtime/RigVM/Source/RigVM/Private/RigVMFunctions/Math/RigVMFunction_MathVector.cpp b/Engine/Plugins/Runtime/RigVM/Source/RigVM/Private/RigVMFunctions/Math/RigVMFunction_MathVector.cpp
--- a/Engine/Plugins/Runtime/RigVM/Source/RigVM/Private/RigVMFunctions/Math/RigVMFunction_MathVector.cpp
+++ b/Engine/Plugins/Runtime/RigVM/Source/RigVM/Private/RigVMFunctions/Math/RigVMFunction_MathVector.cpp
@@ -3,8 +3,6 @@
 #include "RigVMFunctions/Math/RigVMFunction_MathVector.h"
 #include "RigVMFunctions/Math/RigVMFunction_MathTransform.h"
 #include "RigVMFunctions/RigVMDispatch_Core.h"
-#include "Engine/Engine.h"
-#include "Materials/MaterialInstanceDynamic.h"
 
 #include UE_INLINE_GENERATED_CPP_BY_NAME(RigVMFunction_MathVector)
 
