RigVM: FIx + symbol in the Details panel of a Function to add inputs/outputs is cut off
#rnx

diff --git a/Engine/Plugins/Runtime/RigVM/Source/RigVMEditor/Private/Editor/RigVMGraphDetailCustomization.cpp b/Engine/Plugins/Runtime/RigVM/Source/RigVMEditor/Private/Editor/RigVMGraphDetailCustomization.cpp
--- a/Engine/Plugins/Runtime/RigVM/Source/RigVMEditor/Private/Editor/RigVMGraphDetailCustomization.cpp
+++ b/Engine/Plugins/Runtime/RigVM/Source/RigVMEditor/Private/Editor/RigVMGraphDetailCustomization.cpp
@@ -797,7 +797,7 @@ void FRigVMGraphDetailCustomization::CustomizeDetails(IDetailLayoutBuilder& Deta
 		[
 			SNew(SButton)
 			.ButtonStyle(FAppStyle::Get(), "SimpleButton")
-			.ContentPadding(FMargin(1, 0))
+			.ContentPadding(FMargin(10.f, 0))
 			.OnClicked(this, &FRigVMGraphDetailCustomization::OnAddNewInputClicked)
 			.Visibility(this, &FRigVMGraphDetailCustomization::GetAddNewInputOutputVisibility)
 			.HAlign(HAlign_Right)
@@ -831,7 +831,7 @@ void FRigVMGraphDetailCustomization::CustomizeDetails(IDetailLayoutBuilder& Deta
 		[
 			SNew(SButton)
 			.ButtonStyle(FAppStyle::Get(), "SimpleButton")
-			.ContentPadding(FMargin(1, 0))
+			.ContentPadding(FMargin(10.f, 0))
 			.OnClicked(this, &FRigVMGraphDetailCustomization::OnAddNewOutputClicked)
 			.Visibility(this, &FRigVMGraphDetailCustomization::GetAddNewInputOutputVisibility)
 			.HAlign(HAlign_Right)
