RigVM: Support for FText

diff --git a/Engine/Plugins/Runtime/RigVM/Source/RigVM/Private/RigVMCore/RigVMExternalVariable.cpp b/Engine/Plugins/Runtime/RigVM/Source/RigVM/Private/RigVMCore/RigVMExternalVariable.cpp
--- a/Engine/Plugins/Runtime/RigVM/Source/RigVM/Private/RigVMCore/RigVMExternalVariable.cpp
+++ b/Engine/Plugins/Runtime/RigVM/Source/RigVM/Private/RigVMCore/RigVMExternalVariable.cpp
@@ -102,6 +102,11 @@ FRigVMExternalVariable RigVMTypeUtils::ExternalVariableFromPinType(const FName&
 		ExternalVariable.TypeName = FStringTypeName;
 		ExternalVariable.Size = sizeof(FString);
 	}
+	else if (InPinType.PinCategory == UEdGraphSchema_K2::PC_Text)
+	{
+		ExternalVariable.TypeName = FTextTypeName;
+		ExternalVariable.Size = sizeof(FText);
+	}
 	else if (InPinType.PinCategory == UEdGraphSchema_K2::PC_Struct)
 	{
 		if (UScriptStruct* Struct = Cast<UScriptStruct>(InPinType.PinSubCategoryObject))
@@ -111,6 +116,15 @@ FRigVMExternalVariable RigVMTypeUtils::ExternalVariableFromPinType(const FName&
 			ExternalVariable.Size = Struct->GetStructureSize();
 		}
 	}
+	else if (InPinType.PinCategory == UEdGraphSchema_K2::PC_Object)
+	{
+		if (UClass* Class = Cast<UClass>(InPinType.PinSubCategoryObject))
+		{
+			ExternalVariable.TypeName = *(Class->GetPrefixCPP() + Class->GetName());
+			ExternalVariable.TypeObject = Class;
+			ExternalVariable.Size = Class->GetStructureSize();
+		}
+	}
 
 	return ExternalVariable;
 }
@@ -164,6 +178,11 @@ FRigVMExternalVariable RigVMTypeUtils::ExternalVariableFromCPPTypePath(const FNa
 		Variable.TypeName = *CPPTypePath;
 		Variable.Size = sizeof(FName);
 	}
+	else if (CPPTypePath == FTextType)
+	{
+		Variable.TypeName = *CPPTypePath;
+		Variable.Size = sizeof(FText);
+	}
 	else if(UScriptStruct* ScriptStruct = RigVMTypeUtils::FindObjectFromCPPTypeObjectPath<UScriptStruct>(CPPTypePath))
 	{
 		Variable.TypeName = *RigVMTypeUtils::GetUniqueStructTypeName(ScriptStruct);
@@ -233,6 +252,11 @@ FRigVMExternalVariable RigVMTypeUtils::ExternalVariableFromCPPType(const FName&
 		Variable.TypeName = *CPPType;
 		Variable.Size = sizeof(FName);
 	}
+	else if (CPPType == FTextType)
+	{
+		Variable.TypeName = *CPPType;
+		Variable.Size = sizeof(FText);
+	}
 	else if(UScriptStruct* ScriptStruct = RigVMTypeUtils::FindObjectFromCPPTypeObjectPath<UScriptStruct>(CPPType))
 	{
 		Variable.TypeName = *RigVMTypeUtils::GetUniqueStructTypeName(ScriptStruct);
diff --git a/Engine/Plugins/Runtime/RigVM/Source/RigVM/Public/RigVMCore/RigVMExternalVariable.h b/Engine/Plugins/Runtime/RigVM/Source/RigVM/Public/RigVMCore/RigVMExternalVariable.h
--- a/Engine/Plugins/Runtime/RigVM/Source/RigVM/Public/RigVMCore/RigVMExternalVariable.h
+++ b/Engine/Plugins/Runtime/RigVM/Source/RigVM/Public/RigVMCore/RigVMExternalVariable.h
@@ -10,6 +10,7 @@
 #include "RigVMTypeIndex.h"
 #include "RigVMTypeUtils.h"
 #include "UObject/UnrealType.h"
+#include "UObject/TextProperty.h"
 #if WITH_EDITOR
 #include "EdGraphSchema_K2.h"
 #endif
@@ -70,6 +71,11 @@ struct RIGVM_API FRigVMExternalVariableDef
 			OutTypeName = TEXT("FName");
 			OutTypeObject = nullptr;
 		}
+		else if (CastField<FTextProperty>(InProperty))
+		{
+			OutTypeName = TEXT("FText");
+			OutTypeObject = nullptr;
+		}
 		else if (const FEnumProperty* EnumProperty = CastField<FEnumProperty>(InProperty))
 		{
 			OutTypeName = EnumProperty->GetEnum()->GetFName();
@@ -126,7 +132,10 @@ struct RIGVM_API FRigVMExternalVariableDef
 		}
 		else
 		{
-			checkNoEntry();
+			// this can happen due to unsupported property types
+			// within data assets or user defined structs
+			OutTypeName = NAME_None;
+			OutTypeObject = nullptr;
 		}
 	}
 
diff --git a/Engine/Plugins/Runtime/RigVM/Source/RigVM/Public/RigVMTypeUtils.h b/Engine/Plugins/Runtime/RigVM/Source/RigVM/Public/RigVMTypeUtils.h
--- a/Engine/Plugins/Runtime/RigVM/Source/RigVM/Public/RigVMTypeUtils.h
+++ b/Engine/Plugins/Runtime/RigVM/Source/RigVM/Public/RigVMTypeUtils.h
@@ -61,6 +61,7 @@ namespace RigVMTypeUtils
 	const FString UInt8Type = TEXT("uint8");
 	const FString FNameType = TEXT("FName");
 	const FString FStringType = TEXT("FString");
+	const FString FTextType = TEXT("FText");
 	const FString BoolArrayType = TEXT("TArray<bool>");
 	const FString FloatArrayType = TEXT("TArray<float>");
 	const FString DoubleArrayType = TEXT("TArray<double>");
@@ -69,6 +70,7 @@ namespace RigVMTypeUtils
 	const FString UInt8ArrayType = TEXT("TArray<uint8>");
 	const FString FNameArrayType = TEXT("TArray<FName>");
 	const FString FStringArrayType = TEXT("TArray<FString>");
+	const FString FTextArrayType = TEXT("TArray<FText>");
 
 	const FName BoolTypeName = *BoolType;
 	const FName FloatTypeName = *FloatType;
@@ -78,6 +80,7 @@ namespace RigVMTypeUtils
 	const FName UInt8TypeName = *UInt8Type;
 	const FName FNameTypeName = *FNameType;
 	const FName FStringTypeName = *FStringType;
+	const FName FTextTypeName = *FTextType;
 	const FName BoolArrayTypeName = *BoolArrayType;
 	const FName FloatArrayTypeName = *FloatArrayType;
 	const FName DoubleArrayTypeName = *DoubleArrayType;
