Control Rig: Emit compilation error when node has wildcard pins
#rnx

diff --git a/Engine/Plugins/Runtime/RigVM/Source/RigVMDeveloper/Private/RigVMCompiler/RigVMCompiler.cpp b/Engine/Plugins/Runtime/RigVM/Source/RigVMDeveloper/Private/RigVMCompiler/RigVMCompiler.cpp
--- a/Engine/Plugins/Runtime/RigVM/Source/RigVMDeveloper/Private/RigVMCompiler/RigVMCompiler.cpp
+++ b/Engine/Plugins/Runtime/RigVM/Source/RigVMDeveloper/Private/RigVMCompiler/RigVMCompiler.cpp
@@ -850,6 +850,28 @@ bool URigVMCompiler::Compile(const FRigVMCompileSettings& InSettings, TArray<URi
 					}
 				}
 			}
+			
+			if (URigVMTemplateNode* TemplateNode = Cast<URigVMTemplateNode>(ModelNode))
+			{
+				// Make sure there are no wildcard pins left
+				if (TemplateNode->HasWildCardPin())
+				{
+					static const FString UnresolvedUnitNodeMessage = TEXT("Node @@ contains wildcard pins. Resolve the node to continue.");
+					Settings.ASTSettings.Report(EMessageSeverity::Error, ModelNode, UnresolvedUnitNodeMessage);
+					bEncounteredGraphError = true;
+				}
+
+				if (TemplateNode->IsA<URigVMUnitNode>() || TemplateNode->IsA<URigVMDispatchNode>())
+				{
+					// Make sure the node has a resolved function
+					if (TemplateNode->GetResolvedFunction() == nullptr)
+					{
+						static const FString UnresolvedUnitNodeMessage = TEXT("Node @@ could not be resolved.");
+						Settings.ASTSettings.Report(EMessageSeverity::Error, ModelNode, UnresolvedUnitNodeMessage);
+						bEncounteredGraphError = true;
+					}
+				}
+			}
 
 			auto ReportIncompatibleExecuteContextString = [&] (const FString InExecuteContextName)
 			{
