RigVM: Clear compilation data when an error is found during compilation
#rnx

diff --git a/Engine/Plugins/Runtime/RigVM/Source/RigVMDeveloper/Private/RigVMCompiler/RigVMCompiler.cpp b/Engine/Plugins/Runtime/RigVM/Source/RigVMDeveloper/Private/RigVMCompiler/RigVMCompiler.cpp
--- a/Engine/Plugins/Runtime/RigVM/Source/RigVMDeveloper/Private/RigVMCompiler/RigVMCompiler.cpp
+++ b/Engine/Plugins/Runtime/RigVM/Source/RigVMDeveloper/Private/RigVMCompiler/RigVMCompiler.cpp
@@ -1175,7 +1175,11 @@ bool URigVMCompiler::Compile(const FRigVMCompileSettings& InSettings, TArray<URi
 	WorkData.ExprComplete.Reset();
 	for (FRigVMExprAST* RootExpr : *WorkData.AST)
 	{
-		TraverseExpression(RootExpr, WorkData);
+		if (!TraverseExpression(RootExpr, WorkData))
+		{
+			WorkData.Clear();
+			return false;
+		}
 	}
 
 	if(WorkData.WatchedPins.Num() > 0)
@@ -1242,7 +1246,11 @@ bool URigVMCompiler::Compile(const FRigVMCompileSettings& InSettings, TArray<URi
 	WorkData.ExprComplete.Reset();
 	for (FRigVMExprAST* RootExpr : *WorkData.AST)
 	{
-		TraverseExpression(RootExpr, WorkData);
+		if (!TraverseExpression(RootExpr, WorkData))
+		{
+			WorkData.Clear();
+			return false;
+		}
 	}
 
 	if (!CurrentCompilationFunction)
@@ -1493,16 +1501,16 @@ bool URigVMCompiler::CompileFunction(const FRigVMCompileSettings& InSettings, co
 	return bSuccess;
 }
 
-void URigVMCompiler::TraverseExpression(const FRigVMExprAST* InExpr, FRigVMCompilerWorkData& WorkData)
+bool URigVMCompiler::TraverseExpression(const FRigVMExprAST* InExpr, FRigVMCompilerWorkData& WorkData)
 {
 	if (WorkData.ExprToSkip.Contains(InExpr))
 	{
-		return;
+		return true;
 	}
 
 	if (WorkData.ExprComplete.Contains(InExpr))
 	{
-		return;
+		return true;
 	}
 	WorkData.ExprComplete.Add(InExpr, true);
 
@@ -1510,70 +1518,57 @@ void URigVMCompiler::TraverseExpression(const FRigVMExprAST* InExpr, FRigVMCompi
 	{
 		case FRigVMExprAST::EType::Block:
 		{
-			TraverseBlock(InExpr->To<FRigVMBlockExprAST>(), WorkData);
-			break;
+			return TraverseBlock(InExpr->To<FRigVMBlockExprAST>(), WorkData);
 		}
 		case FRigVMExprAST::EType::Entry:
 		{
-			TraverseEntry(InExpr->To<FRigVMEntryExprAST>(), WorkData);
-			break;
+			return TraverseEntry(InExpr->To<FRigVMEntryExprAST>(), WorkData);
 		}
 		case FRigVMExprAST::EType::CallExtern:
 		{
 			const FRigVMCallExternExprAST* CallExternExpr = InExpr->To<FRigVMCallExternExprAST>();
-			TraverseCallExtern(CallExternExpr, WorkData);
-			break;
+			return TraverseCallExtern(CallExternExpr, WorkData);
 		}
 		case FRigVMExprAST::EType::InlineFunction:
 		{
 			const FRigVMInlineFunctionExprAST* InlineExpr = InExpr->To<FRigVMInlineFunctionExprAST>();
-			TraverseInlineFunction(InlineExpr, WorkData);
-			break;
+			return TraverseInlineFunction(InlineExpr, WorkData);
 		}
 		case FRigVMExprAST::EType::NoOp:
 		{
-			TraverseNoOp(InExpr->To<FRigVMNoOpExprAST>(), WorkData);
-			break;
+			return TraverseNoOp(InExpr->To<FRigVMNoOpExprAST>(), WorkData);
 		}
 		case FRigVMExprAST::EType::Var:
 		{
-			TraverseVar(InExpr->To<FRigVMVarExprAST>(), WorkData);
-			break;
+			return TraverseVar(InExpr->To<FRigVMVarExprAST>(), WorkData);
 		}
 		case FRigVMExprAST::EType::Literal:
 		{
-			TraverseLiteral(InExpr->To<FRigVMLiteralExprAST>(), WorkData);
-			break;
+			return TraverseLiteral(InExpr->To<FRigVMLiteralExprAST>(), WorkData);
 		}
 		case FRigVMExprAST::EType::ExternalVar:
 		{
-			TraverseExternalVar(InExpr->To<FRigVMExternalVarExprAST>(), WorkData);
-			break;
+			return TraverseExternalVar(InExpr->To<FRigVMExternalVarExprAST>(), WorkData);
 		}
 		case FRigVMExprAST::EType::Assign:
 		{
-			TraverseAssign(InExpr->To<FRigVMAssignExprAST>(), WorkData);
-			break;
+			return TraverseAssign(InExpr->To<FRigVMAssignExprAST>(), WorkData);
 		}
 		case FRigVMExprAST::EType::Copy:
 		{
-			TraverseCopy(InExpr->To<FRigVMCopyExprAST>(), WorkData);
-			break;
+			return TraverseCopy(InExpr->To<FRigVMCopyExprAST>(), WorkData);
 		}
 		case FRigVMExprAST::EType::CachedValue:
 		{
-			TraverseCachedValue(InExpr->To<FRigVMCachedValueExprAST>(), WorkData);
-			break;
+			return TraverseCachedValue(InExpr->To<FRigVMCachedValueExprAST>(), WorkData);
 		}
 		case FRigVMExprAST::EType::Exit:
 		{
-			TraverseExit(InExpr->To<FRigVMExitExprAST>(), WorkData);
-			break;
+			return TraverseExit(InExpr->To<FRigVMExitExprAST>(), WorkData);
 		}
 		case FRigVMExprAST::EType::InvokeEntry:
 		{
-			TraverseInvokeEntry(InExpr->To<FRigVMInvokeEntryExprAST>(), WorkData);
-			break;
+			return TraverseInvokeEntry(InExpr->To<FRigVMInvokeEntryExprAST>(), WorkData);
 		}
 		default:
 		{
@@ -1581,26 +1576,31 @@ void URigVMCompiler::TraverseExpression(const FRigVMExprAST* InExpr, FRigVMCompi
 			break;
 		}
 	}
+	return false;
 }
 
-void URigVMCompiler::TraverseChildren(const FRigVMExprAST* InExpr, FRigVMCompilerWorkData& WorkData)
+bool URigVMCompiler::TraverseChildren(const FRigVMExprAST* InExpr, FRigVMCompilerWorkData& WorkData)
 {
 	for (FRigVMExprAST* ChildExpr : *InExpr)
 	{
-		TraverseExpression(ChildExpr, WorkData);
+		if (!TraverseExpression(ChildExpr, WorkData))
+		{
+			return false;
+		}
 	}
+	return true;
 }
 
-void URigVMCompiler::TraverseBlock(const FRigVMBlockExprAST* InExpr, FRigVMCompilerWorkData& WorkData)
+bool URigVMCompiler::TraverseBlock(const FRigVMBlockExprAST* InExpr, FRigVMCompilerWorkData& WorkData)
 {
 	if (InExpr->IsObsolete())
 	{
-		return;
+		return true;
 	}
 
 	if (InExpr->NumChildren() == 0)
 	{
-		return;
+		return true;
 	}
 
 	// check if the block is under a lazy pin, in which case we need to set up a branch info
@@ -1709,11 +1709,17 @@ void URigVMCompiler::TraverseBlock(const FRigVMBlockExprAST* InExpr, FRigVMCompi
 		// during the evaluation of the lazy branch, even if some of the expressions
 		// have already been visited in other parts of the traversal. See RigVM.Compiler.IfFromSameNode unit test. 
 		TGuardValue<TMap<const FRigVMExprAST*, bool>> ExprCompletedGuard(WorkData.ExprComplete, {});
-		TraverseChildren(InExpr, WorkData);
+		if (!TraverseChildren(InExpr, WorkData))
+		{
+			return false;
+		}
 	}
 	else
 	{
-		TraverseChildren(InExpr, WorkData);
+		if (!TraverseChildren(InExpr, WorkData))
+		{
+			return false;
+		}
 	}
 
 	if(!BranchInfo.Label.IsNone())
@@ -1721,21 +1727,26 @@ void URigVMCompiler::TraverseBlock(const FRigVMBlockExprAST* InExpr, FRigVMCompi
 		BranchInfo.LastInstruction = WorkData.VM->GetByteCode().GetNumInstructions() - 1;
 		WorkData.BranchInfos.FindOrAdd(CallExternNode).Add(BranchInfo);
 	}
+
+	return true;
 }
 
-void URigVMCompiler::TraverseEntry(const FRigVMEntryExprAST* InExpr, FRigVMCompilerWorkData& WorkData)
+bool URigVMCompiler::TraverseEntry(const FRigVMEntryExprAST* InExpr, FRigVMCompilerWorkData& WorkData)
 {
 	if (URigVMUnitNode* UnitNode = Cast<URigVMUnitNode>(InExpr->GetNode()))
 	{
 		if(!ValidateNode(WorkData.Settings, UnitNode))
 		{
-			return;
+			return false;
 		}
 
 		if (WorkData.bSetupMemory)
 		{
 			TSharedPtr<FStructOnScope> DefaultStruct = UnitNode->ConstructStructInstance();
-			TraverseChildren(InExpr, WorkData);
+			if (!TraverseChildren(InExpr, WorkData))
+			{
+				return false;
+			}
 		}
 		else
 		{
@@ -1805,16 +1816,16 @@ void URigVMCompiler::TraverseEntry(const FRigVMEntryExprAST* InExpr, FRigVMCompi
 					const FRigVMOperand& Source = *SourcePtr;
 					const FRigVMOperand& Target = *TargetPtr;
 	
-					WorkData.VM->GetByteCode().AddCopyOp(WorkData.VM->GetCopyOpForOperands(Source, Target));						
+					WorkData.VM->GetByteCode().AddCopyOp(WorkData.VM->GetCopyOpForOperands(Source, Target));
 				}
 			}
 		}
 	}
 
-	TraverseChildren(InExpr, WorkData);
+	return TraverseChildren(InExpr, WorkData);
 }
 
-int32 URigVMCompiler::TraverseCallExtern(const FRigVMCallExternExprAST* InExpr, FRigVMCompilerWorkData& WorkData)
+bool URigVMCompiler::TraverseCallExtern(const FRigVMCallExternExprAST* InExpr, FRigVMCompilerWorkData& WorkData)
 {
 	URigVMNode* Node = InExpr->GetNode();
 	URigVMUnitNode* UnitNode = Cast<URigVMUnitNode>(Node);
@@ -1822,7 +1833,7 @@ int32 URigVMCompiler::TraverseCallExtern(const FRigVMCallExternExprAST* InExpr,
 	const FRigVMRegistry& Registry = FRigVMRegistry::Get();
 	if(!ValidateNode(WorkData.Settings, UnitNode, false) && !ValidateNode(WorkData.Settings, DispatchNode, false))
 	{
-		return INDEX_NONE;
+		return false;
 	}
 
 	auto CheckExecuteStruct = [this, &WorkData](URigVMNode* Subject, const UScriptStruct* ExecuteStruct) -> bool
@@ -1849,7 +1860,7 @@ int32 URigVMCompiler::TraverseCallExtern(const FRigVMCallExternExprAST* InExpr,
 		{
 			static const FString UnresolvedMessage = TEXT("Node @@ is unresolved.");
 			WorkData.Settings.Report(EMessageSeverity::Error, UnitNode, UnresolvedMessage);
-			return INDEX_NONE;
+			return false;
 		}
 
 		// check execute pins for compatibility
@@ -1859,7 +1870,7 @@ int32 URigVMCompiler::TraverseCallExtern(const FRigVMCallExternExprAST* InExpr,
 			{
 				if(!CheckExecuteStruct(UnitNode, StructProperty->Struct))
 				{
-					return INDEX_NONE;
+					return false;
 				}
 			}
 		}
@@ -1871,13 +1882,13 @@ int32 URigVMCompiler::TraverseCallExtern(const FRigVMCallExternExprAST* InExpr,
 		{
 			static const FString UnresolvedDispatchMessage = TEXT("Dispatch node @@ has no factory.");
 			WorkData.Settings.Report(EMessageSeverity::Error, DispatchNode, UnresolvedDispatchMessage);
-			return INDEX_NONE;
+			return false;
 		}
 
 		// check execute pins for compatibility
 		if(!CheckExecuteStruct(DispatchNode, DispatchNode->GetFactory()->GetExecuteContextStruct()))
 		{
-			return INDEX_NONE;
+			return false;
 		}
 	}
 
@@ -1886,7 +1897,10 @@ int32 URigVMCompiler::TraverseCallExtern(const FRigVMCallExternExprAST* InExpr,
 
 	if (WorkData.bSetupMemory)
 	{
-		TraverseChildren(InExpr, WorkData);
+		if (!TraverseChildren(InExpr, WorkData))
+		{
+			return false;
+		}
 	}
 	else
 	{
@@ -2015,7 +2029,10 @@ int32 URigVMCompiler::TraverseCallExtern(const FRigVMCallExternExprAST* InExpr,
 					}
 				}
 			}
-			TraverseExpression(ChildExpr, WorkData);
+			if (!TraverseExpression(ChildExpr, WorkData))
+			{
+				return false;
+			}
 		}
 
 		if(!LazyChildExprs.IsEmpty())
@@ -2027,7 +2044,10 @@ int32 URigVMCompiler::TraverseCallExtern(const FRigVMCallExternExprAST* InExpr,
 			// traverse the lazy children 
 			for (const FRigVMExprAST* ChildExpr : LazyChildExprs)
 			{
-				TraverseExpression(ChildExpr, WorkData);
+				if (!TraverseExpression(ChildExpr, WorkData))
+				{
+					return false;
+				}
 			}
 
 			// update the operator with the target instruction 
@@ -2066,7 +2086,7 @@ int32 URigVMCompiler::TraverseCallExtern(const FRigVMCallExternExprAST* InExpr,
 
 		if (Operands.Num() > 65535)
 		{
-			return INDEX_NONE;
+			return false;
 		}
 
 		// setup the instruction
@@ -2203,7 +2223,10 @@ int32 URigVMCompiler::TraverseCallExtern(const FRigVMCallExternExprAST* InExpr,
 				const FRigVMVarExprAST* BlockExpr = InExpr->FindVarWithPinName(BlockName);
 				check(BlockExpr);
 				WorkData.ExprToSkip.Remove(BlockExpr);
-				TraverseExpression(BlockExpr, WorkData);
+				if (!TraverseExpression(BlockExpr, WorkData))
+				{
+					return false;
+				}
 
 				// end the block if necessary
 				if(Node->IsControlFlowBlockSliced(BlockName))
@@ -2223,16 +2246,16 @@ int32 URigVMCompiler::TraverseCallExtern(const FRigVMCallExternExprAST* InExpr,
 		}
 	}
 
-	return CallExternInstructionIndex;
+	return true;
 }
 
-int32 URigVMCompiler::TraverseInlineFunction(const FRigVMInlineFunctionExprAST* InExpr, FRigVMCompilerWorkData& WorkData)
+bool URigVMCompiler::TraverseInlineFunction(const FRigVMInlineFunctionExprAST* InExpr, FRigVMCompilerWorkData& WorkData)
 {
 	URigVMNode* Node = InExpr->GetNode();
 	URigVMFunctionReferenceNode* FunctionReferenceNode = Cast<URigVMFunctionReferenceNode>(Node);
 	if(!ValidateNode(WorkData.Settings, FunctionReferenceNode, false))
 	{
-		return INDEX_NONE;
+		return false;
 	}
 	
 	int32 InstructionIndexStart = INDEX_NONE;
@@ -2242,14 +2265,17 @@ int32 URigVMCompiler::TraverseInlineFunction(const FRigVMInlineFunctionExprAST*
 	FString FunctionHash = FunctionReferenceNode->GetReferencedFunctionHeader().GetHash();
 	if (!CompiledFunctions.Contains(FunctionHash))
 	{
-		return INDEX_NONE;
+		return true;
 	}
 	const FRigVMFunctionCompilationData* FunctionCompilationData = CompiledFunctions.FindChecked(FunctionHash);
 	const FRigVMByteCode& FunctionByteCode = FunctionCompilationData->ByteCode;
 	
 	if (WorkData.bSetupMemory)
 	{
-		TraverseChildren(InExpr, WorkData);
+		if (!TraverseChildren(InExpr, WorkData))
+		{
+			return false;
+		}
 
 		// Add internal operands (not the ones represented by interface pins)
 		for (uint8 MemoryIndex=0; MemoryIndex< (uint8)ERigVMMemoryType::Invalid; ++MemoryIndex)
@@ -2366,7 +2392,10 @@ int32 URigVMCompiler::TraverseInlineFunction(const FRigVMInlineFunctionExprAST*
 			}
 		}
 
-		TraverseChildren(InExpr, WorkData);
+		if (!TraverseChildren(InExpr, WorkData))
+		{
+			return false;
+		}
 
 		// Inline the bytecode from the function
 		FRigVMByteCode& ByteCode = WorkData.VM->GetByteCode();
@@ -2398,7 +2427,7 @@ int32 URigVMCompiler::TraverseInlineFunction(const FRigVMInlineFunctionExprAST*
 		{
 			if (!Instructions.IsValidIndex(i))
 			{
-				return INDEX_NONE;
+				return false;
 			}
 			const FRigVMInstruction& Instruction = Instructions[i];
 			const FRigVMOperandArray OperandArray = ByteCode.GetOperandsForOp(Instruction);
@@ -2601,37 +2630,45 @@ int32 URigVMCompiler::TraverseInlineFunction(const FRigVMInlineFunctionExprAST*
 #endif
 	}
 
-	return InstructionIndexEnd;
+	return true;
 }
 
-void URigVMCompiler::TraverseNoOp(const FRigVMNoOpExprAST* InExpr, FRigVMCompilerWorkData& WorkData)
+bool URigVMCompiler::TraverseNoOp(const FRigVMNoOpExprAST* InExpr, FRigVMCompilerWorkData& WorkData)
 {
-	TraverseChildren(InExpr, WorkData);
+	return TraverseChildren(InExpr, WorkData);
 }
 
-void URigVMCompiler::TraverseVar(const FRigVMVarExprAST* InExpr, FRigVMCompilerWorkData& WorkData)
+bool URigVMCompiler::TraverseVar(const FRigVMVarExprAST* InExpr, FRigVMCompilerWorkData& WorkData)
 {
-	TraverseChildren(InExpr, WorkData);
+	if (!TraverseChildren(InExpr, WorkData))
+	{
+		return false;
+	}
 
 	if (WorkData.bSetupMemory)
 	{
 		FindOrAddRegister(InExpr, WorkData);
 	}
+
+	return true;
 }
 
-void URigVMCompiler::TraverseLiteral(const FRigVMVarExprAST* InExpr, FRigVMCompilerWorkData& WorkData)
+bool URigVMCompiler::TraverseLiteral(const FRigVMVarExprAST* InExpr, FRigVMCompilerWorkData& WorkData)
 {
-	TraverseVar(InExpr, WorkData);
+	return TraverseVar(InExpr, WorkData);
 }
 
-void URigVMCompiler::TraverseExternalVar(const FRigVMExternalVarExprAST* InExpr, FRigVMCompilerWorkData& WorkData)
+bool URigVMCompiler::TraverseExternalVar(const FRigVMExternalVarExprAST* InExpr, FRigVMCompilerWorkData& WorkData)
 {
-	TraverseVar(InExpr, WorkData);
+	return TraverseVar(InExpr, WorkData);
 }
 
-void URigVMCompiler::TraverseAssign(const FRigVMAssignExprAST* InExpr, FRigVMCompilerWorkData& WorkData)
+bool URigVMCompiler::TraverseAssign(const FRigVMAssignExprAST* InExpr, FRigVMCompilerWorkData& WorkData)
 {
-	TraverseChildren(InExpr, WorkData);
+	if (!TraverseChildren(InExpr, WorkData))
+	{
+		return false;
+	}
 
 	ensure(InExpr->NumChildren() > 0);
 
@@ -2681,7 +2718,7 @@ void URigVMCompiler::TraverseAssign(const FRigVMAssignExprAST* InExpr, FRigVMCom
 		FRigVMOperand Target = WorkData.ExprToOperand.FindChecked(GetSourceVarExpr(TargetExpr));
 		if(Target == Source)
 		{
-			return;
+			return true;
 		}
 		
 		// if this is a copy - we should check if operands need offsets
@@ -2755,38 +2792,41 @@ void URigVMCompiler::TraverseAssign(const FRigVMAssignExprAST* InExpr, FRigVMCom
 			AddCopyOperator(CopyOp, InExpr, SourceExpr, TargetExpr, WorkData);
 		}
 	}
+
+	return true;
 }
 
-void URigVMCompiler::TraverseCopy(const FRigVMCopyExprAST* InExpr, FRigVMCompilerWorkData& WorkData)
+bool URigVMCompiler::TraverseCopy(const FRigVMCopyExprAST* InExpr, FRigVMCompilerWorkData& WorkData)
 {
-	TraverseAssign(InExpr->To<FRigVMAssignExprAST>(), WorkData);
+	return TraverseAssign(InExpr->To<FRigVMAssignExprAST>(), WorkData);
 }
 
-void URigVMCompiler::TraverseCachedValue(const FRigVMCachedValueExprAST* InExpr, FRigVMCompilerWorkData& WorkData)
+bool URigVMCompiler::TraverseCachedValue(const FRigVMCachedValueExprAST* InExpr, FRigVMCompilerWorkData& WorkData)
 {
-	TraverseChildren(InExpr, WorkData);
+	return TraverseChildren(InExpr, WorkData);
 }
 
-void URigVMCompiler::TraverseExit(const FRigVMExitExprAST* InExpr, FRigVMCompilerWorkData& WorkData)
+bool URigVMCompiler::TraverseExit(const FRigVMExitExprAST* InExpr, FRigVMCompilerWorkData& WorkData)
 {
 	ensure(InExpr->NumChildren() == 0);
 	if (!WorkData.bSetupMemory)
 	{
 		WorkData.VM->GetByteCode().AddExitOp();
 	}
+	return true;
 }
 
-void URigVMCompiler::TraverseInvokeEntry(const FRigVMInvokeEntryExprAST* InExpr, FRigVMCompilerWorkData& WorkData)
+bool URigVMCompiler::TraverseInvokeEntry(const FRigVMInvokeEntryExprAST* InExpr, FRigVMCompilerWorkData& WorkData)
 {
 	URigVMInvokeEntryNode* InvokeEntryNode = Cast<URigVMInvokeEntryNode>(InExpr->GetNode());
 	if(!ValidateNode(WorkData.Settings, InvokeEntryNode))
 	{
-		return;
+		return false;
 	}
 
 	if (WorkData.bSetupMemory)
 	{
-		return;
+		return true;
 	}
 	else
 	{
@@ -2799,6 +2839,7 @@ void URigVMCompiler::TraverseInvokeEntry(const FRigVMInvokeEntryExprAST* InExpr,
 			WorkData.VM->GetByteCode().SetSubject(InstructionIndex, Callstack.GetCallPath(), Callstack.GetStack());
 		}
 	}
+	return true;
 }
 
 void URigVMCompiler::AddCopyOperator(const FRigVMCopyOp& InOp, const FRigVMAssignExprAST* InAssignExpr,
diff --git a/Engine/Plugins/Runtime/RigVM/Source/RigVMDeveloper/Public/RigVMCompiler/RigVMCompiler.h b/Engine/Plugins/Runtime/RigVM/Source/RigVMDeveloper/Public/RigVMCompiler/RigVMCompiler.h
--- a/Engine/Plugins/Runtime/RigVM/Source/RigVMDeveloper/Public/RigVMCompiler/RigVMCompiler.h
+++ b/Engine/Plugins/Runtime/RigVM/Source/RigVMDeveloper/Public/RigVMCompiler/RigVMCompiler.h
@@ -220,6 +220,40 @@ public:
 	FRigVMReportDelegate OriginalReportDelegate;
 	void OverrideReportDelegate(bool& bEncounteredASTError, bool& bSurpressedASTError);
 	void RemoveOverrideReportDelegate();
+
+	void Clear()
+	{
+		if (VM)
+		{
+			if (Context)
+			{
+				VM->ClearMemory(*Context);
+				VM->Reset(*Context);
+			}
+			else
+			{
+				VM->Reset_Internal();
+			}
+		}
+		Graphs.Reset();
+		if (PinPathToOperand)
+		{
+			PinPathToOperand->Reset();
+		}
+		ExprToOperand.Reset();
+		ExprComplete.Reset();
+		ExprToSkip.Reset();
+		ProcessedLinks.Reset();
+		IntegerLiterals.Reset();
+		CachedProxiesWithSharedOperand.Reset();
+		ProxyTargets.Reset();
+		BranchInfos.Reset();
+		FunctionRegisterToOperand.Reset();
+		WatchedPins.Reset();
+		PropertyDescriptions.Reset();
+		PropertyPathDescriptions.Reset();
+		DeferredCopyOps.Reset();
+	}
 };
 
 DECLARE_DELEGATE_RetVal_OneParam(const FRigVMFunctionCompilationData*, FRigVMCompiler_GetFunctionCompilationData, const FRigVMGraphFunctionHeader& Header);
@@ -283,21 +317,21 @@ private:
 
 	static FString GetPinHashImpl(const URigVMPin* InPin, const FRigVMVarExprAST* InVarExpr, bool bIsDebugValue = false, const URigVMLibraryNode* FunctionCompiling = nullptr, const FRigVMASTProxy& InPinProxy = FRigVMASTProxy());
 
-	void TraverseExpression(const FRigVMExprAST* InExpr, FRigVMCompilerWorkData& WorkData);
-	void TraverseChildren(const FRigVMExprAST* InExpr, FRigVMCompilerWorkData& WorkData);
-	void TraverseBlock(const FRigVMBlockExprAST* InExpr, FRigVMCompilerWorkData& WorkData);
-	void TraverseEntry(const FRigVMEntryExprAST* InExpr, FRigVMCompilerWorkData& WorkData);
-	int32 TraverseCallExtern(const FRigVMCallExternExprAST* InExpr, FRigVMCompilerWorkData& WorkData);
-	int32 TraverseInlineFunction(const FRigVMInlineFunctionExprAST* InExpr, FRigVMCompilerWorkData& WorkData);
-	void TraverseNoOp(const FRigVMNoOpExprAST* InExpr, FRigVMCompilerWorkData& WorkData);
-	void TraverseVar(const FRigVMVarExprAST* InExpr, FRigVMCompilerWorkData& WorkData);
-	void TraverseLiteral(const FRigVMVarExprAST* InExpr, FRigVMCompilerWorkData& WorkData);
-	void TraverseExternalVar(const FRigVMExternalVarExprAST* InExpr, FRigVMCompilerWorkData& WorkData);
-	void TraverseAssign(const FRigVMAssignExprAST* InExpr, FRigVMCompilerWorkData& WorkData);
-	void TraverseCopy(const FRigVMCopyExprAST* InExpr, FRigVMCompilerWorkData& WorkData);
-	void TraverseCachedValue(const FRigVMCachedValueExprAST* InExpr, FRigVMCompilerWorkData& WorkData);
-	void TraverseExit(const FRigVMExitExprAST* InExpr, FRigVMCompilerWorkData& WorkData);
-	void TraverseInvokeEntry(const FRigVMInvokeEntryExprAST* InExpr, FRigVMCompilerWorkData& WorkData);
+	bool TraverseExpression(const FRigVMExprAST* InExpr, FRigVMCompilerWorkData& WorkData);
+	bool TraverseChildren(const FRigVMExprAST* InExpr, FRigVMCompilerWorkData& WorkData);
+	bool TraverseBlock(const FRigVMBlockExprAST* InExpr, FRigVMCompilerWorkData& WorkData);
+	bool TraverseEntry(const FRigVMEntryExprAST* InExpr, FRigVMCompilerWorkData& WorkData);
+	bool TraverseCallExtern(const FRigVMCallExternExprAST* InExpr, FRigVMCompilerWorkData& WorkData);
+	bool TraverseInlineFunction(const FRigVMInlineFunctionExprAST* InExpr, FRigVMCompilerWorkData& WorkData);
+	bool TraverseNoOp(const FRigVMNoOpExprAST* InExpr, FRigVMCompilerWorkData& WorkData);
+	bool TraverseVar(const FRigVMVarExprAST* InExpr, FRigVMCompilerWorkData& WorkData);
+	bool TraverseLiteral(const FRigVMVarExprAST* InExpr, FRigVMCompilerWorkData& WorkData);
+	bool TraverseExternalVar(const FRigVMExternalVarExprAST* InExpr, FRigVMCompilerWorkData& WorkData);
+	bool TraverseAssign(const FRigVMAssignExprAST* InExpr, FRigVMCompilerWorkData& WorkData);
+	bool TraverseCopy(const FRigVMCopyExprAST* InExpr, FRigVMCompilerWorkData& WorkData);
+	bool TraverseCachedValue(const FRigVMCachedValueExprAST* InExpr, FRigVMCompilerWorkData& WorkData);
+	bool TraverseExit(const FRigVMExitExprAST* InExpr, FRigVMCompilerWorkData& WorkData);
+	bool TraverseInvokeEntry(const FRigVMInvokeEntryExprAST* InExpr, FRigVMCompilerWorkData& WorkData);
 
 	void AddCopyOperator(
 		const FRigVMCopyOp& InOp,
