Control Rig: Fix crash when a template has no permutations
#rnx

diff --git a/Engine/Plugins/Runtime/RigVM/Source/RigVM/Private/RigVMCore/RigVMRegistry.cpp b/Engine/Plugins/Runtime/RigVM/Source/RigVM/Private/RigVMCore/RigVMRegistry.cpp
--- a/Engine/Plugins/Runtime/RigVM/Source/RigVM/Private/RigVMCore/RigVMRegistry.cpp
+++ b/Engine/Plugins/Runtime/RigVM/Source/RigVM/Private/RigVMCore/RigVMRegistry.cpp
@@ -1536,7 +1536,7 @@ const FRigVMTemplate* FRigVMRegistry::GetOrAddTemplateFromArguments(const FName&
 {
 	FScopeLock RegisterTemplateScopeLock(&RegisterTemplateMutex);
 	
-	FRigVMTemplate Template(InName, InArguments, INDEX_NONE);
+	FRigVMTemplate Template(InName, InArguments);
 
 	// avoid reentry in FindTemplate. try to find an existing
 	// template only if we are not yet in ::FindTemplate.
@@ -1556,10 +1556,10 @@ const FRigVMTemplate* FRigVMRegistry::AddTemplateFromArguments(const FName& InNa
 
 const FRigVMTemplate* FRigVMRegistry::AddTemplateFromArguments_NoLock(const FName& InName, const TArray<FRigVMTemplateArgument>& InArguments, const FRigVMTemplateDelegates& InDelegates)
 {
-	FRigVMTemplate Template(InName, InArguments, INDEX_NONE);
+	FRigVMTemplate Template(InName, InArguments);
 	
 	// we only support to ask for templates here which provide singleton types
-	int32 NumPermutations = 1;
+	int32 NumPermutations = 0;
 	for(const FRigVMTemplateArgument& Argument : InArguments)
 	{
 		if(!Argument.IsSingleton() && NumPermutations > 1)
diff --git a/Engine/Plugins/Runtime/RigVM/Source/RigVM/Private/RigVMCore/RigVMTemplate.cpp b/Engine/Plugins/Runtime/RigVM/Source/RigVM/Private/RigVMCore/RigVMTemplate.cpp
--- a/Engine/Plugins/Runtime/RigVM/Source/RigVM/Private/RigVMCore/RigVMTemplate.cpp
+++ b/Engine/Plugins/Runtime/RigVM/Source/RigVM/Private/RigVMCore/RigVMTemplate.cpp
@@ -502,13 +502,20 @@ FRigVMTemplate::FRigVMTemplate(UScriptStruct* InStruct, const FString& InTemplat
 	{
 		const FString NotationStr = FString::Printf(TEXT("%s(%s)"), *InTemplateName, *FString::Join(ArgumentNotations, TEXT(",")));
 		Notation = *NotationStr;
-		Permutations.Add(InFunctionIndex);
+		if (InFunctionIndex != INDEX_NONE)
+		{
+			Permutations.Add(InFunctionIndex);
+			for (const FRigVMTemplateArgument& Argument : Arguments)
+			{
+				check(Argument.TypeIndices.Num() == 1);
+			}
+		}
 
 		UpdateTypesHashToPermutation(Permutations.Num()-1);
 	}
 }
 
-FRigVMTemplate::FRigVMTemplate(const FName& InTemplateName, const TArray<FRigVMTemplateArgument>& InArguments, int32 InFunctionIndex)
+FRigVMTemplate::FRigVMTemplate(const FName& InTemplateName, const TArray<FRigVMTemplateArgument>& InArguments)
 	: Index(INDEX_NONE)
 	, Notation(NAME_None)
 	, Hash(UINT32_MAX)
@@ -533,7 +540,6 @@ FRigVMTemplate::FRigVMTemplate(const FName& InTemplateName, const TArray<FRigVMT
 	{
 		const FString NotationStr = FString::Printf(TEXT("%s(%s)"), *InTemplateName.ToString(), *FString::Join(ArgumentNotations, TEXT(",")));
 		Notation = *NotationStr;
-		Permutations.Add(InFunctionIndex);
 
 		UpdateTypesHashToPermutation(Permutations.Num()-1);
 	}
diff --git a/Engine/Plugins/Runtime/RigVM/Source/RigVM/Public/RigVMCore/RigVMTemplate.h b/Engine/Plugins/Runtime/RigVM/Source/RigVM/Public/RigVMCore/RigVMTemplate.h
--- a/Engine/Plugins/Runtime/RigVM/Source/RigVM/Public/RigVMCore/RigVMTemplate.h
+++ b/Engine/Plugins/Runtime/RigVM/Source/RigVM/Public/RigVMCore/RigVMTemplate.h
@@ -443,10 +443,10 @@ public:
 private:
 
 	// Constructor from a struct, a template name and a function index
-	FRigVMTemplate(UScriptStruct* InStruct, const FString& InTemplateName, int32 InFunctionIndex);
+	FRigVMTemplate(UScriptStruct* InStruct, const FString& InTemplateName, int32 InFunctionIndex = INDEX_NONE);
 
 	// Constructor from a template name, arguments and a function index
-	FRigVMTemplate(const FName& InTemplateName, const TArray<FRigVMTemplateArgument>& InArguments, int32 InFunctionIndex);
+	FRigVMTemplate(const FName& InTemplateName, const TArray<FRigVMTemplateArgument>& InArguments);
 
 	static FLinearColor GetColorFromMetadata(FString InMetadata);
 
