Fixed FRigVMStorageStruct cook determinism generating GUIDs from property name and cpptype

diff --git a/Engine/Plugins/Runtime/RigVM/Source/RigVM/Private/RigVMCore/RigVMMemoryStorageStruct.cpp b/Engine/Plugins/Runtime/RigVM/Source/RigVM/Private/RigVMCore/RigVMMemoryStorageStruct.cpp
--- a/Engine/Plugins/Runtime/RigVM/Source/RigVM/Private/RigVMCore/RigVMMemoryStorageStruct.cpp
+++ b/Engine/Plugins/Runtime/RigVM/Source/RigVM/Private/RigVMCore/RigVMMemoryStorageStruct.cpp
@@ -3,7 +3,8 @@
 #include "RigVMCore/RigVMMemoryStorageStruct.h"
 #include "RigVMTypeUtils.h"
 #include "RigVMModule.h"
-
+#include "Misc/Guid.h"
+#include "Hash/Blake3.h"
 
 namespace UE::RigVM::RigVMCore::Private
 {
@@ -567,6 +568,14 @@ void FRigVMMemoryStorageStruct::SetDefaultValues(const TArray<FRigVMPropertyDesc
 		{
 			Result = FPropertyBagPropertyDesc(RigVMDescriptor.Name, PropertyBagType, RigVMDescriptor.CPPTypeObject);
 		}
+
+		const FString Name = RigVMDescriptor.Name.ToString();
+
+		FBlake3 Builder;
+		Builder.Update(*Name, Name.Len());
+		Builder.Update(*RigVMDescriptor.CPPType, RigVMDescriptor.CPPType.Len());
+		FBlake3Hash Hash = Builder.Finalize();
+		Result.ID = FGuid::NewGuidFromHash(Hash);
 	}
 
 	return Result;
diff --git a/Engine/Source/Runtime/Core/Private/Misc/Guid.cpp b/Engine/Source/Runtime/Core/Private/Misc/Guid.cpp
--- a/Engine/Source/Runtime/Core/Private/Misc/Guid.cpp
+++ b/Engine/Source/Runtime/Core/Private/Misc/Guid.cpp
@@ -249,9 +249,15 @@ FGuid FGuid::NewDeterministicGuid(FStringView ObjectPath, uint64 Seed)
 	Builder.Update(&Seed, sizeof(Seed));
 	Builder.Update(Utf8ObjectPath.GetData(), Utf8ObjectPath.Len() * sizeof(UTF8CHAR));
 
+	FBlake3Hash Hash = Builder.Finalize();
+
+	return NewGuidFromHash(Hash);
+}
+
+FGuid FGuid::NewGuidFromHash(const FBlake3Hash& Hash)
+{
 	// We use the first 16 bytes of the BLAKE3 hash to create the guid, there is no specific reason why these were
 	// chosen, we could take any pattern or combination of bytes.
-	FBlake3Hash Hash = Builder.Finalize();
 	uint32* HashBytes = (uint32*)Hash.GetBytes();
 	uint32 A = uint32(HashBytes[0]);
 	uint32 B = uint32(HashBytes[1]);
diff --git a/Engine/Source/Runtime/Core/Public/Misc/Guid.h b/Engine/Source/Runtime/Core/Public/Misc/Guid.h
--- a/Engine/Source/Runtime/Core/Public/Misc/Guid.h
+++ b/Engine/Source/Runtime/Core/Public/Misc/Guid.h
@@ -20,6 +20,7 @@ class FOutputDevice;
 class FPointerTableBase;
 class FSHA1;
 class UObject;
+struct FBlake3Hash;
 template <typename CharType> class TStringBuilderBase;
 template <typename T> struct TCanBulkSerialize;
 template <typename T> struct TIsPODType;
@@ -365,6 +366,10 @@ public:
 	 * and will thus avoid cook indeterminism caused by FGuid::NewGuid. ObjectPath and Seed must be deterministic.
 	 */
 	static CORE_API FGuid NewDeterministicGuid(FStringView ObjectPath, uint64 Seed = 0);
+	/**
+	 * Create a guid from a calculated Blake3 Hash
+	 */
+	static CORE_API FGuid NewGuidFromHash(const FBlake3Hash& Hash);
 
 	/**
 	 * Returns a GUID which is a combinationof the two provided ones.
