RigVM: Fixed execution debugging to adapt to new property bag memory
#rnx

diff --git a/Engine/Plugins/Runtime/RigVM/Source/RigVM/Private/RigVMCore/RigVM.cpp b/Engine/Plugins/Runtime/RigVM/Source/RigVM/Private/RigVMCore/RigVM.cpp
--- a/Engine/Plugins/Runtime/RigVM/Source/RigVM/Private/RigVMCore/RigVM.cpp
+++ b/Engine/Plugins/Runtime/RigVM/Source/RigVM/Private/RigVMCore/RigVM.cpp
@@ -21,6 +21,14 @@
 
 #include UE_INLINE_GENERATED_CPP_BY_NAME(RigVM)
 
+#if UE_RIGVM_DEBUG_EXECUTION
+static TAutoConsoleVariable<int32> CVarControlRigDebugAllVMExecutions(
+	TEXT("ControlRig.DebugAllVMExecutions"),
+	0,
+	TEXT("If nonzero we allow to copy the execution of a VM execution."),
+	ECVF_Default);
+#endif
+
 void FRigVMParameter::Serialize(FArchive& Ar)
 {
 	Ar.UsingCustomVersion(FAnimObjectVersion::GUID);
@@ -1488,7 +1496,7 @@ bool URigVM::InitializeInstance(FRigVMExtendedExecuteContext& Context, TArrayVie
 					{
 						Property->CopyCompleteValue_InContainer(WorkMemory, WorkMemoryCDO);
 #if UE_RIGVM_DEBUG_EXECUTION
-						if (Context.GetPublicData<FRigVMExecuteContext>().bDebugExecution)
+						if (CVarControlRigDebugAllVMExecutions->GetBool() || Context.GetPublicData<FRigVMExecuteContext>().bDebugExecution)
 						{
 							FString DefaultValue;
 							const uint8* PropertyMemory = Property->ContainerPtrToValuePtr<uint8>(WorkMemory);
@@ -1636,7 +1644,7 @@ ERigVMExecuteResult URigVM::Execute(FRigVMExtendedExecuteContext& Context, TArra
 	StartProfiling(Context);
 	
 #if UE_RIGVM_DEBUG_EXECUTION
-	if (ContextPublicData.bDebugExecution)
+	if (CVarControlRigDebugAllVMExecutions->GetBool() || ContextPublicData.bDebugExecution)
 	{
 		ContextPublicData.InstanceOpCodeEnum = StaticEnum<ERigVMOpCode>();
 		TRigVMMemoryStorage* LiteralMemory = GetLiteralMemory();
@@ -1693,7 +1701,7 @@ ERigVMExecuteResult URigVM::ExecuteInstructions(FRigVMExtendedExecuteContext& Co
 		if (ShouldHaltAtInstruction(Context, Context.CurrentEntryName, ContextPublicData.InstructionIndex))
 		{
 #if UE_RIGVM_DEBUG_EXECUTION
-			if (ContextPublicData.bDebugExecution)
+			if (CVarControlRigDebugAllVMExecutions->GetBool() || ContextPublicData.bDebugExecution)
 			{
 				ContextPublicData.Log(EMessageSeverity::Info, ContextPublicData.DebugMemoryString);
 			}
@@ -1727,7 +1735,7 @@ ERigVMExecuteResult URigVM::ExecuteInstructions(FRigVMExtendedExecuteContext& Co
 
 #if WITH_EDITOR
 #if UE_RIGVM_DEBUG_EXECUTION
-		if (ContextPublicData.bDebugExecution)
+		if (CVarControlRigDebugAllVMExecutions->GetBool() || ContextPublicData.bDebugExecution)
 		{
 			if (Instruction.OpCode == ERigVMOpCode::Execute)
 			{
@@ -1966,7 +1974,7 @@ ERigVMExecuteResult URigVM::ExecuteInstructions(FRigVMExtendedExecuteContext& Co
 						Context.ExecutionHalted().Broadcast(INDEX_NONE, nullptr, Context.CurrentEntryName);
 					}
 #if UE_RIGVM_DEBUG_EXECUTION
-					if (ContextPublicData.bDebugExecution)
+					if (CVarControlRigDebugAllVMExecutions->GetBool() || ContextPublicData.bDebugExecution)
 					{
 						ContextPublicData.Log(EMessageSeverity::Info, ContextPublicData.DebugMemoryString);
 					}
@@ -2098,10 +2106,10 @@ ERigVMExecuteResult URigVM::ExecuteInstructions(FRigVMExtendedExecuteContext& Co
 		}
 
 #if UE_RIGVM_DEBUG_EXECUTION
-		if (ContextPublicData.bDebugExecution)
+		if (CVarControlRigDebugAllVMExecutions->GetBool() || ContextPublicData.bDebugExecution)
 		{
 			TArray<FString> CurrentWorkMemory;
-			TRigVMMemoryStorage* WorkMemory = GetWorkMemory();
+			TRigVMMemoryStorage* WorkMemory = GetWorkMemory(Context);
 			int32 LineIndex = 0;
 			for (int32 PropertyIndex=0; PropertyIndex<WorkMemory->Num(); ++PropertyIndex, ++LineIndex)
 			{
