RigVM: Invalidate cached memory after load of a RigVMHost
#rnx

diff --git a/Engine/Plugins/Runtime/RigVM/Source/RigVM/Private/RigVMHost.cpp b/Engine/Plugins/Runtime/RigVM/Source/RigVM/Private/RigVMHost.cpp
--- a/Engine/Plugins/Runtime/RigVM/Source/RigVM/Private/RigVMHost.cpp
+++ b/Engine/Plugins/Runtime/RigVM/Source/RigVM/Private/RigVMHost.cpp
@@ -122,6 +122,11 @@ void URigVMHost::Serialize(FArchive& Ar)
 			}
 		}
 	}
+
+	if (Ar.IsLoading())
+	{
+		RecreateCachedMemory();
+	}
 }
 
 void URigVMHost::PostLoad()
