Reduced msvc compile time of RigVMFunction_MathVector.cpp from 2+ minutes to 15 seconds. This is just a mechanical change and should have no behavior difference

diff --git a/Engine/Plugins/Runtime/RigVM/Source/RigVM/Private/RigVMFunctions/Math/RigVMFunction_MathVector.cpp b/Engine/Plugins/Runtime/RigVM/Source/RigVM/Private/RigVMFunctions/Math/RigVMFunction_MathVector.cpp
--- a/Engine/Plugins/Runtime/RigVM/Source/RigVM/Private/RigVMFunctions/Math/RigVMFunction_MathVector.cpp
+++ b/Engine/Plugins/Runtime/RigVM/Source/RigVM/Private/RigVMFunctions/Math/RigVMFunction_MathVector.cpp
@@ -389,214 +389,227 @@ FRigVMStructUpgradeInfo FRigVMFunction_MathVectorMakeBezierFourPoint::GetUpgrade
 	return FRigVMStructUpgradeInfo();
 }
 
-FRigVMFunction_MathVectorClampSpatially_Execute()
+void FRigVMFunction_DrawArrow(FRigVMDrawInterface* DrawInterface, const FTransform& ArrowTransform, const FVector& A, const FVector& B, const FVector& C)
 {
-	Result = FRigVMMathLibrary::ClampSpatially(Value, Axis, Type, Minimum, Maximum, Space);
+	DrawInterface->DrawArrow(FTransform(A) * ArrowTransform, B * 0.25f, C * 0.08f, FLinearColor::Yellow, 0.f);
+}
 
-	if (ExecuteContext.GetDrawInterface() != nullptr && bDrawDebug)
+void FRigVMFunction_DrawPlane(FRigVMDrawInterface* DrawInterface, EAxis::Type Axis, float Minimum, float& CurrentMaximum, FTransform& CurrentSpace, FVector Result)
+{
+	FVector A, B, Normal;
+
+	switch (Axis)
 	{
-		ERigVMClampSpatialMode::Type Mode = Type;
-		float CurrentMaximum = Maximum;
-		FTransform CurrentSpace = Space;
+	case EAxis::X:
+	{
+		A = FVector::YAxisVector;
+		B = FVector::ZAxisVector;
+		Normal = FVector::XAxisVector;
+		break;
+	}
+	case EAxis::Y:
+	{
+		A = FVector::XAxisVector;
+		B = FVector::ZAxisVector;
+		Normal = FVector::YAxisVector;
+		break;
+	}
+	default:
+	{
+		A = FVector::XAxisVector;
+		B = FVector::YAxisVector;
+		Normal = FVector::ZAxisVector;
+		break;
+	}
+	}
 
-		auto DrawPlane = [ExecuteContext, Axis, Minimum, &CurrentMaximum, &CurrentSpace, Result]()
-		{
-			FVector A, B, Normal;
-				
-			switch (Axis)
-			{
-				case EAxis::X:
-				{
-					A = FVector::YAxisVector;
-					B = FVector::ZAxisVector;
-					Normal = FVector::XAxisVector;
-					break;
-				}
-				case EAxis::Y:
-				{
-					A = FVector::XAxisVector;
-					B = FVector::ZAxisVector;
-					Normal = FVector::YAxisVector;
-					break;
-				}
-				default:
-				{
-					A = FVector::XAxisVector;
-					B = FVector::YAxisVector;
-					Normal = FVector::ZAxisVector;
-					break;
-				}
-			}
+	const FVector PointOnPlane = FVector::PointPlaneProject(Result, FPlane(CurrentSpace.GetTranslation(), Normal));
+	const float PlaneScale = FMath::Max<float>(10.f, FVector::Distance(CurrentSpace.GetTranslation(), PointOnPlane)) * 0.525f;
 
-			const FVector PointOnPlane = FVector::PointPlaneProject(Result, FPlane(CurrentSpace.GetTranslation(), Normal));
-			const float PlaneScale = FMath::Max<float>(10.f, FVector::Distance(CurrentSpace.GetTranslation(), PointOnPlane)) * 0.525f;
+	FTransform PlaneRotation = FTransform::Identity;
+	PlaneRotation.SetRotation(FMatrix(A, Normal.Cross(A), Normal, FVector::ZeroVector).ToQuat());
+	const FTransform PlaneTranslation = FTransform((PointOnPlane - CurrentSpace.GetTranslation()) * 0.5);
 
-			FTransform PlaneRotation = FTransform::Identity;
-			PlaneRotation.SetRotation(FMatrix(A, Normal.Cross(A), Normal, FVector::ZeroVector).ToQuat());
-			const FTransform PlaneTranslation = FTransform((PointOnPlane - CurrentSpace.GetTranslation()) * 0.5);
+	const FTransform MinimumWorldTransform = FTransform(Normal * Minimum) * CurrentSpace * PlaneTranslation;
+	DrawInterface->DrawPlane(PlaneRotation * MinimumWorldTransform, FVector2D(PlaneScale, PlaneScale), FLinearColor::Yellow, true, FLinearColor::Yellow, GEngine->ConstraintLimitMaterialPrismatic->GetRenderProxy());
+	DrawInterface->DrawArrow(MinimumWorldTransform, Normal * PlaneScale * 0.25f, A * PlaneScale * 0.08f, FLinearColor::Yellow, 0.f);
 
-			const FTransform MinimumWorldTransform = FTransform(Normal * Minimum) * CurrentSpace * PlaneTranslation;
-			ExecuteContext.GetDrawInterface()->DrawPlane(PlaneRotation * MinimumWorldTransform, FVector2D(PlaneScale, PlaneScale), FLinearColor::Yellow, true, FLinearColor::Yellow, GEngine->ConstraintLimitMaterialPrismatic->GetRenderProxy());
-			ExecuteContext.GetDrawInterface()->DrawArrow(MinimumWorldTransform, Normal * PlaneScale * 0.25f, A * PlaneScale * 0.08f, FLinearColor::Yellow, 0.f);
+	if (CurrentMaximum > SMALL_NUMBER)
+	{
+		const FTransform MaximumWorldTransform = FTransform(Normal * CurrentMaximum) * CurrentSpace * PlaneTranslation;
+		DrawInterface->DrawPlane(PlaneRotation * MaximumWorldTransform, FVector2D(PlaneScale, PlaneScale), FLinearColor::Yellow, true, FLinearColor::Yellow, GEngine->ConstraintLimitMaterialPrismatic->GetRenderProxy());
+		DrawInterface->DrawArrow(MaximumWorldTransform, -Normal * PlaneScale * 0.25f, A * PlaneScale * 0.08f, FLinearColor::Yellow, 0.f);
+	}
+};
 
-			if (CurrentMaximum > SMALL_NUMBER)
-			{
-				const FTransform MaximumWorldTransform = FTransform(Normal * CurrentMaximum) * CurrentSpace * PlaneTranslation;
-				ExecuteContext.GetDrawInterface()->DrawPlane(PlaneRotation * MaximumWorldTransform, FVector2D(PlaneScale, PlaneScale), FLinearColor::Yellow, true, FLinearColor::Yellow, GEngine->ConstraintLimitMaterialPrismatic->GetRenderProxy());
-				ExecuteContext.GetDrawInterface()->DrawArrow(MaximumWorldTransform, -Normal * PlaneScale * 0.25f, A * PlaneScale * 0.08f, FLinearColor::Yellow, 0.f);
-			}
-		};
+void FRigVMFunction_DrawCylinder(FRigVMDrawInterface* DrawInterface, EAxis::Type Axis, float Minimum, float& CurrentMaximum, FTransform& CurrentSpace, FVector Result)
+{
+	FTransform CircleTransform = FTransform::Identity;
+	switch (Axis)
+	{
+	case EAxis::X:
+	{
+		CircleTransform.SetRotation(FQuat(FVector(0.f, 1.f, 0.f), PI * 0.5f));
+		break;
+	}
+	case EAxis::Y:
+	{
+		CircleTransform.SetRotation(FQuat(FVector(1.f, 0.f, 0.f), PI * 0.5f));
+		break;
+	}
+	default:
+	{
+		break;
+	}
+	}
 
-		auto DrawCylinder = [ExecuteContext, Axis, Minimum, &CurrentMaximum, &CurrentSpace, Result]()
+	const FTransform CombinedTransform = CircleTransform * CurrentSpace;
+	const FVector PointOnAxis = CombinedTransform.InverseTransformPosition(Result) * FVector(0, 0, 1);
+	const float Extent = FMath::Abs<float>((float)PointOnAxis.Z);
+	const FTransform CenterTransform = FTransform(PointOnAxis * 0.5) * CombinedTransform;
+	const TArray<FTransform> Transforms = {
+		FTransform(FVector(0, 0, Extent * 0.5)) * CenterTransform,
+		FTransform(FVector(0, 0, -Extent * 0.5)) * CenterTransform
+	};
+
+	if (Minimum > SMALL_NUMBER)
+	{
+		for (const FTransform& Transform : Transforms)
 		{
-			FTransform CircleTransform = FTransform::Identity;
-			switch (Axis)
+			DrawInterface->DrawCircle(Transform, FTransform::Identity, Minimum, FLinearColor::Yellow, 0.f, 32);
+			FRigVMFunction_DrawArrow(DrawInterface, Transform, FVector::XAxisVector * Minimum, FVector::XAxisVector * Minimum, FVector::YAxisVector * Minimum);
+			FRigVMFunction_DrawArrow(DrawInterface, Transform, -FVector::XAxisVector * Minimum, -FVector::XAxisVector * Minimum, FVector::YAxisVector * Minimum);
+			FRigVMFunction_DrawArrow(DrawInterface, Transform, FVector::YAxisVector * Minimum, FVector::YAxisVector * Minimum, FVector::XAxisVector * Minimum);
+			FRigVMFunction_DrawArrow(DrawInterface, Transform, -FVector::YAxisVector * Minimum, -FVector::YAxisVector * Minimum, FVector::XAxisVector * Minimum);
+		}
+		DrawInterface->DrawLines(FTransform::Identity,
 			{
-				case EAxis::X:
-				{
-					CircleTransform.SetRotation(FQuat(FVector(0.f, 1.f, 0.f), PI * 0.5f));
-					break;
-				}
-				case EAxis::Y:
-				{
-					CircleTransform.SetRotation(FQuat(FVector(1.f, 0.f, 0.f), PI * 0.5f));
-					break;
-				}
-				default:
-				{
-					break;
-				}
-			}
+				Transforms[0].TransformPosition(FVector::XAxisVector * Minimum),
+				Transforms[1].TransformPosition(FVector::XAxisVector * Minimum),
+				Transforms[0].TransformPosition(-FVector::XAxisVector * Minimum),
+				Transforms[1].TransformPosition(-FVector::XAxisVector * Minimum),
+				Transforms[0].TransformPosition(FVector::YAxisVector * Minimum),
+				Transforms[1].TransformPosition(FVector::YAxisVector * Minimum),
+				Transforms[0].TransformPosition(-FVector::YAxisVector * Minimum),
+				Transforms[1].TransformPosition(-FVector::YAxisVector * Minimum),
+			},
+			FLinearColor::Yellow,
+			0
+			);
+	}
+	if (CurrentMaximum > SMALL_NUMBER)
+	{
+		for (const FTransform& Transform : Transforms)
+		{
+			TArray<FVector> LinesToDraw;
+			LinesToDraw.Reserve(8);
+			DrawInterface->DrawCircle(Transform, FTransform::Identity, CurrentMaximum, FLinearColor::Yellow, 0.f, 32);
+			FRigVMFunction_DrawArrow(DrawInterface, Transform, -FVector::XAxisVector * CurrentMaximum, FVector::XAxisVector * CurrentMaximum, FVector::YAxisVector * CurrentMaximum);
+			FRigVMFunction_DrawArrow(DrawInterface, Transform, FVector::XAxisVector * CurrentMaximum, -FVector::XAxisVector * CurrentMaximum, FVector::YAxisVector * CurrentMaximum);
+			FRigVMFunction_DrawArrow(DrawInterface, Transform, -FVector::YAxisVector * CurrentMaximum, FVector::YAxisVector * CurrentMaximum, FVector::XAxisVector * CurrentMaximum);
+			FRigVMFunction_DrawArrow(DrawInterface, Transform, FVector::YAxisVector * CurrentMaximum, -FVector::YAxisVector * CurrentMaximum, FVector::XAxisVector * CurrentMaximum);
+		}
+		DrawInterface->DrawLines(FTransform::Identity,
+			{
+				Transforms[0].TransformPosition(FVector::XAxisVector * CurrentMaximum),
+				Transforms[1].TransformPosition(FVector::XAxisVector * CurrentMaximum),
+				Transforms[0].TransformPosition(-FVector::XAxisVector * CurrentMaximum),
+				Transforms[1].TransformPosition(-FVector::XAxisVector * CurrentMaximum),
+				Transforms[0].TransformPosition(FVector::YAxisVector * CurrentMaximum),
+				Transforms[1].TransformPosition(FVector::YAxisVector * CurrentMaximum),
+				Transforms[0].TransformPosition(-FVector::YAxisVector * CurrentMaximum),
+				Transforms[1].TransformPosition(-FVector::YAxisVector * CurrentMaximum),
+			},
+			FLinearColor::Yellow,
+			0
+			);
+	}
+};
 
-			const FTransform CombinedTransform = CircleTransform * CurrentSpace;
-			const FVector PointOnAxis = CombinedTransform.InverseTransformPosition(Result) * FVector(0, 0, 1);
-			const float Extent = FMath::Abs<float>((float)PointOnAxis.Z);
-			const FTransform CenterTransform = FTransform(PointOnAxis * 0.5) * CombinedTransform;
-			const TArray<FTransform> Transforms = {
-				FTransform(FVector(0, 0, Extent * 0.5)) * CenterTransform,
-				FTransform(FVector(0, 0, -Extent * 0.5)) * CenterTransform
-			};
+void FRigVMFunction_DrawSphere(FRigVMDrawInterface* DrawInterface, EAxis::Type Axis, float Minimum, float& CurrentMaximum, FTransform& CurrentSpace, FVector Result)
+{
+	FVector XAxisMin = FVector::XAxisVector * Minimum;
+	FVector XAxisMax = FVector::XAxisVector * CurrentMaximum;
+	FVector YAxisMin = FVector::YAxisVector * Minimum;
+	FVector YAxisMax = FVector::YAxisVector * CurrentMaximum;
+	FVector ZAxisMin = FVector::ZAxisVector * Minimum;
+	FVector ZAxisMax = FVector::ZAxisVector * CurrentMaximum;
 
-			if (Minimum > SMALL_NUMBER)
-			{
-				for(const FTransform& Transform : Transforms)
-				{
-					ExecuteContext.GetDrawInterface()->DrawCircle(Transform, FTransform::Identity, Minimum, FLinearColor::Yellow, 0.f, 32);
-					ExecuteContext.GetDrawInterface()->DrawArrow(FTransform(FVector::XAxisVector * Minimum) * Transform, FVector::XAxisVector * Minimum * 0.25f, FVector::YAxisVector * Minimum * 0.08f, FLinearColor::Yellow, 0.f);
-					ExecuteContext.GetDrawInterface()->DrawArrow(FTransform(-FVector::XAxisVector * Minimum) * Transform, -FVector::XAxisVector * Minimum * 0.25f, FVector::YAxisVector * Minimum * 0.08f, FLinearColor::Yellow, 0.f);
-					ExecuteContext.GetDrawInterface()->DrawArrow(FTransform(FVector::YAxisVector * Minimum) * Transform, FVector::YAxisVector * Minimum * 0.25f, FVector::XAxisVector * Minimum * 0.08f, FLinearColor::Yellow, 0.f);
-					ExecuteContext.GetDrawInterface()->DrawArrow(FTransform(-FVector::YAxisVector * Minimum) * Transform, -FVector::YAxisVector * Minimum * 0.25f, FVector::XAxisVector * Minimum * 0.08f, FLinearColor::Yellow, 0.f);
-				}
-				ExecuteContext.GetDrawInterface()->DrawLines(FTransform::Identity,
-					{
-						Transforms[0].TransformPosition(FVector::XAxisVector * Minimum),
-						Transforms[1].TransformPosition(FVector::XAxisVector * Minimum),
-						Transforms[0].TransformPosition(-FVector::XAxisVector * Minimum),
-						Transforms[1].TransformPosition(-FVector::XAxisVector * Minimum),
-						Transforms[0].TransformPosition(FVector::YAxisVector * Minimum),
-						Transforms[1].TransformPosition(FVector::YAxisVector * Minimum),
-						Transforms[0].TransformPosition(-FVector::YAxisVector * Minimum),
-						Transforms[1].TransformPosition(-FVector::YAxisVector * Minimum),
-					},
-					FLinearColor::Yellow,
-					0
-				);
-			}
-			if (CurrentMaximum > SMALL_NUMBER)
-			{
-				for(const FTransform& Transform : Transforms)
-				{
-					TArray<FVector> LinesToDraw;
-					LinesToDraw.Reserve(8);
-					ExecuteContext.GetDrawInterface()->DrawCircle(Transform, FTransform::Identity, CurrentMaximum, FLinearColor::Yellow, 0.f, 32);
-					ExecuteContext.GetDrawInterface()->DrawArrow(FTransform(-FVector::XAxisVector * CurrentMaximum) * Transform, FVector::XAxisVector * CurrentMaximum * 0.25f, FVector::YAxisVector * CurrentMaximum * 0.08f, FLinearColor::Yellow, 0.f);
-					ExecuteContext.GetDrawInterface()->DrawArrow(FTransform(FVector::XAxisVector * CurrentMaximum) * Transform, -FVector::XAxisVector * CurrentMaximum * 0.25f, FVector::YAxisVector * CurrentMaximum * 0.08f, FLinearColor::Yellow, 0.f);
-					ExecuteContext.GetDrawInterface()->DrawArrow(FTransform(-FVector::YAxisVector * CurrentMaximum) * Transform, FVector::YAxisVector * CurrentMaximum * 0.25f, FVector::XAxisVector * CurrentMaximum * 0.08f, FLinearColor::Yellow, 0.f);
-					ExecuteContext.GetDrawInterface()->DrawArrow(FTransform(FVector::YAxisVector * CurrentMaximum) * Transform, -FVector::YAxisVector * CurrentMaximum * 0.25f, FVector::XAxisVector * CurrentMaximum * 0.08f, FLinearColor::Yellow, 0.f);
-				}
-				ExecuteContext.GetDrawInterface()->DrawLines(FTransform::Identity,
-                    {
-                    	Transforms[0].TransformPosition(FVector::XAxisVector * CurrentMaximum),
-                    	Transforms[1].TransformPosition(FVector::XAxisVector * CurrentMaximum),
-                    	Transforms[0].TransformPosition(-FVector::XAxisVector * CurrentMaximum),
-                    	Transforms[1].TransformPosition(-FVector::XAxisVector * CurrentMaximum),
-                    	Transforms[0].TransformPosition(FVector::YAxisVector * CurrentMaximum),
-                    	Transforms[1].TransformPosition(FVector::YAxisVector * CurrentMaximum),
-                    	Transforms[0].TransformPosition(-FVector::YAxisVector * CurrentMaximum),
-                    	Transforms[1].TransformPosition(-FVector::YAxisVector * CurrentMaximum),
-                    },
-                    FLinearColor::Yellow,
-                    0
-                );
-			}
-		};
+	FTransform CircleTransform = FTransform::Identity;
+	if (Minimum > SMALL_NUMBER)
+	{
+		DrawInterface->DrawCircle(CurrentSpace, CircleTransform, Minimum, FLinearColor::Yellow, 0.f, 32);
+		const FTransform ArrowTransform = CircleTransform * CurrentSpace;
+		FRigVMFunction_DrawArrow(DrawInterface, ArrowTransform, XAxisMin, XAxisMin, YAxisMin);
+		FRigVMFunction_DrawArrow(DrawInterface, ArrowTransform, -XAxisMin, -XAxisMin, YAxisMin);
+		FRigVMFunction_DrawArrow(DrawInterface, ArrowTransform, YAxisMin, YAxisMin, YAxisMin);
+		FRigVMFunction_DrawArrow(DrawInterface, ArrowTransform, -YAxisMin, -YAxisMin, YAxisMin);
+	}
+	if (CurrentMaximum > SMALL_NUMBER)
+	{
+		DrawInterface->DrawCircle(CurrentSpace, CircleTransform, CurrentMaximum, FLinearColor::Yellow, 0.f, 32);
+		const FTransform ArrowTransform = CircleTransform * CurrentSpace;
+		FRigVMFunction_DrawArrow(DrawInterface, ArrowTransform, -XAxisMax, XAxisMax, ZAxisMax);
+		FRigVMFunction_DrawArrow(DrawInterface, ArrowTransform, XAxisMax, -XAxisMax, ZAxisMax);
+		FRigVMFunction_DrawArrow(DrawInterface, ArrowTransform, -YAxisMax, YAxisMax, ZAxisMax);
+		FRigVMFunction_DrawArrow(DrawInterface, ArrowTransform, YAxisMax, -YAxisMax, ZAxisMax);
+	}
+	CircleTransform.SetRotation(FQuat(FVector(0.f, 1.f, 0.f), PI * 0.5f));
+	if (Minimum > SMALL_NUMBER)
+	{
+		DrawInterface->DrawCircle(CurrentSpace, CircleTransform, Minimum, FLinearColor::Yellow, 0.f, 32);
+		const FTransform ArrowTransform = CircleTransform * CurrentSpace;
+		FRigVMFunction_DrawArrow(DrawInterface, ArrowTransform, XAxisMin, XAxisMin, ZAxisMin);
+		FRigVMFunction_DrawArrow(DrawInterface, ArrowTransform, -XAxisMin, -XAxisMin, ZAxisMin);
+		FRigVMFunction_DrawArrow(DrawInterface, ArrowTransform, YAxisMin, YAxisMin, ZAxisMin);
+		FRigVMFunction_DrawArrow(DrawInterface, ArrowTransform, -YAxisMin, -YAxisMin, ZAxisMin);
+	}
+	if (CurrentMaximum > SMALL_NUMBER)
+	{
+		DrawInterface->DrawCircle(CurrentSpace, CircleTransform, CurrentMaximum, FLinearColor::Yellow, 0.f, 32);
+		const FTransform ArrowTransform = CircleTransform * CurrentSpace;
+		FRigVMFunction_DrawArrow(DrawInterface, ArrowTransform, -XAxisMax, XAxisMax, ZAxisMax);
+		FRigVMFunction_DrawArrow(DrawInterface, ArrowTransform, XAxisMax, -XAxisMax, ZAxisMax);
+		FRigVMFunction_DrawArrow(DrawInterface, ArrowTransform, -YAxisMax, YAxisMax, ZAxisMax);
+		FRigVMFunction_DrawArrow(DrawInterface, ArrowTransform, YAxisMax, -YAxisMax, ZAxisMax);
+	}
+	CircleTransform.SetRotation(FQuat(FVector(1.f, 0.f, 0.f), PI * 0.5f));
+	if (Minimum > SMALL_NUMBER)
+	{
+		DrawInterface->DrawCircle(CurrentSpace, CircleTransform, Minimum, FLinearColor::Yellow, 0.f, 32);
+		const FTransform ArrowTransform = CircleTransform * CurrentSpace;
+		FRigVMFunction_DrawArrow(DrawInterface, ArrowTransform, XAxisMin, XAxisMin, ZAxisMin);
+		FRigVMFunction_DrawArrow(DrawInterface, ArrowTransform, -XAxisMin, -XAxisMin, ZAxisMin);
+		FRigVMFunction_DrawArrow(DrawInterface, ArrowTransform, YAxisMin, YAxisMin, ZAxisMin);
+		FRigVMFunction_DrawArrow(DrawInterface, ArrowTransform, -YAxisMin, -YAxisMin, ZAxisMin);
+	}
+	if (CurrentMaximum > SMALL_NUMBER)
+	{
+		DrawInterface->DrawCircle(CurrentSpace, CircleTransform, CurrentMaximum, FLinearColor::Yellow, 0.f, 32);
+		const FTransform ArrowTransform = CircleTransform * CurrentSpace;
+		FRigVMFunction_DrawArrow(DrawInterface, ArrowTransform, -XAxisMax, XAxisMax, ZAxisMax);
+		FRigVMFunction_DrawArrow(DrawInterface, ArrowTransform, XAxisMax, -XAxisMax, ZAxisMax);
+		FRigVMFunction_DrawArrow(DrawInterface, ArrowTransform, -YAxisMax, YAxisMax, ZAxisMax);
+		FRigVMFunction_DrawArrow(DrawInterface, ArrowTransform, YAxisMax, -YAxisMax, ZAxisMax);
+	}
+};
 
-		auto DrawSphere = [ExecuteContext, Axis, Minimum, &CurrentMaximum, &CurrentSpace, Result]()
-		{
-			FTransform CircleTransform = FTransform::Identity;
-			if (Minimum > SMALL_NUMBER)
-			{
-				ExecuteContext.GetDrawInterface()->DrawCircle(CurrentSpace, CircleTransform, Minimum, FLinearColor::Yellow, 0.f, 32);
-				const FTransform ArrowTransform = CircleTransform * CurrentSpace;
-				ExecuteContext.GetDrawInterface()->DrawArrow(FTransform(FVector::XAxisVector * Minimum) * ArrowTransform, FVector::XAxisVector * Minimum * 0.25f, FVector::YAxisVector * Minimum * 0.08f, FLinearColor::Yellow, 0.f);
-				ExecuteContext.GetDrawInterface()->DrawArrow(FTransform(-FVector::XAxisVector * Minimum) * ArrowTransform, -FVector::XAxisVector * Minimum * 0.25f, FVector::YAxisVector * Minimum * 0.08f, FLinearColor::Yellow, 0.f);
-				ExecuteContext.GetDrawInterface()->DrawArrow(FTransform(FVector::YAxisVector * Minimum) * ArrowTransform, FVector::YAxisVector * Minimum * 0.25f, FVector::YAxisVector * Minimum * 0.08f, FLinearColor::Yellow, 0.f);
-				ExecuteContext.GetDrawInterface()->DrawArrow(FTransform(-FVector::YAxisVector * Minimum) * ArrowTransform, -FVector::YAxisVector * Minimum * 0.25f, FVector::YAxisVector * Minimum * 0.08f, FLinearColor::Yellow, 0.f);
-			}
-			if(CurrentMaximum > SMALL_NUMBER)
-			{
-				ExecuteContext.GetDrawInterface()->DrawCircle(CurrentSpace, CircleTransform, CurrentMaximum, FLinearColor::Yellow, 0.f, 32);
-				const FTransform ArrowTransform = CircleTransform * CurrentSpace;
-				ExecuteContext.GetDrawInterface()->DrawArrow(FTransform(-FVector::XAxisVector * CurrentMaximum) * ArrowTransform, FVector::XAxisVector * CurrentMaximum * 0.25f, FVector::ZAxisVector * CurrentMaximum * 0.08f, FLinearColor::Yellow, 0.f);
-				ExecuteContext.GetDrawInterface()->DrawArrow(FTransform(FVector::XAxisVector * CurrentMaximum) * ArrowTransform, -FVector::XAxisVector * CurrentMaximum * 0.25f, FVector::ZAxisVector * CurrentMaximum * 0.08f, FLinearColor::Yellow, 0.f);
-				ExecuteContext.GetDrawInterface()->DrawArrow(FTransform(-FVector::YAxisVector * CurrentMaximum) * ArrowTransform, FVector::YAxisVector * CurrentMaximum * 0.25f, FVector::ZAxisVector * CurrentMaximum * 0.08f, FLinearColor::Yellow, 0.f);
-				ExecuteContext.GetDrawInterface()->DrawArrow(FTransform(FVector::YAxisVector * CurrentMaximum) * ArrowTransform, -FVector::YAxisVector * CurrentMaximum * 0.25f, FVector::ZAxisVector * CurrentMaximum * 0.08f, FLinearColor::Yellow, 0.f);
-			}
-			CircleTransform.SetRotation(FQuat(FVector(0.f, 1.f, 0.f), PI * 0.5f));
-			if (Minimum > SMALL_NUMBER)
-			{
-				ExecuteContext.GetDrawInterface()->DrawCircle(CurrentSpace, CircleTransform, Minimum, FLinearColor::Yellow, 0.f, 32);
-				const FTransform ArrowTransform = CircleTransform * CurrentSpace;
-				ExecuteContext.GetDrawInterface()->DrawArrow(FTransform(FVector::XAxisVector * Minimum) * ArrowTransform, FVector::XAxisVector * Minimum * 0.25f, FVector::ZAxisVector * Minimum * 0.08f, FLinearColor::Yellow, 0.f);
-				ExecuteContext.GetDrawInterface()->DrawArrow(FTransform(-FVector::XAxisVector * Minimum) * ArrowTransform, -FVector::XAxisVector * Minimum * 0.25f, FVector::ZAxisVector * Minimum * 0.08f, FLinearColor::Yellow, 0.f);
-				ExecuteContext.GetDrawInterface()->DrawArrow(FTransform(FVector::YAxisVector * Minimum) * ArrowTransform, FVector::YAxisVector * Minimum * 0.25f, FVector::ZAxisVector * Minimum * 0.08f, FLinearColor::Yellow, 0.f);
-				ExecuteContext.GetDrawInterface()->DrawArrow(FTransform(-FVector::YAxisVector * Minimum) * ArrowTransform, -FVector::YAxisVector * Minimum * 0.25f, FVector::ZAxisVector * Minimum * 0.08f, FLinearColor::Yellow, 0.f);
-			}
-			if(CurrentMaximum > SMALL_NUMBER)
-			{
-				ExecuteContext.GetDrawInterface()->DrawCircle(CurrentSpace, CircleTransform, CurrentMaximum, FLinearColor::Yellow, 0.f, 32);
-				const FTransform ArrowTransform = CircleTransform * CurrentSpace;
-				ExecuteContext.GetDrawInterface()->DrawArrow(FTransform(-FVector::XAxisVector * CurrentMaximum) * ArrowTransform, FVector::XAxisVector * CurrentMaximum * 0.25f, FVector::ZAxisVector * CurrentMaximum * 0.08f, FLinearColor::Yellow, 0.f);
-				ExecuteContext.GetDrawInterface()->DrawArrow(FTransform(FVector::XAxisVector * CurrentMaximum) * ArrowTransform, -FVector::XAxisVector * CurrentMaximum * 0.25f, FVector::ZAxisVector * CurrentMaximum * 0.08f, FLinearColor::Yellow, 0.f);
-				ExecuteContext.GetDrawInterface()->DrawArrow(FTransform(-FVector::YAxisVector * CurrentMaximum) * ArrowTransform, FVector::YAxisVector * CurrentMaximum * 0.25f, FVector::ZAxisVector * CurrentMaximum * 0.08f, FLinearColor::Yellow, 0.f);
-				ExecuteContext.GetDrawInterface()->DrawArrow(FTransform(FVector::YAxisVector * CurrentMaximum) * ArrowTransform, -FVector::YAxisVector * CurrentMaximum * 0.25f, FVector::ZAxisVector * CurrentMaximum * 0.08f, FLinearColor::Yellow, 0.f);
-			}
-			CircleTransform.SetRotation(FQuat(FVector(1.f, 0.f, 0.f), PI * 0.5f));
-			if (Minimum > SMALL_NUMBER)
-			{
-				ExecuteContext.GetDrawInterface()->DrawCircle(CurrentSpace, CircleTransform, Minimum, FLinearColor::Yellow, 0.f, 32);
-				const FTransform ArrowTransform = CircleTransform * CurrentSpace;
-				ExecuteContext.GetDrawInterface()->DrawArrow(FTransform(FVector::XAxisVector * Minimum) * ArrowTransform, FVector::XAxisVector * Minimum * 0.25f, FVector::ZAxisVector * Minimum * 0.08f, FLinearColor::Yellow, 0.f);
-				ExecuteContext.GetDrawInterface()->DrawArrow(FTransform(-FVector::XAxisVector * Minimum) * ArrowTransform, -FVector::XAxisVector * Minimum * 0.25f, FVector::ZAxisVector * Minimum * 0.08f, FLinearColor::Yellow, 0.f);
-				ExecuteContext.GetDrawInterface()->DrawArrow(FTransform(FVector::YAxisVector * Minimum) * ArrowTransform, FVector::YAxisVector * Minimum * 0.25f, FVector::ZAxisVector * Minimum * 0.08f, FLinearColor::Yellow, 0.f);
-				ExecuteContext.GetDrawInterface()->DrawArrow(FTransform(-FVector::YAxisVector * Minimum) * ArrowTransform, -FVector::YAxisVector * Minimum * 0.25f, FVector::ZAxisVector * Minimum * 0.08f, FLinearColor::Yellow, 0.f);
-			}
-			if(CurrentMaximum > SMALL_NUMBER)
-			{
-				ExecuteContext.GetDrawInterface()->DrawCircle(CurrentSpace, CircleTransform, CurrentMaximum, FLinearColor::Yellow, 0.f, 32);
-				const FTransform ArrowTransform = CircleTransform * CurrentSpace;
-				ExecuteContext.GetDrawInterface()->DrawArrow(FTransform(-FVector::XAxisVector * CurrentMaximum) * ArrowTransform, FVector::XAxisVector * CurrentMaximum * 0.25f, FVector::ZAxisVector * CurrentMaximum * 0.08f, FLinearColor::Yellow, 0.f);
-				ExecuteContext.GetDrawInterface()->DrawArrow(FTransform(FVector::XAxisVector * CurrentMaximum) * ArrowTransform, -FVector::XAxisVector * CurrentMaximum * 0.25f, FVector::ZAxisVector * CurrentMaximum * 0.08f, FLinearColor::Yellow, 0.f);
-				ExecuteContext.GetDrawInterface()->DrawArrow(FTransform(-FVector::YAxisVector * CurrentMaximum) * ArrowTransform, FVector::YAxisVector * CurrentMaximum * 0.25f, FVector::ZAxisVector * CurrentMaximum * 0.08f, FLinearColor::Yellow, 0.f);
-				ExecuteContext.GetDrawInterface()->DrawArrow(FTransform(FVector::YAxisVector * CurrentMaximum) * ArrowTransform, -FVector::YAxisVector * CurrentMaximum * 0.25f, FVector::ZAxisVector * CurrentMaximum * 0.08f, FLinearColor::Yellow, 0.f);
-			}
-		};
+FRigVMFunction_MathVectorClampSpatially_Execute()
+{
+	Result = FRigVMMathLibrary::ClampSpatially(Value, Axis, Type, Minimum, Maximum, Space);
+
+	FRigVMDrawInterface* DrawInterface = ExecuteContext.GetDrawInterface();
+	if (DrawInterface != nullptr && bDrawDebug)
+	{
+		ERigVMClampSpatialMode::Type Mode = Type;
+		float CurrentMaximum = Maximum;
+		FTransform CurrentSpace = Space;
 		
 		switch (Mode)
 		{
 			case ERigVMClampSpatialMode::Plane:
 			{
-				DrawPlane();
+				FRigVMFunction_DrawPlane(DrawInterface, Axis, Minimum, CurrentMaximum, CurrentSpace, Result);
 				break;
 			}
 			case ERigVMClampSpatialMode::Capsule:
@@ -605,7 +618,7 @@ FRigVMFunction_MathVectorClampSpatially_Execute()
 				if(CurrentMaximum < Radius * 2.f)
 				{
 					CurrentMaximum = 0.f;
-					DrawSphere();
+					FRigVMFunction_DrawSphere(DrawInterface, Axis, Minimum, CurrentMaximum, CurrentSpace, Result);
 					break;
 				}
 			
@@ -635,7 +648,7 @@ FRigVMFunction_MathVectorClampSpatially_Execute()
 				CircleTransform  = CircleTransform * Space;
 					
 				const FVector CylinderDirection = FVector::ZAxisVector * HalfCylinderLength;
-				ExecuteContext.GetDrawInterface()->DrawLines(CircleTransform,
+				DrawInterface->DrawLines(CircleTransform,
 					{
 						FVector::XAxisVector * Minimum - CylinderDirection,
 						FVector::XAxisVector * Minimum + CylinderDirection,
@@ -653,28 +666,28 @@ FRigVMFunction_MathVectorClampSpatially_Execute()
 				CurrentMaximum = 0.f;
 				{
 					TGuardValue<FTransform> SpaceGuard(CurrentSpace, FTransform(CylinderDirection) * CircleTransform);
-					DrawSphere();
+					FRigVMFunction_DrawSphere(DrawInterface, Axis, Minimum, CurrentMaximum, CurrentSpace, Result);
 				}
 				{
 					TGuardValue<FTransform> SpaceGuard(CurrentSpace, FTransform(-CylinderDirection) * CircleTransform);
-					DrawSphere();
+					FRigVMFunction_DrawSphere(DrawInterface, Axis, Minimum, CurrentMaximum, CurrentSpace, Result);
 				}
 				break;
 			}
 			case ERigVMClampSpatialMode::Cylinder:
 			{
-				DrawCylinder();
+				FRigVMFunction_DrawCylinder(DrawInterface, Axis, Minimum, CurrentMaximum, CurrentSpace, Result);
 				break;
 			}
 			default:
 			case ERigVMClampSpatialMode::Sphere:
 			{
-				DrawSphere();
+				FRigVMFunction_DrawSphere(DrawInterface, Axis, Minimum, CurrentMaximum, CurrentSpace, Result);
 				break;
 			}
 		}
 		
-		ExecuteContext.GetDrawInterface()->DrawPoint(FTransform::Identity, Result, DebugThickness * 8.f, DebugColor);
+		DrawInterface->DrawPoint(FTransform::Identity, Result, DebugThickness * 8.f, DebugColor);
 	}
 }
 
