Changing over to use GetValueOnAnyThread as RigVM is also being used in jobs.

diff --git a/Engine/Plugins/Runtime/RigVM/Source/RigVM/Private/RigVMModule.cpp b/Engine/Plugins/Runtime/RigVM/Source/RigVM/Private/RigVMModule.cpp
--- a/Engine/Plugins/Runtime/RigVM/Source/RigVM/Private/RigVMModule.cpp
+++ b/Engine/Plugins/Runtime/RigVM/Source/RigVM/Private/RigVMModule.cpp
@@ -57,7 +57,7 @@ void FRigVMModule::ShutdownModule()
 bool RigVMCore::SupportsUObjects()
 {
 #if UE_RIGVM_UOBJECT_PROPERTIES_ENABLED
-	return CVarRigVMEnableUObjects.GetValueOnGameThread();
+	return CVarRigVMEnableUObjects.GetValueOnAnyThread();
 #else
 	return false;
 #endif
@@ -66,7 +66,7 @@ bool RigVMCore::SupportsUObjects()
 bool RigVMCore::SupportsUInterfaces()
 {
 #if UE_RIGVM_UINTERFACE_PROPERTIES_ENABLED
-	return CVarRigVMEnableUInterfaces.GetValueOnGameThread();
+	return CVarRigVMEnableUInterfaces.GetValueOnAnyThread();
 #else
 	return false;
 #endif
