Fix for ControlRig silently loosing pin conenctions on UUserStruct members when a  member of the struct is renamed while being used by the rig

diff --git a/Engine/Plugins/Runtime/RigVM/Source/RigVMDeveloper/Private/RigVMBlueprint.cpp b/Engine/Plugins/Runtime/RigVM/Source/RigVMDeveloper/Private/RigVMBlueprint.cpp
--- a/Engine/Plugins/Runtime/RigVM/Source/RigVMDeveloper/Private/RigVMBlueprint.cpp
+++ b/Engine/Plugins/Runtime/RigVM/Source/RigVMDeveloper/Private/RigVMBlueprint.cpp
@@ -1312,6 +1312,10 @@ void URigVMBlueprint::RefreshAllModels(ERigVMBlueprintLoadType InLoadType)
 
 	if (ensure(IsInGameThread()))
 	{
+		TArray<URigVMController::FRepopulatePinsNodeData> RepopulatePinsNodesData;
+		constexpr int32 REPOPULATE_NODES_NUM_RESERVED = 800;
+		RepopulatePinsNodesData.Reserve(REPOPULATE_NODES_NUM_RESERVED);
+
 		for (URigVMGraph* Graph : AllModelsLeavesFirst)
 		{
 			URigVMController* Controller = GetOrCreateController(Graph);
@@ -1320,11 +1324,24 @@ void URigVMBlueprint::RefreshAllModels(ERigVMBlueprintLoadType InLoadType)
 			TGuardValue<bool> GuardEditGraph(Graph->bEditable, true);
 			FRigVMControllerNotifGuard NotifGuard(Controller, true);
 			LinkedPaths.Add(Graph, Controller->GetLinkedPaths());
-			Controller->FastBreakLinkedPaths(LinkedPaths.FindChecked(Graph));
-			TArray<URigVMNode*> Nodes = Graph->GetNodes();
-			for (URigVMNode* Node : Nodes)
+
+			const TArray<URigVMNode*> Nodes = Graph->GetNodes();
+			if (Nodes.Num() > 0)
 			{
-				Controller->RepopulatePinsOnNode(Node, true, true);
+				RepopulatePinsNodesData.Reset();
+
+				for (URigVMNode* Node : Nodes)
+				{
+					Controller->GenerateRepopulatePinsNodeData(RepopulatePinsNodesData, Node, true, true);
+				}
+
+#if UE_RIGVMCONTROLLER_VERBOSE_REPOPULATE
+				UE_LOG(LogRigVMDeveloper, Display, TEXT("--- Graph: [%s/%s]  - NumNodes : [%d]"), *Graph->GetOuter()->GetName(), *Graph->GetName(), RepopulatePinsNodesData.Num());
+#endif
+
+				Controller->OrphanPins(RepopulatePinsNodesData);
+				Controller->FastBreakLinkedPaths(LinkedPaths.FindChecked(Graph));
+				Controller->RepopulatePins(RepopulatePinsNodesData);
 			}
 		}
 		SetupPinRedirectorsForBackwardsCompatibility();
@@ -1338,7 +1355,8 @@ void URigVMBlueprint::RefreshAllModels(ERigVMBlueprintLoadType InLoadType)
 		{
 			URigVMController::FRestoreLinkedPathSettings Settings;
 			Settings.bFollowCoreRedirectors = true;
-			Controller->RestoreLinkedPaths(LinkedPaths.FindChecked(Graph));
+			Settings.bRelayToOrphanPins = true;
+			Controller->RestoreLinkedPaths(LinkedPaths.FindChecked(Graph), Settings);
 		}
 
 		for(URigVMNode* ModelNode : Graph->GetNodes())
diff --git a/Engine/Plugins/Runtime/RigVM/Source/RigVMDeveloper/Private/RigVMModel/RigVMController.cpp b/Engine/Plugins/Runtime/RigVM/Source/RigVMDeveloper/Private/RigVMModel/RigVMController.cpp
--- a/Engine/Plugins/Runtime/RigVM/Source/RigVMDeveloper/Private/RigVMModel/RigVMController.cpp
+++ b/Engine/Plugins/Runtime/RigVM/Source/RigVMDeveloper/Private/RigVMModel/RigVMController.cpp
@@ -14842,7 +14842,7 @@ bool URigVMController::ShouldRedirectPin(const FString& InOldPinPath, FString& I
 	return false;
 }
 
-void URigVMController::RepopulatePinsOnNode(URigVMNode* InNode, bool bFollowCoreRedirectors, bool bSetupOrphanedPins, bool bRecreateLinks)
+void URigVMController::GenerateRepopulatePinsNodeData(TArray<FRepopulatePinsNodeData>& NodesPinData, URigVMNode* InNode, bool bInFollowCoreRedirectors, bool bInSetupOrphanedPins, bool bInRecreateLinks)
 {
 	if (InNode == nullptr)
 	{
@@ -14851,43 +14851,141 @@ void URigVMController::RepopulatePinsOnNode(URigVMNode* InNode, bool bFollowCore
 	}
 
 	FRigVMControllerCompileBracketScope CompileBracketScope(this);
+	FScopeLock Lock(&PinPathCoreRedirectorsLock);
+
+	FRepopulatePinsNodeData NodeData;
+	NodeData.Node = InNode;
+	NodeData.bFollowCoreRedirectors = bInFollowCoreRedirectors;
+	NodeData.bRecreateLinks = bInRecreateLinks;
+	NodeData.bSetupOrphanPinsForThisNode = bInSetupOrphanedPins;
 
-	URigVMUnitNode* UnitNode = Cast<URigVMUnitNode>(InNode);
-	URigVMRerouteNode* RerouteNode = Cast<URigVMRerouteNode>(InNode);
-	URigVMFunctionEntryNode* EntryNode = Cast<URigVMFunctionEntryNode>(InNode);
-	URigVMFunctionReturnNode* ReturnNode = Cast<URigVMFunctionReturnNode>(InNode);
 	URigVMCollapseNode* CollapseNode = Cast<URigVMCollapseNode>(InNode);
-	URigVMFunctionReferenceNode* FunctionRefNode = Cast<URigVMFunctionReferenceNode>(InNode);
-	URigVMVariableNode* VariableNode = Cast<URigVMVariableNode>(InNode);
-	URigVMDispatchNode* DispatchNode = Cast<URigVMDispatchNode>(InNode);
+	if (CollapseNode)
+	{
+		if (CollapseNode->GetOuter()->IsA<URigVMFunctionLibrary>())
+		{
+			NodeData.bSetupOrphanPinsForThisNode = false;
+		}
+	}
 
-	FScopeLock Lock(&PinPathCoreRedirectorsLock);
 	const FRigVMRegistry& Registry = FRigVMRegistry::Get();
 
-	URigVMGraph* Graph = GetGraph();
-	check(Graph);
+	NodeData.PreviousPinInfos = FRigVMPinInfoArray(InNode);
+	NodeData.PreviousPinHash = GetTypeHash(NodeData.PreviousPinInfos);
+	if (!GenerateNewPinInfos(Registry, InNode, NodeData.PreviousPinInfos, NodeData.NewPinInfos, NodeData.bSetupOrphanPinsForThisNode))
+	{
+		return; // skip this node if no Infos can be generated
+	}
 
-	// step 0/3: update execute pins
-	for(URigVMPin* Pin : InNode->Pins)
+	NodeData.bRequireRecreateLinks = false;
+	NodeData.bRequirePinStates = false;
+
+	GenerateRepopulatePinLists(Registry, NodeData);
+
+	NodesPinData.Add(NodeData);
+
+	// Recurse if a collapse node
+	if (CollapseNode != nullptr)
 	{
-		if(Pin->IsExecuteContext())
+		if (URigVMController* CollapseController = GetControllerForGraph(CollapseNode->GetContainedGraph()))
 		{
-			MakeExecutePin(Pin);
+			TGuardValue<bool> GuardEditGraph(CollapseNode->ContainedGraph->bEditable, true);
+			// need to get a copy of the node array since the following function could remove nodes from the graph
+			// we don't want to remove elements from the array we are iterating over.
+			TArray<URigVMNode*> ContainedNodes = CollapseNode->GetContainedNodes();
+			for (URigVMNode* ContainedNode : ContainedNodes)
+			{
+				CollapseController->GenerateRepopulatePinsNodeData(NodesPinData, ContainedNode, bInFollowCoreRedirectors);
+			}
 		}
 	}
+}
 
-	bool bSetupOrphanPinsForThisNode = bSetupOrphanedPins;
-	if(CollapseNode)
+void URigVMController::OrphanPins(const TArray<FRepopulatePinsNodeData>& NodesPinData)
+{
+	FRigVMControllerCompileBracketScope CompileBracketScope(this);
+	FScopeLock Lock(&PinPathCoreRedirectorsLock);
+
+	for (const FRepopulatePinsNodeData& NodeData : NodesPinData)
 	{
-		if(CollapseNode->GetOuter()->IsA<URigVMFunctionLibrary>())
+		if (NodeData.Node == nullptr)
 		{
-			bSetupOrphanPinsForThisNode = false;
+			ReportError(TEXT("InNode is nullptr orphaning pins."));
+			continue;
+		}
+
+		const FRigVMRegistry& Registry = FRigVMRegistry::Get();
+		
+		// orphan pins
+		for (int32 Index = 0; Index < NodeData.PreviousPinsToOrphan.Num(); Index++)
+		{
+			const FString& PinPath = NodeData.PreviousPinInfos.GetPinPath(NodeData.PreviousPinsToOrphan[Index]);
+			if (URigVMPin* Pin = NodeData.Node->FindPin(PinPath))
+			{
+#if UE_RIGVMCONTROLLER_VERBOSE_REPOPULATE
+				UE_LOG(LogRigVMDeveloper, Display, TEXT("Orphaning pin '%s'."), *PinPath);
+#endif
+				check(Pin->IsRootPin());
+
+				const FString OrphanedName = FString::Printf(TEXT("%s%s"), *URigVMPin::OrphanPinPrefix, *Pin->GetName());
+				if (!NodeData.Node->FindPin(OrphanedName))
+				{
+					URigVMPin* OrphanPin = NewObject<URigVMPin>(NodeData.Node, *OrphanedName);
+					ConfigurePinFromPin(OrphanPin, Pin, true);
+
+					for (URigVMPin* SubPin : Pin->SubPins)
+					{
+						const FString SubPinName = SubPin->GetName();
+						URigVMPin* OrphanedSubPin = NewObject<URigVMPin>(OrphanPin, *SubPinName);
+						ConfigurePinFromPin(OrphanedSubPin, SubPin, true);
+						OrphanPin->SubPins.Add(OrphanedSubPin);
+					}
+
+					NodeData.Node->OrphanedPins.Add(OrphanPin);
+				}
+			}
 		}
 	}
+}
 
-	const FRigVMPinInfoArray PreviousPinInfos(InNode);
-	const uint32 PreviousPinHash = GetTypeHash(PreviousPinInfos);
-	FRigVMPinInfoArray NewPinInfos;
+void URigVMController::RepopulatePins(const TArray<FRepopulatePinsNodeData>& NodesPinData)
+{
+	FRigVMControllerCompileBracketScope CompileBracketScope(this);
+	FScopeLock Lock(&PinPathCoreRedirectorsLock);
+
+	const FRigVMRegistry& Registry = FRigVMRegistry::Get();
+	for (const FRepopulatePinsNodeData& NodeData : NodesPinData)
+	{
+		RepopulatePinsOnNode(Registry, NodeData);
+	}
+}
+
+void URigVMController::RepopulatePinsOnNode(URigVMNode* InNode, bool bFollowCoreRedirectors, bool bSetupOrphanedPins, bool bRecreateLinks)
+{
+	if (InNode == nullptr)
+	{
+		ReportError(TEXT("InNode is nullptr."));
+		return;
+	}
+
+	FRigVMControllerCompileBracketScope CompileBracketScope(this);
+	FScopeLock Lock(&PinPathCoreRedirectorsLock);
+
+	TArray<FRepopulatePinsNodeData> NodesPinData;
+	GenerateRepopulatePinsNodeData(NodesPinData, InNode, bFollowCoreRedirectors, bSetupOrphanedPins, bRecreateLinks);
+	RepopulatePins(NodesPinData);
+}
+
+bool URigVMController::GenerateNewPinInfos(const FRigVMRegistry& Registry, URigVMNode* InNode, const FRigVMPinInfoArray& PreviousPinInfos, FRigVMPinInfoArray& NewPinInfos, const bool bSetupOrphanPinsForThisNode)
+{
+	URigVMUnitNode* UnitNode = Cast<URigVMUnitNode>(InNode);
+	URigVMRerouteNode* RerouteNode = Cast<URigVMRerouteNode>(InNode);
+	URigVMFunctionEntryNode* EntryNode = Cast<URigVMFunctionEntryNode>(InNode);
+	URigVMFunctionReturnNode* ReturnNode = Cast<URigVMFunctionReturnNode>(InNode);
+	URigVMCollapseNode* CollapseNode = Cast<URigVMCollapseNode>(InNode);
+	URigVMFunctionReferenceNode* FunctionRefNode = Cast<URigVMFunctionReferenceNode>(InNode);
+	URigVMVariableNode* VariableNode = Cast<URigVMVariableNode>(InNode);
+	URigVMDispatchNode* DispatchNode = Cast<URigVMDispatchNode>(InNode);
 
 	// step 2/3: clear pins on the node and repopulate the node with new pins
 	if (UnitNode != nullptr)
@@ -14897,7 +14995,7 @@ void URigVMController::RepopulatePinsOnNode(URigVMNode* InNode, bool bFollowCore
 		{
 			// this may be an unresolved template node
 			// in that case there's nothing we can do here
-			return;
+			return false;
 		}
 
 		FString NodeColorMetadata;
@@ -14917,58 +15015,58 @@ void URigVMController::RepopulatePinsOnNode(URigVMNode* InNode, bool bFollowCore
 		{
 			PinTypeMap.Add(Pin->GetFName(), Pin->GetTypeIndex());
 		}
-		
+
 		const FRigVMTemplate* Template = DispatchNode->GetTemplate();
 
 		FRigVMDispatchContext DispatchContext = DispatchNode->GetDispatchContext();
 		auto AddExecutePins = [Template, DispatchNode, &Registry, &DispatchContext, &NewPinInfos, &PreviousPinInfos, this](ERigVMPinDirection InPinDirection)
-		{
-			for (int32 ArgIndex = 0; ArgIndex < Template->NumExecuteArguments(DispatchContext); ArgIndex++)
 			{
-				const FRigVMExecuteArgument* Arg = Template->GetExecuteArgument(ArgIndex, DispatchContext);
-				if (Arg->Direction == InPinDirection)
+				for (int32 ArgIndex = 0; ArgIndex < Template->NumExecuteArguments(DispatchContext); ArgIndex++)
 				{
-					const FRigVMTemplateArgumentType Type = Registry.GetType(Arg->TypeIndex);
-					const TRigVMTypeIndex TypeIndex = Registry.GetTypeIndex(Type);
-
-					FString DefaultValue;
-					if(Registry.IsArrayType(Arg->TypeIndex))
+					const FRigVMExecuteArgument* Arg = Template->GetExecuteArgument(ArgIndex, DispatchContext);
+					if (Arg->Direction == InPinDirection)
 					{
-						if(const FRigVMDispatchFactory* Factory = DispatchNode->GetFactory())
+						const FRigVMTemplateArgumentType Type = Registry.GetType(Arg->TypeIndex);
+						const TRigVMTypeIndex TypeIndex = Registry.GetTypeIndex(Type);
+
+						FString DefaultValue;
+						if (Registry.IsArrayType(Arg->TypeIndex))
 						{
-							DefaultValue = Factory->GetArgumentDefaultValue(Arg->Name, Arg->TypeIndex);
+							if (const FRigVMDispatchFactory* Factory = DispatchNode->GetFactory())
+							{
+								DefaultValue = Factory->GetArgumentDefaultValue(Arg->Name, Arg->TypeIndex);
+							}
 						}
-					}
 
-					(void)NewPinInfos.AddPin(this, INDEX_NONE, Arg->Name, Arg->Direction, TypeIndex, DefaultValue, nullptr, &PreviousPinInfos, true);
+						(void)NewPinInfos.AddPin(this, INDEX_NONE, Arg->Name, Arg->Direction, TypeIndex, DefaultValue, nullptr, &PreviousPinInfos, true);
+					}
 				}
-			}
-		};
+			};
 
 		AddExecutePins(ERigVMPinDirection::IO);
 		AddExecutePins(ERigVMPinDirection::Input);
-		
+
 		for (int32 ArgIndex = 0; ArgIndex < Template->NumArguments(); ArgIndex++)
 		{
 			const FRigVMTemplateArgument* Arg = Template->GetArgument(ArgIndex);
 
 			TRigVMTypeIndex TypeIndex = INDEX_NONE;
-			if(const TRigVMTypeIndex* ExistingTypeIndex = PinTypeMap.Find(Arg->GetName()))
+			if (const TRigVMTypeIndex* ExistingTypeIndex = PinTypeMap.Find(Arg->GetName()))
 			{
 				TypeIndex = *ExistingTypeIndex;
-				if(!Arg->SupportsTypeIndex(TypeIndex))
+				if (!Arg->SupportsTypeIndex(TypeIndex))
 				{
 					TypeIndex = INDEX_NONE;
 				}
 			}
 
-			if(TypeIndex == INDEX_NONE)
+			if (TypeIndex == INDEX_NONE)
 			{
-				if(Arg->IsSingleton())
+				if (Arg->IsSingleton())
 				{
 					TypeIndex = Arg->GetSupportedTypeIndices()[0];
 				}
-				else if(Arg->GetArrayType() == FRigVMTemplateArgument::EArrayType_ArrayValue)
+				else if (Arg->GetArrayType() == FRigVMTemplateArgument::EArrayType_ArrayValue)
 				{
 					TypeIndex = RigVMTypeUtils::TypeIndex::WildCardArray;
 				}
@@ -14982,14 +15080,14 @@ void URigVMController::RepopulatePinsOnNode(URigVMNode* InNode, bool bFollowCore
 			UScriptStruct* ArgumentScriptStruct = nullptr;
 			const uint8* DefaultValueMemory = nullptr;
 
-			if(const URigVMPin* ArgumentPin = DispatchNode->FindPin(Arg->Name.ToString()))
+			if (const URigVMPin* ArgumentPin = DispatchNode->FindPin(Arg->Name.ToString()))
 			{
 				DefaultValue = ArgumentPin->GetDefaultValue();
 				ArgumentScriptStruct = Cast<UScriptStruct>(ArgumentPin->GetCPPTypeObject());
 			}
-			else if(const FRigVMDispatchFactory* Factory = DispatchNode->GetFactory())
+			else if (const FRigVMDispatchFactory* Factory = DispatchNode->GetFactory())
 			{
-				if(Arg->IsSingleton())
+				if (Arg->IsSingleton())
 				{
 					DefaultValue = Factory->GetArgumentDefaultValue(Arg->Name, Arg->GetTypeIndices()[0]);
 					const FRigVMTemplateArgumentType& Type = Registry.GetType(Arg->GetTypeIndices()[0]);
@@ -14998,7 +15096,7 @@ void URigVMController::RepopulatePinsOnNode(URigVMNode* InNode, bool bFollowCore
 			}
 
 			FStructOnScope DefaultValueMemoryScope; // has to be in this scope so that DefaultValueMemory is valid
-			if(ArgumentScriptStruct && !DefaultValue.IsEmpty())
+			if (ArgumentScriptStruct && !DefaultValue.IsEmpty())
 			{
 				DefaultValueMemoryScope = FStructOnScope(ArgumentScriptStruct);
 
@@ -15009,7 +15107,7 @@ void URigVMController::RepopulatePinsOnNode(URigVMNode* InNode, bool bFollowCore
 
 			bool bAddSubPinsForArgument = true;
 #if WITH_EDITOR
-			if(!Template->GetArgumentMetaData(Arg->Name, FRigVMStruct::HideSubPinsMetaName).IsEmpty())
+			if (!Template->GetArgumentMetaData(Arg->Name, FRigVMStruct::HideSubPinsMetaName).IsEmpty())
 			{
 				bAddSubPinsForArgument = false;
 			}
@@ -15023,11 +15121,11 @@ void URigVMController::RepopulatePinsOnNode(URigVMNode* InNode, bool bFollowCore
 	{
 		if (InNode->GetPins().Num() == 0)
 		{
-			return;
+			return false;
 		}
 
 		URigVMPin* ValuePin = nullptr;
-		if(RerouteNode)
+		if (RerouteNode)
 		{
 			ValuePin = RerouteNode->Pins[0];
 		}
@@ -15038,7 +15136,7 @@ void URigVMController::RepopulatePinsOnNode(URigVMNode* InNode, bool bFollowCore
 		check(ValuePin);
 		EnsurePinValidity(ValuePin, false);
 
-		if(VariableNode)
+		if (VariableNode)
 		{
 			// this includes local variables for validation
 			const TArray<FRigVMExternalVariable> ExternalVariables = GetAllVariables(false);
@@ -15048,6 +15146,9 @@ void URigVMController::RepopulatePinsOnNode(URigVMNode* InNode, bool bFollowCore
 			FRigVMExternalVariable Variable;
 			if (VariableNode->IsInputArgument())
 			{
+				URigVMGraph* Graph = GetGraph();
+				check(Graph);
+
 				if (URigVMFunctionEntryNode* GraphEntryNode = Graph->GetEntryNode())
 				{
 					if (URigVMPin* EntryPin = GraphEntryNode->FindPin(VariableDescription.Name.ToString()))
@@ -15057,10 +15158,10 @@ void URigVMController::RepopulatePinsOnNode(URigVMNode* InNode, bool bFollowCore
 				}
 			}
 			else
-			{				
-				for(const FRigVMExternalVariable& ExternalVariable : ExternalVariables)
+			{
+				for (const FRigVMExternalVariable& ExternalVariable : ExternalVariables)
 				{
-					if(ExternalVariable.Name == CurrentExternalVariable.Name)
+					if (ExternalVariable.Name == CurrentExternalVariable.Name)
 					{
 						Variable = ExternalVariable;
 						break;
@@ -15070,14 +15171,14 @@ void URigVMController::RepopulatePinsOnNode(URigVMNode* InNode, bool bFollowCore
 
 			if (Variable.IsValid(true))
 			{
-				if(Variable.TypeName != CurrentExternalVariable.TypeName ||
-				   Variable.TypeObject != CurrentExternalVariable.TypeObject ||
-				   Variable.bIsArray != CurrentExternalVariable.bIsArray)
+				if (Variable.TypeName != CurrentExternalVariable.TypeName ||
+					Variable.TypeObject != CurrentExternalVariable.TypeObject ||
+					Variable.bIsArray != CurrentExternalVariable.bIsArray)
 				{
 					FString CPPType;
 					UObject* CPPTypeObject;
-				
-					if(RigVMTypeUtils::CPPTypeFromExternalVariable(Variable, CPPType, &CPPTypeObject))
+
+					if (RigVMTypeUtils::CPPTypeFromExternalVariable(Variable, CPPType, &CPPTypeObject))
 					{
 						RefreshVariableNode(VariableNode->GetFName(), Variable.Name, CPPType, Variable.TypeObject, false, bSetupOrphanPinsForThisNode);
 					}
@@ -15148,7 +15249,7 @@ void URigVMController::RepopulatePinsOnNode(URigVMNode* InNode, bool bFollowCore
 					}
 				}
 
-				ERigVMPinDirection Direction = bIsEntryNode ? ERigVMPinDirection::Output : ERigVMPinDirection::Input;  
+				ERigVMPinDirection Direction = bIsEntryNode ? ERigVMPinDirection::Output : ERigVMPinDirection::Input;
 				(void)NewPinInfos.AddPin(LibraryPin, INDEX_NONE, Direction);
 			}
 		}
@@ -15165,7 +15266,7 @@ void URigVMController::RepopulatePinsOnNode(URigVMNode* InNode, bool bFollowCore
 	else if (FunctionRefNode)
 	{
 		const FRigVMGraphFunctionHeader& FunctionHeader = FunctionRefNode->GetReferencedFunctionHeader();
-		if(FunctionHeader.IsValid())
+		if (FunctionHeader.IsValid())
 		{
 			NewPinInfos = FRigVMPinInfoArray(FunctionHeader, this, &PreviousPinInfos);
 		}
@@ -15178,33 +15279,14 @@ void URigVMController::RepopulatePinsOnNode(URigVMNode* InNode, bool bFollowCore
 	}
 	else
 	{
-		return;
+		return false;
 	}
 
-	auto RecursivelyRepopulatePinsOnCollapseNode = [this, bFollowCoreRedirectors, bSetupOrphanedPins]
-		(URigVMCollapseNode* InCollapseNode)
-		{
-			if(InCollapseNode)
-			{
-				if(URigVMController* CollapseController = GetControllerForGraph(InCollapseNode->GetContainedGraph()))
-				{
-					TGuardValue<bool> GuardEditGraph(InCollapseNode->ContainedGraph->bEditable, true);
-					// need to get a copy of the node array since the following function could remove nodes from the graph
-					// we don't want to remove elements from the array we are iterating over.
-					TArray<URigVMNode*> ContainedNodes = InCollapseNode->GetContainedNodes();
-					for (URigVMNode* ContainedNode : ContainedNodes)
-					{
-						CollapseController->RepopulatePinsOnNode(ContainedNode, bFollowCoreRedirectors, bSetupOrphanedPins);
-					}
-				}
-			}
-		};
-
 	// make sure the new pin infos contains the decorator pins from the last run
-	for(int32 Index = 0; Index < PreviousPinInfos.Num(); Index++)
+	for (int32 Index = 0; Index < PreviousPinInfos.Num(); Index++)
 	{
-		const FRigVMPinInfo& PreviousPin = PreviousPinInfos[Index]; 
-		if(PreviousPin.bIsDecorator)
+		const FRigVMPinInfo& PreviousPin = PreviousPinInfos[Index];
+		if (PreviousPin.bIsDecorator)
 		{
 			const int32 NewPinIndex = NewPinInfos.AddPin(this, INDEX_NONE, PreviousPin.Name, PreviousPin.Direction, PreviousPin.TypeIndex, PreviousPin.DefaultValue, nullptr, &PreviousPinInfos, true);
 			NewPinInfos[NewPinIndex].bIsDecorator = true;
@@ -15219,53 +15301,44 @@ void URigVMController::RepopulatePinsOnNode(URigVMNode* InNode, bool bFollowCore
 		}
 	}
 
-	// if the nodes match in structure - nothing to do here
-	if(GetTypeHash(NewPinInfos) == PreviousPinHash)
-	{
-		// we at least need to recurse into a collapse node
-		RecursivelyRepopulatePinsOnCollapseNode(CollapseNode);
-		return;
-	}
-	
-#if UE_RIGVMCONTROLLER_VERBOSE_REPOPULATE
-	UE_LOG(LogRigVMDeveloper, Display, TEXT("Repopulating pins on node %s"), *InNode->GetPathName());
-#endif
-	
-	bool bRequireRecreateLinks = false;
-	bool bRequirePinStates = false;
+	return true;
+}
 
-	TArray<int32> NewPinsToAdd, PreviousPinsToRemove, PreviousPinsToOrphan, PreviousPinsToUpdate;
-	for(int32 Index = 0; Index < PreviousPinInfos.Num(); Index++)
+void URigVMController::GenerateRepopulatePinLists(const FRigVMRegistry& Registry, FRepopulatePinsNodeData& NodeData)
+{
+	URigVMNode* InNode = NodeData.Node;
+
+	for (int32 Index = 0; Index < NodeData.PreviousPinInfos.Num(); Index++)
 	{
-		const FString PinPath = PreviousPinInfos.GetPinPath(Index);
-		const int32 NewIndex = NewPinInfos.GetIndexFromPinPath(PinPath);
-		
-		if(NewIndex == INDEX_NONE)
+		const FString PinPath = NodeData.PreviousPinInfos.GetPinPath(Index);
+		const int32 NewIndex = NodeData.NewPinInfos.GetIndexFromPinPath(PinPath);
+
+		if (NewIndex == INDEX_NONE)
 		{
-			const int32 RootIndex = PreviousPinInfos.GetRootIndex(Index);
-			if(PreviousPinInfos[Index].Direction != ERigVMPinDirection::Hidden)
+			const int32 RootIndex = NodeData.PreviousPinInfos.GetRootIndex(Index);
+			if (NodeData.PreviousPinInfos[Index].Direction != ERigVMPinDirection::Hidden)
 			{
-				if(URigVMPin* Pin = InNode->FindPin(PinPath))
+				if (URigVMPin* Pin = InNode->FindPin(PinPath))
 				{
-					if(!Pin->GetLinks().IsEmpty())
+					if (!Pin->GetLinks().IsEmpty())
 					{
-						bRequireRecreateLinks = true;
-						bRequirePinStates = true;
+						NodeData.bRequireRecreateLinks = true;
+						NodeData.bRequirePinStates = true;
 					}
 
-					if (bSetupOrphanPinsForThisNode)
+					if (NodeData.bSetupOrphanPinsForThisNode)
 					{
-						if(!PreviousPinsToOrphan.Contains(RootIndex))
+						if (!NodeData.PreviousPinsToOrphan.Contains(RootIndex))
 						{
 							URigVMPin* RootPin = Pin->GetRootPin();
 
-							if(RootPin->GetSourceLinks(true).Num() > 0 ||
-							   RootPin->GetTargetLinks(true).Num() > 0)
+							if (RootPin->GetSourceLinks(true).Num() > 0 ||
+								RootPin->GetTargetLinks(true).Num() > 0)
 							{
-								PreviousPinsToOrphan.Add(RootIndex);
-							
-								bRequireRecreateLinks = true;
-								bRequirePinStates = true;
+								NodeData.PreviousPinsToOrphan.Add(RootIndex);
+
+								NodeData.bRequireRecreateLinks = true;
+								NodeData.bRequirePinStates = true;
 #if UE_RIGVMCONTROLLER_VERBOSE_REPOPULATE
 								UE_LOG(LogRigVMDeveloper, Display, TEXT("Previously existing pin '%s' needs to be orphaned."), *RootPin->GetPinPath());
 #endif
@@ -15275,42 +15348,42 @@ void URigVMController::RepopulatePinsOnNode(URigVMNode* InNode, bool bFollowCore
 				}
 			}
 
-			if(!PreviousPinsToOrphan.Contains(RootIndex))
+			if (!NodeData.PreviousPinsToOrphan.Contains(RootIndex))
 			{
-				PreviousPinsToRemove.Add(Index);
+				NodeData.PreviousPinsToRemove.Add(Index);
 #if UE_RIGVMCONTROLLER_VERBOSE_REPOPULATE
 				UE_LOG(LogRigVMDeveloper, Display, TEXT("Previously existing pin '%s' is now obsolete."), *PinPath);
 #endif
 			}
 		}
-		else if(GetTypeHash(PreviousPinInfos[Index]) != GetTypeHash(NewPinInfos[NewIndex]))
+		else if (GetTypeHash(NodeData.PreviousPinInfos[Index]) != GetTypeHash(NodeData.NewPinInfos[NewIndex]))
 		{
-			const bool bTypesDiffer = !Registry.CanMatchTypes(PreviousPinInfos[Index].TypeIndex, NewPinInfos[NewIndex].TypeIndex, true);
-			if(PreviousPinInfos[Index].Direction != NewPinInfos[NewIndex].Direction)
+			const bool bTypesDiffer = !Registry.CanMatchTypes(NodeData.PreviousPinInfos[Index].TypeIndex, NodeData.NewPinInfos[NewIndex].TypeIndex, true);
+			if (NodeData.PreviousPinInfos[Index].Direction != NodeData.NewPinInfos[NewIndex].Direction)
 			{
-				bRequireRecreateLinks = true;
+				NodeData.bRequireRecreateLinks = true;
 			}
-			else if(PreviousPinInfos[Index].Direction != ERigVMPinDirection::Hidden)
+			else if (NodeData.PreviousPinInfos[Index].Direction != ERigVMPinDirection::Hidden)
 			{
-				bRequireRecreateLinks |= bTypesDiffer;
+				NodeData.bRequireRecreateLinks |= bTypesDiffer;
 			}
 
-			if(Registry.CanMatchTypes(PreviousPinInfos[Index].TypeIndex, NewPinInfos[NewIndex].TypeIndex, true))
+			if (Registry.CanMatchTypes(NodeData.PreviousPinInfos[Index].TypeIndex, NodeData.NewPinInfos[NewIndex].TypeIndex, true))
 			{
-				PreviousPinsToUpdate.Add(Index);
+				NodeData.PreviousPinsToUpdate.Add(Index);
 			}
 			else
-            {
-            	PreviousPinsToRemove.Add(Index);
-            	NewPinsToAdd.Add(NewIndex);
-            }
-			
+			{
+				NodeData.PreviousPinsToRemove.Add(Index);
+				NodeData.NewPinsToAdd.Add(NewIndex);
+			}
+
 #if UE_RIGVMCONTROLLER_VERBOSE_REPOPULATE
-			const FName& PreviousCPPType = Registry.GetType(PreviousPinInfos[Index].TypeIndex).CPPType;
-			const FName& NewCPPType = Registry.GetType(NewPinInfos[NewIndex].TypeIndex).CPPType;
+			const FName& PreviousCPPType = Registry.GetType(NodeData.PreviousPinInfos[Index].TypeIndex).CPPType;
+			const FName& NewCPPType = Registry.GetType(NodeData.NewPinInfos[NewIndex].TypeIndex).CPPType;
 
-			const FString PreviousDirection = StaticEnum<ERigVMPinDirection>()->GetDisplayNameTextByValue((int64)PreviousPinInfos[Index].Direction).ToString();
-			const FString NewDirection = StaticEnum<ERigVMPinDirection>()->GetDisplayNameTextByValue((int64)NewPinInfos[NewIndex].Direction).ToString();
+			const FString PreviousDirection = StaticEnum<ERigVMPinDirection>()->GetDisplayNameTextByValue((int64)NodeData.PreviousPinInfos[Index].Direction).ToString();
+			const FString NewDirection = StaticEnum<ERigVMPinDirection>()->GetDisplayNameTextByValue((int64)NodeData.NewPinInfos[NewIndex].Direction).ToString();
 
 			UE_LOG(LogRigVMDeveloper, Display,
 				TEXT("Previous pin '%s' (Index %d, %s, %s) differs with new pin (Index %d, %s, %s)."),
@@ -15325,13 +15398,13 @@ void URigVMController::RepopulatePinsOnNode(URigVMNode* InNode, bool bFollowCore
 #endif
 		}
 	}
-	for(int32 Index = 0; Index < NewPinInfos.Num(); Index++)
+	for (int32 Index = 0; Index < NodeData.NewPinInfos.Num(); Index++)
 	{
-		const FString PinPath = NewPinInfos.GetPinPath(Index);
-		const int32 PreviousIndex = PreviousPinInfos.GetIndexFromPinPath(PinPath);
-		if(PreviousIndex == INDEX_NONE)
+		const FString PinPath = NodeData.NewPinInfos.GetPinPath(Index);
+		const int32 PreviousIndex = NodeData.PreviousPinInfos.GetIndexFromPinPath(PinPath);
+		if (PreviousIndex == INDEX_NONE)
 		{
-			NewPinsToAdd.Add(Index);
+			NodeData.NewPinsToAdd.Add(Index);
 #if UE_RIGVMCONTROLLER_VERBOSE_REPOPULATE
 			UE_LOG(LogRigVMDeveloper, Display, TEXT("Newly required pin '%s' needs to be added."), *PinPath);
 #endif
@@ -15339,10 +15412,10 @@ void URigVMController::RepopulatePinsOnNode(URigVMNode* InNode, bool bFollowCore
 		else
 		{
 			// the previous pin exists - but it has been orphaned
-			const int32 PreviousRootIndex = PreviousPinInfos.GetRootIndex(PreviousIndex);
-			if(PreviousPinsToOrphan.Contains(PreviousRootIndex))
+			const int32 PreviousRootIndex = NodeData.PreviousPinInfos.GetRootIndex(PreviousIndex);
+			if (NodeData.PreviousPinsToOrphan.Contains(PreviousRootIndex))
 			{
-				NewPinsToAdd.Add(Index);
+				NodeData.NewPinsToAdd.Add(Index);
 
 #if UE_RIGVMCONTROLLER_VERBOSE_REPOPULATE
 				UE_LOG(LogRigVMDeveloper, Display, TEXT("Orphaned pin '%s' needs to be re-added."), *PinPath);
@@ -15350,10 +15423,40 @@ void URigVMController::RepopulatePinsOnNode(URigVMNode* InNode, bool bFollowCore
 			}
 		}
 	}
+}
+
+void URigVMController::RepopulatePinsOnNode(const FRigVMRegistry& Registry, const FRepopulatePinsNodeData& NodeData)
+{
+	URigVMNode* InNode = NodeData.Node;
+
+	if (InNode == nullptr)
+	{
+		ReportError(TEXT("InNode is nullptr repopulating pins."));
+		return;
+	}
+
+	// step 0/3: update execute pins
+	for (URigVMPin* Pin : InNode->Pins)
+	{
+		if (Pin->IsExecuteContext())
+		{
+			MakeExecutePin(Pin);
+		}
+	}
+
+	// if the nodes does not match in structure repopulate
+	if (GetTypeHash(NodeData.NewPinInfos) == NodeData.PreviousPinHash)
+	{
+		return;
+	}
+
+	URigVMCollapseNode* CollapseNode = Cast<URigVMCollapseNode>(InNode);
+	URigVMFunctionReferenceNode* FunctionRefNode = Cast<URigVMFunctionReferenceNode>(InNode);
+	URigVMDispatchNode* DispatchNode = Cast<URigVMDispatchNode>(InNode);
 
 	// step 1/3: keep a record of the current state of the node's pins
 	TMap<FString, FString> RedirectedPinPaths;
-	if (bFollowCoreRedirectors)
+	if (NodeData.bFollowCoreRedirectors)
 	{
 		RedirectedPinPaths = GetRedirectedPinPaths(InNode);
 	}
@@ -15369,13 +15472,13 @@ void URigVMController::RepopulatePinsOnNode(URigVMNode* InNode, bool bFollowCore
 
 	TMap<FString, FPinState> PinStates;
 	TArray<FLinkedPath> LinkedPaths;
-	
-	if(bRequirePinStates)
+
+	if (NodeData.bRequirePinStates)
 	{
 		PinStates = GetPinStates(InNode);
 	}
 
-	if(bRecreateLinks && bRequireRecreateLinks)
+	if (NodeData.bRecreateLinks && NodeData.bRequireRecreateLinks)
 	{
 #if UE_RIGVMCONTROLLER_VERBOSE_REPOPULATE
 		UE_LOG(LogRigVMDeveloper, Display, TEXT("Detaching links of node %s."), *InNode->GetPathName());
@@ -15391,10 +15494,10 @@ void URigVMController::RepopulatePinsOnNode(URigVMNode* InNode, bool bFollowCore
 #endif
 
 	// orphan pins
-	for(int32 Index = 0; Index < PreviousPinsToOrphan.Num(); Index++)
+	for (int32 Index = 0; Index < NodeData.PreviousPinsToOrphan.Num(); Index++)
 	{
-		const FString& PinPath = PreviousPinInfos.GetPinPath(PreviousPinsToOrphan[Index]);
-		if(URigVMPin* Pin = InNode->FindPin(PinPath))
+		const FString& PinPath = NodeData.PreviousPinInfos.GetPinPath(NodeData.PreviousPinsToOrphan[Index]);
+		if (URigVMPin* Pin = InNode->FindPin(PinPath))
 		{
 #if UE_RIGVMCONTROLLER_VERBOSE_REPOPULATE
 			UE_LOG(LogRigVMDeveloper, Display, TEXT("Orphaning pin '%s'."), *PinPath);
@@ -15402,7 +15505,7 @@ void URigVMController::RepopulatePinsOnNode(URigVMNode* InNode, bool bFollowCore
 			check(Pin->IsRootPin());
 
 			const FString OrphanedName = FString::Printf(TEXT("%s%s"), *URigVMPin::OrphanPinPrefix, *Pin->GetName());
-			if(InNode->FindPin(OrphanedName) == nullptr)
+			if (InNode->FindPin(OrphanedName) == nullptr)
 			{
 				Pin->DisplayName = Pin->GetFName();
 				RenameObject(Pin, *OrphanedName, nullptr);
@@ -15420,10 +15523,10 @@ void URigVMController::RepopulatePinsOnNode(URigVMNode* InNode, bool bFollowCore
 	}
 
 	// remove obsolete pins
-	for(int32 Index = PreviousPinsToRemove.Num() - 1; Index >= 0; Index--)
+	for (int32 Index = NodeData.PreviousPinsToRemove.Num() - 1; Index >= 0; Index--)
 	{
-		const FString& PinPath = PreviousPinInfos.GetPinPath(PreviousPinsToRemove[Index]);
-		if(URigVMPin* Pin = InNode->FindPin(PinPath))
+		const FString& PinPath = NodeData.PreviousPinInfos.GetPinPath(NodeData.PreviousPinsToRemove[Index]);
+		if (URigVMPin* Pin = InNode->FindPin(PinPath))
 		{
 #if UE_RIGVMCONTROLLER_VERBOSE_REPOPULATE
 			UE_LOG(LogRigVMDeveloper, Display, TEXT("Removing pin '%s'."), *PinPath);
@@ -15432,36 +15535,36 @@ void URigVMController::RepopulatePinsOnNode(URigVMNode* InNode, bool bFollowCore
 		}
 	}
 	// add missing pins
-	for(int32 Index = 0; Index < NewPinsToAdd.Num(); Index++)
+	for (int32 Index = 0; Index < NodeData.NewPinsToAdd.Num(); Index++)
 	{
-		const FString& PinPath = NewPinInfos.GetPinPath(NewPinsToAdd[Index]);
+		const FString& PinPath = NodeData.NewPinInfos.GetPinPath(NodeData.NewPinsToAdd[Index]);
 		FString ParentPinPath, PinName;
 		UObject* OuterForPin = InNode;
-		if(URigVMPin::SplitPinPathAtEnd(PinPath, ParentPinPath, PinName))
+		if (URigVMPin::SplitPinPathAtEnd(PinPath, ParentPinPath, PinName))
 		{
 			OuterForPin = InNode->FindPin(ParentPinPath);
 		}
-		
-		CreatePinFromPinInfo(Registry, PreviousPinInfos, NewPinInfos[NewPinsToAdd[Index]], PinPath, OuterForPin);
+
+		CreatePinFromPinInfo(Registry, NodeData.PreviousPinInfos, NodeData.NewPinInfos[NodeData.NewPinsToAdd[Index]], PinPath, OuterForPin);
 #if UE_RIGVMCONTROLLER_VERBOSE_REPOPULATE
 		UE_LOG(LogRigVMDeveloper, Display, TEXT("Adding new pin '%s'."), *PinPath);
 #endif
 	}
 	// update existing pins
-	for(int32 Index = 0; Index < PreviousPinsToUpdate.Num(); Index++)
+	for (int32 Index = 0; Index < NodeData.PreviousPinsToUpdate.Num(); Index++)
 	{
-		const FString& PinPath = PreviousPinInfos.GetPinPath(PreviousPinsToUpdate[Index]);
-		const FRigVMPinInfo* NewPinInfo = NewPinInfos.GetPinFromPinPath(PinPath);
+		const FString& PinPath = NodeData.PreviousPinInfos.GetPinPath(NodeData.PreviousPinsToUpdate[Index]);
+		const FRigVMPinInfo* NewPinInfo = NodeData.NewPinInfos.GetPinFromPinPath(PinPath);
 		check(NewPinInfo);
-		
-		if(URigVMPin* Pin = InNode->FindPin(PinPath))
+
+		if (URigVMPin* Pin = InNode->FindPin(PinPath))
 		{
-			if(Pin->IsExecuteContext())
+			if (Pin->IsExecuteContext())
 			{
 				MakeExecutePin(Pin);
 			}
-			
-			if(Pin->GetTypeIndex() != NewPinInfo->TypeIndex)
+
+			if (Pin->GetTypeIndex() != NewPinInfo->TypeIndex)
 			{
 				// we expect these changes to only apply to float and double pins.
 				check((NewPinInfo->TypeIndex == RigVMTypeUtils::TypeIndex::Float) ||
@@ -15490,14 +15593,14 @@ void URigVMController::RepopulatePinsOnNode(URigVMNode* InNode, bool bFollowCore
 
 	// create a map representing the order of expected pins
 	TMap<FString, TArray<FName>> PinOrder;
-	for(int32 Index = 0; Index < NewPinInfos.Num(); Index++)
+	for (int32 Index = 0; Index < NodeData.NewPinInfos.Num(); Index++)
 	{
-		const FRigVMPinInfo& NewPin = NewPinInfos[Index];
+		const FRigVMPinInfo& NewPin = NodeData.NewPinInfos[Index];
 		FString ParentPinPath;
-		if(NewPin.ParentIndex != INDEX_NONE)
+		if (NewPin.ParentIndex != INDEX_NONE)
 		{
-			ParentPinPath = NewPinInfos.GetPinPath(NewPin.ParentIndex);
-			if(NewPinInfos[NewPin.ParentIndex].bIsArray)
+			ParentPinPath = NodeData.NewPinInfos.GetPinPath(NewPin.ParentIndex);
+			if (NodeData.NewPinInfos[NewPin.ParentIndex].bIsArray)
 			{
 				continue;
 			}
@@ -15508,71 +15611,69 @@ void URigVMController::RepopulatePinsOnNode(URigVMNode* InNode, bool bFollowCore
 	}
 
 	auto SortPinArray = [this](TArray<TObjectPtr<URigVMPin>>& Pins, const TArray<FName>* PinOrder)
-	{
-		if(PinOrder == nullptr)
 		{
-			return;
-		}
+			if (PinOrder == nullptr)
+			{
+				return;
+			}
 
-		if(Pins.Num() < 2)
-		{
-			return;
-		}
+			if (Pins.Num() < 2)
+			{
+				return;
+			}
 
-		const TArray<TObjectPtr<URigVMPin>> PreviousPins = Pins;
+			const TArray<TObjectPtr<URigVMPin>> PreviousPins = Pins;
 
-		if(Pins[0]->IsArrayElement())
-		{
-			Algo::Sort(Pins, [PinOrder](const TObjectPtr<URigVMPin>& A, const TObjectPtr<URigVMPin>& B) -> bool
+			if (Pins[0]->IsArrayElement())
 			{
-				return A->GetFName().Compare(B->GetFName()) < 0;
-			});
-		}
-		else
-		{
-			Algo::Sort(Pins, [PinOrder](const TObjectPtr<URigVMPin>& A, const TObjectPtr<URigVMPin>& B) -> bool
+				Algo::Sort(Pins, [PinOrder](const TObjectPtr<URigVMPin>& A, const TObjectPtr<URigVMPin>& B) -> bool
+					{
+						return A->GetFName().Compare(B->GetFName()) < 0;
+					});
+			}
+			else
 			{
-				const int32 IndexA = PinOrder->Find(A->GetFName());
-				const int32 IndexB = PinOrder->Find(B->GetFName());
-				return IndexA < IndexB;
-			});
-		}
+				Algo::Sort(Pins, [PinOrder](const TObjectPtr<URigVMPin>& A, const TObjectPtr<URigVMPin>& B) -> bool
+					{
+						const int32 IndexA = PinOrder->Find(A->GetFName());
+						const int32 IndexB = PinOrder->Find(B->GetFName());
+						return IndexA < IndexB;
+					});
+			}
 
-		for(int32 Index=0;Index<Pins.Num();Index++)
-		{
-			if(PreviousPins[Index] != Pins[Index])
+			for (int32 Index = 0; Index < Pins.Num(); Index++)
 			{
-				Notify(ERigVMGraphNotifType::PinIndexChanged, Pins[Index]);
+				if (PreviousPins[Index] != Pins[Index])
+				{
+					Notify(ERigVMGraphNotifType::PinIndexChanged, Pins[Index]);
 #if UE_RIGVMCONTROLLER_VERBOSE_REPOPULATE
-				UE_LOG(LogRigVMDeveloper, Display, TEXT("Pin '%s' changed index from %d to %d."),
-					*Pins[Index]->GetPinPath(),
-					PreviousPins.Find(Pins[Index]),
-					Index
-				);
+					UE_LOG(LogRigVMDeveloper, Display, TEXT("Pin '%s' changed index from %d to %d."),
+						*Pins[Index]->GetPinPath(),
+						PreviousPins.Find(Pins[Index]),
+						Index
+					);
 #endif
+				}
 			}
-		}
-	};
+		};
 
 	SortPinArray(InNode->Pins, PinOrder.Find(FString()));
-	for(URigVMPin* Pin : InNode->Pins)
+	for (URigVMPin* Pin : InNode->Pins)
 	{
 		SortPinArray(Pin->SubPins, PinOrder.Find(Pin->GetPinPath()));
 	}
-	
-	if(DispatchNode)
+
+	if (DispatchNode)
 	{
 		ResolveTemplateNodeMetaData(DispatchNode, false);
 	}
-	else if(CollapseNode)
+	else if (CollapseNode)
 	{
 		if (!CollapseNode->GetOuter()->IsA<URigVMFunctionLibrary>())
 		{
 			// no need to notify since the function library graph is invisible anyway
 			RemoveUnusedOrphanedPins(CollapseNode);
 		}
-
-		RecursivelyRepopulatePinsOnCollapseNode(CollapseNode);
 	}
 	else if (FunctionRefNode)
 	{
@@ -15581,7 +15682,7 @@ void URigVMController::RepopulatePinsOnNode(URigVMNode* InNode, bool bFollowCore
 		Notify(ERigVMGraphNotifType::NodeRenamed, FunctionRefNode);
 	}
 
-	if(!PinStates.IsEmpty())
+	if (!PinStates.IsEmpty())
 	{
 #if UE_RIGVMCONTROLLER_VERBOSE_REPOPULATE
 		UE_LOG(LogRigVMDeveloper, Display, TEXT("Reapplying pin-states of node %s..."), *InNode->GetPathName());
@@ -15590,16 +15691,16 @@ void URigVMController::RepopulatePinsOnNode(URigVMNode* InNode, bool bFollowCore
 	}
 
 	InNode->DecoratorRootPinNames.Reset();
-	for(int32 Index = 0; Index < NewPinInfos.Num(); Index++)
+	for (int32 Index = 0; Index < NodeData.NewPinInfos.Num(); Index++)
 	{
-		if(NewPinInfos[Index].bIsDecorator)
+		if (NodeData.NewPinInfos[Index].bIsDecorator)
 		{
-			InNode->DecoratorRootPinNames.Add(NewPinInfos[Index].Name.ToString());
+			InNode->DecoratorRootPinNames.Add(NodeData.NewPinInfos[Index].Name.ToString());
 		}
 	}
 	InNode->UpdateDecoratorRootPinNames();
 
-	if(!LinkedPaths.IsEmpty())
+	if (!LinkedPaths.IsEmpty())
 	{
 #if UE_RIGVMCONTROLLER_VERBOSE_REPOPULATE
 		UE_LOG(LogRigVMDeveloper, Display, TEXT("Reattaching links of node %s..."), *InNode->GetPathName());
diff --git a/Engine/Plugins/Runtime/RigVM/Source/RigVMDeveloper/Public/RigVMModel/RigVMController.h b/Engine/Plugins/Runtime/RigVM/Source/RigVMDeveloper/Public/RigVMModel/RigVMController.h
--- a/Engine/Plugins/Runtime/RigVM/Source/RigVMDeveloper/Public/RigVMModel/RigVMController.h
+++ b/Engine/Plugins/Runtime/RigVM/Source/RigVMDeveloper/Public/RigVMModel/RigVMController.h
@@ -1016,7 +1016,32 @@ public:
 	bool ShouldRedirectPin(UScriptStruct* InOwningStruct, const FString& InOldRelativePinPath, FString& InOutNewRelativePinPath) const;
 	bool ShouldRedirectPin(const FString& InOldPinPath, FString& InOutNewPinPath) const;
 
+
+	struct FRepopulatePinsNodeData
+	{
+		URigVMNode* Node = nullptr;
+		uint32 PreviousPinHash = 0;
+		FRigVMPinInfoArray PreviousPinInfos;
+		FRigVMPinInfoArray NewPinInfos;
+		TArray<int32> NewPinsToAdd;
+		TArray<int32> PreviousPinsToRemove;
+		TArray<int32> PreviousPinsToOrphan;
+		TArray<int32> PreviousPinsToUpdate;
+		bool bSetupOrphanPinsForThisNode = false;
+		bool bFollowCoreRedirectors = false;
+		bool bRequirePinStates = false;
+		bool bRecreateLinks = false;
+		bool bRequireRecreateLinks = false;
+	};
+
+	void GenerateRepopulatePinsNodeData(TArray<FRepopulatePinsNodeData>& NodesPinData, URigVMNode* InNode, bool bInFollowCoreRedirectors = true, bool bInSetupOrphanedPins = false, bool bInRecreateLinks = false);
+	void OrphanPins(const TArray<FRepopulatePinsNodeData>& NodesPinData);
+	void RepopulatePins(const TArray<FRepopulatePinsNodeData>& NodesPinData);
+
 	void RepopulatePinsOnNode(URigVMNode* InNode, bool bFollowCoreRedirectors = true, bool bSetupOrphanedPins = false, bool bRecreateLinks = false);
+	bool GenerateNewPinInfos(const FRigVMRegistry& Registry, URigVMNode* InNode, const FRigVMPinInfoArray& PreviousPinInfos, FRigVMPinInfoArray& NewPinInfos, const bool bSetupOrphanPinsForThisNode);
+	void GenerateRepopulatePinLists(const FRigVMRegistry& Registry, FRepopulatePinsNodeData& NodeData);
+	void RepopulatePinsOnNode(const FRigVMRegistry& Registry, const FRepopulatePinsNodeData& NodeData);
 	void RemovePinsDuringRepopulate(URigVMNode* InNode, TArray<URigVMPin*>& InPins, bool bSetupOrphanedPins);
 
 	// removes any orphan pins that no longer holds a link
