RigVM Editor : Fixed crash when a UUserStruct member is renamed while the rig editor is open

diff --git a/Engine/Plugins/Runtime/RigVM/Source/RigVMDeveloper/Public/RigVMBlueprint.h b/Engine/Plugins/Runtime/RigVM/Source/RigVMDeveloper/Public/RigVMBlueprint.h
--- a/Engine/Plugins/Runtime/RigVM/Source/RigVMDeveloper/Public/RigVMBlueprint.h
+++ b/Engine/Plugins/Runtime/RigVM/Source/RigVMDeveloper/Public/RigVMBlueprint.h
@@ -607,6 +607,10 @@ protected:
 	
 	virtual void HandlePackageDone();
 
+	/** Our currently running rig vm instance */
+	UPROPERTY(transient)
+	TObjectPtr<URigVMHost> EditorHost = nullptr;
+
 private:
 	
 	// RigVMBP, once end-loaded, will inform other RigVM-Dependent systems that Host instances are ready.
diff --git a/Engine/Plugins/Runtime/RigVM/Source/RigVMEditor/Private/Editor/RigVMEditor.cpp b/Engine/Plugins/Runtime/RigVM/Source/RigVMEditor/Private/Editor/RigVMEditor.cpp
--- a/Engine/Plugins/Runtime/RigVM/Source/RigVMEditor/Private/Editor/RigVMEditor.cpp
+++ b/Engine/Plugins/Runtime/RigVM/Source/RigVMEditor/Private/Editor/RigVMEditor.cpp
@@ -49,8 +49,7 @@
 const FName FRigVMEditorModes::RigVMEditorMode = TEXT("RigVM");
 
 FRigVMEditor::FRigVMEditor()
-	: Host(nullptr)
-	, bAnyErrorsLeft(false)
+	: bAnyErrorsLeft(false)
 	, KnownInstructionLimitWarnings()
 	, HaltedAtNode(nullptr)
 	, LastDebuggedHost()
@@ -844,9 +843,12 @@ URigVMBlueprint* FRigVMEditor::GetRigVMBlueprint() const
 
 URigVMHost* FRigVMEditor::GetRigVMHost() const
 {
-	if(Host && IsValid(Host))
+	if (URigVMBlueprint* RigVMBlueprint = GetRigVMBlueprint())
 	{
-		return Host;
+		if(RigVMBlueprint->EditorHost && IsValid(RigVMBlueprint->EditorHost))
+		{
+			return RigVMBlueprint->EditorHost;
+		}
 	}
 	return nullptr;
 }
@@ -2813,10 +2815,13 @@ void FRigVMEditor::ClearDetailsViewWrapperObjects()
 
 void FRigVMEditor::SetHost(URigVMHost* InHost)
 {
-	Host = InHost;
-	if(Host && IsValid(Host))
+	if (URigVMBlueprint* RigVMBlueprint = GetRigVMBlueprint())
 	{
-		OnPreviewHostUpdated().Broadcast(this);
+		RigVMBlueprint->EditorHost = InHost;
+		if(RigVMBlueprint->EditorHost && IsValid(RigVMBlueprint->EditorHost))
+		{
+			OnPreviewHostUpdated().Broadcast(this);
+		}
 	}
 }
 
diff --git a/Engine/Plugins/Runtime/RigVM/Source/RigVMEditor/Private/RigVMEditorModule.cpp b/Engine/Plugins/Runtime/RigVM/Source/RigVMEditor/Private/RigVMEditorModule.cpp
--- a/Engine/Plugins/Runtime/RigVM/Source/RigVMEditor/Private/RigVMEditorModule.cpp
+++ b/Engine/Plugins/Runtime/RigVM/Source/RigVMEditor/Private/RigVMEditorModule.cpp
@@ -1707,6 +1707,8 @@ void FRigVMEditorModule::PostChange(const UUserDefinedStruct* Changed, FStructur
 	}
 
 	TArray<URigVMBlueprint*> BlueprintsToRefresh;
+	TArray<UEdGraph*> EdGraphsToRefresh;
+
 	for (TObjectIterator<URigVMPin> It(RF_Transient | RF_ClassDefaultObject, /** bIncludeDerivedClasses */ true, /** InternalExcludeFlags */ EInternalObjectFlags::Garbage); It; ++It)
 	{
 		const URigVMPin* Pin = *It;
@@ -1738,6 +1740,11 @@ void FRigVMEditorModule::PostChange(const UUserDefinedStruct* Changed, FStructur
 						}
 					}	
 				}
+
+				if (URigVMGraph* RigVMGraph = Pin->GetNode()->GetGraph())
+				{
+					EdGraphsToRefresh.AddUnique(Cast<UEdGraph>(RigVMBlueprint->GetEditorObjectForRigVMGraph(RigVMGraph)));
+				}
 			}
 		}
 	}
@@ -1747,7 +1754,13 @@ void FRigVMEditorModule::PostChange(const UUserDefinedStruct* Changed, FStructur
 		RigVMBlueprint->OnRigVMRegistryChanged();
 		(void)RigVMBlueprint->MarkPackageDirty();
 	}
-	
+
+	// Avoid slate crashing after pins get repopulated
+	for (UEdGraph* Graph : EdGraphsToRefresh)
+	{
+		Graph->NotifyGraphChanged();
+	}
+
 	for (URigVMBlueprint* RigVMBlueprint : BlueprintsToRefresh)
 	{
 		// this should make sure variables in BP are updated with the latest struct object
diff --git a/Engine/Plugins/Runtime/RigVM/Source/RigVMEditor/Public/Editor/RigVMEditor.h b/Engine/Plugins/Runtime/RigVM/Source/RigVMEditor/Public/Editor/RigVMEditor.h
--- a/Engine/Plugins/Runtime/RigVM/Source/RigVMEditor/Public/Editor/RigVMEditor.h
+++ b/Engine/Plugins/Runtime/RigVM/Source/RigVMEditor/Public/Editor/RigVMEditor.h
@@ -303,7 +303,7 @@ private:
 	void OnPIEStopped(bool bSimulation);
 
 	/** Our currently running rig vm instance */
-	TObjectPtr<URigVMHost> Host;
+	//TObjectPtr<URigVMHost> Host;
 
 	FPreviewHostUpdated PreviewHostUpdated;
 
