RigVM: Avoid ensures and checks when pin types are invalid
#rnx

diff --git a/Engine/Plugins/Runtime/RigVM/Source/RigVM/Private/RigVMCore/RigVMRegistry.cpp b/Engine/Plugins/Runtime/RigVM/Source/RigVM/Private/RigVMCore/RigVMRegistry.cpp
--- a/Engine/Plugins/Runtime/RigVM/Source/RigVM/Private/RigVMCore/RigVMRegistry.cpp
+++ b/Engine/Plugins/Runtime/RigVM/Source/RigVM/Private/RigVMCore/RigVMRegistry.cpp
@@ -834,7 +834,7 @@ TRigVMTypeIndex FRigVMRegistry::GetTypeIndex(const FRigVMTemplateArgumentType& I
 
 const FRigVMTemplateArgumentType& FRigVMRegistry::GetType(TRigVMTypeIndex InTypeIndex) const
 {
-	if(ensure(Types.IsValidIndex(InTypeIndex)))
+	if((Types.IsValidIndex(InTypeIndex)))
 	{
 		return Types[InTypeIndex].Type;
 	}
@@ -857,7 +857,7 @@ const FRigVMTemplateArgumentType& FRigVMRegistry::FindTypeFromCPPType(const FStr
 TRigVMTypeIndex FRigVMRegistry::GetTypeIndexFromCPPType(const FString& InCPPType) const
 {
 	TRigVMTypeIndex Result = INDEX_NONE;
-	if(ensure(!InCPPType.IsEmpty()))
+	if(!InCPPType.IsEmpty())
 	{
 		const FName CPPTypeName = *InCPPType;
 
@@ -895,7 +895,7 @@ TRigVMTypeIndex FRigVMRegistry::GetTypeIndexFromCPPType(const FString& InCPPType
 
 bool FRigVMRegistry::IsArrayType(TRigVMTypeIndex InTypeIndex) const
 {
-	if(ensure(Types.IsValidIndex(InTypeIndex)))
+	if((Types.IsValidIndex(InTypeIndex)))
 	{
 		return Types[InTypeIndex].bIsArray;
 	}
diff --git a/Engine/Plugins/Runtime/RigVM/Source/RigVMDeveloper/Private/RigVMModel/RigVMController.cpp b/Engine/Plugins/Runtime/RigVM/Source/RigVMDeveloper/Private/RigVMModel/RigVMController.cpp
--- a/Engine/Plugins/Runtime/RigVM/Source/RigVMDeveloper/Private/RigVMModel/RigVMController.cpp
+++ b/Engine/Plugins/Runtime/RigVM/Source/RigVMDeveloper/Private/RigVMModel/RigVMController.cpp
@@ -560,7 +560,10 @@ FRigVMPinInfoArray::FRigVMPinInfoArray(const URigVMNode* InNode, URigVMControlle
 	for(const URigVMPin* Pin : InNode->GetPins())
 	{
 		const FString DefaultValue = Pin->GetDefaultValue();
-		ensureMsgf(Pin->GetTypeIndex() != INDEX_NONE, TEXT("Invalid pin type %s in %s"), *Pin->GetCPPType(), *InNode->GetPackage()->GetPathName());
+		if (Pin->GetTypeIndex() == INDEX_NONE)
+		{
+			InController->ReportErrorf( TEXT("Invalid pin type %s for %s in %s"), *Pin->GetCPPType(), *Pin->GetPathName(), *InNode->GetPackage()->GetPathName());
+		}
 		(void)AddPin(InController, INDEX_NONE, Pin->GetFName(), Pin->GetDirection(), Pin->GetTypeIndex(), DefaultValue, nullptr, InPreviousPinInfos, bAddSubPins);
 	}
 }
@@ -573,6 +576,10 @@ FRigVMPinInfoArray::FRigVMPinInfoArray(const FRigVMGraphFunctionHeader& Function
 	for(const FRigVMGraphFunctionArgument& FunctionArgument : FunctionHeader.Arguments)
 	{
 		const TRigVMTypeIndex TypeIndex = Registry.GetTypeIndexFromCPPType(FunctionArgument.CPPType.ToString());
+		if (TypeIndex == INDEX_NONE)
+		{
+			InController->ReportErrorf( TEXT("Invalid pin type %s for %s in %s"), *FunctionArgument.CPPType.ToString(), *FunctionHeader.LibraryPointer.LibraryNode.ToString(), *InController->GetPackage()->GetPathName());
+		}
 		ensureMsgf(TypeIndex != INDEX_NONE, TEXT("Invalid pin type %s in %s"), *FunctionArgument.CPPType.ToString(), *InController->GetPackage()->GetPathName());
 		(void)AddPin(InController, INDEX_NONE, FunctionArgument.Name, FunctionArgument.Direction, TypeIndex, FunctionArgument.DefaultValue, nullptr, InPreviousPinInfos, true);
 	}
@@ -614,7 +621,6 @@ int32 FRigVMPinInfoArray::AddPin(FProperty* InProperty, URigVMController* InCont
 int32 FRigVMPinInfoArray::AddPin(URigVMController* InController, int32 InParentIndex, const FName& InName, ERigVMPinDirection InDirection,
 	TRigVMTypeIndex InTypeIndex, const FString& InDefaultValue, const uint8* InDefaultValueMemory, const FRigVMPinInfoArray* InPreviousPinInfos, bool bAddSubPins)
 {
-	ensureMsgf(InTypeIndex != INDEX_NONE, TEXT("Invalid pin type for pin %s in %s"), *InName.ToString(), *InController->GetPackage()->GetPathName());
 	const FRigVMRegistry& Registry = FRigVMRegistry::Get();
 		
 	FRigVMPinInfo Info;
@@ -627,6 +633,10 @@ int32 FRigVMPinInfoArray::AddPin(URigVMController* InController, int32 InParentI
 	Info.CorrectExecuteTypeIndex();
 
 	const int32 Index = Pins.Add(Info);
+	if (InTypeIndex == INDEX_NONE)
+	{
+		InController->ReportErrorf(TEXT("Cannot add pin %s due to invalid type in package %s."), *GetPinPath(Index), *InController->GetPackage()->GetPathName());
+	}
 
 	if(InPreviousPinInfos)
 	{
@@ -18293,9 +18303,12 @@ bool URigVMController::BreakLinkedPaths(const TArray<FLinkedPath>& InLinkedPaths
 			if(Link)
 			{
 				URigVMPin* SourcePin = Link->GetSourcePin();
-				URigVMPin* TargetPin = Link->GetTargetPin();				
+				URigVMPin* TargetPin = Link->GetTargetPin();
 
-				check((SourcePin == nullptr) == (TargetPin == nullptr))
+				if ((SourcePin == nullptr) != (TargetPin == nullptr))
+				{
+					ReportErrorf(TEXT("Cannot break link %s in package %s"), *PinPathRepresentation, *GetPackage()->GetPathName());
+				}
 
 				if(SourcePin)
 				{
diff --git a/Engine/Plugins/Runtime/RigVM/Source/RigVMDeveloper/Private/RigVMModel/RigVMGraph.cpp b/Engine/Plugins/Runtime/RigVM/Source/RigVMDeveloper/Private/RigVMModel/RigVMGraph.cpp
--- a/Engine/Plugins/Runtime/RigVM/Source/RigVMDeveloper/Private/RigVMModel/RigVMGraph.cpp
+++ b/Engine/Plugins/Runtime/RigVM/Source/RigVMDeveloper/Private/RigVMModel/RigVMGraph.cpp
@@ -29,7 +29,11 @@ void URigVMGraph::PostLoad()
 	{
 		URigVMPin* SourcePin = Link->GetSourcePin();
 		URigVMPin* TargetPin = Link->GetTargetPin();
-		check((SourcePin == nullptr) == (TargetPin == nullptr));
+		if ((SourcePin == nullptr) != (TargetPin == nullptr))
+		{
+			static constexpr TCHAR Format[] = TEXT("Cannot add link %s in package %s.");
+			UE_LOG(LogRigVMDeveloper, Warning, Format, *Link->GetPinPathRepresentation(), *GetPackage()->GetPathName());
+		}
 		
 		if(SourcePin)
 		{
