RIgVM: Fix test_sequencer_fk_control_rig not passing
#rnx

diff --git a/Engine/Plugins/Runtime/RigVM/Source/RigVMDeveloper/Private/RigVMModel/RigVMController.cpp b/Engine/Plugins/Runtime/RigVM/Source/RigVMDeveloper/Private/RigVMModel/RigVMController.cpp
--- a/Engine/Plugins/Runtime/RigVM/Source/RigVMDeveloper/Private/RigVMModel/RigVMController.cpp
+++ b/Engine/Plugins/Runtime/RigVM/Source/RigVMDeveloper/Private/RigVMModel/RigVMController.cpp
@@ -14997,6 +14997,15 @@ bool URigVMController::GenerateNewPinInfos(const FRigVMRegistry& Registry, URigV
 	URigVMVariableNode* VariableNode = Cast<URigVMVariableNode>(InNode);
 	URigVMDispatchNode* DispatchNode = Cast<URigVMDispatchNode>(InNode);
 
+	// We need to at least fix the execute pins
+	for (URigVMPin* Pin : InNode->Pins)
+	{
+		if (Pin->IsExecuteContext())
+		{
+			MakeExecutePin(Pin);
+		}
+	}
+
 	// step 2/3: clear pins on the node and repopulate the node with new pins
 	if (UnitNode != nullptr)
 	{
