RIgVM: Move correction of execute pins out of GenerateRepopulatePinsNodeData
#rnx

diff --git a/Engine/Plugins/Runtime/RigVM/Source/RigVMDeveloper/Private/RigVMModel/RigVMClient.cpp b/Engine/Plugins/Runtime/RigVM/Source/RigVMDeveloper/Private/RigVMModel/RigVMClient.cpp
--- a/Engine/Plugins/Runtime/RigVM/Source/RigVMDeveloper/Private/RigVMModel/RigVMClient.cpp
+++ b/Engine/Plugins/Runtime/RigVM/Source/RigVMDeveloper/Private/RigVMModel/RigVMClient.cpp
@@ -899,6 +899,7 @@ FRigVMClientPatchResult FRigVMClient::PatchModelsOnLoad()
 		Result.Merge(Controller->PatchArrayNodesOnLoad());
 		Result.Merge(Controller->PatchReduceArrayFloatDoubleConvertsionsOnLoad());
 		Result.Merge(Controller->PatchInvalidLinksOnWildcards());
+		Result.Merge(Controller->PatchExecutePins());
 
 		if (URigVMCollapseNode* CollapseNode = Model->GetTypedOuter<URigVMCollapseNode>())
 		{
diff --git a/Engine/Plugins/Runtime/RigVM/Source/RigVMDeveloper/Private/RigVMModel/RigVMController.cpp b/Engine/Plugins/Runtime/RigVM/Source/RigVMDeveloper/Private/RigVMModel/RigVMController.cpp
--- a/Engine/Plugins/Runtime/RigVM/Source/RigVMDeveloper/Private/RigVMModel/RigVMController.cpp
+++ b/Engine/Plugins/Runtime/RigVM/Source/RigVMDeveloper/Private/RigVMModel/RigVMController.cpp
@@ -14997,15 +14997,6 @@ bool URigVMController::GenerateNewPinInfos(const FRigVMRegistry& Registry, URigV
 	URigVMVariableNode* VariableNode = Cast<URigVMVariableNode>(InNode);
 	URigVMDispatchNode* DispatchNode = Cast<URigVMDispatchNode>(InNode);
 
-	// We need to at least fix the execute pins
-	for (URigVMPin* Pin : InNode->Pins)
-	{
-		if (Pin->IsExecuteContext())
-		{
-			MakeExecutePin(Pin);
-		}
-	}
-
 	// step 2/3: clear pins on the node and repopulate the node with new pins
 	if (UnitNode != nullptr)
 	{
@@ -17952,7 +17943,7 @@ URigVMPin* URigVMController::MakeExecutePin(URigVMNode* InNode, const FName& InN
 	return ExecutePin;
 }
 
-void URigVMController::MakeExecutePin(URigVMPin* InOutPin)
+bool URigVMController::MakeExecutePin(URigVMPin* InOutPin)
 {
 	if(InOutPin->CPPTypeObject != FRigVMExecuteContext::StaticStruct())
 	{
@@ -17971,7 +17962,22 @@ void URigVMController::MakeExecutePin(URigVMPin* InOutPin)
 			InOutPin->LastKnownTypeIndex = RigVMTypeUtils::TypeIndex::Execute;
 		}
 		InOutPin->LastKnownCPPType = InOutPin->CPPType;
+		return true;
 	}
+	return false;
+}
+
+bool URigVMController::CorrectExecutePinsOnNode(URigVMNode* InOutNode)
+{
+	bool bModified = false;
+	for (URigVMPin* Pin : InOutNode->Pins)
+	{
+		if (Pin->IsExecuteContext())
+		{
+			bModified |= MakeExecutePin(Pin);
+		}
+	}
+	return bModified;
 }
 
 bool URigVMController::AddGraphNode(URigVMNode* InNode, bool bNotify)
@@ -19912,6 +19918,19 @@ FRigVMClientPatchResult URigVMController::PatchFunctionsWithInvalidReturnPaths()
 	return Result;
 }
 
+FRigVMClientPatchResult URigVMController::PatchExecutePins()
+{
+	FRigVMClientPatchResult Result;
+	if (const URigVMGraph* Graph = GetGraph())
+	{
+		for (URigVMNode* Node : Graph->GetNodes())
+		{
+			Result.bChangedContent |= CorrectExecutePinsOnNode(Node);
+		}
+	}
+	return Result;
+}
+
 void URigVMController::PostDuplicateHost(const FString& InOldPathName, const FString& InNewPathName)
 {
 	if (const URigVMGraph* Graph = GetGraph())
diff --git a/Engine/Plugins/Runtime/RigVM/Source/RigVMDeveloper/Public/RigVMModel/RigVMController.h b/Engine/Plugins/Runtime/RigVM/Source/RigVMDeveloper/Public/RigVMModel/RigVMController.h
--- a/Engine/Plugins/Runtime/RigVM/Source/RigVMDeveloper/Public/RigVMModel/RigVMController.h
+++ b/Engine/Plugins/Runtime/RigVM/Source/RigVMDeveloper/Public/RigVMModel/RigVMController.h
@@ -1037,6 +1037,7 @@ public:
 	void GenerateRepopulatePinsNodeData(TArray<FRepopulatePinsNodeData>& NodesPinData, URigVMNode* InNode, bool bInFollowCoreRedirectors = true, bool bInSetupOrphanedPins = false, bool bInRecreateLinks = false);
 	void OrphanPins(const TArray<FRepopulatePinsNodeData>& NodesPinData);
 	void RepopulatePins(const TArray<FRepopulatePinsNodeData>& NodesPinData);
+	bool CorrectExecutePinsOnNode(URigVMNode* InOutNode);
 
 	void RepopulatePinsOnNode(URigVMNode* InNode, bool bFollowCoreRedirectors = true, bool bSetupOrphanedPins = false, bool bRecreateLinks = false);
 	bool GenerateNewPinInfos(const FRigVMRegistry& Registry, URigVMNode* InNode, const FRigVMPinInfoArray& PreviousPinInfos, FRigVMPinInfoArray& NewPinInfos, const bool bSetupOrphanPinsForThisNode);
@@ -1230,7 +1231,7 @@ private:
 	bool RenameObject(UObject* InObjectToRename, const TCHAR* InNewName, UObject* InNewOuter = nullptr) const;
 	void DestroyObject(UObject* InObjectToDestroy) const ;
 	static URigVMPin* MakeExecutePin(URigVMNode* InNode, const FName& InName);
-	static void MakeExecutePin(URigVMPin* InOutPin);
+	static bool MakeExecutePin(URigVMPin* InOutPin);
 	bool AddGraphNode(URigVMNode* InNode, bool bNotify);
 	static void AddNodePin(URigVMNode* InNode, URigVMPin* InPin);
 	static void AddSubPin(URigVMPin* InParentPin, URigVMPin* InPin);
@@ -1337,6 +1338,7 @@ protected:
 	FRigVMClientPatchResult PatchReduceArrayFloatDoubleConvertsionsOnLoad();
 	FRigVMClientPatchResult PatchInvalidLinksOnWildcards();
 	FRigVMClientPatchResult PatchFunctionsWithInvalidReturnPaths();
+	FRigVMClientPatchResult PatchExecutePins();
 
 	// work to do after a duplication of the host asset
 	void PostDuplicateHost(const FString& InOldPathName, const FString& InNewPathName);
