RigVM : Moved ExtendedExecuteContext debug information to Host RigVMDebugInfo

diff --git a/Engine/Plugins/Runtime/RigVM/Source/RigVM/Private/RigVMCore/RigVM.cpp b/Engine/Plugins/Runtime/RigVM/Source/RigVM/Private/RigVMCore/RigVM.cpp
--- a/Engine/Plugins/Runtime/RigVM/Source/RigVM/Private/RigVMCore/RigVM.cpp
+++ b/Engine/Plugins/Runtime/RigVM/Source/RigVM/Private/RigVMCore/RigVM.cpp
@@ -909,30 +909,17 @@ bool URigVM::CanExecuteEntry(const FRigVMExtendedExecuteContext& Context, const
 
 #if WITH_EDITOR
 
-bool URigVM::ResumeExecution(FRigVMExtendedExecuteContext& Context)
+bool URigVM::ResumeExecution(FRigVMExtendedExecuteContext& Context, TArrayView<TRigVMMemoryStorage*> Memory, const FName& InEntryName)
 {
-	Context.HaltedAtBreakpoint.Reset();
-	Context.HaltedAtBreakpointHit = INDEX_NONE;
-	if (Context.DebugInfo)
+	if (FRigVMDebugInfo* RigVMDebugInfo = Context.GetRigVMDebugInfo())
 	{
-		if (const FRigVMBreakpoint& CurrentBreakpoint = Context.DebugInfo->GetCurrentActiveBreakpoint())
-		{
-			Context.DebugInfo->IncrementBreakpointActivationOnHit(CurrentBreakpoint);
-			Context.DebugInfo->SetCurrentActiveBreakpoint(FRigVMBreakpoint());
-			return true;
-		}
+		RigVMDebugInfo->ResumeExecution();
 	}
 
-	return false;
-}
-
-bool URigVM::ResumeExecution(FRigVMExtendedExecuteContext& Context, TArrayView<TRigVMMemoryStorage*> Memory, const FName& InEntryName)
-{
-	ResumeExecution(Context);
 	return Execute(Context, Memory, InEntryName) != ERigVMExecuteResult::Failed;
 }
 
-#endif
+#endif 
 
 const TArray<FRigVMParameter>& URigVM::GetParameters() const
 {
@@ -1263,12 +1250,14 @@ void URigVM::RebuildByteCodeOnLoad()
 #if WITH_EDITOR
 bool URigVM::ShouldHaltAtInstruction(FRigVMExtendedExecuteContext& Context, const FName& InEventName, const uint16 InstructionIndex)
 {
-	if(Context.DebugInfo == nullptr)
+	FRigVMDebugInfo* RigVMDebugInfo = Context.GetRigVMDebugInfo();
+
+	if (RigVMDebugInfo == nullptr)
 	{
 		return false;
 	}
 
-	if(Context.DebugInfo->IsEmpty())
+	if(RigVMDebugInfo->IsEmpty())
 	{
 		return false;
 	}
@@ -1276,45 +1265,45 @@ bool URigVM::ShouldHaltAtInstruction(FRigVMExtendedExecuteContext& Context, cons
 	FRigVMByteCode& ByteCode = GetByteCode();
 	FRigVMExecuteContext& ContextPublicData = Context.GetPublicData<>();
 
-	TArray<FRigVMBreakpoint> BreakpointsAtInstruction = Context.DebugInfo->FindBreakpointsAtInstruction(InstructionIndex);
+	TArray<FRigVMBreakpoint> BreakpointsAtInstruction = RigVMDebugInfo->FindBreakpointsAtInstruction(InstructionIndex);
 	for (FRigVMBreakpoint Breakpoint : BreakpointsAtInstruction)
 	{
-		if (Context.DebugInfo->IsActive(Breakpoint))
+		if (RigVMDebugInfo->IsActive(Breakpoint))
 		{
-			switch (Context.CurrentBreakpointAction)
+			switch (RigVMDebugInfo->GetCurrentBreakpointAction())
 			{
 				case ERigVMBreakpointAction::None:
 				{
 					// Halted at breakpoint. Check if this is a new breakpoint different from the previous halt.
-					if (Context.HaltedAtBreakpoint != Breakpoint ||
-						Context.HaltedAtBreakpointHit != Context.DebugInfo->GetBreakpointHits(Breakpoint))
+					if (RigVMDebugInfo->GetHaltedAtBreakpoint() != Breakpoint ||
+						RigVMDebugInfo->GetHaltedAtBreakpointHit() != RigVMDebugInfo->GetBreakpointHits(Breakpoint))
 					{
-						Context.HaltedAtBreakpoint = Breakpoint;
-						Context.HaltedAtBreakpointHit = Context.DebugInfo->GetBreakpointHits(Breakpoint);
-						Context.DebugInfo->SetCurrentActiveBreakpoint(Breakpoint);
+						RigVMDebugInfo->SetHaltedAtBreakpoint(Breakpoint);
+						RigVMDebugInfo->SetHaltedAtBreakpointHit(RigVMDebugInfo->GetBreakpointHits(Breakpoint));
+						RigVMDebugInfo->SetCurrentActiveBreakpoint(Breakpoint);
 						
 						// We want to keep the callstack up to the node that produced the halt
 						const TArray<UObject*>* FullCallstack = ByteCode.GetCallstackForInstruction(ContextPublicData.InstructionIndex);
 						if (FullCallstack)
 						{
-							Context.DebugInfo->SetCurrentActiveBreakpointCallstack(TArray<UObject*>(FullCallstack->GetData(), FullCallstack->Find((UObject*)Breakpoint.Subject)+1));
+							RigVMDebugInfo->SetCurrentActiveBreakpointCallstack(TArray<UObject*>(FullCallstack->GetData(), FullCallstack->Find((UObject*)Breakpoint.Subject)+1));
 						}
-						Context.ExecutionHalted().Broadcast(ContextPublicData.InstructionIndex, Breakpoint.Subject, InEventName);
+						RigVMDebugInfo->ExecutionHalted().Broadcast(ContextPublicData.InstructionIndex, Breakpoint.Subject, InEventName);
 					}
 					return true;
 				}
 				case ERigVMBreakpointAction::Resume:
 				{
-					Context.CurrentBreakpointAction = ERigVMBreakpointAction::None;
+					RigVMDebugInfo->SetCurrentBreakpointAction(ERigVMBreakpointAction::None);
 
-					if (Context.DebugInfo->IsTemporaryBreakpoint(Breakpoint))
+					if (RigVMDebugInfo->IsTemporaryBreakpoint(Breakpoint))
 					{
-						Context.DebugInfo->RemoveBreakpoint(Breakpoint);
+						RigVMDebugInfo->RemoveBreakpoint(Breakpoint);
 					}
 					else
 					{
-						Context.DebugInfo->IncrementBreakpointActivationOnHit(Breakpoint);
-						Context.DebugInfo->HitBreakpoint(Breakpoint);
+						RigVMDebugInfo->IncrementBreakpointActivationOnHit(Breakpoint);
+						RigVMDebugInfo->HitBreakpoint(Breakpoint);
 					}
 					break;
 				}
@@ -1323,15 +1312,15 @@ bool URigVM::ShouldHaltAtInstruction(FRigVMExtendedExecuteContext& Context, cons
 				case ERigVMBreakpointAction::StepOut:
 				{
 					// If we are stepping, check if we were halted at the current instruction, and remember it 
-					if (!Context.DebugInfo->GetCurrentActiveBreakpoint())
+					if (!RigVMDebugInfo->GetCurrentActiveBreakpoint())
 					{
-						Context.DebugInfo->SetCurrentActiveBreakpoint(Breakpoint);
+						RigVMDebugInfo->SetCurrentActiveBreakpoint(Breakpoint);
 						const TArray<UObject*>* FullCallstack = ByteCode.GetCallstackForInstruction(ContextPublicData.InstructionIndex);
 						
 						// We want to keep the callstack up to the node that produced the halt
 						if (FullCallstack)
 						{
-							Context.DebugInfo->SetCurrentActiveBreakpointCallstack(TArray<UObject*>(FullCallstack->GetData(), FullCallstack->Find((UObject*)Context.DebugInfo->GetCurrentActiveBreakpoint().Subject)+1));
+							RigVMDebugInfo->SetCurrentActiveBreakpointCallstack(TArray<UObject*>(FullCallstack->GetData(), FullCallstack->Find((UObject*)RigVMDebugInfo->GetCurrentActiveBreakpoint().Subject)+1));
 						}
 					}							
 					
@@ -1346,12 +1335,12 @@ bool URigVM::ShouldHaltAtInstruction(FRigVMExtendedExecuteContext& Context, cons
 		}
 		else
 		{
-			Context.DebugInfo->HitBreakpoint(Breakpoint);
+			RigVMDebugInfo->HitBreakpoint(Breakpoint);
 		}
 	}
 
 	// If we are stepping, and the last active breakpoint was set, check if this is the new temporary breakpoint
-	if (Context.CurrentBreakpointAction != ERigVMBreakpointAction::None && Context.DebugInfo->GetCurrentActiveBreakpoint())
+	if (RigVMDebugInfo->GetCurrentBreakpointAction() != ERigVMBreakpointAction::None && RigVMDebugInfo->GetCurrentActiveBreakpoint())
 	{
 		const TArray<UObject*>* CurrentCallstack = ByteCode.GetCallstackForInstruction(ContextPublicData.InstructionIndex);
 		if (CurrentCallstack && !CurrentCallstack->IsEmpty())
@@ -1360,7 +1349,7 @@ bool URigVM::ShouldHaltAtInstruction(FRigVMExtendedExecuteContext& Context, cons
 
 			// Find the first difference in the callstack
 			int32 DifferenceIndex = INDEX_NONE;
-			TArray<UObject*>& PreviousCallstack = Context.DebugInfo->GetCurrentActiveBreakpointCallstack();
+			TArray<UObject*>& PreviousCallstack = RigVMDebugInfo->GetCurrentActiveBreakpointCallstack();
 			for (int32 i=0; i<PreviousCallstack.Num(); ++i)
 			{
 				if (CurrentCallstack->Num() == i)
@@ -1375,14 +1364,14 @@ bool URigVM::ShouldHaltAtInstruction(FRigVMExtendedExecuteContext& Context, cons
 				}
 			}
 
-			if (Context.CurrentBreakpointAction == ERigVMBreakpointAction::StepOver)
+			if (RigVMDebugInfo->GetCurrentBreakpointAction() == ERigVMBreakpointAction::StepOver)
 			{
 				if (DifferenceIndex != INDEX_NONE)
 				{
 					NewBreakpointNode = CurrentCallstack->operator[](DifferenceIndex);
 				}
 			}
-			else if (Context.CurrentBreakpointAction == ERigVMBreakpointAction::StepInto)
+			else if (RigVMDebugInfo->GetCurrentBreakpointAction() == ERigVMBreakpointAction::StepInto)
 			{
 				if (DifferenceIndex == INDEX_NONE)
 				{
@@ -1396,7 +1385,7 @@ bool URigVM::ShouldHaltAtInstruction(FRigVMExtendedExecuteContext& Context, cons
 					NewBreakpointNode = CurrentCallstack->operator[](DifferenceIndex);
 				}
 			}
-			else if (Context.CurrentBreakpointAction == ERigVMBreakpointAction::StepOut)
+			else if (RigVMDebugInfo->GetCurrentBreakpointAction() == ERigVMBreakpointAction::StepOut)
 			{
 				if (DifferenceIndex != INDEX_NONE && DifferenceIndex <= PreviousCallstack.Num() - 2)
                 {
@@ -1407,25 +1396,25 @@ bool URigVM::ShouldHaltAtInstruction(FRigVMExtendedExecuteContext& Context, cons
 			if (NewBreakpointNode)
 			{
 				// Remove or hit previous breakpoint
-				if (Context.DebugInfo->IsTemporaryBreakpoint(Context.DebugInfo->GetCurrentActiveBreakpoint()))
+				if (RigVMDebugInfo->IsTemporaryBreakpoint(RigVMDebugInfo->GetCurrentActiveBreakpoint()))
 				{
-					Context.DebugInfo->RemoveBreakpoint(Context.DebugInfo->GetCurrentActiveBreakpoint());
+					RigVMDebugInfo->RemoveBreakpoint(RigVMDebugInfo->GetCurrentActiveBreakpoint());
 				}
 				else
 				{
-					Context.DebugInfo->IncrementBreakpointActivationOnHit(Context.DebugInfo->GetCurrentActiveBreakpoint());
-					Context.DebugInfo->HitBreakpoint(Context.DebugInfo->GetCurrentActiveBreakpoint());
+					RigVMDebugInfo->IncrementBreakpointActivationOnHit(RigVMDebugInfo->GetCurrentActiveBreakpoint());
+					RigVMDebugInfo->HitBreakpoint(RigVMDebugInfo->GetCurrentActiveBreakpoint());
 				}
 
 				// Create new temporary breakpoint
-				const FRigVMBreakpoint& NewBreakpoint = Context.DebugInfo->AddBreakpoint(ContextPublicData.InstructionIndex, NewBreakpointNode, 0, true);
-				Context.DebugInfo->SetBreakpointHits(NewBreakpoint, GetInstructionVisitedCount(Context, ContextPublicData.InstructionIndex));
-				Context.DebugInfo->SetBreakpointActivationOnHit(NewBreakpoint, GetInstructionVisitedCount(Context, ContextPublicData.InstructionIndex));
-				Context.CurrentBreakpointAction = ERigVMBreakpointAction::None;					
-
-				Context.HaltedAtBreakpoint = NewBreakpoint;
-				Context.HaltedAtBreakpointHit = Context.DebugInfo->GetBreakpointHits(Context.HaltedAtBreakpoint);
-				Context.ExecutionHalted().Broadcast(ContextPublicData.InstructionIndex, NewBreakpointNode, InEventName);
+				const FRigVMBreakpoint& NewBreakpoint = RigVMDebugInfo->AddBreakpoint(ContextPublicData.InstructionIndex, NewBreakpointNode, 0, true);
+				RigVMDebugInfo->SetBreakpointHits(NewBreakpoint, GetInstructionVisitedCount(Context, ContextPublicData.InstructionIndex));
+				RigVMDebugInfo->SetBreakpointActivationOnHit(NewBreakpoint, GetInstructionVisitedCount(Context, ContextPublicData.InstructionIndex));
+				RigVMDebugInfo->SetCurrentBreakpointAction(ERigVMBreakpointAction::None);
+
+				RigVMDebugInfo->SetHaltedAtBreakpoint(NewBreakpoint);
+				RigVMDebugInfo->SetHaltedAtBreakpointHit(RigVMDebugInfo->GetBreakpointHits(RigVMDebugInfo->GetHaltedAtBreakpoint()));
+				RigVMDebugInfo->ExecutionHalted().Broadcast(ContextPublicData.InstructionIndex, NewBreakpointNode, InEventName);
 		
 				return true;
 			}
@@ -1565,9 +1554,12 @@ ERigVMExecuteResult URigVM::Execute(FRigVMExtendedExecuteContext& Context, TArra
 
 	if(bIsRootEntry)
 	{
-		if (Context.GetFirstEntryEventInEventQueue() == NAME_None || Context.GetFirstEntryEventInEventQueue() == InEntryName)
+		if (FRigVMDebugInfo* RigVMDebugInfo = Context.GetRigVMDebugInfo())
 		{
-			SetupInstructionTracking(Context, Instructions.Num());
+			if (RigVMDebugInfo->GetFirstEntryEventInEventQueue() == NAME_None || RigVMDebugInfo->GetFirstEntryEventInEventQueue() == InEntryName)
+			{
+				RigVMDebugInfo->SetupInstructionTracking(Instructions.Num(), Context.GetPublicData<>().RuntimeSettings.bEnableProfiling);
+			}
 		}
 	}
 #endif
@@ -1629,9 +1621,12 @@ ERigVMExecuteResult URigVM::Execute(FRigVMExtendedExecuteContext& Context, TArra
 	FEntryExecuteGuard EntryExecuteGuard(Context.EntriesBeingExecuted, EntryIndexToPush);
 
 #if WITH_EDITOR
-	if (Context.DebugInfo && bIsRootEntry)
+	if (bIsRootEntry)
 	{
-		Context.DebugInfo->StartExecution();
+		if (FRigVMDebugInfo* RigVMDebugInfo = Context.GetRigVMDebugInfo())
+		{
+			RigVMDebugInfo->StartExecution();
+		}
 	}
 #endif
 
@@ -1667,11 +1662,14 @@ ERigVMExecuteResult URigVM::Execute(FRigVMExtendedExecuteContext& Context, TArra
 	{
 		Context.CurrentVMMemory = TArrayView<TRigVMMemoryStorage*>();
 		
-		if (Context.HaltedAtBreakpoint.IsValid() && Context.CurrentExecuteResult != ERigVMExecuteResult::Halted)
+		if (FRigVMDebugInfo* RigVMDebugInfo = Context.GetRigVMDebugInfo())
 		{
-			Context.DebugInfo->SetCurrentActiveBreakpoint(FRigVMBreakpoint());
-			Context.HaltedAtBreakpoint.Reset();
-			Context.ExecutionHalted().Broadcast(INDEX_NONE, nullptr, InEntryName);
+			if (RigVMDebugInfo->GetHaltedAtBreakpoint().IsValid() && Context.CurrentExecuteResult != ERigVMExecuteResult::Halted)
+			{
+				RigVMDebugInfo->SetCurrentActiveBreakpoint(FRigVMBreakpoint());
+				RigVMDebugInfo->GetHaltedAtBreakpoint().Reset();
+				RigVMDebugInfo->ExecutionHalted().Broadcast(INDEX_NONE, nullptr, InEntryName);
+			}
 		}
 	}
 #endif
@@ -1693,6 +1691,7 @@ ERigVMExecuteResult URigVM::ExecuteInstructions(FRigVMExtendedExecuteContext& Co
 
 #if WITH_EDITOR
 	TArray<FName>& FunctionNames = GetFunctionNames();
+	FRigVMDebugInfo* RigVMDebugInfo = Context.GetRigVMDebugInfo();
 #endif
 	
 	while (Instructions.IsValidIndex(ContextPublicData.InstructionIndex))
@@ -1726,8 +1725,12 @@ ERigVMExecuteResult URigVM::ExecuteInstructions(FRigVMExtendedExecuteContext& Co
 #if WITH_EDITOR
 
 		const int32 CurrentInstructionIndex = ContextPublicData.InstructionIndex;
-		Context.InstructionVisitedDuringLastRun[ContextPublicData.InstructionIndex]++;
-		Context.InstructionVisitOrder.Add(ContextPublicData.InstructionIndex);
+
+		if (RigVMDebugInfo != nullptr)
+		{
+			RigVMDebugInfo->SetInstructionVisitedDuringLastRun(ContextPublicData.InstructionIndex);
+			RigVMDebugInfo->AddInstructionIndexToVisitOrder(ContextPublicData.InstructionIndex);
+		}
 	
 #endif
 
@@ -1967,11 +1970,14 @@ ERigVMExecuteResult URigVM::ExecuteInstructions(FRigVMExtendedExecuteContext& Co
 					StopProfiling(Context);
 					Context.ExecutionReachedExit().Broadcast(Context.CurrentEntryName);
 #if WITH_EDITOR					
-					if (Context.HaltedAtBreakpoint.IsValid())
+					if (RigVMDebugInfo != nullptr)
 					{
-						Context.HaltedAtBreakpoint.Reset();
-						Context.DebugInfo->SetCurrentActiveBreakpoint(FRigVMBreakpoint());
-						Context.ExecutionHalted().Broadcast(INDEX_NONE, nullptr, Context.CurrentEntryName);
+						if (RigVMDebugInfo->GetHaltedAtBreakpoint().IsValid())
+						{
+							RigVMDebugInfo->GetHaltedAtBreakpoint().Reset();
+							RigVMDebugInfo->SetCurrentActiveBreakpoint(FRigVMBreakpoint());
+							RigVMDebugInfo->ExecutionHalted().Broadcast(INDEX_NONE, nullptr, Context.CurrentEntryName);
+						}
 					}
 #if UE_RIGVM_DEBUG_EXECUTION
 					if (CVarControlRigDebugAllVMExecutions->GetBool() || ContextPublicData.bDebugExecution)
@@ -2088,21 +2094,24 @@ ERigVMExecuteResult URigVM::ExecuteInstructions(FRigVMExtendedExecuteContext& Co
 		}
 
 #if WITH_EDITOR
-		if(ContextPublicData.RuntimeSettings.bEnableProfiling && !Context.InstructionVisitOrder.IsEmpty())
+		if (RigVMDebugInfo != nullptr)
 		{
-			const uint64 EndCycles = FPlatformTime::Cycles64();
-			const uint64 Cycles = EndCycles - Context.StartCycles;
-			if(Context.InstructionCyclesDuringLastRun[CurrentInstructionIndex] == UINT64_MAX)
-			{
-				Context.InstructionCyclesDuringLastRun[CurrentInstructionIndex] = Cycles;
-			}
-			else
+			if (ContextPublicData.RuntimeSettings.bEnableProfiling && !RigVMDebugInfo->GetInstructionVisitOrder().IsEmpty())
 			{
-				Context.InstructionCyclesDuringLastRun[CurrentInstructionIndex] += Cycles;
-			}
+				const uint64 EndCycles = FPlatformTime::Cycles64();
+				const uint64 Cycles = EndCycles - RigVMDebugInfo->GetStartCycles();
+				if (RigVMDebugInfo->GetInstructionCyclesDuringLastRun(CurrentInstructionIndex) == UINT64_MAX)
+				{
+					RigVMDebugInfo->SetInstructionCyclesDuringLastRun(CurrentInstructionIndex, Cycles);
+				}
+				else
+				{
+					RigVMDebugInfo->AddInstructionCyclesDuringLastRun(CurrentInstructionIndex, Cycles);
+				}
 
-			Context.StartCycles = EndCycles;
-			Context.OverallCycles += Cycles;
+				RigVMDebugInfo->SetStartCycles(EndCycles);
+				RigVMDebugInfo->AddOverallCycles(Cycles);
+			}
 		}
 
 #if UE_RIGVM_DEBUG_EXECUTION
@@ -2163,9 +2172,19 @@ bool URigVM::Execute(FRigVMExtendedExecuteContext& Context, const FName& InEntry
 
 ERigVMExecuteResult URigVM::ExecuteBranch(FRigVMExtendedExecuteContext& Context, const FRigVMBranchInfo& InBranchToRun)
 {
+#if WITH_EDITOR
+	const double LastExecutionMicroSecondsGuard = Context.GetRigVMDebugInfo() ? Context.GetRigVMDebugInfo()->GetLastExecutionMicroSeconds() : 0.0;
+	ON_SCOPE_EXIT
+	{
+		if (Context.GetRigVMDebugInfo())
+		{
+			Context.GetRigVMDebugInfo()->SetLastExecutionMicroSeconds(LastExecutionMicroSecondsGuard);
+		}
+	};
+#endif
+
 	// Maintain all settings on the context - to be reset once the branch has executed. 
 	FRigVMExecuteContext& PublicContext = Context.GetPublicData<>();
-	TGuardValue<double> LastExecutionMicroSecondsGuard(Context.LastExecutionMicroSeconds, Context.LastExecutionMicroSeconds);
 	TGuardValue<uint32> NumExecutionsGuard(PublicContext.NumExecutions, PublicContext.NumExecutions);
 	TGuardValue<ERigVMExecuteResult> CurrentExecuteResultGuard(Context.CurrentExecuteResult, Context.CurrentExecuteResult);
 	TGuardValue<FName> CurrentEntryNameGuard(Context.CurrentEntryName, Context.CurrentEntryName);
@@ -2652,18 +2671,9 @@ void URigVM::RefreshExternalPropertyPaths()
 void URigVM::SetupInstructionTracking(FRigVMExtendedExecuteContext& Context, int32 InInstructionCount)
 {
 #if WITH_EDITOR
-	Context.InstructionVisitedDuringLastRun.Reset();
-	Context.InstructionVisitOrder.Reset();
-	Context.InstructionVisitedDuringLastRun.SetNumZeroed(InInstructionCount);
-	Context.InstructionCyclesDuringLastRun.Reset();
-
-	if(Context.GetPublicData<>().RuntimeSettings.bEnableProfiling)
+	if (FRigVMDebugInfo* RigVMDebugInfo = Context.GetRigVMDebugInfo())
 	{
-		Context.InstructionCyclesDuringLastRun.SetNumUninitialized(InInstructionCount);
-		for(int32 DurationIndex=0;DurationIndex<Context.InstructionCyclesDuringLastRun.Num();DurationIndex++)
-		{
-			Context.InstructionCyclesDuringLastRun[DurationIndex] = UINT64_MAX;
-		}
+		RigVMDebugInfo->SetupInstructionTracking(InInstructionCount, Context.GetPublicData<>().RuntimeSettings.bEnableProfiling);
 	}
 #endif
 }
@@ -2671,10 +2681,9 @@ void URigVM::SetupInstructionTracking(FRigVMExtendedExecuteContext& Context, int
 void URigVM::StartProfiling(FRigVMExtendedExecuteContext& Context)
 {
 #if WITH_EDITOR
-	Context.OverallCycles = Context.StartCycles = 0;
-	if(Context.GetPublicData<>().RuntimeSettings.bEnableProfiling)
+	if (FRigVMDebugInfo* RigVMDebugInfo = Context.GetRigVMDebugInfo())
 	{
-		Context.StartCycles = FPlatformTime::Cycles64();
+		RigVMDebugInfo->StartProfiling(Context.GetPublicData<>().RuntimeSettings.bEnableProfiling);
 	}
 #endif
 }
@@ -2682,7 +2691,9 @@ void URigVM::StartProfiling(FRigVMExtendedExecuteContext& Context)
 void URigVM::StopProfiling(FRigVMExtendedExecuteContext& Context)
 {
 #if WITH_EDITOR
-	const uint64 Cycles = Context.OverallCycles > 0 ? Context.OverallCycles : (FPlatformTime::Cycles64() - Context.StartCycles); 
-	Context.LastExecutionMicroSeconds = Cycles * FPlatformTime::GetSecondsPerCycle() * 1000.0 * 1000.0;
+	if (FRigVMDebugInfo* RigVMDebugInfo = Context.GetRigVMDebugInfo())
+	{
+		RigVMDebugInfo->StopProfiling();
+	}
 #endif
 }
diff --git a/Engine/Plugins/Runtime/RigVM/Source/RigVM/Private/RigVMCore/RigVMDebugInfo.cpp b/Engine/Plugins/Runtime/RigVM/Source/RigVM/Private/RigVMCore/RigVMDebugInfo.cpp
--- a/Engine/Plugins/Runtime/RigVM/Source/RigVM/Private/RigVMCore/RigVMDebugInfo.cpp
+++ b/Engine/Plugins/Runtime/RigVM/Source/RigVM/Private/RigVMCore/RigVMDebugInfo.cpp
@@ -228,3 +228,53 @@ uint16 FRigVMDebugInfo::GetBreakpointHits(const FRigVMBreakpoint& InBreakpoint)
 	return 0;
 }
 
+void FRigVMDebugInfo::SetupInstructionTracking(int32 InInstructionCount, bool bEnableProfiling)
+{
+#if WITH_EDITOR
+
+	ResetInstructionVisitedDuringLastRun(InInstructionCount);
+	SetNumInstructionVisitedDuringLastRunZeroed(InInstructionCount);
+	ResetInstructionVisitOrder(InInstructionCount);
+	ResetInstructionCyclesDuringLastRun(InInstructionCount);
+
+	if (bEnableProfiling)
+	{
+		InitInstructionCyclesDuringLastRunValues(InInstructionCount, UINT64_MAX);
+	}
+#endif
+}
+
+bool FRigVMDebugInfo::ResumeExecution()
+{
+	HaltedAtBreakpoint.Reset();
+	HaltedAtBreakpointHit = INDEX_NONE;
+
+	if (const FRigVMBreakpoint& CurrentBreakpoint = GetCurrentActiveBreakpoint())
+	{
+		IncrementBreakpointActivationOnHit(CurrentBreakpoint);
+		SetCurrentActiveBreakpoint(FRigVMBreakpoint());
+		return true;
+	}
+
+	return false;
+}
+
+void FRigVMDebugInfo::StartProfiling(bool bEnableProfiling)
+{
+#if WITH_EDITOR
+		SetStartCycles(0);
+		SetOverallCycles(0);
+		if (bEnableProfiling)
+		{
+			SetStartCycles(FPlatformTime::Cycles64());
+		}
+#endif
+}
+
+void FRigVMDebugInfo::StopProfiling()
+{
+#if WITH_EDITOR
+	const uint64 Cycles = GetOverallCycles() > 0 ? GetOverallCycles() : (FPlatformTime::Cycles64() - GetStartCycles());
+	LastExecutionMicroSeconds = Cycles * FPlatformTime::GetSecondsPerCycle() * 1000.0 * 1000.0;
+#endif
+}
diff --git a/Engine/Plugins/Runtime/RigVM/Source/RigVM/Private/RigVMCore/RigVMExecuteContext.cpp b/Engine/Plugins/Runtime/RigVM/Source/RigVM/Private/RigVMCore/RigVMExecuteContext.cpp
--- a/Engine/Plugins/Runtime/RigVM/Source/RigVM/Private/RigVMCore/RigVMExecuteContext.cpp
+++ b/Engine/Plugins/Runtime/RigVM/Source/RigVM/Private/RigVMCore/RigVMExecuteContext.cpp
@@ -58,6 +58,7 @@ FRigVMExtendedExecuteContext::~FRigVMExtendedExecuteContext()
 }
 PRAGMA_ENABLE_DEPRECATION_WARNINGS
 
+// --- FRigVMExtendedExecuteContext ---
 
 void FRigVMExtendedExecuteContext::Reset()
 {
@@ -76,11 +77,6 @@ void FRigVMExtendedExecuteContext::Reset()
 
 	CachedMemoryHandles.Reset();
 
-#if WITH_EDITOR
-	DebugInfo = nullptr;
-	ExecutionHalted().Clear();
-#endif // WITH_EDITOR
-
 	LazyBranchInstanceData.Reset();
 	ExternalVariableRuntimeData.Reset();
 
@@ -98,17 +94,11 @@ void FRigVMExtendedExecuteContext::Reset()
 	bCurrentlyRunningRootEntry = false;
 
 #if WITH_EDITOR
-	HaltedAtBreakpoint = FRigVMBreakpoint();
-	HaltedAtBreakpointHit = INDEX_NONE;
-	CurrentBreakpointAction = ERigVMBreakpointAction::None;
-
-	InstructionVisitedDuringLastRun.Reset();
-	InstructionCyclesDuringLastRun.Reset();
-	InstructionVisitOrder.Reset();
-	FirstEntryEventInQueue = NAME_None;
-
-	StartCycles = 0;
-	OverallCycles = 0;
+	if (DebugInfo != nullptr)
+	{
+		DebugInfo->Reset();
+		DebugInfo = nullptr;
+	}
 #endif // WITH_EDITOR
 }
 
diff --git a/Engine/Plugins/Runtime/RigVM/Source/RigVM/Private/RigVMHost.cpp b/Engine/Plugins/Runtime/RigVM/Source/RigVM/Private/RigVMHost.cpp
--- a/Engine/Plugins/Runtime/RigVM/Source/RigVM/Private/RigVMHost.cpp
+++ b/Engine/Plugins/Runtime/RigVM/Source/RigVM/Private/RigVMHost.cpp
@@ -173,8 +173,12 @@ void URigVMHost::PreSave(FObjectPreSaveContext SaveContext)
 void URigVMHost::BeginDestroy()
 {
 	Super::BeginDestroy();
+
 	InitializedEvent.Clear();
 	ExecutedEvent.Clear();
+#if WITH_EDITOR
+	DebugInfo.Reset();
+#endif
 
 	GetExtendedExecuteContext().ExecutionReachedExit().RemoveAll(this);
 
@@ -500,17 +504,6 @@ bool URigVMHost::Execute(const FName& InEventName)
 	PublicContext.bDebugExecution = bDebugExecutionEnabled;
 #endif
 
-	if (VM)
-	{
-#if WITH_EDITOR
-		// default to always clear data after each execution
-		// only set a valid first entry event later when execution
-		// has passed the initialization stage and there are multiple events present in one evaluation
-		// first entry event is used to determined when to clear data during an evaluation
-		VM->SetFirstEntryEventInEventQueue(GetExtendedExecuteContext(), NAME_None);
-#endif
-	}
-
 #if WITH_EDITOR
 	if (IsInDebugMode())
 	{
@@ -534,6 +527,17 @@ bool URigVMHost::Execute(const FName& InEventName)
 	}
 #endif
 
+	if (VM)
+	{
+#if WITH_EDITOR
+		// default to always clear data after each execution
+		// only set a valid first entry event later when execution
+		// has passed the initialization stage and there are multiple events present in one evaluation
+		// first entry event is used to determined when to clear data during an evaluation
+		VM->SetFirstEntryEventInEventQueue(GetExtendedExecuteContext(), NAME_None);
+#endif
+	}
+
 	// setup the draw interface for debug drawing
 	if(!bIsEventInQueue || bIsEventFirstInQueue)
 	{
diff --git a/Engine/Plugins/Runtime/RigVM/Source/RigVM/Public/RigVMCore/RigVM.h b/Engine/Plugins/Runtime/RigVM/Source/RigVM/Public/RigVMCore/RigVM.h
--- a/Engine/Plugins/Runtime/RigVM/Source/RigVM/Public/RigVMCore/RigVM.h
+++ b/Engine/Plugins/Runtime/RigVM/Source/RigVM/Public/RigVMCore/RigVM.h
@@ -420,9 +420,9 @@ public:
 	// Returns the number of times an instruction has been hit
 	int32 GetInstructionVisitedCount(const FRigVMExtendedExecuteContext& Context, int32 InIndex) const
 	{
-		if (Context.InstructionVisitedDuringLastRun.IsValidIndex(InIndex))
+		if (const FRigVMDebugInfo* RigVMDebugInfo = Context.GetRigVMDebugInfo())
 		{
-			return Context.InstructionVisitedDuringLastRun[InIndex];
+			return RigVMDebugInfo->GetInstructionVisitedCountDuringLastRun(InIndex);
 		}
 		return 0;
 	}
@@ -435,9 +435,9 @@ public:
 	// If there is no information available this function returns UINT64_MAX.
 	uint64 GetInstructionCycles(const FRigVMExtendedExecuteContext& Context, int32 InIndex) const
 	{
-		if (Context.InstructionCyclesDuringLastRun.IsValidIndex(InIndex))
+		if (const FRigVMDebugInfo* RigVMDebugInfo = Context.GetRigVMDebugInfo())
 		{
-			return Context.InstructionCyclesDuringLastRun[InIndex];
+			return RigVMDebugInfo->GetInstructionCyclesDuringLastRun(InIndex);
 		}
 		return UINT64_MAX;
 	}
@@ -462,12 +462,26 @@ public:
 	const TArray<int32> GetInstructionVisitOrder() const { return TArray<int32>(); }
 
 	// Returns the order of all instructions during the last run
-	const TArray<int32> GetInstructionVisitOrder(const FRigVMExtendedExecuteContext& Context) const { return Context.InstructionVisitOrder; }
+	const TArray<int32> GetInstructionVisitOrder(const FRigVMExtendedExecuteContext& Context) const
+	{
+		if (const FRigVMDebugInfo* RigVMDebugInfo = Context.GetRigVMDebugInfo())
+		{
+			return RigVMDebugInfo->GetInstructionVisitOrder();
+		}
+
+		return TArray<int32>();
+	}
 
 	UE_DEPRECATED(5.3, "Please, use SetFirstEntryEventInEventQueue with Context param")
 	const void SetFirstEntryEventInEventQueue(const FName& InFirstEventName) {}
 
-	const void SetFirstEntryEventInEventQueue(FRigVMExtendedExecuteContext& Context, const FName& InFirstEventName) { Context.SetFirstEntryEventInEventQueue(InFirstEventName); }
+	const void SetFirstEntryEventInEventQueue(FRigVMExtendedExecuteContext& Context, const FName& InFirstEventName) 
+	{ 
+		if (FRigVMDebugInfo* RigVMDebugInfo = Context.GetRigVMDebugInfo())
+		{
+			RigVMDebugInfo->SetFirstEntryEventInEventQueue(InFirstEventName);
+		}
+	}
 
 	UE_DEPRECATED(5.3, "Please, use ResumeExecution with Context param")
 	bool ResumeExecution() { return false; }
@@ -476,9 +490,12 @@ public:
 	UE_DEPRECATED(5.4, "Please, use ResumeExecution with TRigVMMemoryStorage param")
 	bool ResumeExecution(FRigVMExtendedExecuteContext& Context, TArrayView<TRigVMMemoryStorageDeprecatedType*> Memory, const FName& InEntryName = NAME_None) { return false; }
 
-	bool ResumeExecution(FRigVMExtendedExecuteContext& Context);
+	UE_DEPRECATED(5.4, "Please, use ResumeExecution in the Context RigVMDebugInfo")
+	bool ResumeExecution(FRigVMExtendedExecuteContext& Context) { return false; }
+
 	bool ResumeExecution(FRigVMExtendedExecuteContext& Context, TArrayView<TRigVMMemoryStorage*> Memory, const FName& InEntryName = NAME_None);
-#endif
+
+#endif // WITH_EDITOR
 
 	// Returns the parameters of the VM
 	const TArray<FRigVMParameter>& GetParameters() const;
diff --git a/Engine/Plugins/Runtime/RigVM/Source/RigVM/Public/RigVMCore/RigVMDebugInfo.h b/Engine/Plugins/Runtime/RigVM/Source/RigVM/Public/RigVMCore/RigVMDebugInfo.h
--- a/Engine/Plugins/Runtime/RigVM/Source/RigVM/Public/RigVMCore/RigVMDebugInfo.h
+++ b/Engine/Plugins/Runtime/RigVM/Source/RigVM/Public/RigVMCore/RigVMDebugInfo.h
@@ -111,7 +111,22 @@ struct RIGVM_API FRigVMDebugInfo
 	void Reset()
 	{
 		Breakpoints.Reset();
+
 		// Do not remove state
+		OnExecutionHalted.Clear();
+
+		HaltedAtBreakpoint = FRigVMBreakpoint();
+		HaltedAtBreakpointHit = INDEX_NONE;
+		CurrentBreakpointAction = ERigVMBreakpointAction::None;
+
+		InstructionVisitedDuringLastRun.Reset();
+		InstructionCyclesDuringLastRun.Reset();
+		InstructionVisitOrder.Reset();
+		FirstEntryEventInQueue = NAME_None;
+
+		StartCycles = 0;
+		OverallCycles = 0;
+
 	}
 
 	bool IsEmpty() const
@@ -175,9 +190,68 @@ struct RIGVM_API FRigVMDebugInfo
 		}
 	}
 
-	TArray<UObject*>& GetCurrentActiveBreakpointCallstack() { return CurrentActiveBreakpointCallstack; }
+	inline TArray<UObject*>& GetCurrentActiveBreakpointCallstack() { return CurrentActiveBreakpointCallstack; }
+
+	inline void SetCurrentActiveBreakpointCallstack(TArray<UObject*> Callstack) { CurrentActiveBreakpointCallstack = Callstack; }
+
+	inline uint64 GetStartCycles() const {	return StartCycles;	}
+	inline void SetStartCycles(uint64 InStartCycles) { StartCycles = InStartCycles; }
+
+	inline uint64 GetOverallCycles() const { return OverallCycles; }
+	inline void SetOverallCycles(uint64 Cycles) { OverallCycles = Cycles; }
+	inline void AddOverallCycles(uint64 Cycles) { OverallCycles += Cycles; }
+
+	inline FRigVMBreakpoint& GetHaltedAtBreakpoint() { return HaltedAtBreakpoint; }
+	inline const FRigVMBreakpoint& GetHaltedAtBreakpoint() const { return HaltedAtBreakpoint; }
+	inline void SetHaltedAtBreakpoint(const FRigVMBreakpoint& InHaltedAtBreakpoint) { HaltedAtBreakpoint = InHaltedAtBreakpoint; }
+
+	inline int32 GetHaltedAtBreakpointHit() const { return HaltedAtBreakpointHit; }
+	inline void SetHaltedAtBreakpointHit(int32 InHaltedAtBreakpointHit) { HaltedAtBreakpointHit = InHaltedAtBreakpointHit; }
+
+	inline ERigVMBreakpointAction GetCurrentBreakpointAction() const { return CurrentBreakpointAction; }
+	inline void SetCurrentBreakpointAction(ERigVMBreakpointAction InCurrentBreakpointAction) { CurrentBreakpointAction = InCurrentBreakpointAction; }
+
+	inline void ResetInstructionVisitedDuringLastRun(int32 NewSize = 0) { InstructionVisitedDuringLastRun.Reset(NewSize); }
+	inline void SetNumInstructionVisitedDuringLastRunZeroed(int32 Num) { InstructionVisitedDuringLastRun.SetNumZeroed(Num); }
+	inline void SetInstructionVisitedDuringLastRun(int32 InstructionIndex) { InstructionVisitedDuringLastRun[InstructionIndex]++; }
+	inline int32 GetInstructionVisitedCountDuringLastRun(int32 InstructionIndex) const { return InstructionVisitedDuringLastRun.IsValidIndex(InstructionIndex) ? InstructionVisitedDuringLastRun[InstructionIndex] : 0; }
+	inline const TArray<int32>& GetInstructionVisitedCountDuringLastRun() const { return InstructionVisitedDuringLastRun; }
+
+	inline void ResetInstructionCyclesDuringLastRun(int32 NewSize = 0) { InstructionCyclesDuringLastRun.Reset(NewSize); }
+	inline uint64 GetInstructionCyclesDuringLastRun(int32 InstructionIndex) const { return InstructionCyclesDuringLastRun.IsValidIndex(InstructionIndex) ? InstructionCyclesDuringLastRun[InstructionIndex] : UINT64_MAX; }
+	inline void SetInstructionCyclesDuringLastRun(int32 InstructionIndex, uint64 CyclesDuringLastRun) { InstructionCyclesDuringLastRun[InstructionIndex] = CyclesDuringLastRun; }
+	inline void AddInstructionCyclesDuringLastRun(int32 InstructionIndex, uint64 CyclesDuringLastRun) { InstructionCyclesDuringLastRun[InstructionIndex] += CyclesDuringLastRun; }
+	inline void InitInstructionCyclesDuringLastRunValues(int32 NewSize, uint64 DefaultValue)
+	{
+		InstructionCyclesDuringLastRun.SetNumUninitialized(NewSize);
+		for (uint64& Value : InstructionCyclesDuringLastRun)
+		{
+			Value = DefaultValue;
+		}
+	}
+
+	
+	inline void ResetInstructionVisitOrder(int32 NewSize = 0) { InstructionVisitOrder.Reset(NewSize); }
+	inline void AddInstructionIndexToVisitOrder(int32 InstructionIndex) { InstructionVisitOrder.Add(InstructionIndex); }
+	inline const TArray<int32>& GetInstructionVisitOrder() const { return InstructionVisitOrder; }
+
+	inline const void SetFirstEntryEventInEventQueue(const FName& InFirstEventName) { FirstEntryEventInQueue = InFirstEventName; }
+	inline const FName& GetFirstEntryEventInEventQueue() const { return FirstEntryEventInQueue; }
+
+	void SetupInstructionTracking(int32 InInstructionCount, bool bEnableProfiling);
+	bool ResumeExecution();
+
+	DECLARE_EVENT_ThreeParams(URigVM, FExecutionHaltedEvent, int32, UObject*, const FName&);
+	inline FExecutionHaltedEvent& ExecutionHalted()
+	{
+		return OnExecutionHalted;
+	}
+
+	double GetLastExecutionMicroSeconds() const { return LastExecutionMicroSeconds; }
+	void SetLastExecutionMicroSeconds(double InLastExecutionMicroSeconds) { LastExecutionMicroSeconds = InLastExecutionMicroSeconds; }
 
-	void SetCurrentActiveBreakpointCallstack(TArray<UObject*> Callstack) { CurrentActiveBreakpointCallstack = Callstack; }
+	void StartProfiling(bool bEnableProfiling);
+	void StopProfiling();
 
 private:
 	TArray<FRigVMBreakpoint> Breakpoints;
@@ -187,4 +261,26 @@ private:
 
 	FGuid CurrentActiveBreakpoint;
 	TArray<UObject*> CurrentActiveBreakpointCallstack;
+
+	// --- ---
+
+	FRigVMBreakpoint HaltedAtBreakpoint;
+	int32 HaltedAtBreakpointHit = INDEX_NONE;
+	ERigVMBreakpointAction CurrentBreakpointAction = ERigVMBreakpointAction::None;
+
+	// stores the number of times each instruction was visited
+	TArray<int32> InstructionVisitedDuringLastRun;
+	TArray<uint64> InstructionCyclesDuringLastRun;
+	TArray<int32> InstructionVisitOrder;
+
+	// A RigVMHost can run multiple events per evaluation, such as the Backward&Forward Solve Mode,
+	// store the first event such that we know when to reset data for a new round of rig evaluation
+	FName FirstEntryEventInQueue = NAME_None;
+
+	uint64 StartCycles = 0;
+	uint64 OverallCycles = 0;
+
+	FExecutionHaltedEvent OnExecutionHalted;
+
+	double LastExecutionMicroSeconds = 0.0;
 };
diff --git a/Engine/Plugins/Runtime/RigVM/Source/RigVM/Public/RigVMCore/RigVMExecuteContext.h b/Engine/Plugins/Runtime/RigVM/Source/RigVM/Public/RigVMCore/RigVMExecuteContext.h
--- a/Engine/Plugins/Runtime/RigVM/Source/RigVM/Public/RigVMCore/RigVMExecuteContext.h
+++ b/Engine/Plugins/Runtime/RigVM/Source/RigVM/Public/RigVMCore/RigVMExecuteContext.h
@@ -723,7 +723,6 @@ struct RIGVM_API FRigVMExtendedExecuteContext
 	URigVM* VM = nullptr;
 	TArray<FRigVMSlice> Slices;
 	TArray<uint16> SliceOffsets;
-	double LastExecutionMicroSeconds = 0.0;
 	const FRigVMDispatchFactory* Factory = nullptr;
 	FRigVMNameCache NameCache;
 
@@ -762,38 +761,15 @@ struct RIGVM_API FRigVMExtendedExecuteContext
 #if WITH_EDITOR
 
 	FRigVMDebugInfo* DebugInfo = nullptr;
-	FRigVMBreakpoint HaltedAtBreakpoint;
-	int32 HaltedAtBreakpointHit = INDEX_NONE;
-	ERigVMBreakpointAction CurrentBreakpointAction = ERigVMBreakpointAction::None;
-
-	// stores the number of times each instruction was visited
-	TArray<int32> InstructionVisitedDuringLastRun;
-	TArray<uint64> InstructionCyclesDuringLastRun;
-	TArray<int32> InstructionVisitOrder;
-
-	const void SetFirstEntryEventInEventQueue(const FName& InFirstEventName) { FirstEntryEventInQueue = InFirstEventName; }
-	const FName& GetFirstEntryEventInEventQueue() { return FirstEntryEventInQueue; }
-
-	// A RigVMHost can run multiple events per evaluation, such as the Backward&Forward Solve Mode,
-	// store the first event such that we know when to reset data for a new round of rig evaluation
-	FName FirstEntryEventInQueue = NAME_None;
-
-	uint64 StartCycles = 0;
-	uint64 OverallCycles = 0;
-
-	DECLARE_EVENT_ThreeParams(URigVM, FExecutionHaltedEvent, int32, UObject*, const FName&);
-	FExecutionHaltedEvent OnExecutionHalted;
-
-	FExecutionHaltedEvent& ExecutionHalted()
-	{
-		return OnExecutionHalted;
-	}
 
 	void SetDebugInfo(FRigVMDebugInfo* InDebugInfo)
 	{
 		DebugInfo = InDebugInfo;
 	}
 
+	FRigVMDebugInfo* GetRigVMDebugInfo() const
+	{
+		return DebugInfo;
+	}
 #endif // WITH_EDITOR
-
 };
diff --git a/Engine/Plugins/Runtime/RigVM/Source/RigVM/Public/RigVMCore/RigVMNativized.h b/Engine/Plugins/Runtime/RigVM/Source/RigVM/Public/RigVMCore/RigVMNativized.h
--- a/Engine/Plugins/Runtime/RigVM/Source/RigVM/Public/RigVMCore/RigVMNativized.h
+++ b/Engine/Plugins/Runtime/RigVM/Source/RigVM/Public/RigVMCore/RigVMNativized.h
@@ -64,8 +64,11 @@ public:
 	{
 		Super::SetInstructionIndex(Context, InInstructionIndex);
 #if WITH_EDITOR
-		Context.InstructionVisitedDuringLastRun[InInstructionIndex]++;
-		Context.InstructionVisitOrder.Add(InInstructionIndex);
+		if (FRigVMDebugInfo* RigVMDebugInfo = Context.GetRigVMDebugInfo())
+		{
+			RigVMDebugInfo->SetInstructionVisitedDuringLastRun(InInstructionIndex);
+			RigVMDebugInfo->AddInstructionIndexToVisitOrder(InInstructionIndex);
+		}
 #endif
 	}
 
diff --git a/Engine/Plugins/Runtime/RigVM/Source/RigVM/Public/RigVMHost.h b/Engine/Plugins/Runtime/RigVM/Source/RigVM/Public/RigVMHost.h
--- a/Engine/Plugins/Runtime/RigVM/Source/RigVM/Public/RigVMHost.h
+++ b/Engine/Plugins/Runtime/RigVM/Source/RigVM/Public/RigVMHost.h
@@ -10,6 +10,9 @@
 #include "Engine/AssetUserData.h"
 #include "Interfaces/Interface_AssetUserData.h"
 #include "SceneManagement.h"
+#if WITH_EDITOR
+#include "RigVMCore/RigVMDebugInfo.h"
+#endif
 #include "RigVMHost.generated.h"
 
 // set this to something larger than 0 to profile N runs
@@ -444,10 +447,18 @@ public:
 	 *  it is applied on the next VM execution */
 	bool ExecuteBreakpointAction(const ERigVMBreakpointAction BreakpointAction);
 
-	const FRigVMBreakpoint& GetHaltedAtBreakpoint() const { return ExtendedExecuteContext.HaltedAtBreakpoint; }
-	void SetBreakpointAction(const ERigVMBreakpointAction& Action) { ExtendedExecuteContext.CurrentBreakpointAction = Action; }
+	const FRigVMBreakpoint& GetHaltedAtBreakpoint() const
+	{
+		return GetDebugInfo().GetHaltedAtBreakpoint();
+	}
+
+	void SetBreakpointAction(const ERigVMBreakpointAction& Action)
+	{
+		GetDebugInfo().SetCurrentBreakpointAction(Action);
+	}
 	
 	FRigVMDebugInfo& GetDebugInfo() { return DebugInfo; }
+	const FRigVMDebugInfo& GetDebugInfo() const { return DebugInfo; }
 
 	/** Creates the snapshot VM if required and returns it */
 	URigVM* GetSnapshotVM(bool bCreateIfNeeded = true);
diff --git a/Engine/Plugins/Runtime/RigVM/Source/RigVMEditor/Private/Editor/RigVMEditor.cpp b/Engine/Plugins/Runtime/RigVM/Source/RigVMEditor/Private/Editor/RigVMEditor.cpp
--- a/Engine/Plugins/Runtime/RigVM/Source/RigVMEditor/Private/Editor/RigVMEditor.cpp
+++ b/Engine/Plugins/Runtime/RigVM/Source/RigVMEditor/Private/Editor/RigVMEditor.cpp
@@ -1029,7 +1029,7 @@ void FRigVMEditor::Compile()
 		{
 			RigVMHost->OnInitialized_AnyThread().Clear();
 			RigVMHost->OnExecuted_AnyThread().Clear();
-			RigVMHost->GetExtendedExecuteContext().ExecutionHalted().RemoveAll(this);
+			RigVMHost->GetDebugInfo().ExecutionHalted().RemoveAll(this);
 		}
 
 		SetHost(nullptr);
@@ -1827,7 +1827,7 @@ void FRigVMEditor::HandleVMExecutedEvent(URigVMHost* InHost, const FName& InEven
 		{
 			if(DebuggedHost)
 			{
-				RigVMBlueprint->RigGraphDisplaySettings.TotalMicroSeconds = DebuggedHost->GetExtendedExecuteContext().LastExecutionMicroSeconds;
+				RigVMBlueprint->RigGraphDisplaySettings.TotalMicroSeconds = DebuggedHost->GetDebugInfo().GetLastExecutionMicroSeconds();
 			}
 
 			if(RigVMBlueprint->RigGraphDisplaySettings.bAutoDetermineRange)
@@ -2411,7 +2411,7 @@ void FRigVMEditor::UpdateRigVMHost()
 		RigVMHost->OnInitialized_AnyThread().AddSP(this, &FRigVMEditor::HandleVMExecutedEvent);
 		RigVMHost->OnExecuted_AnyThread().AddSP(this, &FRigVMEditor::HandleVMExecutedEvent);
 		RigVMHost->RequestInit();
-		RigVMHost->GetExtendedExecuteContext().ExecutionHalted().AddSP(this, &FRigVMEditor::HandleVMExecutionHalted);
+		RigVMHost->GetDebugInfo().ExecutionHalted().AddSP(this, &FRigVMEditor::HandleVMExecutionHalted);
 	}
 }
 
diff --git a/Engine/Plugins/Runtime/RigVM/Source/RigVMEditor/Private/Widgets/SRigVMExecutionStackView.cpp b/Engine/Plugins/Runtime/RigVM/Source/RigVMEditor/Private/Widgets/SRigVMExecutionStackView.cpp
--- a/Engine/Plugins/Runtime/RigVM/Source/RigVMEditor/Private/Widgets/SRigVMExecutionStackView.cpp
+++ b/Engine/Plugins/Runtime/RigVM/Source/RigVMEditor/Private/Widgets/SRigVMExecutionStackView.cpp
@@ -760,7 +760,7 @@ void SRigVMExecutionStackView::RefreshTreeView(URigVM* InVM, FRigVMExtendedExecu
 					}
 				}
 
-				Host->GetExtendedExecuteContext().ExecutionHalted().AddSP(this, &SRigVMExecutionStackView::HandleExecutionHalted);
+				Host->GetDebugInfo().ExecutionHalted().AddSP(this, &SRigVMExecutionStackView::HandleExecutionHalted);
 			}
 		}
 	}
