AnimNext: RigVM support upgrades

Added class support to RigVM
Added object casting support to RigVM
Rebased graph/node/pin context menus to use IRigVMClientHost so non-blueprint hosts can use them
Moved functionality SRigVMGraphChangePinType::GetBindingText into shared util func
A few tweaks here and there to allow non-class based RigVM hosts to function (mainly CDO checks). Inserted a small hack to allow AnimNext to continue to be able to use non-CDO work memory initialization (until we can get the work memory in TLS).

diff --git a/Engine/Plugins/Animation/ControlRig/Source/ControlRigEditor/Private/ControlRigEditorModule.cpp b/Engine/Plugins/Animation/ControlRig/Source/ControlRigEditor/Private/ControlRigEditorModule.cpp
--- a/Engine/Plugins/Animation/ControlRig/Source/ControlRigEditor/Private/ControlRigEditorModule.cpp
+++ b/Engine/Plugins/Animation/ControlRig/Source/ControlRigEditor/Private/ControlRigEditorModule.cpp
@@ -285,19 +285,19 @@ UClass* FControlRigEditorModule::GetRigVMBlueprintClass() const
 	return UControlRigBlueprint::StaticClass();
 }
 
-void FControlRigEditorModule::GetNodeContextMenuActions(URigVMBlueprint* RigVMBlueprint,
+void FControlRigEditorModule::GetNodeContextMenuActions(IRigVMClientHost* RigVMClientHost,
 	const URigVMEdGraphNode* EdGraphNode, URigVMNode* ModelNode, UToolMenu* Menu) const
 {
-	FRigVMEditorModule::GetNodeContextMenuActions(RigVMBlueprint, EdGraphNode, ModelNode, Menu);
+	FRigVMEditorModule::GetNodeContextMenuActions(RigVMClientHost, EdGraphNode, ModelNode, Menu);
 
-	UControlRigBlueprint* ControlRigBlueprint = Cast<UControlRigBlueprint>(RigVMBlueprint);
+	UControlRigBlueprint* ControlRigBlueprint = Cast<UControlRigBlueprint>(RigVMClientHost);
 	if(ControlRigBlueprint == nullptr)
 	{
 		return;
 	}
 
-	URigVMGraph* Model = RigVMBlueprint->GetModel(EdGraphNode->GetGraph());
-	URigVMController* Controller = RigVMBlueprint->GetController(Model);
+	URigVMGraph* Model = RigVMClientHost->GetRigVMClient()->GetModel(EdGraphNode->GetGraph());
+	URigVMController* Controller = RigVMClientHost->GetRigVMClient()->GetController(Model);
 
 	TArray<FName> SelectedNodeNames = Model->GetSelectNodes();
 	SelectedNodeNames.AddUnique(ModelNode->GetFName());
@@ -476,7 +476,7 @@ void FControlRigEditorModule::GetNodeContextMenuActions(URigVMBlueprint* RigVMBl
 		}
 	}
 
-	GetDirectManipulationMenuActions(RigVMBlueprint, ModelNode, nullptr, Menu);
+	GetDirectManipulationMenuActions(RigVMClientHost, ModelNode, nullptr, Menu);
 
 	if (RigElementsToSelect.Num() > 0)
 	{
@@ -621,18 +621,18 @@ void FControlRigEditorModule::GetNodeContextMenuActions(URigVMBlueprint* RigVMBl
 	}
 }
 
-void FControlRigEditorModule::GetPinContextMenuActions(URigVMBlueprint* RigVMBlueprint, const UEdGraphPin* EdGraphPin, URigVMPin* ModelPin, UToolMenu* Menu) const
+void FControlRigEditorModule::GetPinContextMenuActions(IRigVMClientHost* RigVMClientHost, const UEdGraphPin* EdGraphPin, URigVMPin* ModelPin, UToolMenu* Menu) const
 {
-	FRigVMEditorModule::GetPinContextMenuActions(RigVMBlueprint, EdGraphPin, ModelPin, Menu);
-	GetDirectManipulationMenuActions(RigVMBlueprint, ModelPin->GetNode(), ModelPin, Menu);
+	FRigVMEditorModule::GetPinContextMenuActions(RigVMClientHost, EdGraphPin, ModelPin, Menu);
+	GetDirectManipulationMenuActions(RigVMClientHost, ModelPin->GetNode(), ModelPin, Menu);
 }
 
-void FControlRigEditorModule::GetDirectManipulationMenuActions(URigVMBlueprint* RigVMBlueprint, URigVMNode* InNode, URigVMPin* ModelPin, UToolMenu* Menu) const
+void FControlRigEditorModule::GetDirectManipulationMenuActions(IRigVMClientHost* RigVMClientHost, URigVMNode* InNode, URigVMPin* ModelPin, UToolMenu* Menu) const
 {
     // Add direct manipulation context menu entries
-	if(UControlRigBlueprint* ControlRigBlueprint = Cast<UControlRigBlueprint>(RigVMBlueprint))
+	if(UControlRigBlueprint* ControlRigBlueprint = Cast<UControlRigBlueprint>(RigVMClientHost))
 	{
-		UControlRig* DebuggedRig = Cast<UControlRig>(RigVMBlueprint->GetObjectBeingDebugged());
+		UControlRig* DebuggedRig = Cast<UControlRig>(ControlRigBlueprint->GetObjectBeingDebugged());
 		if(DebuggedRig == nullptr)
 		{
 			return;
diff --git a/Engine/Plugins/Animation/ControlRig/Source/ControlRigEditor/Private/ControlRigEditorModule.h b/Engine/Plugins/Animation/ControlRig/Source/ControlRigEditor/Private/ControlRigEditorModule.h
--- a/Engine/Plugins/Animation/ControlRig/Source/ControlRigEditor/Private/ControlRigEditorModule.h
+++ b/Engine/Plugins/Animation/ControlRig/Source/ControlRigEditor/Private/ControlRigEditorModule.h
@@ -31,8 +31,8 @@ public:
 
 	/** FRigVMEditorModule interface */
 	virtual UClass* GetRigVMBlueprintClass() const override;
-	virtual void GetNodeContextMenuActions(URigVMBlueprint* RigVMBlueprint, const URigVMEdGraphNode* EdGraphNode, URigVMNode* ModelNode, UToolMenu* Menu) const override;
-	virtual void GetPinContextMenuActions(URigVMBlueprint* RigVMBlueprint, const UEdGraphPin* EdGraphPin, URigVMPin* ModelPin, UToolMenu* Menu) const override;
+	virtual void GetNodeContextMenuActions(IRigVMClientHost* RigVMClientHost, const URigVMEdGraphNode* EdGraphNode, URigVMNode* ModelNode, UToolMenu* Menu) const override;
+	virtual void GetPinContextMenuActions(IRigVMClientHost* RigVMClientHost, const UEdGraphPin* EdGraphPin, URigVMPin* ModelPin, UToolMenu* Menu) const override;
 
 	/** IControlRigEditorModule interface */
 	virtual TSharedRef<IControlRigEditor> CreateControlRigEditor(const EToolkitMode::Type Mode, const TSharedPtr<IToolkitHost>& InitToolkitHost, UControlRigBlueprint* Blueprint) override;
@@ -48,7 +48,7 @@ public:
 	static void UnLinkLevelSequence(UAnimSequence* AnimSequence);
 	void ExtendAnimSequenceMenu();
 
-	void GetDirectManipulationMenuActions(URigVMBlueprint* RigVMBlueprint, URigVMNode* InNode, URigVMPin* ModelPin, UToolMenu* Menu) const;
+	void GetDirectManipulationMenuActions(IRigVMClientHost* RigVMClientHost, URigVMNode* InNode, URigVMPin* ModelPin, UToolMenu* Menu) const;
 
 private:
 
diff --git a/Engine/Plugins/Experimental/AnimNext/Source/AnimNext/Private/Module.cpp b/Engine/Plugins/Experimental/AnimNext/Source/AnimNext/Private/Module.cpp
--- a/Engine/Plugins/Experimental/AnimNext/Source/AnimNext/Private/Module.cpp
+++ b/Engine/Plugins/Experimental/AnimNext/Source/AnimNext/Private/Module.cpp
@@ -18,13 +18,13 @@ class FModule : public IModuleInterface
 public:
 	virtual void StartupModule() override
 	{
-		static UClass* const AllowedAssetTypes[] =
+		static TPair<UClass*, FRigVMRegistry::ERegisterObjectOperation> const AllowedObjectTypes[] =
 		{
-			UAnimSequence::StaticClass(),
-			UScriptStruct::StaticClass(),
+			{ UAnimSequence::StaticClass(), FRigVMRegistry::ERegisterObjectOperation::Class },
+			{ UScriptStruct::StaticClass(), FRigVMRegistry::ERegisterObjectOperation::Class }
 		};
 
-		FRigVMRegistry::Get().RegisterObjectTypes(AllowedAssetTypes);
+		FRigVMRegistry::Get().RegisterObjectTypes(AllowedObjectTypes);
 
 		FDataRegistry::Init();
 		FDecoratorRegistry::Init();
diff --git a/Engine/Plugins/Runtime/PropertyAccess/Source/PropertyAccessEditor/Private/SPropertyBinding.cpp b/Engine/Plugins/Runtime/PropertyAccess/Source/PropertyAccessEditor/Private/SPropertyBinding.cpp
--- a/Engine/Plugins/Runtime/PropertyAccess/Source/PropertyAccessEditor/Private/SPropertyBinding.cpp
+++ b/Engine/Plugins/Runtime/PropertyAccess/Source/PropertyAccessEditor/Private/SPropertyBinding.cpp
@@ -415,11 +415,8 @@ TSharedRef<SWidget> SPropertyBinding::OnGenerateDelegateMenu()
 		// Get the current skeleton class, think header for the blueprint.
 		UBlueprintGeneratedClass* SkeletonClass = Blueprint ? Cast<UBlueprintGeneratedClass>(Blueprint->SkeletonGeneratedClass) : nullptr;
 
-		if (SkeletonClass)
-		{
-			TArray<TSharedPtr<FBindingChainElement>> BindingChain;
-			FillPropertyMenu(MenuBuilder, SkeletonClass, BindingChain);
-		}
+		TArray<TSharedPtr<FBindingChainElement>> BindingChain;
+		FillPropertyMenu(MenuBuilder, SkeletonClass, BindingChain);
 	}
 
 	if (BindingContextStructs.Num() > 0)
diff --git a/Engine/Plugins/Runtime/RigVM/Source/RigVM/Private/RigVMCore/RigVM.cpp b/Engine/Plugins/Runtime/RigVM/Source/RigVM/Private/RigVMCore/RigVM.cpp
--- a/Engine/Plugins/Runtime/RigVM/Source/RigVM/Private/RigVMCore/RigVM.cpp
+++ b/Engine/Plugins/Runtime/RigVM/Source/RigVM/Private/RigVMCore/RigVM.cpp
@@ -1473,7 +1473,7 @@ bool URigVM::InitializeInstance(FRigVMExtendedExecuteContext& Context, TArrayVie
 	// re-initialize work memory from CDO
 	if (TRigVMMemoryStorage* WorkMemory = Memory[(int32)ERigVMMemoryType::Work])
 	{
-		if (ensure(!WorkMemory->HasAnyFlags(RF_ClassDefaultObject)))
+		if (!WorkMemory->HasAnyFlags(RF_ClassDefaultObject))
 		{
 			if (const URigVMMemoryStorageGeneratorClass* MemoryClass = Cast<URigVMMemoryStorageGeneratorClass>(WorkMemory->GetClass()))
 			{
diff --git a/Engine/Plugins/Runtime/RigVM/Source/RigVM/Private/RigVMCore/RigVMRegistry.cpp b/Engine/Plugins/Runtime/RigVM/Source/RigVM/Private/RigVMCore/RigVMRegistry.cpp
--- a/Engine/Plugins/Runtime/RigVM/Source/RigVM/Private/RigVMCore/RigVMRegistry.cpp
+++ b/Engine/Plugins/Runtime/RigVM/Source/RigVM/Private/RigVMCore/RigVMRegistry.cpp
@@ -357,7 +357,9 @@ void FRigVMRegistry::RefreshEngineTypes()
 		UClass* Class = *ClassIt;
 		if (IsAllowedType(Class))
 		{
-			FindOrAddType(FRigVMTemplateArgumentType(Class));
+			// Register both the class and the object type for use
+			FindOrAddType(FRigVMTemplateArgumentType(Class, RigVMTypeUtils::EClassArgType::AsClass));
+			FindOrAddType(FRigVMTemplateArgumentType(Class, RigVMTypeUtils::EClassArgType::AsObject));
 		}
 	}
 
@@ -1139,7 +1141,7 @@ bool FRigVMRegistry::IsAllowedType(const UStruct* InStruct) const
 
 bool FRigVMRegistry::IsAllowedType(const UClass* InClass) const
 {
-	if(!InClass || InClass->HasAnyClassFlags(CLASS_Hidden | CLASS_Abstract))
+	if(!InClass || InClass->HasAnyClassFlags(CLASS_Hidden))
 	{
 		return false;
 	}
@@ -1291,20 +1293,44 @@ void FRigVMRegistry::RegisterPredicate(UScriptStruct* InStruct, const TCHAR* InN
 	Predicates.Add(Function);
 }
 
-void FRigVMRegistry::RegisterObjectTypes(TConstArrayView<UClass*> InClasses)
+void FRigVMRegistry::RegisterObjectTypes(TConstArrayView<TPair<UClass*, ERegisterObjectOperation>> InClasses)
 {
-	for (UClass* Class : InClasses)
+	for (TPair<UClass*, ERegisterObjectOperation> ClassOpPair : InClasses)
 	{
+		UClass* Class = ClassOpPair.Key;
+		ERegisterObjectOperation Operation = ClassOpPair.Value;
+
 		// Only allow native object types
 		if (Class->HasAnyClassFlags(CLASS_Native))
 		{
-			// Add all child classes
-			TArray<UClass*> DerivedClasses({ Class });
-			GetDerivedClasses(Class, DerivedClasses, /*bRecursive=*/true);
-			for(UClass* DerivedClass : DerivedClasses)
+			switch (Operation)
 			{
-				AllowedClasses.Add(DerivedClass);
+			case ERegisterObjectOperation::Class:
+				AllowedClasses.Add(Class);
+				break;
+			case ERegisterObjectOperation::ClassAndParents:
+				{
+					// Add all parent classes
+					do
+					{
+						AllowedClasses.Add(Class);
+						Class = Class->GetSuperClass();
+					} while (Class);
+					break;
+				}
+			case ERegisterObjectOperation::ClassAndChildren:
+				{
+					// Add all child classes
+					TArray<UClass*> DerivedClasses({ Class });
+					GetDerivedClasses(Class, DerivedClasses, /*bRecursive=*/true);
+					for (UClass* DerivedClass : DerivedClasses)
+					{
+						AllowedClasses.Add(DerivedClass);
+					}
+					break;
+				}
 			}
+
 		}
 	}
 }
diff --git a/Engine/Plugins/Runtime/RigVM/Source/RigVM/Private/RigVMFunctions/RigVMDispatch_CastObject.cpp b/Engine/Plugins/Runtime/RigVM/Source/RigVM/Private/RigVMFunctions/RigVMDispatch_CastObject.cpp
new file mode 100644
--- /dev/null
+++ b/Engine/Plugins/Runtime/RigVM/Source/RigVM/Private/RigVMFunctions/RigVMDispatch_CastObject.cpp
@@ -0,0 +1,113 @@
+﻿// Copyright Epic Games, Inc. All Rights Reserved.
+
+#include "RigVMFunctions/RigVMDispatch_CastObject.h"
+#include "RigVMCore/RigVMRegistry.h"
+
+#define LOCTEXT_NAMESPACE "RigVMDispatch_CastObject"
+
+const FName FRigVMDispatch_CastObject::ValueName = TEXT("Value");
+const FName FRigVMDispatch_CastObject::ResultName = TEXT("Result");
+
+const TArray<FRigVMTemplateArgument>& FRigVMDispatch_CastObject::GetArguments() const
+{
+	static const TArray<FRigVMTemplateArgument::ETypeCategory> ElementCategories =
+	{
+		FRigVMTemplateArgument::ETypeCategory_SingleObjectValue
+	};
+
+	static const TArray<FRigVMTemplateArgument> Arguments = BuildArgumentListFromPrimaryArgument(
+		{
+			FRigVMTemplateArgument(ValueName, ERigVMPinDirection::Input, ElementCategories),
+			FRigVMTemplateArgument(ResultName, ERigVMPinDirection::Output)
+		},
+		ValueName);
+
+	return Arguments;
+}
+
+TArray<FRigVMTemplateTypeMap> FRigVMDispatch_CastObject::GetPermutationsFromArgumentType(const FName& InArgumentName, const TRigVMTypeIndex& InTypeIndex) const
+{
+	TArray<FRigVMTemplateTypeMap> Permutations;
+	if (InArgumentName == ValueName)
+	{
+		const TArray<TRigVMTypeIndex>& ObjectTypes = FRigVMRegistry::Get().GetTypesForCategory(FRigVMTemplateArgument::ETypeCategory_SingleObjectValue);
+		for (const TRigVMTypeIndex& Type : ObjectTypes)
+		{
+			Permutations.Add(
+	{
+				{ ValueName, InTypeIndex },
+				{ ResultName, Type }
+			});
+		}
+	}
+	else if (InArgumentName == ResultName)
+	{
+		const TArray<TRigVMTypeIndex>& ObjectTypes = FRigVMRegistry::Get().GetTypesForCategory(FRigVMTemplateArgument::ETypeCategory_SingleObjectValue);
+		for (const TRigVMTypeIndex& Type : ObjectTypes)
+		{
+			Permutations.Add(
+	{
+				{ ValueName, Type },
+				{ ResultName, InTypeIndex }
+			});
+		}
+	}
+	return Permutations;
+}
+
+#if WITH_EDITOR
+
+FString FRigVMDispatch_CastObject::GetNodeTitle(const FRigVMTemplateTypeMap& InTypes) const
+{
+	const TRigVMTypeIndex* ValueTypeIndexPtr = InTypes.Find(ValueName);
+	const TRigVMTypeIndex* ResultTypeIndexPtr = InTypes.Find(ResultName);
+	
+	if(ValueTypeIndexPtr && ResultTypeIndexPtr)
+	{
+		const TRigVMTypeIndex& ValueTypeIndex = *ValueTypeIndexPtr;
+		const TRigVMTypeIndex& ResultTypeIndex = *ResultTypeIndexPtr;
+		if(ValueTypeIndex != INDEX_NONE && ResultTypeIndex != INDEX_NONE)
+		{
+			const FRigVMRegistry& Registry = FRigVMRegistry::Get();
+			if(!Registry.IsWildCardType(ValueTypeIndex) && !Registry.IsWildCardType(ResultTypeIndex))
+			{
+				const FRigVMTemplateArgumentType& ValueType = Registry.GetType(ValueTypeIndex);
+				const FRigVMTemplateArgumentType& ResultType = Registry.GetType(ResultTypeIndex);
+				static constexpr TCHAR Format[] = TEXT("Cast to %s");
+				if(ValueType.CPPTypeObject && ResultType.CPPTypeObject)
+				{
+					return FString::Printf(Format, *ResultType.CPPTypeObject->GetName());
+				}
+			}
+		}
+	}
+	return Super::GetNodeTitle(InTypes);
+}
+
+FText FRigVMDispatch_CastObject::GetNodeTooltip(const FRigVMTemplateTypeMap& InTypes) const
+{
+	return LOCTEXT("CastToolTip", "Casts between object types");
+}
+
+#endif
+
+void FRigVMDispatch_CastObject::Execute(FRigVMExtendedExecuteContext& InContext, FRigVMMemoryHandleArray Handles, FRigVMPredicateBranchArray RigVMBranches)
+{
+	UObject** ValueObjectPtr = (UObject**)Handles[0].GetData();
+	UObject** ResultObjectPtr = (UObject**)Handles[1].GetData();
+	const FObjectPropertyBase* ValueObjectProperty = CastFieldChecked<FObjectPropertyBase>(Handles[0].GetProperty());
+	const FObjectPropertyBase* ResultObjectProperty = CastFieldChecked<FObjectPropertyBase>(Handles[1].GetProperty());
+	if(ValueObjectPtr && ResultObjectPtr && ValueObjectProperty && ResultObjectProperty)
+	{
+		if(ResultObjectProperty->GetClass()->IsChildOf(ValueObjectProperty->GetClass()))
+		{
+			*ResultObjectPtr = *ValueObjectPtr;
+		}
+		else
+		{
+			*ResultObjectPtr = nullptr;
+		}
+	}
+}
+
+#undef LOCTEXT_NAMESPACE
diff --git a/Engine/Plugins/Runtime/RigVM/Source/RigVM/Private/RigVMHost.cpp b/Engine/Plugins/Runtime/RigVM/Source/RigVM/Private/RigVMHost.cpp
--- a/Engine/Plugins/Runtime/RigVM/Source/RigVM/Private/RigVMHost.cpp
+++ b/Engine/Plugins/Runtime/RigVM/Source/RigVM/Private/RigVMHost.cpp
@@ -270,7 +270,7 @@ bool URigVMHost::InitializeVM(const FName& InEventName)
 	}
 
 #if UE_RIGVM_PROPERTY_BAG_STORAGE_ENABLED
-	if (!HasAnyFlags(RF_ClassDefaultObject))
+	if (bTEMP_CopyDefaultsFromCDO && !HasAnyFlags(RF_ClassDefaultObject))
 	{
 		if (URigVMHost* CDO = GetClass()->GetDefaultObject<URigVMHost>())
 		{
diff --git a/Engine/Plugins/Runtime/RigVM/Source/RigVM/Public/RigVMCore/RigVMDispatchFactory.h b/Engine/Plugins/Runtime/RigVM/Source/RigVM/Public/RigVMCore/RigVMDispatchFactory.h
--- a/Engine/Plugins/Runtime/RigVM/Source/RigVM/Public/RigVMCore/RigVMDispatchFactory.h
+++ b/Engine/Plugins/Runtime/RigVM/Source/RigVM/Public/RigVMCore/RigVMDispatchFactory.h
@@ -22,6 +22,7 @@
 #include "RigVMDispatchFactory.generated.h"
 
 struct FRigVMDispatchFactory;
+class URigVMPin;
 
 /**
  * A context used for inquiring from dispatch factories
diff --git a/Engine/Plugins/Runtime/RigVM/Source/RigVM/Public/RigVMCore/RigVMRegistry.h b/Engine/Plugins/Runtime/RigVM/Source/RigVM/Public/RigVMCore/RigVMRegistry.h
--- a/Engine/Plugins/Runtime/RigVM/Source/RigVM/Public/RigVMCore/RigVMRegistry.h
+++ b/Engine/Plugins/Runtime/RigVM/Source/RigVM/Public/RigVMCore/RigVMRegistry.h
@@ -62,8 +62,18 @@ public:
 	// Register a predicate contained in the input struct
 	void RegisterPredicate(UScriptStruct* InStruct, const TCHAR* InName, const TArray<FRigVMFunctionArgument>& InArguments);
 
+	// How to register an object's class when passed to RegisterObjectTypes
+	enum class ERegisterObjectOperation
+	{
+		Class,
+
+		ClassAndParents,
+
+		ClassAndChildren,
+	};
+
 	// Register a set of allowed object types
-	void RegisterObjectTypes(TConstArrayView<UClass*> InClasses);
+	void RegisterObjectTypes(TConstArrayView<TPair<UClass*, ERegisterObjectOperation>> InClasses);
 
 	// Initializes the registry by storing the defaults
 	void InitializeIfNeeded();
@@ -157,7 +167,7 @@ public:
 	>
 	TRigVMTypeIndex GetTypeIndex(bool bAsArray = false) const
 	{
-		FRigVMTemplateArgumentType Type(T::StaticClass());
+		FRigVMTemplateArgumentType Type(T::StaticClass(), RigVMTypeUtils::EClassArgType::AsObject);
 		if(bAsArray)
 		{
 			Type.ConvertToArray();
diff --git a/Engine/Plugins/Runtime/RigVM/Source/RigVM/Public/RigVMCore/RigVMTemplate.h b/Engine/Plugins/Runtime/RigVM/Source/RigVM/Public/RigVMCore/RigVMTemplate.h
--- a/Engine/Plugins/Runtime/RigVM/Source/RigVM/Public/RigVMCore/RigVMTemplate.h
+++ b/Engine/Plugins/Runtime/RigVM/Source/RigVM/Public/RigVMCore/RigVMTemplate.h
@@ -70,9 +70,9 @@ struct RIGVM_API FRigVMTemplateArgumentType
 	}
 
 	FRigVMTemplateArgumentType(const FName& InCPPType, UObject* InCPPTypeObject = nullptr);
-
-	FRigVMTemplateArgumentType(UClass* InClass)
-	: CPPType(*RigVMTypeUtils::CPPTypeFromObject(InClass))
+	
+	FRigVMTemplateArgumentType(UClass* InClass, RigVMTypeUtils::EClassArgType InClassArgType = RigVMTypeUtils::EClassArgType::AsObject)
+	: CPPType(*RigVMTypeUtils::CPPTypeFromObject(InClass, InClassArgType))
 	, CPPTypeObject(InClass)
 	{
 	}
@@ -240,6 +240,9 @@ struct RIGVM_API FRigVMTemplateArgument
 
 	RIGVM_API friend uint32 GetTypeHash(const FRigVMTemplateArgument& InArgument);
 
+	// Get the map of types to permutation indices
+	const TMap<TRigVMTypeIndex, TArray<int32>>& GetTypeToPermutations() const { return TypeToPermutations; }
+
 protected:
 
 	int32 Index;
diff --git a/Engine/Plugins/Runtime/RigVM/Source/RigVM/Public/RigVMFunctions/RigVMDispatch_CastObject.h b/Engine/Plugins/Runtime/RigVM/Source/RigVM/Public/RigVMFunctions/RigVMDispatch_CastObject.h
new file mode 100644
--- /dev/null
+++ b/Engine/Plugins/Runtime/RigVM/Source/RigVM/Public/RigVMFunctions/RigVMDispatch_CastObject.h
@@ -0,0 +1,33 @@
+﻿// Copyright Epic Games, Inc. All Rights Reserved.
+
+#pragma once
+
+#include "RigVMCore/RigVMDispatchFactory.h"
+#include "RigVMDispatch_CastObject.generated.h"
+
+USTRUCT(meta=(DisplayName = "Cast", Category = "Object", Keywords = "As", NodeColor = "1,1,1,1"))
+struct RIGVM_API FRigVMDispatch_CastObject : public FRigVMDispatchFactory
+{
+	GENERATED_BODY()
+
+public:
+	FRigVMDispatch_CastObject()
+	{
+		FactoryScriptStruct = StaticStruct();
+	}
+	
+	virtual const TArray<FRigVMTemplateArgument>& GetArguments() const override;
+	virtual TArray<FRigVMTemplateTypeMap> GetPermutationsFromArgumentType(const FName& InArgumentName, const TRigVMTypeIndex& InTypeIndex) const override;
+#if WITH_EDITOR
+	virtual FString GetNodeTitle(const FRigVMTemplateTypeMap& InTypes) const override;
+	virtual FText GetNodeTooltip(const FRigVMTemplateTypeMap& InTypes) const override;
+#endif
+	virtual bool IsSingleton() const override { return true; }
+	
+protected:
+	virtual FRigVMFunctionPtr GetDispatchFunctionImpl(const FRigVMTemplateTypeMap& InTypes) const override { return &FRigVMDispatch_CastObject::Execute; }
+	static void Execute(FRigVMExtendedExecuteContext& InContext, FRigVMMemoryHandleArray Handles, FRigVMPredicateBranchArray RigVMBranches);
+
+	static const FName ValueName;
+	static const FName ResultName;
+};
diff --git a/Engine/Plugins/Runtime/RigVM/Source/RigVM/Public/RigVMHost.h b/Engine/Plugins/Runtime/RigVM/Source/RigVM/Public/RigVMHost.h
--- a/Engine/Plugins/Runtime/RigVM/Source/RigVM/Public/RigVMHost.h
+++ b/Engine/Plugins/Runtime/RigVM/Source/RigVM/Public/RigVMHost.h
@@ -389,6 +389,11 @@ public:
 	virtual const TArray<UAssetUserData*>* GetAssetUserDataArray() const override;
 	//~ End IInterface_AssetUserData Interface
 
+#if UE_RIGVM_PROPERTY_BAG_STORAGE_ENABLED
+	// Temporary flag while we get TLS-based work memory working, controls copying from the CDO in InitializeVM
+	bool bTEMP_CopyDefaultsFromCDO = true;
+#endif
+
 protected:
 	/** Array of user data stored with the asset */
 	UPROPERTY(EditAnywhere, AdvancedDisplay, Instanced, Category = "Default")
diff --git a/Engine/Plugins/Runtime/RigVM/Source/RigVM/Public/RigVMTypeUtils.h b/Engine/Plugins/Runtime/RigVM/Source/RigVM/Public/RigVMTypeUtils.h
--- a/Engine/Plugins/Runtime/RigVM/Source/RigVM/Public/RigVMTypeUtils.h
+++ b/Engine/Plugins/Runtime/RigVM/Source/RigVM/Public/RigVMTypeUtils.h
@@ -12,7 +12,7 @@
 #include "UObject/Package.h"
 #include "UObject/SoftObjectPath.h"
 
-
+struct FRigVMTemplateArgumentType;
 
 struct RIGVM_API FRigVMUserDefinedTypeResolver
 {
@@ -48,9 +48,11 @@ namespace RigVMTypeUtils
 {
 	const TCHAR TArrayPrefix[] = TEXT("TArray<");
 	const TCHAR TObjectPtrPrefix[] = TEXT("TObjectPtr<");
+	const TCHAR TSubclassOfPrefix[] = TEXT("TSubclassOf<");
 	const TCHAR TScriptInterfacePrefix[] = TEXT("TScriptInterface<");
 	const TCHAR TArrayTemplate[] = TEXT("TArray<%s>");
 	const TCHAR TObjectPtrTemplate[] = TEXT("TObjectPtr<%s%s>");
+	const TCHAR TSubclassOfTemplate[] = TEXT("TSubclassOf<%s%s>");
 	const TCHAR TScriptInterfaceTemplate[] = TEXT("TScriptInterface<%s%s>");
 
 	const FString BoolType = TEXT("bool");
@@ -156,6 +158,11 @@ namespace RigVMTypeUtils
 		return CPPType;
 	}
 
+	inline bool IsUClassType(const FString& InCPPType)
+	{
+		return InCPPType.StartsWith(TSubclassOfPrefix);
+	}
+
 	inline bool IsUObjectType(const FString& InCPPType)
 	{
 		return InCPPType.StartsWith(TObjectPtrPrefix);
@@ -212,6 +219,7 @@ namespace RigVMTypeUtils
 			TEXT("TArray<TObjectPtr<"),
 			TEXT("TArray<TArray<TObjectPtr<"),
 			TEXT("TScriptInterface<")
+			TEXT("TSubclassOf<")
 		};
 		static const TArray<FString> PrefixesNotRequiringCPPTypeObject = {
 			TEXT("UInt"), 
@@ -314,11 +322,26 @@ namespace RigVMTypeUtils
 		return Object;
 	}
 
-	static FString CPPTypeFromObject(const UObject* InCPPTypeObject)
+	// A UClass argument is used to signify both the object type and its class type.
+	// This argument differentiates between the two
+	enum class EClassArgType
+	{
+		// This type signifies a class
+		AsClass,
+
+		// This type signifies an object 
+		AsObject
+	};
+
+	static FString CPPTypeFromObject(const UObject* InCPPTypeObject, EClassArgType InClassArgType = EClassArgType::AsObject)
 	{
 		if (const UClass* Class = Cast<UClass>(InCPPTypeObject))
 		{
-			if (Class->IsChildOf(UInterface::StaticClass()))
+			if (InClassArgType == EClassArgType::AsClass)
+			{
+				return FString::Printf(RigVMTypeUtils::TSubclassOfPrefix, Class->GetPrefixCPP(), *Class->GetName());
+			}
+			else if (Class->IsChildOf(UInterface::StaticClass()))
 			{
 				return FString::Printf(RigVMTypeUtils::TScriptInterfaceTemplate, TEXT("I"), *Class->GetName());
 			}
@@ -357,9 +380,23 @@ namespace RigVMTypeUtils
 			BaseCPPType = BaseTypeFromArrayType(BaseCPPType);
 		}
 		FString CPPType = BaseCPPType;
+		const bool bIsClass = CPPType.StartsWith(TSubclassOfPrefix);
 
+		static const FString PrefixObjectPtr = TObjectPtrPrefix;
+		static const FString PrefixSubclassOf = TSubclassOfPrefix;
 		static const FString PrefixScriptInterface = TScriptInterfacePrefix;
-		if (CPPType.StartsWith(TScriptInterfacePrefix))
+
+		if (CPPType.StartsWith(PrefixObjectPtr))
+		{
+			// Chop the prefix + the U indicating object class
+			CPPType = CPPType.RightChop(PrefixObjectPtr.Len() + 1).LeftChop(1);
+		}
+		else if (CPPType.StartsWith(PrefixSubclassOf))
+		{
+			// Chop the prefix + the U indicating object class
+			CPPType = CPPType.RightChop(PrefixSubclassOf.Len() + 1).LeftChop(1);
+		}
+		else if (CPPType.StartsWith(TScriptInterfacePrefix))
 		{
 			// Chop the prefix + the I indicating interface class
 			CPPType = CPPType.RightChop(PrefixScriptInterface.Len() + 1).LeftChop(1);
@@ -389,7 +426,7 @@ namespace RigVMTypeUtils
 			return nullptr;
 		}
 
-		CPPType = CPPTypeFromObject(CPPTypeObject);
+		CPPType = CPPTypeFromObject(CPPTypeObject, bIsClass ? EClassArgType::AsClass : EClassArgType::AsObject);
 		InOutCPPType.ReplaceInline(*BaseCPPType, *CPPType);
 		return CPPTypeObject;
 	}
@@ -399,7 +436,8 @@ namespace RigVMTypeUtils
 		FString CPPType = InCPPType;
 		if (InCPPTypeObject)
 		{
-			CPPType = CPPTypeFromObject(InCPPTypeObject);	
+			const bool bIsClass = CPPType.StartsWith(TSubclassOfPrefix);
+			CPPType = CPPTypeFromObject(InCPPTypeObject, bIsClass ? EClassArgType::AsClass : EClassArgType::AsObject);
 			if(CPPType != InCPPType)
 			{
 				FString TemplateType = InCPPType;
diff --git a/Engine/Plugins/Runtime/RigVM/Source/RigVMDeveloper/Private/EdGraph/RigVMEdGraphNode.cpp b/Engine/Plugins/Runtime/RigVM/Source/RigVMDeveloper/Private/EdGraph/RigVMEdGraphNode.cpp
--- a/Engine/Plugins/Runtime/RigVM/Source/RigVMDeveloper/Private/EdGraph/RigVMEdGraphNode.cpp
+++ b/Engine/Plugins/Runtime/RigVM/Source/RigVMDeveloper/Private/EdGraph/RigVMEdGraphNode.cpp
@@ -16,6 +16,7 @@
 #include "RigVMBlueprintUtils.h"
 #include "../../../RigVMEditor/Public/RigVMEditorModule.h"
 #include "RigVMCore/RigVMExecuteContext.h"
+#include "RigVMFunctions/RigVMDispatch_CastObject.h"
 #include "RigVMModel/Nodes/RigVMAggregateNode.h"
 #include "RigVMModel/Nodes/RigVMFunctionReferenceNode.h"
 #include "RigVMModel/Nodes/RigVMFunctionEntryNode.h"
@@ -1047,7 +1048,8 @@ bool URigVMEdGraphNode::DrawAsCompactNode() const
 		DrawAsCompactNodeCache = false;
 		if(const URigVMTemplateNode* TemplateModelNode = Cast<URigVMTemplateNode>(GetModelNode()))
 		{
-			if(TemplateModelNode->GetNotation() == RigVMTypeUtils::GetCastTemplateNotation())
+			if(TemplateModelNode->GetNotation() == RigVMTypeUtils::GetCastTemplateNotation() ||
+				TemplateModelNode->GetNotation() == FRigVMDispatch_CastObject().GetTemplateNotation())
 			{
 				DrawAsCompactNodeCache = true;
 				
diff --git a/Engine/Plugins/Runtime/RigVM/Source/RigVMDeveloper/Private/RigVMDeveloperTypeUtils.cpp b/Engine/Plugins/Runtime/RigVM/Source/RigVMDeveloper/Private/RigVMDeveloperTypeUtils.cpp
--- a/Engine/Plugins/Runtime/RigVM/Source/RigVMDeveloper/Private/RigVMDeveloperTypeUtils.cpp
+++ b/Engine/Plugins/Runtime/RigVM/Source/RigVMDeveloper/Private/RigVMDeveloperTypeUtils.cpp
@@ -3,9 +3,12 @@
 #include "RigVMDeveloperTypeUtils.h"
 
 #include "Internationalization/StringTableCore.h"
+#include "RigVMFunctions/RigVMDispatch_CastObject.h"
 #include "RigVMModel/RigVMController.h"
 #include "RigVMModel/RigVMVariableDescription.h"
 
+#define LOCTEXT_NAMESPACE "RigVMTypeUtils"
+
 #if WITH_EDITOR
 
 #include "EdGraphSchema_K2.h"
@@ -102,7 +105,8 @@ FEdGraphPinType RigVMTypeUtils::PinTypeFromCPPType(const FName& InCPPType, UObje
 	}
 	else if (Cast<UClass>(InCPPTypeObject))
 	{
-		PinType.PinCategory = UEdGraphSchema_K2::PC_Object;
+		const bool bIsClass = RigVMTypeUtils::IsUClassType(InCPPType.ToString());
+		PinType.PinCategory = bIsClass ? UEdGraphSchema_K2::PC_Class : UEdGraphSchema_K2::PC_Object;
 		PinType.PinSubCategoryObject = InCPPTypeObject;
 	}
 
@@ -629,21 +633,36 @@ bool RigVMTypeUtils::CanCastTypes(const TRigVMTypeIndex& InSourceTypeIndex, cons
 const FRigVMFunction* RigVMTypeUtils::GetCastForTypeIndices(const TRigVMTypeIndex& InSourceTypeIndex, const TRigVMTypeIndex& InTargetTypeIndex)
 {
 	const FRigVMRegistry& Registry = FRigVMRegistry::Get();
-	if(const FRigVMTemplate* CastTemplate = Registry.FindTemplate(CastTemplateNotation))
+
+	const FRigVMTemplate* CastTemplates[] =
 	{
-		const FRigVMTemplateArgument* SourceArgument = CastTemplate->FindArgument(CastTemplateValueName);
-		const FRigVMTemplateArgument* TargetArgument = CastTemplate->FindArgument(CastTemplateResultName);
+		Registry.FindTemplate(CastTemplateNotation),
+		Registry.FindTemplate(FRigVMDispatch_CastObject().GetTemplateNotation())
+	};
 
-		if(SourceArgument && TargetArgument)
+	for(const FRigVMTemplate* CastTemplate : CastTemplates)
+	{
+		if(CastTemplate)
 		{
-			for(int32 Index = 0; Index < SourceArgument->GetTypeIndices().Num(); Index++)
+			const FRigVMTemplateArgument* SourceArgument = CastTemplate->FindArgument(CastTemplateValueName);
+			const FRigVMTemplateArgument* TargetArgument = CastTemplate->FindArgument(CastTemplateResultName);
+
+			if(SourceArgument && TargetArgument)
 			{
-				const TRigVMTypeIndex& SourceTypeIndex = SourceArgument->GetTypeIndices()[Index];
-				const TRigVMTypeIndex& TargetTypeIndex = TargetArgument->GetTypeIndices()[Index];
-				if(SourceTypeIndex == InSourceTypeIndex &&
-					TargetTypeIndex == InTargetTypeIndex)
+				const TArray<int32>* SourcePermutations = SourceArgument->GetTypeToPermutations().Find(InSourceTypeIndex);
+				const TArray<int32>* TargetPermutations = TargetArgument->GetTypeToPermutations().Find(InTargetTypeIndex);
+				if(SourcePermutations && TargetPermutations)
 				{
-					return CastTemplate->GetPermutation(Index);
+					for(int32 SourceIndex = 0, SourceCount = SourcePermutations->Num(); SourceIndex < SourceCount; ++SourceIndex)
+					{
+						for(int32 TargetIndex = 0, TargetCount = TargetPermutations->Num(); TargetIndex < TargetCount; ++TargetIndex)
+						{
+							if((*SourcePermutations)[SourceIndex] == (*TargetPermutations)[TargetIndex])
+							{
+								return const_cast<FRigVMTemplate*>(CastTemplate)->GetOrCreatePermutation((*SourcePermutations)[SourceIndex]);
+							}
+						}
+					}
 				}
 			}
 		}
@@ -666,4 +685,63 @@ const FName& RigVMTypeUtils::GetCastTemplateNotation()
 	return CastTemplateNotation;
 }
 
+FText RigVMTypeUtils::GetDisplayTextForArgumentType(const FRigVMTemplateArgumentType& InType)
+{
+	if(UObject* CPPTypeObject = InType.CPPTypeObject)
+	{
+		if (const UClass* Class = Cast<UClass>(CPPTypeObject))
+		{
+			return Class->GetDisplayNameText();
+		}
+		if(const UScriptStruct* ScriptStruct = Cast<UScriptStruct>(CPPTypeObject))
+		{
+			return ScriptStruct->GetDisplayNameText();
+		}
+		if(const UEnum* Enum = Cast<UEnum>(CPPTypeObject))
+		{
+			return Enum->GetDisplayNameText();
+		}
+	}
+	else
+	{
+		FString CPPType = InType.CPPType.ToString();
+		if(IsArrayType(CPPType))
+		{
+			CPPType = BaseTypeFromArrayType(CPPType);
+		}
+
+		static const FText BoolLabel = LOCTEXT("BoolLabel", "Boolean");
+		static const FText FloatLabel = LOCTEXT("FloatLabel", "Float");
+		static const FText Int32Label = LOCTEXT("Int32Label", "Integer");
+		static const FText FNameLabel = LOCTEXT("FNameLabel", "Name");
+		static const FText FStringLabel = LOCTEXT("FStringLabel", "String");
+
+		if(CPPType == BoolType)
+		{
+			return BoolLabel;
+		}
+		if(CPPType == FloatType || CPPType == DoubleType)
+		{
+			return FloatLabel;
+		}
+		if(CPPType == Int32Type)
+		{
+			return Int32Label;
+		}
+		if(CPPType == FNameType)
+		{
+			return FNameLabel;
+		}
+		if(CPPType == FStringType)
+		{
+			return FStringLabel;
+		}
+
+		return FText::FromString(CPPType);
+	}
+	return FText();
+}
+
 #endif
+
+#undef LOCTEXT_NAMESPACE
\ No newline at end of file
diff --git a/Engine/Plugins/Runtime/RigVM/Source/RigVMDeveloper/Private/RigVMModel/RigVMController.cpp b/Engine/Plugins/Runtime/RigVM/Source/RigVMDeveloper/Private/RigVMModel/RigVMController.cpp
--- a/Engine/Plugins/Runtime/RigVM/Source/RigVMDeveloper/Private/RigVMModel/RigVMController.cpp
+++ b/Engine/Plugins/Runtime/RigVM/Source/RigVMDeveloper/Private/RigVMModel/RigVMController.cpp
@@ -9866,7 +9866,14 @@ bool URigVMController::AddLink(URigVMPin* OutputPin, URigVMPin* InputPin, bool b
 
 				if(CastNode == nullptr)
 				{
-					CastNode = AddUnitNode(CastFunction->Struct, CastFunction->GetMethodName(), CastPosition, FString(), bSetupUndoRedo, false);
+					if(CastFunction->Factory)
+					{
+						CastNode = AddTemplateNode(CastFunction->Factory->GetTemplateNotation(), CastPosition, FString(), bSetupUndoRedo, false);
+					}
+					else
+					{
+						CastNode = AddUnitNode(CastFunction->Struct, CastFunction->GetMethodName(), CastPosition, FString(), bSetupUndoRedo, false);
+					}
 				}
 				
 				if(CastNode == nullptr)
diff --git a/Engine/Plugins/Runtime/RigVM/Source/RigVMDeveloper/Public/RigVMCompiler/RigVMCompiler.h b/Engine/Plugins/Runtime/RigVM/Source/RigVMDeveloper/Public/RigVMCompiler/RigVMCompiler.h
--- a/Engine/Plugins/Runtime/RigVM/Source/RigVMDeveloper/Public/RigVMCompiler/RigVMCompiler.h
+++ b/Engine/Plugins/Runtime/RigVM/Source/RigVMDeveloper/Public/RigVMCompiler/RigVMCompiler.h
@@ -310,7 +310,7 @@ public:
 
 private:
 
-	const URigVMLibraryNode* CurrentCompilationFunction;
+	const URigVMLibraryNode* CurrentCompilationFunction = nullptr;
 
 	TArray<URigVMPin*> GetLinkedPins(URigVMPin* InPin, bool bInputs = true, bool bOutputs = true, bool bRecursive = true);
 	uint16 GetElementSizeFromCPPType(const FString& InCPPType, UScriptStruct* InScriptStruct);
diff --git a/Engine/Plugins/Runtime/RigVM/Source/RigVMDeveloper/Public/RigVMDeveloperTypeUtils.h b/Engine/Plugins/Runtime/RigVM/Source/RigVMDeveloper/Public/RigVMDeveloperTypeUtils.h
--- a/Engine/Plugins/Runtime/RigVM/Source/RigVMDeveloper/Public/RigVMDeveloperTypeUtils.h
+++ b/Engine/Plugins/Runtime/RigVM/Source/RigVMDeveloper/Public/RigVMDeveloperTypeUtils.h
@@ -60,5 +60,7 @@ namespace RigVMTypeUtils
 
 	RIGVMDEVELOPER_API const FName& GetCastTemplateNotation();
 
+	// Get a user-facing name for an argument type
+	RIGVMDEVELOPER_API FText GetDisplayTextForArgumentType(const FRigVMTemplateArgumentType& InType);
 #endif
 }
diff --git a/Engine/Plugins/Runtime/RigVM/Source/RigVMEditor/Private/RigVMEditorModule.cpp b/Engine/Plugins/Runtime/RigVM/Source/RigVMEditor/Private/RigVMEditorModule.cpp
--- a/Engine/Plugins/Runtime/RigVM/Source/RigVMEditor/Private/RigVMEditorModule.cpp
+++ b/Engine/Plugins/Runtime/RigVM/Source/RigVMEditor/Private/RigVMEditorModule.cpp
@@ -490,34 +490,37 @@ void FRigVMEditorModule::GetInstanceActions(URigVMBlueprint* RigVMBlueprint, FBl
 	}
 }
 
-void FRigVMEditorModule::GetNodeContextMenuActions(URigVMBlueprint* RigVMBlueprint, const URigVMEdGraphNode* EdGraphNode, URigVMNode* ModelNode, UToolMenu* Menu) const
+void FRigVMEditorModule::GetNodeContextMenuActions(IRigVMClientHost* RigVMClientHost, const URigVMEdGraphNode* EdGraphNode, URigVMNode* ModelNode, UToolMenu* Menu) const
 {
-	// Only register menu actions for ourselves
-	UClass* BPClass = RigVMBlueprint->GetClass();
-	UClass* BaseClass = GetRigVMBlueprintClass();
-	if(BPClass != GetRigVMBlueprintClass() && !BPClass->IsChildOf(BaseClass))
+	// Only register menu actions for ourselves if we are a blueprint
+	if(URigVMBlueprint* RigVMBlueprint = Cast<URigVMBlueprint>(RigVMClientHost))
 	{
-		return;
+		UClass* BPClass = RigVMBlueprint->GetClass();
+		UClass* BaseClass = GetRigVMBlueprintClass();
+		if(BPClass != GetRigVMBlueprintClass() && !BPClass->IsChildOf(BaseClass))
+		{
+			return;
+		}
 	}
 
-	GetNodeWorkflowContextMenuActions(RigVMBlueprint, EdGraphNode, ModelNode, Menu);
-	GetNodeEventsContextMenuActions(RigVMBlueprint, EdGraphNode, ModelNode, Menu);
-	GetNodeConversionContextMenuActions(RigVMBlueprint, EdGraphNode, ModelNode, Menu);
-	GetNodeDebugContextMenuActions(RigVMBlueprint, EdGraphNode, ModelNode, Menu);
-	GetNodeVariablesContextMenuActions(RigVMBlueprint, EdGraphNode, ModelNode, Menu);
-	GetNodeTemplatesContextMenuActions(RigVMBlueprint, EdGraphNode, ModelNode, Menu);
-	GetNodeOrganizationContextMenuActions(RigVMBlueprint, EdGraphNode, ModelNode, Menu);
-	GetNodeVersioningContextMenuActions(RigVMBlueprint, EdGraphNode, ModelNode, Menu);
-	GetNodeTestContextMenuActions(RigVMBlueprint, EdGraphNode, ModelNode, Menu);
+	GetNodeWorkflowContextMenuActions(RigVMClientHost, EdGraphNode, ModelNode, Menu);
+	GetNodeEventsContextMenuActions(RigVMClientHost, EdGraphNode, ModelNode, Menu);
+	GetNodeConversionContextMenuActions(RigVMClientHost, EdGraphNode, ModelNode, Menu);
+	GetNodeDebugContextMenuActions(RigVMClientHost, EdGraphNode, ModelNode, Menu);
+	GetNodeVariablesContextMenuActions(RigVMClientHost, EdGraphNode, ModelNode, Menu);
+	GetNodeTemplatesContextMenuActions(RigVMClientHost, EdGraphNode, ModelNode, Menu);
+	GetNodeOrganizationContextMenuActions(RigVMClientHost, EdGraphNode, ModelNode, Menu);
+	GetNodeVersioningContextMenuActions(RigVMClientHost, EdGraphNode, ModelNode, Menu);
+	GetNodeTestContextMenuActions(RigVMClientHost, EdGraphNode, ModelNode, Menu);
 }
 
-void FRigVMEditorModule::GetNodeWorkflowContextMenuActions(URigVMBlueprint* RigVMBlueprint, const URigVMEdGraphNode* EdGraphNode, URigVMNode* ModelNode, UToolMenu* Menu) const
+void FRigVMEditorModule::GetNodeWorkflowContextMenuActions(IRigVMClientHost* RigVMClientHost, const URigVMEdGraphNode* EdGraphNode, URigVMNode* ModelNode, UToolMenu* Menu) const
 {
 	const TArray<FRigVMUserWorkflow> Workflows = ModelNode->GetSupportedWorkflows(ERigVMUserWorkflowType::NodeContext, ModelNode);
 	if(!Workflows.IsEmpty())
 	{
 		FToolMenuSection& SettingsSection = Menu->AddSection("RigVMEditorContextMenuWorkflow", LOCTEXT("WorkflowHeader", "Workflow"));
-		URigVMController* Controller = RigVMBlueprint->GetController(ModelNode->GetGraph());
+		URigVMController* Controller = RigVMClientHost->GetRigVMClient()->GetController(ModelNode->GetGraph());
 
 		for(const FRigVMUserWorkflow& Workflow : Workflows)
 		{
@@ -545,9 +548,10 @@ void FRigVMEditorModule::GetNodeWorkflowContextMenuActions(URigVMBlueprint* RigV
 	}
 }
 
-void FRigVMEditorModule::GetNodeEventsContextMenuActions(URigVMBlueprint* RigVMBlueprint, const URigVMEdGraphNode* EdGraphNode, URigVMNode* ModelNode, UToolMenu* Menu) const
+void FRigVMEditorModule::GetNodeEventsContextMenuActions(IRigVMClientHost* RigVMClientHost, const URigVMEdGraphNode* EdGraphNode, URigVMNode* ModelNode, UToolMenu* Menu) const
 {
-	if(ModelNode->IsEvent())
+	URigVMBlueprint* RigVMBlueprint = Cast<URigVMBlueprint>(RigVMClientHost);
+	if(RigVMBlueprint && ModelNode->IsEvent())
 	{
 		const FName& EventName = ModelNode->GetEventName();
 		const bool bCanRunOnce = !CastChecked<URigVMEdGraphSchema>(EdGraphNode->GetSchema())->IsRigVMDefaultEvent(EventName);
@@ -600,7 +604,7 @@ void FRigVMEditorModule::GetNodeEventsContextMenuActions(URigVMBlueprint* RigVMB
 	}
 }
 
-void FRigVMEditorModule::GetNodeConversionContextMenuActions(URigVMBlueprint* RigVMBlueprint, const URigVMEdGraphNode* EdGraphNode, URigVMNode* ModelNode, UToolMenu* Menu) const
+void FRigVMEditorModule::GetNodeConversionContextMenuActions(IRigVMClientHost* RigVMClientHost, const URigVMEdGraphNode* EdGraphNode, URigVMNode* ModelNode, UToolMenu* Menu) const
 {
 	if(URigVMDispatchNode* DispatchNode = Cast<URigVMDispatchNode>(ModelNode))
 	{
@@ -611,7 +615,7 @@ void FRigVMEditorModule::GetNodeConversionContextMenuActions(URigVMBlueprint* Ri
 				// if the value pin has only links on the root pin
 				if((ValuePin->GetSourceLinks(false).Num() > 0) && (ValuePin->GetTargetLinks(false).Num() > 0))
 				{
-					URigVMController* Controller = RigVMBlueprint->GetController(ModelNode->GetGraph());
+					URigVMController* Controller = RigVMClientHost->GetRigVMClient()->GetController(ModelNode->GetGraph());
 
 					FToolMenuSection& ConversionSection = Menu->AddSection("RigVMEditorContextMenuConversion", LOCTEXT("ConversionHeader", "Conversion"));
 					ConversionSection.AddMenuEntry(
@@ -645,157 +649,165 @@ void FRigVMEditorModule::GetNodeConversionContextMenuActions(URigVMBlueprint* Ri
 	}
 }
 
-void FRigVMEditorModule::GetNodeDebugContextMenuActions(URigVMBlueprint* RigVMBlueprint, const URigVMEdGraphNode* EdGraphNode, URigVMNode* ModelNode, UToolMenu* Menu) const
+void FRigVMEditorModule::GetNodeDebugContextMenuActions(IRigVMClientHost* RigVMClientHost, const URigVMEdGraphNode* EdGraphNode, URigVMNode* ModelNode, UToolMenu* Menu) const
 {
-	FToolMenuSection& DebugSection = Menu->AddSection("RigVMEditorContextMenuDebug", LOCTEXT("DebugHeader", "Debug"));
-	bool bNoneHasBreakpoint = true;
+	URigVMBlueprint* RigVMBlueprint = Cast<URigVMBlueprint>(RigVMClientHost);
+	if(RigVMBlueprint)
+	{
+		FToolMenuSection& DebugSection = Menu->AddSection("RigVMEditorContextMenuDebug", LOCTEXT("DebugHeader", "Debug"));
+		bool bNoneHasBreakpoint = true;
 
-	const URigVMGraph* Model = ModelNode->GetGraph();
-	URigVMController* Controller = RigVMBlueprint->GetController(Model);
+		const URigVMGraph* Model = ModelNode->GetGraph();
+		URigVMController* Controller = RigVMBlueprint->GetController(Model);
 
-	TArray<URigVMNode*> SelectedNodes;
-	TArray<FName> SelectedNodeNames = Model->GetSelectNodes();
-	SelectedNodeNames.AddUnique(ModelNode->GetFName());
+		TArray<URigVMNode*> SelectedNodes;
+		TArray<FName> SelectedNodeNames = Model->GetSelectNodes();
+		SelectedNodeNames.AddUnique(ModelNode->GetFName());
 
-	for (FName SelectedNodeName : SelectedNodeNames)
-	{
-		if (URigVMNode* FoundNode = Model->FindNodeByName(SelectedNodeName))
+		for (FName SelectedNodeName : SelectedNodeNames)
 		{
-			SelectedNodes.Add(FoundNode);
-			if (FoundNode->HasBreakpoint())
+			if (URigVMNode* FoundNode = Model->FindNodeByName(SelectedNodeName))
 			{
-				bNoneHasBreakpoint = false;
+				SelectedNodes.Add(FoundNode);
+				if (FoundNode->HasBreakpoint())
+				{
+					bNoneHasBreakpoint = false;
+				}
 			}
 		}
-	}
 	
-	if (bNoneHasBreakpoint)
-	{
-		DebugSection.AddMenuEntry(
-		"Add Breakpoint",
-		LOCTEXT("AddBreakpoint", "Add Breakpoint"),
-		LOCTEXT("AddBreakpoint_Tooltip", "Adds a breakpoint to the graph at this node"),
-		FSlateIcon(),
-		FUIAction(FExecuteAction::CreateLambda([Controller, SelectedNodes, RigVMBlueprint]()
+		if (bNoneHasBreakpoint)
 		{
-			for (URigVMNode* SelectedNode : SelectedNodes)
+			DebugSection.AddMenuEntry(
+			"Add Breakpoint",
+			LOCTEXT("AddBreakpoint", "Add Breakpoint"),
+			LOCTEXT("AddBreakpoint_Tooltip", "Adds a breakpoint to the graph at this node"),
+			FSlateIcon(),
+			FUIAction(FExecuteAction::CreateLambda([Controller, SelectedNodes, RigVMBlueprint]()
 			{
-				if (RigVMBlueprint->AddBreakpoint(SelectedNode))
+				for (URigVMNode* SelectedNode : SelectedNodes)
 				{
-					SelectedNode->SetHasBreakpoint(true);
+					if (RigVMBlueprint->AddBreakpoint(SelectedNode))
+					{
+						SelectedNode->SetHasBreakpoint(true);
+					}
 				}
-			}
-		})));
-	}
-	else
-	{						
-		DebugSection.AddMenuEntry(
-		"Remove Breakpoint",
-		LOCTEXT("RemoveBreakpoint", "Remove Breakpoint"),
-		LOCTEXT("RemoveBreakpoint_Tooltip", "Removes a breakpoint to the graph at this node"),
-		FSlateIcon(),
-		FUIAction(FExecuteAction::CreateLambda([Controller, SelectedNodes, RigVMBlueprint]()
-		{
-			for (URigVMNode* SelectedNode : SelectedNodes)
+			})));
+		}
+		else
+		{						
+			DebugSection.AddMenuEntry(
+			"Remove Breakpoint",
+			LOCTEXT("RemoveBreakpoint", "Remove Breakpoint"),
+			LOCTEXT("RemoveBreakpoint_Tooltip", "Removes a breakpoint to the graph at this node"),
+			FSlateIcon(),
+			FUIAction(FExecuteAction::CreateLambda([Controller, SelectedNodes, RigVMBlueprint]()
 			{
-				if (SelectedNode->HasBreakpoint())
+				for (URigVMNode* SelectedNode : SelectedNodes)
 				{
-					if (RigVMBlueprint->RemoveBreakpoint(SelectedNode))
+					if (SelectedNode->HasBreakpoint())
 					{
-						SelectedNode->SetHasBreakpoint(false);
+						if (RigVMBlueprint->RemoveBreakpoint(SelectedNode))
+						{
+							SelectedNode->SetHasBreakpoint(false);
+						}
 					}
-				}
-			}                            
-		})));
+				}                            
+			})));
+		}
 	}
 }
 
-void FRigVMEditorModule::GetNodeVariablesContextMenuActions(URigVMBlueprint* RigVMBlueprint, const URigVMEdGraphNode* EdGraphNode, URigVMNode* ModelNode, UToolMenu* Menu) const
+void FRigVMEditorModule::GetNodeVariablesContextMenuActions(IRigVMClientHost* RigVMClientHost, const URigVMEdGraphNode* EdGraphNode, URigVMNode* ModelNode, UToolMenu* Menu) const
 {
-	const URigVMGraph* Model = ModelNode->GetGraph();
-	URigVMController* Controller = RigVMBlueprint->GetController(Model);
-
-	if (URigVMVariableNode* VariableNode = Cast<URigVMVariableNode>(EdGraphNode->GetModelNode()))
+	URigVMBlueprint* RigVMBlueprint = Cast<URigVMBlueprint>(RigVMClientHost);
+	if (RigVMBlueprint)
 	{
-		FToolMenuSection& VariablesSection = Menu->AddSection("RigVMEditorContextMenuVariables", LOCTEXT("VariablesSettingsHeader", "Variables"));
-		VariablesSection.AddMenuEntry(
-			"MakePindingsFromVariableNode",
-			LOCTEXT("MakeBindingsFromVariableNode", "Make Bindings From Node"),
-			LOCTEXT("MakeBindingsFromVariableNode_Tooltip", "Turns the variable node into one ore more variable bindings on the pin(s)"),
-			FSlateIcon(),
-			FUIAction(FExecuteAction::CreateLambda([RigVMBlueprint, Controller, VariableNode]() {
-				Controller->MakeBindingsFromVariableNode(VariableNode->GetFName());
-			})
-		));
-	}
+		const URigVMGraph* Model = ModelNode->GetGraph();
+		URigVMController* Controller = RigVMBlueprint->GetController(Model);
 
-	if (URigVMFunctionReferenceNode* FunctionReferenceNode = Cast<URigVMFunctionReferenceNode>(EdGraphNode->GetModelNode()))
-	{
-		TSoftObjectPtr<URigVMFunctionReferenceNode> RefPtr(FunctionReferenceNode);
-		if(RefPtr.GetLongPackageName() != FunctionReferenceNode->GetReferencedFunctionHeader().LibraryPointer.LibraryNode.GetLongPackageName())
+		if (URigVMVariableNode* VariableNode = Cast<URigVMVariableNode>(EdGraphNode->GetModelNode()))
+		{
+			FToolMenuSection& VariablesSection = Menu->AddSection("RigVMEditorContextMenuVariables", LOCTEXT("VariablesSettingsHeader", "Variables"));
+			VariablesSection.AddMenuEntry(
+				"MakePindingsFromVariableNode",
+				LOCTEXT("MakeBindingsFromVariableNode", "Make Bindings From Node"),
+				LOCTEXT("MakeBindingsFromVariableNode_Tooltip", "Turns the variable node into one ore more variable bindings on the pin(s)"),
+				FSlateIcon(),
+				FUIAction(FExecuteAction::CreateLambda([RigVMBlueprint, Controller, VariableNode]() {
+					Controller->MakeBindingsFromVariableNode(VariableNode->GetFName());
+				})
+			));
+		}
+
+		if (URigVMFunctionReferenceNode* FunctionReferenceNode = Cast<URigVMFunctionReferenceNode>(EdGraphNode->GetModelNode()))
 		{
-			if(!FunctionReferenceNode->IsFullyRemapped() && FunctionReferenceNode->GetReferencedFunctionHeader().LibraryPointer.LibraryNode.ResolveObject())
+			TSoftObjectPtr<URigVMFunctionReferenceNode> RefPtr(FunctionReferenceNode);
+			if(RefPtr.GetLongPackageName() != FunctionReferenceNode->GetReferencedFunctionHeader().LibraryPointer.LibraryNode.GetLongPackageName())
 			{
-				FToolMenuSection& VariablesSection = Menu->AddSection("RigVMEditorContextMenuVariables", LOCTEXT("Variables", "Variables"));
-				VariablesSection.AddMenuEntry(
-                    "MakeVariablesFromFunctionReferenceNode",
-                    LOCTEXT("MakeVariablesFromFunctionReferenceNode", "Create required variables"),
-                    LOCTEXT("MakeVariablesFromFunctionReferenceNode_Tooltip", "Creates all required variables for this function and binds them"),
-                    FSlateIcon(),
-                    FUIAction(FExecuteAction::CreateLambda([Controller, FunctionReferenceNode, RigVMBlueprint]() {
-
-                        const TArray<FRigVMExternalVariable> ExternalVariables = FunctionReferenceNode->GetExternalVariables(false);
-                        if(!ExternalVariables.IsEmpty())
-                        {
-							FScopedTransaction Transaction(LOCTEXT("MakeVariablesFromFunctionReferenceNode", "Create required variables"));
-                            RigVMBlueprint->Modify();
-
-                            if (URigVMLibraryNode* LibraryNode = FunctionReferenceNode->LoadReferencedNode())
-                            {
-                                URigVMBlueprint* ReferencedBlueprint = LibraryNode->GetTypedOuter<URigVMBlueprint>();
-							   // ReferencedBlueprint != RigVMBlueprint - since only FunctionReferenceNodes from other assets have the potential to be unmapped
+				if(!FunctionReferenceNode->IsFullyRemapped() && FunctionReferenceNode->GetReferencedFunctionHeader().LibraryPointer.LibraryNode.ResolveObject())
+				{
+					FToolMenuSection& VariablesSection = Menu->AddSection("RigVMEditorContextMenuVariables", LOCTEXT("Variables", "Variables"));
+					VariablesSection.AddMenuEntry(
+						"MakeVariablesFromFunctionReferenceNode",
+						LOCTEXT("MakeVariablesFromFunctionReferenceNode", "Create required variables"),
+						LOCTEXT("MakeVariablesFromFunctionReferenceNode_Tooltip", "Creates all required variables for this function and binds them"),
+						FSlateIcon(),
+						FUIAction(FExecuteAction::CreateLambda([Controller, FunctionReferenceNode, RigVMBlueprint]() {
+
+							const TArray<FRigVMExternalVariable> ExternalVariables = FunctionReferenceNode->GetExternalVariables(false);
+							if(!ExternalVariables.IsEmpty())
+							{
+								FScopedTransaction Transaction(LOCTEXT("MakeVariablesFromFunctionReferenceNode", "Create required variables"));
+								RigVMBlueprint->Modify();
+
+								if (URigVMLibraryNode* LibraryNode = FunctionReferenceNode->LoadReferencedNode())
+								{
+									URigVMBlueprint* ReferencedBlueprint = LibraryNode->GetTypedOuter<URigVMBlueprint>();
+								   // ReferencedBlueprint != RigVMBlueprint - since only FunctionReferenceNodes from other assets have the potential to be unmapped
                             
-							   for(const FRigVMExternalVariable& ExternalVariable : ExternalVariables)
-							   {
-								   FString DefaultValue;
-								   if(ReferencedBlueprint)
+								   for(const FRigVMExternalVariable& ExternalVariable : ExternalVariables)
 								   {
-									   for(const FBPVariableDescription& NewVariable : ReferencedBlueprint->NewVariables)
+									   FString DefaultValue;
+									   if(ReferencedBlueprint)
 									   {
-										   if(NewVariable.VarName == ExternalVariable.Name)
+										   for(const FBPVariableDescription& NewVariable : ReferencedBlueprint->NewVariables)
 										   {
-											   DefaultValue = NewVariable.DefaultValue;
-											   break;
+											   if(NewVariable.VarName == ExternalVariable.Name)
+											   {
+												   DefaultValue = NewVariable.DefaultValue;
+												   break;
+											   }
 										   }
 									   }
-								   }
                                 
-								   FName NewVariableName = RigVMBlueprint->AddHostMemberVariableFromExternal(ExternalVariable, DefaultValue);
-								   if(!NewVariableName.IsNone())
-								   {
-									   Controller->SetRemappedVariable(FunctionReferenceNode, ExternalVariable.Name, NewVariableName);
+									   FName NewVariableName = RigVMBlueprint->AddHostMemberVariableFromExternal(ExternalVariable, DefaultValue);
+									   if(!NewVariableName.IsNone())
+									   {
+										   Controller->SetRemappedVariable(FunctionReferenceNode, ExternalVariable.Name, NewVariableName);
+									   }
 								   }
-							   }
-                            }
+								}
 
-                            FBlueprintEditorUtils::MarkBlueprintAsModified(RigVMBlueprint);
-                        }
+								FBlueprintEditorUtils::MarkBlueprintAsModified(RigVMBlueprint);
+							}
                         
-                    })
-                ));
+						})
+					));
+				}
 			}
 		}
 	}
 }
 
-void FRigVMEditorModule::GetNodeTemplatesContextMenuActions(URigVMBlueprint* RigVMBlueprint, const URigVMEdGraphNode* EdGraphNode, URigVMNode* ModelNode, UToolMenu* Menu) const
+void FRigVMEditorModule::GetNodeTemplatesContextMenuActions(IRigVMClientHost* RigVMClientHost, const URigVMEdGraphNode* EdGraphNode, URigVMNode* ModelNode, UToolMenu* Menu) const
 {
 	if (URigVMTemplateNode* TemplateNode = Cast<URigVMTemplateNode>(EdGraphNode->GetModelNode()))
 	{
 		if (!TemplateNode->IsSingleton())
 		{
 			const URigVMGraph* Model = ModelNode->GetGraph();
-			URigVMController* Controller = RigVMBlueprint->GetController(Model);
+			URigVMController* Controller = RigVMClientHost->GetRigVMClient()->GetController(Model);
 			
 			FToolMenuSection& TemplatesSection = Menu->AddSection("RigVMEditorContextMenuTemplates", LOCTEXT("TemplatesHeader", "Templates"));
 			TemplatesSection.AddMenuEntry(
@@ -812,10 +824,11 @@ void FRigVMEditorModule::GetNodeTemplatesContextMenuActions(URigVMBlueprint* Rig
 	}
 }
 
-void FRigVMEditorModule::GetNodeOrganizationContextMenuActions(URigVMBlueprint* RigVMBlueprint, const URigVMEdGraphNode* EdGraphNode, URigVMNode* ModelNode, UToolMenu* Menu) const
+void FRigVMEditorModule::GetNodeOrganizationContextMenuActions(IRigVMClientHost* RigVMClientHost, const URigVMEdGraphNode* EdGraphNode, URigVMNode* ModelNode, UToolMenu* Menu) const
 {
+	URigVMBlueprint* RigVMBlueprint = Cast<URigVMBlueprint>(RigVMClientHost);
 	const URigVMGraph* Model = ModelNode->GetGraph();
-	URigVMController* Controller = RigVMBlueprint->GetController(Model);
+	URigVMController* Controller = RigVMClientHost->GetRigVMClient()->GetController(Model);
 
 	FToolMenuSection& OrganizationSection = Menu->AddSection("RigVMEditorContextMenuOrganization", LOCTEXT("OrganizationHeader", "Organization"));
 	OrganizationSection.AddMenuEntry(
@@ -868,34 +881,37 @@ void FRigVMEditorModule::GetNodeOrganizationContextMenuActions(URigVMBlueprint*
 		));
 	}
 
-	if (URigVMFunctionReferenceNode* FunctionReferenceNode = Cast<URigVMFunctionReferenceNode>(EdGraphNode->GetModelNode()))
+	if(RigVMBlueprint)
 	{
-		TSoftObjectPtr<URigVMFunctionReferenceNode> RefPtr(FunctionReferenceNode);
-		if(RefPtr.GetLongPackageName() != FunctionReferenceNode->GetReferencedFunctionHeader().LibraryPointer.LibraryNode.GetLongPackageName())
+		if (URigVMFunctionReferenceNode* FunctionReferenceNode = Cast<URigVMFunctionReferenceNode>(EdGraphNode->GetModelNode()))
 		{
-			OrganizationSection.AddMenuEntry(
-			   "Localize Function",
-			   LOCTEXT("LocalizeFunction", "Localize Function"),
-			   LOCTEXT("LocalizeFunction_Tooltip", "Creates a local copy of the function backing the node."),
-			   FSlateIcon(),
-			   FUIAction(FExecuteAction::CreateLambda([RigVMBlueprint, FunctionReferenceNode]() {
-				   RigVMBlueprint->BroadcastRequestLocalizeFunctionDialog(FunctionReferenceNode->GetFunctionIdentifier(), true);
-			   })
-		    ));
+			TSoftObjectPtr<URigVMFunctionReferenceNode> RefPtr(FunctionReferenceNode);
+			if(RefPtr.GetLongPackageName() != FunctionReferenceNode->GetReferencedFunctionHeader().LibraryPointer.LibraryNode.GetLongPackageName())
+			{
+				OrganizationSection.AddMenuEntry(
+				   "Localize Function",
+				   LOCTEXT("LocalizeFunction", "Localize Function"),
+				   LOCTEXT("LocalizeFunction_Tooltip", "Creates a local copy of the function backing the node."),
+				   FSlateIcon(),
+				   FUIAction(FExecuteAction::CreateLambda([RigVMBlueprint, FunctionReferenceNode]() {
+					   RigVMBlueprint->BroadcastRequestLocalizeFunctionDialog(FunctionReferenceNode->GetFunctionIdentifier(), true);
+				   })
+				));
+			}
 		}
-	}
 
-	if (URigVMFunctionReferenceNode* FunctionRefNode = Cast<URigVMFunctionReferenceNode>(EdGraphNode->GetModelNode()))
-	{
-		OrganizationSection.AddMenuEntry(
-			"Promote To Collapse Node",
-			LOCTEXT("PromoteToCollapseNode", "Promote To Collapse Node"),
-			LOCTEXT("PromoteToCollapseNode_Tooltip", "Turns the Function Ref Node into a Collapse Node"),
-			FSlateIcon(),
-			FUIAction(FExecuteAction::CreateLambda([Controller, FunctionRefNode]() {
-				Controller->PromoteFunctionReferenceNodeToCollapseNode(FunctionRefNode->GetFName());
-				})
-			));
+		if (URigVMFunctionReferenceNode* FunctionRefNode = Cast<URigVMFunctionReferenceNode>(EdGraphNode->GetModelNode()))
+		{
+			OrganizationSection.AddMenuEntry(
+				"Promote To Collapse Node",
+				LOCTEXT("PromoteToCollapseNode", "Promote To Collapse Node"),
+				LOCTEXT("PromoteToCollapseNode_Tooltip", "Turns the Function Ref Node into a Collapse Node"),
+				FSlateIcon(),
+				FUIAction(FExecuteAction::CreateLambda([Controller, FunctionRefNode]() {
+					Controller->PromoteFunctionReferenceNodeToCollapseNode(FunctionRefNode->GetFName());
+					})
+				));
+		}
 	}
 
 	if (URigVMLibraryNode* LibraryNode = Cast<URigVMLibraryNode>(EdGraphNode->GetModelNode()))
@@ -943,10 +959,10 @@ void FRigVMEditorModule::GetNodeOrganizationContextMenuActions(URigVMBlueprint*
 	}));
 }
 
-void FRigVMEditorModule::GetNodeVersioningContextMenuActions(URigVMBlueprint* RigVMBlueprint, const URigVMEdGraphNode* EdGraphNode, URigVMNode* ModelNode, UToolMenu* Menu) const
+void FRigVMEditorModule::GetNodeVersioningContextMenuActions(IRigVMClientHost* RigVMClientHost, const URigVMEdGraphNode* EdGraphNode, URigVMNode* ModelNode, UToolMenu* Menu) const
 {
 	const URigVMGraph* Model = ModelNode->GetGraph();
-	URigVMController* Controller = RigVMBlueprint->GetController(Model);
+	URigVMController* Controller = RigVMClientHost->GetRigVMClient()->GetController(Model);
 
 	bool bCanNodeBeUpgraded = false;
 	TArray<FName> SelectedNodeNames = Model->GetSelectNodes();
@@ -976,7 +992,7 @@ void FRigVMEditorModule::GetNodeVersioningContextMenuActions(URigVMBlueprint* Ri
 	}
 }
 
-void FRigVMEditorModule::GetNodeTestContextMenuActions(URigVMBlueprint* RigVMBlueprint, const URigVMEdGraphNode* EdGraphNode, URigVMNode* ModelNode, UToolMenu* Menu) const
+void FRigVMEditorModule::GetNodeTestContextMenuActions(IRigVMClientHost* RigVMClientHost, const URigVMEdGraphNode* EdGraphNode, URigVMNode* ModelNode, UToolMenu* Menu) const
 {
 	// this struct is only available in EngineTest for now
 	static const FString DecoratorObjectPath = TEXT("/Script/EngineTestEditor.EngineTestRigVM_SimpleDecorator");
@@ -987,7 +1003,7 @@ void FRigVMEditorModule::GetNodeTestContextMenuActions(URigVMBlueprint* RigVMBlu
 	}
 
 	const URigVMGraph* Model = ModelNode->GetGraph();
-	URigVMController* Controller = RigVMBlueprint->GetController(Model);
+	URigVMController* Controller = RigVMClientHost->GetRigVMClient()->GetController(Model);
 
 	FToolMenuSection& EngineTestSection = Menu->AddSection("RigVMEditorContextMenuEngineTest", LOCTEXT("EngineTestHeader", "EngineTest"));
 	EngineTestSection.AddMenuEntry(
@@ -1006,25 +1022,25 @@ void FRigVMEditorModule::GetNodeTestContextMenuActions(URigVMBlueprint* RigVMBlu
 	);
 }
 
-void FRigVMEditorModule::GetPinContextMenuActions(URigVMBlueprint* RigVMBlueprint, const UEdGraphPin* EdGraphPin, URigVMPin* ModelPin, UToolMenu* Menu) const
+void FRigVMEditorModule::GetPinContextMenuActions(IRigVMClientHost* RigVMClientHost, const UEdGraphPin* EdGraphPin, URigVMPin* ModelPin, UToolMenu* Menu) const
 {
-	GetPinWorkflowContextMenuActions(RigVMBlueprint, EdGraphPin, ModelPin, Menu);
-	GetPinDebugContextMenuActions(RigVMBlueprint, EdGraphPin, ModelPin, Menu);
-	GetPinArrayContextMenuActions(RigVMBlueprint, EdGraphPin, ModelPin, Menu);
-	GetPinAggregateContextMenuActions(RigVMBlueprint, EdGraphPin, ModelPin, Menu);
-	GetPinTemplateContextMenuActions(RigVMBlueprint, EdGraphPin, ModelPin, Menu);
-	GetPinConversionContextMenuActions(RigVMBlueprint, EdGraphPin, ModelPin, Menu);
-	GetPinVariableContextMenuActions(RigVMBlueprint, EdGraphPin, ModelPin, Menu);
-	GetPinResetDefaultContextMenuActions(RigVMBlueprint, EdGraphPin, ModelPin, Menu);
-	GetPinInjectedNodesContextMenuActions(RigVMBlueprint, EdGraphPin, ModelPin, Menu);
+	GetPinWorkflowContextMenuActions(RigVMClientHost, EdGraphPin, ModelPin, Menu);
+	GetPinDebugContextMenuActions(RigVMClientHost, EdGraphPin, ModelPin, Menu);
+	GetPinArrayContextMenuActions(RigVMClientHost, EdGraphPin, ModelPin, Menu);
+	GetPinAggregateContextMenuActions(RigVMClientHost, EdGraphPin, ModelPin, Menu);
+	GetPinTemplateContextMenuActions(RigVMClientHost, EdGraphPin, ModelPin, Menu);
+	GetPinConversionContextMenuActions(RigVMClientHost, EdGraphPin, ModelPin, Menu);
+	GetPinVariableContextMenuActions(RigVMClientHost, EdGraphPin, ModelPin, Menu);
+	GetPinResetDefaultContextMenuActions(RigVMClientHost, EdGraphPin, ModelPin, Menu);
+	GetPinInjectedNodesContextMenuActions(RigVMClientHost, EdGraphPin, ModelPin, Menu);
 }
 
-void FRigVMEditorModule::GetPinWorkflowContextMenuActions(URigVMBlueprint* RigVMBlueprint, const UEdGraphPin* EdGraphPin, URigVMPin* ModelPin, UToolMenu* Menu) const
+void FRigVMEditorModule::GetPinWorkflowContextMenuActions(IRigVMClientHost* RigVMClientHost, const UEdGraphPin* EdGraphPin, URigVMPin* ModelPin, UToolMenu* Menu) const
 {
 	const TArray<FRigVMUserWorkflow> Workflows = ModelPin->GetNode()->GetSupportedWorkflows(ERigVMUserWorkflowType::PinContext, ModelPin);
 	if(!Workflows.IsEmpty())
 	{
-		URigVMController* Controller = RigVMBlueprint->GetController(ModelPin->GetGraph());
+		URigVMController* Controller = RigVMClientHost->GetRigVMClient()->GetController(ModelPin->GetGraph());
 
 		FToolMenuSection& SettingsSection = Menu->AddSection("RigVMEditorContextMenuWorkflow", LOCTEXT("WorkflowHeader", "Workflow"));
 
@@ -1054,31 +1070,35 @@ void FRigVMEditorModule::GetPinWorkflowContextMenuActions(URigVMBlueprint* RigVM
 	}
 }
 
-void FRigVMEditorModule::GetPinDebugContextMenuActions(URigVMBlueprint* RigVMBlueprint, const UEdGraphPin* EdGraphPin, URigVMPin* ModelPin, UToolMenu* Menu) const
+void FRigVMEditorModule::GetPinDebugContextMenuActions(IRigVMClientHost* RigVMClientHost, const UEdGraphPin* EdGraphPin, URigVMPin* ModelPin, UToolMenu* Menu) const
 {
-	const bool bIsEditablePin = !ModelPin->IsExecuteContext() && !ModelPin->IsWildCard();
-
-	if(Cast<URigVMEdGraphNode>(EdGraphPin->GetOwningNode()))
+	URigVMBlueprint* RigVMBlueprint = Cast<URigVMBlueprint>(RigVMClientHost);
+	if (RigVMBlueprint)
 	{
-		if(bIsEditablePin)
+		const bool bIsEditablePin = !ModelPin->IsExecuteContext() && !ModelPin->IsWildCard();
+
+		if(Cast<URigVMEdGraphNode>(EdGraphPin->GetOwningNode()))
 		{
-			// Add the watch pin / unwatch pin menu items
-			FToolMenuSection& Section = Menu->AddSection("RigVMEditorContextMenuWatches", LOCTEXT("WatchesHeader", "Watches"));
-			if (FKismetDebugUtilities::IsPinBeingWatched(RigVMBlueprint, EdGraphPin))
-			{
-				Section.AddMenuEntry(FGraphEditorCommands::Get().StopWatchingPin);
-			}
-			else
+			if(bIsEditablePin)
 			{
-				Section.AddMenuEntry(FGraphEditorCommands::Get().StartWatchingPin);
+				// Add the watch pin / unwatch pin menu items
+				FToolMenuSection& Section = Menu->AddSection("RigVMEditorContextMenuWatches", LOCTEXT("WatchesHeader", "Watches"));
+				if (FKismetDebugUtilities::IsPinBeingWatched(RigVMBlueprint, EdGraphPin))
+				{
+					Section.AddMenuEntry(FGraphEditorCommands::Get().StopWatchingPin);
+				}
+				else
+				{
+					Section.AddMenuEntry(FGraphEditorCommands::Get().StartWatchingPin);
+				}
 			}
 		}
 	}
 }
 
-void FRigVMEditorModule::GetPinArrayContextMenuActions(URigVMBlueprint* RigVMBlueprint, const UEdGraphPin* EdGraphPin, URigVMPin* ModelPin, UToolMenu* Menu) const
+void FRigVMEditorModule::GetPinArrayContextMenuActions(IRigVMClientHost* RigVMClientHost, const UEdGraphPin* EdGraphPin, URigVMPin* ModelPin, UToolMenu* Menu) const
 {
-	URigVMController* Controller = RigVMBlueprint->GetController(ModelPin->GetGraph());
+	URigVMController* Controller = RigVMClientHost->GetRigVMClient()->GetController(ModelPin->GetGraph());
 	if (ModelPin->IsArray() && !ModelPin->IsExecuteContext())
 	{
 		FToolMenuSection& Section = Menu->AddSection("RigVMEditorContextMenuPinArrays", LOCTEXT("PinArrays", "Arrays"));
@@ -1117,9 +1137,9 @@ void FRigVMEditorModule::GetPinArrayContextMenuActions(URigVMBlueprint* RigVMBlu
 	}
 }
 
-void FRigVMEditorModule::GetPinAggregateContextMenuActions(URigVMBlueprint* RigVMBlueprint, const UEdGraphPin* EdGraphPin, URigVMPin* ModelPin, UToolMenu* Menu) const
+void FRigVMEditorModule::GetPinAggregateContextMenuActions(IRigVMClientHost* RigVMClient, const UEdGraphPin* EdGraphPin, URigVMPin* ModelPin, UToolMenu* Menu) const
 {
-	URigVMController* Controller = RigVMBlueprint->GetController(ModelPin->GetGraph());
+	URigVMController* Controller = RigVMClient->GetRigVMClient()->GetController(ModelPin->GetGraph());
 	if (Cast<URigVMAggregateNode>(ModelPin->GetNode()))
 	{
 		FToolMenuSection& Section = Menu->AddSection("RigVMEditorContextMenuAggregatePin", LOCTEXT("AggregatePin", "Aggregates"));
@@ -1135,13 +1155,13 @@ void FRigVMEditorModule::GetPinAggregateContextMenuActions(URigVMBlueprint* RigV
 	}
 }
 
-void FRigVMEditorModule::GetPinTemplateContextMenuActions(URigVMBlueprint* RigVMBlueprint, const UEdGraphPin* EdGraphPin, URigVMPin* ModelPin, UToolMenu* Menu) const
+void FRigVMEditorModule::GetPinTemplateContextMenuActions(IRigVMClientHost* RigVMClientHost, const UEdGraphPin* EdGraphPin, URigVMPin* ModelPin, UToolMenu* Menu) const
 {
 	if (URigVMTemplateNode* TemplateNode = Cast<URigVMTemplateNode>(ModelPin->GetNode()))
 	{
 		if (!TemplateNode->IsSingleton())
 		{
-			URigVMController* Controller = RigVMBlueprint->GetController(ModelPin->GetGraph());
+			URigVMController* Controller = RigVMClientHost->GetRigVMClient()->GetController(ModelPin->GetGraph());
 			FToolMenuSection& TemplatesSection = Menu->AddSection("RigVMEditorContextMenuTemplates", LOCTEXT("TemplatesHeader", "Templates"));
 
 			if(const FRigVMTemplate* Template = TemplateNode->GetTemplate())
@@ -1155,11 +1175,10 @@ void FRigVMEditorModule::GetPinTemplateContextMenuActions(URigVMBlueprint* RigVM
 							TArray<TRigVMTypeIndex> ResolvedTypeIndices = Argument->GetSupportedTypeIndices(TemplateNode->GetResolvedPermutationIndices(true));
 							TSharedRef<SRigVMGraphChangePinType> ChangePinTypeWidget =
 							SNew(SRigVMGraphChangePinType)
-							.Blueprint(RigVMBlueprint)
 							.Types(ResolvedTypeIndices)
-							.OnTypeSelected_Lambda([RigVMBlueprint, ModelPin](const TRigVMTypeIndex& TypeSelected)
+							.OnTypeSelected_Lambda([RigVMClientHost, ModelPin](const TRigVMTypeIndex& TypeSelected)
 							{
-								if (URigVMController* Controller = RigVMBlueprint->GetController(ModelPin->GetGraph()))
+								if (URigVMController* Controller = RigVMClientHost->GetRigVMClient()->GetController(ModelPin->GetGraph()))
 								{
 									Controller->ResolveWildCardPin(ModelPin, TypeSelected, true, true);
 								}
@@ -1185,7 +1204,7 @@ void FRigVMEditorModule::GetPinTemplateContextMenuActions(URigVMBlueprint* RigVM
 	}
 }
 
-void FRigVMEditorModule::GetPinConversionContextMenuActions(URigVMBlueprint* RigVMBlueprint, const UEdGraphPin* EdGraphPin, URigVMPin* ModelPin, UToolMenu* Menu) const
+void FRigVMEditorModule::GetPinConversionContextMenuActions(IRigVMClientHost* RigVMClientHost, const UEdGraphPin* EdGraphPin, URigVMPin* ModelPin, UToolMenu* Menu) const
 {
 	if (URigVMRerouteNode* RerouteNode = Cast<URigVMRerouteNode>(ModelPin->GetNode()))
 	{
@@ -1195,7 +1214,7 @@ void FRigVMEditorModule::GetPinConversionContextMenuActions(URigVMBlueprint* Rig
 			{
 				if(!ValuePin->IsExecuteContext())
 				{
-					URigVMController* Controller = RigVMBlueprint->GetController(ModelPin->GetGraph());
+					URigVMController* Controller = RigVMClientHost->GetRigVMClient()->GetController(ModelPin->GetGraph());
 					FToolMenuSection& ConversionSection = Menu->AddSection("RigVMEditorContextMenuConversion", LOCTEXT("ConversionHeader", "Conversion"));
 
 					if(ValuePin->IsArray())
@@ -1243,76 +1262,79 @@ void FRigVMEditorModule::GetPinConversionContextMenuActions(URigVMBlueprint* Rig
 	}
 }
 
-void FRigVMEditorModule::GetPinVariableContextMenuActions(URigVMBlueprint* RigVMBlueprint, const UEdGraphPin* EdGraphPin, URigVMPin* ModelPin, UToolMenu* Menu) const
+void FRigVMEditorModule::GetPinVariableContextMenuActions(IRigVMClientHost* RigVMClientHost, const UEdGraphPin* EdGraphPin, URigVMPin* ModelPin, UToolMenu* Menu) const
 {
-	const bool bIsEditablePin = !ModelPin->IsExecuteContext() && !ModelPin->IsWildCard();
-
-	if (ModelPin->GetDirection() == ERigVMPinDirection::Input && bIsEditablePin)
+	if(URigVMBlueprint* RigVMBlueprint = Cast<URigVMBlueprint>(RigVMClientHost))
 	{
-		const UEdGraphNode* EdGraphNode = EdGraphPin->GetOwningNode();
-		URigVMController* Controller = RigVMBlueprint->GetController(ModelPin->GetGraph());
+		const bool bIsEditablePin = !ModelPin->IsExecuteContext() && !ModelPin->IsWildCard();
 
-		if (ModelPin->IsBoundToVariable())
+		if (ModelPin->GetDirection() == ERigVMPinDirection::Input && bIsEditablePin)
 		{
-			FVector2D NodePosition = FVector2D(EdGraphNode->NodePosX - 200.f, EdGraphNode->NodePosY);
+			const UEdGraphNode* EdGraphNode = EdGraphPin->GetOwningNode();
+			URigVMController* Controller = RigVMBlueprint->GetController(ModelPin->GetGraph());
 
-			FToolMenuSection& VariablesSection = Menu->AddSection("RigVMEditorContextMenuVariables", LOCTEXT("Variables", "Variables"));
-			VariablesSection.AddMenuEntry(
-				"MakeVariableNodeFromBinding",
-				LOCTEXT("MakeVariableNodeFromBinding", "Make Variable Node"),
-				LOCTEXT("MakeVariableNodeFromBinding_Tooltip", "Turns the variable binding on the pin to a variable node"),
-				FSlateIcon(),
-				FUIAction(FExecuteAction::CreateLambda([Controller, ModelPin, NodePosition]() {
-					Controller->MakeVariableNodeFromBinding(ModelPin->GetPinPath(), NodePosition, true, true);
-				})
-			));
-		}
-		else
-		{
-			FVector2D NodePosition = FVector2D(EdGraphNode->NodePosX - 200.f, EdGraphNode->NodePosY);
+			if (ModelPin->IsBoundToVariable())
+			{
+				FVector2D NodePosition = FVector2D(EdGraphNode->NodePosX - 200.f, EdGraphNode->NodePosY);
 
-			FToolMenuSection& VariablesSection = Menu->AddSection("RigVMEditorContextMenuVariables", LOCTEXT("Variables", "Variables"));
-			VariablesSection.AddMenuEntry(
-				"PromotePinToVariable",
-				LOCTEXT("PromotePinToVariable", "Promote Pin To Variable"),
-				LOCTEXT("PromotePinToVariable_Tooltip", "Turns the pin into a variable"),
-				FSlateIcon(),
-				FUIAction(FExecuteAction::CreateLambda([Controller, ModelPin, NodePosition]() {
+				FToolMenuSection& VariablesSection = Menu->AddSection("RigVMEditorContextMenuVariables", LOCTEXT("Variables", "Variables"));
+				VariablesSection.AddMenuEntry(
+					"MakeVariableNodeFromBinding",
+					LOCTEXT("MakeVariableNodeFromBinding", "Make Variable Node"),
+					LOCTEXT("MakeVariableNodeFromBinding_Tooltip", "Turns the variable binding on the pin to a variable node"),
+					FSlateIcon(),
+					FUIAction(FExecuteAction::CreateLambda([Controller, ModelPin, NodePosition]() {
+						Controller->MakeVariableNodeFromBinding(ModelPin->GetPinPath(), NodePosition, true, true);
+					})
+				));
+			}
+			else
+			{
+				FVector2D NodePosition = FVector2D(EdGraphNode->NodePosX - 200.f, EdGraphNode->NodePosY);
 
-					FModifierKeysState KeyState = FSlateApplication::Get().GetModifierKeys();
-					bool bCreateVariableNode = !KeyState.IsAltDown();
+				FToolMenuSection& VariablesSection = Menu->AddSection("RigVMEditorContextMenuVariables", LOCTEXT("Variables", "Variables"));
+				VariablesSection.AddMenuEntry(
+					"PromotePinToVariable",
+					LOCTEXT("PromotePinToVariable", "Promote Pin To Variable"),
+					LOCTEXT("PromotePinToVariable_Tooltip", "Turns the pin into a variable"),
+					FSlateIcon(),
+					FUIAction(FExecuteAction::CreateLambda([Controller, ModelPin, NodePosition]() {
 
-					Controller->PromotePinToVariable(ModelPin->GetPinPath(), bCreateVariableNode, NodePosition, true, true);
-				})
-			));
+						FModifierKeysState KeyState = FSlateApplication::Get().GetModifierKeys();
+						bool bCreateVariableNode = !KeyState.IsAltDown();
+
+						Controller->PromotePinToVariable(ModelPin->GetPinPath(), bCreateVariableNode, NodePosition, true, true);
+					})
+				));
+			}
 		}
-	}
 
-	if (Cast<URigVMUnitNode>(ModelPin->GetNode()) != nullptr || 
-		Cast<URigVMDispatchNode>(ModelPin->GetNode()) != nullptr || 
-		Cast<URigVMLibraryNode>(ModelPin->GetNode()) != nullptr)
-	{
-		if (ModelPin->GetDirection() == ERigVMPinDirection::Input &&
-			ModelPin->IsRootPin() &&
-			bIsEditablePin)
+		if (Cast<URigVMUnitNode>(ModelPin->GetNode()) != nullptr || 
+			Cast<URigVMDispatchNode>(ModelPin->GetNode()) != nullptr || 
+			Cast<URigVMLibraryNode>(ModelPin->GetNode()) != nullptr)
 		{
-			if (!ModelPin->IsBoundToVariable())
+			if (ModelPin->GetDirection() == ERigVMPinDirection::Input &&
+				ModelPin->IsRootPin() &&
+				bIsEditablePin)
 			{
-				FToolMenuSection& VariablesSection = Menu->FindOrAddSection(TEXT("Variables"));
+				if (!ModelPin->IsBoundToVariable())
+				{
+					FToolMenuSection& VariablesSection = Menu->FindOrAddSection(TEXT("Variables"));
 
-				TSharedRef<SRigVMGraphVariableBinding> VariableBindingWidget =
-					SNew(SRigVMGraphVariableBinding)
-					.Blueprint(RigVMBlueprint)
-					.ModelPins({ModelPin})
-					.CanRemoveBinding(false);
+					TSharedRef<SRigVMGraphVariableBinding> VariableBindingWidget =
+						SNew(SRigVMGraphVariableBinding)
+						.Blueprint(RigVMBlueprint)
+						.ModelPins({ModelPin})
+						.CanRemoveBinding(false);
 
-				VariablesSection.AddEntry(FToolMenuEntry::InitWidget("BindPinToVariableWidget", VariableBindingWidget, FText(), true));
+					VariablesSection.AddEntry(FToolMenuEntry::InitWidget("BindPinToVariableWidget", VariableBindingWidget, FText(), true));
+				}
 			}
 		}
 	}
 }
 
-void FRigVMEditorModule::GetPinResetDefaultContextMenuActions(URigVMBlueprint* RigVMBlueprint, const UEdGraphPin* EdGraphPin, URigVMPin* ModelPin, UToolMenu* Menu) const
+void FRigVMEditorModule::GetPinResetDefaultContextMenuActions(IRigVMClientHost* RigVMClientHost, const UEdGraphPin* EdGraphPin, URigVMPin* ModelPin, UToolMenu* Menu) const
 {
 	if (Cast<URigVMUnitNode>(ModelPin->GetNode()) != nullptr || 
 		Cast<URigVMDispatchNode>(ModelPin->GetNode()) != nullptr || 
@@ -1325,7 +1347,7 @@ void FRigVMEditorModule::GetPinResetDefaultContextMenuActions(URigVMBlueprint* R
 			bIsEditablePin)
 		{
 			FToolMenuSection& Section = Menu->AddSection("RigVMEditorContextMenuPinDefaults", LOCTEXT("PinDefaults", "Pin Defaults"));
-			URigVMController* Controller = RigVMBlueprint->GetController(ModelPin->GetGraph());
+			URigVMController* Controller = RigVMClientHost->GetRigVMClient()->GetController(ModelPin->GetGraph());
 
 			Section.AddMenuEntry(
 				"ResetPinDefaultValue",
@@ -1340,13 +1362,13 @@ void FRigVMEditorModule::GetPinResetDefaultContextMenuActions(URigVMBlueprint* R
 	}
 }
 
-void FRigVMEditorModule::GetPinInjectedNodesContextMenuActions(URigVMBlueprint* RigVMBlueprint, const UEdGraphPin* EdGraphPin, URigVMPin* ModelPin, UToolMenu* Menu) const
+void FRigVMEditorModule::GetPinInjectedNodesContextMenuActions(IRigVMClientHost* RigVMClientHost, const UEdGraphPin* EdGraphPin, URigVMPin* ModelPin, UToolMenu* Menu) const
 {
 	if (ModelPin->GetRootPin() == ModelPin && (
 	Cast<URigVMUnitNode>(ModelPin->GetNode()) != nullptr ||
 	Cast<URigVMLibraryNode>(ModelPin->GetNode()) != nullptr))
 	{
-		URigVMController* Controller = RigVMBlueprint->GetController(ModelPin->GetGraph());
+		URigVMController* Controller = RigVMClientHost->GetRigVMClient()->GetController(ModelPin->GetGraph());
 
 		if (ModelPin->HasInjectedNodes())
 		{
@@ -1438,10 +1460,10 @@ void FRigVMEditorModule::GetPinInjectedNodesContextMenuActions(URigVMBlueprint*
 						LOCTEXT("EditAlphaInterp", "Edit Interpolate"),
 						LOCTEXT("EditAlphaInterp_Tooltip", "Edit the interpolate node"),
 						FSlateIcon(),
-						FUIAction(FExecuteAction::CreateLambda([RigVMBlueprint, InterpNode]() {
+						FUIAction(FExecuteAction::CreateLambda([RigVMClientHost, InterpNode]() {
 						TArray<FName> NodeNames;
 						NodeNames.Add(InterpNode->GetFName());
-						RigVMBlueprint->GetController(InterpNode->GetGraph())->SetNodeSelection(NodeNames);
+						RigVMClientHost->GetRigVMClient()->GetController(InterpNode->GetGraph())->SetNodeSelection(NodeNames);
 					})
 						));
 					Section.AddMenuEntry(
@@ -1514,7 +1536,7 @@ void FRigVMEditorModule::GetPinInjectedNodesContextMenuActions(URigVMBlueprint*
 						LOCTEXT("AddVisualDebug", "Add Visual Debug"),
 						LOCTEXT("AddVisualDebug_Tooltip", "Injects a visual debugging node"),
 						FSlateIcon(),
-						FUIAction(FExecuteAction::CreateLambda([RigVMBlueprint, Controller, EdGraphPin, ModelPin, ScriptStruct]() {
+						FUIAction(FExecuteAction::CreateLambda([RigVMClientHost, Controller, EdGraphPin, ModelPin, ScriptStruct]() {
 							URigVMInjectionInfo* Injection = Controller->AddInjectedNode(ModelPin->GetPinPath(), ModelPin->GetDirection() != ERigVMPinDirection::Output, ScriptStruct, FRigVMStruct::ExecuteName, TEXT("Value"), TEXT("Value"), FString(), true, true);
 							if (Injection)
 							{
@@ -1609,24 +1631,29 @@ void FRigVMEditorModule::GetContextMenuActions(const URigVMEdGraphSchema* Schema
 	{
 		Schema->UEdGraphSchema::GetContextMenuActions(Menu, Context);
 
-		if (const UEdGraphPin* InGraphPin = (UEdGraphPin* )Context->Pin)
+		if (const UEdGraphPin* InGraphPin = (UEdGraphPin*)Context->Pin)
 		{
-			if(URigVMBlueprint* RigVMBlueprint = Cast<URigVMBlueprint>((UBlueprint*)Context->Blueprint))
+			if(const UEdGraph* Graph = Context->Graph)
 			{
-				const UEdGraph* Graph = InGraphPin->GetOwningNode()->GetGraph();
-				if (URigVMPin* ModelPin = RigVMBlueprint->GetModel(Graph)->FindPin(InGraphPin->GetName()))
+				if(IRigVMClientHost* RigVMClientHost = Graph->GetImplementingOuter<IRigVMClientHost>())
 				{
-					GetPinContextMenuActions(RigVMBlueprint, InGraphPin, ModelPin, Menu);
+					if (URigVMPin* ModelPin = RigVMClientHost->GetRigVMClient()->GetModel(Graph)->FindPin(InGraphPin->GetName()))
+					{
+						GetPinContextMenuActions(RigVMClientHost, InGraphPin, ModelPin, Menu);
+					}
 				}
 			}
 		}
 		else if(const URigVMEdGraphNode* EdGraphNode = Cast<URigVMEdGraphNode>(Context->Node))
 		{
-			if (URigVMBlueprint* RigVMBlueprint = Cast<URigVMBlueprint>((UBlueprint*)Context->Blueprint))
+			if (const UEdGraph* Graph = Context->Graph)
 			{
-				if(URigVMNode* ModelNode = EdGraphNode->GetModelNode())
+				if (IRigVMClientHost* RigVMClientHost = Graph->GetImplementingOuter<IRigVMClientHost>())
 				{
-					GetNodeContextMenuActions(RigVMBlueprint, EdGraphNode, ModelNode, Menu);
+					if(URigVMNode* ModelNode = EdGraphNode->GetModelNode())
+					{
+						GetNodeContextMenuActions(RigVMClientHost, EdGraphNode, ModelNode, Menu);
+					}
 				}
 			}
 		}
diff --git a/Engine/Plugins/Runtime/RigVM/Source/RigVMEditor/Private/Widgets/SRigVMGraphChangePinType.cpp b/Engine/Plugins/Runtime/RigVM/Source/RigVMEditor/Private/Widgets/SRigVMGraphChangePinType.cpp
--- a/Engine/Plugins/Runtime/RigVM/Source/RigVMEditor/Private/Widgets/SRigVMGraphChangePinType.cpp
+++ b/Engine/Plugins/Runtime/RigVM/Source/RigVMEditor/Private/Widgets/SRigVMGraphChangePinType.cpp
@@ -20,7 +20,6 @@ static const FText RigVMEdGraphChangePinTypeMultipleValues = LOCTEXT("MultipleVa
 void SRigVMGraphChangePinType::Construct(const FArguments& InArgs)
 {
 	this->Types = InArgs._Types;
-	this->Blueprint = InArgs._Blueprint;
 	this->OnTypeSelected = InArgs._OnTypeSelected;
 
 	IPropertyAccessEditor& PropertyAccessEditor = IModularFeatures::Get().GetModularFeature<IPropertyAccessEditor>("PropertyAccessEditor");
@@ -47,61 +46,13 @@ void SRigVMGraphChangePinType::Construct(const FArguments& InArgs)
 
 	this->ChildSlot
 	[
-		PropertyAccessEditor.MakePropertyBindingWidget(Blueprint, BindingArgs)
+		PropertyAccessEditor.MakePropertyBindingWidget(nullptr, BindingArgs)
 	];
 }
 
 FText SRigVMGraphChangePinType::GetBindingText(const FRigVMTemplateArgumentType& InType)
 {
-	if(UObject* CPPTypeObject = InType.CPPTypeObject)
-	{
-		if(const UScriptStruct* ScriptStruct = Cast<UScriptStruct>(CPPTypeObject))
-		{
-			return ScriptStruct->GetDisplayNameText();
-		}
-		if(const UEnum* Enum = Cast<UEnum>(CPPTypeObject))
-		{
-			return Enum->GetDisplayNameText();
-		}
-	}
-	else
-	{
-		FString CPPType = InType.CPPType.ToString();
-		if(RigVMTypeUtils::IsArrayType(CPPType))
-		{
-			CPPType = RigVMTypeUtils::BaseTypeFromArrayType(CPPType);
-		}
-
-		static const FText BoolLabel = LOCTEXT("BoolLabel", "Boolean");
-		static const FText FloatLabel = LOCTEXT("FloatLabel", "Float");
-		static const FText Int32Label = LOCTEXT("Int32Label", "Integer");
-		static const FText FNameLabel = LOCTEXT("FNameLabel", "Name");
-		static const FText FStringLabel = LOCTEXT("FStringLabel", "String");
-
-		if(CPPType == RigVMTypeUtils::BoolType)
-		{
-			return BoolLabel;
-		}
-		if(CPPType == RigVMTypeUtils::FloatType || CPPType == RigVMTypeUtils::DoubleType)
-		{
-			return FloatLabel;
-		}
-		if(CPPType == RigVMTypeUtils::Int32Type)
-		{
-			return Int32Label;
-		}
-		if(CPPType == RigVMTypeUtils::FNameType)
-		{
-			return FNameLabel;
-		}
-		if(CPPType == RigVMTypeUtils::FStringType)
-		{
-			return FStringLabel;
-		}
-
-		return FText::FromString(CPPType);
-	}
-	return FText();
+	return RigVMTypeUtils::GetDisplayTextForArgumentType(InType);
 }
 
 FText SRigVMGraphChangePinType::GetBindingText() const
diff --git a/Engine/Plugins/Runtime/RigVM/Source/RigVMEditor/Public/RigVMEditorModule.h b/Engine/Plugins/Runtime/RigVM/Source/RigVMEditor/Public/RigVMEditorModule.h
--- a/Engine/Plugins/Runtime/RigVM/Source/RigVMEditor/Public/RigVMEditorModule.h
+++ b/Engine/Plugins/Runtime/RigVM/Source/RigVMEditor/Public/RigVMEditorModule.h
@@ -30,8 +30,8 @@ public:
 	virtual void GetContextMenuActions(const URigVMEdGraphSchema* Schema, class UToolMenu* Menu, class UGraphNodeContextMenuContext* Context) const = 0;
 	virtual void GetTypeActions(URigVMBlueprint* RigVMBlueprint, FBlueprintActionDatabaseRegistrar& ActionRegistrar) = 0;
 	virtual void GetInstanceActions(URigVMBlueprint* RigVMBlueprint, FBlueprintActionDatabaseRegistrar& ActionRegistrar) = 0;
-	virtual void GetNodeContextMenuActions(URigVMBlueprint* RigVMBlueprint, const URigVMEdGraphNode* EdGraphNode, URigVMNode* ModelNode, class UToolMenu* Menu) const = 0;
-	virtual void GetPinContextMenuActions(URigVMBlueprint* RigVMBlueprint, const UEdGraphPin* EdGraphPin, URigVMPin* ModelPin, class UToolMenu* Menu) const = 0;
+	virtual void GetNodeContextMenuActions(IRigVMClientHost* RigVMClientHost, const URigVMEdGraphNode* EdGraphNode, URigVMNode* ModelNode, class UToolMenu* Menu) const = 0;
+	virtual void GetPinContextMenuActions(IRigVMClientHost* RigVMClientHost, const UEdGraphPin* EdGraphPin, URigVMPin* ModelPin, class UToolMenu* Menu) const = 0;
 	virtual FConnectionDrawingPolicy* CreateConnectionDrawingPolicy(int32 InBackLayerID, int32 InFrontLayerID, float InZoomFactor, const FSlateRect& InClippingRect, class FSlateWindowElementList& InDrawElements, class UEdGraph* InGraphObj) const = 0;
 };
 
@@ -53,8 +53,8 @@ public:
 	const URigVMBlueprint* GetRigVMBlueprintCDO() const;
 	virtual void GetTypeActions(URigVMBlueprint* RigVMBlueprint, FBlueprintActionDatabaseRegistrar& ActionRegistrar);
 	virtual void GetInstanceActions(URigVMBlueprint* RigVMBlueprint, FBlueprintActionDatabaseRegistrar& ActionRegistrar);
-	virtual void GetNodeContextMenuActions(URigVMBlueprint* RigVMBlueprint, const URigVMEdGraphNode* EdGraphNode, URigVMNode* ModelNode, class UToolMenu* Menu) const;
-	virtual void GetPinContextMenuActions(URigVMBlueprint* RigVMBlueprint, const UEdGraphPin* EdGraphPin, URigVMPin* ModelPin, class UToolMenu* Menu) const;
+	virtual void GetNodeContextMenuActions(IRigVMClientHost* RigVMClientHost, const URigVMEdGraphNode* EdGraphNode, URigVMNode* ModelNode, class UToolMenu* Menu) const;
+	virtual void GetPinContextMenuActions(IRigVMClientHost* RigVMClientHost, const UEdGraphPin* EdGraphPin, URigVMPin* ModelPin, class UToolMenu* Menu) const;
 
 	/** IHasMenuExtensibility interface */
 	virtual TSharedPtr<FExtensibilityManager> GetMenuExtensibilityManager() override { return MenuExtensibilityManager; }
@@ -82,24 +82,24 @@ protected:
 	bool IsRigVMEditorModuleBase() const;
 
 	/** Specific section callbacks for the context menu */
-	virtual void GetNodeWorkflowContextMenuActions(URigVMBlueprint* RigVMBlueprint, const URigVMEdGraphNode* EdGraphNode, URigVMNode* ModelNode, UToolMenu* Menu) const;
-	virtual void GetNodeEventsContextMenuActions(URigVMBlueprint* RigVMBlueprint, const URigVMEdGraphNode* EdGraphNode, URigVMNode* ModelNode, UToolMenu* Menu) const;
-	virtual void GetNodeConversionContextMenuActions(URigVMBlueprint* RigVMBlueprint, const URigVMEdGraphNode* EdGraphNode, URigVMNode* ModelNode, UToolMenu* Menu) const;
-	virtual void GetNodeDebugContextMenuActions(URigVMBlueprint* RigVMBlueprint, const URigVMEdGraphNode* EdGraphNode, URigVMNode* ModelNode, UToolMenu* Menu) const;
-	virtual void GetNodeVariablesContextMenuActions(URigVMBlueprint* RigVMBlueprint, const URigVMEdGraphNode* EdGraphNode, URigVMNode* ModelNode, UToolMenu* Menu) const;
-	virtual void GetNodeTemplatesContextMenuActions(URigVMBlueprint* RigVMBlueprint, const URigVMEdGraphNode* EdGraphNode, URigVMNode* ModelNode, UToolMenu* Menu) const;
-	virtual void GetNodeOrganizationContextMenuActions(URigVMBlueprint* RigVMBlueprint, const URigVMEdGraphNode* EdGraphNode, URigVMNode* ModelNode, UToolMenu* Menu) const;
-	virtual void GetNodeVersioningContextMenuActions(URigVMBlueprint* RigVMBlueprint, const URigVMEdGraphNode* EdGraphNode, URigVMNode* ModelNode, UToolMenu* Menu) const;
-	virtual void GetNodeTestContextMenuActions(URigVMBlueprint* RigVMBlueprint, const URigVMEdGraphNode* EdGraphNode, URigVMNode* ModelNode, UToolMenu* Menu) const;
-	virtual void GetPinWorkflowContextMenuActions(URigVMBlueprint* RigVMBlueprint, const UEdGraphPin* EdGraphPin, URigVMPin* ModelPin, UToolMenu* Menu) const;
-	virtual void GetPinDebugContextMenuActions(URigVMBlueprint* RigVMBlueprint, const UEdGraphPin* EdGraphPin, URigVMPin* ModelPin, UToolMenu* Menu) const;
-	virtual void GetPinArrayContextMenuActions(URigVMBlueprint* RigVMBlueprint, const UEdGraphPin* EdGraphPin, URigVMPin* ModelPin, UToolMenu* Menu) const;
-	virtual void GetPinAggregateContextMenuActions(URigVMBlueprint* RigVMBlueprint, const UEdGraphPin* EdGraphPin, URigVMPin* ModelPin, UToolMenu* Menu) const;
-	virtual void GetPinTemplateContextMenuActions(URigVMBlueprint* RigVMBlueprint, const UEdGraphPin* EdGraphPin, URigVMPin* ModelPin, UToolMenu* Menu) const;
-	virtual void GetPinConversionContextMenuActions(URigVMBlueprint* RigVMBlueprint, const UEdGraphPin* EdGraphPin, URigVMPin* ModelPin, UToolMenu* Menu) const;
-	virtual void GetPinVariableContextMenuActions(URigVMBlueprint* RigVMBlueprint, const UEdGraphPin* EdGraphPin, URigVMPin* ModelPin, UToolMenu* Menu) const;
-	virtual void GetPinResetDefaultContextMenuActions(URigVMBlueprint* RigVMBlueprint, const UEdGraphPin* EdGraphPin, URigVMPin* ModelPin, UToolMenu* Menu) const;
-	virtual void GetPinInjectedNodesContextMenuActions(URigVMBlueprint* RigVMBlueprint, const UEdGraphPin* EdGraphPin, URigVMPin* ModelPin, UToolMenu* Menu) const;
+	virtual void GetNodeWorkflowContextMenuActions(IRigVMClientHost* RigVMClientHost, const URigVMEdGraphNode* EdGraphNode, URigVMNode* ModelNode, UToolMenu* Menu) const;
+	virtual void GetNodeEventsContextMenuActions(IRigVMClientHost* RigVMClientHost, const URigVMEdGraphNode* EdGraphNode, URigVMNode* ModelNode, UToolMenu* Menu) const;
+	virtual void GetNodeConversionContextMenuActions(IRigVMClientHost* RigVMClientHost, const URigVMEdGraphNode* EdGraphNode, URigVMNode* ModelNode, UToolMenu* Menu) const;
+	virtual void GetNodeDebugContextMenuActions(IRigVMClientHost* RigVMClientHost, const URigVMEdGraphNode* EdGraphNode, URigVMNode* ModelNode, UToolMenu* Menu) const;
+	virtual void GetNodeVariablesContextMenuActions(IRigVMClientHost* RigVMClientHost, const URigVMEdGraphNode* EdGraphNode, URigVMNode* ModelNode, UToolMenu* Menu) const;
+	virtual void GetNodeTemplatesContextMenuActions(IRigVMClientHost* RigVMClientHost, const URigVMEdGraphNode* EdGraphNode, URigVMNode* ModelNode, UToolMenu* Menu) const;
+	virtual void GetNodeOrganizationContextMenuActions(IRigVMClientHost* RigVMClientHost, const URigVMEdGraphNode* EdGraphNode, URigVMNode* ModelNode, UToolMenu* Menu) const;
+	virtual void GetNodeVersioningContextMenuActions(IRigVMClientHost* RigVMClientHost, const URigVMEdGraphNode* EdGraphNode, URigVMNode* ModelNode, UToolMenu* Menu) const;
+	virtual void GetNodeTestContextMenuActions(IRigVMClientHost* RigVMClientHost, const URigVMEdGraphNode* EdGraphNode, URigVMNode* ModelNode, UToolMenu* Menu) const;
+	virtual void GetPinWorkflowContextMenuActions(IRigVMClientHost* RigVMClientHost, const UEdGraphPin* EdGraphPin, URigVMPin* ModelPin, UToolMenu* Menu) const;
+	virtual void GetPinDebugContextMenuActions(IRigVMClientHost* RigVMClientHost, const UEdGraphPin* EdGraphPin, URigVMPin* ModelPin, UToolMenu* Menu) const;
+	virtual void GetPinArrayContextMenuActions(IRigVMClientHost* RigVMClientHost, const UEdGraphPin* EdGraphPin, URigVMPin* ModelPin, UToolMenu* Menu) const;
+	virtual void GetPinAggregateContextMenuActions(IRigVMClientHost* RigVMClientHost, const UEdGraphPin* EdGraphPin, URigVMPin* ModelPin, UToolMenu* Menu) const;
+	virtual void GetPinTemplateContextMenuActions(IRigVMClientHost* RigVMClientHost, const UEdGraphPin* EdGraphPin, URigVMPin* ModelPin, UToolMenu* Menu) const;
+	virtual void GetPinConversionContextMenuActions(IRigVMClientHost* RigVMClientHost, const UEdGraphPin* EdGraphPin, URigVMPin* ModelPin, UToolMenu* Menu) const;
+	virtual void GetPinVariableContextMenuActions(IRigVMClientHost* RigVMClientHost, const UEdGraphPin* EdGraphPin, URigVMPin* ModelPin, UToolMenu* Menu) const;
+	virtual void GetPinResetDefaultContextMenuActions(IRigVMClientHost* RigVMClientHost, const UEdGraphPin* EdGraphPin, URigVMPin* ModelPin, UToolMenu* Menu) const;
+	virtual void GetPinInjectedNodesContextMenuActions(IRigVMClientHost* RigVMClientHost, const UEdGraphPin* EdGraphPin, URigVMPin* ModelPin, UToolMenu* Menu) const;
 
 	/** Extensibility managers */
 	TSharedPtr<FExtensibilityManager> MenuExtensibilityManager;
diff --git a/Engine/Plugins/Runtime/RigVM/Source/RigVMEditor/Public/Widgets/SRigVMGraphChangePinType.h b/Engine/Plugins/Runtime/RigVM/Source/RigVMEditor/Public/Widgets/SRigVMGraphChangePinType.h
--- a/Engine/Plugins/Runtime/RigVM/Source/RigVMEditor/Public/Widgets/SRigVMGraphChangePinType.h
+++ b/Engine/Plugins/Runtime/RigVM/Source/RigVMEditor/Public/Widgets/SRigVMGraphChangePinType.h
@@ -18,12 +18,10 @@ public:
 
 	SLATE_BEGIN_ARGS(SRigVMGraphChangePinType)
 	: _Types()
-    , _Blueprint(nullptr)
 	, _OnTypeSelected(nullptr)
 	{}
 
 		SLATE_ARGUMENT(TArray<TRigVMTypeIndex>, Types)
-		SLATE_ARGUMENT(URigVMBlueprint*, Blueprint)
 		SLATE_EVENT(FOnTypeSelected, OnTypeSelected)
 
 	SLATE_END_ARGS()
@@ -44,7 +42,6 @@ protected:
 	void HandlePinTypeChanged(FRigVMTemplateArgumentType InType);
 
 	TArray<TRigVMTypeIndex> Types;
-	URigVMBlueprint* Blueprint;
 	FOnTypeSelected OnTypeSelected;
 	FPropertyBindingWidgetArgs BindingArgs;
 };
