RigVM : Removed object based storage and the option to use it with a define. Now we run based on struct based storage only.

diff --git a/Engine/Plugins/Animation/ControlRig/Source/ControlRig/Private/ControlRig.cpp b/Engine/Plugins/Animation/ControlRig/Source/ControlRig/Private/ControlRig.cpp
--- a/Engine/Plugins/Animation/ControlRig/Source/ControlRig/Private/ControlRig.cpp
+++ b/Engine/Plugins/Animation/ControlRig/Source/ControlRig/Private/ControlRig.cpp
@@ -1338,12 +1338,6 @@ bool UControlRig::Execute_Internal(const FName& InEventName)
 			if(!VM->IsContextValidForExecution(Context)
 				|| !IsValidLowLevel()
 				|| !VM->IsValidLowLevel()
-#if !UE_RIGVM_PROPERTY_BAG_STORAGE_ENABLED
-				|| !GetLiteralMemory()
-				|| !GetLiteralMemory()->IsValidLowLevel()
-				|| !GetWorkMemory()
-				||!GetWorkMemory()->IsValidLowLevel()
-#endif // !UE_RIGVM_PROPERTY_BAG_STORAGE_ENABLED
 			)
 			{
 				UE_LOG(LogControlRig, Warning, InvalidatedVMFormat, *GetClass()->GetName());
@@ -1413,7 +1407,7 @@ bool UControlRig::Execute_Internal(const FName& InEventName)
 		TGuardValue<FString> ModuleInstanceNameSpaceGuard(PublicContext.ModuleInstanceNameSpace, GetRigModuleNameSpace());
 		FRigHierarchyRedirectorGuard ElementRedirectorGuard(this);
 
-		TArray<TRigVMMemoryStorage*> LocalMemory = VM->GetLocalMemoryArray(Context);
+		TArray<FRigVMMemoryStorageStruct*> LocalMemory = VM->GetLocalMemoryArray(Context);
 		const bool bSuccess = VM->Execute(Context, LocalMemory, InEventName) != ERigVMExecuteResult::Failed;
 
 #if UE_RIGVM_PROFILE_EXECUTE_UNITS_NUM
diff --git a/Engine/Plugins/Animation/ControlRig/Source/ControlRigEditor/Private/ControlRigCompilerDetails.cpp b/Engine/Plugins/Animation/ControlRig/Source/ControlRigEditor/Private/ControlRigCompilerDetails.cpp
--- a/Engine/Plugins/Animation/ControlRig/Source/ControlRigEditor/Private/ControlRigCompilerDetails.cpp
+++ b/Engine/Plugins/Animation/ControlRig/Source/ControlRigEditor/Private/ControlRigCompilerDetails.cpp
@@ -192,15 +192,10 @@ FReply FRigVMCompileSettingsDetails::OnInspectMemory(ERigVMMemoryType InMemoryTy
 	{
 		if(UControlRig* DebuggedRig = Cast<UControlRig>(BlueprintBeingCustomized->GetObjectBeingDebugged()))
 		{
-			if(TRigVMMemoryStorage* MemoryStorage = DebuggedRig->GetMemoryByType(InMemoryType))
+			if(FRigVMMemoryStorageStruct* MemoryStorage = DebuggedRig->GetMemoryByType(InMemoryType))
 			{
-#if UE_RIGVM_PROPERTY_BAG_STORAGE_ENABLED
-				TArray<TRigVMMemoryStorage*> InStructs = { MemoryStorage };
+				TArray<FRigVMMemoryStorageStruct*> InStructs = { MemoryStorage };
 				BlueprintBeingCustomized->RequestInspectMemoryStorage(InStructs);
-#else
-				TArray<UObject*> ObjectsToSelect = {MemoryStorage};
-				BlueprintBeingCustomized->RequestInspectObject(ObjectsToSelect);
-#endif // !UE_RIGVM_PROPERTY_BAG_STORAGE_ENABLED
 			}
 		}
 	}
diff --git a/Engine/Plugins/Experimental/AnimNext/Source/AnimNext/Private/Param/AnimNextParameterBlock.cpp b/Engine/Plugins/Experimental/AnimNext/Source/AnimNext/Private/Param/AnimNextParameterBlock.cpp
--- a/Engine/Plugins/Experimental/AnimNext/Source/AnimNext/Private/Param/AnimNextParameterBlock.cpp
+++ b/Engine/Plugins/Experimental/AnimNext/Source/AnimNext/Private/Param/AnimNextParameterBlock.cpp
@@ -43,7 +43,7 @@ void UAnimNextParameterBlock::Run(const UE::AnimNext::FContext& Context) const
 		FRigVMExtendedExecuteContext RigVMExtendedExecuteContext;
 		FAnimNextExecuteContext& AnimNextContext = RigVMExtendedExecuteContext.GetPublicDataSafe<FAnimNextExecuteContext>();
 		AnimNextContext.SetContextData(Context);
-		RigVM->Execute(RigVMExtendedExecuteContext, TArray<TRigVMMemoryStorage*>(), FRigUnit_AnimNextBeginExecution::EventName);
+		RigVM->Execute(RigVMExtendedExecuteContext, TArray<FRigVMMemoryStorageStruct*>(), FRigUnit_AnimNextBeginExecution::EventName);
 	}
 }
 
diff --git a/Engine/Plugins/Runtime/RigVM/Source/RigVM/Private/RigVMCore/RigVM.cpp b/Engine/Plugins/Runtime/RigVM/Source/RigVM/Private/RigVMCore/RigVM.cpp
--- a/Engine/Plugins/Runtime/RigVM/Source/RigVM/Private/RigVMCore/RigVM.cpp
+++ b/Engine/Plugins/Runtime/RigVM/Source/RigVM/Private/RigVMCore/RigVM.cpp
@@ -370,19 +370,8 @@ uint32 URigVM::ComputeVMHash(const FRigVMExtendedExecuteContext& Context) const
 		Hash = HashCombine(Hash, GetTypeHash(ExternalVariable.TypeName.ToString()));
 	}
 
-#if UE_RIGVM_PROPERTY_BAG_STORAGE_ENABLED
 	Hash = HashCombine(Hash, LiteralMemoryStorage.GetMemoryHash());
 	Hash = HashCombine(Hash, Context.WorkMemoryStorage.GetMemoryHash());
-#else
-	if(LiteralMemoryStorageObject)
-	{
-		Hash = HashCombine(Hash, LiteralMemoryStorageObject->GetMemoryHash());
-	}
-	if(GetWorkMemory(Context))
-	{
-		Hash = HashCombine(Hash, GetWorkMemory(Context)->GetMemoryHash());
-	}
-#endif
 	
 	return Hash;
 }
@@ -419,7 +408,7 @@ bool URigVM::ValidateAllOperandsDuringLoad(FRigVMExtendedExecuteContext& Context
 	// check all operands on all ops for validity
 	bool bAllOperandsValid = true;
 
-	TArray<TRigVMMemoryStorage*> LocalMemory = { GetWorkMemory(Context), GetLiteralMemory(), GetDebugMemory(Context) };
+	TArray<FRigVMMemoryStorageStruct*> LocalMemory = { GetWorkMemory(Context), GetLiteralMemory(), GetDebugMemory(Context) };
 	
 	auto CheckOperandValidity = [LocalMemory, &bAllOperandsValid, this](const FRigVMOperand& InOperand) -> bool
 	{
@@ -430,7 +419,7 @@ bool URigVM::ValidateAllOperandsDuringLoad(FRigVMExtendedExecuteContext& Context
 		}
 
 
-		const TRigVMMemoryStorage* MemoryForOperand = LocalMemory[InOperand.GetContainerIndex()];
+		const FRigVMMemoryStorageStruct* MemoryForOperand = LocalMemory[InOperand.GetContainerIndex()];
 		if(InOperand.GetMemoryType() != ERigVMMemoryType::External)
 		{
 			if(!MemoryForOperand->IsValidIndex(InOperand.GetRegisterIndex()))
@@ -659,44 +648,7 @@ URigVMMemoryStorage* URigVM::CreateMemoryByType(FRigVMExtendedExecuteContext& Co
 {
 	URigVMMemoryStorage* MemoryStorage = nullptr;
 
-#if !UE_RIGVM_PROPERTY_BAG_STORAGE_ENABLED
-	switch(InMemoryType)
-	{
-		case ERigVMMemoryType::Literal:
-		{
-			// since literal memory object can be shared across packages, it needs to have the RF_Public flag
-			// for example, a control rig instance in a level sequence pacakge can references
-			// the literal memory object in the control rig package
-			CreateMemoryByType(GetOuter(), LiteralMemoryStorageObject, InMemoryType, RF_Public, bForceCreation);
-			//check(LiteralMemoryStorageObject->GetOuter() == GetOuter());
-			MemoryStorage = LiteralMemoryStorageObject;
-			break;
-		}
-
-		case ERigVMMemoryType::Work:
-		{
-			CreateMemoryByType(GetOuter(), Context.WorkMemoryStorageObject, InMemoryType, RF_NoFlags, bForceCreation);
-			//check(GetWorkMemory(Context)->GetOuter() == GetOuter());
-			MemoryStorage = Context.WorkMemoryStorageObject;
-			break;
-		}
-
-		case ERigVMMemoryType::Debug:
-		{
-#if WITH_EDITOR
-			CreateMemoryByType(GetOuter(), Context.DebugMemoryStorageObject, InMemoryType, RF_NoFlags, bForceCreation);
-			//check(Context.DebugMemoryStorageObject->GetOuter() == GetOuter());
-			MemoryStorage = Context.DebugMemoryStorageObject;
-#endif
-			break;
-		}
-
-		default:
-		{
-			break;
-		}
-	}
-#endif // UE_RIGVM_PROPERTY_BAG_STORAGE_ENABLED
+	// TODO zzz : Deprecate function
 
 	return MemoryStorage;
 }
@@ -734,9 +686,9 @@ void URigVM::GenerateMemoryType(FRigVMExtendedExecuteContext& Context, ERigVMMem
 	}
 }
 
-TRigVMMemoryStorage* URigVM::GetMemoryByType(FRigVMExtendedExecuteContext& Context, ERigVMMemoryType InMemoryType/*, bool bCreateIfNeeded*/)
+FRigVMMemoryStorageStruct* URigVM::GetMemoryByType(FRigVMExtendedExecuteContext& Context, ERigVMMemoryType InMemoryType/*, bool bCreateIfNeeded*/)
 {
-	TRigVMMemoryStorage* MemoryStorage = nullptr;
+	FRigVMMemoryStorageStruct* MemoryStorage = nullptr;
 
 	switch(InMemoryType)
 	{
@@ -748,23 +700,13 @@ TRigVMMemoryStorage* URigVM::GetMemoryByType(FRigVMExtendedExecuteContext& Conte
 
 		case ERigVMMemoryType::Work:
 		{
-#if UE_RIGVM_PROPERTY_BAG_STORAGE_ENABLED
 			MemoryStorage = &Context.WorkMemoryStorage;
-#else
-			//check(Context.WorkMemoryStorageObject->GetOuter() == GetOuter());
-			MemoryStorage = Context.WorkMemoryStorageObject;
-#endif
 			break;
 		}
 
 		case ERigVMMemoryType::Debug:
 		{
-#if UE_RIGVM_PROPERTY_BAG_STORAGE_ENABLED
 			MemoryStorage = &Context.DebugMemoryStorage;
-#else
-			//check(Context.DebugMemoryStorageObject->GetOuter() == GetOuter());
-			MemoryStorage = Context.DebugMemoryStorageObject;
-#endif
 			break;
 		}
 
@@ -777,7 +719,7 @@ TRigVMMemoryStorage* URigVM::GetMemoryByType(FRigVMExtendedExecuteContext& Conte
 	return MemoryStorage;
 }
 
-const TRigVMMemoryStorage* URigVM::GetMemoryByType(const FRigVMExtendedExecuteContext& Context, ERigVMMemoryType InMemoryType) const
+const FRigVMMemoryStorageStruct* URigVM::GetMemoryByType(const FRigVMExtendedExecuteContext& Context, ERigVMMemoryType InMemoryType) const
 {
 	return const_cast<URigVM*>(this)->GetMemoryByType(const_cast<FRigVMExtendedExecuteContext&>(Context), InMemoryType/*, false*/);
 }
@@ -815,11 +757,7 @@ void URigVM::ClearMemory_Internal()
 	}
 #endif
 
-#if UE_RIGVM_PROPERTY_BAG_STORAGE_ENABLED
 	LiteralMemoryStorage = FRigVMMemoryStorageStruct(ERigVMMemoryType::Literal);
-#else
-	LiteralMemoryStorageObject = nullptr;
-#endif
 
 	InvalidateCachedMemory_Internal();
 }
@@ -828,20 +766,8 @@ void URigVM::ClearMemory(FRigVMExtendedExecuteContext& Context)
 {
 	ClearMemory_Internal();
 
-#if UE_RIGVM_PROPERTY_BAG_STORAGE_ENABLED
 	Context.WorkMemoryStorage = FRigVMMemoryStorageStruct(ERigVMMemoryType::Work);
 	Context.DebugMemoryStorage = FRigVMMemoryStorageStruct(ERigVMMemoryType::Debug);
-#else
-	if(GetWorkMemory(Context))
-	{
-		Context.WorkMemoryStorageObject = nullptr;
-	}
-
-	if(GetDebugMemory(Context))
-	{
-		Context.DebugMemoryStorageObject = nullptr;
-	}
-#endif
 
 	InvalidateCachedMemory(Context);
 }
@@ -909,7 +835,7 @@ bool URigVM::CanExecuteEntry(const FRigVMExtendedExecuteContext& Context, const
 
 #if WITH_EDITOR
 
-bool URigVM::ResumeExecution(FRigVMExtendedExecuteContext& Context, TArrayView<TRigVMMemoryStorage*> Memory, const FName& InEntryName)
+bool URigVM::ResumeExecution(FRigVMExtendedExecuteContext& Context, TArrayView<FRigVMMemoryStorageStruct*> Memory, const FName& InEntryName)
 {
 	if (FRigVMDebugInfo* RigVMDebugInfo = Context.GetRigVMDebugInfo())
 	{
@@ -1141,7 +1067,7 @@ void URigVM::InstructionOpEval(FRigVMExtendedExecuteContext& Context, int32 Inst
 	}
 }
 
-void URigVM::PrepareMemoryForExecution(FRigVMExtendedExecuteContext& Context, TArrayView<TRigVMMemoryStorage*> InMemory)
+void URigVM::PrepareMemoryForExecution(FRigVMExtendedExecuteContext& Context, TArrayView<FRigVMMemoryStorageStruct*> InMemory)
 {
 	ensureMsgf(Context.ExecutingThreadId == FPlatformTLS::GetCurrentThreadId(), TEXT("RigVM::CacheMemoryHandlesIfRequired from multiple threads (%d and %d)"), Context.ExecutingThreadId, (int32)FPlatformTLS::GetCurrentThreadId());
 
@@ -1183,7 +1109,7 @@ void URigVM::PrepareMemoryForExecution(FRigVMExtendedExecuteContext& Context, TA
 	Context.CachedMemoryHandles.Reset(MemoryHandleCount);
 }
 
-void URigVM::CacheMemoryHandlesIfRequired(FRigVMExtendedExecuteContext& Context, TArrayView<TRigVMMemoryStorage*> InMemory)
+void URigVM::CacheMemoryHandlesIfRequired(FRigVMExtendedExecuteContext& Context, TArrayView<FRigVMMemoryStorageStruct*> InMemory)
 {
 	if (Instructions.Num() == 0 || InMemory.Num() == 0)
 	{
@@ -1425,7 +1351,7 @@ bool URigVM::ShouldHaltAtInstruction(FRigVMExtendedExecuteContext& Context, cons
 }
 #endif
 
-bool URigVM::Initialize(FRigVMExtendedExecuteContext& Context, TArrayView<TRigVMMemoryStorage*> Memory)
+bool URigVM::Initialize(FRigVMExtendedExecuteContext& Context, TArrayView<FRigVMMemoryStorageStruct*> Memory)
 {
 	if (Context.ExecutingThreadId != INDEX_NONE)
 	{
@@ -1443,7 +1369,7 @@ bool URigVM::Initialize(FRigVMExtendedExecuteContext& Context, TArrayView<TRigVM
 	}
 
 	// changes to the layout of memory array should be reflected in GetContainerIndex()
-	TArray<TRigVMMemoryStorage*> LocalMemory;
+	TArray<FRigVMMemoryStorageStruct*> LocalMemory;
 	if (Memory.Num() == 0)
 	{
 		LocalMemory = GetLocalMemoryArray(Context);
@@ -1457,7 +1383,7 @@ bool URigVM::Initialize(FRigVMExtendedExecuteContext& Context, TArrayView<TRigVM
 	return true;
 }
 
-bool URigVM::InitializeInstance(FRigVMExtendedExecuteContext& Context, TArrayView<TRigVMMemoryStorage*> Memory)
+bool URigVM::InitializeInstance(FRigVMExtendedExecuteContext& Context, TArrayView<FRigVMMemoryStorageStruct*> Memory)
 {
 	check(!Memory.IsEmpty());
 
@@ -1469,51 +1395,13 @@ bool URigVM::InitializeInstance(FRigVMExtendedExecuteContext& Context, TArrayVie
 	Context.LazyBranchInstanceData.Reset(LazyBranchSize);
 	Context.LazyBranchInstanceData.SetNumZeroed(LazyBranchSize);
 
-#if !UE_RIGVM_PROPERTY_BAG_STORAGE_ENABLED
-	// re-initialize work memory from CDO
-	if (TRigVMMemoryStorage* WorkMemory = Memory[(int32)ERigVMMemoryType::Work])
-	{
-		if (!WorkMemory->HasAnyFlags(RF_ClassDefaultObject))
-		{
-			if (const URigVMMemoryStorageGeneratorClass* MemoryClass = Cast<URigVMMemoryStorageGeneratorClass>(WorkMemory->GetClass()))
-			{
-				if (URigVMMemoryStorage* WorkMemoryCDO = MemoryClass->GetDefaultObject<URigVMMemoryStorage>())
-				{
-					check(WorkMemory != WorkMemoryCDO);
-
-					for (const FProperty* Property : MemoryClass->LinkedProperties)
-					{
-						Property->CopyCompleteValue_InContainer(WorkMemory, WorkMemoryCDO);
-#if UE_RIGVM_DEBUG_EXECUTION
-						if (CVarControlRigDebugAllVMExecutions->GetBool() || Context.GetPublicData<FRigVMExecuteContext>().bDebugExecution)
-						{
-							FString DefaultValue;
-							const uint8* PropertyMemory = Property->ContainerPtrToValuePtr<uint8>(WorkMemory);
-							Property->ExportText_Direct(
-								DefaultValue,
-								PropertyMemory,
-								PropertyMemory,
-								nullptr,
-								PPF_None,
-								nullptr);
-
-							UE_LOG(LogRigVM, Warning, TEXT("Property %s defaults to '%s'."), *Property->GetName(), *DefaultValue);
-							UE_LOG(LogRigVM, Warning, TEXT("Property %s defaults to '%s'."), *Property->GetName(), *DefaultValue);
-						}
-#endif // UE_RIGVM_DEBUG_EXECUTION
-					}
-				}
-			}
-		}
-	}
-#endif // UE_RIGVM_PROPERTY_BAG_STORAGE_ENABLED
 
 	CacheMemoryHandlesIfRequired(Context, Memory);
 
 	return true;
 }
 
-ERigVMExecuteResult URigVM::Execute(FRigVMExtendedExecuteContext& Context, TArrayView<TRigVMMemoryStorage*> Memory, const FName& InEntryName)
+ERigVMExecuteResult URigVM::Execute(FRigVMExtendedExecuteContext& Context, TArrayView<FRigVMMemoryStorageStruct*> Memory, const FName& InEntryName)
 {
 	// if this the first entry being executed - get ready for execution
 	const bool bIsRootEntry = Context.EntriesBeingExecuted.IsEmpty();
@@ -1642,7 +1530,7 @@ ERigVMExecuteResult URigVM::Execute(FRigVMExtendedExecuteContext& Context, TArra
 	if (CVarControlRigDebugAllVMExecutions->GetBool() || ContextPublicData.bDebugExecution)
 	{
 		ContextPublicData.InstanceOpCodeEnum = StaticEnum<ERigVMOpCode>();
-		TRigVMMemoryStorage* LiteralMemory = GetLiteralMemory();
+		FRigVMMemoryStorageStruct* LiteralMemory = GetLiteralMemory();
 		ContextPublicData.DebugMemoryString = FString("\n\nLiteral Memory\n\n");
 		for (int32 PropertyIndex=0; PropertyIndex<LiteralMemory->Num(); ++PropertyIndex)
 		{
@@ -1660,7 +1548,7 @@ ERigVMExecuteResult URigVM::Execute(FRigVMExtendedExecuteContext& Context, TArra
 #if WITH_EDITOR
 	if(bIsRootEntry)
 	{
-		Context.CurrentVMMemory = TArrayView<TRigVMMemoryStorage*>();
+		Context.CurrentVMMemory = TArrayView<FRigVMMemoryStorageStruct*>();
 		
 		if (FRigVMDebugInfo* RigVMDebugInfo = Context.GetRigVMDebugInfo())
 		{
@@ -1847,7 +1735,7 @@ ERigVMExecuteResult URigVM::ExecuteInstructions(FRigVMExtendedExecuteContext& Co
 
 				FRigVMMemoryHandle& SourceHandle = Context.CachedMemoryHandles[FirstHandleForInstruction[ContextPublicData.InstructionIndex]];
 				FRigVMMemoryHandle& TargetHandle = Context.CachedMemoryHandles[FirstHandleForInstruction[ContextPublicData.InstructionIndex] + 1];
-				TRigVMMemoryStorage::CopyProperty(TargetHandle, SourceHandle);
+				FRigVMMemoryStorageStruct::CopyProperty(TargetHandle, SourceHandle);
 					
 #if WITH_EDITOR
 				if(GetDebugMemory(Context) && GetDebugMemory(Context)->Num() > 0)
@@ -2118,7 +2006,7 @@ ERigVMExecuteResult URigVM::ExecuteInstructions(FRigVMExtendedExecuteContext& Co
 		if (CVarControlRigDebugAllVMExecutions->GetBool() || ContextPublicData.bDebugExecution)
 		{
 			TArray<FString> CurrentWorkMemory;
-			TRigVMMemoryStorage* WorkMemory = GetWorkMemory(Context);
+			FRigVMMemoryStorageStruct* WorkMemory = GetWorkMemory(Context);
 			int32 LineIndex = 0;
 			for (int32 PropertyIndex=0; PropertyIndex<WorkMemory->Num(); ++PropertyIndex, ++LineIndex)
 			{
@@ -2228,7 +2116,7 @@ FRigVMExternalVariable URigVM::GetExternalVariableByName(const FRigVMExtendedExe
 
 void URigVM::SetPropertyValueFromString(FRigVMExtendedExecuteContext& Context, const FRigVMOperand& InOperand, const FString& InDefaultValue)
 {
-	TRigVMMemoryStorage* Memory = GetMemoryByType(Context, InOperand.GetMemoryType());
+	FRigVMMemoryStorageStruct* Memory = GetMemoryByType(Context, InOperand.GetMemoryType());
 	if(Memory == nullptr)
 	{
 		return;
@@ -2465,7 +2353,7 @@ FString URigVM::GetOperandLabel(FRigVMExtendedExecuteContext& Context, const FRi
 	}
 	else
 	{
-		TRigVMMemoryStorage* Memory = GetMemoryByType(Context, InOperand.GetMemoryType());
+		FRigVMMemoryStorageStruct* Memory = GetMemoryByType(Context, InOperand.GetMemoryType());
 		if(Memory == nullptr)
 		{
 			return FString();
@@ -2512,7 +2400,7 @@ void URigVM::ClearDebugMemory(FRigVMExtendedExecuteContext& Context)
 
 void URigVM::CacheSingleMemoryHandle(FRigVMExtendedExecuteContext& Context, int32 InHandleIndex, const FRigVMBranchInfoKey& InBranchInfoKey, const FRigVMOperand& InArg, bool bForExecute)
 {
-	TRigVMMemoryStorage* Memory = GetMemoryByType(Context, InArg.GetMemoryType()/*, false*/);
+	FRigVMMemoryStorageStruct* Memory = GetMemoryByType(Context, InArg.GetMemoryType()/*, false*/);
 
 	if (InArg.GetMemoryType() == ERigVMMemoryType::External)
 	{
@@ -2584,7 +2472,7 @@ void URigVM::CopyOperandForDebuggingImpl(FRigVMExtendedExecuteContext& Context,
 {
 #if WITH_EDITOR
 
-	TRigVMMemoryStorage* TargetMemory = GetDebugMemory(Context);
+	FRigVMMemoryStorageStruct* TargetMemory = GetDebugMemory(Context);
 	if(TargetMemory == nullptr)
 	{
 		return;
@@ -2624,12 +2512,12 @@ void URigVM::CopyOperandForDebuggingImpl(FRigVMExtendedExecuteContext& Context,
 			const FProperty* SourceProperty = ExternalVariable.Property;
 			FRigVMExternalVariableRuntimeData& ExternalVariableRuntimeData = Context.ExternalVariableRuntimeData[ExternalVariableIndex];
 			const uint8* SourcePtr = ExternalVariableRuntimeData.Memory;
-			TRigVMMemoryStorage::CopyProperty(TargetProperty, TargetPtr, SourceProperty, SourcePtr);
+			FRigVMMemoryStorageStruct::CopyProperty(TargetProperty, TargetPtr, SourceProperty, SourcePtr);
 		}
 		return;
 	}
 
-	TRigVMMemoryStorage* SourceMemory = GetMemoryByType(Context, InArg.GetMemoryType());
+	FRigVMMemoryStorageStruct* SourceMemory = GetMemoryByType(Context, InArg.GetMemoryType());
 	if(SourceMemory == nullptr)
 	{
 		return;
@@ -2637,7 +2525,7 @@ void URigVM::CopyOperandForDebuggingImpl(FRigVMExtendedExecuteContext& Context,
 	const FProperty* SourceProperty = SourceMemory->GetProperties()[InArg.GetRegisterIndex()];
 	const uint8* SourcePtr = SourceMemory->GetData<uint8>(InArg.GetRegisterIndex());
 
-	TRigVMMemoryStorage::CopyProperty(TargetProperty, TargetPtr, SourceProperty, SourcePtr);
+	FRigVMMemoryStorageStruct::CopyProperty(TargetProperty, TargetPtr, SourceProperty, SourcePtr);
 	
 #endif
 }
diff --git a/Engine/Plugins/Runtime/RigVM/Source/RigVM/Private/RigVMCore/RigVMExecuteContext.cpp b/Engine/Plugins/Runtime/RigVM/Source/RigVM/Private/RigVMCore/RigVMExecuteContext.cpp
--- a/Engine/Plugins/Runtime/RigVM/Source/RigVM/Private/RigVMCore/RigVMExecuteContext.cpp
+++ b/Engine/Plugins/Runtime/RigVM/Source/RigVM/Private/RigVMCore/RigVMExecuteContext.cpp
@@ -72,7 +72,7 @@ void FRigVMExtendedExecuteContext::Reset()
 	WorkMemoryStorage = FRigVMMemoryStorageStruct();
 	DebugMemoryStorage = FRigVMMemoryStorageStruct();
 
-	CurrentVMMemory = TArrayView<TRigVMMemoryStorage*>();
+	CurrentVMMemory = TArrayView<FRigVMMemoryStorageStruct*>();
 	ExecutionReachedExit().Clear();
 
 	CachedMemoryHandles.Reset();
@@ -120,13 +120,8 @@ void FRigVMExtendedExecuteContext::ResetExecutionState()
 
 void FRigVMExtendedExecuteContext::CopyMemoryStorage(const FRigVMExtendedExecuteContext& Other, UObject* Outer)
 {
-#if UE_RIGVM_PROPERTY_BAG_STORAGE_ENABLED
 	WorkMemoryStorage = Other.WorkMemoryStorage;
 	DebugMemoryStorage = Other.DebugMemoryStorage;
-#else
-	CopyMemoryStorage(WorkMemoryStorageObject, Other.WorkMemoryStorageObject, Outer);
-	CopyMemoryStorage(DebugMemoryStorageObject, Other.DebugMemoryStorageObject, Outer);
-#endif // UE_RIGVM_PROPERTY_BAG_STORAGE_ENABLED
 }
 
 void FRigVMExtendedExecuteContext::CopyMemoryStorage(TObjectPtr<URigVMMemoryStorage>& TargetMemory, const TObjectPtr <URigVMMemoryStorage>& SourceMemory, UObject* Outer)
diff --git a/Engine/Plugins/Runtime/RigVM/Source/RigVM/Private/RigVMCore/RigVMNativized.cpp b/Engine/Plugins/Runtime/RigVM/Source/RigVM/Private/RigVMCore/RigVMNativized.cpp
--- a/Engine/Plugins/Runtime/RigVM/Source/RigVM/Private/RigVMCore/RigVMNativized.cpp
+++ b/Engine/Plugins/Runtime/RigVM/Source/RigVM/Private/RigVMCore/RigVMNativized.cpp
@@ -27,13 +27,13 @@ void URigVMNativized::Reset(bool IsIgnoringArchetypeRef)
 	ByteCodeStorage.Reset();
 }
 
-bool URigVMNativized::Initialize(FRigVMExtendedExecuteContext& Context, TArrayView<TRigVMMemoryStorage*> Memory)
+bool URigVMNativized::Initialize(FRigVMExtendedExecuteContext& Context, TArrayView<FRigVMMemoryStorageStruct*> Memory)
 {
 	// nothing to do here 
 	return true;
 }
 
-ERigVMExecuteResult URigVMNativized::Execute(FRigVMExtendedExecuteContext& Context, TArrayView<TRigVMMemoryStorage*> Memory, const FName& InEntryName)
+ERigVMExecuteResult URigVMNativized::Execute(FRigVMExtendedExecuteContext& Context, TArrayView<FRigVMMemoryStorageStruct*> Memory, const FName& InEntryName)
 {
 	// to be implemented by the generated code
 	return ERigVMExecuteResult::Failed;
diff --git a/Engine/Plugins/Runtime/RigVM/Source/RigVM/Private/RigVMHost.cpp b/Engine/Plugins/Runtime/RigVM/Source/RigVM/Private/RigVMHost.cpp
--- a/Engine/Plugins/Runtime/RigVM/Source/RigVM/Private/RigVMHost.cpp
+++ b/Engine/Plugins/Runtime/RigVM/Source/RigVM/Private/RigVMHost.cpp
@@ -269,7 +269,6 @@ bool URigVMHost::InitializeVM(const FName& InEventName)
 		return false;	// The rig did compile with errors
 	}
 
-#if UE_RIGVM_PROPERTY_BAG_STORAGE_ENABLED
 	if (bTEMP_CopyDefaultsFromCDO && !HasAnyFlags(RF_ClassDefaultObject))
 	{
 		if (URigVMHost* CDO = GetClass()->GetDefaultObject<URigVMHost>())
@@ -277,12 +276,11 @@ bool URigVMHost::InitializeVM(const FName& InEventName)
 			ExtendedExecuteContext.WorkMemoryStorage = CDO->GetExtendedExecuteContext().WorkMemoryStorage; // initialize the memory to CDO default state
 		}
 	}
-#endif // UE_RIGVM_PROPERTY_BAG_STORAGE_ENABLED
 
 	// update the VM's external variables
 	VM->SetExternalVariablesInstanceData(GetExtendedExecuteContext(), ExternalVariables);
 
-	TArray<TRigVMMemoryStorage*> LocalMemory = VM->GetLocalMemoryArray(ExtendedExecuteContext);
+	TArray<FRigVMMemoryStorageStruct*> LocalMemory = VM->GetLocalMemoryArray(ExtendedExecuteContext);
 	const bool bResult = VM->InitializeInstance(GetExtendedExecuteContext(), LocalMemory);
 	if(bResult)
 	{
@@ -625,29 +623,6 @@ bool URigVMHost::InitializeCDOVM()
 	check(VM != nullptr);
 	check(VM->HasAnyFlags(RF_ClassDefaultObject | RF_DefaultSubObject));
 
-#if !UE_RIGVM_PROPERTY_BAG_STORAGE_ENABLED
-	// Create the memory
-	const TArray<ERigVMMemoryType> MemoryTypes = { ERigVMMemoryType::Literal, ERigVMMemoryType::Work, ERigVMMemoryType::Debug };
-	for (ERigVMMemoryType MemoryType : MemoryTypes)
-	{
-		// Refresh if storage class exists
-		if (URigVMMemoryStorageGeneratorClass* Class = URigVMMemoryStorageGeneratorClass::GetStorageClass(VM->GetOuter(), MemoryType))
-		{
-			if (Class->LinkedProperties.Num() == 0)
-			{
-				Class->RefreshLinkedProperties();
-			}
-			if (Class->PropertyPathDescriptions.Num() != Class->PropertyPaths.Num())
-			{
-				Class->RefreshPropertyPaths();
-			}
-		}
-
-		// But create a memory object in any case (required by sequencer rigs)
-		VM->CreateMemoryByType(ExtendedExecuteContext, MemoryType);
-	}
-#endif // !UE_RIGVM_PROPERTY_BAG_STORAGE_ENABLED
-
 	// update the VM's external variables
 	VM->ClearExternalVariables(GetExtendedExecuteContext());
 	VM->SetExternalVariableDefs(GetExternalVariablesImpl(false));
@@ -679,12 +654,6 @@ bool URigVMHost::Execute_Internal(const FName& InEventName)
 		if(!VM->IsContextValidForExecution(Context)
 			|| !IsValidLowLevel()
 			|| !VM->IsValidLowLevel()
-#if !UE_RIGVM_PROPERTY_BAG_STORAGE_ENABLED
-			|| !GetLiteralMemory()
-			|| !GetLiteralMemory()->IsValidLowLevel()
-			|| !GetWorkMemory()
-			|| !GetWorkMemory()->IsValidLowLevel()
-#endif // !UE_RIGVM_PROPERTY_BAG_STORAGE_ENABLED
 		)
 		{
 			UE_LOG(LogRigVM, Warning, InvalidatedVMFormat, *GetClass()->GetName());
@@ -703,7 +672,7 @@ bool URigVMHost::Execute_Internal(const FName& InEventName)
 	
 	const bool bUseDebuggingSnapshots = !VM->IsNativized();
 	
-	TArray<TRigVMMemoryStorage*> LocalMemory = VM->GetLocalMemoryArray(ExtendedExecuteContext);
+	TArray<FRigVMMemoryStorageStruct*> LocalMemory = VM->GetLocalMemoryArray(ExtendedExecuteContext);
 
 #if WITH_EDITOR
 	if(bUseDebuggingSnapshots)
@@ -836,13 +805,13 @@ URigVM* URigVMHost::GetVM()
 	return VM;
 }
 
-TRigVMMemoryStorage* URigVMHost::GetMemoryByType(ERigVMMemoryType InMemoryType)
+FRigVMMemoryStorageStruct* URigVMHost::GetMemoryByType(ERigVMMemoryType InMemoryType)
 {
 	check(VM);
 	return VM->GetMemoryByType(ExtendedExecuteContext, InMemoryType);
 }
 
-const TRigVMMemoryStorage* URigVMHost::GetMemoryByType(ERigVMMemoryType InMemoryType) const
+const FRigVMMemoryStorageStruct* URigVMHost::GetMemoryByType(ERigVMMemoryType InMemoryType) const
 {
 	check(VM);
 	return VM->GetMemoryByType(ExtendedExecuteContext, InMemoryType);
@@ -1148,10 +1117,10 @@ void URigVMHost::GenerateUserDefinedDependenciesData(FRigVMExtendedExecuteContex
 	}
 }
 
-TArray<const UObject*> URigVMHost::GetUserDefinedDependencies(const TArray<const TRigVMMemoryStorage*> InMemory)
+TArray<const UObject*> URigVMHost::GetUserDefinedDependencies(const TArray<const FRigVMMemoryStorageStruct*> InMemory)
 {
 	TArray<const UObject*> Dependencies;
-	auto ProcessMemory = [&Dependencies](const TRigVMMemoryStorage* Memory)
+	auto ProcessMemory = [&Dependencies](const FRigVMMemoryStorageStruct* Memory)
 	{
 		if (Memory == nullptr)
 		{
@@ -1191,7 +1160,7 @@ TArray<const UObject*> URigVMHost::GetUserDefinedDependencies(const TArray<const
 		}
 	};
 
-	for (const TRigVMMemoryStorage* MemoryStorage : InMemory)
+	for (const FRigVMMemoryStorageStruct* MemoryStorage : InMemory)
 	{
 		ProcessMemory(MemoryStorage);
 	}
diff --git a/Engine/Plugins/Runtime/RigVM/Source/RigVM/Private/RigVMModule.cpp b/Engine/Plugins/Runtime/RigVM/Source/RigVM/Private/RigVMModule.cpp
--- a/Engine/Plugins/Runtime/RigVM/Source/RigVM/Private/RigVMModule.cpp
+++ b/Engine/Plugins/Runtime/RigVM/Source/RigVM/Private/RigVMModule.cpp
@@ -9,7 +9,7 @@
 #include "HAL/IConsoleManager.h"
 #include "RigVMObjectVersion.h"
 #include "UObject/DevObjectVersion.h"
-#if WITH_EDITOR && UE_RIGVM_PROPERTY_BAG_STORAGE_ENABLED
+#if WITH_EDITOR
 #include "PropertyBagDetails.h"
 #include "PropertyEditorModule.h"
 #endif
@@ -34,7 +34,7 @@ TAutoConsoleVariable<bool> CVarRigVMEnableUInterfaces(TEXT("RigVM.UInterfaceSupp
 
 void FRigVMModule::StartupModule()
 {
-#if WITH_EDITOR && UE_RIGVM_PROPERTY_BAG_STORAGE_ENABLED
+#if WITH_EDITOR
 	// Register the details customizer
 	FPropertyEditorModule& PropertyModule = FModuleManager::LoadModuleChecked<FPropertyEditorModule>("PropertyEditor");
 	PropertyModule.RegisterCustomPropertyTypeLayout("RigVMMemoryStorageStruct", FOnGetPropertyTypeCustomizationInstance::CreateStatic(&FPropertyBagDetails::MakeInstance));
@@ -44,7 +44,7 @@ void FRigVMModule::StartupModule()
 
 void FRigVMModule::ShutdownModule()
 {
-#if WITH_EDITOR && UE_RIGVM_PROPERTY_BAG_STORAGE_ENABLED
+#if WITH_EDITOR
 	// Unregister the details customization
 	if (FPropertyEditorModule* PropertyModule = FModuleManager::Get().GetModulePtr<FPropertyEditorModule>("PropertyEditor"))
 	{
diff --git a/Engine/Plugins/Runtime/RigVM/Source/RigVM/Public/RigVMCore/RigVM.h b/Engine/Plugins/Runtime/RigVM/Source/RigVM/Public/RigVMCore/RigVM.h
--- a/Engine/Plugins/Runtime/RigVM/Source/RigVM/Public/RigVMCore/RigVM.h
+++ b/Engine/Plugins/Runtime/RigVM/Source/RigVM/Public/RigVMCore/RigVM.h
@@ -222,28 +222,28 @@ public:
 	UE_DEPRECATED(5.3, "Please, use Initialize with Context param")
 	virtual bool Initialize(TArrayView<URigVMMemoryStorage*> Memory) { return false; }
 
-	UE_DEPRECATED(5.4, "Please, use Initialize with TRigVMMemoryStorage param")
-	virtual bool Initialize(FRigVMExtendedExecuteContext& Context, TArrayView<TRigVMMemoryStorageDeprecatedType*> Memory) { return false; }
+	UE_DEPRECATED(5.4, "Please, use Initialize with FRigVMMemoryStorageStruct param")
+	virtual bool Initialize(FRigVMExtendedExecuteContext& Context, TArrayView<URigVMMemoryStorage*> Memory) { return false; }
 
 	// Prepares caches and memory for execution
-	virtual bool Initialize(FRigVMExtendedExecuteContext& Context, TArrayView<TRigVMMemoryStorage*> Memory);
+	virtual bool Initialize(FRigVMExtendedExecuteContext& Context, TArrayView<FRigVMMemoryStorageStruct*> Memory);
 
-	UE_DEPRECATED(5.4, "Please, use InitializeInstance with TRigVMMemoryStorage param")
-	virtual bool InitializeInstance(FRigVMExtendedExecuteContext& Context, TArrayView<TRigVMMemoryStorageDeprecatedType*> Memory) { return false; }
+	UE_DEPRECATED(5.4, "Please, use InitializeInstance with FRigVMMemoryStorageStruct param")
+	virtual bool InitializeInstance(FRigVMExtendedExecuteContext& Context, TArrayView<URigVMMemoryStorage*> Memory) { return false; }
 
 	// Initializes cached memory handles and copies work memory from the CDO to the Context
-	virtual bool InitializeInstance(FRigVMExtendedExecuteContext& Context, TArrayView<TRigVMMemoryStorage*> Memory);
+	virtual bool InitializeInstance(FRigVMExtendedExecuteContext& Context, TArrayView<FRigVMMemoryStorageStruct*> Memory);
 
 	UE_DEPRECATED(5.3, "Please, use Execute with Context param")
 	virtual ERigVMExecuteResult Execute(TArrayView<URigVMMemoryStorage*> Memory, const FName& InEntryName = NAME_None) { return ERigVMExecuteResult::Failed; }
 
-	UE_DEPRECATED(5.4, "Please, use Execute with TRigVMMemoryStorage param")
-	virtual ERigVMExecuteResult Execute(FRigVMExtendedExecuteContext& Context, TArrayView<TRigVMMemoryStorageDeprecatedType*> Memory, const FName& InEntryName = NAME_None) { return ERigVMExecuteResult::Failed; }
+	UE_DEPRECATED(5.4, "Please, use Execute with FRigVMMemoryStorageStruct param")
+	virtual ERigVMExecuteResult Execute(FRigVMExtendedExecuteContext& Context, TArrayView<URigVMMemoryStorage*> Memory, const FName& InEntryName = NAME_None) { return ERigVMExecuteResult::Failed; }
 
 	// Executes the VM.
 	// You can optionally provide external memory to the execution
 	// and provide optional additional operands.
-	virtual ERigVMExecuteResult Execute(FRigVMExtendedExecuteContext& Context, TArrayView<TRigVMMemoryStorage*> Memory, const FName& InEntryName = NAME_None);
+	virtual ERigVMExecuteResult Execute(FRigVMExtendedExecuteContext& Context, TArrayView<FRigVMMemoryStorageStruct*> Memory, const FName& InEntryName = NAME_None);
 
 	UE_DEPRECATED(5.3, "Please, use Execute with Context param")
 	virtual bool Execute(const FName& InEntryName = NAME_None) { return false; }
@@ -276,7 +276,7 @@ public:
 	virtual FString GetRigVMFunctionName(int32 InFunctionIndex) const;
 
 	UE_DEPRECATED(5.4, "GetMemoryByType has been deprecated from the VM. Please, use GetWorkMemory from VMHost or the version with a Context param")
-	virtual TRigVMMemoryStorageDeprecatedType* GetMemoryByType(ERigVMMemoryType InMemoryType, bool bCreateIfNeeded = true) { return nullptr; }
+	virtual URigVMMemoryStorage* GetMemoryByType(ERigVMMemoryType InMemoryType, bool bCreateIfNeeded = true) { return nullptr; }
 
 	// Creates a new memory storage by type
 	void CreateMemoryByType(UObject* Outer, TObjectPtr<URigVMMemoryStorage>& MemoryStorage, ERigVMMemoryType InMemoryType, EObjectFlags ObjectFlags, bool bForceCreation);
@@ -285,65 +285,57 @@ public:
 	virtual void GenerateMemoryType(FRigVMExtendedExecuteContext& Context, ERigVMMemoryType InMemoryType, const TArray<FRigVMPropertyDescription>* InProperties, bool bForceCreation = false);
 	
 	// Returns a memory storage by type
-	virtual TRigVMMemoryStorage* GetMemoryByType(FRigVMExtendedExecuteContext& Context, ERigVMMemoryType InMemoryType);
-	virtual const TRigVMMemoryStorage* GetMemoryByType(const FRigVMExtendedExecuteContext& Context, ERigVMMemoryType InMemoryType) const;
+	virtual FRigVMMemoryStorageStruct* GetMemoryByType(FRigVMExtendedExecuteContext& Context, ERigVMMemoryType InMemoryType);
+	virtual const FRigVMMemoryStorageStruct* GetMemoryByType(const FRigVMExtendedExecuteContext& Context, ERigVMMemoryType InMemoryType) const;
 
 	UE_DEPRECATED(5.4, "This version of GetLiteralMemory has been deprecated from the VM. Please, use GetLiteralMemory from Host or version with no parameter.")
-	TRigVMMemoryStorageDeprecatedType* GetLiteralMemory(bool bCreateIfNeeded) { return nullptr; }
+	URigVMMemoryStorage* GetLiteralMemory(bool bCreateIfNeeded) { return nullptr; }
 
 	// The default const literal memory
-	TRigVMMemoryStorage* GetLiteralMemory()
+	FRigVMMemoryStorageStruct* GetLiteralMemory()
 	{
-#if UE_RIGVM_PROPERTY_BAG_STORAGE_ENABLED
 		return &LiteralMemoryStorage;
-#else
-		return LiteralMemoryStorageObject;
-#endif
 	}
-	const TRigVMMemoryStorage* GetLiteralMemory() const
+	const FRigVMMemoryStorageStruct* GetLiteralMemory() const
 	{
-#if UE_RIGVM_PROPERTY_BAG_STORAGE_ENABLED
 		return &LiteralMemoryStorage;
-#else
-		return LiteralMemoryStorageObject;
-#endif
 	}
 
 	UE_DEPRECATED(5.4, "GetWorkMemory has been deprecated from the VM. Please, use GetWorkMemory from VMHost or the version with a Context param")
-	TRigVMMemoryStorageDeprecatedType* GetWorkMemory(bool bCreateIfNeeded = true) { return nullptr; }
+	URigVMMemoryStorage* GetWorkMemory(bool bCreateIfNeeded = true) { return nullptr; }
 
 	// The default mutable work memory
-	TRigVMMemoryStorage* GetWorkMemory(FRigVMExtendedExecuteContext& Context)
+	FRigVMMemoryStorageStruct* GetWorkMemory(FRigVMExtendedExecuteContext& Context)
 	{
 		return GetMemoryByType(Context, ERigVMMemoryType::Work);
 	}
 
-	const TRigVMMemoryStorage* GetWorkMemory(const FRigVMExtendedExecuteContext& Context) const
+	const FRigVMMemoryStorageStruct* GetWorkMemory(const FRigVMExtendedExecuteContext& Context) const
 	{
 		return GetMemoryByType(Context, ERigVMMemoryType::Work);
 	}
 
 	UE_DEPRECATED(5.4, "GetDebugMemory has been deprecated from the VM. Please, use GetDebugMemory from VMHost or the version with a Context param")
-	TRigVMMemoryStorageDeprecatedType* GetDebugMemory(bool bCreateIfNeeded = true) { return nullptr; }
+	URigVMMemoryStorage* GetDebugMemory(bool bCreateIfNeeded = true) { return nullptr; }
 
 	// The default debug watch memory
-	TRigVMMemoryStorage* GetDebugMemory(FRigVMExtendedExecuteContext& Context)
+	FRigVMMemoryStorageStruct* GetDebugMemory(FRigVMExtendedExecuteContext& Context)
 	{
 		return GetMemoryByType(Context, ERigVMMemoryType::Debug);
 	}
 
-	const TRigVMMemoryStorage* GetDebugMemory(const FRigVMExtendedExecuteContext& Context) const
+	const FRigVMMemoryStorageStruct* GetDebugMemory(const FRigVMExtendedExecuteContext& Context) const
 	{
 		return GetMemoryByType(Context, ERigVMMemoryType::Debug);
 	}
 
 	UE_DEPRECATED(5.4, "Please, use GetLocalMemoryArray with a Context param")
-	TArray<TRigVMMemoryStorageDeprecatedType*> GetLocalMemoryArray() { return TArray<TRigVMMemoryStorageDeprecatedType*>(); }
+	TArray<URigVMMemoryStorage*> GetLocalMemoryArray() { return TArray<URigVMMemoryStorage*>(); }
 
 	// returns all memory storages as an array
-	TArray<TRigVMMemoryStorage*> GetLocalMemoryArray(FRigVMExtendedExecuteContext& Context)
+	TArray<FRigVMMemoryStorageStruct*> GetLocalMemoryArray(FRigVMExtendedExecuteContext& Context)
 	{
-		TArray<TRigVMMemoryStorage*> LocalMemory;
+		TArray<FRigVMMemoryStorageStruct*> LocalMemory;
 		LocalMemory.Add(GetWorkMemory(Context));
 		LocalMemory.Add(GetLiteralMemory());
 		LocalMemory.Add(GetDebugMemory(Context));
@@ -487,13 +479,13 @@ public:
 	bool ResumeExecution() { return false; }
 	UE_DEPRECATED(5.3, "Please, use ResumeExecution with Context param")
 	bool ResumeExecution(TArrayView<URigVMMemoryStorage*> Memory, const FName& InEntryName = NAME_None) { return false; }
-	UE_DEPRECATED(5.4, "Please, use ResumeExecution with TRigVMMemoryStorage param")
-	bool ResumeExecution(FRigVMExtendedExecuteContext& Context, TArrayView<TRigVMMemoryStorageDeprecatedType*> Memory, const FName& InEntryName = NAME_None) { return false; }
+	UE_DEPRECATED(5.4, "Please, use ResumeExecution with FRigVMMemoryStorageStruct param")
+	bool ResumeExecution(FRigVMExtendedExecuteContext& Context, TArrayView<URigVMMemoryStorage*> Memory, const FName& InEntryName = NAME_None) { return false; }
 
 	UE_DEPRECATED(5.4, "Please, use ResumeExecution in the Context RigVMDebugInfo")
 	bool ResumeExecution(FRigVMExtendedExecuteContext& Context) { return false; }
 
-	bool ResumeExecution(FRigVMExtendedExecuteContext& Context, TArrayView<TRigVMMemoryStorage*> Memory, const FName& InEntryName = NAME_None);
+	bool ResumeExecution(FRigVMExtendedExecuteContext& Context, TArrayView<FRigVMMemoryStorageStruct*> Memory, const FName& InEntryName = NAME_None);
 
 #endif // WITH_EDITOR
 
@@ -945,8 +937,8 @@ public:
 
 private:
 	void InstructionOpEval(FRigVMExtendedExecuteContext& Context, int32 InstructionIndex, int32 InHandleBaseIndex, const TFunctionRef<void(FRigVMExtendedExecuteContext& Context, int32 InHandleIndex, const FRigVMBranchInfoKey& InBranchInfoKey, const FRigVMOperand& InArg)>& InOpFunc);
-	void PrepareMemoryForExecution(FRigVMExtendedExecuteContext& Context, TArrayView<TRigVMMemoryStorage*> InMemory);
-	void CacheMemoryHandlesIfRequired(FRigVMExtendedExecuteContext& Context, TArrayView<TRigVMMemoryStorage*> InMemory);
+	void PrepareMemoryForExecution(FRigVMExtendedExecuteContext& Context, TArrayView<FRigVMMemoryStorageStruct*> InMemory);
+	void CacheMemoryHandlesIfRequired(FRigVMExtendedExecuteContext& Context, TArrayView<FRigVMMemoryStorageStruct*> InMemory);
 	void RebuildByteCodeOnLoad();
 
 	UPROPERTY(transient)
diff --git a/Engine/Plugins/Runtime/RigVM/Source/RigVM/Public/RigVMCore/RigVMExecuteContext.h b/Engine/Plugins/Runtime/RigVM/Source/RigVM/Public/RigVMCore/RigVMExecuteContext.h
--- a/Engine/Plugins/Runtime/RigVM/Source/RigVM/Public/RigVMCore/RigVMExecuteContext.h
+++ b/Engine/Plugins/Runtime/RigVM/Source/RigVM/Public/RigVMCore/RigVMExecuteContext.h
@@ -736,14 +736,14 @@ struct RIGVM_API FRigVMExtendedExecuteContext
 	FName CurrentEntryName = NAME_None;
 	bool bCurrentlyRunningRootEntry = false;
 
-	TArrayView<TRigVMMemoryStorage*> CurrentVMMemory;
+	TArrayView<FRigVMMemoryStorageStruct*> CurrentVMMemory;
 
 #if WITH_EDITORONLY_DATA
 	// changes to the layout of cached memory array should be reflected in GetContainerIndex()
 	UE_DEPRECATED(5.4, "CachedMemory has been deprecated.")
 	TArray<URigVMMemoryStorage*> CachedMemory;
-	UE_DEPRECATED(5.4, "CurrentMemory has been deprecated, please use CurrentVMMemory with TRigVMMemoryStorage.")
-	TArrayView<TRigVMMemoryStorageDeprecatedType*> CurrentMemory;
+	UE_DEPRECATED(5.4, "CurrentMemory has been deprecated, please use CurrentVMMemory with FRigVMMemoryStorageStruct.")
+	TArrayView<URigVMMemoryStorage*> CurrentMemory;
 	UE_DEPRECATED(5.4, "DeferredVMToCopy has been deprecated.")
 	TObjectPtr<URigVM> DeferredVMToCopy = nullptr;
 	UE_DEPRECATED(5.4, "DeferredVMContextToCopy has been deprecated.")
diff --git a/Engine/Plugins/Runtime/RigVM/Source/RigVM/Public/RigVMCore/RigVMMemoryStorage.h b/Engine/Plugins/Runtime/RigVM/Source/RigVM/Public/RigVMCore/RigVMMemoryStorage.h
--- a/Engine/Plugins/Runtime/RigVM/Source/RigVM/Public/RigVMCore/RigVMMemoryStorage.h
+++ b/Engine/Plugins/Runtime/RigVM/Source/RigVM/Public/RigVMCore/RigVMMemoryStorage.h
@@ -799,7 +799,6 @@ private:
 			const FArrayProperty* ArrayProperty = CastFieldChecked<FArrayProperty>(Property);
 			FScriptArrayHelper ArrayHelper(ArrayProperty, Ptr);
 
-#if UE_RIGVM_PROPERTY_BAG_STORAGE_ENABLED
 			if (InSliceIndex >= ArrayHelper.Num() - 1)
 			{
 				// For each slice we copy the default value to the next slice index
@@ -838,48 +837,6 @@ private:
 					}
 				}
 			}
-#else // !UE_RIGVM_PROPERTY_BAG_STORAGE_ENABLED
-			if (ArrayHelper.Num() <= InSliceIndex)
-			{
-				const int32 NumValuesToAdd = 1 + InSliceIndex - ArrayHelper.Num();
-				const int32 FirstAddedIndex = ArrayHelper.AddValues(NumValuesToAdd);
-
-				if (FirstAddedIndex > 0)
-				{
-					// Adding slices, we need to initialize the new values.
-					// Otherwise, if we are adding the first slice, we don't need to worry about initializing.
-					UObject* CDO = ArrayProperty->GetOwnerClass()->GetDefaultObject();
-					const uint8* DefaultArrayMemory = ArrayProperty->ContainerPtrToValuePtr<uint8>(CDO);		
-					FScriptArrayHelper DefaultArrayHelper(ArrayProperty, DefaultArrayMemory);
-					if (const uint8* DefaultElementMemory = DefaultArrayHelper.GetRawPtr(0))
-					{
-						const FProperty* ElementProperty = ArrayProperty->Inner;
-						for (int32 i=FirstAddedIndex; i<ArrayHelper.Num(); ++i)
-						{
-#if UE_RIGVM_DEBUG_EXECUTION
-							 FString DefaultValue;
-							 ElementProperty->ExportText_Direct(
-							 	DefaultValue,
-							 	DefaultElementMemory,
-							 	DefaultElementMemory,
-							 	nullptr,
-							 	PPF_None,
-							 	nullptr);
-							
-							 UE_LOG(LogRigVM, Display, TEXT("Adding slice %d for Property '%s', defaulting to '%s'."),
-							 	InSliceIndex,
-							 	*ArrayProperty->GetName(),
-							 	*DefaultValue
-							 );
-#endif
-
-							uint8* DestMemory = ArrayHelper.GetRawPtr(i);
-							ElementProperty->CopyCompleteValue(DestMemory, DefaultElementMemory);
-						}
-					}
-				}
-			}
-#endif // UE_RIGVM_PROPERTY_BAG_STORAGE_ENABLED
 
 #if UE_RIGVM_DEBUG_EXECUTION
 			// const FProperty* ElementProperty = ArrayProperty->Inner;
diff --git a/Engine/Plugins/Runtime/RigVM/Source/RigVM/Public/RigVMCore/RigVMNativized.h b/Engine/Plugins/Runtime/RigVM/Source/RigVM/Public/RigVMCore/RigVMNativized.h
--- a/Engine/Plugins/Runtime/RigVM/Source/RigVM/Public/RigVMCore/RigVMNativized.h
+++ b/Engine/Plugins/Runtime/RigVM/Source/RigVM/Public/RigVMCore/RigVMNativized.h
@@ -50,11 +50,11 @@ public:
 	virtual bool IsNativized() const override { return true; }
 	virtual void Empty(FRigVMExtendedExecuteContext& Context) override { return; }
 	virtual void CopyFrom(URigVM* InVM, bool bDeferCopy = false, bool bReferenceLiteralMemory = false, bool bReferenceByteCode = false, bool bCopyExternalVariables = false, bool bCopyDynamicRegisters = false) override { return; }
-	virtual bool Initialize(FRigVMExtendedExecuteContext& Context, TArrayView<TRigVMMemoryStorage*> Memory) override;
-	virtual ERigVMExecuteResult Execute(FRigVMExtendedExecuteContext& Context, TArrayView<TRigVMMemoryStorage*> Memory, const FName& InEntryName = NAME_None) override;
+	virtual bool Initialize(FRigVMExtendedExecuteContext& Context, TArrayView<FRigVMMemoryStorageStruct*> Memory) override;
+	virtual ERigVMExecuteResult Execute(FRigVMExtendedExecuteContext& Context, TArrayView<FRigVMMemoryStorageStruct*> Memory, const FName& InEntryName = NAME_None) override;
 	virtual int32 AddRigVMFunction(UScriptStruct* InRigVMStruct, const FName& InMethodName) override { return INDEX_NONE; }
 	virtual FString GetRigVMFunctionName(int32 InFunctionIndex) const override { return FString(); }
-	virtual TRigVMMemoryStorage* GetMemoryByType(FRigVMExtendedExecuteContext& Context, ERigVMMemoryType InMemoryType) override { return nullptr; }
+	virtual FRigVMMemoryStorageStruct* GetMemoryByType(FRigVMExtendedExecuteContext& Context, ERigVMMemoryType InMemoryType) override { return nullptr; }
 	virtual void ClearMemory() override { return; }
 	virtual const FRigVMInstructionArray& GetInstructions() override;
 	virtual bool ContainsEntry(const FName& InEntryName) const override { return GetEntryNames().Contains(InEntryName); }
@@ -76,14 +76,14 @@ public:
 	virtual void Empty() override {}
 	UE_DEPRECATED(5.3, "Please, use Initialize with Context param")
 	virtual bool Initialize(TArrayView<URigVMMemoryStorage*> Memory) override { return false; }
-	UE_DEPRECATED(5.4, "Please, use Execute with TRigVMMemoryStorage param")
-	virtual bool Initialize(FRigVMExtendedExecuteContext& Context, TArrayView<TRigVMMemoryStorageDeprecatedType*> Memory) override { return false; }
-	UE_DEPRECATED(5.4, "GetMemoryByType has been deprecated from the VM. Please, use GetWorkMemory from VMHost or the new GetMemoryByType with a TRigVMMemoryStorage and Context parameter.")
-	virtual TRigVMMemoryStorageDeprecatedType* GetMemoryByType(ERigVMMemoryType InMemoryType, bool bCreateIfNeeded = false) override { return nullptr; }
+	UE_DEPRECATED(5.4, "Please, use Execute with FRigVMMemoryStorageStruct param")
+	virtual bool Initialize(FRigVMExtendedExecuteContext& Context, TArrayView<URigVMMemoryStorage*> Memory) override { return false; }
+	UE_DEPRECATED(5.4, "GetMemoryByType has been deprecated from the VM. Please, use GetWorkMemory from VMHost or the new GetMemoryByType with a FRigVMMemoryStorageStruct and Context parameter.")
+	virtual URigVMMemoryStorage* GetMemoryByType(ERigVMMemoryType InMemoryType, bool bCreateIfNeeded = false) override { return nullptr; }
 	UE_DEPRECATED(5.3, "Please, use Execute with Context param")
 	virtual ERigVMExecuteResult Execute(TArrayView<URigVMMemoryStorage*> Memory, const FName& InEntryName = NAME_None) override { return ERigVMExecuteResult::Failed; }
-	UE_DEPRECATED(5.4, "Please, use Execute with TRigVMMemoryStorage param")
-	virtual ERigVMExecuteResult Execute(FRigVMExtendedExecuteContext& Context, TArrayView<TRigVMMemoryStorageDeprecatedType*> Memory, const FName& InEntryName = NAME_None) override { return ERigVMExecuteResult::Failed; }
+	UE_DEPRECATED(5.4, "Please, use Execute with FRigVMMemoryStorageStruct param")
+	virtual ERigVMExecuteResult Execute(FRigVMExtendedExecuteContext& Context, TArrayView<URigVMMemoryStorage*> Memory, const FName& InEntryName = NAME_None) override { return ERigVMExecuteResult::Failed; }
 	UE_DEPRECATED(5.3, "Please, use SetInstructionIndex with Context param")
 	virtual void SetInstructionIndex(uint16 InInstructionIndex) {}
 
diff --git a/Engine/Plugins/Runtime/RigVM/Source/RigVM/Public/RigVMDefines.h b/Engine/Plugins/Runtime/RigVM/Source/RigVM/Public/RigVMDefines.h
--- a/Engine/Plugins/Runtime/RigVM/Source/RigVM/Public/RigVMDefines.h
+++ b/Engine/Plugins/Runtime/RigVM/Source/RigVM/Public/RigVMDefines.h
@@ -33,19 +33,3 @@
 #ifndef UE_RIGVMCONTROLLER_VERBOSE_REPOPULATE
 #define UE_RIGVMCONTROLLER_VERBOSE_REPOPULATE 0
 #endif
-
-#ifndef UE_RIGVM_PROPERTY_BAG_STORAGE_ENABLED
-#define UE_RIGVM_PROPERTY_BAG_STORAGE_ENABLED 1
-#endif
-
-#if UE_RIGVM_PROPERTY_BAG_STORAGE_ENABLED
-	using TRigVMMemoryStorage = struct FRigVMMemoryStorageStruct;
-	using TRigVMMemoryStorageDeprecatedType = class URigVMMemoryStorage;
-#else
-	struct URigVMMemoryStorageDummy
-	{
-	};
-
-	using TRigVMMemoryStorage = class URigVMMemoryStorage;
-	using TRigVMMemoryStorageDeprecatedType = struct URigVMMemoryStorageDummy;
-#endif
diff --git a/Engine/Plugins/Runtime/RigVM/Source/RigVM/Public/RigVMHost.h b/Engine/Plugins/Runtime/RigVM/Source/RigVM/Public/RigVMHost.h
--- a/Engine/Plugins/Runtime/RigVM/Source/RigVM/Public/RigVMHost.h
+++ b/Engine/Plugins/Runtime/RigVM/Source/RigVM/Public/RigVMHost.h
@@ -229,20 +229,20 @@ public:
 #endif
 
 	// Returns a VM memory storage by type
-	virtual TRigVMMemoryStorage* GetMemoryByType(ERigVMMemoryType InMemoryType);
-	virtual const TRigVMMemoryStorage* GetMemoryByType(ERigVMMemoryType InMemoryType) const;
+	virtual FRigVMMemoryStorageStruct* GetMemoryByType(ERigVMMemoryType InMemoryType);
+	virtual const FRigVMMemoryStorageStruct* GetMemoryByType(ERigVMMemoryType InMemoryType) const;
 
 	// The default mutable work memory
-	TRigVMMemoryStorage* GetWorkMemory() { return GetMemoryByType(ERigVMMemoryType::Work); }
-	const TRigVMMemoryStorage* GetWorkMemory() const { return GetMemoryByType(ERigVMMemoryType::Work); }
+	FRigVMMemoryStorageStruct* GetWorkMemory() { return GetMemoryByType(ERigVMMemoryType::Work); }
+	const FRigVMMemoryStorageStruct* GetWorkMemory() const { return GetMemoryByType(ERigVMMemoryType::Work); }
 
 	// The default const literal memory
-	TRigVMMemoryStorage* GetLiteralMemory() { return GetMemoryByType(ERigVMMemoryType::Literal); }
-	const TRigVMMemoryStorage* GetLiteralMemory() const { return GetMemoryByType(ERigVMMemoryType::Literal); }
+	FRigVMMemoryStorageStruct* GetLiteralMemory() { return GetMemoryByType(ERigVMMemoryType::Literal); }
+	const FRigVMMemoryStorageStruct* GetLiteralMemory() const { return GetMemoryByType(ERigVMMemoryType::Literal); }
 
 	// The default debug watch memory
-	TRigVMMemoryStorage* GetDebugMemory() { return GetMemoryByType(ERigVMMemoryType::Debug); }
-	const TRigVMMemoryStorage* GetDebugMemory() const { return GetMemoryByType(ERigVMMemoryType::Debug); }
+	FRigVMMemoryStorageStruct* GetDebugMemory() { return GetMemoryByType(ERigVMMemoryType::Debug); }
+	const FRigVMMemoryStorageStruct* GetDebugMemory() const { return GetMemoryByType(ERigVMMemoryType::Debug); }
 
 	DECLARE_EVENT_TwoParams(URigVM, FRigVMExecutedEvent, class URigVMHost*, const FName&);
 	FRigVMExecutedEvent& OnInitialized_AnyThread() { return InitializedEvent; }
@@ -311,7 +311,7 @@ protected:
 #endif
 
 	void GenerateUserDefinedDependenciesData(FRigVMExtendedExecuteContext& Context);
-	TArray<const UObject*> GetUserDefinedDependencies(const TArray<const TRigVMMemoryStorage*> InMemory);
+	TArray<const UObject*> GetUserDefinedDependencies(const TArray<const FRigVMMemoryStorageStruct*> InMemory);
 
 	UPROPERTY()
 	TMap<FString, FSoftObjectPath> UserDefinedStructGuidToPathName;
@@ -389,10 +389,8 @@ public:
 	virtual const TArray<UAssetUserData*>* GetAssetUserDataArray() const override;
 	//~ End IInterface_AssetUserData Interface
 
-#if UE_RIGVM_PROPERTY_BAG_STORAGE_ENABLED
 	// Temporary flag while we get TLS-based work memory working, controls copying from the CDO in InitializeVM
 	bool bTEMP_CopyDefaultsFromCDO = true;
-#endif
 
 protected:
 	/** Array of user data stored with the asset */
diff --git a/Engine/Plugins/Runtime/RigVM/Source/RigVMDeveloper/Private/RigVMCompiler/RigVMCodeGenerator.cpp b/Engine/Plugins/Runtime/RigVM/Source/RigVMDeveloper/Private/RigVMCompiler/RigVMCodeGenerator.cpp
--- a/Engine/Plugins/Runtime/RigVM/Source/RigVMDeveloper/Private/RigVMCompiler/RigVMCodeGenerator.cpp
+++ b/Engine/Plugins/Runtime/RigVM/Source/RigVMDeveloper/Private/RigVMCompiler/RigVMCodeGenerator.cpp
@@ -98,12 +98,12 @@ static constexpr TCHAR RigVM_GetVMHashFormat[] = TEXT("\tvirtual uint32 GetVMHas
 static constexpr TCHAR RigVM_GetEntryNamesFormat[] = TEXT("\tvirtual const TArray<FName>& GetEntryNames() const override\r\n\t{\r\n\t\tstatic const TArray<FName> StaticEntryNames = { {0} };\r\n\t\treturn StaticEntryNames;\r\n\t}");
 static constexpr TCHAR RigVM_DeclareUpdateExternalVariablesFormat[] = TEXT("\tvirtual void UpdateExternalVariables(FRigVMExtendedExecuteContext& Context) override;");
 static constexpr TCHAR RigVM_DeclareInvokeEntryByNameFormat[] = TEXT("\tERigVMExecuteResult InvokeEntryByName(FRigVMExtendedExecuteContext& Context, const FName& InEntryName{0});");
-static constexpr TCHAR RigVM_DeclareInitializeFormat[] = TEXT("\tvirtual bool Initialize(FRigVMExtendedExecuteContext& Context, TArrayView<TRigVMMemoryStorage*> Memory) override;");
-static constexpr TCHAR RigVM_DefineInitializeFormat[] = TEXT("bool U{0}::Initialize(FRigVMExtendedExecuteContext& Context, TArrayView<TRigVMMemoryStorage*> Memory)\r\n{");
-static constexpr TCHAR RigVM_DeclareExecuteFormat[] = TEXT("\tvirtual ERigVMExecuteResult Execute(FRigVMExtendedExecuteContext& Context, TArrayView<TRigVMMemoryStorage*> Memory, const FName& InEntryName) override;");
+static constexpr TCHAR RigVM_DeclareInitializeFormat[] = TEXT("\tvirtual bool Initialize(FRigVMExtendedExecuteContext& Context, TArrayView<FRigVMMemoryStorageStruct*> Memory) override;");
+static constexpr TCHAR RigVM_DefineInitializeFormat[] = TEXT("bool U{0}::Initialize(FRigVMExtendedExecuteContext& Context, TArrayView<FRigVMMemoryStorageStruct*> Memory)\r\n{");
+static constexpr TCHAR RigVM_DeclareExecuteFormat[] = TEXT("\tvirtual ERigVMExecuteResult Execute(FRigVMExtendedExecuteContext& Context, TArrayView<FRigVMMemoryStorageStruct*> Memory, const FName& InEntryName) override;");
 static constexpr TCHAR RigVM_DefineUpdateExternalVariablesFormat[] = TEXT("void U{0}::UpdateExternalVariables(FRigVMExtendedExecuteContext& Context)\r\n{");
 static constexpr TCHAR RigVM_DefineInvokeEntryByNameFormat[] = TEXT("ERigVMExecuteResult U{0}::InvokeEntryByName(FRigVMExtendedExecuteContext& Context, const FName& InEntryName{1})\r\n{");
-static constexpr TCHAR RigVM_DefineExecuteFormat[] = TEXT("ERigVMExecuteResult U{0}::Execute(FRigVMExtendedExecuteContext& Context, TArrayView<TRigVMMemoryStorage*> Memory, const FName& InEntryName)\r\n{");
+static constexpr TCHAR RigVM_DefineExecuteFormat[] = TEXT("ERigVMExecuteResult U{0}::Execute(FRigVMExtendedExecuteContext& Context, TArrayView<FRigVMMemoryStorageStruct*> Memory, const FName& InEntryName)\r\n{");
 static constexpr TCHAR RigVM_DeclareExecuteEntryFormat[] = TEXT("\tERigVMExecuteResult ExecuteEntry_{0}(FRigVMExtendedExecuteContext& Context, {1});");
 static constexpr TCHAR RigVM_DefineExecuteEntryFormat[] = TEXT("ERigVMExecuteResult U{0}::ExecuteEntry_{1}(FRigVMExtendedExecuteContext& Context, {2})\r\n{");
 static constexpr TCHAR RigVM_DeclareExecuteGroupFormat[] = TEXT("\tERigVMExecuteResult ExecuteGroup_{0}_{1}({2});");
@@ -1241,73 +1241,55 @@ void FRigVMCodeGenerator::ParseRequiredUProperties(const FRigVMExtendedExecuteCo
 	}
 }
 
-void FRigVMCodeGenerator::ParseMemory(const FRigVMExtendedExecuteContext& Context, TRigVMMemoryStorage* InMemory)
+void FRigVMCodeGenerator::ParseMemory(const FRigVMExtendedExecuteContext& Context, FRigVMMemoryStorageStruct* InMemory)
 {
 	if (InMemory == nullptr)
 	{
 		return;
 	}
-#if UE_RIGVM_PROPERTY_BAG_STORAGE_ENABLED
 	for (const FProperty* Property : InMemory->GetProperties())
 	{
 		ParseProperty(Context, InMemory->GetMemoryType(), Property, InMemory);
 	}
-#else
-	if (InMemory->GetClass() == TRigVMMemoryStorage::StaticClass())
-	{
-		return;
-	}
-
-	for (TFieldIterator<FProperty> PropertyIt(InMemory->GetClass()); PropertyIt; ++PropertyIt)
-	{
-		const FProperty* Property = *PropertyIt;
-		ParseProperty(Context, InMemory->GetMemoryType(), Property, InMemory);
-	}
-#endif // UE_RIGVM_PROPERTY_BAG_STORAGE_ENABLED
 }
 
-void FRigVMCodeGenerator::ParseProperty(const FRigVMExtendedExecuteContext& Context, ERigVMMemoryType InMemoryType, const FProperty* InProperty, TRigVMMemoryStorage* InMemory)
+void FRigVMCodeGenerator::ParseProperty(const FRigVMExtendedExecuteContext& Context, ERigVMMemoryType InMemoryType, const FProperty* InProperty, FRigVMMemoryStorageStruct* InMemory)
 {
-#if !UE_RIGVM_PROPERTY_BAG_STORAGE_ENABLED
-	if (URigVMMemoryStorageGeneratorClass* MemoryClass = Cast<URigVMMemoryStorageGeneratorClass>(InMemory->GetClass()))
-#endif
-	{
-		const int32 PropertyIndex = InMemory->GetPropertyIndex(InProperty);
-		const FRigVMOperand Operand(InMemoryType, PropertyIndex);
-		const FRigVMPropertyDescription& PropertyDescription = GetPropertyDescForOperand(Context, Operand);
+	const int32 PropertyIndex = InMemory->GetPropertyIndex(InProperty);
+	const FRigVMOperand Operand(InMemoryType, PropertyIndex);
+	const FRigVMPropertyDescription& PropertyDescription = GetPropertyDescForOperand(Context, Operand);
 
-		FPropertyInfo Info;
-		Info.MemoryPropertyIndex = PropertyIndex;
-		Info.Description = PropertyDescription;
+	FPropertyInfo Info;
+	Info.MemoryPropertyIndex = PropertyIndex;
+	Info.Description = PropertyDescription;
 
-		check(InMemoryType == ERigVMMemoryType::Literal || InMemoryType == ERigVMMemoryType::Work);
-		if(InMemoryType == ERigVMMemoryType::Literal)
-		{
-			Info.PropertyType = ERigVMNativizedPropertyType::Literal;
-		}
-		else
-		{
-			Info.PropertyType = ERigVMNativizedPropertyType::Work;
+	check(InMemoryType == ERigVMMemoryType::Literal || InMemoryType == ERigVMMemoryType::Work);
+	if(InMemoryType == ERigVMMemoryType::Literal)
+	{
+		Info.PropertyType = ERigVMNativizedPropertyType::Literal;
+	}
+	else
+	{
+		Info.PropertyType = ERigVMNativizedPropertyType::Work;
 
-			// work state may be a hidden / sliced property.
-			if(RigVMTypeUtils::IsArrayType(PropertyDescription.CPPType))
+		// work state may be a hidden / sliced property.
+		if(RigVMTypeUtils::IsArrayType(PropertyDescription.CPPType))
+		{
+			if(const FString* PinPath = OperandToPinMap.Find(Operand))
 			{
-				if(const FString* PinPath = OperandToPinMap.Find(Operand))
+				if(const URigVMPin* Pin = Model->FindPin(*PinPath))
 				{
-					if(const URigVMPin* Pin = Model->FindPin(*PinPath))
+					if(Pin->GetDirection() == ERigVMPinDirection::Hidden)
 					{
-						if(Pin->GetDirection() == ERigVMPinDirection::Hidden)
-						{
-							Info.PropertyType = ERigVMNativizedPropertyType::Sliced;
-						}
+						Info.PropertyType = ERigVMNativizedPropertyType::Sliced;
 					}
 				}
 			}
 		}
-
-		const int32 LookupIndex = Properties.Add(Info);
-		PropertyNameToIndex.Add(Info.Description.Name, LookupIndex);
 	}
+
+	const int32 LookupIndex = Properties.Add(Info);
+	PropertyNameToIndex.Add(Info.Description.Name, LookupIndex);
 }
 
 void FRigVMCodeGenerator::ParseInstructionGroups(const FRigVMExtendedExecuteContext& Context)
@@ -1984,7 +1966,7 @@ FRigVMPropertyDescription FRigVMCodeGenerator::GetPropertyDescForOperand(const F
 	}
 	else
 	{
-		if (const TRigVMMemoryStorage* MemoryStorage = VM->GetMemoryByType(Context, InOperand.GetMemoryType()))
+		if (const FRigVMMemoryStorageStruct* MemoryStorage = VM->GetMemoryByType(Context, InOperand.GetMemoryType()))
 		{
 			Property = MemoryStorage->GetProperty(InOperand.GetRegisterIndex());
 		}
@@ -2009,16 +1991,12 @@ FRigVMPropertyDescription FRigVMCodeGenerator::GetPropertyForOperand(const FRigV
 	}
 	else
 	{
-		if (const TRigVMMemoryStorage* MemoryStorage = VM->GetMemoryByType(Context, InOperand.GetMemoryType()))
+		if (const FRigVMMemoryStorageStruct* MemoryStorage = VM->GetMemoryByType(Context, InOperand.GetMemoryType()))
 		{
 			Property = MemoryStorage->GetProperty(InOperand.GetRegisterIndex());
 			if (Property)
 			{
-#if UE_RIGVM_PROPERTY_BAG_STORAGE_ENABLED
 				Memory = Property->ContainerPtrToValuePtr<uint8>(MemoryStorage->GetContainerPtr());
-#else
-				Memory = Property->ContainerPtrToValuePtr<uint8>(MemoryStorage);
-#endif // UE_RIGVM_PROPERTY_BAG_STORAGE_ENABLED
 			}
 		}
 	}
@@ -2062,22 +2040,12 @@ const FRigVMPropertyPathDescription& FRigVMCodeGenerator::GetPropertyPathForOper
 		}
 		else
 		{
-			if (const TRigVMMemoryStorage* MemoryStorage = VM->GetMemoryByType(Context, InOperand.GetMemoryType()))
+			if (const FRigVMMemoryStorageStruct* MemoryStorage = VM->GetMemoryByType(Context, InOperand.GetMemoryType()))
 			{
-#if UE_RIGVM_PROPERTY_BAG_STORAGE_ENABLED
 				if (MemoryStorage->IsValidPropertyPathDescriptionIndex((RegisterOffsetIndex)))
 				{
 					return *MemoryStorage->GetPropertyPathDescriptionByIndex(RegisterOffsetIndex);
 				}
-#else
-				if (const URigVMMemoryStorageGeneratorClass* MemoryClass = Cast<URigVMMemoryStorageGeneratorClass>(MemoryStorage->GetClass()))
-				{
-					if (MemoryClass->IsValidPropertyPathDescriptionIndex((RegisterOffsetIndex)))
-					{
-						return *MemoryClass->GetPropertyPathDescriptionByIndex(RegisterOffsetIndex);
-					}
-				}
-#endif // UE_RIGVM_PROPERTY_BAG_STORAGE_ENABLED
 			}
 		}
 	}
@@ -2133,17 +2101,13 @@ void FRigVMCodeGenerator::CheckOperand(const FRigVMExtendedExecuteContext& Conte
 	}
 	else
 	{
-		const TRigVMMemoryStorage* Memory = VM->GetMemoryByType(Context, InOperand.GetMemoryType());
+		const FRigVMMemoryStorageStruct* Memory = VM->GetMemoryByType(Context, InOperand.GetMemoryType());
 		check(Memory);
 
 		ensure(Memory->GetProperties().IsValidIndex(InOperand.GetRegisterIndex()));
 		if (InOperand.GetRegisterOffset() != INDEX_NONE)
 		{
 			ensure(Memory->GetPropertyPaths().IsValidIndex(InOperand.GetRegisterOffset()));
-#if !UE_RIGVM_PROPERTY_BAG_STORAGE_ENABLED
-			URigVMMemoryStorageGeneratorClass* Class = CastChecked<URigVMMemoryStorageGeneratorClass>(Memory->GetClass()); 
-			ensure(Class->PropertyPathDescriptions.IsValidIndex(InOperand.GetRegisterOffset()));
-#endif // UE_RIGVM_PROPERTY_BAG_STORAGE_ENABLED
 		}
 	}
 }
diff --git a/Engine/Plugins/Runtime/RigVM/Source/RigVMDeveloper/Private/RigVMCompiler/RigVMCompiler.cpp b/Engine/Plugins/Runtime/RigVM/Source/RigVMDeveloper/Private/RigVMCompiler/RigVMCompiler.cpp
--- a/Engine/Plugins/Runtime/RigVM/Source/RigVMDeveloper/Private/RigVMCompiler/RigVMCompiler.cpp
+++ b/Engine/Plugins/Runtime/RigVM/Source/RigVMDeveloper/Private/RigVMCompiler/RigVMCompiler.cpp
@@ -195,7 +195,7 @@ const FProperty* FRigVMCompilerWorkData::GetPropertyForOperand(const FRigVMOpera
 	
 	check(!bSetupMemory);
 
-	auto GetPropertyFromMemory = [this](TRigVMMemoryStorage* InMemory, const FRigVMOperand& InOperand)
+	auto GetPropertyFromMemory = [this](FRigVMMemoryStorageStruct* InMemory, const FRigVMOperand& InOperand)
 	{
 		if(InOperand.GetRegisterOffset() == INDEX_NONE)
 		{
@@ -203,16 +203,8 @@ const FProperty* FRigVMCompilerWorkData::GetPropertyForOperand(const FRigVMOpera
 		}
 		if(!InMemory->GetPropertyPaths().IsValidIndex(InOperand.GetRegisterOffset()))
 		{
-#if UE_RIGVM_PROPERTY_BAG_STORAGE_ENABLED
 			InMemory->SetPropertyPathDescriptions(PropertyPathDescriptions.FindChecked(InOperand.GetMemoryType()));
 			InMemory->RefreshPropertyPaths();
-#else
-			if(URigVMMemoryStorageGeneratorClass* MemoryClass = Cast<URigVMMemoryStorageGeneratorClass>(InMemory->GetClass()))
-			{
-				MemoryClass->PropertyPathDescriptions = PropertyPathDescriptions.FindChecked(InOperand.GetMemoryType());;
-				MemoryClass->RefreshPropertyPaths();
-			}
-#endif
 		}
 		return InMemory->GetPropertyPaths()[InOperand.GetRegisterOffset()].GetTailProperty();
 	};
@@ -1215,32 +1207,11 @@ bool URigVMCompiler::Compile(const FRigVMCompileSettings& InSettings, TArray<URi
 
 	const TArray<ERigVMMemoryType> MemoryTypes = { ERigVMMemoryType::Literal, ERigVMMemoryType::Work, ERigVMMemoryType::Debug };
 
-#if UE_RIGVM_PROPERTY_BAG_STORAGE_ENABLED
 	for (ERigVMMemoryType MemoryType : MemoryTypes)
 	{
 		const TArray<FRigVMPropertyDescription>* Properties = WorkData.PropertyDescriptions.Find(MemoryType);
 		WorkData.VM->GenerateMemoryType(*WorkData.Context, MemoryType, Properties);
 	}
-#else
-	for(ERigVMMemoryType MemoryType : MemoryTypes)
-	{
-		UPackage* Package = WorkData.VM->GetOutermost();
-
-		const TArray<FRigVMPropertyDescription>* Properties = WorkData.PropertyDescriptions.Find(MemoryType);
-		if(Properties == nullptr)
-		{
-			URigVMMemoryStorageGeneratorClass::RemoveStorageClass(Package, MemoryType);
-			continue;
-		}
-
-		URigVMMemoryStorageGeneratorClass::CreateStorageClass(Package, MemoryType, *Properties);
-	}
-
-	for (ERigVMMemoryType MemoryType : MemoryTypes)
-	{
-		WorkData.VM->CreateMemoryByType(*WorkData.Context, MemoryType);
-	}
-#endif
 
 	WorkData.bSetupMemory = false;
 	WorkData.ExprComplete.Reset();
@@ -1277,9 +1248,8 @@ bool URigVMCompiler::Compile(const FRigVMCompileSettings& InSettings, TArray<URi
 	for(ERigVMMemoryType MemoryType : MemoryTypes)
 	{
 		const TArray<FRigVMPropertyPathDescription>* Descriptions = WorkData.PropertyPathDescriptions.Find(MemoryType);
-		if(TRigVMMemoryStorage* MemoryStorageObject = WorkData.VM->GetMemoryByType(*WorkData.Context, MemoryType))
+		if(FRigVMMemoryStorageStruct* MemoryStorageObject = WorkData.VM->GetMemoryByType(*WorkData.Context, MemoryType))
 		{
-#if UE_RIGVM_PROPERTY_BAG_STORAGE_ENABLED
 			if (Descriptions)
 			{
 				MemoryStorageObject->SetPropertyPathDescriptions(*Descriptions);
@@ -1289,20 +1259,6 @@ bool URigVMCompiler::Compile(const FRigVMCompileSettings& InSettings, TArray<URi
 				MemoryStorageObject->ResetPropertyPathDescriptions();
 			}
 			MemoryStorageObject->RefreshPropertyPaths();
-#else
-			if(URigVMMemoryStorageGeneratorClass* Class = Cast<URigVMMemoryStorageGeneratorClass>(MemoryStorageObject->GetClass()))
-			{
-				if(Descriptions)
-				{
-					Class->PropertyPathDescriptions = *Descriptions;
-				}
-				else
-				{
-					Class->PropertyPathDescriptions.Reset();
-				}
-				Class->RefreshPropertyPaths();
-			}
-#endif
 		}
 	}
 
diff --git a/Engine/Plugins/Runtime/RigVM/Source/RigVMDeveloper/Private/RigVMModel/Nodes/RigVMUnitNode.cpp b/Engine/Plugins/Runtime/RigVM/Source/RigVMDeveloper/Private/RigVMModel/Nodes/RigVMUnitNode.cpp
--- a/Engine/Plugins/Runtime/RigVM/Source/RigVMDeveloper/Private/RigVMModel/Nodes/RigVMUnitNode.cpp
+++ b/Engine/Plugins/Runtime/RigVM/Source/RigVMDeveloper/Private/RigVMModel/Nodes/RigVMUnitNode.cpp
@@ -455,15 +455,11 @@ void URigVMUnitNode::EnumeratePropertiesOnHostAndStructInstance(URigVMHost* InHo
 				}
 			case ERigVMMemoryType::Literal:
 				{
-					if(TRigVMMemoryStorage* LiteralMemory = InHost->GetLiteralMemory())
+					if(FRigVMMemoryStorageStruct* LiteralMemory = InHost->GetLiteralMemory())
 					{
 						if(const FProperty* Property = LiteralMemory->GetProperty(Operand.GetRegisterIndex()))
 						{
-#if UE_RIGVM_PROPERTY_BAG_STORAGE_ENABLED
 							HostMemory = Property->ContainerPtrToValuePtr<uint8>(LiteralMemory->GetContainerPtr());
-#else
-							HostMemory = Property->ContainerPtrToValuePtr<uint8>(LiteralMemory);
-#endif // UE_RIGVM_PROPERTY_BAG_STORAGE_ENABLED
 							HostProperty = Property;
 						}
 					}
@@ -471,15 +467,11 @@ void URigVMUnitNode::EnumeratePropertiesOnHostAndStructInstance(URigVMHost* InHo
 				}
 			case ERigVMMemoryType::Work:
 				{
-					if(TRigVMMemoryStorage* WorkMemory = InHost->GetWorkMemory())
+					if(FRigVMMemoryStorageStruct* WorkMemory = InHost->GetWorkMemory())
 					{
 						if(const FProperty* Property = WorkMemory->GetProperty(Operand.GetRegisterIndex()))
 						{
-#if UE_RIGVM_PROPERTY_BAG_STORAGE_ENABLED
 							HostMemory = Property->ContainerPtrToValuePtr<uint8>(WorkMemory->GetContainerPtr());
-#else
-							HostMemory = Property->ContainerPtrToValuePtr<uint8>(WorkMemory);
-#endif // UE_RIGVM_PROPERTY_BAG_STORAGE_ENABLED
 							HostProperty = Property;
 						}
 					}
@@ -578,7 +570,7 @@ bool URigVMUnitNode::UpdateHostFromStructInstance(URigVMHost* InHost, TSharedPtr
 		}
 
 		// copy the memory back to the host
-		TRigVMMemoryStorage::CopyProperty(HostProperty, HostMemory, InstanceProperty, InstanceMemory);
+		FRigVMMemoryStorageStruct::CopyProperty(HostProperty, HostMemory, InstanceProperty, InstanceMemory);
 		bChangedSomething = true;
 
 	}, InSliceIndex);
diff --git a/Engine/Plugins/Runtime/RigVM/Source/RigVMDeveloper/Public/RigVMCompiler/RigVMCodeGenerator.h b/Engine/Plugins/Runtime/RigVM/Source/RigVMDeveloper/Public/RigVMCompiler/RigVMCodeGenerator.h
--- a/Engine/Plugins/Runtime/RigVM/Source/RigVMDeveloper/Public/RigVMCompiler/RigVMCodeGenerator.h
+++ b/Engine/Plugins/Runtime/RigVM/Source/RigVMDeveloper/Public/RigVMCompiler/RigVMCodeGenerator.h
@@ -105,8 +105,8 @@ private:
 		int32 InMaxInstructionsPerFunction);
 	void ParseInclude(UStruct* InDependency, const FName& InMethodName = NAME_None);
 	void ParseRequiredUProperties(const FRigVMExtendedExecuteContext& Context);
-	void ParseMemory(const FRigVMExtendedExecuteContext& Context, TRigVMMemoryStorage* InMemory);
-	void ParseProperty(const FRigVMExtendedExecuteContext& Context, ERigVMMemoryType InMemoryType, const FProperty* InProperty, TRigVMMemoryStorage* InMemory);
+	void ParseMemory(const FRigVMExtendedExecuteContext& Context, FRigVMMemoryStorageStruct* InMemory);
+	void ParseProperty(const FRigVMExtendedExecuteContext& Context, ERigVMMemoryType InMemoryType, const FProperty* InProperty, FRigVMMemoryStorageStruct* InMemory);
 	void ParseInstructionGroups(const FRigVMExtendedExecuteContext& Context);
 	FString DumpInstructions(const FRigVMExtendedExecuteContext& Context, const FString& InPrefix, int32 InFirstInstruction, int32 InLastInstruction, const FInstructionGroup& InGroup, bool bLog = false);
 	FString DumpInstructions(const FRigVMExtendedExecuteContext& Context, const FString& InPrefix, const TArray<int32> InInstructionIndices, const FInstructionGroup& InGroup, bool bLog = false);
diff --git a/Engine/Plugins/Runtime/RigVM/Source/RigVMEditor/Private/Editor/RigVMEditor.cpp b/Engine/Plugins/Runtime/RigVM/Source/RigVMEditor/Private/Editor/RigVMEditor.cpp
--- a/Engine/Plugins/Runtime/RigVM/Source/RigVMEditor/Private/Editor/RigVMEditor.cpp
+++ b/Engine/Plugins/Runtime/RigVM/Source/RigVMEditor/Private/Editor/RigVMEditor.cpp
@@ -1719,9 +1719,7 @@ void FRigVMEditor::HandleVMCompiledEvent(UObject* InCompiledObject, URigVM* InVM
 			{
 				ERigVMMemoryType MemoryType = ActiveTab->GetMetaData<FMemoryTypeMetaData>()->MemoryType;
 				// TODO zzz : UE-195014 - Fix memory tab losing values on VM recompile
-				TRigVMMemoryStorage* Memory = InVM->GetMemoryByType(InContext, MemoryType);
-
-#if UE_RIGVM_PROPERTY_BAG_STORAGE_ENABLED
+				FRigVMMemoryStorageStruct* Memory = InVM->GetMemoryByType(InContext, MemoryType);
 
 			#if 1
 				ActiveTab->RequestCloseTab();
@@ -1733,10 +1731,6 @@ void FRigVMEditor::HandleVMCompiledEvent(UObject* InCompiledObject, URigVM* InVM
 				TSharedPtr<FStructOnScope> StructOnScope = MakeShareable(new FStructOnScope(Memory->GetPropertyBagStruct(), (uint8*)Memory->GetContainerPtr()));
 				StructDetailsView->SetStructureData(StructOnScope);
 			#endif
-#else
-				TSharedRef<IDetailsView> DetailsView = StaticCastSharedRef<IDetailsView>(ActiveTab->GetContent());
-				DetailsView->SetObject(Memory);
-#endif // UE_RIGVM_PROPERTY_BAG_STORAGE_ENABLED
 			}
 		}
 	}
diff --git a/Engine/Plugins/Runtime/RigVM/Source/RigVMEditor/Private/Editor/RigVMGraphDetailCustomization.cpp b/Engine/Plugins/Runtime/RigVM/Source/RigVMEditor/Private/Editor/RigVMGraphDetailCustomization.cpp
--- a/Engine/Plugins/Runtime/RigVM/Source/RigVMEditor/Private/Editor/RigVMGraphDetailCustomization.cpp
+++ b/Engine/Plugins/Runtime/RigVM/Source/RigVMEditor/Private/Editor/RigVMGraphDetailCustomization.cpp
@@ -1729,7 +1729,7 @@ void FRigVMWrappedNodeDetailCustomization::CustomizeLiveValues(IDetailLayoutBuil
 
 				const FProperty* Property = nullptr;
 
-				TArray<TRigVMMemoryStorage*> ExternalStructs;
+				TArray<FRigVMMemoryStorageStruct*> ExternalStructs;
 				TArray<UObject*> ExternalObjects;
 
 				if(Operand->GetMemoryType() == ERigVMMemoryType::External)
@@ -1743,30 +1743,11 @@ void FRigVMWrappedNodeDetailCustomization::CustomizeLiveValues(IDetailLayoutBuil
 				}
 				else
 				{
-					TRigVMMemoryStorage* Memory = DebuggedHost->GetMemoryByType(Operand->GetMemoryType());
+					FRigVMMemoryStorageStruct* Memory = DebuggedHost->GetMemoryByType(Operand->GetMemoryType());
 					if(Memory == nullptr)
 					{
 						continue;
 					}
-#if !UE_RIGVM_PROPERTY_BAG_STORAGE_ENABLED
-					if(Memory->GetOuter() == GetTransientPackage())
-					{
-						continue;
-					}
-
-					// the UClass must be alive for the details view to access it
-					// this ensure can fail if VM memory is not updated immediately after compile
-					// because of deferred copy
-					if(!ensure(IsValidChecked(Memory->GetClass())))
-					{
-						continue;
-					}
-
-					if(Memory->GetClass()->GetOuter() == GetTransientPackage())
-					{
-						continue;
-					}
-#endif // UE_RIGVM_PROPERTY_BAG_STORAGE_ENABLED
 
 					if(!Memory->IsValidIndex(Operand->GetRegisterIndex()))
 					{
@@ -1779,11 +1760,7 @@ void FRigVMWrappedNodeDetailCustomization::CustomizeLiveValues(IDetailLayoutBuil
 						continue;
 					}
 
-#if UE_RIGVM_PROPERTY_BAG_STORAGE_ENABLED
 					ExternalStructs.Add(Memory);
-#else
-					ExternalObjects.Add(Memory);
-#endif // UE_RIGVM_PROPERTY_BAG_STORAGE_ENABLED
 				}
 
 				check(ExternalObjects.Num() > 0 || ExternalStructs.Num() > 0);
@@ -1799,8 +1776,7 @@ void FRigVMWrappedNodeDetailCustomization::CustomizeLiveValues(IDetailLayoutBuil
 					NameSuffix = FString::Printf(TEXT("_%d"), SuffixIndex);
 				}
 
-#if UE_RIGVM_PROPERTY_BAG_STORAGE_ENABLED
-				for (TRigVMMemoryStorage* Memory : ExternalStructs)
+				for (FRigVMMemoryStorageStruct* Memory : ExternalStructs)
 				{
 					TSharedPtr<FStructOnScope> StructOnScope = MakeShareable(new FStructOnScope(Memory->GetPropertyBagStruct(), (uint8*)Memory->GetContainerPtr()));
 					if (IDetailPropertyRow* PropertyRow = DebugCategory.AddExternalStructureProperty(StructOnScope, Property->GetFName(), EPropertyLocation::Default, FAddPropertyParams().ForceShowProperty()))
@@ -1813,7 +1789,6 @@ void FRigVMWrappedNodeDetailCustomization::CustomizeLiveValues(IDetailLayoutBuil
 						NameSuffix = FString::Printf(TEXT("_%d"), SuffixIndex);
 					}
 				}
-#endif // UE_RIGVM_PROPERTY_BAG_STORAGE_ENABLED
 
 				KnownOperands.Add(*Operand);
 			}
diff --git a/Engine/Plugins/Runtime/RigVM/Source/RigVMEditor/Private/Editor/RigVMLocalVariableDetailCustomization.cpp b/Engine/Plugins/Runtime/RigVM/Source/RigVMEditor/Private/Editor/RigVMLocalVariableDetailCustomization.cpp
--- a/Engine/Plugins/Runtime/RigVM/Source/RigVMEditor/Private/Editor/RigVMLocalVariableDetailCustomization.cpp
+++ b/Engine/Plugins/Runtime/RigVM/Source/RigVMEditor/Private/Editor/RigVMLocalVariableDetailCustomization.cpp
@@ -96,7 +96,7 @@ void FRigVMLocalVariableDetailCustomization::CustomizeDetails(IDetailLayoutBuild
 		if (CDO->GetVM() != nullptr)
 		{
 			FString SourcePath = FString::Printf(TEXT("LocalVariableDefault::%s|%s::Const"), *GraphBeingCustomized->GetGraphName(), *VariableDescription.Name.ToString());
-			TRigVMMemoryStorage* LiteralMemory = CDO->GetVM()->GetLiteralMemory();
+			FRigVMMemoryStorageStruct* LiteralMemory = CDO->GetVM()->GetLiteralMemory();
 			FProperty* Property = LiteralMemory->FindPropertyByName(*SourcePath);
 			if (Property)
 			{
@@ -104,13 +104,9 @@ void FRigVMLocalVariableDetailCustomization::CustomizeDetails(IDetailLayoutBuild
 				Property->ClearPropertyFlags(CPF_EditConst);
 			
 				const FName SanitizedName = FRigVMPropertyDescription::SanitizeName(*SourcePath);
-#if UE_RIGVM_PROPERTY_BAG_STORAGE_ENABLED
 				TSharedPtr<FStructOnScope> StructOnScope = MakeShareable(new FStructOnScope(LiteralMemory->GetPropertyBagStruct(), (uint8*)LiteralMemory->GetContainerPtr()));
 				IDetailPropertyRow* Row = DefaultValueCategory.AddExternalStructureProperty(StructOnScope, SanitizedName);
-#else
-				TArray<UObject*> Objects = { LiteralMemory };
-				IDetailPropertyRow* Row = DefaultValueCategory.AddExternalObjectProperty(Objects, SanitizedName);
-#endif
+
 				Row->DisplayName(FText::FromName(VariableDescription.Name));
 
 				const FSimpleDelegate OnDefaultValueChanged = FSimpleDelegate::CreateLambda([this, Property, LiteralMemory]()
diff --git a/Engine/Plugins/Runtime/RigVM/Source/RigVMEditor/Private/Widgets/SRigVMGraphNode.cpp b/Engine/Plugins/Runtime/RigVM/Source/RigVMEditor/Private/Widgets/SRigVMGraphNode.cpp
--- a/Engine/Plugins/Runtime/RigVM/Source/RigVMEditor/Private/Widgets/SRigVMGraphNode.cpp
+++ b/Engine/Plugins/Runtime/RigVM/Source/RigVMEditor/Private/Widgets/SRigVMGraphNode.cpp
@@ -837,7 +837,7 @@ void SRigVMGraphNode::GetNodeInfoPopups(FNodeInfoContext* Context, TArray<FGraph
 							FString PinHash = URigVMCompiler::GetPinHash(ModelPin, nullptr, true);
 							if (const FRigVMOperand* WatchOperand = RigBlueprint->PinToOperandMap.Find(PinHash))
 							{
-								TRigVMMemoryStorage* Memory = ActiveObject->GetDebugMemory();
+								FRigVMMemoryStorageStruct* Memory = ActiveObject->GetDebugMemory();
 								// We mark PPF_ExternalEditor so that default values are also printed
 								const FString DebugValue = Memory->GetDataAsStringSafe(WatchOperand->GetRegisterIndex(), PPF_ExternalEditor | STRUCT_ExportTextItemNative);
 								if(!DebugValue.IsEmpty())
