Control Rig Modules; Polish - clean up of workflow issues

diff --git a/Engine/Plugins/Animation/ControlRig/Source/ControlRig/Private/ControlRig.cpp b/Engine/Plugins/Animation/ControlRig/Source/ControlRig/Private/ControlRig.cpp
--- a/Engine/Plugins/Animation/ControlRig/Source/ControlRig/Private/ControlRig.cpp
+++ b/Engine/Plugins/Animation/ControlRig/Source/ControlRig/Private/ControlRig.cpp
@@ -1641,7 +1641,7 @@ const FRigModuleSettings& UControlRig::GetRigModuleSettings() const
 
 bool UControlRig::IsRigModule() const
 {
-	return GetRigModuleSettings().IsValidModule();
+	return GetRigModuleSettings().IsValidModule(false);
 }
 
 bool UControlRig::IsRigModuleInstance() const
diff --git a/Engine/Plugins/Animation/ControlRig/Source/ControlRig/Private/Rigs/RigConnectionRules.cpp b/Engine/Plugins/Animation/ControlRig/Source/ControlRig/Private/Rigs/RigConnectionRules.cpp
--- a/Engine/Plugins/Animation/ControlRig/Source/ControlRig/Private/Rigs/RigConnectionRules.cpp
+++ b/Engine/Plugins/Animation/ControlRig/Source/ControlRig/Private/Rigs/RigConnectionRules.cpp
@@ -303,7 +303,7 @@ bool FRigTagConnectionRule::CanConnect(const FRigConnectionInfo* InConnectionInf
 
 TArray<FRigElementKey> FRigChainConnectionRule::GetAdditionalConnectors() const
 {
-	return {EndConnector};
+	return {RootConnector};
 }
 
 bool FRigChainConnectionRule::CanConnect(const FRigConnectionInfo* InConnectionInfo, FString* OutFailureReason) const
@@ -324,8 +324,8 @@ bool FRigChainConnectionRule::CanConnect(const FRigConnectionInfo* InConnectionI
 		return false;
 	}
 
-	const FRigElementKey StartKey = Targets[0];
-	const FRigElementKey EndKey = Targets.Last();
+	const FRigElementKey EndKey = Targets[0];
+	const FRigElementKey StartKey = Targets.Last();
 
 	FRigElementKey NextKey = EndKey;
 	TArray<FRigElementKey> ReversedChain;
diff --git a/Engine/Plugins/Animation/ControlRig/Source/ControlRig/Public/Rigs/RigConnectionRules.h b/Engine/Plugins/Animation/ControlRig/Source/ControlRig/Public/Rigs/RigConnectionRules.h
--- a/Engine/Plugins/Animation/ControlRig/Source/ControlRig/Public/Rigs/RigConnectionRules.h
+++ b/Engine/Plugins/Animation/ControlRig/Source/ControlRig/Public/Rigs/RigConnectionRules.h
@@ -188,8 +188,8 @@ public:
 	, bAllowBranches(false)
 	{}
 
-	FRigChainConnectionRule(const FRigElementKey& InEndConnector, int32 InMinNumBones = 2, int32 InMaxNumBones = 0, bool InAllowBranches = false)
-	: EndConnector(InEndConnector)
+	FRigChainConnectionRule(const FRigElementKey& InRootConnector, int32 InMinNumBones = 2, int32 InMaxNumBones = 0, bool InAllowBranches = false)
+	: RootConnector(InRootConnector)
 	, MinNumBones(InMinNumBones)
 	, MaxNumBones(InMaxNumBones)
 	, bAllowBranches(InAllowBranches)
@@ -202,7 +202,7 @@ public:
 	virtual bool CanConnect(const FRigConnectionInfo* InConnectionInfo, FString* OutFailureReason) const override;
 
 	UPROPERTY(BlueprintReadWrite, EditAnywhere, Category=Settings)
-	FRigElementKey EndConnector;
+	FRigElementKey RootConnector;
 
 	UPROPERTY(BlueprintReadWrite, EditAnywhere, Category=Settings)
 	int32 MinNumBones;
diff --git a/Engine/Plugins/Animation/ControlRig/Source/ControlRig/Public/Rigs/RigModuleDefines.h b/Engine/Plugins/Animation/ControlRig/Source/ControlRig/Public/Rigs/RigModuleDefines.h
--- a/Engine/Plugins/Animation/ControlRig/Source/ControlRig/Public/Rigs/RigModuleDefines.h
+++ b/Engine/Plugins/Animation/ControlRig/Source/ControlRig/Public/Rigs/RigModuleDefines.h
@@ -54,11 +54,11 @@ struct CONTROLRIG_API FRigModuleSettings
 	FRigModuleSettings()
 	{}
 
-	bool IsValidModule() const
+	bool IsValidModule(bool bRequireExposedConnectors = true) const
 	{
 		return
 			Identifier.IsValid() &&
-			!ExposedConnectors.IsEmpty();
+			(!bRequireExposedConnectors || !ExposedConnectors.IsEmpty());
 	}
 
 	// The identifier used to retrieve the module in the module library
diff --git a/Engine/Plugins/Animation/ControlRig/Source/ControlRigDeveloper/Private/ControlRigBlueprint.cpp b/Engine/Plugins/Animation/ControlRig/Source/ControlRigDeveloper/Private/ControlRigBlueprint.cpp
--- a/Engine/Plugins/Animation/ControlRig/Source/ControlRigDeveloper/Private/ControlRigBlueprint.cpp
+++ b/Engine/Plugins/Animation/ControlRig/Source/ControlRigDeveloper/Private/ControlRigBlueprint.cpp
@@ -739,7 +739,19 @@ bool UControlRigBlueprint::ResolveConnector(const FRigElementKey& DraggedKey, co
 	}
 
 	PropagateHierarchyFromBPToInstances();
-	RequestAutoVMRecompilation();
+
+	if(UControlRig* ControlRig = Cast<UControlRig>(GetObjectBeingDebugged()))
+	{
+		for (UEdGraph* Graph : UbergraphPages)
+		{
+			UControlRigGraph* RigGraph = Cast<UControlRigGraph>(Graph);
+			if (RigGraph == nullptr)
+			{
+				continue;
+			}
+			RigGraph->CacheNameLists(ControlRig->GetHierarchy(), &DrawContainer, ShapeLibraries);
+		}
+	}
 
 	return true;
 }
diff --git a/Engine/Plugins/Animation/ControlRig/Source/ControlRigDeveloper/Private/Graph/ControlRigGraph.cpp b/Engine/Plugins/Animation/ControlRig/Source/ControlRigDeveloper/Private/Graph/ControlRigGraph.cpp
--- a/Engine/Plugins/Animation/ControlRig/Source/ControlRigDeveloper/Private/Graph/ControlRigGraph.cpp
+++ b/Engine/Plugins/Animation/ControlRig/Source/ControlRigDeveloper/Private/Graph/ControlRigGraph.cpp
@@ -40,9 +40,9 @@ void UControlRigGraph::InitializeFromBlueprint(URigVMBlueprint* InBlueprint)
 
 #if WITH_EDITOR
 
-TArray<TSharedPtr<FString>> UControlRigGraph::EmptyElementNameList;
+TArray<TSharedPtr<FRigVMStringWithTag>> UControlRigGraph::EmptyElementNameList;
 
-const TArray<TSharedPtr<FString>>* UControlRigGraph::GetNameListForWidget(const FString& InWidgetName) const
+const TArray<TSharedPtr<FRigVMStringWithTag>>* UControlRigGraph::GetNameListForWidget(const FString& InWidgetName) const
 {
 	if (InWidgetName == TEXT("BoneName"))
 	{
@@ -73,10 +73,10 @@ void UControlRigGraph::CacheNameLists(URigHierarchy* InHierarchy, const FRigVMDr
 	check(InHierarchy);
 	check(DrawContainer);
 
+	UControlRig* ControlRig = InHierarchy->GetTypedOuter<UControlRig>();
+
 	if(LastHierarchyTopologyVersion != InHierarchy->GetTopologyVersion())
 	{
-		UControlRig* ControlRig = InHierarchy->GetTypedOuter<UControlRig>();
-
 		ElementNameLists.FindOrAdd(ERigElementType::All);
 		ElementNameLists.FindOrAdd(ERigElementType::Bone);
 		ElementNameLists.FindOrAdd(ERigElementType::Null);
@@ -86,14 +86,13 @@ void UControlRigGraph::CacheNameLists(URigHierarchy* InHierarchy, const FRigVMDr
 		ElementNameLists.FindOrAdd(ERigElementType::Reference);
 		ElementNameLists.FindOrAdd(ERigElementType::Connector);
 
-		TArray<TSharedPtr<FString>>& AllNameList = ElementNameLists.FindChecked(ERigElementType::All);
-		TArray<TSharedPtr<FString>>& BoneNameList = ElementNameLists.FindChecked(ERigElementType::Bone);
-		TArray<TSharedPtr<FString>>& NullNameList = ElementNameLists.FindChecked(ERigElementType::Null);
-		TArray<TSharedPtr<FString>>& ControlNameList = ElementNameLists.FindChecked(ERigElementType::Control);
-		TArray<TSharedPtr<FString>>& CurveNameList = ElementNameLists.FindChecked(ERigElementType::Curve);
-		TArray<TSharedPtr<FString>>& RigidBodyNameList = ElementNameLists.FindChecked(ERigElementType::RigidBody);
-		TArray<TSharedPtr<FString>>& ReferenceNameList = ElementNameLists.FindChecked(ERigElementType::Reference);
-		TArray<TSharedPtr<FString>>& ConnectorNameList = ElementNameLists.FindChecked(ERigElementType::Connector);
+		TArray<TSharedPtr<FRigVMStringWithTag>>& AllNameList = ElementNameLists.FindChecked(ERigElementType::All);
+		TArray<TSharedPtr<FRigVMStringWithTag>>& BoneNameList = ElementNameLists.FindChecked(ERigElementType::Bone);
+		TArray<TSharedPtr<FRigVMStringWithTag>>& NullNameList = ElementNameLists.FindChecked(ERigElementType::Null);
+		TArray<TSharedPtr<FRigVMStringWithTag>>& ControlNameList = ElementNameLists.FindChecked(ERigElementType::Control);
+		TArray<TSharedPtr<FRigVMStringWithTag>>& CurveNameList = ElementNameLists.FindChecked(ERigElementType::Curve);
+		TArray<TSharedPtr<FRigVMStringWithTag>>& RigidBodyNameList = ElementNameLists.FindChecked(ERigElementType::RigidBody);
+		TArray<TSharedPtr<FRigVMStringWithTag>>& ReferenceNameList = ElementNameLists.FindChecked(ERigElementType::Reference);
 		
 		CacheNameListForHierarchy<FRigBaseElement>(ControlRig, InHierarchy, AllNameList, false);
 		CacheNameListForHierarchy<FRigBoneElement>(ControlRig, InHierarchy, BoneNameList, false);
@@ -103,17 +102,20 @@ void UControlRigGraph::CacheNameLists(URigHierarchy* InHierarchy, const FRigVMDr
 		CacheNameListForHierarchy<FRigCurveElement>(ControlRig, InHierarchy, CurveNameList, false);
 		CacheNameListForHierarchy<FRigRigidBodyElement>(ControlRig, InHierarchy, RigidBodyNameList, false);
 		CacheNameListForHierarchy<FRigReferenceElement>(ControlRig, InHierarchy, ReferenceNameList, false);
-		CacheNameListForHierarchy<FRigConnectorElement>(ControlRig, InHierarchy, ConnectorNameList, false);
 
 		LastHierarchyTopologyVersion = InHierarchy->GetTopologyVersion();
 	}
 	CacheNameList<FRigVMDrawContainer>(*DrawContainer, DrawingNameList);
 
+	// always update the connector name list since the connector may have been re-resolved
+	TArray<TSharedPtr<FRigVMStringWithTag>>& ConnectorNameList = ElementNameLists.FindChecked(ERigElementType::Connector);
+	CacheNameListForHierarchy<FRigConnectorElement>(ControlRig, InHierarchy, ConnectorNameList, false);
+
 	ShapeNameList.Reset();
-	ShapeNameList.Add(MakeShared<FString>(FName(NAME_None).ToString()));
+	ShapeNameList.Add(MakeShared<FRigVMStringWithTag>(FName(NAME_None).ToString()));
 
 	TMap<FString, FString> LibraryNameMap;
-	if(UControlRig* ControlRig = InHierarchy->GetTypedOuter<UControlRig>())
+	if(ControlRig)
 	{
 		LibraryNameMap = ControlRig->ShapeLibraryNameMap;
 	}
@@ -138,15 +140,15 @@ void UControlRigGraph::CacheNameLists(URigHierarchy* InHierarchy, const FRigVMDr
 		}
 		
 		const FString NameSpace = bUseNameSpace ? LibraryName + TEXT(".") : FString();
-		ShapeNameList.Add(MakeShared<FString>(UControlRigShapeLibrary::GetShapeName(ShapeLibrary.Get(), bUseNameSpace, LibraryNameMap, ShapeLibrary->DefaultShape)));
+		ShapeNameList.Add(MakeShared<FRigVMStringWithTag>(UControlRigShapeLibrary::GetShapeName(ShapeLibrary.Get(), bUseNameSpace, LibraryNameMap, ShapeLibrary->DefaultShape)));
 		for (const FControlRigShapeDefinition& Shape : ShapeLibrary->Shapes)
 		{
-			ShapeNameList.Add(MakeShared<FString>(UControlRigShapeLibrary::GetShapeName(ShapeLibrary.Get(), bUseNameSpace, LibraryNameMap, Shape)));
+			ShapeNameList.Add(MakeShared<FRigVMStringWithTag>(UControlRigShapeLibrary::GetShapeName(ShapeLibrary.Get(), bUseNameSpace, LibraryNameMap, Shape)));
 		}
 	}
 }
 
-const TArray<TSharedPtr<FString>>* UControlRigGraph::GetElementNameList(ERigElementType InElementType) const
+const TArray<TSharedPtr<FRigVMStringWithTag>>* UControlRigGraph::GetElementNameList(ERigElementType InElementType) const
 {
 	if (UControlRigGraph* OuterGraph = Cast<UControlRigGraph>(GetOuter()))
 	{
@@ -178,7 +180,7 @@ const TArray<TSharedPtr<FString>>* UControlRigGraph::GetElementNameList(ERigElem
 	return &ElementNameLists.FindChecked(InElementType);
 }
 
-const TArray<TSharedPtr<FString>>* UControlRigGraph::GetElementNameList(URigVMPin* InPin) const
+const TArray<TSharedPtr<FRigVMStringWithTag>>* UControlRigGraph::GetElementNameList(URigVMPin* InPin) const
 {
 	if (InPin)
 	{
@@ -202,9 +204,9 @@ const TArray<TSharedPtr<FString>>* UControlRigGraph::GetElementNameList(URigVMPi
 	return GetBoneNameList(nullptr);
 }
 
-const TArray<TSharedPtr<FString>> UControlRigGraph::GetSelectedElementsNameList() const
+const TArray<TSharedPtr<FRigVMStringWithTag>> UControlRigGraph::GetSelectedElementsNameList() const
 {
-	TArray<TSharedPtr<FString>> Result;
+	TArray<TSharedPtr<FRigVMStringWithTag>> Result;
 	if (UControlRigGraph* OuterGraph = Cast<UControlRigGraph>(GetOuter()))
 	{
 		return OuterGraph->GetSelectedElementsNameList();
@@ -221,13 +223,13 @@ const TArray<TSharedPtr<FString>> UControlRigGraph::GetSelectedElementsNameList(
 	{
 		FString ValueStr;
 		FRigElementKey::StaticStruct()->ExportText(ValueStr, &Key, nullptr, nullptr, PPF_None, nullptr);
-		Result.Add(MakeShared<FString>(ValueStr));
+		Result.Add(MakeShared<FRigVMStringWithTag>(ValueStr));
 	}
 	
 	return Result;
 }
 
-const TArray<TSharedPtr<FString>>* UControlRigGraph::GetDrawingNameList(URigVMPin* InPin) const
+const TArray<TSharedPtr<FRigVMStringWithTag>>* UControlRigGraph::GetDrawingNameList(URigVMPin* InPin) const
 {
 	if (UControlRigGraph* OuterGraph = Cast<UControlRigGraph>(GetOuter()))
 	{
@@ -236,7 +238,7 @@ const TArray<TSharedPtr<FString>>* UControlRigGraph::GetDrawingNameList(URigVMPi
 	return &DrawingNameList;
 }
 
-const TArray<TSharedPtr<FString>>* UControlRigGraph::GetShapeNameList(URigVMPin* InPin) const
+const TArray<TSharedPtr<FRigVMStringWithTag>>* UControlRigGraph::GetShapeNameList(URigVMPin* InPin) const
 {
 	if (UControlRigGraph* OuterGraph = Cast<UControlRigGraph>(GetOuter()))
 	{
diff --git a/Engine/Plugins/Animation/ControlRig/Source/ControlRigDeveloper/Public/Graph/ControlRigGraph.h b/Engine/Plugins/Animation/ControlRig/Source/ControlRigDeveloper/Public/Graph/ControlRigGraph.h
--- a/Engine/Plugins/Animation/ControlRig/Source/ControlRigDeveloper/Public/Graph/ControlRigGraph.h
+++ b/Engine/Plugins/Animation/ControlRig/Source/ControlRigDeveloper/Public/Graph/ControlRigGraph.h
@@ -33,39 +33,39 @@ public:
 
 #if WITH_EDITOR
 
-	const TArray<TSharedPtr<FString>>* GetBoneNameList(URigVMPin* InPin = nullptr) const
+	const TArray<TSharedPtr<FRigVMStringWithTag>>* GetBoneNameList(URigVMPin* InPin = nullptr) const
 	{
 		return GetElementNameList(ERigElementType::Bone);
 	}
-	const TArray<TSharedPtr<FString>>* GetControlNameList(URigVMPin* InPin = nullptr) const
+	const TArray<TSharedPtr<FRigVMStringWithTag>>* GetControlNameList(URigVMPin* InPin = nullptr) const
 	{
 		return GetElementNameList(ERigElementType::Control);
 	}
-	const TArray<TSharedPtr<FString>>* GetControlNameListWithoutAnimationChannels(URigVMPin* InPin = nullptr) const
+	const TArray<TSharedPtr<FRigVMStringWithTag>>* GetControlNameListWithoutAnimationChannels(URigVMPin* InPin = nullptr) const
 	{
 		return &ControlNameListWithoutAnimationChannels;
 	}
-	const TArray<TSharedPtr<FString>>* GetNullNameList(URigVMPin* InPin = nullptr) const
+	const TArray<TSharedPtr<FRigVMStringWithTag>>* GetNullNameList(URigVMPin* InPin = nullptr) const
 	{
 		return GetElementNameList(ERigElementType::Null);
 	}
-	const TArray<TSharedPtr<FString>>* GetCurveNameList(URigVMPin* InPin = nullptr) const
+	const TArray<TSharedPtr<FRigVMStringWithTag>>* GetCurveNameList(URigVMPin* InPin = nullptr) const
 	{
 		return GetElementNameList(ERigElementType::Curve);
 	}
-	const TArray<TSharedPtr<FString>>* GetConnectorNameList(URigVMPin* InPin = nullptr) const
+	const TArray<TSharedPtr<FRigVMStringWithTag>>* GetConnectorNameList(URigVMPin* InPin = nullptr) const
 	{
 		return GetElementNameList(ERigElementType::Connector);
 	}
 
-	virtual const TArray<TSharedPtr<FString>>* GetNameListForWidget(const FString& InWidgetName) const override;
+	virtual const TArray<TSharedPtr<FRigVMStringWithTag>>* GetNameListForWidget(const FString& InWidgetName) const override;
 
 	void CacheNameLists(URigHierarchy* InHierarchy, const FRigVMDrawContainer* DrawContainer, TArray<TSoftObjectPtr<UControlRigShapeLibrary>> ShapeLibraries);
-	const TArray<TSharedPtr<FString>>* GetElementNameList(ERigElementType InElementType = ERigElementType::Bone) const;
-	const TArray<TSharedPtr<FString>>* GetElementNameList(URigVMPin* InPin = nullptr) const;
-	const TArray<TSharedPtr<FString>> GetSelectedElementsNameList() const;
-	const TArray<TSharedPtr<FString>>* GetDrawingNameList(URigVMPin* InPin = nullptr) const;
-	const TArray<TSharedPtr<FString>>* GetShapeNameList(URigVMPin* InPin = nullptr) const;
+	const TArray<TSharedPtr<FRigVMStringWithTag>>* GetElementNameList(ERigElementType InElementType = ERigElementType::Bone) const;
+	const TArray<TSharedPtr<FRigVMStringWithTag>>* GetElementNameList(URigVMPin* InPin = nullptr) const;
+	const TArray<TSharedPtr<FRigVMStringWithTag>> GetSelectedElementsNameList() const;
+	const TArray<TSharedPtr<FRigVMStringWithTag>>* GetDrawingNameList(URigVMPin* InPin = nullptr) const;
+	const TArray<TSharedPtr<FRigVMStringWithTag>>* GetShapeNameList(URigVMPin* InPin = nullptr) const;
 
 	FReply HandleGetSelectedClicked(URigVMEdGraph* InEdGraph, URigVMPin* InPin, FString InDefaultValue);
 	FReply HandleBrowseClicked(URigVMEdGraph* InEdGraph, URigVMPin* InPin, FString InDefaultValue);
@@ -79,39 +79,53 @@ private:
 	}
 
 	template<class T>
-	void CacheNameListForHierarchy(UControlRig* InControlRig, URigHierarchy* InHierarchy, TArray<TSharedPtr<FString>>& OutNameList, bool bFilter = true)
+	void CacheNameListForHierarchy(UControlRig* InControlRig, URigHierarchy* InHierarchy, TArray<TSharedPtr<FRigVMStringWithTag>>& OutNameList, bool bFilter = true)
 	{
-        TArray<FString> Names;
+        TArray<FRigVMStringWithTag> Names;
 		for (auto Element : *InHierarchy)
 		{
 			if(Element->IsA<T>())
 			{
 				if(!bFilter || IncludeElementInNameList<T>(Cast<T>(Element)))
 				{
+					static const FLinearColor Color(FLinearColor(0.0, 112.f/255.f, 224.f/255.f));
+					FRigVMStringTag Tag;
+					
 					if(InControlRig)
 					{
 						if(const FRigElementKey* SourceKey = InControlRig->GetElementKeyRedirector().FindReverse(Element->GetKey()))
 						{
-							Names.Add(SourceKey->Name.ToString());
+							Tag = FRigVMStringTag(Element->GetKey().Name, Color);
+							Names.Emplace(SourceKey->Name.ToString(), Tag);
 							continue;
 						}
+
+						// look up the resolved name
+						if(Element->GetType() == ERigElementType::Connector)
+						{
+							if(const FCachedRigElement* Cache = InControlRig->GetElementKeyRedirector().Find(Element->GetKey()))
+							{
+								Tag = FRigVMStringTag(Cache->GetKey().Name, Color);
+							}
+						}
 					}
-					Names.Add(Element->GetName());
+
+					Names.Emplace(Element->GetName(), Tag);
 				}
 			}
 		}
 		Names.Sort();
 
 		OutNameList.Reset();
-		OutNameList.Add(MakeShared<FString>(FName(NAME_None).ToString()));
-		for (const FString& Name : Names)
+		OutNameList.Add(MakeShared<FRigVMStringWithTag>(FName(NAME_None).ToString()));
+		for (const FRigVMStringWithTag& Name : Names)
 		{
-			OutNameList.Add(MakeShared<FString>(Name));
+			OutNameList.Add(MakeShared<FRigVMStringWithTag>(Name));
 		}
 	}
 
 	template<class T>
-	void CacheNameList(const T& ElementList, TArray<TSharedPtr<FString>>& OutNameList)
+	void CacheNameList(const T& ElementList, TArray<TSharedPtr<FRigVMStringWithTag>>& OutNameList)
 	{
 		TArray<FString> Names;
 		for (auto Element : ElementList)
@@ -121,20 +135,20 @@ private:
 		Names.Sort();
 
 		OutNameList.Reset();
-		OutNameList.Add(MakeShared<FString>(FName(NAME_None).ToString()));
+		OutNameList.Add(MakeShared<FRigVMStringWithTag>(FName(NAME_None).ToString()));
 		for (const FString& Name : Names)
 		{
-			OutNameList.Add(MakeShared<FString>(Name));
+			OutNameList.Add(MakeShared<FRigVMStringWithTag>(Name));
 		}
 	}
 
-	TMap<ERigElementType, TArray<TSharedPtr<FString>>> ElementNameLists;
-	TArray<TSharedPtr<FString>>	ControlNameListWithoutAnimationChannels;
-	TArray<TSharedPtr<FString>> DrawingNameList;
-	TArray<TSharedPtr<FString>> ShapeNameList;
+	TMap<ERigElementType, TArray<TSharedPtr<FRigVMStringWithTag>>> ElementNameLists;
+	TArray<TSharedPtr<FRigVMStringWithTag>>	ControlNameListWithoutAnimationChannels;
+	TArray<TSharedPtr<FRigVMStringWithTag>> DrawingNameList;
+	TArray<TSharedPtr<FRigVMStringWithTag>> ShapeNameList;
 	int32 LastHierarchyTopologyVersion;
 
-	static TArray<TSharedPtr<FString>> EmptyElementNameList;
+	static TArray<TSharedPtr<FRigVMStringWithTag>> EmptyElementNameList;
 
 #endif
 
diff --git a/Engine/Plugins/Animation/ControlRig/Source/ControlRigEditor/Private/ControlRigElementDetails.cpp b/Engine/Plugins/Animation/ControlRig/Source/ControlRigEditor/Private/ControlRigElementDetails.cpp
--- a/Engine/Plugins/Animation/ControlRig/Source/ControlRigEditor/Private/ControlRigElementDetails.cpp
+++ b/Engine/Plugins/Animation/ControlRig/Source/ControlRigEditor/Private/ControlRigElementDetails.cpp
@@ -338,6 +338,45 @@ void FRigElementKeyDetails::SetElementName(FString InName)
 	if (NameHandle.IsValid())
 	{
 		NameHandle->SetValue(InName);
+
+		// if this is nested below a connection rule
+		const TSharedPtr<IPropertyHandle> KeyHandle = NameHandle->GetParentHandle();
+		if(KeyHandle.IsValid())
+		{
+			const TSharedPtr<IPropertyHandle> ParentHandle = KeyHandle->GetParentHandle();
+			if(ParentHandle.IsValid())
+			{
+				if (const TSharedPtr<IPropertyHandleStruct> StructPropertyHandle = ParentHandle->AsStruct())
+				{
+					if(const UScriptStruct* RuleStruct = Cast<UScriptStruct>(StructPropertyHandle->GetStructData()->GetStruct()))
+					{
+						if (RuleStruct->IsChildOf(FRigConnectionRule::StaticStruct()))
+						{
+							const void* RuleMemory = StructPropertyHandle->GetStructData()->GetStructMemory();
+							FString RuleContent;
+							RuleStruct->ExportText(RuleContent, RuleMemory, RuleMemory, nullptr, PPF_None, nullptr);
+
+							const TSharedPtr<IPropertyHandle> RuleStashHandle = ParentHandle->GetParentHandle();
+							
+							FRigConnectionRuleStash Stash;
+							Stash.ScriptStructPath = RuleStruct->GetPathName();
+							Stash.ExportedText = RuleContent;
+							
+							FString StashContent;
+							FRigConnectionRuleStash::StaticStruct()->ExportText(StashContent, &Stash, &Stash, nullptr, PPF_None, nullptr);
+
+							TArray<UObject*> Objects;
+							RuleStashHandle->GetOuterObjects(Objects);
+							FString FirstObjectValue;
+							for (int32 Index = 0; Index < Objects.Num(); Index++)
+							{
+								(void)RuleStashHandle->SetPerObjectValue(Index, StashContent, EPropertyValueSetFlags::DefaultFlags);
+							}
+						}
+					}
+				}
+			}
+		}
 	}
 }
 
@@ -356,7 +395,16 @@ void FRigElementKeyDetails::UpdateElementNameList()
 		{
 			if (UControlRigGraph* RigGraph = Cast<UControlRigGraph>(Graph))
 			{
-				ElementNameList = *RigGraph->GetElementNameList(GetElementType());
+				const TArray<TSharedPtr<FRigVMStringWithTag>>* NameList =
+					RigGraph->GetElementNameList(GetElementType());
+
+				ElementNameList.Reset();
+				ElementNameList.Reserve(NameList->Num());
+				for(const TSharedPtr<FRigVMStringWithTag>& Name : *NameList)
+				{
+					ElementNameList.Add(MakeShared<FString>(Name->GetString()));
+				}
+				
 				if(SearchableComboBox.IsValid())
 				{
 					SearchableComboBox->RefreshOptions();
@@ -3763,10 +3811,10 @@ void FRigControlElementDetails::CustomizeShape(IDetailLayoutBuilder& DetailBuild
 			if (ShapeLibrary.IsValid())
 			{
 				const FString NameSpace = bUseNameSpace ? ShapeLibrary->GetName() + TEXT(".") : FString();
-				ShapeNameList.Add(MakeShared<FString>(NameSpace + ShapeLibrary->DefaultShape.ShapeName.ToString()));
+				ShapeNameList.Add(MakeShared<FRigVMStringWithTag>(NameSpace + ShapeLibrary->DefaultShape.ShapeName.ToString()));
 				for (const FControlRigShapeDefinition& Shape : ShapeLibrary->Shapes)
 				{
-					ShapeNameList.Add(MakeShared<FString>(NameSpace + Shape.ShapeName.ToString()));
+					ShapeNameList.Add(MakeShared<FRigVMStringWithTag>(NameSpace + Shape.ShapeName.ToString()));
 				}
 			}
 		}
@@ -4239,7 +4287,7 @@ bool FRigControlElementDetails::IsShapeEnabled() const
 	});
 }
 
-const TArray<TSharedPtr<FString>>& FRigControlElementDetails::GetShapeNameList() const
+const TArray<TSharedPtr<FRigVMStringWithTag>>& FRigControlElementDetails::GetShapeNameList() const
 {
 	return ShapeNameList;
 }
diff --git a/Engine/Plugins/Animation/ControlRig/Source/ControlRigEditor/Private/ControlRigElementDetails.h b/Engine/Plugins/Animation/ControlRig/Source/ControlRigEditor/Private/ControlRigElementDetails.h
--- a/Engine/Plugins/Animation/ControlRig/Source/ControlRigEditor/Private/ControlRigElementDetails.h
+++ b/Engine/Plugins/Animation/ControlRig/Source/ControlRigEditor/Private/ControlRigElementDetails.h
@@ -357,7 +357,7 @@ public:
 
 	bool IsShapeEnabled() const;
 
-	const TArray<TSharedPtr<FString>>& GetShapeNameList() const;
+	const TArray<TSharedPtr<FRigVMStringWithTag>>& GetShapeNameList() const;
 
 	FText GetDisplayName() const;
 	void SetDisplayName(const FText& InNewText, ETextCommit::Type InCommitType);
@@ -684,7 +684,7 @@ private:
 	void HandleControlTypeChanged(TSharedPtr<ERigControlType> ControlType, ESelectInfo::Type SelectInfo, FRigElementKey ControlKey, const TSharedRef<IPropertyUtilities> PropertyUtilities);
 	void HandleControlTypeChanged(ERigControlType ControlType, TArray<FRigElementKey> ControlKeys, const TSharedRef<IPropertyUtilities> PropertyUtilities);
 
-	TArray<TSharedPtr<FString>> ShapeNameList;
+	TArray<TSharedPtr<FRigVMStringWithTag>> ShapeNameList;
 	TSharedPtr<FRigInfluenceEntryModifier> InfluenceModifier;
 	TSharedPtr<FStructOnScope> InfluenceModifierStruct;
 
diff --git a/Engine/Plugins/Animation/ControlRig/Source/ControlRigEditor/Private/EditMode/ControlRigEditMode.cpp b/Engine/Plugins/Animation/ControlRig/Source/ControlRigEditor/Private/EditMode/ControlRigEditMode.cpp
--- a/Engine/Plugins/Animation/ControlRig/Source/ControlRigEditor/Private/EditMode/ControlRigEditMode.cpp
+++ b/Engine/Plugins/Animation/ControlRig/Source/ControlRigEditor/Private/EditMode/ControlRigEditMode.cpp
@@ -3313,37 +3313,66 @@ void FControlRigEditMode::RecreateControlShapeActors(const TArray<FRigElementKey
 			if (auto* ShapeActors = ControlRigShapeActors.Find(ControlRig))
 			{
 				TArray<FRigControlElement*> Controls = ControlRig->AvailableControls();
-				for (int32 Index = Controls.Num() - 1; Index >= 0; --Index)
+				TArray<FRigControlElement*> ControlPerShapeActor;
+				ControlPerShapeActor.SetNumZeroed(ShapeActors->Num());
+				
+				if(Controls.Num() == ShapeActors->Num())
 				{
-					FRigControlElement* ControlElement = Controls[Index];
-					if (!ControlElement->Settings.SupportsShape() || !IsSupportedControlType(ControlElement->Settings.ControlType))
+					for (int32 ControlIndex = Controls.Num() - 1; ControlIndex >= 0; --ControlIndex)
 					{
-						Controls.RemoveAtSwap(Index);
+						FRigControlElement* ControlElement = Controls[ControlIndex];
+						if (!ControlElement->Settings.SupportsShape() || !IsSupportedControlType(ControlElement->Settings.ControlType))
+						{
+							Controls.RemoveAtSwap(ControlIndex);
+						}
 					}
-				}
-				//unfortunately n*n-ish but this should be very rare and much faster than recreating them
-				for (AControlRigShapeActor* Actor : *ShapeActors)
-				{
-					if (Actor)
+					//unfortunately n*n-ish but this should be very rare and much faster than recreating them
+					for (int32 ShapeActorIndex = 0; ShapeActorIndex < ShapeActors->Num(); ShapeActorIndex++)
 					{
-						for (int32 Index = 0; Index < Controls.Num(); ++Index)
+						const AControlRigShapeActor* Actor = ShapeActors->operator[](ShapeActorIndex).Get();
+						if (Actor)
 						{
-							FRigControlElement* Element = Controls[Index];
-							if (Element && Element->GetFName() == Actor->ControlName)
+							for (int32 ControlIndex = 0; ControlIndex < Controls.Num(); ++ControlIndex)
 							{
-								Controls.RemoveAtSwap(Index);
-								break;
+								FRigControlElement* Element = Controls[ControlIndex];
+								if (Element && Element->GetFName() == Actor->ControlName)
+								{
+									Controls.RemoveAtSwap(ControlIndex);
+									ControlPerShapeActor[ShapeActorIndex] = Element;
+									break;
+								}
 							}
 						}
-					}
-					else //no actor just recreate
-					{
-						break;
+						else //no actor just recreate
+						{
+							break;
+						}
 					}
 				}
 				if (Controls.Num() == 0)
 				{
 					bRecreateThem = false;
+
+					// we have matching controls - we should at least sync their settings.
+					// PostPoseUpdate / TickControlShape is going to take care of color, visibility etc.
+					// MeshTransform has to be handled here.
+					for (int32 ShapeActorIndex = 0; ShapeActorIndex < ShapeActors->Num(); ShapeActorIndex++)
+					{
+						const AControlRigShapeActor* ShapeActor = ShapeActors->operator[](ShapeActorIndex).Get();
+						FRigControlElement* ControlElement = ControlPerShapeActor[ShapeActorIndex];
+						if (ShapeActor && ControlElement)
+						{
+							const FTransform ShapeTransform = ControlRig->GetHierarchy()->GetControlShapeTransform(ControlElement, ERigTransformType::CurrentLocal);
+							FTransform MeshTransform = FTransform::Identity;
+							if (const FControlRigShapeDefinition* Gizmo = UControlRigShapeLibrary::GetShapeByName(ControlElement->Settings.ShapeName, ControlRig->GetShapeLibraries(), ControlRig->ShapeLibraryNameMap))
+							{
+								MeshTransform = Gizmo->Transform;
+							}
+							ShapeActor->StaticMeshComponent->SetRelativeTransform(MeshTransform * ShapeTransform);
+						}
+					}
+					
+					PostPoseUpdate();
 				}
 			}
 			if (bRecreateThem)
diff --git a/Engine/Plugins/Animation/ControlRig/Source/ControlRigEditor/Private/Editor/ControlRigEditor.cpp b/Engine/Plugins/Animation/ControlRig/Source/ControlRigEditor/Private/Editor/ControlRigEditor.cpp
--- a/Engine/Plugins/Animation/ControlRig/Source/ControlRigEditor/Private/Editor/ControlRigEditor.cpp
+++ b/Engine/Plugins/Animation/ControlRig/Source/ControlRigEditor/Private/Editor/ControlRigEditor.cpp
@@ -2348,6 +2348,54 @@ void FControlRigEditor::GenerateEventQueueMenuContent(FMenuBuilder& MenuBuilder)
     }
 }
 
+void FControlRigEditor::FilterDraggedKeys(TArray<FRigElementKey>& Keys, bool bRemoveNameSpace)
+{
+	// if the keys being dragged contain something mapped to a connector - use that instead
+	if(UControlRigBlueprint* ControlRigBlueprint = GetControlRigBlueprint())
+	{
+		TArray<FRigElementKey> FilteredKeys;
+		FilteredKeys.Reserve(Keys.Num());
+		for (FRigElementKey Key : Keys)
+		{
+			for(const TPair<FRigElementKey,FRigElementKey>& Pair : ControlRigBlueprint->ConnectionMap)
+			{
+				if(Pair.Value == Key)
+				{
+					Key = Pair.Key;
+					break;
+				}
+			}
+
+			if(bRemoveNameSpace)
+			{
+				const FString Name = Key.Name.ToString();
+				int32 LastCharIndex = INDEX_NONE;
+				if(Name.FindLastChar(TEXT(':'), LastCharIndex))
+				{
+					Key.Name = *Name.Mid(LastCharIndex+1);
+				}
+			}
+			else
+			{
+				if(const UControlRig* DebuggedControlRig = Cast<UControlRig>(ControlRigBlueprint->GetObjectBeingDebugged()))
+				{
+					if(!DebuggedControlRig->GetHierarchy()->Contains(Key))
+					{
+						const FString NameSpace = DebuggedControlRig->GetRigModuleNameSpace();
+						if(!NameSpace.IsEmpty())
+						{
+							static constexpr TCHAR Format[] = TEXT("%s%s");
+							Key.Name = *FString::Printf(Format, *NameSpace, *Key.Name.ToString());
+						}
+					}
+				}
+			}
+			FilteredKeys.Add(Key);
+		}
+		Keys = FilteredKeys;
+	}
+}
+
 FTransform FControlRigEditor::GetRigElementTransform(const FRigElementKey& InElement, bool bLocal, bool bOnDebugInstance) const
 {
 	DECLARE_SCOPE_HIERARCHICAL_COUNTER_FUNC()
@@ -3459,7 +3507,9 @@ void FControlRigEditor::CreateRigHierarchyToGraphDragAndDropMenu() const
 					const FControlRigRigHierarchyToGraphDragAndDropContext& DragDropContext = MainContext->GetRigHierarchyToGraphDragAndDropContext();
 
 					URigHierarchy* Hierarchy = ControlRigEditor->GetHierarchyBeingDebugged();
-					const TArray<FRigElementKey>& DraggedKeys = DragDropContext.DraggedElementKeys;
+					TArray<FRigElementKey> DraggedKeys = DragDropContext.DraggedElementKeys;
+					ControlRigEditor->FilterDraggedKeys(DraggedKeys, true);
+					
 					UEdGraph* Graph = DragDropContext.Graph.Get();
 					const FVector2D& NodePosition = DragDropContext.NodePosition;
 					
@@ -3806,8 +3856,14 @@ void FControlRigEditor::HandleMakeElementGetterSetter(ERigElementGetterSetterTyp
 	};
 	TArray<FNewNodeData> NewNodes;
 
-	for (const FRigElementKey& Key : Keys)
+	TArray<FRigElementKey> KeysIncludingNameSpace = Keys;
+	FilterDraggedKeys(KeysIncludingNameSpace, false);
+
+	for (int32 Index = 0; Index < Keys.Num(); Index++)
 	{
+		const FRigElementKey& Key = Keys[Index];
+		const FRigElementKey& KeyIncludingNameSpace = KeysIncludingNameSpace[Index];
+		
 		UScriptStruct* StructTemplate = nullptr;
 
 		FNewNodeData NewNode;
@@ -3823,7 +3879,7 @@ void FControlRigEditor::HandleMakeElementGetterSetter(ERigElementGetterSetterTyp
 		TArray<FName> ChannelPins;
 		TMap<FName, int32> PinsToResolve; 
 
-		if(FRigControlElement* ControlElement = Hierarchy->Find<FRigControlElement>(Key))
+		if(FRigControlElement* ControlElement = Hierarchy->Find<FRigControlElement>(KeyIncludingNameSpace))
 		{
 			if(ControlElement->IsAnimationChannel())
 			{
@@ -3969,8 +4025,11 @@ void FControlRigEditor::HandleMakeElementGetterSetter(ERigElementGetterSetterTyp
 				{
 					if (Key.Type == ERigElementType::Control)
 					{
-						FRigControlElement* ControlElement = Hierarchy->Find<FRigControlElement>(Key);
-						check(ControlElement);
+						FRigControlElement* ControlElement = Hierarchy->Find<FRigControlElement>(KeyIncludingNameSpace);
+						if(ControlElement == nullptr)
+						{
+							return;
+						}
 						
 						switch (ControlElement->Settings.ControlType)
 						{
@@ -4058,8 +4117,11 @@ void FControlRigEditor::HandleMakeElementGetterSetter(ERigElementGetterSetterTyp
 				{
 					if (Key.Type == ERigElementType::Control)
 					{
-						FRigControlElement* ControlElement = Hierarchy->Find<FRigControlElement>(Key);
-						check(ControlElement);
+						FRigControlElement* ControlElement = Hierarchy->Find<FRigControlElement>(KeyIncludingNameSpace);
+						if(ControlElement == nullptr)
+						{
+							return;
+						}
 
 						switch (ControlElement->Settings.ControlType)
 						{
diff --git a/Engine/Plugins/Animation/ControlRig/Source/ControlRigEditor/Private/Editor/ControlRigEditor.h b/Engine/Plugins/Animation/ControlRig/Source/ControlRigEditor/Private/Editor/ControlRigEditor.h
--- a/Engine/Plugins/Animation/ControlRig/Source/ControlRigEditor/Private/Editor/ControlRigEditor.h
+++ b/Engine/Plugins/Animation/ControlRig/Source/ControlRigEditor/Private/Editor/ControlRigEditor.h
@@ -245,9 +245,10 @@ private:
 		ERigElementGetterSetterType_Name
 	};
 
-	 void HandleMakeElementGetterSetter(ERigElementGetterSetterType Type, bool bIsGetter, TArray<FRigElementKey> Keys, UEdGraph* Graph, FVector2D NodePosition);
+	void FilterDraggedKeys(TArray<FRigElementKey>& Keys, bool bRemoveNameSpace);
+	void HandleMakeElementGetterSetter(ERigElementGetterSetterType Type, bool bIsGetter, TArray<FRigElementKey> Keys, UEdGraph* Graph, FVector2D NodePosition);
 
-	 void HandleOnControlModified(UControlRig* Subject, FRigControlElement* ControlElement, const FRigControlModifiedContext& Context);
+	void HandleOnControlModified(UControlRig* Subject, FRigControlElement* ControlElement, const FRigControlModifiedContext& Context);
 
 	virtual void HandleRefreshEditorFromBlueprint(URigVMBlueprint* InBlueprint) override;
 
diff --git a/Engine/Plugins/Animation/ControlRig/Source/ControlRigEditor/Private/Editor/SControlRigGizmoNameList.cpp b/Engine/Plugins/Animation/ControlRig/Source/ControlRigEditor/Private/Editor/SControlRigGizmoNameList.cpp
--- a/Engine/Plugins/Animation/ControlRig/Source/ControlRigEditor/Private/Editor/SControlRigGizmoNameList.cpp
+++ b/Engine/Plugins/Animation/ControlRig/Source/ControlRigEditor/Private/Editor/SControlRigGizmoNameList.cpp
@@ -57,8 +57,8 @@ void SControlRigShapeNameList::ConstructCommon()
 {
 	SBox::Construct(SBox::FArguments());
 
-	TSharedPtr<FString> InitialSelected;
-	for (TSharedPtr<FString> Item : GetNameList())
+	TSharedPtr<FRigVMStringWithTag> InitialSelected;
+	for (TSharedPtr<FRigVMStringWithTag> Item : GetNameList())
 	{
 		if (Item->Equals(GetNameListText().ToString()))
 		{
@@ -86,7 +86,7 @@ void SControlRigShapeNameList::ConstructCommon()
 	);
 }
 
-const TArray<TSharedPtr<FString>>& SControlRigShapeNameList::GetNameList() const
+const TArray<TSharedPtr<FRigVMStringWithTag>>& SControlRigShapeNameList::GetNameList() const
 {
 	if (OnGetNameListContent.IsBound())
 	{
@@ -146,18 +146,18 @@ void SControlRigShapeNameList::SetNameListText(const FText& NewTypeInValue, ETex
 	}
 }
 
-TSharedRef<SWidget> SControlRigShapeNameList::MakeNameListItemWidget(TSharedPtr<FString> InItem)
+TSharedRef<SWidget> SControlRigShapeNameList::MakeNameListItemWidget(TSharedPtr<FRigVMStringWithTag> InItem)
 {
-	return 	SNew(STextBlock).Text(FText::FromString(*InItem));// .Font(FAppStyle::GetFontStyle(TEXT("PropertyWindow.NormalFont")));
+	return 	SNew(STextBlock).Text(FText::FromString(InItem->GetString()));// .Font(FAppStyle::GetFontStyle(TEXT("PropertyWindow.NormalFont")));
 }
 
-void SControlRigShapeNameList::OnNameListChanged(TSharedPtr<FString> NewSelection, ESelectInfo::Type SelectInfo)
+void SControlRigShapeNameList::OnNameListChanged(TSharedPtr<FRigVMStringWithTag> NewSelection, ESelectInfo::Type SelectInfo)
 {
 	if (SelectInfo != ESelectInfo::Direct)
 	{
 		if (NewSelection.IsValid())
 		{
-			FString NewValue = *NewSelection.Get();
+			const FString& NewValue = NewSelection->GetString();
 			SetNameListText(FText::FromString(NewValue), ETextCommit::OnEnter);
 		}
 		else
@@ -169,8 +169,8 @@ void SControlRigShapeNameList::OnNameListChanged(TSharedPtr<FString> NewSelectio
 
 void SControlRigShapeNameList::OnNameListComboBox()
 {
-	TSharedPtr<FString> CurrentlySelected;
-	for (TSharedPtr<FString> Item : GetNameList())
+	TSharedPtr<FRigVMStringWithTag> CurrentlySelected;
+	for (TSharedPtr<FRigVMStringWithTag> Item : GetNameList())
 	{
 		if (Item->Equals(GetNameListText().ToString()))
 		{
diff --git a/Engine/Plugins/Animation/ControlRig/Source/ControlRigEditor/Private/Editor/SControlRigGizmoNameList.h b/Engine/Plugins/Animation/ControlRig/Source/ControlRigEditor/Private/Editor/SControlRigGizmoNameList.h
--- a/Engine/Plugins/Animation/ControlRig/Source/ControlRigEditor/Private/Editor/SControlRigGizmoNameList.h
+++ b/Engine/Plugins/Animation/ControlRig/Source/ControlRigEditor/Private/Editor/SControlRigGizmoNameList.h
@@ -14,7 +14,7 @@ class SControlRigShapeNameList : public SBox
 {
 public:
 
-	DECLARE_DELEGATE_RetVal( const TArray<TSharedPtr<FString>>&, FOnGetNameListContent );
+	DECLARE_DELEGATE_RetVal( const TArray<TSharedPtr<FRigVMStringWithTag>>&, FOnGetNameListContent );
 
 	SLATE_BEGIN_ARGS(SControlRigShapeNameList){}
 
@@ -31,17 +31,17 @@ protected:
 
 	void ConstructCommon();
 
-	const TArray<TSharedPtr<FString>>& GetNameList() const;
+	const TArray<TSharedPtr<FRigVMStringWithTag>>& GetNameList() const;
 	FText GetNameListText() const;
 	virtual void SetNameListText(const FText& NewTypeInValue, ETextCommit::Type CommitInfo);
 
-	TSharedRef<SWidget> MakeNameListItemWidget(TSharedPtr<FString> InItem);
-	void OnNameListChanged(TSharedPtr<FString> NewSelection, ESelectInfo::Type SelectInfo);
+	TSharedRef<SWidget> MakeNameListItemWidget(TSharedPtr<FRigVMStringWithTag> InItem);
+	void OnNameListChanged(TSharedPtr<FRigVMStringWithTag> NewSelection, ESelectInfo::Type SelectInfo);
 	void OnNameListComboBox();
 
-	static const TArray< TSharedPtr<FString> >& GetEmptyList()
+	static const TArray< TSharedPtr<FRigVMStringWithTag> >& GetEmptyList()
 	{
-		static const TArray< TSharedPtr<FString> > EmptyList;
+		static const TArray< TSharedPtr<FRigVMStringWithTag> > EmptyList;
 		return EmptyList;
 	}
 
diff --git a/Engine/Plugins/Animation/ControlRig/Source/ControlRigEditor/Private/Editor/SRigHierarchy.cpp b/Engine/Plugins/Animation/ControlRig/Source/ControlRigEditor/Private/Editor/SRigHierarchy.cpp
--- a/Engine/Plugins/Animation/ControlRig/Source/ControlRigEditor/Private/Editor/SRigHierarchy.cpp
+++ b/Engine/Plugins/Animation/ControlRig/Source/ControlRigEditor/Private/Editor/SRigHierarchy.cpp
@@ -2429,7 +2429,14 @@ TOptional<EItemDropZone> SRigHierarchy::OnCanAcceptDrop(const FDragDropEvent& Dr
 					return InvalidDropZone;
 				}
 
-				if(DropZone == EItemDropZone::OntoItem && !RigDragDropOp->IsDraggingSingleConnector())
+				if(RigDragDropOp->IsDraggingSingleConnector())
+				{
+					if(DropZone != EItemDropZone::OntoItem)
+					{
+						return InvalidDropZone;
+					}
+				}
+				else if(DropZone == EItemDropZone::OntoItem)
 				{
 					if(Hierarchy->IsParentedTo(TargetKey, DraggedKey))
 					{
diff --git a/Engine/Plugins/Animation/ControlRig/Source/ControlRigEditor/Private/Editor/SRigHierarchyTreeView.cpp b/Engine/Plugins/Animation/ControlRig/Source/ControlRigEditor/Private/Editor/SRigHierarchyTreeView.cpp
--- a/Engine/Plugins/Animation/ControlRig/Source/ControlRigEditor/Private/Editor/SRigHierarchyTreeView.cpp
+++ b/Engine/Plugins/Animation/ControlRig/Source/ControlRigEditor/Private/Editor/SRigHierarchyTreeView.cpp
@@ -595,6 +595,11 @@ bool SRigHierarchyTreeView::ReparentElement(FRigElementKey InKey, FRigElementKey
 		return false;
 	}
 
+	if(InKey.Type == ERigElementType::Connector)
+	{
+		return false;
+	}
+
 	const FRigTreeDisplaySettings& Settings = Delegates.GetDisplaySettings();
 
 	TSharedPtr<FRigTreeElement>* FoundItem = ElementMap.Find(InKey);
diff --git a/Engine/Plugins/Animation/ControlRig/Source/ControlRigEditor/Private/Graph/ControlRigGraphPanelPinFactory.cpp b/Engine/Plugins/Animation/ControlRig/Source/ControlRigEditor/Private/Graph/ControlRigGraphPanelPinFactory.cpp
--- a/Engine/Plugins/Animation/ControlRig/Source/ControlRigEditor/Private/Graph/ControlRigGraphPanelPinFactory.cpp
+++ b/Engine/Plugins/Animation/ControlRig/Source/ControlRigEditor/Private/Graph/ControlRigGraphPanelPinFactory.cpp
@@ -209,7 +209,7 @@ TSharedPtr<SGraphPin> FControlRigGraphPanelPinFactory::CreatePin_Internal(UEdGra
 					struct FCachedAnimationChannelNames
 					{
 						int32 TopologyVersion;
-						TSharedPtr<TArray<TSharedPtr<FString>>> Names;
+						TSharedPtr<TArray<TSharedPtr<FRigVMStringWithTag>>> Names;
 						
 						FCachedAnimationChannelNames()
 						: TopologyVersion(INDEX_NONE)
@@ -280,10 +280,10 @@ TSharedPtr<SGraphPin> FControlRigGraphPanelPinFactory::CreatePin_Internal(UEdGra
 								{
 									if(!ChannelNames.Names.IsValid())
 									{
-										ChannelNames.Names = MakeShareable(new TArray<TSharedPtr<FString>>());
+										ChannelNames.Names = MakeShareable(new TArray<TSharedPtr<FRigVMStringWithTag>>());
 									}
 									ChannelNames.Names->Reset();
-									ChannelNames.Names->Add(MakeShareable(new FString(FName(NAME_None).ToString())));
+									ChannelNames.Names->Add(MakeShareable(new FRigVMStringWithTag(FName(NAME_None).ToString())));
 
 									if(const FRigControlElement* ControlElement = Hierarchy->Find<FRigControlElement>(ControlKey))
 									{
@@ -294,7 +294,7 @@ TSharedPtr<SGraphPin> FControlRigGraphPanelPinFactory::CreatePin_Internal(UEdGra
 											{
 												if(ChildControl->IsAnimationChannel())
 												{
-													ChannelNames.Names->Add(MakeShareable(new FString(ChildControl->GetDisplayName().ToString())));
+													ChannelNames.Names->Add(MakeShareable(new FRigVMStringWithTag(ChildControl->GetDisplayName().ToString())));
 												}
 											}
 										}
@@ -303,7 +303,7 @@ TSharedPtr<SGraphPin> FControlRigGraphPanelPinFactory::CreatePin_Internal(UEdGra
 								return ChannelNames.Names.Get();
 							}
 
-							static TArray<TSharedPtr<FString>> EmptyNameList;
+							static TArray<TSharedPtr<FRigVMStringWithTag>> EmptyNameList;
 							return &EmptyNameList;
 						})
 						.OnGetSelectedClicked_UObject(RigGraph, &UControlRigGraph::HandleGetSelectedClicked)
@@ -314,7 +314,7 @@ TSharedPtr<SGraphPin> FControlRigGraphPanelPinFactory::CreatePin_Internal(UEdGra
 					struct FCachedMetadataNames
 					{
 						int32 MetadataVersion;
-						TSharedPtr<TArray<TSharedPtr<FString>>> Names;
+						TSharedPtr<TArray<TSharedPtr<FRigVMStringWithTag>>> Names;
 						
 						FCachedMetadataNames()
 						: MetadataVersion(INDEX_NONE)
@@ -352,20 +352,20 @@ TSharedPtr<SGraphPin> FControlRigGraphPanelPinFactory::CreatePin_Internal(UEdGra
 
 										if(!MetadataNames.Names.IsValid())
 										{
-											MetadataNames.Names = MakeShareable(new TArray<TSharedPtr<FString>>());
+											MetadataNames.Names = MakeShareable(new TArray<TSharedPtr<FRigVMStringWithTag>>());
 										}
 										MetadataNames.Names->Reset();
 
 										for(const FName& Name : Names)
 										{
-											MetadataNames.Names->Add(MakeShareable(new FString(Name.ToString())));
+											MetadataNames.Names->Add(MakeShareable(new FRigVMStringWithTag(Name.ToString())));
 										}
 
-										MetadataNames.Names->Sort([](const TSharedPtr<FString>& A, const TSharedPtr<FString>& B)
+										MetadataNames.Names->Sort([](const TSharedPtr<FRigVMStringWithTag>& A, const TSharedPtr<FRigVMStringWithTag>& B)
 										{
 											return A.Get() > B.Get();
 										});
-										MetadataNames.Names->Insert(MakeShareable(new FString(FName(NAME_None).ToString())), 0);
+										MetadataNames.Names->Insert(MakeShareable(new FRigVMStringWithTag(FName(NAME_None).ToString())), 0);
 
 										MetadataNames.MetadataVersion = MetadataVersion;
 									}
@@ -373,7 +373,7 @@ TSharedPtr<SGraphPin> FControlRigGraphPanelPinFactory::CreatePin_Internal(UEdGra
 								}
 							}
 
-							static TArray<TSharedPtr<FString>> EmptyNameList;
+							static TArray<TSharedPtr<FRigVMStringWithTag>> EmptyNameList;
 							return &EmptyNameList;
 						});
 				}
@@ -382,7 +382,7 @@ TSharedPtr<SGraphPin> FControlRigGraphPanelPinFactory::CreatePin_Internal(UEdGra
 					struct FCachedMetadataTagNames
 					{
 						int32 MetadataTagVersion;
-						TSharedPtr<TArray<TSharedPtr<FString>>> Names;
+						TSharedPtr<TArray<TSharedPtr<FRigVMStringWithTag>>> Names;
 						FCachedMetadataTagNames()
 						: MetadataTagVersion(INDEX_NONE)
 						{}
@@ -422,19 +422,19 @@ TSharedPtr<SGraphPin> FControlRigGraphPanelPinFactory::CreatePin_Internal(UEdGra
 
 										if(!MetadataTagNames.Names.IsValid())
 										{
-											MetadataTagNames.Names = MakeShareable(new TArray<TSharedPtr<FString>>());
+											MetadataTagNames.Names = MakeShareable(new TArray<TSharedPtr<FRigVMStringWithTag>>());
 										}
 										MetadataTagNames.Names->Reset();
 
 										for(const FName& Tag : Tags)
 										{
-											MetadataTagNames.Names->Add(MakeShareable(new FString(Tag.ToString())));
+											MetadataTagNames.Names->Add(MakeShareable(new FRigVMStringWithTag(Tag.ToString())));
 										}
-										MetadataTagNames.Names->Sort([](const TSharedPtr<FString>& A, const TSharedPtr<FString>& B)
+										MetadataTagNames.Names->Sort([](const TSharedPtr<FRigVMStringWithTag>& A, const TSharedPtr<FRigVMStringWithTag>& B)
 										{
 											return A.Get() > B.Get();
 										});
-										MetadataTagNames.Names->Insert(MakeShareable(new FString(FName(NAME_None).ToString())), 0);
+										MetadataTagNames.Names->Insert(MakeShareable(new FRigVMStringWithTag(FName(NAME_None).ToString())), 0);
 
 										MetadataTagNames.MetadataTagVersion = MetadataTagVersion;
 									}
@@ -443,7 +443,7 @@ TSharedPtr<SGraphPin> FControlRigGraphPanelPinFactory::CreatePin_Internal(UEdGra
 								}
 							}
 
-							static TArray<TSharedPtr<FString>> EmptyNameList;
+							static TArray<TSharedPtr<FRigVMStringWithTag>> EmptyNameList;
 							return &EmptyNameList;
 						});
 				}
diff --git a/Engine/Plugins/Runtime/RigVM/Source/RigVMDeveloper/Private/EdGraph/RigVMEdGraph.cpp b/Engine/Plugins/Runtime/RigVM/Source/RigVMDeveloper/Private/EdGraph/RigVMEdGraph.cpp
--- a/Engine/Plugins/Runtime/RigVM/Source/RigVMDeveloper/Private/EdGraph/RigVMEdGraph.cpp
+++ b/Engine/Plugins/Runtime/RigVM/Source/RigVMDeveloper/Private/EdGraph/RigVMEdGraph.cpp
@@ -828,19 +828,19 @@ int32 URigVMEdGraph::GetInstructionIndex(const URigVMEdGraphNode* InNode, bool b
 void URigVMEdGraph::CacheEntryNameList()
 {
 	EntryNameList.Reset();
-	EntryNameList.Add(MakeShared<FString>(FName(NAME_None).ToString()));
+	EntryNameList.Add(MakeShared<FRigVMStringWithTag>(FName(NAME_None).ToString()));
 
 	if(const URigVMBlueprint* Blueprint = CastChecked<URigVMBlueprint>(GetBlueprint()))
 	{
 		const TArray<FName> EntryNames = Blueprint->GetRigVMClient()->GetEntryNames();
 		for (const FName& EntryName : EntryNames)
 		{
-			EntryNameList.Add(MakeShared<FString>(EntryName.ToString()));
+			EntryNameList.Add(MakeShared<FRigVMStringWithTag>(EntryName.ToString()));
 		}
 	}
 }
 
-const TArray<TSharedPtr<FString>>* URigVMEdGraph::GetEntryNameList(URigVMPin* InPin) const
+const TArray<TSharedPtr<FRigVMStringWithTag>>* URigVMEdGraph::GetEntryNameList(URigVMPin* InPin) const
 {
 	if (const URigVMEdGraph* OuterGraph = Cast<URigVMEdGraph>(GetOuter()))
 	{
diff --git a/Engine/Plugins/Runtime/RigVM/Source/RigVMDeveloper/Public/EdGraph/RigVMEdGraph.h b/Engine/Plugins/Runtime/RigVM/Source/RigVMDeveloper/Public/EdGraph/RigVMEdGraph.h
--- a/Engine/Plugins/Runtime/RigVM/Source/RigVMDeveloper/Public/EdGraph/RigVMEdGraph.h
+++ b/Engine/Plugins/Runtime/RigVM/Source/RigVMDeveloper/Public/EdGraph/RigVMEdGraph.h
@@ -13,6 +13,118 @@ class URigVMBlueprint;
 class URigVMEdGraphSchema;
 class URigVMController;
 
+struct FRigVMStringTag
+{
+public:
+	
+	FRigVMStringTag()
+		: Name(NAME_None)
+		, Color(FLinearColor::Red)
+	{}
+
+	explicit FRigVMStringTag(const FName& InName, const FLinearColor& InColor)
+		: Name(InName)
+		, Color(InColor)
+	{
+	}
+
+	const FName& GetName() const
+	{
+		return Name;
+	}
+	
+	const FLinearColor& GetColor() const
+	{
+		return Color;
+	}
+
+	bool IsValid() const
+	{
+		return !Name.IsNone();
+	}
+
+	bool Equals(const FName& InOther) const
+	{
+		return GetName().IsEqual(InOther, ENameCase::CaseSensitive);
+	}
+
+	bool Equals(const FRigVMStringTag& InOther) const
+	{
+		return Equals(InOther.GetName());
+	}
+
+private:
+	FName Name;
+	FLinearColor Color;
+};
+
+struct FRigVMStringWithTag
+{
+public:
+	
+	FRigVMStringWithTag() = default;
+
+	FRigVMStringWithTag(const FString& InString, const FRigVMStringTag& InTag = FRigVMStringTag())
+		: String(InString)
+		, Tag(InTag)
+	{
+	}
+
+	const FString& GetString() const
+	{
+		return String;
+	}
+
+	FString GetStringWithTag() const
+	{
+		if(HasTag())
+		{
+			static constexpr TCHAR Format[] = TEXT("%s (%s)");
+			return FString::Printf(Format, *GetString(), *GetTag().GetName().ToString());
+		}
+		return GetString();
+	}
+
+	bool HasTag() const
+	{
+		return Tag.IsValid();
+	}
+
+	const FRigVMStringTag& GetTag() const
+	{
+		return Tag;
+	}
+
+	bool operator ==(const FRigVMStringWithTag& InOther) const
+	{
+		return Equals(InOther);
+	}
+
+	bool operator >(const FRigVMStringWithTag& InOther) const
+	{
+		return GetString() > InOther.GetString();
+	}
+
+	bool operator <(const FRigVMStringWithTag& InOther) const
+	{
+		return GetString() < InOther.GetString();
+	}
+
+	bool Equals(const FString& InOther) const
+	{
+		return GetString().Equals(InOther, ESearchCase::CaseSensitive);
+	}
+
+	bool Equals(const FRigVMStringWithTag& InOther) const
+	{
+		return Equals(InOther.GetString());
+	}
+
+private:
+	FString String;
+	FRigVMStringTag Tag;
+};
+
 DECLARE_MULTICAST_DELEGATE_OneParam(FRigVMEdGraphNodeClicked, URigVMEdGraphNode*);
 
 UCLASS()
@@ -58,9 +170,9 @@ public:
 	int32 GetInstructionIndex(const URigVMEdGraphNode* InNode, bool bAsInput);
 
 	void CacheEntryNameList();
-	const TArray<TSharedPtr<FString>>* GetEntryNameList(URigVMPin* InPin = nullptr) const;
+	const TArray<TSharedPtr<FRigVMStringWithTag>>* GetEntryNameList(URigVMPin* InPin = nullptr) const;
 
-	virtual const TArray<TSharedPtr<FString>>* GetNameListForWidget(const FString& InWidgetName) const { return nullptr; }
+	virtual const TArray<TSharedPtr<FRigVMStringWithTag>>* GetNameListForWidget(const FString& InWidgetName) const { return nullptr; }
 
 	UPROPERTY()
 	FString ModelNodePath;
@@ -88,7 +200,7 @@ protected:
 private:
 	TMap<FName, UEdGraphNode*> ModelNodePathToEdNode;
 	mutable TWeakObjectPtr<URigVMGraph> CachedModelGraph;
-	TArray<TSharedPtr<FString>> EntryNameList;
+	TArray<TSharedPtr<FRigVMStringWithTag>> EntryNameList;
 
 #endif
 	friend class URigVMEdGraphNode;
diff --git a/Engine/Plugins/Runtime/RigVM/Source/RigVMEditor/Private/Editor/RigVMGraphDetailCustomization.cpp b/Engine/Plugins/Runtime/RigVM/Source/RigVMEditor/Private/Editor/RigVMGraphDetailCustomization.cpp
--- a/Engine/Plugins/Runtime/RigVM/Source/RigVMEditor/Private/Editor/RigVMGraphDetailCustomization.cpp
+++ b/Engine/Plugins/Runtime/RigVM/Source/RigVMEditor/Private/Editor/RigVMGraphDetailCustomization.cpp
@@ -921,8 +921,8 @@ void FRigVMGraphDetailCustomization::CustomizeDetails(IDetailLayoutBuilder& Deta
 
 		if(AccessSpecifierStrings.IsEmpty())
 		{
-			AccessSpecifierStrings.Add(TSharedPtr<FString>(new FString(TEXT("Public"))));
-			AccessSpecifierStrings.Add(TSharedPtr<FString>(new FString(TEXT("Private"))));
+			AccessSpecifierStrings.Add(TSharedPtr<FRigVMStringWithTag>(new FRigVMStringWithTag(TEXT("Public"))));
+			AccessSpecifierStrings.Add(TSharedPtr<FRigVMStringWithTag>(new FRigVMStringWithTag(TEXT("Private"))));
 		}
 
 		// access specifier
@@ -945,7 +945,7 @@ void FRigVMGraphDetailCustomization::CustomizeDetails(IDetailLayoutBuilder& Deta
             ]
             .MenuContent()
             [
-                SNew(SListView<TSharedPtr<FString> >)
+                SNew(SListView<TSharedPtr<FRigVMStringWithTag> >)
                     .ListItemsSource( &AccessSpecifierStrings )
                     .OnGenerateRow(this, &FRigVMGraphDetailCustomization::HandleGenerateRowAccessSpecifier)
                     .OnSelectionChanged(this, &FRigVMGraphDetailCustomization::OnAccessSpecifierSelected)
@@ -1248,7 +1248,7 @@ FReply FRigVMGraphDetailCustomization::OnNodeColorClicked()
 	return FReply::Handled();
 }
 
-TArray<TSharedPtr<FString>> FRigVMGraphDetailCustomization::AccessSpecifierStrings;
+TArray<TSharedPtr<FRigVMStringWithTag>> FRigVMGraphDetailCustomization::AccessSpecifierStrings;
 
 FText FRigVMGraphDetailCustomization::GetCurrentAccessSpecifierName() const
 {
@@ -1263,16 +1263,16 @@ FText FRigVMGraphDetailCustomization::GetCurrentAccessSpecifierName() const
 			{
 				if(RigVMBlueprint->IsFunctionPublic(LibraryNode->GetFName()))
 				{
-					return FText::FromString(*AccessSpecifierStrings[0].Get()); // public
+					return FText::FromString(AccessSpecifierStrings[0]->GetString()); // public
 				}
 			}
 		}
 	}
 
-	return FText::FromString(*AccessSpecifierStrings[1].Get()); // private
+	return FText::FromString(AccessSpecifierStrings[1]->GetString()); // private
 }
 
-void FRigVMGraphDetailCustomization::OnAccessSpecifierSelected( TSharedPtr<FString> SpecifierName, ESelectInfo::Type SelectInfo )
+void FRigVMGraphDetailCustomization::OnAccessSpecifierSelected( TSharedPtr<FRigVMStringWithTag> SpecifierName, ESelectInfo::Type SelectInfo )
 {
 	if(RigVMBlueprintPtr.IsValid() && GraphPtr.IsValid())
 	{
@@ -1296,13 +1296,13 @@ void FRigVMGraphDetailCustomization::OnAccessSpecifierSelected( TSharedPtr<FStri
 	}
 }
 
-TSharedRef<ITableRow> FRigVMGraphDetailCustomization::HandleGenerateRowAccessSpecifier( TSharedPtr<FString> SpecifierName, const TSharedRef<STableViewBase>& OwnerTable )
+TSharedRef<ITableRow> FRigVMGraphDetailCustomization::HandleGenerateRowAccessSpecifier( TSharedPtr<FRigVMStringWithTag> SpecifierName, const TSharedRef<STableViewBase>& OwnerTable )
 {
-	return SNew(STableRow< TSharedPtr<FString> >, OwnerTable)
+	return SNew(STableRow< TSharedPtr<FRigVMStringWithTag> >, OwnerTable)
         .Content()
         [
             SNew( STextBlock ) 
-                .Text(FText::FromString(*SpecifierName.Get()) )
+                .Text(FText::FromString(SpecifierName->GetString()) )
         ];
 }
 
@@ -1411,7 +1411,7 @@ void FRigVMWrappedNodeDetailCustomization::CustomizeDetails(IDetailLayoutBuilder
         			BlueprintBeingCustomized->GetEdGraph(NodesBeingCustomized[0]->GetGraph()));
         		ensure(GraphBeingCustomized);
         		
-        		const TArray<TSharedPtr<FString>>* NameList = GraphBeingCustomized->GetNameListForWidget(CustomWidgetName);
+        		const TArray<TSharedPtr<FRigVMStringWithTag>>* NameList = GraphBeingCustomized->GetNameListForWidget(CustomWidgetName);
         		if (NameList)
         		{
         			TSharedPtr<SRigVMGraphPinNameListValueWidget> NameListWidget;
@@ -1518,9 +1518,10 @@ void FRigVMWrappedNodeDetailCustomization::CustomizeDetails(IDetailLayoutBuilder
 	CustomizeLiveValues(DetailLayout);
 }
 
-TSharedRef<SWidget> FRigVMWrappedNodeDetailCustomization::MakeNameListItemWidget(TSharedPtr<FString> InItem)
+TSharedRef<SWidget> FRigVMWrappedNodeDetailCustomization::MakeNameListItemWidget(TSharedPtr<FRigVMStringWithTag> InItem)
 {
-	return 	SNew(STextBlock).Text(FText::FromString(*InItem));// .Font(FAppStyle::GetFontStyle(TEXT("PropertyWindow.NormalFont")));
+	//TODO: make this prettier
+	return 	SNew(STextBlock).Text(FText::FromString(InItem->GetStringWithTag()));// .Font(FAppStyle::GetFontStyle(TEXT("PropertyWindow.NormalFont")));
 }
 
 FText FRigVMWrappedNodeDetailCustomization::GetNameListText(FNameProperty* InProperty) const
@@ -1544,17 +1545,17 @@ FText FRigVMWrappedNodeDetailCustomization::GetNameListText(FNameProperty* InPro
 	return FirstText;
 }
 
-TSharedPtr<FString> FRigVMWrappedNodeDetailCustomization::GetCurrentlySelectedItem(FNameProperty* InProperty, const TArray<TSharedPtr<FString>>* InNameList) const
+TSharedPtr<FRigVMStringWithTag> FRigVMWrappedNodeDetailCustomization::GetCurrentlySelectedItem(FNameProperty* InProperty, const TArray<TSharedPtr<FRigVMStringWithTag>>* InNameList) const
 {
 	FString CurrentItem = GetNameListText(InProperty).ToString();
-	for (const TSharedPtr<FString>& Item : *InNameList)
+	for (const TSharedPtr<FRigVMStringWithTag>& Item : *InNameList)
 	{
 		if (Item->Equals(CurrentItem))
 		{
 			return Item;
 		}
 	}
-	return TSharedPtr<FString>();
+	return TSharedPtr<FRigVMStringWithTag>();
 }
 
 
@@ -1576,19 +1577,19 @@ void FRigVMWrappedNodeDetailCustomization::SetNameListText(const FText& NewTypeI
 	Controller->CloseUndoBracket();
 }
 
-void FRigVMWrappedNodeDetailCustomization::OnNameListChanged(TSharedPtr<FString> NewSelection, ESelectInfo::Type SelectInfo, FNameProperty* InProperty, TSharedRef<IPropertyUtilities> PropertyUtilities)
+void FRigVMWrappedNodeDetailCustomization::OnNameListChanged(TSharedPtr<FRigVMStringWithTag> NewSelection, ESelectInfo::Type SelectInfo, FNameProperty* InProperty, TSharedRef<IPropertyUtilities> PropertyUtilities)
 {
 	if (SelectInfo != ESelectInfo::Direct)
 	{
-		FString NewValue = *NewSelection.Get();
+		const FString& NewValue = NewSelection->GetString();
 		SetNameListText(FText::FromString(NewValue), ETextCommit::OnEnter, InProperty, PropertyUtilities);
 	}
 }
 
-void FRigVMWrappedNodeDetailCustomization::OnNameListComboBox(FNameProperty* InProperty, const TArray<TSharedPtr<FString>>* InNameList)
+void FRigVMWrappedNodeDetailCustomization::OnNameListComboBox(FNameProperty* InProperty, const TArray<TSharedPtr<FRigVMStringWithTag>>* InNameList)
 {
 	TSharedPtr<SRigVMGraphPinNameListValueWidget> Widget = NameListWidgets.FindChecked(InProperty->GetFName());
-	const TSharedPtr<FString> CurrentlySelected = GetCurrentlySelectedItem(InProperty, InNameList);
+	const TSharedPtr<FRigVMStringWithTag> CurrentlySelected = GetCurrentlySelectedItem(InProperty, InNameList);
 	Widget->SetSelectedItem(CurrentlySelected);
 }
 
diff --git a/Engine/Plugins/Runtime/RigVM/Source/RigVMEditor/Private/Widgets/SRigVMGraphPinNameList.cpp b/Engine/Plugins/Runtime/RigVM/Source/RigVMEditor/Private/Widgets/SRigVMGraphPinNameList.cpp
--- a/Engine/Plugins/Runtime/RigVM/Source/RigVMEditor/Private/Widgets/SRigVMGraphPinNameList.cpp
+++ b/Engine/Plugins/Runtime/RigVM/Source/RigVMEditor/Private/Widgets/SRigVMGraphPinNameList.cpp
@@ -44,9 +44,9 @@ void SRigVMGraphPinNameList::Construct(const FArguments& InArgs, UEdGraphPin* In
 
 TSharedRef<SWidget>	SRigVMGraphPinNameList::GetDefaultValueWidget()
 {
-	TSharedPtr<FString> InitialSelected;
-	const TArray<TSharedPtr<FString>>* List = GetNameList();
-	for (TSharedPtr<FString> Item : (*List))
+	TSharedPtr<FRigVMStringWithTag> InitialSelected;
+	const TArray<TSharedPtr<FRigVMStringWithTag>>* List = GetNameList();
+	for (const TSharedPtr<FRigVMStringWithTag>& Item : (*List))
 	{
 		if (Item->Equals(GetNameListText().ToString()))
 		{
@@ -121,9 +121,9 @@ TSharedRef<SWidget>	SRigVMGraphPinNameList::GetDefaultValueWidget()
 		];
 }
 
-const TArray<TSharedPtr<FString>>* SRigVMGraphPinNameList::GetNameList(bool bForContent) const
+const TArray<TSharedPtr<FRigVMStringWithTag>>* SRigVMGraphPinNameList::GetNameList(bool bForContent) const
 {
-	const TArray<TSharedPtr<FString>>* Result = nullptr;
+	const TArray<TSharedPtr<FRigVMStringWithTag>>* Result = nullptr;
 	
 	// if we are looking for the validation name list, try to get it from the
 	// dedicated delegate first - but fall back on the content list nevertheless
@@ -154,7 +154,15 @@ const TArray<TSharedPtr<FString>>* SRigVMGraphPinNameList::GetNameList(bool bFor
 
 FText SRigVMGraphPinNameList::GetNameListText() const
 {
-	return FText::FromString( GraphPinObj->GetDefaultAsString() );
+	const FString DefaultString = GraphPinObj->GetDefaultAsString();
+	for (const TSharedPtr<FRigVMStringWithTag>& Item : (*CurrentList))
+	{
+		if (Item->Equals(DefaultString))
+		{
+			return FText::FromString(Item->GetStringWithTag());
+		}
+	}
+	return FText::FromString( DefaultString );
 }
 
 void SRigVMGraphPinNameList::SetNameListText(const FText& NewTypeInValue, ETextCommit::Type /*CommitInfo*/)
@@ -183,7 +191,7 @@ FSlateColor SRigVMGraphPinNameList::GetNameColor() const
 {
 	if(bMarkupInvalidItems)
 	{
-		FString CurrentItem = GetNameListText().ToString();
+		const FString DefaultString = GraphPinObj->GetDefaultAsString();
 
 		if(!EnableNameListCache)
 		{
@@ -191,16 +199,16 @@ FSlateColor SRigVMGraphPinNameList::GetNameColor() const
 		}
 
 		bool bFound = false;
-		for (TSharedPtr<FString> Item : (*ValidationList))
+		for (const TSharedPtr<FRigVMStringWithTag>& Item : (*ValidationList))
 		{
-			if (Item->Equals(CurrentItem))
+			if (Item->Equals(DefaultString))
 			{
 				bFound = true;
 				break;
 			}
 		}
 
-		if(!bFound || CurrentItem.IsEmpty() || CurrentItem == FName(NAME_None).ToString())
+		if(!bFound || DefaultString.IsEmpty() || DefaultString == FName(NAME_None).ToString())
 		{
 			return FSlateColor(FLinearColor::Red);
 		}
@@ -208,19 +216,20 @@ FSlateColor SRigVMGraphPinNameList::GetNameColor() const
 	return FSlateColor::UseForeground();
 }
 
-TSharedRef<SWidget> SRigVMGraphPinNameList::MakeNameListItemWidget(TSharedPtr<FString> InItem)
+TSharedRef<SWidget> SRigVMGraphPinNameList::MakeNameListItemWidget(TSharedPtr<FRigVMStringWithTag> InItem)
 {
-	return 	SNew(STextBlock).Text(FText::FromString(*InItem)).Font(FAppStyle::GetFontStyle(TEXT("PropertyWindow.NormalFont")));
+	//TODO: make this prettier
+	return 	SNew(STextBlock).Text(FText::FromString(InItem->GetStringWithTag())).Font(FAppStyle::GetFontStyle(TEXT("PropertyWindow.NormalFont")));
 }
 
-void SRigVMGraphPinNameList::OnNameListChanged(TSharedPtr<FString> NewSelection, ESelectInfo::Type SelectInfo)
+void SRigVMGraphPinNameList::OnNameListChanged(TSharedPtr<FRigVMStringWithTag> NewSelection, ESelectInfo::Type SelectInfo)
 {
 	if (SelectInfo != ESelectInfo::Direct)
 	{
 		FString NewValue = FName(NAME_None).ToString();
 		if (NewSelection.IsValid())
 		{
-			NewValue = *NewSelection.Get();
+			NewValue = NewSelection->GetString();
 		}
 		SetNameListText(FText::FromString(NewValue), ETextCommit::OnEnter);
 	}
@@ -230,8 +239,8 @@ void SRigVMGraphPinNameList::OnNameListComboBox()
 {
 	UpdateNameLists();
 
-	TSharedPtr<FString> CurrentlySelected;
-	for (TSharedPtr<FString> Item : (*CurrentList))
+	TSharedPtr<FRigVMStringWithTag> CurrentlySelected;
+	for (const TSharedPtr<FRigVMStringWithTag>& Item : (*CurrentList))
 	{
 		if (Item->Equals(GetNameListText().ToString()))
 		{
@@ -264,12 +273,12 @@ FReply SRigVMGraphPinNameList::HandleGetSelectedClicked()
 	{
 		if (OnGetNameFromSelection.IsBound() && OnGetSelectedClicked.IsBound())
 		{
-			const TArray<TSharedPtr<FString>> Result = OnGetNameFromSelection.Execute();
+			const TArray<TSharedPtr<FRigVMStringWithTag>> Result = OnGetNameFromSelection.Execute();
 			if (Result.Num() > 0)
 			{
 				if (Result[0].IsValid() && Result[0] != nullptr)
 				{
-					const FString DefaultValue = *Result[0].Get();
+					const FString DefaultValue = Result[0]->GetString();
 					FReply Reply = OnGetSelectedClicked.Execute(Graph, ModelPin, DefaultValue);
 					if(Reply.IsEventHandled())
 					{
@@ -289,10 +298,10 @@ FReply SRigVMGraphPinNameList::HandleBrowseClicked()
 	{
 		if(OnBrowseClicked.IsBound())
 		{
-			TSharedPtr<FString> Selected = NameListComboBox->GetSelectedItem();
+			TSharedPtr<FRigVMStringWithTag> Selected = NameListComboBox->GetSelectedItem();
 			if (Selected.IsValid() && ModelPin)
 			{
-				return OnBrowseClicked.Execute(Graph, ModelPin, *Selected.Get());
+				return OnBrowseClicked.Execute(Graph, ModelPin, Selected->GetString());
 			}
 		}
 	}
diff --git a/Engine/Plugins/Runtime/RigVM/Source/RigVMEditor/Private/Widgets/SRigVMGraphPinNameListValueWidget.cpp b/Engine/Plugins/Runtime/RigVM/Source/RigVMEditor/Private/Widgets/SRigVMGraphPinNameListValueWidget.cpp
--- a/Engine/Plugins/Runtime/RigVM/Source/RigVMEditor/Private/Widgets/SRigVMGraphPinNameListValueWidget.cpp
+++ b/Engine/Plugins/Runtime/RigVM/Source/RigVMEditor/Private/Widgets/SRigVMGraphPinNameListValueWidget.cpp
@@ -76,7 +76,7 @@ void SRigVMGraphPinNameListValueWidget::Construct(const FArguments& InArgs)
 	// Need to establish the selected item at point of construction so its available for querying
 	// NB: If you need a selection to fire use SetItemSelection rather than setting an IntiallySelectedItem
 	SelectedItem = InArgs._InitiallySelectedItem;
-	if (TListTypeTraits<TSharedPtr<FString>>::IsPtrValid(SelectedItem))
+	if (TListTypeTraits<TSharedPtr<FRigVMStringWithTag>>::IsPtrValid(SelectedItem))
 	{
 		ComboListView->Private_SetItemSelection(SelectedItem, true);
 	}
@@ -88,9 +88,9 @@ void SRigVMGraphPinNameListValueWidget::ClearSelection()
 	ComboListView->ClearSelection();
 }
 
-void SRigVMGraphPinNameListValueWidget::SetSelectedItem(TSharedPtr<FString> InSelectedItem)
+void SRigVMGraphPinNameListValueWidget::SetSelectedItem(TSharedPtr<FRigVMStringWithTag> InSelectedItem)
 {
-	if (TListTypeTraits<TSharedPtr<FString>>::IsPtrValid(InSelectedItem))
+	if (TListTypeTraits<TSharedPtr<FRigVMStringWithTag>>::IsPtrValid(InSelectedItem))
 	{
 		ComboListView->SetSelection(InSelectedItem);
 	}
@@ -100,14 +100,14 @@ void SRigVMGraphPinNameListValueWidget::SetSelectedItem(TSharedPtr<FString> InSe
 	}
 }
 
-void SRigVMGraphPinNameListValueWidget::SetOptionsSource(const TArray<TSharedPtr<FString>>* InOptionsSource)
+void SRigVMGraphPinNameListValueWidget::SetOptionsSource(const TArray<TSharedPtr<FRigVMStringWithTag>>* InOptionsSource)
 {
 	check(InOptionsSource);
 	OptionsSource = InOptionsSource;
 	ComboListView->SetItemsSource(OptionsSource);
 }
 
-TSharedPtr<FString> SRigVMGraphPinNameListValueWidget::GetSelectedItem()
+TSharedPtr<FRigVMStringWithTag> SRigVMGraphPinNameListValueWidget::GetSelectedItem()
 {
 	return SelectedItem;
 }
@@ -120,7 +120,7 @@ void SRigVMGraphPinNameListValueWidget::RefreshOptions()
 	}
 }
 
-TSharedRef<ITableRow> SRigVMGraphPinNameListValueWidget::GenerateMenuItemRow(TSharedPtr<FString> InItem, const TSharedRef<STableViewBase>& OwnerTable)
+TSharedRef<ITableRow> SRigVMGraphPinNameListValueWidget::GenerateMenuItemRow(TSharedPtr<FRigVMStringWithTag> InItem, const TSharedRef<STableViewBase>& OwnerTable)
 {
 	if (OnGenerateWidget.IsBound())
 	{
@@ -129,13 +129,13 @@ TSharedRef<ITableRow> SRigVMGraphPinNameListValueWidget::GenerateMenuItemRow(TSh
 		if (!SearchToken.IsEmpty())
 		{
 			FString SearchTokenUnderscores = SearchToken.Replace(TEXT(" "), TEXT("_"));
-			if ((InItem->ToLower().Find(SearchToken) < 0) &&
-				(InItem->ToLower().Find(SearchTokenUnderscores) < 0))
+			if ((InItem->GetString().ToLower().Find(SearchToken) < 0) &&
+				(InItem->GetString().ToLower().Find(SearchTokenUnderscores) < 0))
 			{
 				WidgetVisibility = EVisibility::Collapsed;
 			}
 		}
-		return SNew(SComboRow<TSharedPtr<FString>>, OwnerTable)
+		return SNew(SComboRow<TSharedPtr<FRigVMStringWithTag>>, OwnerTable)
 			.Visibility(WidgetVisibility)
 			[
 				OnGenerateWidget.Execute(InItem)
@@ -143,7 +143,7 @@ TSharedRef<ITableRow> SRigVMGraphPinNameListValueWidget::GenerateMenuItemRow(TSh
 	}
 	else
 	{
-		return SNew(SComboRow<TSharedPtr<FString>>, OwnerTable)
+		return SNew(SComboRow<TSharedPtr<FRigVMStringWithTag>>, OwnerTable)
 			[
 				SNew(STextBlock).Text(NSLOCTEXT("SlateCore", "ComboBoxMissingOnGenerateWidgetMethod", "Please provide a .OnGenerateWidget() handler."))
 			];
@@ -155,7 +155,7 @@ void SRigVMGraphPinNameListValueWidget::OnMenuOpenChanged(bool bOpen)
 {
 	if (bOpen == false)
 	{
-		if (TListTypeTraits<TSharedPtr<FString>>::IsPtrValid(SelectedItem))
+		if (TListTypeTraits<TSharedPtr<FRigVMStringWithTag>>::IsPtrValid(SelectedItem))
 		{
 			// Ensure the ListView selection is set back to the last committed selection
 			ComboListView->SetSelection(SelectedItem, ESelectInfo::OnNavigation);
@@ -196,7 +196,7 @@ EActiveTimerReturnType SRigVMGraphPinNameListValueWidget::SetFocusPostConstruct(
 	return EActiveTimerReturnType::Continue;
 }
 
-void SRigVMGraphPinNameListValueWidget::OnSelectionChanged_Internal(TSharedPtr<FString> ProposedSelection, ESelectInfo::Type SelectInfo, bool bForce)
+void SRigVMGraphPinNameListValueWidget::OnSelectionChanged_Internal(TSharedPtr<FRigVMStringWithTag> ProposedSelection, ESelectInfo::Type SelectInfo, bool bForce)
 {
 	// Ensure that the proposed selection is different
 	if (SelectInfo != ESelectInfo::OnNavigation || bForce)
@@ -219,7 +219,7 @@ void SRigVMGraphPinNameListValueWidget::OnSearchTextChanged(const FText& Changed
 
 	for (int32 i = 0; i < OptionsSource->Num(); i++)
 	{
-		const TSharedPtr<FString>& Option = (*OptionsSource)[i];
+		const TSharedPtr<FRigVMStringWithTag>& Option = (*OptionsSource)[i];
 
 		TSharedPtr<ITableRow> Row = ComboListView->WidgetFromItem(Option);
 		if (Row)
@@ -228,14 +228,14 @@ void SRigVMGraphPinNameListValueWidget::OnSearchTextChanged(const FText& Changed
 			{
 				Row->AsWidget()->SetVisibility(EVisibility::Visible);
 			}
-			else if (Option->ToLower() == SearchToken)
+			else if (Option->GetString().ToLower() == SearchToken)
 			{
 				Row->AsWidget()->SetVisibility(EVisibility::Visible);
 				ComboListView->Private_ClearSelection();
 				ComboListView->Private_SetItemSelection(Option, true);
 			}
-			else if ((Option->ToLower().Find(SearchToken) >= 0) ||
-				(Option->ToLower().Find(SearchTokenUnderscores) >= 0))
+			else if ((Option->GetString().ToLower().Find(SearchToken) >= 0) ||
+				(Option->GetString().ToLower().Find(SearchTokenUnderscores) >= 0))
 			{
 				Row->AsWidget()->SetVisibility(EVisibility::Visible);
 			}
@@ -248,7 +248,7 @@ void SRigVMGraphPinNameListValueWidget::OnSearchTextChanged(const FText& Changed
 
 	ComboListView->RequestListRefresh();
 
-	SelectedItem = TSharedPtr< FString >();
+	SelectedItem = TSharedPtr< FRigVMStringWithTag >();
 }
 
 void SRigVMGraphPinNameListValueWidget::OnSearchTextCommitted(const FText& ChangedText, ETextCommit::Type CommitType)
@@ -258,8 +258,8 @@ void SRigVMGraphPinNameListValueWidget::OnSearchTextCommitted(const FText& Chang
 		const FString LowerCaseProposedSelection = ChangedText.ToString().ToLower();
 		for (int32 i = 0; i < OptionsSource->Num(); i++)
 		{
-			const TSharedPtr<FString>& Option = (*OptionsSource)[i];
-			if (Option->ToLower() == LowerCaseProposedSelection)
+			const TSharedPtr<FRigVMStringWithTag>& Option = (*OptionsSource)[i];
+			if (Option->GetString().ToLower() == LowerCaseProposedSelection)
 			{
 				OnSelectionChanged_Internal(Option, ESelectInfo::OnKeyPress, true);
 				return;
@@ -268,7 +268,7 @@ void SRigVMGraphPinNameListValueWidget::OnSearchTextCommitted(const FText& Chang
 
 		if(AllowUserProvidedText)
 		{
-			const TSharedPtr<FString> UserDefinedOption = MakeShareable(new FString(ChangedText.ToString()));
+			const TSharedPtr<FRigVMStringWithTag> UserDefinedOption = MakeShareable(new FRigVMStringWithTag(ChangedText.ToString()));
 			OnSelectionChanged_Internal(UserDefinedOption, ESelectInfo::OnKeyPress, true);
 		}
 	}
@@ -280,7 +280,7 @@ FReply SRigVMGraphPinNameListValueWidget::OnSearchTextKeyDown(const FGeometry& M
 	{
 		for (int32 i = 0; i < OptionsSource->Num(); i++)
 		{
-			const TSharedPtr<FString>& Option = (*OptionsSource)[i];
+			const TSharedPtr<FRigVMStringWithTag>& Option = (*OptionsSource)[i];
 			TSharedPtr<ITableRow> Row = ComboListView->WidgetFromItem(Option);
 			if (Row)
 			{
@@ -306,7 +306,7 @@ FReply SRigVMGraphPinNameListValueWidget::OnComboListKeyDown(const FGeometry& My
 {
 	if (InKeyEvent.GetKey() == EKeys::Enter)
 	{
-		TArray<TSharedPtr<FString>> SelectedItems = ComboListView->GetSelectedItems();
+		TArray<TSharedPtr<FRigVMStringWithTag>> SelectedItems = ComboListView->GetSelectedItems();
 		if (SelectedItems.Num() > 0)
 		{
 			OnSelectionChanged_Internal(SelectedItems[0], ESelectInfo::OnKeyPress);
@@ -324,7 +324,7 @@ FReply SRigVMGraphPinNameListValueWidget::OnButtonClicked()
 	if (this->IsOpen())
 	{
 		// Re-select first selected item, just in case it was selected by navigation previously
-		TArray<TSharedPtr<FString>> SelectedItems = ComboListView->GetSelectedItems();
+		TArray<TSharedPtr<FRigVMStringWithTag>> SelectedItems = ComboListView->GetSelectedItems();
 		if (SelectedItems.Num() > 0)
 		{
 			OnSelectionChanged_Internal(SelectedItems[0], ESelectInfo::Direct);
diff --git a/Engine/Plugins/Runtime/RigVM/Source/RigVMEditor/Public/Editor/RigVMGraphDetailCustomization.h b/Engine/Plugins/Runtime/RigVM/Source/RigVMEditor/Public/Editor/RigVMGraphDetailCustomization.h
--- a/Engine/Plugins/Runtime/RigVM/Source/RigVMEditor/Public/Editor/RigVMGraphDetailCustomization.h
+++ b/Engine/Plugins/Runtime/RigVM/Source/RigVMEditor/Public/Editor/RigVMGraphDetailCustomization.h
@@ -189,8 +189,8 @@ public:
 	void OnNodeColorCancelled(FLinearColor OriginalColor);
 	FReply OnNodeColorClicked();
 	FText GetCurrentAccessSpecifierName() const;
-	void OnAccessSpecifierSelected( TSharedPtr<FString> SpecifierName, ESelectInfo::Type SelectInfo );
-	TSharedRef<ITableRow> HandleGenerateRowAccessSpecifier( TSharedPtr<FString> SpecifierName, const TSharedRef<STableViewBase>& OwnerTable );
+	void OnAccessSpecifierSelected( TSharedPtr<FRigVMStringWithTag> SpecifierName, ESelectInfo::Type SelectInfo );
+	TSharedRef<ITableRow> HandleGenerateRowAccessSpecifier( TSharedPtr<FRigVMStringWithTag> SpecifierName, const TSharedRef<STableViewBase>& OwnerTable );
 
 private:
 
@@ -209,7 +209,7 @@ private:
 	/** Set to true if the UI is currently picking a color */
 	bool bIsPickingColor;
 
-	static TArray<TSharedPtr<FString>> AccessSpecifierStrings;
+	static TArray<TSharedPtr<FRigVMStringWithTag>> AccessSpecifierStrings;
 };
 
 /** Customization for editing a rig vm node */
@@ -225,12 +225,12 @@ public:
 	// IDetailCustomization interface
 	virtual void CustomizeDetails(IDetailLayoutBuilder& DetailLayout) override;
 
-	TSharedRef<SWidget> MakeNameListItemWidget(TSharedPtr<FString> InItem);
+	TSharedRef<SWidget> MakeNameListItemWidget(TSharedPtr<FRigVMStringWithTag> InItem);
 	FText GetNameListText(FNameProperty* InProperty) const;
-	TSharedPtr<FString> GetCurrentlySelectedItem(FNameProperty* InProperty, const TArray<TSharedPtr<FString>>* InNameList) const;
+	TSharedPtr<FRigVMStringWithTag> GetCurrentlySelectedItem(FNameProperty* InProperty, const TArray<TSharedPtr<FRigVMStringWithTag>>* InNameList) const;
 	void SetNameListText(const FText& NewTypeInValue, ETextCommit::Type, FNameProperty* InProperty, TSharedRef<IPropertyUtilities> PropertyUtilities);
-	void OnNameListChanged(TSharedPtr<FString> NewSelection, ESelectInfo::Type SelectInfo, FNameProperty* InProperty, TSharedRef<IPropertyUtilities> PropertyUtilities);
-	void OnNameListComboBox(FNameProperty* InProperty, const TArray<TSharedPtr<FString>>* InNameList);
+	void OnNameListChanged(TSharedPtr<FRigVMStringWithTag> NewSelection, ESelectInfo::Type SelectInfo, FNameProperty* InProperty, TSharedRef<IPropertyUtilities> PropertyUtilities);
+	void OnNameListComboBox(FNameProperty* InProperty, const TArray<TSharedPtr<FRigVMStringWithTag>>* InNameList);
 	void CustomizeLiveValues(IDetailLayoutBuilder& DetailLayout);
 
 	URigVMBlueprint* BlueprintBeingCustomized;
diff --git a/Engine/Plugins/Runtime/RigVM/Source/RigVMEditor/Public/Widgets/SRigVMGraphPinNameList.h b/Engine/Plugins/Runtime/RigVM/Source/RigVMEditor/Public/Widgets/SRigVMGraphPinNameList.h
--- a/Engine/Plugins/Runtime/RigVM/Source/RigVMEditor/Public/Widgets/SRigVMGraphPinNameList.h
+++ b/Engine/Plugins/Runtime/RigVM/Source/RigVMEditor/Public/Widgets/SRigVMGraphPinNameList.h
@@ -15,8 +15,8 @@ class RIGVMEDITOR_API SRigVMGraphPinNameList : public SGraphPin
 {
 public:
 
-	DECLARE_DELEGATE_RetVal_OneParam( const TArray<TSharedPtr<FString>>*, FOnGetNameListContent, URigVMPin*);
-	DECLARE_DELEGATE_RetVal( const TArray<TSharedPtr<FString>>, FOnGetNameFromSelection);
+	DECLARE_DELEGATE_RetVal_OneParam( const TArray<TSharedPtr<FRigVMStringWithTag>>*, FOnGetNameListContent, URigVMPin*);
+	DECLARE_DELEGATE_RetVal( const TArray<TSharedPtr<FRigVMStringWithTag>>, FOnGetNameFromSelection);
 	DECLARE_DELEGATE_RetVal_ThreeParams( FReply, FOnGetSelectedClicked, URigVMEdGraph*, URigVMPin*, FString);
 	DECLARE_DELEGATE_RetVal_ThreeParams( FReply, FOnBrowseClicked, URigVMEdGraph*, URigVMPin*, FString);
 
@@ -47,23 +47,23 @@ protected:
 	virtual TSharedRef<SWidget>	GetDefaultValueWidget() override;
 	//~ End SGraphPin Interface
 
-	const TArray<TSharedPtr<FString>>* GetNameList(bool bForContent = true) const;
+	const TArray<TSharedPtr<FRigVMStringWithTag>>* GetNameList(bool bForContent = true) const;
 	FText GetNameListText() const;
 	FSlateColor GetNameColor() const;
 	virtual void SetNameListText(const FText& NewTypeInValue, ETextCommit::Type CommitInfo);
 	void UpdateNameLists(bool bUpdateCurrent = true, bool bUpdateValidation = true);
 
-	TSharedRef<SWidget> MakeNameListItemWidget(TSharedPtr<FString> InItem);
-	void OnNameListChanged(TSharedPtr<FString> NewSelection, ESelectInfo::Type SelectInfo);
+	TSharedRef<SWidget> MakeNameListItemWidget(TSharedPtr<FRigVMStringWithTag> InItem);
+	void OnNameListChanged(TSharedPtr<FRigVMStringWithTag> NewSelection, ESelectInfo::Type SelectInfo);
 	void OnNameListComboBox();
 
 	FOnGetNameListContent OnGetNameListContent;
 	FOnGetNameListContent OnGetNameListContentForValidation;
 	URigVMPin* ModelPin;
 	TSharedPtr<SRigVMGraphPinNameListValueWidget> NameListComboBox;
-	TArray<TSharedPtr<FString>> EmptyList;
-	const TArray<TSharedPtr<FString>>* CurrentList;
-	const TArray<TSharedPtr<FString>>* ValidationList;
+	TArray<TSharedPtr<FRigVMStringWithTag>> EmptyList;
+	const TArray<TSharedPtr<FRigVMStringWithTag>>* CurrentList;
+	const TArray<TSharedPtr<FRigVMStringWithTag>>* ValidationList;
 	bool bMarkupInvalidItems;
 	bool EnableNameListCache;
 	FText SearchHintText;
diff --git a/Engine/Plugins/Runtime/RigVM/Source/RigVMEditor/Public/Widgets/SRigVMGraphPinNameListValueWidget.h b/Engine/Plugins/Runtime/RigVM/Source/RigVMEditor/Public/Widgets/SRigVMGraphPinNameListValueWidget.h
--- a/Engine/Plugins/Runtime/RigVM/Source/RigVMEditor/Public/Widgets/SRigVMGraphPinNameListValueWidget.h
+++ b/Engine/Plugins/Runtime/RigVM/Source/RigVMEditor/Public/Widgets/SRigVMGraphPinNameListValueWidget.h
@@ -9,6 +9,7 @@
 #include "Widgets/Input/SEditableTextBox.h"
 #include "Widgets/Views/SListView.h"
 #include "Widgets/Input/SComboBox.h"
+#include "EdGraph/RigVMEdGraph.h"
 
 /**
  * A searchable text combo box
@@ -17,10 +18,10 @@ class RIGVMEDITOR_API SRigVMGraphPinNameListValueWidget : public SComboButton
 {
 public:
 	/** Type of list used for showing menu options. */
-	typedef SListView< TSharedPtr<FString> > SComboListType;
+	typedef SListView< TSharedPtr<FRigVMStringWithTag> > SComboListType;
 	/** Delegate type used to generate widgets that represent Options */
-	typedef typename TSlateDelegates< TSharedPtr<FString> >::FOnGenerateWidget FOnGenerateWidget;
-	typedef typename TSlateDelegates< TSharedPtr<FString> >::FOnSelectionChanged FOnSelectionChanged;
+	typedef typename TSlateDelegates< TSharedPtr<FRigVMStringWithTag> >::FOnGenerateWidget FOnGenerateWidget;
+	typedef typename TSlateDelegates< TSharedPtr<FRigVMStringWithTag> >::FOnSelectionChanged FOnSelectionChanged;
 
 	SLATE_BEGIN_ARGS(SRigVMGraphPinNameListValueWidget)
 		: _Content()
@@ -41,7 +42,7 @@ public:
 
 		SLATE_ATTRIBUTE(FMargin, ContentPadding)
 
-		SLATE_ARGUMENT(const TArray< TSharedPtr<FString> >*, OptionsSource)
+		SLATE_ARGUMENT(const TArray< TSharedPtr<FRigVMStringWithTag> >*, OptionsSource)
 		SLATE_EVENT(FOnSelectionChanged, OnSelectionChanged)
 		SLATE_EVENT(FOnGenerateWidget, OnGenerateWidget)
 
@@ -52,7 +53,7 @@ public:
 		SLATE_ARGUMENT(TSharedPtr<SScrollBar>, CustomScrollbar)
 
 		/** The option that should be selected when the combo box is first created */
-		SLATE_ARGUMENT(TSharedPtr<FString>, InitiallySelectedItem)
+		SLATE_ARGUMENT(TSharedPtr<FRigVMStringWithTag>, InitiallySelectedItem)
 
 		SLATE_ARGUMENT(TOptional<EPopupMethod>, Method)
 
@@ -86,11 +87,11 @@ public:
 
 	void ClearSelection();
 
-	void SetSelectedItem(TSharedPtr<FString> InSelectedItem);
-	void SetOptionsSource(const TArray< TSharedPtr<FString> >* InOptionsSource);
+	void SetSelectedItem(TSharedPtr<FRigVMStringWithTag> InSelectedItem);
+	void SetOptionsSource(const TArray< TSharedPtr<FRigVMStringWithTag> >* InOptionsSource);
 
 	/** @return the item currently selected by the combo box. */
-	TSharedPtr<FString> GetSelectedItem();
+	TSharedPtr<FRigVMStringWithTag> GetSelectedItem();
 
 	/**
 	 * Requests a list refresh after updating options
@@ -102,13 +103,13 @@ public:
 private:
 
 	/** Generate a row for the InItem in the combo box's list (passed in as OwnerTable). Do this by calling the user-specified OnGenerateWidget */
-	TSharedRef<ITableRow> GenerateMenuItemRow(TSharedPtr<FString> InItem, const TSharedRef<STableViewBase>& OwnerTable);
+	TSharedRef<ITableRow> GenerateMenuItemRow(TSharedPtr<FRigVMStringWithTag> InItem, const TSharedRef<STableViewBase>& OwnerTable);
 
 	//** Called if the menu is closed
 	void OnMenuOpenChanged(bool bOpen);
 
 	/** Invoked when the selection in the list changes */
-	void OnSelectionChanged_Internal(TSharedPtr<FString> ProposedSelection, ESelectInfo::Type SelectInfo, bool bForce = false);
+	void OnSelectionChanged_Internal(TSharedPtr<FRigVMStringWithTag> ProposedSelection, ESelectInfo::Type SelectInfo, bool bForce = false);
 
 	/** Invoked when the search text changes */
 	void OnSearchTextChanged(const FText& ChangedText);
@@ -132,7 +133,7 @@ private:
 	/** Delegate that is invoked when the selected item in the combo box changes */
 	FOnSelectionChanged OnSelectionChanged;
 	/** The item currently selected in the combo box */
-	TSharedPtr<FString> SelectedItem;
+	TSharedPtr<FRigVMStringWithTag> SelectedItem;
 	/** The search field used for the combox box's contents */
 	TSharedPtr< SEditableTextBox > SearchField;
 	/** The ListView that we pop up; visualized the available options. */
@@ -144,7 +145,7 @@ private:
 	/** Delegate to invoke when we need to visualize an option as a widget. */
 	FOnGenerateWidget OnGenerateWidget;
 
-	const TArray< TSharedPtr<FString> >* OptionsSource;
+	const TArray< TSharedPtr<FRigVMStringWithTag> >* OptionsSource;
 	bool AllowUserProvidedText;
 
 	/** Used to focus the name box immediately following construction */
