RigVM: Avoid infinite loop when compiling cyclical dependencies
#rnx

diff --git a/Engine/Plugins/Runtime/RigVM/Source/RigVMDeveloper/Private/RigVMCompiler/RigVMCompiler.cpp b/Engine/Plugins/Runtime/RigVM/Source/RigVMDeveloper/Private/RigVMCompiler/RigVMCompiler.cpp
--- a/Engine/Plugins/Runtime/RigVM/Source/RigVMDeveloper/Private/RigVMCompiler/RigVMCompiler.cpp
+++ b/Engine/Plugins/Runtime/RigVM/Source/RigVMDeveloper/Private/RigVMCompiler/RigVMCompiler.cpp
@@ -1393,6 +1393,12 @@ bool URigVMCompiler::CompileFunction(const URigVMLibraryNode* InLibraryNode, URi
 
 bool URigVMCompiler::CompileFunction(const FRigVMCompileSettings& InSettings, const URigVMLibraryNode* InLibraryNode, URigVMController* InController, const TArray<FRigVMExternalVariable>& InExternalVariables, FRigVMFunctionCompilationData* OutFunctionCompilationData, FRigVMExtendedExecuteContext& OutVMContext)
 {
+	if (FunctionCompilationStack.Contains(InLibraryNode))
+	{
+		return false;
+	}
+
+	FFunctionCompilationScope FunctionCompilationScope(this, InLibraryNode);
 	TGuardValue<const URigVMLibraryNode*> CompilationGuard(CurrentCompilationFunction, InLibraryNode);
 
 	URigVMController* LibraryController = InController->GetControllerForGraph(InLibraryNode->GetContainedGraph());
diff --git a/Engine/Plugins/Runtime/RigVM/Source/RigVMDeveloper/Public/RigVMCompiler/RigVMCompiler.h b/Engine/Plugins/Runtime/RigVM/Source/RigVMDeveloper/Public/RigVMCompiler/RigVMCompiler.h
--- a/Engine/Plugins/Runtime/RigVM/Source/RigVMDeveloper/Public/RigVMCompiler/RigVMCompiler.h
+++ b/Engine/Plugins/Runtime/RigVM/Source/RigVMDeveloper/Public/RigVMCompiler/RigVMCompiler.h
@@ -14,6 +14,7 @@
 #include "RigVMCore/RigVMStruct.h"
 #include "RigVMCompiler/RigVMAST.h"
 #include "Logging/TokenizedMessage.h"
+#include "RigVMModel/Nodes/RigVMLibraryNode.h"
 
 #include "RigVMCompiler.generated.h"
 
@@ -311,6 +312,7 @@ public:
 private:
 
 	const URigVMLibraryNode* CurrentCompilationFunction = nullptr;
+	TSet<const URigVMLibraryNode*> FunctionCompilationStack;
 
 	TArray<URigVMPin*> GetLinkedPins(URigVMPin* InPin, bool bInputs = true, bool bOutputs = true, bool bRecursive = true);
 	uint16 GetElementSizeFromCPPType(const FString& InCPPType, UScriptStruct* InScriptStruct);
@@ -377,4 +379,24 @@ private:
 	}
 	
 	friend class FRigVMCompilerImportErrorContext;
+	friend class FFunctionCompilationScope;
 };
+
+class RIGVMDEVELOPER_API FFunctionCompilationScope
+{
+public:
+	FFunctionCompilationScope(URigVMCompiler* InCompiler, const URigVMLibraryNode* InLibraryNode)
+		: Compiler(InCompiler), LibraryNode(InLibraryNode)
+	{
+		InCompiler->FunctionCompilationStack.Add(InLibraryNode);
+	}
+
+	~FFunctionCompilationScope()
+	{
+		Compiler->FunctionCompilationStack.Remove(LibraryNode);
+	}
+
+private:
+	URigVMCompiler* Compiler;
+	const URigVMLibraryNode* LibraryNode;
+};
\ No newline at end of file
