RigVM: Fix copy/paste transform of BP variables in details panel
#rnx

diff --git a/Engine/Plugins/Runtime/RigVM/Source/RigVMEditor/Public/Editor/RigVMGraphDetailCustomization.h b/Engine/Plugins/Runtime/RigVM/Source/RigVMEditor/Public/Editor/RigVMGraphDetailCustomization.h
--- a/Engine/Plugins/Runtime/RigVM/Source/RigVMEditor/Public/Editor/RigVMGraphDetailCustomization.h
+++ b/Engine/Plugins/Runtime/RigVM/Source/RigVMEditor/Public/Editor/RigVMGraphDetailCustomization.h
@@ -17,6 +17,7 @@
 #include "EdGraph/RigVMEdGraphSchema.h"
 #include "SAdvancedTransformInputBox.h"
 #include "Widgets/SRigVMGraphPinNameListValueWidget.h"
+#include "HAL/PlatformApplicationMisc.h"
 
 class IDetailLayoutBuilder;
 
@@ -917,6 +918,179 @@ protected:
 			}
 		});
 
+		WidgetArgs.OnCopyToClipboard_Lambda([this, InPropertyHandle](
+			ESlateTransformComponent::Type InComponent
+			)
+		{
+			TOptional<FReal> Result;
+			FEditPropertyChain PropertyChain;
+			TArray<int32> PropertyArrayIndices;
+			bool bEnabled;
+			if (!GetPropertyChain(InPropertyHandle, PropertyChain, PropertyArrayIndices, bEnabled))
+			{
+				return;
+			}
+			
+			for(const TWeakObjectPtr<UObject>& Object : ObjectsBeingCustomized)
+			{
+				if(Object.Get() && InPropertyHandle->IsValidHandle())
+				{
+					const TransformType& Transform = ContainerUObjectToValueRef<TransformType>(Object.Get(), Identity, PropertyChain, PropertyArrayIndices);
+					FString Content;
+					switch(InComponent)
+					{
+						case ESlateTransformComponent::Location:
+						{
+							const FVector Data = Transform.GetLocation();
+							TBaseStructure<FVector>::Get()->ExportText(Content, &Data, &Data, nullptr, PPF_None, nullptr);
+							break;
+						}
+						case ESlateTransformComponent::Rotation:
+						{
+							const FRotator Data = Transform.Rotator();
+							TBaseStructure<FRotator>::Get()->ExportText(Content, &Data, &Data, nullptr, PPF_None, nullptr);
+							break;
+						}
+						case ESlateTransformComponent::Scale:
+						{
+							const FVector Data = Transform.GetScale3D();
+							TBaseStructure<FVector>::Get()->ExportText(Content, &Data, &Data, nullptr, PPF_None, nullptr);
+							break;
+						}
+						case ESlateTransformComponent::Max:
+						default:
+						{
+							TBaseStructure<TransformType>::Get()->ExportText(Content, &Transform, &Transform, nullptr, PPF_None, nullptr);
+							break;
+						}
+					}
+
+					if(!Content.IsEmpty())
+					{
+						FPlatformApplicationMisc::ClipboardCopy(*Content);
+					}
+				}
+			}
+		});
+
+		WidgetArgs.OnPasteFromClipboard_Lambda([this, InPropertyHandle, OnNumericValueChanged](
+			ESlateTransformComponent::Type InComponent
+			)
+		{
+			FString Content;
+			FPlatformApplicationMisc::ClipboardPaste(Content);
+		
+			if(Content.IsEmpty())
+			{
+				return;
+			}
+
+			TOptional<FReal> Result;
+			FEditPropertyChain PropertyChain;
+			TArray<int32> PropertyArrayIndices;
+			bool bEnabled;
+			if (!GetPropertyChain(InPropertyHandle, PropertyChain, PropertyArrayIndices, bEnabled))
+			{
+				return;
+			}
+
+			TArray<UObject*> ObjectsView;
+			for(int32 Index = 0; Index < ObjectsBeingCustomized.Num(); Index++)
+			{
+				const TWeakObjectPtr<UObject>& Object = ObjectsBeingCustomized[Index];
+				if (Object.Get())
+				{
+					ObjectsView.Add(Object.Get());
+				}
+			}
+			FPropertyChangedEvent PropertyChangedEvent(InPropertyHandle->GetProperty(), EPropertyChangeType::ValueSet, ObjectsView);
+			FPropertyChangedChainEvent PropertyChangedChainEvent(PropertyChain, PropertyChangedEvent);
+
+			URigVMController* Controller = nullptr;
+			if(BlueprintBeingCustomized && GraphBeingCustomized)
+			{
+				Controller = BlueprintBeingCustomized->GetController(GraphBeingCustomized);
+				Controller->OpenUndoBracket(FString::Printf(TEXT("Set %s"), *InPropertyHandle->GetProperty()->GetName()));
+			}
+			
+			for(const TWeakObjectPtr<UObject>& Object : ObjectsBeingCustomized)
+			{
+				if(Object.Get() && InPropertyHandle->IsValidHandle())
+				{
+					TransformType& Transform = ContainerUObjectToValueRef<TransformType>(Object.Get(), Identity, PropertyChain, PropertyArrayIndices);
+					const TransformType PreviousTransform = Transform;
+
+					// Apply the new value
+					{
+						class FRigPasteTransformWidgetErrorPipe : public FOutputDevice
+						{
+						public:
+					
+							int32 NumErrors;
+					
+							FRigPasteTransformWidgetErrorPipe()
+								: FOutputDevice()
+								, NumErrors(0)
+							{
+							}
+					
+							virtual void Serialize(const TCHAR* V, ELogVerbosity::Type Verbosity, const class FName& Category) override
+							{
+								UE_LOG(LogRigVM, Error, TEXT("Error Pasting to Widget: %s"), V);
+								NumErrors++;
+							}
+						};
+					
+						FRigPasteTransformWidgetErrorPipe ErrorPipe;
+						
+						switch(InComponent)
+						{
+							case ESlateTransformComponent::Location:
+							{
+								FVector Data = Transform.GetLocation();
+								TBaseStructure<FVector>::Get()->ImportText(*Content, &Data, nullptr, PPF_None, &ErrorPipe, TBaseStructure<FVector>::Get()->GetName(), true);
+								Transform.SetLocation(Data);
+								break;
+							}
+							case ESlateTransformComponent::Rotation:
+							{
+								FRotator Data = Transform.Rotator();
+								TBaseStructure<FRotator>::Get()->ImportText(*Content, &Data, nullptr, PPF_None, &ErrorPipe, TBaseStructure<FRotator>::Get()->GetName(), true);
+								FQuat Quat = Data.Quaternion();
+								Transform.SetRotation(Quat);
+								
+								break;
+							}
+							case ESlateTransformComponent::Scale:
+							{
+								FVector Data = Transform.GetScale3D();
+								TBaseStructure<FVector>::Get()->ImportText(*Content, &Data, nullptr, PPF_None, &ErrorPipe, TBaseStructure<FVector>::Get()->GetName(), true);
+								Transform.SetScale3D(Data);
+								break;
+							}
+							case ESlateTransformComponent::Max:
+							default:
+							{
+								TBaseStructure<TransformType>::Get()->ImportText(*Content, &Transform, nullptr, PPF_None, &ErrorPipe, TBaseStructure<TransformType>::Get()->GetName(), true);
+								break;
+							}
+						}
+					
+						if(ErrorPipe.NumErrors == 0 && !PreviousTransform.Equals(Transform))
+						{
+							Object->PostEditChangeChainProperty(PropertyChangedChainEvent);
+							InPropertyHandle->NotifyPostChange(PropertyChangedEvent.ChangeType);
+						}
+					}
+				}
+			}
+			
+			if (Controller)
+			{
+				Controller->CloseUndoBracket();
+			}
+		});
+
 		SAdvancedTransformInputBox<TransformType>::ConstructGroupedTransformRows(
 			StructBuilder,
 			InPropertyHandle->GetPropertyDisplayName(),
