RigVM: Remove the notion of deprecated nodes

diff --git a/Engine/Plugins/Runtime/RigVM/Source/RigVMDeveloper/Private/EdGraph/RigVMEdGraphNode.cpp b/Engine/Plugins/Runtime/RigVM/Source/RigVMDeveloper/Private/EdGraph/RigVMEdGraphNode.cpp
--- a/Engine/Plugins/Runtime/RigVM/Source/RigVMDeveloper/Private/EdGraph/RigVMEdGraphNode.cpp
+++ b/Engine/Plugins/Runtime/RigVM/Source/RigVMDeveloper/Private/EdGraph/RigVMEdGraphNode.cpp
@@ -169,6 +169,11 @@ FText URigVMEdGraphNode::GetNodeTitle(ENodeTitleType::Type TitleType) const
 			NodeTitle = FText::FromString(FString::Printf(TEXT("%s (Deprecated)"), *NodeTitle.ToString()));
 		}
 
+		if(IsOutDated())
+		{
+			NodeTitle = FText::FromString(FString::Printf(TEXT("%s (OutDated)"), *NodeTitle.ToString()));
+		}
+
 		FullNodeTitle = NodeTitle;
 
 		if(!SubTitle.IsEmpty())
@@ -276,15 +281,19 @@ void URigVMEdGraphNode::ReconstructNode_Internal(bool bForce)
 }
 
 bool URigVMEdGraphNode::IsDeprecated() const
+{
+	// we longer consider nodes to be deprecated ever,
+	// instead we refer to them as outdated
+	return false;
+}
+
+bool URigVMEdGraphNode::IsOutDated() const
 {
 	if(URigVMNode* ModelNode = GetModelNode())
 	{
-		if(URigVMUnitNode* StructModelNode = Cast<URigVMUnitNode>(ModelNode))
-		{
-			return StructModelNode->IsDeprecated();
-		}
+		return ModelNode->IsOutDated();
 	}
-	return Super::IsDeprecated();
+	return false;
 }
 
 FEdGraphNodeDeprecationResponse URigVMEdGraphNode::GetDeprecationResponse(EEdGraphNodeDeprecationType DeprecationType) const
@@ -293,15 +302,12 @@ FEdGraphNodeDeprecationResponse URigVMEdGraphNode::GetDeprecationResponse(EEdGra
 
 	if(URigVMNode* ModelNode = GetModelNode())
 	{
-		if(URigVMUnitNode* StructModelNode = Cast<URigVMUnitNode>(ModelNode))
+		const FString DeprecatedMetadata = ModelNode->GetDeprecatedMetadata();
+		if (!DeprecatedMetadata.IsEmpty())
 		{
-			FString DeprecatedMetadata = StructModelNode->GetDeprecatedMetadata();
-			if (!DeprecatedMetadata.IsEmpty())
-			{
-				FFormatNamedArguments Args;
-				Args.Add(TEXT("DeprecatedMetadata"), FText::FromString(DeprecatedMetadata));
-				Response.MessageText = FText::Format(LOCTEXT("RigVMEdGraphNodeDeprecationMessage", "Warning: This node is deprecated from: {DeprecatedMetadata}"), Args);
-			}
+			FFormatNamedArguments Args;
+			Args.Add(TEXT("DeprecatedMetadata"), FText::FromString(DeprecatedMetadata));
+			Response.MessageText = FText::Format(LOCTEXT("RigVMEdGraphNodeDeprecationMessage", "Warning: This node is deprecated from: {DeprecatedMetadata}"), Args);
 		}
 	}
 
@@ -375,7 +381,12 @@ void URigVMEdGraphNode::PostReconstructNode()
 
 	bCanRenameNode = false;
 
-	if(URigVMNode* ModelNode = GetModelNode())
+	if(IsOutDated())
+	{
+		static const FLinearColor WarningColor = FAppStyle::GetColor("ErrorReporting.WarningBackgroundColor");
+		SetColorFromModel(WarningColor);
+	}
+	else if(URigVMNode* ModelNode = GetModelNode())
 	{
 		SetColorFromModel(ModelNode->GetNodeColor());
 	}
@@ -1666,7 +1677,14 @@ FText URigVMEdGraphNode::GetTooltipText() const
 {
 	if(URigVMNode* ModelNode = GetModelNode())
 	{
-		return ModelNode->GetToolTipText();
+		FText Tooltip = ModelNode->GetToolTipText();
+		if(IsOutDated())
+		{
+			return FText::Format(
+				LOCTEXT("OutDatedTooltipFormat", "This node is outdated and can potentially be upgraded.\nFor this right click and choose 'Upgrade Nodes'.\n\n{0}"),
+					Tooltip);
+		}
+		return Tooltip;;
 	}
 	return FText::FromString(ModelNodePath);
 }
diff --git a/Engine/Plugins/Runtime/RigVM/Source/RigVMDeveloper/Private/RigVMModel/Nodes/RigVMDispatchNode.cpp b/Engine/Plugins/Runtime/RigVM/Source/RigVMDeveloper/Private/RigVMModel/Nodes/RigVMDispatchNode.cpp
--- a/Engine/Plugins/Runtime/RigVM/Source/RigVMDeveloper/Private/RigVMModel/Nodes/RigVMDispatchNode.cpp
+++ b/Engine/Plugins/Runtime/RigVM/Source/RigVMDeveloper/Private/RigVMModel/Nodes/RigVMDispatchNode.cpp
@@ -78,7 +78,7 @@ FText URigVMDispatchNode::GetToolTipTextForPin(const URigVMPin* InPin) const
 	return URigVMNode::GetToolTipTextForPin(InPin);
 }
 
-bool URigVMDispatchNode::IsDeprecated() const
+bool URigVMDispatchNode::IsOutDated() const
 {
 	return !GetDeprecatedMetadata().IsEmpty();
 }
diff --git a/Engine/Plugins/Runtime/RigVM/Source/RigVMDeveloper/Private/RigVMModel/Nodes/RigVMUnitNode.cpp b/Engine/Plugins/Runtime/RigVM/Source/RigVMDeveloper/Private/RigVMModel/Nodes/RigVMUnitNode.cpp
--- a/Engine/Plugins/Runtime/RigVM/Source/RigVMDeveloper/Private/RigVMModel/Nodes/RigVMUnitNode.cpp
+++ b/Engine/Plugins/Runtime/RigVM/Source/RigVMDeveloper/Private/RigVMModel/Nodes/RigVMUnitNode.cpp
@@ -17,7 +17,7 @@ void URigVMUnitNode::PostLoad()
 	// if we have a script struct but no notation let's figure out the template
 	if(GetScriptStruct() != nullptr)
 	{
-		if (IsDeprecated())
+		if (IsOutDated())
 		{
 			TemplateNotation = NAME_None;
 			if(const FRigVMFunction* Function = FRigVMRegistry::Get().FindFunction(GetScriptStruct(), *GetMethodName().ToString()))
@@ -175,7 +175,7 @@ bool URigVMUnitNode::ShouldInputPinComputeLazily(const URigVMPin* InPin) const
 	return Super::ShouldInputPinComputeLazily(InPin);
 }
 
-bool URigVMUnitNode::IsDeprecated() const
+bool URigVMUnitNode::IsOutDated() const
 {
 	return !GetDeprecatedMetadata().IsEmpty();
 }
diff --git a/Engine/Plugins/Runtime/RigVM/Source/RigVMDeveloper/Private/RigVMModel/RigVMController.cpp b/Engine/Plugins/Runtime/RigVM/Source/RigVMDeveloper/Private/RigVMModel/RigVMController.cpp
--- a/Engine/Plugins/Runtime/RigVM/Source/RigVMDeveloper/Private/RigVMModel/RigVMController.cpp
+++ b/Engine/Plugins/Runtime/RigVM/Source/RigVMDeveloper/Private/RigVMModel/RigVMController.cpp
@@ -2559,7 +2559,7 @@ TArray<URigVMNode*> URigVMController::UpgradeNodes(const TArray<FName>& InNodeNa
 	{
 		if(URigVMUnitNode* UnitNode = Cast<URigVMUnitNode>(Node))
 		{
-			if(UnitNode->IsDeprecated())
+			if(UnitNode->IsOutDated())
 			{
 				ReportWarningf(TEXT("Node %s cannot be upgraded. There is no automatic upgrade path available."), *UnitNode->GetNodePath());
 			}
diff --git a/Engine/Plugins/Runtime/RigVM/Source/RigVMDeveloper/Public/EdGraph/RigVMEdGraphNode.h b/Engine/Plugins/Runtime/RigVM/Source/RigVMDeveloper/Public/EdGraph/RigVMEdGraphNode.h
--- a/Engine/Plugins/Runtime/RigVM/Source/RigVMDeveloper/Public/EdGraph/RigVMEdGraphNode.h
+++ b/Engine/Plugins/Runtime/RigVM/Source/RigVMDeveloper/Public/EdGraph/RigVMEdGraphNode.h
@@ -118,6 +118,7 @@ public:
 	virtual void BeginDestroy() override;
 
 	virtual bool IsDeprecated() const override;
+	bool IsOutDated() const;
 	virtual FEdGraphNodeDeprecationResponse GetDeprecationResponse(EEdGraphNodeDeprecationType DeprecationType) const override;
 
 	/** Set the cached dimensions of this node */
diff --git a/Engine/Plugins/Runtime/RigVM/Source/RigVMDeveloper/Public/RigVMModel/Nodes/RigVMDispatchNode.h b/Engine/Plugins/Runtime/RigVM/Source/RigVMDeveloper/Public/RigVMModel/Nodes/RigVMDispatchNode.h
--- a/Engine/Plugins/Runtime/RigVM/Source/RigVMDeveloper/Public/RigVMModel/Nodes/RigVMDispatchNode.h
+++ b/Engine/Plugins/Runtime/RigVM/Source/RigVMDeveloper/Public/RigVMModel/Nodes/RigVMDispatchNode.h
@@ -47,8 +47,8 @@ public:
 	virtual FRigVMStructUpgradeInfo GetUpgradeInfo() const override;
 
 	const FRigVMDispatchFactory* GetFactory() const;
-	bool IsDeprecated() const;
-	FString GetDeprecatedMetadata() const;
+	virtual bool IsOutDated() const override;
+	virtual FString GetDeprecatedMetadata() const override;
 
 	FRigVMDispatchContext GetDispatchContext() const;
 
diff --git a/Engine/Plugins/Runtime/RigVM/Source/RigVMDeveloper/Public/RigVMModel/Nodes/RigVMUnitNode.h b/Engine/Plugins/Runtime/RigVM/Source/RigVMDeveloper/Public/RigVMModel/Nodes/RigVMUnitNode.h
--- a/Engine/Plugins/Runtime/RigVM/Source/RigVMDeveloper/Public/RigVMModel/Nodes/RigVMUnitNode.h
+++ b/Engine/Plugins/Runtime/RigVM/Source/RigVMDeveloper/Public/RigVMModel/Nodes/RigVMUnitNode.h
@@ -38,9 +38,8 @@ public:
 	virtual TArray<URigVMPin*> GetAggregateOutputs() const override;
 	virtual FName GetNextAggregateName(const FName& InLastAggregatePinName) const override;
 
-	bool IsDeprecated() const;
-	FString GetDeprecatedMetadata() const;
-
+	virtual bool IsOutDated() const override;
+	virtual FString GetDeprecatedMetadata() const override;
 
 	// URigVMTemplateNode interface
 	virtual UScriptStruct* GetScriptStruct() const override;
diff --git a/Engine/Plugins/Runtime/RigVM/Source/RigVMDeveloper/Public/RigVMModel/RigVMNode.h b/Engine/Plugins/Runtime/RigVM/Source/RigVMDeveloper/Public/RigVMModel/RigVMNode.h
--- a/Engine/Plugins/Runtime/RigVM/Source/RigVMDeveloper/Public/RigVMModel/RigVMNode.h
+++ b/Engine/Plugins/Runtime/RigVM/Source/RigVMDeveloper/Public/RigVMModel/RigVMNode.h
@@ -312,6 +312,16 @@ public:
 
 	TSharedPtr<FStructOnScope> GetDecoratorInstance(const URigVMPin* InDecoratorPin, bool bUseDefaultValueFromPin = true) const;
 
+	virtual bool IsOutDated() const
+	{
+		return false;
+	}
+	
+	virtual FString GetDeprecatedMetadata() const
+	{
+		return FString();
+	}
+
 private:
 
 	static const FString NodeColorName;
