RigVM: Update TypeActions related to public functions when blueprint is automatically saved or is deleted
#rnx

diff --git a/Engine/Plugins/Runtime/RigVM/Source/RigVMDeveloper/Private/RigVMBlueprint.cpp b/Engine/Plugins/Runtime/RigVM/Source/RigVMDeveloper/Private/RigVMBlueprint.cpp
--- a/Engine/Plugins/Runtime/RigVM/Source/RigVMDeveloper/Private/RigVMBlueprint.cpp
+++ b/Engine/Plugins/Runtime/RigVM/Source/RigVMDeveloper/Private/RigVMBlueprint.cpp
@@ -837,6 +837,16 @@ void URigVMBlueprint::PreSave(FObjectPreSaveContext ObjectSaveContext)
 	}
 }
 
+void URigVMBlueprint::PostSaveRoot(FObjectPostSaveRootContext ObjectSaveContext)
+{
+	Super::PostSaveRoot(ObjectSaveContext);
+
+	// Make sure all the tags are accounted for in the TypeActions after we save
+	FBlueprintActionDatabase& ActionDatabase = FBlueprintActionDatabase::Get();
+	ActionDatabase.ClearAssetActions(GetClass());
+	ActionDatabase.RefreshClassActions(GetClass());
+}
+
 void URigVMBlueprint::PostLoad()
 {
 	Super::PostLoad();
diff --git a/Engine/Plugins/Runtime/RigVM/Source/RigVMDeveloper/Private/RigVMBlueprintUtils.cpp b/Engine/Plugins/Runtime/RigVM/Source/RigVMDeveloper/Private/RigVMBlueprintUtils.cpp
--- a/Engine/Plugins/Runtime/RigVM/Source/RigVMDeveloper/Private/RigVMBlueprintUtils.cpp
+++ b/Engine/Plugins/Runtime/RigVM/Source/RigVMDeveloper/Private/RigVMBlueprintUtils.cpp
@@ -1,6 +1,8 @@
 // Copyright Epic Games, Inc. All Rights Reserved.
 
 #include "RigVMBlueprintUtils.h"
+
+#include "BlueprintActionDatabase.h"
 #include "Kismet2/BlueprintEditorUtils.h"
 #include "RigVMCore/RigVMStruct.h"
 #include "UObject/UObjectIterator.h"
@@ -137,6 +139,17 @@ void FRigVMBlueprintUtils::HandleRefreshAllNodes(UBlueprint* InBlueprint)
 #endif
 }
 
+void FRigVMBlueprintUtils::HandleAssetDeleted(const FAssetData& InAssetData)
+{
+	if (InAssetData.GetClass()->IsChildOf(URigVMBlueprint::StaticClass()))
+	{
+		// Make sure any RigVMBlueprint removes any TypeActions related to this asset (e.g. public functions)
+		FBlueprintActionDatabase& ActionDatabase = FBlueprintActionDatabase::Get();
+		ActionDatabase.ClearAssetActions(InAssetData.GetClass());
+		ActionDatabase.RefreshClassActions(InAssetData.GetClass());
+	}
+}
+
 void FRigVMBlueprintUtils::RemoveMemberVariableIfNotUsed(UBlueprint* Blueprint, const FName VarName)
 {
 	DECLARE_SCOPE_HIERARCHICAL_COUNTER_FUNC()
diff --git a/Engine/Plugins/Runtime/RigVM/Source/RigVMDeveloper/Public/RigVMBlueprint.h b/Engine/Plugins/Runtime/RigVM/Source/RigVMDeveloper/Public/RigVMBlueprint.h
--- a/Engine/Plugins/Runtime/RigVM/Source/RigVMDeveloper/Public/RigVMBlueprint.h
+++ b/Engine/Plugins/Runtime/RigVM/Source/RigVMDeveloper/Public/RigVMBlueprint.h
@@ -274,6 +274,7 @@ public:
 	virtual void PreSave(const class ITargetPlatform* TargetPlatform) override;
 	PRAGMA_ENABLE_DEPRECATION_WARNINGS
 	virtual void PreSave(FObjectPreSaveContext ObjectSaveContext) override;
+	virtual void PostSaveRoot(FObjectPostSaveRootContext ObjectSaveContext) override;
 	virtual void PostLoad() override;
 #if WITH_EDITORONLY_DATA
 	static void DeclareConstructClasses(TArray<FTopLevelAssetPath>& OutConstructClasses, const UClass* SpecificSubclass);
diff --git a/Engine/Plugins/Runtime/RigVM/Source/RigVMDeveloper/Public/RigVMBlueprintUtils.h b/Engine/Plugins/Runtime/RigVM/Source/RigVMDeveloper/Public/RigVMBlueprintUtils.h
--- a/Engine/Plugins/Runtime/RigVM/Source/RigVMDeveloper/Public/RigVMBlueprintUtils.h
+++ b/Engine/Plugins/Runtime/RigVM/Source/RigVMDeveloper/Public/RigVMBlueprintUtils.h
@@ -23,6 +23,9 @@ static void HandleReconstructAllNodes(UBlueprint* InBlueprint);
 /** Handle blueprint node refresh */
 static void HandleRefreshAllNodes(UBlueprint* InBlueprint);
 
+/** Handle blueprint deleted */
+static void HandleAssetDeleted(const FAssetData& InAssetData);
+
 /** remove the variable if not used by anybody else but ToBeDeleted*/
 static void RemoveMemberVariableIfNotUsed(UBlueprint* Blueprint, const FName VarName);
 
diff --git a/Engine/Plugins/Runtime/RigVM/Source/RigVMEditor/Private/Editor/RigVMEditor.cpp b/Engine/Plugins/Runtime/RigVM/Source/RigVMEditor/Private/Editor/RigVMEditor.cpp
--- a/Engine/Plugins/Runtime/RigVM/Source/RigVMEditor/Private/Editor/RigVMEditor.cpp
+++ b/Engine/Plugins/Runtime/RigVM/Source/RigVMEditor/Private/Editor/RigVMEditor.cpp
@@ -1125,20 +1125,12 @@ void FRigVMEditor::SaveAsset_Execute()
 {
 	LastDebuggedHost = GetCustomDebugObjectLabel(GetBlueprintObj()->GetObjectBeingDebugged());
 	FBlueprintEditor::SaveAsset_Execute();
-
-	FBlueprintActionDatabase& ActionDatabase = FBlueprintActionDatabase::Get();
-	ActionDatabase.ClearAssetActions(URigVMBlueprint::StaticClass());
-	ActionDatabase.RefreshClassActions(URigVMBlueprint::StaticClass());
 }
 
 void FRigVMEditor::SaveAssetAs_Execute()
 {
 	LastDebuggedHost = GetCustomDebugObjectLabel(GetBlueprintObj()->GetObjectBeingDebugged());
 	FBlueprintEditor::SaveAssetAs_Execute();
-
-	FBlueprintActionDatabase& ActionDatabase = FBlueprintActionDatabase::Get();
-	ActionDatabase.ClearAssetActions(URigVMBlueprint::StaticClass());
-	ActionDatabase.RefreshClassActions(URigVMBlueprint::StaticClass());
 }
 
 bool FRigVMEditor::IsEditable(UEdGraph* InGraph) const
diff --git a/Engine/Plugins/Runtime/RigVM/Source/RigVMEditor/Private/RigVMEditorModule.cpp b/Engine/Plugins/Runtime/RigVM/Source/RigVMEditor/Private/RigVMEditorModule.cpp
--- a/Engine/Plugins/Runtime/RigVM/Source/RigVMEditor/Private/RigVMEditorModule.cpp
+++ b/Engine/Plugins/Runtime/RigVM/Source/RigVMEditor/Private/RigVMEditorModule.cpp
@@ -69,6 +69,12 @@ void FRigVMEditorModule::StartupModule()
 		// Register Blueprint editor variable customization
 		FBlueprintEditorModule& BlueprintEditorModule = FModuleManager::LoadModuleChecked<FBlueprintEditorModule>("Kismet");
 		BlueprintVariableCustomizationHandle = BlueprintEditorModule.RegisterVariableCustomization(FProperty::StaticClass(), FOnGetVariableCustomizationInstance::CreateStatic(&FRigVMVariableDetailCustomization::MakeInstance));
+
+		IAssetRegistry* AssetRegistry = IAssetRegistry::Get();
+		if (ensure(AssetRegistry))
+		{
+			AssetRegistry->OnAssetRemoved().AddStatic(&FRigVMBlueprintUtils::HandleAssetDeleted);
+		}
 	}
 
 	StartupModuleCommon();
