RigVM: Fix crash when deleting public function used by other currently running rigs. Convert UObject pointers from ByteCode to weak pointers.
#rnx

diff --git a/Engine/Plugins/Runtime/RigVM/Source/RigVM/Private/RigVMCore/RigVM.cpp b/Engine/Plugins/Runtime/RigVM/Source/RigVM/Private/RigVMCore/RigVM.cpp
--- a/Engine/Plugins/Runtime/RigVM/Source/RigVM/Private/RigVMCore/RigVM.cpp
+++ b/Engine/Plugins/Runtime/RigVM/Source/RigVM/Private/RigVMCore/RigVM.cpp
@@ -1209,12 +1209,13 @@ bool URigVM::ShouldHaltAtInstruction(FRigVMExtendedExecuteContext& Context, cons
 						RigVMDebugInfo->SetCurrentActiveBreakpoint(Breakpoint);
 						
 						// We want to keep the callstack up to the node that produced the halt
-						const TArray<UObject*>* FullCallstack = ByteCode.GetCallstackForInstruction(ContextPublicData.InstructionIndex);
+						const TArray<TWeakObjectPtr<UObject>>* FullCallstack = ByteCode.GetCallstackForInstruction(ContextPublicData.InstructionIndex);
 						if (FullCallstack)
 						{
-							RigVMDebugInfo->SetCurrentActiveBreakpointCallstack(TArray<UObject*>(FullCallstack->GetData(), FullCallstack->Find((UObject*)Breakpoint.Subject)+1));
+							RigVMDebugInfo->SetCurrentActiveBreakpointCallstack(TArray<TWeakObjectPtr<UObject>>(FullCallstack->GetData(), FullCallstack->Find((TWeakObjectPtr<UObject>)Breakpoint.Subject)+1));
 						}
-						RigVMDebugInfo->ExecutionHalted().Broadcast(ContextPublicData.InstructionIndex, Breakpoint.Subject, InEventName);
+						UObject* BreakpointNode = Breakpoint.Subject.IsValid() ? Breakpoint.Subject.Get() : nullptr;
+						RigVMDebugInfo->ExecutionHalted().Broadcast(ContextPublicData.InstructionIndex, BreakpointNode, InEventName);
 					}
 					return true;
 				}
@@ -1241,14 +1242,14 @@ bool URigVM::ShouldHaltAtInstruction(FRigVMExtendedExecuteContext& Context, cons
 					if (!RigVMDebugInfo->GetCurrentActiveBreakpoint())
 					{
 						RigVMDebugInfo->SetCurrentActiveBreakpoint(Breakpoint);
-						const TArray<UObject*>* FullCallstack = ByteCode.GetCallstackForInstruction(ContextPublicData.InstructionIndex);
+						const TArray<TWeakObjectPtr<UObject>>* FullCallstack = ByteCode.GetCallstackForInstruction(ContextPublicData.InstructionIndex);
 						
 						// We want to keep the callstack up to the node that produced the halt
 						if (FullCallstack)
 						{
-							RigVMDebugInfo->SetCurrentActiveBreakpointCallstack(TArray<UObject*>(FullCallstack->GetData(), FullCallstack->Find((UObject*)RigVMDebugInfo->GetCurrentActiveBreakpoint().Subject)+1));
+							RigVMDebugInfo->SetCurrentActiveBreakpointCallstack(TArray<TWeakObjectPtr<UObject>>(FullCallstack->GetData(), FullCallstack->Find((TWeakObjectPtr<UObject>)RigVMDebugInfo->GetCurrentActiveBreakpoint().Subject)+1));
 						}
-					}							
+					}
 					
 					break;	
 				}
@@ -1268,14 +1269,14 @@ bool URigVM::ShouldHaltAtInstruction(FRigVMExtendedExecuteContext& Context, cons
 	// If we are stepping, and the last active breakpoint was set, check if this is the new temporary breakpoint
 	if (RigVMDebugInfo->GetCurrentBreakpointAction() != ERigVMBreakpointAction::None && RigVMDebugInfo->GetCurrentActiveBreakpoint())
 	{
-		const TArray<UObject*>* CurrentCallstack = ByteCode.GetCallstackForInstruction(ContextPublicData.InstructionIndex);
+		const TArray<TWeakObjectPtr<UObject>>* CurrentCallstack = ByteCode.GetCallstackForInstruction(ContextPublicData.InstructionIndex);
 		if (CurrentCallstack && !CurrentCallstack->IsEmpty())
 		{
 			UObject* NewBreakpointNode = nullptr;
 
 			// Find the first difference in the callstack
 			int32 DifferenceIndex = INDEX_NONE;
-			TArray<UObject*>& PreviousCallstack = RigVMDebugInfo->GetCurrentActiveBreakpointCallstack();
+			TArray<TWeakObjectPtr<UObject>>& PreviousCallstack = RigVMDebugInfo->GetCurrentActiveBreakpointCallstack();
 			for (int32 i=0; i<PreviousCallstack.Num(); ++i)
 			{
 				if (CurrentCallstack->Num() == i)
@@ -1294,7 +1295,10 @@ bool URigVM::ShouldHaltAtInstruction(FRigVMExtendedExecuteContext& Context, cons
 			{
 				if (DifferenceIndex != INDEX_NONE)
 				{
-					NewBreakpointNode = CurrentCallstack->operator[](DifferenceIndex);
+					if ((*CurrentCallstack)[DifferenceIndex].IsValid())
+					{
+						NewBreakpointNode = (*CurrentCallstack)[DifferenceIndex].Get();
+					}
 				}
 			}
 			else if (RigVMDebugInfo->GetCurrentBreakpointAction() == ERigVMBreakpointAction::StepInto)
@@ -1303,19 +1307,29 @@ bool URigVM::ShouldHaltAtInstruction(FRigVMExtendedExecuteContext& Context, cons
 				{
 					if (!CurrentCallstack->IsEmpty() && !PreviousCallstack.IsEmpty() && CurrentCallstack->Last() != PreviousCallstack.Last())
 					{
-						NewBreakpointNode = CurrentCallstack->operator[](FMath::Min(PreviousCallstack.Num(), CurrentCallstack->Num()-1));
+						const int32 MinIndex = FMath::Min(PreviousCallstack.Num(), CurrentCallstack->Num()-1);
+						if ((*CurrentCallstack)[MinIndex].IsValid())
+						{
+							NewBreakpointNode = (*CurrentCallstack)[MinIndex].Get();
+						}
 					}
 				}
 				else
 				{
-					NewBreakpointNode = CurrentCallstack->operator[](DifferenceIndex);
+					if ((*CurrentCallstack)[DifferenceIndex].IsValid())
+					{
+						NewBreakpointNode = (*CurrentCallstack)[DifferenceIndex].Get();
+					}
 				}
 			}
 			else if (RigVMDebugInfo->GetCurrentBreakpointAction() == ERigVMBreakpointAction::StepOut)
 			{
 				if (DifferenceIndex != INDEX_NONE && DifferenceIndex <= PreviousCallstack.Num() - 2)
                 {
-                	NewBreakpointNode = CurrentCallstack->operator[](DifferenceIndex);
+					if ((*CurrentCallstack)[DifferenceIndex].IsValid())
+					{
+						NewBreakpointNode = (*CurrentCallstack)[DifferenceIndex].Get();
+					}
                 }
 			}
 			
diff --git a/Engine/Plugins/Runtime/RigVM/Source/RigVM/Private/RigVMCore/RigVMByteCode.cpp b/Engine/Plugins/Runtime/RigVM/Source/RigVM/Private/RigVMCore/RigVMByteCode.cpp
--- a/Engine/Plugins/Runtime/RigVM/Source/RigVM/Private/RigVMCore/RigVMByteCode.cpp
+++ b/Engine/Plugins/Runtime/RigVM/Source/RigVM/Private/RigVMCore/RigVMByteCode.cpp
@@ -1699,7 +1699,10 @@ UObject* FRigVMByteCode::GetSubjectForInstruction(int32 InInstructionIndex) cons
 {
 	if (SubjectPerInstruction.IsValidIndex(InInstructionIndex))
 	{
-		return SubjectPerInstruction[InInstructionIndex];
+		if (SubjectPerInstruction[InInstructionIndex].IsValid())
+		{
+			return SubjectPerInstruction[InInstructionIndex].Get();
+		}
 	}
 	return nullptr;
 }
@@ -1781,7 +1784,7 @@ TArray<int32> FRigVMByteCode::GetAllInstructionIndicesForCallPath(const FString&
 	return MatchedInstructions;
 }
 
-int32 FRigVMByteCode::GetFirstInstructionIndexForCallstack(const TArray<UObject*>& InCallstack) const
+int32 FRigVMByteCode::GetFirstInstructionIndexForCallstack(const TArray<TWeakObjectPtr<UObject>>& InCallstack) const
 {
 	const TArray<int32>& InstructionIndices = GetAllInstructionIndicesForCallstack(InCallstack);
 	if (InstructionIndices.Num() > 0)
@@ -1791,7 +1794,7 @@ int32 FRigVMByteCode::GetFirstInstructionIndexForCallstack(const TArray<UObject*
 	return INDEX_NONE;
 }
 
-const TArray<int32>& FRigVMByteCode::GetAllInstructionIndicesForCallstack(const TArray<UObject*>& InCallstack) const
+const TArray<int32>& FRigVMByteCode::GetAllInstructionIndicesForCallstack(const TArray<TWeakObjectPtr<UObject>>& InCallstack) const
 {
 	if(InCallstack.IsEmpty())
 	{
@@ -1807,9 +1810,9 @@ const TArray<int32>& FRigVMByteCode::GetAllInstructionIndicesForCallstack(const
 	return EmptyInstructionIndices;
 }
 
-void FRigVMByteCode::SetSubject(int32 InInstructionIndex, const FString& InCallPath, const TArray<UObject*>& InCallstack)
+void FRigVMByteCode::SetSubject(int32 InInstructionIndex, const FString& InCallPath, const TArray<TWeakObjectPtr<UObject>>& InCallstack)
 {
-	UObject* Subject = InCallstack.Last();
+	TWeakObjectPtr<UObject> Subject = InCallstack.Last();
 	if (SubjectPerInstruction.Num() <= InInstructionIndex)
 	{
 		SubjectPerInstruction.AddZeroed(1 + InInstructionIndex - SubjectPerInstruction.Num());
@@ -1838,14 +1841,14 @@ void FRigVMByteCode::SetSubject(int32 InInstructionIndex, const FString& InCallP
 
 	for(int32 CallstackLength = InCallstack.Num(); CallstackLength > 0; CallstackLength--)
 	{
-		UObject* const* DataPtr = &InCallstack[InCallstack.Num() - CallstackLength];
-		TArrayView<UObject* const> View(DataPtr, CallstackLength);
+		TWeakObjectPtr<UObject> const* DataPtr = &InCallstack[InCallstack.Num() - CallstackLength];
+		TArrayView<TWeakObjectPtr<UObject> const> View(DataPtr, CallstackLength);
 		uint32 Hash = GetCallstackHash(View);
 		CallstackHashToInstructions.FindOrAdd(Hash).AddUnique(InInstructionIndex);
 	}
 }
 
-const TArray<UObject*>* FRigVMByteCode::GetCallstackForInstruction(int32 InInstructionIndex) const
+const TArray<TWeakObjectPtr<UObject>>* FRigVMByteCode::GetCallstackForInstruction(int32 InInstructionIndex) const
 {
 	if (CallstackPerInstruction.IsValidIndex(InInstructionIndex))
 	{
@@ -1863,21 +1866,21 @@ uint32 FRigVMByteCode::GetCallstackHashForInstruction(int32 InInstructionIndex)
 	return 0;
 }
 
-uint32 FRigVMByteCode::GetCallstackHash(const TArray<UObject*>& InCallstack)
+uint32 FRigVMByteCode::GetCallstackHash(const TArray<TWeakObjectPtr<UObject>>& InCallstack)
 {
-	UObject*const * DataPtr = nullptr;
+	TWeakObjectPtr<UObject> const * DataPtr = nullptr;
 	if(InCallstack.Num() > 0)
 	{
 		DataPtr = &InCallstack[0];
 	}
-	TArrayView<UObject* const> View(DataPtr, InCallstack.Num());
+	TArrayView<TWeakObjectPtr<UObject> const> View(DataPtr, InCallstack.Num());
 	return GetCallstackHash(View);
 }
 
-uint32 FRigVMByteCode::GetCallstackHash(const TArrayView<UObject* const>& InCallstack)
+uint32 FRigVMByteCode::GetCallstackHash(const TArrayView<TWeakObjectPtr<UObject> const>& InCallstack)
 {
 	uint32 Hash = GetTypeHash(InCallstack.Num());
-	for(const UObject* Object : InCallstack)
+	for(const TWeakObjectPtr<UObject> Object : InCallstack)
 	{
 		Hash = HashCombine(Hash, GetTypeHash(Object));
 	}
diff --git a/Engine/Plugins/Runtime/RigVM/Source/RigVM/Private/RigVMCore/RigVMDebugInfo.cpp b/Engine/Plugins/Runtime/RigVM/Source/RigVM/Private/RigVMCore/RigVMDebugInfo.cpp
--- a/Engine/Plugins/Runtime/RigVM/Source/RigVM/Private/RigVMCore/RigVMDebugInfo.cpp
+++ b/Engine/Plugins/Runtime/RigVM/Source/RigVM/Private/RigVMCore/RigVMDebugInfo.cpp
@@ -79,7 +79,7 @@ const FRigVMBreakpoint& FRigVMDebugInfo::FindBreakpoint(const FGuid& InGuid) con
 	return EmptyBreakpoint;
 }
 
-const FRigVMBreakpoint& FRigVMDebugInfo::AddBreakpoint(const uint16 InstructionIndex, UObject* InNode, const uint16 InDepth,
+const FRigVMBreakpoint& FRigVMDebugInfo::AddBreakpoint(const uint16 InstructionIndex, TWeakObjectPtr<UObject> InNode, const uint16 InDepth,
                                                             const bool bIsTemporary)
 {
 	static const FRigVMBreakpoint EmptyBreakpoint;
diff --git a/Engine/Plugins/Runtime/RigVM/Source/RigVM/Public/RigVMCore/RigVMByteCode.h b/Engine/Plugins/Runtime/RigVM/Source/RigVM/Public/RigVMCore/RigVMByteCode.h
--- a/Engine/Plugins/Runtime/RigVM/Source/RigVM/Public/RigVMCore/RigVMByteCode.h
+++ b/Engine/Plugins/Runtime/RigVM/Source/RigVM/Public/RigVMCore/RigVMByteCode.h
@@ -1195,20 +1195,20 @@ public:
 	TArray<int32> GetAllInstructionIndicesForCallPath(const FString& InCallPath, bool bStartsWith = false, bool bEndsWith = false) const;
 
 	// returns the first hit instruction index for a given callpath (or INDEX_NONE)
-	int32 GetFirstInstructionIndexForCallstack(const TArray<UObject*>& InCallstack) const;
+	int32 GetFirstInstructionIndexForCallstack(const TArray<TWeakObjectPtr<UObject>>& InCallstack) const;
 
 	// returns all found instruction indices for a given callpath
-	const TArray<int32>& GetAllInstructionIndicesForCallstack(const TArray<UObject*>& InCallstack) const;
+	const TArray<int32>& GetAllInstructionIndicesForCallstack(const TArray<TWeakObjectPtr<UObject>>& InCallstack) const;
 
 	// returns the callstack which was used to inject a given instruction
-	const TArray<UObject*>* GetCallstackForInstruction(int32 InInstructionIndex) const;
+	const TArray<TWeakObjectPtr<UObject>>* GetCallstackForInstruction(int32 InInstructionIndex) const;
 
 	// returns the callstack hash which was used to inject a given instruction
 	uint32 GetCallstackHashForInstruction(int32 InInstructionIndex) const;
 
 	// computes a hash for a given callstack
-	static uint32 GetCallstackHash(const TArray<UObject*>& InCallstack);
-	static uint32 GetCallstackHash(const TArrayView<UObject* const>& InCallstack);
+	static uint32 GetCallstackHash(const TArray<TWeakObjectPtr<UObject>>& InCallstack);
+	static uint32 GetCallstackHash(const TArrayView<TWeakObjectPtr<UObject> const>& InCallstack);
 
 	// returns the input operands of a given instruction
 	FRigVMOperandArray GetInputOperands(int32 InInstructionIndex) const
@@ -1270,11 +1270,11 @@ private:
 	
 #if WITH_EDITORONLY_DATA
 
-	TArray<UObject*> SubjectPerInstruction;
-	TMap<UObject*, TArray<int32>> SubjectToInstructions;
+	TArray<TWeakObjectPtr<UObject>> SubjectPerInstruction;
+	TMap<TWeakObjectPtr<UObject>, TArray<int32>> SubjectToInstructions;
 	TArray<FString> CallPathPerInstruction;
 	TMap<FString, TArray<int32>> CallPathToInstructions;
-	TArray<TArray<UObject*>> CallstackPerInstruction;
+	TArray<TArray<TWeakObjectPtr<UObject>>> CallstackPerInstruction;
 	TMap<uint32, TArray<int32>> CallstackHashToInstructions;
 	TArray<uint32> CallstackHashPerInstruction;
 	TArray<TArray<FRigVMOperand>> InputOperandsPerInstruction;
@@ -1284,7 +1284,7 @@ private:
 
 #if WITH_EDITOR
 
-	void SetSubject(int32 InInstructionIndex, const FString& InCallPath, const TArray<UObject*>& InCallstack);
+	void SetSubject(int32 InInstructionIndex, const FString& InCallPath, const TArray<TWeakObjectPtr<UObject>>& InCallstack);
 
 #endif
 
diff --git a/Engine/Plugins/Runtime/RigVM/Source/RigVM/Public/RigVMCore/RigVMDebugInfo.h b/Engine/Plugins/Runtime/RigVM/Source/RigVM/Public/RigVMCore/RigVMDebugInfo.h
--- a/Engine/Plugins/Runtime/RigVM/Source/RigVM/Public/RigVMCore/RigVMDebugInfo.h
+++ b/Engine/Plugins/Runtime/RigVM/Source/RigVM/Public/RigVMCore/RigVMDebugInfo.h
@@ -40,7 +40,7 @@ struct RIGVM_API FRigVMBreakpoint
 		Guid.Invalidate();
 	}
 	
-	FRigVMBreakpoint(const uint16 InInstructionIndex, UObject* InNode, const uint16 InDepth)
+	FRigVMBreakpoint(const uint16 InInstructionIndex, TWeakObjectPtr<UObject> InNode, const uint16 InDepth)
 	: bIsActive(true)
 	, Guid(FGuid::NewGuid())
 	, InstructionIndex(InInstructionIndex)
@@ -79,7 +79,7 @@ struct RIGVM_API FRigVMBreakpoint
 	uint16 InstructionIndex;
 
 	// Node related to the breakpoint
-	UObject* Subject; 
+	TWeakObjectPtr<UObject> Subject; 
 
 	// The same instruction might be related to multiple breakpoints at different callstack depths
 	uint16 Depth; 
@@ -138,7 +138,7 @@ struct RIGVM_API FRigVMDebugInfo
 	TArray<FRigVMBreakpoint> FindBreakpointsAtInstruction(const uint16 InInstructionIndex) const;
 	const FRigVMBreakpoint& FindBreakpoint(const FGuid& InGuid) const;
 	
-	const FRigVMBreakpoint& AddBreakpoint(const uint16 InstructionIndex, UObject* InNode, const uint16 InDepth, const bool bIsTemporary = false);
+	const FRigVMBreakpoint& AddBreakpoint(const uint16 InstructionIndex, TWeakObjectPtr<UObject> InNode, const uint16 InDepth, const bool bIsTemporary = false);
 
 	bool RemoveBreakpoint(const FRigVMBreakpoint& Breakpoint);
 
@@ -190,9 +190,9 @@ struct RIGVM_API FRigVMDebugInfo
 		}
 	}
 
-	inline TArray<UObject*>& GetCurrentActiveBreakpointCallstack() { return CurrentActiveBreakpointCallstack; }
+	inline TArray<TWeakObjectPtr<UObject>>& GetCurrentActiveBreakpointCallstack() { return CurrentActiveBreakpointCallstack; }
 
-	inline void SetCurrentActiveBreakpointCallstack(TArray<UObject*> Callstack) { CurrentActiveBreakpointCallstack = Callstack; }
+	inline void SetCurrentActiveBreakpointCallstack(TArray<TWeakObjectPtr<UObject>> Callstack) { CurrentActiveBreakpointCallstack = Callstack; }
 
 	inline uint64 GetStartCycles() const {	return StartCycles;	}
 	inline void SetStartCycles(uint64 InStartCycles) { StartCycles = InStartCycles; }
@@ -260,7 +260,7 @@ private:
 	TMap<FGuid, uint16> BreakpointHits; // How many times this instruction has been executed
 
 	FGuid CurrentActiveBreakpoint;
-	TArray<UObject*> CurrentActiveBreakpointCallstack;
+	TArray<TWeakObjectPtr<UObject>> CurrentActiveBreakpointCallstack;
 
 	// --- ---
 
diff --git a/Engine/Plugins/Runtime/RigVM/Source/RigVMDeveloper/Private/RigVMCompiler/RigVMASTProxy.cpp b/Engine/Plugins/Runtime/RigVM/Source/RigVMDeveloper/Private/RigVMCompiler/RigVMASTProxy.cpp
--- a/Engine/Plugins/Runtime/RigVM/Source/RigVMDeveloper/Private/RigVMCompiler/RigVMASTProxy.cpp
+++ b/Engine/Plugins/Runtime/RigVM/Source/RigVMDeveloper/Private/RigVMCompiler/RigVMASTProxy.cpp
@@ -11,28 +11,31 @@
 FString FRigVMCallstack::GetCallPath(bool bIncludeLast) const
 {
 	TArray<FString> Segments;
-	for (UObject* Entry : Stack)
+	for (TWeakObjectPtr<UObject> Entry : Stack)
 	{
-		if (URigVMNode* Node = Cast<URigVMNode>(Entry))
+		if (Entry.IsValid())
 		{
-			if (URigVMGraph* Graph = Node->GetGraph())
+			if (URigVMNode* Node = Cast<URigVMNode>(Entry))
 			{
-				if(Graph->IsRootGraph())
+				if (URigVMGraph* Graph = Node->GetGraph())
 				{
-					Segments.Add(Node->GetNodePath(true));
-				}
-				else
-				{
-					Segments.Add(Node->GetName());
+					if(Graph->IsRootGraph())
+					{
+						Segments.Add(Node->GetNodePath(true));
+					}
+					else
+					{
+						Segments.Add(Node->GetName());
+					}
 				}
 			}
-		}
-		else if (URigVMPin* Pin = Cast<URigVMPin>(Entry))
-		{
-			if (URigVMGraph* Graph = Pin->GetGraph())
+			else if (URigVMPin* Pin = Cast<URigVMPin>(Entry))
 			{
-				const bool bUseNodePath = Graph->IsRootGraph();
-				Segments.Add(Pin->GetPinPath(bUseNodePath));
+				if (URigVMGraph* Graph = Pin->GetGraph())
+				{
+					const bool bUseNodePath = Graph->IsRootGraph();
+					Segments.Add(Pin->GetPinPath(bUseNodePath));
+				}
 			}
 		}
 	}
@@ -63,16 +66,20 @@ int32 FRigVMCallstack::Num() const
 
 const UObject* FRigVMCallstack::Last() const
 {
-	if(Stack.IsEmpty())
+	if(Stack.IsEmpty() || !Stack.Last().IsValid())
 	{
 		return nullptr;
 	}
-	return Stack.Last();
+	return Stack.Last().Get();
 }
 
 const UObject* FRigVMCallstack::operator[](int32 InIndex) const
 {
-	return Stack[InIndex];
+	if (!Stack[InIndex].IsValid())
+	{
+		return nullptr;
+	}
+	return Stack[InIndex].Get();
 }
 
 bool FRigVMCallstack::Contains(const UObject* InEntry) const
@@ -211,7 +218,7 @@ FRigVMASTProxy FRigVMASTProxy::MakeFromCallstack(const FRigVMCallstack& InCallst
 	return Proxy;
 }
 
-FRigVMASTProxy FRigVMASTProxy::MakeFromCallstack(const TArray<UObject*>* InCallstack)
+FRigVMASTProxy FRigVMASTProxy::MakeFromCallstack(const TArray<TWeakObjectPtr<UObject>>* InCallstack)
 {
 	check(InCallstack);
 	FRigVMCallstack Callstack;
diff --git a/Engine/Plugins/Runtime/RigVM/Source/RigVMDeveloper/Private/RigVMCompiler/RigVMCompiler.cpp b/Engine/Plugins/Runtime/RigVM/Source/RigVMDeveloper/Private/RigVMCompiler/RigVMCompiler.cpp
--- a/Engine/Plugins/Runtime/RigVM/Source/RigVMDeveloper/Private/RigVMCompiler/RigVMCompiler.cpp
+++ b/Engine/Plugins/Runtime/RigVM/Source/RigVMDeveloper/Private/RigVMCompiler/RigVMCompiler.cpp
@@ -2496,7 +2496,7 @@ bool URigVMCompiler::TraverseInlineFunction(const FRigVMInlineFunctionExprAST* I
 
 			if (WorkData.Settings.SetupNodeInstructionIndex)
 			{
-				if (const TArray<UObject*>* Callstack = FunctionByteCode.GetCallstackForInstruction(i-InstructionIndexStart))
+				if (const TArray<TWeakObjectPtr<UObject>>* Callstack = FunctionByteCode.GetCallstackForInstruction(i-InstructionIndexStart))
 				{
 					if (Callstack->Num() > 1)
 					{
diff --git a/Engine/Plugins/Runtime/RigVM/Source/RigVMDeveloper/Public/RigVMCompiler/RigVMASTProxy.h b/Engine/Plugins/Runtime/RigVM/Source/RigVMDeveloper/Public/RigVMCompiler/RigVMASTProxy.h
--- a/Engine/Plugins/Runtime/RigVM/Source/RigVMDeveloper/Public/RigVMCompiler/RigVMASTProxy.h
+++ b/Engine/Plugins/Runtime/RigVM/Source/RigVMDeveloper/Public/RigVMCompiler/RigVMASTProxy.h
@@ -95,7 +95,7 @@ public:
 		return operator[](0) > Other[0];
 	}
 
-	TArray<UObject*> GetStack() const { return Stack; }
+	TArray<TWeakObjectPtr<UObject>> GetStack() const { return Stack; }
 
 	FRigVMCallstack GetCallStackUpTo(int32 InIndex) const;
 
@@ -114,7 +114,7 @@ private:
 		return HashCombine(GetEntryTypeHash(Index - 1), GetTypeHash(Stack[Index]));
 	}
 
-	TArray<UObject*> Stack;
+	TArray<TWeakObjectPtr<UObject>> Stack;
 
 	friend class FRigVMParserAST;
 	friend class FRigVMCallstackGuard;
@@ -162,7 +162,7 @@ public:
 	static FRigVMASTProxy MakeFromUObject(UObject* InSubject);
 	static FRigVMASTProxy MakeFromCallPath(const FString& InCallPath, UObject* InRootObject);
 	static FRigVMASTProxy MakeFromCallstack(const FRigVMCallstack& InCallstack);
-	static FRigVMASTProxy MakeFromCallstack(const TArray<UObject*>* InCallstack);
+	static FRigVMASTProxy MakeFromCallstack(const TArray<TWeakObjectPtr<UObject>>* InCallstack);
 
 	FRigVMASTProxy GetSibling(UObject* InSubject) const
 	{
@@ -204,7 +204,7 @@ public:
 
 	const FRigVMCallstack& GetCallstack() const { return Callstack; }
 
-	UObject* GetSubject() const { return Callstack.Stack.Last(); }
+	UObject* GetSubject() const { return Callstack.Stack.Last().IsValid() ? Callstack.Stack.Last().Get() : nullptr; }
 
 	template<class T>
 	T* GetSubject() const { return Cast<T>(GetSubject()); }
diff --git a/Engine/Plugins/Runtime/RigVM/Source/RigVMEditor/Private/Widgets/SRigVMExecutionStackView.cpp b/Engine/Plugins/Runtime/RigVM/Source/RigVMEditor/Private/Widgets/SRigVMExecutionStackView.cpp
--- a/Engine/Plugins/Runtime/RigVM/Source/RigVMEditor/Private/Widgets/SRigVMExecutionStackView.cpp
+++ b/Engine/Plugins/Runtime/RigVM/Source/RigVMEditor/Private/Widgets/SRigVMExecutionStackView.cpp
@@ -591,7 +591,7 @@ void SRigVMExecutionStackView::PopulateStackView(URigVM* InVM, FRigVMExtendedExe
 				if (FilterText.IsEmpty() || Label.Contains(FilterText.ToString()))
 				{
 					FRigVMASTProxy Proxy;
-					if(const TArray<UObject*>* Callstack = ByteCode.GetCallstackForInstruction(InstructionIndex))
+					if(const TArray<TWeakObjectPtr<UObject>>* Callstack = ByteCode.GetCallstackForInstruction(InstructionIndex))
 					{
 						Proxy = FRigVMASTProxy::MakeFromCallstack(Callstack);
 					}
