RigVM: Fix crash when renaming public function
#rnx

diff --git a/Engine/Plugins/Runtime/RigVM/Source/RigVMDeveloper/Private/RigVMModel/RigVMClient.cpp b/Engine/Plugins/Runtime/RigVM/Source/RigVMDeveloper/Private/RigVMModel/RigVMClient.cpp
--- a/Engine/Plugins/Runtime/RigVM/Source/RigVMDeveloper/Private/RigVMModel/RigVMClient.cpp
+++ b/Engine/Plugins/Runtime/RigVM/Source/RigVMDeveloper/Private/RigVMModel/RigVMClient.cpp
@@ -1075,7 +1075,15 @@ void FRigVMClient::HandleGraphModifiedEvent(ERigVMGraphNotifType InNotifType, UR
 
 								for (int32 i=0; i<NewFunctionReferences.Num(); ++i)
 								{
-									NewFunctionReferences[i]->ReferencedFunctionHeader = Data->Header;
+									if (!NewFunctionReferences[i].IsValid())
+									{
+										NewFunctionReferences[i].LoadSynchronous();
+									}
+									if (NewFunctionReferences[i].IsValid())
+									{
+										NewFunctionReferences[i]->ReferencedFunctionHeader = Data->Header;
+										NewFunctionReferences[i]->MarkPackageDirty();
+									}
 								}
 							}
 
