RigVM: Fix crash when accessing instruction in the ExecutionStack
#rnx

diff --git a/Engine/Plugins/Runtime/RigVM/Source/RigVMEditor/Private/Widgets/SRigVMExecutionStackView.cpp b/Engine/Plugins/Runtime/RigVM/Source/RigVMEditor/Private/Widgets/SRigVMExecutionStackView.cpp
--- a/Engine/Plugins/Runtime/RigVM/Source/RigVMEditor/Private/Widgets/SRigVMExecutionStackView.cpp
+++ b/Engine/Plugins/Runtime/RigVM/Source/RigVMEditor/Private/Widgets/SRigVMExecutionStackView.cpp
@@ -730,7 +730,7 @@ void SRigVMExecutionStackView::RefreshTreeView(URigVM* InVM, FRigVMExtendedExecu
 				const TArray<FRigVMLog::FLogEntry>& LogEntries = Host->GetLog()->Entries;
 				for (const FRigVMLog::FLogEntry& LogEntry : LogEntries)
 				{
-					if (Operators.Num() <= LogEntry.InstructionIndex)
+					if (!Operators.IsValidIndex(LogEntry.InstructionIndex))
 					{
 						continue;
 					}
