RigVMEditor: Avoid displaying errors on incorrect nodes

diff --git a/Engine/Plugins/Runtime/RigVM/Source/RigVMEditor/Private/Editor/RigVMEditor.cpp b/Engine/Plugins/Runtime/RigVM/Source/RigVMEditor/Private/Editor/RigVMEditor.cpp
--- a/Engine/Plugins/Runtime/RigVM/Source/RigVMEditor/Private/Editor/RigVMEditor.cpp
+++ b/Engine/Plugins/Runtime/RigVM/Source/RigVMEditor/Private/Editor/RigVMEditor.cpp
@@ -3326,13 +3326,20 @@ void FRigVMEditor::UpdateGraphCompilerErrors()
 					continue;
 				}
 
-				bFoundError = bFoundError || Entry.Severity <= EMessageSeverity::Error;
-				bFoundWarning = bFoundWarning || Entry.Severity <= EMessageSeverity::Warning;
-
 				if (URigVMEdGraphNode* RigVMEdGraphNode = Cast<URigVMEdGraphNode>(GraphNode))
 				{
+					// The node in this graph may have the same local node path,
+					// but may be backed by another model node.
+					if(RigVMEdGraphNode->GetModelNode() != ModelNode)
+					{
+						continue;
+					}
+
 					RigVMEdGraphNode->AddErrorInfo(Entry.Severity, Entry.Message);
 				}
+
+				bFoundError = bFoundError || Entry.Severity <= EMessageSeverity::Error;
+				bFoundWarning = bFoundWarning || Entry.Severity <= EMessageSeverity::Warning;
 			}
 
 			bAnyErrorsLeft = false;
