RigVMRegistry: avoid adding array type into the user defined type index loopup map, which assumes element type only

diff --git a/Engine/Plugins/Runtime/RigVM/Source/RigVM/Private/RigVMCore/RigVMRegistry.cpp b/Engine/Plugins/Runtime/RigVM/Source/RigVM/Private/RigVMCore/RigVMRegistry.cpp
--- a/Engine/Plugins/Runtime/RigVM/Source/RigVM/Private/RigVMCore/RigVMRegistry.cpp
+++ b/Engine/Plugins/Runtime/RigVM/Source/RigVM/Private/RigVMCore/RigVMRegistry.cpp
@@ -676,8 +676,9 @@ TRigVMTypeIndex FRigVMRegistry::FindOrAddType_Internal(const FRigVMTemplateArgum
 		{
 			if (CPPTypeObject->IsA<UUserDefinedStruct>() || CPPTypeObject->IsA<UUserDefinedEnum>())
 			{
-				// used to track name changes to user defined types
-				UserDefinedTypeToIndex.FindOrAdd(CPPTypeObject) = Index;
+				TRigVMTypeIndex ElementTypeIndex = GetTypeIndex(ElementType);
+				// used to track name changes to user defined types, stores the element type index, see RemoveType()
+				UserDefinedTypeToIndex.FindOrAdd(CPPTypeObject) = ElementTypeIndex;
 			}
 		}
 		
