Add null check for non-loaded asset classes.

diff --git a/Engine/Plugins/Runtime/RigVM/Source/RigVMDeveloper/Private/RigVMBlueprintUtils.cpp b/Engine/Plugins/Runtime/RigVM/Source/RigVMDeveloper/Private/RigVMBlueprintUtils.cpp
--- a/Engine/Plugins/Runtime/RigVM/Source/RigVMDeveloper/Private/RigVMBlueprintUtils.cpp
+++ b/Engine/Plugins/Runtime/RigVM/Source/RigVMDeveloper/Private/RigVMBlueprintUtils.cpp
@@ -141,7 +141,7 @@ void FRigVMBlueprintUtils::HandleRefreshAllNodes(UBlueprint* InBlueprint)
 
 void FRigVMBlueprintUtils::HandleAssetDeleted(const FAssetData& InAssetData)
 {
-	if (InAssetData.GetClass()->IsChildOf(URigVMBlueprint::StaticClass()))
+	if (InAssetData.GetClass() && InAssetData.GetClass()->IsChildOf(URigVMBlueprint::StaticClass()))
 	{
 		// Make sure any RigVMBlueprint removes any TypeActions related to this asset (e.g. public functions)
 		FBlueprintActionDatabase& ActionDatabase = FBlueprintActionDatabase::Get();
