RigVM: FIx FRigVMFunction_MathIntersectPlane and FRigVMFunction_MathRayIntersectPlane producing NaN values when normal is zero
#rnx

diff --git a/Engine/Plugins/Runtime/RigVM/Source/RigVM/Private/RigVMFunctions/Math/RigVMFunction_MathRay.cpp b/Engine/Plugins/Runtime/RigVM/Source/RigVM/Private/RigVMFunctions/Math/RigVMFunction_MathRay.cpp
--- a/Engine/Plugins/Runtime/RigVM/Source/RigVM/Private/RigVMFunctions/Math/RigVMFunction_MathRay.cpp
+++ b/Engine/Plugins/Runtime/RigVM/Source/RigVM/Private/RigVMFunctions/Math/RigVMFunction_MathRay.cpp
@@ -31,7 +31,15 @@ FRigVMFunction_MathRayIntersectRay_Execute()
 
 FRigVMFunction_MathRayIntersectPlane_Execute()
 {
-	Ratio = FMath::RayPlaneIntersectionParam(Ray.Origin, Ray.Direction, FPlane(PlanePoint, PlaneNormal));
+	const FVector NormalizedPlaneNormal = PlaneNormal.GetSafeNormal(UE_SMALL_NUMBER, FVector::UpVector);
+	if (FMath::IsNearlyZero(FVector::DotProduct(NormalizedPlaneNormal, Ray.Direction)))
+	{
+		Result = FVector::ZeroVector;
+		Distance = Ratio = 0.f;
+		return;
+	}
+	
+	Ratio = FMath::RayPlaneIntersectionParam(Ray.Origin, Ray.Direction, FPlane(PlanePoint, NormalizedPlaneNormal));
 	Result = Ray.PointAt(Ratio);
 	Distance = 0;
 	if(!Ray.Origin.Equals(Result))
diff --git a/Engine/Plugins/Runtime/RigVM/Source/RigVM/Private/RigVMFunctions/Math/RigVMFunction_MathVector.cpp b/Engine/Plugins/Runtime/RigVM/Source/RigVM/Private/RigVMFunctions/Math/RigVMFunction_MathVector.cpp
--- a/Engine/Plugins/Runtime/RigVM/Source/RigVM/Private/RigVMFunctions/Math/RigVMFunction_MathVector.cpp
+++ b/Engine/Plugins/Runtime/RigVM/Source/RigVM/Private/RigVMFunctions/Math/RigVMFunction_MathVector.cpp
@@ -693,7 +693,15 @@ FRigVMFunction_MathVectorClampSpatially_Execute()
 
 FRigVMFunction_MathIntersectPlane_Execute()
 {
-	FPlane Plane(PlanePoint, PlaneNormal);
+	const FVector NormalizedPlaneNormal = PlaneNormal.GetSafeNormal(UE_SMALL_NUMBER, FVector(0.f, 0.f, 1.f));
+	if (FMath::IsNearlyZero(FVector::DotProduct(NormalizedPlaneNormal, Direction)))
+	{
+		Result = FVector::ZeroVector;
+		Distance = 0.f;
+		return;
+	}
+	
+	FPlane Plane(PlanePoint, NormalizedPlaneNormal);
 
 	Result = FMath::RayPlaneIntersection(Start, Direction, Plane);
 	Distance = (Start - Result).Size();
