RIgVM: Fix crash when URigVMMemoryStorageObjects are loaded after we destroy its GeneratedClass
#rnx

#localization none
#tests none

diff --git a/Engine/Plugins/Runtime/RigVM/Source/RigVMDeveloper/Private/RigVMBlueprint.cpp b/Engine/Plugins/Runtime/RigVM/Source/RigVMDeveloper/Private/RigVMBlueprint.cpp
--- a/Engine/Plugins/Runtime/RigVM/Source/RigVMDeveloper/Private/RigVMBlueprint.cpp
+++ b/Engine/Plugins/Runtime/RigVM/Source/RigVMDeveloper/Private/RigVMBlueprint.cpp
@@ -835,6 +835,10 @@ void URigVMBlueprint::PreSave(FObjectPreSaveContext ObjectSaveContext)
 			}
 		}
 	}
+
+#if WITH_EDITORONLY_DATA
+	OldMemoryStorageGeneratorClasses.Reset();
+#endif
 }
 
 void URigVMBlueprint::PostSaveRoot(FObjectPostSaveRootContext ObjectSaveContext)
@@ -1121,18 +1125,22 @@ void URigVMBlueprint::BroadcastRigVMPackageDone()
 	}
 }
 
-void URigVMBlueprint::RemoveDeprecatedVMMemoryClass() const
+void URigVMBlueprint::RemoveDeprecatedVMMemoryClass() 
 {
 	TArray<UObject*> Objects;
 	GetObjectsWithOuter(this, Objects, false);
 
+#if WITH_EDITORONLY_DATA
+	OldMemoryStorageGeneratorClasses.Reserve(Objects.Num());
 	for (UObject* Object : Objects)
 	{
 		if (URigVMMemoryStorageGeneratorClass* DeprecatedClass = Cast<URigVMMemoryStorageGeneratorClass>(Object))
 		{
 			DeprecatedClass->Rename(nullptr, GetTransientPackage(), REN_ForceNoResetLoaders | REN_DoNotDirty | REN_DontCreateRedirectors | REN_NonTransactional);
+			OldMemoryStorageGeneratorClasses.Add(DeprecatedClass);
 		}
 	}
+#endif
 }
 #endif
 
diff --git a/Engine/Plugins/Runtime/RigVM/Source/RigVMDeveloper/Public/RigVMBlueprint.h b/Engine/Plugins/Runtime/RigVM/Source/RigVMDeveloper/Public/RigVMBlueprint.h
--- a/Engine/Plugins/Runtime/RigVM/Source/RigVMDeveloper/Public/RigVMBlueprint.h
+++ b/Engine/Plugins/Runtime/RigVM/Source/RigVMDeveloper/Public/RigVMBlueprint.h
@@ -623,9 +623,18 @@ private:
 	// this function removes those deprecated class.
 	// new classes should be created by RecompileVM and parented to the Package
 	// during PostLoad
-	void RemoveDeprecatedVMMemoryClass() const;
+	void RemoveDeprecatedVMMemoryClass();
+
+#if WITH_EDITORONLY_DATA
+	// During load, we do not want the GC to destroy the generator classes until all URigVMMemoryStorage objects
+	// are loaded, so we need to keep a pointer to the classes. These pointers will be removed on PreSave so that the
+	// GC can do its work.
+	UPROPERTY(Transient)
+	TArray<TObjectPtr<URigVMMemoryStorageGeneratorClass>> OldMemoryStorageGeneratorClasses;
+
 #endif
 
+#endif
 #if WITH_EDITOR
 
 public:
