RigVM compiler now generates default memory directly in the VM. The memory is then copied to the ExtendedExecuteContext when an instance is created from the CDO.
RigVM : Removed local memory array parameter to VM Initialize and Execute, as it is not used by the VM. Deprecated functions and renamed Execute to ExecuteVM due to function collisions
AnimNext ParameterBlock creates VM Runtime ExtendedExecuteContexts in the TLS, using CDO context as reference context (it is created once for each VM and only refreshed if the VM hash changes due to recompilation)

diff --git a/Engine/Plugins/Animation/ControlRig/Source/ControlRig/Private/ControlRig.cpp b/Engine/Plugins/Animation/ControlRig/Source/ControlRig/Private/ControlRig.cpp
--- a/Engine/Plugins/Animation/ControlRig/Source/ControlRig/Private/ControlRig.cpp
+++ b/Engine/Plugins/Animation/ControlRig/Source/ControlRig/Private/ControlRig.cpp
@@ -65,6 +65,7 @@ static TAutoConsoleVariable<float> CVarControlRigEnableDrawInterfaceInGame(
 	TEXT("If nonzero debug drawing will be enabled during play."),
 	ECVF_Default);
 
+
 UControlRig::UControlRig(const FObjectInitializer& ObjectInitializer)
 	: Super(ObjectInitializer)
 #if WITH_EDITOR
@@ -91,11 +92,15 @@ UControlRig::UControlRig(const FObjectInitializer& ObjectInitializer)
 #endif
 {
 	EventQueue.Add(FRigUnit_BeginExecution::EventName);
+
+	SetRigVMExtendedExecuteContext(&RigVMExtendedExecuteContext);
 }
 
 void UControlRig::BeginDestroy()
 {
 	Super::BeginDestroy();
+	SetRigVMExtendedExecuteContext(nullptr);
+
 	PreConstructionEvent.Clear();
 	PostConstructionEvent.Clear();
 	PreForwardsSolveEvent.Clear();
@@ -306,7 +311,7 @@ bool UControlRig::InitializeVM(const FName& InEventName)
 	{
 		HierarchyController->LogFunction = [this](EMessageSeverity::Type InSeverity, const FString& Message)
 		{
-			const FRigVMExecuteContext& PublicContext = GetExtendedExecuteContext().GetPublicData<>();
+			const FRigVMExecuteContext& PublicContext = GetRigVMExtendedExecuteContext().GetPublicData<>();
 			if(RigVMLog)
 			{
 				RigVMLog->Report(InSeverity,PublicContext.GetFunctionName(),PublicContext.GetInstructionIndex(), Message);
@@ -642,7 +647,9 @@ bool UControlRig::Execute(const FName& InEventName)
 	DECLARE_SCOPE_HIERARCHICAL_COUNTER_FUNC()
 	QUICK_SCOPE_CYCLE_COUNTER(STAT_ControlRig_Execute);
 	
-	FControlRigExecuteContext& PublicContext = GetExtendedExecuteContext().GetPublicDataSafe<FControlRigExecuteContext>();
+	FRigVMExtendedExecuteContext& ExtendedExecuteContext = GetRigVMExtendedExecuteContext();
+
+	FControlRigExecuteContext& PublicContext = ExtendedExecuteContext.GetPublicDataSafe<FControlRigExecuteContext>();
 
 #if WITH_EDITOR
 	PublicContext.SetLog(RigVMLog); // may be nullptr
@@ -663,7 +670,7 @@ bool UControlRig::Execute(const FName& InEventName)
 		// only set a valid first entry event later when execution
 		// has passed the initialization stage and there are multiple events present in one evaluation
 		// first entry event is used to determined when to clear data during an evaluation
-		VM->SetFirstEntryEventInEventQueue(GetExtendedExecuteContext(), NAME_None);
+		VM->SetFirstEntryEventInEventQueue(ExtendedExecuteContext, NAME_None);
 #endif
 	}
 
@@ -682,12 +689,11 @@ bool UControlRig::Execute(const FName& InEventName)
 			}
 		}
 
-		GetExtendedExecuteContext().SetDebugInfo(&DebugInfo);
-		GetSnapshotContext() = GetExtendedExecuteContext();
+		ExtendedExecuteContext.SetDebugInfo(&DebugInfo);
 	}
 	else
 	{
-		GetExtendedExecuteContext().SetDebugInfo(nullptr);
+		ExtendedExecuteContext.SetDebugInfo(nullptr);
 		GetSnapshotContext().Reset();
 	}
 #endif
@@ -967,7 +973,7 @@ bool UControlRig::Execute(const FName& InEventName)
 		// has passed the initialization stage and there are multiple events present
 		if (EventQueueToRun.Num() >= 2 && VM)
 		{
-			VM->SetFirstEntryEventInEventQueue(GetExtendedExecuteContext(), EventQueueToRun[0]);
+			VM->SetFirstEntryEventInEventQueue(ExtendedExecuteContext, EventQueueToRun[0]);
 		}
 
 		// Transform Overrride is generated using a Transient Control 
@@ -1323,7 +1329,7 @@ bool UControlRig::Execute_Internal(const FName& InEventName)
 	
 	if (VM)
 	{
-		FRigVMExtendedExecuteContext& Context = GetExtendedExecuteContext();
+		FRigVMExtendedExecuteContext& Context = GetRigVMExtendedExecuteContext();
 
 		static constexpr TCHAR InvalidatedVMFormat[] = TEXT("%s: Invalidated VM - aborting execution.");
 		if(VM->IsNativized())
@@ -1371,12 +1377,12 @@ bool UControlRig::Execute_Internal(const FName& InEventName)
 				{
 					if(bIsEventFirstInQueue)
 					{
-						CopyVMMemory(GetExtendedExecuteContext(), GetSnapshotContext());
+						CopyVMMemory(GetRigVMExtendedExecuteContext(), GetSnapshotContext());
 					}
 				}
 				else if(bIsEventLastInQueue)
 				{
-					CopyVMMemory(GetSnapshotContext(), GetExtendedExecuteContext());
+					CopyVMMemory(GetSnapshotContext(), GetRigVMExtendedExecuteContext());
 				}
 			}
 		}
@@ -1409,8 +1415,7 @@ bool UControlRig::Execute_Internal(const FName& InEventName)
 		FControlRigExecuteContextRigModuleGuard RigModuleGuard(PublicContext, this);
 		FRigHierarchyRedirectorGuard ElementRedirectorGuard(this);
 
-		TArray<FRigVMMemoryStorageStruct*> LocalMemory = VM->GetLocalMemoryArray(Context);
-		const bool bSuccess = VM->Execute(Context, LocalMemory, InEventName) != ERigVMExecuteResult::Failed;
+		const bool bSuccess = VM->ExecuteVM(Context, InEventName) != ERigVMExecuteResult::Failed;
 
 #if UE_RIGVM_PROFILE_EXECUTE_UNITS_NUM
 		const uint64 EndCycles = FPlatformTime::Cycles64();
@@ -2256,7 +2261,7 @@ bool UControlRig::SetTransientControlValue(const URigVMUnitNode* InNode, TShared
 		return false;
 	}
 
-	FControlRigExecuteContext& PublicContext = GetExtendedExecuteContext().GetPublicDataSafe<FControlRigExecuteContext>();
+	FControlRigExecuteContext& PublicContext = GetRigVMExtendedExecuteContext().GetPublicDataSafe<FControlRigExecuteContext>();
 	const bool bResult = UnitInstance->UpdateHierarchyForDirectManipulation(InNode, NodeInstance, PublicContext, InInfo);
 	InInfo->bInitialized = true;
 	return bResult;
@@ -3070,7 +3075,7 @@ USceneComponent* UControlRig::GetOwningSceneComponent()
 {
 	if(OuterSceneComponent == nullptr)
 	{
-		const FControlRigExecuteContext& PublicContext = GetExtendedExecuteContext().GetPublicDataSafe<FControlRigExecuteContext>();
+		const FControlRigExecuteContext& PublicContext = GetRigVMExtendedExecuteContext().GetPublicDataSafe<FControlRigExecuteContext>();
 		const FRigUnitContext& Context = PublicContext.UnitContext;
 
 		USceneComponent* SceneComponentFromRegistry = Context.DataSourceRegistry->RequestSource<USceneComponent>(UControlRig::OwnerComponent);
@@ -3094,7 +3099,7 @@ void UControlRig::PostInitInstance(URigVMHost* InCDO)
 		RF_Public | RF_DefaultSubObject :
 		RF_Transient | RF_Transactional;
 
-	FRigVMExtendedExecuteContext& Context = GetExtendedExecuteContext();
+	FRigVMExtendedExecuteContext& Context = GetRigVMExtendedExecuteContext();
 	
 	Context.SetContextPublicDataStruct(FControlRigExecuteContext::StaticStruct());
 
@@ -3135,24 +3140,6 @@ void UControlRig::PostInitInstance(URigVMHost* InCDO)
 				VM = InCDO->GetVM();
 			}
 
-			// This is needed by some tests that use Templates to run (and RigVMHost does not initialize template CDO rigs)
-			if (InCDO->HasAnyFlags(RF_ArchetypeObject))
-			{
-				CopyVMMemory(GetExtendedExecuteContext(), InCDO->GetExtendedExecuteContext());
-
-				Context.VMHash = InCDO->GetExtendedExecuteContext().VMHash;
-
-				if (!ensure(VM->GetVMHash() == Context.VMHash))
-				{
-					UE_LOG(LogRigVM
-						, Warning
-						, TEXT("ControlRig : VM Hash [%u] is different from Context stored VM Hash [%u]. Please recompile ControlRig used at Asset : [%s]")
-						, Context.VMHash
-						, VM->GetVMHash()
-						, *GetPathName());
-				}
-			}
-
 			DynamicHierarchy->CopyHierarchy(CastChecked<UControlRig>(InCDO)->GetHierarchy());
 		}
 	}
diff --git a/Engine/Plugins/Animation/ControlRig/Source/ControlRig/Public/ControlRig.h b/Engine/Plugins/Animation/ControlRig/Source/ControlRig/Public/ControlRig.h
--- a/Engine/Plugins/Animation/ControlRig/Source/ControlRig/Public/ControlRig.h
+++ b/Engine/Plugins/Animation/ControlRig/Source/ControlRig/Public/ControlRig.h
@@ -407,6 +407,9 @@ public:
 	};
 	
 private:
+	UPROPERTY(Transient)
+	FRigVMExtendedExecuteContext RigVMExtendedExecuteContext;
+
 	UE::Anim::FStackAttributeContainer* ExternalAnimAttributeContainer;
 
 #if WITH_EDITOR
diff --git a/Engine/Plugins/Animation/ControlRig/Source/ControlRigEditor/Private/ControlRigCompilerDetails.cpp b/Engine/Plugins/Animation/ControlRig/Source/ControlRigEditor/Private/ControlRigCompilerDetails.cpp
--- a/Engine/Plugins/Animation/ControlRig/Source/ControlRigEditor/Private/ControlRigCompilerDetails.cpp
+++ b/Engine/Plugins/Animation/ControlRig/Source/ControlRigEditor/Private/ControlRigCompilerDetails.cpp
@@ -223,7 +223,7 @@ FReply FRigVMCompileSettingsDetails::OnCopyByteCodeClicked()
 		{
 			if(UControlRig* ControlRig = Cast<UControlRig>(BlueprintBeingCustomized->GetObjectBeingDebugged()))
 			{
-				FString ByteCodeContent = ControlRig->GetVM()->DumpByteCodeAsText(ControlRig->GetExtendedExecuteContext());
+				FString ByteCodeContent = ControlRig->GetVM()->DumpByteCodeAsText(ControlRig->GetRigVMExtendedExecuteContext());
 				FPlatformApplicationMisc::ClipboardCopy(*ByteCodeContent);
 			}
 		}
@@ -261,14 +261,14 @@ FReply FRigVMCompileSettingsDetails::OnCopyGeneratedCodeClicked()
 			{
 				if(CDO->GetVM())
 				{
-					CDO->GetVM()->ClearExternalVariables(CDO->GetExtendedExecuteContext());
+					CDO->GetVM()->ClearExternalVariables(CDO->GetRigVMExtendedExecuteContext());
 					TArray<FRigVMExternalVariable> ExternalVariables = CDO->GetExternalVariables();
 					for(const FRigVMExternalVariable& ExternalVariable : ExternalVariables)
 					{
-						CDO->GetVM()->AddExternalVariable(CDO->GetExtendedExecuteContext(), ExternalVariable);
+						CDO->GetVM()->AddExternalVariable(CDO->GetRigVMExtendedExecuteContext(), ExternalVariable);
 					}
 					
-					FRigVMExtendedExecuteContext& CDOContext = CDO->GetExtendedExecuteContext();
+					FRigVMExtendedExecuteContext& CDOContext = CDO->GetRigVMExtendedExecuteContext();
 
 					FRigVMCodeGenerator CodeGenerator(ClassName,
 						TEXT("TestModule"), BlueprintBeingCustomized->GetDefaultModel(), CDO->GetVM(), CDOContext,
diff --git a/Engine/Plugins/Animation/ControlRig/Source/ControlRigEditor/Private/Editor/ControlRigEditor.cpp b/Engine/Plugins/Animation/ControlRig/Source/ControlRigEditor/Private/Editor/ControlRigEditor.cpp
--- a/Engine/Plugins/Animation/ControlRig/Source/ControlRigEditor/Private/Editor/ControlRigEditor.cpp
+++ b/Engine/Plugins/Animation/ControlRig/Source/ControlRigEditor/Private/Editor/ControlRigEditor.cpp
@@ -1484,7 +1484,7 @@ void FControlRigEditor::HandleVMExecutedEvent(URigVMHost* InHost, const FName& I
 			if(!DebuggedControlRig->RigUnitManipulationInfos.IsEmpty())
 			{
 				const FRigHierarchyRedirectorGuard RedirectorGuard(DebuggedControlRig);
-				FControlRigExecuteContext& ExecuteContext = DebuggedControlRig->GetExtendedExecuteContext().GetPublicDataSafe<FControlRigExecuteContext>();
+				FControlRigExecuteContext& ExecuteContext = DebuggedControlRig->GetRigVMExtendedExecuteContext().GetPublicDataSafe<FControlRigExecuteContext>();
 				
 				for(const TSharedPtr<FRigDirectManipulationInfo>& ManipulationInfo : DebuggedControlRig->RigUnitManipulationInfos)
 				{
@@ -4389,7 +4389,7 @@ void FControlRigEditor::HandleOnControlModified(UControlRig* Subject, FRigContro
 
 			// update the node based on the incoming pose. once that is done we'll need to compare the node instance
 			// with the settings on the node in the graph and update them accordingly.
-			FControlRigExecuteContext& ExecuteContext = DebuggedControlRig->GetExtendedExecuteContext().GetPublicDataSafe<FControlRigExecuteContext>();
+			FControlRigExecuteContext& ExecuteContext = DebuggedControlRig->GetRigVMExtendedExecuteContext().GetPublicDataSafe<FControlRigExecuteContext>();
 			const FRigHierarchyRedirectorGuard RedirectorGuard(DebuggedControlRig);
 			if(UnitInstance->UpdateDirectManipulationFromHierarchy(UnitNode, NodeInstance, ExecuteContext, ManipulationInfo))
 			{
diff --git a/Engine/Plugins/Experimental/AnimNext/Source/AnimNext/Internal/Param/AnimNextParameterBlock.h b/Engine/Plugins/Experimental/AnimNext/Source/AnimNext/Internal/Param/AnimNextParameterBlock.h
--- a/Engine/Plugins/Experimental/AnimNext/Source/AnimNext/Internal/Param/AnimNextParameterBlock.h
+++ b/Engine/Plugins/Experimental/AnimNext/Source/AnimNext/Internal/Param/AnimNextParameterBlock.h
@@ -66,6 +66,8 @@ class UAnimNextParameterBlock : public URigVMHost, public IAnimNextParameterSour
 {
 	GENERATED_BODY()
 
+	UAnimNextParameterBlock(const FObjectInitializer& ObjectInitializer);
+
 	friend class UAnimNextParameterBlockFactory;
 	friend class UAnimNextParameterBlock_EditorData;
 	friend struct UE::AnimNext::UncookedOnly::FUtils;
@@ -86,9 +88,8 @@ class UAnimNextParameterBlock : public URigVMHost, public IAnimNextParameterSour
 	virtual void PostRename(UObject* OldOuter, const FName OldName) override;
 	virtual void GetPreloadDependencies(TArray<UObject*>& OutDeps) override;
 	virtual void GetAssetRegistryTags(TArray<FAssetRegistryTag>& OutTags) const override;
-
-	UPROPERTY()
-	TObjectPtr<URigVM> RigVM;
+	
+	FRigVMExtendedExecuteContext BaseRigVMContext;
 
 	UPROPERTY()
 	FInstancedPropertyBag PropertyBag;
@@ -97,4 +98,4 @@ class UAnimNextParameterBlock : public URigVMHost, public IAnimNextParameterSour
 	UPROPERTY(VisibleAnywhere, Instanced, Category = "Parameters", meta = (ShowInnerProperties))
 	TObjectPtr<UObject> EditorData;
 #endif
-};
\ No newline at end of file
+};
diff --git a/Engine/Plugins/Experimental/AnimNext/Source/AnimNext/Internal/RigVMRuntimeDataRegistry.h b/Engine/Plugins/Experimental/AnimNext/Source/AnimNext/Internal/RigVMRuntimeDataRegistry.h
new file mode 100644
--- /dev/null
+++ b/Engine/Plugins/Experimental/AnimNext/Source/AnimNext/Internal/RigVMRuntimeDataRegistry.h
@@ -0,0 +1,69 @@
+// Copyright Epic Games, Inc. All Rights Reserved.
+
+#pragma once
+
+#include "CoreMinimal.h"
+#include "RigVMCore/RigVMExecuteContext.h"
+
+
+namespace UE::AnimNext
+{
+
+struct FRigVMRuntimeData
+{
+	FRigVMExtendedExecuteContext Context;
+};
+
+using FRigVMRuntimeDataID = TObjectKey<URigVM>;
+
+/**
+	* RigVMRuntimeDataRegistry
+	* 
+	* A global registry of all existing VMs that require TLS data instantiation
+	* 
+	*/
+struct ANIMNEXT_API FRigVMRuntimeDataRegistry final
+{
+	/**
+	 * Finds the VM runtime data ID in the TLS storage
+	 *
+	 * @param RigVMRuntimeDataID The key to associate the value with.
+	 * @return A pointer to the data as stored in the map. The data is only valid until the next change to any key in the map.
+	 */
+	static FRigVMRuntimeData* FindRuntimeData(const FRigVMRuntimeDataID& RigVMRuntimeDataID);
+
+	/**
+	 * Adds VM runtime data for the passed ID
+	 *
+	 * @param RigVMRuntimeDataID The key to associate the value with.
+	 * @param ReferenceContext Reference context to copy data from
+	 * @return A pointer to the data as stored in the map. The data is only valid until the next change to any key in the map.
+	 */
+	static FRigVMRuntimeData* AddRuntimeData(const FRigVMRuntimeDataID& RigVMRuntimeDataID, const FRigVMExtendedExecuteContext& ReferenceContext);
+
+	
+	/**
+	 * Finds or adds the VM runtime data ID
+	 *
+	 * @param RigVMRuntimeDataID The key to associate the value with.
+	 * @param ReferenceContext Reference context to copy data from (if not found)
+	 * @return A pointer to the data as stored in the map. The data is only valid until the next change to any key in the map.
+	 */
+	static FRigVMRuntimeData* FindOrAddRuntimeData(const FRigVMRuntimeDataID& RigVMRuntimeDataID, const FRigVMExtendedExecuteContext& ReferenceContext);
+
+private:
+
+	// Post GC callback to signal a compactation has to be done (and perform it on game thread)
+	static void HandlePostGarbageCollect();
+
+	// Checks if any of the stored VM datas have been deleted and removes deleted elements.
+	static void PerformStorageCompaction();
+
+	// Module lifetime functions
+	static void Init();
+	static void Destroy();
+
+	friend class FModule;
+};
+
+} // end namespace UE::AnimNext
diff --git a/Engine/Plugins/Experimental/AnimNext/Source/AnimNext/Private/Graph/AnimNextGraph.cpp b/Engine/Plugins/Experimental/AnimNext/Source/AnimNext/Private/Graph/AnimNextGraph.cpp
--- a/Engine/Plugins/Experimental/AnimNext/Source/AnimNext/Private/Graph/AnimNextGraph.cpp
+++ b/Engine/Plugins/Experimental/AnimNext/Source/AnimNext/Private/Graph/AnimNextGraph.cpp
@@ -58,9 +58,9 @@ void UAnimNextGraph::AllocateInstance(FAnimNextGraphInstance& Instance) const
 	UE::AnimNext::FExecutionContext Context(SharedDataBuffer);
 	Instance.GraphInstancePtr = Context.AllocateNodeInstance(UE::AnimNext::FWeakDecoratorPtr(), ResolvedRootDecoratorHandle);
 
-	Instance.ExtendedExecuteContext.CopyMemoryStorage(ExtendedExecuteContext, RigVM);
+	Instance.ExtendedExecuteContext.CopyMemoryStorage(ExtendedExecuteContext);
 
-	RigVM->InitializeInstance(Instance.ExtendedExecuteContext, RigVM->GetLocalMemoryArray(Instance.ExtendedExecuteContext));
+	RigVM->InitializeInstance(Instance.ExtendedExecuteContext);
 }
 
 void UAnimNextGraph::Run(const UE::AnimNext::FContext& Context, FAnimNextGraphInstance& GraphInstance, EAnimNextGraphSimulationSteps SimulationSteps) const
@@ -72,7 +72,7 @@ void UAnimNextGraph::Run(const UE::AnimNext::FContext& Context, FAnimNextGraphIn
 		AnimNextContext.InitializeWithGraph(SharedDataBuffer, GraphInstance.GraphInstancePtr);
 		AnimNextContext.SetSimulationSteps(SimulationSteps);
 
-		RigVM->Execute(GraphInstance.ExtendedExecuteContext, FRigUnit_AnimNextShimRoot::EventName);
+		RigVM->ExecuteVM(GraphInstance.ExtendedExecuteContext, FRigUnit_AnimNextShimRoot::EventName);
 
 		// Reset the context to avoid issues if we forget to reset it the next time we use it
 		AnimNextContext.DebugReset();
diff --git a/Engine/Plugins/Experimental/AnimNext/Source/AnimNext/Private/Module.cpp b/Engine/Plugins/Experimental/AnimNext/Source/AnimNext/Private/Module.cpp
--- a/Engine/Plugins/Experimental/AnimNext/Source/AnimNext/Private/Module.cpp
+++ b/Engine/Plugins/Experimental/AnimNext/Source/AnimNext/Private/Module.cpp
@@ -9,6 +9,7 @@
 #include "DecoratorBase/DecoratorRegistry.h"
 #include "DecoratorBase/NodeTemplateRegistry.h"
 #include "RigVMCore/RigVMRegistry.h"
+#include "RigVMRuntimeDataRegistry.h"
 #include "Animation/AnimSequence.h"
 #include "Scheduler/Scheduler.h"
 #include "Param/Params.h"
@@ -38,10 +39,12 @@ public:
 		FDecoratorRegistry::Init();
 		FNodeTemplateRegistry::Init();
 		FScheduler::Init();
+		FRigVMRuntimeDataRegistry::Init();
 	}
 
 	virtual void ShutdownModule() override
 	{
+		FRigVMRuntimeDataRegistry::Destroy();
 		FScheduler::Destroy();
 		FNodeTemplateRegistry::Destroy();
 		FDecoratorRegistry::Destroy();
diff --git a/Engine/Plugins/Experimental/AnimNext/Source/AnimNext/Private/Param/AnimNextParameterBlock.cpp b/Engine/Plugins/Experimental/AnimNext/Source/AnimNext/Private/Param/AnimNextParameterBlock.cpp
--- a/Engine/Plugins/Experimental/AnimNext/Source/AnimNext/Private/Param/AnimNextParameterBlock.cpp
+++ b/Engine/Plugins/Experimental/AnimNext/Source/AnimNext/Private/Param/AnimNextParameterBlock.cpp
@@ -2,6 +2,7 @@
 
 #include "Param/AnimNextParameterBlock.h"
 #include "RigVMCore/RigVMMemoryStorage.h"
+#include "RigVMRuntimeDataRegistry.h"
 #include "UObject/ObjectSaveContext.h"
 #include "UObject/Package.h"
 #include "Graph/RigUnit_AnimNextBeginExecution.h"
@@ -31,15 +32,27 @@ static TArray<UClass*> GetClassObjectsInPackage(UPackage* InPackage)
 
 }
 
+UAnimNextParameterBlock::UAnimNextParameterBlock(const FObjectInitializer& ObjectInitializer)
+	: Super(ObjectInitializer)
+{
+	SetRigVMExtendedExecuteContext(&BaseRigVMContext);
+}
+
 void UAnimNextParameterBlock::UpdateLayer(UE::AnimNext::FParamStackLayerHandle& InHandle) const
 {
 	if (VM)
 	{
-		FRigVMExtendedExecuteContext Context = GetExtendedExecuteContext();
+		UE::AnimNext::FRigVMRuntimeData* RuntimeData = UE::AnimNext::FRigVMRuntimeDataRegistry::FindRuntimeData(VM);
+		if (RuntimeData == nullptr || RuntimeData->Context.VMHash != VM->GetVMHash())
+		{
+			RuntimeData = UE::AnimNext::FRigVMRuntimeDataRegistry::AddRuntimeData(VM, GetRigVMExtendedExecuteContext());
+			VM->InitializeInstance(RuntimeData->Context, false); // TODO zzz : Temp until VM supports WorkData cached handles using offsets (UE-197067)
+		}
+
+		FRigVMExtendedExecuteContext& Context = RuntimeData->Context;
 		FAnimNextParameterExecuteContext& AnimNextParameterContext = Context.GetPublicDataSafe<FAnimNextParameterExecuteContext>();
 		AnimNextParameterContext.SetParamContextData(InHandle);
-		TArray<FRigVMMemoryStorageStruct*> LocalMemory = VM->GetLocalMemoryArray(Context);
-		VM->Execute(Context, LocalMemory, FRigUnit_AnimNextBeginExecution::EventName);
+		VM->ExecuteVM(Context, FRigUnit_AnimNextBeginExecution::EventName);
 	}
 }
 
@@ -104,8 +117,6 @@ void UAnimNextParameterBlock::GetPreloadDependencies(TArray<UObject*>& OutDeps)
 
 void UAnimNextParameterBlock::PostLoad()
 {
-	VM = RigVM;
-
 	Super::PostLoad();
 
 	// In packaged builds, initialize the VM
diff --git a/Engine/Plugins/Experimental/AnimNext/Source/AnimNext/Private/RigVMRuntimeDataRegistry.cpp b/Engine/Plugins/Experimental/AnimNext/Source/AnimNext/Private/RigVMRuntimeDataRegistry.cpp
new file mode 100644
--- /dev/null
+++ b/Engine/Plugins/Experimental/AnimNext/Source/AnimNext/Private/RigVMRuntimeDataRegistry.cpp
@@ -0,0 +1,109 @@
+#include "RigVMRuntimeDataRegistry.h"
+#include "Misc/ScopeRWLock.h"
+
+namespace UE::AnimNext
+{
+
+namespace Private
+{
+
+static bool bInitialized = false;
+static FDelegateHandle PostGarbageCollectHandle;
+
+static std::atomic<uint32> GCCycle = 0; // Main thread GC counter, incremented when a VM has been deleted
+
+static thread_local uint32 LocalGCCycle = 0;	// Local thread GC counter, used to compare with main and trigger compaction if different
+static thread_local TMap<FRigVMRuntimeDataID, FRigVMRuntimeData> RuntimeDataStorage;
+
+} // end namespace Private
+
+
+/*static*/ void FRigVMRuntimeDataRegistry::Init()
+{
+	if (!Private::bInitialized)
+	{
+		Private::PostGarbageCollectHandle = FCoreUObjectDelegates::GetPostGarbageCollect().AddStatic(&FRigVMRuntimeDataRegistry::HandlePostGarbageCollect);
+
+		Private::bInitialized = true;
+	}
+}
+
+/*static*/ void FRigVMRuntimeDataRegistry::Destroy()
+{
+	if (Private::bInitialized)
+	{
+		Private::bInitialized = false;
+
+		FCoreUObjectDelegates::GetPostGarbageCollect().Remove(Private::PostGarbageCollectHandle);
+		Private::RuntimeDataStorage.Empty();
+	}
+}
+
+/*static*/ FRigVMRuntimeData* FRigVMRuntimeDataRegistry::FindRuntimeData(const FRigVMRuntimeDataID& RigVMRuntimeDataID)
+{
+	check(Private::bInitialized);
+
+	const uint32 CurrentCycle = Private::GCCycle;
+	if (CurrentCycle != Private::LocalGCCycle)
+	{
+		PerformStorageCompaction();
+
+		Private::LocalGCCycle = CurrentCycle;
+	}
+
+	return Private::RuntimeDataStorage.Find(RigVMRuntimeDataID);
+}
+
+/*static*/ FRigVMRuntimeData* FRigVMRuntimeDataRegistry::AddRuntimeData(const FRigVMRuntimeDataID& RigVMRuntimeDataID, const FRigVMExtendedExecuteContext& ReferenceContext)
+{
+	check(Private::bInitialized);
+
+	FRigVMRuntimeData* RigVMRuntimeData = nullptr;
+
+	FRigVMRuntimeData& NewRigVMRuntimeData = Private::RuntimeDataStorage.Add(RigVMRuntimeDataID, FRigVMRuntimeData());
+	NewRigVMRuntimeData.Context = ReferenceContext;
+	RigVMRuntimeData = &NewRigVMRuntimeData;
+
+	return RigVMRuntimeData;
+}
+
+/*static*/ FRigVMRuntimeData* FRigVMRuntimeDataRegistry::FindOrAddRuntimeData(const FRigVMRuntimeDataID& RigVMRuntimeDataID, const FRigVMExtendedExecuteContext& ReferenceContext)
+{
+	check(Private::bInitialized);
+
+	FRigVMRuntimeData* RigVMRuntimeData = nullptr;
+
+	if (RigVMRuntimeData = FindRuntimeData(RigVMRuntimeDataID); RigVMRuntimeData != nullptr)
+	{
+		if (RigVMRuntimeData->Context.VMHash != RigVMRuntimeData->Context.VMHash)
+		{
+			RigVMRuntimeData->Context = ReferenceContext;
+		}
+
+		return RigVMRuntimeData;
+	}
+
+	return AddRuntimeData(RigVMRuntimeDataID, ReferenceContext);
+}
+
+/*static*/ void FRigVMRuntimeDataRegistry::HandlePostGarbageCollect()
+{
+	Private::GCCycle++;
+	Private::LocalGCCycle = Private::GCCycle; // avoid additional compactions on main thread
+
+	PerformStorageCompaction();
+}
+
+/*static*/ void FRigVMRuntimeDataRegistry::PerformStorageCompaction()
+{
+	for (auto Iter = Private::RuntimeDataStorage.CreateIterator(); Iter; ++Iter)
+	{
+		const FRigVMRuntimeDataID& RuntimeDataID = Iter.Key();
+		if (RuntimeDataID.ResolveObjectPtr() == nullptr)
+		{
+			Iter.RemoveCurrent();
+		}
+	}
+}
+
+} // end namespace UE::AnimNext
diff --git a/Engine/Plugins/Experimental/AnimNext/Source/AnimNextUncookedOnly/Private/UncookedOnlyUtils.cpp b/Engine/Plugins/Experimental/AnimNext/Source/AnimNextUncookedOnly/Private/UncookedOnlyUtils.cpp
--- a/Engine/Plugins/Experimental/AnimNext/Source/AnimNextUncookedOnly/Private/UncookedOnlyUtils.cpp
+++ b/Engine/Plugins/Experimental/AnimNext/Source/AnimNextUncookedOnly/Private/UncookedOnlyUtils.cpp
@@ -30,6 +30,7 @@
 #include "DecoratorBase/Decorator.h"
 #include "Graph/RigUnit_AnimNextBeginExecution.h"
 #include "Serialization/MemoryReader.h"
+#include "RigVMRuntimeDataRegistry.h"
 
 #include "RigVMCompiler/RigVMCompiler.h"
 #include "RigVMCore/RigVM.h"
@@ -463,7 +464,7 @@ void FUtils::Compile(UAnimNextGraph* InGraph)
 	Compiler->Compile(Settings, { VMTempGraph }, TempController, InGraph->RigVM, InGraph->ExtendedExecuteContext, InGraph->GetRigVMExternalVariables(), & EditorData->PinToOperandMap);
 
 	// Initialize right away, in packaged builds we initialize during PostLoad
-	InGraph->RigVM->Initialize(InGraph->ExtendedExecuteContext, InGraph->RigVM->GetLocalMemoryArray(InGraph->ExtendedExecuteContext));
+	InGraph->RigVM->Initialize(InGraph->ExtendedExecuteContext);
 
 	if (EditorData->bErrorsDuringCompilation)
 	{
@@ -489,15 +490,6 @@ void FUtils::Compile(UAnimNextGraph* InGraph)
 void FUtils::RecreateVM(UAnimNextGraph* InGraph)
 {
 	InGraph->RigVM = NewObject<URigVM>(InGraph, TEXT("VM"), RF_NoFlags);
-
-	// Cooked platforms will load these pointers from disk
-	if (!FPlatformProperties::RequiresCookedData())
-	{
-		InGraph->RigVM->CreateMemoryByType(InGraph->ExtendedExecuteContext, ERigVMMemoryType::Work);
-		InGraph->RigVM->CreateMemoryByType(InGraph->ExtendedExecuteContext, ERigVMMemoryType::Literal);
-		InGraph->RigVM->CreateMemoryByType(InGraph->ExtendedExecuteContext, ERigVMMemoryType::Debug);
-	}
-
 	InGraph->RigVM->Reset(InGraph->ExtendedExecuteContext);
 }
 
@@ -641,19 +633,15 @@ void FUtils::CompileVM(UAnimNextParameterBlock* InParameterBlock)
 
 	URigVMCompiler* Compiler = URigVMCompiler::StaticClass()->GetDefaultObject<URigVMCompiler>();
 	EditorData->VMCompileSettings.SetExecuteContextStruct(EditorData->RigVMClient.GetExecuteContextStruct());
-	FRigVMExtendedExecuteContext& CDOContext = InParameterBlock->GetExtendedExecuteContext();
+	FRigVMExtendedExecuteContext& CDOContext = InParameterBlock->GetRigVMExtendedExecuteContext();
 	const FRigVMCompileSettings Settings = (EditorData->bCompileInDebugMode) ? FRigVMCompileSettings::Fast(EditorData->VMCompileSettings.GetExecuteContextStruct()) : EditorData->VMCompileSettings;
 	URigVMController* RootController = EditorData->GetRigVMClient()->GetOrCreateController(EditorData->GetRigVMClient()->GetDefaultModel());
-	Compiler->Compile(Settings, EditorData->GetRigVMClient()->GetAllModels(false, false), RootController, InParameterBlock->RigVM, InParameterBlock->GetExtendedExecuteContext(), InParameterBlock->GetExternalVariables(), &EditorData->PinToOperandMap);
+	Compiler->Compile(Settings, EditorData->GetRigVMClient()->GetAllModels(false, false), RootController, InParameterBlock->VM, CDOContext, InParameterBlock->GetExternalVariables(), &EditorData->PinToOperandMap);
 
-	InParameterBlock->RigVM->SetVMHash(InParameterBlock->RigVM->ComputeVMHash(CDOContext));
-	CDOContext.VMHash = InParameterBlock->RigVM->GetVMHash();
-	InParameterBlock->RigVM->Initialize(CDOContext, InParameterBlock->RigVM->GetLocalMemoryArray(CDOContext));
+	InParameterBlock->VM->Initialize(CDOContext);
 	InParameterBlock->GenerateUserDefinedDependenciesData(CDOContext);
 
 	// Notable difference with vanilla RigVM host behavior - we init the VM here at the moment as we only have one 'instance'
-	InParameterBlock->bTEMP_CopyDefaultsFromCDO = false;
-	InParameterBlock->VM = InParameterBlock->RigVM;
 	InParameterBlock->InitializeVM(FRigUnit_AnimNextBeginExecution::EventName);
 
 	if (EditorData->bErrorsDuringCompilation)
@@ -664,9 +652,9 @@ void FUtils::CompileVM(UAnimNextParameterBlock* InParameterBlock)
 				TEXT("Compilation Errors may be suppressed for ControlRigBlueprint: %s. See VM Compile Setting in Class Settings for more Details"), *InParameterBlock->GetName());
 		}
 		EditorData->bVMRecompilationRequired = false;
-		if(InParameterBlock->RigVM)
+		if(InParameterBlock->VM)
 		{
-			EditorData->RigVMCompiledEvent.Broadcast(InParameterBlock, InParameterBlock->RigVM, InParameterBlock->GetExtendedExecuteContext());
+			EditorData->RigVMCompiledEvent.Broadcast(InParameterBlock, InParameterBlock->VM, InParameterBlock->GetRigVMExtendedExecuteContext());
 		}
 		return;
 	}
@@ -727,18 +715,8 @@ void FUtils::Compile(UAnimNextParameterBlock* InParameterBlock)
 
 void FUtils::RecreateVM(UAnimNextParameterBlock* InParameterBlock)
 {
-	InParameterBlock->RigVM = NewObject<URigVM>(InParameterBlock, TEXT("VM"), RF_NoFlags);
-
-	// Cooked platforms will load these pointers from disk
-	if (!FPlatformProperties::RequiresCookedData())
-	{
-		// We dont support ERigVMMemoryType::Work memory as we dont operate on an instance
-	//	InParameterBlock->RigVM->GetMemoryByType(ERigVMMemoryType::Work, true);
-		InParameterBlock->RigVM->CreateMemoryByType(InParameterBlock->GetExtendedExecuteContext(), ERigVMMemoryType::Literal);
-		InParameterBlock->RigVM->CreateMemoryByType(InParameterBlock->GetExtendedExecuteContext(), ERigVMMemoryType::Debug);
-	}
-
-	InParameterBlock->RigVM->Reset(InParameterBlock->GetExtendedExecuteContext());
+	InParameterBlock->VM = NewObject<URigVM>(InParameterBlock, TEXT("VM"), RF_NoFlags);
+	InParameterBlock->VM->Reset(InParameterBlock->GetRigVMExtendedExecuteContext());
 }
 
 UAnimNextParameterBlock_EditorData* FUtils::GetEditorData(const UAnimNextParameterBlock* InParameterBlock)
diff --git a/Engine/Plugins/Runtime/RigVM/Source/RigVM/Private/RigVMCore/RigVM.cpp b/Engine/Plugins/Runtime/RigVM/Source/RigVM/Private/RigVMCore/RigVM.cpp
--- a/Engine/Plugins/Runtime/RigVM/Source/RigVM/Private/RigVMCore/RigVM.cpp
+++ b/Engine/Plugins/Runtime/RigVM/Source/RigVM/Private/RigVMCore/RigVM.cpp
@@ -86,8 +86,7 @@ UScriptStruct* FRigVMParameter::GetScriptStruct() const
 }
 
 URigVM::URigVM()
-	: LiteralMemoryStorageObject(nullptr)
-	, ByteCodePtr(&ByteCodeStorage)
+	: ByteCodePtr(&ByteCodeStorage)
     , FunctionNamesPtr(&FunctionNamesStorage)
     , FunctionsPtr(&FunctionsStorage)
     , FactoriesPtr(&FactoriesStorage)
@@ -150,6 +149,8 @@ void URigVM::Save(FArchive& Ar)
 	Ar << OperandToDebugRegisters;
 	
 	Ar << LiteralMemoryStorage;
+	Ar << DefaultWorkMemoryStorage;
+	Ar << DefaultDebugMemoryStorage;
 }
 
 void URigVM::Load(FArchive& Ar)
@@ -258,6 +259,12 @@ void URigVM::Load(FArchive& Ar)
 	{
 		Ar << LiteralMemoryStorage;
 	}
+
+	if (Ar.CustomVer(FRigVMObjectVersion::GUID) >= FRigVMObjectVersion::VMMemoryStorageDefaultsGeneratedAtVM)
+	{
+		Ar << DefaultWorkMemoryStorage;
+		Ar << DefaultDebugMemoryStorage;
+	}
 }
 
 void URigVM::CopyDataForSerialization(URigVM* InVM)
@@ -276,6 +283,8 @@ void URigVM::CopyDataForSerialization(URigVM* InVM)
 	OperandToDebugRegisters = InVM->OperandToDebugRegisters;
 
 	LiteralMemoryStorage = InVM->LiteralMemoryStorage;
+	DefaultWorkMemoryStorage = InVM->DefaultWorkMemoryStorage;
+	DefaultDebugMemoryStorage = InVM->DefaultDebugMemoryStorage;
 }
 
 void URigVM::PostLoad()
@@ -354,7 +363,7 @@ uint32 URigVM::GetVMHash() const
 	return CachedVMHash;
 }
 
-uint32 URigVM::ComputeVMHash(const FRigVMExtendedExecuteContext& Context) const
+uint32 URigVM::ComputeVMHash() const
 {
 	uint32 Hash = 0;
 	for(const FName& FunctionName : GetFunctionNames())
@@ -371,7 +380,7 @@ uint32 URigVM::ComputeVMHash(const FRigVMExtendedExecuteContext& Context) const
 	}
 
 	Hash = HashCombine(Hash, LiteralMemoryStorage.GetMemoryHash());
-	Hash = HashCombine(Hash, Context.WorkMemoryStorage.GetMemoryHash());
+	Hash = HashCombine(Hash, DefaultWorkMemoryStorage.GetMemoryHash());
 	
 	return Hash;
 }
@@ -403,12 +412,12 @@ UClass* URigVM::GetNativizedClass(const TArray<FRigVMExternalVariableDef>& InExt
 	return nullptr;
 }
 
-bool URigVM::ValidateAllOperandsDuringLoad(FRigVMExtendedExecuteContext& Context)
+bool URigVM::ValidateAllOperandsDuringLoad()
 {
 	// check all operands on all ops for validity
 	bool bAllOperandsValid = true;
 
-	TArray<FRigVMMemoryStorageStruct*> LocalMemory = { GetWorkMemory(Context), GetLiteralMemory(), GetDebugMemory(Context) };
+	const TArray<const FRigVMMemoryStorageStruct*> LocalMemory = { &GetDefaultWorkMemory(), &GetDefaultLiteralMemory(), &GetDefaultDebugMemory() };
 	
 	auto CheckOperandValidity = [LocalMemory, &bAllOperandsValid, this](const FRigVMOperand& InOperand) -> bool
 	{
@@ -621,92 +630,132 @@ FString URigVM::GetRigVMFunctionName(int32 InFunctionIndex) const
 	return GetFunctionNames()[InFunctionIndex].ToString();
 }
 
-void URigVM::CreateMemoryByType(UObject* Outer, TObjectPtr<URigVMMemoryStorage>& MemoryStorage, ERigVMMemoryType InMemoryType, EObjectFlags InObjectFlags, bool bForceCreation)
+FRigVMMemoryStorageStruct* URigVM::GetMemoryByType(FRigVMExtendedExecuteContext& Context, ERigVMMemoryType InMemoryType/*, bool bCreateIfNeeded*/)
 {
-	if (MemoryStorage != nullptr && bForceCreation)
-	{
-		MemoryStorage->Rename(nullptr, GetTransientPackage(), REN_ForceNoResetLoaders | REN_DoNotDirty | REN_DontCreateRedirectors | REN_NonTransactional);
-		MemoryStorage = nullptr;
-	}
+	FRigVMMemoryStorageStruct* MemoryStorage = nullptr;
 
-	if (MemoryStorage == nullptr)
+	switch(InMemoryType)
 	{
-		if (UClass* Class = URigVMMemoryStorageGeneratorClass::GetStorageClass(Outer, InMemoryType))
+		case ERigVMMemoryType::Literal:
 		{
-			// The compiler creates the memory in the CDO, so we just use the storage CDO for all memory types
-			// Later, when we instantiate a VM, we will clone Work and Debug objects for each instance
-			MemoryStorage = Cast<URigVMMemoryStorage>(Class->GetDefaultObject(true));
+			MemoryStorage = GetLiteralMemory();
+			break;
 		}
-		else
+
+		case ERigVMMemoryType::Work:
 		{
-			MemoryStorage = NewObject<URigVMMemoryStorage>(Outer, FName(), InObjectFlags);
+			MemoryStorage = &Context.WorkMemoryStorage;
+			break;
+		}
+
+		case ERigVMMemoryType::Debug:
+		{
+			MemoryStorage = &Context.DebugMemoryStorage;
+			break;
+		}
+
+		default:
+		{
+			break;
 		}
 	}
+
+	return MemoryStorage;
 }
 
-URigVMMemoryStorage* URigVM::CreateMemoryByType(FRigVMExtendedExecuteContext& Context, ERigVMMemoryType InMemoryType, bool bForceCreation)
+const FRigVMMemoryStorageStruct* URigVM::GetMemoryByType(const FRigVMExtendedExecuteContext& Context, ERigVMMemoryType InMemoryType) const
 {
-	URigVMMemoryStorage* MemoryStorage = nullptr;
+	return const_cast<URigVM*>(this)->GetMemoryByType(const_cast<FRigVMExtendedExecuteContext&>(Context), InMemoryType/*, false*/);
+}
 
-	// TODO zzz : Deprecate function
+void URigVM::GenerateMemoryType(FRigVMExtendedExecuteContext& Context, ERigVMMemoryType InMemoryType, const TArray<FRigVMPropertyDescription>* InProperties)
+{
+	GenerateDefaultMemoryType(InMemoryType, InProperties);
 
-	return MemoryStorage;
+	switch (InMemoryType)
+	{
+	case ERigVMMemoryType::Work:
+	{
+		Context.WorkMemoryStorage = DefaultWorkMemoryStorage;
+		break;
+	}
+
+	case ERigVMMemoryType::Debug:
+	{
+#if WITH_EDITOR
+		Context.DebugMemoryStorage = DefaultDebugMemoryStorage;
+#endif
+		break;
+	}
+
+	default:
+	{
+		break;
+	}
+	}
 }
 
-void URigVM::GenerateMemoryType(FRigVMExtendedExecuteContext& Context, ERigVMMemoryType InMemoryType, const TArray<FRigVMPropertyDescription>* InProperties, bool bForceCreation)
+void URigVM::GenerateDefaultMemoryType(ERigVMMemoryType InMemoryType, const TArray<FRigVMPropertyDescription>* InProperties)
 {
-	FRigVMMemoryStorageStruct NewBag = (InProperties != nullptr) ? FRigVMMemoryStorageStruct(InMemoryType, *InProperties) : FRigVMMemoryStorageStruct(InMemoryType);
+	FRigVMMemoryStorageStruct* Memory = nullptr;
 
-	switch(InMemoryType)
+	switch (InMemoryType)
 	{
 		case ERigVMMemoryType::Literal:
 		{
-			LiteralMemoryStorage = NewBag;
+			Memory = &LiteralMemoryStorage;
 			break;
 		}
 
 		case ERigVMMemoryType::Work:
 		{
-			Context.WorkMemoryStorage = NewBag;
+			Memory = &DefaultWorkMemoryStorage;
 			break;
 		}
 
+#if WITH_EDITOR
 		case ERigVMMemoryType::Debug:
 		{
-#if WITH_EDITOR
-			Context.DebugMemoryStorage = NewBag;
-#endif
+			Memory = &DefaultDebugMemoryStorage;
 			break;
 		}
+#endif
 
 		default:
 		{
 			break;
 		}
 	}
+
+	if (Memory != nullptr)
+	{
+		*Memory = (InProperties != nullptr) ? FRigVMMemoryStorageStruct(InMemoryType, *InProperties) : FRigVMMemoryStorageStruct(InMemoryType);
+	}
 }
 
-FRigVMMemoryStorageStruct* URigVM::GetMemoryByType(FRigVMExtendedExecuteContext& Context, ERigVMMemoryType InMemoryType/*, bool bCreateIfNeeded*/)
+FRigVMMemoryStorageStruct* URigVM::GetDefaultMemoryByType(ERigVMMemoryType InMemoryType)
 {
 	FRigVMMemoryStorageStruct* MemoryStorage = nullptr;
 
-	switch(InMemoryType)
+	switch (InMemoryType)
 	{
 		case ERigVMMemoryType::Literal:
 		{
-			MemoryStorage = GetLiteralMemory();
+			MemoryStorage = &LiteralMemoryStorage;
 			break;
 		}
 
 		case ERigVMMemoryType::Work:
 		{
-			MemoryStorage = &Context.WorkMemoryStorage;
+			MemoryStorage = &DefaultWorkMemoryStorage;
 			break;
 		}
 
 		case ERigVMMemoryType::Debug:
 		{
-			MemoryStorage = &Context.DebugMemoryStorage;
+#if WITH_EDITOR
+			MemoryStorage = &DefaultDebugMemoryStorage;
+#endif
 			break;
 		}
 
@@ -719,9 +768,9 @@ FRigVMMemoryStorageStruct* URigVM::GetMemoryByType(FRigVMExtendedExecuteContext&
 	return MemoryStorage;
 }
 
-const FRigVMMemoryStorageStruct* URigVM::GetMemoryByType(const FRigVMExtendedExecuteContext& Context, ERigVMMemoryType InMemoryType) const
+const FRigVMMemoryStorageStruct* URigVM::GetDefaultMemoryByType(ERigVMMemoryType InMemoryType) const
 {
-	return const_cast<URigVM*>(this)->GetMemoryByType(const_cast<FRigVMExtendedExecuteContext&>(Context), InMemoryType/*, false*/);
+	return const_cast<URigVM*>(this)->GetDefaultMemoryByType(InMemoryType/*, false*/);
 }
 
 void URigVM::ClearMemory_Internal()
@@ -757,7 +806,11 @@ void URigVM::ClearMemory_Internal()
 	}
 #endif
 
-	LiteralMemoryStorage = FRigVMMemoryStorageStruct(ERigVMMemoryType::Literal);
+	LiteralMemoryStorage = FRigVMMemoryStorageStruct(ERigVMMemoryType::Invalid);
+	DefaultWorkMemoryStorage = FRigVMMemoryStorageStruct(ERigVMMemoryType::Invalid);
+#if WITH_EDITOR
+	DefaultDebugMemoryStorage = FRigVMMemoryStorageStruct(ERigVMMemoryType::Invalid);
+#endif
 
 	InvalidateCachedMemory_Internal();
 }
@@ -766,15 +819,14 @@ void URigVM::ClearMemory(FRigVMExtendedExecuteContext& Context)
 {
 	ClearMemory_Internal();
 
-	Context.WorkMemoryStorage = FRigVMMemoryStorageStruct(ERigVMMemoryType::Work);
-	Context.DebugMemoryStorage = FRigVMMemoryStorageStruct(ERigVMMemoryType::Debug);
+	Context.WorkMemoryStorage = FRigVMMemoryStorageStruct(ERigVMMemoryType::Invalid);
+	Context.DebugMemoryStorage = FRigVMMemoryStorageStruct(ERigVMMemoryType::Invalid);
 
 	InvalidateCachedMemory(Context);
 }
 
 const FRigVMInstructionArray& URigVM::GetInstructions()
 {
-	//RefreshInstructionsIfRequired();
 	return Instructions;
 }
 
@@ -835,14 +887,14 @@ bool URigVM::CanExecuteEntry(const FRigVMExtendedExecuteContext& Context, const
 
 #if WITH_EDITOR
 
-bool URigVM::ResumeExecution(FRigVMExtendedExecuteContext& Context, TArrayView<FRigVMMemoryStorageStruct*> Memory, const FName& InEntryName)
+bool URigVM::ResumeExecution(FRigVMExtendedExecuteContext& Context, const FName& InEntryName)
 {
 	if (FRigVMDebugInfo* RigVMDebugInfo = Context.GetRigVMDebugInfo())
 	{
 		RigVMDebugInfo->ResumeExecution();
 	}
 
-	return Execute(Context, Memory, InEntryName) != ERigVMExecuteResult::Failed;
+	return ExecuteVM(Context, InEntryName) != ERigVMExecuteResult::Failed;
 }
 
 #endif 
@@ -1067,13 +1119,13 @@ void URigVM::InstructionOpEval(FRigVMExtendedExecuteContext& Context, int32 Inst
 	}
 }
 
-void URigVM::PrepareMemoryForExecution(FRigVMExtendedExecuteContext& Context, TArrayView<FRigVMMemoryStorageStruct*> InMemory)
+void URigVM::PrepareMemoryForExecution(FRigVMExtendedExecuteContext& Context)
 {
 	ensureMsgf(Context.ExecutingThreadId == FPlatformTLS::GetCurrentThreadId(), TEXT("RigVM::CacheMemoryHandlesIfRequired from multiple threads (%d and %d)"), Context.ExecutingThreadId, (int32)FPlatformTLS::GetCurrentThreadId());
 
 	InvalidateCachedMemory(Context);
 
-	if (Instructions.Num() == 0 || InMemory.Num() == 0)
+	if (Instructions.Num() == 0)
 	{
 		return;
 	}
@@ -1109,9 +1161,9 @@ void URigVM::PrepareMemoryForExecution(FRigVMExtendedExecuteContext& Context, TA
 	Context.CachedMemoryHandles.Reset(MemoryHandleCount);
 }
 
-void URigVM::CacheMemoryHandlesIfRequired(FRigVMExtendedExecuteContext& Context, TArrayView<FRigVMMemoryStorageStruct*> InMemory)
+void URigVM::CacheMemoryHandlesIfRequired(FRigVMExtendedExecuteContext& Context)
 {
-	if (Instructions.Num() == 0 || InMemory.Num() == 0)
+	if (Instructions.Num() == 0)
 	{
 		return;
 	}
@@ -1365,7 +1417,7 @@ bool URigVM::ShouldHaltAtInstruction(FRigVMExtendedExecuteContext& Context, cons
 }
 #endif
 
-bool URigVM::Initialize(FRigVMExtendedExecuteContext& Context, TArrayView<FRigVMMemoryStorageStruct*> Memory)
+bool URigVM::Initialize(FRigVMExtendedExecuteContext& Context)
 {
 	if (Context.ExecutingThreadId != INDEX_NONE)
 	{
@@ -1382,25 +1434,15 @@ bool URigVM::Initialize(FRigVMExtendedExecuteContext& Context, TArrayView<FRigVM
 		return true;
 	}
 
-	// changes to the layout of memory array should be reflected in GetContainerIndex()
-	TArray<FRigVMMemoryStorageStruct*> LocalMemory;
-	if (Memory.Num() == 0)
-	{
-		LocalMemory = GetLocalMemoryArray(Context);
-		Memory = LocalMemory;
-	}
-
 	RefreshArgumentNameCaches();
 
-	PrepareMemoryForExecution(Context, Memory);
+	PrepareMemoryForExecution(Context);
 
 	return true;
 }
 
-bool URigVM::InitializeInstance(FRigVMExtendedExecuteContext& Context, TArrayView<FRigVMMemoryStorageStruct*> Memory)
+bool URigVM::InitializeInstance(FRigVMExtendedExecuteContext& Context, bool bCopyMemory)
 {
-	check(!Memory.IsEmpty());
-
 	TGuardValue<int32> GuardThreadId(Context.ExecutingThreadId, FPlatformTLS::GetCurrentThreadId());
 
 	Context.CachedMemoryHandles.Reset(MemoryHandleCount);
@@ -1409,13 +1451,21 @@ bool URigVM::InitializeInstance(FRigVMExtendedExecuteContext& Context, TArrayVie
 	Context.LazyBranchInstanceData.Reset(LazyBranchSize);
 	Context.LazyBranchInstanceData.SetNumZeroed(LazyBranchSize);
 
+	if (bCopyMemory)
+	{
+		Context.VMHash = GetVMHash();
+		Context.WorkMemoryStorage = DefaultWorkMemoryStorage;
+#if WITH_EDITOR
+		Context.DebugMemoryStorage = DefaultDebugMemoryStorage;
+#endif // WITH_EDITOR
+	}
 
-	CacheMemoryHandlesIfRequired(Context, Memory);
+	CacheMemoryHandlesIfRequired(Context);
 
 	return true;
 }
 
-ERigVMExecuteResult URigVM::Execute(FRigVMExtendedExecuteContext& Context, TArrayView<FRigVMMemoryStorageStruct*> Memory, const FName& InEntryName)
+ERigVMExecuteResult URigVM::ExecuteVM(FRigVMExtendedExecuteContext& Context, const FName& InEntryName)
 {
 	// if this the first entry being executed - get ready for execution
 	const bool bIsRootEntry = Context.EntriesBeingExecuted.IsEmpty();
@@ -1441,8 +1491,7 @@ ERigVMExecuteResult URigVM::Execute(FRigVMExtendedExecuteContext& Context, TArra
 			return Context.CurrentExecuteResult = ERigVMExecuteResult::Failed;
 		}
 
-		check(!Memory.IsEmpty());
-		Context.CurrentVMMemory = Memory;
+		Context.CurrentVMMemory = GetInstanceMemory(Context);
 	}
 
 	FRigVMByteCode& ByteCode = GetByteCode();
@@ -1917,7 +1966,7 @@ ERigVMExecuteResult URigVM::ExecuteInstructions(FRigVMExtendedExecuteContext& Co
 				{
 					// this will restore the public data after invoking the entry
 					TGuardValue<FRigVMExecuteContext> PublicDataGuard(ContextPublicData, ContextPublicData);
-					const ERigVMExecuteResult ExecuteResult = Execute(Context, Context.CurrentVMMemory, Op.EntryName);
+					const ERigVMExecuteResult ExecuteResult = ExecuteVM(Context, Op.EntryName);
 					if(ExecuteResult != ERigVMExecuteResult::Succeeded)
 					{
 						return Context.CurrentExecuteResult = ExecuteResult;
@@ -2067,11 +2116,6 @@ ERigVMExecuteResult URigVM::ExecuteInstructions(FRigVMExtendedExecuteContext& Co
 	return Context.CurrentExecuteResult = ERigVMExecuteResult::Succeeded;
 }
 
-bool URigVM::Execute(FRigVMExtendedExecuteContext& Context, const FName& InEntryName)
-{
-	return Execute(Context, GetLocalMemoryArray(Context), InEntryName) != ERigVMExecuteResult::Failed;
-}
-
 ERigVMExecuteResult URigVM::ExecuteBranch(FRigVMExtendedExecuteContext& Context, const FRigVMBranchInfo& InBranchToRun)
 {
 #if WITH_EDITOR
diff --git a/Engine/Plugins/Runtime/RigVM/Source/RigVM/Private/RigVMCore/RigVMExecuteContext.cpp b/Engine/Plugins/Runtime/RigVM/Source/RigVM/Private/RigVMCore/RigVMExecuteContext.cpp
--- a/Engine/Plugins/Runtime/RigVM/Source/RigVM/Private/RigVMCore/RigVMExecuteContext.cpp
+++ b/Engine/Plugins/Runtime/RigVM/Source/RigVM/Private/RigVMCore/RigVMExecuteContext.cpp
@@ -66,9 +66,6 @@ void FRigVMExtendedExecuteContext::Reset()
 
 	ResetExecutionState();
 
-	WorkMemoryStorageObject = nullptr;
-	DebugMemoryStorageObject = nullptr;
-
 	WorkMemoryStorage = FRigVMMemoryStorageStruct();
 	DebugMemoryStorage = FRigVMMemoryStorageStruct();
 
@@ -118,37 +115,16 @@ void FRigVMExtendedExecuteContext::ResetExecutionState()
 	Factory = nullptr;
 }
 
-void FRigVMExtendedExecuteContext::CopyMemoryStorage(const FRigVMExtendedExecuteContext& Other, UObject* Outer)
+void FRigVMExtendedExecuteContext::CopyMemoryStorage(const FRigVMExtendedExecuteContext& Other)
 {
+	VMHash = Other.VMHash;
 	WorkMemoryStorage = Other.WorkMemoryStorage;
 	DebugMemoryStorage = Other.DebugMemoryStorage;
 }
 
-void FRigVMExtendedExecuteContext::CopyMemoryStorage(TObjectPtr<URigVMMemoryStorage>& TargetMemory, const TObjectPtr <URigVMMemoryStorage>& SourceMemory, UObject* Outer)
-{
-	if(SourceMemory != nullptr)
-	{
-		if(TargetMemory == nullptr)
-		{
-			TargetMemory = NewObject<URigVMMemoryStorage>(Outer, SourceMemory->GetClass());
-		}
-		else if(TargetMemory->GetClass() != SourceMemory->GetClass()
-			|| TargetMemory == SourceMemory) // when a instance comes with CDO data automatically copied during instantiation
-		{
-			TargetMemory = NewObject<URigVMMemoryStorage>(Outer, SourceMemory->GetClass());
-		}
-
-		TargetMemory->CopyFrom(SourceMemory);
-	}
-	else if(TargetMemory != nullptr)
-	{
-		TargetMemory = nullptr;
-	}
-}
-
 FRigVMExtendedExecuteContext& FRigVMExtendedExecuteContext::operator =(const FRigVMExtendedExecuteContext& Other)
 {
-	VMHash = Other.VMHash;
+	CopyMemoryStorage(Other);
 
 	const UScriptStruct* OtherPublicDataStruct = Cast<UScriptStruct>(Other.PublicDataScope.GetStruct());
 	check(OtherPublicDataStruct);
diff --git a/Engine/Plugins/Runtime/RigVM/Source/RigVM/Private/RigVMCore/RigVMNativized.cpp b/Engine/Plugins/Runtime/RigVM/Source/RigVM/Private/RigVMCore/RigVMNativized.cpp
--- a/Engine/Plugins/Runtime/RigVM/Source/RigVM/Private/RigVMCore/RigVMNativized.cpp
+++ b/Engine/Plugins/Runtime/RigVM/Source/RigVM/Private/RigVMCore/RigVMNativized.cpp
@@ -27,13 +27,13 @@ void URigVMNativized::Reset(bool IsIgnoringArchetypeRef)
 	ByteCodeStorage.Reset();
 }
 
-bool URigVMNativized::Initialize(FRigVMExtendedExecuteContext& Context, TArrayView<FRigVMMemoryStorageStruct*> Memory)
+bool URigVMNativized::Initialize(FRigVMExtendedExecuteContext& Context)
 {
 	// nothing to do here 
 	return true;
 }
 
-ERigVMExecuteResult URigVMNativized::Execute(FRigVMExtendedExecuteContext& Context, TArrayView<FRigVMMemoryStorageStruct*> Memory, const FName& InEntryName)
+ERigVMExecuteResult URigVMNativized::ExecuteVM(FRigVMExtendedExecuteContext& Context, const FName& InEntryName)
 {
 	// to be implemented by the generated code
 	return ERigVMExecuteResult::Failed;
diff --git a/Engine/Plugins/Runtime/RigVM/Source/RigVM/Private/RigVMHost.cpp b/Engine/Plugins/Runtime/RigVM/Source/RigVM/Private/RigVMHost.cpp
--- a/Engine/Plugins/Runtime/RigVM/Source/RigVM/Private/RigVMHost.cpp
+++ b/Engine/Plugins/Runtime/RigVM/Source/RigVM/Private/RigVMHost.cpp
@@ -111,7 +111,7 @@ void URigVMHost::Serialize(FArchive& Ar)
 	// to make sure they are loaded prior to the VM.
 	if (Ar.IsObjectReferenceCollector() && VM != nullptr)
 	{
-		const TArray<const UObject*> UserDefinedDependencies = GetUserDefinedDependencies({ GetLiteralMemory(), GetWorkMemory() });
+		const TArray<const UObject*> UserDefinedDependencies = GetUserDefinedDependencies({ GetDefaultMemoryByType(ERigVMMemoryType::Literal), GetDefaultMemoryByType(ERigVMMemoryType::Work) });
 		for (const UObject* UserDefinedDependency : UserDefinedDependencies)
 		{
 			if (Cast<UUserDefinedStruct>(UserDefinedDependency) ||
@@ -133,6 +133,8 @@ void URigVMHost::PostLoad()
 {
 	Super::PostLoad();
 
+	FRigVMExtendedExecuteContext& ExtendedExecuteContext = GetRigVMExtendedExecuteContext();
+
 	ExtendedExecuteContext.InvalidateCachedMemory();
 
 	// In packaged builds, initialize the CDO VM
@@ -146,7 +148,7 @@ void URigVMHost::PostLoad()
 			InitializeCDOVM();
 		}
 
-		if (!VM->ValidateAllOperandsDuringLoad(ExtendedExecuteContext))
+		if (!ensure(VM->ValidateAllOperandsDuringLoad()))
 		{
 			VM->Reset(ExtendedExecuteContext);
 		}
@@ -168,7 +170,7 @@ void URigVMHost::PreSave(FObjectPreSaveContext SaveContext)
 {
 	Super::PreSave(SaveContext);
 
-	GenerateUserDefinedDependenciesData(ExtendedExecuteContext);
+	GenerateUserDefinedDependenciesData(GetRigVMExtendedExecuteContext());
 }
 
 void URigVMHost::BeginDestroy()
@@ -181,7 +183,10 @@ void URigVMHost::BeginDestroy()
 	DebugInfo.Reset();
 #endif
 
-	GetExtendedExecuteContext().ExecutionReachedExit().RemoveAll(this);
+	if (RigVMExtendedExecuteContext != nullptr)
+	{
+		RigVMExtendedExecuteContext->ExecutionReachedExit().RemoveAll(this);
+	}
 
 #if WITH_EDITORONLY_DATA
 	if (VMSnapshotBeforeExecution)
@@ -264,25 +269,18 @@ void URigVMHost::Initialize(bool bRequestInit)
 
 bool URigVMHost::InitializeVM(const FName& InEventName)
 {
+	FRigVMExtendedExecuteContext& ExtendedExecuteContext = GetRigVMExtendedExecuteContext();
+
 	const TArray<FRigVMExternalVariable> ExternalVariables = GetExternalVariablesImpl(false);
 	if (VM->GetExternalVariableDefs().Num() != ExternalVariables.Num())
 	{
 		return false;	// The rig did compile with errors
 	}
 
-	if (bTEMP_CopyDefaultsFromCDO && !HasAnyFlags(RF_ClassDefaultObject))
-	{
-		if (URigVMHost* CDO = GetClass()->GetDefaultObject<URigVMHost>())
-		{
-			ExtendedExecuteContext.WorkMemoryStorage = CDO->GetExtendedExecuteContext().WorkMemoryStorage; // initialize the memory to CDO default state
-		}
-	}
-
 	// update the VM's external variables
-	VM->SetExternalVariablesInstanceData(GetExtendedExecuteContext(), ExternalVariables);
+	VM->SetExternalVariablesInstanceData(ExtendedExecuteContext, ExternalVariables);
 
-	TArray<FRigVMMemoryStorageStruct*> LocalMemory = VM->GetLocalMemoryArray(ExtendedExecuteContext);
-	const bool bResult = VM->InitializeInstance(GetExtendedExecuteContext(), LocalMemory);
+	const bool bResult = VM->InitializeInstance(ExtendedExecuteContext);
 	if(bResult)
 	{
 		bRequiresInitExecution = false;
@@ -291,10 +289,7 @@ bool URigVMHost::InitializeVM(const FName& InEventName)
 	// reset the time and caches during init
 	AbsoluteTime = DeltaTime = 0.f;
 		
-	if(VM)
-	{
-		GetExtendedExecuteContext().GetPublicDataSafe<>().GetNameCache()->Reset();
-	}
+	ExtendedExecuteContext.GetPublicDataSafe<>().GetNameCache()->Reset();
 
 	if (InitializedEvent.IsBound())
 	{
@@ -450,7 +445,7 @@ void URigVMHost::InvalidateCachedMemory()
 {
 	if (VM)
 	{
-		VM->InvalidateCachedMemory(ExtendedExecuteContext);
+		VM->InvalidateCachedMemory(GetRigVMExtendedExecuteContext());
 	}
 }
 
@@ -494,7 +489,9 @@ bool URigVMHost::Execute(const FName& InEventName)
 
 	ensure(!HasAnyFlags(RF_ClassDefaultObject));
 	
-	FRigVMExecuteContext& PublicContext = GetExtendedExecuteContext().GetPublicData<>();
+	FRigVMExtendedExecuteContext& ExtendedExecuteContext = GetRigVMExtendedExecuteContext();
+
+	FRigVMExecuteContext& PublicContext = ExtendedExecuteContext.GetPublicData<>();
 	PublicContext.SetDeltaTime(DeltaTime);
 	PublicContext.SetAbsoluteTime(AbsoluteTime);
 	PublicContext.SetFramesPerSecond(GetCurrentFramesPerSecond());
@@ -518,11 +515,11 @@ bool URigVMHost::Execute(const FName& InEventName)
 			}
 		}
 
-		GetExtendedExecuteContext().SetDebugInfo(&DebugInfo);
+		ExtendedExecuteContext.SetDebugInfo(&DebugInfo);
 	}
 	else
 	{
-		GetExtendedExecuteContext().SetDebugInfo(nullptr);
+		ExtendedExecuteContext.SetDebugInfo(nullptr);
 	}
 #endif
 
@@ -533,7 +530,8 @@ bool URigVMHost::Execute(const FName& InEventName)
 		// only set a valid first entry event later when execution
 		// has passed the initialization stage and there are multiple events present in one evaluation
 		// first entry event is used to determined when to clear data during an evaluation
-		VM->SetFirstEntryEventInEventQueue(GetExtendedExecuteContext(), NAME_None);
+		VM->SetFirstEntryEventInEventQueue(GetRigVMExtendedExecuteContext(), NAME_None);
+		VM->SetFirstEntryEventInEventQueue(GetRigVMExtendedExecuteContext(), NAME_None);
 #endif
 	}
 
@@ -624,10 +622,12 @@ bool URigVMHost::InitializeCDOVM()
 	check(VM != nullptr);
 	check(VM->HasAnyFlags(RF_ClassDefaultObject | RF_DefaultSubObject));
 
+	FRigVMExtendedExecuteContext& ExtendedExecuteContext = GetRigVMExtendedExecuteContext();
+
 	// update the VM's external variables
-	VM->ClearExternalVariables(GetExtendedExecuteContext());
+	VM->ClearExternalVariables(ExtendedExecuteContext);
 	VM->SetExternalVariableDefs(GetExternalVariablesImpl(false));
-	return VM->Initialize(GetExtendedExecuteContext(), VM->GetLocalMemoryArray(ExtendedExecuteContext));
+	return VM->Initialize(ExtendedExecuteContext);
 }
 
 bool URigVMHost::Execute_Internal(const FName& InEventName)
@@ -637,7 +637,7 @@ bool URigVMHost::Execute_Internal(const FName& InEventName)
 		return false;
 	}
 
-	FRigVMExtendedExecuteContext& Context = GetExtendedExecuteContext();
+	FRigVMExtendedExecuteContext& ExtendedExecuteContext = GetRigVMExtendedExecuteContext();
 	
 	static constexpr TCHAR InvalidatedVMFormat[] = TEXT("%s: Invalidated VM - aborting execution.");
 	if(VM->IsNativized())
@@ -652,7 +652,7 @@ bool URigVMHost::Execute_Internal(const FName& InEventName)
 	else
 	{
 		// sanity check the validity of the VM to ensure stability.
-		if(!VM->IsContextValidForExecution(Context)
+		if(!VM->IsContextValidForExecution(ExtendedExecuteContext)
 			|| !IsValidLowLevel()
 			|| !VM->IsValidLowLevel()
 		)
@@ -673,8 +673,6 @@ bool URigVMHost::Execute_Internal(const FName& InEventName)
 	
 	const bool bUseDebuggingSnapshots = !VM->IsNativized();
 	
-	TArray<FRigVMMemoryStorageStruct*> LocalMemory = VM->GetLocalMemoryArray(ExtendedExecuteContext);
-
 #if WITH_EDITOR
 	if(bUseDebuggingSnapshots)
 	{
@@ -687,18 +685,18 @@ bool URigVMHost::Execute_Internal(const FName& InEventName)
 			{
 				if(bIsEventFirstInQueue)
 				{
-					CopyVMMemory(GetExtendedExecuteContext(), GetSnapshotContext());
+					CopyVMMemory(GetRigVMExtendedExecuteContext(), GetSnapshotContext());
 				}
 			}
 			else if(bIsEventLastInQueue)
 			{
-				CopyVMMemory(GetSnapshotContext(), GetExtendedExecuteContext());
+				CopyVMMemory(GetSnapshotContext(), GetRigVMExtendedExecuteContext());
 			}
 		}
 	}
 #endif
 
-	const bool bSuccess = VM->Execute(Context, LocalMemory, InEventName) != ERigVMExecuteResult::Failed;
+	const bool bSuccess = VM->ExecuteVM(ExtendedExecuteContext, InEventName) != ERigVMExecuteResult::Failed;
 
 #if UE_RIGVM_PROFILE_EXECUTE_UNITS_NUM
 	const uint64 EndCycles = FPlatformTime::Cycles64();
@@ -792,7 +790,7 @@ void URigVMHost::UpdateVMSettings()
 			});
 #endif
 
-		ExtendedExecuteContext.SetRuntimeSettings(VMRuntimeSettings);
+		GetRigVMExtendedExecuteContext().SetRuntimeSettings(VMRuntimeSettings);
 	}
 }
 
@@ -806,16 +804,22 @@ URigVM* URigVMHost::GetVM()
 	return VM;
 }
 
+const FRigVMMemoryStorageStruct* URigVMHost::GetDefaultMemoryByType(ERigVMMemoryType InMemoryType) const
+{
+	check(VM);
+	return VM->GetDefaultMemoryByType(InMemoryType);
+}
+
 FRigVMMemoryStorageStruct* URigVMHost::GetMemoryByType(ERigVMMemoryType InMemoryType)
 {
 	check(VM);
-	return VM->GetMemoryByType(ExtendedExecuteContext, InMemoryType);
+	return VM->GetMemoryByType(GetRigVMExtendedExecuteContext(), InMemoryType);
 }
 
 const FRigVMMemoryStorageStruct* URigVMHost::GetMemoryByType(ERigVMMemoryType InMemoryType) const
 {
 	check(VM);
-	return VM->GetMemoryByType(ExtendedExecuteContext, InMemoryType);
+	return VM->GetMemoryByType(GetRigVMExtendedExecuteContext(), InMemoryType);
 }
 
 void URigVMHost::DrawIntoPDI(FPrimitiveDrawInterface* PDI, const FTransform& InTransform)
@@ -940,7 +944,7 @@ void URigVMHost::SwapVMToNativizedIfRequired(UClass* InNativizedClass)
 			VM->Rename(nullptr, GetTransientPackage(), REN_ForceNoResetLoaders | REN_DoNotDirty | REN_DontCreateRedirectors | REN_NonTransactional);
 			VM->MarkAsGarbage();
 			VM = NewObject<URigVM>(this, InNativizedClass, TEXT("RigVM_NVMB"), PreviousFlags);
-			GetExtendedExecuteContext().ExecutionReachedExit().AddUObject(this, &URigVMHost::HandleExecutionReachedExit);
+			GetRigVMExtendedExecuteContext().ExecutionReachedExit().AddUObject(this, &URigVMHost::HandleExecutionReachedExit);
 #if UE_RIGVM_PROFILE_EXECUTE_UNITS_NUM
 			ProfilingRunsLeft = 0;
 			AccumulatedCycles = 0;
@@ -985,7 +989,7 @@ void URigVMHost::DeclareConstructClasses(TArray<FTopLevelAssetPath>& OutConstruc
 const FString URigVMHost::GetDebugExecutionString()
 {
 	TGuardValue<bool> DebugExecutionGuard(bDebugExecutionEnabled, true);
-	FRigVMExecuteContext& PublicContext = GetExtendedExecuteContext().GetPublicData<FRigVMExecuteContext>();
+	FRigVMExecuteContext& PublicContext = GetRigVMExtendedExecuteContext().GetPublicData<FRigVMExecuteContext>();
 	PublicContext.DebugMemoryString.Reset();
 	
 	Evaluate_AnyThread();
@@ -1026,15 +1030,15 @@ void URigVMHost::PostInitInstance(URigVMHost* InCDO)
 		RF_Public | RF_DefaultSubObject :
 		RF_Transient | RF_Transactional;
 
-	FRigVMExtendedExecuteContext& Context = GetExtendedExecuteContext();
+	FRigVMExtendedExecuteContext& ExtendedExecuteContext = GetRigVMExtendedExecuteContext();
 
-	Context.SetContextPublicDataStruct(GetPublicContextStruct());
+	ExtendedExecuteContext.SetContextPublicDataStruct(GetPublicContextStruct());
 
-	GetExtendedExecuteContext().ExecutionReachedExit().RemoveAll(this);
-	GetExtendedExecuteContext().ExecutionReachedExit().AddUObject(this, &URigVMHost::HandleExecutionReachedExit);
+	ExtendedExecuteContext.ExecutionReachedExit().RemoveAll(this);
+	ExtendedExecuteContext.ExecutionReachedExit().AddUObject(this, &URigVMHost::HandleExecutionReachedExit);
 
 #if WITH_EDITOR
-	Context.GetPublicData<>().SetLog(RigVMLog); // may be nullptr
+	ExtendedExecuteContext.GetPublicData<>().SetLog(RigVMLog); // may be nullptr
 #endif
 
 	UpdateVMSettings();
@@ -1048,24 +1052,6 @@ void URigVMHost::PostInitInstance(URigVMHost* InCDO)
 			{
 				VM = InCDO->GetVM();
 			}
-
-			// This is needed by some tests that use Templates to run (and RigVMHost does not initialize template CDO rigs)
-			if (InCDO->HasAnyFlags(RF_ArchetypeObject))
-			{
-				CopyVMMemory(GetExtendedExecuteContext(), InCDO->GetExtendedExecuteContext());
-
-				Context.VMHash = InCDO->GetExtendedExecuteContext().VMHash;
-
-				if (!ensure(VM->GetVMHash() == Context.VMHash))
-				{
-					UE_LOG(LogRigVM
-						, Warning
-						, TEXT("RigVMHost : VM Hash [%u] is different from Context stored VM Hash [%u]. Please recompile ControlRig used at Asset : [%s]")
-						, Context.VMHash
-						, VM->GetVMHash()
-						, *GetPathName());
-				}
-			}
 		}
 	}
 	else // we are the CDO
@@ -1097,7 +1083,7 @@ void URigVMHost::GenerateUserDefinedDependenciesData(FRigVMExtendedExecuteContex
 {
 	if (VM)
 	{
-		const TArray<const UObject*> UserDefinedDependencies = GetUserDefinedDependencies({ GetLiteralMemory(), GetWorkMemory() });
+		const TArray<const UObject*> UserDefinedDependencies = GetUserDefinedDependencies({ GetDefaultMemoryByType(ERigVMMemoryType::Literal), GetDefaultMemoryByType(ERigVMMemoryType::Work) });
 		UserDefinedStructGuidToPathName.Reset();
 		UserDefinedEnumToPathName.Reset();
 		UserDefinedTypesInUse.Reset();
@@ -1192,7 +1178,7 @@ void URigVMHost::HandleExecutionReachedExit(const FName& InEventName)
 	{
 		if(URigVM* SnapShotVM = GetSnapshotVM(false))
 		{
-			CopyVMMemory(GetSnapshotContext(), GetExtendedExecuteContext());
+			CopyVMMemory(GetSnapshotContext(), GetRigVMExtendedExecuteContext());
 		}
 		DebugInfo.ResetState();
 		SetBreakpointAction(ERigVMBreakpointAction::None);
@@ -1290,33 +1276,21 @@ void URigVMHost::InstantiateVMFromCDO()
 	{
 		SwapVMToNativizedIfRequired();
 
+		FRigVMExtendedExecuteContext& ExtendedExecuteContext = GetRigVMExtendedExecuteContext();
+
 		URigVMHost* CDO = GetClass()->GetDefaultObject<URigVMHost>();
 		if (VM && CDO && CDO->VM)
 		{
 			if(!VM->IsNativized())
 			{
-				// clone work and debug memory between contexts
-				CopyVMMemory(ExtendedExecuteContext, CDO->GetExtendedExecuteContext());
-
-				ExtendedExecuteContext.VMHash = CDO->GetExtendedExecuteContext().VMHash;
+				ExtendedExecuteContext.WorkMemoryStorage = CDO->VM->GetDefaultWorkMemory();
+				ExtendedExecuteContext.DebugMemoryStorage = CDO->VM->GetDefaultDebugMemory();
+				ExtendedExecuteContext.VMHash = CDO->VM->GetVMHash();
 
 #if WITH_EDITOR
 				// Fix AutoCompile while stopped in a breakpoint
-				if (URigVM* SnapShotVM = GetSnapshotVM(false)) // don't create it if not already there
-				{
-					CopyVMMemory(GetSnapshotContext(), CDO->GetExtendedExecuteContext());
-				}
+				CopyVMMemory(GetSnapshotContext(), ExtendedExecuteContext);
 #endif // WITH_EDITOR
-
-				if (!ensure(VM->GetVMHash() == ExtendedExecuteContext.VMHash))
-				{
-					UE_LOG(LogRigVM
-						, Warning
-						, TEXT("RigVMHost : VM Hash [%u] is different from Context stored VM Hash [%u]. Please recompile ControlRig used at Asset : [%s]")
-						, ExtendedExecuteContext.VMHash
-						, VM->GetVMHash()
-						, *GetPathName());
-				}
 			}
 		}
 		else if (VM)
@@ -1374,21 +1348,9 @@ void URigVMHost::InitializeFromCDO()
 	}
 }
 
-/*static*/ uint32 URigVMHost::ComputeAndUpdateCDOHash(URigVMHost* InCDO)
-{
-	// This has to be calculated after the CDO has been fully instantiated and initialized, as the CDO memory is lazily instantiated and it affects the Hash
-	const uint32 CDOVMHash = InCDO->GetVM()->ComputeVMHash(InCDO->GetExtendedExecuteContext());
-	if (InCDO->GetVM()->GetVMHash() != InCDO->GetExtendedExecuteContext().VMHash || CDOVMHash != InCDO->GetVM()->GetVMHash())
-	{
-		InCDO->GetVM()->SetVMHash(CDOVMHash);
-		InCDO->GetExtendedExecuteContext().VMHash = CDOVMHash;
-	}
-	return CDOVMHash;
-}
-
 void URigVMHost::CopyVMMemory(FRigVMExtendedExecuteContext& TargetContext, const FRigVMExtendedExecuteContext& SourceContext)
 {
-	TargetContext.CopyMemoryStorage(SourceContext, this);
+	TargetContext.CopyMemoryStorage(SourceContext);
 }
 
 void URigVMHost::AddAssetUserData(UAssetUserData* InUserData)
diff --git a/Engine/Plugins/Runtime/RigVM/Source/RigVM/Public/RigVMCore/RigVM.h b/Engine/Plugins/Runtime/RigVM/Source/RigVM/Public/RigVMCore/RigVM.h
--- a/Engine/Plugins/Runtime/RigVM/Source/RigVM/Public/RigVMCore/RigVM.h
+++ b/Engine/Plugins/Runtime/RigVM/Source/RigVM/Public/RigVMCore/RigVM.h
@@ -188,11 +188,8 @@ public:
 	// returns the cached VM hash
 	virtual uint32 GetVMHash() const;
 
-	UE_DEPRECATED(5.4, "Please, use ComputeVMHash with FRigVMExtendedExecuteContext parameter.")
-	virtual uint32 ComputeVMHash() const { return 0; }
-
 	// Generates a unique hash to compare VMs
-	virtual uint32 ComputeVMHash(const FRigVMExtendedExecuteContext& Context) const;
+	virtual uint32 ComputeVMHash() const;
 
 	UE_DEPRECATED(5.3, "Please, use GetNativizedClass with FRigVMExternalVariableDef array.")
 	UClass* GetNativizedClass(const TArray<FRigVMExternalVariable>& InExternalVariables) { return nullptr; }
@@ -222,37 +219,31 @@ public:
 	UE_DEPRECATED(5.3, "Please, use Initialize with Context param")
 	virtual bool Initialize(TArrayView<URigVMMemoryStorage*> Memory) { return false; }
 
-	UE_DEPRECATED(5.4, "Please, use Initialize with FRigVMMemoryStorageStruct param")
+	UE_DEPRECATED(5.4, "Please, use Initialize with just Context param")
 	virtual bool Initialize(FRigVMExtendedExecuteContext& Context, TArrayView<URigVMMemoryStorage*> Memory) { return false; }
 
 	// Prepares caches and memory for execution
-	virtual bool Initialize(FRigVMExtendedExecuteContext& Context, TArrayView<FRigVMMemoryStorageStruct*> Memory);
-
-	UE_DEPRECATED(5.4, "Please, use InitializeInstance with FRigVMMemoryStorageStruct param")
-	virtual bool InitializeInstance(FRigVMExtendedExecuteContext& Context, TArrayView<URigVMMemoryStorage*> Memory) { return false; }
+	virtual bool Initialize(FRigVMExtendedExecuteContext& Context);
 
-	// Initializes cached memory handles and copies work memory from the CDO to the Context
-	virtual bool InitializeInstance(FRigVMExtendedExecuteContext& Context, TArrayView<FRigVMMemoryStorageStruct*> Memory);
+	// Initializes cached memory handles and optionally copies work memory from the CDO to the Context
+	virtual bool InitializeInstance(FRigVMExtendedExecuteContext& Context, bool bCopyMemory = true);
 
 	UE_DEPRECATED(5.3, "Please, use Execute with Context param")
 	virtual ERigVMExecuteResult Execute(TArrayView<URigVMMemoryStorage*> Memory, const FName& InEntryName = NAME_None) { return ERigVMExecuteResult::Failed; }
 
-	UE_DEPRECATED(5.4, "Please, use Execute with FRigVMMemoryStorageStruct param")
+	UE_DEPRECATED(5.4, "Please, use ExecuteVM with Context param")
 	virtual ERigVMExecuteResult Execute(FRigVMExtendedExecuteContext& Context, TArrayView<URigVMMemoryStorage*> Memory, const FName& InEntryName = NAME_None) { return ERigVMExecuteResult::Failed; }
 
-	// Executes the VM.
-	// You can optionally provide external memory to the execution
-	// and provide optional additional operands.
-	virtual ERigVMExecuteResult Execute(FRigVMExtendedExecuteContext& Context, TArrayView<FRigVMMemoryStorageStruct*> Memory, const FName& InEntryName = NAME_None);
-
 	UE_DEPRECATED(5.3, "Please, use Execute with Context param")
 	virtual bool Execute(const FName& InEntryName = NAME_None) { return false; }
 
 	// Executes the VM.
-	// You can optionally provide external memory to the execution
-	// and provide optional additional operands.
-	UFUNCTION(BlueprintCallable, Category = RigVM)
-	virtual bool Execute(FRigVMExtendedExecuteContext& Context, const FName& InEntryName = NAME_None);
+	// You can optionally provide optional additional operands.
+	virtual ERigVMExecuteResult ExecuteVM(FRigVMExtendedExecuteContext& Context, const FName& InEntryName = NAME_None);
+
+	UE_DEPRECATED(5.4, "Please, use ExecuteVM with Context param")
+	UFUNCTION(BlueprintCallable, Category = RigVM, meta = (DeprecatedFunction, DeprecationMessage = "This function has been deprecated and it is no longer supported."))
+	virtual bool Execute(FRigVMExtendedExecuteContext& Context, const FName& InEntryName = NAME_None) { return false; }
 
 	UE_DEPRECATED(5.3, "Please, use ExecuteLazyBranch with Context param")
 	ERigVMExecuteResult ExecuteLazyBranch(const FRigVMBranchInfo& InBranchToRun) { return ERigVMExecuteResult::Failed; }
@@ -278,12 +269,6 @@ public:
 	UE_DEPRECATED(5.4, "GetMemoryByType has been deprecated from the VM. Please, use GetWorkMemory from VMHost or the version with a Context param")
 	virtual URigVMMemoryStorage* GetMemoryByType(ERigVMMemoryType InMemoryType, bool bCreateIfNeeded = true) { return nullptr; }
 
-	// Creates a new memory storage by type
-	void CreateMemoryByType(UObject* Outer, TObjectPtr<URigVMMemoryStorage>& MemoryStorage, ERigVMMemoryType InMemoryType, EObjectFlags ObjectFlags, bool bForceCreation);
-	virtual URigVMMemoryStorage* CreateMemoryByType(FRigVMExtendedExecuteContext& Context, ERigVMMemoryType InMemoryType, bool bForceCreation = false);
-
-	virtual void GenerateMemoryType(FRigVMExtendedExecuteContext& Context, ERigVMMemoryType InMemoryType, const TArray<FRigVMPropertyDescription>* InProperties, bool bForceCreation = false);
-	
 	// Returns a memory storage by type
 	virtual FRigVMMemoryStorageStruct* GetMemoryByType(FRigVMExtendedExecuteContext& Context, ERigVMMemoryType InMemoryType);
 	virtual const FRigVMMemoryStorageStruct* GetMemoryByType(const FRigVMExtendedExecuteContext& Context, ERigVMMemoryType InMemoryType) const;
@@ -304,12 +289,11 @@ public:
 	UE_DEPRECATED(5.4, "GetWorkMemory has been deprecated from the VM. Please, use GetWorkMemory from VMHost or the version with a Context param")
 	URigVMMemoryStorage* GetWorkMemory(bool bCreateIfNeeded = true) { return nullptr; }
 
-	// The default mutable work memory
+	// The instance mutable work memory the VM will use to execute
 	FRigVMMemoryStorageStruct* GetWorkMemory(FRigVMExtendedExecuteContext& Context)
 	{
 		return GetMemoryByType(Context, ERigVMMemoryType::Work);
 	}
-
 	const FRigVMMemoryStorageStruct* GetWorkMemory(const FRigVMExtendedExecuteContext& Context) const
 	{
 		return GetMemoryByType(Context, ERigVMMemoryType::Work);
@@ -318,22 +302,21 @@ public:
 	UE_DEPRECATED(5.4, "GetDebugMemory has been deprecated from the VM. Please, use GetDebugMemory from VMHost or the version with a Context param")
 	URigVMMemoryStorage* GetDebugMemory(bool bCreateIfNeeded = true) { return nullptr; }
 
-	// The default debug watch memory
+	// The instance debug watch memory
 	FRigVMMemoryStorageStruct* GetDebugMemory(FRigVMExtendedExecuteContext& Context)
 	{
 		return GetMemoryByType(Context, ERigVMMemoryType::Debug);
 	}
-
 	const FRigVMMemoryStorageStruct* GetDebugMemory(const FRigVMExtendedExecuteContext& Context) const
 	{
 		return GetMemoryByType(Context, ERigVMMemoryType::Debug);
 	}
 
-	UE_DEPRECATED(5.4, "Please, use GetLocalMemoryArray with a Context param")
+	UE_DEPRECATED(5.4, "Please, use GetInstanceMemoryArray with a Context param")
 	TArray<URigVMMemoryStorage*> GetLocalMemoryArray() { return TArray<URigVMMemoryStorage*>(); }
 
-	// returns all memory storages as an array
-	TArray<FRigVMMemoryStorageStruct*> GetLocalMemoryArray(FRigVMExtendedExecuteContext& Context)
+	// Returns the VM Literal memory and context instanced memory.
+	TArray<FRigVMMemoryStorageStruct*> GetInstanceMemory(FRigVMExtendedExecuteContext& Context)
 	{
 		TArray<FRigVMMemoryStorageStruct*> LocalMemory;
 		LocalMemory.Add(GetWorkMemory(Context));
@@ -342,6 +325,47 @@ public:
 		return LocalMemory;
 	}
 
+	// The default non mutable reference literal memory (generated by the compiler and copied to the instances as initial state)
+	const FRigVMMemoryStorageStruct& GetDefaultLiteralMemory() const
+	{
+		return LiteralMemoryStorage;
+	}
+	FRigVMMemoryStorageStruct& GetDefaultLiteralMemory()
+	{
+		return LiteralMemoryStorage;
+	}
+
+	// The default non mutable reference work memory (generated by the compiler and copied to the instances as initial state)
+	const FRigVMMemoryStorageStruct& GetDefaultWorkMemory() const
+	{
+		return DefaultWorkMemoryStorage;
+	}
+	FRigVMMemoryStorageStruct& GetDefaultWorkMemory()
+	{
+		return DefaultWorkMemoryStorage;
+	}
+
+	// The default non mutable reference debug memory (generated by the compiler and copied to the instances as initial state)
+	const FRigVMMemoryStorageStruct& GetDefaultDebugMemory() const
+	{
+		return DefaultDebugMemoryStorage;
+	}
+	FRigVMMemoryStorageStruct& GetDefaultDebugMemory()
+	{
+		return DefaultDebugMemoryStorage;
+	}
+
+	// Generates the Default memory and copy it to the Context if required by the type of memory
+	virtual void GenerateMemoryType(FRigVMExtendedExecuteContext& Context, ERigVMMemoryType InMemoryType, const TArray<FRigVMPropertyDescription>* InProperties);
+
+	// Used by the compiler and EngineTests to generate the default VM memory storages
+	virtual void GenerateDefaultMemoryType(ERigVMMemoryType InMemoryType, const TArray<FRigVMPropertyDescription>* InProperties);
+
+	// Returns a default memory storage by type
+	virtual const FRigVMMemoryStorageStruct* GetDefaultMemoryByType(ERigVMMemoryType InMemoryType) const;
+	virtual FRigVMMemoryStorageStruct* GetDefaultMemoryByType(ERigVMMemoryType InMemoryType);
+
+
 public:
 
 	UE_DEPRECATED(5.4, "Please, use GetLocalMemoryArray with a Context param")
@@ -349,21 +373,30 @@ public:
 
 	virtual void ClearMemory(FRigVMExtendedExecuteContext& Context);
 
-#if WITH_EDITORONLY_DATA
-	// Deprecated 5.4, please use the WorkMemory in the ExtendedExecuteContext
 	UPROPERTY()
-	TObjectPtr<URigVMMemoryStorage> WorkMemoryStorageObject_DEPRECATED;
-#endif
+	FRigVMMemoryStorageStruct LiteralMemoryStorage;
 
 	UPROPERTY()
-	TObjectPtr<URigVMMemoryStorage> LiteralMemoryStorageObject;
+	FRigVMMemoryStorageStruct DefaultWorkMemoryStorage;
 
 	UPROPERTY()
-	FRigVMMemoryStorageStruct LiteralMemoryStorage;
+	FRigVMMemoryStorageStruct DefaultDebugMemoryStorage;
 
 #if WITH_EDITORONLY_DATA
-	// Deprecated 5.4, please use the DebugMemory in the ExtendedExecuteContext
-	UPROPERTY()
+	// Deprecated 5.4
+	UPROPERTY(meta = (DeprecatedProperty, DeprecationMessage = "Please, use DefaultWorkMemoryStorage for compiling and WorkMemoryStorage in the ExtendedExecuteContext for intance execution"))
+	TObjectPtr<URigVMMemoryStorage> WorkMemoryStorageObject_DEPRECATED;
+#endif
+
+#if WITH_EDITORONLY_DATA
+	// Deprecated 5.4
+	UPROPERTY(meta = (DeprecatedProperty, DeprecationMessage = "Please, use LiteralMemoryStorage"))
+	TObjectPtr<URigVMMemoryStorage> LiteralMemoryStorageObject_DEPRECATED;
+#endif
+
+#if WITH_EDITORONLY_DATA
+	// Deprecated 5.4, 
+	UPROPERTY(meta = (DeprecatedProperty, DeprecationMessage = "Please use DefaultDebugMemoryStorage for compiling and DebugMemoryStorage in the ExtendedExecuteContext for intance execution"))
 	TObjectPtr<URigVMMemoryStorage> DebugMemoryStorageObject_DEPRECATED;
 #endif
 
@@ -475,18 +508,17 @@ public:
 		}
 	}
 
+	bool ResumeExecution(FRigVMExtendedExecuteContext& Context, const FName& InEntryName = NAME_None);
+
 	UE_DEPRECATED(5.3, "Please, use ResumeExecution with Context param")
 	bool ResumeExecution() { return false; }
 	UE_DEPRECATED(5.3, "Please, use ResumeExecution with Context param")
 	bool ResumeExecution(TArrayView<URigVMMemoryStorage*> Memory, const FName& InEntryName = NAME_None) { return false; }
-	UE_DEPRECATED(5.4, "Please, use ResumeExecution with FRigVMMemoryStorageStruct param")
+	UE_DEPRECATED(5.4, "Please, use ResumeExecution with Context param and entry name")
 	bool ResumeExecution(FRigVMExtendedExecuteContext& Context, TArrayView<URigVMMemoryStorage*> Memory, const FName& InEntryName = NAME_None) { return false; }
-
-	UE_DEPRECATED(5.4, "Please, use ResumeExecution in the Context RigVMDebugInfo")
+	UE_DEPRECATED(5.4, "Please, use ResumeExecution in the Host RigVMDebugInfo")
 	bool ResumeExecution(FRigVMExtendedExecuteContext& Context) { return false; }
 
-	bool ResumeExecution(FRigVMExtendedExecuteContext& Context, TArrayView<FRigVMMemoryStorageStruct*> Memory, const FName& InEntryName = NAME_None);
-
 #endif // WITH_EDITOR
 
 	// Returns the parameters of the VM
@@ -937,8 +969,8 @@ public:
 
 private:
 	void InstructionOpEval(FRigVMExtendedExecuteContext& Context, int32 InstructionIndex, int32 InHandleBaseIndex, const TFunctionRef<void(FRigVMExtendedExecuteContext& Context, int32 InHandleIndex, const FRigVMBranchInfoKey& InBranchInfoKey, const FRigVMOperand& InArg)>& InOpFunc);
-	void PrepareMemoryForExecution(FRigVMExtendedExecuteContext& Context, TArrayView<FRigVMMemoryStorageStruct*> InMemory);
-	void CacheMemoryHandlesIfRequired(FRigVMExtendedExecuteContext& Context, TArrayView<FRigVMMemoryStorageStruct*> InMemory);
+	void PrepareMemoryForExecution(FRigVMExtendedExecuteContext& Context);
+	void CacheMemoryHandlesIfRequired(FRigVMExtendedExecuteContext& Context);
 	void RebuildByteCodeOnLoad();
 
 	UPROPERTY(transient)
@@ -961,7 +993,7 @@ public:
 		return ActiveExecutions.load();
 	}
 
-	bool ValidateAllOperandsDuringLoad(FRigVMExtendedExecuteContext& Context);
+	bool ValidateAllOperandsDuringLoad();
 	void RefreshArgumentNameCaches();
 
 private:
@@ -1000,9 +1032,6 @@ private:
 	TArray<FRigVMExternalVariableDef> ExternalVariables;
 	TArray<FRigVMLazyBranch> LazyBranches;
 	
-	//TMap<FString, FSoftObjectPath> UserDefinedStructGuidToPathName;
-	//TMap<FString, FSoftObjectPath> UserDefinedEnumToPathName;
-
 	// this function should be kept in sync with FRigVMOperand::GetContainerIndex()
 	static int32 GetContainerIndex(ERigVMMemoryType InType)
 	{
@@ -1067,6 +1096,7 @@ protected:
 	
 private:
 
+	UPROPERTY()
 	uint32 CachedVMHash = 0;
 
 	mutable TArray<FName> EntryNames;
diff --git a/Engine/Plugins/Runtime/RigVM/Source/RigVM/Public/RigVMCore/RigVMExecuteContext.h b/Engine/Plugins/Runtime/RigVM/Source/RigVM/Public/RigVMCore/RigVMExecuteContext.h
--- a/Engine/Plugins/Runtime/RigVM/Source/RigVM/Public/RigVMCore/RigVMExecuteContext.h
+++ b/Engine/Plugins/Runtime/RigVM/Source/RigVM/Public/RigVMCore/RigVMExecuteContext.h
@@ -549,8 +549,12 @@ struct RIGVM_API FRigVMExtendedExecuteContext
 	/** Resets VM execution state */
 	void ResetExecutionState();
 
-	void CopyMemoryStorage(const FRigVMExtendedExecuteContext& Other, UObject* Outer);
-	static void CopyMemoryStorage(TObjectPtr<URigVMMemoryStorage>& TargetMemory, const TObjectPtr <URigVMMemoryStorage>& SourceMemory, UObject* Outer);
+	void CopyMemoryStorage(const FRigVMExtendedExecuteContext& Other);
+
+	UE_DEPRECATED(5.4, "This function has been deprecated. Please, use CopyMemoryStorage with Context param.")
+	void CopyMemoryStorage(const FRigVMExtendedExecuteContext& Other, UObject* Outer) {}
+	UE_DEPRECATED(5.4, "This function has been deprecated. Please, use CopyMemoryStorage with Context param.")
+	static void CopyMemoryStorage(TObjectPtr<URigVMMemoryStorage>& TargetMemory, const TObjectPtr <URigVMMemoryStorage>& SourceMemory, UObject* Outer) {}
 
 	FRigVMExtendedExecuteContext& operator =(const FRigVMExtendedExecuteContext& Other);
 
@@ -704,20 +708,26 @@ struct RIGVM_API FRigVMExtendedExecuteContext
 		return 0;
 	}
 
-	UPROPERTY()
+	UPROPERTY(transient)
 	uint32 VMHash = 0;
 
 	UPROPERTY(transient)
-	TObjectPtr<URigVMMemoryStorage> WorkMemoryStorageObject;
+	FRigVMMemoryStorageStruct WorkMemoryStorage;
 
 	UPROPERTY(transient)
-	TObjectPtr<URigVMMemoryStorage> DebugMemoryStorageObject;
+	FRigVMMemoryStorageStruct DebugMemoryStorage;
 
-	UPROPERTY()
-	FRigVMMemoryStorageStruct WorkMemoryStorage;
+#if WITH_EDITORONLY_DATA
+	// Deprecated 5.4
+	UPROPERTY(transient, meta = (DeprecatedProperty, DeprecationMessage = "Please, use WorkMemoryStorage"))
+	TObjectPtr<URigVMMemoryStorage> WorkMemoryStorageObject_DEPRECATED;
+#endif
 
-	UPROPERTY()
-	FRigVMMemoryStorageStruct DebugMemoryStorage;
+#if WITH_EDITORONLY_DATA
+	// Deprecated 5.4
+	UPROPERTY(transient, meta = (DeprecatedProperty, DeprecationMessage = "Please, use DebugMemoryStorage"))
+	TObjectPtr<URigVMMemoryStorage> DebugMemoryStorageObject_DEPRECATED;
+#endif
 
 	FStructOnScope PublicDataScope;
 	URigVM* VM = nullptr;
diff --git a/Engine/Plugins/Runtime/RigVM/Source/RigVM/Public/RigVMCore/RigVMNativized.h b/Engine/Plugins/Runtime/RigVM/Source/RigVM/Public/RigVMCore/RigVMNativized.h
--- a/Engine/Plugins/Runtime/RigVM/Source/RigVM/Public/RigVMCore/RigVMNativized.h
+++ b/Engine/Plugins/Runtime/RigVM/Source/RigVM/Public/RigVMCore/RigVMNativized.h
@@ -50,8 +50,8 @@ public:
 	virtual bool IsNativized() const override { return true; }
 	virtual void Empty(FRigVMExtendedExecuteContext& Context) override { return; }
 	virtual void CopyFrom(URigVM* InVM, bool bDeferCopy = false, bool bReferenceLiteralMemory = false, bool bReferenceByteCode = false, bool bCopyExternalVariables = false, bool bCopyDynamicRegisters = false) override { return; }
-	virtual bool Initialize(FRigVMExtendedExecuteContext& Context, TArrayView<FRigVMMemoryStorageStruct*> Memory) override;
-	virtual ERigVMExecuteResult Execute(FRigVMExtendedExecuteContext& Context, TArrayView<FRigVMMemoryStorageStruct*> Memory, const FName& InEntryName = NAME_None) override;
+	virtual bool Initialize(FRigVMExtendedExecuteContext& Context) override;
+	virtual ERigVMExecuteResult ExecuteVM(FRigVMExtendedExecuteContext& Context, const FName& InEntryName = NAME_None) override;
 	virtual int32 AddRigVMFunction(UScriptStruct* InRigVMStruct, const FName& InMethodName) override { return INDEX_NONE; }
 	virtual FString GetRigVMFunctionName(int32 InFunctionIndex) const override { return FString(); }
 	virtual FRigVMMemoryStorageStruct* GetMemoryByType(FRigVMExtendedExecuteContext& Context, ERigVMMemoryType InMemoryType) override { return nullptr; }
@@ -76,13 +76,13 @@ public:
 	virtual void Empty() override {}
 	UE_DEPRECATED(5.3, "Please, use Initialize with Context param")
 	virtual bool Initialize(TArrayView<URigVMMemoryStorage*> Memory) override { return false; }
-	UE_DEPRECATED(5.4, "Please, use Execute with FRigVMMemoryStorageStruct param")
-	virtual bool Initialize(FRigVMExtendedExecuteContext& Context, TArrayView<URigVMMemoryStorage*> Memory) override { return false; }
-	UE_DEPRECATED(5.4, "GetMemoryByType has been deprecated from the VM. Please, use GetWorkMemory from VMHost or the new GetMemoryByType with a FRigVMMemoryStorageStruct and Context parameter.")
+	UE_DEPRECATED(5.4, "Please, use Initialize just with Context param")
+	virtual bool Initialize(FRigVMExtendedExecuteContext& Context, TArrayView<URigVMMemoryStorage*> Memory) { return false; }
+	UE_DEPRECATED(5.4, "GetMemoryByType has been deprecated from the VM. Please, use GetWorkMemory from VMHost or the new GetMemoryByType with Context parameter.")
 	virtual URigVMMemoryStorage* GetMemoryByType(ERigVMMemoryType InMemoryType, bool bCreateIfNeeded = false) override { return nullptr; }
 	UE_DEPRECATED(5.3, "Please, use Execute with Context param")
 	virtual ERigVMExecuteResult Execute(TArrayView<URigVMMemoryStorage*> Memory, const FName& InEntryName = NAME_None) override { return ERigVMExecuteResult::Failed; }
-	UE_DEPRECATED(5.4, "Please, use Execute with FRigVMMemoryStorageStruct param")
+	UE_DEPRECATED(5.4, "Please, use Execute with just Context param")
 	virtual ERigVMExecuteResult Execute(FRigVMExtendedExecuteContext& Context, TArrayView<URigVMMemoryStorage*> Memory, const FName& InEntryName = NAME_None) override { return ERigVMExecuteResult::Failed; }
 	UE_DEPRECATED(5.3, "Please, use SetInstructionIndex with Context param")
 	virtual void SetInstructionIndex(uint16 InInstructionIndex) {}
diff --git a/Engine/Plugins/Runtime/RigVM/Source/RigVM/Public/RigVMHost.h b/Engine/Plugins/Runtime/RigVM/Source/RigVM/Public/RigVMHost.h
--- a/Engine/Plugins/Runtime/RigVM/Source/RigVM/Public/RigVMHost.h
+++ b/Engine/Plugins/Runtime/RigVM/Source/RigVM/Public/RigVMHost.h
@@ -228,11 +228,14 @@ public:
 	void SetLog(FRigVMLog* InLog) { RigVMLog = InLog; }
 #endif
 
-	// Returns a VM memory storage by type
+	// Returns the compiler generated VM memory storage by type
+	virtual const FRigVMMemoryStorageStruct* GetDefaultMemoryByType(ERigVMMemoryType InMemoryType) const;
+
+	// Returns an instanced VM memory storage by type
 	virtual FRigVMMemoryStorageStruct* GetMemoryByType(ERigVMMemoryType InMemoryType);
 	virtual const FRigVMMemoryStorageStruct* GetMemoryByType(ERigVMMemoryType InMemoryType) const;
 
-	// The default mutable work memory
+	// The instanced mutable work memory
 	FRigVMMemoryStorageStruct* GetWorkMemory() { return GetMemoryByType(ERigVMMemoryType::Work); }
 	const FRigVMMemoryStorageStruct* GetWorkMemory() const { return GetMemoryByType(ERigVMMemoryType::Work); }
 
@@ -240,7 +243,7 @@ public:
 	FRigVMMemoryStorageStruct* GetLiteralMemory() { return GetMemoryByType(ERigVMMemoryType::Literal); }
 	const FRigVMMemoryStorageStruct* GetLiteralMemory() const { return GetMemoryByType(ERigVMMemoryType::Literal); }
 
-	// The default debug watch memory
+	// The instanced debug watch memory
 	FRigVMMemoryStorageStruct* GetDebugMemory() { return GetMemoryByType(ERigVMMemoryType::Debug); }
 	const FRigVMMemoryStorageStruct* GetDebugMemory() const { return GetMemoryByType(ERigVMMemoryType::Debug); }
 
@@ -272,16 +275,36 @@ public:
 #endif
 
 	/** Provide access to the ExtendedExecuteContext */
-	UFUNCTION(BlueprintCallable, Category = RigVM)
+	UE_DEPRECATED(5.4, "Please, use GetRigVMExtendedExecuteContext")
+	UFUNCTION(BlueprintCallable, Category = RigVM, meta = (DeprecatedFunction, DeprecationMessage = "This function has been deprecated and it is no longer supported."))
 	virtual FRigVMExtendedExecuteContext& GetExtendedExecuteContext()
 	{
-		return ExtendedExecuteContext;
+		static FRigVMExtendedExecuteContext DummyContext;
+		return DummyContext;
 	};
 
 	/** Provide access to the ExtendedExecuteContext */
+	UE_DEPRECATED(5.4, "Please, use GetRigVMExtendedExecuteContext")
 	virtual const FRigVMExtendedExecuteContext& GetExtendedExecuteContext() const
 	{
-		return ExtendedExecuteContext;
+		static FRigVMExtendedExecuteContext DummyContext;
+		return DummyContext;
+	};
+
+	inline void SetRigVMExtendedExecuteContext(FRigVMExtendedExecuteContext* InRigVMExtendedExecuteContext)
+	{
+		RigVMExtendedExecuteContext = InRigVMExtendedExecuteContext;
+	}
+
+	inline FRigVMExtendedExecuteContext& GetRigVMExtendedExecuteContext()
+	{
+		check(RigVMExtendedExecuteContext != nullptr);
+		return *RigVMExtendedExecuteContext;
+	};
+	inline const FRigVMExtendedExecuteContext& GetRigVMExtendedExecuteContext() const
+	{
+		check(RigVMExtendedExecuteContext != nullptr);
+		return *RigVMExtendedExecuteContext;
 	};
 
 	UObject* ResolveUserDefinedTypeById(const FString& InTypeName) const;
@@ -302,7 +325,7 @@ protected:
 	/** true if we should increase the AbsoluteTime */
 	bool bAccumulateTime;
 
-	UPROPERTY(transient)
+	UPROPERTY()
 	TObjectPtr<URigVM> VM;
 
 #if WITH_EDITOR
@@ -321,9 +344,13 @@ protected:
 private:
 	UPROPERTY(transient)
 	TSet<TObjectPtr<UObject>> UserDefinedTypesInUse;
-	
+
+#if WITH_EDITORONLY_DATA
 	UPROPERTY()
-	FRigVMExtendedExecuteContext ExtendedExecuteContext;
+	FRigVMExtendedExecuteContext ExtendedExecuteContext_DEPRECATED;
+#endif
+	
+	FRigVMExtendedExecuteContext* RigVMExtendedExecuteContext = nullptr;
 
 protected:
 	
@@ -377,8 +404,6 @@ protected:
 
 	virtual void InitializeFromCDO();
 
-	static uint32 ComputeAndUpdateCDOHash(URigVMHost* InCDO);
-
 	virtual void CopyVMMemory(FRigVMExtendedExecuteContext& TargetContext, const FRigVMExtendedExecuteContext& SourceContext);
 
 public:
@@ -389,9 +414,6 @@ public:
 	virtual const TArray<UAssetUserData*>* GetAssetUserDataArray() const override;
 	//~ End IInterface_AssetUserData Interface
 
-	// Temporary flag while we get TLS-based work memory working, controls copying from the CDO in InitializeVM
-	bool bTEMP_CopyDefaultsFromCDO = true;
-
 protected:
 	/** Array of user data stored with the asset */
 	UPROPERTY(EditAnywhere, AdvancedDisplay, Instanced, Category = "Default")
diff --git a/Engine/Plugins/Runtime/RigVM/Source/RigVM/Public/RigVMObjectVersion.h b/Engine/Plugins/Runtime/RigVM/Source/RigVM/Public/RigVMObjectVersion.h
--- a/Engine/Plugins/Runtime/RigVM/Source/RigVM/Public/RigVMObjectVersion.h
+++ b/Engine/Plugins/Runtime/RigVM/Source/RigVM/Public/RigVMObjectVersion.h
@@ -30,6 +30,9 @@ struct RIGVM_API FRigVMObjectVersion
 		// VM Memory Storage Struct serialized
 		VMMemoryStorageStructSerialized,
 
+		// VM Memory Storage Defaults generated at VM
+		VMMemoryStorageDefaultsGeneratedAtVM,
+
 		// -----<new versions can be added above this line>-------------------------------------------------
 		VersionPlusOne,
 		LatestVersion = VersionPlusOne - 1
diff --git a/Engine/Plugins/Runtime/RigVM/Source/RigVMDeveloper/Private/EdGraph/RigVMEdGraphNode.cpp b/Engine/Plugins/Runtime/RigVM/Source/RigVMDeveloper/Private/EdGraph/RigVMEdGraphNode.cpp
--- a/Engine/Plugins/Runtime/RigVM/Source/RigVMDeveloper/Private/EdGraph/RigVMEdGraphNode.cpp
+++ b/Engine/Plugins/Runtime/RigVM/Source/RigVMDeveloper/Private/EdGraph/RigVMEdGraphNode.cpp
@@ -1132,7 +1132,7 @@ FLinearColor URigVMEdGraphNode::GetNodeProfilingColor() const
 			{
 				if(const URigVMNode* ModelNode = GetModelNode())
 				{
-					const double MicroSeconds = ModelNode->GetInstructionMicroSeconds(DebuggedHost->GetExtendedExecuteContext(), DebuggedHost->GetVM(), FRigVMASTProxy());
+					const double MicroSeconds = ModelNode->GetInstructionMicroSeconds(DebuggedHost->GetRigVMExtendedExecuteContext(), DebuggedHost->GetVM(), FRigVMASTProxy());
 					if(MicroSeconds >= 0.0)
 					{
 						if(Blueprint->RigGraphDisplaySettings.bAutoDetermineRange)
diff --git a/Engine/Plugins/Runtime/RigVM/Source/RigVMDeveloper/Private/RigVMBlueprint.cpp b/Engine/Plugins/Runtime/RigVM/Source/RigVMDeveloper/Private/RigVMBlueprint.cpp
--- a/Engine/Plugins/Runtime/RigVM/Source/RigVMDeveloper/Private/RigVMBlueprint.cpp
+++ b/Engine/Plugins/Runtime/RigVM/Source/RigVMDeveloper/Private/RigVMBlueprint.cpp
@@ -1176,7 +1176,7 @@ void URigVMBlueprint::RecompileVM()
 		{
 			CDO->Modify(false);
 		}
-		CDO->VM->Reset(CDO->GetExtendedExecuteContext());
+		CDO->VM->Reset(CDO->GetRigVMExtendedExecuteContext());
 
 		// Clear all Errors
 		CompileLog.Messages.Reset();
@@ -1205,13 +1205,11 @@ void URigVMBlueprint::RecompileVM()
 		URigVMCompiler* Compiler = URigVMCompiler::StaticClass()->GetDefaultObject<URigVMCompiler>();
 		VMCompileSettings.SetExecuteContextStruct(RigVMClient.GetExecuteContextStruct());
 
-	    FRigVMExtendedExecuteContext& CDOContext = CDO->GetExtendedExecuteContext();
+	    FRigVMExtendedExecuteContext& CDOContext = CDO->GetRigVMExtendedExecuteContext();
 		const FRigVMCompileSettings Settings = (bCompileInDebugMode) ? FRigVMCompileSettings::Fast(VMCompileSettings.GetExecuteContextStruct()) : VMCompileSettings;
 		Compiler->Compile(Settings, RigVMClient.GetAllModels(false, false), GetOrCreateController(), CDO->VM, CDOContext, CDO->GetExternalVariablesImpl(false), &PinToOperandMap);
 
-		CDO->VM->SetVMHash(CDO->VM->ComputeVMHash(CDOContext));
-		CDOContext.VMHash = CDO->VM->GetVMHash();
-		CDO->VM->Initialize(CDOContext, CDO->VM->GetLocalMemoryArray(CDOContext));
+		CDO->VM->Initialize(CDOContext);
 		CDO->GenerateUserDefinedDependenciesData(CDOContext);
 
 		if (bErrorsDuringCompilation)
@@ -1224,7 +1222,7 @@ void URigVMBlueprint::RecompileVM()
 			bVMRecompilationRequired = false;
 			if(CDO->VM)
 			{
-				VMCompiledEvent.Broadcast(this, CDO->GetVM(), CDO->GetExtendedExecuteContext());
+				VMCompiledEvent.Broadcast(this, CDO->GetVM(), CDO->GetRigVMExtendedExecuteContext());
 			}
 			return;
 		}
@@ -1232,7 +1230,7 @@ void URigVMBlueprint::RecompileVM()
 		InitializeArchetypeInstances();
 
 		bVMRecompilationRequired = false;
-		VMCompiledEvent.Broadcast(this, CDO->GetVM(), CDO->GetExtendedExecuteContext());
+		VMCompiledEvent.Broadcast(this, CDO->GetVM(), CDO->GetRigVMExtendedExecuteContext());
 
 #if WITH_EDITOR
 		RefreshBreakpoints();
@@ -2942,7 +2940,7 @@ void URigVMBlueprint::HandleModifiedEvent(ERigVMGraphNotifType InNotifType, URig
 							URigVMHost* CDO = Cast<URigVMHost>(RigClass->GetDefaultObject(true /* create if needed */));
 							if (CDO->VM != nullptr)
 							{
-								CDO->VM->SetPropertyValueFromString(CDO->GetExtendedExecuteContext(), *Operand, DefaultValue);
+								CDO->VM->SetPropertyValueFromString(CDO->GetRigVMExtendedExecuteContext(), *Operand, DefaultValue);
 							}
 
 							TArray<UObject*> ArchetypeInstances;
@@ -2954,7 +2952,7 @@ void URigVMBlueprint::HandleModifiedEvent(ERigVMGraphNotifType InNotifType, URig
 								{
 									if (InstancedHost->VM)
 									{
-										InstancedHost->VM->SetPropertyValueFromString(InstancedHost->GetExtendedExecuteContext(), *Operand, DefaultValue);
+										InstancedHost->VM->SetPropertyValueFromString(InstancedHost->GetRigVMExtendedExecuteContext(), *Operand, DefaultValue);
 									}
 								}
 							}
diff --git a/Engine/Plugins/Runtime/RigVM/Source/RigVMDeveloper/Private/RigVMCompiler/RigVMCodeGenerator.cpp b/Engine/Plugins/Runtime/RigVM/Source/RigVMDeveloper/Private/RigVMCompiler/RigVMCodeGenerator.cpp
--- a/Engine/Plugins/Runtime/RigVM/Source/RigVMDeveloper/Private/RigVMCompiler/RigVMCodeGenerator.cpp
+++ b/Engine/Plugins/Runtime/RigVM/Source/RigVMDeveloper/Private/RigVMCompiler/RigVMCodeGenerator.cpp
@@ -98,12 +98,12 @@ static constexpr TCHAR RigVM_GetVMHashFormat[] = TEXT("\tvirtual uint32 GetVMHas
 static constexpr TCHAR RigVM_GetEntryNamesFormat[] = TEXT("\tvirtual const TArray<FName>& GetEntryNames() const override\r\n\t{\r\n\t\tstatic const TArray<FName> StaticEntryNames = { {0} };\r\n\t\treturn StaticEntryNames;\r\n\t}");
 static constexpr TCHAR RigVM_DeclareUpdateExternalVariablesFormat[] = TEXT("\tvirtual void UpdateExternalVariables(FRigVMExtendedExecuteContext& Context) override;");
 static constexpr TCHAR RigVM_DeclareInvokeEntryByNameFormat[] = TEXT("\tERigVMExecuteResult InvokeEntryByName(FRigVMExtendedExecuteContext& Context, const FName& InEntryName{0});");
-static constexpr TCHAR RigVM_DeclareInitializeFormat[] = TEXT("\tvirtual bool Initialize(FRigVMExtendedExecuteContext& Context, TArrayView<FRigVMMemoryStorageStruct*> Memory) override;");
-static constexpr TCHAR RigVM_DefineInitializeFormat[] = TEXT("bool U{0}::Initialize(FRigVMExtendedExecuteContext& Context, TArrayView<FRigVMMemoryStorageStruct*> Memory)\r\n{");
-static constexpr TCHAR RigVM_DeclareExecuteFormat[] = TEXT("\tvirtual ERigVMExecuteResult Execute(FRigVMExtendedExecuteContext& Context, TArrayView<FRigVMMemoryStorageStruct*> Memory, const FName& InEntryName) override;");
+static constexpr TCHAR RigVM_DeclareInitializeFormat[] = TEXT("\tvirtual bool Initialize(FRigVMExtendedExecuteContext& Context) override;");
+static constexpr TCHAR RigVM_DefineInitializeFormat[] = TEXT("bool U{0}::Initialize(FRigVMExtendedExecuteContext& Context)\r\n{");
+static constexpr TCHAR RigVM_DeclareExecuteFormat[] = TEXT("\tvirtual ERigVMExecuteResult Execute(FRigVMExtendedExecuteContext& Context, const FName& InEntryName) override;");
 static constexpr TCHAR RigVM_DefineUpdateExternalVariablesFormat[] = TEXT("void U{0}::UpdateExternalVariables(FRigVMExtendedExecuteContext& Context)\r\n{");
 static constexpr TCHAR RigVM_DefineInvokeEntryByNameFormat[] = TEXT("ERigVMExecuteResult U{0}::InvokeEntryByName(FRigVMExtendedExecuteContext& Context, const FName& InEntryName{1})\r\n{");
-static constexpr TCHAR RigVM_DefineExecuteFormat[] = TEXT("ERigVMExecuteResult U{0}::Execute(FRigVMExtendedExecuteContext& Context, TArrayView<FRigVMMemoryStorageStruct*> Memory, const FName& InEntryName)\r\n{");
+static constexpr TCHAR RigVM_DefineExecuteFormat[] = TEXT("ERigVMExecuteResult U{0}::Execute(FRigVMExtendedExecuteContext& Context, const FName& InEntryName)\r\n{");
 static constexpr TCHAR RigVM_DeclareExecuteEntryFormat[] = TEXT("\tERigVMExecuteResult ExecuteEntry_{0}(FRigVMExtendedExecuteContext& Context, {1});");
 static constexpr TCHAR RigVM_DefineExecuteEntryFormat[] = TEXT("ERigVMExecuteResult U{0}::ExecuteEntry_{1}(FRigVMExtendedExecuteContext& Context, {2})\r\n{");
 static constexpr TCHAR RigVM_DeclareExecuteGroupFormat[] = TEXT("\tERigVMExecuteResult ExecuteGroup_{0}_{1}({2});");
@@ -700,7 +700,8 @@ FString FRigVMCodeGenerator::DumpInstructions(const FRigVMExtendedExecuteContext
 			case ERigVMOpCode::Zero:
 			{
 				const FRigVMUnaryOp& Op = ByteCode.GetOpAt<FRigVMUnaryOp>(Instruction);
-				const FProperty* Property = Context.WorkMemoryStorageObject->GetProperty(Op.Arg.GetRegisterIndex());
+				const FProperty* Property = VM->GetDefaultWorkMemory().GetProperty(Op.Arg.GetRegisterIndex());
+
 				if(Property->IsA<FIntProperty>())
 				{
 					Lines.Add(Prefix + Format(RigVM_ZeroOpIntFormat, *GetOperandName(Context, Op.Arg, false)));
@@ -1127,10 +1128,10 @@ void FRigVMCodeGenerator::ParseVM(const FString& InClassName, const FString& InM
 	ParseInclude(URigVM::StaticClass());
 
 	Properties.Reserve(
-		InVMToNativize->GetLiteralMemory()->Num() +
-		InVMToNativize->GetWorkMemory(InVMContext)->Num());
-	ParseMemory(InVMContext, InVMToNativize->GetLiteralMemory());
-	ParseMemory(InVMContext, InVMToNativize->GetWorkMemory(InVMContext));
+		InVMToNativize->GetDefaultLiteralMemory().Num() +
+		InVMToNativize->GetDefaultWorkMemory().Num());
+	ParseMemory(InVMContext, &InVMToNativize->GetDefaultLiteralMemory());
+	ParseMemory(InVMContext, &InVMToNativize->GetDefaultWorkMemory());
 
 	ParseRequiredUProperties(InVMContext);
 	ParseInstructionGroups(InVMContext);
@@ -1966,7 +1967,7 @@ FRigVMPropertyDescription FRigVMCodeGenerator::GetPropertyDescForOperand(const F
 	}
 	else
 	{
-		if (const FRigVMMemoryStorageStruct* MemoryStorage = VM->GetMemoryByType(Context, InOperand.GetMemoryType()))
+		if (const FRigVMMemoryStorageStruct* MemoryStorage = VM->GetDefaultMemoryByType(InOperand.GetMemoryType()))
 		{
 			Property = MemoryStorage->GetProperty(InOperand.GetRegisterIndex());
 		}
@@ -1991,7 +1992,7 @@ FRigVMPropertyDescription FRigVMCodeGenerator::GetPropertyForOperand(const FRigV
 	}
 	else
 	{
-		if (const FRigVMMemoryStorageStruct* MemoryStorage = VM->GetMemoryByType(Context, InOperand.GetMemoryType()))
+		if (const FRigVMMemoryStorageStruct* MemoryStorage = VM->GetDefaultMemoryByType(InOperand.GetMemoryType()))
 		{
 			Property = MemoryStorage->GetProperty(InOperand.GetRegisterIndex());
 			if (Property)
@@ -2040,7 +2041,7 @@ const FRigVMPropertyPathDescription& FRigVMCodeGenerator::GetPropertyPathForOper
 		}
 		else
 		{
-			if (const FRigVMMemoryStorageStruct* MemoryStorage = VM->GetMemoryByType(Context, InOperand.GetMemoryType()))
+			if (const FRigVMMemoryStorageStruct* MemoryStorage = VM->GetDefaultMemoryByType(InOperand.GetMemoryType()))
 			{
 				if (MemoryStorage->IsValidPropertyPathDescriptionIndex((RegisterOffsetIndex)))
 				{
@@ -2101,7 +2102,7 @@ void FRigVMCodeGenerator::CheckOperand(const FRigVMExtendedExecuteContext& Conte
 	}
 	else
 	{
-		const FRigVMMemoryStorageStruct* Memory = VM->GetMemoryByType(Context, InOperand.GetMemoryType());
+		const FRigVMMemoryStorageStruct* Memory = VM->GetDefaultMemoryByType(InOperand.GetMemoryType());
 		check(Memory);
 
 		ensure(Memory->GetProperties().IsValidIndex(InOperand.GetRegisterIndex()));
diff --git a/Engine/Plugins/Runtime/RigVM/Source/RigVMDeveloper/Private/RigVMCompiler/RigVMCompiler.cpp b/Engine/Plugins/Runtime/RigVM/Source/RigVMDeveloper/Private/RigVMCompiler/RigVMCompiler.cpp
--- a/Engine/Plugins/Runtime/RigVM/Source/RigVMDeveloper/Private/RigVMCompiler/RigVMCompiler.cpp
+++ b/Engine/Plugins/Runtime/RigVM/Source/RigVMDeveloper/Private/RigVMCompiler/RigVMCompiler.cpp
@@ -195,18 +195,18 @@ const FProperty* FRigVMCompilerWorkData::GetPropertyForOperand(const FRigVMOpera
 	
 	check(!bSetupMemory);
 
-	auto GetPropertyFromMemory = [this](FRigVMMemoryStorageStruct* InMemory, const FRigVMOperand& InOperand)
+	auto GetPropertyFromMemory = [this](FRigVMMemoryStorageStruct& InMemory, const FRigVMOperand& InOperand)
 	{
 		if(InOperand.GetRegisterOffset() == INDEX_NONE)
 		{
-			return  InMemory->GetProperty(InOperand.GetRegisterIndex());
+			return  InMemory.GetProperty(InOperand.GetRegisterIndex());
 		}
-		if(!InMemory->GetPropertyPaths().IsValidIndex(InOperand.GetRegisterOffset()))
+		if(!InMemory.GetPropertyPaths().IsValidIndex(InOperand.GetRegisterOffset()))
 		{
-			InMemory->SetPropertyPathDescriptions(PropertyPathDescriptions.FindChecked(InOperand.GetMemoryType()));
-			InMemory->RefreshPropertyPaths();
+			InMemory.SetPropertyPathDescriptions(PropertyPathDescriptions.FindChecked(InOperand.GetMemoryType()));
+			InMemory.RefreshPropertyPaths();
 		}
-		return InMemory->GetPropertyPaths()[InOperand.GetRegisterOffset()].GetTailProperty();
+		return InMemory.GetPropertyPaths()[InOperand.GetRegisterOffset()].GetTailProperty();
 	};
 
 	const FProperty* Property = nullptr;
@@ -214,17 +214,17 @@ const FProperty* FRigVMCompilerWorkData::GetPropertyForOperand(const FRigVMOpera
 	{
 	case ERigVMMemoryType::Literal:
 		{
-			Property = GetPropertyFromMemory(VM->GetLiteralMemory(), InOperand);
+			Property = GetPropertyFromMemory(VM->GetDefaultLiteralMemory(), InOperand);
 			break;
 		}
 	case ERigVMMemoryType::Work:
 		{
-			Property = GetPropertyFromMemory(VM->GetWorkMemory(*Context), InOperand);
+			Property = GetPropertyFromMemory(VM->GetDefaultWorkMemory(), InOperand);
 			break;
 		}
 	case ERigVMMemoryType::Debug:
 		{
-			Property = GetPropertyFromMemory(VM->GetDebugMemory(*Context), InOperand);
+			Property = GetPropertyFromMemory(VM->GetDefaultDebugMemory(), InOperand);
 			break;
 		}
 	case ERigVMMemoryType::External:
@@ -435,6 +435,7 @@ bool URigVMCompiler::Compile(const FRigVMCompileSettings& InSettings, TArray<URi
 	DECLARE_SCOPE_HIERARCHICAL_COUNTER_FUNC()
 
 	OutVM->Reset(OutVMContext);
+	OutVMContext.VMHash = 0;	// this has to be set when the memory is copied to the context
 
 	TMap<FString, FRigVMOperand> LocalOperands;
 	if (OutOperands == nullptr)
@@ -1210,7 +1211,7 @@ bool URigVMCompiler::Compile(const FRigVMCompileSettings& InSettings, TArray<URi
 	for (ERigVMMemoryType MemoryType : MemoryTypes)
 	{
 		const TArray<FRigVMPropertyDescription>* Properties = WorkData.PropertyDescriptions.Find(MemoryType);
-		WorkData.VM->GenerateMemoryType(*WorkData.Context, MemoryType, Properties);
+		WorkData.VM->GenerateDefaultMemoryType(MemoryType, Properties);
 	}
 
 	WorkData.bSetupMemory = false;
@@ -1248,7 +1249,7 @@ bool URigVMCompiler::Compile(const FRigVMCompileSettings& InSettings, TArray<URi
 	for(ERigVMMemoryType MemoryType : MemoryTypes)
 	{
 		const TArray<FRigVMPropertyPathDescription>* Descriptions = WorkData.PropertyPathDescriptions.Find(MemoryType);
-		if(FRigVMMemoryStorageStruct* MemoryStorageObject = WorkData.VM->GetMemoryByType(*WorkData.Context, MemoryType))
+		if(FRigVMMemoryStorageStruct* MemoryStorageObject = WorkData.VM->GetDefaultMemoryByType(MemoryType))
 		{
 			if (Descriptions)
 			{
@@ -1381,6 +1382,8 @@ bool URigVMCompiler::Compile(const FRigVMCompileSettings& InSettings, TArray<URi
 	{
 		CompileTimer.Stop();
 		WorkData.ReportInfof(TEXT("Total Compilation time %f\n"), CompilationTime*1000);
+
+		WorkData.VM->SetVMHash(WorkData.VM->ComputeVMHash());
 	}
 
 	return true;
diff --git a/Engine/Plugins/Runtime/RigVM/Source/RigVMEditor/Private/EdGraph/RigVMEdGraphConnectionDrawingPolicy.cpp b/Engine/Plugins/Runtime/RigVM/Source/RigVMEditor/Private/EdGraph/RigVMEdGraphConnectionDrawingPolicy.cpp
--- a/Engine/Plugins/Runtime/RigVM/Source/RigVMEditor/Private/EdGraph/RigVMEdGraphConnectionDrawingPolicy.cpp
+++ b/Engine/Plugins/Runtime/RigVM/Source/RigVMEditor/Private/EdGraph/RigVMEdGraphConnectionDrawingPolicy.cpp
@@ -321,8 +321,8 @@ void FRigVMEdGraphConnectionDrawingPolicy::DetermineWiringStyle(UEdGraphPin* Out
 					{
 						if (const URigVM* VM = Host->GetVM())
 						{
-							if (VM->WasInstructionVisitedDuringLastRun(Host->GetExtendedExecuteContext(), OutputInstructionIndex) &&
-								VM->WasInstructionVisitedDuringLastRun(Host->GetExtendedExecuteContext(), InputInstructionIndex))
+							if (VM->WasInstructionVisitedDuringLastRun(Host->GetRigVMExtendedExecuteContext(), OutputInstructionIndex) &&
+								VM->WasInstructionVisitedDuringLastRun(Host->GetRigVMExtendedExecuteContext(), InputInstructionIndex))
 							{
 								bVisited = true;
 							}
diff --git a/Engine/Plugins/Runtime/RigVM/Source/RigVMEditor/Private/Editor/RigVMEditor.cpp b/Engine/Plugins/Runtime/RigVM/Source/RigVMEditor/Private/Editor/RigVMEditor.cpp
--- a/Engine/Plugins/Runtime/RigVM/Source/RigVMEditor/Private/Editor/RigVMEditor.cpp
+++ b/Engine/Plugins/Runtime/RigVM/Source/RigVMEditor/Private/Editor/RigVMEditor.cpp
@@ -1751,7 +1751,7 @@ void FRigVMEditor::HandleVMExecutedEvent(URigVMHost* InHost, const FName& InEven
 					const FRigVMByteCode& ByteCode = VM->GetByteCode();
 					for(int32 InstructionIndex = 0; InstructionIndex < ByteCode.GetNumInstructions(); InstructionIndex++)
 					{
-						const int32 Count = VM->GetInstructionVisitedCount(DebuggedHost->GetExtendedExecuteContext(), InstructionIndex);
+						const int32 Count = VM->GetInstructionVisitedCount(DebuggedHost->GetRigVMExtendedExecuteContext(), InstructionIndex);
 						if(Count > RigVMBlueprint->RigGraphDisplaySettings.NodeRunLimit)
 						{
 							bFoundLimitWarnings = true;
@@ -3300,8 +3300,8 @@ void FRigVMEditor::UpdateGraphCompilerErrors()
 
 					if(RigVMEdGraphNode->ErrorType <= (int32)EMessageSeverity::Warning)
 					{
-						if(!VM->WasInstructionVisitedDuringLastRun(RigVMHost->GetExtendedExecuteContext(), RigVMEdGraphNode->GetInstructionIndex(true)) &&
-							!VM->WasInstructionVisitedDuringLastRun(RigVMHost->GetExtendedExecuteContext(), RigVMEdGraphNode->GetInstructionIndex(false)))
+						if(!VM->WasInstructionVisitedDuringLastRun(RigVMHost->GetRigVMExtendedExecuteContext(), RigVMEdGraphNode->GetInstructionIndex(true)) &&
+							!VM->WasInstructionVisitedDuringLastRun(RigVMHost->GetRigVMExtendedExecuteContext(), RigVMEdGraphNode->GetInstructionIndex(false)))
 						{
 							continue;
 						}
diff --git a/Engine/Plugins/Runtime/RigVM/Source/RigVMEditor/Private/RigVMEditorModule.cpp b/Engine/Plugins/Runtime/RigVM/Source/RigVMEditor/Private/RigVMEditorModule.cpp
--- a/Engine/Plugins/Runtime/RigVM/Source/RigVMEditor/Private/RigVMEditorModule.cpp
+++ b/Engine/Plugins/Runtime/RigVM/Source/RigVMEditor/Private/RigVMEditorModule.cpp
@@ -1813,7 +1813,7 @@ void FRigVMEditorModule::PostChange(const UUserDefinedStruct* Changed, FStructur
 			URigVMHost* CDO = Cast<URigVMHost>(RigClass->GetDefaultObject(true /* create if needed */));
 			if (CDO->GetVM() != nullptr)
 			{
-				CDO->GetVM()->Reset(CDO->GetExtendedExecuteContext());
+				CDO->GetVM()->Reset(CDO->GetRigVMExtendedExecuteContext());
 			}
 			TArray<UObject*> ArchetypeInstances;
 			CDO->GetArchetypeInstances(ArchetypeInstances);
@@ -1821,7 +1821,7 @@ void FRigVMEditorModule::PostChange(const UUserDefinedStruct* Changed, FStructur
 			{
 				if (URigVMHost* InstanceHost = Cast<URigVMHost>(Instance))
 				{
-					InstanceHost->GetVM()->Reset(InstanceHost->GetExtendedExecuteContext());
+					InstanceHost->GetVM()->Reset(InstanceHost->GetRigVMExtendedExecuteContext());
 				}
 			}
 		}
diff --git a/Engine/Plugins/Runtime/RigVM/Source/RigVMEditor/Private/Widgets/SRigVMExecutionStackView.cpp b/Engine/Plugins/Runtime/RigVM/Source/RigVMEditor/Private/Widgets/SRigVMExecutionStackView.cpp
--- a/Engine/Plugins/Runtime/RigVM/Source/RigVMEditor/Private/Widgets/SRigVMExecutionStackView.cpp
+++ b/Engine/Plugins/Runtime/RigVM/Source/RigVMEditor/Private/Widgets/SRigVMExecutionStackView.cpp
@@ -200,7 +200,7 @@ FText SRigStackItem::GetVisitedCountText() const
 				{
 					if(URigVM* VM = RigVMHost->GetVM())
 					{
-						const int32 Count = VM->GetInstructionVisitedCount(RigVMHost->GetExtendedExecuteContext(), WeakStackEntry.Pin()->InstructionIndex);
+						const int32 Count = VM->GetInstructionVisitedCount(RigVMHost->GetRigVMExtendedExecuteContext(), WeakStackEntry.Pin()->InstructionIndex);
 						if(Count > 0)
 						{
 							return FText::FromString(FString::FromInt(Count));
@@ -225,7 +225,7 @@ FText SRigStackItem::GetDurationText() const
 				{
 					if(URigVM* VM = RigVMHost->GetVM())
 					{
-						const double MicroSeconds = VM->GetInstructionMicroSeconds(RigVMHost->GetExtendedExecuteContext(), WeakStackEntry.Pin()->InstructionIndex);
+						const double MicroSeconds = VM->GetInstructionMicroSeconds(RigVMHost->GetRigVMExtendedExecuteContext(), WeakStackEntry.Pin()->InstructionIndex);
 						if(MicroSeconds > 0.0)
 						{
 							return FText::FromString(FString::Printf(TEXT("%d µs"), (int32)MicroSeconds));
@@ -350,7 +350,7 @@ void SRigVMExecutionStackView::Construct( const FArguments& InArgs, TSharedRef<F
 		{
 			if (URigVMHost* RigVMHost = RigVMEditor.Pin()->GetRigVMHost())
 			{
-				OnVMCompiled(RigVMBlueprint.Get(),  RigVMHost->GetVM(),  RigVMHost->GetExtendedExecuteContext());
+				OnVMCompiled(RigVMBlueprint.Get(),  RigVMHost->GetVM(),  RigVMHost->GetRigVMExtendedExecuteContext());
 			}
 		}
 	}
@@ -920,7 +920,7 @@ void SRigVMExecutionStackView::OnFilterTextChanged(const FText& SearchText)
 	URigVMHost* RigVMHost = RigVMEditor.Pin()->GetRigVMHost();
 	if (RigVMHost != nullptr)
 	{
-		RefreshTreeView(RigVMHost->GetVM(), &RigVMHost->GetExtendedExecuteContext());
+		RefreshTreeView(RigVMHost->GetVM(), &RigVMHost->GetRigVMExtendedExecuteContext());
 	}
 }
 
@@ -981,7 +981,7 @@ void SRigVMExecutionStackView::HandleHostInitializedEvent(URigVMHost* InHost, co
 {
 	TGuardValue<bool> SuspendControllerSelection(bSuspendControllerSelection, true);
 
-	RefreshTreeView(InHost->GetVM(), &InHost->GetExtendedExecuteContext());
+	RefreshTreeView(InHost->GetVM(), &InHost->GetRigVMExtendedExecuteContext());
 	OnSelectionChanged(TSharedPtr<FRigStackEntry>(), ESelectInfo::Direct);
 
 	for (TSharedPtr<FRigStackEntry>& Operator : Operators)
@@ -1022,7 +1022,7 @@ void SRigVMExecutionStackView::HandlePreviewHostUpdated(FRigVMEditor* InEditor)
 {
 	if(URigVMHost* RigVMHost = InEditor->GetRigVMHost())
 	{
-		RefreshTreeView(RigVMHost->GetVM(), &RigVMHost->GetExtendedExecuteContext());
+		RefreshTreeView(RigVMHost->GetVM(), &RigVMHost->GetRigVMExtendedExecuteContext());
 	}
 }
 
diff --git a/Engine/Plugins/Runtime/RigVM/Source/RigVMEditor/Private/Widgets/SRigVMGraphNode.cpp b/Engine/Plugins/Runtime/RigVM/Source/RigVMEditor/Private/Widgets/SRigVMGraphNode.cpp
--- a/Engine/Plugins/Runtime/RigVM/Source/RigVMEditor/Private/Widgets/SRigVMGraphNode.cpp
+++ b/Engine/Plugins/Runtime/RigVM/Source/RigVMEditor/Private/Widgets/SRigVMGraphNode.cpp
@@ -900,7 +900,7 @@ void SRigVMGraphNode::GetNodeInfoPopups(FNodeInfoContext* Context, TArray<FGraph
 				{
 					if(URigVMHost* DebuggedHost = Cast<URigVMHost>(Blueprint->GetObjectBeingDebugged()))
 					{
-						const int32 Count = ModelNode->GetInstructionVisitedCount(DebuggedHost->GetExtendedExecuteContext(), DebuggedHost->GetVM(), FRigVMASTProxy());
+						const int32 Count = ModelNode->GetInstructionVisitedCount(DebuggedHost->GetRigVMExtendedExecuteContext(), DebuggedHost->GetVM(), FRigVMASTProxy());
 						if(Count == 0)
 						{
 							PinnedWatchText = FString::Printf(TEXT("Node is not running - wrong event?\n%s"), *PinnedWatchText);
@@ -985,7 +985,7 @@ TArray<FOverlayWidgetInfo> SRigVMGraphNode::GetOverlayWidgets(bool bSelected, co
 				{
 					if(bShowNodeCounts || bShowInstructionIndex)
 					{
-						const int32 Count = ModelNode->GetInstructionVisitedCount(DebuggedHost->GetExtendedExecuteContext(), DebuggedHost->GetVM(), FRigVMASTProxy());
+						const int32 Count = ModelNode->GetInstructionVisitedCount(DebuggedHost->GetRigVMExtendedExecuteContext(), DebuggedHost->GetVM(), FRigVMASTProxy());
 						if((Count > Blueprint->RigGraphDisplaySettings.NodeRunLowerBound) || bShowInstructionIndex)
 						{
 							const int32 VOffset = bSelected ? -2 : 2;
@@ -999,7 +999,7 @@ TArray<FOverlayWidgetInfo> SRigVMGraphNode::GetOverlayWidgets(bool bSelected, co
 
 					if(bEnableProfiling)
 					{
-						const double MicroSeconds = ModelNode->GetInstructionMicroSeconds(DebuggedHost->GetExtendedExecuteContext(), DebuggedHost->GetVM(), FRigVMASTProxy());
+						const double MicroSeconds = ModelNode->GetInstructionMicroSeconds(DebuggedHost->GetRigVMExtendedExecuteContext(), DebuggedHost->GetVM(), FRigVMASTProxy());
 						if(MicroSeconds >= 0.0)
 						{
 							const int32 VOffset = bSelected ? -2 : 2;
@@ -1261,13 +1261,13 @@ FText SRigVMGraphNode::GetInstructionCountText() const
 					int32 FirstInstructionIndex = INDEX_NONE;
 					if(bShowNodeRunCount)
 					{
-						RunCount = ModelNode->GetInstructionVisitedCount(DebuggedHost->GetExtendedExecuteContext(), DebuggedHost->GetVM(), FRigVMASTProxy());
+						RunCount = ModelNode->GetInstructionVisitedCount(DebuggedHost->GetRigVMExtendedExecuteContext(), DebuggedHost->GetVM(), FRigVMASTProxy());
 						bShowNodeRunCount = RunCount > Blueprint->RigGraphDisplaySettings.NodeRunLowerBound;
 					}
 
 					if(bShowInstructionIndex)
 					{
-						const TArray<int32> Instructions = ModelNode->GetInstructionsForVM(DebuggedHost->GetExtendedExecuteContext(), DebuggedHost->GetVM());
+						const TArray<int32> Instructions = ModelNode->GetInstructionsForVM(DebuggedHost->GetRigVMExtendedExecuteContext(), DebuggedHost->GetVM());
 						bShowInstructionIndex = Instructions.Num() > 0;
 						if(bShowInstructionIndex)
 						{
@@ -1316,7 +1316,7 @@ FText SRigVMGraphNode::GetInstructionDurationText() const
 			{
 				if(URigVMHost* DebuggedHost = Cast<URigVMHost>(Blueprint->GetObjectBeingDebugged()))
 				{
-					const double MicroSeconds = ModelNode->GetInstructionMicroSeconds(DebuggedHost->GetExtendedExecuteContext(), DebuggedHost->GetVM(), FRigVMASTProxy());
+					const double MicroSeconds = ModelNode->GetInstructionMicroSeconds(DebuggedHost->GetRigVMExtendedExecuteContext(), DebuggedHost->GetVM(), FRigVMASTProxy());
 					if(MicroSeconds >= 0)
 					{
 						return FText::FromString(FString::Printf(TEXT("%d µs"), (int32)MicroSeconds));
