Fixed function and collapse nodes not able to change their color

#rnx

diff --git a/Engine/Plugins/Runtime/RigVM/Source/RigVMDeveloper/Private/RigVMModel/RigVMSchema.cpp b/Engine/Plugins/Runtime/RigVM/Source/RigVMDeveloper/Private/RigVMModel/RigVMSchema.cpp
--- a/Engine/Plugins/Runtime/RigVM/Source/RigVMDeveloper/Private/RigVMModel/RigVMSchema.cpp
+++ b/Engine/Plugins/Runtime/RigVM/Source/RigVMDeveloper/Private/RigVMModel/RigVMSchema.cpp
@@ -475,7 +475,7 @@ bool URigVMSchema::CanRecolorNode(URigVMController* InController, const URigVMNo
 		return false;
 	}
 
-	return InNode->IsA<URigVMCommentNode>();
+	return InNode->IsA<URigVMCommentNode>() || InNode->IsA<URigVMLibraryNode>();
 }
 
 bool URigVMSchema::CanAddLink(URigVMController* InController, const URigVMPin* InSourcePin, const URigVMPin* InTargetPin, const FRigVMByteCode* InByteCode, ERigVMPinDirection InUserLinkDirection, bool bInAllowWildcard, bool bEnableTypeCasting, FString* OutFailureReason) const
