RigVM: Ensure to zero out padded bits on operands and operators when storing bytecode.

Decided against tight packing the bytecode since we already have cooked assets. Changing the size of structs like that is going to cause backwards compatibility issues.

diff --git a/Engine/Plugins/Runtime/RigVM/Source/RigVM/Private/RigVMCore/RigVMByteCode.cpp b/Engine/Plugins/Runtime/RigVM/Source/RigVM/Private/RigVMCore/RigVMByteCode.cpp
--- a/Engine/Plugins/Runtime/RigVM/Source/RigVM/Private/RigVMCore/RigVMByteCode.cpp
+++ b/Engine/Plugins/Runtime/RigVM/Source/RigVM/Private/RigVMCore/RigVMByteCode.cpp
@@ -1064,8 +1064,15 @@ uint64 FRigVMByteCode::AddExecuteOp(uint16 InFunctionIndex, const FRigVMOperandA
 	Op.PredicateCount = PredicateCount;
 	uint64 OpByteIndex = AddOp(Op);
 
-	uint64 OperandsByteIndex = (uint64)ByteCode.AddZeroed(sizeof(FRigVMOperand) * InOperands.Num());
+	const uint64 OperandsByteIndex = (uint64)ByteCode.AddZeroed(sizeof(FRigVMOperand) * InOperands.Num());
 	FMemory::Memcpy(ByteCode.GetData() + OperandsByteIndex, InOperands.GetData(), sizeof(FRigVMOperand) * InOperands.Num());
+
+	for(int32 Index = 0; Index < InOperands.Num(); Index++)
+	{
+		FRigVMOperand* Operand = reinterpret_cast<FRigVMOperand*>(ByteCode.GetData() + OperandsByteIndex + sizeof(FRigVMOperand) * Index);
+		FRigVMOperand::ZeroPaddedMemoryIfNeeded(Operand);
+	}
+	
 	return OpByteIndex;
 }
 
diff --git a/Engine/Plugins/Runtime/RigVM/Source/RigVM/Public/RigVMCore/RigVMByteCode.h b/Engine/Plugins/Runtime/RigVM/Source/RigVM/Public/RigVMCore/RigVMByteCode.h
--- a/Engine/Plugins/Runtime/RigVM/Source/RigVM/Public/RigVMCore/RigVMByteCode.h
+++ b/Engine/Plugins/Runtime/RigVM/Source/RigVM/Public/RigVMCore/RigVMByteCode.h
@@ -242,6 +242,11 @@ struct RIGVM_API FRigVMBaseOp
 	{
 		return GetTypeHash(Op.OpCode);
 	}
+
+	void ZeroPaddedMemoryIfNeeded(FRigVMBaseOp* InMemory) const
+	{
+		RigVM::ZeroPaddedMemory(&InMemory->OpCode, reinterpret_cast<uint8*>(InMemory) + sizeof(FRigVMBaseOp));
+	}
 };
 
 
@@ -291,6 +296,12 @@ struct RIGVM_API FRigVMExecuteOp : public FRigVMBaseOp
 		P.Serialize(Ar);
 		return Ar;
 	}
+	
+	void ZeroPaddedMemoryIfNeeded(FRigVMBaseOp* InMemory) const
+	{
+		FRigVMExecuteOp* This = reinterpret_cast<FRigVMExecuteOp*>(InMemory);
+		RigVM::ZeroPaddedMemory(&This->OpCode, &This->FunctionIndex);
+	}
 };
 
 // operator used for zero, false, true, increment, decrement
@@ -336,6 +347,13 @@ struct RIGVM_API FRigVMUnaryOp : public FRigVMBaseOp
 		P.Serialize(Ar);
 		return Ar;
 	}
+
+	void ZeroPaddedMemoryIfNeeded(FRigVMBaseOp* InMemory) const
+	{
+		FRigVMUnaryOp* This = reinterpret_cast<FRigVMUnaryOp*>(InMemory);
+		RigVM::ZeroPaddedMemory(&This->OpCode, &This->Arg);
+		FRigVMOperand::ZeroPaddedMemoryIfNeeded(&This->Arg);
+	}
 };
 
 // operator used for beginblock and array reset
@@ -375,6 +393,14 @@ struct RIGVM_API FRigVMBinaryOp : public FRigVMBaseOp
 		P.Serialize(Ar);
 		return Ar;
 	}
+	
+	void ZeroPaddedMemoryIfNeeded(FRigVMBaseOp* InMemory) const
+	{
+		FRigVMBinaryOp* This = reinterpret_cast<FRigVMBinaryOp*>(InMemory);
+		RigVM::ZeroPaddedMemory(&This->OpCode, &This->ArgA);
+		FRigVMOperand::ZeroPaddedMemoryIfNeeded(&This->ArgA);
+		FRigVMOperand::ZeroPaddedMemoryIfNeeded(&This->ArgB);
+	}
 };
 
 // operator used for some array operations
@@ -418,6 +444,15 @@ struct RIGVM_API FRigVMTernaryOp : public FRigVMBaseOp
 		P.Serialize(Ar);
 		return Ar;
 	}
+
+	void ZeroPaddedMemoryIfNeeded(FRigVMBaseOp* InMemory) const
+	{
+		FRigVMTernaryOp* This = reinterpret_cast<FRigVMTernaryOp*>(InMemory);
+		RigVM::ZeroPaddedMemory(&This->OpCode, &This->ArgA);
+		FRigVMOperand::ZeroPaddedMemoryIfNeeded(&This->ArgA);
+		FRigVMOperand::ZeroPaddedMemoryIfNeeded(&This->ArgB);
+		FRigVMOperand::ZeroPaddedMemoryIfNeeded(&This->ArgC);
+	}
 };
 
 // operator used for some array operations
@@ -468,6 +503,16 @@ struct RIGVM_API FRigVMQuaternaryOp : public FRigVMBaseOp
 		P.Serialize(Ar);
 		return Ar;
 	}
+	
+	void ZeroPaddedMemoryIfNeeded(FRigVMBaseOp* InMemory) const
+	{
+		FRigVMQuaternaryOp* This = reinterpret_cast<FRigVMQuaternaryOp*>(InMemory);
+		RigVM::ZeroPaddedMemory(&This->OpCode, &This->ArgA);
+		FRigVMOperand::ZeroPaddedMemoryIfNeeded(&This->ArgA);
+		FRigVMOperand::ZeroPaddedMemoryIfNeeded(&This->ArgB);
+		FRigVMOperand::ZeroPaddedMemoryIfNeeded(&This->ArgC);
+		FRigVMOperand::ZeroPaddedMemoryIfNeeded(&This->ArgD);
+	}
 };
 
 // operator used for some array operations
@@ -519,6 +564,17 @@ struct RIGVM_API FRigVMQuinaryOp : public FRigVMBaseOp
 		P.Serialize(Ar);
 		return Ar;
 	}
+	
+	void ZeroPaddedMemoryIfNeeded(FRigVMBaseOp* InMemory) const
+	{
+		FRigVMQuinaryOp* This = reinterpret_cast<FRigVMQuinaryOp*>(InMemory);
+		RigVM::ZeroPaddedMemory(&This->OpCode, &This->ArgA);
+		FRigVMOperand::ZeroPaddedMemoryIfNeeded(&This->ArgA);
+		FRigVMOperand::ZeroPaddedMemoryIfNeeded(&This->ArgB);
+		FRigVMOperand::ZeroPaddedMemoryIfNeeded(&This->ArgC);
+		FRigVMOperand::ZeroPaddedMemoryIfNeeded(&This->ArgD);
+		FRigVMOperand::ZeroPaddedMemoryIfNeeded(&This->ArgE);
+	}
 };
 
 // operator used for some array operations
@@ -574,6 +630,18 @@ struct RIGVM_API FRigVMSenaryOp : public FRigVMBaseOp
 		P.Serialize(Ar);
 		return Ar;
 	}
+	
+	void ZeroPaddedMemoryIfNeeded(FRigVMBaseOp* InMemory) const
+	{
+		FRigVMSenaryOp* This = reinterpret_cast<FRigVMSenaryOp*>(InMemory);
+		RigVM::ZeroPaddedMemory(&This->OpCode, &This->ArgA);
+		FRigVMOperand::ZeroPaddedMemoryIfNeeded(&This->ArgA);
+		FRigVMOperand::ZeroPaddedMemoryIfNeeded(&This->ArgB);
+		FRigVMOperand::ZeroPaddedMemoryIfNeeded(&This->ArgC);
+		FRigVMOperand::ZeroPaddedMemoryIfNeeded(&This->ArgD);
+		FRigVMOperand::ZeroPaddedMemoryIfNeeded(&This->ArgE);
+		FRigVMOperand::ZeroPaddedMemoryIfNeeded(&This->ArgF);
+	}
 };
 
 // The kind of copy operation to perform
@@ -648,6 +716,16 @@ public:
 		P.Serialize(Ar);
 		return Ar;
 	}
+
+	void ZeroPaddedMemoryIfNeeded(FRigVMBaseOp* InMemory) const
+	{
+		FRigVMCopyOp* This = reinterpret_cast<FRigVMCopyOp*>(InMemory);
+		RigVM::ZeroPaddedMemory(&This->OpCode, &This->Source);
+		FRigVMOperand::ZeroPaddedMemoryIfNeeded(&This->Source);
+		FRigVMOperand::ZeroPaddedMemoryIfNeeded(&This->Target);
+		RigVM::ZeroPaddedMemory(&This->RegisterType, &This->CopyType);
+		RigVM::ZeroPaddedMemory(&This->CopyType, reinterpret_cast<uint8*>(This) + sizeof(FRigVMCopyOp));
+	}
 };
 
 // used for equals and not equals comparisons
@@ -700,6 +778,15 @@ struct RIGVM_API FRigVMComparisonOp : public FRigVMBaseOp
 		P.Serialize(Ar);
 		return Ar;
 	}
+
+	void ZeroPaddedMemoryIfNeeded(FRigVMBaseOp* InMemory) const
+	{
+		FRigVMComparisonOp* This = reinterpret_cast<FRigVMComparisonOp*>(InMemory);
+		RigVM::ZeroPaddedMemory(&This->OpCode, &This->A);
+		FRigVMOperand::ZeroPaddedMemoryIfNeeded(&This->A);
+		FRigVMOperand::ZeroPaddedMemoryIfNeeded(&This->B);
+		FRigVMOperand::ZeroPaddedMemoryIfNeeded(&This->Result);
+	}
 };
 
 // jump to a new instruction index.
@@ -739,6 +826,12 @@ struct RIGVM_API FRigVMJumpOp : public FRigVMBaseOp
 		P.Serialize(Ar);
 		return Ar;
 	}
+	
+	void ZeroPaddedMemoryIfNeeded(FRigVMBaseOp* InMemory) const
+	{
+		FRigVMJumpOp* This = reinterpret_cast<FRigVMJumpOp*>(InMemory);
+		RigVM::ZeroPaddedMemory(&This->OpCode, &This->InstructionIndex);
+	}
 };
 
 // jump to a new instruction index based on a condition.
@@ -782,6 +875,14 @@ struct RIGVM_API FRigVMJumpIfOp : public FRigVMUnaryOp
 		P.Serialize(Ar);
 		return Ar;
 	}
+	
+	void ZeroPaddedMemoryIfNeeded(FRigVMBaseOp* InMemory) const
+	{
+		FRigVMJumpIfOp* This = reinterpret_cast<FRigVMJumpIfOp*>(InMemory);
+		RigVM::ZeroPaddedMemory(&This->OpCode, &This->Arg);
+		FRigVMOperand::ZeroPaddedMemoryIfNeeded(&This->Arg);
+		RigVM::ZeroPaddedMemory<bool>(&This->Condition, reinterpret_cast<uint8*>(This) + sizeof(FRigVMJumpIfOp));
+	}
 };
 
 // change the type of a register
@@ -827,6 +928,12 @@ struct RIGVM_API FRigVMInvokeEntryOp : public FRigVMBaseOp
 		P.Serialize(Ar);
 		return Ar;
 	}
+
+	void ZeroPaddedMemoryIfNeeded(FRigVMBaseOp* InMemory) const
+	{
+		FRigVMInvokeEntryOp* This = reinterpret_cast<FRigVMInvokeEntryOp*>(InMemory);
+		RigVM::ZeroPaddedMemory(&This->OpCode, &This->EntryName);
+	}
 };
 
 // jump into a branch based on a name argument
@@ -1251,9 +1358,10 @@ private:
 			(InOp.OpCode < ERigVMOpCode::Invalid) 
 		);
 		
-		uint64 ByteIndex = (uint64)ByteCode.AddZeroed(sizeof(OpType));
+		const uint64 ByteIndex = (uint64)ByteCode.AddZeroed(sizeof(OpType));
 		uint8* Pointer = &ByteCode[ByteIndex];
 		FMemory::Memcpy(Pointer, &InOp, sizeof(OpType));
+		InOp.ZeroPaddedMemoryIfNeeded(reinterpret_cast<OpType*>(Pointer));
 		NumInstructions++;
 		return ByteIndex;
 	}
diff --git a/Engine/Plugins/Runtime/RigVM/Source/RigVM/Public/RigVMCore/RigVMGraphFunctionDefinition.h b/Engine/Plugins/Runtime/RigVM/Source/RigVM/Public/RigVMCore/RigVMGraphFunctionDefinition.h
--- a/Engine/Plugins/Runtime/RigVM/Source/RigVM/Public/RigVMCore/RigVMGraphFunctionDefinition.h
+++ b/Engine/Plugins/Runtime/RigVM/Source/RigVM/Public/RigVMCore/RigVMGraphFunctionDefinition.h
@@ -35,9 +35,10 @@ struct FRigVMFunctionCompilationPropertyDescription
 
 	friend uint32 GetTypeHash(const FRigVMFunctionCompilationPropertyDescription& Description) 
 	{
-		uint32 Hash = GetTypeHash(Description.Name);
+		uint32 Hash = GetTypeHash(Description.Name.ToString());
 		Hash = HashCombine(Hash, GetTypeHash(Description.CPPType));
-		Hash = HashCombine(Hash, GetTypeHash(Description.CPPTypeObject));
+		// we can't hash based on the pointer since that's not deterministic across sessions
+		// Hash = HashCombine(Hash, GetTypeHash(Description.CPPTypeObject));
 		Hash = HashCombine(Hash, GetTypeHash(Description.DefaultValue));
 		return Hash;
 	}
@@ -152,7 +153,7 @@ struct RIGVM_API FRigVMFunctionCompilationData
 		uint32 DataHash = Data.ByteCode.GetByteCodeHash();
 		for (const FName& Name : Data.FunctionNames)
 		{
-			DataHash = HashCombine(DataHash, GetTypeHash(Name));
+			DataHash = HashCombine(DataHash, GetTypeHash(Name.ToString()));
 		}
 
 		for (const FRigVMFunctionCompilationPropertyDescription& Description : Data.WorkPropertyDescriptions)
@@ -194,7 +195,7 @@ struct RIGVM_API FRigVMFunctionCompilationData
 		for (const TPair<int32,FName>& Pair : Data.ExternalRegisterIndexToVariable)
 		{
 			DataHash = HashCombine(DataHash, GetTypeHash(Pair.Key));
-			DataHash = HashCombine(DataHash, GetTypeHash(Pair.Value));
+			DataHash = HashCombine(DataHash, GetTypeHash(Pair.Value.ToString()));
 		}
 
 		for (const TPair<FString, FRigVMOperand>& Pair : Data.Operands)
diff --git a/Engine/Plugins/Runtime/RigVM/Source/RigVM/Public/RigVMCore/RigVMMemoryCommon.h b/Engine/Plugins/Runtime/RigVM/Source/RigVM/Public/RigVMCore/RigVMMemoryCommon.h
--- a/Engine/Plugins/Runtime/RigVM/Source/RigVM/Public/RigVMCore/RigVMMemoryCommon.h
+++ b/Engine/Plugins/Runtime/RigVM/Source/RigVM/Public/RigVMCore/RigVMMemoryCommon.h
@@ -10,6 +10,7 @@
 #include "Templates/TypeHash.h"
 #include "UObject/ObjectMacros.h"
 #include "Misc/OutputDevice.h"
+#include "GenericPlatform/GenericPlatformMemory.h"
 
 #include "RigVMMemoryCommon.generated.h"
 
@@ -45,6 +46,21 @@ enum class ERigVMMemoryType: uint8
 	Invalid
 };
 
+namespace RigVM
+{
+	template<typename T = uint8>
+	void ZeroPaddedMemory(void* InFirstMember, const void* InSecondMember)
+	{
+		char* FirstMember = static_cast<char*>(InFirstMember);
+		const char* SecondMember = static_cast<const char*>(InSecondMember);
+		FirstMember += sizeof(T);
+		if(FirstMember < InSecondMember)
+		{
+			FPlatformMemory::Memzero(FirstMember, SecondMember - FirstMember);
+		}
+	}
+}
+
 /**
  * The FRigVMOperand represents an argument used for an operator
  * within the virtual machine. Operands provide information about
@@ -122,14 +138,17 @@ public:
   	}
 	
 	void Serialize(FArchive& Ar);
-	void Save(FArchive& Ar) const;
-	void Load(FArchive& Ar);
 	friend FArchive& operator<<(FArchive& Ar, FRigVMOperand& P)
 	{
 		P.Serialize(Ar);
 		return Ar;
 	}
 
+	static void ZeroPaddedMemoryIfNeeded(FRigVMOperand* InOperand)
+	{
+		RigVM::ZeroPaddedMemory(&InOperand->MemoryType, &InOperand->RegisterIndex);
+	}
+
 private:
 
 	UPROPERTY()
diff --git a/Engine/Plugins/Runtime/RigVM/Source/RigVMDeveloper/Private/RigVMCompiler/RigVMCompiler.cpp b/Engine/Plugins/Runtime/RigVM/Source/RigVMDeveloper/Private/RigVMCompiler/RigVMCompiler.cpp
--- a/Engine/Plugins/Runtime/RigVM/Source/RigVMDeveloper/Private/RigVMCompiler/RigVMCompiler.cpp
+++ b/Engine/Plugins/Runtime/RigVM/Source/RigVMDeveloper/Private/RigVMCompiler/RigVMCompiler.cpp
@@ -1371,7 +1371,7 @@ bool URigVMCompiler::Compile(const FRigVMCompileSettings& InSettings, TArray<URi
 			}
 		}
 
-		OutFunctionCompilationData->Hash = GetTypeHash(OutFunctionCompilationData);
+		OutFunctionCompilationData->Hash = GetTypeHash(*OutFunctionCompilationData);
 		OutFunctionCompilationData->OperandToDebugRegisters = WorkData.VM->OperandToDebugRegisters;
 	}
 
