RigVM: Toggle to debug mode automatically when profiling is enabled.

Also disable profiling if release mode is re-enabled

diff --git a/Engine/Plugins/Runtime/RigVM/Source/RigVM/Public/RigVMCore/RigVMExecuteContext.h b/Engine/Plugins/Runtime/RigVM/Source/RigVM/Public/RigVMCore/RigVMExecuteContext.h
--- a/Engine/Plugins/Runtime/RigVM/Source/RigVM/Public/RigVMCore/RigVMExecuteContext.h
+++ b/Engine/Plugins/Runtime/RigVM/Source/RigVM/Public/RigVMCore/RigVMExecuteContext.h
@@ -207,6 +207,9 @@ struct RIGVM_API FRigVMRuntimeSettings
 	// Keep in mind when looking at nodes in a function the duration
 	// represents the accumulated duration of all invocations
 	// of the function currently running.
+	// 
+	// Note: This can only be used when in Debug Mode. Click the "Release" button
+	// in the top toolbar to switch to Debug mode.
 	UPROPERTY(EditAnywhere, Category = "VM")
 	bool bEnableProfiling = false;
 #endif
diff --git a/Engine/Plugins/Runtime/RigVM/Source/RigVMDeveloper/Private/RigVMBlueprint.cpp b/Engine/Plugins/Runtime/RigVM/Source/RigVMDeveloper/Private/RigVMBlueprint.cpp
--- a/Engine/Plugins/Runtime/RigVM/Source/RigVMDeveloper/Private/RigVMBlueprint.cpp
+++ b/Engine/Plugins/Runtime/RigVM/Source/RigVMDeveloper/Private/RigVMBlueprint.cpp
@@ -1603,6 +1603,15 @@ TArray<IRigVMGraphFunctionHost*> URigVMBlueprint::GetReferencedFunctionHosts(boo
 
 #if WITH_EDITOR
 
+void URigVMBlueprint::SetDebugMode(const bool bValue)
+{
+	bCompileInDebugMode = bValue;
+	if(!bCompileInDebugMode)
+	{
+		VMRuntimeSettings.bEnableProfiling = false;
+	}
+}
+
 void URigVMBlueprint::ClearBreakpoints()
 {
 	for(URigVMNode* Node : RigVMBreakpointNodes)
diff --git a/Engine/Plugins/Runtime/RigVM/Source/RigVMDeveloper/Public/RigVMBlueprint.h b/Engine/Plugins/Runtime/RigVM/Source/RigVMDeveloper/Public/RigVMBlueprint.h
--- a/Engine/Plugins/Runtime/RigVM/Source/RigVMDeveloper/Public/RigVMBlueprint.h
+++ b/Engine/Plugins/Runtime/RigVM/Source/RigVMDeveloper/Public/RigVMBlueprint.h
@@ -719,7 +719,10 @@ private:
 public:
 
 	/** Sets the execution mode. In Release mode the rig will ignore all breakpoints. */
-	void SetDebugMode(const bool bValue) { bCompileInDebugMode = bValue; }
+	void SetDebugMode(const bool bValue);
+
+	/** Returns the execution mode */
+	bool IsInDebugMode() const { return bCompileInDebugMode; }
 
 	/** Removes all the breakpoints from the blueprint and the VM */
 	void ClearBreakpoints();
diff --git a/Engine/Plugins/Runtime/RigVM/Source/RigVMEditor/Private/Editor/RigVMEditor.cpp b/Engine/Plugins/Runtime/RigVM/Source/RigVMEditor/Private/Editor/RigVMEditor.cpp
--- a/Engine/Plugins/Runtime/RigVM/Source/RigVMEditor/Private/Editor/RigVMEditor.cpp
+++ b/Engine/Plugins/Runtime/RigVM/Source/RigVMEditor/Private/Editor/RigVMEditor.cpp
@@ -1947,6 +1947,11 @@ void FRigVMEditor::OnFinishedChangingProperties(const FPropertyChangedEvent& Pro
 		{
 			RigVMBlueprint->VMRuntimeSettings.Validate();
 			RigVMBlueprint->PropagateRuntimeSettingsFromBPToInstances();
+
+			if(RigVMBlueprint->VMRuntimeSettings.bEnableProfiling && !RigVMBlueprint->IsInDebugMode())
+			{
+				SetExecutionMode(ERigVMEditorExecutionModeType_Debug);
+			}
 		}
 	}
 }
