RigVM: Enable profiling in release mode

diff --git a/Engine/Plugins/Animation/ControlRig/Source/ControlRig/Private/BaseControlRig.cpp b/Engine/Plugins/Animation/ControlRig/Source/ControlRig/Private/BaseControlRig.cpp
--- a/Engine/Plugins/Animation/ControlRig/Source/ControlRig/Private/BaseControlRig.cpp
+++ b/Engine/Plugins/Animation/ControlRig/Source/ControlRig/Private/BaseControlRig.cpp
@@ -705,6 +705,15 @@ bool UBaseControlRig::Execute(const FName& InEventName)
 		ExtendedExecuteContext.SetDebugInfo(nullptr);
 		GetSnapshotContext().Reset();
 	}
+	
+	if (IsProfilingEnabled())
+    {
+    	ExtendedExecuteContext.SetProfilingInfo(&ProfilingInfo);
+    }
+    else
+    {
+    	ExtendedExecuteContext.SetProfilingInfo(nullptr);
+    }
 #endif
 
 	FRigUnitContext& Context = PublicContext.UnitContext;
diff --git a/Engine/Plugins/Runtime/RigVM/Source/RigVM/Private/RigVMCore/RigVM.cpp b/Engine/Plugins/Runtime/RigVM/Source/RigVM/Private/RigVMCore/RigVM.cpp
--- a/Engine/Plugins/Runtime/RigVM/Source/RigVM/Private/RigVMCore/RigVM.cpp
+++ b/Engine/Plugins/Runtime/RigVM/Source/RigVM/Private/RigVMCore/RigVM.cpp
@@ -1505,11 +1505,11 @@ ERigVMExecuteResult URigVM::ExecuteVM(FRigVMExtendedExecuteContext& Context, con
 
 	if(bIsRootEntry)
 	{
-		if (FRigVMDebugInfo* RigVMDebugInfo = Context.GetRigVMDebugInfo())
+		if (FRigVMProfilingInfo* RigVMProfilingInfo = Context.GetRigVMProfilingInfo())
 		{
-			if (RigVMDebugInfo->GetFirstEntryEventInEventQueue() == NAME_None || RigVMDebugInfo->GetFirstEntryEventInEventQueue() == InEntryName)
+			if (RigVMProfilingInfo->GetFirstEntryEventInEventQueue() == NAME_None || RigVMProfilingInfo->GetFirstEntryEventInEventQueue() == InEntryName)
 			{
-				RigVMDebugInfo->SetupInstructionTracking(Instructions.Num(), Context.GetPublicData<>().RuntimeSettings.bEnableProfiling);
+				RigVMProfilingInfo->SetupInstructionTracking(Instructions.Num(), true);
 			}
 		}
 	}
@@ -1643,6 +1643,7 @@ ERigVMExecuteResult URigVM::ExecuteInstructions(FRigVMExtendedExecuteContext& Co
 #if WITH_EDITOR
 	TArray<FName>& FunctionNames = GetFunctionNames();
 	FRigVMDebugInfo* RigVMDebugInfo = Context.GetRigVMDebugInfo();
+	FRigVMProfilingInfo* RigVMProfilingInfo = Context.GetRigVMProfilingInfo();
 #endif
 	
 	while (Instructions.IsValidIndex(ContextPublicData.InstructionIndex))
@@ -1677,10 +1678,10 @@ ERigVMExecuteResult URigVM::ExecuteInstructions(FRigVMExtendedExecuteContext& Co
 
 		const int32 CurrentInstructionIndex = ContextPublicData.InstructionIndex;
 
-		if (RigVMDebugInfo != nullptr)
+		if (RigVMProfilingInfo != nullptr)
 		{
-			RigVMDebugInfo->SetInstructionVisitedDuringLastRun(ContextPublicData.InstructionIndex);
-			RigVMDebugInfo->AddInstructionIndexToVisitOrder(ContextPublicData.InstructionIndex);
+			RigVMProfilingInfo->SetInstructionVisitedDuringLastRun(ContextPublicData.InstructionIndex);
+			RigVMProfilingInfo->AddInstructionIndexToVisitOrder(ContextPublicData.InstructionIndex);
 		}
 	
 #endif
@@ -2045,23 +2046,23 @@ ERigVMExecuteResult URigVM::ExecuteInstructions(FRigVMExtendedExecuteContext& Co
 		}
 
 #if WITH_EDITOR
-		if (RigVMDebugInfo != nullptr)
+		if (RigVMProfilingInfo != nullptr)
 		{
-			if (ContextPublicData.RuntimeSettings.bEnableProfiling && !RigVMDebugInfo->GetInstructionVisitOrder().IsEmpty())
+			if (!RigVMProfilingInfo->GetInstructionVisitOrder().IsEmpty())
 			{
 				const uint64 EndCycles = FPlatformTime::Cycles64();
-				const uint64 Cycles = EndCycles - RigVMDebugInfo->GetStartCycles();
-				if (RigVMDebugInfo->GetInstructionCyclesDuringLastRun(CurrentInstructionIndex) == UINT64_MAX)
+				const uint64 Cycles = EndCycles - RigVMProfilingInfo->GetStartCycles();
+				if (RigVMProfilingInfo->GetInstructionCyclesDuringLastRun(CurrentInstructionIndex) == UINT64_MAX)
 				{
-					RigVMDebugInfo->SetInstructionCyclesDuringLastRun(CurrentInstructionIndex, Cycles);
+					RigVMProfilingInfo->SetInstructionCyclesDuringLastRun(CurrentInstructionIndex, Cycles);
 				}
 				else
 				{
-					RigVMDebugInfo->AddInstructionCyclesDuringLastRun(CurrentInstructionIndex, Cycles);
+					RigVMProfilingInfo->AddInstructionCyclesDuringLastRun(CurrentInstructionIndex, Cycles);
 				}
 
-				RigVMDebugInfo->SetStartCycles(EndCycles);
-				RigVMDebugInfo->AddOverallCycles(Cycles);
+				RigVMProfilingInfo->SetStartCycles(EndCycles);
+				RigVMProfilingInfo->AddOverallCycles(Cycles);
 			}
 		}
 
@@ -2119,12 +2120,12 @@ ERigVMExecuteResult URigVM::ExecuteInstructions(FRigVMExtendedExecuteContext& Co
 ERigVMExecuteResult URigVM::ExecuteBranch(FRigVMExtendedExecuteContext& Context, const FRigVMBranchInfo& InBranchToRun)
 {
 #if WITH_EDITOR
-	const double LastExecutionMicroSecondsGuard = Context.GetRigVMDebugInfo() ? Context.GetRigVMDebugInfo()->GetLastExecutionMicroSeconds() : 0.0;
+	const double LastExecutionMicroSecondsGuard = Context.GetRigVMProfilingInfo() ? Context.GetRigVMProfilingInfo()->GetLastExecutionMicroSeconds() : 0.0;
 	ON_SCOPE_EXIT
 	{
-		if (Context.GetRigVMDebugInfo())
+		if (Context.GetRigVMProfilingInfo())
 		{
-			Context.GetRigVMDebugInfo()->SetLastExecutionMicroSeconds(LastExecutionMicroSecondsGuard);
+			Context.GetRigVMProfilingInfo()->SetLastExecutionMicroSeconds(LastExecutionMicroSecondsGuard);
 		}
 	};
 #endif
@@ -2617,9 +2618,9 @@ void URigVM::RefreshExternalPropertyPaths()
 void URigVM::SetupInstructionTracking(FRigVMExtendedExecuteContext& Context, int32 InInstructionCount)
 {
 #if WITH_EDITOR
-	if (FRigVMDebugInfo* RigVMDebugInfo = Context.GetRigVMDebugInfo())
+	if (FRigVMProfilingInfo* RigVMProfilingInfo = Context.GetRigVMProfilingInfo())
 	{
-		RigVMDebugInfo->SetupInstructionTracking(InInstructionCount, Context.GetPublicData<>().RuntimeSettings.bEnableProfiling);
+		RigVMProfilingInfo->SetupInstructionTracking(InInstructionCount, true);
 	}
 #endif
 }
@@ -2627,9 +2628,9 @@ void URigVM::SetupInstructionTracking(FRigVMExtendedExecuteContext& Context, int
 void URigVM::StartProfiling(FRigVMExtendedExecuteContext& Context)
 {
 #if WITH_EDITOR
-	if (FRigVMDebugInfo* RigVMDebugInfo = Context.GetRigVMDebugInfo())
+	if (FRigVMProfilingInfo* RigVMProfilingInfo = Context.GetRigVMProfilingInfo())
 	{
-		RigVMDebugInfo->StartProfiling(Context.GetPublicData<>().RuntimeSettings.bEnableProfiling);
+		RigVMProfilingInfo->StartProfiling(true);
 	}
 #endif
 }
@@ -2637,9 +2638,9 @@ void URigVM::StartProfiling(FRigVMExtendedExecuteContext& Context)
 void URigVM::StopProfiling(FRigVMExtendedExecuteContext& Context)
 {
 #if WITH_EDITOR
-	if (FRigVMDebugInfo* RigVMDebugInfo = Context.GetRigVMDebugInfo())
+	if (FRigVMProfilingInfo* RigVMProfilingInfo = Context.GetRigVMProfilingInfo())
 	{
-		RigVMDebugInfo->StopProfiling();
+		RigVMProfilingInfo->StopProfiling();
 	}
 #endif
 }
diff --git a/Engine/Plugins/Runtime/RigVM/Source/RigVM/Private/RigVMCore/RigVMDebugInfo.cpp b/Engine/Plugins/Runtime/RigVM/Source/RigVM/Private/RigVMCore/RigVMDebugInfo.cpp
--- a/Engine/Plugins/Runtime/RigVM/Source/RigVM/Private/RigVMCore/RigVMDebugInfo.cpp
+++ b/Engine/Plugins/Runtime/RigVM/Source/RigVM/Private/RigVMCore/RigVMDebugInfo.cpp
@@ -228,22 +228,6 @@ uint16 FRigVMDebugInfo::GetBreakpointHits(const FRigVMBreakpoint& InBreakpoint)
 	return 0;
 }
 
-void FRigVMDebugInfo::SetupInstructionTracking(int32 InInstructionCount, bool bEnableProfiling)
-{
-#if WITH_EDITOR
-
-	ResetInstructionVisitedDuringLastRun(InInstructionCount);
-	SetNumInstructionVisitedDuringLastRunZeroed(InInstructionCount);
-	ResetInstructionVisitOrder(InInstructionCount);
-	ResetInstructionCyclesDuringLastRun(InInstructionCount);
-
-	if (bEnableProfiling)
-	{
-		InitInstructionCyclesDuringLastRunValues(InInstructionCount, UINT64_MAX);
-	}
-#endif
-}
-
 bool FRigVMDebugInfo::ResumeExecution()
 {
 	HaltedAtBreakpoint.Reset();
@@ -258,23 +242,3 @@ bool FRigVMDebugInfo::ResumeExecution()
 
 	return false;
 }
-
-void FRigVMDebugInfo::StartProfiling(bool bEnableProfiling)
-{
-#if WITH_EDITOR
-		SetStartCycles(0);
-		SetOverallCycles(0);
-		if (bEnableProfiling)
-		{
-			SetStartCycles(FPlatformTime::Cycles64());
-		}
-#endif
-}
-
-void FRigVMDebugInfo::StopProfiling()
-{
-#if WITH_EDITOR
-	const uint64 Cycles = GetOverallCycles() > 0 ? GetOverallCycles() : (FPlatformTime::Cycles64() - GetStartCycles());
-	LastExecutionMicroSeconds = Cycles * FPlatformTime::GetSecondsPerCycle() * 1000.0 * 1000.0;
-#endif
-}
diff --git a/Engine/Plugins/Runtime/RigVM/Source/RigVM/Private/RigVMCore/RigVMProfilingInfo.cpp b/Engine/Plugins/Runtime/RigVM/Source/RigVM/Private/RigVMCore/RigVMProfilingInfo.cpp
new file mode 100644
--- /dev/null
+++ b/Engine/Plugins/Runtime/RigVM/Source/RigVM/Private/RigVMCore/RigVMProfilingInfo.cpp
@@ -0,0 +1,41 @@
+﻿// Copyright Epic Games, Inc. All Rights Reserved.
+
+#include "RigVMCore/RigVMProfilingInfo.h"
+
+#include UE_INLINE_GENERATED_CPP_BY_NAME(RigVMProfilingInfo)
+
+void FRigVMProfilingInfo::SetupInstructionTracking(int32 InInstructionCount, bool bEnableProfiling)
+{
+#if WITH_EDITOR
+
+	ResetInstructionVisitedDuringLastRun(InInstructionCount);
+	SetNumInstructionVisitedDuringLastRunZeroed(InInstructionCount);
+	ResetInstructionVisitOrder(InInstructionCount);
+	ResetInstructionCyclesDuringLastRun(InInstructionCount);
+
+	if (bEnableProfiling)
+	{
+		InitInstructionCyclesDuringLastRunValues(InInstructionCount, UINT64_MAX);
+	}
+#endif
+}
+
+void FRigVMProfilingInfo::StartProfiling(bool bEnableProfiling)
+{
+#if WITH_EDITOR
+		SetStartCycles(0);
+		SetOverallCycles(0);
+		if (bEnableProfiling)
+		{
+			SetStartCycles(FPlatformTime::Cycles64());
+		}
+#endif
+}
+
+void FRigVMProfilingInfo::StopProfiling()
+{
+#if WITH_EDITOR
+	const uint64 Cycles = GetOverallCycles() > 0 ? GetOverallCycles() : (FPlatformTime::Cycles64() - GetStartCycles());
+	LastExecutionMicroSeconds = Cycles * FPlatformTime::GetSecondsPerCycle() * 1000.0 * 1000.0;
+#endif
+}
diff --git a/Engine/Plugins/Runtime/RigVM/Source/RigVM/Private/RigVMHost.cpp b/Engine/Plugins/Runtime/RigVM/Source/RigVM/Private/RigVMHost.cpp
--- a/Engine/Plugins/Runtime/RigVM/Source/RigVM/Private/RigVMHost.cpp
+++ b/Engine/Plugins/Runtime/RigVM/Source/RigVM/Private/RigVMHost.cpp
@@ -521,6 +521,15 @@ bool URigVMHost::Execute(const FName& InEventName)
 	{
 		ExtendedExecuteContext.SetDebugInfo(nullptr);
 	}
+	
+	if (IsProfilingEnabled())
+	{
+		ExtendedExecuteContext.SetProfilingInfo(&ProfilingInfo);
+	}
+	else
+	{
+		ExtendedExecuteContext.SetProfilingInfo(nullptr);
+	}
 #endif
 
 	if (VM)
diff --git a/Engine/Plugins/Runtime/RigVM/Source/RigVM/Public/RigVMCore/RigVM.h b/Engine/Plugins/Runtime/RigVM/Source/RigVM/Public/RigVMCore/RigVM.h
--- a/Engine/Plugins/Runtime/RigVM/Source/RigVM/Public/RigVMCore/RigVM.h
+++ b/Engine/Plugins/Runtime/RigVM/Source/RigVM/Public/RigVMCore/RigVM.h
@@ -445,9 +445,9 @@ public:
 	// Returns the number of times an instruction has been hit
 	int32 GetInstructionVisitedCount(const FRigVMExtendedExecuteContext& Context, int32 InIndex) const
 	{
-		if (const FRigVMDebugInfo* RigVMDebugInfo = Context.GetRigVMDebugInfo())
+		if (const FRigVMProfilingInfo* RigVMProfilingInfo = Context.GetRigVMProfilingInfo())
 		{
-			return RigVMDebugInfo->GetInstructionVisitedCountDuringLastRun(InIndex);
+			return RigVMProfilingInfo->GetInstructionVisitedCountDuringLastRun(InIndex);
 		}
 		return 0;
 	}
@@ -460,9 +460,9 @@ public:
 	// If there is no information available this function returns UINT64_MAX.
 	uint64 GetInstructionCycles(const FRigVMExtendedExecuteContext& Context, int32 InIndex) const
 	{
-		if (const FRigVMDebugInfo* RigVMDebugInfo = Context.GetRigVMDebugInfo())
+		if (const FRigVMProfilingInfo* RigVMProfilingInfo = Context.GetRigVMProfilingInfo())
 		{
-			return RigVMDebugInfo->GetInstructionCyclesDuringLastRun(InIndex);
+			return RigVMProfilingInfo->GetInstructionCyclesDuringLastRun(InIndex);
 		}
 		return UINT64_MAX;
 	}
@@ -489,9 +489,9 @@ public:
 	// Returns the order of all instructions during the last run
 	const TArray<int32> GetInstructionVisitOrder(const FRigVMExtendedExecuteContext& Context) const
 	{
-		if (const FRigVMDebugInfo* RigVMDebugInfo = Context.GetRigVMDebugInfo())
+		if (const FRigVMProfilingInfo* RigVMProfilingInfo = Context.GetRigVMProfilingInfo())
 		{
-			return RigVMDebugInfo->GetInstructionVisitOrder();
+			return RigVMProfilingInfo->GetInstructionVisitOrder();
 		}
 
 		return TArray<int32>();
@@ -502,9 +502,9 @@ public:
 
 	const void SetFirstEntryEventInEventQueue(FRigVMExtendedExecuteContext& Context, const FName& InFirstEventName) 
 	{ 
-		if (FRigVMDebugInfo* RigVMDebugInfo = Context.GetRigVMDebugInfo())
+		if (FRigVMProfilingInfo* RigVMProfilingInfo = Context.GetRigVMProfilingInfo())
 		{
-			RigVMDebugInfo->SetFirstEntryEventInEventQueue(InFirstEventName);
+			RigVMProfilingInfo->SetFirstEntryEventInEventQueue(InFirstEventName);
 		}
 	}
 
diff --git a/Engine/Plugins/Runtime/RigVM/Source/RigVM/Public/RigVMCore/RigVMDebugInfo.h b/Engine/Plugins/Runtime/RigVM/Source/RigVM/Public/RigVMCore/RigVMDebugInfo.h
--- a/Engine/Plugins/Runtime/RigVM/Source/RigVM/Public/RigVMCore/RigVMDebugInfo.h
+++ b/Engine/Plugins/Runtime/RigVM/Source/RigVM/Public/RigVMCore/RigVMDebugInfo.h
@@ -118,15 +118,6 @@ struct RIGVM_API FRigVMDebugInfo
 		HaltedAtBreakpoint = FRigVMBreakpoint();
 		HaltedAtBreakpointHit = INDEX_NONE;
 		CurrentBreakpointAction = ERigVMBreakpointAction::None;
-
-		InstructionVisitedDuringLastRun.Reset();
-		InstructionCyclesDuringLastRun.Reset();
-		InstructionVisitOrder.Reset();
-		FirstEntryEventInQueue = NAME_None;
-
-		StartCycles = 0;
-		OverallCycles = 0;
-
 	}
 
 	bool IsEmpty() const
@@ -194,13 +185,6 @@ struct RIGVM_API FRigVMDebugInfo
 
 	inline void SetCurrentActiveBreakpointCallstack(TArray<TWeakObjectPtr<UObject>> Callstack) { CurrentActiveBreakpointCallstack = Callstack; }
 
-	inline uint64 GetStartCycles() const {	return StartCycles;	}
-	inline void SetStartCycles(uint64 InStartCycles) { StartCycles = InStartCycles; }
-
-	inline uint64 GetOverallCycles() const { return OverallCycles; }
-	inline void SetOverallCycles(uint64 Cycles) { OverallCycles = Cycles; }
-	inline void AddOverallCycles(uint64 Cycles) { OverallCycles += Cycles; }
-
 	inline FRigVMBreakpoint& GetHaltedAtBreakpoint() { return HaltedAtBreakpoint; }
 	inline const FRigVMBreakpoint& GetHaltedAtBreakpoint() const { return HaltedAtBreakpoint; }
 	inline void SetHaltedAtBreakpoint(const FRigVMBreakpoint& InHaltedAtBreakpoint) { HaltedAtBreakpoint = InHaltedAtBreakpoint; }
@@ -211,34 +195,6 @@ struct RIGVM_API FRigVMDebugInfo
 	inline ERigVMBreakpointAction GetCurrentBreakpointAction() const { return CurrentBreakpointAction; }
 	inline void SetCurrentBreakpointAction(ERigVMBreakpointAction InCurrentBreakpointAction) { CurrentBreakpointAction = InCurrentBreakpointAction; }
 
-	inline void ResetInstructionVisitedDuringLastRun(int32 NewSize = 0) { InstructionVisitedDuringLastRun.Reset(NewSize); }
-	inline void SetNumInstructionVisitedDuringLastRunZeroed(int32 Num) { InstructionVisitedDuringLastRun.SetNumZeroed(Num); }
-	inline void SetInstructionVisitedDuringLastRun(int32 InstructionIndex) { InstructionVisitedDuringLastRun[InstructionIndex]++; }
-	inline int32 GetInstructionVisitedCountDuringLastRun(int32 InstructionIndex) const { return InstructionVisitedDuringLastRun.IsValidIndex(InstructionIndex) ? InstructionVisitedDuringLastRun[InstructionIndex] : 0; }
-	inline const TArray<int32>& GetInstructionVisitedCountDuringLastRun() const { return InstructionVisitedDuringLastRun; }
-
-	inline void ResetInstructionCyclesDuringLastRun(int32 NewSize = 0) { InstructionCyclesDuringLastRun.Reset(NewSize); }
-	inline uint64 GetInstructionCyclesDuringLastRun(int32 InstructionIndex) const { return InstructionCyclesDuringLastRun.IsValidIndex(InstructionIndex) ? InstructionCyclesDuringLastRun[InstructionIndex] : UINT64_MAX; }
-	inline void SetInstructionCyclesDuringLastRun(int32 InstructionIndex, uint64 CyclesDuringLastRun) { InstructionCyclesDuringLastRun[InstructionIndex] = CyclesDuringLastRun; }
-	inline void AddInstructionCyclesDuringLastRun(int32 InstructionIndex, uint64 CyclesDuringLastRun) { InstructionCyclesDuringLastRun[InstructionIndex] += CyclesDuringLastRun; }
-	inline void InitInstructionCyclesDuringLastRunValues(int32 NewSize, uint64 DefaultValue)
-	{
-		InstructionCyclesDuringLastRun.SetNumUninitialized(NewSize);
-		for (uint64& Value : InstructionCyclesDuringLastRun)
-		{
-			Value = DefaultValue;
-		}
-	}
-
-	
-	inline void ResetInstructionVisitOrder(int32 NewSize = 0) { InstructionVisitOrder.Reset(NewSize); }
-	inline void AddInstructionIndexToVisitOrder(int32 InstructionIndex) { InstructionVisitOrder.Add(InstructionIndex); }
-	inline const TArray<int32>& GetInstructionVisitOrder() const { return InstructionVisitOrder; }
-
-	inline const void SetFirstEntryEventInEventQueue(const FName& InFirstEventName) { FirstEntryEventInQueue = InFirstEventName; }
-	inline const FName& GetFirstEntryEventInEventQueue() const { return FirstEntryEventInQueue; }
-
-	void SetupInstructionTracking(int32 InInstructionCount, bool bEnableProfiling);
 	bool ResumeExecution();
 
 	DECLARE_EVENT_ThreeParams(URigVM, FExecutionHaltedEvent, int32, UObject*, const FName&);
@@ -247,12 +203,6 @@ struct RIGVM_API FRigVMDebugInfo
 		return OnExecutionHalted;
 	}
 
-	double GetLastExecutionMicroSeconds() const { return LastExecutionMicroSeconds; }
-	void SetLastExecutionMicroSeconds(double InLastExecutionMicroSeconds) { LastExecutionMicroSeconds = InLastExecutionMicroSeconds; }
-
-	void StartProfiling(bool bEnableProfiling);
-	void StopProfiling();
-
 private:
 	TArray<FRigVMBreakpoint> Breakpoints;
 	FRigVMBreakpoint TemporaryBreakpoint;
@@ -268,19 +218,5 @@ private:
 	int32 HaltedAtBreakpointHit = INDEX_NONE;
 	ERigVMBreakpointAction CurrentBreakpointAction = ERigVMBreakpointAction::None;
 
-	// stores the number of times each instruction was visited
-	TArray<int32> InstructionVisitedDuringLastRun;
-	TArray<uint64> InstructionCyclesDuringLastRun;
-	TArray<int32> InstructionVisitOrder;
-
-	// A RigVMHost can run multiple events per evaluation, such as the Backward&Forward Solve Mode,
-	// store the first event such that we know when to reset data for a new round of rig evaluation
-	FName FirstEntryEventInQueue = NAME_None;
-
-	uint64 StartCycles = 0;
-	uint64 OverallCycles = 0;
-
 	FExecutionHaltedEvent OnExecutionHalted;
-
-	double LastExecutionMicroSeconds = 0.0;
 };
diff --git a/Engine/Plugins/Runtime/RigVM/Source/RigVM/Public/RigVMCore/RigVMExecuteContext.h b/Engine/Plugins/Runtime/RigVM/Source/RigVM/Public/RigVMCore/RigVMExecuteContext.h
--- a/Engine/Plugins/Runtime/RigVM/Source/RigVM/Public/RigVMCore/RigVMExecuteContext.h
+++ b/Engine/Plugins/Runtime/RigVM/Source/RigVM/Public/RigVMCore/RigVMExecuteContext.h
@@ -17,6 +17,7 @@
 #include "RigVMExternalVariable.h"
 #include "RigVMModule.h"
 #include "RigVMCore/RigVMDebugInfo.h"
+#include "RigVMCore/RigVMProfilingInfo.h"
 #include "RigVMCore/RigVMNameCache.h"
 #include "RigVMCore/RigVMMemoryStorageStruct.h"
 #include "RigVMLog.h"
@@ -784,5 +785,18 @@ struct RIGVM_API FRigVMExtendedExecuteContext
 	{
 		return DebugInfo;
 	}
+
+	FRigVMProfilingInfo* ProfilingInfo = nullptr;
+
+	void SetProfilingInfo(FRigVMProfilingInfo* InProfilingInfo)
+	{
+		ProfilingInfo = InProfilingInfo;
+	}
+
+	FRigVMProfilingInfo* GetRigVMProfilingInfo() const
+	{
+		return ProfilingInfo;
+	}
+
 #endif // WITH_EDITOR
 };
diff --git a/Engine/Plugins/Runtime/RigVM/Source/RigVM/Public/RigVMCore/RigVMNativized.h b/Engine/Plugins/Runtime/RigVM/Source/RigVM/Public/RigVMCore/RigVMNativized.h
--- a/Engine/Plugins/Runtime/RigVM/Source/RigVM/Public/RigVMCore/RigVMNativized.h
+++ b/Engine/Plugins/Runtime/RigVM/Source/RigVM/Public/RigVMCore/RigVMNativized.h
@@ -64,10 +64,10 @@ public:
 	{
 		Super::SetInstructionIndex(Context, InInstructionIndex);
 #if WITH_EDITOR
-		if (FRigVMDebugInfo* RigVMDebugInfo = Context.GetRigVMDebugInfo())
+		if (FRigVMProfilingInfo* RigVMProfilingInfo = Context.GetRigVMProfilingInfo())
 		{
-			RigVMDebugInfo->SetInstructionVisitedDuringLastRun(InInstructionIndex);
-			RigVMDebugInfo->AddInstructionIndexToVisitOrder(InInstructionIndex);
+			RigVMProfilingInfo->SetInstructionVisitedDuringLastRun(InInstructionIndex);
+			RigVMProfilingInfo->AddInstructionIndexToVisitOrder(InInstructionIndex);
 		}
 #endif
 	}
diff --git a/Engine/Plugins/Runtime/RigVM/Source/RigVM/Public/RigVMCore/RigVMProfilingInfo.h b/Engine/Plugins/Runtime/RigVM/Source/RigVM/Public/RigVMCore/RigVMProfilingInfo.h
new file mode 100644
--- /dev/null
+++ b/Engine/Plugins/Runtime/RigVM/Source/RigVM/Public/RigVMCore/RigVMProfilingInfo.h
@@ -0,0 +1,89 @@
+﻿// Copyright Epic Games, Inc. All Rights Reserved.
+
+#pragma once
+
+#include "Containers/Array.h"
+#include "CoreMinimal.h"
+#include "CoreTypes.h"
+#include "UObject/ObjectMacros.h"
+
+#include "RigVMProfilingInfo.generated.h"
+
+USTRUCT()
+struct RIGVM_API FRigVMProfilingInfo
+{
+	GENERATED_BODY()
+
+	FRigVMProfilingInfo()
+	{
+	}
+
+	void Reset()
+	{
+		InstructionVisitedDuringLastRun.Reset();
+		InstructionCyclesDuringLastRun.Reset();
+		InstructionVisitOrder.Reset();
+		FirstEntryEventInQueue = NAME_None;
+
+		StartCycles = 0;
+		OverallCycles = 0;
+	}
+
+	inline uint64 GetStartCycles() const {	return StartCycles;	}
+	inline void SetStartCycles(uint64 InStartCycles) { StartCycles = InStartCycles; }
+
+	inline uint64 GetOverallCycles() const { return OverallCycles; }
+	inline void SetOverallCycles(uint64 Cycles) { OverallCycles = Cycles; }
+	inline void AddOverallCycles(uint64 Cycles) { OverallCycles += Cycles; }
+
+	inline void ResetInstructionVisitedDuringLastRun(int32 NewSize = 0) { InstructionVisitedDuringLastRun.Reset(NewSize); }
+	inline void SetNumInstructionVisitedDuringLastRunZeroed(int32 Num) { InstructionVisitedDuringLastRun.SetNumZeroed(Num); }
+	inline void SetInstructionVisitedDuringLastRun(int32 InstructionIndex) { InstructionVisitedDuringLastRun[InstructionIndex]++; }
+	inline int32 GetInstructionVisitedCountDuringLastRun(int32 InstructionIndex) const { return InstructionVisitedDuringLastRun.IsValidIndex(InstructionIndex) ? InstructionVisitedDuringLastRun[InstructionIndex] : 0; }
+	inline const TArray<int32>& GetInstructionVisitedCountDuringLastRun() const { return InstructionVisitedDuringLastRun; }
+
+	inline void ResetInstructionCyclesDuringLastRun(int32 NewSize = 0) { InstructionCyclesDuringLastRun.Reset(NewSize); }
+	inline uint64 GetInstructionCyclesDuringLastRun(int32 InstructionIndex) const { return InstructionCyclesDuringLastRun.IsValidIndex(InstructionIndex) ? InstructionCyclesDuringLastRun[InstructionIndex] : UINT64_MAX; }
+	inline void SetInstructionCyclesDuringLastRun(int32 InstructionIndex, uint64 CyclesDuringLastRun) { InstructionCyclesDuringLastRun[InstructionIndex] = CyclesDuringLastRun; }
+	inline void AddInstructionCyclesDuringLastRun(int32 InstructionIndex, uint64 CyclesDuringLastRun) { InstructionCyclesDuringLastRun[InstructionIndex] += CyclesDuringLastRun; }
+	inline void InitInstructionCyclesDuringLastRunValues(int32 NewSize, uint64 DefaultValue)
+	{
+		InstructionCyclesDuringLastRun.SetNumUninitialized(NewSize);
+		for (uint64& Value : InstructionCyclesDuringLastRun)
+		{
+			Value = DefaultValue;
+		}
+	}
+
+	
+	inline void ResetInstructionVisitOrder(int32 NewSize = 0) { InstructionVisitOrder.Reset(NewSize); }
+	inline void AddInstructionIndexToVisitOrder(int32 InstructionIndex) { InstructionVisitOrder.Add(InstructionIndex); }
+	inline const TArray<int32>& GetInstructionVisitOrder() const { return InstructionVisitOrder; }
+
+	inline const void SetFirstEntryEventInEventQueue(const FName& InFirstEventName) { FirstEntryEventInQueue = InFirstEventName; }
+	inline const FName& GetFirstEntryEventInEventQueue() const { return FirstEntryEventInQueue; }
+
+	void SetupInstructionTracking(int32 InInstructionCount, bool bEnableProfiling);
+
+	double GetLastExecutionMicroSeconds() const { return LastExecutionMicroSeconds; }
+	void SetLastExecutionMicroSeconds(double InLastExecutionMicroSeconds) { LastExecutionMicroSeconds = InLastExecutionMicroSeconds; }
+
+	void StartProfiling(bool bEnableProfiling);
+	void StopProfiling();
+
+private:
+
+	// stores the number of times each instruction was visited
+	TArray<int32> InstructionVisitedDuringLastRun;
+	TArray<uint64> InstructionCyclesDuringLastRun;
+	TArray<int32> InstructionVisitOrder;
+
+	// A RigVMHost can run multiple events per evaluation, such as the Backward&Forward Solve Mode,
+	// store the first event such that we know when to reset data for a new round of rig evaluation
+	FName FirstEntryEventInQueue = NAME_None;
+
+	uint64 StartCycles = 0;
+	uint64 OverallCycles = 0;
+
+	double LastExecutionMicroSeconds = 0.0;
+};
diff --git a/Engine/Plugins/Runtime/RigVM/Source/RigVM/Public/RigVMHost.h b/Engine/Plugins/Runtime/RigVM/Source/RigVM/Public/RigVMHost.h
--- a/Engine/Plugins/Runtime/RigVM/Source/RigVM/Public/RigVMHost.h
+++ b/Engine/Plugins/Runtime/RigVM/Source/RigVM/Public/RigVMHost.h
@@ -465,6 +465,7 @@ private:
 protected:
 	
 	FRigVMDebugInfo DebugInfo;
+	FRigVMProfilingInfo ProfilingInfo;
 	TMap<FString, bool> LoggedMessages;
 	void LogOnce(EMessageSeverity::Type InSeverity, int32 InInstructionIndex, const FString& InMessage);
 	
@@ -472,6 +473,10 @@ public:
 
 	void SetIsInDebugMode(const bool bValue) { bIsInDebugMode = bValue; }
 	bool IsInDebugMode() const { return bIsInDebugMode; }
+	bool IsProfilingEnabled() const
+	{
+		return IsInDebugMode() || VMRuntimeSettings.bEnableProfiling;
+	}
 	
 	/** Adds a breakpoint in the VM at the InstructionIndex for the Node / Subject */
 	void AddBreakpoint(int32 InstructionIndex, UObject* InSubject, uint16 InDepth);
@@ -493,6 +498,9 @@ public:
 	FRigVMDebugInfo& GetDebugInfo() { return DebugInfo; }
 	const FRigVMDebugInfo& GetDebugInfo() const { return DebugInfo; }
 
+	FRigVMProfilingInfo& GetProfilingInfo() { return ProfilingInfo; }
+	const FRigVMProfilingInfo& GetProfilingInfo() const { return ProfilingInfo; }
+
 	/** Creates the snapshot VM if required and returns it */
 	URigVM* GetSnapshotVM(bool bCreateIfNeeded = true);
 	FRigVMExtendedExecuteContext& GetSnapshotContext();
diff --git a/Engine/Plugins/Runtime/RigVM/Source/RigVMDeveloper/Private/RigVMBlueprint.cpp b/Engine/Plugins/Runtime/RigVM/Source/RigVMDeveloper/Private/RigVMBlueprint.cpp
--- a/Engine/Plugins/Runtime/RigVM/Source/RigVMDeveloper/Private/RigVMBlueprint.cpp
+++ b/Engine/Plugins/Runtime/RigVM/Source/RigVMDeveloper/Private/RigVMBlueprint.cpp
@@ -1183,6 +1183,9 @@ void URigVMBlueprint::RecompileVM()
 
 	RigGraphDisplaySettings.MinMicroSeconds = RigGraphDisplaySettings.LastMinMicroSeconds = DBL_MAX;
 	RigGraphDisplaySettings.MaxMicroSeconds = RigGraphDisplaySettings.LastMaxMicroSeconds = (double)INDEX_NONE;
+	RigGraphDisplaySettings.MinMicroSecondsFrames.Reset();
+	RigGraphDisplaySettings.MaxMicroSecondsFrames.Reset();
+	RigGraphDisplaySettings.TotalMicroSecondsFrames.Reset();
 
 	URigVMHost* CDO = Cast<URigVMHost>(RigClass->GetDefaultObject(true /* create if needed */));
 	if (CDO && CDO->VM != nullptr)
@@ -1606,10 +1609,6 @@ TArray<IRigVMGraphFunctionHost*> URigVMBlueprint::GetReferencedFunctionHosts(boo
 void URigVMBlueprint::SetDebugMode(const bool bValue)
 {
 	bCompileInDebugMode = bValue;
-	if(!bCompileInDebugMode)
-	{
-		VMRuntimeSettings.bEnableProfiling = false;
-	}
 }
 
 void URigVMBlueprint::ClearBreakpoints()
diff --git a/Engine/Plugins/Runtime/RigVM/Source/RigVMDeveloper/Public/RigVMBlueprint.h b/Engine/Plugins/Runtime/RigVM/Source/RigVMDeveloper/Public/RigVMBlueprint.h
--- a/Engine/Plugins/Runtime/RigVM/Source/RigVMDeveloper/Public/RigVMBlueprint.h
+++ b/Engine/Plugins/Runtime/RigVM/Source/RigVMDeveloper/Public/RigVMBlueprint.h
@@ -67,6 +67,7 @@ struct RIGVMDEVELOPER_API FRigVMEdGraphDisplaySettings
 		, MinMicroSeconds(0.0)
 		, MaxMicroSeconds(1.0)
 		, TotalMicroSeconds(0.0)
+		, AverageFrames(16)
 		, bAutoDetermineRange(true)
 		, LastMinMicroSeconds(0.0)
 		, LastMaxMicroSeconds(1.0)
@@ -114,6 +115,14 @@ struct RIGVMDEVELOPER_API FRigVMEdGraphDisplaySettings
 	UPROPERTY(VisibleAnywhere, Category = "Graph Display Settings", transient)
 	double TotalMicroSeconds;
 
+	// If you set this to more than 1 the results will be averaged across multiple frames
+	UPROPERTY(EditAnywhere, Category = "Graph Display Settings")
+	int32 AverageFrames;
+
+	TArray<double> MinMicroSecondsFrames;
+	TArray<double> MaxMicroSecondsFrames;
+	TArray<double> TotalMicroSecondsFrames;
+
 	UPROPERTY(EditAnywhere, Category = "Graph Display Settings")
 	bool bAutoDetermineRange;
 
diff --git a/Engine/Plugins/Runtime/RigVM/Source/RigVMEditor/Private/Editor/RigVMEditor.cpp b/Engine/Plugins/Runtime/RigVM/Source/RigVMEditor/Private/Editor/RigVMEditor.cpp
--- a/Engine/Plugins/Runtime/RigVM/Source/RigVMEditor/Private/Editor/RigVMEditor.cpp
+++ b/Engine/Plugins/Runtime/RigVM/Source/RigVMEditor/Private/Editor/RigVMEditor.cpp
@@ -1813,7 +1813,7 @@ void FRigVMEditor::HandleVMExecutedEvent(URigVMHost* InHost, const FName& InEven
 		{
 			if(DebuggedHost)
 			{
-				RigVMBlueprint->RigGraphDisplaySettings.TotalMicroSeconds = DebuggedHost->GetDebugInfo().GetLastExecutionMicroSeconds();
+				RigVMBlueprint->RigGraphDisplaySettings.TotalMicroSeconds = DebuggedHost->GetProfilingInfo().GetLastExecutionMicroSeconds();
 			}
 
 			if(RigVMBlueprint->RigGraphDisplaySettings.bAutoDetermineRange)
@@ -1947,11 +1947,6 @@ void FRigVMEditor::OnFinishedChangingProperties(const FPropertyChangedEvent& Pro
 		{
 			RigVMBlueprint->VMRuntimeSettings.Validate();
 			RigVMBlueprint->PropagateRuntimeSettingsFromBPToInstances();
-
-			if(RigVMBlueprint->VMRuntimeSettings.bEnableProfiling && !RigVMBlueprint->IsInDebugMode())
-			{
-				SetExecutionMode(ERigVMEditorExecutionModeType_Debug);
-			}
 		}
 	}
 }
