Removed useless logs

#rnx

diff --git a/Engine/Plugins/Runtime/RigVM/Source/RigVMEditor/Private/EdGraph/NodeSpawners/RigVMEdGraphTemplateNodeSpawner.cpp b/Engine/Plugins/Runtime/RigVM/Source/RigVMEditor/Private/EdGraph/NodeSpawners/RigVMEdGraphTemplateNodeSpawner.cpp
--- a/Engine/Plugins/Runtime/RigVM/Source/RigVMEditor/Private/EdGraph/NodeSpawners/RigVMEdGraphTemplateNodeSpawner.cpp
+++ b/Engine/Plugins/Runtime/RigVM/Source/RigVMEditor/Private/EdGraph/NodeSpawners/RigVMEdGraphTemplateNodeSpawner.cpp
@@ -76,22 +76,11 @@ FBlueprintNodeSignature URigVMEdGraphTemplateNodeSpawner::GetSpawnerSignature()
 
 bool URigVMEdGraphTemplateNodeSpawner::IsTemplateNodeFilteredOut(FBlueprintActionFilter const& Filter) const
 {
-	UE_LOG(LogRigVM, Display, TEXT("Filtering notation '%s'"), *Template->GetNotation().ToString());
-
 	if(URigVMEdGraphNodeSpawner::IsTemplateNodeFilteredOut(Filter))
 	{
 		return true;
 	}
 
-	if(Template->GetNotation().ToString().Contains(TEXT("Add::Execute(")))
-	{
-		UE_LOG(LogRigVM, Display, TEXT("Hitting notation '%s'"), *Template->GetNotation().ToString());
-		UE_LOG(LogRigVM, Display, TEXT("Hitting notation '%s'"), *Template->GetNotation().ToString());
-		UE_LOG(LogRigVM, Display, TEXT("Hitting notation '%s'"), *Template->GetNotation().ToString());
-		UE_LOG(LogRigVM, Display, TEXT("Hitting notation '%s'"), *Template->GetNotation().ToString());
-		UE_LOG(LogRigVM, Display, TEXT("Hitting notation '%s'"), *Template->GetNotation().ToString());
-	}
-
 	for (const UEdGraphPin* Pin : Filter.Context.Pins)
 	{
 		for (UEdGraph* Graph : Filter.Context.Graphs)
