RigVM: General improvements of the profiling featureset

#rnx

diff --git a/Engine/Plugins/Runtime/RigVM/Source/RigVMDeveloper/Private/EdGraph/RigVMEdGraphNode.cpp b/Engine/Plugins/Runtime/RigVM/Source/RigVMDeveloper/Private/EdGraph/RigVMEdGraphNode.cpp
--- a/Engine/Plugins/Runtime/RigVM/Source/RigVMDeveloper/Private/EdGraph/RigVMEdGraphNode.cpp
+++ b/Engine/Plugins/Runtime/RigVM/Source/RigVMDeveloper/Private/EdGraph/RigVMEdGraphNode.cpp
@@ -44,9 +44,14 @@ URigVMEdGraphNode::URigVMEdGraphNode()
 , bEnableProfiling(false)
 #endif
 , CachedTemplate(nullptr)
+, MicroSeconds(0)
 {
 	bHasCompilerMessage = false;
 	ErrorType = (int32)EMessageSeverity::Info + 1;
+	
+#if WITH_EDITOR
+	UpdateProfilingSettings();
+#endif
 }
 
 FText URigVMEdGraphNode::GetNodeTitle(ENodeTitleType::Type TitleType) const
@@ -208,14 +213,7 @@ void URigVMEdGraphNode::ReconstructNode_Internal(bool bForce)
 	}
 
 #if WITH_EDITOR
-	bEnableProfiling = false;
-	if(RigGraph)
-	{
-		if(URigVMBlueprint* RigBlueprint = RigGraph->GetBlueprint())
-		{
-			bEnableProfiling = RigBlueprint->VMRuntimeSettings.bEnableProfiling;
-		}
-	}
+	UpdateProfilingSettings();
 #endif
 
 	// Clear previously set messages
@@ -1132,18 +1130,19 @@ FLinearColor URigVMEdGraphNode::GetNodeProfilingColor() const
 			{
 				if(const URigVMNode* ModelNode = GetModelNode())
 				{
-					const double MicroSeconds = ModelNode->GetInstructionMicroSeconds(DebuggedHost->GetRigVMExtendedExecuteContext(), DebuggedHost->GetVM(), FRigVMASTProxy());
+					const double CurrentMicroSeconds = ModelNode->GetInstructionMicroSeconds(DebuggedHost->GetRigVMExtendedExecuteContext(), DebuggedHost->GetVM(), FRigVMASTProxy());
+					MicroSeconds = Blueprint->RigGraphDisplaySettings.AggregateAverage(MicroSecondsFrames, MicroSeconds, CurrentMicroSeconds);
 					if(MicroSeconds >= 0.0)
 					{
 						if(Blueprint->RigGraphDisplaySettings.bAutoDetermineRange)
 						{
-							if(MicroSeconds < Blueprint->RigGraphDisplaySettings.MinMicroSeconds)
+							if(CurrentMicroSeconds < Blueprint->RigGraphDisplaySettings.MinMicroSeconds)
 							{
-								Blueprint->RigGraphDisplaySettings.MinMicroSeconds = MicroSeconds;
+								Blueprint->RigGraphDisplaySettings.MinMicroSeconds = CurrentMicroSeconds;
 							}
-							if(MicroSeconds > Blueprint->RigGraphDisplaySettings.MaxMicroSeconds)
+							if(CurrentMicroSeconds > Blueprint->RigGraphDisplaySettings.MaxMicroSeconds)
 							{
-								Blueprint->RigGraphDisplaySettings.MaxMicroSeconds = MicroSeconds;
+								Blueprint->RigGraphDisplaySettings.MaxMicroSeconds = CurrentMicroSeconds;
 							}
 						}
 							
@@ -1164,6 +1163,11 @@ FLinearColor URigVMEdGraphNode::GetNodeProfilingColor() const
 			}
 		}
 	}
+	else
+	{
+		MicroSeconds = 0;
+		MicroSecondsFrames.Reset();
+	}
 #endif
 	return FLinearColor::Black;
 }
@@ -1811,5 +1815,22 @@ TArray<URigVMPin*>& URigVMEdGraphNode::PinListForPin(const URigVMPin* InModelPin
 	return EmptyList;
 }
 
+#if WITH_EDITOR
+void URigVMEdGraphNode::UpdateProfilingSettings()
+{
+	bEnableProfiling = false;
+	MicroSeconds = 0;
+	MicroSecondsFrames.Reset();
+
+	if(const URigVMEdGraph* RigGraph = Cast<URigVMEdGraph>(GetOuter()))
+	{
+		if(const URigVMBlueprint* RigBlueprint = RigGraph->GetBlueprint())
+		{
+			bEnableProfiling = RigBlueprint->VMRuntimeSettings.bEnableProfiling;
+		}
+	}
+}
+#endif
+
 #undef LOCTEXT_NAMESPACE
 
diff --git a/Engine/Plugins/Runtime/RigVM/Source/RigVMDeveloper/Private/RigVMBlueprint.cpp b/Engine/Plugins/Runtime/RigVM/Source/RigVMDeveloper/Private/RigVMBlueprint.cpp
--- a/Engine/Plugins/Runtime/RigVM/Source/RigVMDeveloper/Private/RigVMBlueprint.cpp
+++ b/Engine/Plugins/Runtime/RigVM/Source/RigVMDeveloper/Private/RigVMBlueprint.cpp
@@ -66,6 +66,50 @@ static TArray<UClass*> GetClassObjectsInPackage(UPackage* InPackage)
 	return ClassObjects;
 }
 
+void FRigVMEdGraphDisplaySettings::SetTotalMicroSeconds(double InTotalMicroSeconds)
+{
+	TotalMicroSeconds = AggregateAverage(TotalMicroSecondsFrames, TotalMicroSeconds, InTotalMicroSeconds);
+}
+
+void FRigVMEdGraphDisplaySettings::SetLastMinMicroSeconds(double InMinMicroSeconds)
+{
+	LastMinMicroSeconds = AggregateAverage(MinMicroSecondsFrames, LastMinMicroSeconds, InMinMicroSeconds);
+}
+
+void FRigVMEdGraphDisplaySettings::SetLastMaxMicroSeconds(double InMaxMicroSeconds)
+{
+	LastMaxMicroSeconds = AggregateAverage(MaxMicroSecondsFrames, LastMaxMicroSeconds, InMaxMicroSeconds);
+}
+
+double FRigVMEdGraphDisplaySettings::AggregateAverage(TArray<double>& InFrames, double InPrevious, double InNext) const
+{
+	const int32 NbFrames = FMath::Min(AverageFrames, 256);
+	if(NbFrames < 2)
+	{
+		InFrames.Reset();
+		return InNext;
+	}
+	
+	InFrames.Add(InNext);
+	if(InFrames.Num() >= NbFrames)
+	{
+		double Average = 0;
+		for(const double Value : InFrames)
+		{
+			Average += Value;
+		}
+		Average /= double(NbFrames);
+		InFrames.Reset();
+		return Average;
+	}
+
+	if(InPrevious == DBL_MAX || InPrevious < -SMALL_NUMBER)
+	{
+		return InNext;
+	}
+	return InPrevious;
+}
+
 FEdGraphPinType FRigVMOldPublicFunctionArg::GetPinType() const
 {
 	FRigVMExternalVariable Variable;
@@ -1181,8 +1225,18 @@ void URigVMBlueprint::RecompileVM()
 	
 	bErrorsDuringCompilation = false;
 
-	RigGraphDisplaySettings.MinMicroSeconds = RigGraphDisplaySettings.LastMinMicroSeconds = DBL_MAX;
-	RigGraphDisplaySettings.MaxMicroSeconds = RigGraphDisplaySettings.LastMaxMicroSeconds = (double)INDEX_NONE;
+	if(RigGraphDisplaySettings.bAutoDetermineRange)
+	{
+		RigGraphDisplaySettings.MinMicroSeconds = RigGraphDisplaySettings.LastMinMicroSeconds = DBL_MAX;
+		RigGraphDisplaySettings.MaxMicroSeconds = RigGraphDisplaySettings.LastMaxMicroSeconds = (double)INDEX_NONE;
+	}
+	else if(RigGraphDisplaySettings.MaxMicroSeconds < RigGraphDisplaySettings.MinMicroSeconds)
+	{
+		RigGraphDisplaySettings.MinMicroSeconds = 0;
+		RigGraphDisplaySettings.MaxMicroSeconds = 5;
+	}
+	
+	RigGraphDisplaySettings.TotalMicroSeconds = 0.0;
 	RigGraphDisplaySettings.MinMicroSecondsFrames.Reset();
 	RigGraphDisplaySettings.MaxMicroSecondsFrames.Reset();
 	RigGraphDisplaySettings.TotalMicroSecondsFrames.Reset();
diff --git a/Engine/Plugins/Runtime/RigVM/Source/RigVMDeveloper/Private/RigVMCompiler/RigVMCompiler.cpp b/Engine/Plugins/Runtime/RigVM/Source/RigVMDeveloper/Private/RigVMCompiler/RigVMCompiler.cpp
--- a/Engine/Plugins/Runtime/RigVM/Source/RigVMDeveloper/Private/RigVMCompiler/RigVMCompiler.cpp
+++ b/Engine/Plugins/Runtime/RigVM/Source/RigVMDeveloper/Private/RigVMCompiler/RigVMCompiler.cpp
@@ -2499,15 +2499,7 @@ bool URigVMCompiler::TraverseInlineFunction(const FRigVMInlineFunctionExprAST* I
 
 			if (WorkData.Settings.SetupNodeInstructionIndex)
 			{
-				if (const TArray<TWeakObjectPtr<UObject>>* Callstack = FunctionByteCode.GetCallstackForInstruction(i-InstructionIndexStart))
-				{
-					if (Callstack->Num() > 1)
-					{
-						FRigVMCallstack InstructionCallstack = FuncRefCallstack;
-						InstructionCallstack.Stack.Append(&(*Callstack)[1], Callstack->Num()-1);
-						WorkData.VM->GetByteCode().SetSubject(i, InstructionCallstack.GetCallPath(), InstructionCallstack.GetStack());
-					}
-				}
+				WorkData.VM->GetByteCode().SetSubject(i, FuncRefCallstack.GetCallPath(), FuncRefCallstack.GetStack());
 			}
 		}
 
diff --git a/Engine/Plugins/Runtime/RigVM/Source/RigVMDeveloper/Private/RigVMModel/Nodes/RigVMFunctionReferenceNode.cpp b/Engine/Plugins/Runtime/RigVM/Source/RigVMDeveloper/Private/RigVMModel/Nodes/RigVMFunctionReferenceNode.cpp
--- a/Engine/Plugins/Runtime/RigVM/Source/RigVMDeveloper/Private/RigVMModel/Nodes/RigVMFunctionReferenceNode.cpp
+++ b/Engine/Plugins/Runtime/RigVM/Source/RigVMDeveloper/Private/RigVMModel/Nodes/RigVMFunctionReferenceNode.cpp
@@ -221,3 +221,16 @@ URigVMLibraryNode* URigVMFunctionReferenceNode::LoadReferencedNode() const
 	return Cast<URigVMLibraryNode>(LibraryNode);
 	
 }
+
+TArray<int32> URigVMFunctionReferenceNode::GetInstructionsForVMImpl(const FRigVMExtendedExecuteContext& Context, URigVM* InVM, const FRigVMASTProxy& InProxy) const
+{
+	TArray<int32> Instructions = URigVMNode::GetInstructionsForVMImpl(Context, InVM, InProxy);
+
+	// if the base cannot find any matching instructions, fall back to the library node's implementation
+	if(Instructions.IsEmpty())
+	{
+		Instructions = Super::GetInstructionsForVMImpl(Context, InVM, InProxy);
+	}
+	
+	return Instructions;
+}
\ No newline at end of file
diff --git a/Engine/Plugins/Runtime/RigVM/Source/RigVMDeveloper/Public/EdGraph/RigVMEdGraphNode.h b/Engine/Plugins/Runtime/RigVM/Source/RigVMDeveloper/Public/EdGraph/RigVMEdGraphNode.h
--- a/Engine/Plugins/Runtime/RigVM/Source/RigVMDeveloper/Public/EdGraph/RigVMEdGraphNode.h
+++ b/Engine/Plugins/Runtime/RigVM/Source/RigVMDeveloper/Public/EdGraph/RigVMEdGraphNode.h
@@ -254,6 +254,10 @@ private:
 	static void ConfigurePin(UEdGraphPin* EdGraphPin, const URigVMPin* ModelPin);
 	TArray<URigVMPin*>& PinListForPin(const URigVMPin* InModelPin);
 
+#if WITH_EDITOR
+	void UpdateProfilingSettings();
+#endif
+	
 	FLinearColor CachedTitleColor;
 	FLinearColor CachedNodeColor;
 
@@ -279,6 +283,8 @@ private:
 
 	mutable const FRigVMTemplate* CachedTemplate;
 	mutable TOptional<bool> DrawAsCompactNodeCache;
+	mutable double MicroSeconds;
+	mutable TArray<double> MicroSecondsFrames;
 
 	friend class SRigVMGraphNode;
 	friend class FRigVMFunctionArgumentLayout;
diff --git a/Engine/Plugins/Runtime/RigVM/Source/RigVMDeveloper/Public/RigVMBlueprint.h b/Engine/Plugins/Runtime/RigVM/Source/RigVMDeveloper/Public/RigVMBlueprint.h
--- a/Engine/Plugins/Runtime/RigVM/Source/RigVMDeveloper/Public/RigVMBlueprint.h
+++ b/Engine/Plugins/Runtime/RigVM/Source/RigVMDeveloper/Public/RigVMBlueprint.h
@@ -67,7 +67,7 @@ struct RIGVMDEVELOPER_API FRigVMEdGraphDisplaySettings
 		, MinMicroSeconds(0.0)
 		, MaxMicroSeconds(1.0)
 		, TotalMicroSeconds(0.0)
-		, AverageFrames(16)
+		, AverageFrames(64)
 		, bAutoDetermineRange(true)
 		, LastMinMicroSeconds(0.0)
 		, LastMaxMicroSeconds(1.0)
@@ -116,7 +116,7 @@ struct RIGVMDEVELOPER_API FRigVMEdGraphDisplaySettings
 	double TotalMicroSeconds;
 
 	// If you set this to more than 1 the results will be averaged across multiple frames
-	UPROPERTY(EditAnywhere, Category = "Graph Display Settings")
+	UPROPERTY(EditAnywhere, Category = "Graph Display Settings", meta = (UIMin=1, UIMax=256))
 	int32 AverageFrames;
 
 	TArray<double> MinMicroSecondsFrames;
@@ -139,6 +139,11 @@ struct RIGVMDEVELOPER_API FRigVMEdGraphDisplaySettings
 	// The color of the slowest instruction / node
 	UPROPERTY(EditAnywhere, Category = "Graph Display Settings")
 	FLinearColor MaxDurationColor;
+
+	void SetTotalMicroSeconds(double InTotalMicroSeconds);
+	void SetLastMinMicroSeconds(double InMinMicroSeconds);
+	void SetLastMaxMicroSeconds(double InMaxMicroSeconds);
+	double AggregateAverage(TArray<double>& InFrames, double InPrevious, double InNext) const;
 };
 
 enum class ERigVMBlueprintLoadType : uint8
diff --git a/Engine/Plugins/Runtime/RigVM/Source/RigVMDeveloper/Public/RigVMModel/Nodes/RigVMFunctionReferenceNode.h b/Engine/Plugins/Runtime/RigVM/Source/RigVMDeveloper/Public/RigVMModel/Nodes/RigVMFunctionReferenceNode.h
--- a/Engine/Plugins/Runtime/RigVM/Source/RigVMDeveloper/Public/RigVMModel/Nodes/RigVMFunctionReferenceNode.h
+++ b/Engine/Plugins/Runtime/RigVM/Source/RigVMDeveloper/Public/RigVMModel/Nodes/RigVMFunctionReferenceNode.h
@@ -60,6 +60,7 @@ private:
 
 	virtual FText GetToolTipTextForPin(const URigVMPin* InPin) const override;
 	bool RequiresVariableRemappingInternal(TArray<FRigVMExternalVariable>& InnerVariables) const;
+	virtual TArray<int32> GetInstructionsForVMImpl(const FRigVMExtendedExecuteContext& Context, URigVM* InVM, const FRigVMASTProxy& InProxy = FRigVMASTProxy()) const override; 
 
 	//void SetReferencedFunctionData(FRigVMGraphFunctionData* Data);
 
diff --git a/Engine/Plugins/Runtime/RigVM/Source/RigVMEditor/Private/Editor/RigVMEditor.cpp b/Engine/Plugins/Runtime/RigVM/Source/RigVMEditor/Private/Editor/RigVMEditor.cpp
--- a/Engine/Plugins/Runtime/RigVM/Source/RigVMEditor/Private/Editor/RigVMEditor.cpp
+++ b/Engine/Plugins/Runtime/RigVM/Source/RigVMEditor/Private/Editor/RigVMEditor.cpp
@@ -1813,21 +1813,20 @@ void FRigVMEditor::HandleVMExecutedEvent(URigVMHost* InHost, const FName& InEven
 		{
 			if(DebuggedHost)
 			{
-				RigVMBlueprint->RigGraphDisplaySettings.TotalMicroSeconds = DebuggedHost->GetProfilingInfo().GetLastExecutionMicroSeconds();
+				RigVMBlueprint->RigGraphDisplaySettings.SetTotalMicroSeconds(DebuggedHost->GetProfilingInfo().GetLastExecutionMicroSeconds());
 			}
 
 			if(RigVMBlueprint->RigGraphDisplaySettings.bAutoDetermineRange)
 			{
 				if(RigVMBlueprint->RigGraphDisplaySettings.LastMaxMicroSeconds < 0.0)
 				{
-					RigVMBlueprint->RigGraphDisplaySettings.LastMinMicroSeconds = RigVMBlueprint->RigGraphDisplaySettings.MinMicroSeconds; 
-					RigVMBlueprint->RigGraphDisplaySettings.LastMaxMicroSeconds = RigVMBlueprint->RigGraphDisplaySettings.MaxMicroSeconds;
+					RigVMBlueprint->RigGraphDisplaySettings.SetLastMinMicroSeconds(RigVMBlueprint->RigGraphDisplaySettings.MinMicroSeconds); 
+					RigVMBlueprint->RigGraphDisplaySettings.SetLastMaxMicroSeconds(RigVMBlueprint->RigGraphDisplaySettings.MaxMicroSeconds);
 				}
 				else if(RigVMBlueprint->RigGraphDisplaySettings.MaxMicroSeconds >= 0.0)
 				{
-					const double T = 0.05;
-					RigVMBlueprint->RigGraphDisplaySettings.LastMinMicroSeconds = FMath::Lerp<double>(RigVMBlueprint->RigGraphDisplaySettings.LastMinMicroSeconds, RigVMBlueprint->RigGraphDisplaySettings.MinMicroSeconds, T); 
-					RigVMBlueprint->RigGraphDisplaySettings.LastMaxMicroSeconds = FMath::Lerp<double>(RigVMBlueprint->RigGraphDisplaySettings.LastMaxMicroSeconds, RigVMBlueprint->RigGraphDisplaySettings.MaxMicroSeconds, T); 
+					RigVMBlueprint->RigGraphDisplaySettings.SetLastMinMicroSeconds(RigVMBlueprint->RigGraphDisplaySettings.MinMicroSeconds); 
+					RigVMBlueprint->RigGraphDisplaySettings.SetLastMaxMicroSeconds(RigVMBlueprint->RigGraphDisplaySettings.MaxMicroSeconds); 
 				}
 
 				RigVMBlueprint->RigGraphDisplaySettings.MinMicroSeconds = DBL_MAX; 
@@ -1835,8 +1834,8 @@ void FRigVMEditor::HandleVMExecutedEvent(URigVMHost* InHost, const FName& InEven
 			}
 			else
 			{
-				RigVMBlueprint->RigGraphDisplaySettings.LastMinMicroSeconds = RigVMBlueprint->RigGraphDisplaySettings.MinMicroSeconds; 
-				RigVMBlueprint->RigGraphDisplaySettings.LastMaxMicroSeconds = RigVMBlueprint->RigGraphDisplaySettings.MaxMicroSeconds;
+				RigVMBlueprint->RigGraphDisplaySettings.SetLastMinMicroSeconds(RigVMBlueprint->RigGraphDisplaySettings.MinMicroSeconds); 
+				RigVMBlueprint->RigGraphDisplaySettings.SetLastMaxMicroSeconds(RigVMBlueprint->RigGraphDisplaySettings.MaxMicroSeconds);
 			}
 		}
 	}
diff --git a/Engine/Plugins/Runtime/RigVM/Source/RigVMEditor/Private/Widgets/SRigVMExecutionStackView.cpp b/Engine/Plugins/Runtime/RigVM/Source/RigVMEditor/Private/Widgets/SRigVMExecutionStackView.cpp
--- a/Engine/Plugins/Runtime/RigVM/Source/RigVMEditor/Private/Widgets/SRigVMExecutionStackView.cpp
+++ b/Engine/Plugins/Runtime/RigVM/Source/RigVMEditor/Private/Widgets/SRigVMExecutionStackView.cpp
@@ -50,6 +50,8 @@ void SRigStackItem::Construct(const FArguments& InArgs, const TSharedRef<STableV
 	WeakStackEntry = InStackEntry;
 	WeakBlueprint = InBlueprint;
 	WeakCommandList = InCommandList;
+	MicroSeconds = 0;
+	MicroSecondsFrames.Reset();
 
 	TSharedPtr< STextBlock > NumberWidget;
 	TSharedPtr< STextBlock > TextWidget;
@@ -225,10 +227,11 @@ FText SRigStackItem::GetDurationText() const
 				{
 					if(URigVM* VM = RigVMHost->GetVM())
 					{
-						const double MicroSeconds = VM->GetInstructionMicroSeconds(RigVMHost->GetRigVMExtendedExecuteContext(), WeakStackEntry.Pin()->InstructionIndex);
+						const double CurrentMicroSeconds = VM->GetInstructionMicroSeconds(RigVMHost->GetRigVMExtendedExecuteContext(), WeakStackEntry.Pin()->InstructionIndex);
+						MicroSeconds = WeakBlueprint->RigGraphDisplaySettings.AggregateAverage(MicroSecondsFrames, MicroSeconds, CurrentMicroSeconds);
 						if(MicroSeconds > 0.0)
 						{
-							return FText::FromString(FString::Printf(TEXT("%d µs"), (int32)MicroSeconds));
+							return FText::FromString(FString::Printf(TEXT("%.02f µs"), (float)MicroSeconds));
 						}
 					}
 				}
@@ -967,7 +970,12 @@ void SRigVMExecutionStackView::HandleModifiedEvent(ERigVMGraphNotifType InNotifT
 					}
 				}
 			}
-			TreeView->SetItemSelection(SelectedItems, true, ESelectInfo::Direct);
+
+			if(!SelectedItems.IsEmpty())
+			{
+				TreeView->SetItemSelection(SelectedItems, true, ESelectInfo::Direct);
+				TreeView->RequestScrollIntoView(SelectedItems[0]);
+			}
 			break;
 		}
 		default:
diff --git a/Engine/Plugins/Runtime/RigVM/Source/RigVMEditor/Private/Widgets/SRigVMGraphNode.cpp b/Engine/Plugins/Runtime/RigVM/Source/RigVMEditor/Private/Widgets/SRigVMGraphNode.cpp
--- a/Engine/Plugins/Runtime/RigVM/Source/RigVMEditor/Private/Widgets/SRigVMGraphNode.cpp
+++ b/Engine/Plugins/Runtime/RigVM/Source/RigVMEditor/Private/Widgets/SRigVMGraphNode.cpp
@@ -103,7 +103,7 @@ void SRigVMGraphNode::Construct( const FArguments& InArgs )
 	.ColorAndOpacity(FLinearColor::White)
 	.ShadowColorAndOpacity(FLinearColor(0.1f, 0.1f, 0.1f, 1.f))
 	.Visibility(EVisibility::Visible)
-	.ToolTipText(LOCTEXT("NodeHitCountToolTip", "This number represents the hit count for a node.\nFor functions / collapse nodes it represents the sum of all hit counts of contained nodes.\n\nYou can enable / disable the display of the number in the Class Settings\n(Rig Graph Display Settings -> Show Node Run Counts)"));
+	.ToolTipText(LOCTEXT("NodeHitCountToolTip", "This number represents the number of instructions hit for a node.\nIf the node has auxiliary instructions (such as Copies) you'll first see the primary instructions followed by the overall instructions in braces.\nFor functions / collapse nodes it represents the sum of all hit instructions of contained nodes.\n\nYou can enable / disable the display of the number in the Class Settings\n(Rig Graph Display Settings -> Show Node Run Counts)"));
 
 	SAssignNew(InstructionDurationTextBlockWidget, STextBlock)
 	.Margin(FMargin(2.0f, 2.0f, 2.0f, 1.0f))
@@ -112,7 +112,7 @@ void SRigVMGraphNode::Construct( const FArguments& InArgs )
 	.ColorAndOpacity(FLinearColor::White)
 	.ShadowColorAndOpacity(FLinearColor(0.1f, 0.1f, 0.1f, 1.f))
 	.Visibility(EVisibility::Visible)
-	.ToolTipText(LOCTEXT("NodeDurationToolTip", "This number represents the duration in microseconds for a node.\nFor functions / collapse nodes it represents the accumulated time of contained nodes.\n\nYou can enable / disable the display of the number in the Class Settings\n(VM Runtime Settings -> Enable Profiling)"));
+	.ToolTipText(LOCTEXT("NodeDurationToolTip", "This number represents the duration in microseconds for a node.\nFor functions / collapse nodes it represents the accumulated time of contained nodes.\n\nIf you have more than one node selected you'll see also the overall summed up time of the selection.\n\nYou can enable / disable the display of the number in the Class Settings\n(VM Runtime Settings -> Enable Profiling)"));
 
 	EdGraphNode->OnNodeTitleDirtied().AddSP(this, &SRigVMGraphNode::HandleNodeTitleDirtied);
 	EdGraphNode->OnNodePinsChanged().AddSP(this, &SRigVMGraphNode::HandleNodePinsChanged);
@@ -1263,11 +1263,14 @@ FText SRigVMGraphNode::GetInstructionCountText() const
 					{
 						RunCount = ModelNode->GetInstructionVisitedCount(DebuggedHost->GetRigVMExtendedExecuteContext(), DebuggedHost->GetVM(), FRigVMASTProxy());
 						bShowNodeRunCount = RunCount > Blueprint->RigGraphDisplaySettings.NodeRunLowerBound;
+
+						// toodoo here we want to differentiate between primary and secondary instructions. (call extern vs copy).
+						// const TArray<int32>& Instructions = ModelNode->GetInstructionsForVM(DebuggedHost->GetRigVMExtendedExecuteContext(), DebuggedHost->GetVM());
 					}
 
 					if(bShowInstructionIndex)
 					{
-						const TArray<int32> Instructions = ModelNode->GetInstructionsForVM(DebuggedHost->GetRigVMExtendedExecuteContext(), DebuggedHost->GetVM());
+						const TArray<int32>& Instructions = ModelNode->GetInstructionsForVM(DebuggedHost->GetRigVMExtendedExecuteContext(), DebuggedHost->GetVM());
 						bShowInstructionIndex = Instructions.Num() > 0;
 						if(bShowInstructionIndex)
 						{
@@ -1312,15 +1315,30 @@ FText SRigVMGraphNode::GetInstructionDurationText() const
 	{
 		if(Blueprint->VMRuntimeSettings.bEnableProfiling)
 		{
-			if (ModelNode.IsValid())
+			if(const URigVMEdGraphNode* RigGraphNode = Cast<URigVMEdGraphNode>(GraphNode))
 			{
-				if(URigVMHost* DebuggedHost = Cast<URigVMHost>(Blueprint->GetObjectBeingDebugged()))
+				const double MicroSeconds = RigGraphNode->MicroSeconds;
+				if(MicroSeconds >= 0)
 				{
-					const double MicroSeconds = ModelNode->GetInstructionMicroSeconds(DebuggedHost->GetRigVMExtendedExecuteContext(), DebuggedHost->GetVM(), FRigVMASTProxy());
-					if(MicroSeconds >= 0)
+					if (const SGraphPanel* MyOwnerPanel = GetOwnerPanel().Get())
 					{
-						return FText::FromString(FString::Printf(TEXT("%d µs"), (int32)MicroSeconds));
+						const TArray<UEdGraphNode*> SelectedNodes = MyOwnerPanel->GetSelectedGraphNodes();
+						if((SelectedNodes.Num() > 1) && SelectedNodes.Contains(RigGraphNode))
+						{
+							double OverallMicroSeconds = 0;
+							for(const UEdGraphNode* SelectedNode : SelectedNodes)
+							{
+								if(const URigVMEdGraphNode* SelectedRigGraphNode = Cast<URigVMEdGraphNode>(SelectedNode))
+								{
+									OverallMicroSeconds += SelectedRigGraphNode->MicroSeconds;
+								}
+							}
+							
+							return FText::FromString(FString::Printf(TEXT("%.02f µs of %.02f µs"), (float)MicroSeconds, (float)OverallMicroSeconds));
+						}
 					}
+
+					return FText::FromString(FString::Printf(TEXT("%.02f µs"), (float)MicroSeconds));
 				}
 			}
 		}
diff --git a/Engine/Plugins/Runtime/RigVM/Source/RigVMEditor/Public/Widgets/SRigVMExecutionStackView.h b/Engine/Plugins/Runtime/RigVM/Source/RigVMEditor/Public/Widgets/SRigVMExecutionStackView.h
--- a/Engine/Plugins/Runtime/RigVM/Source/RigVMEditor/Public/Widgets/SRigVMExecutionStackView.h
+++ b/Engine/Plugins/Runtime/RigVM/Source/RigVMEditor/Public/Widgets/SRigVMExecutionStackView.h
@@ -56,6 +56,9 @@ private:
 	TWeakObjectPtr<URigVMBlueprint> WeakBlueprint;
 	TWeakPtr<FUICommandList> WeakCommandList;
 
+	mutable double MicroSeconds;
+	mutable TArray<double> MicroSecondsFrames;
+
 	FText GetIndexText() const;
 	FText GetLabelText() const;
 	FSlateFontInfo GetLabelFont() const;
