RigVM: Localize Function to affect all graphs

#rnx

diff --git a/Engine/Plugins/Runtime/RigVM/Source/RigVMDeveloper/Private/RigVMModel/RigVMController.cpp b/Engine/Plugins/Runtime/RigVM/Source/RigVMDeveloper/Private/RigVMModel/RigVMController.cpp
--- a/Engine/Plugins/Runtime/RigVM/Source/RigVMDeveloper/Private/RigVMModel/RigVMController.cpp
+++ b/Engine/Plugins/Runtime/RigVM/Source/RigVMDeveloper/Private/RigVMModel/RigVMController.cpp
@@ -4536,29 +4536,30 @@ TMap<FRigVMGraphFunctionIdentifier, URigVMLibraryNode*> URigVMController::Locali
 
 	// once we have all local functions available, clean up the references
 	TArray<URigVMGraph*> GraphsToUpdate;
-	GraphsToUpdate.AddUnique(Graph);
-	if(URigVMFunctionLibrary* DefaultFunctionLibrary = Graph->GetDefaultFunctionLibrary())
+	if(IRigVMClientHost* ClientHost = GetImplementingOuter<IRigVMClientHost>())
 	{
-		GraphsToUpdate.AddUnique(DefaultFunctionLibrary);
+		if(FRigVMClient* Client = ClientHost->GetRigVMClient())
+		{
+			GraphsToUpdate = Client->GetAllModels(true, true);
+		}
 	}
+
 	for(int32 GraphToUpdateIndex=0; GraphToUpdateIndex<GraphsToUpdate.Num(); GraphToUpdateIndex++)
 	{
 		URigVMGraph* GraphToUpdate = GraphsToUpdate[GraphToUpdateIndex];
-		
-		const TArray<URigVMNode*> NodesToUpdate = GraphToUpdate->GetNodes();
-		for(URigVMNode* NodeToUpdate : NodesToUpdate)
+		if(URigVMController* GraphController = GetControllerForGraph(GraphToUpdate))
 		{
-			if(URigVMCollapseNode* CollapseNode = Cast<URigVMCollapseNode>(NodeToUpdate))
-			{
-				GraphsToUpdate.AddUnique(CollapseNode->GetContainedGraph());
-			}
-			else if(URigVMFunctionReferenceNode* FunctionReferenceNode = Cast<URigVMFunctionReferenceNode>(NodeToUpdate))
+			const TArray<URigVMNode*> NodesToUpdate = GraphToUpdate->GetNodes();
+			for(URigVMNode* NodeToUpdate : NodesToUpdate)
 			{
-				URigVMLibraryNode** RemappedNodePtr = LocalizedFunctions.Find(FunctionReferenceNode->GetReferencedFunctionHeader().LibraryPointer);
-				if(RemappedNodePtr)
+				if(URigVMFunctionReferenceNode* FunctionReferenceNode = Cast<URigVMFunctionReferenceNode>(NodeToUpdate))
 				{
-					URigVMLibraryNode* RemappedNode = *RemappedNodePtr;
-					SetReferencedFunction(FunctionReferenceNode, RemappedNode, bSetupUndoRedo);
+					URigVMLibraryNode** RemappedNodePtr = LocalizedFunctions.Find(FunctionReferenceNode->GetReferencedFunctionHeader().LibraryPointer);
+					if(RemappedNodePtr)
+					{
+						URigVMLibraryNode* RemappedNode = *RemappedNodePtr;
+						GraphController->SetReferencedFunction(FunctionReferenceNode, RemappedNode, bSetupUndoRedo);
+					}
 				}
 			}
 		}
