RigVM: Fix debug unit tests failing due to profiler changes

#rnx

diff --git a/Engine/Plugins/Runtime/RigVM/Source/RigVM/Private/RigVMCore/RigVMByteCode.cpp b/Engine/Plugins/Runtime/RigVM/Source/RigVM/Private/RigVMCore/RigVMByteCode.cpp
--- a/Engine/Plugins/Runtime/RigVM/Source/RigVM/Private/RigVMCore/RigVMByteCode.cpp
+++ b/Engine/Plugins/Runtime/RigVM/Source/RigVM/Private/RigVMCore/RigVMByteCode.cpp
@@ -1855,6 +1855,13 @@ void FRigVMByteCode::SetSubject(int32 InInstructionIndex, const FString& InCallP
 	}
 }
 
+void FRigVMByteCode::AddInstructionForSubject(UObject* InSubject, int32 InInstructionIndex)
+{
+	check(InSubject);
+	TWeakObjectPtr<UObject> WeakSubject(InSubject);
+	SubjectToInstructions.FindOrAdd(WeakSubject).AddUnique(InInstructionIndex);
+}
+
 const TArray<TWeakObjectPtr<UObject>>* FRigVMByteCode::GetCallstackForInstruction(int32 InInstructionIndex) const
 {
 	if (CallstackPerInstruction.IsValidIndex(InInstructionIndex))
diff --git a/Engine/Plugins/Runtime/RigVM/Source/RigVM/Public/RigVMCore/RigVMByteCode.h b/Engine/Plugins/Runtime/RigVM/Source/RigVM/Public/RigVMCore/RigVMByteCode.h
--- a/Engine/Plugins/Runtime/RigVM/Source/RigVM/Public/RigVMCore/RigVMByteCode.h
+++ b/Engine/Plugins/Runtime/RigVM/Source/RigVM/Public/RigVMCore/RigVMByteCode.h
@@ -1393,6 +1393,7 @@ private:
 #if WITH_EDITOR
 
 	void SetSubject(int32 InInstructionIndex, const FString& InCallPath, const TArray<TWeakObjectPtr<UObject>>& InCallstack);
+	void AddInstructionForSubject(UObject* InSubject, int32 InInstructionIndex);
 
 #endif
 
diff --git a/Engine/Plugins/Runtime/RigVM/Source/RigVMDeveloper/Private/RigVMCompiler/RigVMCompiler.cpp b/Engine/Plugins/Runtime/RigVM/Source/RigVMDeveloper/Private/RigVMCompiler/RigVMCompiler.cpp
--- a/Engine/Plugins/Runtime/RigVM/Source/RigVMDeveloper/Private/RigVMCompiler/RigVMCompiler.cpp
+++ b/Engine/Plugins/Runtime/RigVM/Source/RigVMDeveloper/Private/RigVMCompiler/RigVMCompiler.cpp
@@ -2499,7 +2499,18 @@ bool URigVMCompiler::TraverseInlineFunction(const FRigVMInlineFunctionExprAST* I
 
 			if (WorkData.Settings.SetupNodeInstructionIndex)
 			{
-				WorkData.VM->GetByteCode().SetSubject(i, FuncRefCallstack.GetCallPath(), FuncRefCallstack.GetStack());
+				if (const TArray<TWeakObjectPtr<UObject>>* Callstack = FunctionByteCode.GetCallstackForInstruction(i-InstructionIndexStart))
+				{
+					if (Callstack->Num() > 1)
+					{
+						FRigVMCallstack InstructionCallstack = FuncRefCallstack;
+						InstructionCallstack.Stack.Append(&(*Callstack)[1], Callstack->Num()-1);
+						WorkData.VM->GetByteCode().SetSubject(i, InstructionCallstack.GetCallPath(), InstructionCallstack.GetStack());
+					}
+				}
+
+				// also store the instruction for the subject so that profiling can determine the cost of functions.
+				WorkData.VM->GetByteCode().AddInstructionForSubject(FunctionReferenceNode, i);
 			}
 		}
 
