RigVM - Templates memory optimization (part #1)
    - removed useless FRigVMTemplate delegates
    - factories now return FRigVMTemplateInfos to avoid storing types twice. they are used to store argument name and direction + a callback to create the real argument when needed.

diff --git a/Engine/Plugins/Animation/ControlRig/Source/ControlRig/Private/Units/Animation/RigUnit_AnimAttribute.cpp b/Engine/Plugins/Animation/ControlRig/Source/ControlRig/Private/Units/Animation/RigUnit_AnimAttribute.cpp
--- a/Engine/Plugins/Animation/ControlRig/Source/ControlRig/Private/Units/Animation/RigUnit_AnimAttribute.cpp
+++ b/Engine/Plugins/Animation/ControlRig/Source/ControlRig/Private/Units/Animation/RigUnit_AnimAttribute.cpp
@@ -321,18 +321,18 @@ FString FRigDispatch_AnimAttributeBase::GetNodeTitle(const FRigVMTemplateTypeMap
 }
 #endif
 
-const TArray<FRigVMTemplateArgument>& FRigDispatch_AnimAttributeBase::GetArguments() const
+const TArray<FRigVMTemplateArgumentInfo>& FRigDispatch_AnimAttributeBase::GetArgumentInfos() const
 {
-	if (Arguments.IsEmpty())
+	if (Infos.IsEmpty())
 	{
-		NameArgIndex = Arguments.Emplace(NameArgName, ERigVMPinDirection::Input, RigVMTypeUtils::TypeIndex::FName);
-		BoneNameArgIndex = Arguments.Emplace(BoneNameArgName, ERigVMPinDirection::Input, RigVMTypeUtils::TypeIndex::FName);
+		NameArgIndex = Infos.Emplace(NameArgName, ERigVMPinDirection::Input, RigVMTypeUtils::TypeIndex::FName);
+		BoneNameArgIndex = Infos.Emplace(BoneNameArgName, ERigVMPinDirection::Input, RigVMTypeUtils::TypeIndex::FName);
 
-		CachedBoneNameArgIndex = Arguments.Emplace(CachedBoneNameArgName, ERigVMPinDirection::Hidden, RigVMTypeUtils::TypeIndex::FName);
-		CachedBoneIndexArgIndex = Arguments.Emplace(CachedBoneIndexArgName, ERigVMPinDirection::Hidden, RigVMTypeUtils::TypeIndex::Int32);
+		CachedBoneNameArgIndex = Infos.Emplace(CachedBoneNameArgName, ERigVMPinDirection::Hidden, RigVMTypeUtils::TypeIndex::FName);
+		CachedBoneIndexArgIndex = Infos.Emplace(CachedBoneIndexArgName, ERigVMPinDirection::Hidden, RigVMTypeUtils::TypeIndex::Int32);
 	}
 	
-	return Arguments;
+	return Infos;
 }
 
 #if WITH_EDITOR
@@ -372,23 +372,23 @@ FText FRigDispatch_AnimAttributeBase::GetArgumentTooltip(const FName& InArgument
 
 
 
-const TArray<FRigVMTemplateArgument>& FRigDispatch_GetAnimAttribute::GetArguments() const
+const TArray<FRigVMTemplateArgumentInfo>& FRigDispatch_GetAnimAttribute::GetArgumentInfos() const
 {
 	if (ValueArgIndex == INDEX_NONE)
 	{
-		Arguments = Super::GetArguments(); 
+		Infos = Super::GetArgumentInfos(); 
 
 		FRigVMTemplateArgument::FTypeFilter	TypeFilter;
 		TypeFilter.BindStatic(&FRigDispatch_AnimAttributeBase::IsTypeSupported);
+		const TArray<TRigVMTypeIndex> Types = FRigVMTemplateArgumentInfo::GetTypesFromCategories(GetValueTypeCategory(), TypeFilter);
 		
-		DefaultArgIndex = Arguments.Emplace(DefaultArgName, ERigVMPinDirection::Input, GetValueTypeCategory(), TypeFilter);
-		ValueArgIndex = Arguments.Emplace(ValueArgName, ERigVMPinDirection::Output, GetValueTypeCategory(), TypeFilter);
+		DefaultArgIndex = Infos.Emplace(DefaultArgName, ERigVMPinDirection::Input, Types);
+		ValueArgIndex = Infos.Emplace(ValueArgName, ERigVMPinDirection::Output, Types);
 		
-		FoundArgIndex = Arguments.Emplace(FoundArgName, ERigVMPinDirection::Output, RigVMTypeUtils::TypeIndex::Bool);
+		FoundArgIndex = Infos.Emplace(FoundArgName, ERigVMPinDirection::Output, RigVMTypeUtils::TypeIndex::Bool);
 	}
 
-	
-	return Arguments;
+	return Infos;
 }
 
 FRigVMTemplateTypeMap FRigDispatch_GetAnimAttribute::OnNewArgumentType(const FName& InArgumentName,
@@ -455,20 +455,21 @@ FRigVMFunctionPtr FRigDispatch_GetAnimAttribute::GetDispatchFunctionImpl(const F
 	return nullptr;
 }
 
-const TArray<FRigVMTemplateArgument>& FRigDispatch_SetAnimAttribute::GetArguments() const
+const TArray<FRigVMTemplateArgumentInfo>& FRigDispatch_SetAnimAttribute::GetArgumentInfos() const
 {
 	if (ValueArgIndex == INDEX_NONE)
 	{
-		Arguments = Super::GetArguments();
+		Infos = Super::GetArgumentInfos();
 
 		FRigVMTemplateArgument::FTypeFilter	TypeFilter;
 		TypeFilter.BindStatic(&FRigDispatch_AnimAttributeBase::IsTypeSupported);
+		const TArray<TRigVMTypeIndex> Types = FRigVMTemplateArgumentInfo::GetTypesFromCategories(GetValueTypeCategory(), TypeFilter);
 		
-		ValueArgIndex = Arguments.Emplace(ValueArgName, ERigVMPinDirection::Input, GetValueTypeCategory(), TypeFilter);
-		SuccessArgIndex = Arguments.Emplace(SuccessArgName, ERigVMPinDirection::Output, RigVMTypeUtils::TypeIndex::Bool);
+		ValueArgIndex = Infos.Emplace(ValueArgName, ERigVMPinDirection::Input, Types);
+		SuccessArgIndex = Infos.Emplace(SuccessArgName, ERigVMPinDirection::Output, RigVMTypeUtils::TypeIndex::Bool);
 	}
 	
-	return Arguments;
+	return Infos;
 }
 
 const TArray<FRigVMExecuteArgument>& FRigDispatch_SetAnimAttribute::GetExecuteArguments_Impl(const FRigVMDispatchContext& InContext) const
diff --git a/Engine/Plugins/Animation/ControlRig/Source/ControlRig/Private/Units/Animation/RigUnit_AnimAttribute.h b/Engine/Plugins/Animation/ControlRig/Source/ControlRig/Private/Units/Animation/RigUnit_AnimAttribute.h
--- a/Engine/Plugins/Animation/ControlRig/Source/ControlRig/Private/Units/Animation/RigUnit_AnimAttribute.h
+++ b/Engine/Plugins/Animation/ControlRig/Source/ControlRig/Private/Units/Animation/RigUnit_AnimAttribute.h
@@ -134,7 +134,7 @@ struct CONTROLRIG_API FRigDispatch_AnimAttributeBase : public FRigDispatchFactor
 	virtual FString GetNodeTitle(const FRigVMTemplateTypeMap& InTypes) const override;;
 #endif
 	
-	virtual const TArray<FRigVMTemplateArgument>& GetArguments() const override;
+	virtual const TArray<FRigVMTemplateArgumentInfo>& GetArgumentInfos() const override;
 	virtual bool IsSet() const { return false; }
 
 #if WITH_EDITOR
@@ -145,7 +145,7 @@ protected:
 	static bool IsTypeSupported(const TRigVMTypeIndex& InTypeIndex);
 	static const TArray<FRigVMTemplateArgument::ETypeCategory>& GetValueTypeCategory();
 	
-	mutable TArray<FRigVMTemplateArgument> Arguments;
+	mutable TArray<FRigVMTemplateArgumentInfo> Infos;
 
 	// input
 	mutable int32 NameArgIndex = INDEX_NONE;
@@ -187,7 +187,7 @@ struct CONTROLRIG_API FRigDispatch_GetAnimAttribute: public FRigDispatch_AnimAtt
 		FactoryScriptStruct = StaticStruct();
 	}
 
-	virtual const TArray<FRigVMTemplateArgument>& GetArguments() const override;
+	virtual const TArray<FRigVMTemplateArgumentInfo>& GetArgumentInfos() const override;
 	virtual FRigVMTemplateTypeMap OnNewArgumentType(const FName& InArgumentName, TRigVMTypeIndex InTypeIndex) const override;
 
 	
@@ -267,7 +267,7 @@ struct CONTROLRIG_API FRigDispatch_SetAnimAttribute: public FRigDispatch_AnimAtt
 	}
 
 	virtual bool IsSet() const override { return true; }
-	virtual const TArray<FRigVMTemplateArgument>& GetArguments() const override;
+	virtual const TArray<FRigVMTemplateArgumentInfo>& GetArgumentInfos() const override;
 	virtual const TArray<FRigVMExecuteArgument>& GetExecuteArguments_Impl(const FRigVMDispatchContext& InContext) const override;
 	virtual FRigVMTemplateTypeMap OnNewArgumentType(const FName& InArgumentName, TRigVMTypeIndex InTypeIndex) const override;
 	
diff --git a/Engine/Plugins/Animation/ControlRig/Source/ControlRig/Private/Units/Core/RigUnit_UserData.cpp b/Engine/Plugins/Animation/ControlRig/Source/ControlRig/Private/Units/Core/RigUnit_UserData.cpp
--- a/Engine/Plugins/Animation/ControlRig/Source/ControlRig/Private/Units/Core/RigUnit_UserData.cpp
+++ b/Engine/Plugins/Animation/ControlRig/Source/ControlRig/Private/Units/Core/RigUnit_UserData.cpp
@@ -28,20 +28,24 @@ FString FRigDispatch_GetUserData::GetArgumentMetaData(const FName& InArgumentNam
 
 #endif
 
-const TArray<FRigVMTemplateArgument>& FRigDispatch_GetUserData::GetArguments() const
+const TArray<FRigVMTemplateArgumentInfo>& FRigDispatch_GetUserData::GetArgumentInfos() const
 {
-	static const TArray<FRigVMTemplateArgument::ETypeCategory> ValueCategories = {
-		FRigVMTemplateArgument::ETypeCategory_SingleAnyValue,
-		FRigVMTemplateArgument::ETypeCategory_ArrayAnyValue
-	};
-	static const TArray<FRigVMTemplateArgument> Arguments = {
-		FRigVMTemplateArgument(ArgNameSpaceName, ERigVMPinDirection::Input, RigVMTypeUtils::TypeIndex::FString),
-		FRigVMTemplateArgument(ArgPathName, ERigVMPinDirection::Input, RigVMTypeUtils::TypeIndex::FString),
-		FRigVMTemplateArgument(ArgDefaultName, ERigVMPinDirection::Input, ValueCategories),
-		FRigVMTemplateArgument(ArgResultName, ERigVMPinDirection::Output, ValueCategories),
-		FRigVMTemplateArgument(ArgFoundName, ERigVMPinDirection::Output, RigVMTypeUtils::TypeIndex::Bool)
-	};
-	return Arguments;
+	static TArray<FRigVMTemplateArgumentInfo> Infos;
+	if (Infos.IsEmpty())
+	{
+		static const TArray<FRigVMTemplateArgument::ETypeCategory> ValueCategories = {
+			FRigVMTemplateArgument::ETypeCategory_SingleAnyValue,
+			FRigVMTemplateArgument::ETypeCategory_ArrayAnyValue
+		};
+		
+		Infos.Emplace(ArgNameSpaceName, ERigVMPinDirection::Input, RigVMTypeUtils::TypeIndex::FString);
+		Infos.Emplace(ArgPathName, ERigVMPinDirection::Input, RigVMTypeUtils::TypeIndex::FString);
+		Infos.Emplace(ArgDefaultName, ERigVMPinDirection::Input, ValueCategories);
+		Infos.Emplace(ArgResultName, ERigVMPinDirection::Output, ValueCategories);
+		Infos.Emplace(ArgFoundName, ERigVMPinDirection::Output, RigVMTypeUtils::TypeIndex::Bool);
+	}
+	
+	return Infos;
 }
 
 FRigVMTemplateTypeMap FRigDispatch_GetUserData::OnNewArgumentType(const FName& InArgumentName,
diff --git a/Engine/Plugins/Animation/ControlRig/Source/ControlRig/Private/Units/Core/RigUnit_UserData.h b/Engine/Plugins/Animation/ControlRig/Source/ControlRig/Private/Units/Core/RigUnit_UserData.h
--- a/Engine/Plugins/Animation/ControlRig/Source/ControlRig/Private/Units/Core/RigUnit_UserData.h
+++ b/Engine/Plugins/Animation/ControlRig/Source/ControlRig/Private/Units/Core/RigUnit_UserData.h
@@ -23,7 +23,7 @@ public:
 #if WITH_EDITOR
 	virtual FString GetArgumentMetaData(const FName& InArgumentName, const FName& InMetaDataKey) const override;
 #endif
-	virtual const TArray<FRigVMTemplateArgument>& GetArguments() const override;
+	virtual const TArray<FRigVMTemplateArgumentInfo>& GetArgumentInfos() const override;
 	virtual FRigVMTemplateTypeMap OnNewArgumentType(const FName& InArgumentName, TRigVMTypeIndex InTypeIndex) const override;
 
 protected:
diff --git a/Engine/Plugins/Animation/ControlRig/Source/ControlRig/Private/Units/Hierarchy/RigUnit_Metadata.cpp b/Engine/Plugins/Animation/ControlRig/Source/ControlRig/Private/Units/Hierarchy/RigUnit_Metadata.cpp
--- a/Engine/Plugins/Animation/ControlRig/Source/ControlRig/Private/Units/Hierarchy/RigUnit_Metadata.cpp
+++ b/Engine/Plugins/Animation/ControlRig/Source/ControlRig/Private/Units/Hierarchy/RigUnit_Metadata.cpp
@@ -66,16 +66,16 @@ FString FRigDispatch_MetadataBase::GetNodeTitle(const FRigVMTemplateTypeMap& InT
 
 #endif
 
-const TArray<FRigVMTemplateArgument>& FRigDispatch_MetadataBase::GetArguments() const
+const TArray<FRigVMTemplateArgumentInfo>& FRigDispatch_MetadataBase::GetArgumentInfos() const
 {
-	if(Arguments.IsEmpty())
+	if(Infos.IsEmpty())
 	{
-		ItemArgIndex = Arguments.Emplace(ItemArgName, ERigVMPinDirection::Input, FRigVMRegistry::Get().GetTypeIndex<FRigElementKey>());
-		NameArgIndex = Arguments.Emplace(NameArgName, ERigVMPinDirection::Input, RigVMTypeUtils::TypeIndex::FName);
-		UseNameSpaceArgIndex = Arguments.Emplace(UseNameSpaceArgName, ERigVMPinDirection::Input, RigVMTypeUtils::TypeIndex::Bool);
-		CacheArgIndex = Arguments.Emplace(CacheArgName, ERigVMPinDirection::Hidden, FRigVMRegistry::Get().GetTypeIndex<FCachedRigElement>());
+		ItemArgIndex = Infos.Emplace(ItemArgName, ERigVMPinDirection::Input, FRigVMRegistry::Get().GetTypeIndex<FRigElementKey>());
+		NameArgIndex = Infos.Emplace(NameArgName, ERigVMPinDirection::Input, RigVMTypeUtils::TypeIndex::FName);
+		UseNameSpaceArgIndex = Infos.Emplace(UseNameSpaceArgName, ERigVMPinDirection::Input, RigVMTypeUtils::TypeIndex::Bool);
+		CacheArgIndex = Infos.Emplace(CacheArgName, ERigVMPinDirection::Hidden, FRigVMRegistry::Get().GetTypeIndex<FCachedRigElement>());
 	};
-	return Arguments;
+	return Infos;
 }
 
 #if WITH_EDITOR
@@ -171,16 +171,16 @@ const TArray<TRigVMTypeIndex>& FRigDispatch_MetadataBase::GetValueTypes() const
 
 ////////////////////////////////////////////////////////////////////////////////////////////////
 
-const TArray<FRigVMTemplateArgument>& FRigDispatch_GetMetadata::GetArguments() const
+const TArray<FRigVMTemplateArgumentInfo>& FRigDispatch_GetMetadata::GetArgumentInfos() const
 {
 	if(ValueArgIndex == INDEX_NONE)
 	{
-		Arguments = Super::GetArguments(); 
-		DefaultArgIndex = Arguments.Emplace(DefaultArgName, ERigVMPinDirection::Input, GetValueTypes());
-		ValueArgIndex = Arguments.Emplace(ValueArgName, ERigVMPinDirection::Output, GetValueTypes());
-		FoundArgIndex = Arguments.Emplace(FoundArgName, ERigVMPinDirection::Output, RigVMTypeUtils::TypeIndex::Bool);
+		Infos = Super::GetArgumentInfos(); 
+		DefaultArgIndex = Infos.Emplace(DefaultArgName, ERigVMPinDirection::Input, GetValueTypes());
+		ValueArgIndex = Infos.Emplace(ValueArgName, ERigVMPinDirection::Output, GetValueTypes());
+		FoundArgIndex = Infos.Emplace(FoundArgName, ERigVMPinDirection::Output, RigVMTypeUtils::TypeIndex::Bool);
 	};
-	return Arguments;
+	return Infos;
 }
 
 FRigBaseMetadata* FRigDispatch_GetMetadata::FindMetadata(const FRigVMExtendedExecuteContext& InContext,
@@ -296,15 +296,15 @@ FRigVMFunctionPtr FRigDispatch_GetMetadata::GetDispatchFunctionImpl(const FRigVM
 
 ////////////////////////////////////////////////////////////////////////////////////////////////
 
-const TArray<FRigVMTemplateArgument>& FRigDispatch_SetMetadata::GetArguments() const
+const TArray<FRigVMTemplateArgumentInfo>& FRigDispatch_SetMetadata::GetArgumentInfos() const
 {
 	if(ValueArgIndex == INDEX_NONE)
 	{
-		Arguments = Super::GetArguments(); 
-		ValueArgIndex = Arguments.Emplace(ValueArgName, ERigVMPinDirection::Input, GetValueTypes());
-		SuccessArgIndex = Arguments.Emplace(SuccessArgName, ERigVMPinDirection::Output, RigVMTypeUtils::TypeIndex::Bool);
+		Infos = Super::GetArgumentInfos(); 
+		ValueArgIndex = Infos.Emplace(ValueArgName, ERigVMPinDirection::Input, GetValueTypes());
+		SuccessArgIndex = Infos.Emplace(SuccessArgName, ERigVMPinDirection::Output, RigVMTypeUtils::TypeIndex::Bool);
 	};
-	return Arguments;
+	return Infos;
 }
 
 const TArray<FRigVMExecuteArgument>& FRigDispatch_SetMetadata::GetExecuteArguments_Impl(const FRigVMDispatchContext& InContext) const
diff --git a/Engine/Plugins/Animation/ControlRig/Source/ControlRig/Private/Units/Hierarchy/RigUnit_Metadata.h b/Engine/Plugins/Animation/ControlRig/Source/ControlRig/Private/Units/Hierarchy/RigUnit_Metadata.h
--- a/Engine/Plugins/Animation/ControlRig/Source/ControlRig/Private/Units/Hierarchy/RigUnit_Metadata.h
+++ b/Engine/Plugins/Animation/ControlRig/Source/ControlRig/Private/Units/Hierarchy/RigUnit_Metadata.h
@@ -19,7 +19,7 @@ struct CONTROLRIG_API FRigDispatch_MetadataBase : public FRigDispatchFactory
 #if WITH_EDITOR
 	virtual FString GetNodeTitle(const FRigVMTemplateTypeMap& InTypes) const override;;
 #endif
-	virtual const TArray<FRigVMTemplateArgument>& GetArguments() const override;
+	virtual const TArray<FRigVMTemplateArgumentInfo>& GetArgumentInfos() const override;
 	virtual bool IsSetMetadata() const { return false; }
 
 #if WITH_EDITOR
@@ -32,7 +32,7 @@ protected:
 
 	const TArray<TRigVMTypeIndex>& GetValueTypes() const;
 
-	mutable TArray<FRigVMTemplateArgument> Arguments;
+	mutable TArray<FRigVMTemplateArgumentInfo> Infos;
 	
 	mutable int32 ItemArgIndex = INDEX_NONE;
 	mutable int32 NameArgIndex = INDEX_NONE;
@@ -60,7 +60,7 @@ struct CONTROLRIG_API FRigDispatch_GetMetadata : public FRigDispatch_MetadataBas
 {
 	GENERATED_BODY()
 
-	virtual const TArray<FRigVMTemplateArgument>& GetArguments() const override;
+	virtual const TArray<FRigVMTemplateArgumentInfo>& GetArgumentInfos() const override;
 
 protected:
 
@@ -123,7 +123,7 @@ struct CONTROLRIG_API FRigDispatch_SetMetadata : public FRigDispatch_MetadataBas
 {
 	GENERATED_BODY()
 
-	virtual const TArray<FRigVMTemplateArgument>& GetArguments() const override;
+	virtual const TArray<FRigVMTemplateArgumentInfo>& GetArgumentInfos() const override;
 	virtual const TArray<FRigVMExecuteArgument>& GetExecuteArguments_Impl(const FRigVMDispatchContext& InContext) const override;
 	virtual bool IsSetMetadata() const override { return true; }
 
diff --git a/Engine/Plugins/Experimental/AnimNext/Source/AnimNext/Internal/Param/RigVMDispatch_GetLayerParameter.h b/Engine/Plugins/Experimental/AnimNext/Source/AnimNext/Internal/Param/RigVMDispatch_GetLayerParameter.h
--- a/Engine/Plugins/Experimental/AnimNext/Source/AnimNext/Internal/Param/RigVMDispatch_GetLayerParameter.h
+++ b/Engine/Plugins/Experimental/AnimNext/Source/AnimNext/Internal/Param/RigVMDispatch_GetLayerParameter.h
@@ -26,7 +26,7 @@ private:
 
 	virtual UScriptStruct* GetExecuteContextStruct() const override { return FAnimNextParameterExecuteContext::StaticStruct(); }
 	virtual FName GetArgumentNameForOperandIndex(int32 InOperandIndex, int32 InTotalOperands) const override;
-	virtual const TArray<FRigVMTemplateArgument>& GetArguments() const override;
+	virtual const TArray<FRigVMTemplateArgumentInfo>& GetArgumentInfos() const override;
 #if WITH_EDITOR
 	virtual FString GetArgumentMetaData(const FName& InArgumentName, const FName& InMetaDataKey) const override;
 	virtual FString GetArgumentDefaultValue(const FName& InArgumentName, TRigVMTypeIndex InTypeIndex) const override;
diff --git a/Engine/Plugins/Experimental/AnimNext/Source/AnimNext/Internal/Param/RigVMDispatch_GetParameter.h b/Engine/Plugins/Experimental/AnimNext/Source/AnimNext/Internal/Param/RigVMDispatch_GetParameter.h
--- a/Engine/Plugins/Experimental/AnimNext/Source/AnimNext/Internal/Param/RigVMDispatch_GetParameter.h
+++ b/Engine/Plugins/Experimental/AnimNext/Source/AnimNext/Internal/Param/RigVMDispatch_GetParameter.h
@@ -26,7 +26,7 @@ private:
 
 	virtual UScriptStruct* GetExecuteContextStruct() const { return FAnimNextExecuteContext::StaticStruct(); }
 	virtual FName GetArgumentNameForOperandIndex(int32 InOperandIndex, int32 InTotalOperands) const override;
-	virtual const TArray<FRigVMTemplateArgument>& GetArguments() const override;
+	virtual const TArray<FRigVMTemplateArgumentInfo>& GetArgumentInfos() const override;
 #if WITH_EDITOR
 	virtual FString GetArgumentMetaData(const FName& InArgumentName, const FName& InMetaDataKey) const override;
 	virtual FString GetArgumentDefaultValue(const FName& InArgumentName, TRigVMTypeIndex InTypeIndex) const override;
diff --git a/Engine/Plugins/Experimental/AnimNext/Source/AnimNext/Internal/Param/RigVMDispatch_SetLayerParameter.h b/Engine/Plugins/Experimental/AnimNext/Source/AnimNext/Internal/Param/RigVMDispatch_SetLayerParameter.h
--- a/Engine/Plugins/Experimental/AnimNext/Source/AnimNext/Internal/Param/RigVMDispatch_SetLayerParameter.h
+++ b/Engine/Plugins/Experimental/AnimNext/Source/AnimNext/Internal/Param/RigVMDispatch_SetLayerParameter.h
@@ -26,7 +26,7 @@ private:
 
 	virtual UScriptStruct* GetExecuteContextStruct() const { return FAnimNextParameterExecuteContext::StaticStruct(); }
 	virtual FName GetArgumentNameForOperandIndex(int32 InOperandIndex, int32 InTotalOperands) const override;
-	virtual const TArray<FRigVMTemplateArgument>& GetArguments() const override;
+	virtual const TArray<FRigVMTemplateArgumentInfo>& GetArgumentInfos() const override;
 #if WITH_EDITOR
 	virtual FString GetArgumentMetaData(const FName& InArgumentName, const FName& InMetaDataKey) const override;
 	virtual FString GetArgumentDefaultValue(const FName& InArgumentName, TRigVMTypeIndex InTypeIndex) const override;
diff --git a/Engine/Plugins/Experimental/AnimNext/Source/AnimNext/Private/Param/RigVMDispatch_GetLayerParameter.cpp b/Engine/Plugins/Experimental/AnimNext/Source/AnimNext/Private/Param/RigVMDispatch_GetLayerParameter.cpp
--- a/Engine/Plugins/Experimental/AnimNext/Source/AnimNext/Private/Param/RigVMDispatch_GetLayerParameter.cpp
+++ b/Engine/Plugins/Experimental/AnimNext/Source/AnimNext/Private/Param/RigVMDispatch_GetLayerParameter.cpp
@@ -58,22 +58,24 @@ FString FRigVMDispatch_GetLayerParameter::GetArgumentDefaultValue(const FName& I
 }
 #endif
 
-const TArray<FRigVMTemplateArgument>& FRigVMDispatch_GetLayerParameter::GetArguments() const
+const TArray<FRigVMTemplateArgumentInfo>& FRigVMDispatch_GetLayerParameter::GetArgumentInfos() const
 {
-	static const TArray<FRigVMTemplateArgument::ETypeCategory> ValueCategories =
+	static TArray<FRigVMTemplateArgumentInfo> Infos;
+	if(Infos.IsEmpty())
 	{
-		FRigVMTemplateArgument::ETypeCategory_SingleAnyValue,
-		FRigVMTemplateArgument::ETypeCategory_ArrayAnyValue
-	};
-
-	static const TArray<FRigVMTemplateArgument> Arguments = {
-		FRigVMTemplateArgument(ParameterName, ERigVMPinDirection::Input, RigVMTypeUtils::TypeIndex::FName),
-		FRigVMTemplateArgument(ValueName, ERigVMPinDirection::Output, ValueCategories),
-		FRigVMTemplateArgument(ParameterIdName, ERigVMPinDirection::Hidden, RigVMTypeUtils::TypeIndex::UInt32),
-		FRigVMTemplateArgument(TypeHandleName, ERigVMPinDirection::Hidden, RigVMTypeUtils::TypeIndex::UInt32),
-	};
+		static const TArray<FRigVMTemplateArgument::ETypeCategory> ValueCategories =
+		{
+			FRigVMTemplateArgument::ETypeCategory_SingleAnyValue,
+			FRigVMTemplateArgument::ETypeCategory_ArrayAnyValue
+		};
+		
+		Infos.Emplace(ParameterName, ERigVMPinDirection::Input, RigVMTypeUtils::TypeIndex::FName);
+		Infos.Emplace(ValueName, ERigVMPinDirection::Output, ValueCategories);
+		Infos.Emplace(ParameterIdName, ERigVMPinDirection::Hidden, RigVMTypeUtils::TypeIndex::UInt32);
+		Infos.Emplace(TypeHandleName, ERigVMPinDirection::Hidden, RigVMTypeUtils::TypeIndex::UInt32);
+	}
 
-	return Arguments;
+	return Infos;
 }
 
 FRigVMTemplateTypeMap FRigVMDispatch_GetLayerParameter::OnNewArgumentType(const FName& InArgumentName, TRigVMTypeIndex InTypeIndex) const
diff --git a/Engine/Plugins/Experimental/AnimNext/Source/AnimNext/Private/Param/RigVMDispatch_GetParameter.cpp b/Engine/Plugins/Experimental/AnimNext/Source/AnimNext/Private/Param/RigVMDispatch_GetParameter.cpp
--- a/Engine/Plugins/Experimental/AnimNext/Source/AnimNext/Private/Param/RigVMDispatch_GetParameter.cpp
+++ b/Engine/Plugins/Experimental/AnimNext/Source/AnimNext/Private/Param/RigVMDispatch_GetParameter.cpp
@@ -60,22 +60,24 @@ FString FRigVMDispatch_GetParameter::GetArgumentDefaultValue(const FName& InArgu
 }
 #endif
 
-const TArray<FRigVMTemplateArgument>& FRigVMDispatch_GetParameter::GetArguments() const
+const TArray<FRigVMTemplateArgumentInfo>& FRigVMDispatch_GetParameter::GetArgumentInfos() const
 {
-	static const TArray<FRigVMTemplateArgument::ETypeCategory> ValueCategories =
+	static TArray<FRigVMTemplateArgumentInfo> Infos;
+	if(Infos.IsEmpty())
 	{
-		FRigVMTemplateArgument::ETypeCategory_SingleAnyValue,
-		FRigVMTemplateArgument::ETypeCategory_ArrayAnyValue
-	};
-
-	static const TArray<FRigVMTemplateArgument> Arguments = {
-		FRigVMTemplateArgument(ParameterName, ERigVMPinDirection::Input, RigVMTypeUtils::TypeIndex::FName),
-		FRigVMTemplateArgument(ValueName, ERigVMPinDirection::Output, ValueCategories),
-		FRigVMTemplateArgument(ParameterIdName, ERigVMPinDirection::Hidden, RigVMTypeUtils::TypeIndex::UInt32),
-		FRigVMTemplateArgument(TypeHandleName, ERigVMPinDirection::Hidden, RigVMTypeUtils::TypeIndex::UInt32),
-	};
+		static const TArray<FRigVMTemplateArgument::ETypeCategory> ValueCategories =
+		{
+			FRigVMTemplateArgument::ETypeCategory_SingleAnyValue,
+			FRigVMTemplateArgument::ETypeCategory_ArrayAnyValue
+		};
+
+		Infos.Emplace(ParameterName, ERigVMPinDirection::Input, RigVMTypeUtils::TypeIndex::FName);
+		Infos.Emplace(ValueName, ERigVMPinDirection::Output, ValueCategories);
+		Infos.Emplace(ParameterIdName, ERigVMPinDirection::Hidden, RigVMTypeUtils::TypeIndex::UInt32);
+		Infos.Emplace(TypeHandleName, ERigVMPinDirection::Hidden, RigVMTypeUtils::TypeIndex::UInt32);
+	}
 
-	return Arguments;
+	return Infos;
 }
 
 FRigVMTemplateTypeMap FRigVMDispatch_GetParameter::OnNewArgumentType(const FName& InArgumentName, TRigVMTypeIndex InTypeIndex) const
diff --git a/Engine/Plugins/Experimental/AnimNext/Source/AnimNext/Private/Param/RigVMDispatch_SetLayerParameter.cpp b/Engine/Plugins/Experimental/AnimNext/Source/AnimNext/Private/Param/RigVMDispatch_SetLayerParameter.cpp
--- a/Engine/Plugins/Experimental/AnimNext/Source/AnimNext/Private/Param/RigVMDispatch_SetLayerParameter.cpp
+++ b/Engine/Plugins/Experimental/AnimNext/Source/AnimNext/Private/Param/RigVMDispatch_SetLayerParameter.cpp
@@ -61,22 +61,24 @@ FString FRigVMDispatch_SetLayerParameter::GetArgumentDefaultValue(const FName& I
 }
 #endif
 
-const TArray<FRigVMTemplateArgument>& FRigVMDispatch_SetLayerParameter::GetArguments() const
+const TArray<FRigVMTemplateArgumentInfo>& FRigVMDispatch_SetLayerParameter::GetArgumentInfos() const
 {
-	static const TArray<FRigVMTemplateArgument::ETypeCategory> ValueCategories =
+	static TArray<FRigVMTemplateArgumentInfo> Infos;
+	if(Infos.IsEmpty())
 	{
-		FRigVMTemplateArgument::ETypeCategory_SingleAnyValue,
-		FRigVMTemplateArgument::ETypeCategory_ArrayAnyValue
-	};
+		static const TArray<FRigVMTemplateArgument::ETypeCategory> ValueCategories =
+		{
+			FRigVMTemplateArgument::ETypeCategory_SingleAnyValue,
+			FRigVMTemplateArgument::ETypeCategory_ArrayAnyValue
+		};
 	
-	static const TArray<FRigVMTemplateArgument> Arguments = { 
-		FRigVMTemplateArgument(ParameterName, ERigVMPinDirection::Input, RigVMTypeUtils::TypeIndex::FName),
-		FRigVMTemplateArgument(ValueName, ERigVMPinDirection::Input, ValueCategories),
-		FRigVMTemplateArgument(ParameterIdName, ERigVMPinDirection::Hidden, RigVMTypeUtils::TypeIndex::UInt32),
-		FRigVMTemplateArgument(TypeHandleName, ERigVMPinDirection::Hidden, RigVMTypeUtils::TypeIndex::UInt32),
-	};
+		Infos.Emplace(ParameterName, ERigVMPinDirection::Input, RigVMTypeUtils::TypeIndex::FName);
+		Infos.Emplace(ValueName, ERigVMPinDirection::Input, ValueCategories);
+		Infos.Emplace(ParameterIdName, ERigVMPinDirection::Hidden, RigVMTypeUtils::TypeIndex::UInt32);
+		Infos.Emplace(TypeHandleName, ERigVMPinDirection::Hidden, RigVMTypeUtils::TypeIndex::UInt32);
+	}
 
-	return Arguments;
+	return Infos;
 }
 
 TArray<FRigVMExecuteArgument>& FRigVMDispatch_SetLayerParameter::GetExecuteArguments_Impl(const FRigVMDispatchContext& InContext) const
diff --git a/Engine/Plugins/Runtime/RigVM/Source/RigVM/Private/RigVMCore/RigVMDispatchFactory.cpp b/Engine/Plugins/Runtime/RigVM/Source/RigVM/Private/RigVMCore/RigVMDispatchFactory.cpp
--- a/Engine/Plugins/Runtime/RigVM/Source/RigVM/Private/RigVMCore/RigVMDispatchFactory.cpp
+++ b/Engine/Plugins/Runtime/RigVM/Source/RigVM/Private/RigVMCore/RigVMDispatchFactory.cpp
@@ -155,8 +155,8 @@ bool FRigVMDispatchFactory::IsLazyInputArgument(const FName& InArgumentName) con
 
 FName FRigVMDispatchFactory::GetArgumentNameForOperandIndex(int32 InOperandIndex, int32 InTotalOperands) const
 {
-	check(GetArguments().Num() == InTotalOperands);
-	return GetArguments()[InOperandIndex].GetName();
+	check(GetArgumentInfos().Num() == InTotalOperands);
+	return GetArgumentInfos()[InOperandIndex].Name;
 }
 
 const TArray<FName>& FRigVMDispatchFactory::GetControlFlowBlocks(const FRigVMDispatchContext& InContext) const
@@ -200,9 +200,9 @@ const TArray<FName>* FRigVMDispatchFactory::UpdateArgumentNameCache_NoLock(int32
 	return ArgumentNames.Get();
 }
 
-const TArray<FRigVMTemplateArgument>& FRigVMDispatchFactory::GetArguments() const
+const TArray<FRigVMTemplateArgumentInfo>& FRigVMDispatchFactory::GetArgumentInfos() const
 {
-	static TArray<FRigVMTemplateArgument> EmptyArguments;
+	static const TArray<FRigVMTemplateArgumentInfo> EmptyArguments;
 	return EmptyArguments;
 }
 
@@ -268,68 +268,80 @@ TArray<FRigVMFunction> FRigVMDispatchFactory::CreateDispatchPredicates_NoLock(co
 FString FRigVMDispatchFactory::GetPermutationName(const FRigVMTemplateTypeMap& InTypes) const
 {
 #if WITH_EDITOR
-	const TArray<FRigVMTemplateArgument>& Arguments = GetArguments();
-	check(InTypes.Num() == Arguments.Num());
-	for(const FRigVMTemplateArgument& Argument : Arguments)
+	const TArray<FRigVMTemplateArgumentInfo>& Arguments = GetArgumentInfos();
+	
+	checkf(InTypes.Num() == Arguments.Num(), TEXT("Failed getting permutation names for '%s' "), *GetFactoryName().ToString());
+	
+	for(const FRigVMTemplateArgumentInfo& Argument : Arguments)
 	{
-		check(InTypes.Contains(Argument.GetName()));
+		check(InTypes.Contains(Argument.Name));
 	}
 #endif
 	return GetPermutationNameImpl(InTypes);
 }
 
-TArray<FRigVMTemplateArgument> FRigVMDispatchFactory::BuildArgumentListFromPrimaryArgument(const TArray<FRigVMTemplateArgument>& InArguments, const FName& InPrimaryArgumentName) const
+TArray<FRigVMTemplateArgumentInfo> FRigVMDispatchFactory::BuildArgumentListFromPrimaryArgument(const TArray<FRigVMTemplateArgumentInfo>& InInfos, const FName& InPrimaryArgumentName) const
 {
-	TArray<FRigVMTemplateArgument> NewArguments;
-	const FRigVMTemplateArgument* PrimaryArgument = InArguments.FindByPredicate([InPrimaryArgumentName](const FRigVMTemplateArgument& Arg)
+	TArray<FRigVMTemplateArgumentInfo> NewInfos;
+	
+	const FRigVMTemplateArgumentInfo* PrimaryInfo = InInfos.FindByPredicate([InPrimaryArgumentName](const FRigVMTemplateArgumentInfo& Arg)
 	{
 		return Arg.Name == InPrimaryArgumentName;
 	});
-	
-	if (!PrimaryArgument)
-	{
-		return NewArguments;
-	}
 
-	NewArguments.SetNum(InArguments.Num());
-	for (int32 Index=0; Index < InArguments.Num(); ++Index)
+	if (!PrimaryInfo)
 	{
-		NewArguments[Index].Name = InArguments[Index].Name;
-		NewArguments[Index].Direction = InArguments[Index].Direction;
-		NewArguments[Index].TypeCategories = InArguments[Index].TypeCategories;
+		return NewInfos;
 	}
-	
-	TSet<TRigVMTypeIndex> ProcessedTypes;
-	for (const TRigVMTypeIndex& Type : PrimaryArgument->TypeIndices)
-	{
-		if (ProcessedTypes.Contains(Type))
-		{
-			continue;
-		}
 
+	const int32 NumInfos = InInfos.Num();
+	TArray< TArray<TRigVMTypeIndex> > TypeIndicesArray;
+	TypeIndicesArray.SetNum(NumInfos);
+
+	const FRigVMTemplateArgument PrimaryArgument = PrimaryInfo->GetArgument();
+	for (const TRigVMTypeIndex& Type : PrimaryArgument.TypeIndices)
+	{
 		const TArray<FRigVMTemplateTypeMap> Permutations = GetPermutationsFromArgumentType(InPrimaryArgumentName, Type);
 		for (const FRigVMTemplateTypeMap& Permutation : Permutations)
 		{
-			for (int32 Index=0; Index < InArguments.Num(); ++Index)
+			for (int32 Index=0; Index < InInfos.Num(); ++Index)
 			{
-				const TRigVMTypeIndex* PermutationArg = Permutation.Find(NewArguments[Index].Name);
+				const TRigVMTypeIndex* PermutationArg = Permutation.Find(InInfos[Index].Name);
 				if (!PermutationArg)
 				{
-					NewArguments.Reset();
-					return NewArguments;
+					NewInfos.Reset();
+					return NewInfos;
 				}
-				NewArguments[Index].TypeIndices.Add(*PermutationArg);
+				TypeIndicesArray[Index].Add(*PermutationArg);
 			}
 		}
 	}
 
-	for (FRigVMTemplateArgument& Argument : NewArguments)
+	NewInfos.Reserve(NumInfos);
+	for (int32 Index=0; Index < NumInfos; ++Index)
 	{
-		Argument.EnsureValidExecuteType();
-		Argument.UpdateTypeToPermutations();
+		const FRigVMTemplateArgument Argument = InInfos[Index].GetArgument();
+		const TArray<TRigVMTypeIndex>& TypeIndices = TypeIndicesArray[Index];
+
+		const FName Name = InInfos[Index].Name;
+		const ERigVMPinDirection Direction = Argument.Direction;
+		const TArray<FRigVMTemplateArgument::ETypeCategory> TypeCategories = Argument.TypeCategories;
+		
+		NewInfos.Emplace(Name, Direction, [TypeIndices, TypeCategories](const FName InName, ERigVMPinDirection InDirection)
+		{
+			FRigVMTemplateArgument Argument;
+			Argument.Name = InName;
+			Argument.Direction = InDirection;
+			Argument.TypeCategories = TypeCategories;
+			Argument.TypeIndices = TypeIndices;
+			
+			Argument.EnsureValidExecuteType();
+			Argument.UpdateTypeToPermutations();
+			return Argument;
+		} );
 	}
 	
-	return NewArguments;
+	return NewInfos;
 }
 
 FString FRigVMDispatchFactory::GetPermutationNameImpl(const FRigVMTemplateTypeMap& InTypes) const
@@ -363,11 +375,12 @@ const FRigVMTemplate* FRigVMDispatchFactory::GetTemplate() const
 		return CachedTemplate;
 	}
 
-	const TArray<FRigVMTemplateArgument>& Arguments = GetArguments();
+	const TArray<FRigVMTemplateArgumentInfo>& Arguments = GetArgumentInfos();
 
 	// we don't allow execute types on arguments
-	for(const FRigVMTemplateArgument& Argument : Arguments)
+	for(const FRigVMTemplateArgumentInfo& Info : Arguments)
 	{
+		const FRigVMTemplateArgument Argument = Info.GetArgument();
 		for(const TRigVMTypeIndex& TypeIndex : Argument.GetTypeIndices())
 		{
 			if(Registry.IsExecuteType(TypeIndex))
@@ -381,36 +394,12 @@ const FRigVMTemplate* FRigVMDispatchFactory::GetTemplate() const
 	}
 
 	FRigVMTemplateDelegates Delegates;
-	Delegates.NewArgumentTypeDelegate = FRigVMTemplate_NewArgumentTypeDelegate::CreateLambda(
-		[this](const FRigVMTemplate*, const FName& InArgumentName, int32 InTypeIndex)
-		{
-			return OnNewArgumentType(InArgumentName, InTypeIndex);
-		});
-
-	Delegates.GetPermutationsFromArgumentTypeDelegate = FRigVMTemplate_GetPermutationsFromArgumentTypeDelegate::CreateLambda(
-		[this](const FRigVMTemplate*, const FName& InArgumentName, int32 InTypeIndex)
-		{
-			return GetPermutationsFromArgumentType(InArgumentName, InTypeIndex);
-		});
-
 	Delegates.GetDispatchFactoryDelegate = FRigVMTemplate_GetDispatchFactoryDelegate::CreateLambda(
-[FactoryName]()
+	[FactoryName]()
 	{
 		return FRigVMRegistry::Get().FindDispatchFactory(FactoryName);
 	});
 
-	Delegates.RequestDispatchFunctionDelegate = FRigVMTemplate_RequestDispatchFunctionDelegate::CreateLambda(
-	[this](const FRigVMTemplate*, const FRigVMTemplateTypeMap& InTypes)
-	{
-		return CreateDispatchFunction(InTypes);
-	});
-
-	Delegates.RequestDispatchPredicatesDelegate = FRigVMTemplate_RequestDispatchPredicatesDelegate::CreateLambda(
-	[this](const FRigVMTemplate*, const FRigVMTemplateTypeMap& InTypes)
-	{
-		return CreateDispatchPredicates(InTypes);
-	});
-
 	CachedTemplate = Registry.AddTemplateFromArguments(GetFactoryName(), Arguments, Delegates); 
 	return CachedTemplate;
 }
diff --git a/Engine/Plugins/Runtime/RigVM/Source/RigVM/Private/RigVMCore/RigVMRegistry.cpp b/Engine/Plugins/Runtime/RigVM/Source/RigVM/Private/RigVMCore/RigVMRegistry.cpp
--- a/Engine/Plugins/Runtime/RigVM/Source/RigVM/Private/RigVMCore/RigVMRegistry.cpp
+++ b/Engine/Plugins/Runtime/RigVM/Source/RigVM/Private/RigVMCore/RigVMRegistry.cpp
@@ -1559,35 +1559,33 @@ const TChunkedArray<FRigVMTemplate>& FRigVMRegistry::GetTemplates() const
 	return Templates;
 }
 
-const FRigVMTemplate* FRigVMRegistry::GetOrAddTemplateFromArguments(const FName& InName, const TArray<FRigVMTemplateArgument>& InArguments, const FRigVMTemplateDelegates& InDelegates)
+const FRigVMTemplate* FRigVMRegistry::GetOrAddTemplateFromArguments(const FName& InName, const TArray<FRigVMTemplateArgumentInfo>& InInfos, const FRigVMTemplateDelegates& InDelegates)
 {
 	FScopeLock RegisterTemplateScopeLock(&RegisterTemplateMutex);
 	
-	FRigVMTemplate Template(InName, InArguments);
-
 	// avoid reentry in FindTemplate. try to find an existing
 	// template only if we are not yet in ::FindTemplate.
-	if(const FRigVMTemplate* ExistingTemplate = FindTemplate(Template.GetNotation()))
+	const FName Notation = FRigVMTemplateArgumentInfo::ComputeTemplateNotation(InName, InInfos);
+	if(const FRigVMTemplate* ExistingTemplate = FindTemplate(Notation))
 	{
 		return ExistingTemplate;
 	}
 
-	return AddTemplateFromArguments_NoLock(InName, InArguments, InDelegates);
+	return AddTemplateFromArguments_NoLock(InName, InInfos, InDelegates);
 }
 
-const FRigVMTemplate* FRigVMRegistry::AddTemplateFromArguments(const FName& InName, const TArray<FRigVMTemplateArgument>& InArguments, const FRigVMTemplateDelegates& InDelegates)
+const FRigVMTemplate* FRigVMRegistry::AddTemplateFromArguments(const FName& InName, const TArray<FRigVMTemplateArgumentInfo>& InInfos, const FRigVMTemplateDelegates& InDelegates)
 {
 	FScopeLock RegisterTemplateScopeLock(&RegisterTemplateMutex);
-	return AddTemplateFromArguments_NoLock(InName, InArguments, InDelegates);
+	return AddTemplateFromArguments_NoLock(InName, InInfos, InDelegates);
 }
 
-const FRigVMTemplate* FRigVMRegistry::AddTemplateFromArguments_NoLock(const FName& InName, const TArray<FRigVMTemplateArgument>& InArguments, const FRigVMTemplateDelegates& InDelegates)
+const FRigVMTemplate* FRigVMRegistry::AddTemplateFromArguments_NoLock(const FName& InName, const TArray<FRigVMTemplateArgumentInfo>& InInfos, const FRigVMTemplateDelegates& InDelegates)
 {
-	FRigVMTemplate Template(InName, InArguments);
-	
 	// we only support to ask for templates here which provide singleton types
 	int32 NumPermutations = 0;
-	for(const FRigVMTemplateArgument& Argument : InArguments)
+	FRigVMTemplate Template(InName, InInfos);
+	for(const FRigVMTemplateArgument& Argument : Template.Arguments)
 	{
 		if(!Argument.IsSingleton() && NumPermutations > 1)
 		{
diff --git a/Engine/Plugins/Runtime/RigVM/Source/RigVM/Private/RigVMCore/RigVMTemplate.cpp b/Engine/Plugins/Runtime/RigVM/Source/RigVM/Private/RigVMCore/RigVMTemplate.cpp
--- a/Engine/Plugins/Runtime/RigVM/Source/RigVM/Private/RigVMCore/RigVMTemplate.cpp
+++ b/Engine/Plugins/Runtime/RigVM/Source/RigVM/Private/RigVMCore/RigVMTemplate.cpp
@@ -419,6 +419,102 @@ TArray<FString> FRigVMTemplateArgument::GetSupportedTypeStrings(const TArray<int
 
 #endif
 
+/**
+ * FRigVMTemplateArgumentInfo 
+ */
+
+FRigVMTemplateArgumentInfo::FRigVMTemplateArgumentInfo(const FName InName, ERigVMPinDirection InDirection, const TArray<TRigVMTypeIndex>& InTypeIndices)
+	: Name(InName)
+	, Direction(InDirection)
+	, FactoryCallback([InTypeIndices](const FName InName, ERigVMPinDirection InDirection){ return FRigVMTemplateArgument(InName, InDirection, InTypeIndices); } )
+{}
+
+FRigVMTemplateArgumentInfo::FRigVMTemplateArgumentInfo(const FName InName, ERigVMPinDirection InDirection, TRigVMTypeIndex InTypeIndex)
+	: Name(InName)
+	, Direction(InDirection)
+	, FactoryCallback( [InTypeIndex](const FName InName, ERigVMPinDirection InDirection){ return FRigVMTemplateArgument(InName, InDirection, InTypeIndex); } )
+{}
+
+FRigVMTemplateArgumentInfo::FRigVMTemplateArgumentInfo(
+	const FName InName, ERigVMPinDirection InDirection,
+	const TArray<FRigVMTemplateArgument::ETypeCategory>& InTypeCategories)
+	: Name(InName)
+	, Direction(InDirection)
+	, FactoryCallback([InTypeCategories](const FName InName, ERigVMPinDirection InDirection){ return FRigVMTemplateArgument(InName, InDirection, InTypeCategories); })
+{}
+
+FRigVMTemplateArgumentInfo::FRigVMTemplateArgumentInfo(const FName InName, ERigVMPinDirection InDirection)
+	: Name(InName)
+	, Direction(InDirection)
+	, FactoryCallback([](const FName InName, ERigVMPinDirection InDirection){ return FRigVMTemplateArgument(InName, InDirection); })
+{}
+
+FRigVMTemplateArgumentInfo::FRigVMTemplateArgumentInfo(const FName InName, ERigVMPinDirection InDirection, ArgumentCallback&& InCallback)
+	: Name(InName)
+	, Direction(InDirection)
+	, FactoryCallback(InCallback)
+{}
+
+FRigVMTemplateArgument FRigVMTemplateArgumentInfo::GetArgument() const
+{
+	FRigVMTemplateArgument Argument = FactoryCallback(Name, Direction);
+	return MoveTemp(Argument);
+}
+
+FName FRigVMTemplateArgumentInfo::ComputeTemplateNotation(const FName InTemplateName, const TArray<FRigVMTemplateArgumentInfo>& InInfos)
+{
+	if (InInfos.IsEmpty())
+	{
+		return NAME_None;	
+	}
+		
+	TArray<FString> ArgumentNotations;
+	Algo::TransformIf(
+		InInfos,
+		ArgumentNotations,
+		[](const FRigVMTemplateArgumentInfo& Info){ return Info.Direction != ERigVMPinDirection::Invalid && Info.Direction != ERigVMPinDirection::Hidden; },
+		[](const FRigVMTemplateArgumentInfo& Info){ return FRigVMTemplate::GetArgumentNotation(Info.Name, Info.Direction); }
+	);
+
+	if (ArgumentNotations.IsEmpty())
+	{
+		return NAME_None;
+	}
+	
+	const FString NotationStr = FString::Printf(TEXT("%s(%s)"), *InTemplateName.ToString(), *FString::Join(ArgumentNotations, TEXT(",")));
+	return *NotationStr;
+}
+
+TArray<TRigVMTypeIndex> FRigVMTemplateArgumentInfo::GetTypesFromCategories(
+	const TArray<FRigVMTemplateArgument::ETypeCategory>& InTypeCategories,
+	const FRigVMTemplateArgument::FTypeFilter& InTypeFilter)
+{
+	TSet<TRigVMTypeIndex> AllTypes;
+	for (const FRigVMTemplateArgument::ETypeCategory TypeCategory : InTypeCategories)
+	{
+		AllTypes.Append(FRigVMRegistry::Get().GetTypesForCategory(TypeCategory));
+	}
+
+	TArray<TRigVMTypeIndex> Types;
+	if (!InTypeFilter.IsBound())
+	{
+		Types = AllTypes.Array();
+	}
+	else
+	{
+		Types.Reserve(AllTypes.Num());
+		for (const TRigVMTypeIndex& Type : AllTypes)
+		{
+			if (InTypeFilter.Execute(Type))
+			{
+				Types.Add(Type);
+			}
+		}
+	}
+	
+	return MoveTemp(Types);
+}
+
 ////////////////////////////////////////////////////////////////////////////////////////////////
 
 FRigVMTemplate::FRigVMTemplate()
@@ -478,7 +574,7 @@ FRigVMTemplate::FRigVMTemplate(UScriptStruct* InStruct, const FString& InTemplat
 			FRigVMTemplateArgument Argument(*It);
 			Argument.Index = Arguments.Num();
 
-			if(!Argument.IsExecute() && IsValidArgumentForTemplate(Argument) && Argument.GetDirection() != ERigVMPinDirection::Hidden)
+			if(!Argument.IsExecute() && IsValidArgumentForTemplate(Argument.GetDirection()) && Argument.GetDirection() != ERigVMPinDirection::Hidden)
 			{
 				Arguments.Add(Argument);
 			}
@@ -493,7 +589,7 @@ FRigVMTemplate::FRigVMTemplate(UScriptStruct* InStruct, const FString& InTemplat
 		{
 			if(!Argument->IsExecute() && Argument->GetDirection() != ERigVMPinDirection::Hidden)
 			{
-				ArgumentNotations.Add(GetArgumentNotation(*Argument));
+				ArgumentNotations.Add(GetArgumentNotation(Argument->Name, Argument->Direction));
 			}
 		}
 	}
@@ -515,34 +611,23 @@ FRigVMTemplate::FRigVMTemplate(UScriptStruct* InStruct, const FString& InTemplat
 	}
 }
 
-FRigVMTemplate::FRigVMTemplate(const FName& InTemplateName, const TArray<FRigVMTemplateArgument>& InArguments)
+FRigVMTemplate::FRigVMTemplate(const FName& InTemplateName, const TArray<FRigVMTemplateArgumentInfo>& InInfos)
 	: Index(INDEX_NONE)
 	, Notation(NAME_None)
 	, Hash(UINT32_MAX)
 {
-	TArray<FString> ArgumentNotations;
-	for (const FRigVMTemplateArgument& InArgument : InArguments)
+	for (const FRigVMTemplateArgumentInfo& InInfo : InInfos)
 	{
-		FRigVMTemplateArgument Argument = InArgument;
-		Argument.Index = Arguments.Num();
-
-		if(IsValidArgumentForTemplate(Argument))
+		if(IsValidArgumentForTemplate(InInfo.Direction))
 		{
-			Arguments.Add(Argument);
-			if(Argument.GetDirection() != ERigVMPinDirection::Hidden)
-			{
-				ArgumentNotations.Add(GetArgumentNotation(Argument));
-			}
+			FRigVMTemplateArgument Argument = InInfo.GetArgument();
+			Argument.Index = Arguments.Num();
+			Arguments.Emplace(MoveTemp(Argument));
 		}
 	}
-
-	if (ArgumentNotations.Num() > 0)
-	{
-		const FString NotationStr = FString::Printf(TEXT("%s(%s)"), *InTemplateName.ToString(), *FString::Join(ArgumentNotations, TEXT(",")));
-		Notation = *NotationStr;
-
-		UpdateTypesHashToPermutation(Permutations.Num()-1);
-	}
+	
+	Notation = FRigVMTemplateArgumentInfo::ComputeTemplateNotation(InTemplateName, InInfos);
+	UpdateTypesHashToPermutation(Permutations.Num()-1);
 }
 
 FLinearColor FRigVMTemplate::GetColorFromMetadata(FString InMetadata)
@@ -572,47 +657,35 @@ FLinearColor FRigVMTemplate::GetColorFromMetadata(FString InMetadata)
 	return Color;
 }
 
-bool FRigVMTemplate::IsValidArgumentForTemplate(const FRigVMTemplateArgument& InArgument)
+bool FRigVMTemplate::IsValidArgumentForTemplate(const ERigVMPinDirection InDirection)
 {
-	static const TArray<ERigVMPinDirection> ValidDirections = {
-		ERigVMPinDirection::Input,
-		ERigVMPinDirection::Output,
-		ERigVMPinDirection::IO,
-		ERigVMPinDirection::Hidden,
-		ERigVMPinDirection::Visible
-	};
-
-	if(!ValidDirections.Contains(InArgument.Direction))
-	{
-		return false;
-	}
-	return true;
+	return InDirection != ERigVMPinDirection::Invalid;
 }
 
 
-const FString& FRigVMTemplate::GetArgumentNotationPrefix(const FRigVMTemplateArgument& InArgument)
+const FString& FRigVMTemplate::GetDirectionPrefix(const ERigVMPinDirection InDirection)
 {
 	static const FString EmptyPrefix = FString();
 	static const FString InPrefix = TEXT("in ");
 	static const FString OutPrefix = TEXT("out ");
 	static const FString IOPrefix = TEXT("io ");
 
-	switch(InArgument.Direction)
+	switch(InDirection)
 	{
-		case ERigVMPinDirection::Input:
-		case ERigVMPinDirection::Visible:
+	case ERigVMPinDirection::Input:
+	case ERigVMPinDirection::Visible:
 		{
 			return InPrefix;
 		}
-		case ERigVMPinDirection::Output:
+	case ERigVMPinDirection::Output:
 		{
 			return OutPrefix;
 		}
-		case ERigVMPinDirection::IO:
+	case ERigVMPinDirection::IO:
 		{
 			return IOPrefix;
 		}
-		default:
+	default:
 		{
 			break;
 		}
@@ -621,21 +694,19 @@ const FString& FRigVMTemplate::GetArgumentNotationPrefix(const FRigVMTemplateArg
 	return EmptyPrefix;
 }
 
-FString FRigVMTemplate::GetArgumentNotation(const FRigVMTemplateArgument& InArgument)
+FString FRigVMTemplate::GetArgumentNotation(const FName InName, const ERigVMPinDirection InDirection)
 {
-	return FString::Printf(TEXT("%s%s"),
-		*GetArgumentNotationPrefix(InArgument),
-		*InArgument.GetName().ToString());
+	return FString::Printf(TEXT("%s%s"), *GetDirectionPrefix(InDirection), *InName.ToString());
 }
 
 void FRigVMTemplate::ComputeNotationFromArguments(const FString& InTemplateName)
 {
 	TArray<FString> ArgumentNotations;			
-	for (FRigVMTemplateArgument& Argument : Arguments)
+	for (const FRigVMTemplateArgument& Argument : Arguments)
 	{
-		if(FRigVMTemplate::IsValidArgumentForTemplate(Argument))
+		if(IsValidArgumentForTemplate(Argument.GetDirection()))
 		{
-			ArgumentNotations.Add(FRigVMTemplate::GetArgumentNotation(Argument));
+			ArgumentNotations.Add(GetArgumentNotation(Argument.Name, Argument.Direction));
 		}
 	}
 
@@ -1233,11 +1304,11 @@ const FRigVMFunction* FRigVMTemplate::GetOrCreatePermutation(int32 InIndex)
 		{
 			Types.Add(Argument.GetName(), Argument.TypeIndices[InIndex]);
 		}
-		
-		if(const FRigVMFunctionPtr DispatchFunction = Delegates.RequestDispatchFunctionDelegate.Execute(this, Types))
+
+		FRigVMDispatchFactory* Factory = Delegates.GetDispatchFactoryDelegate.Execute();
+		if (ensure(Factory))
 		{
-			FRigVMDispatchFactory* Factory = Delegates.GetDispatchFactoryDelegate.Execute();
-			check(Factory);
+			const FRigVMFunctionPtr DispatchFunction = Factory->CreateDispatchFunction(Types);
 
 			TArray<FRigVMFunctionArgument> FunctionArguments;
 			for(const FRigVMTemplateArgument& Argument : Arguments)
@@ -1262,11 +1333,8 @@ const FRigVMFunction* FRigVMTemplate::GetOrCreatePermutation(int32 InIndex)
 			Registry.Functions[FunctionIndex].TemplateIndex = Index;
 			Registry.FunctionNameToIndex.Add(*PermutationName, FunctionIndex);
 
-			if (Delegates.RequestDispatchPredicatesDelegate.IsBound())
-			{
-				TArray<FRigVMFunction> Predicates = Delegates.RequestDispatchPredicatesDelegate.Execute(this, Types);
-				Registry.StructNameToPredicates.Add(*PermutationName, Predicates);
-			}
+			TArray<FRigVMFunction> Predicates = Factory->CreateDispatchPredicates(Types);
+			Registry.StructNameToPredicates.Add(*PermutationName, MoveTemp(Predicates));
 
 			return &Registry.Functions[FunctionIndex];
 		}
@@ -1693,13 +1761,18 @@ bool FRigVMTemplate::AddTypeForArgument(const FName& InArgumentName, TRigVMTypeI
 	InvalidateHash();
 
 	TArray<FRigVMTemplateTypeMap> TypesArray;
-	if (OnGetPermutationsFromArgumentType().IsBound())
+
+	if(Delegates.GetDispatchFactoryDelegate.IsBound())
 	{
-		TypesArray = OnGetPermutationsFromArgumentType().Execute(this, InArgumentName, InTypeIndex);
+		const FRigVMDispatchFactory* Factory = Delegates.GetDispatchFactoryDelegate.Execute();
+		if(ensure(Factory))
+		{
+			TypesArray = Factory->GetPermutationsFromArgumentType(InArgumentName, InTypeIndex);
+		}
 	}
 	else if(OnNewArgumentType().IsBound())
 	{
-		FRigVMTemplateTypeMap Types = OnNewArgumentType().Execute(this, InArgumentName, InTypeIndex);
+		FRigVMTemplateTypeMap Types = OnNewArgumentType().Execute(InArgumentName, InTypeIndex);
 		TypesArray = {Types};
 	}
 
diff --git a/Engine/Plugins/Runtime/RigVM/Source/RigVM/Private/RigVMFunctions/RigVMDispatch_Array.cpp b/Engine/Plugins/Runtime/RigVM/Source/RigVM/Private/RigVMFunctions/RigVMDispatch_Array.cpp
--- a/Engine/Plugins/Runtime/RigVM/Source/RigVM/Private/RigVMFunctions/RigVMDispatch_Array.cpp
+++ b/Engine/Plugins/Runtime/RigVM/Source/RigVM/Private/RigVMFunctions/RigVMDispatch_Array.cpp
@@ -127,7 +127,7 @@ FString FRigVMDispatch_ArrayBase::GetArgumentDefaultValue(const FName& InArgumen
 
 #endif
 
-FRigVMTemplateArgument FRigVMDispatch_ArrayBase::CreateArgument(const FName& InName, ERigVMPinDirection InDirection)
+FRigVMTemplateArgumentInfo FRigVMDispatch_ArrayBase::CreateArgumentInfo(const FName& InName, ERigVMPinDirection InDirection)
 {
 	static const TArray<FRigVMTemplateArgument::ETypeCategory> ArrayCategories = {
 		FRigVMTemplateArgument::ETypeCategory_ArrayAnyValue
@@ -139,27 +139,27 @@ FRigVMTemplateArgument FRigVMDispatch_ArrayBase::CreateArgument(const FName& InN
 
 	if(InName == ArrayName || InName == ValuesName || InName == OtherName || InName == CloneName || InName == ResultName)
 	{
-		return FRigVMTemplateArgument(InName, InDirection, ArrayCategories);
+		return FRigVMTemplateArgumentInfo(InName, InDirection, ArrayCategories);
 	}
 	if(InName == ElementName)
 	{
-		return FRigVMTemplateArgument(InName, InDirection, ElementCategories);
+		return FRigVMTemplateArgumentInfo(InName, InDirection, ElementCategories);
 	}
 	if(InName == NumName || InName == IndexName || InName == CountName)
 	{
-		return FRigVMTemplateArgument(InName, InDirection, RigVMTypeUtils::TypeIndex::Int32);
+		return FRigVMTemplateArgumentInfo(InName, InDirection, RigVMTypeUtils::TypeIndex::Int32);
 	}
 	if(InName == SuccessName)
 	{
-		return FRigVMTemplateArgument(InName, InDirection, RigVMTypeUtils::TypeIndex::Bool);
+		return FRigVMTemplateArgumentInfo(InName, InDirection, RigVMTypeUtils::TypeIndex::Bool);
 	}
 	if(InName == RatioName)
 	{
-		return FRigVMTemplateArgument(InName, InDirection, RigVMTypeUtils::TypeIndex::Float);
+		return FRigVMTemplateArgumentInfo(InName, InDirection, RigVMTypeUtils::TypeIndex::Float);
 	}
 
 	checkNoEntry();
-	return FRigVMTemplateArgument();
+	return FRigVMTemplateArgumentInfo(NAME_None, ERigVMPinDirection::Invalid, INDEX_NONE);
 }
 
 TMap<uint32, int32> FRigVMDispatch_ArrayBase::GetArrayHash(FScriptArrayHelper& InArrayHelper, const FArrayProperty* InArrayProperty)
@@ -205,13 +205,15 @@ const TArray<FRigVMExecuteArgument>& FRigVMDispatch_ArrayBaseMutable::GetExecute
 
 ////////////////////////////////////////////////////////////////////////////////////////////////////////////////////////
 
-const TArray<FRigVMTemplateArgument>& FRigVMDispatch_ArrayMake::GetArguments() const
+const TArray<FRigVMTemplateArgumentInfo>& FRigVMDispatch_ArrayMake::GetArgumentInfos() const
 {
-	static const TArray<FRigVMTemplateArgument> Arguments = {
-		CreateArgument(ValuesName, ERigVMPinDirection::Input),
-		CreateArgument(ArrayName, ERigVMPinDirection::Output)
-	};
-	return Arguments;
+	static TArray<FRigVMTemplateArgumentInfo> Infos;
+	if (Infos.IsEmpty())
+	{
+		Infos.Emplace(CreateArgumentInfo(ValuesName, ERigVMPinDirection::Input));
+		Infos.Emplace(CreateArgumentInfo(ArrayName, ERigVMPinDirection::Output));
+	}
+	return Infos;
 }
 
 FRigVMTemplateTypeMap FRigVMDispatch_ArrayMake::OnNewArgumentType(const FName& InArgumentName, TRigVMTypeIndex InTypeIndex) const
@@ -315,12 +317,14 @@ FName FRigVMDispatch_ArrayReset::GetArgumentNameForOperandIndex(int32 InOperandI
 	return ArrayName;
 }
 
-const TArray<FRigVMTemplateArgument>& FRigVMDispatch_ArrayReset::GetArguments() const
+const TArray<FRigVMTemplateArgumentInfo>& FRigVMDispatch_ArrayReset::GetArgumentInfos() const
 {
-	static const TArray<FRigVMTemplateArgument> Arguments = {
-		CreateArgument(ArrayName, ERigVMPinDirection::IO)
-	};
-	return Arguments;
+	static TArray<FRigVMTemplateArgumentInfo> Infos;
+	if(Infos.IsEmpty())
+	{
+		Infos.Emplace(CreateArgumentInfo(ArrayName, ERigVMPinDirection::IO));
+	}
+	return Infos;
 }
 
 FRigVMTemplateTypeMap FRigVMDispatch_ArrayReset::OnNewArgumentType(const FName& InArgumentName, TRigVMTypeIndex InTypeIndex) const
@@ -368,13 +372,15 @@ FName FRigVMDispatch_ArrayGetNum::GetArgumentNameForOperandIndex(int32 InOperand
 	return ArgumentNames[InOperandIndex];
 }
 
-const TArray<FRigVMTemplateArgument>& FRigVMDispatch_ArrayGetNum::GetArguments() const
+const TArray<FRigVMTemplateArgumentInfo>& FRigVMDispatch_ArrayGetNum::GetArgumentInfos() const
 {
-	static const TArray<FRigVMTemplateArgument> Arguments = {
-		CreateArgument(ArrayName, ERigVMPinDirection::Input),
-		CreateArgument(NumName, ERigVMPinDirection::Output)
-	};
-	return Arguments;
+	static TArray<FRigVMTemplateArgumentInfo> Infos;
+	if(Infos.IsEmpty())
+	{
+		Infos.Emplace( CreateArgumentInfo(ArrayName, ERigVMPinDirection::Input) );
+		Infos.Emplace( CreateArgumentInfo(NumName, ERigVMPinDirection::Output) );
+	}
+	return Infos;
 }
 
 FRigVMTemplateTypeMap FRigVMDispatch_ArrayGetNum::OnNewArgumentType(const FName& InArgumentName, TRigVMTypeIndex InTypeIndex) const
@@ -428,13 +434,15 @@ FName FRigVMDispatch_ArraySetNum::GetArgumentNameForOperandIndex(int32 InOperand
 	return ArgumentNames[InOperandIndex];
 }
 
-const TArray<FRigVMTemplateArgument>& FRigVMDispatch_ArraySetNum::GetArguments() const
+const TArray<FRigVMTemplateArgumentInfo>& FRigVMDispatch_ArraySetNum::GetArgumentInfos() const
 {
-	static const TArray<FRigVMTemplateArgument> Arguments = {
-		CreateArgument(ArrayName, ERigVMPinDirection::IO),
-		CreateArgument(NumName, ERigVMPinDirection::Input)
+	static TArray<FRigVMTemplateArgumentInfo> Infos;
+	if (Infos.IsEmpty())
+	{
+		Infos.Emplace(CreateArgumentInfo(ArrayName, ERigVMPinDirection::IO));
+		Infos.Emplace(CreateArgumentInfo(NumName, ERigVMPinDirection::Input));
 	};
-	return Arguments;
+	return Infos;
 }
 
 FRigVMTemplateTypeMap FRigVMDispatch_ArraySetNum::OnNewArgumentType(const FName& InArgumentName, TRigVMTypeIndex InTypeIndex) const
@@ -492,14 +500,16 @@ FName FRigVMDispatch_ArrayGetAtIndex::GetArgumentNameForOperandIndex(int32 InOpe
 	return ArgumentNames[InOperandIndex];
 }
 
-const TArray<FRigVMTemplateArgument>& FRigVMDispatch_ArrayGetAtIndex::GetArguments() const
+const TArray<FRigVMTemplateArgumentInfo>& FRigVMDispatch_ArrayGetAtIndex::GetArgumentInfos() const
 {
-	static const TArray<FRigVMTemplateArgument> Arguments = {
-		CreateArgument(ArrayName, ERigVMPinDirection::Input),
-		CreateArgument(IndexName, ERigVMPinDirection::Input),
-		CreateArgument(ElementName, ERigVMPinDirection::Output)
-	};
-	return Arguments;
+	static TArray<FRigVMTemplateArgumentInfo> Infos;
+	if (Infos.IsEmpty())
+	{
+		Infos.Emplace( CreateArgumentInfo(ArrayName, ERigVMPinDirection::Input) );
+		Infos.Emplace( CreateArgumentInfo(IndexName, ERigVMPinDirection::Input) );
+		Infos.Emplace( CreateArgumentInfo(ElementName, ERigVMPinDirection::Output) );
+	}
+	return Infos;
 }
 
 FRigVMTemplateTypeMap FRigVMDispatch_ArrayGetAtIndex::OnNewArgumentType(const FName& InArgumentName, TRigVMTypeIndex InTypeIndex) const
@@ -575,14 +585,16 @@ FName FRigVMDispatch_ArraySetAtIndex::GetArgumentNameForOperandIndex(int32 InOpe
 	return ArgumentNames[InOperandIndex];
 }
 
-const TArray<FRigVMTemplateArgument>& FRigVMDispatch_ArraySetAtIndex::GetArguments() const
+const TArray<FRigVMTemplateArgumentInfo>& FRigVMDispatch_ArraySetAtIndex::GetArgumentInfos() const
 {
-	static const TArray<FRigVMTemplateArgument> Arguments = {
-		CreateArgument(ArrayName, ERigVMPinDirection::IO),
-		CreateArgument(IndexName, ERigVMPinDirection::Input),
-		CreateArgument(ElementName, ERigVMPinDirection::Input)
-	};
-	return Arguments;
+	static TArray<FRigVMTemplateArgumentInfo> Infos;
+	if (Infos.IsEmpty())
+	{
+		Infos.Emplace( CreateArgumentInfo(ArrayName, ERigVMPinDirection::IO) );
+		Infos.Emplace( CreateArgumentInfo(IndexName, ERigVMPinDirection::Input) );
+		Infos.Emplace( CreateArgumentInfo(ElementName, ERigVMPinDirection::Input) );
+	}
+	return Infos;
 }
 
 FRigVMTemplateTypeMap FRigVMDispatch_ArraySetAtIndex::OnNewArgumentType(const FName& InArgumentName, TRigVMTypeIndex InTypeIndex) const
@@ -658,14 +670,16 @@ FName FRigVMDispatch_ArrayAdd::GetArgumentNameForOperandIndex(int32 InOperandInd
 	return ArgumentNames[InOperandIndex];
 }
 
-const TArray<FRigVMTemplateArgument>& FRigVMDispatch_ArrayAdd::GetArguments() const
+const TArray<FRigVMTemplateArgumentInfo>& FRigVMDispatch_ArrayAdd::GetArgumentInfos() const
 {
-	static const TArray<FRigVMTemplateArgument> Arguments = {
-		CreateArgument(ArrayName, ERigVMPinDirection::IO),
-		CreateArgument(ElementName, ERigVMPinDirection::Input),
-		CreateArgument(IndexName, ERigVMPinDirection::Output)
-	};
-	return Arguments;
+	static TArray<FRigVMTemplateArgumentInfo> Infos;
+	if (Infos.IsEmpty())
+	{
+		Infos.Emplace( CreateArgumentInfo(ArrayName, ERigVMPinDirection::IO) );
+		Infos.Emplace( CreateArgumentInfo(ElementName, ERigVMPinDirection::Input) );
+		Infos.Emplace( CreateArgumentInfo(IndexName, ERigVMPinDirection::Output) );
+	}
+	return Infos;
 }
 
 void FRigVMDispatch_ArrayAdd::Execute(FRigVMExtendedExecuteContext& InContext, FRigVMMemoryHandleArray Handles, FRigVMPredicateBranchArray Predicates)
@@ -776,13 +790,15 @@ FName FRigVMDispatch_ArrayRemove::GetArgumentNameForOperandIndex(int32 InOperand
 	return ArgumentNames[InOperandIndex];
 }
 
-const TArray<FRigVMTemplateArgument>& FRigVMDispatch_ArrayRemove::GetArguments() const
+const TArray<FRigVMTemplateArgumentInfo>& FRigVMDispatch_ArrayRemove::GetArgumentInfos() const
 {
-	static const TArray<FRigVMTemplateArgument> Arguments = {
-		CreateArgument(ArrayName, ERigVMPinDirection::IO),
-		CreateArgument(IndexName, ERigVMPinDirection::Input)
-	};
-	return Arguments;
+	static TArray<FRigVMTemplateArgumentInfo> Infos;
+	if (Infos.IsEmpty())
+	{
+		Infos.Emplace( CreateArgumentInfo(ArrayName, ERigVMPinDirection::IO) );
+		Infos.Emplace( CreateArgumentInfo(IndexName, ERigVMPinDirection::Input) );
+	}
+	return Infos;
 }
 
 FRigVMTemplateTypeMap FRigVMDispatch_ArrayRemove::OnNewArgumentType(const FName& InArgumentName, TRigVMTypeIndex InTypeIndex) const
@@ -873,15 +889,17 @@ FName FRigVMDispatch_ArrayFind::GetArgumentNameForOperandIndex(int32 InOperandIn
 	return ArgumentNames[InOperandIndex];
 }
 
-const TArray<FRigVMTemplateArgument>& FRigVMDispatch_ArrayFind::GetArguments() const
+const TArray<FRigVMTemplateArgumentInfo>& FRigVMDispatch_ArrayFind::GetArgumentInfos() const
 {
-	static const TArray<FRigVMTemplateArgument> Arguments = {
-		CreateArgument(ArrayName, ERigVMPinDirection::Input),
-		CreateArgument(ElementName, ERigVMPinDirection::Input),
-		CreateArgument(IndexName, ERigVMPinDirection::Output),
-		CreateArgument(SuccessName, ERigVMPinDirection::Output)
-	};
-	return Arguments;
+	static TArray<FRigVMTemplateArgumentInfo> Infos;
+	if (Infos.IsEmpty())
+	{
+		Infos.Emplace(CreateArgumentInfo(ArrayName, ERigVMPinDirection::Input));
+		Infos.Emplace(CreateArgumentInfo(ElementName, ERigVMPinDirection::Input));
+		Infos.Emplace(CreateArgumentInfo(IndexName, ERigVMPinDirection::Output));
+		Infos.Emplace(CreateArgumentInfo(SuccessName, ERigVMPinDirection::Output));
+	}
+	return Infos;
 }
 
 FRigVMTemplateTypeMap FRigVMDispatch_ArrayFind::OnNewArgumentType(const FName& InArgumentName, TRigVMTypeIndex InTypeIndex) const
@@ -983,15 +1001,18 @@ FName FRigVMDispatch_ArrayAppend::GetArgumentNameForOperandIndex(int32 InOperand
 	return ArgumentNames[InOperandIndex];
 }
 
-const TArray<FRigVMTemplateArgument>& FRigVMDispatch_ArrayAppend::GetArguments() const
+const TArray<FRigVMTemplateArgumentInfo>& FRigVMDispatch_ArrayAppend::GetArgumentInfos() const
 {
-	static const TArray<FRigVMTemplateArgument> Arguments = {
-		CreateArgument(ArrayName, ERigVMPinDirection::IO),
-		CreateArgument(OtherName, ERigVMPinDirection::Input)
-	};
-	return Arguments;
+	static TArray<FRigVMTemplateArgumentInfo> Infos;
+	if (Infos.IsEmpty())
+	{
+		Infos.Emplace(CreateArgumentInfo(ArrayName, ERigVMPinDirection::IO));
+		Infos.Emplace(CreateArgumentInfo(OtherName, ERigVMPinDirection::Input));
+	}
+	return Infos;
 }
 
+
 FRigVMTemplateTypeMap FRigVMDispatch_ArrayAppend::OnNewArgumentType(const FName& InArgumentName, TRigVMTypeIndex InTypeIndex) const
 {
 	check(InArgumentName == ArrayName || InArgumentName == OtherName);
@@ -1063,13 +1084,15 @@ FName FRigVMDispatch_ArrayClone::GetArgumentNameForOperandIndex(int32 InOperandI
 	return ArgumentNames[InOperandIndex];
 }
 
-const TArray<FRigVMTemplateArgument>& FRigVMDispatch_ArrayClone::GetArguments() const
+const TArray<FRigVMTemplateArgumentInfo>& FRigVMDispatch_ArrayClone::GetArgumentInfos() const
 {
-	static const TArray<FRigVMTemplateArgument> Arguments = {
-		CreateArgument(ArrayName, ERigVMPinDirection::Input),
-		CreateArgument(CloneName, ERigVMPinDirection::Output)
-	};
-	return Arguments;
+	static TArray<FRigVMTemplateArgumentInfo> Infos;
+	if (Infos.IsEmpty())
+	{
+		Infos.Emplace(CreateArgumentInfo(ArrayName, ERigVMPinDirection::Input));
+		Infos.Emplace(CreateArgumentInfo(CloneName, ERigVMPinDirection::Output));
+	}
+	return Infos;
 }
 
 FRigVMTemplateTypeMap FRigVMDispatch_ArrayClone::OnNewArgumentType(const FName& InArgumentName, TRigVMTypeIndex InTypeIndex) const
@@ -1244,14 +1267,16 @@ FName FRigVMDispatch_ArrayDifference::GetArgumentNameForOperandIndex(int32 InOpe
 	return ArgumentNames[InOperandIndex];
 }
 
-const TArray<FRigVMTemplateArgument>& FRigVMDispatch_ArrayDifference::GetArguments() const
+const TArray<FRigVMTemplateArgumentInfo>& FRigVMDispatch_ArrayDifference::GetArgumentInfos() const
 {
-	static const TArray<FRigVMTemplateArgument> Arguments = {
-		CreateArgument(ArrayName, ERigVMPinDirection::Input),
-		CreateArgument(OtherName, ERigVMPinDirection::Input),
-		CreateArgument(ResultName, ERigVMPinDirection::Output)
-	};
-	return Arguments;
+	static TArray<FRigVMTemplateArgumentInfo> Infos;
+	if (Infos.IsEmpty())
+	{
+		Infos.Emplace(CreateArgumentInfo(ArrayName, ERigVMPinDirection::Input));
+		Infos.Emplace(CreateArgumentInfo(OtherName, ERigVMPinDirection::Input));
+		Infos.Emplace(CreateArgumentInfo(ResultName, ERigVMPinDirection::Output));
+	}
+	return Infos;
 }
 
 FRigVMTemplateTypeMap FRigVMDispatch_ArrayDifference::OnNewArgumentType(const FName& InArgumentName, TRigVMTypeIndex InTypeIndex) const
@@ -1413,17 +1438,19 @@ FName FRigVMDispatch_ArrayIterator::GetArgumentNameForOperandIndex(int32 InOpera
 	return ArgumentNames[InOperandIndex];
 }
 
-const TArray<FRigVMTemplateArgument>& FRigVMDispatch_ArrayIterator::GetArguments() const
+const TArray<FRigVMTemplateArgumentInfo>& FRigVMDispatch_ArrayIterator::GetArgumentInfos() const
 {
-	static const TArray<FRigVMTemplateArgument> Arguments = {
-		CreateArgument(ArrayName, ERigVMPinDirection::Input),
-		CreateArgument(ElementName, ERigVMPinDirection::Output),
-		CreateArgument(IndexName, ERigVMPinDirection::Output),
-		CreateArgument(CountName, ERigVMPinDirection::Output),
-		CreateArgument(RatioName, ERigVMPinDirection::Output),
-		{FRigVMStruct::ControlFlowBlockToRunName, ERigVMPinDirection::Hidden, RigVMTypeUtils::TypeIndex::FName}
-	};
-	return Arguments;
+	static TArray<FRigVMTemplateArgumentInfo> Infos;
+	if (Infos.IsEmpty())
+	{
+		Infos.Emplace(CreateArgumentInfo(ArrayName, ERigVMPinDirection::Input));
+		Infos.Emplace(CreateArgumentInfo(ElementName, ERigVMPinDirection::Output));
+		Infos.Emplace(CreateArgumentInfo(IndexName, ERigVMPinDirection::Output));
+		Infos.Emplace(CreateArgumentInfo(CountName, ERigVMPinDirection::Output));
+		Infos.Emplace(CreateArgumentInfo(RatioName, ERigVMPinDirection::Output));
+		Infos.Emplace(FRigVMStruct::ControlFlowBlockToRunName, ERigVMPinDirection::Hidden, RigVMTypeUtils::TypeIndex::FName);
+	}
+	return Infos;
 }
 
 const TArray<FRigVMExecuteArgument>& FRigVMDispatch_ArrayIterator::GetExecuteArguments_Impl(
diff --git a/Engine/Plugins/Runtime/RigVM/Source/RigVM/Private/RigVMFunctions/RigVMDispatch_CastObject.cpp b/Engine/Plugins/Runtime/RigVM/Source/RigVM/Private/RigVMFunctions/RigVMDispatch_CastObject.cpp
--- a/Engine/Plugins/Runtime/RigVM/Source/RigVM/Private/RigVMFunctions/RigVMDispatch_CastObject.cpp
+++ b/Engine/Plugins/Runtime/RigVM/Source/RigVM/Private/RigVMFunctions/RigVMDispatch_CastObject.cpp
@@ -8,21 +8,24 @@
 const FName FRigVMDispatch_CastObject::ValueName = TEXT("Value");
 const FName FRigVMDispatch_CastObject::ResultName = TEXT("Result");
 
-const TArray<FRigVMTemplateArgument>& FRigVMDispatch_CastObject::GetArguments() const
+const TArray<FRigVMTemplateArgumentInfo>& FRigVMDispatch_CastObject::GetArgumentInfos() const
 {
-	static const TArray<FRigVMTemplateArgument::ETypeCategory> ElementCategories =
+	static TArray<FRigVMTemplateArgumentInfo> OutInfos;
+	if (OutInfos.IsEmpty())
 	{
-		FRigVMTemplateArgument::ETypeCategory_SingleObjectValue
-	};
-
-	static const TArray<FRigVMTemplateArgument> Arguments = BuildArgumentListFromPrimaryArgument(
+		static const TArray<FRigVMTemplateArgument::ETypeCategory> ElementCategories =
 		{
-			FRigVMTemplateArgument(ValueName, ERigVMPinDirection::Input, ElementCategories),
-			FRigVMTemplateArgument(ResultName, ERigVMPinDirection::Output)
-		},
-		ValueName);
+			FRigVMTemplateArgument::ETypeCategory_SingleObjectValue
+		};
+		
+		static TArray<FRigVMTemplateArgumentInfo> Infos;
+		Infos.Emplace(ValueName, ERigVMPinDirection::Input, ElementCategories);
+		Infos.Emplace(ResultName, ERigVMPinDirection::Output, INDEX_NONE);
+		
+		OutInfos = BuildArgumentListFromPrimaryArgument(Infos, ValueName);
+	}
 
-	return Arguments;
+	return OutInfos;
 }
 
 TArray<FRigVMTemplateTypeMap> FRigVMDispatch_CastObject::GetPermutationsFromArgumentType(const FName& InArgumentName, const TRigVMTypeIndex& InTypeIndex) const
diff --git a/Engine/Plugins/Runtime/RigVM/Source/RigVM/Private/RigVMFunctions/RigVMDispatch_Constant.cpp b/Engine/Plugins/Runtime/RigVM/Source/RigVM/Private/RigVMFunctions/RigVMDispatch_Constant.cpp
--- a/Engine/Plugins/Runtime/RigVM/Source/RigVM/Private/RigVMFunctions/RigVMDispatch_Constant.cpp
+++ b/Engine/Plugins/Runtime/RigVM/Source/RigVM/Private/RigVMFunctions/RigVMDispatch_Constant.cpp
@@ -10,15 +10,18 @@
 
 const FName FRigVMDispatch_Constant::ValueName = TEXT("Value");
 
-const TArray<FRigVMTemplateArgument>& FRigVMDispatch_Constant::GetArguments() const
+const TArray<FRigVMTemplateArgumentInfo>& FRigVMDispatch_Constant::GetArgumentInfos() const
 {
-	static const TArray<FRigVMTemplateArgument::ETypeCategory> Categories = {
-		FRigVMTemplateArgument::ETypeCategory_SingleAnyValue
-	};
-	static const TArray<FRigVMTemplateArgument> Arguments = {
-		{ValueName, ERigVMPinDirection::IO, Categories}
-	};
-	return Arguments;
+	static TArray<FRigVMTemplateArgumentInfo> Infos;
+	if (Infos.IsEmpty())
+	{
+		static const TArray<FRigVMTemplateArgument::ETypeCategory> Categories = {
+			FRigVMTemplateArgument::ETypeCategory_SingleAnyValue
+		};
+		
+		Infos.Emplace(FRigVMTemplateArgumentInfo(ValueName, ERigVMPinDirection::IO, Categories));
+	}
+	return Infos;
 }
 
 FRigVMTemplateTypeMap FRigVMDispatch_Constant::OnNewArgumentType(const FName& InArgumentName, TRigVMTypeIndex InTypeIndex) const
diff --git a/Engine/Plugins/Runtime/RigVM/Source/RigVM/Private/RigVMFunctions/RigVMDispatch_Core.cpp b/Engine/Plugins/Runtime/RigVM/Source/RigVM/Private/RigVMFunctions/RigVMDispatch_Core.cpp
--- a/Engine/Plugins/Runtime/RigVM/Source/RigVM/Private/RigVMFunctions/RigVMDispatch_Core.cpp
+++ b/Engine/Plugins/Runtime/RigVM/Source/RigVM/Private/RigVMFunctions/RigVMDispatch_Core.cpp
@@ -23,18 +23,20 @@ FName FRigVMDispatch_CoreEquals::GetArgumentNameForOperandIndex(int32 InOperandI
 	return ArgumentNames[InOperandIndex];
 }
 
-const TArray<FRigVMTemplateArgument>& FRigVMDispatch_CoreEquals::GetArguments() const
+const TArray<FRigVMTemplateArgumentInfo>& FRigVMDispatch_CoreEquals::GetArgumentInfos() const
 {
-	static const TArray<FRigVMTemplateArgument::ETypeCategory> ValueCategories = {
-		FRigVMTemplateArgument::ETypeCategory_SingleAnyValue,
-		FRigVMTemplateArgument::ETypeCategory_ArrayAnyValue
-	};
-	static const TArray<FRigVMTemplateArgument> Arguments = {
-		FRigVMTemplateArgument(AName, ERigVMPinDirection::Input, ValueCategories),
-		FRigVMTemplateArgument(BName, ERigVMPinDirection::Input, ValueCategories),
-		FRigVMTemplateArgument(ResultName, ERigVMPinDirection::Output, RigVMTypeUtils::TypeIndex::Bool)
-	};
-	return Arguments;
+	static TArray<FRigVMTemplateArgumentInfo> Infos;
+	if (Infos.IsEmpty())
+	{
+		static const TArray<FRigVMTemplateArgument::ETypeCategory> ValueCategories = {
+			FRigVMTemplateArgument::ETypeCategory_SingleAnyValue,
+			FRigVMTemplateArgument::ETypeCategory_ArrayAnyValue
+		};
+		Infos.Emplace(AName, ERigVMPinDirection::Input, ValueCategories);
+		Infos.Emplace(BName, ERigVMPinDirection::Input, ValueCategories);
+		Infos.Emplace(ResultName, ERigVMPinDirection::Output, RigVMTypeUtils::TypeIndex::Bool);
+	}
+	return Infos;
 }
 
 FRigVMTemplateTypeMap FRigVMDispatch_CoreEquals::OnNewArgumentType(const FName& InArgumentName,
@@ -128,16 +130,18 @@ void FRigVMDispatch_CoreEquals::Execute(FRigVMExtendedExecuteContext& InContext,
 
 // duplicate the code here so that the FRigVMDispatch_CoreNotEquals has it's own static variables
 // to store the types are registration time.
-const TArray<FRigVMTemplateArgument>& FRigVMDispatch_CoreNotEquals::GetArguments() const
+const TArray<FRigVMTemplateArgumentInfo>& FRigVMDispatch_CoreNotEquals::GetArgumentInfos() const
 {
-	static const TArray<FRigVMTemplateArgument::ETypeCategory> ValueCategories = {
-		FRigVMTemplateArgument::ETypeCategory_SingleAnyValue,
-		FRigVMTemplateArgument::ETypeCategory_ArrayAnyValue
-	};
-	static const TArray<FRigVMTemplateArgument> Arguments = {
-		FRigVMTemplateArgument(AName, ERigVMPinDirection::Input, ValueCategories),
-		FRigVMTemplateArgument(BName, ERigVMPinDirection::Input, ValueCategories),
-		FRigVMTemplateArgument(ResultName, ERigVMPinDirection::Output, RigVMTypeUtils::TypeIndex::Bool)
-	};
-	return Arguments;
+	static TArray<FRigVMTemplateArgumentInfo> Infos;
+	if (Infos.IsEmpty())
+	{
+		static const TArray<FRigVMTemplateArgument::ETypeCategory> ValueCategories = {
+			FRigVMTemplateArgument::ETypeCategory_SingleAnyValue,
+			FRigVMTemplateArgument::ETypeCategory_ArrayAnyValue
+		};
+		Infos.Emplace(AName, ERigVMPinDirection::Input, ValueCategories);
+		Infos.Emplace(BName, ERigVMPinDirection::Input, ValueCategories);
+		Infos.Emplace(ResultName, ERigVMPinDirection::Output, RigVMTypeUtils::TypeIndex::Bool);
+	}
+	return Infos;
 }
diff --git a/Engine/Plugins/Runtime/RigVM/Source/RigVM/Private/RigVMFunctions/RigVMDispatch_If.cpp b/Engine/Plugins/Runtime/RigVM/Source/RigVM/Private/RigVMFunctions/RigVMDispatch_If.cpp
--- a/Engine/Plugins/Runtime/RigVM/Source/RigVM/Private/RigVMFunctions/RigVMDispatch_If.cpp
+++ b/Engine/Plugins/Runtime/RigVM/Source/RigVM/Private/RigVMFunctions/RigVMDispatch_If.cpp
@@ -23,19 +23,22 @@ FName FRigVMDispatch_If::GetArgumentNameForOperandIndex(int32 InOperandIndex, in
 	return ArgumentNames[InOperandIndex];
 }
 
-const TArray<FRigVMTemplateArgument>& FRigVMDispatch_If::GetArguments() const
+const TArray<FRigVMTemplateArgumentInfo>& FRigVMDispatch_If::GetArgumentInfos() const
 {
-	static const TArray<FRigVMTemplateArgument::ETypeCategory> ValueCategories = {
-		FRigVMTemplateArgument::ETypeCategory_SingleAnyValue,
-		FRigVMTemplateArgument::ETypeCategory_ArrayAnyValue
-	};
-	static const TArray<FRigVMTemplateArgument> Arguments = {
-		FRigVMTemplateArgument(ConditionName, ERigVMPinDirection::Input, RigVMTypeUtils::TypeIndex::Bool),
-		FRigVMTemplateArgument(TrueName, ERigVMPinDirection::Input, ValueCategories),
-		FRigVMTemplateArgument(FalseName, ERigVMPinDirection::Input, ValueCategories),
-		FRigVMTemplateArgument(ResultName, ERigVMPinDirection::Output, ValueCategories)
-	};
-	return Arguments;
+	static TArray<FRigVMTemplateArgumentInfo> Infos;
+	if(Infos.IsEmpty())
+	{
+		static const TArray<FRigVMTemplateArgument::ETypeCategory> ValueCategories = {
+			FRigVMTemplateArgument::ETypeCategory_SingleAnyValue,
+			FRigVMTemplateArgument::ETypeCategory_ArrayAnyValue
+		};
+		
+		Infos.Emplace(ConditionName, ERigVMPinDirection::Input, RigVMTypeUtils::TypeIndex::Bool);
+		Infos.Emplace(TrueName, ERigVMPinDirection::Input, ValueCategories);
+		Infos.Emplace(FalseName, ERigVMPinDirection::Input, ValueCategories);
+		Infos.Emplace(ResultName, ERigVMPinDirection::Output, ValueCategories);
+	}
+	return Infos;
 }
 
 FRigVMTemplateTypeMap FRigVMDispatch_If::OnNewArgumentType(const FName& InArgumentName,
diff --git a/Engine/Plugins/Runtime/RigVM/Source/RigVM/Private/RigVMFunctions/RigVMDispatch_MakeStruct.cpp b/Engine/Plugins/Runtime/RigVM/Source/RigVM/Private/RigVMFunctions/RigVMDispatch_MakeStruct.cpp
--- a/Engine/Plugins/Runtime/RigVM/Source/RigVM/Private/RigVMFunctions/RigVMDispatch_MakeStruct.cpp
+++ b/Engine/Plugins/Runtime/RigVM/Source/RigVM/Private/RigVMFunctions/RigVMDispatch_MakeStruct.cpp
@@ -11,16 +11,19 @@
 const FName FRigVMDispatch_MakeStruct::ElementsName = TEXT("Elements");
 const FName FRigVMDispatch_MakeStruct::StructName = TEXT("Struct");
 
-const TArray<FRigVMTemplateArgument>& FRigVMDispatch_MakeStruct::GetArguments() const
+const TArray<FRigVMTemplateArgumentInfo>& FRigVMDispatch_MakeStruct::GetArgumentInfos() const
 {
-	static const TArray<FRigVMTemplateArgument::ETypeCategory> Categories = {
-		FRigVMTemplateArgument::ETypeCategory_SingleScriptStructValue
-	};
-	static TArray<FRigVMTemplateArgument> Arguments = {
-		{ElementsName, ERigVMPinDirection::Input, Categories},
-		{StructName, ERigVMPinDirection::Output, Categories}
-	};
-	return Arguments;
+	static TArray<FRigVMTemplateArgumentInfo> Infos;
+	if(Infos.IsEmpty())
+	{
+		static const TArray<FRigVMTemplateArgument::ETypeCategory> Categories = {
+			FRigVMTemplateArgument::ETypeCategory_SingleScriptStructValue
+		};
+		
+		Infos.Emplace(ElementsName, ERigVMPinDirection::Input, Categories);
+		Infos.Emplace(StructName, ERigVMPinDirection::Output, Categories);
+	}
+	return Infos;
 }
 
 FRigVMTemplateTypeMap FRigVMDispatch_MakeStruct::OnNewArgumentType(const FName& InArgumentName, TRigVMTypeIndex InTypeIndex) const
@@ -102,16 +105,19 @@ void FRigVMDispatch_MakeStruct::Execute(FRigVMExtendedExecuteContext& InContext,
 	URigVMMemoryStorage::CopyProperty(TargetProperty, TargetMemory, SourceProperty, SourceMemory);
 }
 
-const TArray<FRigVMTemplateArgument>& FRigVMDispatch_BreakStruct::GetArguments() const
+const TArray<FRigVMTemplateArgumentInfo>& FRigVMDispatch_BreakStruct::GetArgumentInfos() const
 {
-	static const TArray<FRigVMTemplateArgument::ETypeCategory> Categories = {
-		FRigVMTemplateArgument::ETypeCategory_SingleScriptStructValue
-	};
-	static const TArray<FRigVMTemplateArgument> Arguments = {
-		{StructName, ERigVMPinDirection::Input, Categories},
-		{ElementsName, ERigVMPinDirection::Output, Categories}
-	};
-	return Arguments;
+	static TArray<FRigVMTemplateArgumentInfo> Infos;
+	if(Infos.IsEmpty())
+	{
+		static const TArray<FRigVMTemplateArgument::ETypeCategory> Categories = {
+			FRigVMTemplateArgument::ETypeCategory_SingleScriptStructValue
+		};
+
+		Infos.Emplace(StructName, ERigVMPinDirection::Input, Categories);
+		Infos.Emplace(ElementsName, ERigVMPinDirection::Output, Categories);
+	}
+	return Infos;
 }
 
 #if WITH_EDITOR
diff --git a/Engine/Plugins/Runtime/RigVM/Source/RigVM/Private/RigVMFunctions/RigVMDispatch_Print.cpp b/Engine/Plugins/Runtime/RigVM/Source/RigVM/Private/RigVMFunctions/RigVMDispatch_Print.cpp
--- a/Engine/Plugins/Runtime/RigVM/Source/RigVM/Private/RigVMFunctions/RigVMDispatch_Print.cpp
+++ b/Engine/Plugins/Runtime/RigVM/Source/RigVM/Private/RigVMFunctions/RigVMDispatch_Print.cpp
@@ -28,20 +28,23 @@ FName FRigVMDispatch_Print::GetArgumentNameForOperandIndex(int32 InOperandIndex,
 	return ArgumentNames[InOperandIndex];
 }
 
-const TArray<FRigVMTemplateArgument>& FRigVMDispatch_Print::GetArguments() const
+const TArray<FRigVMTemplateArgumentInfo>& FRigVMDispatch_Print::GetArgumentInfos() const
 {
-	static const TArray<FRigVMTemplateArgument::ETypeCategory> ValueCategories = {
-		FRigVMTemplateArgument::ETypeCategory_SingleAnyValue,
-		FRigVMTemplateArgument::ETypeCategory_ArrayAnyValue
-	};
-	static const TArray<FRigVMTemplateArgument> Arguments = {
-		FRigVMTemplateArgument(PrefixName, ERigVMPinDirection::Input, RigVMTypeUtils::TypeIndex::FString),
-		FRigVMTemplateArgument(ValueName, ERigVMPinDirection::Input, ValueCategories),
-		FRigVMTemplateArgument(EnabledName, ERigVMPinDirection::Input, RigVMTypeUtils::TypeIndex::Bool),
-		FRigVMTemplateArgument(ScreenDurationName, ERigVMPinDirection::Input, RigVMTypeUtils::TypeIndex::Float),
-		FRigVMTemplateArgument(ScreenColorName, ERigVMPinDirection::Input, FRigVMRegistry::Get().GetTypeIndex<FLinearColor>())
-	};
-	return Arguments;
+	static TArray<FRigVMTemplateArgumentInfo> Infos;
+	if (Infos.IsEmpty())
+	{
+		static const TArray<FRigVMTemplateArgument::ETypeCategory> ValueCategories = {
+			FRigVMTemplateArgument::ETypeCategory_SingleAnyValue,
+			FRigVMTemplateArgument::ETypeCategory_ArrayAnyValue
+		};
+	
+		Infos.Emplace(PrefixName, ERigVMPinDirection::Input, RigVMTypeUtils::TypeIndex::FString);
+		Infos.Emplace(ValueName, ERigVMPinDirection::Input, ValueCategories);
+		Infos.Emplace(EnabledName, ERigVMPinDirection::Input, RigVMTypeUtils::TypeIndex::Bool);
+		Infos.Emplace(ScreenDurationName, ERigVMPinDirection::Input, RigVMTypeUtils::TypeIndex::Float);
+		Infos.Emplace(ScreenColorName, ERigVMPinDirection::Input, FRigVMRegistry::Get().GetTypeIndex<FLinearColor>());
+	}
+	return Infos;
 }
 
 const TArray<FRigVMExecuteArgument>& FRigVMDispatch_Print::GetExecuteArguments_Impl(const FRigVMDispatchContext& InContext) const
diff --git a/Engine/Plugins/Runtime/RigVM/Source/RigVM/Private/RigVMFunctions/RigVMDispatch_Select.cpp b/Engine/Plugins/Runtime/RigVM/Source/RigVM/Private/RigVMFunctions/RigVMDispatch_Select.cpp
--- a/Engine/Plugins/Runtime/RigVM/Source/RigVM/Private/RigVMFunctions/RigVMDispatch_Select.cpp
+++ b/Engine/Plugins/Runtime/RigVM/Source/RigVM/Private/RigVMFunctions/RigVMDispatch_Select.cpp
@@ -23,24 +23,26 @@ FName FRigVMDispatch_SelectInt32::GetArgumentNameForOperandIndex(int32 InOperand
 }
 
 
-const TArray<FRigVMTemplateArgument>& FRigVMDispatch_SelectInt32::GetArguments() const
+const TArray<FRigVMTemplateArgumentInfo>& FRigVMDispatch_SelectInt32::GetArgumentInfos() const
 {
-	static const TArray<FRigVMTemplateArgument::ETypeCategory> ArrayCategories = {
-		FRigVMTemplateArgument::ETypeCategory_ArrayAnyValue,
-		FRigVMTemplateArgument::ETypeCategory_ArrayArrayAnyValue
-	};
-	static const TArray<FRigVMTemplateArgument::ETypeCategory> SingleCategories = {
-		FRigVMTemplateArgument::ETypeCategory_SingleAnyValue,
-		FRigVMTemplateArgument::ETypeCategory_ArrayAnyValue
-	};
-	
-	static const TArray<FRigVMTemplateArgument> Arguments = {
-		FRigVMTemplateArgument(IndexName, ERigVMPinDirection::Input, RigVMTypeUtils::TypeIndex::Int32),
-		FRigVMTemplateArgument(ValuesName, ERigVMPinDirection::Input, ArrayCategories),
-		FRigVMTemplateArgument(ResultName, ERigVMPinDirection::Output, SingleCategories)
-	};
+	static TArray<FRigVMTemplateArgumentInfo> Infos;
+	if (Infos.IsEmpty())
+	{
+		static const TArray<FRigVMTemplateArgument::ETypeCategory> ArrayCategories = {
+			FRigVMTemplateArgument::ETypeCategory_ArrayAnyValue,
+			FRigVMTemplateArgument::ETypeCategory_ArrayArrayAnyValue
+		};
+		static const TArray<FRigVMTemplateArgument::ETypeCategory> SingleCategories = {
+			FRigVMTemplateArgument::ETypeCategory_SingleAnyValue,
+			FRigVMTemplateArgument::ETypeCategory_ArrayAnyValue
+		};
+		
+		Infos.Emplace(IndexName, ERigVMPinDirection::Input, RigVMTypeUtils::TypeIndex::Int32);
+		Infos.Emplace(ValuesName, ERigVMPinDirection::Input, ArrayCategories);
+		Infos.Emplace(ResultName, ERigVMPinDirection::Output, SingleCategories);
+	}
 
-	return Arguments;
+	return Infos;
 }
 
 FRigVMTemplateTypeMap FRigVMDispatch_SelectInt32::OnNewArgumentType(const FName& InArgumentName,
diff --git a/Engine/Plugins/Runtime/RigVM/Source/RigVM/Private/RigVMFunctions/RigVMDispatch_Switch.cpp b/Engine/Plugins/Runtime/RigVM/Source/RigVM/Private/RigVMFunctions/RigVMDispatch_Switch.cpp
--- a/Engine/Plugins/Runtime/RigVM/Source/RigVM/Private/RigVMFunctions/RigVMDispatch_Switch.cpp
+++ b/Engine/Plugins/Runtime/RigVM/Source/RigVM/Private/RigVMFunctions/RigVMDispatch_Switch.cpp
@@ -20,12 +20,14 @@ FName FRigVMDispatch_SwitchInt32::GetArgumentNameForOperandIndex(int32 InOperand
 	return ArgumentNames[InOperandIndex];
 }
 
-const TArray<FRigVMTemplateArgument>& FRigVMDispatch_SwitchInt32::GetArguments() const
+const TArray<FRigVMTemplateArgumentInfo>& FRigVMDispatch_SwitchInt32::GetArgumentInfos() const
 {
-	static const TArray<FRigVMTemplateArgument> Arguments = {
-		FRigVMTemplateArgument(IndexName, ERigVMPinDirection::Input, RigVMTypeUtils::TypeIndex::Int32),
-		FRigVMTemplateArgument(FRigVMStruct::ControlFlowBlockToRunName, ERigVMPinDirection::Hidden, RigVMTypeUtils::TypeIndex::FName)
-	};
+	static TArray<FRigVMTemplateArgumentInfo> Arguments;
+	if(Arguments.IsEmpty())
+	{
+		Arguments.Emplace(IndexName, ERigVMPinDirection::Input, RigVMTypeUtils::TypeIndex::Int32);
+		Arguments.Emplace(FRigVMStruct::ControlFlowBlockToRunName, ERigVMPinDirection::Hidden, RigVMTypeUtils::TypeIndex::FName);
+	}
 	return Arguments;
 }
 
diff --git a/Engine/Plugins/Runtime/RigVM/Source/RigVM/Private/RigVMFunctions/RigVMFunction_String.cpp b/Engine/Plugins/Runtime/RigVM/Source/RigVM/Private/RigVMFunctions/RigVMFunction_String.cpp
--- a/Engine/Plugins/Runtime/RigVM/Source/RigVM/Private/RigVMFunctions/RigVMFunction_String.cpp
+++ b/Engine/Plugins/Runtime/RigVM/Source/RigVM/Private/RigVMFunctions/RigVMFunction_String.cpp
@@ -252,17 +252,20 @@ FRigVMFunction_StringPadInteger_Execute()
 	}
 }
 
-const TArray<FRigVMTemplateArgument>& FRigDispatch_ToString::GetArguments() const
+const TArray<FRigVMTemplateArgumentInfo>& FRigDispatch_ToString::GetArgumentInfos() const
 {
-	static const TArray<FRigVMTemplateArgument::ETypeCategory> ValueCategories = {
-		FRigVMTemplateArgument::ETypeCategory_SingleAnyValue,
-		FRigVMTemplateArgument::ETypeCategory_ArrayAnyValue
-	};
-	static const TArray<FRigVMTemplateArgument> Arguments = {
-		FRigVMTemplateArgument(TEXT("Value"), ERigVMPinDirection::Input, ValueCategories),
-		FRigVMTemplateArgument(TEXT("Result"), ERigVMPinDirection::Output, RigVMTypeUtils::TypeIndex::FString)
-	};
-	return Arguments;
+	static TArray<FRigVMTemplateArgumentInfo> Infos;
+	if(Infos.IsEmpty())
+	{
+		static const TArray<FRigVMTemplateArgument::ETypeCategory> ValueCategories = {
+			FRigVMTemplateArgument::ETypeCategory_SingleAnyValue,
+			FRigVMTemplateArgument::ETypeCategory_ArrayAnyValue
+		};
+		
+		Infos.Emplace(TEXT("Value"), ERigVMPinDirection::Input, ValueCategories);
+		Infos.Emplace(TEXT("Result"), ERigVMPinDirection::Output, RigVMTypeUtils::TypeIndex::FString);
+	}
+	return Infos;
 }
 
 FRigVMTemplateTypeMap FRigDispatch_ToString::OnNewArgumentType(const FName& InArgumentName,
@@ -318,17 +321,20 @@ void FRigDispatch_ToString::Execute(FRigVMExtendedExecuteContext& InContext, FRi
 	ValueProperty->ExportText_Direct(Result, Value, Value, nullptr, PPF_None, nullptr);
 }
 
-const TArray<FRigVMTemplateArgument>& FRigDispatch_FromString::GetArguments() const
+const TArray<FRigVMTemplateArgumentInfo>& FRigDispatch_FromString::GetArgumentInfos() const
 {
-	static const TArray<FRigVMTemplateArgument::ETypeCategory> ValueCategories = {
-		FRigVMTemplateArgument::ETypeCategory_SingleAnyValue,
-		FRigVMTemplateArgument::ETypeCategory_ArrayAnyValue
-	};
-	static const TArray<FRigVMTemplateArgument> Arguments = {
-		FRigVMTemplateArgument(TEXT("String"), ERigVMPinDirection::Input, RigVMTypeUtils::TypeIndex::FString),
-		FRigVMTemplateArgument(TEXT("Result"), ERigVMPinDirection::Output, ValueCategories)
-	};
-	return Arguments;
+	static TArray<FRigVMTemplateArgumentInfo> Infos;
+	if(Infos.IsEmpty())
+	{
+		static const TArray<FRigVMTemplateArgument::ETypeCategory> ValueCategories = {
+			FRigVMTemplateArgument::ETypeCategory_SingleAnyValue,
+			FRigVMTemplateArgument::ETypeCategory_ArrayAnyValue
+		};
+		
+		Infos.Emplace(TEXT("String"), ERigVMPinDirection::Input, RigVMTypeUtils::TypeIndex::FString);
+		Infos.Emplace(TEXT("Result"), ERigVMPinDirection::Output, ValueCategories);
+	}
+	return Infos;
 }
 
 FRigVMTemplateTypeMap FRigDispatch_FromString::OnNewArgumentType(const FName& InArgumentName,
diff --git a/Engine/Plugins/Runtime/RigVM/Source/RigVM/Public/RigVMCore/RigVMDispatchFactory.h b/Engine/Plugins/Runtime/RigVM/Source/RigVM/Public/RigVMCore/RigVMDispatchFactory.h
--- a/Engine/Plugins/Runtime/RigVM/Source/RigVM/Public/RigVMCore/RigVMDispatchFactory.h
+++ b/Engine/Plugins/Runtime/RigVM/Source/RigVM/Public/RigVMCore/RigVMDispatchFactory.h
@@ -141,7 +141,7 @@ public:
 	virtual void RegisterDependencyTypes() const {}
 
 	// returns the arguments of the template
-	RIGVM_API virtual const TArray<FRigVMTemplateArgument>& GetArguments() const;
+	RIGVM_API virtual const TArray<FRigVMTemplateArgumentInfo>& GetArgumentInfos() const;
 
 	// returns the execute arguments of the template
 	RIGVM_API TArray<FRigVMExecuteArgument> GetExecuteArguments(const FRigVMDispatchContext& InContext) const;
@@ -183,7 +183,7 @@ public:
 protected:
 
 	// for each type defined in the primary argument, this function will call GetPermutationsFromArgumentType to construct an array of arguments with the appropiate permutations
-	RIGVM_API TArray<FRigVMTemplateArgument> BuildArgumentListFromPrimaryArgument(const TArray<FRigVMTemplateArgument>& InArguments, const FName& InPrimaryArgumentName) const;
+	RIGVM_API TArray<FRigVMTemplateArgumentInfo> BuildArgumentListFromPrimaryArgument(const TArray<FRigVMTemplateArgumentInfo>& InInfos, const FName& InPrimaryArgumentName) const;
 
 	// returns the name of the permutation for a given set of types
 	RIGVM_API FString GetPermutationNameImpl(const FRigVMTemplateTypeMap& InTypes) const;
diff --git a/Engine/Plugins/Runtime/RigVM/Source/RigVM/Public/RigVMCore/RigVMRegistry.h b/Engine/Plugins/Runtime/RigVM/Source/RigVM/Public/RigVMCore/RigVMRegistry.h
--- a/Engine/Plugins/Runtime/RigVM/Source/RigVM/Public/RigVMCore/RigVMRegistry.h
+++ b/Engine/Plugins/Runtime/RigVM/Source/RigVM/Public/RigVMCore/RigVMRegistry.h
@@ -235,13 +235,13 @@ public:
 	// Defines and retrieves a template given its arguments
 	const FRigVMTemplate* GetOrAddTemplateFromArguments(
 		const FName& InName,
-		const TArray<FRigVMTemplateArgument>& InArguments,
+		const TArray<FRigVMTemplateArgumentInfo>& InInfos,
 		const FRigVMTemplateDelegates& InDelegates);
 
 	// Adds a new template given its arguments
 	const FRigVMTemplate* AddTemplateFromArguments(
 		const FName& InName,
-		const TArray<FRigVMTemplateArgument>& InArguments,
+		const TArray<FRigVMTemplateArgumentInfo>& InInfos,
 		const FRigVMTemplateDelegates& InDelegates);
 
 	// Returns a dispatch factory given its name (or nullptr)
@@ -348,7 +348,7 @@ private:
 
 	const FRigVMTemplate* AddTemplateFromArguments_NoLock(
 		const FName& InName,
-		const TArray<FRigVMTemplateArgument>& InArguments,
+		const TArray<FRigVMTemplateArgumentInfo>& InInfos,
 		const FRigVMTemplateDelegates& InDelegates);
 
 	// memory for all (known) types
diff --git a/Engine/Plugins/Runtime/RigVM/Source/RigVM/Public/RigVMCore/RigVMTemplate.h b/Engine/Plugins/Runtime/RigVM/Source/RigVM/Public/RigVMCore/RigVMTemplate.h
--- a/Engine/Plugins/Runtime/RigVM/Source/RigVM/Public/RigVMCore/RigVMTemplate.h
+++ b/Engine/Plugins/Runtime/RigVM/Source/RigVM/Public/RigVMCore/RigVMTemplate.h
@@ -35,19 +35,13 @@ struct FRigVMUserDefinedTypeResolver;
 typedef TMap<FName, TRigVMTypeIndex> FRigVMTemplateTypeMap;
 
 // FRigVMTemplate_NewArgumentTypeDelegate is deprecated, use FRigVMTemplate_GetPermutationsFromArgumentTypeDelegate
-DECLARE_DELEGATE_RetVal_ThreeParams(FRigVMTemplateTypeMap, FRigVMTemplate_NewArgumentTypeDelegate, const FRigVMTemplate* /* InTemplate */, const FName& /* InArgumentName */, TRigVMTypeIndex /* InTypeIndexToAdd */);
-DECLARE_DELEGATE_RetVal_ThreeParams(TArray<FRigVMTemplateTypeMap>, FRigVMTemplate_GetPermutationsFromArgumentTypeDelegate, const FRigVMTemplate* /* InTemplate */, const FName& /* InArgumentName */, TRigVMTypeIndex /* InTypeIndexToAdd */);
-DECLARE_DELEGATE_RetVal_TwoParams(FRigVMFunctionPtr, FRigVMTemplate_RequestDispatchFunctionDelegate, const FRigVMTemplate* /* InTemplate */,  const FRigVMTemplateTypeMap& /* InTypes */);
-DECLARE_DELEGATE_RetVal_TwoParams(TArray<FRigVMFunction>, FRigVMTemplate_RequestDispatchPredicatesDelegate, const FRigVMTemplate* /* InTemplate */,  const FRigVMTemplateTypeMap& /* InTypes */);
+DECLARE_DELEGATE_RetVal_TwoParams(FRigVMTemplateTypeMap, FRigVMTemplate_NewArgumentTypeDelegate, const FName& /* InArgumentName */, TRigVMTypeIndex /* InTypeIndexToAdd */);
 DECLARE_DELEGATE_RetVal(FRigVMDispatchFactory*, FRigVMTemplate_GetDispatchFactoryDelegate);
 
 struct RIGVM_API FRigVMTemplateDelegates
 {
 	FRigVMTemplate_NewArgumentTypeDelegate NewArgumentTypeDelegate;
-	FRigVMTemplate_GetPermutationsFromArgumentTypeDelegate GetPermutationsFromArgumentTypeDelegate;
 	FRigVMTemplate_GetDispatchFactoryDelegate GetDispatchFactoryDelegate;
-	FRigVMTemplate_RequestDispatchFunctionDelegate RequestDispatchFunctionDelegate;
-	FRigVMTemplate_RequestDispatchPredicatesDelegate RequestDispatchPredicatesDelegate;
 };
 
 USTRUCT()
@@ -268,6 +262,35 @@ protected:
 	friend class URigVMCompiler;
 };
 
+/**
+ * FRigVMTemplateArgumentInfo 
+ */
+
+struct RIGVM_API FRigVMTemplateArgumentInfo
+{
+	using ArgumentCallback = TFunction<FRigVMTemplateArgument(const FName /*InName*/, ERigVMPinDirection /*InDirection*/)>;
+	
+	FName Name = NAME_None;
+	ERigVMPinDirection Direction = ERigVMPinDirection::Invalid;
+	TFunction<FRigVMTemplateArgument(const FName /*InName*/, ERigVMPinDirection /*InDirection*/)> FactoryCallback = [](const FName, ERigVMPinDirection) { return FRigVMTemplateArgument(); };
+	
+	FRigVMTemplateArgumentInfo(const FName InName, ERigVMPinDirection InDirection, TRigVMTypeIndex InTypeIndex);
+	FRigVMTemplateArgumentInfo(const FName InName, ERigVMPinDirection InDirection, const TArray<TRigVMTypeIndex>& InTypeIndices);
+	FRigVMTemplateArgumentInfo(const FName InName, ERigVMPinDirection InDirection, const TArray<FRigVMTemplateArgument::ETypeCategory>& InTypeCategories);
+	FRigVMTemplateArgumentInfo(const FName InName, ERigVMPinDirection InDirection);
+	FRigVMTemplateArgumentInfo(const FName InName, ERigVMPinDirection InDirection, ArgumentCallback&& InCallback);
+
+	FRigVMTemplateArgument GetArgument() const;
+	
+	static FName ComputeTemplateNotation(
+		const FName InTemplateName,
+		const TArray<FRigVMTemplateArgumentInfo>& InInfos);
+	
+	static TArray<TRigVMTypeIndex> GetTypesFromCategories(
+		const TArray<FRigVMTemplateArgument::ETypeCategory>& InTypeCategories,
+		const FRigVMTemplateArgument::FTypeFilter& InTypeFilter = {});	
+};
+
 /**
  * The template is used to group multiple rigvm functions
  * that share the same notation. Templates can then be used
@@ -367,13 +390,13 @@ public:
 	FRigVMTemplateTypeMap GetTypesForPermutation(const int32 InPermutationIndex) const;
 
 	// returns true if a given argument is valid for a template
-	static bool IsValidArgumentForTemplate(const FRigVMTemplateArgument& InArgument);
+	static bool IsValidArgumentForTemplate(const ERigVMPinDirection InDirection);
 
 	// returns the prefix for an argument in the notation
-	static const FString& GetArgumentNotationPrefix(const FRigVMTemplateArgument& InArgument);
+	static const FString& GetDirectionPrefix(const ERigVMPinDirection InDirection);
 
 	// returns the notation of an argument
-	static FString GetArgumentNotation(const FRigVMTemplateArgument& InArgument);
+	static FString GetArgumentNotation(const FName InName, const ERigVMPinDirection InDirection);
 
 	// recomputes the notation from its arguments
 	void ComputeNotationFromArguments(const FString& InTemplateName);
@@ -417,24 +440,16 @@ public:
 	// This delegate is deprecated
 	FRigVMTemplate_NewArgumentTypeDelegate& OnNewArgumentType() { return Delegates.NewArgumentTypeDelegate; }
 
-	// Returns the delegate to be able to react to type changes dynamically
-	FRigVMTemplate_GetPermutationsFromArgumentTypeDelegate& OnGetPermutationsFromArgumentType() { return Delegates.GetPermutationsFromArgumentTypeDelegate; }
-
 	// Returns the factory this template was created by
 	const FRigVMDispatchFactory* GetDispatchFactory() const
 	{
-		if(Delegates.GetDispatchFactoryDelegate.IsBound())
-		{
-			return Delegates.GetDispatchFactoryDelegate.Execute();
-		}
-		return nullptr;
+		return UsesDispatch() ? Delegates.GetDispatchFactoryDelegate.Execute() : nullptr;
 	}
 
 	// Returns true if this template is backed by a dispatch factory
 	bool UsesDispatch() const
 	{
-		return Delegates.RequestDispatchFunctionDelegate.IsBound() &&
-			Delegates.GetDispatchFactoryDelegate.IsBound();
+		return Delegates.GetDispatchFactoryDelegate.IsBound();
 	}
 
 	void RecomputeTypesHashToPermutations();
@@ -448,8 +463,8 @@ private:
 	// Constructor from a struct, a template name and a function index
 	FRigVMTemplate(UScriptStruct* InStruct, const FString& InTemplateName, int32 InFunctionIndex = INDEX_NONE);
 
-	// Constructor from a template name, arguments and a function index
-	FRigVMTemplate(const FName& InTemplateName, const TArray<FRigVMTemplateArgument>& InArguments);
+	// Constructor from a template name and argument infos
+	FRigVMTemplate(const FName& InTemplateName, const TArray<FRigVMTemplateArgumentInfo>& InInfos);
 
 	static FLinearColor GetColorFromMetadata(FString InMetadata);
 
diff --git a/Engine/Plugins/Runtime/RigVM/Source/RigVM/Public/RigVMFunctions/RigVMDispatch_Array.h b/Engine/Plugins/Runtime/RigVM/Source/RigVM/Public/RigVMFunctions/RigVMDispatch_Array.h
--- a/Engine/Plugins/Runtime/RigVM/Source/RigVM/Public/RigVMFunctions/RigVMDispatch_Array.h
+++ b/Engine/Plugins/Runtime/RigVM/Source/RigVM/Public/RigVMFunctions/RigVMDispatch_Array.h
@@ -22,7 +22,7 @@ public:
 	virtual bool IsSingleton() const override { return true; } 
 
 protected:
-	static FRigVMTemplateArgument CreateArgument(const FName& InName, ERigVMPinDirection InDirection);
+	static FRigVMTemplateArgumentInfo CreateArgumentInfo(const FName& InName, ERigVMPinDirection InDirection);
 	static TMap<uint32, int32> GetArrayHash(FScriptArrayHelper& InArrayHelper, const FArrayProperty* InArrayProperty);
 	
 	static const FName& ExecuteName;
@@ -62,7 +62,7 @@ public:
 		FactoryScriptStruct = StaticStruct();
 	}
 
-	virtual const TArray<FRigVMTemplateArgument>& GetArguments() const override;
+	virtual const TArray<FRigVMTemplateArgumentInfo>& GetArgumentInfos() const override;
 	virtual FRigVMTemplateTypeMap OnNewArgumentType(const FName& InArgumentName, TRigVMTypeIndex InTypeIndex) const override;
 #if WITH_EDITOR
 	virtual FText GetNodeTooltip(const FRigVMTemplateTypeMap& InTypes) const override;
@@ -90,7 +90,7 @@ public:
 
 	virtual ERigVMOpCode GetOpCode() const override { return ERigVMOpCode::ArrayReset; }
 	virtual FName GetArgumentNameForOperandIndex(int32 InOperandIndex, int32 InTotalOperands) const override;
-	virtual const TArray<FRigVMTemplateArgument>& GetArguments() const override;
+	virtual const TArray<FRigVMTemplateArgumentInfo>& GetArgumentInfos() const override;
 	virtual FRigVMTemplateTypeMap OnNewArgumentType(const FName& InArgumentName, TRigVMTypeIndex InTypeIndex) const override;
 #if WITH_EDITOR
 	virtual FText GetNodeTooltip(const FRigVMTemplateTypeMap& InTypes) const override;
@@ -115,7 +115,7 @@ public:
 
 	virtual ERigVMOpCode GetOpCode() const override { return ERigVMOpCode::ArrayGetNum; }
 	virtual FName GetArgumentNameForOperandIndex(int32 InOperandIndex, int32 InTotalOperands) const override;
-	virtual const TArray<FRigVMTemplateArgument>& GetArguments() const override;
+	virtual const TArray<FRigVMTemplateArgumentInfo>& GetArgumentInfos() const override;
 	virtual FRigVMTemplateTypeMap OnNewArgumentType(const FName& InArgumentName, TRigVMTypeIndex InTypeIndex) const override;
 #if WITH_EDITOR
 	virtual FText GetNodeTooltip(const FRigVMTemplateTypeMap& InTypes) const override;
@@ -140,7 +140,7 @@ public:
 
 	virtual ERigVMOpCode GetOpCode() const override { return ERigVMOpCode::ArraySetNum; }
 	virtual FName GetArgumentNameForOperandIndex(int32 InOperandIndex, int32 InTotalOperands) const override;
-	virtual const TArray<FRigVMTemplateArgument>& GetArguments() const override;
+	virtual const TArray<FRigVMTemplateArgumentInfo>& GetArgumentInfos() const override;
 	virtual FRigVMTemplateTypeMap OnNewArgumentType(const FName& InArgumentName, TRigVMTypeIndex InTypeIndex) const override;
 #if WITH_EDITOR
 	virtual FText GetNodeTooltip(const FRigVMTemplateTypeMap& InTypes) const override;
@@ -165,7 +165,7 @@ public:
 
 	virtual ERigVMOpCode GetOpCode() const override { return ERigVMOpCode::ArrayGetAtIndex; }
 	virtual FName GetArgumentNameForOperandIndex(int32 InOperandIndex, int32 InTotalOperands) const override;
-	virtual const TArray<FRigVMTemplateArgument>& GetArguments() const override;
+	virtual const TArray<FRigVMTemplateArgumentInfo>& GetArgumentInfos() const override;
 	virtual FRigVMTemplateTypeMap OnNewArgumentType(const FName& InArgumentName, TRigVMTypeIndex InTypeIndex) const override;
 #if WITH_EDITOR
 	virtual FText GetNodeTooltip(const FRigVMTemplateTypeMap& InTypes) const override;
@@ -190,7 +190,7 @@ public:
 
 	virtual ERigVMOpCode GetOpCode() const override { return ERigVMOpCode::ArraySetAtIndex; }
 	virtual FName GetArgumentNameForOperandIndex(int32 InOperandIndex, int32 InTotalOperands) const override;
-	virtual const TArray<FRigVMTemplateArgument>& GetArguments() const override;
+	virtual const TArray<FRigVMTemplateArgumentInfo>& GetArgumentInfos() const override;
 	virtual FRigVMTemplateTypeMap OnNewArgumentType(const FName& InArgumentName, TRigVMTypeIndex InTypeIndex) const override;
 #if WITH_EDITOR
 	virtual FText GetNodeTooltip(const FRigVMTemplateTypeMap& InTypes) const override;
@@ -215,7 +215,7 @@ public:
 
 	virtual ERigVMOpCode GetOpCode() const override { return ERigVMOpCode::ArrayAdd; }
 	virtual FName GetArgumentNameForOperandIndex(int32 InOperandIndex, int32 InTotalOperands) const override;
-	virtual const TArray<FRigVMTemplateArgument>& GetArguments() const override;
+	virtual const TArray<FRigVMTemplateArgumentInfo>& GetArgumentInfos() const override;
 #if WITH_EDITOR
 	virtual FText GetNodeTooltip(const FRigVMTemplateTypeMap& InTypes) const override;
 	virtual FText GetArgumentTooltip(const FName& InArgumentName, TRigVMTypeIndex InTypeIndex) const override;
@@ -261,7 +261,7 @@ public:
 
 	virtual ERigVMOpCode GetOpCode() const override { return ERigVMOpCode::ArrayRemove; }
 	virtual FName GetArgumentNameForOperandIndex(int32 InOperandIndex, int32 InTotalOperands) const override;
-	virtual const TArray<FRigVMTemplateArgument>& GetArguments() const override;
+	virtual const TArray<FRigVMTemplateArgumentInfo>& GetArgumentInfos() const override;
 	virtual FRigVMTemplateTypeMap OnNewArgumentType(const FName& InArgumentName, TRigVMTypeIndex InTypeIndex) const override;
 #if WITH_EDITOR
    	virtual FText GetNodeTooltip(const FRigVMTemplateTypeMap& InTypes) const override;
@@ -309,7 +309,7 @@ public:
 
 	virtual ERigVMOpCode GetOpCode() const override { return ERigVMOpCode::ArrayFind; }
 	virtual FName GetArgumentNameForOperandIndex(int32 InOperandIndex, int32 InTotalOperands) const override;
-	virtual const TArray<FRigVMTemplateArgument>& GetArguments() const override;
+	virtual const TArray<FRigVMTemplateArgumentInfo>& GetArgumentInfos() const override;
 	virtual FRigVMTemplateTypeMap OnNewArgumentType(const FName& InArgumentName, TRigVMTypeIndex InTypeIndex) const override;
 #if WITH_EDITOR
    	virtual FText GetNodeTooltip(const FRigVMTemplateTypeMap& InTypes) const override;
@@ -334,7 +334,7 @@ public:
 
 	virtual ERigVMOpCode GetOpCode() const override { return ERigVMOpCode::ArrayAppend; }
 	virtual FName GetArgumentNameForOperandIndex(int32 InOperandIndex, int32 InTotalOperands) const override;
-	virtual const TArray<FRigVMTemplateArgument>& GetArguments() const override;
+	virtual const TArray<FRigVMTemplateArgumentInfo>& GetArgumentInfos() const override;
 	virtual FRigVMTemplateTypeMap OnNewArgumentType(const FName& InArgumentName, TRigVMTypeIndex InTypeIndex) const override;
 #if WITH_EDITOR
    	virtual FText GetNodeTooltip(const FRigVMTemplateTypeMap& InTypes) const override;
@@ -359,7 +359,7 @@ public:
 
 	virtual ERigVMOpCode GetOpCode() const override { return ERigVMOpCode::ArrayClone; }
 	virtual FName GetArgumentNameForOperandIndex(int32 InOperandIndex, int32 InTotalOperands) const override;
-	virtual const TArray<FRigVMTemplateArgument>& GetArguments() const override;
+	virtual const TArray<FRigVMTemplateArgumentInfo>& GetArgumentInfos() const override;
 	virtual FRigVMTemplateTypeMap OnNewArgumentType(const FName& InArgumentName, TRigVMTypeIndex InTypeIndex) const override;
 #if WITH_EDITOR
    	virtual FText GetNodeTooltip(const FRigVMTemplateTypeMap& InTypes) const override;
@@ -406,7 +406,7 @@ public:
 
 	virtual ERigVMOpCode GetOpCode() const override { return ERigVMOpCode::ArrayDifference; }
 	virtual FName GetArgumentNameForOperandIndex(int32 InOperandIndex, int32 InTotalOperands) const override;
-	virtual const TArray<FRigVMTemplateArgument>& GetArguments() const override;
+	virtual const TArray<FRigVMTemplateArgumentInfo>& GetArgumentInfos() const override;
 	virtual FRigVMTemplateTypeMap OnNewArgumentType(const FName& InArgumentName, TRigVMTypeIndex InTypeIndex) const override;
 #if WITH_EDITOR
    	virtual FText GetNodeTooltip(const FRigVMTemplateTypeMap& InTypes) const override;
@@ -453,7 +453,7 @@ public:
 
 	virtual ERigVMOpCode GetOpCode() const override { return ERigVMOpCode::ArrayIterator; }
 	virtual FName GetArgumentNameForOperandIndex(int32 InOperandIndex, int32 InTotalOperands) const override;
-	virtual const TArray<FRigVMTemplateArgument>& GetArguments() const override;
+	virtual const TArray<FRigVMTemplateArgumentInfo>& GetArgumentInfos() const override;
 	virtual const TArray<FRigVMExecuteArgument>& GetExecuteArguments_Impl(const FRigVMDispatchContext& InContext) const override;
 	virtual const TArray<FName>& GetControlFlowBlocks_Impl(const FRigVMDispatchContext& InContext) const override;
 	virtual const bool IsControlFlowBlockSliced(const FName& InBlockName) const override;
diff --git a/Engine/Plugins/Runtime/RigVM/Source/RigVM/Public/RigVMFunctions/RigVMDispatch_CastObject.h b/Engine/Plugins/Runtime/RigVM/Source/RigVM/Public/RigVMFunctions/RigVMDispatch_CastObject.h
--- a/Engine/Plugins/Runtime/RigVM/Source/RigVM/Public/RigVMFunctions/RigVMDispatch_CastObject.h
+++ b/Engine/Plugins/Runtime/RigVM/Source/RigVM/Public/RigVMFunctions/RigVMDispatch_CastObject.h
@@ -16,7 +16,7 @@ public:
 		FactoryScriptStruct = StaticStruct();
 	}
 	
-	virtual const TArray<FRigVMTemplateArgument>& GetArguments() const override;
+	virtual const TArray<FRigVMTemplateArgumentInfo>& GetArgumentInfos() const override;
 	virtual TArray<FRigVMTemplateTypeMap> GetPermutationsFromArgumentType(const FName& InArgumentName, const TRigVMTypeIndex& InTypeIndex) const override;
 #if WITH_EDITOR
 	virtual FString GetNodeTitle(const FRigVMTemplateTypeMap& InTypes) const override;
diff --git a/Engine/Plugins/Runtime/RigVM/Source/RigVM/Public/RigVMFunctions/RigVMDispatch_Constant.h b/Engine/Plugins/Runtime/RigVM/Source/RigVM/Public/RigVMFunctions/RigVMDispatch_Constant.h
--- a/Engine/Plugins/Runtime/RigVM/Source/RigVM/Public/RigVMFunctions/RigVMDispatch_Constant.h
+++ b/Engine/Plugins/Runtime/RigVM/Source/RigVM/Public/RigVMFunctions/RigVMDispatch_Constant.h
@@ -18,7 +18,7 @@ public:
 		FactoryScriptStruct = StaticStruct();
 	}
 
-	virtual const TArray<FRigVMTemplateArgument>& GetArguments() const override;
+	virtual const TArray<FRigVMTemplateArgumentInfo>& GetArgumentInfos() const override;
 	virtual FRigVMTemplateTypeMap OnNewArgumentType(const FName& InArgumentName, TRigVMTypeIndex InTypeIndex) const override;
 #if WITH_EDITOR
 	virtual FString GetNodeTitle(const FRigVMTemplateTypeMap& InTypes) const override;
diff --git a/Engine/Plugins/Runtime/RigVM/Source/RigVM/Public/RigVMFunctions/RigVMDispatch_Core.h b/Engine/Plugins/Runtime/RigVM/Source/RigVM/Public/RigVMFunctions/RigVMDispatch_Core.h
--- a/Engine/Plugins/Runtime/RigVM/Source/RigVM/Public/RigVMFunctions/RigVMDispatch_Core.h
+++ b/Engine/Plugins/Runtime/RigVM/Source/RigVM/Public/RigVMFunctions/RigVMDispatch_Core.h
@@ -27,7 +27,7 @@ public:
 	}
 
 	virtual FName GetArgumentNameForOperandIndex(int32 InOperandIndex, int32 InTotalOperands) const override;
-	virtual const TArray<FRigVMTemplateArgument>& GetArguments() const override;
+	virtual const TArray<FRigVMTemplateArgumentInfo>& GetArgumentInfos() const override;
 	virtual FRigVMTemplateTypeMap OnNewArgumentType(const FName& InArgumentName, TRigVMTypeIndex InTypeIndex) const override;
 
 protected:
@@ -100,7 +100,7 @@ struct RIGVM_API FRigVMDispatch_CoreNotEquals : public FRigVMDispatch_CoreEquals
 		FactoryScriptStruct = StaticStruct();
 	}
 
-	virtual const TArray<FRigVMTemplateArgument>& GetArguments() const override;
+	virtual const TArray<FRigVMTemplateArgumentInfo>& GetArgumentInfos() const override;
 
 	// we are inheriting everything from the equals dispatch,
 	// and due to the check of the factory within FRigVMDispatch_CoreEquals::Execute we can
diff --git a/Engine/Plugins/Runtime/RigVM/Source/RigVM/Public/RigVMFunctions/RigVMDispatch_If.h b/Engine/Plugins/Runtime/RigVM/Source/RigVM/Public/RigVMFunctions/RigVMDispatch_If.h
--- a/Engine/Plugins/Runtime/RigVM/Source/RigVM/Public/RigVMFunctions/RigVMDispatch_If.h
+++ b/Engine/Plugins/Runtime/RigVM/Source/RigVM/Public/RigVMFunctions/RigVMDispatch_If.h
@@ -22,7 +22,7 @@ public:
 	}
 
 	virtual FName GetArgumentNameForOperandIndex(int32 InOperandIndex, int32 InTotalOperands) const override;
-	virtual const TArray<FRigVMTemplateArgument>& GetArguments() const override;
+	virtual const TArray<FRigVMTemplateArgumentInfo>& GetArgumentInfos() const override;
 	virtual FRigVMTemplateTypeMap OnNewArgumentType(const FName& InArgumentName, TRigVMTypeIndex InTypeIndex) const override;
 	virtual bool IsSingleton() const override { return true; } 
 
diff --git a/Engine/Plugins/Runtime/RigVM/Source/RigVM/Public/RigVMFunctions/RigVMDispatch_MakeStruct.h b/Engine/Plugins/Runtime/RigVM/Source/RigVM/Public/RigVMFunctions/RigVMDispatch_MakeStruct.h
--- a/Engine/Plugins/Runtime/RigVM/Source/RigVM/Public/RigVMFunctions/RigVMDispatch_MakeStruct.h
+++ b/Engine/Plugins/Runtime/RigVM/Source/RigVM/Public/RigVMFunctions/RigVMDispatch_MakeStruct.h
@@ -18,7 +18,7 @@ public:
 		FactoryScriptStruct = StaticStruct();
 	}
 	
-	virtual const TArray<FRigVMTemplateArgument>& GetArguments() const override;
+	virtual const TArray<FRigVMTemplateArgumentInfo>& GetArgumentInfos() const override;
 	virtual FRigVMTemplateTypeMap OnNewArgumentType(const FName& InArgumentName, TRigVMTypeIndex InTypeIndex) const override;
 #if WITH_EDITOR
 	virtual FString GetNodeTitle(const FRigVMTemplateTypeMap& InTypes) const override;
@@ -49,10 +49,10 @@ public:
 		FactoryScriptStruct = StaticStruct();
 	}
 
-	virtual const TArray<FRigVMTemplateArgument>& GetArguments() const override;
+	virtual const TArray<FRigVMTemplateArgumentInfo>& GetArgumentInfos() const override;
 #if WITH_EDITOR
 	virtual FText GetNodeTooltip(const FRigVMTemplateTypeMap& InTypes) const override;
-	virtual FString GetKeywords() const;
+	virtual FString GetKeywords() const override;
 #endif
 };
 
diff --git a/Engine/Plugins/Runtime/RigVM/Source/RigVM/Public/RigVMFunctions/RigVMDispatch_Print.h b/Engine/Plugins/Runtime/RigVM/Source/RigVM/Public/RigVMFunctions/RigVMDispatch_Print.h
--- a/Engine/Plugins/Runtime/RigVM/Source/RigVM/Public/RigVMFunctions/RigVMDispatch_Print.h
+++ b/Engine/Plugins/Runtime/RigVM/Source/RigVM/Public/RigVMFunctions/RigVMDispatch_Print.h
@@ -21,7 +21,7 @@ public:
 	}
 
 	virtual FName GetArgumentNameForOperandIndex(int32 InOperandIndex, int32 InTotalOperands) const override;
-	virtual const TArray<FRigVMTemplateArgument>& GetArguments() const override;
+	virtual const TArray<FRigVMTemplateArgumentInfo>& GetArgumentInfos() const override;
 	virtual const TArray<FRigVMExecuteArgument>& GetExecuteArguments_Impl(const FRigVMDispatchContext& InContext) const override;
 	virtual FRigVMTemplateTypeMap OnNewArgumentType(const FName& InArgumentName, TRigVMTypeIndex InTypeIndex) const override;
 	virtual bool IsSingleton() const override { return true; } 
diff --git a/Engine/Plugins/Runtime/RigVM/Source/RigVM/Public/RigVMFunctions/RigVMDispatch_Select.h b/Engine/Plugins/Runtime/RigVM/Source/RigVM/Public/RigVMFunctions/RigVMDispatch_Select.h
--- a/Engine/Plugins/Runtime/RigVM/Source/RigVM/Public/RigVMFunctions/RigVMDispatch_Select.h
+++ b/Engine/Plugins/Runtime/RigVM/Source/RigVM/Public/RigVMFunctions/RigVMDispatch_Select.h
@@ -22,7 +22,7 @@ public:
 	}
 
 	virtual FName GetArgumentNameForOperandIndex(int32 InOperandIndex, int32 InTotalOperands) const override;
-	virtual const TArray<FRigVMTemplateArgument>& GetArguments() const override;
+	virtual const TArray<FRigVMTemplateArgumentInfo>& GetArgumentInfos() const override;
 	virtual FRigVMTemplateTypeMap OnNewArgumentType(const FName& InArgumentName, TRigVMTypeIndex InTypeIndex) const override;
 #if WITH_EDITOR
 	virtual FString GetArgumentMetaData(const FName& InArgumentName, const FName& InMetaDataKey) const override;
diff --git a/Engine/Plugins/Runtime/RigVM/Source/RigVM/Public/RigVMFunctions/RigVMDispatch_Switch.h b/Engine/Plugins/Runtime/RigVM/Source/RigVM/Public/RigVMFunctions/RigVMDispatch_Switch.h
--- a/Engine/Plugins/Runtime/RigVM/Source/RigVM/Public/RigVMFunctions/RigVMDispatch_Switch.h
+++ b/Engine/Plugins/Runtime/RigVM/Source/RigVM/Public/RigVMFunctions/RigVMDispatch_Switch.h
@@ -22,7 +22,7 @@ public:
 	}
 
 	virtual FName GetArgumentNameForOperandIndex(int32 InOperandIndex, int32 InTotalOperands) const override;
-	virtual const TArray<FRigVMTemplateArgument>& GetArguments() const override;
+	virtual const TArray<FRigVMTemplateArgumentInfo>& GetArgumentInfos() const override;
 	virtual const TArray<FRigVMExecuteArgument>& GetExecuteArguments_Impl(const FRigVMDispatchContext& InContext) const override;
 	virtual bool IsSingleton() const override { return true; } 
 
diff --git a/Engine/Plugins/Runtime/RigVM/Source/RigVM/Public/RigVMFunctions/RigVMFunction_String.h b/Engine/Plugins/Runtime/RigVM/Source/RigVM/Public/RigVMFunctions/RigVMFunction_String.h
--- a/Engine/Plugins/Runtime/RigVM/Source/RigVM/Public/RigVMFunctions/RigVMFunction_String.h
+++ b/Engine/Plugins/Runtime/RigVM/Source/RigVM/Public/RigVMFunctions/RigVMFunction_String.h
@@ -510,7 +510,7 @@ struct FRigDispatch_ToString : public FRigVMDispatchFactory
 
 public:
 
-	virtual const TArray<FRigVMTemplateArgument>& GetArguments() const override;
+	virtual const TArray<FRigVMTemplateArgumentInfo>& GetArgumentInfos() const override;
 	virtual FRigVMTemplateTypeMap OnNewArgumentType(const FName& InArgumentName, TRigVMTypeIndex InTypeIndex) const override;
 
 protected:
@@ -529,7 +529,7 @@ struct FRigDispatch_FromString : public FRigVMDispatchFactory
 
 public:
 
-	virtual const TArray<FRigVMTemplateArgument>& GetArguments() const override;
+	virtual const TArray<FRigVMTemplateArgumentInfo>& GetArgumentInfos() const override;
 	virtual FRigVMTemplateTypeMap OnNewArgumentType(const FName& InArgumentName, TRigVMTypeIndex InTypeIndex) const override;
 
 protected:
diff --git a/Engine/Plugins/Runtime/RigVM/Source/RigVMDeveloper/Private/RigVMModel/Nodes/RigVMIfNode.cpp b/Engine/Plugins/Runtime/RigVM/Source/RigVMDeveloper/Private/RigVMModel/Nodes/RigVMIfNode.cpp
--- a/Engine/Plugins/Runtime/RigVM/Source/RigVMDeveloper/Private/RigVMModel/Nodes/RigVMIfNode.cpp
+++ b/Engine/Plugins/Runtime/RigVM/Source/RigVMDeveloper/Private/RigVMModel/Nodes/RigVMIfNode.cpp
@@ -38,23 +38,23 @@ const FRigVMTemplate* UDEPRECATED_RigVMIfNode::GetTemplate() const
 		static const FName FalseFName = *FalseName;
 		static const FName ResultFName = *ResultName;
 
-		static TArray<FRigVMTemplateArgument> Arguments;
-		if(Arguments.IsEmpty())
+		static TArray<FRigVMTemplateArgumentInfo> Infos;
+		if(Infos.IsEmpty())
 		{
 			static const TArray<FRigVMTemplateArgument::ETypeCategory> Categories = {
 				FRigVMTemplateArgument::ETypeCategory_SingleAnyValue,
 				FRigVMTemplateArgument::ETypeCategory_ArrayAnyValue
 			};
-			Arguments.Reserve(4);
-			Arguments.Emplace(ConditionFName, ERigVMPinDirection::Input, RigVMTypeUtils::TypeIndex::Bool);
-			Arguments.Emplace(TrueFName, ERigVMPinDirection::Input, Categories);
-			Arguments.Emplace(FalseFName, ERigVMPinDirection::Input, Categories);
-			Arguments.Emplace(ResultFName, ERigVMPinDirection::Output, Categories);
+			Infos.Reserve(4);
+			Infos.Emplace(ConditionFName, ERigVMPinDirection::Input, RigVMTypeUtils::TypeIndex::Bool);
+			Infos.Emplace(TrueFName, ERigVMPinDirection::Input, Categories);
+			Infos.Emplace(FalseFName, ERigVMPinDirection::Input, Categories);
+			Infos.Emplace(ResultFName, ERigVMPinDirection::Output, Categories);
 		}
 
 		FRigVMTemplateDelegates Delegates;
 		Delegates.NewArgumentTypeDelegate = 
-			FRigVMTemplate_NewArgumentTypeDelegate::CreateLambda([](const FRigVMTemplate*, const FName& InArgumentName, int32 InTypeIndex)
+			FRigVMTemplate_NewArgumentTypeDelegate::CreateLambda([](const FName& InArgumentName, int32 InTypeIndex)
 			{
 				FRigVMTemplateTypeMap Types;
 
@@ -69,7 +69,7 @@ const FRigVMTemplate* UDEPRECATED_RigVMIfNode::GetTemplate() const
 				return Types;
 			});
 
-		IfNodeTemplate = CachedTemplate = FRigVMRegistry::Get().GetOrAddTemplateFromArguments(*IfName, Arguments, Delegates);
+		IfNodeTemplate = CachedTemplate = FRigVMRegistry::Get().GetOrAddTemplateFromArguments(*IfName, Infos, Delegates);
 	}
 	return CachedTemplate;
 }
diff --git a/Engine/Plugins/Runtime/RigVM/Source/RigVMDeveloper/Private/RigVMModel/Nodes/RigVMSelectNode.cpp b/Engine/Plugins/Runtime/RigVM/Source/RigVMDeveloper/Private/RigVMModel/Nodes/RigVMSelectNode.cpp
--- a/Engine/Plugins/Runtime/RigVM/Source/RigVMDeveloper/Private/RigVMModel/Nodes/RigVMSelectNode.cpp
+++ b/Engine/Plugins/Runtime/RigVM/Source/RigVMDeveloper/Private/RigVMModel/Nodes/RigVMSelectNode.cpp
@@ -48,8 +48,8 @@ const FRigVMTemplate* UDEPRECATED_RigVMSelectNode::GetTemplate() const
 		static const FName ValueFName = *ValueName;
 		static const FName ResultFName = *ResultName;
 
-		static TArray<FRigVMTemplateArgument> Arguments;
-		if(Arguments.IsEmpty())
+		static TArray<FRigVMTemplateArgumentInfo> Infos;
+		if(Infos.IsEmpty())
 		{
 			static const TArray<FRigVMTemplateArgument::ETypeCategory> ValueTypeCategories = {
 				FRigVMTemplateArgument::ETypeCategory_ArrayAnyValue,
@@ -59,15 +59,15 @@ const FRigVMTemplate* UDEPRECATED_RigVMSelectNode::GetTemplate() const
 				FRigVMTemplateArgument::ETypeCategory_SingleAnyValue,
 				FRigVMTemplateArgument::ETypeCategory_ArrayAnyValue
 			};
-			Arguments.Reserve(3);
-			Arguments.Emplace(IndexFName, ERigVMPinDirection::Input, RigVMTypeUtils::TypeIndex::Int32);
-			Arguments.Emplace(ValueFName, ERigVMPinDirection::Input, ValueTypeCategories);
-			Arguments.Emplace(ResultFName, ERigVMPinDirection::Output, ResultTypeCategories);
+			Infos.Reserve(3);
+			Infos.Emplace(IndexFName, ERigVMPinDirection::Input, RigVMTypeUtils::TypeIndex::Int32);
+			Infos.Emplace(ValueFName, ERigVMPinDirection::Input, ValueTypeCategories);
+			Infos.Emplace(ResultFName, ERigVMPinDirection::Output, ResultTypeCategories);
 		}
 		
 		FRigVMTemplateDelegates Delegates;
 		Delegates.NewArgumentTypeDelegate = 
-			FRigVMTemplate_NewArgumentTypeDelegate::CreateLambda([](const FRigVMTemplate*, const FName& InArgumentName, int32 InTypeIndex)
+			FRigVMTemplate_NewArgumentTypeDelegate::CreateLambda([](const FName& InArgumentName, int32 InTypeIndex)
 			{
 				FRigVMTemplateTypeMap Types;
 
@@ -95,7 +95,7 @@ const FRigVMTemplate* UDEPRECATED_RigVMSelectNode::GetTemplate() const
 				return Types;
 			});
 
-		SelectNodeTemplate = CachedTemplate = FRigVMRegistry::Get().GetOrAddTemplateFromArguments(*SelectName, Arguments, Delegates);
+		SelectNodeTemplate = CachedTemplate = FRigVMRegistry::Get().GetOrAddTemplateFromArguments(*SelectName, Infos, Delegates);
 	}
 	return CachedTemplate;
 }
