RigVM: Fix profiler to reset start cycles only for root event

#rnx

diff --git a/Engine/Plugins/Runtime/RigVM/Source/RigVM/Private/RigVMCore/RigVM.cpp b/Engine/Plugins/Runtime/RigVM/Source/RigVM/Private/RigVMCore/RigVM.cpp
--- a/Engine/Plugins/Runtime/RigVM/Source/RigVM/Private/RigVMCore/RigVM.cpp
+++ b/Engine/Plugins/Runtime/RigVM/Source/RigVM/Private/RigVMCore/RigVM.cpp
@@ -1587,7 +1587,10 @@ ERigVMExecuteResult URigVM::ExecuteVM(FRigVMExtendedExecuteContext& Context, con
 	}
 
 #if WITH_EDITOR
-	StartProfiling(Context);
+	if(Context.bCurrentlyRunningRootEntry)
+	{
+		StartProfiling(Context);
+	}
 	
 #if UE_RIGVM_DEBUG_EXECUTION
 	if (CVarControlRigDebugAllVMExecutions->GetBool() || ContextPublicData.bDebugExecution)
