RigVM: Display profiling total microseconds in bottom left of graph

#rnx

diff --git a/Engine/Plugins/Runtime/RigVM/Source/RigVMEditor/Private/Editor/RigVMEditor.cpp b/Engine/Plugins/Runtime/RigVM/Source/RigVMEditor/Private/Editor/RigVMEditor.cpp
--- a/Engine/Plugins/Runtime/RigVM/Source/RigVMEditor/Private/Editor/RigVMEditor.cpp
+++ b/Engine/Plugins/Runtime/RigVM/Source/RigVMEditor/Private/Editor/RigVMEditor.cpp
@@ -1512,7 +1512,7 @@ FGraphAppearanceInfo FRigVMEditor::GetGraphAppearance(UEdGraph* InGraph) const
 {
 	FGraphAppearanceInfo AppearanceInfo = FBlueprintEditor::GetGraphAppearance(InGraph);
 
-	if (GetBlueprintObj()->IsA(URigVMBlueprint::StaticClass()))
+	if (const URigVMBlueprint* RigVMBlueprint = GetRigVMBlueprint())
 	{
 		AppearanceInfo.CornerText = LOCTEXT("AppearanceCornerText_RigVMEditor", "RigVM");
 
@@ -1525,6 +1525,12 @@ FGraphAppearanceInfo FRigVMEditor::GetGraphAppearance(UEdGraph* InGraph) const
 					FString::Printf(TEXT("This graph runs a nativized VM (U%s)."), *RigVMHost->GetVM()->GetNativizedClass()->GetName())
 				);
 			}
+
+			if(RigVMHost->VMRuntimeSettings.bEnableProfiling)
+			{
+				static constexpr TCHAR Format[] = TEXT("Total %.02f µs");
+				AppearanceInfo.WarningText = FText::FromString(FString::Printf(Format, (float)RigVMBlueprint->RigGraphDisplaySettings.TotalMicroSeconds));
+			}
 		}
 	}
 
