Fixed customization not working with UPropertyBag
    - note that setting values are disabled (only get + copy)

diff --git a/Engine/Plugins/Runtime/RigVM/Source/RigVMEditor/Private/Editor/RigVMGraphDetailCustomization.cpp b/Engine/Plugins/Runtime/RigVM/Source/RigVMEditor/Private/Editor/RigVMGraphDetailCustomization.cpp
--- a/Engine/Plugins/Runtime/RigVM/Source/RigVMEditor/Private/Editor/RigVMGraphDetailCustomization.cpp
+++ b/Engine/Plugins/Runtime/RigVM/Source/RigVMEditor/Private/Editor/RigVMGraphDetailCustomization.cpp
@@ -1595,7 +1595,7 @@ void FRigVMWrappedNodeDetailCustomization::OnNameListComboBox(FNameProperty* InP
 
 void FRigVMWrappedNodeDetailCustomization::CustomizeLiveValues(IDetailLayoutBuilder& DetailLayout)
 {
-	if(ObjectsBeingCustomized.Num() != 1)
+	if(ObjectsBeingCustomized.Num() > 1)
 	{
 		return;
 	}
@@ -1636,6 +1636,11 @@ void FRigVMWrappedNodeDetailCustomization::CustomizeLiveValues(IDetailLayoutBuil
 
 	for(URigVMPin* Pin : FirstNode->GetPins())
 	{
+		if (Pin->IsExecuteContext())
+		{
+			continue;
+		}
+		
 		// only show hidden pins in debug mode
 		if(Pin->GetDirection() == ERigVMPinDirection::Hidden)
 		{
@@ -1710,7 +1715,19 @@ void FRigVMWrappedNodeDetailCustomization::CustomizeLiveValues(IDetailLayoutBuil
 		int32 SuffixIndex = 1;
 		FString NameSuffix;
 
-		TArray<FRigVMOperand> KnownOperands; 
+		auto UpdateRow = [&](IDetailPropertyRow* PropertyRow)
+		{
+			PropertyRow->DisplayName(FText::FromString(FString::Printf(TEXT("%s%s"), *Pin->GetName(), *NameSuffix)));
+			PropertyRow->IsEnabled(false);
+
+			SuffixIndex++;
+			bAddedProperty = true;
+			NameSuffix = FString::Printf(TEXT("_%d"), SuffixIndex);
+		};
+
+		static const FAddPropertyParams AddPropertyParams = FAddPropertyParams().ForceShowProperty();
+
+		TArray<FRigVMOperand> KnownOperands;
 		for(const FRigVMExprAST* Expression : FilteredExpressions)
 		{
 			const FRigVMVarExprAST* VarExpr = Expression->To<FRigVMVarExprAST>();
@@ -1767,27 +1784,20 @@ void FRigVMWrappedNodeDetailCustomization::CustomizeLiveValues(IDetailLayoutBuil
 				check(ExternalObjects.Num() > 0 || ExternalStructs.Num() > 0);
 				check(Property);
 
-				if(IDetailPropertyRow* PropertyRow = DebugCategory.AddExternalObjectProperty(ExternalObjects, Property->GetFName(), EPropertyLocation::Default, FAddPropertyParams().ForceShowProperty()))
+				if (!ExternalObjects.IsEmpty())
 				{
-					PropertyRow->DisplayName(FText::FromString(FString::Printf(TEXT("%s%s"), *Pin->GetName(), *NameSuffix)));
-					PropertyRow->IsEnabled(false);
-
-					SuffixIndex++;
-					bAddedProperty = true;
-					NameSuffix = FString::Printf(TEXT("_%d"), SuffixIndex);
+					if(IDetailPropertyRow* PropertyRow = DebugCategory.AddExternalObjectProperty(ExternalObjects, Property->GetFName(), EPropertyLocation::Default, AddPropertyParams))
+					{
+						UpdateRow(PropertyRow);
+					}
 				}
 
-				for (FRigVMMemoryStorageStruct* Memory : ExternalStructs)
+				for (const FRigVMMemoryStorageStruct* Memory : ExternalStructs)
 				{
-					TSharedPtr<FStructOnScope> StructOnScope = MakeShareable(new FStructOnScope(Memory->GetPropertyBagStruct(), (uint8*)Memory->GetContainerPtr()));
-					if (IDetailPropertyRow* PropertyRow = DebugCategory.AddExternalStructureProperty(StructOnScope, Property->GetFName(), EPropertyLocation::Default, FAddPropertyParams().ForceShowProperty()))
+					const TSharedPtr<FStructOnScope> StructOnScope = MakeShareable(new FStructOnScope(Memory->GetPropertyBagStruct(), (uint8*)Memory->GetContainerPtr()));
+					if (IDetailPropertyRow* PropertyRow = DebugCategory.AddExternalStructureProperty(StructOnScope, Property->GetFName(), EPropertyLocation::Default, AddPropertyParams))
 					{
-						PropertyRow->DisplayName(FText::FromString(FString::Printf(TEXT("%s%s"), *Pin->GetName(), *NameSuffix)));
-						PropertyRow->IsEnabled(false);
-
-						SuffixIndex++;
-						bAddedProperty = true;
-						NameSuffix = FString::Printf(TEXT("_%d"), SuffixIndex);
+						UpdateRow(PropertyRow);
 					}
 				}
 
@@ -1811,7 +1821,8 @@ void FRigVMWrappedNodeDetailCustomization::CustomizeLiveValues(IDetailLayoutBuil
 				.ValueContent()
 				[
 					PinHandle->CreatePropertyValueWidget()
-				];
+				]
+				.IsEnabled(false);
 			}
 		}
 	}
@@ -1830,6 +1841,9 @@ void FRigVMGraphMathTypeDetailCustomization::CustomizeHeader(TSharedRef<IPropert
 	TArray<UObject*> Objects;
 	InPropertyHandle->GetOuterObjects(Objects);
 
+	StructsBeingCustomized.Reset();
+	InPropertyHandle->GetOuterStructs(StructsBeingCustomized);
+
 	for (UObject* Object : Objects)
 	{
 		ObjectsBeingCustomized.Add(Object);
diff --git a/Engine/Plugins/Runtime/RigVM/Source/RigVMEditor/Public/Editor/RigVMGraphDetailCustomization.h b/Engine/Plugins/Runtime/RigVM/Source/RigVMEditor/Public/Editor/RigVMGraphDetailCustomization.h
--- a/Engine/Plugins/Runtime/RigVM/Source/RigVMEditor/Public/Editor/RigVMGraphDetailCustomization.h
+++ b/Engine/Plugins/Runtime/RigVM/Source/RigVMEditor/Public/Editor/RigVMGraphDetailCustomization.h
@@ -258,29 +258,59 @@ public:
 
 protected:
 
+	TArray<uint8*> GetMemoryBeingCustomized()
+	{
+		TArray<uint8*> MemoryPtr;
+		MemoryPtr.Reserve(ObjectsBeingCustomized.Num() + StructsBeingCustomized.Num());
+
+		for(const TWeakObjectPtr<UObject>& Object : ObjectsBeingCustomized)
+		{
+			if(Object.IsValid())
+			{
+				MemoryPtr.Add((uint8*)Object.Get());
+			}
+		}
+
+		for(const TSharedPtr<FStructOnScope>& StructPtr: StructsBeingCustomized)
+		{
+			if(StructPtr.IsValid())
+			{
+				MemoryPtr.Add(StructPtr->GetStructMemory());
+			}
+		}
+
+		return MemoryPtr;
+	}
+	
 	bool GetPropertyChain(TSharedRef<class IPropertyHandle> InPropertyHandle, FEditPropertyChain& OutPropertyChain, TArray<int32> &OutPropertyArrayIndices, bool& bOutEnabled)
 	{
+		if (!InPropertyHandle->IsValidHandle())
+		{
+			return false;
+		}
+		
 		OutPropertyChain.Empty();
 		OutPropertyArrayIndices.Reset();
 		bOutEnabled = false;
-		if (!ObjectsBeingCustomized.IsEmpty())
+
+		const bool bHasObject = !ObjectsBeingCustomized.IsEmpty() && ObjectsBeingCustomized[0].Get();
+		const bool bHasStruct = !StructsBeingCustomized.IsEmpty() && StructsBeingCustomized[0].Get();
+		
+		if (bHasStruct || bHasObject)
 		{
-			if (ObjectsBeingCustomized[0].Get())
+			TSharedPtr<class IPropertyHandle> ChainHandle = InPropertyHandle;
+			while (ChainHandle.IsValid() && ChainHandle->GetProperty() != nullptr)
 			{
-				TSharedPtr<class IPropertyHandle> ChainHandle = InPropertyHandle;
-				while (ChainHandle.IsValid() && ChainHandle->GetProperty() != nullptr)
-				{
-					OutPropertyChain.AddHead(ChainHandle->GetProperty());
-					OutPropertyArrayIndices.Insert(ChainHandle->GetIndexInArray(), 0);
-					ChainHandle = ChainHandle->GetParentHandle();					
-				}
+				OutPropertyChain.AddHead(ChainHandle->GetProperty());
+				OutPropertyArrayIndices.Insert(ChainHandle->GetIndexInArray(), 0);
+				ChainHandle = ChainHandle->GetParentHandle();					
+			}
 
-				if (OutPropertyChain.GetHead() != nullptr)
-				{
-					OutPropertyChain.SetActiveMemberPropertyNode(OutPropertyChain.GetTail()->GetValue());
-					bOutEnabled = !OutPropertyChain.GetHead()->GetValue()->HasAnyPropertyFlags(CPF_EditConst);
-					return true;
-				}
+			if (OutPropertyChain.GetHead() != nullptr)
+			{
+				OutPropertyChain.SetActiveMemberPropertyNode(OutPropertyChain.GetTail()->GetValue());
+				bOutEnabled = !OutPropertyChain.GetHead()->GetValue()->HasAnyPropertyFlags(CPF_EditConst);
+				return true;
 			}
 		}
 		return false;
@@ -288,7 +318,7 @@ protected:
 
 	// extracts the value for a nested property (for Example Settings.WorldTransform) from an outer owner
 	template<typename ValueType>
-	ValueType& ContainerUObjectToValueRef(UObject* InOwner, ValueType& InDefault, FEditPropertyChain& InPropertyChain, TArray<int32> &InPropertyArrayIndices) const
+	static ValueType& ContainerMemoryBlockToValueRef(uint8* InMemoryBlock, ValueType& InDefault, FEditPropertyChain& InPropertyChain, TArray<int32> &InPropertyArrayIndices)
 	{
 		if (InPropertyChain.GetHead() == nullptr)
 		{
@@ -296,7 +326,7 @@ protected:
 		}
 		
 		FEditPropertyChain::TDoubleLinkedListNode* PropertyNode = InPropertyChain.GetHead();
-		uint8* MemoryPtr = (uint8*)InOwner;
+		uint8* MemoryPtr = InMemoryBlock;
 		int32 ChainIndex = 0;
 		do
 		{
@@ -326,7 +356,6 @@ protected:
 					ChainIndex++;
 				}
 			}
-
 		}
 		while (PropertyNode);
 
@@ -349,13 +378,14 @@ protected:
 		{
 			return Result;
 		}
-	
-		for(const TWeakObjectPtr<UObject>& Object : ObjectsBeingCustomized)
+
+		const TArray<uint8*> MemoryBlocks = GetMemoryBeingCustomized();
+		for(uint8* MemoryBlock: MemoryBlocks)
 		{
-			if(Object.Get() && InPropertyHandle->IsValidHandle())
+			if(MemoryBlock)
 			{
 				static VectorType ZeroVector = VectorType();
-				const VectorType& Vector = ContainerUObjectToValueRef<VectorType>(Object.Get(), ZeroVector, PropertyChain, PropertyArrayIndices);
+				const VectorType& Vector = ContainerMemoryBlockToValueRef<VectorType>(MemoryBlock, ZeroVector, PropertyChain, PropertyArrayIndices);
 				NumericType Component = Vector[InComponent];
 				if(Result.IsSet())
 				{
@@ -377,6 +407,11 @@ protected:
 	template<typename VectorType, typename NumericType>
 	void OnVectorComponentChanged(TSharedRef<class IPropertyHandle> InPropertyHandle, int32 InComponent, NumericType InValue, bool bIsCommit, ETextCommit::Type InCommitType = ETextCommit::Default)
 	{
+		if (ObjectsBeingCustomized.IsEmpty())
+		{
+			return;
+		}
+		
 		FEditPropertyChain PropertyChain;
 		TArray<int32> PropertyArrayIndices;
 		bool bEnabled;
@@ -413,7 +448,7 @@ protected:
 			if(Object.Get() && InPropertyHandle->IsValidHandle())
 			{
 				static VectorType ZeroVector = VectorType();
-				VectorType& Vector = ContainerUObjectToValueRef<VectorType>(Object.Get(), ZeroVector, PropertyChain, PropertyArrayIndices);
+				VectorType& Vector = ContainerMemoryBlockToValueRef<VectorType>((uint8*)Object.Get(), ZeroVector, PropertyChain, PropertyArrayIndices);
 				VectorType PreviousVector = Vector;
 				Vector[InComponent] = InValue;
 					
@@ -509,13 +544,14 @@ protected:
 		{
 			return Result;
 		}
-		
-		for(const TWeakObjectPtr<UObject>& Object : ObjectsBeingCustomized)
+
+		const TArray<uint8*> MemoryBlocks = GetMemoryBeingCustomized();
+		for(uint8* MemoryBlock: MemoryBlocks)
 		{
-			if(Object.Get() && InPropertyHandle->IsValidHandle())
+			if(MemoryBlock)
 			{
 				static RotationType ZeroRotation = RotationType();
-				const RotationType& Rotation = ContainerUObjectToValueRef<RotationType>(Object.Get(), ZeroRotation, PropertyChain, PropertyArrayIndices);
+				const RotationType& Rotation = ContainerMemoryBlockToValueRef<RotationType>(MemoryBlock, ZeroRotation, PropertyChain, PropertyArrayIndices);
 				if(Result.IsSet())
 				{
 					if(!Rotation.Equals(Result.GetValue()))
@@ -536,6 +572,11 @@ protected:
 	template<typename RotationType>
 	void OnRotationChanged(TSharedRef<class IPropertyHandle> InPropertyHandle, RotationType InValue, bool bIsCommit, ETextCommit::Type InCommitType = ETextCommit::Default)
 	{
+		if (ObjectsBeingCustomized.IsEmpty())
+		{
+			return;
+		}
+		
 		FEditPropertyChain PropertyChain;
         TArray<int32> PropertyArrayIndices;
         bool bEnabled;
@@ -572,7 +613,7 @@ protected:
 			if(Object.Get() && InPropertyHandle->IsValidHandle())
 			{
 				static RotationType ZeroRotation = RotationType();
-				RotationType& Rotation = ContainerUObjectToValueRef<RotationType>(Object.Get(), ZeroRotation, PropertyChain, PropertyArrayIndices);
+				RotationType& Rotation = ContainerMemoryBlockToValueRef<RotationType>((uint8*)Object.Get(), ZeroRotation, PropertyChain, PropertyArrayIndices);
 				RotationType PreviousRotation = Rotation;
 				Rotation = InValue;
 					
@@ -620,7 +661,6 @@ protected:
 		.IsEnabled(bEnabled)
 		.NameContent()
 		[
-			
 			InPropertyHandle->CreatePropertyNameWidget()
 		]
 		.ValueContent()
@@ -651,10 +691,24 @@ protected:
 		WidgetArgs.UseQuaternionForRotation(IsQuaternionBasedRotation<TransformType>());
 
 		static TransformType Identity = TransformType::Identity;
-		TransformType DefaultValue = ContainerUObjectToValueRef<TransformType>(ObjectsBeingCustomized[0]->GetClass()->GetDefaultObject(), Identity, PropertyChain, PropertyArrayIndices);
 
+		UObject* DefaultObject = !ObjectsBeingCustomized.IsEmpty() ? ObjectsBeingCustomized[0]->GetClass()->GetDefaultObject() :
+		!StructsBeingCustomized.IsEmpty() ? StructsBeingCustomized[0]->GetStruct()->GetClass() : nullptr;
+
+		if (!DefaultObject)
+		{
+			return;
+		}
+		
+		TransformType DefaultValue = ContainerMemoryBlockToValueRef<TransformType>((uint8*)DefaultObject, Identity, PropertyChain, PropertyArrayIndices);
+		
 		WidgetArgs.DiffersFromDefault_Lambda([this, InPropertyHandle, DefaultValue](ESlateTransformComponent::Type InTransformComponent) -> bool
 		{
+			if (ObjectsBeingCustomized.IsEmpty())
+			{
+				return false;
+			}
+			
 			FEditPropertyChain PropertyChain;
 			TArray<int32> PropertyArrayIndices;
 			bool bEnabled;
@@ -667,7 +721,7 @@ protected:
 			{
 				if(Object.Get() && InPropertyHandle->IsValidHandle())
 				{
-					const TransformType& Transform = ContainerUObjectToValueRef<TransformType>(Object.Get(), Identity, PropertyChain, PropertyArrayIndices);
+					const TransformType& Transform = ContainerMemoryBlockToValueRef<TransformType>((uint8*)Object.Get(), Identity, PropertyChain, PropertyArrayIndices);
 
 					switch(InTransformComponent)
 					{
@@ -718,12 +772,13 @@ protected:
 			{
 				return Result;
 			}
-			
-			for(const TWeakObjectPtr<UObject>& Object : ObjectsBeingCustomized)
+
+			const TArray<uint8*> MemoryBlocks = GetMemoryBeingCustomized();
+			for(uint8* MemoryBlock: MemoryBlocks)
 			{
-				if(Object.Get() && InPropertyHandle->IsValidHandle())
+				if(MemoryBlock)
 				{
-					const TransformType& Transform = ContainerUObjectToValueRef<TransformType>(Object.Get(), Identity, PropertyChain, PropertyArrayIndices);
+					const TransformType& Transform = ContainerMemoryBlockToValueRef<TransformType>(MemoryBlock, Identity, PropertyChain, PropertyArrayIndices);
 					
 					TOptional<FReal> Value = SAdvancedTransformInputBox<TransformType>::GetNumericValueFromTransform(
 						Transform,
@@ -795,7 +850,7 @@ protected:
 				const TWeakObjectPtr<UObject>& Object = ObjectsBeingCustomized[Index];
 				if(Object.Get() && InPropertyHandle->IsValidHandle())
 				{
-					TransformType& Transform = ContainerUObjectToValueRef<TransformType>(Object.Get(), Identity, PropertyChain, PropertyArrayIndices);
+					TransformType& Transform = ContainerMemoryBlockToValueRef<TransformType>((uint8*)Object.Get(), Identity, PropertyChain, PropertyArrayIndices);
 					TransformType PreviousTransform = Transform;
 					
 					SAdvancedTransformInputBox<TransformType>::ApplyNumericValueChange(
@@ -840,6 +895,11 @@ protected:
 
 		WidgetArgs.OnResetToDefault_Lambda([this, DefaultValue, InPropertyHandle](ESlateTransformComponent::Type InTransformComponent)
 		{
+			if (ObjectsBeingCustomized.IsEmpty())
+			{
+				return;
+			}
+			
 			FEditPropertyChain PropertyChain;
 			TArray<int32> PropertyArrayIndices;
 			bool bEnabled;
@@ -875,7 +935,7 @@ protected:
 				const TWeakObjectPtr<UObject>& Object = ObjectsBeingCustomized[Index];
 				if(Object.Get() && InPropertyHandle->IsValidHandle())
 				{
-					TransformType& Transform = ContainerUObjectToValueRef<TransformType>(Object.Get(), Identity, PropertyChain, PropertyArrayIndices);
+					TransformType& Transform = ContainerMemoryBlockToValueRef<TransformType>((uint8*)Object.Get(), Identity, PropertyChain, PropertyArrayIndices);
 					TransformType PreviousTransform = Transform;
 
 					switch(InTransformComponent)
@@ -930,12 +990,13 @@ protected:
 			{
 				return;
 			}
-			
-			for(const TWeakObjectPtr<UObject>& Object : ObjectsBeingCustomized)
+
+			const TArray<uint8*> MemoryBlocks = GetMemoryBeingCustomized();
+			for(uint8* MemoryBlock: MemoryBlocks)
 			{
-				if(Object.Get() && InPropertyHandle->IsValidHandle())
+				if(MemoryBlock)
 				{
-					const TransformType& Transform = ContainerUObjectToValueRef<TransformType>(Object.Get(), Identity, PropertyChain, PropertyArrayIndices);
+					const TransformType& Transform = ContainerMemoryBlockToValueRef<TransformType>(MemoryBlock, Identity, PropertyChain, PropertyArrayIndices);
 					FString Content;
 					switch(InComponent)
 					{
@@ -985,6 +1046,11 @@ protected:
 				return;
 			}
 
+			if (ObjectsBeingCustomized.IsEmpty())
+			{
+				return;
+			}
+
 			TOptional<FReal> Result;
 			FEditPropertyChain PropertyChain;
 			TArray<int32> PropertyArrayIndices;
@@ -1017,7 +1083,7 @@ protected:
 			{
 				if(Object.Get() && InPropertyHandle->IsValidHandle())
 				{
-					TransformType& Transform = ContainerUObjectToValueRef<TransformType>(Object.Get(), Identity, PropertyChain, PropertyArrayIndices);
+					TransformType& Transform = ContainerMemoryBlockToValueRef<TransformType>((uint8*)Object.Get(), Identity, PropertyChain, PropertyArrayIndices);
 					const TransformType PreviousTransform = Transform;
 
 					// Apply the new value
@@ -1101,7 +1167,8 @@ protected:
 	UScriptStruct* ScriptStruct;
 	URigVMBlueprint* BlueprintBeingCustomized;
 	URigVMGraph* GraphBeingCustomized;
-	TArray<TWeakObjectPtr<UObject>> ObjectsBeingCustomized; 
+	TArray<TWeakObjectPtr<UObject>> ObjectsBeingCustomized;
+	TArray<TSharedPtr<FStructOnScope>> StructsBeingCustomized;
 };
 
 template<>
diff --git a/Engine/Source/Editor/PropertyEditor/Private/PropertyHandleImpl.cpp b/Engine/Source/Editor/PropertyEditor/Private/PropertyHandleImpl.cpp
--- a/Engine/Source/Editor/PropertyEditor/Private/PropertyHandleImpl.cpp
+++ b/Engine/Source/Editor/PropertyEditor/Private/PropertyHandleImpl.cpp
@@ -2599,6 +2599,17 @@ void FPropertyHandleBase::GetOuterObjects( TArray<UObject*>& OuterObjects ) cons
 	}
 }
 
+void FPropertyHandleBase::GetOuterStructs( TArray<TSharedPtr<FStructOnScope>>& OutStructs ) const
+{
+	if (Implementation->GetPropertyNode().IsValid())
+	{
+		if (const FStructurePropertyNode* StructNode = Implementation->GetPropertyNode()->FindStructureItemParent())
+		{
+			StructNode->GetAllStructureData(OutStructs);
+		}
+	}
+}
+
 const UClass* FPropertyHandleBase::GetOuterBaseClass() const
 {
 	if (Implementation->GetPropertyNode().IsValid())
diff --git a/Engine/Source/Editor/PropertyEditor/Private/PropertyHandleImpl.h b/Engine/Source/Editor/PropertyEditor/Private/PropertyHandleImpl.h
--- a/Engine/Source/Editor/PropertyEditor/Private/PropertyHandleImpl.h
+++ b/Engine/Source/Editor/PropertyEditor/Private/PropertyHandleImpl.h
@@ -458,6 +458,7 @@ public:
 	virtual void AccessRawData( TArray<const void*>& RawData ) const override;
 	virtual uint32 GetNumOuterObjects() const override;
 	virtual void GetOuterObjects(TArray<UObject*>& OuterObjects) const override;
+	virtual void GetOuterStructs( TArray<TSharedPtr<FStructOnScope>>& OutStructs ) const override;
 	virtual const UClass* GetOuterBaseClass() const override;
 	virtual void ReplaceOuterObjects(const TArray<UObject*>& OuterObjects) override;
 	virtual void GetOuterPackages(TArray<UPackage*>& OuterPackages) const override;
diff --git a/Engine/Source/Editor/PropertyEditor/Public/PropertyHandle.h b/Engine/Source/Editor/PropertyEditor/Public/PropertyHandle.h
--- a/Engine/Source/Editor/PropertyEditor/Public/PropertyHandle.h
+++ b/Engine/Source/Editor/PropertyEditor/Public/PropertyHandle.h
@@ -465,6 +465,13 @@ public:
 	 */
 	virtual void GetOuterObjects( TArray<UObject*>& OuterObjects ) const = 0;
 
+	/**
+	 * Get the structures that contain this property 
+	 *
+	 * @param OutStructs	An array that will be populated with the outer structures 
+	 */
+	virtual void GetOuterStructs( TArray<TSharedPtr<FStructOnScope>>& OutStructs ) const = 0;
+
 	/**
 	 * Get the shared base class of the objects that contain this property.
 	 *
