RigVM Nativization - interim changes

diff --git a/Engine/Plugins/Animation/ControlRig/Source/ControlRigEditor/Private/ControlRigCompilerDetails.cpp b/Engine/Plugins/Animation/ControlRig/Source/ControlRigEditor/Private/ControlRigCompilerDetails.cpp
--- a/Engine/Plugins/Animation/ControlRig/Source/ControlRigEditor/Private/ControlRigCompilerDetails.cpp
+++ b/Engine/Plugins/Animation/ControlRig/Source/ControlRigEditor/Private/ControlRigCompilerDetails.cpp
@@ -262,7 +262,7 @@ FReply FRigVMCompileSettingsDetails::OnCopyGeneratedCodeClicked()
 				if(CDO->GetVM())
 				{
 					CDO->GetVM()->ClearExternalVariables(CDO->GetRigVMExtendedExecuteContext());
-					TArray<FRigVMExternalVariable> ExternalVariables = CDO->GetExternalVariables();
+					TArray<FRigVMExternalVariable> ExternalVariables = CDO->GetExternalVariablesImpl(false);
 					for(const FRigVMExternalVariable& ExternalVariable : ExternalVariables)
 					{
 						CDO->GetVM()->AddExternalVariable(CDO->GetRigVMExtendedExecuteContext(), ExternalVariable);
diff --git a/Engine/Plugins/Runtime/RigVM/Source/RigVM/Public/RigVMHost.h b/Engine/Plugins/Runtime/RigVM/Source/RigVM/Public/RigVMHost.h
--- a/Engine/Plugins/Runtime/RigVM/Source/RigVM/Public/RigVMHost.h
+++ b/Engine/Plugins/Runtime/RigVM/Source/RigVM/Public/RigVMHost.h
@@ -520,6 +520,7 @@ protected:
 
 	friend class URigVMBlueprint;
 	friend class URigVMBlueprintGeneratedClass;
+	friend class FRigVMCompileSettingsDetails;
 };
 
 class RIGVM_API FRigVMBracketScope
diff --git a/Engine/Plugins/Runtime/RigVM/Source/RigVMDeveloper/Private/RigVMCompiler/RigVMCodeGenerator.cpp b/Engine/Plugins/Runtime/RigVM/Source/RigVMDeveloper/Private/RigVMCompiler/RigVMCodeGenerator.cpp
--- a/Engine/Plugins/Runtime/RigVM/Source/RigVMDeveloper/Private/RigVMCompiler/RigVMCodeGenerator.cpp
+++ b/Engine/Plugins/Runtime/RigVM/Source/RigVMDeveloper/Private/RigVMCompiler/RigVMCodeGenerator.cpp
@@ -18,7 +18,7 @@ static constexpr TCHAR RigVM_UPropertyDeclareFormat[] = TEXT("\tUPROPERTY()\r\n\
 static constexpr TCHAR RigVM_UPropertyMemberFormat[] = TEXT("\tstatic const FProperty* {0}_Ptr;");
 static constexpr TCHAR RigVM_UPropertyMember2Format[] = TEXT("const FProperty* U{0}::{1}_Ptr = nullptr;");
 static constexpr TCHAR RigVM_UPropertyDefineFormat[] = TEXT("\tif({1}_Ptr == nullptr)\r\n\t{\r\n\t\t{1}_Ptr = StaticClass()->FindPropertyByName(GET_MEMBER_NAME_CHECKED(U{0}, {1}));\r\n\t}");
-static constexpr TCHAR RigVM_InvokeDispatchPtrFormat[] = TEXT("\t\t(*Dispatch->FunctionPtr)({0});");
+static constexpr TCHAR RigVM_InvokeDispatchPtrFormat[] = TEXT("\t\t(*Dispatch->FunctionPtr)({0}, FRigVMPredicateBranchArray()); // todo: predicates are not implemented yet");
 static constexpr TCHAR RigVM_InvokeDispatchFormat[] = TEXT("\t{0}({1});");
 static constexpr TCHAR RigVM_WrappedArrayTypeFormat[] = TEXT("struct {0}_API {1}\r\n{\r\n\tTArray<{2}> Array;\r\n};");
 static constexpr TCHAR RigVM_WrappedTypeNameFormat[] = TEXT("{0}Array_{1}");
@@ -100,10 +100,10 @@ static constexpr TCHAR RigVM_DeclareUpdateExternalVariablesFormat[] = TEXT("\tvi
 static constexpr TCHAR RigVM_DeclareInvokeEntryByNameFormat[] = TEXT("\tERigVMExecuteResult InvokeEntryByName(FRigVMExtendedExecuteContext& Context, const FName& InEntryName{0});");
 static constexpr TCHAR RigVM_DeclareInitializeFormat[] = TEXT("\tvirtual bool Initialize(FRigVMExtendedExecuteContext& Context) override;");
 static constexpr TCHAR RigVM_DefineInitializeFormat[] = TEXT("bool U{0}::Initialize(FRigVMExtendedExecuteContext& Context)\r\n{");
-static constexpr TCHAR RigVM_DeclareExecuteFormat[] = TEXT("\tvirtual ERigVMExecuteResult Execute(FRigVMExtendedExecuteContext& Context, const FName& InEntryName) override;");
+static constexpr TCHAR RigVM_DeclareExecuteFormat[] = TEXT("\tvirtual ERigVMExecuteResult ExecuteVM(FRigVMExtendedExecuteContext& Context, const FName& InEntryName = NAME_None) override;");
 static constexpr TCHAR RigVM_DefineUpdateExternalVariablesFormat[] = TEXT("void U{0}::UpdateExternalVariables(FRigVMExtendedExecuteContext& Context)\r\n{");
 static constexpr TCHAR RigVM_DefineInvokeEntryByNameFormat[] = TEXT("ERigVMExecuteResult U{0}::InvokeEntryByName(FRigVMExtendedExecuteContext& Context, const FName& InEntryName{1})\r\n{");
-static constexpr TCHAR RigVM_DefineExecuteFormat[] = TEXT("ERigVMExecuteResult U{0}::Execute(FRigVMExtendedExecuteContext& Context, const FName& InEntryName)\r\n{");
+static constexpr TCHAR RigVM_DefineExecuteFormat[] = TEXT("ERigVMExecuteResult U{0}::ExecuteVM(FRigVMExtendedExecuteContext& Context, const FName& InEntryName)\r\n{");
 static constexpr TCHAR RigVM_DeclareExecuteEntryFormat[] = TEXT("\tERigVMExecuteResult ExecuteEntry_{0}(FRigVMExtendedExecuteContext& Context, {1});");
 static constexpr TCHAR RigVM_DefineExecuteEntryFormat[] = TEXT("ERigVMExecuteResult U{0}::ExecuteEntry_{1}(FRigVMExtendedExecuteContext& Context, {2})\r\n{");
 static constexpr TCHAR RigVM_DeclareExecuteGroupFormat[] = TEXT("\tERigVMExecuteResult ExecuteGroup_{0}_{1}({2});");
@@ -623,7 +623,7 @@ FString FRigVMCodeGenerator::DumpInstructions(const FRigVMExtendedExecuteContext
 		if (InGroup.RequiredLabels.Contains(InstructionIndex))
 		{
 			// check if the last line was a jump to this label
-			if(Lines.Last().Contains(Format(RigVM_JumpOpFormat, InstructionIndex)))
+			if(!Lines.IsEmpty() && Lines.Last().Contains(Format(RigVM_JumpOpFormat, InstructionIndex)))
 			{
 				Lines.Pop();
 			}
@@ -1149,7 +1149,23 @@ void FRigVMCodeGenerator::ParseInclude(UStruct* InDependency, const FName& InMet
 				FunctionModuleName.Split(TEXT("/"), nullptr, &FunctionModuleName, ESearchCase::IgnoreCase, ESearchDir::FromEnd);
 			}
 			Libraries.AddUnique(FunctionModuleName);
-			Includes.AddUnique(Format(RigVM_JoinFilePathFormat, *FunctionModuleName, *Function->GetModuleRelativeHeaderPath()));
+
+			static const FString PrivatePrefix = TEXT("Private/");
+			static const FString PublicPrefix = TEXT("Public/");
+
+			const FString RelativeHeaderPath = Function->GetModuleRelativeHeaderPath();
+			if(RelativeHeaderPath.StartsWith(PrivatePrefix))
+			{
+				Includes.Add(RelativeHeaderPath.Mid(PrivatePrefix.Len()));
+			}
+			else if(RelativeHeaderPath.StartsWith(PublicPrefix))
+			{
+				Includes.Add(RelativeHeaderPath.Mid(PublicPrefix.Len()));
+			}
+			else
+			{
+				Includes.AddUnique(Format(RigVM_JoinFilePathFormat, *FunctionModuleName, *RelativeHeaderPath));
+			}
 			return;
 		}
 	}
@@ -1258,8 +1274,8 @@ void FRigVMCodeGenerator::ParseProperty(const FRigVMExtendedExecuteContext& Cont
 {
 	const int32 PropertyIndex = InMemory->GetPropertyIndex(InProperty);
 	const FRigVMOperand Operand(InMemoryType, PropertyIndex);
-	const FRigVMPropertyDescription& PropertyDescription = GetPropertyDescForOperand(Context, Operand);
-
+	const FRigVMPropertyDescription& PropertyDescription = GetPropertyForOperand(Context, Operand);
+	
 	FPropertyInfo Info;
 	Info.MemoryPropertyIndex = PropertyIndex;
 	Info.Description = PropertyDescription;
