RigVM: Make sure that instruction tracking is enabled independent of debug / profiling

#rnx

diff --git a/Engine/Plugins/Animation/ControlRig/Source/ControlRig/Private/ControlRig.cpp b/Engine/Plugins/Animation/ControlRig/Source/ControlRig/Private/ControlRig.cpp
--- a/Engine/Plugins/Animation/ControlRig/Source/ControlRig/Private/ControlRig.cpp
+++ b/Engine/Plugins/Animation/ControlRig/Source/ControlRig/Private/ControlRig.cpp
@@ -688,6 +688,8 @@ bool UControlRig::Execute(const FName& InEventName)
 	}
 
 #if WITH_EDITOR
+	ExtendedExecuteContext.SetInstructionVisitInfo(&InstructionVisitInfo);
+
 	if (IsInDebugMode())
 	{
 		if (UControlRig* CDO = GetClass()->GetDefaultObject<UControlRig>())
diff --git a/Engine/Plugins/Runtime/RigVM/Source/RigVM/Private/RigVMCore/RigVM.cpp b/Engine/Plugins/Runtime/RigVM/Source/RigVM/Private/RigVMCore/RigVM.cpp
--- a/Engine/Plugins/Runtime/RigVM/Source/RigVM/Private/RigVMCore/RigVM.cpp
+++ b/Engine/Plugins/Runtime/RigVM/Source/RigVM/Private/RigVMCore/RigVM.cpp
@@ -1505,11 +1505,15 @@ ERigVMExecuteResult URigVM::ExecuteVM(FRigVMExtendedExecuteContext& Context, con
 
 	if(bIsRootEntry)
 	{
-		if (FRigVMProfilingInfo* RigVMProfilingInfo = Context.GetRigVMProfilingInfo())
+		if (FRigVMInstructionVisitInfo* RigVMInstructionVisitInfo = Context.GetRigVMInstructionVisitInfo())
 		{
-			if (RigVMProfilingInfo->GetFirstEntryEventInEventQueue() == NAME_None || RigVMProfilingInfo->GetFirstEntryEventInEventQueue() == InEntryName)
+			if (RigVMInstructionVisitInfo->GetFirstEntryEventInEventQueue() == NAME_None || RigVMInstructionVisitInfo->GetFirstEntryEventInEventQueue() == InEntryName)
 			{
-				RigVMProfilingInfo->SetupInstructionTracking(Instructions.Num(), true);
+				RigVMInstructionVisitInfo->SetupInstructionTracking(Instructions.Num());
+				if (FRigVMProfilingInfo* RigVMProfilingInfo = Context.GetRigVMProfilingInfo())
+				{
+					RigVMProfilingInfo->SetupInstructionTracking(Instructions.Num(), true);
+				}
 			}
 		}
 	}
@@ -1645,6 +1649,7 @@ ERigVMExecuteResult URigVM::ExecuteInstructions(FRigVMExtendedExecuteContext& Co
 
 #if WITH_EDITOR
 	TArray<FName>& FunctionNames = GetFunctionNames();
+	FRigVMInstructionVisitInfo* RigVMInstructionVisitInfo = Context.GetRigVMInstructionVisitInfo();
 	FRigVMDebugInfo* RigVMDebugInfo = Context.GetRigVMDebugInfo();
 	FRigVMProfilingInfo* RigVMProfilingInfo = Context.GetRigVMProfilingInfo();
 #endif
@@ -1681,10 +1686,10 @@ ERigVMExecuteResult URigVM::ExecuteInstructions(FRigVMExtendedExecuteContext& Co
 
 		const int32 CurrentInstructionIndex = ContextPublicData.InstructionIndex;
 
-		if (RigVMProfilingInfo != nullptr)
+		if (RigVMInstructionVisitInfo != nullptr)
 		{
-			RigVMProfilingInfo->SetInstructionVisitedDuringLastRun(ContextPublicData.InstructionIndex);
-			RigVMProfilingInfo->AddInstructionIndexToVisitOrder(ContextPublicData.InstructionIndex);
+			RigVMInstructionVisitInfo->SetInstructionVisitedDuringLastRun(ContextPublicData.InstructionIndex);
+			RigVMInstructionVisitInfo->AddInstructionIndexToVisitOrder(ContextPublicData.InstructionIndex);
 		}
 	
 #endif
@@ -2049,9 +2054,9 @@ ERigVMExecuteResult URigVM::ExecuteInstructions(FRigVMExtendedExecuteContext& Co
 		}
 
 #if WITH_EDITOR
-		if (RigVMProfilingInfo != nullptr)
+		if ((RigVMInstructionVisitInfo != nullptr) && (RigVMProfilingInfo != nullptr))
 		{
-			if (!RigVMProfilingInfo->GetInstructionVisitOrder().IsEmpty())
+			if (!RigVMInstructionVisitInfo->GetInstructionVisitOrder().IsEmpty())
 			{
 				const uint64 EndCycles = FPlatformTime::Cycles64();
 				const uint64 Cycles = EndCycles - RigVMProfilingInfo->GetStartCycles();
diff --git a/Engine/Plugins/Runtime/RigVM/Source/RigVM/Private/RigVMCore/RigVMProfilingInfo.cpp b/Engine/Plugins/Runtime/RigVM/Source/RigVM/Private/RigVMCore/RigVMProfilingInfo.cpp
--- a/Engine/Plugins/Runtime/RigVM/Source/RigVM/Private/RigVMCore/RigVMProfilingInfo.cpp
+++ b/Engine/Plugins/Runtime/RigVM/Source/RigVM/Private/RigVMCore/RigVMProfilingInfo.cpp
@@ -4,15 +4,19 @@
 
 #include UE_INLINE_GENERATED_CPP_BY_NAME(RigVMProfilingInfo)
 
-void FRigVMProfilingInfo::SetupInstructionTracking(int32 InInstructionCount, bool bEnableProfiling)
+void FRigVMInstructionVisitInfo::SetupInstructionTracking(int32 InInstructionCount)
 {
 #if WITH_EDITOR
-
 	ResetInstructionVisitedDuringLastRun(InInstructionCount);
 	SetNumInstructionVisitedDuringLastRunZeroed(InInstructionCount);
 	ResetInstructionVisitOrder(InInstructionCount);
-	ResetInstructionCyclesDuringLastRun(InInstructionCount);
+#endif
+}
 
+void FRigVMProfilingInfo::SetupInstructionTracking(int32 InInstructionCount, bool bEnableProfiling)
+{
+#if WITH_EDITOR
+	ResetInstructionCyclesDuringLastRun(InInstructionCount);
 	if (bEnableProfiling)
 	{
 		InitInstructionCyclesDuringLastRunValues(InInstructionCount, UINT64_MAX);
diff --git a/Engine/Plugins/Runtime/RigVM/Source/RigVM/Private/RigVMHost.cpp b/Engine/Plugins/Runtime/RigVM/Source/RigVM/Private/RigVMHost.cpp
--- a/Engine/Plugins/Runtime/RigVM/Source/RigVM/Private/RigVMHost.cpp
+++ b/Engine/Plugins/Runtime/RigVM/Source/RigVM/Private/RigVMHost.cpp
@@ -501,6 +501,8 @@ bool URigVMHost::Execute(const FName& InEventName)
 #endif
 
 #if WITH_EDITOR
+	ExtendedExecuteContext.SetInstructionVisitInfo(&InstructionVisitInfo);
+
 	if (IsInDebugMode())
 	{
 		if (CDO)
diff --git a/Engine/Plugins/Runtime/RigVM/Source/RigVM/Public/RigVMCore/RigVM.h b/Engine/Plugins/Runtime/RigVM/Source/RigVM/Public/RigVMCore/RigVM.h
--- a/Engine/Plugins/Runtime/RigVM/Source/RigVM/Public/RigVMCore/RigVM.h
+++ b/Engine/Plugins/Runtime/RigVM/Source/RigVM/Public/RigVMCore/RigVM.h
@@ -445,9 +445,9 @@ public:
 	// Returns the number of times an instruction has been hit
 	int32 GetInstructionVisitedCount(const FRigVMExtendedExecuteContext& Context, int32 InIndex) const
 	{
-		if (const FRigVMProfilingInfo* RigVMProfilingInfo = Context.GetRigVMProfilingInfo())
+		if (const FRigVMInstructionVisitInfo* InstructionVisitInfo = Context.GetRigVMInstructionVisitInfo())
 		{
-			return RigVMProfilingInfo->GetInstructionVisitedCountDuringLastRun(InIndex);
+			return InstructionVisitInfo->GetInstructionVisitedCountDuringLastRun(InIndex);
 		}
 		return 0;
 	}
@@ -489,9 +489,9 @@ public:
 	// Returns the order of all instructions during the last run
 	const TArray<int32> GetInstructionVisitOrder(const FRigVMExtendedExecuteContext& Context) const
 	{
-		if (const FRigVMProfilingInfo* RigVMProfilingInfo = Context.GetRigVMProfilingInfo())
+		if (const FRigVMInstructionVisitInfo* InstructionVisitInfo = Context.GetRigVMInstructionVisitInfo())
 		{
-			return RigVMProfilingInfo->GetInstructionVisitOrder();
+			return InstructionVisitInfo->GetInstructionVisitOrder();
 		}
 
 		return TArray<int32>();
@@ -502,9 +502,9 @@ public:
 
 	const void SetFirstEntryEventInEventQueue(FRigVMExtendedExecuteContext& Context, const FName& InFirstEventName) 
 	{ 
-		if (FRigVMProfilingInfo* RigVMProfilingInfo = Context.GetRigVMProfilingInfo())
+		if (FRigVMInstructionVisitInfo* InstructionVisitInfo = Context.GetRigVMInstructionVisitInfo())
 		{
-			RigVMProfilingInfo->SetFirstEntryEventInEventQueue(InFirstEventName);
+			InstructionVisitInfo->SetFirstEntryEventInEventQueue(InFirstEventName);
 		}
 	}
 
diff --git a/Engine/Plugins/Runtime/RigVM/Source/RigVM/Public/RigVMCore/RigVMExecuteContext.h b/Engine/Plugins/Runtime/RigVM/Source/RigVM/Public/RigVMCore/RigVMExecuteContext.h
--- a/Engine/Plugins/Runtime/RigVM/Source/RigVM/Public/RigVMCore/RigVMExecuteContext.h
+++ b/Engine/Plugins/Runtime/RigVM/Source/RigVM/Public/RigVMCore/RigVMExecuteContext.h
@@ -776,6 +776,18 @@ struct RIGVM_API FRigVMExtendedExecuteContext
 
 	FRigVMDebugInfo* DebugInfo = nullptr;
 
+	FRigVMInstructionVisitInfo* InstructionVisitInfo = nullptr;
+
+	void SetInstructionVisitInfo(FRigVMInstructionVisitInfo* InInstructionVisitInfo)
+	{
+		InstructionVisitInfo = InInstructionVisitInfo;
+	}
+
+	FRigVMInstructionVisitInfo* GetRigVMInstructionVisitInfo() const
+	{
+		return InstructionVisitInfo;
+	}
+
 	void SetDebugInfo(FRigVMDebugInfo* InDebugInfo)
 	{
 		DebugInfo = InDebugInfo;
diff --git a/Engine/Plugins/Runtime/RigVM/Source/RigVM/Public/RigVMCore/RigVMNativized.h b/Engine/Plugins/Runtime/RigVM/Source/RigVM/Public/RigVMCore/RigVMNativized.h
--- a/Engine/Plugins/Runtime/RigVM/Source/RigVM/Public/RigVMCore/RigVMNativized.h
+++ b/Engine/Plugins/Runtime/RigVM/Source/RigVM/Public/RigVMCore/RigVMNativized.h
@@ -64,10 +64,10 @@ public:
 	{
 		Super::SetInstructionIndex(Context, InInstructionIndex);
 #if WITH_EDITOR
-		if (FRigVMProfilingInfo* RigVMProfilingInfo = Context.GetRigVMProfilingInfo())
+		if (FRigVMInstructionVisitInfo* InstructionVisitInfo = Context.GetRigVMInstructionVisitInfo())
 		{
-			RigVMProfilingInfo->SetInstructionVisitedDuringLastRun(InInstructionIndex);
-			RigVMProfilingInfo->AddInstructionIndexToVisitOrder(InInstructionIndex);
+			InstructionVisitInfo->SetInstructionVisitedDuringLastRun(InInstructionIndex);
+			InstructionVisitInfo->AddInstructionIndexToVisitOrder(InInstructionIndex);
 		}
 #endif
 	}
diff --git a/Engine/Plugins/Runtime/RigVM/Source/RigVM/Public/RigVMCore/RigVMProfilingInfo.h b/Engine/Plugins/Runtime/RigVM/Source/RigVM/Public/RigVMCore/RigVMProfilingInfo.h
--- a/Engine/Plugins/Runtime/RigVM/Source/RigVM/Public/RigVMCore/RigVMProfilingInfo.h
+++ b/Engine/Plugins/Runtime/RigVM/Source/RigVM/Public/RigVMCore/RigVMProfilingInfo.h
@@ -10,21 +10,59 @@
 #include "RigVMProfilingInfo.generated.h"
 
 USTRUCT()
-struct RIGVM_API FRigVMProfilingInfo
+struct RIGVM_API FRigVMInstructionVisitInfo
 {
 	GENERATED_BODY()
 
-	FRigVMProfilingInfo()
+	FRigVMInstructionVisitInfo()
 	{
 	}
 
 	void Reset()
 	{
 		InstructionVisitedDuringLastRun.Reset();
-		InstructionCyclesDuringLastRun.Reset();
 		InstructionVisitOrder.Reset();
 		FirstEntryEventInQueue = NAME_None;
+	}
+
+	inline void ResetInstructionVisitedDuringLastRun(int32 NewSize = 0) { InstructionVisitedDuringLastRun.Reset(NewSize); }
+	inline void SetNumInstructionVisitedDuringLastRunZeroed(int32 Num) { InstructionVisitedDuringLastRun.SetNumZeroed(Num); }
+	inline void SetInstructionVisitedDuringLastRun(int32 InstructionIndex) { InstructionVisitedDuringLastRun[InstructionIndex]++; }
+	inline int32 GetInstructionVisitedCountDuringLastRun(int32 InstructionIndex) const { return InstructionVisitedDuringLastRun.IsValidIndex(InstructionIndex) ? InstructionVisitedDuringLastRun[InstructionIndex] : 0; }
+	inline const TArray<int32>& GetInstructionVisitedCountDuringLastRun() const { return InstructionVisitedDuringLastRun; }
+
+	inline void ResetInstructionVisitOrder(int32 NewSize = 0) { InstructionVisitOrder.Reset(NewSize); }
+	inline void AddInstructionIndexToVisitOrder(int32 InstructionIndex) { InstructionVisitOrder.Add(InstructionIndex); }
+	inline const TArray<int32>& GetInstructionVisitOrder() const { return InstructionVisitOrder; }
+
+	inline const void SetFirstEntryEventInEventQueue(const FName& InFirstEventName) { FirstEntryEventInQueue = InFirstEventName; }
+	inline const FName& GetFirstEntryEventInEventQueue() const { return FirstEntryEventInQueue; }
+
+	void SetupInstructionTracking(int32 InInstructionCount);
+
+private:
+
+	// stores the number of times each instruction was visited
+	TArray<int32> InstructionVisitedDuringLastRun;
+	TArray<int32> InstructionVisitOrder;
+
+	// A RigVMHost can run multiple events per evaluation, such as the Backward&Forward Solve Mode,
+	// store the first event such that we know when to reset data for a new round of rig evaluation
+	FName FirstEntryEventInQueue = NAME_None;
+};
+
+USTRUCT()
+struct RIGVM_API FRigVMProfilingInfo
+{
+	GENERATED_BODY()
+
+	FRigVMProfilingInfo()
+	{
+	}
 
+	void Reset()
+	{
+		InstructionCyclesDuringLastRun.Reset();
 		StartCycles = 0;
 		OverallCycles = 0;
 	}
@@ -36,12 +74,6 @@ struct RIGVM_API FRigVMProfilingInfo
 	inline void SetOverallCycles(uint64 Cycles) { OverallCycles = Cycles; }
 	inline void AddOverallCycles(uint64 Cycles) { OverallCycles += Cycles; }
 
-	inline void ResetInstructionVisitedDuringLastRun(int32 NewSize = 0) { InstructionVisitedDuringLastRun.Reset(NewSize); }
-	inline void SetNumInstructionVisitedDuringLastRunZeroed(int32 Num) { InstructionVisitedDuringLastRun.SetNumZeroed(Num); }
-	inline void SetInstructionVisitedDuringLastRun(int32 InstructionIndex) { InstructionVisitedDuringLastRun[InstructionIndex]++; }
-	inline int32 GetInstructionVisitedCountDuringLastRun(int32 InstructionIndex) const { return InstructionVisitedDuringLastRun.IsValidIndex(InstructionIndex) ? InstructionVisitedDuringLastRun[InstructionIndex] : 0; }
-	inline const TArray<int32>& GetInstructionVisitedCountDuringLastRun() const { return InstructionVisitedDuringLastRun; }
-
 	inline void ResetInstructionCyclesDuringLastRun(int32 NewSize = 0) { InstructionCyclesDuringLastRun.Reset(NewSize); }
 	inline uint64 GetInstructionCyclesDuringLastRun(int32 InstructionIndex) const { return InstructionCyclesDuringLastRun.IsValidIndex(InstructionIndex) ? InstructionCyclesDuringLastRun[InstructionIndex] : UINT64_MAX; }
 	inline void SetInstructionCyclesDuringLastRun(int32 InstructionIndex, uint64 CyclesDuringLastRun) { InstructionCyclesDuringLastRun[InstructionIndex] = CyclesDuringLastRun; }
@@ -55,14 +87,6 @@ struct RIGVM_API FRigVMProfilingInfo
 		}
 	}
 
-	
-	inline void ResetInstructionVisitOrder(int32 NewSize = 0) { InstructionVisitOrder.Reset(NewSize); }
-	inline void AddInstructionIndexToVisitOrder(int32 InstructionIndex) { InstructionVisitOrder.Add(InstructionIndex); }
-	inline const TArray<int32>& GetInstructionVisitOrder() const { return InstructionVisitOrder; }
-
-	inline const void SetFirstEntryEventInEventQueue(const FName& InFirstEventName) { FirstEntryEventInQueue = InFirstEventName; }
-	inline const FName& GetFirstEntryEventInEventQueue() const { return FirstEntryEventInQueue; }
-
 	void SetupInstructionTracking(int32 InInstructionCount, bool bEnableProfiling);
 
 	double GetLastExecutionMicroSeconds() const { return LastExecutionMicroSeconds; }
@@ -74,13 +98,7 @@ struct RIGVM_API FRigVMProfilingInfo
 private:
 
 	// stores the number of times each instruction was visited
-	TArray<int32> InstructionVisitedDuringLastRun;
 	TArray<uint64> InstructionCyclesDuringLastRun;
-	TArray<int32> InstructionVisitOrder;
-
-	// A RigVMHost can run multiple events per evaluation, such as the Backward&Forward Solve Mode,
-	// store the first event such that we know when to reset data for a new round of rig evaluation
-	FName FirstEntryEventInQueue = NAME_None;
 
 	uint64 StartCycles = 0;
 	uint64 OverallCycles = 0;
diff --git a/Engine/Plugins/Runtime/RigVM/Source/RigVM/Public/RigVMHost.h b/Engine/Plugins/Runtime/RigVM/Source/RigVM/Public/RigVMHost.h
--- a/Engine/Plugins/Runtime/RigVM/Source/RigVM/Public/RigVMHost.h
+++ b/Engine/Plugins/Runtime/RigVM/Source/RigVM/Public/RigVMHost.h
@@ -464,6 +464,7 @@ private:
 
 protected:
 	
+	FRigVMInstructionVisitInfo InstructionVisitInfo;
 	FRigVMDebugInfo DebugInfo;
 	FRigVMProfilingInfo ProfilingInfo;
 	TMap<FString, bool> LoggedMessages;
