RigVM: Fix cannot resolve wildcard for array element pins
#rnx

diff --git a/Engine/Plugins/Runtime/RigVM/Source/RigVMEditor/Private/RigVMEditorModule.cpp b/Engine/Plugins/Runtime/RigVM/Source/RigVMEditor/Private/RigVMEditorModule.cpp
--- a/Engine/Plugins/Runtime/RigVM/Source/RigVMEditor/Private/RigVMEditorModule.cpp
+++ b/Engine/Plugins/Runtime/RigVM/Source/RigVMEditor/Private/RigVMEditorModule.cpp
@@ -1174,7 +1174,8 @@ void FRigVMEditorModule::GetPinTemplateContextMenuActions(IRigVMClientHost* RigV
 			{
 				if(!ModelPin->IsExecuteContext())
 				{
-					if(const FRigVMTemplateArgument* Argument = Template->FindArgument(ModelPin->GetRootPin()->GetFName()))
+					URigVMPin* RootPin = ModelPin->GetRootPin();
+					if(const FRigVMTemplateArgument* Argument = Template->FindArgument(RootPin->GetFName()))
 					{
 						if(!Argument->IsSingleton())
 						{
@@ -1182,11 +1183,11 @@ void FRigVMEditorModule::GetPinTemplateContextMenuActions(IRigVMClientHost* RigV
 							TSharedRef<SRigVMGraphChangePinType> ChangePinTypeWidget =
 							SNew(SRigVMGraphChangePinType)
 							.Types(ResolvedTypeIndices)
-							.OnTypeSelected_Lambda([RigVMClientHost, ModelPin](const TRigVMTypeIndex& TypeSelected)
+							.OnTypeSelected_Lambda([RigVMClientHost, RootPin](const TRigVMTypeIndex& TypeSelected)
 							{
-								if (URigVMController* Controller = RigVMClientHost->GetRigVMClient()->GetController(ModelPin->GetGraph()))
+								if (URigVMController* Controller = RigVMClientHost->GetRigVMClient()->GetController(RootPin->GetGraph()))
 								{
-									Controller->ResolveWildCardPin(ModelPin, TypeSelected, true, true);
+									Controller->ResolveWildCardPin(RootPin, TypeSelected, true, true);
 								}
 							});
 
