RigVM: Harden against invalid/malicious bytecode.

#rnx

diff --git a/Engine/Plugins/Runtime/RigVM/Source/RigVM/Private/RigVMCore/RigVM.cpp b/Engine/Plugins/Runtime/RigVM/Source/RigVM/Private/RigVMCore/RigVM.cpp
--- a/Engine/Plugins/Runtime/RigVM/Source/RigVM/Private/RigVMCore/RigVM.cpp
+++ b/Engine/Plugins/Runtime/RigVM/Source/RigVM/Private/RigVMCore/RigVM.cpp
@@ -412,18 +412,15 @@ UClass* URigVM::GetNativizedClass(const TArray<FRigVMExternalVariableDef>& InExt
 	return nullptr;
 }
 
-bool URigVM::ValidateAllOperandsDuringLoad()
+bool URigVM::ValidateBytecode()
 {
 	// check all operands on all ops for validity
-	bool bAllOperandsValid = true;
-
 	const TArray<const FRigVMMemoryStorageStruct*> LocalMemory = { &GetDefaultWorkMemory(), &GetDefaultLiteralMemory(), &GetDefaultDebugMemory() };
 	
-	auto CheckOperandValidity = [LocalMemory, &bAllOperandsValid, this](const FRigVMOperand& InOperand) -> bool
+	auto CheckOperandValidity = [LocalMemory, this](const FRigVMOperand& InOperand) -> bool
 	{
 		if(InOperand.GetContainerIndex() < 0 || InOperand.GetContainerIndex() >= (int32)ERigVMMemoryType::Invalid)
 		{
-			bAllOperandsValid = false;
 			return false;
 		}
 
@@ -433,7 +430,6 @@ bool URigVM::ValidateAllOperandsDuringLoad()
 		{
 			if(!MemoryForOperand->IsValidIndex(InOperand.GetRegisterIndex()))
 			{
-				bAllOperandsValid = false;
 				return false;
 			}
 
@@ -441,7 +437,6 @@ bool URigVM::ValidateAllOperandsDuringLoad()
 			{
 				if(!MemoryForOperand->GetPropertyPaths().IsValidIndex(InOperand.GetRegisterOffset()))
 				{
-					bAllOperandsValid = false;
 					return false;
 				}
 			}
@@ -454,13 +449,14 @@ bool URigVM::ValidateAllOperandsDuringLoad()
 			{
 				if (!ExternalPropertyPathDescriptions.IsValidIndex(InOperand.GetRegisterOffset()))
 				{
-					bAllOperandsValid = false;
 					return false;
 				}
 			}
 		}
 		return true;
 	};
+
+	const TArray<const FRigVMFunction*>& Functions = GetFunctions();
 	
 	const FRigVMInstructionArray ByteCodeInstructions = ByteCodeStorage.GetInstructions();
 	for(const FRigVMInstruction& ByteCodeInstruction : ByteCodeInstructions)
@@ -470,10 +466,17 @@ bool URigVM::ValidateAllOperandsDuringLoad()
 			case ERigVMOpCode::Execute:
 			{
 				const FRigVMExecuteOp& Op = ByteCodeStorage.GetOpAt<FRigVMExecuteOp>(ByteCodeInstruction);
+				if (!Functions.IsValidIndex(Op.FunctionIndex))
+				{
+					return false;
+				}
 				FRigVMOperandArray Operands = ByteCodeStorage.GetOperandsForExecuteOp(ByteCodeInstruction);
 				for (const FRigVMOperand& Arg : Operands)
 				{
-					CheckOperandValidity(Arg);
+					if (!CheckOperandValidity(Arg))
+					{
+						return false;
+					}
 				}
 				break;
 			}
@@ -484,23 +487,32 @@ bool URigVM::ValidateAllOperandsDuringLoad()
 			case ERigVMOpCode::Decrement:
 			{
 				const FRigVMUnaryOp& Op = ByteCodeStorage.GetOpAt<FRigVMUnaryOp>(ByteCodeInstruction);
-				CheckOperandValidity(Op.Arg);
+				if (!CheckOperandValidity(Op.Arg))
+				{
+					return false;
+				}
 				break;
 			}
 			case ERigVMOpCode::Copy:
 			{
 				const FRigVMCopyOp& Op = ByteCodeStorage.GetOpAt<FRigVMCopyOp>(ByteCodeInstruction);
-				CheckOperandValidity(Op.Source);
-				CheckOperandValidity(Op.Target);
+				if (!CheckOperandValidity(Op.Source) ||
+					!CheckOperandValidity(Op.Target))
+				{
+					return false;
+				}
 				break;
 			}
 			case ERigVMOpCode::Equals:
 			case ERigVMOpCode::NotEquals:
 			{
 				const FRigVMComparisonOp& Op = ByteCodeStorage.GetOpAt<FRigVMComparisonOp>(ByteCodeInstruction);
-				CheckOperandValidity(Op.A);
-				CheckOperandValidity(Op.B);
-				CheckOperandValidity(Op.Result);
+				if (!CheckOperandValidity(Op.A) ||
+					!CheckOperandValidity(Op.B) ||
+					!CheckOperandValidity(Op.Result))
+				{
+					return false;
+				}
 				break;
 			}
 			case ERigVMOpCode::JumpAbsoluteIf:
@@ -508,20 +520,29 @@ bool URigVM::ValidateAllOperandsDuringLoad()
 			case ERigVMOpCode::JumpBackwardIf:
 			{
 				const FRigVMJumpIfOp& Op = ByteCodeStorage.GetOpAt<FRigVMJumpIfOp>(ByteCodeInstruction);
-				CheckOperandValidity(Op.Arg);
+				if (!CheckOperandValidity(Op.Arg))
+				{
+					return false;
+				}
 				break;
 			}
 			case ERigVMOpCode::BeginBlock:
 			{
 				const FRigVMBinaryOp& Op = ByteCodeStorage.GetOpAt<FRigVMBinaryOp>(ByteCodeInstruction);
-				CheckOperandValidity(Op.ArgA);
-				CheckOperandValidity(Op.ArgB);
+				if (!CheckOperandValidity(Op.ArgA) ||
+					!CheckOperandValidity(Op.ArgB))
+				{
+					return false;
+				}
 				break;
 			}
 			case ERigVMOpCode::JumpToBranch:
 			{
 				const FRigVMJumpToBranchOp& Op = ByteCodeStorage.GetOpAt<FRigVMJumpToBranchOp>(ByteCodeInstruction);
-				CheckOperandValidity(Op.Arg);
+				if (!CheckOperandValidity(Op.Arg))
+				{
+					return false;
+				}
 				break;
 			}
 			case ERigVMOpCode::Invalid:
@@ -536,7 +557,7 @@ bool URigVM::ValidateAllOperandsDuringLoad()
 		}
 	}
 
-	return bAllOperandsValid;
+	return true;
 }
 
 const URigVMHost* URigVM::GetHostCDO() const
diff --git a/Engine/Plugins/Runtime/RigVM/Source/RigVM/Private/RigVMHost.cpp b/Engine/Plugins/Runtime/RigVM/Source/RigVM/Private/RigVMHost.cpp
--- a/Engine/Plugins/Runtime/RigVM/Source/RigVM/Private/RigVMHost.cpp
+++ b/Engine/Plugins/Runtime/RigVM/Source/RigVM/Private/RigVMHost.cpp
@@ -69,7 +69,12 @@ TArray<URigVMHost*> URigVMHost::FindRigVMHosts(UObject* Outer, TSubclassOf<URigV
 	for (TObjectIterator<URigVMHost> Itr; Itr; ++Itr)
 	{
 		URigVMHost* RigInstance = *Itr;
-		const UClass* RigInstanceClass = RigInstance ? RigInstance->GetClass() : nullptr;
+		if (!RigInstance)
+		{
+			continue;
+		}
+		
+		const UClass* RigInstanceClass = RigInstance->GetClass();
 		if (OptionalClass == nullptr || (RigInstanceClass && RigInstanceClass->IsChildOf(OptionalClass)))
 		{
 			if(RigInstance->IsInOuter(Outer))
@@ -148,8 +153,9 @@ void URigVMHost::PostLoad()
 			InitializeCDOVM();
 		}
 
-		if (!ensure(VM->ValidateAllOperandsDuringLoad()))
+		if (!ensure(VM->ValidateBytecode()))
 		{
+			UE_LOG(LogRigVM, Warning, TEXT("%s: Invalid bytecode detected. VM will be reset."), *GetPathName());
 			VM->Reset(ExtendedExecuteContext);
 		}
 	}
diff --git a/Engine/Plugins/Runtime/RigVM/Source/RigVM/Public/RigVMCore/RigVM.h b/Engine/Plugins/Runtime/RigVM/Source/RigVM/Public/RigVMCore/RigVM.h
--- a/Engine/Plugins/Runtime/RigVM/Source/RigVM/Public/RigVMCore/RigVM.h
+++ b/Engine/Plugins/Runtime/RigVM/Source/RigVM/Public/RigVMCore/RigVM.h
@@ -992,7 +992,7 @@ public:
 		return ActiveExecutions.load();
 	}
 
-	bool ValidateAllOperandsDuringLoad();
+	bool ValidateBytecode();
 	void RefreshArgumentNameCaches();
 
 private:
