Fix warning V586: The 'Reset' function is called twice for deallocation of the same resource.
#rnx

diff --git a/Engine/Plugins/Runtime/RigVM/Source/RigVM/Private/RigVMCore/RigVM.cpp b/Engine/Plugins/Runtime/RigVM/Source/RigVM/Private/RigVMCore/RigVM.cpp
--- a/Engine/Plugins/Runtime/RigVM/Source/RigVM/Private/RigVMCore/RigVM.cpp
+++ b/Engine/Plugins/Runtime/RigVM/Source/RigVM/Private/RigVMCore/RigVM.cpp
@@ -586,7 +586,7 @@ void URigVM::Reset_Internal()
 	FactoriesPtr = &FactoriesStorage;
 	ByteCodePtr = &ByteCodeStorage;
 
-	ExternalPropertyPaths.Reset();
+	ExternalVariables.Reset();
 	LazyBranches.Reset();
 
 	InvalidateCachedMemory_Internal();
