RigVM - Optimizing template memory (part #2 - reduces memory footprint by around 70%)
    - take advantage of categories to avoid copying unnecessary types
    - eliminate duplication of unnecessary types for singletons

TODO
   - permutations and permutation hashes could probably be computed lazilly to save additional memory too

diff --git a/Engine/Plugins/Runtime/RigVM/Source/RigVM/Private/RigVMCore/RigVMDispatchFactory.cpp b/Engine/Plugins/Runtime/RigVM/Source/RigVM/Private/RigVMCore/RigVMDispatchFactory.cpp
--- a/Engine/Plugins/Runtime/RigVM/Source/RigVM/Private/RigVMCore/RigVMDispatchFactory.cpp
+++ b/Engine/Plugins/Runtime/RigVM/Source/RigVM/Private/RigVMCore/RigVMDispatchFactory.cpp
@@ -299,22 +299,37 @@ TArray<FRigVMTemplateArgumentInfo> FRigVMDispatchFactory::BuildArgumentListFromP
 	TypeIndicesArray.SetNum(NumInfos);
 
 	const FRigVMTemplateArgument PrimaryArgument = PrimaryInfo->GetArgument();
-	for (const TRigVMTypeIndex& Type : PrimaryArgument.TypeIndices)
+	bool bFoundArg = true;
+	PrimaryArgument.ForEachType([&](const TRigVMTypeIndex Type)
 	{
-		const TArray<FRigVMTemplateTypeMap> Permutations = GetPermutationsFromArgumentType(InPrimaryArgumentName, Type);
-		for (const FRigVMTemplateTypeMap& Permutation : Permutations)
+		if (bFoundArg)
 		{
-			for (int32 Index=0; Index < InInfos.Num(); ++Index)
+			const TArray<FRigVMTemplateTypeMap> Permutations = GetPermutationsFromArgumentType(InPrimaryArgumentName, Type);
+			for (const FRigVMTemplateTypeMap& Permutation : Permutations)
 			{
-				const TRigVMTypeIndex* PermutationArg = Permutation.Find(InInfos[Index].Name);
-				if (!PermutationArg)
+				for (int32 Index=0; Index < InInfos.Num(); ++Index)
 				{
-					NewInfos.Reset();
-					return NewInfos;
+					if (const TRigVMTypeIndex* PermutationArg = Permutation.Find(InInfos[Index].Name))
+					{
+						TypeIndicesArray[Index].Add(*PermutationArg);
+					}
+					else
+					{
+						bFoundArg = false;
+						break;
+					}
+				}
+				if (!bFoundArg)
+				{
+					break;
 				}
-				TypeIndicesArray[Index].Add(*PermutationArg);
 			}
 		}
+	});
+
+	if (!bFoundArg)
+	{
+		return NewInfos;	
 	}
 
 	NewInfos.Reserve(NumInfos);
@@ -322,23 +337,14 @@ TArray<FRigVMTemplateArgumentInfo> FRigVMDispatchFactory::BuildArgumentListFromP
 	{
 		const FRigVMTemplateArgument Argument = InInfos[Index].GetArgument();
 		const TArray<TRigVMTypeIndex>& TypeIndices = TypeIndicesArray[Index];
-
-		const FName Name = InInfos[Index].Name;
-		const ERigVMPinDirection Direction = Argument.Direction;
-		const TArray<FRigVMTemplateArgument::ETypeCategory> TypeCategories = Argument.TypeCategories;
-		
-		NewInfos.Emplace(Name, Direction, [TypeIndices, TypeCategories](const FName InName, ERigVMPinDirection InDirection)
+		if (TypeIndices.IsEmpty())
+		{
+			NewInfos.Emplace(InInfos[Index].Name, Argument.Direction, Argument.TypeCategories);
+		}
+		else
 		{
-			FRigVMTemplateArgument Argument;
-			Argument.Name = InName;
-			Argument.Direction = InDirection;
-			Argument.TypeCategories = TypeCategories;
-			Argument.TypeIndices = TypeIndices;
-			
-			Argument.EnsureValidExecuteType();
-			Argument.UpdateTypeToPermutations();
-			return Argument;
-		} );
+			NewInfos.Emplace(InInfos[Index].Name, Argument.Direction, TypeIndices);
+		}
 	}
 	
 	return NewInfos;
@@ -375,21 +381,22 @@ const FRigVMTemplate* FRigVMDispatchFactory::GetTemplate() const
 		return CachedTemplate;
 	}
 
-	const TArray<FRigVMTemplateArgumentInfo>& Arguments = GetArgumentInfos();
-
-	// we don't allow execute types on arguments
-	for(const FRigVMTemplateArgumentInfo& Info : Arguments)
+	
+	// we don't allow execute types on arguments	
+	const TArray<FRigVMTemplateArgumentInfo>& Infos = GetArgumentInfos();
+	for (const FRigVMTemplateArgumentInfo& Info : Infos)
 	{
 		const FRigVMTemplateArgument Argument = Info.GetArgument();
-		for(const TRigVMTypeIndex& TypeIndex : Argument.GetTypeIndices())
+		const int32 Index = Argument.IndexOfByPredicate([&](const TRigVMTypeIndex TypeIndex)
 		{
-			if(Registry.IsExecuteType(TypeIndex))
-			{
-				UE_LOG(LogRigVM, Error, TEXT("Failed to add template for dispatch '%s'. Argument '%s' is an execute type."),
-					*GetFactoryName().ToString(),
-					*Argument.GetName().ToString());
-				return nullptr;
-			}
+			return Registry.IsExecuteType(TypeIndex);
+		});
+		
+		if (Index != INDEX_NONE)
+		{
+			UE_LOG(LogRigVM, Error, TEXT("Failed to add template for dispatch '%s'. Argument '%s' is an execute type."),
+				*FactoryName.ToString(), *Info.Name.ToString());
+			return nullptr;			
 		}
 	}
 
@@ -400,7 +407,7 @@ const FRigVMTemplate* FRigVMDispatchFactory::GetTemplate() const
 		return FRigVMRegistry::Get().FindDispatchFactory(FactoryName);
 	});
 
-	CachedTemplate = Registry.AddTemplateFromArguments(GetFactoryName(), Arguments, Delegates); 
+	CachedTemplate = Registry.AddTemplateFromArguments(GetFactoryName(), Infos, Delegates); 
 	return CachedTemplate;
 }
 
diff --git a/Engine/Plugins/Runtime/RigVM/Source/RigVM/Private/RigVMCore/RigVMFunction.cpp b/Engine/Plugins/Runtime/RigVM/Source/RigVM/Private/RigVMCore/RigVMFunction.cpp
--- a/Engine/Plugins/Runtime/RigVM/Source/RigVM/Private/RigVMCore/RigVMFunction.cpp
+++ b/Engine/Plugins/Runtime/RigVM/Source/RigVM/Private/RigVMCore/RigVMFunction.cpp
@@ -104,7 +104,7 @@ const TArray<TRigVMTypeIndex>& FRigVMFunction::GetArgumentTypeIndices() const
 			{
 				const FRigVMTemplateArgument* TemplateArgument = Template->FindArgument(FunctionArgument.Name);
 				check(TemplateArgument);
-				ArgumentTypeIndices.Add(TemplateArgument->GetTypeIndices()[PermutationIndex]);
+				ArgumentTypeIndices.Add(TemplateArgument->GetTypeIndex(PermutationIndex));
 			}
 		}
 		else
diff --git a/Engine/Plugins/Runtime/RigVM/Source/RigVM/Private/RigVMCore/RigVMRegistry.cpp b/Engine/Plugins/Runtime/RigVM/Source/RigVM/Private/RigVMCore/RigVMRegistry.cpp
--- a/Engine/Plugins/Runtime/RigVM/Source/RigVM/Private/RigVMCore/RigVMRegistry.cpp
+++ b/Engine/Plugins/Runtime/RigVM/Source/RigVM/Private/RigVMCore/RigVMRegistry.cpp
@@ -502,6 +502,14 @@ TRigVMTypeIndex FRigVMRegistry::FindOrAddType_Internal(const FRigVMTemplateArgum
 			Types[Indices[1]].ArrayTypeIndex = Indices[2];
 		}
 
+		// update the categories first then propagate to ArgumentsPerCategory once all categories up to date
+		TArray<TPair<FRigVMTemplateArgument::ETypeCategory, int32>> ToPropagate;
+		auto RegisterNewType = [&](FRigVMTemplateArgument::ETypeCategory InCategory, int32 NewIndex)
+		{
+			RegisterTypeInCategory(InCategory, NewIndex);
+			ToPropagate.Emplace(InCategory, NewIndex);
+		}; 
+
 		for (int32 ArrayDimension=0; ArrayDimension<3; ++ArrayDimension)
 		{
 			if(bIsExecute && ArrayDimension > 1)
@@ -519,20 +527,20 @@ TRigVMTypeIndex FRigVMRegistry::FindOrAddType_Internal(const FRigVMTemplateArgum
 					default:
 					case 0:
 					{
-						RegisterTypeInCategory(FRigVMTemplateArgument::ETypeCategory_SingleSimpleValue, Index);
-						RegisterTypeInCategory(FRigVMTemplateArgument::ETypeCategory_SingleAnyValue, Index);
+						RegisterNewType(FRigVMTemplateArgument::ETypeCategory_SingleSimpleValue, Index);
+						RegisterNewType(FRigVMTemplateArgument::ETypeCategory_SingleAnyValue, Index);
 						break;
 					}
 					case 1:
 					{
-						RegisterTypeInCategory(FRigVMTemplateArgument::ETypeCategory_ArraySimpleValue, Index);
-						RegisterTypeInCategory(FRigVMTemplateArgument::ETypeCategory_ArrayAnyValue, Index);
+						RegisterNewType(FRigVMTemplateArgument::ETypeCategory_ArraySimpleValue, Index);
+						RegisterNewType(FRigVMTemplateArgument::ETypeCategory_ArrayAnyValue, Index);
 						break;
 					}
 					case 2:
 					{
-						RegisterTypeInCategory(FRigVMTemplateArgument::ETypeCategory_ArrayArraySimpleValue, Index);
-						RegisterTypeInCategory(FRigVMTemplateArgument::ETypeCategory_ArrayArrayAnyValue, Index);
+						RegisterNewType(FRigVMTemplateArgument::ETypeCategory_ArrayArraySimpleValue, Index);
+						RegisterNewType(FRigVMTemplateArgument::ETypeCategory_ArrayArrayAnyValue, Index);
 						break;
 					}
 				}
@@ -544,20 +552,20 @@ TRigVMTypeIndex FRigVMRegistry::FindOrAddType_Internal(const FRigVMTemplateArgum
 					default:
 					case 0:
 					{
-						RegisterTypeInCategory(FRigVMTemplateArgument::ETypeCategory_SingleObjectValue, Index);
-						RegisterTypeInCategory(FRigVMTemplateArgument::ETypeCategory_SingleAnyValue, Index);
+						RegisterNewType(FRigVMTemplateArgument::ETypeCategory_SingleObjectValue, Index);
+						RegisterNewType(FRigVMTemplateArgument::ETypeCategory_SingleAnyValue, Index);
 						break;
 					}
 					case 1:
 					{
-						RegisterTypeInCategory(FRigVMTemplateArgument::ETypeCategory_ArrayObjectValue, Index);
-						RegisterTypeInCategory(FRigVMTemplateArgument::ETypeCategory_ArrayAnyValue, Index);
+						RegisterNewType(FRigVMTemplateArgument::ETypeCategory_ArrayObjectValue, Index);
+						RegisterNewType(FRigVMTemplateArgument::ETypeCategory_ArrayAnyValue, Index);
 						break;
 					}
 					case 2:
 					{
-						RegisterTypeInCategory(FRigVMTemplateArgument::ETypeCategory_ArrayArrayObjectValue, Index);
-						RegisterTypeInCategory(FRigVMTemplateArgument::ETypeCategory_ArrayArrayAnyValue, Index);
+						RegisterNewType(FRigVMTemplateArgument::ETypeCategory_ArrayArrayObjectValue, Index);
+						RegisterNewType(FRigVMTemplateArgument::ETypeCategory_ArrayArrayAnyValue, Index);
 						break;
 					}
 				}
@@ -569,20 +577,20 @@ TRigVMTypeIndex FRigVMRegistry::FindOrAddType_Internal(const FRigVMTemplateArgum
 					default:
 					case 0:
 					{
-						RegisterTypeInCategory(FRigVMTemplateArgument::ETypeCategory_SingleEnumValue, Index);
-						RegisterTypeInCategory(FRigVMTemplateArgument::ETypeCategory_SingleAnyValue, Index);
+						RegisterNewType(FRigVMTemplateArgument::ETypeCategory_SingleEnumValue, Index);
+						RegisterNewType(FRigVMTemplateArgument::ETypeCategory_SingleAnyValue, Index);
 						break;
 					}
 					case 1:
 					{
-						RegisterTypeInCategory(FRigVMTemplateArgument::ETypeCategory_ArrayEnumValue, Index);
-						RegisterTypeInCategory(FRigVMTemplateArgument::ETypeCategory_ArrayAnyValue, Index);
+						RegisterNewType(FRigVMTemplateArgument::ETypeCategory_ArrayEnumValue, Index);
+						RegisterNewType(FRigVMTemplateArgument::ETypeCategory_ArrayAnyValue, Index);
 						break;
 					}
 					case 2:
 					{
-						RegisterTypeInCategory(FRigVMTemplateArgument::ETypeCategory_ArrayArrayEnumValue, Index);
-						RegisterTypeInCategory(FRigVMTemplateArgument::ETypeCategory_ArrayArrayAnyValue, Index);
+						RegisterNewType(FRigVMTemplateArgument::ETypeCategory_ArrayArrayEnumValue, Index);
+						RegisterNewType(FRigVMTemplateArgument::ETypeCategory_ArrayArrayAnyValue, Index);
 						break;
 					}
 				}
@@ -593,7 +601,7 @@ TRigVMTypeIndex FRigVMRegistry::FindOrAddType_Internal(const FRigVMTemplateArgum
 				{
 					if(ArrayDimension == 0)
 					{
-						RegisterTypeInCategory(FRigVMTemplateArgument::ETypeCategory_Execute, Index);
+						RegisterNewType(FRigVMTemplateArgument::ETypeCategory_Execute, Index);
 					}
 				}
 				else
@@ -605,17 +613,17 @@ TRigVMTypeIndex FRigVMRegistry::FindOrAddType_Internal(const FRigVMTemplateArgum
 							default:
 							case 0:
 							{
-								RegisterTypeInCategory(FRigVMTemplateArgument::ETypeCategory_SingleMathStructValue, Index);
+								RegisterNewType(FRigVMTemplateArgument::ETypeCategory_SingleMathStructValue, Index);
 								break;
 							}
 							case 1:
 							{
-								RegisterTypeInCategory(FRigVMTemplateArgument::ETypeCategory_ArrayMathStructValue, Index);
+								RegisterNewType(FRigVMTemplateArgument::ETypeCategory_ArrayMathStructValue, Index);
 								break;
 							}
 							case 2:
 							{
-								RegisterTypeInCategory(FRigVMTemplateArgument::ETypeCategory_ArrayArrayMathStructValue, Index);
+								RegisterNewType(FRigVMTemplateArgument::ETypeCategory_ArrayArrayMathStructValue, Index);
 								break;
 							}
 						}
@@ -626,20 +634,20 @@ TRigVMTypeIndex FRigVMRegistry::FindOrAddType_Internal(const FRigVMTemplateArgum
 						default:
 						case 0:
 						{
-							RegisterTypeInCategory(FRigVMTemplateArgument::ETypeCategory_SingleScriptStructValue, Index);
-							RegisterTypeInCategory(FRigVMTemplateArgument::ETypeCategory_SingleAnyValue, Index);
+							RegisterNewType(FRigVMTemplateArgument::ETypeCategory_SingleScriptStructValue, Index);
+							RegisterNewType(FRigVMTemplateArgument::ETypeCategory_SingleAnyValue, Index);
 							break;
 						}
 						case 1:
 						{
-							RegisterTypeInCategory(FRigVMTemplateArgument::ETypeCategory_ArrayScriptStructValue, Index);
-							RegisterTypeInCategory(FRigVMTemplateArgument::ETypeCategory_ArrayAnyValue, Index);
+							RegisterNewType(FRigVMTemplateArgument::ETypeCategory_ArrayScriptStructValue, Index);
+							RegisterNewType(FRigVMTemplateArgument::ETypeCategory_ArrayAnyValue, Index);
 							break;
 						}
 						case 2:
 						{
-							RegisterTypeInCategory(FRigVMTemplateArgument::ETypeCategory_ArrayArrayScriptStructValue, Index);
-							RegisterTypeInCategory(FRigVMTemplateArgument::ETypeCategory_ArrayArrayAnyValue, Index);
+							RegisterNewType(FRigVMTemplateArgument::ETypeCategory_ArrayArrayScriptStructValue, Index);
+							RegisterNewType(FRigVMTemplateArgument::ETypeCategory_ArrayArrayAnyValue, Index);
 							break;
 						}
 					}
@@ -647,6 +655,12 @@ TRigVMTypeIndex FRigVMRegistry::FindOrAddType_Internal(const FRigVMTemplateArgum
 			}
 		}
 
+		// propagate new type to templates once they have all been added to the categories
+		for (const auto& [Category, NewIndex]: ToPropagate)
+		{
+			PropagateTypeAddedToCategory(Category, NewIndex);	
+		}
+
 		// if the type is a structure
 		// then add all of its sub property types
 		if(const UStruct* Struct = Cast<UStruct>(CPPTypeObject))
@@ -693,14 +707,21 @@ void FRigVMRegistry::RegisterTypeInCategory(FRigVMTemplateArgument::ETypeCategor
 	check(InCategory != FRigVMTemplateArgument::ETypeCategory_Invalid);
 
 	TypesPerCategory.FindChecked(InCategory).Add(InTypeIndex);
+}
 
-	// when adding a new type - we need to update template arguments which expect to have access to that type 
-	const TArray<TPair<int32,int32>>& ArgumentsToUseType = ArgumentsPerCategory.FindChecked(InCategory);
-	for(const TPair<int32,int32>& Pair : ArgumentsToUseType)
+void FRigVMRegistry::PropagateTypeAddedToCategory(const FRigVMTemplateArgument::ETypeCategory InCategory, const TRigVMTypeIndex InTypeIndex)
+{
+	check(InCategory != FRigVMTemplateArgument::ETypeCategory_Invalid);
+	if ( ensure(TypesPerCategory.FindChecked(InCategory).Contains(InTypeIndex)) )
 	{
-		FRigVMTemplate& Template = Templates[Pair.Key];
-		const FRigVMTemplateArgument* Argument = Template.GetArgument(Pair.Value);
-		Template.AddTypeForArgument(Argument->GetName(), InTypeIndex);
+		// when adding a new type - we need to update template arguments which expect to have access to that type 
+		const TArray<TPair<int32,int32>>& ArgumentsToUseType = ArgumentsPerCategory.FindChecked(InCategory);
+		for(const TPair<int32,int32>& Pair : ArgumentsToUseType)
+		{
+			FRigVMTemplate& Template = Templates[Pair.Key];
+			const FRigVMTemplateArgument* Argument = Template.GetArgument(Pair.Value);
+			Template.AddTypeForArgument(Argument->GetName(), InTypeIndex);
+		}
 	}
 }
 
@@ -1587,51 +1608,36 @@ const FRigVMTemplate* FRigVMRegistry::AddTemplateFromArguments_NoLock(const FNam
 	FRigVMTemplate Template(InName, InInfos);
 	for(const FRigVMTemplateArgument& Argument : Template.Arguments)
 	{
+		const int32 NumIndices = Argument.GetNumTypes();
 		if(!Argument.IsSingleton() && NumPermutations > 1)
 		{
-			if(Argument.TypeIndices.Num() != NumPermutations)
+			if(NumIndices != NumPermutations)
 			{
 				UE_LOG(LogRigVM, Error, TEXT("Failed to add template '%s' since the arguments' types counts don't match."), *InName.ToString());
 				return nullptr;
 			}
 		}
-		NumPermutations = FMath::Max(NumPermutations, Argument.TypeIndices.Num()); 
+		NumPermutations = FMath::Max(NumPermutations, NumIndices); 
 	}
 
 	// if any of the arguments are wildcards we'll need to update the types
 	for(FRigVMTemplateArgument& Argument : Template.Arguments)
 	{
-		if(Argument.TypeIndices.Num() == 1 && IsWildCardType(Argument.TypeIndices[0]))
+		if(Argument.GetNumTypes() == 1 && IsWildCardType(Argument.GetTypeIndex(0)))
 		{
-			if(IsArrayType(Argument.TypeIndices[0]))
+			Argument.InvalidatePermutations(Argument.GetTypeIndex(0));
+			if(IsArrayType(Argument.GetTypeIndex(0)))
 			{
-				Argument.TypeIndices = GetTypesForCategory(FRigVMTemplateArgument::ETypeCategory_ArrayAnyValue);
 				Argument.TypeCategories.Add(FRigVMTemplateArgument::ETypeCategory_ArrayAnyValue);
 			}
 			else
 			{
-				Argument.TypeIndices = GetTypesForCategory(FRigVMTemplateArgument::ETypeCategory_SingleAnyValue);
 				Argument.TypeCategories.Add(FRigVMTemplateArgument::ETypeCategory_SingleAnyValue);
 			}
-
-			NumPermutations = FMath::Max(NumPermutations, Argument.TypeIndices.Num()); 
-		}
-	}
-
-	// if we have more than one permutation we may need to upgrade the types for singleton args
-	if(NumPermutations > 1)
-	{
-		for(FRigVMTemplateArgument& Argument : Template.Arguments)
-		{
-			if(Argument.TypeIndices.Num() == 1)
-			{
-				const int32 TypeIndex = Argument.TypeIndices[0];
-				Argument.TypeIndices.SetNum(NumPermutations);
-				for(int32 Index=0;Index<NumPermutations;Index++)
-				{
-					Argument.TypeIndices[Index] = TypeIndex;
-				}
-			}
+			Argument.bUseCategories = true;
+			Argument.TypeIndices.Reset();
+	
+			NumPermutations = FMath::Max(NumPermutations, Argument.GetNumTypes()); 
 		}
 	}
 
@@ -1644,9 +1650,8 @@ const FRigVMTemplate* FRigVMRegistry::AddTemplateFromArguments_NoLock(const FNam
 		{
 			TSet< TRigVMTypeIndex > PermutationTypes; PermutationTypes.Reserve(NumPermutations);
 			for(int32 Index = 0; Index < NumPermutations; Index++)
- 
 			{
-				const TRigVMTypeIndex ArgType = Template.Arguments[0].TypeIndices[Index];
+				const TRigVMTypeIndex ArgType = Template.Arguments[0].GetTypeIndex(Index);
 				if (PermutationTypes.Contains(ArgType))
 				{
 					ToRemove.Add(Index);
@@ -1666,7 +1671,7 @@ const FRigVMTemplate* FRigVMRegistry::AddTemplateFromArguments_NoLock(const FNam
 			{
 				for(int32 ArgIndex = 0; ArgIndex < NumArguments; ArgIndex++)
 				{
-					ArgTypes[ArgIndex] = Template.Arguments[ArgIndex].TypeIndices[Index];
+					ArgTypes[ArgIndex] = Template.Arguments[ArgIndex].GetTypeIndex(Index);
 				}
 				
 				if (PermutationTypes.Contains(ArgTypes))
@@ -1684,7 +1689,10 @@ const FRigVMTemplate* FRigVMRegistry::AddTemplateFromArguments_NoLock(const FNam
 		{
 			for(FRigVMTemplateArgument& Argument : Template.Arguments)
 			{
-				Argument.TypeIndices.RemoveAt(ToRemove[i]);
+				if (Argument.TypeIndices.IsValidIndex(ToRemove[i]))
+				{
+					Argument.TypeIndices.RemoveAt(ToRemove[i]);
+				}
 			}
 		}
 		NumPermutations -= ToRemove.Num();
@@ -1695,11 +1703,7 @@ const FRigVMTemplate* FRigVMRegistry::AddTemplateFromArguments_NoLock(const FNam
 		Argument.UpdateTypeToPermutations();
 	}
 
-	Template.Permutations.SetNum(NumPermutations);
-	for(int32 Index=0;Index<NumPermutations;Index++)
-	{
-		Template.Permutations[Index] = INDEX_NONE;
-	}
+	Template.Permutations.Init(INDEX_NONE, NumPermutations);
 	Template.RecomputeTypesHashToPermutations();
 
 	const int32 Index = Templates.AddElement(Template);
diff --git a/Engine/Plugins/Runtime/RigVM/Source/RigVM/Private/RigVMCore/RigVMStruct.cpp b/Engine/Plugins/Runtime/RigVM/Source/RigVM/Private/RigVMCore/RigVMStruct.cpp
--- a/Engine/Plugins/Runtime/RigVM/Source/RigVM/Private/RigVMCore/RigVMStruct.cpp
+++ b/Engine/Plugins/Runtime/RigVM/Source/RigVM/Private/RigVMCore/RigVMStruct.cpp
@@ -130,7 +130,7 @@ FRigVMTemplateTypeMap FRigVMStructUpgradeInfo::GetTypeMapFromStruct(UScriptStruc
 		}
 		
 		const FName PropertyName = It->GetFName();
-		const TRigVMTypeIndex TypeIndex = FRigVMTemplateArgument(*It).GetTypeIndices()[0];
+		const TRigVMTypeIndex TypeIndex = FRigVMTemplateArgument(*It).GetTypeIndex(0);
 		Types.Add(PropertyName, TypeIndex);
 	}
 #endif
diff --git a/Engine/Plugins/Runtime/RigVM/Source/RigVM/Private/RigVMCore/RigVMTemplate.cpp b/Engine/Plugins/Runtime/RigVM/Source/RigVM/Private/RigVMCore/RigVMTemplate.cpp
--- a/Engine/Plugins/Runtime/RigVM/Source/RigVM/Private/RigVMCore/RigVMTemplate.cpp
+++ b/Engine/Plugins/Runtime/RigVM/Source/RigVM/Private/RigVMCore/RigVMTemplate.cpp
@@ -5,6 +5,8 @@
 #include "RigVMCore/RigVMStruct.h"
 #include "RigVMCore/RigVMDispatchFactory.h"
 #include "RigVMModule.h"
+#include "Algo/Accumulate.h"
+#include "Algo/ForEach.h"
 #include "Algo/Sort.h"
 #include "UObject/UObjectIterator.h"
 
@@ -30,23 +32,15 @@ FRigVMTemplateArgumentType::FRigVMTemplateArgumentType(const FName& InCPPType, U
 }
 
 FRigVMTemplateArgument::FRigVMTemplateArgument()
-	: Index(INDEX_NONE)
-	, Name(NAME_None)
-	, Direction(ERigVMPinDirection::IO)
-{
-}
+{}
 
 FRigVMTemplateArgument::FRigVMTemplateArgument(const FName& InName, ERigVMPinDirection InDirection)
-: Index(INDEX_NONE)
-, Name(InName)
-, Direction(InDirection)
-{
-}
+	: Name(InName)
+	, Direction(InDirection)
+{}
 
 FRigVMTemplateArgument::FRigVMTemplateArgument(FProperty* InProperty)
-	: Index(INDEX_NONE)
-	, Name(InProperty->GetFName())
-	, Direction(ERigVMPinDirection::IO)
+	: Name(InProperty->GetFName())
 {
 #if WITH_EDITOR
 	Direction = FRigVMStruct::GetPinDirectionFromProperty(InProperty);
@@ -91,126 +85,69 @@ FRigVMTemplateArgument::FRigVMTemplateArgument(FProperty* InProperty)
 }
 
 FRigVMTemplateArgument::FRigVMTemplateArgument(const FName& InName, ERigVMPinDirection InDirection, TRigVMTypeIndex InTypeIndex)
-: Index(INDEX_NONE)
-, Name(InName)
-, Direction(InDirection)
-, TypeIndices({InTypeIndex})
+	: Name(InName)
+	, Direction(InDirection)
+	, TypeIndices({InTypeIndex})
 {
 	EnsureValidExecuteType();
 	UpdateTypeToPermutations();
 }
 
 FRigVMTemplateArgument::FRigVMTemplateArgument(const FName& InName, ERigVMPinDirection InDirection, const TArray<TRigVMTypeIndex>& InTypeIndices)
-: Index(INDEX_NONE)
-, Name(InName)
-, Direction(InDirection)
-, TypeIndices(InTypeIndices)
+	: Name(InName)
+	, Direction(InDirection)
+	, TypeIndices(InTypeIndices)
 {
 	check(TypeIndices.Num() > 0);
-
 	EnsureValidExecuteType();
 	UpdateTypeToPermutations();
 }
 
-FRigVMTemplateArgument::FRigVMTemplateArgument(const FName& InName, ERigVMPinDirection InDirection, const TArray<ETypeCategory>& InTypeCategories, const FTypeFilter& InTypeFilter)
-: Index(INDEX_NONE)
-, Name(InName)
-, Direction(InDirection)
-, TypeCategories(InTypeCategories)
+FRigVMTemplateArgument::FRigVMTemplateArgument(const FName& InName, ERigVMPinDirection InDirection, const TArray<ETypeCategory>& InTypeCategories)
+	: Name(InName)
+	, Direction(InDirection)
+	, TypeCategories(InTypeCategories)
 {
-	TSet<TRigVMTypeIndex> AllTypes;
-	for(ETypeCategory TypeCategory : TypeCategories)
+	const int32 NumCategories = TypeCategories.Num();
+	if (NumCategories > 0)
 	{
-		AllTypes.Append(FRigVMRegistry::Get().GetTypesForCategory(TypeCategory));
-	}
-	
-	TypeIndices.Reserve(AllTypes.Num());
-	for (const TRigVMTypeIndex& Type : AllTypes)
-	{
-		if (InTypeFilter.IsBound() && !InTypeFilter.Execute(Type))
-		{
-			continue;
-		}
-		
-		TypeIndices.Add(Type);
-	}
-
-	EnsureValidExecuteType();
-	UpdateTypeToPermutations();
-}
+		TSet<TRigVMTypeIndex> AllTypes;
 
-void FRigVMTemplateArgument::Serialize(FArchive& Ar)
-{
-	Ar << Index;
-	Ar << Name;
-	Ar << Direction;
+		TArray<int32> NumTypesByCategory, TypesAddedByCategory;
+		TypesAddedByCategory.Reserve(NumCategories);
+		NumTypesByCategory.Reserve(NumCategories);
 
-#if UE_RIGVM_DEBUG_TYPEINDEX
-	if (Ar.IsSaving())
-	{
-		TArray<int32> TypeIndicesInt;
-		TypeIndicesInt.Reserve(TypeIndices.Num());
-		for (TRigVMTypeIndex& TypeIndex : TypeIndices)
+		bUseCategories = NumCategories == 1;
+		if (!bUseCategories)
 		{
-			TypeIndicesInt.Add(TypeIndex.GetIndex());
-		}
-		Ar << TypeIndicesInt;
-	}
-	else if (Ar.IsLoading())
-	{
-		TArray<int32> TypeIndicesInt;
-		Ar << TypeIndicesInt;
-		TypeIndices.Reset();
-		TypeIndices.SetNum(TypeIndicesInt.Num());
-		static FRigVMRegistry& Registry = FRigVMRegistry::Get();
-		for (int32 i=0; i<TypeIndicesInt.Num(); ++i)
-		{
-			TypeIndices[i].Index = TypeIndicesInt[i];
-			TypeIndices[i].Name = Registry.GetType(TypeIndices[i]).CPPType;
-		}
-	}
-#else
-	Ar << TypeIndices;
-#endif
-
-	if (Ar.IsLoading())
-	{
-		for (int32 i=0; i<TypeIndices.Num(); ++i)
-		{
-			if (TArray<int32>* Permutations = TypeToPermutations.Find(TypeIndices[i]))
+			for (const ETypeCategory TypeCategory : TypeCategories)
 			{
-				Permutations->Add(i);
-			}
-			else
-			{
-				TypeToPermutations.Add(TypeIndices[i], {i});
+				const TArray<TRigVMTypeIndex>& Types = FRigVMRegistry::Get().GetTypesForCategory(TypeCategory);
+				AllTypes.Reserve(AllTypes.Num() + Types.Num());
+				for (const TRigVMTypeIndex Type: Types)
+				{
+					AllTypes.Add(Type);
+				}
+			
+				NumTypesByCategory.Add(Types.Num());
+				TypesAddedByCategory.Add(AllTypes.Num());
 			}
+			bUseCategories = NumTypesByCategory[0] == TypesAddedByCategory[0] && Algo::Accumulate(NumTypesByCategory, 0) == AllTypes.Num();
 		}
-	}
 
-	if (Ar.IsSaving())
-	{
-		int32 CategoriesNum = TypeCategories.Num();
-		Ar << CategoriesNum;
-		for (int32 i=0; i<CategoriesNum; ++i)
+		if (bUseCategories)
 		{
-			uint8 CategoryIndex = (int8)TypeCategories[i]; 
-			Ar << CategoryIndex;
+			TypeIndices.Reset();
 		}
-	}
-	else if (Ar.IsLoading())
-	{
-		int32 CategoriesNum = 0;
-		Ar << CategoriesNum;
-		TypeCategories.SetNumUninitialized(CategoriesNum);
-		
-		for (int32 i=0; i<CategoriesNum; ++i)
+		else
 		{
-			uint8 CategoryIndex = INDEX_NONE;
-			Ar << CategoryIndex;
-			TypeCategories[i] = (FRigVMTemplateArgument::ETypeCategory) CategoryIndex;
+			TArray<TRigVMTypeIndex> Indices = AllTypes.Array();
+			TypeIndices = MoveTemp(Indices);
+			EnsureValidExecuteType();
 		}
-	}	
+
+		UpdateTypeToPermutations();
+	}
 }
 
 void FRigVMTemplateArgument::EnsureValidExecuteType()
@@ -225,15 +162,15 @@ void FRigVMTemplateArgument::EnsureValidExecuteType()
 void FRigVMTemplateArgument::UpdateTypeToPermutations()
 {
 	TypeToPermutations.Reset();
-	for(int32 TypeIndex=0;TypeIndex<TypeIndices.Num();TypeIndex++)
+	for(int32 TypeIndex=0;TypeIndex<GetNumTypes();TypeIndex++)
 	{
-		if (TArray<int32>* Permutations = TypeToPermutations.Find(TypeIndices[TypeIndex]))
+		if (TArray<int32>* Permutations = TypeToPermutations.Find(GetTypeIndex(TypeIndex)))
 		{
 			Permutations->Add(TypeIndex);
 		}
 		else
 		{
-			TypeToPermutations.Add(TypeIndices[TypeIndex], {TypeIndex});
+			TypeToPermutations.Add(GetTypeIndex(TypeIndex), {TypeIndex});
 		}
 	}
 }
@@ -256,12 +193,13 @@ bool FRigVMTemplateArgument::SupportsTypeIndex(TRigVMTypeIndex InTypeIndex, TRig
 			InTypeIndex = Registry.GetArrayTypeFromBaseTypeIndex(InTypeIndex);
 		}
 	}
-	
-	if (const TArray<int32>* Permutations = TypeToPermutations.Find(InTypeIndex))
+
+	const TArray<int32>& Permutations = GetPermutations(InTypeIndex);
+	if (!Permutations.IsEmpty())
 	{
 		if(OutTypeIndex)
 		{
-			(*OutTypeIndex) = TypeIndices[(*Permutations)[0]];
+			(*OutTypeIndex) = GetTypeIndex(Permutations[0]);
 		}
 		return true;
 	}
@@ -270,11 +208,12 @@ bool FRigVMTemplateArgument::SupportsTypeIndex(TRigVMTypeIndex InTypeIndex, TRig
 	const TArray<TRigVMTypeIndex>& CompatibleTypes = Registry.GetCompatibleTypes(InTypeIndex);
 	for (const TRigVMTypeIndex& CompatibleTypeIndex : CompatibleTypes)
 	{
-		if (const TArray<int32>* Permutations = TypeToPermutations.Find(CompatibleTypeIndex))
+		const TArray<int32>& CompatiblePermutations = GetPermutations(CompatibleTypeIndex);
+		if (!CompatiblePermutations.IsEmpty())
 		{
 			if(OutTypeIndex)
 			{
-				(*OutTypeIndex) = TypeIndices[(*Permutations)[0]];
+				(*OutTypeIndex) = GetTypeIndex(CompatiblePermutations[0]);
 			}
 			return true;
 		}
@@ -294,10 +233,10 @@ bool FRigVMTemplateArgument::IsSingleton(const TArray<int32>& InPermutationIndic
 		return false;
 	}
 
-	const TRigVMTypeIndex TRigVMTypeIndex = TypeIndices[InPermutationIndices[0]];
+	const TRigVMTypeIndex InType0 = GetTypeIndex(InPermutationIndices[0]);
 	for (int32 PermutationIndex = 1; PermutationIndex < InPermutationIndices.Num(); PermutationIndex++)
 	{
-		if (TypeIndices[InPermutationIndices[PermutationIndex]] != TRigVMTypeIndex)
+		if (GetTypeIndex(InPermutationIndices[PermutationIndex]) != InType0)
 		{
 			return false;
 		}
@@ -308,40 +247,39 @@ bool FRigVMTemplateArgument::IsSingleton(const TArray<int32>& InPermutationIndic
 bool FRigVMTemplateArgument::IsExecute() const
 {
 	const FRigVMRegistry& Registry = FRigVMRegistry::Get();
-	for (int32 PermutationIndex = 0; PermutationIndex < TypeIndices.Num(); PermutationIndex++)
+	const int32 FoundAnyNotExec = IndexOfByPredicate([&](const TRigVMTypeIndex Type)
 	{
-		if (!Registry.IsExecuteType(TypeIndices[PermutationIndex]))
-		{
-			return false;
-		}
-	}
-	return true;
+		return !Registry.IsExecuteType(Type);
+	});
+	return FoundAnyNotExec == INDEX_NONE;
 }
 
 FRigVMTemplateArgument::EArrayType FRigVMTemplateArgument::GetArrayType() const
 {
-	if(TypeIndices.Num() > 0)
+	const int32 NumTypes = GetNumTypes();
+	if (GetNumTypes() > 0)
 	{
 		const FRigVMRegistry& Registry = FRigVMRegistry::Get();
-		const EArrayType ArrayType = Registry.IsArrayType(TypeIndices[0]) ? EArrayType::EArrayType_ArrayValue : EArrayType::EArrayType_SingleValue;
+		const EArrayType ArrayType = Registry.IsArrayType(GetTypeIndex(0)) ? EArrayType_ArrayValue : EArrayType_SingleValue;
 		
 		if(IsSingleton())
 		{
 			return ArrayType;
 		}
 
-		for(int32 PermutationIndex=1;PermutationIndex<TypeIndices.Num();PermutationIndex++)
+		for(int32 PermutationIndex=1; PermutationIndex<NumTypes;PermutationIndex++)
 		{
+			const TRigVMTypeIndex TypeIndex = GetTypeIndex(PermutationIndex);
 			// INDEX_NONE indicates deleted permutation
-			if (TypeIndices[PermutationIndex] == INDEX_NONE)
+			if (TypeIndex == INDEX_NONE)
 			{
 				continue;
 			}
 			
-			const EArrayType OtherArrayType = Registry.IsArrayType(TypeIndices[PermutationIndex]) ? EArrayType::EArrayType_ArrayValue : EArrayType::EArrayType_SingleValue;
+			const EArrayType OtherArrayType = Registry.IsArrayType(TypeIndex) ? EArrayType_ArrayValue : EArrayType_SingleValue;
 			if(OtherArrayType != ArrayType)
 			{
-				return EArrayType::EArrayType_Mixed;
+				return EArrayType_Mixed;
 			}
 		}
 
@@ -351,9 +289,107 @@ FRigVMTemplateArgument::EArrayType FRigVMTemplateArgument::GetArrayType() const
 	return EArrayType_Invalid;
 }
 
-const TArray<TRigVMTypeIndex>& FRigVMTemplateArgument::GetTypeIndices() const
+const TArray<int32>& FRigVMTemplateArgument::GetPermutations(const TRigVMTypeIndex InType) const
+{
+	if (const TArray<int32>* Found = TypeToPermutations.Find(InType))
+	{
+		return *Found;
+	}
+
+	int32 IndexInTypes = 0;
+	TArray<int32> Permutations;
+	ForEachType([&](const TRigVMTypeIndex Type)
+	{
+		if (Type == InType)
+		{
+			Permutations.Add(IndexInTypes);
+		}
+		IndexInTypes++;
+	});
+
+	if (!Permutations.IsEmpty())
+	{
+		return TypeToPermutations.Emplace(InType, MoveTemp(Permutations));
+	}
+
+	static const TArray<int32> Dummy;
+	return Dummy;
+}
+
+void FRigVMTemplateArgument::InvalidatePermutations(const TRigVMTypeIndex InType)
 {
-	return TypeIndices;
+	TypeToPermutations.Remove(InType);
+}
+
+void FRigVMTemplateArgument::GetAllTypes(TArray<TRigVMTypeIndex>& OutTypes) const
+{
+	if (!bUseCategories)
+	{
+		OutTypes = TypeIndices;
+		return;
+	}
+	
+	OutTypes.Reset();
+	for (const ETypeCategory Category: TypeCategories)
+	{
+		OutTypes.Append(FRigVMRegistry::Get().GetTypesForCategory(Category));
+	}
+}
+
+TRigVMTypeIndex FRigVMTemplateArgument::GetTypeIndex(const int32 InIndex) const
+{
+	if (!bUseCategories)
+	{
+		check(!TypeIndices.IsEmpty())
+		return TypeIndices.IsValidIndex(InIndex) ? TypeIndices[InIndex] : TypeIndices[0];		
+	}
+
+	return CategoryViews(TypeCategories).GetTypeIndex(InIndex);
+}
+
+int32 FRigVMTemplateArgument::FindTypeIndex(const TRigVMTypeIndex InTypeIndex) const
+{
+	if (!bUseCategories)
+	{
+		return TypeIndices.IndexOfByKey(InTypeIndex);
+	}
+
+	return CategoryViews(TypeCategories).FindIndex(InTypeIndex);
+}
+
+int32 FRigVMTemplateArgument::GetNumTypes() const
+{
+	if (!bUseCategories)
+	{
+		return TypeIndices.Num();
+	}
+	
+	return Algo::Accumulate(TypeCategories, 0, [](int32 Sum, const ETypeCategory Category)
+	{
+		return Sum + FRigVMRegistry::Get().GetTypesForCategory(Category).Num();
+	});
+}
+
+void FRigVMTemplateArgument::AddTypeIndex(const TRigVMTypeIndex InTypeIndex)
+{
+	ensure( TypeCategories.IsEmpty() );
+	TypeIndices.AddUnique(InTypeIndex);
+}
+
+void FRigVMTemplateArgument::RemoveType(const int32 InIndex)
+{
+	ensure( TypeCategories.IsEmpty() );
+	TypeIndices.RemoveAt(InIndex);
+}
+
+void FRigVMTemplateArgument::ForEachType(TFunction<void(const TRigVMTypeIndex InType)>&& InCallback) const
+{
+	if (!bUseCategories)
+	{
+		return Algo::ForEach(TypeIndices, InCallback);
+	}
+	
+	return CategoryViews(TypeCategories).ForEachType(MoveTemp(InCallback));
 }
 
 TArray<TRigVMTypeIndex> FRigVMTemplateArgument::GetSupportedTypeIndices(const TArray<int32>& InPermutationIndices) const
@@ -361,23 +397,24 @@ TArray<TRigVMTypeIndex> FRigVMTemplateArgument::GetSupportedTypeIndices(const TA
 	TArray<TRigVMTypeIndex> SupportedTypes;
 	if(InPermutationIndices.IsEmpty())
 	{
-		for (const TRigVMTypeIndex& TypeIndex : TypeIndices)
+		ForEachType([&](const TRigVMTypeIndex TypeIndex)
 		{
 			// INDEX_NONE indicates deleted permutation
 			if (TypeIndex != INDEX_NONE)
 			{
 				SupportedTypes.AddUnique(TypeIndex);
 			}
-		}
+		});
 	}
 	else
 	{
 		for(const int32 PermutationIndex : InPermutationIndices)
 		{
 			// INDEX_NONE indicates deleted permutation
-			if (TypeIndices[PermutationIndex] != INDEX_NONE)
+			const TRigVMTypeIndex Type = GetTypeIndex(PermutationIndex);
+			if (Type != INDEX_NONE)
 			{
-				SupportedTypes.AddUnique(TypeIndices[PermutationIndex]);
+				SupportedTypes.AddUnique(Type);
 			}
 		}
 	}
@@ -389,27 +426,28 @@ TArray<TRigVMTypeIndex> FRigVMTemplateArgument::GetSupportedTypeIndices(const TA
 TArray<FString> FRigVMTemplateArgument::GetSupportedTypeStrings(const TArray<int32>& InPermutationIndices) const
 {
 	TArray<FString> SupportedTypes;
+	const FRigVMRegistry& Registry = FRigVMRegistry::Get();
 	if(InPermutationIndices.IsEmpty())
 	{
-		for (const TRigVMTypeIndex& TypeIndex : TypeIndices)
+		ForEachType([&](const TRigVMTypeIndex TypeIndex)
 		{
 			// INDEX_NONE indicates deleted permutation
 			if (TypeIndex != INDEX_NONE)
 			{
-				const FString TypeString = FRigVMRegistry::Get().GetType(TypeIndex).CPPType.ToString();
+				const FString TypeString = Registry.GetType(TypeIndex).CPPType.ToString();
 				SupportedTypes.AddUnique(TypeString);
 			}
-		}
+		});
 	}
 	else
 	{
 		for(const int32 PermutationIndex : InPermutationIndices)
 		{
-			const TRigVMTypeIndex TypeIndex = TypeIndices[PermutationIndex];
+			const TRigVMTypeIndex TypeIndex = GetTypeIndex(PermutationIndex);
 			// INDEX_NONE indicates deleted permutation
 			if (TypeIndex != INDEX_NONE)
 			{
-				const FString TypeString = FRigVMRegistry::Get().GetType(TypeIndex).CPPType.ToString();
+				const FString TypeString = Registry.GetType(TypeIndex).CPPType.ToString();
 				SupportedTypes.AddUnique(TypeString);
 			}
 		}
@@ -419,6 +457,52 @@ TArray<FString> FRigVMTemplateArgument::GetSupportedTypeStrings(const TArray<int
 
 #endif
 
+
+FRigVMTemplateArgument::CategoryViews::CategoryViews(const TArray<ETypeCategory>& InCategories)
+{
+	Types.Reserve(InCategories.Num());
+	for (const ETypeCategory Category: InCategories)
+	{
+		Types.Emplace(FRigVMRegistry::Get().GetTypesForCategory(Category));
+	}
+}
+
+void FRigVMTemplateArgument::CategoryViews::ForEachType(TFunction<void(const TRigVMTypeIndex InType)>&& InCallback) const
+{
+	for (const TArrayView<const TRigVMTypeIndex>& TypeView: Types)
+	{
+		Algo::ForEach(TypeView, InCallback);
+	}
+}
+
+TRigVMTypeIndex FRigVMTemplateArgument::CategoryViews::GetTypeIndex(int32 InIndex) const
+{
+	for (const TArrayView<const TRigVMTypeIndex>& TypeView: Types)
+	{
+		if (TypeView.IsValidIndex(InIndex))
+		{
+			return TypeView[InIndex];  
+		}
+		InIndex -= TypeView.Num();
+	}
+	return INDEX_NONE;
+}
+	
+int32 FRigVMTemplateArgument::CategoryViews::FindIndex(const TRigVMTypeIndex InTypeIndex) const
+{
+	int32 Offset = 0;
+	for (const TArrayView<const TRigVMTypeIndex>& TypeView: Types)
+	{
+		const int32 Found = TypeView.IndexOfByKey(InTypeIndex);
+		if (Found != INDEX_NONE)
+		{
+			return Found + Offset;
+		}
+		Offset += TypeView.Num();
+	}
+	return INDEX_NONE;
+}
+
 /**
  * FRigVMTemplateArgumentInfo 
  */
@@ -525,38 +609,6 @@ FRigVMTemplate::FRigVMTemplate()
 
 }
 
-void FRigVMTemplate::Serialize(FArchive& Ar)
-{
-	Ar << Index;
-	Ar << Notation;
-	Ar << Permutations;
-
-	if (Ar.IsSaving())
-	{
-		int32 ArgumentsNum = Arguments.Num();
-		Ar << ArgumentsNum;
-
-		for (FRigVMTemplateArgument& Argument : Arguments)
-		{
-			Argument.Serialize(Ar);
-		}
-	}
-	else if(Ar.IsLoading())
-	{
-		int32 ArgumentsNum = 0;
-		Ar << ArgumentsNum;
-
-		Arguments.Reset();
-		Arguments.Reserve(ArgumentsNum);
-		for (int32 i=0; i<ArgumentsNum; ++i)
-		{
-			FRigVMTemplateArgument Argument;
-			Argument.Serialize(Ar);
-			Arguments.Add(Argument);
-		}
-	}	
-}
-
 FRigVMTemplate::FRigVMTemplate(UScriptStruct* InStruct, const FString& InTemplateName, int32 InFunctionIndex)
 	: Index(INDEX_NONE)
 	, Notation(NAME_None)
@@ -1116,8 +1168,7 @@ bool FRigVMTemplate::Merge(const FRigVMTemplate& InOther)
 		int32 MatchingArguments = 0;
 		for(int32 ArgumentIndex = 0; ArgumentIndex < Arguments.Num(); ArgumentIndex++)
 		{
-			if(Arguments[ArgumentIndex].GetTypeIndices()[PermutationIndex] ==
-				InOther.Arguments[ArgumentIndex].TypeIndices[0])
+			if (Arguments[ArgumentIndex].GetTypeIndex(PermutationIndex) == InOther.Arguments[ArgumentIndex].GetTypeIndex(0))
 			{
 				MatchingArguments++;
 			}
@@ -1137,27 +1188,27 @@ bool FRigVMTemplate::Merge(const FRigVMTemplate& InOther)
 
 	for (int32 ArgumentIndex = 0; ArgumentIndex < Arguments.Num(); ArgumentIndex++)
 	{
-		if (InOther.Arguments[ArgumentIndex].TypeIndices.Num() != 1)
+		const FRigVMTemplateArgument& OtherArg = InOther.Arguments[ArgumentIndex];
+		if (OtherArg.GetNumTypes() != 1)
 		{
 			return false;
 		}
 
-		NewArgs.Add(Arguments[ArgumentIndex]);
-
 		// Add Other argument information into the TypeToPermutations map
 		{
-			const TRigVMTypeIndex OtherTypeIndex = InOther.Arguments[ArgumentIndex].TypeIndices[0];
-			const int32 NewPermutationIndex = NewArgs[ArgumentIndex].TypeIndices.Num();
-			if (TArray<int32>* ArgTypePermutations = NewArgs[ArgumentIndex].TypeToPermutations.Find(OtherTypeIndex))
+			FRigVMTemplateArgument& NewArg = NewArgs.Add_GetRef(Arguments[ArgumentIndex]);
+			const TRigVMTypeIndex OtherTypeIndex = OtherArg.GetTypeIndex(0);
+			const int32 NewPermutationIndex = NewArg.GetNumTypes();
+			if (TArray<int32>* ArgTypePermutations = NewArg.TypeToPermutations.Find(OtherTypeIndex))
 			{
 				ArgTypePermutations->Add(NewPermutationIndex);
 			}
 			else
 			{
-				NewArgs[ArgumentIndex].TypeToPermutations.Add(OtherTypeIndex, {NewPermutationIndex});
+				NewArg.TypeToPermutations.Add(OtherTypeIndex, {NewPermutationIndex});
 			}
+			NewArg.TypeIndices.Add(OtherTypeIndex);
 		}
-		NewArgs[ArgumentIndex].TypeIndices.Add(InOther.Arguments[ArgumentIndex].TypeIndices[0]);
 	}
 
 	Arguments = NewArgs;
@@ -1225,7 +1276,7 @@ const TArray<FRigVMExecuteArgument>& FRigVMTemplate::GetExecuteArguments(const F
 						FRigVMTemplateArgument Argument(*It);
 						if(Argument.IsExecute())
 						{
-							ExecuteArguments.Emplace(Argument.Name, Argument.Direction, Argument.TypeIndices[0]);
+							ExecuteArguments.Emplace(Argument.Name, Argument.Direction, Argument.GetTypeIndex(0));
 						}
 					}
 				}
@@ -1293,7 +1344,7 @@ const FRigVMFunction* FRigVMTemplate::GetOrCreatePermutation(int32 InIndex)
 	{
 		return Function;
 	}
-	
+
 	if(Permutations[InIndex] == INDEX_NONE && UsesDispatch())
 	{
 		FScopeLock RegisterFunctionScopeLock(&FRigVMRegistry::RegisterFunctionMutex);
@@ -1302,7 +1353,7 @@ const FRigVMFunction* FRigVMTemplate::GetOrCreatePermutation(int32 InIndex)
 		FTypeMap Types;
 		for(const FRigVMTemplateArgument& Argument : Arguments)
 		{
-			Types.Add(Argument.GetName(), Argument.TypeIndices[InIndex]);
+			Types.Add(Argument.GetName(), Argument.GetTypeIndex(InIndex));
 		}
 
 		FRigVMDispatchFactory* Factory = Delegates.GetDispatchFactoryDelegate.Execute();
@@ -1313,7 +1364,7 @@ const FRigVMFunction* FRigVMTemplate::GetOrCreatePermutation(int32 InIndex)
 			TArray<FRigVMFunctionArgument> FunctionArguments;
 			for(const FRigVMTemplateArgument& Argument : Arguments)
 			{
-				const FRigVMTemplateArgumentType& Type = Registry.GetType(Argument.TypeIndices[InIndex]);
+				const FRigVMTemplateArgumentType& Type = Registry.GetType(Argument.GetTypeIndex(InIndex));
 				FunctionArguments.Add(FRigVMFunctionArgument(Argument.Name.ToString(), Type.CPPType.ToString()));
 			}
 			
@@ -1388,7 +1439,7 @@ bool FRigVMTemplate::Resolve(FTypeMap& InOutTypes, TArray<int32>& OutPermutation
 	InOutTypes.Reset();
 
 	OutPermutationIndices.Reset();
-	for (int32 PermutationIndex = 0; PermutationIndex < Permutations.Num(); PermutationIndex++)
+	for (int32 PermutationIndex = 0; PermutationIndex < NumPermutations(); PermutationIndex++)
 	{
 		OutPermutationIndices.Add(PermutationIndex);
 	}
@@ -1397,7 +1448,7 @@ bool FRigVMTemplate::Resolve(FTypeMap& InOutTypes, TArray<int32>& OutPermutation
 	{
 		if (Argument.IsSingleton())
 		{
-			InOutTypes.Add(Argument.Name, Argument.TypeIndices[0]);
+			InOutTypes.Add(Argument.Name, Argument.GetTypeIndex(0));
 			continue;
 		}
 		else if (const TRigVMTypeIndex* InputType = InputTypes.Find(Argument.Name))
@@ -1409,10 +1460,11 @@ bool FRigVMTemplate::Resolve(FTypeMap& InOutTypes, TArray<int32>& OutPermutation
 			// Using a map to collect all permutations that we can keep/remove
 			// instead of removing them one by one, which can be costly
 			TMap<int32, bool> PermutationsToKeep;
-			
-			for (int32 PermutationIndex = 0; PermutationIndex < Argument.TypeIndices.Num(); PermutationIndex++)
+
+			for (int32 PermutationIndex = 0; PermutationIndex < Argument.GetNumTypes(); PermutationIndex++)
 			{
-				if(!Registry.CanMatchTypes(Argument.TypeIndices[PermutationIndex], *InputType, bAllowFloatingPointCasts))
+				const TRigVMTypeIndex Type = Argument.GetTypeIndex(PermutationIndex);
+				if(!Registry.CanMatchTypes(Type, *InputType, bAllowFloatingPointCasts))
 				{
 					PermutationsToKeep.FindOrAdd(PermutationIndex) = false;
 				}
@@ -1424,10 +1476,10 @@ bool FRigVMTemplate::Resolve(FTypeMap& InOutTypes, TArray<int32>& OutPermutation
 					// if the type matches - but it's not the exact same
 					if(!bFoundPerfectMatch)
 					{
-						MatchedType = Argument.TypeIndices[PermutationIndex];
+						MatchedType = Type;
 
 						// if we found the perfect match - let's stop here
-						if(Argument.TypeIndices[PermutationIndex] == *InputType)
+						if(Type == *InputType)
 						{
 							bFoundPerfectMatch = true;
 						}
@@ -1437,7 +1489,8 @@ bool FRigVMTemplate::Resolve(FTypeMap& InOutTypes, TArray<int32>& OutPermutation
 
 			OutPermutationIndices = OutPermutationIndices.FilterByPredicate([PermutationsToKeep](int32 PermutationIndex) -> bool
 			{
-				return PermutationsToKeep[PermutationIndex];
+				const bool* Value = PermutationsToKeep.Find(PermutationIndex);
+				return Value ? *Value : false;
 			});
 			
 			if(bFoundMatch)
@@ -1476,7 +1529,7 @@ bool FRigVMTemplate::Resolve(FTypeMap& InOutTypes, TArray<int32>& OutPermutation
 		InOutTypes.Reset();
 		for (const FRigVMTemplateArgument& Argument : Arguments)
 		{
-			InOutTypes.Add(Argument.Name, Argument.TypeIndices[OutPermutationIndices[0]]);
+			InOutTypes.Add(Argument.Name, Argument.GetTypeIndex(OutPermutationIndices[0]));
 		}
 	}
 	else if (OutPermutationIndices.Num() > 1)
@@ -1485,7 +1538,7 @@ bool FRigVMTemplate::Resolve(FTypeMap& InOutTypes, TArray<int32>& OutPermutation
 		{
 			if (Argument.IsSingleton(OutPermutationIndices))
 			{
-				InOutTypes.FindChecked(Argument.Name) = Argument.TypeIndices[OutPermutationIndices[0]];
+				InOutTypes.FindChecked(Argument.Name) = Argument.GetTypeIndex(OutPermutationIndices[0]);
 			}
 		}
 	}
@@ -1526,12 +1579,10 @@ uint32 FRigVMTemplate::GetTypesHashFromTypes(const FTypeMap& InTypes) const
 
 bool FRigVMTemplate::ContainsPermutation(const FTypeMap& InTypes) const
 {
-	const FRigVMRegistry& Registry = FRigVMRegistry::Get();
-
 	// If they type map is valid (full description of arguments), then we can rely on
 	// the TypesHashToPermutation cache. Otherwise, we will have to search for a specific permutation
 	// by filtering types
-	uint32 TypesHash = GetTypesHashFromTypes(InTypes);
+	const uint32 TypesHash = GetTypesHashFromTypes(InTypes);
 	if (TypesHash != 0)
 	{
 		return TypesHashToPermutation.Contains(TypesHash);
@@ -1542,17 +1593,18 @@ bool FRigVMTemplate::ContainsPermutation(const FTypeMap& InTypes) const
 	{
 		if (const FRigVMTemplateArgument* Argument = FindArgument(Pair.Key))
 		{
-			if (const TArray<int32>* ArgumentPermutations = Argument->TypeToPermutations.Find(Pair.Value))
+			const TArray<int32>& ArgumentPermutations = Argument->GetPermutations(Pair.Value);
+			if (!ArgumentPermutations.IsEmpty())
 			{
 				// If possible permutations is empty, initialize it
 				if (PossiblePermutations.IsEmpty())
 				{
-					PossiblePermutations = *ArgumentPermutations;
+					PossiblePermutations = ArgumentPermutations;
 				}
 				else
 				{
 					// Intersect possible permutations and the permutations found for this argument
-					PossiblePermutations = ArgumentPermutations->FilterByPredicate([PossiblePermutations](const int32& ArgPermutation)
+					PossiblePermutations = ArgumentPermutations.FilterByPredicate([PossiblePermutations](const int32& ArgPermutation)
 					{
 						return PossiblePermutations.Contains(ArgPermutation);
 					});
@@ -1672,9 +1724,13 @@ FRigVMTemplateTypeMap FRigVMTemplate::GetTypesForPermutation(const int32 InPermu
 	for (int32 ArgIndex = 0; ArgIndex < NumArguments(); ++ArgIndex)
 	{
 		const FRigVMTemplateArgument* Argument = GetArgument(ArgIndex);
-		if (Argument->GetTypeIndices().Num() > InPermutationIndex)
+		if (Argument->GetNumTypes() > InPermutationIndex)
 		{
-			TypeMap.Add(Argument->GetName(), Argument->GetTypeIndices()[InPermutationIndex]);
+			TypeMap.Add(Argument->GetName(), Argument->GetTypeIndex(InPermutationIndex));
+		}
+		else if (Argument->IsSingleton())
+		{
+			TypeMap.Add(Argument->GetName(), Argument->GetTypeIndex(0));
 		}
 		else
 		{
@@ -1808,17 +1864,25 @@ bool FRigVMTemplate::AddTypeForArgument(const FName& InArgumentName, TRigVMTypeI
 				{
 					return false;
 				}
-			
+
+				uint32 TypeHash=0;
 				for(FRigVMTemplateArgument& Argument : Arguments)
 				{
 					const TRigVMTypeIndex TypeIndex = Types.FindChecked(Argument.Name);
-					Argument.TypeIndices.Add(TypeIndex);
-					Argument.TypeToPermutations.FindOrAdd(TypeIndex).Add(Permutations.Num());
+					if (!Argument.bUseCategories)
+					{
+						Argument.TypeIndices.Add(TypeIndex);
+						Argument.TypeToPermutations.FindOrAdd(TypeIndex).Add(Permutations.Num());
+					}
+					else
+					{
+						Argument.UpdateTypeToPermutations();
+					}
+					TypeHash = HashCombine(TypeHash, GetTypeHash(TypeIndex));
 				}
-
+				
 				Permutations.Add(INDEX_NONE);
-
-				UpdateTypesHashToPermutation(Permutations.Num()-1);
+				TypesHashToPermutation.Add(TypeHash, Permutations.Num()-1);
 			}
 			else
 			{
@@ -1875,23 +1939,23 @@ void FRigVMTemplate::RecomputeTypesHashToPermutations()
 		uint32 TypesHash=0;
 		for (int32 ArgIndex=0; ArgIndex<NumArguments(); ++ArgIndex)
 		{
-			TypesHash = HashCombine(Hash, GetTypeHash(Arguments[ArgIndex].TypeIndices[PermutationIndex]));
+			TypesHash = HashCombine(Hash, GetTypeHash(Arguments[ArgIndex].GetTypeIndex(PermutationIndex)));
 		}
 		TypesHashToPermutation.Add(TypesHash, PermutationIndex);
 	}
 }
 
-void FRigVMTemplate::UpdateTypesHashToPermutation(const int32& InPermutation)
+void FRigVMTemplate::UpdateTypesHashToPermutation(const int32 InPermutation)
 {
-	if (InPermutation < 0 || InPermutation >= NumPermutations())
+	if (!Permutations.IsValidIndex(InPermutation))
 	{
 		return;
 	}
-	
+
 	uint32 TypeHash=0;
 	for (const FRigVMTemplateArgument& Argument : Arguments)
 	{
-		TypeHash = HashCombine(TypeHash, GetTypeHash(Argument.TypeIndices[InPermutation]));
+		TypeHash = HashCombine(TypeHash, GetTypeHash(Argument.GetTypeIndex(InPermutation)));
 	}
 	TypesHashToPermutation.Add(TypeHash, InPermutation);
 }
@@ -1901,10 +1965,10 @@ uint32 GetTypeHash(const FRigVMTemplateArgument& InArgument)
 	const FRigVMRegistry& Registry = FRigVMRegistry::Get();
 	uint32 Hash = GetTypeHash(InArgument.Name.ToString());
 	Hash = HashCombine(Hash, GetTypeHash((int32)InArgument.Direction));
-	for(const TRigVMTypeIndex& TypeIndex : InArgument.TypeIndices)
+	InArgument.ForEachType([&](const TRigVMTypeIndex TypeIndex)
 	{
 		Hash = HashCombine(Hash, Registry.GetHashForType(TypeIndex));
-	}
+	});
 	return Hash;
 }
 
diff --git a/Engine/Plugins/Runtime/RigVM/Source/RigVM/Public/RigVMCore/RigVMRegistry.h b/Engine/Plugins/Runtime/RigVM/Source/RigVM/Public/RigVMCore/RigVMRegistry.h
--- a/Engine/Plugins/Runtime/RigVM/Source/RigVM/Public/RigVMCore/RigVMRegistry.h
+++ b/Engine/Plugins/Runtime/RigVM/Source/RigVM/Public/RigVMCore/RigVMRegistry.h
@@ -338,7 +338,8 @@ private:
 	bool IsAllowedType(const UStruct* InStruct) const;
 	bool IsAllowedType(const UClass* InClass) const;
 
-	void RegisterTypeInCategory(FRigVMTemplateArgument::ETypeCategory InCategory, TRigVMTypeIndex InTypeIndex);
+	void RegisterTypeInCategory(const FRigVMTemplateArgument::ETypeCategory InCategory, const TRigVMTypeIndex InTypeIndex);
+	void PropagateTypeAddedToCategory(const FRigVMTemplateArgument::ETypeCategory InCategory, const TRigVMTypeIndex InTypeIndex);
 	
 	void RemoveTypeInCategory(FRigVMTemplateArgument::ETypeCategory InCategory, TRigVMTypeIndex InTypeIndex);
 
diff --git a/Engine/Plugins/Runtime/RigVM/Source/RigVM/Public/RigVMCore/RigVMTemplate.h b/Engine/Plugins/Runtime/RigVM/Source/RigVM/Public/RigVMCore/RigVMTemplate.h
--- a/Engine/Plugins/Runtime/RigVM/Source/RigVM/Public/RigVMCore/RigVMTemplate.h
+++ b/Engine/Plugins/Runtime/RigVM/Source/RigVM/Public/RigVMCore/RigVMTemplate.h
@@ -198,10 +198,7 @@ struct RIGVM_API FRigVMTemplateArgument
 	FRigVMTemplateArgument(const FName& InName, ERigVMPinDirection InDirection);
 	FRigVMTemplateArgument(const FName& InName, ERigVMPinDirection InDirection, TRigVMTypeIndex InType);
 	FRigVMTemplateArgument(const FName& InName, ERigVMPinDirection InDirection, const TArray<TRigVMTypeIndex>& InTypeIndices);
-	FRigVMTemplateArgument(const FName& InName, ERigVMPinDirection InDirection, const TArray<ETypeCategory>& InTypeCategories, const FTypeFilter& InTypeFilter = {});
-
-	// Serialize
-	void Serialize(FArchive& Ar);
+	FRigVMTemplateArgument(const FName& InName, ERigVMPinDirection InDirection, const TArray<ETypeCategory>& InTypeCategories);
 
 	// returns the name of the argument
 	const FName& GetName() const { return Name; }
@@ -213,7 +210,24 @@ struct RIGVM_API FRigVMTemplateArgument
 	bool SupportsTypeIndex(TRigVMTypeIndex InTypeIndex, TRigVMTypeIndex* OutTypeIndex = nullptr) const;
 
 	// returns the flat list of types (including duplicates) of this argument
-	const TArray<TRigVMTypeIndex>& GetTypeIndices() const;
+	void GetAllTypes(TArray<TRigVMTypeIndex>& OutTypes) const;
+	
+	TRigVMTypeIndex GetTypeIndex(const int32 InIndex) const;
+	int32 GetNumTypes() const;
+	void AddTypeIndex(const TRigVMTypeIndex InTypeIndex);
+	void RemoveType(const int32 InIndex);
+	void ForEachType(TFunction<void(const TRigVMTypeIndex InType)>&& InCallback) const;
+	int32 FindTypeIndex(const TRigVMTypeIndex InTypeIndex) const;
+
+	template <typename Predicate>
+	int32 IndexOfByPredicate(Predicate Pred) const
+	{
+		if (!bUseCategories)
+		{
+			return TypeIndices.IndexOfByPredicate(Pred);
+		}
+		return CategoryViews(TypeCategories).IndexOfByPredicate(Pred);
+	}
 
 	// returns an array of all of the supported types
 	TArray<TRigVMTypeIndex> GetSupportedTypeIndices(const TArray<int32>& InPermutationIndices = TArray<int32>()) const;
@@ -235,16 +249,19 @@ struct RIGVM_API FRigVMTemplateArgument
 	RIGVM_API friend uint32 GetTypeHash(const FRigVMTemplateArgument& InArgument);
 
 	// Get the map of types to permutation indices
-	const TMap<TRigVMTypeIndex, TArray<int32>>& GetTypeToPermutations() const { return TypeToPermutations; }
+	const TArray<int32>& GetPermutations(const TRigVMTypeIndex InType) const;
+	void InvalidatePermutations(const TRigVMTypeIndex InType);
 
 protected:
 
-	int32 Index;
-	FName Name;
-	ERigVMPinDirection Direction;
+	int32 Index = INDEX_NONE;
+	FName Name = NAME_None;
+	ERigVMPinDirection Direction = ERigVMPinDirection::IO;
+
 	TArray<TRigVMTypeIndex> TypeIndices;
+	mutable TMap<TRigVMTypeIndex, TArray<int32>> TypeToPermutations;
 
-	TMap<TRigVMTypeIndex, TArray<int32>> TypeToPermutations;
+	bool bUseCategories = false;
 	TArray<ETypeCategory> TypeCategories;
 
 	// constructor from a property. this forces the type to be created
@@ -258,8 +275,39 @@ protected:
 	friend class URigVMController;
 	friend struct FRigVMRegistry;
 	friend struct FRigVMStructUpgradeInfo;
-	friend struct FRigVMSetLibraryTemplateAction;
 	friend class URigVMCompiler;
+
+private:
+	struct CategoryViews
+	{
+		CategoryViews() = delete;
+		CategoryViews(const TArray<ETypeCategory>& InCategories);
+		
+		void ForEachType(TFunction<void(const TRigVMTypeIndex InType)>&& InCallback) const;
+
+		TRigVMTypeIndex GetTypeIndex(int32 InIndex) const;
+		
+		int32 FindIndex(const TRigVMTypeIndex InTypeIndex) const;
+
+		template <typename Predicate>
+		int32 IndexOfByPredicate(Predicate Pred) const
+		{
+			int32 Offset = 0;
+			for (const TArrayView<const TRigVMTypeIndex>& TypeView: Types)
+			{
+				const int32 Found = TypeView.IndexOfByPredicate(Pred);
+				if (Found != INDEX_NONE)
+				{
+					return Found + Offset;
+				}
+				Offset += TypeView.Num();
+			}
+			return INDEX_NONE;
+		}
+		
+	private:
+		TArray<TArrayView<const TRigVMTypeIndex>> Types;
+	};
 };
 
 /**
@@ -307,9 +355,6 @@ public:
 	// Default constructor
 	FRigVMTemplate();
 
-	// Serialize
-	void Serialize(FArchive& Ar);
-
 	// returns true if this is a valid template
 	bool IsValid() const;
 
@@ -454,7 +499,7 @@ public:
 
 	void RecomputeTypesHashToPermutations();
 
-	void UpdateTypesHashToPermutation(const int32& InPermutation);
+	void UpdateTypesHashToPermutation(const int32 InPermutation);
 
 	RIGVM_API friend uint32 GetTypeHash(const FRigVMTemplate& InTemplate);
 
@@ -486,6 +531,5 @@ private:
 	friend struct FRigVMRegistry;
 	friend class URigVMController;
 	friend class URigVMLibraryNode;
-	friend struct FRigVMSetLibraryTemplateAction;
 	friend struct FRigVMDispatchFactory;
 };
diff --git a/Engine/Plugins/Runtime/RigVM/Source/RigVMDeveloper/Private/RigVMCompiler/RigVMCompiler.cpp b/Engine/Plugins/Runtime/RigVM/Source/RigVMDeveloper/Private/RigVMCompiler/RigVMCompiler.cpp
--- a/Engine/Plugins/Runtime/RigVM/Source/RigVMDeveloper/Private/RigVMCompiler/RigVMCompiler.cpp
+++ b/Engine/Plugins/Runtime/RigVM/Source/RigVMDeveloper/Private/RigVMCompiler/RigVMCompiler.cpp
@@ -826,16 +826,16 @@ bool URigVMCompiler::Compile(const FRigVMCompileSettings& InSettings, TArray<URi
 							const TRigVMTypeIndex ExpectedTypeIndex = ExpectedArgument.GetSupportedTypeIndices()[0];
 							if (URigVMPin* Pin = UnitNode->FindPin(ExpectedArgument.Name.ToString()))
 							{
-								if (Pin->GetTypeIndex() != ExpectedArgument.GetTypeIndices()[0])
+								if (Pin->GetTypeIndex() != ExpectedArgument.GetTypeIndex(0))
 								{
-									FString MissingPinMessage = FString::Printf(TEXT("Could not find pin %s of type %s in Node @@."), *ExpectedArgument.Name.ToString(), *FRigVMRegistry::Get().GetType(ExpectedArgument.TypeIndices[0]).CPPType.ToString());
+									FString MissingPinMessage = FString::Printf(TEXT("Could not find pin %s of type %s in Node @@."), *ExpectedArgument.Name.ToString(), *Registry.GetType(ExpectedArgument.GetTypeIndex(0)).CPPType.ToString());
 									Settings.ASTSettings.Report(EMessageSeverity::Error, ModelNode, MissingPinMessage);
 									bEncounteredGraphError = true;
 								}
 							}
 							else
 							{
-								FString MissingPinMessage = FString::Printf(TEXT("Could not find pin %s of type %s in Node @@."), *ExpectedArgument.Name.ToString(), *FRigVMRegistry::Get().GetType(ExpectedArgument.TypeIndices[0]).CPPType.ToString());
+								FString MissingPinMessage = FString::Printf(TEXT("Could not find pin %s of type %s in Node @@."), *ExpectedArgument.Name.ToString(), *Registry.GetType(ExpectedArgument.GetTypeIndex(0)).CPPType.ToString());
 								Settings.ASTSettings.Report(EMessageSeverity::Error, ModelNode, MissingPinMessage);
 								bEncounteredGraphError = true;
 							}
diff --git a/Engine/Plugins/Runtime/RigVM/Source/RigVMDeveloper/Private/RigVMDeveloperTypeUtils.cpp b/Engine/Plugins/Runtime/RigVM/Source/RigVMDeveloper/Private/RigVMDeveloperTypeUtils.cpp
--- a/Engine/Plugins/Runtime/RigVM/Source/RigVMDeveloper/Private/RigVMDeveloperTypeUtils.cpp
+++ b/Engine/Plugins/Runtime/RigVM/Source/RigVMDeveloper/Private/RigVMDeveloperTypeUtils.cpp
@@ -608,14 +608,15 @@ const TArray<TRigVMTypeIndex>& RigVMTypeUtils::GetAvailableCasts(const TRigVMTyp
 		{
 			TArray<TRigVMTypeIndex>& AvailableCasts = AvailableCastMap.Add(InTypeIndex);
 			
-			for(int32 Index = 0; Index < ArgumentA->GetTypeIndices().Num(); Index++)
+			int32 Index = 0;
+			ArgumentA->ForEachType([&](const TRigVMTypeIndex TypeIndexA)
 			{
-				const TRigVMTypeIndex& TypeIndexA = ArgumentA->GetTypeIndices()[Index];
 				if(TypeIndexA == InTypeIndex)
 				{
-					AvailableCasts.Add(ArgumentB->GetTypeIndices()[Index]);
+					AvailableCasts.Add(ArgumentB->GetTypeIndex(Index));
 				}
-			}
+				Index++;
+			});
 
 			return AvailableCasts;
 		}
@@ -649,17 +650,17 @@ const FRigVMFunction* RigVMTypeUtils::GetCastForTypeIndices(const TRigVMTypeInde
 
 			if(SourceArgument && TargetArgument)
 			{
-				const TArray<int32>* SourcePermutations = SourceArgument->GetTypeToPermutations().Find(InSourceTypeIndex);
-				const TArray<int32>* TargetPermutations = TargetArgument->GetTypeToPermutations().Find(InTargetTypeIndex);
-				if(SourcePermutations && TargetPermutations)
+				const TArray<int32>& SourcePermutations = SourceArgument->GetPermutations(InSourceTypeIndex);
+				const TArray<int32>& TargetPermutations = TargetArgument->GetPermutations(InTargetTypeIndex);
+				if(!SourcePermutations.IsEmpty() && !TargetPermutations.IsEmpty())
 				{
-					for(int32 SourceIndex = 0, SourceCount = SourcePermutations->Num(); SourceIndex < SourceCount; ++SourceIndex)
+					for(int32 SourceIndex = 0, SourceCount = SourcePermutations.Num(); SourceIndex < SourceCount; ++SourceIndex)
 					{
-						for(int32 TargetIndex = 0, TargetCount = TargetPermutations->Num(); TargetIndex < TargetCount; ++TargetIndex)
+						for(int32 TargetIndex = 0, TargetCount = TargetPermutations.Num(); TargetIndex < TargetCount; ++TargetIndex)
 						{
-							if((*SourcePermutations)[SourceIndex] == (*TargetPermutations)[TargetIndex])
+							if(SourcePermutations[SourceIndex] == TargetPermutations[TargetIndex])
 							{
-								return const_cast<FRigVMTemplate*>(CastTemplate)->GetOrCreatePermutation((*SourcePermutations)[SourceIndex]);
+								return const_cast<FRigVMTemplate*>(CastTemplate)->GetOrCreatePermutation(SourcePermutations[SourceIndex]);
 							}
 						}
 					}
diff --git a/Engine/Plugins/Runtime/RigVM/Source/RigVMDeveloper/Private/RigVMModel/Nodes/RigVMTemplateNode.cpp b/Engine/Plugins/Runtime/RigVM/Source/RigVMDeveloper/Private/RigVMModel/Nodes/RigVMTemplateNode.cpp
--- a/Engine/Plugins/Runtime/RigVM/Source/RigVMDeveloper/Private/RigVMModel/Nodes/RigVMTemplateNode.cpp
+++ b/Engine/Plugins/Runtime/RigVM/Source/RigVMDeveloper/Private/RigVMModel/Nodes/RigVMTemplateNode.cpp
@@ -196,7 +196,7 @@ FText URigVMTemplateNode::GetToolTipTextForPin(const URigVMPin* InPin) const
 			{
 				FString Tooltip;
 				
-				if(Arg->GetTypeIndices().Num() > 100)
+				if(Arg->GetNumTypes() > 100)
 				{
 					Tooltip = TEXT("Supports any type.");
 				}
@@ -207,7 +207,7 @@ FText URigVMTemplateNode::GetToolTipTextForPin(const URigVMPin* InPin) const
 					TArray<TRigVMTypeIndex> TypesPrinted;
 					for (int32 Index=0; Index<Template->NumPermutations(); ++Index)
 					{
-						const TRigVMTypeIndex TypeIndex = Arg->GetTypeIndices()[Index];
+						const TRigVMTypeIndex TypeIndex = Arg->GetTypeIndex(Index);
 						if (TypesPrinted.Contains(TypeIndex))
 						{
 							continue;
@@ -573,7 +573,7 @@ FString URigVMTemplateNode::GetInitialDefaultValueForPin(const FName& InRootPinN
 		{
 			FString NewDefaultValue;
 
-			const TRigVMTypeIndex TypeIndex = Argument->GetTypeIndices()[PermutationIndex];
+			const TRigVMTypeIndex TypeIndex = Argument->GetTypeIndex(PermutationIndex);
 
 			// INDEX_NONE indicates deleted permutation
 			if (TypeIndex == INDEX_NONE)
@@ -763,8 +763,8 @@ TArray<int32> URigVMTemplateNode::FindPermutationsForTypes(const TArray<FRigVMTe
 				const FRigVMTemplateArgument* Argument = Args[ArgIndex];
 				if (Argument)
 				{
-					if ((bAllowCasting && !Registry.CanMatchTypes(Argument->GetTypeIndices()[i], TypeIndices[ArgIndex], true)) ||
-						(!bAllowCasting && Argument->GetTypeIndices()[i] != TypeIndices[ArgIndex]))
+					if ((bAllowCasting && !Registry.CanMatchTypes(Argument->GetTypeIndex(i), TypeIndices[ArgIndex], true)) ||
+						(!bAllowCasting && Argument->GetTypeIndex(i) != TypeIndices[ArgIndex]))
 					{
 						bAllArgsMatched = false;
 						break;
diff --git a/Engine/Plugins/Runtime/RigVM/Source/RigVMDeveloper/Private/RigVMModel/RigVMController.cpp b/Engine/Plugins/Runtime/RigVM/Source/RigVMDeveloper/Private/RigVMModel/RigVMController.cpp
--- a/Engine/Plugins/Runtime/RigVM/Source/RigVMDeveloper/Private/RigVMModel/RigVMController.cpp
+++ b/Engine/Plugins/Runtime/RigVM/Source/RigVMDeveloper/Private/RigVMModel/RigVMController.cpp
@@ -8320,7 +8320,7 @@ FString URigVMController::AddAggregatePin(URigVMNode* InNode, const FString& InP
 					{
 						if (const FRigVMTemplateArgument* Argument = Template->FindArgument(Arg1->GetFName()))
 						{
-							const TArray<TRigVMTypeIndex>& Types = Argument->GetTypeIndices();
+							TArray<TRigVMTypeIndex> Types; Argument->GetAllTypes(Types);
 							AnswerType = RequestPinTypeSelectionDelegate.Execute(Types);
 						}
 					}
@@ -15107,8 +15107,9 @@ bool URigVMController::GenerateNewPinInfos(const FRigVMRegistry& Registry, URigV
 			{
 				if (Arg->IsSingleton())
 				{
-					DefaultValue = Factory->GetArgumentDefaultValue(Arg->Name, Arg->GetTypeIndices()[0]);
-					const FRigVMTemplateArgumentType& Type = Registry.GetType(Arg->GetTypeIndices()[0]);
+					const TRigVMTypeIndex Type0 = Arg->GetTypeIndex(0);
+					DefaultValue = Factory->GetArgumentDefaultValue(Arg->Name, Type0);
+					const FRigVMTemplateArgumentType& Type = Registry.GetType(Type0);
 					ArgumentScriptStruct = Cast<UScriptStruct>(Type.CPPTypeObject);
 				}
 			}
@@ -16568,7 +16569,7 @@ bool URigVMController::FullyResolveTemplateNode(URigVMTemplateNode* InNode, int3
 			const TArray<FRigVMTemplateArgument>& Arguments = Template->Arguments;
 			for(const FRigVMTemplateArgument& Argument : Arguments)
 			{
-				const TRigVMTypeIndex ExpectedTypeIndex = Argument.TypeIndices[InNode->ResolvedPermutation];
+				const TRigVMTypeIndex ExpectedTypeIndex = Argument.GetTypeIndex(InNode->ResolvedPermutation);
 				if(URigVMPin* Pin = InNode->FindPin(Argument.GetName().ToString()))
 				{
 					if(Pin->GetTypeIndex() != ExpectedTypeIndex && ExpectedTypeIndex != RigVMTypeUtils::TypeIndex::Execute)
@@ -16717,7 +16718,7 @@ bool URigVMController::FullyResolveTemplateNode(URigVMTemplateNode* InNode, int3
 				
 				URigVMPin* Pin = NewObject<URigVMPin>(Cast<UObject>(InNode), MissingPin.GetName());
 
-				const TRigVMTypeIndex TypeIndex = MissingPin.TypeIndices[InNode->ResolvedPermutation];
+				const TRigVMTypeIndex TypeIndex = MissingPin.GetTypeIndex(InNode->ResolvedPermutation);
 				const FRigVMTemplateArgumentType& Type = FRigVMRegistry::Get().GetType(TypeIndex); 
 				
 				Pin->Direction = MissingPin.GetDirection();
@@ -16986,7 +16987,7 @@ bool URigVMController::UpdateTemplateNodePinTypes(URigVMTemplateNode* InNode, bo
 	{
 		for (const FRigVMTemplateArgument* Argument : Arguments)
 		{
-			if (Argument && Argument->TypeIndices[Permutation] == INDEX_NONE)
+			if (Argument && Argument->GetTypeIndex(Permutation) == INDEX_NONE)
 			{
 				return true;
 			}
@@ -17017,7 +17018,7 @@ bool URigVMController::UpdateTemplateNodePinTypes(URigVMTemplateNode* InNode, bo
 			TArray<TRigVMTypeIndex>& Types = PinTypes.FindOrAdd(ResolvedPermutation);
 			if (const FRigVMTemplateArgument* Argument = Arguments[PinIndex])
 			{
-				Types.Add(Argument->TypeIndices[ResolvedPermutation]);
+				Types.Add(Argument->GetTypeIndex(ResolvedPermutation));
 				bAddedType = true;
 			}
 			else if (const FRigVMExecuteArgument* ExecuteArgument = Template->FindExecuteArgument(Pin->GetFName(), DispatchContext))
diff --git a/Engine/Plugins/Runtime/RigVM/Source/RigVMDeveloper/Private/RigVMModel/RigVMPin.cpp b/Engine/Plugins/Runtime/RigVM/Source/RigVMDeveloper/Private/RigVMModel/RigVMPin.cpp
--- a/Engine/Plugins/Runtime/RigVM/Source/RigVMDeveloper/Private/RigVMModel/RigVMPin.cpp
+++ b/Engine/Plugins/Runtime/RigVM/Source/RigVMDeveloper/Private/RigVMModel/RigVMPin.cpp
@@ -1938,7 +1938,7 @@ bool URigVMPin::CanLink(const URigVMPin* InSourcePin, const URigVMPin* InTargetP
 					FRigVMRegistry& Registry = FRigVMRegistry::Get();
 					for (int32 Permutation : SourcePermutations)
 					{
-						TRigVMTypeIndex Type = SourceRootArgument->GetTypeIndices()[Permutation];
+						TRigVMTypeIndex Type = SourceRootArgument->GetTypeIndex(Permutation);
 						for (int32 i=0; i<SourceLevels; ++i)
 						{
 							check(Registry.IsArrayType(Type));
@@ -1948,7 +1948,7 @@ bool URigVMPin::CanLink(const URigVMPin* InSourcePin, const URigVMPin* InTargetP
 					}
 					for (int32 Permutation : TargetPermutations)
 					{
-						TRigVMTypeIndex Type = TargetRootArgument->GetTypeIndices()[Permutation];
+						TRigVMTypeIndex Type = TargetRootArgument->GetTypeIndex(Permutation);
 						for (int32 i=0; i<TargetLevels; ++i)
 						{
 							check(Registry.IsArrayType(Type));
diff --git a/Engine/Plugins/Runtime/RigVM/Source/RigVMEditor/Private/EdGraph/NodeSpawners/RigVMEdGraphTemplateNodeSpawner.cpp b/Engine/Plugins/Runtime/RigVM/Source/RigVMEditor/Private/EdGraph/NodeSpawners/RigVMEdGraphTemplateNodeSpawner.cpp
--- a/Engine/Plugins/Runtime/RigVM/Source/RigVMEditor/Private/EdGraph/NodeSpawners/RigVMEdGraphTemplateNodeSpawner.cpp
+++ b/Engine/Plugins/Runtime/RigVM/Source/RigVMEditor/Private/EdGraph/NodeSpawners/RigVMEdGraphTemplateNodeSpawner.cpp
@@ -188,7 +188,7 @@ UEdGraphNode* URigVMEdGraphTemplateNodeSpawner::Invoke(UEdGraph* ParentGraph, FB
 
 				if(Argument->IsSingleton())
 				{
-					const TRigVMTypeIndex TypeIndex = Argument->GetTypeIndices()[0];
+					const TRigVMTypeIndex TypeIndex = Argument->GetTypeIndex(0);
 					const FRigVMTemplateArgumentType& Type = Registry.GetType(TypeIndex);
 					CPPType = Type.CPPType;
 					CPPTypeObject = Type.CPPTypeObject;
