RigVM: Improve log performance / memory footprint

diff --git a/Engine/Plugins/Runtime/RigVM/Source/RigVM/Private/RigVMCore/RigVMLog.cpp b/Engine/Plugins/Runtime/RigVM/Source/RigVM/Private/RigVMCore/RigVMLog.cpp
--- a/Engine/Plugins/Runtime/RigVM/Source/RigVM/Private/RigVMCore/RigVMLog.cpp
+++ b/Engine/Plugins/Runtime/RigVM/Source/RigVM/Private/RigVMCore/RigVMLog.cpp
@@ -24,11 +24,11 @@ void FRigVMLog::Reset()
 #endif
 }
 
-void FRigVMLog::Report(EMessageSeverity::Type InSeverity, const FName& InOperatorName, int32 InInstructionIndex, const FString& InMessage)
+void FRigVMLog::Report(const FRigVMLogSettings& InLogSettings, const FName& InFunctionName, int32 InInstructionIndex, const FString& InMessage)
 {
 #if WITH_EDITOR
 
-	if (!InMessage.IsEmpty())
+	if (!InMessage.IsEmpty() && InLogSettings.bLogOnce)
 	{
 		if (KnownMessages.Contains(InMessage))
 		{
@@ -36,7 +36,60 @@ void FRigVMLog::Report(EMessageSeverity::Type InSeverity, const FName& InOperato
 		}
 	}
 
-	int32 EntryIndex = Entries.Add(FLogEntry(InSeverity, InOperatorName, InInstructionIndex, InMessage));
-	KnownMessages.Add(InMessage, true);
+	int32 EntryIndex = Entries.Add(FLogEntry(InLogSettings.Severity, InFunctionName, InInstructionIndex, InMessage));
+
+	if(InLogSettings.bLogOnce)
+	{
+		KnownMessages.Add(InMessage, true);
+	}
 #endif
 }
+
+#if WITH_EDITOR
+
+void FRigVMLog::RemoveRedundantEntries()
+{
+	if(Entries.IsEmpty())
+	{
+		return;
+	}
+
+	TArray<FLogEntry> OldEntries;
+	Swap(OldEntries, Entries);
+
+	// find a single entry for each severity index per instruction.
+	// this prefers the latest (most recent) entries over old entries.
+	TMap<int32, FIntVector3> EntriesPerInstruction;
+	for(int32 EntryIndex = 0; EntryIndex < OldEntries.Num(); EntryIndex++)
+	{
+		const FLogEntry& OldEntry = OldEntries[EntryIndex];
+		if(OldEntry.InstructionIndex == INDEX_NONE)
+		{
+			Entries.Add(OldEntry);
+			continue;
+		}
+		
+		const int32 SeverityIndex = OldEntry.Severity == EMessageSeverity::Error ? 2 :
+			(OldEntry.Severity == EMessageSeverity::Warning ? 1 : 0);
+		
+		FIntVector3& EntryPerInstruction =
+			EntriesPerInstruction.FindOrAdd(OldEntry.InstructionIndex, {INDEX_NONE, INDEX_NONE, INDEX_NONE});
+
+		EntryPerInstruction[SeverityIndex] = EntryIndex;
+	}
+
+	// add the entries back - culling redundant / least recent entries
+	for(const TPair<int32, FIntVector3>& Pair : EntriesPerInstruction)
+	{
+		for(int32 SeverityIndex = 0; SeverityIndex < 3; SeverityIndex++)
+		{
+			const int32 EntryIndex = Pair.Value[SeverityIndex];
+			if(EntryIndex != INDEX_NONE)
+			{
+				Entries.Add(OldEntries[EntryIndex]);
+			}
+		}
+	}
+}
+
+#endif
\ No newline at end of file
diff --git a/Engine/Plugins/Runtime/RigVM/Source/RigVM/Private/RigVMFunctions/RigVMDispatch_Print.cpp b/Engine/Plugins/Runtime/RigVM/Source/RigVM/Private/RigVMFunctions/RigVMDispatch_Print.cpp
--- a/Engine/Plugins/Runtime/RigVM/Source/RigVM/Private/RigVMFunctions/RigVMDispatch_Print.cpp
+++ b/Engine/Plugins/Runtime/RigVM/Source/RigVM/Private/RigVMFunctions/RigVMDispatch_Print.cpp
@@ -127,7 +127,7 @@ void FRigVMDispatch_Print::Execute(FRigVMExtendedExecuteContext& InContext, FRig
 	}
 
 	static constexpr TCHAR LogFormat[] = TEXT("%s[%04d] %s%s");
-	InContext.GetPublicData<>().Logf(EMessageSeverity::Info, LogFormat, *ObjectPath, InContext.GetPublicData<>().GetInstructionIndex(), *Prefix, *String);
+	InContext.GetPublicData<>().Logf({EMessageSeverity::Info, false}, LogFormat, *ObjectPath, InContext.GetPublicData<>().GetInstructionIndex(), *Prefix, *String);
 	const UObject* WorldObject = (const UObject*)InContext.VM;
 
 	if(ScreenDuration > SMALL_NUMBER && WorldObject)
diff --git a/Engine/Plugins/Runtime/RigVM/Source/RigVM/Private/RigVMHost.cpp b/Engine/Plugins/Runtime/RigVM/Source/RigVM/Private/RigVMHost.cpp
--- a/Engine/Plugins/Runtime/RigVM/Source/RigVM/Private/RigVMHost.cpp
+++ b/Engine/Plugins/Runtime/RigVM/Source/RigVM/Private/RigVMHost.cpp
@@ -792,17 +792,17 @@ void URigVMHost::UpdateVMSettings()
 	{
 #if WITH_EDITOR
 		// setup array handling and error reporting on the VM
-		VMRuntimeSettings.SetLogFunction([this](EMessageSeverity::Type InSeverity, const FRigVMExecuteContext* InContext, const FString& Message)
+		VMRuntimeSettings.SetLogFunction([this](const FRigVMLogSettings& InLogSettings, const FRigVMExecuteContext* InContext, const FString& Message)
 			{
 				check(InContext);
 
 				if (RigVMLog)
 				{
-					RigVMLog->Report(InSeverity, InContext->GetFunctionName(), InContext->GetInstructionIndex(), Message);
+					RigVMLog->Report(InLogSettings, InContext->GetFunctionName(), InContext->GetInstructionIndex(), Message);
 				}
 				else
 				{
-					LogOnce(InSeverity, InContext->GetInstructionIndex(), Message);
+					LogOnce(InLogSettings.Severity, InContext->GetInstructionIndex(), Message);
 				}
 			});
 #endif
diff --git a/Engine/Plugins/Runtime/RigVM/Source/RigVM/Public/RigVMCore/RigVMExecuteContext.h b/Engine/Plugins/Runtime/RigVM/Source/RigVM/Public/RigVMCore/RigVMExecuteContext.h
--- a/Engine/Plugins/Runtime/RigVM/Source/RigVM/Public/RigVMCore/RigVMExecuteContext.h
+++ b/Engine/Plugins/Runtime/RigVM/Source/RigVM/Public/RigVMCore/RigVMExecuteContext.h
@@ -218,7 +218,7 @@ struct RIGVM_API FRigVMRuntimeSettings
 	/*
 	 * The function to use for logging anything from the VM to the host
 	 */
-	using LogFunctionType = TFunction<void(EMessageSeverity::Type,const FRigVMExecuteContext*,const FString&)>;
+	using LogFunctionType = TFunction<void(const FRigVMLogSettings&,const FRigVMExecuteContext*,const FString&)>;
 	TSharedPtr<LogFunctionType> LogFunction = nullptr;
 
 	void SetLogFunction(LogFunctionType InLogFunction)
@@ -268,19 +268,19 @@ struct FRigVMExecuteContext
 
 	virtual ~FRigVMExecuteContext() {}
 
-	void Log(EMessageSeverity::Type InSeverity, const FString& InMessage) const
+	void Log(const FRigVMLogSettings& InLogSettings, const FString& InMessage) const
 	{
 		if(RuntimeSettings.LogFunction.IsValid())
 		{
-			(*RuntimeSettings.LogFunction)(InSeverity, this, InMessage);
+			(*RuntimeSettings.LogFunction)(InLogSettings, this, InMessage);
 		}
 		else
 		{
-			if(InSeverity == EMessageSeverity::Error)
+			if(InLogSettings.Severity == EMessageSeverity::Error)
 			{
 				UE_LOG(LogRigVM, Error, TEXT("Instruction %d: %s"), InstructionIndex, *InMessage);
 			}
-			else if(InSeverity == EMessageSeverity::Warning)
+			else if(InLogSettings.Severity == EMessageSeverity::Warning)
 			{
 				UE_LOG(LogRigVM, Warning, TEXT("Instruction %d: %s"), InstructionIndex, *InMessage);
 			}
@@ -292,9 +292,9 @@ struct FRigVMExecuteContext
 	}
 
 	template <typename FmtType, typename... Types>
-	void Logf(EMessageSeverity::Type InSeverity, const FmtType& Fmt, Types... Args) const
+	void Logf(const FRigVMLogSettings& InLogSettings, const FmtType& Fmt, Types... Args) const
 	{
-		Log(InSeverity, FString::Printf(Fmt, Args...));
+		Log(InLogSettings, FString::Printf(Fmt, Args...));
 	}
 
 	uint16 GetInstructionIndex() const { return InstructionIndex; }
diff --git a/Engine/Plugins/Runtime/RigVM/Source/RigVM/Public/RigVMCore/RigVMLog.h b/Engine/Plugins/Runtime/RigVM/Source/RigVM/Public/RigVMCore/RigVMLog.h
--- a/Engine/Plugins/Runtime/RigVM/Source/RigVM/Public/RigVMCore/RigVMLog.h
+++ b/Engine/Plugins/Runtime/RigVM/Source/RigVM/Public/RigVMCore/RigVMLog.h
@@ -5,6 +5,17 @@
 #include "CoreMinimal.h"
 #include "Logging/TokenizedMessage.h"
 
+struct RIGVM_API FRigVMLogSettings
+{
+	FRigVMLogSettings(EMessageSeverity::Type InSeverity, bool InLogOnce = true)
+		: Severity(InSeverity)
+		, bLogOnce(InLogOnce)
+	{}
+
+	EMessageSeverity::Type Severity;
+	bool bLogOnce;
+};
+
 struct RIGVM_API FRigVMLog
 {
 public:
@@ -34,5 +45,9 @@ public:
 #endif
 
 	virtual void Reset();
-	virtual void Report(EMessageSeverity::Type InSeverity, const FName& InFunctionName, int32 InInstructionIndex, const FString& InMessage);
+	virtual void Report(const FRigVMLogSettings& InLogSettings, const FName& InFunctionName, int32 InInstructionIndex, const FString& InMessage);
+
+#if WITH_EDITOR
+	void RemoveRedundantEntries();
+#endif
 };
diff --git a/Engine/Plugins/Runtime/RigVM/Source/RigVMEditor/Private/Editor/RigVMEditor.cpp b/Engine/Plugins/Runtime/RigVM/Source/RigVMEditor/Private/Editor/RigVMEditor.cpp
--- a/Engine/Plugins/Runtime/RigVM/Source/RigVMEditor/Private/Editor/RigVMEditor.cpp
+++ b/Engine/Plugins/Runtime/RigVM/Source/RigVMEditor/Private/Editor/RigVMEditor.cpp
@@ -3406,6 +3406,8 @@ void FRigVMEditor::UpdateGraphCompilerErrors()
 				Blueprint->Status = BS_Error;
 				(void)Blueprint->MarkPackageDirty();
 			}
+
+			RigVMLog.RemoveRedundantEntries();
 		}
 	}
 
