RigVM: Optimize argument refresh for type category changes

diff --git a/Engine/Plugins/Runtime/RigVM/Source/RigVM/Private/RigVMCore/RigVMDispatchFactory.cpp b/Engine/Plugins/Runtime/RigVM/Source/RigVM/Private/RigVMCore/RigVMDispatchFactory.cpp
--- a/Engine/Plugins/Runtime/RigVM/Source/RigVM/Private/RigVMCore/RigVMDispatchFactory.cpp
+++ b/Engine/Plugins/Runtime/RigVM/Source/RigVM/Private/RigVMCore/RigVMDispatchFactory.cpp
@@ -304,7 +304,8 @@ TArray<FRigVMTemplateArgumentInfo> FRigVMDispatchFactory::BuildArgumentListFromP
 	{
 		if (bFoundArg)
 		{
-			const TArray<FRigVMTemplateTypeMap> Permutations = GetPermutationsFromArgumentType(InPrimaryArgumentName, Type);
+			TArray<FRigVMTemplateTypeMap, TInlineAllocator<1>> Permutations;
+			GetPermutationsFromArgumentType(InPrimaryArgumentName, Type, Permutations);
 			for (const FRigVMTemplateTypeMap& Permutation : Permutations)
 			{
 				for (int32 Index=0; Index < InInfos.Num(); ++Index)
diff --git a/Engine/Plugins/Runtime/RigVM/Source/RigVM/Private/RigVMCore/RigVMRegistry.cpp b/Engine/Plugins/Runtime/RigVM/Source/RigVM/Private/RigVMCore/RigVMRegistry.cpp
--- a/Engine/Plugins/Runtime/RigVM/Source/RigVM/Private/RigVMCore/RigVMRegistry.cpp
+++ b/Engine/Plugins/Runtime/RigVM/Source/RigVM/Private/RigVMCore/RigVMRegistry.cpp
@@ -224,7 +224,7 @@ void FRigVMRegistry::InitializeIfNeeded()
 	Types.Reserve(512);
 	TypeToIndex.Reserve(512);
 	TypesPerCategory.Reserve(19);
-	ArgumentsPerCategory.Reserve(19);
+	TemplatesPerCategory.Reserve(19);
 	
 	TypesPerCategory.Add(FRigVMTemplateArgument::ETypeCategory_Execute, TArray<TRigVMTypeIndex>()).Reserve(8);
 	TypesPerCategory.Add(FRigVMTemplateArgument::ETypeCategory_SingleAnyValue, TArray<TRigVMTypeIndex>()).Reserve(256);
@@ -246,25 +246,25 @@ void FRigVMRegistry::InitializeIfNeeded()
 	TypesPerCategory.Add(FRigVMTemplateArgument::ETypeCategory_ArrayObjectValue, TArray<TRigVMTypeIndex>()).Reserve(128);
 	TypesPerCategory.Add(FRigVMTemplateArgument::ETypeCategory_ArrayArrayObjectValue, TArray<TRigVMTypeIndex>()).Reserve(128);
 
-	ArgumentsPerCategory.Add(FRigVMTemplateArgument::ETypeCategory_Execute, TArray<TPair<int32,int32>>()).Reserve(8);
-	ArgumentsPerCategory.Add(FRigVMTemplateArgument::ETypeCategory_SingleAnyValue, TArray<TPair<int32,int32>>()).Reserve(64);
-	ArgumentsPerCategory.Add(FRigVMTemplateArgument::ETypeCategory_ArrayAnyValue, TArray<TPair<int32,int32>>()).Reserve(64);
-	ArgumentsPerCategory.Add(FRigVMTemplateArgument::ETypeCategory_ArrayArrayAnyValue, TArray<TPair<int32,int32>>()).Reserve(64);
-	ArgumentsPerCategory.Add(FRigVMTemplateArgument::ETypeCategory_SingleSimpleValue, TArray<TPair<int32,int32>>()).Reserve(64);
-	ArgumentsPerCategory.Add(FRigVMTemplateArgument::ETypeCategory_ArraySimpleValue, TArray<TPair<int32,int32>>()).Reserve(64);
-	ArgumentsPerCategory.Add(FRigVMTemplateArgument::ETypeCategory_ArrayArraySimpleValue, TArray<TPair<int32,int32>>()).Reserve(64);
-	ArgumentsPerCategory.Add(FRigVMTemplateArgument::ETypeCategory_SingleMathStructValue, TArray<TPair<int32,int32>>()).Reserve(64);
-	ArgumentsPerCategory.Add(FRigVMTemplateArgument::ETypeCategory_ArrayMathStructValue, TArray<TPair<int32,int32>>()).Reserve(64);
-	ArgumentsPerCategory.Add(FRigVMTemplateArgument::ETypeCategory_ArrayArrayMathStructValue, TArray<TPair<int32,int32>>()).Reserve(64);
-	ArgumentsPerCategory.Add(FRigVMTemplateArgument::ETypeCategory_SingleScriptStructValue, TArray<TPair<int32,int32>>()).Reserve(64);
-	ArgumentsPerCategory.Add(FRigVMTemplateArgument::ETypeCategory_ArrayScriptStructValue, TArray<TPair<int32,int32>>()).Reserve(64);
-	ArgumentsPerCategory.Add(FRigVMTemplateArgument::ETypeCategory_ArrayArrayScriptStructValue, TArray<TPair<int32,int32>>()).Reserve(64);
-	ArgumentsPerCategory.Add(FRigVMTemplateArgument::ETypeCategory_SingleEnumValue, TArray<TPair<int32,int32>>()).Reserve(64);
-	ArgumentsPerCategory.Add(FRigVMTemplateArgument::ETypeCategory_ArrayEnumValue, TArray<TPair<int32,int32>>()).Reserve(64);
-	ArgumentsPerCategory.Add(FRigVMTemplateArgument::ETypeCategory_ArrayArrayEnumValue, TArray<TPair<int32,int32>>()).Reserve(64);
-	ArgumentsPerCategory.Add(FRigVMTemplateArgument::ETypeCategory_SingleObjectValue, TArray<TPair<int32,int32>>()).Reserve(64);
-	ArgumentsPerCategory.Add(FRigVMTemplateArgument::ETypeCategory_ArrayObjectValue, TArray<TPair<int32,int32>>()).Reserve(64);
-	ArgumentsPerCategory.Add(FRigVMTemplateArgument::ETypeCategory_ArrayArrayObjectValue, TArray<TPair<int32,int32>>()).Reserve(64);
+	TemplatesPerCategory.Add(FRigVMTemplateArgument::ETypeCategory_Execute, TArray<int32>()).Reserve(8);
+	TemplatesPerCategory.Add(FRigVMTemplateArgument::ETypeCategory_SingleAnyValue, TArray<int32>()).Reserve(64);
+	TemplatesPerCategory.Add(FRigVMTemplateArgument::ETypeCategory_ArrayAnyValue, TArray<int32>()).Reserve(64);
+	TemplatesPerCategory.Add(FRigVMTemplateArgument::ETypeCategory_ArrayArrayAnyValue, TArray<int32>()).Reserve(64);
+	TemplatesPerCategory.Add(FRigVMTemplateArgument::ETypeCategory_SingleSimpleValue, TArray<int32>()).Reserve(64);
+	TemplatesPerCategory.Add(FRigVMTemplateArgument::ETypeCategory_ArraySimpleValue, TArray<int32>()).Reserve(64);
+	TemplatesPerCategory.Add(FRigVMTemplateArgument::ETypeCategory_ArrayArraySimpleValue, TArray<int32>()).Reserve(64);
+	TemplatesPerCategory.Add(FRigVMTemplateArgument::ETypeCategory_SingleMathStructValue, TArray<int32>()).Reserve(64);
+	TemplatesPerCategory.Add(FRigVMTemplateArgument::ETypeCategory_ArrayMathStructValue, TArray<int32>()).Reserve(64);
+	TemplatesPerCategory.Add(FRigVMTemplateArgument::ETypeCategory_ArrayArrayMathStructValue, TArray<int32>()).Reserve(64);
+	TemplatesPerCategory.Add(FRigVMTemplateArgument::ETypeCategory_SingleScriptStructValue, TArray<int32>()).Reserve(64);
+	TemplatesPerCategory.Add(FRigVMTemplateArgument::ETypeCategory_ArrayScriptStructValue, TArray<int32>()).Reserve(64);
+	TemplatesPerCategory.Add(FRigVMTemplateArgument::ETypeCategory_ArrayArrayScriptStructValue, TArray<int32>()).Reserve(64);
+	TemplatesPerCategory.Add(FRigVMTemplateArgument::ETypeCategory_SingleEnumValue, TArray<int32>()).Reserve(64);
+	TemplatesPerCategory.Add(FRigVMTemplateArgument::ETypeCategory_ArrayEnumValue, TArray<int32>()).Reserve(64);
+	TemplatesPerCategory.Add(FRigVMTemplateArgument::ETypeCategory_ArrayArrayEnumValue, TArray<int32>()).Reserve(64);
+	TemplatesPerCategory.Add(FRigVMTemplateArgument::ETypeCategory_SingleObjectValue, TArray<int32>()).Reserve(64);
+	TemplatesPerCategory.Add(FRigVMTemplateArgument::ETypeCategory_ArrayObjectValue, TArray<int32>()).Reserve(64);
+	TemplatesPerCategory.Add(FRigVMTemplateArgument::ETypeCategory_ArrayArrayObjectValue, TArray<int32>()).Reserve(64);
 
 	RigVMTypeUtils::TypeIndex::Execute = FindOrAddType(FRigVMTemplateArgumentType(FRigVMExecuteContext::StaticStruct()));
 	RigVMTypeUtils::TypeIndex::ExecuteArray = FindOrAddType(FRigVMTemplateArgumentType(FRigVMExecuteContext::StaticStruct()).ConvertToArray());
@@ -320,6 +320,9 @@ void FRigVMRegistry::InitializeIfNeeded()
 void FRigVMRegistry::RefreshEngineTypes()
 {
 	FScopeLock RefreshTypesScopeLock(&RefreshTypesMutex);
+	TGuardValue<bool> EnableGuardRefresh(bIsRefreshingEngineTypes, true);
+
+	const int32 NumTypesBefore = Types.Num(); 
 	
 	// Register all user-defined types that the engine knows about. Enumerating over the entire object hierarchy is
 	// slow, so we do it for structs, enums and dispatch factories in one shot.
@@ -368,6 +371,26 @@ void FRigVMRegistry::RefreshEngineTypes()
 	{
 		RegisterFactory(DispatchFactoryStruct);
 	}
+
+	const int32 NumTypesNow = Types.Num();
+	if(NumTypesBefore != NumTypesNow)
+	{
+		// update all of the templates once
+		TArray<bool> TemplateProcessed;
+		TemplateProcessed.AddZeroed(Templates.Num());
+		for(const TPair<FRigVMTemplateArgument::ETypeCategory, TArray<int32>>& Pair : TemplatesPerCategory)
+		{
+			for(const int32 TemplateIndex : Pair.Value)
+			{
+				if(!TemplateProcessed[TemplateIndex])
+				{
+					FRigVMTemplate& Template = Templates[TemplateIndex];
+					(void)Template.UpdateArgumentTypes();
+					TemplateProcessed[TemplateIndex] = true;
+				}
+			}
+		}
+	}
 }
 
 void FRigVMRegistry::OnAssetRenamed(const FAssetData& InAssetData, const FString& InOldObjectPath)
@@ -502,7 +525,7 @@ TRigVMTypeIndex FRigVMRegistry::FindOrAddType_Internal(const FRigVMTemplateArgum
 			Types[Indices[1]].ArrayTypeIndex = Indices[2];
 		}
 
-		// update the categories first then propagate to ArgumentsPerCategory once all categories up to date
+		// update the categories first then propagate to TemplatesPerCategory once all categories up to date
 		TArray<TPair<FRigVMTemplateArgument::ETypeCategory, int32>> ToPropagate;
 		auto RegisterNewType = [&](FRigVMTemplateArgument::ETypeCategory InCategory, int32 NewIndex)
 		{
@@ -711,16 +734,20 @@ void FRigVMRegistry::RegisterTypeInCategory(FRigVMTemplateArgument::ETypeCategor
 
 void FRigVMRegistry::PropagateTypeAddedToCategory(const FRigVMTemplateArgument::ETypeCategory InCategory, const TRigVMTypeIndex InTypeIndex)
 {
+	if(bIsRefreshingEngineTypes)
+	{
+		return;
+	}
+	
 	check(InCategory != FRigVMTemplateArgument::ETypeCategory_Invalid);
 	if ( ensure(TypesPerCategory.FindChecked(InCategory).Contains(InTypeIndex)) )
 	{
 		// when adding a new type - we need to update template arguments which expect to have access to that type 
-		const TArray<TPair<int32,int32>>& ArgumentsToUseType = ArgumentsPerCategory.FindChecked(InCategory);
-		for(const TPair<int32,int32>& Pair : ArgumentsToUseType)
+		const TArray<int32>& TemplatesToUseType = TemplatesPerCategory.FindChecked(InCategory);
+		for(const int32 TemplateIndex : TemplatesToUseType)
 		{
-			FRigVMTemplate& Template = Templates[Pair.Key];
-			const FRigVMTemplateArgument* Argument = Template.GetArgument(Pair.Value);
-			Template.AddTypeForArgument(Argument->GetName(), InTypeIndex);
+			FRigVMTemplate& Template = Templates[TemplateIndex];
+			(void)Template.UpdateArgumentTypes();
 		}
 	}
 }
@@ -831,15 +858,7 @@ void FRigVMRegistry::RemoveTypeInCategory(FRigVMTemplateArgument::ETypeCategory
 
 	TypesPerCategory.FindChecked(InCategory).Remove(InTypeIndex);
 
-	const TArray<TPair<int32,int32>>& ArgumentsToUseType = ArgumentsPerCategory.FindChecked(InCategory);
-
-	TSet<int32> TemplatesToUseType;
-	
-	for(const TPair<int32,int32>& Pair : ArgumentsToUseType)
-	{
-		TemplatesToUseType.Add(Pair.Key);
-	}
-	
+	const TArray<int32>& TemplatesToUseType = TemplatesPerCategory.FindChecked(InCategory);
 	for (const int32 TemplateIndex : TemplatesToUseType)
 	{
 		FRigVMTemplate& Template = Templates[TemplateIndex];
@@ -1715,7 +1734,7 @@ const FRigVMTemplate* FRigVMRegistry::AddTemplateFromArguments_NoLock(const FNam
 	{
 		for(const FRigVMTemplateArgument::ETypeCategory& ArgumentTypeCategory : Templates[Index].Arguments[ArgumentIndex].TypeCategories)
 		{
-			ArgumentsPerCategory.FindChecked(ArgumentTypeCategory).AddUnique(TPair<int32, int32>(Index, ArgumentIndex));
+			TemplatesPerCategory.FindChecked(ArgumentTypeCategory).AddUnique(Index);
 		}
 	}
 	
diff --git a/Engine/Plugins/Runtime/RigVM/Source/RigVM/Private/RigVMCore/RigVMTemplate.cpp b/Engine/Plugins/Runtime/RigVM/Source/RigVM/Private/RigVMCore/RigVMTemplate.cpp
--- a/Engine/Plugins/Runtime/RigVM/Source/RigVM/Private/RigVMCore/RigVMTemplate.cpp
+++ b/Engine/Plugins/Runtime/RigVM/Source/RigVM/Private/RigVMCore/RigVMTemplate.cpp
@@ -162,17 +162,13 @@ void FRigVMTemplateArgument::EnsureValidExecuteType()
 void FRigVMTemplateArgument::UpdateTypeToPermutations()
 {
 	TypeToPermutations.Reset();
-	for(int32 TypeIndex=0;TypeIndex<GetNumTypes();TypeIndex++)
+	TypeToPermutations.Reserve(GetNumTypes());
+
+	int32 TypeIndex = 0;
+	ForEachType([&](const TRigVMTypeIndex Type)
 	{
-		if (TArray<int32>* Permutations = TypeToPermutations.Find(GetTypeIndex(TypeIndex)))
-		{
-			Permutations->Add(TypeIndex);
-		}
-		else
-		{
-			TypeToPermutations.Add(GetTypeIndex(TypeIndex), {TypeIndex});
-		}
-	}
+	   TypeToPermutations.FindOrAdd(Type).Add(TypeIndex++);
+	});
 }
 
 bool FRigVMTemplateArgument::SupportsTypeIndex(TRigVMTypeIndex InTypeIndex, TRigVMTypeIndex* OutTypeIndex) const
@@ -1812,86 +1808,127 @@ FString FRigVMTemplate::GetKeywords() const
 
 #endif
 
-bool FRigVMTemplate::AddTypeForArgument(const FName& InArgumentName, TRigVMTypeIndex InTypeIndex)
+bool FRigVMTemplate::UpdateArgumentTypes()
 {
-	InvalidateHash();
+	const int32 PrimaryArgumentIndex = Arguments.IndexOfByPredicate([](const FRigVMTemplateArgument& Argument)
+	{
+		return Argument.bUseCategories;
+	});
 
-	TArray<FRigVMTemplateTypeMap> TypesArray;
+	// this template may not be affected at all by this
+	if(PrimaryArgumentIndex == INDEX_NONE)
+	{
+		return true;
+	}
 
-	if(Delegates.GetDispatchFactoryDelegate.IsBound())
+	InvalidateHash();
+
+	for(int32 ArgumentIndex = 0; ArgumentIndex < Arguments.Num(); ArgumentIndex++)
 	{
-		const FRigVMDispatchFactory* Factory = Delegates.GetDispatchFactoryDelegate.Execute();
-		if(ensure(Factory))
+		FRigVMTemplateArgument& Argument = Arguments[ArgumentIndex];
+		if(Argument.bUseCategories || Argument.IsSingleton())
 		{
-			TypesArray = Factory->GetPermutationsFromArgumentType(InArgumentName, InTypeIndex);
+			continue;
 		}
+
+		Argument.TypeIndices.Reset();
+		Argument.TypeToPermutations.Reset();
 	}
-	else if(OnNewArgumentType().IsBound())
+
+	const FRigVMRegistry& Registry = FRigVMRegistry::Get();
+	const FRigVMTemplateArgument& PrimaryArgument = Arguments[PrimaryArgumentIndex];
+	TArray<FRigVMTemplateTypeMap, TInlineAllocator<1>> TypesArray;
+	bool bResult = true;
+
+	const FRigVMDispatchFactory* Factory = nullptr;
+	if(Delegates.GetDispatchFactoryDelegate.IsBound())
 	{
-		FRigVMTemplateTypeMap Types = OnNewArgumentType().Execute(InArgumentName, InTypeIndex);
-		TypesArray = {Types};
+		Factory = Delegates.GetDispatchFactoryDelegate.Execute();
+		ensure(Factory);
 	}
 
-	if (!TypesArray.IsEmpty())
+	PrimaryArgument.ForEachType([&](const TRigVMTypeIndex PrimaryTypeIndex)
 	{
-		for (FRigVMTemplateTypeMap& Types : TypesArray)
+		TypesArray.Reset();
+		if(Factory)
+		{
+			Factory->GetPermutationsFromArgumentType(PrimaryArgument.Name, PrimaryTypeIndex, TypesArray);
+		}
+		else if(OnNewArgumentType().IsBound())
 		{
-			if(Types.Num() == Arguments.Num())
+			FRigVMTemplateTypeMap Types = OnNewArgumentType().Execute(PrimaryArgument.Name, PrimaryTypeIndex);
+			TypesArray.Add(Types);
+		}
+
+		if (!TypesArray.IsEmpty())
+		{
+			for (FRigVMTemplateTypeMap& Types : TypesArray)
 			{
-				const FRigVMRegistry& Registry = FRigVMRegistry::Get();
-				for (TPair<FName, TRigVMTypeIndex>& ArgumentAndType : Types)
-				{
-					// similar to FRigVMTemplateArgument::EnsureValidExecuteType
-					Registry.ConvertExecuteContextToBaseType(ArgumentAndType.Value);
-				}
-			
-				for(FRigVMTemplateArgument& Argument : Arguments)
+				if(Types.Num() == Arguments.Num())
 				{
-					const TRigVMTypeIndex* TypeIndex = Types.Find(Argument.Name);
-					if(TypeIndex == nullptr)
+					for (TPair<FName, TRigVMTypeIndex>& ArgumentAndType : Types)
 					{
-						return false;
+						// similar to FRigVMTemplateArgument::EnsureValidExecuteType
+						Registry.ConvertExecuteContextToBaseType(ArgumentAndType.Value);
 					}
-					if(*TypeIndex == INDEX_NONE)
+
+					// Find if these types were already registered
+					if (ContainsPermutation(Types))
 					{
-						return false;
+						return;
 					}
-				}
-
-				// Find if these types were already registered
-				FRigVMTemplateTypeMap TestTypes = Types;
-				if (ContainsPermutation(TestTypes))
-				{
-					return false;
-				}
 
-				uint32 TypeHash=0;
-				for(FRigVMTemplateArgument& Argument : Arguments)
-				{
-					const TRigVMTypeIndex TypeIndex = Types.FindChecked(Argument.Name);
-					if (!Argument.bUseCategories)
+					uint32 TypeHash=0;
+					for(int32 ArgumentIndex = 0; ArgumentIndex < Arguments.Num(); ArgumentIndex++)
 					{
+						FRigVMTemplateArgument& Argument = Arguments[ArgumentIndex];
+
+						const TRigVMTypeIndex* TypeIndexPtr = Types.Find(Argument.Name);
+						if(TypeIndexPtr == nullptr)
+						{
+							bResult = false;
+							return;
+						}
+						if(*TypeIndexPtr == INDEX_NONE)
+						{
+							bResult = false;
+							return;
+						}
+
+						const TRigVMTypeIndex& TypeIndex = *TypeIndexPtr;
+						TypeHash = HashCombine(TypeHash, GetTypeHash(TypeIndex));
+
+						if(Argument.bUseCategories || Argument.IsSingleton())
+						{
+							continue;
+						}
+					
 						Argument.TypeIndices.Add(TypeIndex);
 						Argument.TypeToPermutations.FindOrAdd(TypeIndex).Add(Permutations.Num());
 					}
-					else
-					{
-						Argument.UpdateTypeToPermutations();
-					}
-					TypeHash = HashCombine(TypeHash, GetTypeHash(TypeIndex));
+
+					Permutations.Add(INDEX_NONE);
+					TypesHashToPermutation.Add(TypeHash, Permutations.Num()-1);
+				}
+				else
+				{
+					bResult = false;
 				}
-				
-				Permutations.Add(INDEX_NONE);
-				TypesHashToPermutation.Add(TypeHash, Permutations.Num()-1);
-			}
-			else
-			{
-				return false;
 			}
 		}
-		return true;
+		else
+		{
+			bResult = false;
+		}
+	});
+
+	for(int32 ArgumentIndex = 0; ArgumentIndex < Arguments.Num(); ArgumentIndex++)
+	{
+		FRigVMTemplateArgument& Argument = Arguments[ArgumentIndex];
+		Argument.UpdateTypeToPermutations();
 	}
-	return false;
+
+	return bResult;
 }
 
 void FRigVMTemplate::HandleTypeRemoval(TRigVMTypeIndex InTypeIndex)
diff --git a/Engine/Plugins/Runtime/RigVM/Source/RigVM/Private/RigVMFunctions/RigVMDispatch_CastObject.cpp b/Engine/Plugins/Runtime/RigVM/Source/RigVM/Private/RigVMFunctions/RigVMDispatch_CastObject.cpp
--- a/Engine/Plugins/Runtime/RigVM/Source/RigVM/Private/RigVMFunctions/RigVMDispatch_CastObject.cpp
+++ b/Engine/Plugins/Runtime/RigVM/Source/RigVM/Private/RigVMFunctions/RigVMDispatch_CastObject.cpp
@@ -28,15 +28,14 @@ const TArray<FRigVMTemplateArgumentInfo>& FRigVMDispatch_CastObject::GetArgument
 	return OutInfos;
 }
 
-TArray<FRigVMTemplateTypeMap> FRigVMDispatch_CastObject::GetPermutationsFromArgumentType(const FName& InArgumentName, const TRigVMTypeIndex& InTypeIndex) const
+bool FRigVMDispatch_CastObject::GetPermutationsFromArgumentType(const FName& InArgumentName, const TRigVMTypeIndex& InTypeIndex, TArray<FRigVMTemplateTypeMap, TInlineAllocator<1>>& OutPermutations) const
 {
-	TArray<FRigVMTemplateTypeMap> Permutations;
 	if (InArgumentName == ValueName)
 	{
 		const TArray<TRigVMTypeIndex>& ObjectTypes = FRigVMRegistry::Get().GetTypesForCategory(FRigVMTemplateArgument::ETypeCategory_SingleObjectValue);
 		for (const TRigVMTypeIndex& Type : ObjectTypes)
 		{
-			Permutations.Add(
+			OutPermutations.Add(
 	{
 				{ ValueName, InTypeIndex },
 				{ ResultName, Type }
@@ -48,14 +47,14 @@ TArray<FRigVMTemplateTypeMap> FRigVMDispatch_CastObject::GetPermutationsFromArgu
 		const TArray<TRigVMTypeIndex>& ObjectTypes = FRigVMRegistry::Get().GetTypesForCategory(FRigVMTemplateArgument::ETypeCategory_SingleObjectValue);
 		for (const TRigVMTypeIndex& Type : ObjectTypes)
 		{
-			Permutations.Add(
+			OutPermutations.Add(
 	{
 				{ ValueName, Type },
 				{ ResultName, InTypeIndex }
 			});
 		}
 	}
-	return Permutations;
+	return !OutPermutations.IsEmpty();
 }
 
 #if WITH_EDITOR
diff --git a/Engine/Plugins/Runtime/RigVM/Source/RigVM/Public/RigVMCore/RigVMDispatchFactory.h b/Engine/Plugins/Runtime/RigVM/Source/RigVM/Public/RigVMCore/RigVMDispatchFactory.h
--- a/Engine/Plugins/Runtime/RigVM/Source/RigVM/Public/RigVMCore/RigVMDispatchFactory.h
+++ b/Engine/Plugins/Runtime/RigVM/Source/RigVM/Public/RigVMCore/RigVMDispatchFactory.h
@@ -152,14 +152,15 @@ public:
 	virtual FRigVMTemplateTypeMap OnNewArgumentType(const FName& InArgumentName, TRigVMTypeIndex InTypeIndex) const { return FRigVMTemplateTypeMap(); }
 
 	// returns the new permutations argument types after a new type is defined for one argument
-	virtual TArray<FRigVMTemplateTypeMap> GetPermutationsFromArgumentType(const FName& InArgumentName, const TRigVMTypeIndex& InTypeIndex) const
+	virtual bool GetPermutationsFromArgumentType(const FName& InArgumentName, const TRigVMTypeIndex& InTypeIndex, TArray<FRigVMTemplateTypeMap, TInlineAllocator<1>>& OutPermutations) const
 	{
 		FRigVMTemplateTypeMap Permutation = OnNewArgumentType(InArgumentName, InTypeIndex);
 		if (!Permutation.IsEmpty())
 		{
-			return {Permutation};
+			OutPermutations.Add(Permutation);
+			return true;
 		}
-		return {};
+		return false;
 	}
 
 	// returns the upgrade info to use for this factory
diff --git a/Engine/Plugins/Runtime/RigVM/Source/RigVM/Public/RigVMCore/RigVMRegistry.h b/Engine/Plugins/Runtime/RigVM/Source/RigVM/Public/RigVMCore/RigVMRegistry.h
--- a/Engine/Plugins/Runtime/RigVM/Source/RigVM/Public/RigVMCore/RigVMRegistry.h
+++ b/Engine/Plugins/Runtime/RigVM/Source/RigVM/Public/RigVMCore/RigVMRegistry.h
@@ -296,7 +296,11 @@ private:
 	static const FName TemplateNameMetaName;
 
 	// disable default constructor
-	FRigVMRegistry() {}
+	FRigVMRegistry()
+		: bIsRefreshingEngineTypes(false)
+	{
+	}
+	
 	// disable copy constructor
 	FRigVMRegistry(const FRigVMRegistry&) = delete;
 	// disable assignment operator
@@ -384,8 +388,8 @@ private:
 	// Maps storing the default types per type category
 	TMap<FRigVMTemplateArgument::ETypeCategory, TArray<TRigVMTypeIndex>> TypesPerCategory;
 
-	// Lookup per type category to know which argument to keep in sync
-	TMap<FRigVMTemplateArgument::ETypeCategory, TArray<TPair<int32,int32>>> ArgumentsPerCategory;
+	// Lookup per type category to know which template to keep in sync
+	TMap<FRigVMTemplateArgument::ETypeCategory, TArray<int32>> TemplatesPerCategory;
 
 	// Name loop up for user defined types since they can be deleted.
 	// When that happens, it won't be safe to reload deleted assets so only type names are reliable
@@ -397,6 +401,9 @@ private:
 	// Notifies other system that types have been added/removed, and template permutations have been updated
 	FOnRigVMRegistryChanged OnRigVMRegistryChangedDelegate;
 
+	// If this is true the registry is currently refreshing all types
+	bool bIsRefreshingEngineTypes;
+
 	static FCriticalSection RefreshTypesMutex;
 	static FCriticalSection RegisterFunctionMutex;
 	static FCriticalSection RegisterTemplateMutex;
diff --git a/Engine/Plugins/Runtime/RigVM/Source/RigVM/Public/RigVMCore/RigVMTemplate.h b/Engine/Plugins/Runtime/RigVM/Source/RigVM/Public/RigVMCore/RigVMTemplate.h
--- a/Engine/Plugins/Runtime/RigVM/Source/RigVM/Public/RigVMCore/RigVMTemplate.h
+++ b/Engine/Plugins/Runtime/RigVM/Source/RigVM/Public/RigVMCore/RigVMTemplate.h
@@ -474,9 +474,9 @@ public:
 
 #endif
 
-	// Adds a new argument to the template. The template needs to rely on the delegate to ask the
-	// template factory how to deal with the new type.
-	bool AddTypeForArgument(const FName& InArgumentName, TRigVMTypeIndex InTypeIndex);
+	// Updates the template's argument types. This only affects templates which have category based
+	// arguments and will resolve the other arguments to the expected types.
+	bool UpdateArgumentTypes();
 
 	// Invalidates template permutations whenever a type such as a user defined struct is removed
 	void HandleTypeRemoval(TRigVMTypeIndex InTypeIndex);
diff --git a/Engine/Plugins/Runtime/RigVM/Source/RigVM/Public/RigVMFunctions/RigVMDispatch_CastObject.h b/Engine/Plugins/Runtime/RigVM/Source/RigVM/Public/RigVMFunctions/RigVMDispatch_CastObject.h
--- a/Engine/Plugins/Runtime/RigVM/Source/RigVM/Public/RigVMFunctions/RigVMDispatch_CastObject.h
+++ b/Engine/Plugins/Runtime/RigVM/Source/RigVM/Public/RigVMFunctions/RigVMDispatch_CastObject.h
@@ -17,7 +17,7 @@ public:
 	}
 	
 	virtual const TArray<FRigVMTemplateArgumentInfo>& GetArgumentInfos() const override;
-	virtual TArray<FRigVMTemplateTypeMap> GetPermutationsFromArgumentType(const FName& InArgumentName, const TRigVMTypeIndex& InTypeIndex) const override;
+	virtual bool GetPermutationsFromArgumentType(const FName& InArgumentName, const TRigVMTypeIndex& InTypeIndex, TArray<FRigVMTemplateTypeMap, TInlineAllocator<1>>& OutPermutations) const override;
 #if WITH_EDITOR
 	virtual FString GetNodeTitle(const FRigVMTemplateTypeMap& InTypes) const override;
 	virtual FText GetNodeTooltip(const FRigVMTemplateTypeMap& InTypes) const override;
