RigVM: Remove invocation to ResolveAllPinReferences

#rnx

diff --git a/Engine/Plugins/Runtime/RigVM/Source/RigVMDeveloper/Private/RigVMBlueprint.cpp b/Engine/Plugins/Runtime/RigVM/Source/RigVMDeveloper/Private/RigVMBlueprint.cpp
--- a/Engine/Plugins/Runtime/RigVM/Source/RigVMDeveloper/Private/RigVMBlueprint.cpp
+++ b/Engine/Plugins/Runtime/RigVM/Source/RigVMDeveloper/Private/RigVMBlueprint.cpp
@@ -1084,8 +1084,6 @@ void URigVMBlueprint::PostLoad()
 	// in which case we have to refresh the model
 	FRigVMRegistry::Get().OnRigVMRegistryChanged().RemoveAll(this);
 	FRigVMRegistry::Get().OnRigVMRegistryChanged().AddUObject(this, &URigVMBlueprint::OnRigVMRegistryChanged);
-	
-	UEdGraphPin::ResolveAllPinReferences();
 }
 
 #if WITH_EDITORONLY_DATA
