RigVM: Introduce the notion of lazy blocks

#rnx

diff --git a/Engine/Plugins/Runtime/RigVM/Source/RigVM/Private/RigVMCore/RigVM.cpp b/Engine/Plugins/Runtime/RigVM/Source/RigVM/Private/RigVMCore/RigVM.cpp
--- a/Engine/Plugins/Runtime/RigVM/Source/RigVM/Private/RigVMCore/RigVM.cpp
+++ b/Engine/Plugins/Runtime/RigVM/Source/RigVM/Private/RigVMCore/RigVM.cpp
@@ -545,6 +545,12 @@ bool URigVM::ValidateBytecode()
 				}
 				break;
 			}
+			case ERigVMOpCode::RunInstructions:
+			{
+				const FRigVMRunInstructionsOp& Op = ByteCodeStorage.GetOpAt<FRigVMRunInstructionsOp>(ByteCodeInstruction);
+				CheckOperandValidity(Op.Arg);
+				break;
+			}
 			case ERigVMOpCode::Invalid:
 			{
 				ensure(false);
@@ -1130,6 +1136,13 @@ void URigVM::InstructionOpEval(FRigVMExtendedExecuteContext& Context, int32 Inst
 			InOpFunc(Context, InHandleBaseIndex, {}, Arg);
 			break;
 		}
+		case ERigVMOpCode::RunInstructions:
+		{
+			const FRigVMRunInstructionsOp& Op = ByteCode.GetOpAt<FRigVMRunInstructionsOp>(Instructions[InstructionIndex]);
+			const FRigVMOperand& Arg = Op.Arg;
+			InOpFunc(Context, InHandleBaseIndex, {}, Arg);
+			break;
+		}
 		case ERigVMOpCode::Invalid:
 		default:
 		{
@@ -2067,6 +2080,32 @@ ERigVMExecuteResult URigVM::ExecuteInstructions(FRigVMExtendedExecuteContext& Co
 				}
 				break;
 			}
+			case ERigVMOpCode::RunInstructions:
+			{
+				const FRigVMRunInstructionsOp& Op = ByteCode.GetOpAt<FRigVMRunInstructionsOp>(Instruction);
+				FRigVMInstructionSetExecuteState& ExecutionState = *(FRigVMInstructionSetExecuteState*)Context.CachedMemoryHandles[FirstHandleForInstruction[ContextPublicData.InstructionIndex]].GetData();
+
+				if((Op.StartInstruction != INDEX_NONE) &&
+					(Op.EndInstruction != INDEX_NONE) &&
+					(Op.EndInstruction >= Op.StartInstruction))
+				{
+					const int32 SliceIndex = FMath::Max(0, Context.GetSlice().GetIndex());
+					while(!ExecutionState.HashPerSlice.IsValidIndex(SliceIndex))
+					{
+						ExecutionState.HashPerSlice.Add(UINT32_MAX);
+					}
+						
+					const uint32 Hash = GetTypeHash(ContextPublicData.GetNumExecutions());
+					if(ExecutionState.HashPerSlice[SliceIndex] != Hash)
+					{
+						ExecuteInstructions(Context, Op.StartInstruction, Op.EndInstruction);
+						ExecutionState.HashPerSlice[SliceIndex] = Hash;
+					}
+				}
+
+				ContextPublicData.InstructionIndex++;
+				break;
+			}
 			case ERigVMOpCode::Invalid:
 			{
 				ensure(false);
@@ -2392,6 +2431,12 @@ TArray<FString> URigVM::DumpByteCodeAsTextArray(FRigVMExtendedExecuteContext& Co
 				ResultLine = FString::Printf(TEXT("Jump To Branch %s"), *GetOperandLabel(Context, Op.Arg, OperandFormatFunction));
 				break;
 			}
+			case ERigVMOpCode::RunInstructions:
+			{
+				const FRigVMRunInstructionsOp& Op = ByteCode.GetOpAt<FRigVMRunInstructionsOp>(Instructions[InstructionIndex]);
+				ResultLine = FString::Printf(TEXT("Run Instructions %d-%d (%s)"), Op.StartInstruction, Op.EndInstruction, *GetOperandLabel(Context, Op.Arg, OperandFormatFunction));
+				break;
+			}
 			default:
 			{
 				ensure(false);
@@ -2447,13 +2492,28 @@ FString URigVM::GetOperandLabel(FRigVMExtendedExecuteContext& Context, const FRi
 			return FString();
 		}
 
-		check(Memory->IsValidIndex(InOperand.GetRegisterIndex()));
-		
-		RegisterName = Memory->GetProperties()[InOperand.GetRegisterIndex()]->GetName();
-		RegisterOffsetName =
-			InOperand.GetRegisterOffset() != INDEX_NONE ?
-			Memory->GetPropertyPaths()[InOperand.GetRegisterOffset()].ToString() :
-			FString();
+		if(!Memory->IsValidIndex(InOperand.GetRegisterIndex()))
+		{
+			if(Memory->Num() == 0)
+			{
+				Memory = GetDefaultMemoryByType(InOperand.GetMemoryType());
+			}
+		}
+
+		if(!Memory->IsValidIndex(InOperand.GetRegisterIndex()))
+		{
+			static const UEnum* MemoryTypeEnum = StaticEnum<ERigVMMemoryType>();
+			static constexpr TCHAR Format[] = TEXT("%s_%d");
+			RegisterName = FString::Printf(Format, *MemoryTypeEnum->GetDisplayNameTextByValue((int64)InOperand.GetMemoryType()).ToString(), InOperand.GetRegisterIndex());
+		}
+		else
+		{
+			RegisterName = Memory->GetProperties()[InOperand.GetRegisterIndex()]->GetName();
+			RegisterOffsetName =
+				InOperand.GetRegisterOffset() != INDEX_NONE ?
+				Memory->GetPropertyPaths()[InOperand.GetRegisterOffset()].ToString() :
+				FString();
+		}
 	}
 	
 	FString OperandLabel = RegisterName;
diff --git a/Engine/Plugins/Runtime/RigVM/Source/RigVM/Private/RigVMCore/RigVMByteCode.cpp b/Engine/Plugins/Runtime/RigVM/Source/RigVM/Private/RigVMCore/RigVMByteCode.cpp
--- a/Engine/Plugins/Runtime/RigVM/Source/RigVM/Private/RigVMCore/RigVMByteCode.cpp
+++ b/Engine/Plugins/Runtime/RigVM/Source/RigVM/Private/RigVMCore/RigVMByteCode.cpp
@@ -178,6 +178,14 @@ void FRigVMJumpToBranchOp::Serialize(FArchive& Ar)
 	Ar << FirstBranchInfoIndex;
 }
 
+void FRigVMRunInstructionsOp::Serialize(FArchive& Ar)
+{
+	Ar << OpCode;
+	Ar << Arg;
+	Ar << StartInstruction;
+	Ar << EndInstruction;
+}
+
 FRigVMInstructionArray::FRigVMInstructionArray()
 {
 }
@@ -417,6 +425,12 @@ void FRigVMByteCode::Save(FArchive& Ar) const
 				Ar << Op;
 				break;
 			}
+			case ERigVMOpCode::RunInstructions:
+			{
+				FRigVMRunInstructionsOp Op = GetOpAt<FRigVMRunInstructionsOp>(Instruction.ByteCodeIndex);
+				Ar << Op;
+				break;
+			}
 			default:
 			{
 				ensure(false);
@@ -600,6 +614,13 @@ void FRigVMByteCode::Load(FArchive& Ar)
 				AddOp<FRigVMJumpToBranchOp>(Op);
 				break;
 			}
+			case ERigVMOpCode::RunInstructions:
+			{
+				FRigVMRunInstructionsOp Op;
+				Ar << Op;
+				AddOp<FRigVMRunInstructionsOp>(Op);
+				break;
+			}
 			default:
 			{
 				ensure(false);
@@ -801,6 +822,10 @@ uint32 FRigVMByteCode::GetOperatorHash(const FRigVMInstruction& InInstruction) c
 		{
 			return GetTypeHash(GetOpAt<FRigVMJumpToBranchOp>(InInstruction));
 		}
+		case ERigVMOpCode::RunInstructions:
+		{
+			return GetTypeHash(GetOpAt<FRigVMRunInstructionsOp>(InInstruction));
+		}
 		case ERigVMOpCode::Invalid:
 		{
 			ensure(false);
@@ -948,6 +973,10 @@ uint64 FRigVMByteCode::GetOpNumBytesAt(uint64 InByteCodeIndex, bool bIncludeOper
 		{
 			return (uint64)sizeof(FRigVMJumpToBranchOp);
 		}
+		case ERigVMOpCode::RunInstructions:
+		{
+			return (uint64)sizeof(FRigVMRunInstructionsOp);
+		}
 		case ERigVMOpCode::Invalid:
 		{
 			ensure(false);
@@ -1226,6 +1255,17 @@ FString FRigVMByteCode::DumpToText() const
 				Line += FString::Join(BranchInfoTexts, TEXT(", "));
 				break;
 			}
+			case ERigVMOpCode::RunInstructions:
+			{
+				const FRigVMRunInstructionsOp& Op = GetOpAt<FRigVMRunInstructionsOp>(Instruction.ByteCodeIndex);
+				FString Arg;
+				FRigVMOperand::StaticStruct()->ExportText(Arg, &Op.Arg, &Op.Arg, nullptr, PPF_None, nullptr);
+				Line += TEXT(" Instructions ");
+				Line += FString::FromInt(Op.StartInstruction);
+				Line += TEXT("-");
+				Line += FString::FromInt(Op.EndInstruction);
+				break;
+			}
 			case ERigVMOpCode::Invalid:
 			{
 				ensure(false);
@@ -1269,6 +1309,11 @@ uint64 FRigVMByteCode::AddJumpToBranchOp(FRigVMOperand InBranchNameArg, int32 In
 	return AddOp(FRigVMJumpToBranchOp(InBranchNameArg, InFirstBranchInfoIndex));
 }
 
+uint64 FRigVMByteCode::AddRunInstructionsOp(FRigVMOperand InExecuteStateArg, int32 InStartInstruction, int32 InEndInstruction)
+{
+	return AddOp(FRigVMRunInstructionsOp(InExecuteStateArg, InStartInstruction, InEndInstruction));
+}
+
 int32 FRigVMByteCode::AddBranchInfo(const FRigVMBranchInfo& InBranchInfo)
 {
 	FRigVMBranchInfo BranchInfo = InBranchInfo;
@@ -1368,6 +1413,11 @@ FRigVMOperandArray FRigVMByteCode::GetOperandsForOp(const FRigVMInstruction& InI
 			const FRigVMJumpToBranchOp& Op = GetOpAt<FRigVMJumpToBranchOp>(InInstruction.ByteCodeIndex);
 			return FRigVMOperandArray(&Op.Arg, 1);
 		}
+		case ERigVMOpCode::RunInstructions:
+		{
+			const FRigVMRunInstructionsOp& Op = GetOpAt<FRigVMRunInstructionsOp>(InInstruction.ByteCodeIndex);
+			return FRigVMOperandArray(&Op.Arg, 1);
+		}
 		case ERigVMOpCode::JumpAbsolute:
 		case ERigVMOpCode::JumpForward:
 		case ERigVMOpCode::JumpBackward:
@@ -1461,6 +1511,11 @@ uint64 FRigVMByteCode::GetFirstOperandByteIndex(const FRigVMInstruction& InInstr
 			const FRigVMJumpToBranchOp& Op = GetOpAt<FRigVMJumpToBranchOp>(InInstruction.ByteCodeIndex);
 			return InInstruction.ByteCodeIndex + ((uint64)&Op.Arg - (uint64)&Op);
 		}
+		case ERigVMOpCode::RunInstructions:
+		{
+			const FRigVMRunInstructionsOp& Op = GetOpAt<FRigVMRunInstructionsOp>(InInstruction.ByteCodeIndex);
+			return InInstruction.ByteCodeIndex + ((uint64)&Op.Arg - (uint64)&Op);
+		}
 		case ERigVMOpCode::JumpAbsolute:
 		case ERigVMOpCode::JumpForward:
 		case ERigVMOpCode::JumpBackward:
@@ -1594,6 +1649,11 @@ uint64 FRigVMByteCode::GetOpAlignment(ERigVMOpCode InOpCode) const
 			static const uint64 Alignment = FRigVMJumpToBranchOp::StaticStruct()->GetCppStructOps()->GetAlignment();
 			return Alignment;
 		}
+		case ERigVMOpCode::RunInstructions:
+		{
+			static const uint64 Alignment = FRigVMRunInstructionsOp::StaticStruct()->GetCppStructOps()->GetAlignment();
+			return Alignment;
+		}
 		case ERigVMOpCode::Invalid:
 		{
 			ensure(false);
diff --git a/Engine/Plugins/Runtime/RigVM/Source/RigVM/Public/RigVMCore/RigVMByteCode.h b/Engine/Plugins/Runtime/RigVM/Source/RigVM/Public/RigVMCore/RigVMByteCode.h
--- a/Engine/Plugins/Runtime/RigVM/Source/RigVM/Public/RigVMCore/RigVMByteCode.h
+++ b/Engine/Plugins/Runtime/RigVM/Source/RigVM/Public/RigVMCore/RigVMByteCode.h
@@ -94,6 +94,21 @@ struct FRigVMBranchInfoKey
 	FName Label;
 };
 
+// A runtime cache for determining if a set of instruction has to
+// run for this execution of the VM
+USTRUCT(BlueprintType)
+struct RIGVM_API FRigVMInstructionSetExecuteState
+{
+	GENERATED_BODY()
+
+	FRigVMInstructionSetExecuteState()
+	{
+	}
+
+	UPROPERTY()
+	TArray<uint32> HashPerSlice;
+};
+
 // A description of a predicate branch in the VM's bytecode
 USTRUCT()
 struct RIGVM_API FRigVMPredicateBranch
@@ -220,6 +235,7 @@ enum class ERigVMOpCode : uint8
 	InvokeEntry, // invokes an entry from the entry list
 	JumpToBranch, // jumps to a branch based on a name operand
 	Execute, // single execute op (formerly Execute_0_Operands to Execute_64_Operands)
+	RunInstructions, // runs a set of instructions lazily
 	Invalid,
 	FirstArrayOpCode = ArrayReset,
 	LastArrayOpCode = ArrayReverse,
@@ -330,7 +346,8 @@ struct RIGVM_API FRigVMUnaryOp : public FRigVMBaseOp
 			uint8(InOpCode) == uint8(ERigVMOpCode::JumpForwardIf) ||
 			uint8(InOpCode) == uint8(ERigVMOpCode::JumpBackwardIf) ||
 			uint8(InOpCode) == uint8(ERigVMOpCode::ChangeType) ||
-			uint8(InOpCode) == uint8(ERigVMOpCode::JumpToBranch)
+			uint8(InOpCode) == uint8(ERigVMOpCode::JumpToBranch) ||
+			uint8(InOpCode) == uint8(ERigVMOpCode::RunInstructions)
 		);
 	}
 
@@ -971,6 +988,44 @@ struct RIGVM_API FRigVMJumpToBranchOp : public FRigVMUnaryOp
 	}
 };
 
+// runs a set of instructions lazily
+USTRUCT()
+struct RIGVM_API FRigVMRunInstructionsOp : public FRigVMUnaryOp
+{
+	GENERATED_USTRUCT_BODY()
+
+	FRigVMRunInstructionsOp()
+		: FRigVMUnaryOp()
+		, StartInstruction(INDEX_NONE)
+		, EndInstruction(INDEX_NONE)
+	{
+	}
+
+	FRigVMRunInstructionsOp(FRigVMOperand InExecutionStateArg, int32 InStartInstruction, int32 InEndInstruction)
+		: FRigVMUnaryOp(ERigVMOpCode::RunInstructions, InExecutionStateArg)
+		, StartInstruction(InStartInstruction)
+		, EndInstruction(InEndInstruction)
+	{
+	}
+
+	int32 StartInstruction;
+	int32 EndInstruction;
+
+	friend uint32 GetTypeHash(const FRigVMRunInstructionsOp& Op)
+	{
+		uint32 Hash = GetTypeHash((const FRigVMBaseOp&)Op);
+		Hash = HashCombine(Hash, GetTypeHash(Op.StartInstruction));
+		Hash = HashCombine(Hash, GetTypeHash(Op.EndInstruction));
+		return Hash;
+	}
+
+	void Serialize(FArchive& Ar);
+	friend FArchive& operator<<(FArchive& Ar, FRigVMRunInstructionsOp& P)
+	{
+		P.Serialize(Ar);
+		return Ar;
+	}
+};
 
 /**
  * The FRigVMInstruction represents
@@ -1171,6 +1226,9 @@ public:
 	// adds a jump to branch operator
 	uint64 AddJumpToBranchOp(FRigVMOperand InBranchNameArg, int32 InFirstBranchInfoIndex);
 
+	// adds a run instructions op
+	uint64 AddRunInstructionsOp(FRigVMOperand InExecuteStateArg, int32 InStartInstruction, int32 InEndInstruction);
+
 	// adds information about a branch for an instruction's argument
 	int32 AddBranchInfo(const FRigVMBranchInfo& InBranchInfo);
 	int32 AddBranchInfo(const FName& InBranchLabel, int32 InInstructionIndex, int32 InArgumentIndex, int32 InFirstBranchInstruction, int32 InLastBranchInstruction);
diff --git a/Engine/Plugins/Runtime/RigVM/Source/RigVMDeveloper/Private/RigVMCompiler/RigVMAST.cpp b/Engine/Plugins/Runtime/RigVM/Source/RigVMDeveloper/Private/RigVMCompiler/RigVMAST.cpp
--- a/Engine/Plugins/Runtime/RigVM/Source/RigVMDeveloper/Private/RigVMCompiler/RigVMAST.cpp
+++ b/Engine/Plugins/Runtime/RigVM/Source/RigVMDeveloper/Private/RigVMCompiler/RigVMAST.cpp
@@ -19,6 +19,7 @@
 #include "VisualGraphUtils.h"
 #include "RigVMModel/Nodes/RigVMDispatchNode.h"
 #include "UObject/FieldIterator.h"
+#include "Algo/Sort.h"
 
 #include UE_INLINE_GENERATED_CPP_BY_NAME(RigVMAST)
 
@@ -121,6 +122,30 @@ const FRigVMExprAST* FRigVMExprAST::GetFirstParentOfType(EType InExprType) const
 	return nullptr;
 }
 
+bool FRigVMExprAST::IsParentedTo(const FRigVMExprAST* InParentExpr) const
+{
+	check(InParentExpr);
+
+	for(const FRigVMExprAST* Parent : Parents)
+	{
+		if(Parent == InParentExpr)
+		{
+			return true;
+		}
+		if(Parent->IsParentedTo(InParentExpr))
+		{
+			return true;
+		}
+	}
+	return false;
+}
+
+bool FRigVMExprAST::IsParentOf(const FRigVMExprAST* InChildExpr) const
+{
+	check(InChildExpr);
+	return InChildExpr->IsParentedTo(this);
+}
+
 const FRigVMExprAST* FRigVMExprAST::GetFirstChildOfType(EType InExprType) const
 {
 	for (const FRigVMExprAST* Child : Children)
@@ -181,6 +206,167 @@ const FRigVMBlockExprAST* FRigVMExprAST::GetRootBlock() const
 	return nullptr;
 }
 
+FRigVMExprAST::FRigVMBlockArray FRigVMExprAST::GetBlocks(bool bSortByDepth) const
+{
+	FRigVMBlockArray Blocks;
+	GetBlocksImpl(Blocks);
+
+	if(bSortByDepth)
+	{
+		auto SortByDepth = [](const FRigVMExprAST* A, const FRigVMExprAST* B) -> bool
+		{
+			if(A->GetMaximumDepth() > B->GetMaximumDepth())
+			{
+				return true;
+			}
+			return A->Index < B->Index;
+		};
+		Algo::Sort(Blocks, SortByDepth);
+	}
+
+	return Blocks;
+}
+
+TOptional<uint32> FRigVMExprAST::GetBlockCombinationHash() const
+{
+	if(IsA(EType::Literal))
+	{
+		return TOptional<uint32>();
+	}
+
+	// if this expression is a node which is mutable, we are not part of a lazy block
+	if(IsNode())
+	{
+		const FRigVMNodeExprAST* NodeExpr = To<FRigVMNodeExprAST>();
+		check(NodeExpr);
+		if(const URigVMNode* Node = NodeExpr->GetNode())
+		{
+			if(Node->IsMutable())
+			{
+				return TOptional<uint32>();
+			}
+		}
+	}
+
+	// if this expression is a var expression on a mutable node - we are not part of a lazy block
+	if(IsVar())
+	{
+		const FRigVMVarExprAST* VarExpr = To<FRigVMVarExprAST>();
+		check(VarExpr);
+		if(const URigVMPin* Pin = VarExpr->GetPin())
+		{
+			if(const URigVMNode* Node = Pin->GetNode())
+			{
+				if(Node->IsMutable())
+				{
+					return TOptional<uint32>();
+				}
+			}
+		}
+	}
+
+	if(BlockCombinationHash.IsSet())
+	{
+		return BlockCombinationHash.GetValue();
+	}
+	
+	// only consider cached / computed values for lazy blocks
+	if (IsA(FRigVMExprAST::CachedValue))
+	{
+		// collect all of the blocks this expression is in
+		const FRigVMExprAST::FRigVMBlockArray Blocks = GetBlocks(true);
+		if(Blocks.Num() > 1)
+		{
+			// compute a hash for the blocks
+			uint32 ComputedHash = 0;
+			for(const FRigVMBlockExprAST* Block : Blocks)
+			{
+				ComputedHash = HashCombine(ComputedHash, GetTypeHash(Block->Index));
+			}
+
+			BlockCombinationHash = TOptional<uint32>(ComputedHash);
+
+			if(const FRigVMParserAST* Parser = GetParser())
+			{
+				// compute a unique name for the block combination
+				if(!Parser->BlockCombinationHashToName.Contains(ComputedHash))
+				{
+					FString BlockCombinationName;
+					for(const FRigVMBlockExprAST* Block : Blocks)
+					{
+						static constexpr TCHAR BlockNameFormat[] = TEXT("%s%d");
+						const FString BlockNameAndIndex = FString::Printf(BlockNameFormat, *Block->GetName().ToString(), Block->Index);
+
+						if(BlockCombinationName.IsEmpty())
+						{
+							BlockCombinationName = BlockNameAndIndex;
+						}
+						else
+						{
+							static constexpr TCHAR JoinFormat[] = TEXT("%s_%s");
+							BlockCombinationName = FString::Printf(JoinFormat, *BlockCombinationName, *BlockNameAndIndex);
+						}
+					}
+					Parser->BlockCombinationHashToName.Add(ComputedHash, BlockCombinationName);
+				}
+			}
+
+			return BlockCombinationHash.GetValue();
+		}
+	}
+
+	// if we still don't have a hash - rely on the parents
+	if(!BlockCombinationHash.IsSet())
+	{
+		for(int32 ParentIndex = 0; ParentIndex < NumParents(); ParentIndex++)
+		{
+			const TOptional<uint32> ParentHash = ParentAt(ParentIndex)->GetBlockCombinationHash();
+			if(ParentHash.IsSet())
+			{
+				BlockCombinationHash = ParentHash;
+				return ParentHash;
+			}
+		}
+	}
+
+	BlockCombinationHash = TOptional<uint32>();
+	return BlockCombinationHash.GetValue();
+}
+
+const FString& FRigVMExprAST::GetBlockCombinationName() const
+{
+	const TOptional<uint32> CombinationHash = GetBlockCombinationHash();
+	if(CombinationHash.IsSet())
+	{
+		if(const FRigVMParserAST* Parser = GetParser())
+		{
+			if(const FString* CombinationName = Parser->BlockCombinationHashToName.Find(CombinationHash.GetValue()))
+			{
+				return *CombinationName;
+			}
+		}
+	}
+
+	static const FString EmptyString;
+	return EmptyString;
+}
+
+void FRigVMExprAST::GetBlocksImpl(FRigVMBlockArray& InOutBlocks) const
+{
+	if (IsA(EType::Block))
+	{
+		InOutBlocks.AddUnique(this->To<FRigVMBlockExprAST>());
+		return;
+	}
+	
+	for(int32 ParentIndex = 0; ParentIndex < NumParents(); ParentIndex++)
+	{
+		const FRigVMExprAST* ParentExpression = ParentAt(ParentIndex);
+		ParentExpression->GetBlocksImpl(InOutBlocks);
+	}
+}
+
+
 int32 FRigVMExprAST::GetMinChildIndexWithinParent(const FRigVMExprAST* InParentExpr) const
 {
 	int32 MinIndex = INDEX_NONE;
@@ -238,6 +424,8 @@ void FRigVMExprAST::AddParent(FRigVMExprAST* InParent)
 
 	InParent->Children.Add(this);
 	Parents.Add(InParent);
+
+	InvalidateCaches();
 }
 
 void FRigVMExprAST::RemoveParent(FRigVMExprAST* InParent)
@@ -266,6 +454,8 @@ void FRigVMExprAST::RemoveParent(FRigVMExprAST* InParent)
 		}
 		InParent->Children.Remove(this);
 	}
+	
+	InvalidateCaches();
 }
 
 void FRigVMExprAST::RemoveChild(FRigVMExprAST* InChild)
@@ -289,6 +479,7 @@ void FRigVMExprAST::ReplaceParent(FRigVMExprAST* InCurrentParent, FRigVMExprAST*
 			Parents[ParentIndex] = InNewParent;
 			InCurrentParent->Children.Remove(this);
 			InNewParent->Children.Add(this);
+			InvalidateCaches();
 		}
 	}
 }
@@ -306,6 +497,8 @@ void FRigVMExprAST::ReplaceChild(FRigVMExprAST* InCurrentChild, FRigVMExprAST* I
 			Children[ChildIndex] = InNewChild;
 			InCurrentChild->Parents.Remove(this);
 			InNewChild->Parents.Add(this);
+			InCurrentChild->InvalidateCaches();
+			InNewChild->InvalidateCaches();
 		}
 	}
 }
@@ -333,6 +526,22 @@ bool FRigVMExprAST::IsConstant() const
 	return true;
 }
 
+int32 FRigVMExprAST::GetMaximumDepth() const
+{
+	if(MaximumDepth.IsSet())
+	{
+		return MaximumDepth.GetValue();
+	}
+	int32 Depth = 0;
+	for(int32 ParentIndex = 0; ParentIndex < NumParents(); ParentIndex++)
+	{
+		const FRigVMExprAST* ParentExpression = ParentAt(ParentIndex);
+		Depth = FMath::Max<int32>(Depth, ParentExpression->GetMaximumDepth() + 1);
+	}
+	MaximumDepth = Depth;
+	return Depth;
+}
+
 FString FRigVMExprAST::DumpText(const FString& InPrefix) const
 {
 	FString Result;
@@ -364,6 +573,30 @@ FString FRigVMExprAST::DumpText(const FString& InPrefix) const
 	return Result;
 }
 
+void FRigVMExprAST::InvalidateCaches()
+{
+	TArray<bool> Processed;
+	Processed.AddZeroed(GetParser()->Expressions.Num());
+	InvalidateCachesImpl(Processed);
+}
+
+void FRigVMExprAST::InvalidateCachesImpl(TArray<bool>& OutProcessed)
+{
+	if(OutProcessed[Index])
+	{
+		return;
+	}
+	
+	BlockCombinationHash.Reset();
+	MaximumDepth.Reset();
+	OutProcessed[Index] = true;
+
+	for(FRigVMExprAST* ChildExpression : Children)
+	{
+		ChildExpression->InvalidateCachesImpl(OutProcessed);
+	}
+}
+
 bool FRigVMBlockExprAST::ShouldExecute() const
 {
 	return ContainsEntry();
@@ -713,8 +946,6 @@ FRigVMParserAST::FRigVMParserAST(TArray<URigVMGraph*> InGraphs, URigVMController
 	InjectExitsToEntries();
 	FoldNoOps();
 
-	BubbleUpExpressions();
-
 	if (InSettings.bFoldAssignments)
 	{
 		FoldAssignments();
@@ -726,39 +957,6 @@ FRigVMParserAST::FRigVMParserAST(TArray<URigVMGraph*> InGraphs, URigVMController
 	}
 }
 
-FRigVMParserAST::FRigVMParserAST(TArray<URigVMGraph*> InGraphs, const TArray<FRigVMASTProxy>& InNodesToCompute)
-	: LibraryNodeBeingCompiled(nullptr)
-{
-	check(!InGraphs.IsEmpty());
-
-	LastCycleCheckExpr = nullptr;
-
-	FRigVMBlockExprAST* Block = MakeExpr<FRigVMBlockExprAST>(FRigVMExprAST::EType::Block, FRigVMASTProxy());
-	Block->Name = TEXT("NodesToCompute");
-	RootExpressions.Add(Block);
-
-	NodeProxies = InNodesToCompute;
-	
-	Inline(InGraphs, InNodesToCompute);
-
-	for (const FRigVMASTProxy& NodeProxy : NodeProxies)
-	{
-		URigVMNode* Node = NodeProxy.GetSubjectChecked<URigVMNode>();
-		if (Node->IsEvent())
-		{
-			continue;
-		}
-		if (Node->IsMutable())
-		{
-			continue;
-		}
-		TraverseNode(NodeProxy, Block);
-	}
-
-	FRigVMExprAST* ExitExpr = MakeExpr<FRigVMExitExprAST>(FRigVMASTProxy());
-	ExitExpr->AddParent(Block);
-}
-
 FRigVMParserAST::~FRigVMParserAST()
 {
 	for (FRigVMExprAST* Expression : Expressions)
@@ -1500,152 +1698,6 @@ void FRigVMParserAST::InjectExitsToEntries()
 	}
 }
 
-void FRigVMParserAST::BubbleUpExpressions()
-{
-	for (int32 ExpressionIndex = 0; ExpressionIndex < Expressions.Num(); ExpressionIndex++)
-	{
-		FRigVMExprAST* Expression = Expressions[ExpressionIndex];
-		if (!Expression->IsA(FRigVMExprAST::CachedValue))
-		{
-			continue;
-		}
-
-		if (Expression->NumParents() < 2)
-		{
-			continue;
-		}
-
-		// collect all of the blocks this is in and make sure it's bubbled up before that
-		TArray<FRigVMBlockExprAST*> Blocks;
-		for (int32 ParentIndex = 0; ParentIndex < Expression->NumParents(); ParentIndex++)
-		{
-			const FRigVMExprAST* ParentExpression = Expression->ParentAt(ParentIndex);
-			if (ParentExpression->IsA(FRigVMExprAST::Block))
-			{
-				Blocks.AddUnique((FRigVMBlockExprAST*)ParentExpression->To<FRigVMBlockExprAST>());
-			}
-			else
-			{
-				Blocks.AddUnique((FRigVMBlockExprAST*)ParentExpression->GetBlock());
-			}
-		}
-
-		if (Blocks.Num() > 1)
-		{
-			TArray<FRigVMBlockExprAST*> BlockCandidates;
-			BlockCandidates.Append(Blocks);
-			FRigVMBlockExprAST* OuterBlock = nullptr;
-
-			// check if we need to add this expression to the block representing the non-lazy pins on the node.
-			// the expressions in that block are supposed to run before the node can execute - since they are greedy. 
-			// note: for backwards compatibility this also works for the condition pin on the if node and the index pin
-			// on the select node
-			TArray<const FRigVMExprAST*> Parents;
-			TArray<const FRigVMExprAST*> GrandParents;
-
-			for(const FRigVMBlockExprAST* Block : Blocks)
-			{
-				const FRigVMExprAST* Parent = Block->GetParent();
-				Parents.AddUnique(Parent);
-				if(Parent)
-				{
-					GrandParents.AddUnique(Parent->GetParent());
-				}
-			}
-
-			// if any of the parents is nullptr - it means that the block didn't
-			// have a parent and thus is the top level block. in that case we cannot
-			// bubble up further.
-			if(!Parents.Contains(nullptr))
-			{
-				// if all blocks here are part of the same grandparent
-				// parent == pin, grandparent == node
-				if(GrandParents.Num() == 1 && !GrandParents.Contains(nullptr))
-				{
-					if(GrandParents[0]->IsA(FRigVMExprAST::EType::CallExtern))
-					{
-						const FRigVMNodeExprAST* NodeExpr = GrandParents[0]->To<FRigVMNodeExprAST>();
-
-						// find a pin on the node which is either an input or an io and not lazy
-						const URigVMNode* Node = NodeExpr->GetNode();
-						const FRigVMExprAST* TopLevelBlockExpression = nullptr;
-						for(const URigVMPin* Pin : Node->GetPins())
-						{
-							if(Pin->IsLazy())
-							{
-								continue;
-							}
-
-							if(Pin->GetDirection() == ERigVMPinDirection::Input ||
-								Pin->GetDirection() == ERigVMPinDirection::IO)
-							{
-								if(const FRigVMExprAST* PinExpr = NodeExpr->FindExprWithPinName(Pin->GetFName()))
-								{
-									TopLevelBlockExpression = PinExpr->GetFirstChildOfType(FRigVMExprAST::Block);
-								}
-							}
-
-							if(TopLevelBlockExpression)
-							{
-								break;
-							}
-						}
-
-						if(TopLevelBlockExpression)
-						{
-							OuterBlock = (FRigVMBlockExprAST*)TopLevelBlockExpression->To<FRigVMBlockExprAST>();
-							OuterBlock->Children.Add(Expression);
-							Expression->Parents.Insert(OuterBlock, 0);
-							continue;
-						}
-					}
-				}
-			}
-
-			// this expression is part of multiple blocks, and it needs to be bubbled up.
-			// for this we'll walk up the block tree and find the first block which contains all of them
-			for (int32 BlockCandidateIndex = 0; BlockCandidateIndex < BlockCandidates.Num(); BlockCandidateIndex++)
-			{
-				FRigVMBlockExprAST* BlockCandidate = BlockCandidates[BlockCandidateIndex];
-
-				bool bFoundCandidate = true;
-				TMap<const FRigVMExprAST*, bool> ContainedExpressionsCache;
-				for (int32 BlockIndex = 0; BlockIndex < Blocks.Num(); BlockIndex++)
-				{
-					FRigVMBlockExprAST* Block = Blocks[BlockIndex];
-					if (!BlockCandidate->Contains(Block, &ContainedExpressionsCache))
-					{
-						bFoundCandidate = false;
-						break;
-					}
-				}
-
-				if (bFoundCandidate)
-				{
-					OuterBlock = BlockCandidate;
-					break;
-				}
-
-				BlockCandidates.AddUnique((FRigVMBlockExprAST*)BlockCandidate->GetBlock());
-			}
-
-			// we found a block which contains all of our blocks.
-			// we are now going to inject this block as the first parent
-			// of the cached value, so that the traverser sees it earlier
-			if (OuterBlock)
-			{
-				MinIndexOfChildWithinParent.Reset();
-				int32 ChildIndex = Expression->GetMinChildIndexWithinParent(OuterBlock);
-				if (ChildIndex != INDEX_NONE)
-				{
-					OuterBlock->Children.Insert(Expression, ChildIndex);
-					Expression->Parents.Insert(OuterBlock, 0);
-				}
-			}
-		}
-	}
-}
-
 void FRigVMParserAST::RefreshExprIndices()
 {
 	for (int32 Index = 0; Index < Expressions.Num(); Index++)
@@ -2230,7 +2282,78 @@ FString FRigVMParserAST::DumpDot() const
 
 	struct Local
 	{
-		static TArray<int32> VisitChildren(const FRigVMExprAST* InExpr, int32 InSubGraphIndex, FVisualGraph& OutGraph)
+		const TArray<FColor> BlockCombinationColors =
+		{
+			FColor(255,235,205),
+			FColor(220,20,60),
+			FColor(46,139,87),
+			FColor(0,0,205),
+			FColor(255,228,225),
+			FColor(255,218,185),
+			FColor(255,140,0),
+			FColor(255,228,181),
+			FColor(216,191,216),
+			FColor(210,105,30),
+			FColor(240,248,255),
+			FColor(30,144,255),
+			FColor(47,79,79),
+			FColor(245,245,220),
+			FColor(165,42,42),
+			FColor(255,245,238),
+			FColor(112,128,144),
+			FColor(220,220,220),
+			FColor(123,104,238),
+			FColor(139,0,139),
+			FColor(255,182,193),
+			FColor(250,128,114),
+			FColor(148,0,211),
+			FColor(224,255,255),
+			FColor(255,165,0),
+			FColor(255,250,240),
+			FColor(0,128,128),
+			FColor(175,238,238),
+			FColor(147,112,219),
+			FColor(255,160,122),
+			FColor(0,255,127),
+			FColor(255,240,245),
+			FColor(211,211,211),
+			FColor(173,255,47),
+			FColor(0,100,0),
+			FColor(0,128,0),
+			FColor(32,178,170),
+			FColor(123,104,238),
+			FColor(240,230,140),
+			FColor(139,69,19),
+			FColor(153,50,204),
+			FColor(219,112,147),
+			FColor(138,43,226),
+			FColor(245,245,245),
+			FColor(255,255,240),
+			FColor(255,69,0),
+			FColor(135,206,235),
+			FColor(240,255,255),
+			FColor(205,92,92),
+			FColor(255,250,205),
+			FColor(105,105,105),
+			FColor(255,250,250),
+			FColor(72,61,139),
+			FColor(255,248,220),
+			FColor(255,192,203),
+			FColor(222,184,135),
+			FColor(245,222,179),
+			FColor(0,0,128),
+			FColor(245,255,250),
+			FColor(25,25,112),
+			FColor(244,164,96),
+			FColor(238,130,238),
+			FColor(240,255,240),
+			FColor(34,139,34),
+		};
+
+		int32 BlockCombinationColorIndex = 0;
+		TMap<uint32, int32> BlockCombinationHashToColor; 
+
+		TArray<int32> VisitChildren(const FRigVMExprAST* InExpr, int32 InSubGraphIndex, FVisualGraph& OutGraph)
 		{
 			TArray<int32> ChildNodeIndices;
 			for (FRigVMExprAST* Child : InExpr->Children)
@@ -2240,7 +2363,7 @@ FString FRigVMParserAST::DumpDot() const
 			return ChildNodeIndices;
 		}
 		
-		static int32 VisitExpr(const FRigVMExprAST* InExpr, int32 InSubGraphIndex, FVisualGraph& OutGraph)
+		int32 VisitExpr(const FRigVMExprAST* InExpr, int32 InSubGraphIndex, FVisualGraph& OutGraph)
 		{
 			const FName NodeName = *FString::Printf(TEXT("node_%d"), InExpr->GetIndex());
 
@@ -2390,8 +2513,24 @@ FString FRigVMParserAST::DumpDot() const
 
 			if (!Label.IsEmpty())
 			{
+				TOptional<FLinearColor> Color;
+				const TOptional<uint32> BlockCombinationHash = InExpr->GetBlockCombinationHash();
+				if(BlockCombinationHash.IsSet())
+				{
+					const uint32 Hash = BlockCombinationHash.GetValue();
+					if(!BlockCombinationHashToColor.Contains(Hash))
+					{
+						BlockCombinationHashToColor.Add(Hash, BlockCombinationColorIndex++);
+						if(BlockCombinationColorIndex >= BlockCombinationColors.Num())
+						{
+							BlockCombinationColorIndex = 0;
+						}
+					}
+					Color = BlockCombinationColors[BlockCombinationHashToColor.FindChecked(Hash)];
+				}
+				
 				const TOptional<FName> DisplayName = FName(*Label);
-				NodeIndex = OutGraph.AddNode(NodeName, DisplayName, TOptional<FLinearColor>(), Shape);
+				NodeIndex = OutGraph.AddNode(NodeName, DisplayName, Color, Shape);
 				OutGraph.AddNodeToSubGraph(NodeIndex, SubGraphIndex);
 			}
 
@@ -2412,6 +2551,7 @@ FString FRigVMParserAST::DumpDot() const
 		}
 	};
 
+	Local LocalStruct;
 	for (FRigVMExprAST* Expr : RootExpressions)
 	{
 		if (Expr == GetObsoleteBlock(false))
@@ -2419,7 +2559,7 @@ FString FRigVMParserAST::DumpDot() const
 			continue;
 		}
 
-		Local::VisitExpr(Expr, INDEX_NONE, VisualGraph);
+		LocalStruct.VisitExpr(Expr, INDEX_NONE, VisualGraph);
 	}
 
 	return VisualGraph.DumpDot();
diff --git a/Engine/Plugins/Runtime/RigVM/Source/RigVMDeveloper/Private/RigVMCompiler/RigVMCodeGenerator.cpp b/Engine/Plugins/Runtime/RigVM/Source/RigVMDeveloper/Private/RigVMCompiler/RigVMCodeGenerator.cpp
--- a/Engine/Plugins/Runtime/RigVM/Source/RigVMDeveloper/Private/RigVMCompiler/RigVMCodeGenerator.cpp
+++ b/Engine/Plugins/Runtime/RigVM/Source/RigVMDeveloper/Private/RigVMCompiler/RigVMCodeGenerator.cpp
@@ -865,6 +865,12 @@ FString FRigVMCodeGenerator::DumpInstructions(const FRigVMExtendedExecuteContext
 				}
 				break;
 			}
+			case ERigVMOpCode::RunInstructions:
+			{
+				const FRigVMRunInstructionsOp& Op = ByteCode.GetOpAt<FRigVMRunInstructionsOp>(Instruction);
+				// todo
+				break;
+			}
 			case ERigVMOpCode::Invalid:
 			case ERigVMOpCode::ChangeType:
 			default:
@@ -1407,6 +1413,12 @@ void FRigVMCodeGenerator::ParseInstructionGroups(const FRigVMExtendedExecuteCont
 						}
 						break;
 					}
+					case ERigVMOpCode::RunInstructions:
+					{
+						const FRigVMRunInstructionsOp& Op = ByteCode.GetOpAt<FRigVMRunInstructionsOp>(Instruction);
+						// todo
+						break;
+					}
 					default:
 					{
 						break;
@@ -1556,6 +1568,12 @@ void FRigVMCodeGenerator::ParseInstructionGroups(const FRigVMExtendedExecuteCont
 					}
 					break;
 				}
+				case ERigVMOpCode::RunInstructions:
+				{
+					const FRigVMRunInstructionsOp& Op = ByteCode.GetOpAt<FRigVMRunInstructionsOp>(Instruction);
+					// todo
+					break;
+				}
 				default:
 				{
 					break;
diff --git a/Engine/Plugins/Runtime/RigVM/Source/RigVMDeveloper/Private/RigVMCompiler/RigVMCompiler.cpp b/Engine/Plugins/Runtime/RigVM/Source/RigVMDeveloper/Private/RigVMCompiler/RigVMCompiler.cpp
--- a/Engine/Plugins/Runtime/RigVM/Source/RigVMDeveloper/Private/RigVMCompiler/RigVMCompiler.cpp
+++ b/Engine/Plugins/Runtime/RigVM/Source/RigVMDeveloper/Private/RigVMCompiler/RigVMCompiler.cpp
@@ -17,6 +17,7 @@
 #include "RigVMModel/RigVMClient.h"
 #include "RigVMFunctions/RigVMDispatch_Array.h"
 #include "Algo/Count.h"
+#include "Algo/Sort.h"
 #include "String/Join.h"
 
 #if WITH_EDITOR
@@ -1165,6 +1166,28 @@ bool URigVMCompiler::Compile(const FRigVMCompileSettings& InSettings, TArray<URi
 		}
 	}
 
+	// find all blocks
+	for (FRigVMExprAST* Expression : WorkData.AST->Expressions)
+	{
+		const TOptional<uint32> OptionalHash = Expression->GetBlockCombinationHash();
+		const uint32 Hash = OptionalHash.Get(0);
+		if(!WorkData.LazyBlocks.Contains(Hash))
+		{
+			WorkData.LazyBlocks.Add(Hash, MakeShared<FRigVMCompilerWorkData::FLazyBlockInfo>());
+		}
+
+		if(OptionalHash.IsSet())
+		{
+			TSharedPtr<FRigVMCompilerWorkData::FLazyBlockInfo> Info = WorkData.LazyBlocks.FindChecked(Hash);
+			Info->Hash = OptionalHash;
+			if(Info->BlockCombinationName.IsEmpty())
+			{
+				Info->BlockCombinationName = Expression->GetBlockCombinationName();
+			}
+			Info->Expressions.Add(Expression);
+		}
+	}
+
 	WorkData.ExprComplete.Reset();
 	for (FRigVMExprAST* RootExpr : *WorkData.AST)
 	{
@@ -1216,6 +1239,19 @@ bool URigVMCompiler::Compile(const FRigVMCompileSettings& InSettings, TArray<URi
 
 	WorkData.bSetupMemory = false;
 	WorkData.ExprComplete.Reset();
+
+	// sort the expressions in each block by depth
+	for(auto Pair : WorkData.LazyBlocks)
+	{
+		auto SortByDepth = [](const FRigVMExprAST* InExpression) -> int32
+		{
+			return InExpression->GetMaximumDepth();
+		};
+		Algo::SortBy(Pair.Value->Expressions, SortByDepth);
+	}
+	
+	// traverse the top level blocks
+	WorkData.CurrentBlockHash = TOptional<uint32>();
 	for (FRigVMExprAST* RootExpr : *WorkData.AST)
 	{
 		if (!TraverseExpression(RootExpr, WorkData))
@@ -1231,7 +1267,57 @@ bool URigVMCompiler::Compile(const FRigVMCompileSettings& InSettings, TArray<URi
 		{
 			WorkData.VM->GetByteCode().AddExitOp();
 		}
+	}
+
+	// traverse all other blocks - this has to be an index based loop
+	if(!WorkData.LazyBlocksToProcess.IsEmpty())
+	{
+		const uint64 JumpToEndOfBlocksExternByte = WorkData.VM->GetByteCode().AddJumpOp(ERigVMOpCode::JumpForward, INDEX_NONE);
+		const int32 JumpToEndOfBlocksExternInstruction = WorkData.VM->GetByteCode().GetNumInstructions() - 1;
+
+		FRigVMByteCode& ByteCode = WorkData.VM->GetByteCode();
+		for(int32 LazyBlockHashIndex = 0; LazyBlockHashIndex < WorkData.LazyBlocksToProcess.Num(); LazyBlockHashIndex++)
+		{
+			TSharedPtr<FRigVMCompilerWorkData::FLazyBlockInfo>& BlockInfo =
+				WorkData.LazyBlocks.FindChecked(WorkData.LazyBlocksToProcess[LazyBlockHashIndex]);
+			if(BlockInfo->bProcessed)
+			{
+				continue;
+			}
+
+			BlockInfo->StartInstruction = ByteCode.GetNumInstructions();
+
+			TGuardValue<TOptional<uint32>> HashGuard(WorkData.CurrentBlockHash, BlockInfo->Hash);
+			for(const FRigVMExprAST* Expression : BlockInfo->Expressions)
+			{
+				TraverseExpression(Expression, WorkData);
+			}
+
+			BlockInfo->EndInstruction = ByteCode.GetNumInstructions() - 1;
+			BlockInfo->bProcessed = true;
+		}
+
+		for(int32 LazyBlockHashIndex = 0; LazyBlockHashIndex < WorkData.LazyBlocksToProcess.Num(); LazyBlockHashIndex++)
+		{
+			TSharedPtr<FRigVMCompilerWorkData::FLazyBlockInfo>& BlockInfo =
+				WorkData.LazyBlocks.FindChecked(WorkData.LazyBlocksToProcess[LazyBlockHashIndex]);
+
+			// update all run instructions ops in the bytecode
+			for(uint64 RunInstructionsByteCode : BlockInfo->RunInstructionsToUpdate)
+			{
+				FRigVMRunInstructionsOp& RunInstructionsOp = ByteCode.GetOpAt<FRigVMRunInstructionsOp>(RunInstructionsByteCode);
+				RunInstructionsOp.StartInstruction = BlockInfo->StartInstruction;
+				RunInstructionsOp.EndInstruction = BlockInfo->EndInstruction;
+			}
+		}
+		
+		// update the operator with the target instruction 
+		const int32 InstructionsToJump = WorkData.VM->GetByteCode().GetNumInstructions() - JumpToEndOfBlocksExternInstruction;
+		WorkData.VM->GetByteCode().GetOpAt<FRigVMJumpOp>(JumpToEndOfBlocksExternByte).InstructionIndex = InstructionsToJump;
+	}
 	
+	if (!CurrentCompilationFunction)
+	{
 		WorkData.VM->GetByteCode().AlignByteCode();
 	}
 
@@ -1473,12 +1559,101 @@ bool URigVMCompiler::TraverseExpression(const FRigVMExprAST* InExpr, FRigVMCompi
 		return true;
 	}
 
+	// if we hit an expression which is on a different block
+	// take care of redirecting this to the RunInstructions op
+	if(!WorkData.TraversalExpressions.IsEmpty())
+	{
+		const TOptional<uint32> PreviousBlockHash = WorkData.TraversalExpressions.Last()->GetBlockCombinationHash();
+		const TOptional<uint32> NextBlockHash = InExpr->GetBlockCombinationHash();
+
+		if((PreviousBlockHash != NextBlockHash) && NextBlockHash.IsSet())
+		{
+			TSharedPtr<FRigVMCompilerWorkData::FLazyBlockInfo>& BlockInfo = WorkData.LazyBlocks.FindChecked(NextBlockHash.GetValue());
+
+			// if we are hitting the block for the first time - let's register the execution state operand
+			if(WorkData.bSetupMemory)
+			{
+				if(!BlockInfo->ExecuteStateOperand.IsValid())
+				{
+					static constexpr TCHAR Format[] = TEXT("BlockExecuteState_%zu");
+					const FString BlockStateName = FString::Printf(Format, NextBlockHash.GetValue());
+					BlockInfo->ExecuteStateOperand = WorkData.AddProperty(
+						ERigVMMemoryType::Work,
+						*BlockStateName,
+						FRigVMInstructionSetExecuteState::StaticStruct()->GetStructCPPName(),
+						FRigVMInstructionSetExecuteState::StaticStruct()
+					);
+				}
+			}
+			else
+			{
+				check(BlockInfo->ExecuteStateOperand.IsValid());
+
+				FRigVMByteCode& ByteCode = WorkData.VM->GetByteCode();
+
+				// add an operator to invoke the block lazily
+				const int32 RunInstructionsOpIndex = ByteCode.GetNumInstructions();
+				BlockInfo->RunInstructionsToUpdate.Add(
+					ByteCode.AddRunInstructionsOp(BlockInfo->ExecuteStateOperand, INDEX_NONE, INDEX_NONE)
+				);
+				if (WorkData.Settings.SetupNodeInstructionIndex)
+				{
+					const FRigVMCallstack Callstack = InExpr->GetProxy().GetCallstack();
+					if(Callstack.Num() > 0)
+					{
+						WorkData.VM->GetByteCode().SetSubject(RunInstructionsOpIndex, Callstack.GetCallPath(), Callstack.GetStack());
+					}
+				}
+
+				// mark the block to be processed
+				if(!BlockInfo->bProcessed)
+				{
+					WorkData.LazyBlocksToProcess.AddUnique(NextBlockHash.GetValue());
+				}
+
+				// return here and stop traversal
+				return true;
+			}
+		}
+	}
+
+	if(!WorkData.bSetupMemory)
+	{
+		// skip any expression which is not part of this block
+		if(WorkData.CurrentBlockHash.IsSet())
+		{
+			const TOptional<uint32> NextBlockHash = InExpr->GetBlockCombinationHash();
+			if(WorkData.CurrentBlockHash != NextBlockHash)
+			{
+				return true;
+			}
+		}
+	}
+
 	if (WorkData.ExprComplete.Contains(InExpr))
 	{
 		return true;
 	}
 	WorkData.ExprComplete.Add(InExpr, true);
 
+	struct FTraversalGuard
+	{
+		FTraversalGuard(const FRigVMExprAST* InExpr, FRigVMCompilerWorkData& InWorkData)
+			: WorkData(InWorkData)
+		{
+			WorkData.TraversalExpressions.Push(InExpr);
+		}
+
+		~FTraversalGuard()
+		{
+			WorkData.TraversalExpressions.Pop();
+		}
+		
+		FRigVMCompilerWorkData& WorkData;
+	};
+
+	const FTraversalGuard TraversalGuard(InExpr, WorkData);
+	
 	switch (InExpr->GetType())
 	{
 		case FRigVMExprAST::EType::Block:
@@ -2304,10 +2479,16 @@ bool URigVMCompiler::TraverseInlineFunction(const FRigVMInlineFunctionExprAST* I
 			int32 StartIndex = MemoryType == ERigVMMemoryType::Work ? NumProperties : 0;
 			for (int32 PropertyIndex = StartIndex; PropertyIndex < Properties.Num(); ++PropertyIndex)
 			{
+				static const FString RigVMInstructionSetExecuteStateName = FRigVMInstructionSetExecuteState::StaticStruct()->GetStructCPPName();
+				
 				const FRigVMFunctionCompilationPropertyDescription& Description = Properties[PropertyIndex];
 				FString NewName = Description.Name.ToString();
 				static const FString FunctionLibraryPrefix = TEXT("FunctionLibrary");
-				if (NewName.StartsWith(FunctionLibraryPrefix))
+
+				// instantiate function library specific work state as well as
+				// instruction set execute state - which is used for lazy blocks.
+				if (NewName.StartsWith(FunctionLibraryPrefix) ||
+					Description.CPPType.Equals(RigVMInstructionSetExecuteStateName))
 				{
 					NewName = FString::Printf(TEXT("%s%s"), *FunctionReferenceNode->GetNodePath(), *NewName.RightChop(FunctionLibraryPrefix.Len()));
 					FRigVMPropertyDescription::SanitizeName(NewName);
@@ -2490,12 +2671,17 @@ bool URigVMCompiler::TraverseInlineFunction(const FRigVMInlineFunctionExprAST* I
 				const int32 FunctionIndex = WorkData.VM->AddRigVMFunction(FunctionCompilationData->FunctionNames[Op.FunctionIndex].ToString());
 				Op.FunctionIndex = FunctionIndex;
 			}
-
-			if (Instruction.OpCode == ERigVMOpCode::JumpToBranch)
+			else if (Instruction.OpCode == ERigVMOpCode::JumpToBranch)
 			{
 				FRigVMJumpToBranchOp& Op = ByteCode.GetOpAt<FRigVMJumpToBranchOp>(Instruction);
 				Op.FirstBranchInfoIndex = Op.FirstBranchInfoIndex + BranchIndexStart;
 			}
+			else if (Instruction.OpCode == ERigVMOpCode::RunInstructions)
+			{
+				FRigVMRunInstructionsOp& Op = ByteCode.GetOpAt<FRigVMRunInstructionsOp>(Instruction);
+				Op.StartInstruction = Op.StartInstruction + InstructionIndexStart;
+				Op.EndInstruction = Op.EndInstruction + InstructionIndexStart;
+			}
 
 			if (WorkData.Settings.SetupNodeInstructionIndex)
 			{
diff --git a/Engine/Plugins/Runtime/RigVM/Source/RigVMDeveloper/Public/RigVMCompiler/RigVMAST.h b/Engine/Plugins/Runtime/RigVM/Source/RigVMDeveloper/Public/RigVMCompiler/RigVMAST.h
--- a/Engine/Plugins/Runtime/RigVM/Source/RigVMDeveloper/Public/RigVMCompiler/RigVMAST.h
+++ b/Engine/Plugins/Runtime/RigVM/Source/RigVMDeveloper/Public/RigVMCompiler/RigVMAST.h
@@ -102,70 +102,117 @@ public:
 	// disable copy constructor
 	FRigVMExprAST(const FRigVMExprAST&) = delete;
 
-	// returns the parser this expression is owned by
-	// @return the parser this expression is owned by 
+	/// returns the parser this expression is owned by
+	/// @return the parser this expression is owned by 
 	const FRigVMParserAST* GetParser() const { return ParserPtr; }
 
-	// returns the name of the expression (can be NAME_None)
-	// @return the name of the expression
+	/// returns the name of the expression (can be NAME_None)
+	/// @return the name of the expression
 	FName GetName() const { return Name; }
 
-	// returns the exact type of the expression
-	// @return the exact type of the expression
+	/// returns the exact type of the expression
+	/// @return the exact type of the expression
 	EType GetType() const { return Type; }
 
-	// returns the name of the expression's type
-	// @return the name of the expression's type
+	/// returns the name of the expression's type
+	/// @return the name of the expression's type
 	FName GetTypeName() const;
 
-	// provides type checking for inherited types
-	// @param InType the type to check against
-	// @return true if this expression is of a given type
+	/// provides type checking for inherited types
+	/// @param InType the type to check against
+	/// @return true if this expression is of a given type
 	virtual bool IsA(EType InType) const = 0;
 
-	// returns the index of this expression within the parser's storage
-	// @return the index of this expression within the parser's storage
+	/// returns true if the expression is a node expression
+	virtual bool IsNode() const
+	{
+		return false;
+	}
+
+	/// returns true if the expression is a var expression
+	virtual bool IsVar() const
+	{
+		return false;
+	}
+
+	/// returns the index of this expression within the parser's storage
+	/// @return the index of this expression within the parser's storage
 	int32 GetIndex() const { return Index; }
 
-	// returns true if the expressoin is valid
+	/// returns true if the expressoin is valid
 	bool IsValid() const { return GetIndex() != INDEX_NONE; }
 
-	// returns the parent of this expression
-	// @return the parent of this expression
+	/// returns the model proxy this expression relates to
+	virtual const FRigVMASTProxy& GetProxy() const
+	{
+		static const FRigVMASTProxy EmptyProxy;
+		return EmptyProxy;
+	}
+
+	/// returns the parent of this expression
+	/// @return the parent of this expression
 	const FRigVMExprAST* GetParent() const;
 
-	// returns the first parent in the tree of a given type
-	// @param InExprType The type of expression to look for in the parent tree
-	// @return the first parent in the tree of a given type
+	/// returns the first parent in the tree of a given type
+	/// @param InExprType The type of expression to look for in the parent tree
+	/// @return the first parent in the tree of a given type
 	virtual const FRigVMExprAST* GetFirstParentOfType(EType InExprType) const;
 
-	// returns the block of this expression
-	// @return the block of this expression
+	/// returns true if a given expression is parented to another expression
+	/// @param InParentExpr The potential parent to check
+	/// @return True if this expression is parented to the InParentExpr
+	bool IsParentedTo(const FRigVMExprAST* InParentExpr) const;
+
+	/// returns true if a given expression is a parent of another expression
+	/// @param InChildExpr The potential child to check
+	/// @return True if this expression is a parent of the InChildExpr
+	bool IsParentOf(const FRigVMExprAST* InChildExpr) const;
+
+	/// returns the block of this expression
+	/// @return the block of this expression
 	const FRigVMBlockExprAST* GetBlock() const;
 
-	// returns the root / top level block of this expression
-	// @return the root / top level block of this expression
+	/// returns the root / top level block of this expression
+	/// @return the root / top level block of this expression
 	const FRigVMBlockExprAST* GetRootBlock() const;
 
-	// returns the lowest child index found for this expression within a parent (or INDEX_NONE)
-	// @return the lowest child index found for this expression within a parent (or INDEX_NONE)
+	/// returns all of the block this expression is in
+	/// @param bSortByDepth If true the blocks will be returned sorted (head to tail)
+	/// @return all of the block this expression is in
+	typedef TArray<const FRigVMBlockExprAST*, TInlineAllocator<5>> FRigVMBlockArray;
+	FRigVMBlockArray GetBlocks(bool bSortByDepth) const;
+
+	/// returns the unique block combination hash for this expression (or INDEX_NONE)
+	/// @return the unique block combination name for this expression (or INDEX_NONE)
+	TOptional<uint32> GetBlockCombinationHash() const;
+
+	/// returns the unique block combination name for this expression (or an empty string)
+	/// @return the unique block combination name for this expression (or an empty string)
+	const FString& GetBlockCombinationName() const;
+
+	/// returns the lowest child index found for this expression within a parent (or INDEX_NONE)
+	/// @return the lowest child index found for this expression within a parent (or INDEX_NONE)
 	int32 GetMinChildIndexWithinParent(const FRigVMExprAST* InParentExpr) const;
 
-	// returns the number of children of this expression
-	// @return the number of children of this expression
+	/// returns the number of children of this expression
+	/// @return the number of children of this expression
 	int32 NumChildren() const { return Children.Num(); }
 
-	// returns true if this expressions is constant (non varying)
-	// @return true if this expressions is constant (non varying)
+	/// returns true if this expressions is constant (non varying)
+	/// @return true if this expressions is constant (non varying)
 	virtual bool IsConstant() const;
 
-	// returns true if this expressions is varying (non constant)
-	// @return true if this expressions is varying (non constant)
+	/// returns true if this expressions is varying (non constant)
+	/// @return true if this expressions is varying (non constant)
 	bool IsVarying() const { return !IsConstant(); }
 
-	// accessor operator for a given child
-	// @param InIndex the index of the child to retrieve (bound = NumChildren() - 1)
-	// @return the child at the given index
+	/// returns the maximum depth of the expression in tree (number of recursive parents)
+	/// @return the maximum depth of the expression in tree (number of recursive parents)
+	int32 GetMaximumDepth() const;
+
+	/// accessor operator for a given child
+	/// @param InIndex the index of the child to retrieve (bound = NumChildren() - 1)
+	/// @return the child at the given index
 	const FRigVMExprAST* operator[](int32 InIndex) const { return Children[InIndex]; }
 
 	// begin iterator accessor for the children
@@ -173,55 +220,55 @@ public:
 	// end iterator accessor for the children
 	TArray<FRigVMExprAST*>::RangedForConstIteratorType end() const { return Children.end(); }
 
-	// templated getter to retrieve a child with a given index
-	// type checking will occur within the ::To method and raise
-	// @param InIndex the index of the child to retrieve
-	// @return the child at a given index cast to the provided class
+	/// templated getter to retrieve a child with a given index
+	/// type checking will occur within the ::To method and raise
+	/// @param InIndex the index of the child to retrieve
+	/// @return the child at a given index cast to the provided class
 	template<class ObjectType>
 	const ObjectType* ChildAt(int32 InIndex) const
 	{
 		return Children[InIndex]->To<ObjectType>();
 	}
 
-	// getter to retrieve a child with a given index
-	// @param InIndex the index of the child to retrieve
-	// @return the child at a given index
+	/// getter to retrieve a child with a given index
+	/// @param InIndex the index of the child to retrieve
+	/// @return the child at a given index
 	const FRigVMExprAST* ChildAt(int32 InIndex) const
 	{
 		return Children[InIndex];
 	}
 
-	// returns the first parent in the tree of a given type
-	// @param InExprType The type of expression to look for in the parent tree
-	// @return the first parent in the tree of a given type
+	/// returns the first parent in the tree of a given type
+	/// @param InExprType The type of expression to look for in the parent tree
+	/// @return the first parent in the tree of a given type
 	virtual const FRigVMExprAST* GetFirstChildOfType(EType InExprType) const;
 
-	// returns the number of parents of this expression
-	// @return the number of parents of this expression
+	/// returns the number of parents of this expression
+	/// @return the number of parents of this expression
 	int32 NumParents() const { return Parents.Num(); }
 
-	// templated getter to retrieve a parent with a given index
-	// type checking will occur within the ::To method and raise
-	// @param InIndex the index of the parent to retrieve
-	// @return the parent at a given index cast to the provided class
+	/// templated getter to retrieve a parent with a given index
+	/// type checking will occur within the ::To method and raise
+	/// @param InIndex the index of the parent to retrieve
+	/// @return the parent at a given index cast to the provided class
 	template<class ObjectType>
 	const ObjectType* ParentAt(int32 InIndex) const
 	{
 		return Parents[InIndex]->To<ObjectType>();
 	}
 
-	// getter to retrieve a parent with a given index
-	// @param InIndex the index of the parent to retrieve
-	// @return the parent at a given index
+	/// getter to retrieve a parent with a given index
+	/// @param InIndex the index of the parent to retrieve
+	/// @return the parent at a given index
 	const FRigVMExprAST* ParentAt(int32 InIndex) const
 	{
 		return Parents[InIndex];
 	}
 
-	// const templated cast for casting between
-	// different expression types.
-	// specializations below are used for type checking
-	// @return this object cast to the provided class
+	/// const templated cast for casting between
+	/// different expression types.
+	/// specializations below are used for type checking
+	/// @return this object cast to the provided class
 	template<class ObjectType>
 	const ObjectType* To() const
 	{
@@ -229,10 +276,10 @@ public:
 		return nullptr;
 	}
 
-	// templated cast for casting between
-	// different expression types.
-	// specializations below are used for type checking
-	// @return this object cast to the provided class
+	/// templated cast for casting between
+	/// different expression types.
+	/// specializations below are used for type checking
+	/// @return this object cast to the provided class
 	template<class ObjectType>
 	ObjectType* To()
 	{
@@ -273,7 +320,8 @@ protected:
 	// @param InReplacement the expression to replace this one
 	void ReplaceBy(FRigVMExprAST* InReplacement);
 
-protected:
+	// computes all of the block this expression is in
+	void GetBlocksImpl(FRigVMBlockArray& InOutBlocks) const;
 
 	// returns a string containing an indented tree structure
 	// for debugging purposes. this is only used by the parser
@@ -281,6 +329,10 @@ protected:
 	// @return the text representation of this part of the tree
 	virtual FString DumpText(const FString& InPrefix = FString()) const;
 
+	// Resets the internal caches of the expression
+	void InvalidateCaches();
+	void InvalidateCachesImpl(TArray<bool>& OutProcessed);
+
 	FName Name;
 	EType Type;
 	int32 Index;
@@ -288,6 +340,8 @@ protected:
 	TArray<FRigVMExprAST*> Parents;
 	TArray<FRigVMExprAST*> Children;
 	TMap<FName, int32> PinNameToChildIndex;
+	mutable TOptional<TOptional<uint32>> BlockCombinationHash;
+	mutable TOptional<int32> MaximumDepth;
 
 	friend class FRigVMParserAST;
 	friend class URigVMCompiler;
@@ -521,7 +575,12 @@ public:
 	FRigVMNodeExprAST(const FRigVMNodeExprAST&) = delete;
 
 	// returns the proxy this expression is using
-	const FRigVMASTProxy& GetProxy() const { return Proxy; }
+	virtual const FRigVMASTProxy& GetProxy() const override { return Proxy; }
+
+	virtual bool IsNode() const
+	{
+		return true;
+	}
 
 	// returns the node from the model this expression is referencing
 	// @return the node from the model this expression is referencing
@@ -755,10 +814,15 @@ public:
 		return InType == EType::Var;
 	};
 
+	virtual bool IsVar() const
+	{
+		return true;
+	}
+
 	virtual bool IsConstant() const override;
 
 	// returns the proxy this expression is using
-	const FRigVMASTProxy& GetProxy() const { return Proxy; }
+	virtual const FRigVMASTProxy& GetProxy() const override { return Proxy; }
 
 	// returns the pin in the model this variable is representing
 	// @return the pin in the model this variable is representing
@@ -924,6 +988,8 @@ public:
 	// @return the target pin for this assignment
 	URigVMPin* GetTargetPin() const { return Link.TargetProxy.GetSubjectChecked<URigVMPin>(); }
 
+	virtual const FRigVMASTProxy& GetProxy() const override { return GetTargetProxy(); }
+
 	// overload of the type checking mechanism
 	virtual bool IsA(EType InType) const override
 	{
@@ -1204,9 +1270,6 @@ public:
 
 private:
 
-	// private constructor for a partial build
-	FRigVMParserAST(TArray<URigVMGraph*> InGraphs, const TArray<FRigVMASTProxy>& InNodesToCompute);
-
 	// make function to create an expression
 	template<class ExprType>
 	ExprType* MakeExpr(FRigVMExprAST::EType InType, const FRigVMASTProxy& InProxy)
@@ -1301,9 +1364,6 @@ private:
 	// helper function to inject an exit expression at the end of every entry expressoin
 	void InjectExitsToEntries();
 
-	// helper function to bubble up expressions which are part of multiple blocks
-	void BubbleUpExpressions();
-
 	// helper function to refresh the expression indices (used after deleting an expression)
 	void RefreshExprIndices();
 
@@ -1356,6 +1416,7 @@ private:
 	TArray<FRigVMExprAST*> RootExpressions;
 	TArray<FRigVMExprAST*> DeletedExpressions;
 	FRigVMBlockExprAST* ObsoleteBlock;
+	mutable TMap<uint32, FString> BlockCombinationHashToName;
 
 	TArray<FRigVMASTProxy> NodeProxies;
 	TMap<FRigVMASTProxy, FRigVMASTProxy> SharedOperandPins;
diff --git a/Engine/Plugins/Runtime/RigVM/Source/RigVMDeveloper/Public/RigVMCompiler/RigVMCompiler.h b/Engine/Plugins/Runtime/RigVM/Source/RigVMDeveloper/Public/RigVMCompiler/RigVMCompiler.h
--- a/Engine/Plugins/Runtime/RigVM/Source/RigVMDeveloper/Public/RigVMCompiler/RigVMCompiler.h
+++ b/Engine/Plugins/Runtime/RigVM/Source/RigVMDeveloper/Public/RigVMCompiler/RigVMCompiler.h
@@ -139,6 +139,7 @@ public:
 	TMap<const FRigVMVarExprAST*, FRigVMOperand> ExprToOperand;
 	TMap<const FRigVMExprAST*, bool> ExprComplete;
 	TArray<const FRigVMExprAST*> ExprToSkip;
+	TArray<const FRigVMExprAST*> TraversalExpressions;
 	TMap<FString, int32> ProcessedLinks;
 	TMap<int32, FRigVMOperand> IntegerLiterals;
 	FRigVMOperand ComparisonOperand;
@@ -196,6 +197,28 @@ public:
 	// operators that have been delayed for injection into the bytecode
 	TMap<FRigVMOperand, FCopyOpInfo> DeferredCopyOps;
 
+	struct FLazyBlockInfo
+	{
+		FLazyBlockInfo()
+			: StartInstruction(INDEX_NONE)
+			, EndInstruction(INDEX_NONE)
+			, bProcessed(false)
+		{}
+
+		TOptional<uint32> Hash;
+		FString BlockCombinationName;
+		FRigVMOperand ExecuteStateOperand;
+		int32 StartInstruction;
+		int32 EndInstruction;
+		TArray<const FRigVMExprAST*> Expressions;
+		TArray<uint64> RunInstructionsToUpdate;
+		bool bProcessed;
+	};
+
+	TMap<uint32, TSharedPtr<FLazyBlockInfo>> LazyBlocks;
+	TArray<uint32> LazyBlocksToProcess;
+	TOptional<uint32> CurrentBlockHash;
+
 	void ReportInfo(const FString& InMessage) const;
 	void ReportWarning(const FString& InMessage) const;
 	void ReportError(const FString& InMessage) const;
diff --git a/Engine/Plugins/Runtime/RigVM/Source/RigVMEditor/Private/Editor/RigVMExecutionStackCommands.cpp b/Engine/Plugins/Runtime/RigVM/Source/RigVMEditor/Private/Editor/RigVMExecutionStackCommands.cpp
--- a/Engine/Plugins/Runtime/RigVM/Source/RigVMEditor/Private/Editor/RigVMExecutionStackCommands.cpp
+++ b/Engine/Plugins/Runtime/RigVM/Source/RigVMEditor/Private/Editor/RigVMExecutionStackCommands.cpp
@@ -8,6 +8,7 @@ void FRigVMExecutionStackCommands::RegisterCommands()
 {
 	UI_COMMAND(FocusOnSelection, "Focus On Selection", "Finds the selected operator's node in the graph.", EUserInterfaceActionType::Button, FInputChord());
 	UI_COMMAND(GoToInstruction, "Go To Instruction", "Looks for a specific instruction by index and brings it into focus.", EUserInterfaceActionType::Button, FInputChord(EKeys::G, EModifierKey::Control));
+	UI_COMMAND(SelectTargetInstructions, "Select Target Instruction(s)", "Looks up the target instructions and selects them.", EUserInterfaceActionType::Button, FInputChord(EKeys::T, EModifierKey::Control));
 }
 
 #undef LOCTEXT_NAMESPACE
diff --git a/Engine/Plugins/Runtime/RigVM/Source/RigVMEditor/Private/Widgets/SRigVMExecutionStackView.cpp b/Engine/Plugins/Runtime/RigVM/Source/RigVMEditor/Private/Widgets/SRigVMExecutionStackView.cpp
--- a/Engine/Plugins/Runtime/RigVM/Source/RigVMEditor/Private/Widgets/SRigVMExecutionStackView.cpp
+++ b/Engine/Plugins/Runtime/RigVM/Source/RigVMEditor/Private/Widgets/SRigVMExecutionStackView.cpp
@@ -428,6 +428,8 @@ void SRigVMExecutionStackView::OnSelectionChanged(TSharedPtr<FRigStackEntry> Sel
 			RigVMBlueprint->GetOrCreateController(Pair.Key)->SetNodeSelection(Pair.Value);
 		}
 	}
+
+	UpdateTargetItemHighlighting();
 }
 
 void SRigVMExecutionStackView::BindCommands()
@@ -436,6 +438,7 @@ void SRigVMExecutionStackView::BindCommands()
 	const FRigVMExecutionStackCommands& Commands = FRigVMExecutionStackCommands::Get();
 	CommandList->MapAction(Commands.FocusOnSelection, FExecuteAction::CreateSP(this, &SRigVMExecutionStackView::HandleFocusOnSelectedGraphNode));
 	CommandList->MapAction(Commands.GoToInstruction, FExecuteAction::CreateSP(this, &SRigVMExecutionStackView::HandleGoToInstruction));
+	CommandList->MapAction(Commands.SelectTargetInstructions, FExecuteAction::CreateSP(this, &SRigVMExecutionStackView::HandleSelectTargetInstructions));
 }
 
 TSharedRef<ITableRow> SRigVMExecutionStackView::MakeTableRowWidget(TSharedPtr<FRigStackEntry> InItem, const TSharedRef<STableViewBase>& OwnerTable, TWeakObjectPtr<URigVMBlueprint> InBlueprint)
@@ -577,6 +580,12 @@ void SRigVMExecutionStackView::PopulateStackView(URigVM* InVM, FRigVMExtendedExe
 						Label = TEXT("Jump To Branch");
 						break;
 					}
+					case ERigVMOpCode::RunInstructions:
+					{
+						const FRigVMRunInstructionsOp& Op = ByteCode.GetOpAt<FRigVMRunInstructionsOp>(Instructions[InstructionIndex]);
+						Label = FString::Printf(TEXT("Run Instructions %d-%d"), Op.StartInstruction, Op.EndInstruction);
+						break;
+					}
 					case ERigVMOpCode::Exit:
 					{
 						Label = TEXT("Exit");
@@ -787,6 +796,21 @@ TSharedPtr< SWidget > SRigVMExecutionStackView::CreateContextMenu()
 		MenuBuilder.BeginSection("RigStackToolsAction", LOCTEXT("ToolsAction", "Tools"));
 		MenuBuilder.AddMenuEntry(Actions.FocusOnSelection);
 		MenuBuilder.AddMenuEntry(Actions.GoToInstruction);
+
+		if(SelectedItems.ContainsByPredicate([](const TSharedPtr<FRigStackEntry>& InEntry) -> bool
+		{
+			return InEntry->OpCode == ERigVMOpCode::JumpAbsolute ||
+				InEntry->OpCode == ERigVMOpCode::JumpBackward || 
+				InEntry->OpCode == ERigVMOpCode::JumpForward ||
+				InEntry->OpCode == ERigVMOpCode::JumpAbsoluteIf || 
+				InEntry->OpCode == ERigVMOpCode::JumpBackwardIf || 
+				InEntry->OpCode == ERigVMOpCode::JumpForwardIf ||
+				InEntry->OpCode == ERigVMOpCode::RunInstructions;
+		}))
+		{
+			MenuBuilder.AddMenuEntry(Actions.SelectTargetInstructions);
+		}
+		
 		MenuBuilder.EndSection();
 	}
 
@@ -883,6 +907,128 @@ void SRigVMExecutionStackView::HandleGoToInstruction()
 	}
 }
 
+void SRigVMExecutionStackView::HandleSelectTargetInstructions()
+{
+	const TArray<TSharedPtr<FRigStackEntry>> TargetItems = GetTargetItems(TreeView->GetSelectedItems());
+	if(!TargetItems.IsEmpty())
+	{
+		TreeView->ClearSelection();
+		for(const TSharedPtr<FRigStackEntry>& TargetItem : TargetItems)
+		{
+			TreeView->SetItemSelection(TargetItem, true, ESelectInfo::Direct);
+		}
+		TreeView->RequestScrollIntoView(TargetItems[0]);
+	}
+}
+
+TArray<TSharedPtr<FRigStackEntry>> SRigVMExecutionStackView::GetTargetItems(const TArray<TSharedPtr<FRigStackEntry>>& InItems) const
+{
+	TArray<TSharedPtr<FRigStackEntry>> TargetItems;
+	
+	URigVMHost* Host = RigVMEditor.Pin()->GetRigVMHost();
+	if (Host == nullptr || Host->GetVM() == nullptr)
+	{
+		return TargetItems;
+	}
+
+	const FRigVMByteCode& ByteCode = Host->GetVM()->GetByteCode();
+	const FRigVMInstructionArray Instructions = ByteCode.GetInstructions();
+
+	TArray<int32> TargetInstructionIndices;
+	const TArray<TSharedPtr<FRigStackEntry>> SelectedItems = TreeView->GetSelectedItems();
+	for(const TSharedPtr<FRigStackEntry>& SelectedItem : SelectedItems)
+	{
+		if(!Instructions.IsValidIndex(SelectedItem->InstructionIndex))
+		{
+			continue;
+		}
+		
+		const FRigVMInstruction& Instruction = Instructions[SelectedItem->InstructionIndex];
+		switch(SelectedItem->OpCode)
+		{
+			case ERigVMOpCode::JumpAbsolute:
+			{
+				const FRigVMJumpOp& Op = ByteCode.GetOpAt<FRigVMJumpOp>(Instruction);
+				TargetInstructionIndices.AddUnique(Op.InstructionIndex);
+				break;
+			}
+			case ERigVMOpCode::JumpAbsoluteIf:
+			{
+				const FRigVMJumpIfOp& Op = ByteCode.GetOpAt<FRigVMJumpIfOp>(Instruction);
+				TargetInstructionIndices.AddUnique(Op.InstructionIndex);
+				break;
+			}
+			case ERigVMOpCode::JumpForward:
+			{
+				const FRigVMJumpOp& Op = ByteCode.GetOpAt<FRigVMJumpOp>(Instruction);
+				TargetInstructionIndices.AddUnique(SelectedItem->InstructionIndex + Op.InstructionIndex);
+				break;
+			}
+			case ERigVMOpCode::JumpForwardIf:
+			{
+				const FRigVMJumpIfOp& Op = ByteCode.GetOpAt<FRigVMJumpIfOp>(Instruction);
+				TargetInstructionIndices.AddUnique(SelectedItem->InstructionIndex + Op.InstructionIndex);
+				break;
+			}
+			case ERigVMOpCode::JumpBackward:
+			{
+				const FRigVMJumpOp& Op = ByteCode.GetOpAt<FRigVMJumpOp>(Instruction);
+				TargetInstructionIndices.AddUnique(SelectedItem->InstructionIndex - Op.InstructionIndex);
+				break;
+			}
+			case ERigVMOpCode::JumpBackwardIf:
+			{
+				const FRigVMJumpIfOp& Op = ByteCode.GetOpAt<FRigVMJumpIfOp>(Instruction);
+				TargetInstructionIndices.AddUnique(SelectedItem->InstructionIndex - Op.InstructionIndex);
+				break;
+			}
+			case ERigVMOpCode::RunInstructions:
+			{
+				const FRigVMRunInstructionsOp& Op = ByteCode.GetOpAt<FRigVMRunInstructionsOp>(Instruction);
+				for(int32 Index = Op.StartInstruction; Index <= Op.EndInstruction; Index++)
+				{
+					TargetInstructionIndices.AddUnique(Index);
+				}
+				break;
+			}
+			default:
+			{
+				break;
+			}
+		}
+	}
+
+	TMap<int32, TSharedPtr<FRigStackEntry>> InstructionToEntry;
+	for(const TSharedPtr<FRigStackEntry>& Entry : Operators)
+	{
+		InstructionToEntry.Add(Entry->InstructionIndex, Entry);
+	}
+
+	for(const int32 TargetInstructionIndex : TargetInstructionIndices)
+	{
+		if(const TSharedPtr<FRigStackEntry>* EntryPtr = InstructionToEntry.Find(TargetInstructionIndex))
+		{
+			TargetItems.Add(*EntryPtr);
+		}
+	}
+
+	return TargetItems;
+}
+
+void SRigVMExecutionStackView::UpdateTargetItemHighlighting()
+{
+	TreeView->ClearHighlightedItems();
+	
+	const TArray<TSharedPtr<FRigStackEntry>> TargetItems = GetTargetItems(TreeView->GetSelectedItems());
+	for(const TSharedPtr<FRigStackEntry>& TargetItem : TargetItems)
+	{
+		if(!TreeView->IsItemSelected(TargetItem))
+		{
+			TreeView->SetItemHighlighted(TargetItem, true);
+		}
+	}
+}
+
 void SRigVMExecutionStackView::OnVMCompiled(UObject* InCompiledObject, URigVM* InCompiledVM, FRigVMExtendedExecuteContext& InVMContext)
 {
 	RefreshTreeView(InCompiledVM, &InVMContext);
@@ -976,6 +1122,8 @@ void SRigVMExecutionStackView::HandleModifiedEvent(ERigVMGraphNotifType InNotifT
 				TreeView->SetItemSelection(SelectedItems, true, ESelectInfo::Direct);
 				TreeView->RequestScrollIntoView(SelectedItems[0]);
 			}
+
+			UpdateTargetItemHighlighting();
 			break;
 		}
 		default:
diff --git a/Engine/Plugins/Runtime/RigVM/Source/RigVMEditor/Public/Editor/RigVMExecutionStackCommands.h b/Engine/Plugins/Runtime/RigVM/Source/RigVMEditor/Public/Editor/RigVMExecutionStackCommands.h
--- a/Engine/Plugins/Runtime/RigVM/Source/RigVMEditor/Public/Editor/RigVMExecutionStackCommands.h
+++ b/Engine/Plugins/Runtime/RigVM/Source/RigVMEditor/Public/Editor/RigVMExecutionStackCommands.h
@@ -23,6 +23,9 @@ public:
 	/** Looks for a specific instruction by index and brings it into focus */
 	TSharedPtr< FUICommandInfo > GoToInstruction;
 
+	/** Selects the target instruction(s) of an instruction */
+	TSharedPtr< FUICommandInfo > SelectTargetInstructions;
+
 	/**
 	 * Initialize commands
 	 */
diff --git a/Engine/Plugins/Runtime/RigVM/Source/RigVMEditor/Public/Widgets/SRigVMExecutionStackView.h b/Engine/Plugins/Runtime/RigVM/Source/RigVMEditor/Public/Widgets/SRigVMExecutionStackView.h
--- a/Engine/Plugins/Runtime/RigVM/Source/RigVMEditor/Public/Widgets/SRigVMExecutionStackView.h
+++ b/Engine/Plugins/Runtime/RigVM/Source/RigVMEditor/Public/Widgets/SRigVMExecutionStackView.h
@@ -104,6 +104,13 @@ private:
 	/** Offers a dialog to move to a specific instruction */
 	void HandleGoToInstruction();
 
+	/** Selects the target instructions for the current selection */
+	void HandleSelectTargetInstructions();
+
+	TArray<TSharedPtr<FRigStackEntry>> GetTargetItems(const TArray<TSharedPtr<FRigStackEntry>>& InItems) const;
+
+	void UpdateTargetItemHighlighting();
+
 	void OnVMCompiled(UObject* InCompiledObject, URigVM* InCompiledVM, FRigVMExtendedExecuteContext& InVMContext);
 
 	//* Focus on the instruction when the execution is halted */
