RigVM ByteCode now stores PublicContext path name, so the VM can check at Runtime

diff --git a/Engine/Plugins/Runtime/RigVM/Source/RigVM/Private/RigVMCore/RigVM.cpp b/Engine/Plugins/Runtime/RigVM/Source/RigVM/Private/RigVMCore/RigVM.cpp
--- a/Engine/Plugins/Runtime/RigVM/Source/RigVM/Private/RigVMCore/RigVM.cpp
+++ b/Engine/Plugins/Runtime/RigVM/Source/RigVM/Private/RigVMCore/RigVM.cpp
@@ -1514,6 +1514,8 @@ ERigVMExecuteResult URigVM::ExecuteVM(FRigVMExtendedExecuteContext& Context, con
 	TGuardValue<FName> EntryNameGuard(Context.CurrentEntryName, InEntryName);
 	TGuardValue<bool> RootEntryGuard(Context.bCurrentlyRunningRootEntry, bIsRootEntry);
 
+	FRigVMByteCode& ByteCode = GetByteCode();
+
 	if(bIsRootEntry)
 	{
 		Context.CurrentExecuteResult = ERigVMExecuteResult::Succeeded;
@@ -1526,9 +1528,19 @@ ERigVMExecuteResult URigVM::ExecuteVM(FRigVMExtendedExecuteContext& Context, con
 		}
 
 		Context.CurrentVMMemory = GetInstanceMemory(Context);
+
+		if (ByteCode.HasPublicContextPathName())
+		{
+			const FString ContextPublicDataPathName = Context.GetContextPublicDataStruct() != nullptr ? Context.GetContextPublicDataStruct()->GetPathName() : FString();
+			if (ByteCode.GetPublicContextPathName().IsEmpty() || ByteCode.GetPublicContextPathName() != ContextPublicDataPathName)
+			{
+				UE_LOG(LogRigVM, Error, TEXT("Context PublicData [%s] does not match ByteCode Public Data [%s]. Likely a corrupt VM. Exiting."), *ContextPublicDataPathName, *ByteCode.GetPublicContextPathName());
+				return Context.CurrentExecuteResult = ERigVMExecuteResult::Failed;
+			}
+		}
 	}
 
-	FRigVMByteCode& ByteCode = GetByteCode();
+
 	TArray<const FRigVMFunction*>& Functions = GetFunctions();
 	TArray<const FRigVMDispatchFactory*>& Factories = GetFactories();
 
diff --git a/Engine/Plugins/Runtime/RigVM/Source/RigVM/Private/RigVMCore/RigVMByteCode.cpp b/Engine/Plugins/Runtime/RigVM/Source/RigVM/Private/RigVMCore/RigVMByteCode.cpp
--- a/Engine/Plugins/Runtime/RigVM/Source/RigVM/Private/RigVMCore/RigVMByteCode.cpp
+++ b/Engine/Plugins/Runtime/RigVM/Source/RigVM/Private/RigVMCore/RigVMByteCode.cpp
@@ -267,6 +267,7 @@ void FRigVMByteCode::Serialize(FArchive& Ar)
 	Ar.UsingCustomVersion(FAnimObjectVersion::GUID);
 	Ar.UsingCustomVersion(FFortniteMainBranchObjectVersion::GUID);
 	Ar.UsingCustomVersion(FUE5MainStreamObjectVersion::GUID);
+	Ar.UsingCustomVersion(FRigVMObjectVersion::GUID);
 
 	if (Ar.CustomVer(FAnimObjectVersion::GUID) < FAnimObjectVersion::StoreMarkerNamesOnSkeleton)
 	{
@@ -287,7 +288,7 @@ void FRigVMByteCode::Serialize(FArchive& Ar)
 	}
 }
 
-void FRigVMByteCode::Save(FArchive& Ar) const
+void FRigVMByteCode::Save(FArchive& Ar)
 {
 	FRigVMInstructionArray Instructions;
 
@@ -455,6 +456,8 @@ void FRigVMByteCode::Save(FArchive& Ar) const
 
 	TArray<FRigVMBranchInfo> TempBranchInfos = BranchInfos;
 	Ar << TempBranchInfos;
+
+	Ar << PublicContextPathName;
 }
 
 void FRigVMByteCode::Load(FArchive& Ar)
@@ -659,6 +662,12 @@ void FRigVMByteCode::Load(FArchive& Ar)
 		BranchInfos.Reset();
 		BranchInfoLookup.Reset();
 	}
+
+	if (Ar.CustomVer(FRigVMObjectVersion::GUID) >= FRigVMObjectVersion::VMBytecodeStorePublicContextPath)
+	{
+		Ar << PublicContextPathName;
+		bHasPublicContextPathName = true;
+	}
 }
 
 void FRigVMByteCode::Reset()
@@ -669,6 +678,7 @@ void FRigVMByteCode::Reset()
 	Entries.Reset();
 	BranchInfos.Reset();
 	BranchInfoLookup.Reset();
+	PublicContextPathName.Reset();
 
 #if WITH_EDITORONLY_DATA
 	SubjectPerInstruction.Reset();
@@ -689,6 +699,7 @@ void FRigVMByteCode::Empty()
 	bByteCodeIsAligned = false;
 	NumInstructions = 0;
 	Entries.Empty();
+	PublicContextPathName.Empty();
 
 #if WITH_EDITORONLY_DATA
 	SubjectPerInstruction.Empty();
diff --git a/Engine/Plugins/Runtime/RigVM/Source/RigVM/Public/RigVMCore/RigVMByteCode.h b/Engine/Plugins/Runtime/RigVM/Source/RigVM/Public/RigVMCore/RigVMByteCode.h
--- a/Engine/Plugins/Runtime/RigVM/Source/RigVM/Public/RigVMCore/RigVMByteCode.h
+++ b/Engine/Plugins/Runtime/RigVM/Source/RigVM/Public/RigVMCore/RigVMByteCode.h
@@ -1131,7 +1131,7 @@ public:
 	FRigVMByteCode();
 
 	void Serialize(FArchive& Ar);
-	void Save(FArchive& Ar) const;
+	void Save(FArchive& Ar);
 	void Load(FArchive& Ar);
 	friend FArchive& operator<<(FArchive& Ar, FRigVMByteCode& P)
 	{
@@ -1339,6 +1339,22 @@ public:
 
 	FString DumpToText() const;
 
+	bool HasPublicContextPathName() const
+	{
+		return bHasPublicContextPathName;
+	}
+
+	const FString& GetPublicContextPathName() const
+	{
+		return PublicContextPathName;
+	}
+
+	void SetPublicContextPathName(const FString& InPublicContextPathName)
+	{
+		PublicContextPathName = InPublicContextPathName;
+		bHasPublicContextPathName = true;
+	}
+
 #if WITH_EDITOR
 
 	// returns the subject which was used to inject a given instruction
@@ -1467,11 +1483,16 @@ private:
 	UPROPERTY()
 	TArray<FRigVMPredicateBranch> PredicateBranches;
 
+	UPROPERTY()
+	FString PublicContextPathName;
+
 	const FRigVMBranchInfo* GetBranchInfo(const FRigVMBranchInfoKey& InBranchInfoKey) const;
 	mutable TMap<FRigVMBranchInfoKey, const FRigVMBranchInfo*> BranchInfoLookup;
 
 	// if this is set to true the stored bytecode is aligned / padded
 	bool bByteCodeIsAligned;
+	// If the serialization has loaded a PublicContextPathName, so we check on new versions and skip check on older
+	bool bHasPublicContextPathName = false;
 
 	static TArray<int32> EmptyInstructionIndices;
 
diff --git a/Engine/Plugins/Runtime/RigVM/Source/RigVM/Public/RigVMObjectVersion.h b/Engine/Plugins/Runtime/RigVM/Source/RigVM/Public/RigVMObjectVersion.h
--- a/Engine/Plugins/Runtime/RigVM/Source/RigVM/Public/RigVMObjectVersion.h
+++ b/Engine/Plugins/Runtime/RigVM/Source/RigVM/Public/RigVMObjectVersion.h
@@ -33,6 +33,10 @@ struct RIGVM_API FRigVMObjectVersion
 		// VM Memory Storage Defaults generated at VM
 		VMMemoryStorageDefaultsGeneratedAtVM,
 
+		// VM Bytecode Stores the Public Context Path
+		VMBytecodeStorePublicContextPath,
+
+
 		// -----<new versions can be added above this line>-------------------------------------------------
 		VersionPlusOne,
 		LatestVersion = VersionPlusOne - 1
diff --git a/Engine/Plugins/Runtime/RigVM/Source/RigVMDeveloper/Private/RigVMCompiler/RigVMCompiler.cpp b/Engine/Plugins/Runtime/RigVM/Source/RigVMDeveloper/Private/RigVMCompiler/RigVMCompiler.cpp
--- a/Engine/Plugins/Runtime/RigVM/Source/RigVMDeveloper/Private/RigVMCompiler/RigVMCompiler.cpp
+++ b/Engine/Plugins/Runtime/RigVM/Source/RigVMDeveloper/Private/RigVMCompiler/RigVMCompiler.cpp
@@ -1470,6 +1470,7 @@ bool URigVMCompiler::Compile(const FRigVMCompileSettings& InSettings, TArray<URi
 		WorkData.ReportInfof(TEXT("Total Compilation time %f\n"), CompilationTime*1000);
 
 		WorkData.VM->SetVMHash(WorkData.VM->ComputeVMHash());
+		WorkData.VM->GetByteCode().SetPublicContextPathName((WorkData.Context&& WorkData.Context->GetContextPublicDataStruct() != nullptr) ? WorkData.Context->GetContextPublicDataStruct()->GetPathName() : FString());
 	}
 
 	return true;
