Rework RigVM function registry locks, in order to avoid a deadlock when a VM is initializing in a thread while a different worker thread registers a function.

diff --git a/Engine/Plugins/Runtime/RigVM/Source/RigVM/Private/RigVMCore/RigVMDispatchFactory.cpp b/Engine/Plugins/Runtime/RigVM/Source/RigVM/Private/RigVMCore/RigVMDispatchFactory.cpp
--- a/Engine/Plugins/Runtime/RigVM/Source/RigVM/Private/RigVMCore/RigVMDispatchFactory.cpp
+++ b/Engine/Plugins/Runtime/RigVM/Source/RigVM/Private/RigVMCore/RigVMDispatchFactory.cpp
@@ -232,7 +232,7 @@ const TArray<FRigVMExecuteArgument>& FRigVMDispatchFactory::GetExecuteArguments_
 
 FRigVMFunctionPtr FRigVMDispatchFactory::GetOrCreateDispatchFunction(const FRigVMTemplateTypeMap& InTypes) const
 {
-	FScopeLock GetTemplateScopeLock(&FRigVMRegistry::GetDispatchFunctionMutex);
+	FScopeLock DispatchFunctionScopeLock(&FRigVMRegistry::DispatchFunctionMutex);
 
 	const FString PermutationName = GetPermutationNameImpl(InTypes);
 	if(const FRigVMFunction* ExistingFunction = FRigVMRegistry::Get().FindFunction(*PermutationName))
@@ -245,7 +245,7 @@ FRigVMFunctionPtr FRigVMDispatchFactory::GetOrCreateDispatchFunction(const FRigV
 
 FRigVMFunctionPtr FRigVMDispatchFactory::CreateDispatchFunction(const FRigVMTemplateTypeMap& InTypes) const
 {
-	FScopeLock GetTemplateScopeLock(&FRigVMRegistry::GetDispatchFunctionMutex);
+	FScopeLock DispatchFunctionScopeLock(&FRigVMRegistry::DispatchFunctionMutex);
 	return CreateDispatchFunction_NoLock(InTypes);
 }
 
@@ -256,7 +256,7 @@ FRigVMFunctionPtr FRigVMDispatchFactory::CreateDispatchFunction_NoLock(const FRi
 
 TArray<FRigVMFunction> FRigVMDispatchFactory::CreateDispatchPredicates(const FRigVMTemplateTypeMap& InTypes) const
 {
-	FScopeLock GetTemplateScopeLock(&FRigVMRegistry::GetDispatchPredicatesMutex);
+	FScopeLock DispatchPredicatesScopeLock(&FRigVMRegistry::DispatchPredicatesMutex);
 	return CreateDispatchPredicates_NoLock(InTypes);
 }
 
diff --git a/Engine/Plugins/Runtime/RigVM/Source/RigVM/Private/RigVMCore/RigVMRegistry.cpp b/Engine/Plugins/Runtime/RigVM/Source/RigVM/Private/RigVMCore/RigVMRegistry.cpp
--- a/Engine/Plugins/Runtime/RigVM/Source/RigVM/Private/RigVMCore/RigVMRegistry.cpp
+++ b/Engine/Plugins/Runtime/RigVM/Source/RigVM/Private/RigVMCore/RigVMRegistry.cpp
@@ -16,16 +16,15 @@
 #include "RigVMFunctions/RigVMDispatch_Core.h"
 
 const FName FRigVMRegistry::TemplateNameMetaName = TEXT("TemplateName");
+
 FCriticalSection FRigVMRegistry::RefreshTypesMutex;
-FCriticalSection FRigVMRegistry::RegisterFunctionMutex;
-FCriticalSection FRigVMRegistry::RegisterTemplateMutex;
-FCriticalSection FRigVMRegistry::RegisterFactoryMutex;
-FCriticalSection FRigVMRegistry::FindFunctionMutex;
-FCriticalSection FRigVMRegistry::FindTemplateMutex;
-FCriticalSection FRigVMRegistry::FindFactoryMutex;
-FCriticalSection FRigVMRegistry::GetDispatchFunctionMutex;
-FCriticalSection FRigVMRegistry::GetDispatchPredicatesMutex;
-FCriticalSection FRigVMRegistry::GetPermutationMutex;
+
+FCriticalSection FRigVMRegistry::FunctionRegistryMutex;
+FCriticalSection FRigVMRegistry::FactoryRegistryMutex;
+FCriticalSection FRigVMRegistry::TemplateRegistryMutex;
+
+FCriticalSection FRigVMRegistry::DispatchFunctionMutex;
+FCriticalSection FRigVMRegistry::DispatchPredicatesMutex;
 
 
 // When the object system has been completely loaded, load in all the engine types that we haven't registered already in InitializeIfNeeded 
@@ -1198,9 +1197,9 @@ bool FRigVMRegistry::IsAllowedType(const UClass* InClass) const
 
 void FRigVMRegistry::Register(const TCHAR* InName, FRigVMFunctionPtr InFunctionPtr, UScriptStruct* InStruct, const TArray<FRigVMFunctionArgument>& InArguments)
 {
-	FScopeLock RegisterFunctionScopeLock(&RegisterFunctionMutex);
+	FScopeLock FunctionRegistryScopeLock(&FunctionRegistryMutex);
 	
-	if (FindFunction(InName) != nullptr)
+	if (FindFunction_NoLock(InName) != nullptr)
 	{
 		return;
 	}
@@ -1280,7 +1279,6 @@ void FRigVMRegistry::Register(const TCHAR* InName, FRigVMFunctionPtr InFunctionP
 
 const FRigVMDispatchFactory* FRigVMRegistry::RegisterFactory(UScriptStruct* InFactoryStruct)
 {
-	FScopeLock RegisterFactoryScopeLock(&RegisterFactoryMutex);
 
 	check(InFactoryStruct);
 	check(InFactoryStruct != FRigVMDispatchFactory::StaticStruct());
@@ -1288,6 +1286,9 @@ const FRigVMDispatchFactory* FRigVMRegistry::RegisterFactory(UScriptStruct* InFa
 
 	// ensure to register factories only once
 	const FRigVMDispatchFactory* ExistingFactory = nullptr;
+
+	FScopeLock FactoryRegistryScopeLock(&FactoryRegistryMutex);
+
 	const bool bFactoryAlreadyRegistered = Factories.ContainsByPredicate([InFactoryStruct, &ExistingFactory](const FRigVMDispatchFactory* Factory)
 	{
 		if(Factory->GetScriptStruct() == InFactoryStruct)
@@ -1378,7 +1379,7 @@ void FRigVMRegistry::RegisterObjectTypes(TConstArrayView<TPair<UClass*, ERegiste
 
 const FRigVMFunction* FRigVMRegistry::FindFunction(const TCHAR* InName, const FRigVMUserDefinedTypeResolver& InTypeResolver) const
 {
-	FScopeLock FindFunctionScopeLock(&FindFunctionMutex);
+	FScopeLock FunctionRegistryScopeLock(&FunctionRegistryMutex);
 	return FindFunction_NoLock(InName, InTypeResolver);
 }
 
@@ -1406,7 +1407,7 @@ const FRigVMFunction* FRigVMRegistry::FindFunction_NoLock(const TCHAR* InName, c
 					const int32 PermutationIndex = Template->FindPermutation(ArgumentTypes);
 					if(PermutationIndex != INDEX_NONE)
 					{
-						return ((FRigVMTemplate*)Template)->GetOrCreatePermutation(PermutationIndex);
+						return ((FRigVMTemplate*)Template)->GetOrCreatePermutation_NoLock(PermutationIndex);
 					}
 				}
 			}
@@ -1478,7 +1479,7 @@ const FRigVMTemplate* FRigVMRegistry::FindTemplate(const FName& InNotation, bool
 		return nullptr;
 	}
 
-	FScopeLock FindTemplateScopeLock(&FindTemplateMutex);
+	FScopeLock TemplateRegistryScopeLock(&TemplateRegistryMutex);
 	return FindTemplate_NoLock(InNotation, bIncludeDeprecated);
 }
 
@@ -1601,7 +1602,7 @@ const TChunkedArray<FRigVMTemplate>& FRigVMRegistry::GetTemplates() const
 
 const FRigVMTemplate* FRigVMRegistry::GetOrAddTemplateFromArguments(const FName& InName, const TArray<FRigVMTemplateArgumentInfo>& InInfos, const FRigVMTemplateDelegates& InDelegates)
 {
-	FScopeLock RegisterTemplateScopeLock(&RegisterTemplateMutex);
+	FScopeLock TemplateRegistryScopeLock(&TemplateRegistryMutex);
 	
 	// avoid reentry in FindTemplate. try to find an existing
 	// template only if we are not yet in ::FindTemplate.
@@ -1616,7 +1617,7 @@ const FRigVMTemplate* FRigVMRegistry::GetOrAddTemplateFromArguments(const FName&
 
 const FRigVMTemplate* FRigVMRegistry::AddTemplateFromArguments(const FName& InName, const TArray<FRigVMTemplateArgumentInfo>& InInfos, const FRigVMTemplateDelegates& InDelegates)
 {
-	FScopeLock RegisterTemplateScopeLock(&RegisterTemplateMutex);
+	FScopeLock TemplateRegistryScopeLock(&TemplateRegistryMutex);
 	return AddTemplateFromArguments_NoLock(InName, InInfos, InDelegates);
 }
 
@@ -1743,7 +1744,7 @@ const FRigVMTemplate* FRigVMRegistry::AddTemplateFromArguments_NoLock(const FNam
 
 FRigVMDispatchFactory* FRigVMRegistry::FindDispatchFactory(const FName& InFactoryName) const
 {
-	FScopeLock FindFactoryScopeLock(&FindFactoryMutex);
+	FScopeLock FactoryRegistryScopeLock(&FactoryRegistryMutex);
 	return FindDispatchFactory_NoLock(InFactoryName);
 }
 
diff --git a/Engine/Plugins/Runtime/RigVM/Source/RigVM/Private/RigVMCore/RigVMTemplate.cpp b/Engine/Plugins/Runtime/RigVM/Source/RigVM/Private/RigVMCore/RigVMTemplate.cpp
--- a/Engine/Plugins/Runtime/RigVM/Source/RigVM/Private/RigVMCore/RigVMTemplate.cpp
+++ b/Engine/Plugins/Runtime/RigVM/Source/RigVM/Private/RigVMCore/RigVMTemplate.cpp
@@ -1317,7 +1317,7 @@ const FRigVMFunction* FRigVMTemplate::GetPrimaryPermutation() const
 
 const FRigVMFunction* FRigVMTemplate::GetPermutation(int32 InIndex) const
 {
-	FScopeLock FindPermutationScopeLock(&FRigVMRegistry::GetPermutationMutex);
+	FScopeLock FunctionRegistryScopeLock(&FRigVMRegistry::FunctionRegistryMutex);
 	return GetPermutation_NoLock(InIndex);
 }
 
@@ -1334,8 +1334,13 @@ const FRigVMFunction* FRigVMTemplate::GetPermutation_NoLock(int32 InIndex) const
 
 const FRigVMFunction* FRigVMTemplate::GetOrCreatePermutation(int32 InIndex)
 {
-	FScopeLock FindPermutationScopeLock(&FRigVMRegistry::GetPermutationMutex);
+	FScopeLock FunctionRegistryScopeLock(&FRigVMRegistry::FunctionRegistryMutex);
 
+	return GetOrCreatePermutation_NoLock(InIndex);
+}
+
+const FRigVMFunction* FRigVMTemplate::GetOrCreatePermutation_NoLock(int32 InIndex)
+{
 	if(const FRigVMFunction* Function = GetPermutation_NoLock(InIndex))
 	{
 		return Function;
@@ -1343,7 +1348,6 @@ const FRigVMFunction* FRigVMTemplate::GetOrCreatePermutation(int32 InIndex)
 
 	if(Permutations[InIndex] == INDEX_NONE && UsesDispatch())
 	{
-		FScopeLock RegisterFunctionScopeLock(&FRigVMRegistry::RegisterFunctionMutex);
 		FRigVMRegistry& Registry = FRigVMRegistry::Get();
 		
 		FTypeMap Types;
diff --git a/Engine/Plugins/Runtime/RigVM/Source/RigVM/Public/RigVMCore/RigVMRegistry.h b/Engine/Plugins/Runtime/RigVM/Source/RigVM/Public/RigVMCore/RigVMRegistry.h
--- a/Engine/Plugins/Runtime/RigVM/Source/RigVM/Public/RigVMCore/RigVMRegistry.h
+++ b/Engine/Plugins/Runtime/RigVM/Source/RigVM/Public/RigVMCore/RigVMRegistry.h
@@ -405,16 +405,13 @@ private:
 	bool bIsRefreshingEngineTypes;
 
 	static FCriticalSection RefreshTypesMutex;
-	static FCriticalSection RegisterFunctionMutex;
-	static FCriticalSection RegisterTemplateMutex;
-	static FCriticalSection RegisterFactoryMutex;
-	static FCriticalSection FindFunctionMutex;
-	static FCriticalSection FindTemplateMutex;
-	static FCriticalSection FindFactoryMutex;
-	static FCriticalSection GetDispatchFunctionMutex;
-	static FCriticalSection GetDispatchPredicatesMutex;
-	static FCriticalSection GetPermutationMutex;
+	static FCriticalSection FunctionRegistryMutex;
+	static FCriticalSection FactoryRegistryMutex;
+	static FCriticalSection TemplateRegistryMutex;
 
+	static FCriticalSection DispatchFunctionMutex;
+	static FCriticalSection DispatchPredicatesMutex;
+	
 	friend struct FRigVMStruct;
 	friend struct FRigVMTemplate;
 	friend struct FRigVMTemplateArgument;
diff --git a/Engine/Plugins/Runtime/RigVM/Source/RigVM/Public/RigVMCore/RigVMTemplate.h b/Engine/Plugins/Runtime/RigVM/Source/RigVM/Public/RigVMCore/RigVMTemplate.h
--- a/Engine/Plugins/Runtime/RigVM/Source/RigVM/Public/RigVMCore/RigVMTemplate.h
+++ b/Engine/Plugins/Runtime/RigVM/Source/RigVM/Public/RigVMCore/RigVMTemplate.h
@@ -517,6 +517,7 @@ private:
 	const TArray<FRigVMExecuteArgument>& GetExecuteArguments(const FRigVMDispatchContext& InContext) const;
 
 	const FRigVMFunction* GetPermutation_NoLock(int32 InIndex) const;
+	const FRigVMFunction* GetOrCreatePermutation_NoLock(int32 InIndex);
 
 	int32 Index;
 	FName Notation;
