RigVM: Patch incorrect state of bLazyPin for old assets

#rnx

diff --git a/Engine/Plugins/Runtime/RigVM/Source/RigVM/Private/RigVMCore/RigVM.cpp b/Engine/Plugins/Runtime/RigVM/Source/RigVM/Private/RigVMCore/RigVM.cpp
--- a/Engine/Plugins/Runtime/RigVM/Source/RigVM/Private/RigVMCore/RigVM.cpp
+++ b/Engine/Plugins/Runtime/RigVM/Source/RigVM/Private/RigVMCore/RigVM.cpp
@@ -1685,9 +1685,17 @@ ERigVMExecuteResult URigVM::ExecuteInstructions(FRigVMExtendedExecuteContext& Co
 {
 	// make we are already executing this VM
 	check(!Context.CurrentVMMemory.IsEmpty());
-	
+
 	FRigVMExecuteContext& ContextPublicData = Context.GetPublicData<>();
 	TGuardValue<uint16> InstructionIndexGuard(ContextPublicData.InstructionIndex, (uint16)InFirstInstruction);
+#if WITH_EDITOR
+	FInstructionBracketGuard InstructionBracket(Context, InFirstInstruction, InLastInstruction);
+	if(!InstructionBracket.ErrorMessage.IsEmpty())
+	{
+		Context.GetPublicDataSafe<>().Log(EMessageSeverity::Error, InstructionBracket.ErrorMessage);
+		return ERigVMExecuteResult::Failed;
+	}
+#endif
 	
 	FRigVMByteCode& ByteCode = GetByteCode();
 	TArray<const FRigVMFunction*>& Functions = GetFunctions();
diff --git a/Engine/Plugins/Runtime/RigVM/Source/RigVM/Public/RigVMCore/RigVM.h b/Engine/Plugins/Runtime/RigVM/Source/RigVM/Public/RigVMCore/RigVM.h
--- a/Engine/Plugins/Runtime/RigVM/Source/RigVM/Public/RigVMCore/RigVM.h
+++ b/Engine/Plugins/Runtime/RigVM/Source/RigVM/Public/RigVMCore/RigVM.h
@@ -1092,7 +1092,37 @@ protected:
 
 		TArray<int32>& Stack;
 	};
-	
+
+#if WITH_EDITOR
+	struct FInstructionBracketGuard
+	{
+	public:
+		FInstructionBracketGuard(FRigVMExtendedExecuteContext& InOutContext, int32 InFirstInstruction, int32 InLastInstruction)
+		: Context(InOutContext)
+		{
+			const TTuple<int32, int32> Tuple(InFirstInstruction, InLastInstruction);
+			if(Context.InstructionBrackets.Contains(Tuple))
+			{
+				static constexpr TCHAR Format[] = TEXT("Re-Entry of Instructions %d - %d.");
+				ErrorMessage = FString::Printf(Format, Tuple.Get<0>(), Tuple.Get<1>());
+			}
+			Context.InstructionBrackets.Add(Tuple);
+		}
+
+		~FInstructionBracketGuard()
+		{
+			Context.InstructionBrackets.Pop();
+		}
+		
+	private:
+
+		FRigVMExtendedExecuteContext& Context;
+		FString ErrorMessage;
+
+		friend class URigVM;
+	};
+#endif
+
 private:
 
 	UPROPERTY()
@@ -1125,3 +1155,4 @@ protected:
 	friend struct FRigVMCompilerWorkData;
 	friend struct FRigVMCodeGenerator;
 };
+
diff --git a/Engine/Plugins/Runtime/RigVM/Source/RigVM/Public/RigVMCore/RigVMExecuteContext.h b/Engine/Plugins/Runtime/RigVM/Source/RigVM/Public/RigVMCore/RigVMExecuteContext.h
--- a/Engine/Plugins/Runtime/RigVM/Source/RigVM/Public/RigVMCore/RigVMExecuteContext.h
+++ b/Engine/Plugins/Runtime/RigVM/Source/RigVM/Public/RigVMCore/RigVMExecuteContext.h
@@ -810,5 +810,7 @@ struct RIGVM_API FRigVMExtendedExecuteContext
 		return ProfilingInfo;
 	}
 
+	TArray<TTuple<int32, int32>> InstructionBrackets;
+
 #endif // WITH_EDITOR
 };
diff --git a/Engine/Plugins/Runtime/RigVM/Source/RigVMDeveloper/Private/RigVMCompiler/RigVMAST.cpp b/Engine/Plugins/Runtime/RigVM/Source/RigVMDeveloper/Private/RigVMCompiler/RigVMAST.cpp
--- a/Engine/Plugins/Runtime/RigVM/Source/RigVMDeveloper/Private/RigVMCompiler/RigVMAST.cpp
+++ b/Engine/Plugins/Runtime/RigVM/Source/RigVMDeveloper/Private/RigVMCompiler/RigVMAST.cpp
@@ -273,45 +273,56 @@ TOptional<uint32> FRigVMExprAST::GetBlockCombinationHash() const
 	// only consider cached / computed values for lazy blocks
 	if (IsA(FRigVMExprAST::CachedValue))
 	{
-		// collect all of the blocks this expression is in
-		const FRigVMExprAST::FRigVMBlockArray Blocks = GetBlocks(true);
-		if(Blocks.Num() > 1)
+		// find the first node expr child
+		FRigVMExprAST* const* NodeChildPtr = Children.FindByPredicate([](const FRigVMExprAST* Child)
 		{
-			// compute a hash for the blocks
-			uint32 ComputedHash = 0;
-			for(const FRigVMBlockExprAST* Block : Blocks)
+			return Child->IsNode();
+		});
+
+		if(NodeChildPtr)
+		{
+			const FRigVMExprAST* NodeExpression = *NodeChildPtr;
+			
+			// collect all of the blocks this expression is in
+			const FRigVMExprAST::FRigVMBlockArray Blocks = NodeExpression->GetBlocks(true);
+			if(Blocks.Num() > 1)
 			{
-				ComputedHash = HashCombine(ComputedHash, GetTypeHash(Block->Index));
-			}
+				// compute a hash for the blocks
+				uint32 ComputedHash = 0;
+				for(const FRigVMBlockExprAST* Block : Blocks)
+				{
+					ComputedHash = HashCombine(ComputedHash, GetTypeHash(Block->Index));
+				}
 
-			BlockCombinationHash = TOptional<uint32>(ComputedHash);
+				BlockCombinationHash = TOptional<uint32>(ComputedHash);
 
-			if(const FRigVMParserAST* Parser = GetParser())
-			{
-				// compute a unique name for the block combination
-				if(!Parser->BlockCombinationHashToName.Contains(ComputedHash))
+				if(const FRigVMParserAST* Parser = GetParser())
 				{
-					FString BlockCombinationName;
-					for(const FRigVMBlockExprAST* Block : Blocks)
+					// compute a unique name for the block combination
+					if(!Parser->BlockCombinationHashToName.Contains(ComputedHash))
 					{
-						static constexpr TCHAR BlockNameFormat[] = TEXT("%s%d");
-						const FString BlockNameAndIndex = FString::Printf(BlockNameFormat, *Block->GetName().ToString(), Block->Index);
-
-						if(BlockCombinationName.IsEmpty())
+						FString BlockCombinationName;
+						for(const FRigVMBlockExprAST* Block : Blocks)
 						{
-							BlockCombinationName = BlockNameAndIndex;
-						}
-						else
-						{
-							static constexpr TCHAR JoinFormat[] = TEXT("%s_%s");
-							BlockCombinationName = FString::Printf(JoinFormat, *BlockCombinationName, *BlockNameAndIndex);
+							static constexpr TCHAR BlockNameFormat[] = TEXT("%s%d");
+							const FString BlockNameAndIndex = FString::Printf(BlockNameFormat, *Block->GetName().ToString(), Block->Index);
+
+							if(BlockCombinationName.IsEmpty())
+							{
+								BlockCombinationName = BlockNameAndIndex;
+							}
+							else
+							{
+								static constexpr TCHAR JoinFormat[] = TEXT("%s_%s");
+								BlockCombinationName = FString::Printf(JoinFormat, *BlockCombinationName, *BlockNameAndIndex);
+							}
 						}
+						Parser->BlockCombinationHashToName.Add(ComputedHash, BlockCombinationName);
 					}
-					Parser->BlockCombinationHashToName.Add(ComputedHash, BlockCombinationName);
 				}
-			}
 
-			return BlockCombinationHash.GetValue();
+				return BlockCombinationHash.GetValue();
+			}
 		}
 	}
 
diff --git a/Engine/Plugins/Runtime/RigVM/Source/RigVMDeveloper/Private/RigVMModel/RigVMClient.cpp b/Engine/Plugins/Runtime/RigVM/Source/RigVMDeveloper/Private/RigVMModel/RigVMClient.cpp
--- a/Engine/Plugins/Runtime/RigVM/Source/RigVMDeveloper/Private/RigVMModel/RigVMClient.cpp
+++ b/Engine/Plugins/Runtime/RigVM/Source/RigVMDeveloper/Private/RigVMModel/RigVMClient.cpp
@@ -900,6 +900,7 @@ FRigVMClientPatchResult FRigVMClient::PatchModelsOnLoad()
 		Result.Merge(Controller->PatchReduceArrayFloatDoubleConvertsionsOnLoad());
 		Result.Merge(Controller->PatchInvalidLinksOnWildcards());
 		Result.Merge(Controller->PatchExecutePins());
+		Result.Merge(Controller->PatchLazyPins());
 
 		if (URigVMCollapseNode* CollapseNode = Model->GetTypedOuter<URigVMCollapseNode>())
 		{
diff --git a/Engine/Plugins/Runtime/RigVM/Source/RigVMDeveloper/Private/RigVMModel/RigVMController.cpp b/Engine/Plugins/Runtime/RigVM/Source/RigVMDeveloper/Private/RigVMModel/RigVMController.cpp
--- a/Engine/Plugins/Runtime/RigVM/Source/RigVMDeveloper/Private/RigVMModel/RigVMController.cpp
+++ b/Engine/Plugins/Runtime/RigVM/Source/RigVMDeveloper/Private/RigVMModel/RigVMController.cpp
@@ -19940,6 +19940,32 @@ FRigVMClientPatchResult URigVMController::PatchExecutePins()
 	return Result;
 }
 
+FRigVMClientPatchResult URigVMController::PatchLazyPins()
+{
+	FRigVMClientPatchResult Result;
+	if (const URigVMGraph* Graph = GetGraph())
+	{
+		for (const URigVMNode* Node : Graph->GetNodes())
+		{
+			for(URigVMPin* Pin : Node->GetPins())
+			{
+				const bool bShouldBeLazy = Node->ShouldInputPinComputeLazily(Pin);
+				if(Pin->bIsLazy != bShouldBeLazy)
+				{
+					TArray<URigVMPin*> AllPins = {Pin};
+					for(int32 Index = 0; Index < AllPins.Num(); Index++)
+					{
+						AllPins[Index]->bIsLazy = true;
+						AllPins.Append(AllPins[Index]->GetSubPins());
+					}
+					Result.bChangedContent = true;
+				}
+			}
+		}
+	}
+	return Result;
+}
+
 void URigVMController::PostDuplicateHost(const FString& InOldPathName, const FString& InNewPathName)
 {
 	if (const URigVMGraph* Graph = GetGraph())
diff --git a/Engine/Plugins/Runtime/RigVM/Source/RigVMDeveloper/Public/RigVMModel/RigVMController.h b/Engine/Plugins/Runtime/RigVM/Source/RigVMDeveloper/Public/RigVMModel/RigVMController.h
--- a/Engine/Plugins/Runtime/RigVM/Source/RigVMDeveloper/Public/RigVMModel/RigVMController.h
+++ b/Engine/Plugins/Runtime/RigVM/Source/RigVMDeveloper/Public/RigVMModel/RigVMController.h
@@ -1339,6 +1339,7 @@ protected:
 	FRigVMClientPatchResult PatchInvalidLinksOnWildcards();
 	FRigVMClientPatchResult PatchFunctionsWithInvalidReturnPaths();
 	FRigVMClientPatchResult PatchExecutePins();
+	FRigVMClientPatchResult PatchLazyPins();
 
 	// work to do after a duplication of the host asset
 	void PostDuplicateHost(const FString& InOldPathName, const FString& InNewPathName);
