Control Rig: Avoid showing public functions duplicated in the context menu
#rnx

diff --git a/Engine/Plugins/Runtime/RigVM/Source/RigVMEditor/Private/RigVMEditorModule.cpp b/Engine/Plugins/Runtime/RigVM/Source/RigVMEditor/Private/RigVMEditorModule.cpp
--- a/Engine/Plugins/Runtime/RigVM/Source/RigVMEditor/Private/RigVMEditorModule.cpp
+++ b/Engine/Plugins/Runtime/RigVM/Source/RigVMEditor/Private/RigVMEditorModule.cpp
@@ -434,8 +434,18 @@ void FRigVMEditorModule::GetInstanceActions(URigVMBlueprint* RigVMBlueprint, FBl
 		if (URigVMFunctionLibrary* LocalFunctionLibrary = RigVMBlueprint->GetLocalFunctionLibrary())
 		{
 			TArray<URigVMLibraryNode*> Functions = LocalFunctionLibrary->GetFunctions();
+			const FSoftObjectPath LocalLibrarySoftPath = LocalFunctionLibrary->GetFunctionHostObjectPath();
 			for (URigVMLibraryNode* Function : Functions)
 			{
+				// Avoid adding functions that are already added by the GetTypeActions functions (public functions that are already saved into the blueprint tag)
+				if (RigVMBlueprint->PublicGraphFunctions.ContainsByPredicate([LocalLibrarySoftPath, Function](const FRigVMGraphFunctionHeader& Header) -> bool
+				{
+					return FRigVMGraphFunctionIdentifier(LocalLibrarySoftPath, Function) == Header.LibraryPointer;
+				}))
+				{
+					continue;
+				}
+				
 				URigVMEdGraphNodeSpawner* NodeSpawner = URigVMEdGraphFunctionRefNodeSpawner::CreateFromFunction(Function);
 				check(NodeSpawner != nullptr);
 				NodeSpawner->SetRelatedBlueprintClass(BlueprintClass);
