RigVM: Disable live instance details + more granular pin value updates

diff --git a/Engine/Plugins/Runtime/RigVM/Source/RigVMEditor/Private/Editor/RigVMEditor.cpp b/Engine/Plugins/Runtime/RigVM/Source/RigVMEditor/Private/Editor/RigVMEditor.cpp
--- a/Engine/Plugins/Runtime/RigVM/Source/RigVMEditor/Private/Editor/RigVMEditor.cpp
+++ b/Engine/Plugins/Runtime/RigVM/Source/RigVMEditor/Private/Editor/RigVMEditor.cpp
@@ -1631,11 +1631,10 @@ void FRigVMEditor::HandleModifiedEvent(ERigVMGraphNotifType InNotifType, URigVMG
 		}
 		case ERigVMGraphNotifType::PinDefaultValueChanged:
 		{
-			URigVMPin* Pin = Cast<URigVMPin>(InSubject);
-
-			if(URigVMPin* RootPin = Pin->GetRootPin())
+			const URigVMPin* Pin = Cast<URigVMPin>(InSubject);
+			if(const URigVMPin* RootPin = Pin->GetRootPin())
 			{
-				const FString DefaultValue = RootPin->GetDefaultValue();
+				const FString DefaultValue = Pin->GetDefaultValue();
 				if(!DefaultValue.IsEmpty())
 				{
 					// sync the value change with the unit(s) displayed 
@@ -1646,16 +1645,24 @@ void FRigVMEditor::HandleModifiedEvent(ERigVMGraphNotifType InNotifType, URigVMG
 						{
 							if(URigVMDetailsViewWrapperObject* WrapperObject = Cast<URigVMDetailsViewWrapperObject>(SelectedObject.Get()))
 							{
-								if(WrapperObject->GetOuter() == Pin->GetNode())
+								if(WrapperObject->GetSubject() == Pin->GetNode())
 								{
-									const FProperty* TargetProperty = WrapperObject->GetClass()->FindPropertyByName(RootPin->GetFName());
-									if(TargetProperty)
+									if(const FProperty* Property = WrapperObject->GetClass()->FindPropertyByName(RootPin->GetFName()))
 									{
-										uint8* PropertyStorage = TargetProperty->ContainerPtrToValuePtr<uint8>(WrapperObject);
+										uint8* PropertyStorage = Property->ContainerPtrToValuePtr<uint8>(WrapperObject);
+
+										// traverse to get to the target pin
+										if(Pin != RootPin)
+										{
+											FString SegmentPath = Pin->GetSegmentPath();
+											const FRigVMPropertyPath PropertyTraverser(Property, SegmentPath);
+											PropertyStorage = PropertyTraverser.GetData<uint8>(PropertyStorage, Property);
+											Property = PropertyTraverser.GetTailProperty();
+										}
 
 										// we are ok with not reacting to errors here
 										FRigVMPinDefaultValueImportErrorContext ErrorPipe;										
-										TargetProperty->ImportText_Direct(*DefaultValue, PropertyStorage, nullptr, PPF_None, &ErrorPipe);
+										Property->ImportText_Direct(*DefaultValue, PropertyStorage, nullptr, PPF_None, &ErrorPipe);
 									}
 								}
 							}
@@ -2089,25 +2096,38 @@ void FRigVMEditor::OnWrappedPropertyChangedChainEvent(URigVMDetailsViewWrapperOb
 	}
 	else if(!InWrapperObject->GetWrappedNodeNotation().IsEmpty())
 	{
-		FName RootPinName = InPropertyChangedChainEvent.PropertyChain.GetHead()->GetValue()->GetFName();
-		FProperty* TargetProperty = WrapperObjects[0]->GetClass()->FindPropertyByName(RootPinName);
-		uint8* FirstPropertyStorage = TargetProperty->ContainerPtrToValuePtr<uint8>(WrapperObjects[0].Get());
-
 		URigVMNode* Node = CastChecked<URigVMNode>(InWrapperObject->GetSubject());
 
-		FString DefaultValue = FRigVMStruct::ExportToFullyQualifiedText(TargetProperty, FirstPropertyStorage);
+		const FName RootPinName = InPropertyChangedChainEvent.PropertyChain.GetHead()->GetValue()->GetFName();
+		const FProperty* Property = WrapperObjects[0]->GetClass()->FindPropertyByName(RootPinName);
+		uint8* PropertyStorage = Property->ContainerPtrToValuePtr<uint8>(WrapperObjects[0].Get());
+		const FString RootPinNameString = RootPinName.ToString();
+		FString PinPath = URigVMPin::JoinPinPath(Node->GetName(), RootPinNameString);
 
-		if(TargetProperty->IsA<FStrProperty>() || TargetProperty->IsA<FNameProperty>())
+		// traverse to get to the target pin
+		if(!InPropertyPath.Equals(RootPinNameString))
 		{
-			DefaultValue.TrimCharInline(TEXT('\"'), nullptr);
+			const FString RootPinPrefix = RootPinNameString + TEXT("->");
+			check(InPropertyPath.StartsWith(RootPinPrefix));
+			const FString RemainingPropertyPath = InPropertyPath.Mid(RootPinPrefix.Len());
+			const FString SegmentPath = RemainingPropertyPath.Replace(TEXT("->"), TEXT("."));
+			
+			const FRigVMPropertyPath PropertyTraverser(Property, SegmentPath);
+			PropertyStorage = PropertyTraverser.GetData<uint8>(PropertyStorage, Property);
+			Property = PropertyTraverser.GetTailProperty();
+			PinPath = URigVMPin::JoinPinPath(PinPath, SegmentPath);
 		}
-		
-		URigVMController* Controller = GetRigVMBlueprint()->GetController(Node->GetGraph());
 
+		FString DefaultValue = FRigVMStruct::ExportToFullyQualifiedText(Property, PropertyStorage);
+		if(Property->IsA<FStrProperty>() || Property->IsA<FNameProperty>())
+		{
+			DefaultValue.TrimCharInline(TEXT('\"'), nullptr);
+		}
 		if (!DefaultValue.IsEmpty())
 		{
-			FString PinPath = FString::Printf(TEXT("%s.%s"), *Node->GetName(), *RootPinName.ToString());
 			const bool bInteractive = InPropertyChangedChainEvent.ChangeType == EPropertyChangeType::Interactive;
+			URigVMController* Controller = GetRigVMBlueprint()->GetController(Node->GetGraph());
+			check(Controller);
 			Controller->SetPinDefaultValue(PinPath, DefaultValue, true, !bInteractive, true, !bInteractive);
 		}
 	}
@@ -2746,6 +2766,10 @@ void FRigVMEditor::SetDetailViewForLocalVariable()
 
 void FRigVMEditor::RefreshDetailView()
 {
+	if(bSuspendDetailsPanelRefresh)
+	{
+		return;
+	}
 	if(DetailViewShowsAnyRigUnit())
 	{
 		SetDetailViewForFocusedGraph();
diff --git a/Engine/Plugins/Runtime/RigVM/Source/RigVMEditor/Private/Editor/RigVMGraphDetailCustomization.cpp b/Engine/Plugins/Runtime/RigVM/Source/RigVMEditor/Private/Editor/RigVMGraphDetailCustomization.cpp
--- a/Engine/Plugins/Runtime/RigVM/Source/RigVMEditor/Private/Editor/RigVMGraphDetailCustomization.cpp
+++ b/Engine/Plugins/Runtime/RigVM/Source/RigVMEditor/Private/Editor/RigVMGraphDetailCustomization.cpp
@@ -1630,7 +1630,8 @@ void FRigVMWrappedNodeDetailCustomization::CustomizeLiveValues(IDetailLayoutBuil
 	{
 		return;
 	}
-	
+
+	/*
 	IDetailCategoryBuilder& DebugCategory = DetailLayout.EditCategory("DebugLiveValues", LOCTEXT("DebugLiveValues", "Inspect Live Values"), ECategoryPriority::Uncommon);
 	DebugCategory.InitiallyCollapsed(true);
 
@@ -1826,6 +1827,7 @@ void FRigVMWrappedNodeDetailCustomization::CustomizeLiveValues(IDetailLayoutBuil
 			}
 		}
 	}
+	*/
 }
 
 FRigVMGraphMathTypeDetailCustomization::FRigVMGraphMathTypeDetailCustomization()
