Moving functions from RigVMBlueprint to the client, through the client host interface, so it can be unified for different users

#rnx

diff --git a/Engine/Plugins/Experimental/AnimNext/Source/AnimNextUncookedOnly/Internal/Graph/AnimNextGraph_EditorData.h b/Engine/Plugins/Experimental/AnimNext/Source/AnimNextUncookedOnly/Internal/Graph/AnimNextGraph_EditorData.h
--- a/Engine/Plugins/Experimental/AnimNext/Source/AnimNextUncookedOnly/Internal/Graph/AnimNextGraph_EditorData.h
+++ b/Engine/Plugins/Experimental/AnimNext/Source/AnimNextUncookedOnly/Internal/Graph/AnimNextGraph_EditorData.h
@@ -28,7 +28,7 @@ namespace UE::AnimNext::Editor
 	struct FUtils;
 }
 
-enum class EAnimNextGraphLoadType : uint8
+enum class UE_DEPRECATED(5.4, "Please, use ERigVMLoadType") EAnimNextGraphLoadType : uint8
 {
 	PostLoad,
 	CheckUserDefinedStructs
@@ -115,7 +115,34 @@ public:
 	virtual void HandleConfigureRigVMController(const FRigVMClient* InClient, URigVMController* InControllerToConfigure) override;
 	virtual UObject* GetEditorObjectForRigVMGraph(URigVMGraph* InVMGraph) const override;
 	virtual URigVMGraph* GetRigVMGraphForEditorObject(UObject* InObject) const override;
-	
+	virtual void RecompileVM() override;
+	virtual void RecompileVMIfRequired() override;
+	virtual void RequestAutoVMRecompilation() override;
+	virtual void SetAutoVMRecompile(bool bAutoRecompile) override;
+	virtual bool GetAutoVMRecompile() const override;
+	virtual void IncrementVMRecompileBracket() override;
+	virtual void DecrementVMRecompileBracket() override;
+	virtual void RefreshAllModels(ERigVMLoadType InLoadType) override;
+	virtual void OnRigVMRegistryChanged() override;
+	virtual void RequestRigVMInit() override;
+	virtual URigVMGraph* GetModel(const UEdGraph* InEdGraph = nullptr) const override;
+	virtual URigVMGraph* GetModel(const FString& InNodePath) const override;
+	virtual URigVMGraph* GetDefaultModel() const override;
+	virtual TArray<URigVMGraph*> GetAllModels() const override;
+	virtual URigVMFunctionLibrary* GetLocalFunctionLibrary() const override;
+	virtual URigVMGraph* AddModel(FString InName = TEXT("Rig Graph"), bool bSetupUndoRedo = true, bool bPrintPythonCommand = true) override;
+	virtual bool RemoveModel(FString InName = TEXT("Rig Graph"), bool bSetupUndoRedo = true, bool bPrintPythonCommand = true) override;
+	virtual FRigVMGetFocusedGraph& OnGetFocusedGraph() override;
+	virtual const FRigVMGetFocusedGraph& OnGetFocusedGraph() const override;
+	virtual URigVMGraph* GetFocusedModel() const override;
+	virtual URigVMController* GetController(const URigVMGraph* InGraph = nullptr) const override;
+	virtual URigVMController* GetControllerByName(const FString InGraphName = TEXT("")) const override;
+	virtual URigVMController* GetOrCreateController(URigVMGraph* InGraph = nullptr) override;
+	virtual URigVMController* GetController(const UEdGraph* InEdGraph) const override;
+	virtual URigVMController* GetOrCreateController(const UEdGraph* InGraph) override;
+	virtual TArray<FString> GeneratePythonCommands(const FString InNewBlueprintName) override;
+
+
 	// IRigVMGraphFunctionHost interface
 	virtual FRigVMGraphFunctionStore* GetRigVMGraphFunctionStore() override;
 	virtual const FRigVMGraphFunctionStore* GetRigVMGraphFunctionStore() const override;
@@ -132,19 +159,9 @@ public:
 
 protected:
 #if WITH_EDITOR
-	void RefreshAllModels(EAnimNextGraphLoadType InLoadType);
 	void GetAllGraphs(TArray<UEdGraph*>& Graphs) const;
 #endif
 
-	void RecompileVM();
-	
-	void RecompileVMIfRequired();
-
-	void RequestAutoVMRecompilation();
-
-	void IncrementVMRecompileBracket();
-
-	void DecrementVMRecompileBracket();
 	
 	void HandleModifiedEvent(ERigVMGraphNotifType InNotifType, URigVMGraph* InGraph, UObject* InSubject);
 
diff --git a/Engine/Plugins/Experimental/AnimNext/Source/AnimNextUncookedOnly/Internal/Param/AnimNextParameterBlock_EditorData.h b/Engine/Plugins/Experimental/AnimNext/Source/AnimNextUncookedOnly/Internal/Param/AnimNextParameterBlock_EditorData.h
--- a/Engine/Plugins/Experimental/AnimNext/Source/AnimNextUncookedOnly/Internal/Param/AnimNextParameterBlock_EditorData.h
+++ b/Engine/Plugins/Experimental/AnimNext/Source/AnimNextUncookedOnly/Internal/Param/AnimNextParameterBlock_EditorData.h
@@ -41,7 +41,7 @@ namespace UE::AnimNext::UncookedOnly
 	DECLARE_MULTICAST_DELEGATE_OneParam(FOnParameterBlockModified, UAnimNextParameterBlock_EditorData* /* InEditorData */);
 }
 
-enum class EAnimNextParameterLoadType : uint8
+enum class UE_DEPRECATED(5.4, "Please use ERigVMLoadType") EAnimNextParameterLoadType : uint8
 {
 	PostLoad,
 	CheckUserDefinedStructs
@@ -141,6 +141,7 @@ class UAnimNextParameterBlock_EditorData : public UObject, public IRigVMClientHo
 	void HandlePackageDone(const FEndLoadPackageContext& Context);
 	void HandlePackageDone();
 #endif
+
 	// IRigVMClientHost interface
 	virtual FRigVMClient* GetRigVMClient() override;
 	virtual const FRigVMClient* GetRigVMClient() const override;
@@ -152,6 +153,33 @@ class UAnimNextParameterBlock_EditorData : public UObject, public IRigVMClientHo
 	virtual void HandleConfigureRigVMController(const FRigVMClient* InClient, URigVMController* InControllerToConfigure) override;
 	virtual UObject* GetEditorObjectForRigVMGraph(URigVMGraph* InVMGraph) const override;
 	virtual URigVMGraph* GetRigVMGraphForEditorObject(UObject* InObject) const override;
+	virtual void RecompileVM() override;
+	virtual void RecompileVMIfRequired() override;
+	virtual void RequestAutoVMRecompilation() override;
+	virtual void SetAutoVMRecompile(bool bAutoRecompile) override;
+	virtual bool GetAutoVMRecompile() const override;
+	virtual void IncrementVMRecompileBracket() override;
+	virtual void DecrementVMRecompileBracket() override;
+	virtual void RefreshAllModels(ERigVMLoadType InLoadType) override;
+	virtual void OnRigVMRegistryChanged() override;
+	virtual void RequestRigVMInit() override;
+	virtual URigVMGraph* GetModel(const UEdGraph* InEdGraph = nullptr) const override;
+	virtual URigVMGraph* GetModel(const FString& InNodePath) const override;
+	virtual URigVMGraph* GetDefaultModel() const override;
+	virtual TArray<URigVMGraph*> GetAllModels() const override;
+	virtual URigVMFunctionLibrary* GetLocalFunctionLibrary() const override;
+	virtual URigVMGraph* AddModel(FString InName = TEXT("Rig Graph"), bool bSetupUndoRedo = true, bool bPrintPythonCommand = true) override;
+	virtual bool RemoveModel(FString InName = TEXT("Rig Graph"), bool bSetupUndoRedo = true, bool bPrintPythonCommand = true) override;
+	virtual FRigVMGetFocusedGraph& OnGetFocusedGraph() override;
+	virtual const FRigVMGetFocusedGraph& OnGetFocusedGraph() const override;
+	virtual URigVMGraph* GetFocusedModel() const override;
+	virtual URigVMController* GetController(const URigVMGraph* InGraph = nullptr) const override;
+	virtual URigVMController* GetControllerByName(const FString InGraphName = TEXT("")) const override;
+	virtual URigVMController* GetOrCreateController(URigVMGraph* InGraph = nullptr) override;
+	virtual URigVMController* GetController(const UEdGraph* InEdGraph) const override;
+	virtual URigVMController* GetOrCreateController(const UEdGraph* InGraph) override;
+	virtual TArray<FString> GeneratePythonCommands(const FString InNewBlueprintName) override;
+
 
 	// IRigVMGraphFunctionHost interface
 	virtual FRigVMGraphFunctionStore* GetRigVMGraphFunctionStore() override;
@@ -162,24 +190,12 @@ class UAnimNextParameterBlock_EditorData : public UObject, public IRigVMClientHo
 	
 	ANIMNEXTUNCOOKEDONLY_API void Initialize(bool bRecompileVM);
 
-	void RefreshAllModels(EAnimNextParameterLoadType InLoadType);
-
-	void RecompileVM();
-
-	void RecompileVMIfRequired();
-
 	void Recompile();
 
 	void RecompileIfRequired();
 
 	void RequestAutoRecompilation();
 	
-	void RequestAutoVMRecompilation();
-
-	void IncrementVMRecompileBracket();
-
-	void DecrementVMRecompileBracket();
-
 	void HandleModifiedEvent(ERigVMGraphNotifType InNotifType, URigVMGraph* InGraph, UObject* InSubject);
 
 	void CreateEdGraphForCollapseNode(URigVMCollapseNode* InNode);
diff --git a/Engine/Plugins/Experimental/AnimNext/Source/AnimNextUncookedOnly/Private/Graph/AnimNextGraph_EditorData.cpp b/Engine/Plugins/Experimental/AnimNext/Source/AnimNextUncookedOnly/Private/Graph/AnimNextGraph_EditorData.cpp
--- a/Engine/Plugins/Experimental/AnimNext/Source/AnimNextUncookedOnly/Private/Graph/AnimNextGraph_EditorData.cpp
+++ b/Engine/Plugins/Experimental/AnimNext/Source/AnimNextUncookedOnly/Private/Graph/AnimNextGraph_EditorData.cpp
@@ -151,7 +151,7 @@ void UAnimNextGraph_EditorData::PostLoad()
 	Super::PostLoad();
 	
 	Initialize(/*bRecompileVM*/false);
-	RefreshAllModels(EAnimNextGraphLoadType::PostLoad);
+	RefreshAllModels(ERigVMLoadType::PostLoad);
 
 #if WITH_EDITOR
 	FExternalPackageHelper::LoadObjectsFromExternalPackages<UAnimNextGraphEntry>(this, [this](UAnimNextGraphEntry* InLoadedEntry)
@@ -213,9 +213,9 @@ void UAnimNextGraph_EditorData::HandlePackageDone()
 	ReconstructAllNodes();
 }
 
-void UAnimNextGraph_EditorData::RefreshAllModels(EAnimNextGraphLoadType InLoadType)
+void UAnimNextGraph_EditorData::RefreshAllModels(ERigVMLoadType InLoadType)
 {
-	const bool bIsPostLoad = InLoadType == EAnimNextGraphLoadType::PostLoad;
+	const bool bIsPostLoad = InLoadType == ERigVMLoadType::PostLoad;
 
 	TGuardValue<bool> IsCompilingGuard(bIsCompiling, true);
 	TGuardValue<bool> ClientIgnoreModificationsGuard(RigVMClient.bIgnoreModelNotifications, true);
@@ -328,6 +328,99 @@ void UAnimNextGraph_EditorData::RefreshAllModels(EAnimNextGraphLoadType InLoadTy
 	}
 }
 
+void UAnimNextGraph_EditorData::OnRigVMRegistryChanged()
+{
+	RefreshAllModels(ERigVMLoadType::PostLoad);
+	//RebuildGraphFromModel(); // TODO zzz : How we do this on AnimNext ?
+}
+
+void UAnimNextGraph_EditorData::RequestRigVMInit()
+{
+	// TODO zzz : How we do this on AnimNext ?
+}
+
+URigVMGraph* UAnimNextGraph_EditorData::GetModel(const UEdGraph* InEdGraph) const
+{
+	return RigVMClient.GetModel(InEdGraph);
+}
+
+URigVMGraph* UAnimNextGraph_EditorData::GetModel(const FString& InNodePath) const
+{
+	return RigVMClient.GetModel(InNodePath);
+}
+
+URigVMGraph* UAnimNextGraph_EditorData::GetDefaultModel() const 
+{
+	return RigVMClient.GetDefaultModel();
+}
+
+TArray<URigVMGraph*> UAnimNextGraph_EditorData::GetAllModels() const
+{
+	return RigVMClient.GetAllModels(true, true);
+}
+
+URigVMFunctionLibrary* UAnimNextGraph_EditorData::GetLocalFunctionLibrary() const
+{
+	return RigVMClient.GetFunctionLibrary();
+}
+
+URigVMGraph* UAnimNextGraph_EditorData::AddModel(FString InName, bool bSetupUndoRedo, bool bPrintPythonCommand)
+{
+	TGuardValue<bool> EnablePythonPrint(bSuspendPythonMessagesForRigVMClient, !bPrintPythonCommand);
+	return RigVMClient.AddModel(InName, bSetupUndoRedo, bPrintPythonCommand);
+}
+
+bool UAnimNextGraph_EditorData::RemoveModel(FString InName, bool bSetupUndoRedo, bool bPrintPythonCommand)
+{
+	TGuardValue<bool> EnablePythonPrint(bSuspendPythonMessagesForRigVMClient, !bPrintPythonCommand);
+	return RigVMClient.RemoveModel(InName, bSetupUndoRedo, bPrintPythonCommand);
+}
+
+FRigVMGetFocusedGraph& UAnimNextGraph_EditorData::OnGetFocusedGraph()
+{
+	return RigVMClient.OnGetFocusedGraph();
+}
+
+const FRigVMGetFocusedGraph& UAnimNextGraph_EditorData::OnGetFocusedGraph() const
+{
+	return RigVMClient.OnGetFocusedGraph();
+}
+
+URigVMGraph* UAnimNextGraph_EditorData::GetFocusedModel() const
+{
+	return RigVMClient.GetFocusedModel();
+}
+
+URigVMController* UAnimNextGraph_EditorData::GetController(const URigVMGraph* InGraph) const
+{
+	return RigVMClient.GetController(InGraph);
+};
+
+URigVMController* UAnimNextGraph_EditorData::GetControllerByName(const FString InGraphName) const
+{
+	return RigVMClient.GetControllerByName(InGraphName);
+};
+
+URigVMController* UAnimNextGraph_EditorData::GetOrCreateController(URigVMGraph* InGraph)
+{
+	return RigVMClient.GetOrCreateController(InGraph);
+};
+
+URigVMController* UAnimNextGraph_EditorData::GetController(const UEdGraph* InEdGraph) const
+{
+	return RigVMClient.GetController(InEdGraph);
+};
+
+URigVMController* UAnimNextGraph_EditorData::GetOrCreateController(const UEdGraph* InEdGraph)
+{
+	return RigVMClient.GetOrCreateController(InEdGraph);
+};
+
+TArray<FString> UAnimNextGraph_EditorData::GeneratePythonCommands(const FString InNewBlueprintName)
+{
+	return TArray<FString>();
+}
+
 void UAnimNextGraph_EditorData::GetAllGraphs(TArray<UEdGraph*>& OutGraphs) const
 {
 	OutGraphs.Reset();
@@ -504,6 +597,16 @@ void UAnimNextGraph_EditorData::RequestAutoVMRecompilation()
 	}
 }
 
+void UAnimNextGraph_EditorData::SetAutoVMRecompile(bool bAutoRecompile)
+{
+	bAutoRecompileVM = bAutoRecompile;
+}
+
+bool UAnimNextGraph_EditorData::GetAutoVMRecompile() const
+{
+	return bAutoRecompileVM;
+}
+
 void UAnimNextGraph_EditorData::IncrementVMRecompileBracket()
 {
 	VMRecompilationBracket++;
diff --git a/Engine/Plugins/Experimental/AnimNext/Source/AnimNextUncookedOnly/Private/Param/AnimNextParameterBlock_EditorData.cpp b/Engine/Plugins/Experimental/AnimNext/Source/AnimNextUncookedOnly/Private/Param/AnimNextParameterBlock_EditorData.cpp
--- a/Engine/Plugins/Experimental/AnimNext/Source/AnimNextUncookedOnly/Private/Param/AnimNextParameterBlock_EditorData.cpp
+++ b/Engine/Plugins/Experimental/AnimNext/Source/AnimNextUncookedOnly/Private/Param/AnimNextParameterBlock_EditorData.cpp
@@ -547,7 +547,7 @@ void UAnimNextParameterBlock_EditorData::PostLoad()
 	Super::PostLoad();
 
 	Initialize(/*bRecompileVM*/false);
-	RefreshAllModels(EAnimNextParameterLoadType::PostLoad);
+	RefreshAllModels(ERigVMLoadType::PostLoad);
 
 #if WITH_EDITOR
 	FExternalPackageHelper::LoadObjectsFromExternalPackages<UAnimNextParameterBlockEntry>(this, [this](UAnimNextParameterBlockEntry* InLoadedEntry)
@@ -625,9 +625,9 @@ void UAnimNextParameterBlock_EditorData::HandlePackageDone()
 	ReconstructAllNodes();
 }
 
-void UAnimNextParameterBlock_EditorData::RefreshAllModels(EAnimNextParameterLoadType InLoadType)
+void UAnimNextParameterBlock_EditorData::RefreshAllModels(ERigVMLoadType InLoadType)
 {
-	const bool bIsPostLoad = InLoadType == EAnimNextParameterLoadType::PostLoad;
+	const bool bIsPostLoad = InLoadType == ERigVMLoadType::PostLoad;
 
 	TGuardValue<bool> IsCompilingGuard(bIsCompiling, true);
 	TGuardValue<bool> ClientIgnoreModificationsGuard(RigVMClient.bIgnoreModelNotifications, true);
@@ -740,6 +740,99 @@ void UAnimNextParameterBlock_EditorData::RefreshAllModels(EAnimNextParameterLoad
 	}
 }
 
+void UAnimNextParameterBlock_EditorData::OnRigVMRegistryChanged()
+{
+	RefreshAllModels(ERigVMLoadType::PostLoad);
+	//RebuildGraphFromModel(); // TODO zzz : How we do this on AnimNext ?
+}
+
+void UAnimNextParameterBlock_EditorData::RequestRigVMInit()
+{
+	// TODO zzz : How we do this on AnimNext ?
+}
+
+URigVMGraph* UAnimNextParameterBlock_EditorData::GetModel(const UEdGraph* InEdGraph) const
+{
+	return RigVMClient.GetModel(InEdGraph);
+}
+
+URigVMGraph* UAnimNextParameterBlock_EditorData::GetModel(const FString& InNodePath) const
+{
+	return RigVMClient.GetModel(InNodePath);
+}
+
+URigVMGraph* UAnimNextParameterBlock_EditorData::GetDefaultModel() const
+{
+	return RigVMClient.GetDefaultModel();
+}
+
+TArray<URigVMGraph*> UAnimNextParameterBlock_EditorData::GetAllModels() const
+{
+	return RigVMClient.GetAllModels(true, true);
+}
+
+URigVMFunctionLibrary* UAnimNextParameterBlock_EditorData::GetLocalFunctionLibrary() const
+{
+	return RigVMClient.GetFunctionLibrary();
+}
+
+URigVMGraph* UAnimNextParameterBlock_EditorData::AddModel(FString InName, bool bSetupUndoRedo, bool bPrintPythonCommand)
+{
+	TGuardValue<bool> EnablePythonPrint(bSuspendPythonMessagesForRigVMClient, !bPrintPythonCommand);
+	return RigVMClient.AddModel(InName, bSetupUndoRedo, bPrintPythonCommand);
+}
+
+bool UAnimNextParameterBlock_EditorData::RemoveModel(FString InName, bool bSetupUndoRedo, bool bPrintPythonCommand)
+{
+	TGuardValue<bool> EnablePythonPrint(bSuspendPythonMessagesForRigVMClient, !bPrintPythonCommand);
+	return RigVMClient.RemoveModel(InName, bSetupUndoRedo, bPrintPythonCommand);
+}
+
+FRigVMGetFocusedGraph& UAnimNextParameterBlock_EditorData::OnGetFocusedGraph()
+{
+	return RigVMClient.OnGetFocusedGraph();
+}
+
+const FRigVMGetFocusedGraph& UAnimNextParameterBlock_EditorData::OnGetFocusedGraph() const
+{
+	return RigVMClient.OnGetFocusedGraph();
+}
+
+URigVMGraph* UAnimNextParameterBlock_EditorData::GetFocusedModel() const
+{
+	return RigVMClient.GetFocusedModel();
+}
+
+URigVMController* UAnimNextParameterBlock_EditorData::GetController(const URigVMGraph* InGraph) const
+{
+	return RigVMClient.GetController(InGraph);
+};
+
+URigVMController* UAnimNextParameterBlock_EditorData::GetControllerByName(const FString InGraphName) const
+{
+	return RigVMClient.GetControllerByName(InGraphName);
+};
+
+URigVMController* UAnimNextParameterBlock_EditorData::GetOrCreateController(URigVMGraph* InGraph)
+{
+	return RigVMClient.GetOrCreateController(InGraph);
+};
+
+URigVMController* UAnimNextParameterBlock_EditorData::GetController(const UEdGraph* InEdGraph) const
+{
+	return RigVMClient.GetController(InEdGraph);
+};
+
+URigVMController* UAnimNextParameterBlock_EditorData::GetOrCreateController(const UEdGraph* InEdGraph)
+{
+	return RigVMClient.GetOrCreateController(InEdGraph);
+};
+
+TArray<FString> UAnimNextParameterBlock_EditorData::GeneratePythonCommands(const FString InNewBlueprintName)
+{
+	return TArray<FString>();
+}
+
 #endif // WITH_EDITOR
 
 FRigVMClient* UAnimNextParameterBlock_EditorData::GetRigVMClient()
@@ -948,6 +1041,17 @@ void UAnimNextParameterBlock_EditorData::RequestAutoVMRecompilation()
 	}
 }
 
+void UAnimNextParameterBlock_EditorData::SetAutoVMRecompile(bool bAutoRecompile)
+{
+	bAutoRecompileVM = bAutoRecompile;
+}
+
+bool UAnimNextParameterBlock_EditorData::GetAutoVMRecompile() const
+{
+	return bAutoRecompileVM;
+}
+
+
 void UAnimNextParameterBlock_EditorData::IncrementVMRecompileBracket()
 {
 	VMRecompilationBracket++;
diff --git a/Engine/Plugins/Runtime/RigVM/Source/RigVMDeveloper/Private/RigVMBlueprint.cpp b/Engine/Plugins/Runtime/RigVM/Source/RigVMDeveloper/Private/RigVMBlueprint.cpp
--- a/Engine/Plugins/Runtime/RigVM/Source/RigVMDeveloper/Private/RigVMBlueprint.cpp
+++ b/Engine/Plugins/Runtime/RigVM/Source/RigVMDeveloper/Private/RigVMBlueprint.cpp
@@ -1354,9 +1354,9 @@ void URigVMBlueprint::DecrementVMRecompileBracket()
 	}
 }
 
-void URigVMBlueprint::RefreshAllModels(ERigVMBlueprintLoadType InLoadType)
+void URigVMBlueprint::RefreshAllModels(ERigVMLoadType InLoadType)
 {
-	const bool bIsPostLoad = InLoadType == ERigVMBlueprintLoadType::PostLoad;
+	const bool bIsPostLoad = InLoadType == ERigVMLoadType::PostLoad;
 
 	// avoid any compute if the current structure hashes match with the serialized ones
 	if(CVarRigVMEnablePostLoadHashing->GetBool() && RigVMClient.GetStructureHash() == RigVMClient.GetSerializedStructureHash())
@@ -1938,26 +1938,14 @@ void URigVMBlueprint::RequestRigVMInit()
 
 URigVMGraph* URigVMBlueprint::GetModel(const UEdGraph* InEdGraph) const
 {
-	if (InEdGraph == nullptr)
-	{
-		return GetDefaultModel();
-	}
-
-	if(InEdGraph->GetOutermost() != GetOutermost())
-	{
-		return nullptr;
-	}
-
 #if WITH_EDITORONLY_DATA
-	if (InEdGraph == FunctionLibraryEdGraph)
+	if (InEdGraph != nullptr && InEdGraph == FunctionLibraryEdGraph)
 	{
 		return RigVMClient.GetFunctionLibrary();
 	}
 #endif
 
-	const URigVMEdGraph* RigGraph = Cast< URigVMEdGraph>(InEdGraph);
-	check(RigGraph);
-	return GetModel(RigGraph->ModelNodePath);
+	return RigVMClient.GetModel(InEdGraph);
 }
 
 URigVMGraph* URigVMBlueprint::GetModel(const FString& InNodePath) const
@@ -1982,28 +1970,29 @@ URigVMFunctionLibrary* URigVMBlueprint::GetLocalFunctionLibrary() const
 
 URigVMGraph* URigVMBlueprint::AddModel(FString InName, bool bSetupUndoRedo, bool bPrintPythonCommand)
 {
-	const FString DesiredName = FString::Printf(TEXT("%s %s"),
-    	FRigVMClient::RigVMModelPrefix, *InName);
-
 	TGuardValue<bool> EnablePythonPrint(bSuspendPythonMessagesForRigVMClient, !bPrintPythonCommand);
-	return RigVMClient.AddModel(*DesiredName, bSetupUndoRedo);
+	return RigVMClient.AddModel(InName, bSetupUndoRedo, bPrintPythonCommand);
 }
 
 bool URigVMBlueprint::RemoveModel(FString InName, bool bSetupUndoRedo, bool bPrintPythonCommand)
 {
 	TGuardValue<bool> EnablePythonPrint(bSuspendPythonMessagesForRigVMClient, !bPrintPythonCommand);
-	return RigVMClient.RemoveModel(InName, bSetupUndoRedo);
+	return RigVMClient.RemoveModel(InName, bSetupUndoRedo, bPrintPythonCommand);
+}
+
+FRigVMGetFocusedGraph& URigVMBlueprint::OnGetFocusedGraph()
+{
+	return RigVMClient.OnGetFocusedGraph();
+}
+
+const FRigVMGetFocusedGraph& URigVMBlueprint::OnGetFocusedGraph() const
+{
+	return RigVMClient.OnGetFocusedGraph();
 }
 
 URigVMGraph* URigVMBlueprint::GetFocusedModel() const
 {
-#if WITH_EDITOR
-	if(OnGetFocusedGraphDelegate.IsBound())
-	{
-		return OnGetFocusedGraphDelegate.Execute(); 
-	}
-#endif
-	return RigVMClient.GetDefaultModel();
+	return RigVMClient.GetFocusedModel();
 }
 
 URigVMController* URigVMBlueprint::GetController(const URigVMGraph* InGraph) const
@@ -2013,23 +2002,7 @@ URigVMController* URigVMBlueprint::GetController(const URigVMGraph* InGraph) con
 
 URigVMController* URigVMBlueprint::GetControllerByName(const FString InGraphName) const
 {
-	if(InGraphName.IsEmpty())
-	{
-		if(const URigVMGraph* DefaultModel = GetRigVMClient()->GetDefaultModel())
-		{
-			return GetController(DefaultModel);
-		}
-	}
-	
-	for (const URigVMGraph* Graph : GetAllModels())
-	{
-		if (Graph->GetName() == InGraphName || Graph->GetGraphName() == InGraphName)
-		{
-			return GetController(Graph);
-		}
-	}
-	
-	return nullptr;
+	return RigVMClient.GetControllerByName(InGraphName);
 }
 
 URigVMController* URigVMBlueprint::GetOrCreateController(URigVMGraph* InGraph)
diff --git a/Engine/Plugins/Runtime/RigVM/Source/RigVMDeveloper/Private/RigVMModel/RigVMClient.cpp b/Engine/Plugins/Runtime/RigVM/Source/RigVMDeveloper/Private/RigVMModel/RigVMClient.cpp
--- a/Engine/Plugins/Runtime/RigVM/Source/RigVMDeveloper/Private/RigVMModel/RigVMClient.cpp
+++ b/Engine/Plugins/Runtime/RigVM/Source/RigVMDeveloper/Private/RigVMModel/RigVMClient.cpp
@@ -7,6 +7,7 @@
 #include "UnrealExporter.h"
 #include "UObject/ObjectSaveContext.h"
 #include "RigVMModel/RigVMControllerActions.h"
+#include "EdGraph/RigVMEdGraph.h"
 
 #include UE_INLINE_GENERATED_CPP_BY_NAME(RigVMClient)
 
@@ -163,6 +164,31 @@ URigVMGraph* FRigVMClient::GetDefaultModel() const
 	return GetModel(0);
 }
 
+URigVMGraph* FRigVMClient::GetModel(const UEdGraph* InEdGraph) const
+{
+	if (InEdGraph == nullptr)
+	{
+		return GetDefaultModel();
+	}
+
+	if (InEdGraph->GetOutermost() != GetOuter()->GetOutermost())
+	{
+		return nullptr;
+	}
+
+//#if WITH_EDITORONLY_DATA
+//	if (InEdGraph == FunctionLibraryEdGraph)
+//	{
+//		return RigVMClient.GetFunctionLibrary();
+//	}
+//#endif
+
+	const URigVMEdGraph* RigGraph = Cast< URigVMEdGraph>(InEdGraph);
+	check(RigGraph);
+	return GetModel(RigGraph->ModelNodePath);
+
+}
+
 URigVMGraph* FRigVMClient::GetModel(const FString& InNodePathOrName) const
 {
 	if(InNodePathOrName.IsEmpty())
@@ -316,6 +342,27 @@ URigVMController* FRigVMClient::GetOrCreateController(const UObject* InEditorSid
 	return nullptr;
 }
 
+URigVMController* FRigVMClient::GetControllerByName(const FString InGraphName) const
+{
+	if (InGraphName.IsEmpty())
+	{
+		if (const URigVMGraph* DefaultModel = GetDefaultModel())
+		{
+			return GetController(DefaultModel);
+		}
+	}
+
+	for (const URigVMGraph* Graph : GetAllModels(true, true))
+	{
+		if (Graph->GetName() == InGraphName || Graph->GetGraphName() == InGraphName)
+		{
+			return GetController(Graph);
+		}
+	}
+
+	return nullptr;
+}
+
 bool FRigVMClient::RemoveController(const URigVMGraph* InModel)
 {
 	const FSoftObjectPath Key(InModel);
@@ -331,6 +378,12 @@ bool FRigVMClient::RemoveController(const URigVMGraph* InModel)
 	return bSuccess;
 }
 
+URigVMGraph* FRigVMClient::AddModel(const FString InName, bool bSetupUndoRedo, bool bPrintPythonCommand)
+{
+	const FString DesiredName = FString::Printf(TEXT("%s %s"), FRigVMClient::RigVMModelPrefix, *InName);
+	return AddModel(*DesiredName, bSetupUndoRedo);
+}
+
 URigVMGraph* FRigVMClient::AddModel(const FName& InName, bool bSetupUndoRedo, const FObjectInitializer* ObjectInitializer, bool bCreateController)
 {
 #if WITH_EDITOR
@@ -482,6 +535,23 @@ void FRigVMClient::SetExecuteContextStruct(UScriptStruct* InExecuteContextStruct
 	GetOrCreateSchema()->SetExecuteContextStruct(InExecuteContextStruct);
 }
 
+URigVMGraph* FRigVMClient::GetFocusedModel() const
+{
+#if WITH_EDITOR
+	if (OnGetFocusedGraph().IsBound())
+	{
+		return OnGetFocusedGraph().Execute();
+	}
+#endif
+
+	return GetDefaultModel();
+}
+
+bool FRigVMClient::RemoveModel(FString InName, bool bSetupUndoRedo, bool bPrintPythonCommand)
+{
+	return RemoveModel(InName, bSetupUndoRedo);
+}
+
 bool FRigVMClient::RemoveModel(const FString& InNodePathOrName, bool bSetupUndoRedo)
 {
 	if(URigVMGraph* Model = GetModel(InNodePathOrName))
diff --git a/Engine/Plugins/Runtime/RigVM/Source/RigVMDeveloper/Public/RigVMBlueprint.h b/Engine/Plugins/Runtime/RigVM/Source/RigVMDeveloper/Public/RigVMBlueprint.h
--- a/Engine/Plugins/Runtime/RigVM/Source/RigVMDeveloper/Public/RigVMBlueprint.h
+++ b/Engine/Plugins/Runtime/RigVM/Source/RigVMDeveloper/Public/RigVMBlueprint.h
@@ -42,7 +42,7 @@ DECLARE_DELEGATE_RetVal_OneParam(TRigVMTypeIndex, FRigVMOnPinTypeSelectionReques
 DECLARE_EVENT(URigVMBlueprint, FOnRigVMBreakpointAdded);
 DECLARE_EVENT_OneParam(URigVMBlueprint, FOnRigVMRequestInspectObject, const TArray<UObject*>& );
 DECLARE_EVENT_OneParam(URigVMBlueprint, FOnRigVMRequestInspectMemoryStorage, const TArray<FRigVMMemoryStorageStruct*>&);
-DECLARE_DELEGATE_RetVal(URigVMGraph*, FRigVMBlueprintGetFocusedGraph);
+//DECLARE_DELEGATE_RetVal(URigVMGraph*, FRigVMBlueprintGetFocusedGraph);
 
 USTRUCT()
 struct RIGVMDEVELOPER_API FRigVMPythonSettings
@@ -146,7 +146,7 @@ struct RIGVMDEVELOPER_API FRigVMEdGraphDisplaySettings
 	double AggregateAverage(TArray<double>& InFrames, double InPrevious, double InNext) const;
 };
 
-enum class ERigVMBlueprintLoadType : uint8
+enum class UE_DEPRECATED(5.4, "Pease, use ERigVMLoadType") ERigVMBlueprintLoadType : uint8
 {
 	PostLoad,
 	CheckUserDefinedStructs
@@ -318,7 +318,7 @@ public:
 	/** Called during cooking. Must return all objects that will be Preload()ed when this is serialized at load time. */
 	virtual void GetPreloadDependencies(TArray<UObject*>& OutDeps) override;
 
-	// IRigVMClientHost interface
+	//  --- IRigVMClientHost interface Start---
 	virtual FRigVMClient* GetRigVMClient() override;
 	virtual const FRigVMClient* GetRigVMClient() const override;
 	virtual IRigVMGraphFunctionHost* GetRigVMGraphFunctionHost() override;
@@ -331,85 +331,87 @@ public:
 	virtual void HandleConfigureRigVMController(const FRigVMClient* InClient, URigVMController* InControllerToConfigure) override;
 	virtual UObject* ResolveUserDefinedTypeById(const FString& InTypeName) const override;
 
-	FOnRigVMRequestInspectObject& OnRequestInspectObject() { return OnRequestInspectObjectEvent; }
-	void RequestInspectObject(const TArray<UObject*>& InObjects) { OnRequestInspectObjectEvent.Broadcast(InObjects); }
-
-	FOnRigVMRequestInspectMemoryStorage& OnRequestInspectMemoryStorage() { return OnRequestInspectMemoryStorageEvent; }
-	void RequestInspectMemoryStorage(const TArray<FRigVMMemoryStorageStruct*>& InMemoryStorageStructs) { OnRequestInspectMemoryStorageEvent.Broadcast(InMemoryStorageStructs); }
-
-#endif	// #if WITH_EDITOR
-
-	FRigVMBlueprintGetFocusedGraph& OnGetFocusedGraph() { return OnGetFocusedGraphDelegate; }
-
-	virtual bool ShouldBeMarkedDirtyUponTransaction() const override { return false; }
-
 	UFUNCTION(BlueprintCallable, Category = "RigVM Blueprint")
-	void RecompileVM();
-
+	virtual void RecompileVM() override;
 	UFUNCTION(BlueprintCallable, Category = "RigVM Blueprint")
-	void RecompileVMIfRequired();
-	
+	virtual void RecompileVMIfRequired() override;
 	UFUNCTION(BlueprintCallable, Category = "RigVM Blueprint")
-	void RequestAutoVMRecompilation();
-
+	virtual void RequestAutoVMRecompilation() override;
 	UFUNCTION(BlueprintCallable, Category = "RigVM Blueprint")
-	void SetAutoVMRecompile(bool bAutoRecompile) { bAutoRecompileVM = bAutoRecompile; }
-
+	virtual void SetAutoVMRecompile(bool bAutoRecompile) override { bAutoRecompileVM = bAutoRecompile; }
 	UFUNCTION(BlueprintPure, Category = "RigVM Blueprint")
-	bool GetAutoVMRecompile() const { return bAutoRecompileVM; }
+	virtual bool GetAutoVMRecompile() const override { return bAutoRecompileVM; }
 
-	void IncrementVMRecompileBracket();
-	void DecrementVMRecompileBracket();
+	virtual void IncrementVMRecompileBracket() override;
+	virtual void DecrementVMRecompileBracket() override;
 
 	// this is needed since even after load
 	// model data can change while the RigVM BP is not opened
 	// for example, if a user defined struct changed after BP load,
 	// any pin that references the struct needs to be regenerated
-	void RefreshAllModels(ERigVMBlueprintLoadType InLoadType = ERigVMBlueprintLoadType::PostLoad);
+	virtual void RefreshAllModels(ERigVMLoadType InLoadType = ERigVMLoadType::PostLoad) override;
 
 	// RigVMRegistry changes can be triggered when new user defined types(structs/enums) are added/removed
 	// in which case we have to refresh the model
-	void OnRigVMRegistryChanged();
-	
+	virtual void OnRigVMRegistryChanged() override;
+
 	UFUNCTION(BlueprintCallable, Category = "RigVM Blueprint")
-	void RequestRigVMInit();
+	virtual void RequestRigVMInit() override;
 
 	UFUNCTION(BlueprintCallable, Category = "RigVM Blueprint")
-	URigVMGraph* GetModel(const UEdGraph* InEdGraph = nullptr) const;
-	URigVMGraph* GetModel(const FString& InNodePath) const;
+	virtual URigVMGraph* GetModel(const UEdGraph* InEdGraph = nullptr) const override;
+	virtual URigVMGraph* GetModel(const FString& InNodePath) const override;
 
 	UFUNCTION(BlueprintCallable, Category = "RigVM Blueprint")
-	URigVMGraph* GetDefaultModel() const;
+	virtual URigVMGraph* GetDefaultModel() const override;
 
 	UFUNCTION(BlueprintCallable, Category = "RigVM Blueprint")
-	TArray<URigVMGraph*> GetAllModels() const;
+	virtual TArray<URigVMGraph*> GetAllModels() const override;
 
 	UFUNCTION(BlueprintCallable, Category = "RigVM Blueprint")
-	URigVMFunctionLibrary* GetLocalFunctionLibrary() const;
+	virtual URigVMFunctionLibrary* GetLocalFunctionLibrary() const override;
 
 	UFUNCTION(BlueprintCallable, Category = "RigVM Blueprint")
-	URigVMGraph* AddModel(FString InName = TEXT("Rig Graph"), bool bSetupUndoRedo = true, bool bPrintPythonCommand = true);
+	virtual URigVMGraph* AddModel(FString InName = TEXT("Rig Graph"), bool bSetupUndoRedo = true, bool bPrintPythonCommand = true) override;
 
 	UFUNCTION(BlueprintCallable, Category = "RigVM Blueprint")
-	bool RemoveModel(FString InName = TEXT("Rig Graph"), bool bSetupUndoRedo = true, bool bPrintPythonCommand = true);
+	virtual bool RemoveModel(FString InName = TEXT("Rig Graph"), bool bSetupUndoRedo = true, bool bPrintPythonCommand = true) override;
+
+	virtual FRigVMGetFocusedGraph& OnGetFocusedGraph() override;
+	virtual const FRigVMGetFocusedGraph& OnGetFocusedGraph() const override;
 
 	UFUNCTION(BlueprintCallable, Category = "RigVM Blueprint")
-	URigVMGraph* GetFocusedModel() const;
+	virtual URigVMGraph* GetFocusedModel() const override;
 
 	UFUNCTION(BlueprintCallable, Category = "RigVM Blueprint")
-	URigVMController* GetController(const URigVMGraph* InGraph = nullptr) const;
+	virtual URigVMController* GetController(const URigVMGraph* InGraph = nullptr) const override;
 
 	UFUNCTION(BlueprintCallable, Category = "RigVM Blueprint")
-	URigVMController* GetControllerByName(const FString InGraphName = TEXT("")) const;
+	virtual URigVMController* GetControllerByName(const FString InGraphName = TEXT("")) const override;
 
 	UFUNCTION(BlueprintCallable, Category = "RigVM Blueprint")
-	URigVMController* GetOrCreateController(URigVMGraph* InGraph = nullptr);
+	virtual URigVMController* GetOrCreateController(URigVMGraph* InGraph = nullptr) override;
 
-	URigVMController* GetController(const UEdGraph* InEdGraph) const;
-	URigVMController* GetOrCreateController(const UEdGraph* InGraph);
+	virtual URigVMController* GetController(const UEdGraph* InEdGraph) const override;
+	virtual URigVMController* GetOrCreateController(const UEdGraph* InGraph) override;
 
 	UFUNCTION(BlueprintCallable, Category = "RigVM Blueprint")
-	virtual TArray<FString> GeneratePythonCommands(const FString InNewBlueprintName);
+	virtual TArray<FString> GeneratePythonCommands(const FString InNewBlueprintName) override;
+
+	//  --- IRigVMClientHost interface End ---
+
+
+	FOnRigVMRequestInspectObject& OnRequestInspectObject() { return OnRequestInspectObjectEvent; }
+	void RequestInspectObject(const TArray<UObject*>& InObjects) { OnRequestInspectObjectEvent.Broadcast(InObjects); }
+
+	FOnRigVMRequestInspectMemoryStorage& OnRequestInspectMemoryStorage() { return OnRequestInspectMemoryStorageEvent; }
+	void RequestInspectMemoryStorage(const TArray<FRigVMMemoryStorageStruct*>& InMemoryStorageStructs) { OnRequestInspectMemoryStorageEvent.Broadcast(InMemoryStorageStructs); }
+
+#endif	// #if WITH_EDITOR
+
+	
+
+	virtual bool ShouldBeMarkedDirtyUponTransaction() const override { return false; }
 
 	URigVMGraph* GetTemplateModel(bool bIsFunctionLibrary = false);
 	URigVMController* GetTemplateController(bool bIsFunctionLibrary = false);
@@ -731,7 +733,6 @@ private:
 
 	FOnRigVMRequestInspectObject OnRequestInspectObjectEvent;
 	FOnRigVMRequestInspectMemoryStorage OnRequestInspectMemoryStorageEvent;
-	FRigVMBlueprintGetFocusedGraph OnGetFocusedGraphDelegate;
 		
 public:
 
diff --git a/Engine/Plugins/Runtime/RigVM/Source/RigVMDeveloper/Public/RigVMModel/RigVMClient.h b/Engine/Plugins/Runtime/RigVM/Source/RigVMDeveloper/Public/RigVMModel/RigVMClient.h
--- a/Engine/Plugins/Runtime/RigVM/Source/RigVMDeveloper/Public/RigVMModel/RigVMClient.h
+++ b/Engine/Plugins/Runtime/RigVM/Source/RigVMDeveloper/Public/RigVMModel/RigVMClient.h
@@ -16,6 +16,14 @@ class RIGVMDEVELOPER_API URigVMClientHost : public UInterface
 	GENERATED_BODY()
 };
 
+enum class ERigVMLoadType : uint8
+{
+	PostLoad,
+	CheckUserDefinedStructs
+};
+
+DECLARE_DELEGATE_RetVal(URigVMGraph*, FRigVMGetFocusedGraph);
+
 // Interface that allows an object to host a rig VM client. Used by graph edting code to interact with the controller.
 class RIGVMDEVELOPER_API IRigVMClientHost
 {
@@ -56,6 +64,64 @@ public:
 	// Given a type name for a user-defined type, either struct or an enum, returns a pointer to the object definition
 	// or nullptr if the client host has no knowledge of such a type.
 	virtual UObject* ResolveUserDefinedTypeById(const FString& InTypeName) const;
+
+	// Recompiles the VM if not already being compiled
+	virtual void RecompileVM() = 0;
+
+	// Recompiles VM if flagged for recompiling is set
+	virtual void RecompileVMIfRequired() = 0;
+
+	// Flags VM requires recompile, and if auto recompile is enabled and no compile bracket is active, requests a recompilation
+	virtual void RequestAutoVMRecompilation() = 0;
+
+	// Sets flag for automatic recompile on model changes
+	virtual void SetAutoVMRecompile(bool bAutoRecompile) = 0;
+
+	// Returns current state of automatic recompile flag
+	virtual bool GetAutoVMRecompile() const = 0;
+
+	// Helper to increase recompile bracket on nested requests
+	virtual void IncrementVMRecompileBracket() = 0;
+
+	// Helper to decrease recompile bracket on nested requests. When value == 1, if autorecompile is enabled, it triggers a VM recompilation
+	virtual void DecrementVMRecompileBracket() = 0;
+
+	// Regenerates model pins if data has changed while the RigVM Graph is not opened (i.e. user defined struct is changed)
+	virtual void RefreshAllModels(ERigVMLoadType InLoadType) = 0;
+
+	virtual void OnRigVMRegistryChanged() = 0;
+
+	virtual void RequestRigVMInit() = 0;
+
+	virtual URigVMGraph* GetModel(const UEdGraph* InEdGraph = nullptr) const = 0;
+	virtual URigVMGraph* GetModel(const FString& InNodePath) const = 0;
+
+	virtual URigVMGraph* GetDefaultModel() const = 0;
+
+	virtual TArray<URigVMGraph*> GetAllModels() const = 0;
+
+	virtual URigVMFunctionLibrary* GetLocalFunctionLibrary() const = 0;
+
+	virtual URigVMGraph* AddModel(FString InName = TEXT("Rig Graph"), bool bSetupUndoRedo = true, bool bPrintPythonCommand = true) = 0;
+
+	virtual bool RemoveModel(FString InName = TEXT("Rig Graph"), bool bSetupUndoRedo = true, bool bPrintPythonCommand = true) = 0;
+
+
+	virtual FRigVMGetFocusedGraph& OnGetFocusedGraph() = 0;
+	virtual const FRigVMGetFocusedGraph& OnGetFocusedGraph() const = 0;
+
+	virtual URigVMGraph* GetFocusedModel() const = 0;
+
+	virtual URigVMController* GetController(const URigVMGraph* InGraph = nullptr) const = 0;
+
+	virtual URigVMController* GetControllerByName(const FString InGraphName = TEXT("")) const = 0;
+
+	virtual URigVMController* GetOrCreateController(URigVMGraph* InGraph = nullptr) = 0;
+
+	virtual URigVMController* GetController(const UEdGraph* InEdGraph) const = 0;
+	virtual URigVMController* GetOrCreateController(const UEdGraph* InGraph) = 0;
+
+	virtual TArray<FString> GeneratePythonCommands(const FString InNewBlueprintName)  = 0;
 };
 
 UINTERFACE()
@@ -113,6 +179,7 @@ public:
 	URigVMSchema* GetOrCreateSchema();
 	URigVMGraph* GetDefaultModel() const;
 	URigVMGraph* GetModel(int32 InIndex) const { return Models[InIndex]; }
+	URigVMGraph* GetModel(const UEdGraph* InEdGraph = nullptr) const;
 	URigVMGraph* GetModel(const FString& InNodePathOrName) const;
 	URigVMGraph* GetModel(const UObject* InEditorSideObject) const;
 	TArray<URigVMGraph*> GetAllModels(bool bIncludeFunctionLibrary, bool bRecursive) const;
@@ -125,15 +192,22 @@ public:
 	URigVMController* GetOrCreateController(const FString& InNodePathOrName);
 	URigVMController* GetOrCreateController(const URigVMGraph* InModel);
 	URigVMController* GetOrCreateController(const UObject* InEditorSideObject);
+	URigVMController* GetControllerByName(const FString InGraphName) const;
 	bool RemoveController(const URigVMGraph* InModel);
 	URigVMFunctionLibrary* GetFunctionLibrary() const { return FunctionLibrary; }
 	URigVMFunctionLibrary* GetOrCreateFunctionLibrary(bool bSetupUndoRedo, const FObjectInitializer* ObjectInitializer = nullptr, bool bCreateController = true);
 	TArray<FName> GetEntryNames() const;
 	UScriptStruct* GetExecuteContextStruct() const;
 	void SetExecuteContextStruct(UScriptStruct* InExecuteContextStruct);
-	
+
+	FRigVMGetFocusedGraph& OnGetFocusedGraph() { return OnGetFocusedGraphDelegate;}
+	const FRigVMGetFocusedGraph& OnGetFocusedGraph() const { return OnGetFocusedGraphDelegate; }
+	URigVMGraph* GetFocusedModel() const;
+
+	URigVMGraph* AddModel(const FString InName, bool bSetupUndoRedo, bool bPrintPythonCommand);
 	URigVMGraph* AddModel(const FName& InName, bool bSetupUndoRedo, const FObjectInitializer* ObjectInitializer = nullptr, bool bCreateController = true);
 	void AddModel(URigVMGraph* InModel, bool bCreateController);
+	bool RemoveModel(FString InName, bool bSetupUndoRedo, bool bPrintPythonCommand);
 	bool RemoveModel(const FString& InNodePathOrName, bool bSetupUndoRedo);
 	FName RenameModel(const FString& InNodePathOrName, const FName& InNewName, bool bSetupUndoRedo);
 	void PostTransacted(const FTransactionObjectEvent& TransactionEvent);
@@ -205,6 +279,8 @@ private:
 
 	void SetSchema(URigVMSchema* InSchema);
 
+	FRigVMGetFocusedGraph OnGetFocusedGraphDelegate;
+
 	UPROPERTY(transient)
 	TObjectPtr<URigVMSchema> SchemaPtr;
 
diff --git a/Engine/Plugins/Runtime/RigVM/Source/RigVMEditor/Private/Editor/RigVMEditor.cpp b/Engine/Plugins/Runtime/RigVM/Source/RigVMEditor/Private/Editor/RigVMEditor.cpp
--- a/Engine/Plugins/Runtime/RigVM/Source/RigVMEditor/Private/Editor/RigVMEditor.cpp
+++ b/Engine/Plugins/Runtime/RigVM/Source/RigVMEditor/Private/Editor/RigVMEditor.cpp
@@ -148,7 +148,7 @@ void FRigVMEditor::InitRigVMEditor(const EToolkitMode::Type Mode, const TSharedP
 	// user-defined-struct can change even after load
 	// refresh the models such that pins are updated to match
 	// the latest struct member layout
-	InRigVMBlueprint->RefreshAllModels(ERigVMBlueprintLoadType::CheckUserDefinedStructs);
+	InRigVMBlueprint->RefreshAllModels(ERigVMLoadType::CheckUserDefinedStructs);
 
 	{
 		TArray<UEdGraph*> EdGraphs;
