Changed ScopedStruct use in favor of InstancePropertyBagStructureDataProvider when adding InstancedStructs to details panels, as the ScopedStruct might crash if the InstancedStruct is changed while being displayed (becasue the UStruct that is passed to the ScopedStruct is destroyed)

#rnx

diff --git a/Engine/Plugins/Experimental/AnimNext/Source/AnimNextEditor/Private/Param/ParameterBlockParameterCustomization.cpp b/Engine/Plugins/Experimental/AnimNext/Source/AnimNextEditor/Private/Param/ParameterBlockParameterCustomization.cpp
--- a/Engine/Plugins/Experimental/AnimNext/Source/AnimNextEditor/Private/Param/ParameterBlockParameterCustomization.cpp
+++ b/Engine/Plugins/Experimental/AnimNext/Source/AnimNextEditor/Private/Param/ParameterBlockParameterCustomization.cpp
@@ -43,22 +43,17 @@ void FParameterBlockParameterCustomization::CustomizeDetails(IDetailLayoutBuilde
 
 		if (UAnimNextParameterBlock_EditorData* EditorData = Cast<UAnimNextParameterBlock_EditorData>(BlockParam->GetOuter()))
 		{
-			if (UAnimNextParameterBlock* ReferencedBlock = UE::AnimNext::UncookedOnly::FUtils::GetBlock(EditorData))
+			const FName ParameterName = BlockParam->GetParameterName();
+			if (UAnimNextParameterBlockEntry* BlockEntry = EditorData->FindBinding(ParameterName)) // filter non parameter entries
 			{
-				FSinglePropertyParams SinglePropertyArgs;
-				const FName ParameterName = BlockParam->GetParameterName();
-				TArray<IDetailPropertyRow*> DetailPropertyRows;
-
-				// FIXME: As AddExternalStructureProperty does not support IStructureDataProvider, we add all properties and only show the one we need. Icky.
-				DefaultValueCategory.AddAllExternalStructureProperties(MakeShared<FInstancePropertyBagStructureDataProvider>(ReferencedBlock->GetPropertyBag()), EPropertyLocation::Default, &DetailPropertyRows);
-				for (IDetailPropertyRow* DetailPropertyRow : DetailPropertyRows)
+				if (UAnimNextParameterBlock* ReferencedBlock = UE::AnimNext::UncookedOnly::FUtils::GetBlock(EditorData))
 				{
-					if(DetailPropertyRow->GetPropertyHandle()->GetProperty()->GetFName() != ParameterName)
-					{
-						DetailPropertyRow->Visibility(EVisibility::Collapsed);
-					}
-					else
+					FAddPropertyParams AddPropertyParams;
+					TArray<IDetailPropertyRow*> DetailPropertyRows;
+
+					if (ReferencedBlock->GetPropertyBag().FindPropertyDescByName(ParameterName))
 					{
+						IDetailPropertyRow* DetailPropertyRow = DefaultValueCategory.AddExternalStructureProperty(MakeShared<FInstancePropertyBagStructureDataProvider>(ReferencedBlock->GetPropertyBag()), ParameterName, EPropertyLocation::Default, AddPropertyParams);
 						if (TSharedPtr<IPropertyHandle> Handle = DetailPropertyRow->GetPropertyHandle(); Handle.IsValid())
 						{
 							Handle->SetOnChildPropertyValuePreChange(FSimpleDelegate::CreateLambda([this, ReferencedBlock]()
diff --git a/Engine/Plugins/Runtime/RigVM/Source/RigVMEditor/Private/Editor/RigVMEditor.cpp b/Engine/Plugins/Runtime/RigVM/Source/RigVMEditor/Private/Editor/RigVMEditor.cpp
--- a/Engine/Plugins/Runtime/RigVM/Source/RigVMEditor/Private/Editor/RigVMEditor.cpp
+++ b/Engine/Plugins/Runtime/RigVM/Source/RigVMEditor/Private/Editor/RigVMEditor.cpp
@@ -43,6 +43,7 @@
 #include "Framework/Docking/TabManager.h"
 #include "ScopedTransaction.h"
 #include "Editor/RigVMEditorMode.h"
+#include "InstancedPropertyBagStructureDataProvider.h"
 
 #define LOCTEXT_NAMESPACE "RigVMEditor"
 
@@ -1778,10 +1779,9 @@ void FRigVMEditor::HandleVMCompiledEvent(UObject* InCompiledObject, URigVM* InVM
 				const TArray<FRigVMMemoryStorageStruct*> MemoryStorage = { Memory };
 				SetMemoryStorageDetails(MemoryStorage);
 			#else
-				// TODO zzz : This does not compile, IStructureDetailsView does not inherit from SCompoundWidget, like IDetailsView
+				// TODO zzz : need a way to get the IStructureDetailsView
 				TSharedRef<IStructureDetailsView> StructDetailsView = StaticCastSharedRef<IStructureDetailsView>(ActiveTab->GetContent());
-				TSharedPtr<FStructOnScope> StructOnScope = MakeShareable(new FStructOnScope(Memory->GetPropertyBagStruct(), (uint8*)Memory->GetContainerPtr()));
-				StructDetailsView->SetStructureData(StructOnScope);
+				StructDetailsView->SetStructureProvider(MakeShared<FInstancePropertyBagStructureDataProvider>(*Memory));
 			#endif
 			}
 		}
@@ -2666,7 +2666,7 @@ void FRigVMEditor::SetMemoryStorageDetails(const TArray<FRigVMMemoryStorageStruc
 
 	if (InStructs.Num() == 1)
 	{
-		if (const FRigVMMemoryStorageStruct* Memory = InStructs[0])
+		if (FRigVMMemoryStorageStruct* Memory = InStructs[0])
 		{
 			FPropertyEditorModule& EditModule = FModuleManager::Get().GetModuleChecked<FPropertyEditorModule>("PropertyEditor");
 
@@ -2676,9 +2676,8 @@ void FRigVMEditor::SetMemoryStorageDetails(const TArray<FRigVMMemoryStorageStruc
 
 			FStructureDetailsViewArgs StructureViewArgs;
 
-			TSharedRef<IStructureDetailsView> DetailsView = EditModule.CreateStructureDetailView(DetailsViewArgs, StructureViewArgs, TSharedPtr<FStructOnScope>());
-			TSharedPtr<FStructOnScope> StructOnScope = MakeShareable(new FStructOnScope(Memory->GetPropertyBagStruct(), (uint8*)Memory->GetContainerPtr()));
-			DetailsView->SetStructureData(StructOnScope);
+			TSharedRef<IStructureDetailsView> DetailsView = EditModule.CreateStructureDetailView(DetailsViewArgs, StructureViewArgs, nullptr);
+			DetailsView->SetStructureProvider(MakeShared<FInstancePropertyBagStructureDataProvider>(*Memory));
 
 			TSharedRef<SDockTab> DockTab = SNew(SDockTab)
 				.Label(LOCTEXT("RigVMMemoryDetails", "RigVM Memory Details"))
diff --git a/Engine/Plugins/Runtime/RigVM/Source/RigVMEditor/Private/Editor/RigVMGraphDetailCustomization.cpp b/Engine/Plugins/Runtime/RigVM/Source/RigVMEditor/Private/Editor/RigVMGraphDetailCustomization.cpp
--- a/Engine/Plugins/Runtime/RigVM/Source/RigVMEditor/Private/Editor/RigVMGraphDetailCustomization.cpp
+++ b/Engine/Plugins/Runtime/RigVM/Source/RigVMEditor/Private/Editor/RigVMGraphDetailCustomization.cpp
@@ -25,6 +25,7 @@
 #include "IPropertyUtilities.h"
 #include "RigVMModel/Nodes/RigVMAggregateNode.h"
 #include "Widgets/SRigVMGraphPinVariableBinding.h"
+#include "InstancedPropertyBagStructureDataProvider.h"
 
 #define LOCTEXT_NAMESPACE "RigVMGraphDetailCustomization"
 
@@ -1793,10 +1794,9 @@ void FRigVMWrappedNodeDetailCustomization::CustomizeLiveValues(IDetailLayoutBuil
 					}
 				}
 
-				for (const FRigVMMemoryStorageStruct* Memory : ExternalStructs)
+				for (FRigVMMemoryStorageStruct* Memory : ExternalStructs)
 				{
-					const TSharedPtr<FStructOnScope> StructOnScope = MakeShareable(new FStructOnScope(Memory->GetPropertyBagStruct(), (uint8*)Memory->GetContainerPtr()));
-					if (IDetailPropertyRow* PropertyRow = DebugCategory.AddExternalStructureProperty(StructOnScope, Property->GetFName(), EPropertyLocation::Default, AddPropertyParams))
+					if (IDetailPropertyRow* PropertyRow = DebugCategory.AddExternalStructureProperty(MakeShared<FInstancePropertyBagStructureDataProvider>(*Memory), Property->GetFName(), EPropertyLocation::Default, AddPropertyParams))
 					{
 						UpdateRow(PropertyRow);
 					}
diff --git a/Engine/Plugins/Runtime/RigVM/Source/RigVMEditor/Private/Editor/RigVMLocalVariableDetailCustomization.cpp b/Engine/Plugins/Runtime/RigVM/Source/RigVMEditor/Private/Editor/RigVMLocalVariableDetailCustomization.cpp
--- a/Engine/Plugins/Runtime/RigVM/Source/RigVMEditor/Private/Editor/RigVMLocalVariableDetailCustomization.cpp
+++ b/Engine/Plugins/Runtime/RigVM/Source/RigVMEditor/Private/Editor/RigVMLocalVariableDetailCustomization.cpp
@@ -12,6 +12,7 @@
 #include "RigVMBlueprintGeneratedClass.h"
 #include "RigVMBlueprint.h"
 #include "RigVMCore/RigVM.h"
+#include "InstancedPropertyBagStructureDataProvider.h"
 
 #define LOCTEXT_NAMESPACE "LocalVariableDetails"
 
@@ -104,8 +105,7 @@ void FRigVMLocalVariableDetailCustomization::CustomizeDetails(IDetailLayoutBuild
 				Property->ClearPropertyFlags(CPF_EditConst);
 			
 				const FName SanitizedName = FRigVMPropertyDescription::SanitizeName(*SourcePath);
-				TSharedPtr<FStructOnScope> StructOnScope = MakeShareable(new FStructOnScope(LiteralMemory->GetPropertyBagStruct(), (uint8*)LiteralMemory->GetContainerPtr()));
-				IDetailPropertyRow* Row = DefaultValueCategory.AddExternalStructureProperty(StructOnScope, SanitizedName);
+				IDetailPropertyRow* Row = DefaultValueCategory.AddExternalStructureProperty(MakeShared<FInstancePropertyBagStructureDataProvider>(*LiteralMemory), SanitizedName);
 
 				Row->DisplayName(FText::FromName(VariableDescription.Name));
 
