Fix compile error from missing forward declaration.

diff --git a/Engine/Plugins/Runtime/RigVM/Source/RigVM/Public/RigVMCore/RigVMExecuteContext.h b/Engine/Plugins/Runtime/RigVM/Source/RigVM/Public/RigVMCore/RigVMExecuteContext.h
--- a/Engine/Plugins/Runtime/RigVM/Source/RigVM/Public/RigVMCore/RigVMExecuteContext.h
+++ b/Engine/Plugins/Runtime/RigVM/Source/RigVM/Public/RigVMCore/RigVMExecuteContext.h
@@ -37,6 +37,7 @@ class URigVM;
 struct FRigVMDispatchFactory;
 struct FRigVMExecuteContext;
 struct FRigVMExtendedExecuteContext;
+struct FRigVMLogSettings;
 
 USTRUCT()
 struct RIGVM_API FRigVMSlice
