Fixes issue with RigVM spamming warning when creating a storage class for an already loaded type.

diff --git a/Engine/Plugins/Runtime/RigVM/Source/RigVM/Private/RigVMCore/RigVMMemoryStorage.cpp b/Engine/Plugins/Runtime/RigVM/Source/RigVM/Private/RigVMCore/RigVMMemoryStorage.cpp
--- a/Engine/Plugins/Runtime/RigVM/Source/RigVM/Private/RigVMCore/RigVMMemoryStorage.cpp
+++ b/Engine/Plugins/Runtime/RigVM/Source/RigVM/Private/RigVMCore/RigVMMemoryStorage.cpp
@@ -510,30 +510,46 @@ URigVMMemoryStorageGeneratorClass* URigVMMemoryStorageGeneratorClass::CreateStor
 	UObject* OldClassObject = StaticFindObjectFastInternal( /*Class=*/ NULL, Package, *ClassName, true, RF_NoFlags, EInternalObjectFlags::None);
 	if(OldClassObject)
 	{
+		auto FindUniqueName = [](const TCHAR* TemplateName, int32& TemplateIndex)
+		{
+			FName UniqueName;
+			do
+			{
+				UniqueName = FName(TemplateName, TemplateIndex++);
+				if(StaticFindObjectFast(nullptr, GetTransientPackage(), UniqueName) == nullptr)
+				{
+					break;
+				}
+			}
+			while (TemplateIndex < INT_MAX);
+ 
+			return UniqueName;
+		};
+		
+		auto RenameAndMarkGarbage = [](UObject* InObject, FName NewName)
+		{
+			InObject->Rename(*NewName.ToString(), GetTransientPackage(), REN_ForceNoResetLoaders | REN_DoNotDirty | REN_DontCreateRedirectors | REN_NonTransactional);
+			InObject->MarkAsGarbage();
+		};
+ 
 		// ensure the OldClass is completely loaded so that we can remove it and replace it with the new class
 		// otherwise we get an assertion trying to replace objects currently being loaded
 		OldClassObject->ConditionalPostLoad();
-
 		UClass* OldClass = Cast<UClass>(OldClassObject);
-		
-		FString DiscardedMemoryClassName;
-		static const TCHAR DiscardedMemoryClassTemplate[] = TEXT("DiscardedMemoryClassTemplate_%d");
+ 
+		// using an FName to avoid formatting a string and then immediately turning around to parse it.
+		static const TCHAR* DiscardedMemoryClassTemplate(TEXT("DiscardedMemoryClassTemplate"));
 		static int32 DiscardedMemoryClassIndex = 0;
-		do
-		{
-			DiscardedMemoryClassName = FString::Printf(DiscardedMemoryClassTemplate, DiscardedMemoryClassIndex++);
-			if(StaticFindObjectFast(nullptr, GetTransientPackage(), *DiscardedMemoryClassName) == nullptr)
-			{
-				break;
-			}
-		}
-		while (DiscardedMemoryClassIndex < INT_MAX);
-
+		const FName DiscardedMemoryClassName = FindUniqueName(DiscardedMemoryClassTemplate, DiscardedMemoryClassIndex);
+		
 		OldClass->ClassFlags |= CLASS_NewerVersionExists;
-		OldClass->Rename(*DiscardedMemoryClassName, GetTransientPackage(), REN_ForceNoResetLoaders | REN_DoNotDirty | REN_DontCreateRedirectors | REN_NonTransactional);
+		RenameAndMarkGarbage(OldClass, DiscardedMemoryClassName);
+ 
 		if (OldClass->ClassDefaultObject)
 		{
-			OldClass->ClassDefaultObject->Rename(nullptr, GetTransientPackage(), REN_ForceNoResetLoaders | REN_DoNotDirty | REN_DontCreateRedirectors | REN_NonTransactional);
+			static const TCHAR* DiscardedMemoryCDOTemplate(TEXT("DiscardedMemoryCDOTemplate"));
+			static int32 DiscardedMemoryCDOIndex = 0;
+			RenameAndMarkGarbage(OldClass->ClassDefaultObject, FindUniqueName(DiscardedMemoryCDOTemplate, DiscardedMemoryCDOIndex));
 		}
 	}
 
