Control Rig: Compile all public functions even if they are not used in the graphs
#rnx

diff --git a/Engine/Plugins/Runtime/RigVM/Source/RigVMDeveloper/Private/RigVMCompiler/RigVMCompiler.cpp b/Engine/Plugins/Runtime/RigVM/Source/RigVMDeveloper/Private/RigVMCompiler/RigVMCompiler.cpp
--- a/Engine/Plugins/Runtime/RigVM/Source/RigVMDeveloper/Private/RigVMCompiler/RigVMCompiler.cpp
+++ b/Engine/Plugins/Runtime/RigVM/Source/RigVMDeveloper/Private/RigVMCompiler/RigVMCompiler.cpp
@@ -450,6 +450,26 @@ bool URigVMCompiler::Compile(const FRigVMCompileSettings& InSettings, TArray<URi
 
 	TMap<FString, const FRigVMFunctionCompilationData*> CurrentCompiledFunctions;
 
+#if WITH_EDITOR
+	if (!CurrentCompilationFunction)
+	{
+		if (FunctionLibrary)
+		{
+			TArray<URigVMLibraryNode*> Functions = FunctionLibrary->GetFunctions();
+			for (URigVMLibraryNode* Function : Functions)
+			{
+				if (FunctionLibrary->IsFunctionPublic(Function->GetFName()))
+				{
+					if (FRigVMGraphFunctionData* FunctionData = Function->GetFunctionHeader().GetFunctionData())
+					{
+						CompileFunction(WorkData.Settings, Function, InController, InExternalVariables, &FunctionData->CompilationData, OutVMContext);
+					}
+				}
+			}
+		}
+	}
+#endif
+
 	// Gather function compilation data
 	for(URigVMGraph* Graph : InGraphs)
 	{
