Control Rig: Fix warning from usage of deprecated function GetAssetRegistryTags
#rnx

diff --git a/Engine/Plugins/Runtime/RigVM/Source/RigVM/Private/RigVMBlueprintGeneratedClass.cpp b/Engine/Plugins/Runtime/RigVM/Source/RigVM/Private/RigVMBlueprintGeneratedClass.cpp
--- a/Engine/Plugins/Runtime/RigVM/Source/RigVM/Private/RigVMBlueprintGeneratedClass.cpp
+++ b/Engine/Plugins/Runtime/RigVM/Source/RigVM/Private/RigVMBlueprintGeneratedClass.cpp
@@ -3,6 +3,7 @@
 #include "RigVMBlueprintGeneratedClass.h"
 #include "UObject/UE5MainStreamObjectVersion.h"
 #include "RigVMHost.h"
+#include "UObject/AssetRegistryTagsContext.h"
 
 #include UE_INLINE_GENERATED_CPP_BY_NAME(RigVMBlueprintGeneratedClass)
 
@@ -102,7 +103,14 @@ void URigVMBlueprintGeneratedClass::PostLoad()
 
 void URigVMBlueprintGeneratedClass::GetAssetRegistryTags(TArray<FAssetRegistryTag>& OutTags) const
 {
+	PRAGMA_DISABLE_DEPRECATION_WARNINGS;
 	Super::GetAssetRegistryTags(OutTags);
+	PRAGMA_ENABLE_DEPRECATION_WARNINGS;
+}
+
+void URigVMBlueprintGeneratedClass::GetAssetRegistryTags(FAssetRegistryTagsContext Context) const
+{
+	Super::GetAssetRegistryTags(Context);
 
 	const FArrayProperty* HeadersProperty = CastField<FArrayProperty>(FRigVMGraphFunctionHeaderArray::StaticStruct()->FindPropertyByName(TEXT("Headers")));
 	FRigVMGraphFunctionHeaderArray HeaderArray;
@@ -115,5 +123,5 @@ void URigVMBlueprintGeneratedClass::GetAssetRegistryTags(TArray<FAssetRegistryTa
 	FString HeadersString;
 	HeadersProperty->ExportText_Direct(HeadersString, &(HeaderArray.Headers), &(HeaderArray.Headers), nullptr, PPF_None, nullptr);
 
-	OutTags.Add(UObject::FAssetRegistryTag(TEXT("PublicGraphFunctions"), HeadersString, UObject::FAssetRegistryTag::TT_Hidden));
+	Context.AddTag(UObject::FAssetRegistryTag(TEXT("PublicGraphFunctions"), HeadersString, UObject::FAssetRegistryTag::TT_Hidden));
 }
diff --git a/Engine/Plugins/Runtime/RigVM/Source/RigVM/Public/RigVMBlueprintGeneratedClass.h b/Engine/Plugins/Runtime/RigVM/Source/RigVM/Public/RigVMBlueprintGeneratedClass.h
--- a/Engine/Plugins/Runtime/RigVM/Source/RigVM/Public/RigVMBlueprintGeneratedClass.h
+++ b/Engine/Plugins/Runtime/RigVM/Source/RigVM/Public/RigVMBlueprintGeneratedClass.h
@@ -32,7 +32,9 @@ public:
 	// UObject interface
 	virtual void Serialize(FArchive& Ar) override;
 	virtual void PostLoad() override;
+	UE_DEPRECATED(5.4, "Implement the version that takes FAssetRegistryTagsContext instead.")
 	virtual void GetAssetRegistryTags(TArray<FAssetRegistryTag>& OutTags) const override;
+	virtual void GetAssetRegistryTags(FAssetRegistryTagsContext Context) const override;
 
 	// IRigVMGraphFunctionHost interface
 	virtual const FRigVMGraphFunctionStore* GetRigVMGraphFunctionStore() const override { return &GraphFunctionStore; }
