More "unreachable code" fixes after enabling detailed warnings.

diff --git a/Engine/Plugins/Online/OnlineServicesEOSGS/Source/Private/Online/AuthEOSGS.cpp b/Engine/Plugins/Online/OnlineServicesEOSGS/Source/Private/Online/AuthEOSGS.cpp
--- a/Engine/Plugins/Online/OnlineServicesEOSGS/Source/Private/Online/AuthEOSGS.cpp
+++ b/Engine/Plugins/Online/OnlineServicesEOSGS/Source/Private/Online/AuthEOSGS.cpp
@@ -598,7 +598,6 @@ TOnlineAsyncOpHandle<FAuthLogin> FAuthEOSGS::Login(FAuthLogin::Params&& Params)
 			if (!IsOnlineStatus(AccountInfoEOS->LoginStatus))
 			{
 				TPromise<void> Promise;
-				TFuture<void> Future = Promise.GetFuture();
 
 				const FAuthLogin::Params& Params = InAsyncOp.GetParams();
 
@@ -627,8 +626,6 @@ TOnlineAsyncOpHandle<FAuthLogin> FAuthEOSGS::Login(FAuthLogin::Params&& Params)
 
 					Promise.EmplaceValue();
 				});
-
-				return Future;
 			}
 			else
 			{
diff --git a/Engine/Plugins/Runtime/RigVM/Source/RigVMDeveloper/Private/RigVMModel/Nodes/RigVMUnitNode.cpp b/Engine/Plugins/Runtime/RigVM/Source/RigVMDeveloper/Private/RigVMModel/Nodes/RigVMUnitNode.cpp
--- a/Engine/Plugins/Runtime/RigVM/Source/RigVMDeveloper/Private/RigVMModel/Nodes/RigVMUnitNode.cpp
+++ b/Engine/Plugins/Runtime/RigVM/Source/RigVMDeveloper/Private/RigVMModel/Nodes/RigVMUnitNode.cpp
@@ -284,12 +284,11 @@ FName URigVMUnitNode::GetNextAggregateName(const FName& InLastAggregatePinName)
 		const FRigVMStruct* StructMemory = (const FRigVMStruct*)StructOnScope->GetStructMemory();
 		return StructMemory->GetNextAggregateName(InLastAggregatePinName);
 	}
-	else
-	{
-		return Super::GetNextAggregateName(InLastAggregatePinName);
-	}
-#endif
+
+	return Super::GetNextAggregateName(InLastAggregatePinName);
+#else
 	return FName();
+#endif
 }
 
 UScriptStruct* URigVMUnitNode::GetScriptStruct() const
