Control Rig: Fix Undo/Redo in AnimSequencer when adding a SkeletalMesh that has a ControlRig set
#rnx

diff --git a/Engine/Plugins/Runtime/RigVM/Source/RigVM/Public/RigVMCore/RigVMExecuteContext.h b/Engine/Plugins/Runtime/RigVM/Source/RigVM/Public/RigVMCore/RigVMExecuteContext.h
--- a/Engine/Plugins/Runtime/RigVM/Source/RigVM/Public/RigVMCore/RigVMExecuteContext.h
+++ b/Engine/Plugins/Runtime/RigVM/Source/RigVM/Public/RigVMCore/RigVMExecuteContext.h
@@ -713,13 +713,10 @@ struct RIGVM_API FRigVMExtendedExecuteContext
 		return 0;
 	}
 
-	UPROPERTY(transient)
 	uint32 VMHash = 0;
 
-	UPROPERTY(transient)
 	FRigVMMemoryStorageStruct WorkMemoryStorage;
 
-	UPROPERTY(transient)
 	FRigVMMemoryStorageStruct DebugMemoryStorage;
 
 #if WITH_EDITORONLY_DATA
