Control Rig: Allow public functions to be localized
#rnx

diff --git a/Engine/Plugins/Runtime/RigVM/Source/RigVMDeveloper/Private/RigVMModel/RigVMController.cpp b/Engine/Plugins/Runtime/RigVM/Source/RigVMDeveloper/Private/RigVMModel/RigVMController.cpp
--- a/Engine/Plugins/Runtime/RigVM/Source/RigVMDeveloper/Private/RigVMModel/RigVMController.cpp
+++ b/Engine/Plugins/Runtime/RigVM/Source/RigVMDeveloper/Private/RigVMModel/RigVMController.cpp
@@ -4476,12 +4476,6 @@ TMap<FRigVMGraphFunctionIdentifier, URigVMLibraryNode*> URigVMController::Locali
 			return LocalizedFunctions;
 		}
 
-		// Do not localize public functions
-		if (bIsPublic)
-		{
-			continue;
-		}
-
 		if (!bLocalizeDependentPrivateFunctions)
 		{
 			ReportAndNotifyErrorf(TEXT("Cannot localize function - dependency %s is private."), *NodeToVisit.LibraryNode.ToString());
