Bind RigVMHost type so blueprint created callback can be called on control rigs

diff --git a/Engine/Plugins/Runtime/RigVM/Source/RigVMEditor/Private/RigVMEditorModule.cpp b/Engine/Plugins/Runtime/RigVM/Source/RigVMEditor/Private/RigVMEditorModule.cpp
--- a/Engine/Plugins/Runtime/RigVM/Source/RigVMEditor/Private/RigVMEditorModule.cpp
+++ b/Engine/Plugins/Runtime/RigVM/Source/RigVMEditor/Private/RigVMEditorModule.cpp
@@ -93,7 +93,7 @@ void FRigVMEditorModule::StartupModuleCommon()
 	BlueprintEditorModule.RegisterGraphCustomization(SchemaCDO, FOnGetGraphCustomizationInstance::CreateStatic(&FRigVMGraphDetailCustomization::MakeInstance, BlueprintClass));
 
 	// Register to fixup newly created BPs
-	FKismetEditorUtilities::RegisterOnBlueprintCreatedCallback(this, BlueprintCDO->GetRigVMHostClass(), FKismetEditorUtilities::FOnBlueprintCreated::CreateRaw(this, &FRigVMEditorModule::HandleNewBlueprintCreated));
+	FKismetEditorUtilities::RegisterOnBlueprintCreatedCallback(this, URigVMHost::StaticClass(), FKismetEditorUtilities::FOnBlueprintCreated::CreateRaw(this, &FRigVMEditorModule::HandleNewBlueprintCreated));
 }
 
 void FRigVMEditorModule::ShutdownModule()
