RigVM: Fix unit test breakage on localize functions

#rnx

diff --git a/Engine/Plugins/Runtime/RigVM/Source/RigVMDeveloper/Private/RigVMModel/RigVMController.cpp b/Engine/Plugins/Runtime/RigVM/Source/RigVMDeveloper/Private/RigVMModel/RigVMController.cpp
--- a/Engine/Plugins/Runtime/RigVM/Source/RigVMDeveloper/Private/RigVMModel/RigVMController.cpp
+++ b/Engine/Plugins/Runtime/RigVM/Source/RigVMDeveloper/Private/RigVMModel/RigVMController.cpp
@@ -4455,6 +4455,8 @@ TMap<FRigVMGraphFunctionIdentifier, URigVMLibraryNode*> URigVMController::Locali
 		FunctionsToLocalize.AddUnique(FRigVMGraphFunctionData::FindFunctionData(FunctionDefinition));
 	}
 
+	const int32 InputNodesToVisitCount = NodesToVisit.Num();
+
 	const FSoftObjectPath ThisFunctionHost = ThisLibrary->GetFunctionHostObjectPath();
 	
 	// find all functions to localize
@@ -4476,6 +4478,12 @@ TMap<FRigVMGraphFunctionIdentifier, URigVMLibraryNode*> URigVMController::Locali
 			return LocalizedFunctions;
 		}
 
+		// Do not localize public functions if they are not part of the input set of functions
+		if (bIsPublic && NodeToVisitIndex >= InputNodesToVisitCount)
+		{
+			continue;
+		}
+
 		if (!bLocalizeDependentPrivateFunctions)
 		{
 			ReportAndNotifyErrorf(TEXT("Cannot localize function - dependency %s is private."), *NodeToVisit.LibraryNode.ToString());
