RigVM: Disable aligning bytecode on old / invalid assets
Allow bytecode retrieval for old execute ops

#rnx

diff --git a/Engine/Plugins/Runtime/RigVM/Source/RigVM/Private/RigVMCore/RigVM.cpp b/Engine/Plugins/Runtime/RigVM/Source/RigVM/Private/RigVMCore/RigVM.cpp
--- a/Engine/Plugins/Runtime/RigVM/Source/RigVM/Private/RigVMCore/RigVM.cpp
+++ b/Engine/Plugins/Runtime/RigVM/Source/RigVM/Private/RigVMCore/RigVM.cpp
@@ -186,7 +186,6 @@ void URigVM::Load(FArchive& Ar)
 			Ar << TmpLiteralMemoryStorage;
 			Ar << FunctionNamesStorage;
 			Ar << ByteCodeStorage;
-			ByteCodeStorage.AlignByteCode();
 			Ar << Parameters;
 
 			if (Ar.CustomVer(FUE5MainStreamObjectVersion::GUID) < FUE5MainStreamObjectVersion::RigVMCopyOpStoreNumBytes)
diff --git a/Engine/Plugins/Runtime/RigVM/Source/RigVM/Private/RigVMCore/RigVMByteCode.cpp b/Engine/Plugins/Runtime/RigVM/Source/RigVM/Private/RigVMCore/RigVMByteCode.cpp
--- a/Engine/Plugins/Runtime/RigVM/Source/RigVM/Private/RigVMCore/RigVMByteCode.cpp
+++ b/Engine/Plugins/Runtime/RigVM/Source/RigVM/Private/RigVMCore/RigVMByteCode.cpp
@@ -876,6 +876,71 @@ uint64 FRigVMByteCode::GetOpNumBytesAt(uint64 InByteCodeIndex, bool bIncludeOper
 	ERigVMOpCode OpCode = GetOpCodeAt(InByteCodeIndex);
 	switch (OpCode)
 	{
+		case ERigVMOpCode::Execute_0_Operands:
+		case ERigVMOpCode::Execute_1_Operands:
+		case ERigVMOpCode::Execute_2_Operands:
+		case ERigVMOpCode::Execute_3_Operands:
+		case ERigVMOpCode::Execute_4_Operands:
+		case ERigVMOpCode::Execute_5_Operands:
+		case ERigVMOpCode::Execute_6_Operands:
+		case ERigVMOpCode::Execute_7_Operands:
+		case ERigVMOpCode::Execute_8_Operands:
+		case ERigVMOpCode::Execute_9_Operands:
+		case ERigVMOpCode::Execute_10_Operands:
+		case ERigVMOpCode::Execute_11_Operands:
+		case ERigVMOpCode::Execute_12_Operands:
+		case ERigVMOpCode::Execute_13_Operands:
+		case ERigVMOpCode::Execute_14_Operands:
+		case ERigVMOpCode::Execute_15_Operands:
+		case ERigVMOpCode::Execute_16_Operands:
+		case ERigVMOpCode::Execute_17_Operands:
+		case ERigVMOpCode::Execute_18_Operands:
+		case ERigVMOpCode::Execute_19_Operands:
+		case ERigVMOpCode::Execute_20_Operands:
+		case ERigVMOpCode::Execute_21_Operands:
+		case ERigVMOpCode::Execute_22_Operands:
+		case ERigVMOpCode::Execute_23_Operands:
+		case ERigVMOpCode::Execute_24_Operands:
+		case ERigVMOpCode::Execute_25_Operands:
+		case ERigVMOpCode::Execute_26_Operands:
+		case ERigVMOpCode::Execute_27_Operands:
+		case ERigVMOpCode::Execute_28_Operands:
+		case ERigVMOpCode::Execute_29_Operands:
+		case ERigVMOpCode::Execute_30_Operands:
+		case ERigVMOpCode::Execute_31_Operands:
+		case ERigVMOpCode::Execute_32_Operands:
+		case ERigVMOpCode::Execute_33_Operands:
+		case ERigVMOpCode::Execute_34_Operands:
+		case ERigVMOpCode::Execute_35_Operands:
+		case ERigVMOpCode::Execute_36_Operands:
+		case ERigVMOpCode::Execute_37_Operands:
+		case ERigVMOpCode::Execute_38_Operands:
+		case ERigVMOpCode::Execute_39_Operands:
+		case ERigVMOpCode::Execute_40_Operands:
+		case ERigVMOpCode::Execute_41_Operands:
+		case ERigVMOpCode::Execute_42_Operands:
+		case ERigVMOpCode::Execute_43_Operands:
+		case ERigVMOpCode::Execute_44_Operands:
+		case ERigVMOpCode::Execute_45_Operands:
+		case ERigVMOpCode::Execute_46_Operands:
+		case ERigVMOpCode::Execute_47_Operands:
+		case ERigVMOpCode::Execute_48_Operands:
+		case ERigVMOpCode::Execute_49_Operands:
+		case ERigVMOpCode::Execute_50_Operands:
+		case ERigVMOpCode::Execute_51_Operands:
+		case ERigVMOpCode::Execute_52_Operands:
+		case ERigVMOpCode::Execute_53_Operands:
+		case ERigVMOpCode::Execute_54_Operands:
+		case ERigVMOpCode::Execute_55_Operands:
+		case ERigVMOpCode::Execute_56_Operands:
+		case ERigVMOpCode::Execute_57_Operands:
+		case ERigVMOpCode::Execute_58_Operands:
+		case ERigVMOpCode::Execute_59_Operands:
+		case ERigVMOpCode::Execute_60_Operands:
+		case ERigVMOpCode::Execute_61_Operands:
+		case ERigVMOpCode::Execute_62_Operands:
+		case ERigVMOpCode::Execute_63_Operands:
+		case ERigVMOpCode::Execute_64_Operands:
 		case ERigVMOpCode::Execute:
 		{
 			uint64 NumBytes = (uint64)sizeof(FRigVMExecuteOp);
