RigVM : Integrated Hash memory optimization from CL 27353614 to main

#rnx

diff --git a/Engine/Plugins/Runtime/RigVM/Source/RigVM/Private/RigVMCore/RigVMFunction.cpp b/Engine/Plugins/Runtime/RigVM/Source/RigVM/Private/RigVMCore/RigVMFunction.cpp
--- a/Engine/Plugins/Runtime/RigVM/Source/RigVM/Private/RigVMCore/RigVMFunction.cpp
+++ b/Engine/Plugins/Runtime/RigVM/Source/RigVM/Private/RigVMCore/RigVMFunction.cpp
@@ -5,6 +5,7 @@
 #include "RigVMCore/RigVMRegistry.h"
 #include "RigVMCore/RigVMExternalVariable.h"
 #include "RigVMCore/RigVMStruct.h"
+#include "RigVMPropertyUtils.h"
 #include "UObject/Package.h"
 
 FRigVMFunction::FRigVMFunction(const TCHAR* InName, FRigVMFunctionPtr InFunctionPtr, UScriptStruct* InStruct, int32 InIndex, const TArray<FRigVMFunctionArgument>& InArguments)
@@ -88,7 +89,7 @@ const TArray<TRigVMTypeIndex>& FRigVMFunction::GetArgumentTypeIndices() const
 				{
 					FName CPPType = NAME_None;
 					UObject* CPPTypeObject = nullptr;
-					FRigVMExternalVariable::GetTypeFromProperty(Property, CPPType, CPPTypeObject);
+					RigVMPropertyUtils::GetTypeFromProperty(Property, CPPType, CPPTypeObject);
 
 					const FRigVMTemplateArgumentType Type(CPPType, CPPTypeObject);
 					ArgumentTypeIndices.Add(FRigVMRegistry::Get().FindOrAddType(Type));
diff --git a/Engine/Plugins/Runtime/RigVM/Source/RigVM/Private/RigVMCore/RigVMMemoryStorage.cpp b/Engine/Plugins/Runtime/RigVM/Source/RigVM/Private/RigVMCore/RigVMMemoryStorage.cpp
--- a/Engine/Plugins/Runtime/RigVM/Source/RigVM/Private/RigVMCore/RigVMMemoryStorage.cpp
+++ b/Engine/Plugins/Runtime/RigVM/Source/RigVM/Private/RigVMCore/RigVMMemoryStorage.cpp
@@ -649,7 +649,7 @@ bool URigVMMemoryStorageGeneratorClass::RemoveStorageClass(UObject* InOuter, ERi
 
 uint32 URigVMMemoryStorageGeneratorClass::GetMemoryHash() const
 {
-	if(CachedMemoryHash != 0)
+ 	if(CachedMemoryHash != 0)
 	{
 		return CachedMemoryHash;
 	}
@@ -658,20 +658,21 @@ uint32 URigVMMemoryStorageGeneratorClass::GetMemoryHash() const
 
 	for(const FProperty* Property : LinkedProperties)
 	{
-		CachedMemoryHash = HashCombine(CachedMemoryHash, GetTypeHash(Property->GetFName().ToString()));
-		CachedMemoryHash = HashCombine(CachedMemoryHash, GetTypeHash(Property->GetCPPType()));
+		CachedMemoryHash = HashCombineFast(CachedMemoryHash, GetTypeHash(Property->GetFName().ToString()));
+		CachedMemoryHash = HashCombineFast(CachedMemoryHash, GetTypeHash(Property->GetCPPType()));
 	}
 	
 	// for literals we also hash the content / defaults for each property
 	if(GetMemoryType() == ERigVMMemoryType::Literal)
 	{
-		if(URigVMMemoryStorage* CDO = Cast<URigVMMemoryStorage>(GetDefaultObject(true)))
+		if(const URigVMMemoryStorage* CDO = Cast<URigVMMemoryStorage>(GetDefaultObject(true)))
 		{
 			for(const FProperty* Property : LinkedProperties)
 			{
-				FString DefaultValue;
-				Property->ExportTextItem_InContainer(DefaultValue, CDO, nullptr, nullptr, PPF_None, nullptr);
-				CachedMemoryHash = HashCombine(CachedMemoryHash, GetTypeHash(DefaultValue));
+				const uint32 PropertyHash =
+					RigVMPropertyUtils::GetPropertyHashFast(Property, reinterpret_cast<const uint8*>(CDO));
+ 
+				CachedMemoryHash = HashCombineFast(CachedMemoryHash, PropertyHash);
 			}
 		}
 	}
diff --git a/Engine/Plugins/Runtime/RigVM/Source/RigVM/Private/RigVMPropertyUtils.cpp b/Engine/Plugins/Runtime/RigVM/Source/RigVM/Private/RigVMPropertyUtils.cpp
new file mode 100644
--- /dev/null
+++ b/Engine/Plugins/Runtime/RigVM/Source/RigVM/Private/RigVMPropertyUtils.cpp
@@ -0,0 +1,210 @@
+﻿// Copyright Epic Games, Inc. All Rights Reserved.
+
+#include "RigVMPropertyUtils.h"
+
+#include "RigVMModule.h"
+#include "RigVMTypeUtils.h"
+#include "Engine/UserDefinedStruct.h"
+#include "Misc/AssertionMacros.h"
+#include "UObject/TextProperty.h"
+
+void RigVMPropertyUtils::GetTypeFromProperty(const FProperty* InProperty, FName& OutTypeName, UObject*& OutTypeObject)
+{
+	if (CastField<FBoolProperty>(InProperty))
+	{
+		OutTypeName = TEXT("bool");
+		OutTypeObject = nullptr;
+	}
+	else if (CastField<FIntProperty>(InProperty))
+	{
+		OutTypeName = TEXT("int32");
+		OutTypeObject = nullptr;
+	}
+	else if (CastField<FFloatProperty>(InProperty))
+	{
+		OutTypeName = TEXT("float");
+		OutTypeObject = nullptr;
+	}
+	else if (CastField<FDoubleProperty>(InProperty))
+	{
+		OutTypeName = TEXT("double");
+		OutTypeObject = nullptr;
+	}
+	else if (CastField<FStrProperty>(InProperty))
+	{
+		OutTypeName = TEXT("FString");
+		OutTypeObject = nullptr;
+	}
+	else if (CastField<FNameProperty>(InProperty))
+	{
+		OutTypeName = TEXT("FName");
+		OutTypeObject = nullptr;
+	}
+	else if (CastField<FTextProperty>(InProperty))
+	{
+		OutTypeName = TEXT("FText");
+		OutTypeObject = nullptr;
+	}
+	else if (const FEnumProperty* EnumProperty = CastField<FEnumProperty>(InProperty))
+	{
+		OutTypeName = EnumProperty->GetEnum()->GetFName();
+		OutTypeObject = EnumProperty->GetEnum();
+	}
+	else if (const FByteProperty* ByteProperty = CastField<FByteProperty>(InProperty))
+	{
+		if (UEnum* BytePropertyEnum = ByteProperty->Enum)
+		{
+			OutTypeName = BytePropertyEnum->GetFName();
+			OutTypeObject = BytePropertyEnum;
+		}
+	}
+	else if (const FStructProperty* StructProperty = CastField<FStructProperty>(InProperty))
+	{
+		OutTypeName = *RigVMTypeUtils::GetUniqueStructTypeName(StructProperty->Struct);
+		OutTypeObject = StructProperty->Struct;
+	}
+	else if (const FArrayProperty* ArrayProperty = CastField<FArrayProperty>(InProperty))
+	{
+		FName InnerTypeName = NAME_None;
+		GetTypeFromProperty(ArrayProperty->Inner, InnerTypeName, OutTypeObject);
+		OutTypeName = *RigVMTypeUtils::ArrayTypeFromBaseType(InnerTypeName.ToString());
+	}
+	else if (const FObjectProperty* ObjectProperty = CastField<FObjectProperty>(InProperty))
+	{
+		if (RigVMCore::SupportsUObjects())
+		{
+			OutTypeName = *FString::Printf(TEXT("TObjectPtr<%s%s>"),
+				ObjectProperty->PropertyClass->GetPrefixCPP(),
+				*ObjectProperty->PropertyClass->GetName());
+			OutTypeObject = ObjectProperty->PropertyClass;
+		}
+		else
+		{
+			OutTypeName = NAME_None;
+			OutTypeObject = nullptr;
+		}
+	}
+	else if (const FInterfaceProperty* InterfaceProperty = CastField<FInterfaceProperty>(InProperty))
+	{
+		if (RigVMCore::SupportsUInterfaces())
+		{
+			OutTypeName = *FString::Printf(TEXT("TScriptInterface<%s%s>"),
+				InterfaceProperty->InterfaceClass->GetPrefixCPP(),
+				*InterfaceProperty->InterfaceClass->GetName());
+			OutTypeObject = InterfaceProperty->InterfaceClass;
+		}
+		else
+		{
+			OutTypeName = NAME_None;
+			OutTypeObject = nullptr;
+		}
+	}
+	else
+	{
+		// this can happen due to unsupported property types
+		// within data assets or user defined structs
+		OutTypeName = NAME_None;
+		OutTypeObject = nullptr;
+	}
+}
+
+struct FHashCombineFast
+{
+	inline uint32 operator()(uint32 A, uint32 B) const { return HashCombineFast(A, B); }
+};
+
+struct FHashCombineStable
+{
+	inline uint32 operator()(uint32 A, uint32 B) const { return HashCombine(A, B); }
+};
+
+template<typename HASH_COMBINE_OP>
+static uint32 GetPropertyHashPrivate(const FProperty* InProperty, const uint8* InMemory, EPropertyPointerType InContainerType)
+{
+	if (!ensure(InProperty) || InMemory == nullptr)
+	{
+		return 0;
+	}
+
+	// Offset into the container to the location where the actual property value is stored.
+	if (InContainerType == EPropertyPointerType::Container)
+	{
+		InMemory += InProperty->GetOffset_ForInternal();
+	}
+
+	// If the property type provides its own hashing function, use that as a preference. The types below do not
+	// have hashing implemented, and so we hand-roll some of them.
+	if (InProperty->PropertyFlags & CPF_HasGetValueTypeHash)
+	{
+		return InProperty->GetValueTypeHash(InMemory);
+	}
+
+	if (const FBoolProperty* BoolProperty = CastField<FBoolProperty>(InProperty))
+	{
+		return BoolProperty->GetPropertyValue(InMemory);
+	}
+	else if (const FArrayProperty* ArrayProperty = CastField<FArrayProperty>(InProperty))
+	{
+		FScriptArrayHelper ArrayHelper(ArrayProperty, InMemory);
+		int32 Hash = ::GetTypeHash(ArrayHelper.Num());
+		for (int32 Index = 0; Index < ArrayHelper.Num(); Index++)
+		{
+			Hash = HASH_COMBINE_OP{}(Hash, GetPropertyHashPrivate<HASH_COMBINE_OP>(ArrayProperty->Inner, ArrayHelper.GetRawPtr(Index), EPropertyPointerType::Direct));
+		}
+		return Hash;
+	}
+	else if (const FMapProperty* MapProperty = CastField<FMapProperty>(InProperty))
+	{
+		FScriptMapHelper MapHelper(MapProperty, InMemory);
+		int32 Hash = ::GetTypeHash(MapHelper.Num());
+		for (int32 Index = 0; Index < MapHelper.Num(); Index++)
+		{
+			Hash = HASH_COMBINE_OP{}(Hash, GetPropertyHashPrivate<HASH_COMBINE_OP>(MapProperty->KeyProp, MapHelper.GetKeyPtr(Index), EPropertyPointerType::Direct));
+			Hash = HASH_COMBINE_OP{}(Hash, GetPropertyHashPrivate<HASH_COMBINE_OP>(MapProperty->ValueProp, MapHelper.GetValuePtr(Index), EPropertyPointerType::Direct));
+		}
+		return Hash;
+	}
+	else if (const FSetProperty* SetProperty = CastField<FSetProperty>(InProperty))
+	{
+		FScriptSetHelper SetHelper(SetProperty, InMemory);
+		int32 Hash = ::GetTypeHash(SetHelper.Num());
+		for (int32 Index = 0; Index < SetHelper.Num(); Index++)
+		{
+			Hash = HASH_COMBINE_OP{}(Hash, GetPropertyHashPrivate<HASH_COMBINE_OP>(SetProperty->ElementProp, SetHelper.GetElementPtr(Index), EPropertyPointerType::Direct));
+		}
+		return Hash;
+	}
+	else if (const FStructProperty* StructProperty = CastField<FStructProperty>(InProperty))
+	{
+		const UScriptStruct* StructType = StructProperty->Struct;
+		// UserDefinedStruct overrides GetStructTypeHash to work without valid CppStructOps
+		if (StructType->IsA<UUserDefinedStruct>() || (StructType->GetCppStructOps() && StructType->GetCppStructOps()->HasGetTypeHash()))
+		{
+			return StructType->GetStructTypeHash(InMemory);
+		}
+		else
+		{
+			uint32 Hash = 0;
+			for (TFieldIterator<FProperty> It(StructType); It; ++It)
+			{
+				const FProperty* SubProperty = *It;
+				Hash = HASH_COMBINE_OP{}(Hash, GetPropertyHashPrivate<HASH_COMBINE_OP>(*It, InMemory + SubProperty->GetOffset_ForInternal(), EPropertyPointerType::Direct));
+			}
+			return Hash;
+		}
+	}
+
+	// If we get here, we're missing support for a property type that doesn't do its own hashing. 
+	checkNoEntry();
+	return 0;
+}
+
+uint32 RigVMPropertyUtils::GetPropertyHashFast(const FProperty* InProperty, const uint8* InMemory, EPropertyPointerType InContainerType)
+{
+	return GetPropertyHashPrivate<FHashCombineFast>(InProperty, InMemory, InContainerType);
+}
+
+uint32 RigVMPropertyUtils::GetPropertyHashStable(const FProperty* InProperty, const uint8* InMemory, EPropertyPointerType InContainerType)
+{
+	return GetPropertyHashPrivate<FHashCombineStable>(InProperty, InMemory, InContainerType);
+}
diff --git a/Engine/Plugins/Runtime/RigVM/Source/RigVM/Public/RigVMCore/RigVMExecuteContext.h b/Engine/Plugins/Runtime/RigVM/Source/RigVM/Public/RigVMCore/RigVMExecuteContext.h
--- a/Engine/Plugins/Runtime/RigVM/Source/RigVM/Public/RigVMCore/RigVMExecuteContext.h
+++ b/Engine/Plugins/Runtime/RigVM/Source/RigVM/Public/RigVMCore/RigVMExecuteContext.h
@@ -14,7 +14,6 @@
 #include "Math/UnrealMathSSE.h"
 #include "Misc/AssertionMacros.h"
 #include "RigVMDefines.h"
-#include "RigVMExternalVariable.h"
 #include "RigVMModule.h"
 #include "RigVMCore/RigVMDebugInfo.h"
 #include "RigVMCore/RigVMProfilingInfo.h"
@@ -23,6 +22,7 @@
 #include "RigVMLog.h"
 #include "RigVMDrawInterface.h"
 #include "RigVMDrawContainer.h"
+#include "RigVMMemoryStorage.h"
 #include "Templates/SharedPointer.h"
 #include "Trace/Detail/Channel.h"
 #include "UObject/NameTypes.h"
diff --git a/Engine/Plugins/Runtime/RigVM/Source/RigVM/Public/RigVMCore/RigVMExternalVariable.h b/Engine/Plugins/Runtime/RigVM/Source/RigVM/Public/RigVMCore/RigVMExternalVariable.h
--- a/Engine/Plugins/Runtime/RigVM/Source/RigVM/Public/RigVMCore/RigVMExternalVariable.h
+++ b/Engine/Plugins/Runtime/RigVM/Source/RigVM/Public/RigVMCore/RigVMExternalVariable.h
@@ -4,9 +4,8 @@
 
 #include "CoreMinimal.h"
 #include "RigVMTraits.h"
-#include "RigVMDefines.h"
-#include "RigVMMemory.h"
 #include "RigVMModule.h"
+#include "RigVMPropertyUtils.h"
 #include "RigVMTypeIndex.h"
 #include "RigVMTypeUtils.h"
 #include "UObject/UnrealType.h"
@@ -39,106 +38,6 @@ struct RIGVM_API FRigVMExternalVariableDef
 			&& Size == Other.Size;
 	}
 
-	static void GetTypeFromProperty(const FProperty* InProperty, FName& OutTypeName, UObject*& OutTypeObject)
-	{
-		if (CastField<FBoolProperty>(InProperty))
-		{
-			OutTypeName = TEXT("bool");
-			OutTypeObject = nullptr;
-		}
-		else if (CastField<FIntProperty>(InProperty))
-		{
-			OutTypeName = TEXT("int32");
-			OutTypeObject = nullptr;
-		}
-		else if (CastField<FFloatProperty>(InProperty))
-		{
-			OutTypeName = TEXT("float");
-			OutTypeObject = nullptr;
-		}
-		else if (CastField<FDoubleProperty>(InProperty))
-		{
-			OutTypeName = TEXT("double");
-			OutTypeObject = nullptr;
-		}
-		else if (CastField<FStrProperty>(InProperty))
-		{
-			OutTypeName = TEXT("FString");
-			OutTypeObject = nullptr;
-		}
-		else if (CastField<FNameProperty>(InProperty))
-		{
-			OutTypeName = TEXT("FName");
-			OutTypeObject = nullptr;
-		}
-		else if (CastField<FTextProperty>(InProperty))
-		{
-			OutTypeName = TEXT("FText");
-			OutTypeObject = nullptr;
-		}
-		else if (const FEnumProperty* EnumProperty = CastField<FEnumProperty>(InProperty))
-		{
-			OutTypeName = EnumProperty->GetEnum()->GetFName();
-			OutTypeObject = EnumProperty->GetEnum();
-		}
-		else if (const FByteProperty* ByteProperty = CastField<FByteProperty>(InProperty))
-		{
-			if (UEnum* BytePropertyEnum = ByteProperty->Enum)
-			{
-				OutTypeName = BytePropertyEnum->GetFName();
-				OutTypeObject = BytePropertyEnum;
-			}
-		}
-		else if (const FStructProperty* StructProperty = CastField<FStructProperty>(InProperty))
-		{
-			OutTypeName = *RigVMTypeUtils::GetUniqueStructTypeName(StructProperty->Struct);
-			OutTypeObject = StructProperty->Struct;
-		}
-		else if (const FArrayProperty* ArrayProperty = CastField<FArrayProperty>(InProperty))
-		{
-			FName InnerTypeName = NAME_None;
-			GetTypeFromProperty(ArrayProperty->Inner, InnerTypeName, OutTypeObject);
-			OutTypeName = *RigVMTypeUtils::ArrayTypeFromBaseType(InnerTypeName.ToString());
-		}
-		else if (const FObjectProperty* ObjectProperty = CastField<FObjectProperty>(InProperty))
-		{
-			if(RigVMCore::SupportsUObjects())
-			{
-				OutTypeName = *FString::Printf(TEXT("TObjectPtr<%s%s>"),
-					ObjectProperty->PropertyClass->GetPrefixCPP(),
-					*ObjectProperty->PropertyClass->GetName());
-				OutTypeObject = ObjectProperty->PropertyClass;
-			}
-			else
-			{
-				OutTypeName = NAME_None;
-				OutTypeObject = nullptr;
-			}
-		}
-		else if (const FInterfaceProperty* InterfaceProperty = CastField<FInterfaceProperty>(InProperty))
-		{
-			if (RigVMCore::SupportsUInterfaces())
-			{
-				OutTypeName = *FString::Printf(TEXT("TScriptInterface<%s%s>"),
-					InterfaceProperty->InterfaceClass->GetPrefixCPP(),
-					*InterfaceProperty->InterfaceClass->GetName());
-				OutTypeObject = InterfaceProperty->InterfaceClass;
-			}
-			else
-			{
-				OutTypeName = NAME_None;
-				OutTypeObject = nullptr;
-			}
-		}
-		else
-		{
-			// this can happen due to unsupported property types
-			// within data assets or user defined structs
-			OutTypeName = NAME_None;
-			OutTypeObject = nullptr;
-		}
-	}
-
 	bool IsValid() const
 	{
 		return Name.IsValid() &&
@@ -188,81 +87,9 @@ struct RIGVM_API FRigVMExternalVariable : public FRigVMExternalVariableDef
 	{
 	}
 
-	static uint32 GetPropertyTypeHash(const FProperty *InProperty, const uint8 *InMemory)
-	{
-		if (!ensure(InProperty))
-		{
-			return 0;
-		}
-
-		if (InProperty->PropertyFlags & CPF_HasGetValueTypeHash)
-		{
-			return InProperty->GetValueTypeHash(InMemory);
-		}
-
-		if (const FBoolProperty* BoolProperty = CastField<FBoolProperty>(InProperty))
-		{
-			return static_cast<uint32>(BoolProperty->GetPropertyValue(InMemory));
-		}
-		else if (const FArrayProperty *ArrayProperty = CastField<FArrayProperty>(InProperty))
-		{
-			FScriptArrayHelper ArrayHelper(ArrayProperty, InMemory);
-			int32 Hash = ::GetTypeHash(ArrayHelper.Num());
-			for (int32 Index = 0; Index < ArrayHelper.Num(); Index++)
-			{
-				Hash = HashCombine(Hash, GetPropertyTypeHash(ArrayProperty->Inner, ArrayHelper.GetRawPtr(Index)));
-			}
-			return Hash;
-		}
-		else if (const FMapProperty *MapProperty = CastField<FMapProperty>(InProperty))
-		{
-			FScriptMapHelper MapHelper(MapProperty, InMemory);
-			int32 Hash = ::GetTypeHash(MapHelper.Num());
-			for (FScriptMapHelper::FIterator It(MapHelper); It; ++It)
-			{
-				Hash = HashCombine(Hash, GetPropertyTypeHash(MapProperty->KeyProp, MapHelper.GetKeyPtr(It)));
-				Hash = HashCombine(Hash, GetPropertyTypeHash(MapProperty->ValueProp, MapHelper.GetValuePtr(It)));
-			}
-			return Hash;
-		}
-		else if (const FSetProperty *SetProperty = CastField<FSetProperty>(InProperty))
-		{
-			FScriptSetHelper SetHelper(SetProperty, InMemory);
-			int32 Hash = ::GetTypeHash(SetHelper.Num());
-			for (FScriptSetHelper::FIterator It(SetHelper); It; ++It)
-			{
-				Hash = HashCombine(Hash, GetPropertyTypeHash(SetProperty->ElementProp, SetHelper.GetElementPtr(It)));
-			}
-			return Hash;
-		}
-		else if (const FStructProperty *StructProperty = CastField<FStructProperty>(InProperty))
-		{
-			const UScriptStruct* StructType = StructProperty->Struct;
-			// UserDefinedStruct overrides GetStructTypeHash to work without valid CppStructOps
-			if (StructType->IsA<UUserDefinedStruct>() || (StructType->GetCppStructOps() && StructType->GetCppStructOps()->HasGetTypeHash()))
-			{
-				return StructType->GetStructTypeHash(InMemory);
-			}
-			else
-			{
-				uint32 Hash = 0;
-				for (TFieldIterator<FProperty> It(StructType); It; ++It)
-				{
-					const FProperty* SubProperty = *It;
-					Hash = HashCombine(Hash, GetPropertyTypeHash(*It, InMemory + SubProperty->GetOffset_ForInternal()));
-				}
-				return Hash;
-			}
-		}
-
-		// If we get here, we're missing support for a property type that doesn't do its own hashing. 
-		checkNoEntry();
-		return 0;
-	}
-
 	uint32 GetTypeHash() const
 	{
-		return GetPropertyTypeHash(Property, Memory);
+		return RigVMPropertyUtils::GetPropertyHashStable(Property, Memory);
 	}
 	
 
@@ -293,7 +120,7 @@ struct RIGVM_API FRigVMExternalVariable : public FRigVMExternalVariableDef
 		}
 
 		ExternalVariable.Size = Property->GetSize();
-		GetTypeFromProperty(Property, ExternalVariable.TypeName, ExternalVariable.TypeObject);
+		RigVMPropertyUtils::GetTypeFromProperty(Property, ExternalVariable.TypeName, ExternalVariable.TypeObject);
 
 		return ExternalVariable;
 	}
diff --git a/Engine/Plugins/Runtime/RigVM/Source/RigVM/Public/RigVMPropertyUtils.h b/Engine/Plugins/Runtime/RigVM/Source/RigVM/Public/RigVMPropertyUtils.h
new file mode 100644
--- /dev/null
+++ b/Engine/Plugins/Runtime/RigVM/Source/RigVM/Public/RigVMPropertyUtils.h
@@ -0,0 +1,48 @@
+﻿// Copyright Epic Games, Inc. All Rights Reserved.
+
+#pragma once
+
+#include "HAL/Platform.h"
+#include "UObject/UnrealType.h"
+
+
+namespace RigVMPropertyUtils
+{
+
+/**
+ * Given a property, return the C++ type name for it and the associated type object. If the type is an unambiguous
+ * simple type, such as a bool or an int, the OutTypeObject value will be nullptr. If the property represents a
+ * ambiguous type, like an object pointer, interface or an enum, then the OutTypeObject will be the class of the
+ * pointed to object, or the enum definition. For container properties, the type object is always the contained type,
+ * rather than the type of the container.
+ * \param InProperty The property to get the type for.
+ * \param OutTypeName The C++ type name.
+ * \param OutTypeObject The object class, if the property type is complex or ambiguous. 
+ */
+RIGVM_API void GetTypeFromProperty(const FProperty* InProperty, FName& OutTypeName, UObject*& OutTypeObject);
+
+/**
+ * Given a property and a value storage for it, return a 32-bit hash value for the underlying value of tha type.
+ * \note The resulting hash is not suitable for serialization, since the underlying hash combine function is not 
+ *    guaranteed to be stable between releases.
+ * \param InProperty The property to get a hash value for.
+ * \param InMemory The underlying memory that the property is defined for.
+ * \param InContainerType The type of memory being pointed at. For UObjects, use EPropertyPointerType::Container.
+ *    If pointing directly at the value location in memory, use EPropertyPointerType::Direct.
+ * \return The hash value of the value of the property at the memory location the property lives at.
+ */
+RIGVM_API uint32 GetPropertyHashFast(const FProperty* InProperty, const uint8* InMemory, EPropertyPointerType InContainerType = EPropertyPointerType::Container);
+
+/**
+ * Given a property and a value storage for it, return a 32-bit hash value for the underlying value of tha type.
+ * \note This hash is suitable for serialization, since the underlying hash combine function is guaranteed to be stable
+ *    between releases.
+ * \param InProperty The property to get a hash value for.
+ * \param InMemory The underlying UObject memory that the property is defined for.
+* \param InContainerType The type of memory being pointed at. For UObjects, use EPropertyPointerType::Container.
+ *    If pointing directly at the value location in memory, use EPropertyPointerType::Direct.
+ * \return The hash value of the value of the property at the memory location the property lives at.
+ */
+RIGVM_API uint32 GetPropertyHashStable(const FProperty* InProperty, const uint8* InMemory, EPropertyPointerType InContainerType = EPropertyPointerType::Container);
+
+}
diff --git a/Engine/Plugins/Runtime/RigVM/Source/RigVMDeveloper/Private/RigVMCompiler/RigVMCompiler.cpp b/Engine/Plugins/Runtime/RigVM/Source/RigVMDeveloper/Private/RigVMCompiler/RigVMCompiler.cpp
--- a/Engine/Plugins/Runtime/RigVM/Source/RigVMDeveloper/Private/RigVMCompiler/RigVMCompiler.cpp
+++ b/Engine/Plugins/Runtime/RigVM/Source/RigVMDeveloper/Private/RigVMCompiler/RigVMCompiler.cpp
@@ -274,7 +274,7 @@ TRigVMTypeIndex FRigVMCompilerWorkData::GetTypeIndexForOperand(const FRigVMOpera
 
 	FName CPPTypeName(NAME_None);
 	UObject* CPPTypeObject = nullptr;
-	FRigVMExternalVariable::GetTypeFromProperty(Property, CPPTypeName, CPPTypeObject);
+	RigVMPropertyUtils::GetTypeFromProperty(Property, CPPTypeName, CPPTypeObject);
 
 	return FRigVMRegistry::Get().GetTypeIndex(CPPTypeName, CPPTypeObject);
 }
diff --git a/Engine/Plugins/Runtime/RigVM/Source/RigVMDeveloper/Private/RigVMModel/RigVMPin.cpp b/Engine/Plugins/Runtime/RigVM/Source/RigVMDeveloper/Private/RigVMModel/RigVMPin.cpp
--- a/Engine/Plugins/Runtime/RigVM/Source/RigVMDeveloper/Private/RigVMModel/RigVMPin.cpp
+++ b/Engine/Plugins/Runtime/RigVM/Source/RigVMDeveloper/Private/RigVMModel/RigVMPin.cpp
@@ -1334,7 +1334,7 @@ bool URigVMPin::CanBeBoundToVariable(const FRigVMExternalVariable& InExternalVar
 		const FRigVMPropertyPath PropertyPath(Property, InSegmentPath);
 		Property = PropertyPath.GetTailProperty();
 
-		FRigVMExternalVariable::GetTypeFromProperty(Property, ExternalCPPType, ExternalCPPTypeObject);
+		RigVMPropertyUtils::GetTypeFromProperty(Property, ExternalCPPType, ExternalCPPTypeObject);
 	}
 
 	const FString CPPBaseType = IsArray() ? GetArrayElementCppType() : GetCPPType();
diff --git a/Engine/Plugins/Runtime/RigVM/Source/RigVMEditor/Private/EdGraph/NodeSpawners/RigVMEdGraphUnitNodeSpawner.cpp b/Engine/Plugins/Runtime/RigVM/Source/RigVMEditor/Private/EdGraph/NodeSpawners/RigVMEdGraphUnitNodeSpawner.cpp
--- a/Engine/Plugins/Runtime/RigVM/Source/RigVMEditor/Private/EdGraph/NodeSpawners/RigVMEdGraphUnitNodeSpawner.cpp
+++ b/Engine/Plugins/Runtime/RigVM/Source/RigVMEditor/Private/EdGraph/NodeSpawners/RigVMEdGraphUnitNodeSpawner.cpp
@@ -16,6 +16,7 @@
 #include "RigVMBlueprintUtils.h"
 #include "ScopedTransaction.h"
 #include "RigVMFunctions/Execution/RigVMFunction_UserDefinedEvent.h"
+#include "RigVMPropertyUtils.h"
 
 #include UE_INLINE_GENERATED_CPP_BY_NAME(RigVMEdGraphUnitNodeSpawner)
 
@@ -175,7 +176,7 @@ URigVMEdGraphNode* URigVMEdGraphUnitNodeSpawner::SpawnNode(UEdGraph* ParentGraph
 				FPinInfo Pin;
 				const FProperty* Property = *It;
 				Pin.Name = Property->GetFName();
-				FRigVMExternalVariableDef::GetTypeFromProperty(Property, Pin.CPPType, Pin.CPPTypeObject);
+				RigVMPropertyUtils::GetTypeFromProperty(Property, Pin.CPPType, Pin.CPPTypeObject);
 				if(Pin.CPPType.IsNone())
 				{
 					continue;
