RigVM: Fix wrong behavior when same name nodes exist in different root graphs
#rnx

diff --git a/Engine/Plugins/Runtime/RigVM/Source/RigVMDeveloper/Private/RigVMCompiler/RigVMASTProxy.cpp b/Engine/Plugins/Runtime/RigVM/Source/RigVMDeveloper/Private/RigVMCompiler/RigVMASTProxy.cpp
--- a/Engine/Plugins/Runtime/RigVM/Source/RigVMDeveloper/Private/RigVMCompiler/RigVMASTProxy.cpp
+++ b/Engine/Plugins/Runtime/RigVM/Source/RigVMDeveloper/Private/RigVMCompiler/RigVMASTProxy.cpp
@@ -7,6 +7,7 @@
 #include "RigVMModel/Nodes/RigVMCollapseNode.h"
 #include "RigVMModel/Nodes/RigVMFunctionReferenceNode.h"
 #include "RigVMModel/Nodes/RigVMLibraryNode.h"
+#include "RigVMModel/RigVMFunctionLibrary.h"
 
 FString FRigVMCallstack::GetCallPath(bool bIncludeLast) const
 {
@@ -33,7 +34,7 @@ FString FRigVMCallstack::GetCallPath(bool bIncludeLast) const
 			{
 				if (URigVMGraph* Graph = Pin->GetGraph())
 				{
-					const bool bUseNodePath = Graph->IsRootGraph();
+					const bool bUseNodePath = Stack.Num() == 1 && !Graph->IsA<URigVMFunctionLibrary>();
 					Segments.Add(Pin->GetPinPath(bUseNodePath));
 				}
 			}
diff --git a/Engine/Plugins/Runtime/RigVM/Source/RigVMDeveloper/Private/RigVMModel/RigVMPin.cpp b/Engine/Plugins/Runtime/RigVM/Source/RigVMDeveloper/Private/RigVMModel/RigVMPin.cpp
--- a/Engine/Plugins/Runtime/RigVM/Source/RigVMDeveloper/Private/RigVMModel/RigVMPin.cpp
+++ b/Engine/Plugins/Runtime/RigVM/Source/RigVMDeveloper/Private/RigVMModel/RigVMPin.cpp
@@ -184,7 +184,7 @@ FString URigVMPin::GetPinPath(bool bUseNodePath) const
 	URigVMPin* ParentPin = GetParentPin();
 	if (ParentPin)
 	{
-		PinPath = JoinPinPath(ParentPin->GetPinPath(), GetName());
+		PinPath = JoinPinPath(ParentPin->GetPinPath(bUseNodePath), GetName());
 	}
 	else
 	{
