Fix warning V1101: The default argument of the second virtual function parameter was changed in the derived class from '1' to '0'. Calling such function via a pointer or a reference to the base class may result in unexpected behavior.
#rnx

diff --git a/Engine/Plugins/Runtime/RigVM/Source/RigVM/Public/RigVMCore/RigVMNativized.h b/Engine/Plugins/Runtime/RigVM/Source/RigVM/Public/RigVMCore/RigVMNativized.h
--- a/Engine/Plugins/Runtime/RigVM/Source/RigVM/Public/RigVMCore/RigVMNativized.h
+++ b/Engine/Plugins/Runtime/RigVM/Source/RigVM/Public/RigVMCore/RigVMNativized.h
@@ -79,7 +79,7 @@ public:
 	UE_DEPRECATED(5.4, "Please, use Initialize just with Context param")
 	virtual bool Initialize(FRigVMExtendedExecuteContext& Context, TArrayView<URigVMMemoryStorage*> Memory) { return false; }
 	UE_DEPRECATED(5.4, "GetMemoryByType has been deprecated from the VM. Please, use GetWorkMemory from VMHost or the new GetMemoryByType with Context parameter.")
-	virtual URigVMMemoryStorage* GetMemoryByType(ERigVMMemoryType InMemoryType, bool bCreateIfNeeded = false) override { return nullptr; }
+	virtual URigVMMemoryStorage* GetMemoryByType(ERigVMMemoryType InMemoryType, bool bCreateIfNeeded = true) override { return nullptr; }
 	UE_DEPRECATED(5.3, "Please, use Execute with Context param")
 	virtual ERigVMExecuteResult Execute(TArrayView<URigVMMemoryStorage*> Memory, const FName& InEntryName = NAME_None) override { return ERigVMExecuteResult::Failed; }
 	UE_DEPRECATED(5.4, "Please, use Execute with just Context param")
