RigVM: Fix changes in inner functions not triggering recompilation of outer functions
#rnx

diff --git a/Engine/Plugins/Runtime/RigVM/Source/RigVMDeveloper/Private/RigVMCompiler/RigVMCompiler.cpp b/Engine/Plugins/Runtime/RigVM/Source/RigVMDeveloper/Private/RigVMCompiler/RigVMCompiler.cpp
--- a/Engine/Plugins/Runtime/RigVM/Source/RigVMDeveloper/Private/RigVMCompiler/RigVMCompiler.cpp
+++ b/Engine/Plugins/Runtime/RigVM/Source/RigVMDeveloper/Private/RigVMCompiler/RigVMCompiler.cpp
@@ -497,17 +497,19 @@ bool URigVMCompiler::Compile(const FRigVMCompileSettings& InSettings, TArray<URi
 						{
 							for (const TPair<FRigVMGraphFunctionIdentifier, uint32>& Pair : FunctionData->Header.Dependencies)
 							{
+								bool bDirty = true;
 								if (IRigVMGraphFunctionHost* HostObj = Cast<IRigVMGraphFunctionHost>(Pair.Key.HostObject.ResolveObject()))
 								{
 									if (FRigVMGraphFunctionData* DependencyData = HostObj->GetRigVMGraphFunctionStore()->FindFunction(Pair.Key))
 									{
-										if (DependencyData->CompilationData.Hash == 0 || Pair.Value != DependencyData->CompilationData.Hash)
-										{
-											FunctionData->ClearCompilationData();
-											break;
-										}
+										bDirty = DependencyData->CompilationData.Hash != Pair.Value || Pair.Value == 0;
 									}
 								}
+								if (bDirty)
+								{
+									FunctionData->ClearCompilationData();
+									break;
+								}
 							}
 						}
 						
diff --git a/Engine/Plugins/Runtime/RigVM/Source/RigVMDeveloper/Private/RigVMModel/RigVMClient.cpp b/Engine/Plugins/Runtime/RigVM/Source/RigVMDeveloper/Private/RigVMModel/RigVMClient.cpp
--- a/Engine/Plugins/Runtime/RigVM/Source/RigVMDeveloper/Private/RigVMModel/RigVMClient.cpp
+++ b/Engine/Plugins/Runtime/RigVM/Source/RigVMDeveloper/Private/RigVMModel/RigVMClient.cpp
@@ -1153,6 +1153,7 @@ void FRigVMClient::HandleGraphModifiedEvent(ERigVMGraphNotifType InNotifType, UR
 								}
 							}
 
+							UpdateFunctionReferences(Data->Header, true, false);
 							UpdateGraphFunctionData(CollapseNode);
 						}
 					}
