Control Rig: Add keywords to Array At node

#rnx

diff --git a/Engine/Plugins/Runtime/RigVM/Source/RigVM/Public/RigVMFunctions/RigVMDispatch_Array.h b/Engine/Plugins/Runtime/RigVM/Source/RigVM/Public/RigVMFunctions/RigVMDispatch_Array.h
--- a/Engine/Plugins/Runtime/RigVM/Source/RigVM/Public/RigVMFunctions/RigVMDispatch_Array.h
+++ b/Engine/Plugins/Runtime/RigVM/Source/RigVM/Public/RigVMFunctions/RigVMDispatch_Array.h
@@ -152,7 +152,7 @@ protected:
 	static void Execute(FRigVMExtendedExecuteContext& InContext, FRigVMMemoryHandleArray Handles, FRigVMPredicateBranchArray Predicates);
 };
 
-USTRUCT(meta=(DisplayName = "At"))
+USTRUCT(meta=(DisplayName = "At", Keywords = "Get Index,At Index,[]"))
 struct RIGVM_API FRigVMDispatch_ArrayGetAtIndex : public FRigVMDispatch_ArrayBase
 {
 	GENERATED_BODY()
