Control Rig: Fix changing exposed pin type to execute context does not convert the pin into IO on function references
#rnx

diff --git a/Engine/Plugins/Runtime/RigVM/Source/RigVMDeveloper/Private/RigVMModel/RigVMController.cpp b/Engine/Plugins/Runtime/RigVM/Source/RigVMDeveloper/Private/RigVMModel/RigVMController.cpp
--- a/Engine/Plugins/Runtime/RigVM/Source/RigVMDeveloper/Private/RigVMModel/RigVMController.cpp
+++ b/Engine/Plugins/Runtime/RigVM/Source/RigVMDeveloper/Private/RigVMModel/RigVMController.cpp
@@ -10996,25 +10996,30 @@ bool URigVMController::ChangeExposedPinType(const FName& InPinName, const FStrin
 		return false;
 	}
 
-	if (Pin->GetDirection() == ERigVMPinDirection::IO)
+	bool bIsExecute = false;
+	if (CPPTypeObject)
 	{
-		bool bIsExecute = false;
-		if (CPPTypeObject)
+		if(const UScriptStruct* CPPTypeStruct = Cast<UScriptStruct>(CPPTypeObject))
 		{
-			if(const UScriptStruct* CPPTypeStruct = Cast<UScriptStruct>(CPPTypeObject))
+			if(CPPTypeStruct->IsChildOf(FRigVMExecuteContext::StaticStruct()))
 			{
-				if(CPPTypeStruct->IsChildOf(FRigVMExecuteContext::StaticStruct()))
-				{
-					bIsExecute = true;
-				}
+				bIsExecute = true;
 			}
 		}
-		if (!bIsExecute)
+	}
+
+	if (bIsExecute)
+	{
+		if (Pin->GetDirection() != ERigVMPinDirection::IO)
 		{
-			ReportAndNotifyError(TEXT("Input/Output pins only allow Execute Context types."));
-			return false;
+			Pin->Direction = ERigVMPinDirection::IO;
 		}
 	}
+	else if(Pin->GetDirection() == ERigVMPinDirection::IO)
+	{
+		ReportAndNotifyError(TEXT("Input/Output pins only allow Execute Context types."));
+		return false;
+	}
 
 	if(bSetupUndoRedo)
 	{
@@ -11120,17 +11125,7 @@ bool URigVMController::ChangeExposedPinType(const FName& InPinName, const FStrin
 	// Change pin type on function references
 	if (URigVMFunctionLibrary* FunctionLibrary = Cast<URigVMFunctionLibrary>(LibraryNode->GetGraph()))
 	{
-		FunctionLibrary->ForEachReference(LibraryNode->GetFName(), [this, &Pin, InCPPType, InCPPTypeObjectPath, bSetupUndoRedo, bSetupOrphanPins](URigVMFunctionReferenceNode* ReferenceNode)
-        {
-			if (URigVMPin* ReferencedNodePin = ReferenceNode->FindPin(Pin->GetName()))
-			{
-				if(URigVMController* ReferenceController = GetControllerForGraph(ReferenceNode->GetGraph()))
-				{
-					ReferenceController->ChangePinType(ReferencedNodePin, InCPPType, InCPPTypeObjectPath, bSetupUndoRedo, bSetupOrphanPins);
-					ReferenceController->RemoveUnusedOrphanedPins(ReferenceNode);
-				}
-			}
-        });
+		RefreshFunctionReferences(LibraryNode, bSetupUndoRedo);
 	}
 
 	// Change pin types on input variable nodes
