RigVM Ed Graph: New pin order execute, output, IO, input

#rnx

diff --git a/Engine/Plugins/Runtime/RigVM/Source/RigVMDeveloper/Private/EdGraph/RigVMEdGraphNode.cpp b/Engine/Plugins/Runtime/RigVM/Source/RigVMDeveloper/Private/EdGraph/RigVMEdGraphNode.cpp
--- a/Engine/Plugins/Runtime/RigVM/Source/RigVMDeveloper/Private/EdGraph/RigVMEdGraphNode.cpp
+++ b/Engine/Plugins/Runtime/RigVM/Source/RigVMDeveloper/Private/EdGraph/RigVMEdGraphNode.cpp
@@ -952,18 +952,18 @@ bool URigVMEdGraphNode::ModelPinsChanged(bool bForce)
 	
 	int32 PinsAdded = 0;
 	PinsAdded += AddMissingPins(ExecutePins);
-	PinsAdded += AddMissingPins(InputPins);
-	PinsAdded += AddMissingPins(InputOutputPins);
 	PinsAdded += AddMissingPins(OutputPins);
+	PinsAdded += AddMissingPins(InputOutputPins);
+	PinsAdded += AddMissingPins(InputPins);
 
 	const int32 PinsRemoved = RemoveObsoletePins();
 
 	// working through it in the opposite order
 	// due to the use of Append within the lambda
 	int32 PinsReordered = 0;
-	PinsReordered += OrderPins(OutputPins);
-	PinsReordered += OrderPins(InputOutputPins);
 	PinsReordered += OrderPins(InputPins);
+	PinsReordered += OrderPins(InputOutputPins);
+	PinsReordered += OrderPins(OutputPins);
 	PinsReordered += OrderPins(ExecutePins);
 
 	const bool bResult = (PinsAdded > 0) || (PinsRemoved > 0) || (PinsReordered > 0); 
@@ -1176,9 +1176,9 @@ FLinearColor URigVMEdGraphNode::GetNodeProfilingColor() const
 void URigVMEdGraphNode::UpdatePinLists()
 {
 	ExecutePins.Reset();
-	InputPins.Reset();
-	InputOutputPins.Reset();
 	OutputPins.Reset();
+	InputOutputPins.Reset();
+	InputPins.Reset();
 
 	if (URigVMNode* ModelNode = GetModelNode())
 	{
@@ -1234,9 +1234,9 @@ void URigVMEdGraphNode::AllocateDefaultPins()
 	};
 
 	CreateGraphPins(ExecutePins);
-	CreateGraphPins(InputPins);
-	CreateGraphPins(InputOutputPins);
 	CreateGraphPins(OutputPins);
+	CreateGraphPins(InputOutputPins);
+	CreateGraphPins(InputPins);
 
 	// Fill the variable list
 	ExternalVariables.Reset();
diff --git a/Engine/Plugins/Runtime/RigVM/Source/RigVMEditor/Private/Widgets/SRigVMGraphNode.cpp b/Engine/Plugins/Runtime/RigVM/Source/RigVMEditor/Private/Widgets/SRigVMGraphNode.cpp
--- a/Engine/Plugins/Runtime/RigVM/Source/RigVMEditor/Private/Widgets/SRigVMGraphNode.cpp
+++ b/Engine/Plugins/Runtime/RigVM/Source/RigVMEditor/Private/Widgets/SRigVMGraphNode.cpp
@@ -1629,26 +1629,26 @@ void SRigVMGraphNode::UpdatePinTreeView()
 					}
 					break;
 				}
-				case 1: // IO pins
+				case 1: // output pins
 				{
-					if(!RootPin->IsExecuteContext() && RootPin->GetDirection() == ERigVMPinDirection::IO)
+					if(RootPin->GetDirection() == ERigVMPinDirection::Output)
 					{
 						Local::VisitPinRecursively(RootPin, ModelPins, bSupportSubPins);
 					}
 					break;
 				}
-				case 2: // input / visible pins
+				case 2: // IO pins
 				{
-					if(RootPin->GetDirection() == ERigVMPinDirection::Input || RootPin->GetDirection() == ERigVMPinDirection::Visible)
+					if(!RootPin->IsExecuteContext() && RootPin->GetDirection() == ERigVMPinDirection::IO)
 					{
 						Local::VisitPinRecursively(RootPin, ModelPins, bSupportSubPins);
 					}
 					break;
 				}
-				case 3: // output pins
+				case 3: // input / visible pins
 				default:
 				{
-					if(RootPin->GetDirection() == ERigVMPinDirection::Output)
+					if(RootPin->GetDirection() == ERigVMPinDirection::Input || RootPin->GetDirection() == ERigVMPinDirection::Visible)
 					{
 						Local::VisitPinRecursively(RootPin, ModelPins, bSupportSubPins);
 					}
