RigVM: transform / quat to vectors nodes

#rnx

diff --git a/Engine/Plugins/Runtime/RigVM/Source/RigVM/Private/RigVMFunctions/Math/RigVMFunction_MathQuaternion.cpp b/Engine/Plugins/Runtime/RigVM/Source/RigVM/Private/RigVMFunctions/Math/RigVMFunction_MathQuaternion.cpp
--- a/Engine/Plugins/Runtime/RigVM/Source/RigVM/Private/RigVMFunctions/Math/RigVMFunction_MathQuaternion.cpp
+++ b/Engine/Plugins/Runtime/RigVM/Source/RigVM/Private/RigVMFunctions/Math/RigVMFunction_MathQuaternion.cpp
@@ -85,6 +85,14 @@ FRigVMFunction_MathQuaternionToAxisAndAngle_Execute()
 	Angle = Angle * AngleSign;
 }
 
+FRigVMFunction_MathQuaternionToVectors_Execute()
+{
+	const FQuat Rotation = Value.GetNormalized();
+	Forward = Rotation.GetForwardVector();
+	Right = Rotation.GetRightVector();
+	Up = Rotation.GetUpVector();
+}
+
 FRigVMFunction_MathQuaternionScale_Execute()
 {
 	FVector Axis = FVector::ZeroVector;
diff --git a/Engine/Plugins/Runtime/RigVM/Source/RigVM/Private/RigVMFunctions/Math/RigVMFunction_MathTransform.cpp b/Engine/Plugins/Runtime/RigVM/Source/RigVM/Private/RigVMFunctions/Math/RigVMFunction_MathTransform.cpp
--- a/Engine/Plugins/Runtime/RigVM/Source/RigVM/Private/RigVMFunctions/Math/RigVMFunction_MathTransform.cpp
+++ b/Engine/Plugins/Runtime/RigVM/Source/RigVM/Private/RigVMFunctions/Math/RigVMFunction_MathTransform.cpp
@@ -36,6 +36,14 @@ FRigVMFunction_MathTransformToEulerTransform_Execute()
 	Result.FromFTransform(Value);
 }
 
+FRigVMFunction_MathTransformToVectors_Execute()
+{
+	const FQuat Rotation = Value.GetRotation().GetNormalized();
+	Forward = Rotation.GetForwardVector();
+	Right = Rotation.GetRightVector();
+	Up = Rotation.GetUpVector();
+}
+
 FRigVMFunction_MathTransformMul_Execute()
 {
 	Result = A * B;
diff --git a/Engine/Plugins/Runtime/RigVM/Source/RigVM/Public/RigVMFunctions/Math/RigVMFunction_MathQuaternion.h b/Engine/Plugins/Runtime/RigVM/Source/RigVM/Public/RigVMFunctions/Math/RigVMFunction_MathQuaternion.h
--- a/Engine/Plugins/Runtime/RigVM/Source/RigVM/Public/RigVMFunctions/Math/RigVMFunction_MathQuaternion.h
+++ b/Engine/Plugins/Runtime/RigVM/Source/RigVM/Public/RigVMFunctions/Math/RigVMFunction_MathQuaternion.h
@@ -263,6 +263,36 @@ struct RIGVM_API FRigVMFunction_MathQuaternionToAxisAndAngle : public FRigVMFunc
 	float Angle;
 };
 
+/**
+ * Retrieves the forward, right and up vectors of a quaternion
+ */
+USTRUCT(meta=(DisplayName="To Vectors", TemplateName="ToVectors", Keywords="Forward,Right,Up"))
+struct RIGVM_API FRigVMFunction_MathQuaternionToVectors : public FRigVMFunction_MathQuaternionBase
+{
+	GENERATED_BODY()
+
+	FRigVMFunction_MathQuaternionToVectors()
+	{
+		Value = FQuat::Identity;
+		Forward = Right = Up = FVector(0.f, 0.f, 0.f);
+	}
+
+	RIGVM_METHOD()
+	virtual void Execute() override;
+
+	UPROPERTY(meta=(Input))
+	FQuat Value;
+
+	UPROPERTY(meta=(Output))
+	FVector Forward;
+
+	UPROPERTY(meta=(Output))
+	FVector Right;
+
+	UPROPERTY(meta=(Output))
+	FVector Up;
+};
+
 /**
  * Scales a quaternion's angle
  */
diff --git a/Engine/Plugins/Runtime/RigVM/Source/RigVM/Public/RigVMFunctions/Math/RigVMFunction_MathTransform.h b/Engine/Plugins/Runtime/RigVM/Source/RigVM/Public/RigVMFunctions/Math/RigVMFunction_MathTransform.h
--- a/Engine/Plugins/Runtime/RigVM/Source/RigVM/Public/RigVMFunctions/Math/RigVMFunction_MathTransform.h
+++ b/Engine/Plugins/Runtime/RigVM/Source/RigVM/Public/RigVMFunctions/Math/RigVMFunction_MathTransform.h
@@ -181,6 +181,36 @@ struct RIGVM_API FRigVMFunction_MathTransformToEulerTransform : public FRigVMFun
 	FEulerTransform Result;
 };
 
+/**
+ * Retrieves the forward, right and up vectors of a transform's quaternion
+ */
+USTRUCT(meta=(DisplayName="To Vectors", TemplateName="ToVectors", Keywords="Forward,Right,Up"))
+struct RIGVM_API FRigVMFunction_MathTransformToVectors : public FRigVMFunction_MathTransformBase
+{
+	GENERATED_BODY()
+
+	FRigVMFunction_MathTransformToVectors()
+	{
+		Value = FTransform::Identity;
+		Forward = Right = Up = FVector(0.f, 0.f, 0.f);
+	}
+
+	RIGVM_METHOD()
+	virtual void Execute() override;
+
+	UPROPERTY(meta=(Input))
+	FTransform Value;
+
+	UPROPERTY(meta=(Output))
+	FVector Forward;
+
+	UPROPERTY(meta=(Output))
+	FVector Right;
+
+	UPROPERTY(meta=(Output))
+	FVector Up;
+};
+
 /**
  * Returns the product of the two values
  */
