Control Rig: Fix pin bindings not working when searching strings
#rnx

diff --git a/Engine/Plugins/Runtime/RigVM/Source/RigVMDeveloper/Private/EdGraph/RigVMEdGraphNode.cpp b/Engine/Plugins/Runtime/RigVM/Source/RigVMDeveloper/Private/EdGraph/RigVMEdGraphNode.cpp
--- a/Engine/Plugins/Runtime/RigVM/Source/RigVMDeveloper/Private/EdGraph/RigVMEdGraphNode.cpp
+++ b/Engine/Plugins/Runtime/RigVM/Source/RigVMDeveloper/Private/EdGraph/RigVMEdGraphNode.cpp
@@ -1758,6 +1758,26 @@ void URigVMEdGraphNode::BeginDestroy()
 	Super::BeginDestroy();
 }
 
+#if WITH_EDITOR
+void URigVMEdGraphNode::AddPinSearchMetaDataInfo(const UEdGraphPin* Pin, TArray<FSearchTagDataPair>& OutTaggedMetaData) const
+{
+	Super::AddPinSearchMetaDataInfo(Pin, OutTaggedMetaData);
+	if (const URigVMNode* Node = GetModelNode())
+	{
+		const FString PinName = Pin->GetName();
+		FString Left, Right = PinName;
+		URigVMPin::SplitPinPathAtStart(PinName, Left, Right);
+		if (URigVMPin* ModelPin = Node->FindPin(Right))
+		{
+			if (ModelPin->IsBoundToVariable())
+			{
+				OutTaggedMetaData.Add(FSearchTagDataPair(FText::FromString(TEXT("Binding")), FText::FromString(ModelPin->GetBoundVariablePath())));
+			}
+		}
+	}
+}
+#endif
+
 FEdGraphPinType URigVMEdGraphNode::GetPinTypeForModelPin(const URigVMPin* InModelPin)
 {
 	FEdGraphPinType PinType = RigVMTypeUtils::PinTypeFromCPPType(*InModelPin->GetCPPType(), InModelPin->GetCPPTypeObject());
diff --git a/Engine/Plugins/Runtime/RigVM/Source/RigVMDeveloper/Public/EdGraph/RigVMEdGraphNode.h b/Engine/Plugins/Runtime/RigVM/Source/RigVMDeveloper/Public/EdGraph/RigVMEdGraphNode.h
--- a/Engine/Plugins/Runtime/RigVM/Source/RigVMDeveloper/Public/EdGraph/RigVMEdGraphNode.h
+++ b/Engine/Plugins/Runtime/RigVM/Source/RigVMDeveloper/Public/EdGraph/RigVMEdGraphNode.h
@@ -116,6 +116,9 @@ public:
 	virtual bool IsSelectedInEditor() const;
 	virtual bool ShouldDrawNodeAsControlPointOnly(int32& OutInputPinIndex, int32& OutOutputPinIndex) const override;
 	virtual void BeginDestroy() override;
+#if WITH_EDITOR
+	virtual void AddPinSearchMetaDataInfo(const UEdGraphPin* Pin, TArray<FSearchTagDataPair>& OutTaggedMetaData) const override;
+#endif
 
 	virtual bool IsDeprecated() const override;
 	bool IsOutDated() const;
