Control Rig: Adding atan2 float and double nodes
#9064
#rnx

diff --git a/Engine/Plugins/Runtime/RigVM/Source/RigVM/Private/RigVMFunctions/Math/RigVMFunction_MathDouble.cpp b/Engine/Plugins/Runtime/RigVM/Source/RigVM/Private/RigVMFunctions/Math/RigVMFunction_MathDouble.cpp
--- a/Engine/Plugins/Runtime/RigVM/Source/RigVM/Private/RigVMFunctions/Math/RigVMFunction_MathDouble.cpp
+++ b/Engine/Plugins/Runtime/RigVM/Source/RigVM/Private/RigVMFunctions/Math/RigVMFunction_MathDouble.cpp
@@ -268,6 +268,11 @@ FRigVMFunction_MathDoubleAtan_Execute()
 	Result = FMath::Atan(Value);
 }
 
+FRigVMFunction_MathDoubleAtan2_Execute()
+{
+	Result = FMath::Atan2(A, B);
+}
+
 FRigVMFunction_MathDoubleLawOfCosine_Execute()
 {
 	if ((A <= 0.0) || (B <= 0.0) || (C <= 0.0) || (A + B < C) || (A + C < B) || (B + C < A))
diff --git a/Engine/Plugins/Runtime/RigVM/Source/RigVM/Private/RigVMFunctions/Math/RigVMFunction_MathFloat.cpp b/Engine/Plugins/Runtime/RigVM/Source/RigVM/Private/RigVMFunctions/Math/RigVMFunction_MathFloat.cpp
--- a/Engine/Plugins/Runtime/RigVM/Source/RigVM/Private/RigVMFunctions/Math/RigVMFunction_MathFloat.cpp
+++ b/Engine/Plugins/Runtime/RigVM/Source/RigVM/Private/RigVMFunctions/Math/RigVMFunction_MathFloat.cpp
@@ -279,6 +279,11 @@ FRigVMFunction_MathFloatAtan_Execute()
 	Result = FMath::Atan(Value);
 }
 
+FRigVMFunction_MathFloatAtan2_Execute()
+{
+	Result = FMath::Atan2(A, B);
+}
+
 FRigVMFunction_MathFloatLawOfCosine_Execute()
 {
 	if ((A <= 0.f) || (B <= 0.f) || (C <= 0.f) || (A + B < C) || (A + C < B) || (B + C < A))
diff --git a/Engine/Plugins/Runtime/RigVM/Source/RigVM/Public/RigVMFunctions/Math/RigVMFunction_MathDouble.h b/Engine/Plugins/Runtime/RigVM/Source/RigVM/Public/RigVMFunctions/Math/RigVMFunction_MathDouble.h
--- a/Engine/Plugins/Runtime/RigVM/Source/RigVM/Public/RigVMFunctions/Math/RigVMFunction_MathDouble.h
+++ b/Engine/Plugins/Runtime/RigVM/Source/RigVM/Public/RigVMFunctions/Math/RigVMFunction_MathDouble.h
@@ -842,6 +842,18 @@ struct RIGVM_API FRigVMFunction_MathDoubleAtan : public FRigVMFunction_MathDoubl
 	virtual void Execute() override;
 };
 
+/**
+ * Returns the arctangent of the specified A and B coordinates.
+ */
+USTRUCT(meta=(DisplayName="Atan2", TemplateName="Atan2", Keywords="Arctan"))
+struct RIGVM_API FRigVMFunction_MathDoubleAtan2 : public FRigVMFunction_MathDoubleBinaryOp
+{
+	GENERATED_BODY()
+
+	RIGVM_METHOD()
+	virtual void Execute() override;
+};
+
 /**
  * Computes the angles alpha, beta and gamma (in radians) between the three sides A, B and C
  */
diff --git a/Engine/Plugins/Runtime/RigVM/Source/RigVM/Public/RigVMFunctions/Math/RigVMFunction_MathFloat.h b/Engine/Plugins/Runtime/RigVM/Source/RigVM/Public/RigVMFunctions/Math/RigVMFunction_MathFloat.h
--- a/Engine/Plugins/Runtime/RigVM/Source/RigVM/Public/RigVMFunctions/Math/RigVMFunction_MathFloat.h
+++ b/Engine/Plugins/Runtime/RigVM/Source/RigVM/Public/RigVMFunctions/Math/RigVMFunction_MathFloat.h
@@ -874,6 +874,18 @@ struct RIGVM_API FRigVMFunction_MathFloatAtan : public FRigVMFunction_MathFloatU
 	virtual void Execute() override;
 };
 
+/**
+ * Returns the arctangent of the specified A and B coordinates.
+ */
+USTRUCT(meta=(DisplayName="Atan2", TemplateName="Atan2", Keywords="Arctan"))
+struct RIGVM_API FRigVMFunction_MathFloatAtan2 : public FRigVMFunction_MathFloatBinaryOp
+{
+	GENERATED_BODY()
+
+	RIGVM_METHOD()
+	virtual void Execute() override;
+};
+
 /**
  * Computes the angles alpha, beta and gamma (in radians) between the three sides A, B and C
  */
