Control Rig: Fix comment node not storing the correct color
#rnx

diff --git a/Engine/Plugins/Runtime/RigVM/Source/RigVMDeveloper/Private/EdGraph/RigVMEdGraph.cpp b/Engine/Plugins/Runtime/RigVM/Source/RigVMDeveloper/Private/EdGraph/RigVMEdGraph.cpp
--- a/Engine/Plugins/Runtime/RigVM/Source/RigVMDeveloper/Private/EdGraph/RigVMEdGraph.cpp
+++ b/Engine/Plugins/Runtime/RigVM/Source/RigVMDeveloper/Private/EdGraph/RigVMEdGraph.cpp
@@ -415,7 +415,7 @@ bool URigVMEdGraph::HandleModifiedEvent_Internal(ERigVMGraphNotifType InNotifTyp
 				UEdGraphNode_Comment* EdNode = Cast<UEdGraphNode_Comment>(FindNodeForModelNodeName(ModelNode->GetFName()));
 				if (EdNode)
 				{
-					EdNode->OnUpdateCommentText(ModelNode->GetCommentText());
+					EdNode->NodeComment = ModelNode->GetCommentText();
 					EdNode->FontSize = ModelNode->GetCommentFontSize();
 					EdNode->bCommentBubbleVisible = ModelNode->GetCommentBubbleVisible();
 					EdNode->bCommentBubbleVisible_InDetailsPanel = ModelNode->GetCommentBubbleVisible();
diff --git a/Engine/Plugins/Runtime/RigVM/Source/RigVMEditor/Private/Editor/RigVMEditor.cpp b/Engine/Plugins/Runtime/RigVM/Source/RigVMEditor/Private/Editor/RigVMEditor.cpp
--- a/Engine/Plugins/Runtime/RigVM/Source/RigVMEditor/Private/Editor/RigVMEditor.cpp
+++ b/Engine/Plugins/Runtime/RigVM/Source/RigVMEditor/Private/Editor/RigVMEditor.cpp
@@ -2502,8 +2502,24 @@ void FRigVMEditor::OnCreateComment()
 	{
 		if (UEdGraph* Graph = GraphEditor->GetCurrentGraph())
 		{
-			FEdGraphSchemaAction_K2AddComment CommentAction;
-			CommentAction.PerformAction(Graph, NULL, GraphEditor->GetPasteLocation());
+			if (URigVMEdGraph* RigVMEdGraph = Cast<URigVMEdGraph>(Graph))
+			{
+				if (URigVMBlueprint* Blueprint = GetRigVMBlueprint())
+				{
+					if (URigVMController* Controller = Blueprint->GetController(RigVMEdGraph))
+					{
+						Controller->OpenUndoBracket(TEXT("Create Comment"));
+						FEdGraphSchemaAction_K2AddComment CommentAction;
+						UEdGraphNode* EdNode = CommentAction.PerformAction(Graph, NULL, GraphEditor->GetPasteLocation());
+						if (UEdGraphNode_Comment* CommentNode = CastChecked<UEdGraphNode_Comment>(EdNode))
+						{
+							Controller->SetNodeColorByName(CommentNode->GetFName(), CommentNode->GetNodeTitleColor(), false);
+							Controller->SetNodePositionByName(CommentNode->GetFName(), FVector2D(CommentNode->NodePosX, CommentNode->NodePosY), false);
+						}
+						Controller->CloseUndoBracket();
+					}
+				}
+			}
 		}
 	}
 }
