Contorl Rig: Fix comment color not stored correctly in the graph
#rnx

diff --git a/Engine/Plugins/Runtime/RigVM/Source/RigVM/Private/RigVMCore/RigVM.cpp b/Engine/Plugins/Runtime/RigVM/Source/RigVM/Private/RigVMCore/RigVM.cpp
--- a/Engine/Plugins/Runtime/RigVM/Source/RigVM/Private/RigVMCore/RigVM.cpp
+++ b/Engine/Plugins/Runtime/RigVM/Source/RigVM/Private/RigVMCore/RigVM.cpp
@@ -2177,6 +2177,11 @@ ERigVMExecuteResult URigVM::ExecuteInstructions(FRigVMExtendedExecuteContext& Co
 			}
 			for (const FRigVMExternalVariable& ExternalVariable : ExternalVariables)
 			{
+				if (ExternalVariable.Memory == nullptr)
+				{
+					continue;
+				}
+				
 				FString Value;
 				ExternalVariable.Property->ExportTextItem_Direct(Value, ExternalVariable.Memory, nullptr, nullptr, PPF_None);
 				FString Line = FString::Printf(TEXT("External %s: %s"), *ExternalVariable.Name.ToString(), *Value);
diff --git a/Engine/Plugins/Runtime/RigVM/Source/RigVMEditor/Private/Editor/RigVMEditor.cpp b/Engine/Plugins/Runtime/RigVM/Source/RigVMEditor/Private/Editor/RigVMEditor.cpp
--- a/Engine/Plugins/Runtime/RigVM/Source/RigVMEditor/Private/Editor/RigVMEditor.cpp
+++ b/Engine/Plugins/Runtime/RigVM/Source/RigVMEditor/Private/Editor/RigVMEditor.cpp
@@ -2513,7 +2513,7 @@ void FRigVMEditor::OnCreateComment()
 						UEdGraphNode* EdNode = CommentAction.PerformAction(Graph, NULL, GraphEditor->GetPasteLocation());
 						if (UEdGraphNode_Comment* CommentNode = CastChecked<UEdGraphNode_Comment>(EdNode))
 						{
-							Controller->SetNodeColorByName(CommentNode->GetFName(), CommentNode->GetNodeTitleColor(), false);
+							Controller->SetNodeColorByName(CommentNode->GetFName(), CommentNode->CommentColor, false);
 							Controller->SetNodePositionByName(CommentNode->GetFName(), FVector2D(CommentNode->NodePosX, CommentNode->NodePosY), false);
 						}
 						Controller->CloseUndoBracket();
