Control Rig: Fix function description not showing in the function reference tooltip
#rnx

diff --git a/Engine/Plugins/Runtime/RigVM/Source/RigVM/Public/RigVMCore/RigVMGraphFunctionDefinition.h b/Engine/Plugins/Runtime/RigVM/Source/RigVM/Public/RigVMCore/RigVMGraphFunctionDefinition.h
--- a/Engine/Plugins/Runtime/RigVM/Source/RigVM/Public/RigVMCore/RigVMGraphFunctionDefinition.h
+++ b/Engine/Plugins/Runtime/RigVM/Source/RigVM/Public/RigVMCore/RigVMGraphFunctionDefinition.h
@@ -2,6 +2,7 @@
 
 #pragma once
 
+#include "RigVMObjectVersion.h"
 #include "RigVMCore/RigVMExternalVariable.h"
 #include "RigVMCore/RigVMByteCode.h"
 #include "UObject/UE5MainStreamObjectVersion.h"
@@ -422,14 +423,17 @@ struct RIGVM_API FRigVMGraphFunctionHeader
 	UPROPERTY(BlueprintReadOnly, VisibleAnywhere, Category=FunctionHeader)
 	FLinearColor NodeColor = FLinearColor::White;
 
+	UPROPERTY(meta=(DeprecatedProperty))
+	FText Tooltip_DEPRECATED;
+
 	UPROPERTY(BlueprintReadOnly, VisibleAnywhere, Category=FunctionHeader)
-	FText Tooltip;
+	FString Description;
 	
 	UPROPERTY(BlueprintReadOnly, VisibleAnywhere, Category=FunctionHeader)
 	FString Category;
 
 	UPROPERTY(BlueprintReadOnly, VisibleAnywhere, Category=FunctionHeader)
-	FString Keywords;	
+	FString Keywords;
 
 	UPROPERTY(BlueprintReadOnly, VisibleAnywhere, Category=FunctionHeader)
 	TArray<FRigVMGraphFunctionArgument> Arguments;
@@ -463,13 +467,40 @@ struct RIGVM_API FRigVMGraphFunctionHeader
 
 	FRigVMGraphFunctionData* GetFunctionData(bool bLoadIfNecessary = true) const;
 
+	FText GetTooltip() const
+	{
+		FString TooltipStr = FString::Printf(TEXT("%s (%s)\n%s"),
+		*NodeTitle,
+		*LibraryPointer.LibraryNode.GetAssetPathString(),
+		*Description);
+		return FText::FromString(TooltipStr);
+	}
+
 	friend FArchive& operator<<(FArchive& Ar, FRigVMGraphFunctionHeader& Data)
 	{
+		Ar.UsingCustomVersion(FRigVMObjectVersion::GUID);
+		
 		Ar << Data.LibraryPointer;
 		Ar << Data.Name;
 		Ar << Data.NodeTitle;
 		Ar << Data.NodeColor;
-		Ar << Data.Tooltip;
+
+		if (Ar.IsLoading())
+		{
+			if (Ar.CustomVer(FRigVMObjectVersion::GUID) < FRigVMObjectVersion::VMRemoveTooltipFromFunctionHeader)
+			{
+				Ar << Data.Tooltip_DEPRECATED;
+			}
+			else
+			{
+				Ar << Data.Description;
+			}
+		}
+		else
+		{
+			Ar << Data.Description;
+		}
+		
 		Ar << Data.Category;
 		Ar << Data.Keywords;
 		Ar << Data.Arguments;
diff --git a/Engine/Plugins/Runtime/RigVM/Source/RigVM/Public/RigVMObjectVersion.h b/Engine/Plugins/Runtime/RigVM/Source/RigVM/Public/RigVMObjectVersion.h
--- a/Engine/Plugins/Runtime/RigVM/Source/RigVM/Public/RigVMObjectVersion.h
+++ b/Engine/Plugins/Runtime/RigVM/Source/RigVM/Public/RigVMObjectVersion.h
@@ -36,6 +36,8 @@ struct RIGVM_API FRigVMObjectVersion
 		// VM Bytecode Stores the Public Context Path
 		VMBytecodeStorePublicContextPath,
 
+		// Removing unused tooltip property from frunction header
+		VMRemoveTooltipFromFunctionHeader,
 
 		// -----<new versions can be added above this line>-------------------------------------------------
 		VersionPlusOne,
diff --git a/Engine/Plugins/Runtime/RigVM/Source/RigVMDeveloper/Private/RigVMModel/Nodes/RigVMFunctionReferenceNode.cpp b/Engine/Plugins/Runtime/RigVM/Source/RigVMDeveloper/Private/RigVMModel/Nodes/RigVMFunctionReferenceNode.cpp
--- a/Engine/Plugins/Runtime/RigVM/Source/RigVMDeveloper/Private/RigVMModel/Nodes/RigVMFunctionReferenceNode.cpp
+++ b/Engine/Plugins/Runtime/RigVM/Source/RigVMDeveloper/Private/RigVMModel/Nodes/RigVMFunctionReferenceNode.cpp
@@ -19,7 +19,7 @@ FLinearColor URigVMFunctionReferenceNode::GetNodeColor() const
 
 FText URigVMFunctionReferenceNode::GetToolTipText() const
 {
-	return ReferencedFunctionHeader.Tooltip;
+	return ReferencedFunctionHeader.GetTooltip();
 }
 
 FString URigVMFunctionReferenceNode::GetNodeCategory() const
@@ -139,7 +139,7 @@ uint32 URigVMFunctionReferenceNode::GetStructureHash() const
 	Hash = HashCombine(Hash, GetTypeHash(ReferencedFunctionHeader.NodeTitle));
 	Hash = HashCombine(Hash, GetTypeHash(ReferencedFunctionHeader.LibraryPointer.LibraryNode.ToString()));
 	Hash = HashCombine(Hash, GetTypeHash(ReferencedFunctionHeader.Keywords));
-	Hash = HashCombine(Hash, GetTypeHash(ReferencedFunctionHeader.Tooltip.ToString()));
+	Hash = HashCombine(Hash, GetTypeHash(ReferencedFunctionHeader.Description));
 	Hash = HashCombine(Hash, GetTypeHash(ReferencedFunctionHeader.NodeColor));
 
 	for(const FRigVMGraphFunctionArgument& Argument : ReferencedFunctionHeader.Arguments)
diff --git a/Engine/Plugins/Runtime/RigVM/Source/RigVMDeveloper/Private/RigVMModel/Nodes/RigVMLibraryNode.cpp b/Engine/Plugins/Runtime/RigVM/Source/RigVMDeveloper/Private/RigVMModel/Nodes/RigVMLibraryNode.cpp
--- a/Engine/Plugins/Runtime/RigVM/Source/RigVMDeveloper/Private/RigVMModel/Nodes/RigVMLibraryNode.cpp
+++ b/Engine/Plugins/Runtime/RigVM/Source/RigVMDeveloper/Private/RigVMModel/Nodes/RigVMLibraryNode.cpp
@@ -238,6 +238,7 @@ FRigVMGraphFunctionHeader URigVMLibraryNode::GetFunctionHeader(IRigVMGraphFuncti
 	}	
     
     Header.Name = GetFName();
+	Header.Description = GetNodeDescription();
     Header.Category = GetNodeCategory();
     Header.Keywords = GetNodeKeywords();
 	
diff --git a/Engine/Plugins/Runtime/RigVM/Source/RigVMDeveloper/Private/RigVMModel/RigVMController.cpp b/Engine/Plugins/Runtime/RigVM/Source/RigVMDeveloper/Private/RigVMModel/RigVMController.cpp
--- a/Engine/Plugins/Runtime/RigVM/Source/RigVMDeveloper/Private/RigVMModel/RigVMController.cpp
+++ b/Engine/Plugins/Runtime/RigVM/Source/RigVMDeveloper/Private/RigVMModel/RigVMController.cpp
@@ -11290,15 +11290,22 @@ URigVMFunctionReferenceNode* URigVMController::AddFunctionReferenceNodeFromDescr
 		return nullptr;
 	}
 
-	if(!GetSchema()->SupportsGraphFunction(this, &InFunctionDefinition))
+	// Update the function header from the host itself (in case the spawner has outdated information)
+	const FRigVMGraphFunctionHeader* FunctionHeader = &InFunctionDefinition;
+	if (const FRigVMGraphFunctionData* FunctionData = InFunctionDefinition.GetFunctionData())
+	{
+		FunctionHeader = &FunctionData->Header;
+	}
+	
+	if(!GetSchema()->SupportsGraphFunction(this, FunctionHeader))
 	{
 		return nullptr;
 	}
 
-	FString NodeName = GetSchema()->GetValidNodeName(Graph, InNodeName.IsEmpty() ? InFunctionDefinition.Name.ToString() : InNodeName);
+	FString NodeName = GetSchema()->GetValidNodeName(Graph, InNodeName.IsEmpty() ? FunctionHeader->Name.ToString() : InNodeName);
 	URigVMFunctionReferenceNode* FunctionRefNode = NewObject<URigVMFunctionReferenceNode>(Graph, *NodeName);
 	FunctionRefNode->Position = InNodePosition;
-	FunctionRefNode->ReferencedFunctionHeader = InFunctionDefinition;
+	FunctionRefNode->ReferencedFunctionHeader = *FunctionHeader;
 
 	if(!AddGraphNode(FunctionRefNode, false))
 	{
@@ -11316,7 +11323,7 @@ URigVMFunctionReferenceNode* URigVMController::AddFunctionReferenceNodeFromDescr
 		BuildData->RegisterFunctionReference(FunctionRefNode->GetReferencedFunctionHeader().LibraryPointer, FunctionRefNode);
 	}
 	
-	for (const FRigVMGraphFunctionArgument& Argument : InFunctionDefinition.Arguments)
+	for (const FRigVMGraphFunctionArgument& Argument : FunctionHeader->Arguments)
 	{
 		if (URigVMPin* TargetPin = FunctionRefNode->FindPin(Argument.Name.ToString()))
 		{
@@ -11338,12 +11345,12 @@ URigVMFunctionReferenceNode* URigVMController::AddFunctionReferenceNodeFromDescr
 	if (bPrintPythonCommand)
 	{
 		const FString GraphName = GetSchema()->GetSanitizedGraphName(GetGraph()->GetGraphName());
-		const FString FunctionDefinitionName = GetSchema()->GetSanitizedNodeName(InFunctionDefinition.Name.ToString());
+		const FString FunctionDefinitionName = GetSchema()->GetSanitizedNodeName(FunctionHeader->Name.ToString());
 
 		bool bLocal = false;
 		if(IRigVMClientHost* ClientHost = GetImplementingOuter<IRigVMClientHost>())
 		{
-			if (InFunctionDefinition.LibraryPointer.HostObject == Cast<UObject>(ClientHost->GetRigVMGraphFunctionHost()))
+			if (FunctionHeader->LibraryPointer.HostObject == Cast<UObject>(ClientHost->GetRigVMGraphFunctionHost()))
 			{
 				bLocal = true;
 				RigVMPythonUtils::Print(GetSchema()->GetGraphOuterName(GetGraph()), 
@@ -11360,7 +11367,7 @@ URigVMFunctionReferenceNode* URigVMController::AddFunctionReferenceNodeFromDescr
 			RigVMPythonUtils::Print(GetSchema()->GetGraphOuterName(GetGraph()), 
 				FString::Printf(TEXT("blueprint.get_controller_by_name('%s').add_external_function_reference_node('%s', '%s', %s, '%s')"),
 						*GraphName,
-						*InFunctionDefinition.LibraryPointer.HostObject.ToString(),
+						*FunctionHeader->LibraryPointer.HostObject.ToString(),
 						*FunctionDefinitionName,
 						*RigVMPythonUtils::Vector2DToPythonString(InNodePosition), 
 						*NodeName));
diff --git a/Engine/Plugins/Runtime/RigVM/Source/RigVMEditor/Private/EdGraph/NodeSpawners/RigVMEdGraphFunctionRefNodeSpawner.cpp b/Engine/Plugins/Runtime/RigVM/Source/RigVMEditor/Private/EdGraph/NodeSpawners/RigVMEdGraphFunctionRefNodeSpawner.cpp
--- a/Engine/Plugins/Runtime/RigVM/Source/RigVMEditor/Private/EdGraph/NodeSpawners/RigVMEdGraphFunctionRefNodeSpawner.cpp
+++ b/Engine/Plugins/Runtime/RigVM/Source/RigVMEditor/Private/EdGraph/NodeSpawners/RigVMEdGraphFunctionRefNodeSpawner.cpp
@@ -97,7 +97,7 @@ URigVMEdGraphFunctionRefNodeSpawner* URigVMEdGraphFunctionRefNodeSpawner::Create
 	MenuSignature.MenuName = FText::FromName(InPublicFunction.Name);
 	MenuSignature.Category = FText::FromString(InPublicFunction.Category);
 	MenuSignature.Keywords = FText::FromString(InPublicFunction.Keywords);
-	MenuSignature.Tooltip = InPublicFunction.Tooltip;
+	MenuSignature.Tooltip = InPublicFunction.GetTooltip();
 
 	const FString PackagePathString = InAssetData.PackageName.ToString();
 	if(MenuSignature.Tooltip.IsEmpty())
