Add proper flag when importing a string default value to allow support for redirects

#changelist validated
#virtualized

diff --git a/Engine/Plugins/Runtime/RigVM/Source/RigVMDeveloper/Private/RigVMModel/RigVMNode.cpp b/Engine/Plugins/Runtime/RigVM/Source/RigVMDeveloper/Private/RigVMModel/RigVMNode.cpp
--- a/Engine/Plugins/Runtime/RigVM/Source/RigVMDeveloper/Private/RigVMModel/RigVMNode.cpp
+++ b/Engine/Plugins/Runtime/RigVM/Source/RigVMDeveloper/Private/RigVMModel/RigVMNode.cpp
@@ -494,8 +494,8 @@ TSharedPtr<FStructOnScope> URigVMNode::GetDecoratorInstance(const URigVMPin* InD
 				FRigVMPinDefaultValueImportErrorContext ErrorPipe;
 				{
 					// force logging to the error pipe for error detection
-					LOG_SCOPE_VERBOSITY_OVERRIDE(LogExec, ELogVerbosity::Verbose); 
-					ScriptStruct->ImportText(*DefaultValue, Decorator, nullptr, PPF_None, &ErrorPipe, ScriptStruct->GetName()); 
+					LOG_SCOPE_VERBOSITY_OVERRIDE(LogExec, ELogVerbosity::Verbose);
+					ScriptStruct->ImportText(*DefaultValue, Decorator, nullptr, PPF_SerializedAsImportText, &ErrorPipe, ScriptStruct->GetName());
 				}
 			}
 		}
