RigVM: Slice caching to be based on hash of structure rather than index

#rnx

#changelist validated
#virtualized

diff --git a/Engine/Plugins/Experimental/AnimNext/Source/AnimNext/Private/Graph/RigUnit_AnimNextGraphEvaluator.cpp b/Engine/Plugins/Experimental/AnimNext/Source/AnimNext/Private/Graph/RigUnit_AnimNextGraphEvaluator.cpp
--- a/Engine/Plugins/Experimental/AnimNext/Source/AnimNext/Private/Graph/RigUnit_AnimNextGraphEvaluator.cpp
+++ b/Engine/Plugins/Experimental/AnimNext/Source/AnimNext/Private/Graph/RigUnit_AnimNextGraphEvaluator.cpp
@@ -54,7 +54,7 @@ void FRigUnit_AnimNextGraphEvaluator::StaticExecute(FRigVMExtendedExecuteContext
 		// value we expect but it'll work. The ensure will signal that we need to fix the bug.
 		if (ensure(MemoryHandle.IsLazy()))
 		{
-			MemoryHandle.ComputeLazyValueIfNecessary(RigVMExecuteContext, RigVMExecuteContext.GetSlice().GetIndex());
+			MemoryHandle.ComputeLazyValueIfNecessary(RigVMExecuteContext, RigVMExecuteContext.GetSliceHash());
 		}
 
 		const uint8* SourcePtr = MemoryHandle.GetData();
diff --git a/Engine/Plugins/Runtime/RigVM/Source/RigVM/Private/RigVMCore/RigVM.cpp b/Engine/Plugins/Runtime/RigVM/Source/RigVM/Private/RigVMCore/RigVM.cpp
--- a/Engine/Plugins/Runtime/RigVM/Source/RigVM/Private/RigVMCore/RigVM.cpp
+++ b/Engine/Plugins/Runtime/RigVM/Source/RigVM/Private/RigVMCore/RigVM.cpp
@@ -1481,8 +1481,8 @@ bool URigVM::InitializeInstance(FRigVMExtendedExecuteContext& Context, bool bCop
 	Context.CachedMemoryHandles.Reset(MemoryHandleCount);
 
 	const int32 LazyBranchSize = GetByteCode().BranchInfos.Num();
-	Context.LazyBranchInstanceData.Reset(LazyBranchSize);
-	Context.LazyBranchInstanceData.SetNumZeroed(LazyBranchSize);
+	Context.LazyBranchExecuteState.Reset(LazyBranchSize);
+	Context.LazyBranchExecuteState.SetNumZeroed(LazyBranchSize);
 
 	if (bCopyMemory)
 	{
@@ -2108,17 +2108,14 @@ ERigVMExecuteResult URigVM::ExecuteInstructions(FRigVMExtendedExecuteContext& Co
 					(Op.EndInstruction != INDEX_NONE) &&
 					(Op.EndInstruction >= Op.StartInstruction))
 				{
-					const int32 SliceIndex = FMath::Max(0, Context.GetSlice().GetIndex());
-					while(!ExecutionState.HashPerSlice.IsValidIndex(SliceIndex))
-					{
-						ExecutionState.HashPerSlice.Add(UINT32_MAX);
-					}
+					const uint32 SliceHash = Context.GetSliceHash();
+					uint32& StoredHash = ExecutionState.HashPerSlice.FindOrAdd(SliceHash, UINT32_MAX);
 						
 					const uint32 Hash = GetTypeHash(ContextPublicData.GetNumExecutions());
-					if(ExecutionState.HashPerSlice[SliceIndex] != Hash)
+					if(StoredHash != Hash)
 					{
 						ExecuteInstructions(Context, Op.StartInstruction, Op.EndInstruction);
-						ExecutionState.HashPerSlice[SliceIndex] = Hash;
+						StoredHash = Hash;
 					}
 				}
 
diff --git a/Engine/Plugins/Runtime/RigVM/Source/RigVM/Private/RigVMCore/RigVMExecuteContext.cpp b/Engine/Plugins/Runtime/RigVM/Source/RigVM/Private/RigVMCore/RigVMExecuteContext.cpp
--- a/Engine/Plugins/Runtime/RigVM/Source/RigVM/Private/RigVMCore/RigVMExecuteContext.cpp
+++ b/Engine/Plugins/Runtime/RigVM/Source/RigVM/Private/RigVMCore/RigVMExecuteContext.cpp
@@ -74,7 +74,7 @@ void FRigVMExtendedExecuteContext::Reset()
 
 	CachedMemoryHandles.Reset();
 
-	LazyBranchInstanceData.Reset();
+	LazyBranchExecuteState.Reset();
 	ExternalVariableRuntimeData.Reset();
 
 	if(PublicDataScope.IsValid())
@@ -149,7 +149,7 @@ FRigVMExtendedExecuteContext& FRigVMExtendedExecuteContext::operator =(const FRi
 
 	CachedMemoryHandles = Other.CachedMemoryHandles;
 
-	LazyBranchInstanceData = Other.LazyBranchInstanceData;
+	LazyBranchExecuteState = Other.LazyBranchExecuteState;
 	ExternalVariableRuntimeData = Other.ExternalVariableRuntimeData;
 
 	return *this;
diff --git a/Engine/Plugins/Runtime/RigVM/Source/RigVM/Private/RigVMCore/RigVMMemoryStorage.cpp b/Engine/Plugins/Runtime/RigVM/Source/RigVM/Private/RigVMCore/RigVMMemoryStorage.cpp
--- a/Engine/Plugins/Runtime/RigVM/Source/RigVM/Private/RigVMCore/RigVMMemoryStorage.cpp
+++ b/Engine/Plugins/Runtime/RigVM/Source/RigVM/Private/RigVMCore/RigVMMemoryStorage.cpp
@@ -91,24 +91,16 @@ ERigVMExecuteResult FRigVMLazyBranch::Execute(FRigVMExtendedExecuteContext& Cont
 	return FunctionPtr();
 }
 
-ERigVMExecuteResult FRigVMLazyBranch::ExecuteIfRequired(FRigVMExtendedExecuteContext& Context, int32 InSliceIndex)
+ERigVMExecuteResult FRigVMLazyBranch::ExecuteIfRequired(FRigVMExtendedExecuteContext& Context, uint32 InSliceHash)
 {
 	check(VM);
 
-	if(InSliceIndex == INDEX_NONE)
-	{
-		InSliceIndex = 0;
-	}
-	
-	while(!Context.LazyBranchInstanceData[BranchInfo.Index].LastVMNumExecutions.IsValidIndex(InSliceIndex))
-	{
-		Context.LazyBranchInstanceData[BranchInfo.Index].LastVMNumExecutions.Add(INDEX_NONE);
-	}
-
-	if(Context.GetNumExecutions() != Context.LazyBranchInstanceData[BranchInfo.Index].LastVMNumExecutions[InSliceIndex])
+	const uint32 Hash = GetTypeHash(Context.GetNumExecutions());
+	uint32& StoredHash = Context.LazyBranchExecuteState[BranchInfo.Index].HashPerSlice.FindOrAdd(InSliceHash, UINT32_MAX);
+	if(Hash != StoredHash)
 	{
 		const ERigVMExecuteResult Result = Execute(Context);
-		Context.LazyBranchInstanceData[BranchInfo.Index].LastVMNumExecutions[InSliceIndex] = Context.GetNumExecutions();
+		StoredHash = Hash;
 		return Result;
 	}
 	
@@ -122,7 +114,7 @@ const uint8* TRigVMLazyValueBase::GetData(const FRigVMExecuteContext& Context) c
 		if(MemoryHandle->IsLazy())
 		{
 			check(Context.ExtendedExecuteContext != nullptr);
-			MemoryHandle->ComputeLazyValueIfNecessary(*Context.ExtendedExecuteContext, SliceIndex);
+			MemoryHandle->ComputeLazyValueIfNecessary(*Context.ExtendedExecuteContext, SliceHash);
 		}
 		return MemoryHandle->GetData(bFollowPropertyPath, INDEX_NONE);
 	}
@@ -131,10 +123,10 @@ const uint8* TRigVMLazyValueBase::GetData(const FRigVMExecuteContext& Context) c
 
 ////////////////////////////////////////////////////////////////////////////////////////////////////////////////////////
 
-bool FRigVMMemoryHandle::ComputeLazyValueIfNecessary(FRigVMExtendedExecuteContext& Context, int32 InSliceIndex)
+bool FRigVMMemoryHandle::ComputeLazyValueIfNecessary(FRigVMExtendedExecuteContext& Context, uint32 InSliceHash)
 {
 	check(IsLazy());
-	return LazyBranch->ExecuteIfRequired(Context, InSliceIndex) != ERigVMExecuteResult::Failed;
+	return LazyBranch->ExecuteIfRequired(Context, InSliceHash) != ERigVMExecuteResult::Failed;
 }
 
 ////////////////////////////////////////////////////////////////////////////////////////////////////////////////////////
diff --git a/Engine/Plugins/Runtime/RigVM/Source/RigVM/Private/RigVMFunctions/RigVMDispatch_If.cpp b/Engine/Plugins/Runtime/RigVM/Source/RigVM/Private/RigVMFunctions/RigVMDispatch_If.cpp
--- a/Engine/Plugins/Runtime/RigVM/Source/RigVM/Private/RigVMFunctions/RigVMDispatch_If.cpp
+++ b/Engine/Plugins/Runtime/RigVM/Source/RigVM/Private/RigVMFunctions/RigVMDispatch_If.cpp
@@ -86,7 +86,7 @@ void FRigVMDispatch_If::Execute(FRigVMExtendedExecuteContext& InContext, FRigVMM
 	FRigVMMemoryHandle& InputHandle = Condition ? Handles[1] : Handles[2];
 	if(InputHandle.IsLazy())
 	{
-		InputHandle.ComputeLazyValueIfNecessary(InContext, InContext.GetSlice().GetIndex());
+		InputHandle.ComputeLazyValueIfNecessary(InContext, InContext.GetSliceHash());
 	}
 
 	const uint8* Input = InputHandle.GetData();
diff --git a/Engine/Plugins/Runtime/RigVM/Source/RigVM/Private/RigVMFunctions/RigVMDispatch_Select.cpp b/Engine/Plugins/Runtime/RigVM/Source/RigVM/Private/RigVMFunctions/RigVMDispatch_Select.cpp
--- a/Engine/Plugins/Runtime/RigVM/Source/RigVM/Private/RigVMFunctions/RigVMDispatch_Select.cpp
+++ b/Engine/Plugins/Runtime/RigVM/Source/RigVM/Private/RigVMFunctions/RigVMDispatch_Select.cpp
@@ -140,7 +140,7 @@ void FRigVMDispatch_SelectInt32::Execute(FRigVMExtendedExecuteContext& InContext
 	FRigVMMemoryHandle& InputHandle = Handles[1 + Index];
 	if(InputHandle.IsLazy())
 	{
-		InputHandle.ComputeLazyValueIfNecessary(InContext, InContext.GetSlice().GetIndex());
+		InputHandle.ComputeLazyValueIfNecessary(InContext, InContext.GetSliceHash());
 	}
 
 	const uint8* Input = InputHandle.GetData();
diff --git a/Engine/Plugins/Runtime/RigVM/Source/RigVM/Public/RigVMCore/RigVMByteCode.h b/Engine/Plugins/Runtime/RigVM/Source/RigVM/Public/RigVMCore/RigVMByteCode.h
--- a/Engine/Plugins/Runtime/RigVM/Source/RigVM/Public/RigVMCore/RigVMByteCode.h
+++ b/Engine/Plugins/Runtime/RigVM/Source/RigVM/Public/RigVMCore/RigVMByteCode.h
@@ -94,21 +94,6 @@ struct FRigVMBranchInfoKey
 	FName Label;
 };
 
-// A runtime cache for determining if a set of instruction has to
-// run for this execution of the VM
-USTRUCT(BlueprintType)
-struct RIGVM_API FRigVMInstructionSetExecuteState
-{
-	GENERATED_BODY()
-
-	FRigVMInstructionSetExecuteState()
-	{
-	}
-
-	UPROPERTY()
-	TArray<uint32> HashPerSlice;
-};
-
 // A description of a predicate branch in the VM's bytecode
 USTRUCT()
 struct RIGVM_API FRigVMPredicateBranch
diff --git a/Engine/Plugins/Runtime/RigVM/Source/RigVM/Public/RigVMCore/RigVMExecuteContext.h b/Engine/Plugins/Runtime/RigVM/Source/RigVM/Public/RigVMCore/RigVMExecuteContext.h
--- a/Engine/Plugins/Runtime/RigVM/Source/RigVM/Public/RigVMCore/RigVMExecuteContext.h
+++ b/Engine/Plugins/Runtime/RigVM/Source/RigVM/Public/RigVMCore/RigVMExecuteContext.h
@@ -182,6 +182,11 @@ public:
 		}
 	}
 
+	uint32 GetHash() const
+	{
+		return HashCombine(GetTypeHash(LowerBound), GetTypeHash(Index));
+	}
+
 private:
 
 	int32 LowerBound;
@@ -496,14 +501,6 @@ struct TStructOpsTypeTraits<FRigVMExecuteContext> : public TStructOpsTypeTraitsB
 	};
 };
 
-/**
- * Lazy branch data required by each instance of the VM
- */
-struct RIGVM_API FRigVMLazyBranchInstanceData
-{
-	TArray<int32> LastVMNumExecutions;
-};
-
 struct RIGVM_API FRigVMExternalVariableRuntimeData
 {
 	explicit FRigVMExternalVariableRuntimeData(uint8* InMemory)
@@ -647,6 +644,16 @@ struct RIGVM_API FRigVMExtendedExecuteContext
 		return GetSlice().IsComplete();
 	}
 
+	uint32 GetSliceHash() const
+	{
+		uint32 Hash = 0;
+		for(const FRigVMSlice& Slice : Slices)
+		{
+			Hash = HashCombine(Hash, Slice.GetHash());
+		}
+		return Hash;
+	}
+
 	bool IsValidArrayIndex(int32& InOutIndex, const FScriptArrayHelper& InArrayHelper) const
 	{
 		return IsValidArrayIndex(InOutIndex, InArrayHelper.Num());
@@ -701,7 +708,7 @@ struct RIGVM_API FRigVMExtendedExecuteContext
 	void InvalidateCachedMemory()
 	{
 		CachedMemoryHandles.Reset();
-		LazyBranchInstanceData.Reset();
+		LazyBranchExecuteState.Reset();
 	}
 
 	uint32 GetNumExecutions() const
@@ -767,7 +774,7 @@ struct RIGVM_API FRigVMExtendedExecuteContext
 	FExecutionReachedExitEvent OnExecutionReachedExit;
 	FExecutionReachedExitEvent& ExecutionReachedExit() { return OnExecutionReachedExit; }
 
-	TArray<FRigVMLazyBranchInstanceData> LazyBranchInstanceData;
+	TArray<FRigVMInstructionSetExecuteState> LazyBranchExecuteState;
 	TArray<FRigVMExternalVariableRuntimeData> ExternalVariableRuntimeData;
 
 #if WITH_EDITOR
diff --git a/Engine/Plugins/Runtime/RigVM/Source/RigVM/Public/RigVMCore/RigVMMemoryCommon.h b/Engine/Plugins/Runtime/RigVM/Source/RigVM/Public/RigVMCore/RigVMMemoryCommon.h
--- a/Engine/Plugins/Runtime/RigVM/Source/RigVM/Public/RigVMCore/RigVMMemoryCommon.h
+++ b/Engine/Plugins/Runtime/RigVM/Source/RigVM/Public/RigVMCore/RigVMMemoryCommon.h
@@ -61,6 +61,21 @@ namespace RigVM
 	}
 }
 
+// A runtime cache for determining if a set of instruction has to
+// run for this execution of the VM
+USTRUCT(BlueprintType)
+struct RIGVM_API FRigVMInstructionSetExecuteState
+{
+	GENERATED_BODY()
+
+	FRigVMInstructionSetExecuteState()
+	{
+	}
+
+	UPROPERTY()
+	TMap<uint32, uint32> HashPerSlice;
+};
+
 /**
  * The FRigVMOperand represents an argument used for an operator
  * within the virtual machine. Operands provide information about
diff --git a/Engine/Plugins/Runtime/RigVM/Source/RigVM/Public/RigVMCore/RigVMMemoryStorage.h b/Engine/Plugins/Runtime/RigVM/Source/RigVM/Public/RigVMCore/RigVMMemoryStorage.h
--- a/Engine/Plugins/Runtime/RigVM/Source/RigVM/Public/RigVMCore/RigVMMemoryStorage.h
+++ b/Engine/Plugins/Runtime/RigVM/Source/RigVM/Public/RigVMCore/RigVMMemoryStorage.h
@@ -146,7 +146,7 @@ public:
 	{}
 
 	ERigVMExecuteResult Execute(FRigVMExtendedExecuteContext& Context);
-	ERigVMExecuteResult ExecuteIfRequired(FRigVMExtendedExecuteContext& Context, int32 InSliceIndex);
+	ERigVMExecuteResult ExecuteIfRequired(FRigVMExtendedExecuteContext& Context, uint32 InSliceHash);
 
 private:
 
@@ -167,7 +167,7 @@ public:
 
 	TRigVMLazyValueBase()
 		: bFollowPropertyPath(false)
-		, SliceIndex(INDEX_NONE)
+		, SliceHash(UINT32_MAX)
 		, MemoryHandle(nullptr)
 	{
 	}
@@ -178,7 +178,7 @@ public:
 protected:
 	
 	bool bFollowPropertyPath;
-	int32 SliceIndex;
+	uint32 SliceHash;
 	FRigVMMemoryHandle* MemoryHandle;
 	
 	friend class URigVM;
@@ -304,7 +304,7 @@ public:
 	 * Computes the data if necessary and returns true if the value is valid
 	 * @return True if the value of the handle is valid after the compute
 	 */
-	RIGVM_API bool ComputeLazyValueIfNecessary(FRigVMExtendedExecuteContext& Context, int32 InSliceIndex = INDEX_NONE);
+	RIGVM_API bool ComputeLazyValueIfNecessary(FRigVMExtendedExecuteContext& Context, uint32 InSliceHash);
 
 	// Returns the head property of this handle
 	const FProperty* GetProperty() const
@@ -869,13 +869,13 @@ private:
 		return Ptr;
 	}
 
-	TRigVMLazyValueBase GetDataLazily_Internal(bool bFollowPropertyPath, int32 InSliceIndex) const
+	TRigVMLazyValueBase GetDataLazily_Internal(bool bFollowPropertyPath, uint32 InSliceHash) const
 	{
 		// note: this works also for memory handles which don't provide a lazy branch
 		TRigVMLazyValueBase LazyValue;
 		LazyValue.MemoryHandle = (FRigVMMemoryHandle*)this;
 		LazyValue.bFollowPropertyPath = bFollowPropertyPath;
-		LazyValue.SliceIndex = InSliceIndex;
+		LazyValue.SliceHash = InSliceHash;
 		return LazyValue;
 	}
 
