RigVM: Fix deprecation issue with InstructionSetExecuteState

#rnx

#changelist validated
#virtualized

diff --git a/Engine/Plugins/Runtime/RigVM/Source/RigVM/Private/RigVMCore/RigVM.cpp b/Engine/Plugins/Runtime/RigVM/Source/RigVM/Private/RigVMCore/RigVM.cpp
--- a/Engine/Plugins/Runtime/RigVM/Source/RigVM/Private/RigVMCore/RigVM.cpp
+++ b/Engine/Plugins/Runtime/RigVM/Source/RigVM/Private/RigVMCore/RigVM.cpp
@@ -2109,7 +2109,7 @@ ERigVMExecuteResult URigVM::ExecuteInstructions(FRigVMExtendedExecuteContext& Co
 					(Op.EndInstruction >= Op.StartInstruction))
 				{
 					const uint32 SliceHash = Context.GetSliceHash();
-					uint32& StoredHash = ExecutionState.HashPerSlice.FindOrAdd(SliceHash, UINT32_MAX);
+					uint32& StoredHash = ExecutionState.SliceHashToNumInstruction.FindOrAdd(SliceHash, UINT32_MAX);
 						
 					const uint32 Hash = GetTypeHash(ContextPublicData.GetNumExecutions());
 					if(StoredHash != Hash)
diff --git a/Engine/Plugins/Runtime/RigVM/Source/RigVM/Private/RigVMCore/RigVMMemoryStorage.cpp b/Engine/Plugins/Runtime/RigVM/Source/RigVM/Private/RigVMCore/RigVMMemoryStorage.cpp
--- a/Engine/Plugins/Runtime/RigVM/Source/RigVM/Private/RigVMCore/RigVMMemoryStorage.cpp
+++ b/Engine/Plugins/Runtime/RigVM/Source/RigVM/Private/RigVMCore/RigVMMemoryStorage.cpp
@@ -96,7 +96,7 @@ ERigVMExecuteResult FRigVMLazyBranch::ExecuteIfRequired(FRigVMExtendedExecuteCon
 	check(VM);
 
 	const uint32 Hash = GetTypeHash(Context.GetNumExecutions());
-	uint32& StoredHash = Context.LazyBranchExecuteState[BranchInfo.Index].HashPerSlice.FindOrAdd(InSliceHash, UINT32_MAX);
+	uint32& StoredHash = Context.LazyBranchExecuteState[BranchInfo.Index].SliceHashToNumInstruction.FindOrAdd(InSliceHash, UINT32_MAX);
 	if(Hash != StoredHash)
 	{
 		const ERigVMExecuteResult Result = Execute(Context);
diff --git a/Engine/Plugins/Runtime/RigVM/Source/RigVM/Public/RigVMCore/RigVMMemoryCommon.h b/Engine/Plugins/Runtime/RigVM/Source/RigVM/Public/RigVMCore/RigVMMemoryCommon.h
--- a/Engine/Plugins/Runtime/RigVM/Source/RigVM/Public/RigVMCore/RigVMMemoryCommon.h
+++ b/Engine/Plugins/Runtime/RigVM/Source/RigVM/Public/RigVMCore/RigVMMemoryCommon.h
@@ -72,8 +72,8 @@ struct RIGVM_API FRigVMInstructionSetExecuteState
 	{
 	}
 
-	UPROPERTY()
-	TMap<uint32, uint32> HashPerSlice;
+	UPROPERTY(transient)
+	TMap<uint32, uint32> SliceHashToNumInstruction;
 };
 
 /**
