Control Rig: Fix crash when accessing an array element in the details panel
#rnx

#changelist validated
#virtualized

diff --git a/Engine/Plugins/Runtime/RigVM/Source/RigVMEditor/Private/Editor/RigVMDetailsViewWrapperObject.cpp b/Engine/Plugins/Runtime/RigVM/Source/RigVMEditor/Private/Editor/RigVMDetailsViewWrapperObject.cpp
--- a/Engine/Plugins/Runtime/RigVM/Source/RigVMEditor/Private/Editor/RigVMDetailsViewWrapperObject.cpp
+++ b/Engine/Plugins/Runtime/RigVM/Source/RigVMEditor/Private/Editor/RigVMDetailsViewWrapperObject.cpp
@@ -597,7 +597,12 @@ void URigVMDetailsViewWrapperObject::PostEditChangeChainProperty(FPropertyChange
 	const TDoubleLinkedList<FProperty*>::TDoubleLinkedListNode* PropertyNode = PropertyChangedEvent.PropertyChain.GetHead();
 	do
 	{
-		const FString PropertyName = PropertyNode->GetValue()->GetNameCPP();
+		FString PropertyName = PropertyNode->GetValue()->GetNameCPP();
+		int32 Index = PropertyChangedEvent.GetArrayIndex(PropertyName);
+		if (Index != INDEX_NONE)
+		{
+			PropertyName = FString::Printf(TEXT("%s[%d]"), *PropertyName, Index);
+		}
 		if(PropertyPath.IsEmpty())
 		{
 			PropertyPath = PropertyName;
diff --git a/Engine/Plugins/Runtime/RigVM/Source/RigVMEditor/Private/Editor/RigVMEditor.cpp b/Engine/Plugins/Runtime/RigVM/Source/RigVMEditor/Private/Editor/RigVMEditor.cpp
--- a/Engine/Plugins/Runtime/RigVM/Source/RigVMEditor/Private/Editor/RigVMEditor.cpp
+++ b/Engine/Plugins/Runtime/RigVM/Source/RigVMEditor/Private/Editor/RigVMEditor.cpp
@@ -2102,36 +2102,46 @@ void FRigVMEditor::OnWrappedPropertyChangedChainEvent(URigVMDetailsViewWrapperOb
 		URigVMNode* Node = CastChecked<URigVMNode>(InWrapperObject->GetSubject());
 
 		const FName RootPinName = InPropertyChangedChainEvent.PropertyChain.GetHead()->GetValue()->GetFName();
-		const FProperty* Property = WrapperObjects[0]->GetClass()->FindPropertyByName(RootPinName);
-		uint8* PropertyStorage = Property->ContainerPtrToValuePtr<uint8>(WrapperObjects[0].Get());
 		const FString RootPinNameString = RootPinName.ToString();
 		FString PinPath = URigVMPin::JoinPinPath(Node->GetName(), RootPinNameString);
-
-		// traverse to get to the target pin
-		if(!InPropertyPath.Equals(RootPinNameString))
+		
+		const FProperty* Property = WrapperObjects[0]->GetClass()->FindPropertyByName(RootPinName);
+		uint8* PropertyStorage = nullptr;
+		if (Property)
 		{
-			const FString RootPinPrefix = RootPinNameString + TEXT("->");
-			check(InPropertyPath.StartsWith(RootPinPrefix));
-			const FString RemainingPropertyPath = InPropertyPath.Mid(RootPinPrefix.Len());
-			const FString SegmentPath = RemainingPropertyPath.Replace(TEXT("->"), TEXT("."));
+			PropertyStorage = Property->ContainerPtrToValuePtr<uint8>(WrapperObjects[0].Get());
+
+			// traverse to get to the target pin
+			if(!InPropertyPath.Equals(RootPinNameString))
+			{
+				check(InPropertyPath.StartsWith(RootPinNameString));
+				FString RemainingPropertyPath = InPropertyPath.Mid(RootPinNameString.Len());
+				RemainingPropertyPath.RemoveFromStart(TEXT("->"));
+				const FString SegmentPath = RemainingPropertyPath.Replace(TEXT("->"), TEXT("."));
 			
-			const FRigVMPropertyPath PropertyTraverser(Property, SegmentPath);
-			PropertyStorage = PropertyTraverser.GetData<uint8>(PropertyStorage, Property);
-			Property = PropertyTraverser.GetTailProperty();
-			PinPath = URigVMPin::JoinPinPath(PinPath, SegmentPath);
+				const FRigVMPropertyPath PropertyTraverser(Property, SegmentPath);
+				PropertyStorage = PropertyTraverser.GetData<uint8>(PropertyStorage, Property);
+				Property = PropertyTraverser.GetTailProperty();
+				PinPath = URigVMPin::JoinPinPath(PinPath, SegmentPath);
+				PinPath.ReplaceInline(TEXT("["), TEXT(""));
+				PinPath.ReplaceInline(TEXT("]"), TEXT(""));
+			}
 		}
 
-		FString DefaultValue = FRigVMStruct::ExportToFullyQualifiedText(Property, PropertyStorage);
-		if(Property->IsA<FStrProperty>() || Property->IsA<FNameProperty>())
+		if (Property)
 		{
-			DefaultValue.TrimCharInline(TEXT('\"'), nullptr);
-		}
-		if (!DefaultValue.IsEmpty())
-		{
-			const bool bInteractive = InPropertyChangedChainEvent.ChangeType == EPropertyChangeType::Interactive;
-			URigVMController* Controller = GetRigVMBlueprint()->GetController(Node->GetGraph());
-			check(Controller);
-			Controller->SetPinDefaultValue(PinPath, DefaultValue, true, !bInteractive, true, !bInteractive);
+			FString DefaultValue = FRigVMStruct::ExportToFullyQualifiedText(Property, PropertyStorage);
+			if(Property->IsA<FStrProperty>() || Property->IsA<FNameProperty>())
+			{
+				DefaultValue.TrimCharInline(TEXT('\"'), nullptr);
+			}
+			if (!DefaultValue.IsEmpty())
+			{
+				const bool bInteractive = InPropertyChangedChainEvent.ChangeType == EPropertyChangeType::Interactive;
+				URigVMController* Controller = GetRigVMBlueprint()->GetController(Node->GetGraph());
+				check(Controller);
+				Controller->SetPinDefaultValue(PinPath, DefaultValue, true, !bInteractive, true, !bInteractive);
+			}
 		}
 	}
 }
