RigVMEdGraph: Correct detect cycles for output execute pins

#rnx

diff --git a/Engine/Plugins/Runtime/RigVM/Source/RigVMDeveloper/Private/RigVMModel/RigVMGraph.cpp b/Engine/Plugins/Runtime/RigVM/Source/RigVMDeveloper/Private/RigVMModel/RigVMGraph.cpp
--- a/Engine/Plugins/Runtime/RigVM/Source/RigVMDeveloper/Private/RigVMModel/RigVMGraph.cpp
+++ b/Engine/Plugins/Runtime/RigVM/Source/RigVMDeveloper/Private/RigVMModel/RigVMGraph.cpp
@@ -520,6 +520,22 @@ void URigVMGraph::PrepareCycleChecking(URigVMPin* InPin, bool bAsInput)
 	if (InPin)
 	{
 		LinksToSkip = InPin->GetLinks();
+
+		// since execute output pins also only allow one link we need to ignore links
+		// on other output links. this is so that cycle checking doesn't kick in irregularly.
+		if(InPin->IsExecuteContext() && InPin->GetDirection() == ERigVMPinDirection::Output)
+		{
+			for(const URigVMPin* OtherPin : InPin->GetNode()->GetPins())
+			{
+				if(OtherPin != InPin)
+				{
+					if(OtherPin->IsExecuteContext() && OtherPin->GetDirection() == ERigVMPinDirection::Output)
+					{
+						LinksToSkip.Append(OtherPin->GetLinks());
+					}
+				}
+			}
+		}
 	}
 
 	GetDiagnosticsAST(false, LinksToSkip)->PrepareCycleChecking(InPin);
