RigVM: Prevent a possible 32-bit overflow when allocating a temporary storage.

#rnx

diff --git a/Engine/Plugins/Runtime/RigVM/Source/RigVM/Private/RigVMFunctions/RigVMDispatch_Array.cpp b/Engine/Plugins/Runtime/RigVM/Source/RigVM/Private/RigVMFunctions/RigVMDispatch_Array.cpp
--- a/Engine/Plugins/Runtime/RigVM/Source/RigVM/Private/RigVMFunctions/RigVMDispatch_Array.cpp
+++ b/Engine/Plugins/Runtime/RigVM/Source/RigVM/Private/RigVMFunctions/RigVMDispatch_Array.cpp
@@ -4,6 +4,7 @@
 #include "RigVMCore/RigVMRegistry.h"
 #include "RigVMCore/RigVMStruct.h"
 #include "RigVMStringUtils.h"
+#include "Math/GuardedInt.h"
 
 #include UE_INLINE_GENERATED_CPP_BY_NAME(RigVMDispatch_Array)
 #define LOCTEXT_NAMESPACE "RigVMDispatch_Array"
@@ -1182,12 +1183,15 @@ void FRigVMDispatch_ArrayUnion::Execute(FRigVMExtendedExecuteContext& InContext,
 			}
 		}
 
-		if (InContext.IsValidArraySize(FinalCount))
+		// Check if we get a 32-bit overflow due to malformed GetSize() result and bail early if so.
+		const FGuardedInt32 TempStorageSize = FGuardedInt32(ArrayHelper.Num()) * ArrayElementProperty->GetSize();
+		
+		if (InContext.IsValidArraySize(FinalCount) && TempStorageSize.IsValid())
 		{
 			// copy the complete array to a temp storage
 			TArray<uint8, TAlignedHeapAllocator<16>> TempStorage;
 			const int32 NumElementsA = ArrayHelper.Num();
-			TempStorage.AddZeroed(NumElementsA * ArrayElementProperty->GetSize());
+			TempStorage.AddZeroed(TempStorageSize.GetChecked());
 			uint8* TempMemory = TempStorage.GetData();
 			for(int32 Index = 0; Index < NumElementsA; Index++)
 			{
