Control Rig: Mark old editor's rigvmhost as garbage when we are changing it.
#rnx

diff --git a/Engine/Plugins/Runtime/RigVM/Source/RigVMEditor/Private/Editor/RigVMEditor.cpp b/Engine/Plugins/Runtime/RigVM/Source/RigVMEditor/Private/Editor/RigVMEditor.cpp
--- a/Engine/Plugins/Runtime/RigVM/Source/RigVMEditor/Private/Editor/RigVMEditor.cpp
+++ b/Engine/Plugins/Runtime/RigVM/Source/RigVMEditor/Private/Editor/RigVMEditor.cpp
@@ -2907,6 +2907,11 @@ void FRigVMEditor::SetHost(URigVMHost* InHost)
 {
 	if (URigVMBlueprint* RigVMBlueprint = GetRigVMBlueprint())
 	{
+		if (IsValid(RigVMBlueprint->EditorHost) && RigVMBlueprint->EditorHost->GetOuter() == GetOuterForHost())
+		{
+			RigVMBlueprint->EditorHost->Rename(nullptr, GetTransientPackage(), REN_ForceNoResetLoaders | REN_DoNotDirty | REN_DontCreateRedirectors | REN_NonTransactional);
+			RigVMBlueprint->EditorHost->MarkAsGarbage();
+		}
 		RigVMBlueprint->EditorHost = InHost;
 		if(RigVMBlueprint->EditorHost && IsValid(RigVMBlueprint->EditorHost))
 		{
