AnimNext: Merging parameter blocks and graphs - part 2

Merged execution contexts into one: FAnimNextExecuteContext. Per-graph type context data is now accessible via FAnimNextExecuteContext::GetContextData<T>.

RigVM changes:
- Move authoratative schema to URigVMGraph. This will eventually also allow multiple different graph types in a single IRigVMClientHost asset.
- Assets still have a 'default' schema that is used for the null controller and for any kinds of assets that dont have per-graph schemas (e.g. control rig).
- Made URigVMSchema usage non-stateful. Access to the schema of a particular class now does directly to the CDO rather than sorting a local ptr to the CDO.

diff --git a/Engine/Plugins/Experimental/AnimNext/Source/AnimNext/Internal/AnimNextExecuteContext.h b/Engine/Plugins/Experimental/AnimNext/Source/AnimNext/Internal/AnimNextExecuteContext.h
new file mode 100644
--- /dev/null
+++ b/Engine/Plugins/Experimental/AnimNext/Source/AnimNext/Internal/AnimNextExecuteContext.h
@@ -0,0 +1,62 @@
+﻿// Copyright Epic Games, Inc. All Rights Reserved.
+
+#pragma once
+
+#include "CoreMinimal.h"
+#include "RigVMCore/RigVMExecuteContext.h"
+#include "Graph/AnimNextGraphContextData.h"
+#include "Param/AnimNextParamContextData.h"
+#include "Misc/TVariant.h"
+#include "AnimNextExecuteContext.generated.h"
+
+struct FAnimNextGraphInstance;
+
+namespace UE::AnimNext
+{
+	struct FLatentPropertyHandle;
+}
+
+USTRUCT(BlueprintType)
+struct FAnimNextExecuteContext : public FRigVMExecuteContext
+{
+	GENERATED_BODY()
+
+	FAnimNextExecuteContext() = default;
+
+	virtual void Copy(const FRigVMExecuteContext* InOtherContext) override
+	{
+		Super::Copy(InOtherContext);
+
+		const FAnimNextExecuteContext* OtherContext = static_cast<const FAnimNextExecuteContext*>(InOtherContext);
+		ContextData = OtherContext->ContextData;
+	}
+
+	// Get the context data as the specified type. This will assert if the type differs from the last call to SetContextData.
+	template<typename ContextType>
+	const ContextType& GetContextData() const
+	{
+		return ContextData.Get<ContextType>();
+	}
+
+protected:
+	// Setup the context data to the specified type
+	template<typename ContextType, typename... ArgsType>
+	void SetContextData(ArgsType&&... Args)
+	{
+		ContextData.Emplace<ContextType>(Forward<ArgsType>(Args)...);
+	}
+
+	// Call this to reset the context to its original state to detect stale usage, can't call it Reset due to virtual in base with that name
+	template<typename ContextType>
+	void DebugReset()
+	{
+		ContextData.Get<ContextType>().Reset();
+	}
+
+	// All possible known variants of our context data. IF we ever want this to be extensible, this can be converted into an FInstancedStruct
+	TVariant<FAnimNextGraphContextData, FAnimNextParamContextData> ContextData;
+
+	friend struct FAnimNextGraphInstance;
+	friend class UAnimNextParameterBlock;
+};
+
diff --git a/Engine/Plugins/Experimental/AnimNext/Source/AnimNext/Internal/Graph/AnimNextExecuteContext.h b/Engine/Plugins/Experimental/AnimNext/Source/AnimNext/Internal/Graph/AnimNextExecuteContext.h
deleted file mode 100644
diff --git a/Engine/Plugins/Experimental/AnimNext/Source/AnimNext/Internal/Graph/AnimNextGraphContextData.h b/Engine/Plugins/Experimental/AnimNext/Source/AnimNext/Internal/Graph/AnimNextGraphContextData.h
new file mode 100644
--- /dev/null
+++ b/Engine/Plugins/Experimental/AnimNext/Source/AnimNext/Internal/Graph/AnimNextGraphContextData.h
@@ -0,0 +1,52 @@
+// Copyright Epic Games, Inc. All Rights Reserved.
+
+#pragma once
+
+#include "CoreMinimal.h"
+#include "AnimNextGraphContextData.generated.h"
+
+struct FAnimNextGraphInstance;
+
+namespace UE::AnimNext
+{
+	struct FLatentPropertyHandle;
+}
+
+USTRUCT()
+struct FAnimNextGraphContextData
+{
+	GENERATED_BODY()
+
+	FAnimNextGraphContextData() = default;
+
+	FAnimNextGraphContextData(const FAnimNextGraphInstance* InInstance, const TConstArrayView<UE::AnimNext::FLatentPropertyHandle>& InLatentHandles, void* InDestinationBasePtr, bool bInIsFrozen)
+		: Instance(InInstance)
+		, LatentHandles(InLatentHandles)
+		, DestinationBasePtr(InDestinationBasePtr)
+		, bIsFrozen(bInIsFrozen)
+	{
+	}
+
+	const FAnimNextGraphInstance& GetGraphInstance() const { check(Instance); return *Instance; }
+	const TConstArrayView<UE::AnimNext::FLatentPropertyHandle>& GetLatentHandles() const { return LatentHandles; }
+	void* GetDestinationBasePtr() const { return DestinationBasePtr; }
+	bool IsFrozen() const { return bIsFrozen; }
+
+private:
+	// Call this to reset the context to its original state to detect stale usage
+	void Reset()
+	{
+		Instance = nullptr;
+		LatentHandles = TConstArrayView<UE::AnimNext::FLatentPropertyHandle>();
+		DestinationBasePtr = nullptr;
+		bIsFrozen = false;
+	}
+
+	const FAnimNextGraphInstance* Instance = nullptr;
+	TConstArrayView<UE::AnimNext::FLatentPropertyHandle> LatentHandles;
+	void* DestinationBasePtr = nullptr;
+	bool bIsFrozen = false;
+
+	friend struct FAnimNextGraphInstance;
+	friend struct FAnimNextExecuteContext;
+};
diff --git a/Engine/Plugins/Experimental/AnimNext/Source/AnimNext/Internal/Graph/RigUnit_AnimNextBase.h b/Engine/Plugins/Experimental/AnimNext/Source/AnimNext/Internal/Graph/RigUnit_AnimNextBase.h
new file mode 100644
--- /dev/null
+++ b/Engine/Plugins/Experimental/AnimNext/Source/AnimNext/Internal/Graph/RigUnit_AnimNextBase.h
@@ -0,0 +1,14 @@
+﻿// Copyright Epic Games, Inc. All Rights Reserved.
+
+#pragma once
+
+#include "CoreMinimal.h"
+#include "Units/RigUnit.h"
+#include "AnimNextExecuteContext.h"
+#include "RigUnit_AnimNextBase.generated.h"
+
+USTRUCT(meta=(ExecuteContext="FAnimNextExecuteContext"))
+struct FRigUnit_AnimNextBase : public FRigUnit
+{
+	GENERATED_BODY()
+};
diff --git a/Engine/Plugins/Experimental/AnimNext/Source/AnimNext/Internal/Graph/RigUnit_AnimNextBeginExecution.h b/Engine/Plugins/Experimental/AnimNext/Source/AnimNext/Internal/Graph/RigUnit_AnimNextBeginExecution.h
--- a/Engine/Plugins/Experimental/AnimNext/Source/AnimNext/Internal/Graph/RigUnit_AnimNextBeginExecution.h
+++ b/Engine/Plugins/Experimental/AnimNext/Source/AnimNext/Internal/Graph/RigUnit_AnimNextBeginExecution.h
@@ -2,7 +2,7 @@
 
 #pragma once
 
-#include "Units/RigUnit.h"
+#include "RigUnit_AnimNextBase.h"
 #include "AnimNextExecuteContext.h"
 #include "RigUnit_AnimNextBeginExecution.generated.h"
 
diff --git a/Engine/Plugins/Experimental/AnimNext/Source/AnimNext/Internal/Graph/RigUnit_AnimNextGraphEvaluator.h b/Engine/Plugins/Experimental/AnimNext/Source/AnimNext/Internal/Graph/RigUnit_AnimNextGraphEvaluator.h
--- a/Engine/Plugins/Experimental/AnimNext/Source/AnimNext/Internal/Graph/RigUnit_AnimNextGraphEvaluator.h
+++ b/Engine/Plugins/Experimental/AnimNext/Source/AnimNext/Internal/Graph/RigUnit_AnimNextGraphEvaluator.h
@@ -2,8 +2,8 @@
 
 #pragma once
 
+#include "RigUnit_AnimNextBase.h"
 #include "AnimNextExecuteContext.h"
-
 #include "RigUnit_AnimNextGraphEvaluator.generated.h"
 
 // Represents an argument to a FRigUnit_AnimNextGraphEvaluator entry point
diff --git a/Engine/Plugins/Experimental/AnimNext/Source/AnimNext/Internal/Graph/RigUnit_AnimNextGraphRoot.h b/Engine/Plugins/Experimental/AnimNext/Source/AnimNext/Internal/Graph/RigUnit_AnimNextGraphRoot.h
--- a/Engine/Plugins/Experimental/AnimNext/Source/AnimNext/Internal/Graph/RigUnit_AnimNextGraphRoot.h
+++ b/Engine/Plugins/Experimental/AnimNext/Source/AnimNext/Internal/Graph/RigUnit_AnimNextGraphRoot.h
@@ -3,8 +3,8 @@
 #pragma once
 
 #include "AnimNextExecuteContext.h"
+#include "RigUnit_AnimNextBase.h"
 #include "TraitCore/TraitHandle.h"
-
 #include "RigUnit_AnimNextGraphRoot.generated.h"
 
 /**
diff --git a/Engine/Plugins/Experimental/AnimNext/Source/AnimNext/Internal/Graph/RigUnit_AnimNextGraphSelf.h b/Engine/Plugins/Experimental/AnimNext/Source/AnimNext/Internal/Graph/RigUnit_AnimNextGraphSelf.h
--- a/Engine/Plugins/Experimental/AnimNext/Source/AnimNext/Internal/Graph/RigUnit_AnimNextGraphSelf.h
+++ b/Engine/Plugins/Experimental/AnimNext/Source/AnimNext/Internal/Graph/RigUnit_AnimNextGraphSelf.h
@@ -2,8 +2,7 @@
 
 #pragma once
 
-#include "AnimNextExecuteContext.h"
-
+#include "RigUnit_AnimNextBase.h"
 #include "RigUnit_AnimNextGraphSelf.generated.h"
 
 class UAnimNextGraph;
diff --git a/Engine/Plugins/Experimental/AnimNext/Source/AnimNext/Internal/Graph/RigUnit_AnimNextShimRoot.h b/Engine/Plugins/Experimental/AnimNext/Source/AnimNext/Internal/Graph/RigUnit_AnimNextShimRoot.h
--- a/Engine/Plugins/Experimental/AnimNext/Source/AnimNext/Internal/Graph/RigUnit_AnimNextShimRoot.h
+++ b/Engine/Plugins/Experimental/AnimNext/Source/AnimNext/Internal/Graph/RigUnit_AnimNextShimRoot.h
@@ -2,8 +2,8 @@
 
 #pragma once
 
+#include "RigUnit_AnimNextBase.h"
 #include "AnimNextExecuteContext.h"
-
 #include "RigUnit_AnimNextShimRoot.generated.h"
 
 /**
diff --git a/Engine/Plugins/Experimental/AnimNext/Source/AnimNext/Internal/Graph/RigUnit_AnimNextTraitStack.h b/Engine/Plugins/Experimental/AnimNext/Source/AnimNext/Internal/Graph/RigUnit_AnimNextTraitStack.h
--- a/Engine/Plugins/Experimental/AnimNext/Source/AnimNext/Internal/Graph/RigUnit_AnimNextTraitStack.h
+++ b/Engine/Plugins/Experimental/AnimNext/Source/AnimNext/Internal/Graph/RigUnit_AnimNextTraitStack.h
@@ -3,6 +3,7 @@
 #pragma once
 
 #include "AnimNextExecuteContext.h"
+#include "RigUnit_AnimNextBase.h"
 #include "TraitCore/TraitHandle.h"
 
 #include "RigUnit_AnimNextTraitStack.generated.h"
diff --git a/Engine/Plugins/Experimental/AnimNext/Source/AnimNext/Internal/Param/AnimNextParamContextData.h b/Engine/Plugins/Experimental/AnimNext/Source/AnimNext/Internal/Param/AnimNextParamContextData.h
new file mode 100644
--- /dev/null
+++ b/Engine/Plugins/Experimental/AnimNext/Source/AnimNext/Internal/Param/AnimNextParamContextData.h
@@ -0,0 +1,41 @@
+// Copyright Epic Games, Inc. All Rights Reserved.
+
+#pragma once
+
+#include "CoreMinimal.h"
+#include "AnimNextParamContextData.generated.h"
+
+namespace UE::AnimNext
+{
+	struct FParamStackLayerHandle;
+}
+
+USTRUCT()
+struct FAnimNextParamContextData
+{
+	GENERATED_BODY()
+
+	FAnimNextParamContextData() = default;
+
+	FAnimNextParamContextData(UE::AnimNext::FParamStackLayerHandle& InLayerHandle)
+		: LayerHandle(&InLayerHandle)
+	{}
+
+	UE::AnimNext::FParamStackLayerHandle& GetLayerHandle() const
+	{
+		return *LayerHandle;
+	}
+
+private:
+	// Call this to reset the context to its original state to detect stale usage
+	void Reset()
+	{
+		LayerHandle = nullptr;
+	}
+
+	// Parameter scope context
+	UE::AnimNext::FParamStackLayerHandle* LayerHandle = nullptr;
+
+	friend class UAnimNextParameterBlock;
+	friend struct FAnimNextExecuteContext;
+};
diff --git a/Engine/Plugins/Experimental/AnimNext/Source/AnimNext/Internal/Param/AnimNextParameterExecuteContext.h b/Engine/Plugins/Experimental/AnimNext/Source/AnimNext/Internal/Param/AnimNextParameterExecuteContext.h
deleted file mode 100644
diff --git a/Engine/Plugins/Experimental/AnimNext/Source/AnimNext/Internal/Param/RigUnit_AnimNextParameterBase.h b/Engine/Plugins/Experimental/AnimNext/Source/AnimNext/Internal/Param/RigUnit_AnimNextParameterBase.h
new file mode 100644
--- /dev/null
+++ b/Engine/Plugins/Experimental/AnimNext/Source/AnimNext/Internal/Param/RigUnit_AnimNextParameterBase.h
@@ -0,0 +1,13 @@
+﻿// Copyright Epic Games, Inc. All Rights Reserved.
+
+#pragma once
+
+#include "CoreMinimal.h"
+#include "Units/RigUnit.h"
+#include "RigUnit_AnimNextParameterBase.generated.h"
+
+USTRUCT(meta = (ExecuteContext = "FAnimNextExecuteContext"))
+struct FRigUnit_AnimNextParameterBase : public FRigUnit
+{
+	GENERATED_BODY()
+};
diff --git a/Engine/Plugins/Experimental/AnimNext/Source/AnimNext/Internal/Param/RigUnit_AnimNextParameterBeginExecution.h b/Engine/Plugins/Experimental/AnimNext/Source/AnimNext/Internal/Param/RigUnit_AnimNextParameterBeginExecution.h
--- a/Engine/Plugins/Experimental/AnimNext/Source/AnimNext/Internal/Param/RigUnit_AnimNextParameterBeginExecution.h
+++ b/Engine/Plugins/Experimental/AnimNext/Source/AnimNext/Internal/Param/RigUnit_AnimNextParameterBeginExecution.h
@@ -2,7 +2,8 @@
 
 #pragma once
 
-#include "AnimNextParameterExecuteContext.h"
+#include "RigUnit_AnimNextParameterBase.h"
+#include "AnimNextExecuteContext.h"
 #include "RigUnit_AnimNextParameterBeginExecution.generated.h"
 
 /**
@@ -22,7 +23,7 @@ struct ANIMNEXT_API FRigUnit_AnimNextParameterBeginExecution : public FRigUnit_A
 
 	// The execution result
 	UPROPERTY(EditAnywhere, DisplayName = "Execute", Category = "Entry Point", meta = (Output))
-	FAnimNextParameterExecuteContext ExecuteContext;
+	FAnimNextExecuteContext ExecuteContext;
 
 	// The name of the entry point
 	UPROPERTY(VisibleAnywhere, Category = "Entry Point", meta = (Hidden))
diff --git a/Engine/Plugins/Experimental/AnimNext/Source/AnimNext/Internal/Param/RigVMDispatch_GetLayerParameter.h b/Engine/Plugins/Experimental/AnimNext/Source/AnimNext/Internal/Param/RigVMDispatch_GetLayerParameter.h
--- a/Engine/Plugins/Experimental/AnimNext/Source/AnimNext/Internal/Param/RigVMDispatch_GetLayerParameter.h
+++ b/Engine/Plugins/Experimental/AnimNext/Source/AnimNext/Internal/Param/RigVMDispatch_GetLayerParameter.h
@@ -3,7 +3,7 @@
 #pragma once
 
 #include "RigVMCore/RigVMDispatchFactory.h"
-#include "AnimNextParameterExecuteContext.h"
+#include "AnimNextExecuteContext.h"
 #include "RigVMDispatch_GetLayerParameter.generated.h"
 
 namespace UE::AnimNext::UncookedOnly
@@ -24,7 +24,7 @@ struct ANIMNEXT_API FRigVMDispatch_GetLayerParameter : public FRigVMDispatchFact
 private:
 	friend struct UE::AnimNext::UncookedOnly::FUtils;
 
-	virtual UScriptStruct* GetExecuteContextStruct() const override { return FAnimNextParameterExecuteContext::StaticStruct(); }
+	virtual UScriptStruct* GetExecuteContextStruct() const override { return FAnimNextExecuteContext::StaticStruct(); }
 	virtual FName GetArgumentNameForOperandIndex(int32 InOperandIndex, int32 InTotalOperands) const override;
 	virtual const TArray<FRigVMTemplateArgumentInfo>& GetArgumentInfos() const override;
 #if WITH_EDITOR
diff --git a/Engine/Plugins/Experimental/AnimNext/Source/AnimNext/Internal/Param/RigVMDispatch_GetParameter.h b/Engine/Plugins/Experimental/AnimNext/Source/AnimNext/Internal/Param/RigVMDispatch_GetParameter.h
--- a/Engine/Plugins/Experimental/AnimNext/Source/AnimNext/Internal/Param/RigVMDispatch_GetParameter.h
+++ b/Engine/Plugins/Experimental/AnimNext/Source/AnimNext/Internal/Param/RigVMDispatch_GetParameter.h
@@ -3,7 +3,7 @@
 #pragma once
 
 #include "RigVMCore/RigVMDispatchFactory.h"
-#include "Graph/AnimNextExecuteContext.h"
+#include "AnimNextExecuteContext.h"
 #include "RigVMDispatch_GetParameter.generated.h"
 
 namespace UE::AnimNext::UncookedOnly
diff --git a/Engine/Plugins/Experimental/AnimNext/Source/AnimNext/Internal/Param/RigVMDispatch_SetLayerParameter.h b/Engine/Plugins/Experimental/AnimNext/Source/AnimNext/Internal/Param/RigVMDispatch_SetLayerParameter.h
--- a/Engine/Plugins/Experimental/AnimNext/Source/AnimNext/Internal/Param/RigVMDispatch_SetLayerParameter.h
+++ b/Engine/Plugins/Experimental/AnimNext/Source/AnimNext/Internal/Param/RigVMDispatch_SetLayerParameter.h
@@ -3,7 +3,7 @@
 #pragma once
 
 #include "RigVMCore/RigVMDispatchFactory.h"
-#include "AnimNextParameterExecuteContext.h"
+#include "AnimNextExecuteContext.h"
 #include "RigVMDispatch_SetLayerParameter.generated.h"
 
 namespace UE::AnimNext::UncookedOnly
@@ -24,7 +24,7 @@ struct ANIMNEXT_API FRigVMDispatch_SetLayerParameter : public FRigVMDispatchFact
 private:
 	friend struct UE::AnimNext::UncookedOnly::FUtils;
 
-	virtual UScriptStruct* GetExecuteContextStruct() const { return FAnimNextParameterExecuteContext::StaticStruct(); }
+	virtual UScriptStruct* GetExecuteContextStruct() const { return FAnimNextExecuteContext::StaticStruct(); }
 	virtual FName GetArgumentNameForOperandIndex(int32 InOperandIndex, int32 InTotalOperands) const override;
 	virtual const TArray<FRigVMTemplateArgumentInfo>& GetArgumentInfos() const override;
 #if WITH_EDITOR
diff --git a/Engine/Plugins/Experimental/AnimNext/Source/AnimNext/Private/Graph/AnimGraph/AnimNode_AnimNextGraph.cpp b/Engine/Plugins/Experimental/AnimNext/Source/AnimNext/Private/Graph/AnimGraph/AnimNode_AnimNextGraph.cpp
--- a/Engine/Plugins/Experimental/AnimNext/Source/AnimNext/Private/Graph/AnimGraph/AnimNode_AnimNextGraph.cpp
+++ b/Engine/Plugins/Experimental/AnimNext/Source/AnimNext/Private/Graph/AnimGraph/AnimNode_AnimNextGraph.cpp
@@ -7,13 +7,9 @@
 #include "Animation/AnimInstanceProxy.h"
 #include "GameFramework/Actor.h"
 #include "AnimationRuntime.h"
-#include "DataRegistry.h"
 #include "GenerationTools.h"
-#include "ReferencePose.h"
 #include "Graph/AnimNext_LODPose.h"
-#include "Graph/AnimNextExecuteContext.h"
 #include "Engine/SkeletalMesh.h"
-#include "BoneContainer.h"
 #include "Param/ParamStack.h"
 #include "AnimGraphParamStackScope.h"
 #include "Scheduler/ScheduleContext.h"
diff --git a/Engine/Plugins/Experimental/AnimNext/Source/AnimNext/Private/Graph/AnimNextGraphInstance.cpp b/Engine/Plugins/Experimental/AnimNext/Source/AnimNext/Private/Graph/AnimNextGraphInstance.cpp
--- a/Engine/Plugins/Experimental/AnimNext/Source/AnimNext/Private/Graph/AnimNextGraphInstance.cpp
+++ b/Engine/Plugins/Experimental/AnimNext/Source/AnimNext/Private/Graph/AnimNextGraphInstance.cpp
@@ -129,12 +129,12 @@ void FAnimNextGraphInstance::ExecuteLatentPins(const TConstArrayView<UE::AnimNex
 	if (URigVM* VM = Graph->VM)
 	{
 		FAnimNextExecuteContext& AnimNextContext = ExtendedExecuteContext.GetPublicDataSafe<FAnimNextExecuteContext>();
-		AnimNextContext.SetupForExecution(this, LatentHandles, DestinationBasePtr, bIsFrozen);
+		AnimNextContext.SetContextData<FAnimNextGraphContextData>(this, LatentHandles, DestinationBasePtr, bIsFrozen);
 
 		VM->ExecuteVM(ExtendedExecuteContext, FRigUnit_AnimNextShimRoot::EventName);
 
 		// Reset the context to avoid issues if we forget to reset it the next time we use it
-		AnimNextContext.DebugReset();
+		AnimNextContext.DebugReset<FAnimNextGraphContextData>();
 	}
 }
 
diff --git a/Engine/Plugins/Experimental/AnimNext/Source/AnimNext/Private/Graph/RigUnit_AnimNextGraphEvaluator.cpp b/Engine/Plugins/Experimental/AnimNext/Source/AnimNext/Private/Graph/RigUnit_AnimNextGraphEvaluator.cpp
--- a/Engine/Plugins/Experimental/AnimNext/Source/AnimNext/Private/Graph/RigUnit_AnimNextGraphEvaluator.cpp
+++ b/Engine/Plugins/Experimental/AnimNext/Source/AnimNext/Private/Graph/RigUnit_AnimNextGraphEvaluator.cpp
@@ -27,10 +27,11 @@ namespace UE::AnimNext::Private
 void FRigUnit_AnimNextGraphEvaluator::StaticExecute(FRigVMExtendedExecuteContext& RigVMExecuteContext, FRigVMMemoryHandleArray RigVMMemoryHandles, FRigVMPredicateBranchArray RigVMBranches)
 {
 	const FAnimNextExecuteContext& VMExecuteContext = RigVMExecuteContext.GetPublicData<FAnimNextExecuteContext>();
+	const FAnimNextGraphContextData& GraphContextData = VMExecuteContext.GetContextData<FAnimNextGraphContextData>();
 
-	const TConstArrayView<UE::AnimNext::FLatentPropertyHandle>& LatentHandles = VMExecuteContext.GetLatentHandles();
-	uint8* DestinationBasePtr = (uint8*)VMExecuteContext.GetDestinationBasePtr();
-	const bool bIsFrozen = VMExecuteContext.IsFrozen();
+	const TConstArrayView<UE::AnimNext::FLatentPropertyHandle>& LatentHandles = GraphContextData.GetLatentHandles();
+	uint8* DestinationBasePtr = (uint8*)GraphContextData.GetDestinationBasePtr();
+	const bool bIsFrozen = GraphContextData.IsFrozen();
 
 	for (UE::AnimNext::FLatentPropertyHandle Handle : LatentHandles)
 	{
diff --git a/Engine/Plugins/Experimental/AnimNext/Source/AnimNext/Private/Graph/RigUnit_AnimNextGraphSelf.cpp b/Engine/Plugins/Experimental/AnimNext/Source/AnimNext/Private/Graph/RigUnit_AnimNextGraphSelf.cpp
--- a/Engine/Plugins/Experimental/AnimNext/Source/AnimNext/Private/Graph/RigUnit_AnimNextGraphSelf.cpp
+++ b/Engine/Plugins/Experimental/AnimNext/Source/AnimNext/Private/Graph/RigUnit_AnimNextGraphSelf.cpp
@@ -5,5 +5,5 @@
 
 FRigUnit_AnimNextGraphSelf_Execute()
 {
-	Self = const_cast<UAnimNextGraph*>(ExecuteContext.GetGraphInstance().GetGraph());
+	Self = const_cast<UAnimNextGraph*>(ExecuteContext.GetContextData<FAnimNextGraphContextData>().GetGraphInstance().GetGraph());
 }
diff --git a/Engine/Plugins/Experimental/AnimNext/Source/AnimNext/Private/Param/AnimNextParameterBlock.cpp b/Engine/Plugins/Experimental/AnimNext/Source/AnimNext/Private/Param/AnimNextParameterBlock.cpp
--- a/Engine/Plugins/Experimental/AnimNext/Source/AnimNext/Private/Param/AnimNextParameterBlock.cpp
+++ b/Engine/Plugins/Experimental/AnimNext/Source/AnimNext/Private/Param/AnimNextParameterBlock.cpp
@@ -5,9 +5,8 @@
 #include "RigVMRuntimeDataRegistry.h"
 #include "UObject/AssetRegistryTagsContext.h"
 #include "UObject/ObjectSaveContext.h"
-#include "UObject/Package.h"
 #include "Graph/RigUnit_AnimNextBeginExecution.h"
-#include "Param/AnimNextParameterExecuteContext.h"
+#include "AnimNextExecuteContext.h"
 #include "AnimNextStats.h"
 
 DEFINE_STAT(STAT_AnimNext_ParamBlock_UpdateLayer);
@@ -17,7 +16,7 @@ DEFINE_STAT(STAT_AnimNext_ParamBlock_UpdateLayer);
 UAnimNextParameterBlock::UAnimNextParameterBlock(const FObjectInitializer& ObjectInitializer)
 	: Super(ObjectInitializer)
 {
-	ExtendedExecuteContext.SetContextPublicDataStruct(FAnimNextParameterExecuteContext::StaticStruct());
+	ExtendedExecuteContext.SetContextPublicDataStruct(FAnimNextExecuteContext::StaticStruct());
 }
 
 void UAnimNextParameterBlock::UpdateLayer(UE::AnimNext::FParamStackLayerHandle& InHandle, float InDeltaTime) const
@@ -32,15 +31,18 @@ void UAnimNextParameterBlock::UpdateLayer(UE::AnimNext::FParamStackLayerHandle&
 
 			check(Context.VMHash == VM->GetVMHash());
 
-			FAnimNextParameterExecuteContext& AnimNextParameterContext = Context.GetPublicDataSafe<FAnimNextParameterExecuteContext>();
+			FAnimNextExecuteContext& AnimNextContext = Context.GetPublicDataSafe<FAnimNextExecuteContext>();
 
 			// Param block setup
-			AnimNextParameterContext.SetParamContextData(InHandle);
+			AnimNextContext.SetContextData<FAnimNextParamContextData>(InHandle);
 
 			// RigVM setup
-			AnimNextParameterContext.SetDeltaTime(InDeltaTime);
+			AnimNextContext.SetDeltaTime(InDeltaTime);
 
 			VM->ExecuteVM(Context, FRigUnit_AnimNextBeginExecution::EventName);
+
+			// Reset the context to avoid issues if we forget to reset it the next time we use it
+			AnimNextContext.DebugReset<FAnimNextParamContextData>();
 		}
 	}
 }
diff --git a/Engine/Plugins/Experimental/AnimNext/Source/AnimNext/Private/Param/RigVMDispatch_GetLayerParameter.cpp b/Engine/Plugins/Experimental/AnimNext/Source/AnimNext/Private/Param/RigVMDispatch_GetLayerParameter.cpp
--- a/Engine/Plugins/Experimental/AnimNext/Source/AnimNext/Private/Param/RigVMDispatch_GetLayerParameter.cpp
+++ b/Engine/Plugins/Experimental/AnimNext/Source/AnimNext/Private/Param/RigVMDispatch_GetLayerParameter.cpp
@@ -2,8 +2,6 @@
 
 #include "Param/RigVMDispatch_GetLayerParameter.h"
 #include "RigVMCore/RigVMStruct.h"
-#include "RigVMCore/RigVM.h"
-#include "Graph/AnimNextExecuteContext.h"
 #include "Param/ParamStack.h"
 
 const FName FRigVMDispatch_GetLayerParameter::ValueName = TEXT("Value");
@@ -97,9 +95,10 @@ void FRigVMDispatch_GetLayerParameter::Execute(FRigVMExtendedExecuteContext& InC
 		TypeHandle = FParamTypeHandle::FromProperty(ValueProperty).ToRaw();
 	}
 
-	FAnimNextParameterExecuteContext& ParamContext = InContext.GetPublicData<FAnimNextParameterExecuteContext>();
+	FAnimNextExecuteContext& Context = InContext.GetPublicData<FAnimNextExecuteContext>();
+	const FAnimNextParamContextData& ParamContextData = Context.GetContextData<FAnimNextParamContextData>();
 	TConstArrayView<uint8> SourceData;
-	if (ParamContext.GetLayerHandle().GetValueRaw(FParamId(Parameter, ParameterHash), FParamTypeHandle::FromRaw(TypeHandle), SourceData).IsSuccessful())
+	if (ParamContextData.GetLayerHandle().GetValueRaw(FParamId(Parameter, ParameterHash), FParamTypeHandle::FromRaw(TypeHandle), SourceData).IsSuccessful())
 	{
 		ValueProperty->CopyCompleteValue(TargetDataPtr, SourceData.GetData());
 	}
diff --git a/Engine/Plugins/Experimental/AnimNext/Source/AnimNext/Private/Param/RigVMDispatch_GetParameter.cpp b/Engine/Plugins/Experimental/AnimNext/Source/AnimNext/Private/Param/RigVMDispatch_GetParameter.cpp
--- a/Engine/Plugins/Experimental/AnimNext/Source/AnimNext/Private/Param/RigVMDispatch_GetParameter.cpp
+++ b/Engine/Plugins/Experimental/AnimNext/Source/AnimNext/Private/Param/RigVMDispatch_GetParameter.cpp
@@ -2,11 +2,8 @@
 
 #include "Param/RigVMDispatch_GetParameter.h"
 #include "RigVMCore/RigVMStruct.h"
-#include "RigVMCore/RigVM.h"
-#include "Graph/AnimNextExecuteContext.h"
 #include "Context.h"
 #include "Param/ParamStack.h"
-#include "Param/ParamStack.h"
 
 const FName FRigVMDispatch_GetParameter::ValueName = TEXT("Value");
 const FName FRigVMDispatch_GetParameter::TypeHandleName = TEXT("Type");
diff --git a/Engine/Plugins/Experimental/AnimNext/Source/AnimNext/Private/Param/RigVMDispatch_SetLayerParameter.cpp b/Engine/Plugins/Experimental/AnimNext/Source/AnimNext/Private/Param/RigVMDispatch_SetLayerParameter.cpp
--- a/Engine/Plugins/Experimental/AnimNext/Source/AnimNext/Private/Param/RigVMDispatch_SetLayerParameter.cpp
+++ b/Engine/Plugins/Experimental/AnimNext/Source/AnimNext/Private/Param/RigVMDispatch_SetLayerParameter.cpp
@@ -3,7 +3,7 @@
 #include "Param/RigVMDispatch_SetLayerParameter.h"
 #include "RigVMCore/RigVMStruct.h"
 #include "RigVMCore/RigVM.h"
-#include "Param/AnimNextParameterExecuteContext.h"
+#include "AnimNextExecuteContext.h"
 #include "Param/ParamId.h"
 #include "Param/ParamTypeHandle.h"
 #include "Param/ParamStackLayerHandle.h"
@@ -110,9 +110,10 @@ void FRigVMDispatch_SetLayerParameter::Execute(FRigVMExtendedExecuteContext& InC
 	FParamTypeHandle TypeHandle = FParamTypeHandle::FromRaw(RawTypeHandle);
 	if(ParameterId.IsValid() && TypeHandle.IsValid())
 	{
-		FAnimNextParameterExecuteContext& ParamContext = InContext.GetPublicData<FAnimNextParameterExecuteContext>();
+		FAnimNextExecuteContext& AnimNextContext = InContext.GetPublicData<FAnimNextExecuteContext>();
+		const FAnimNextParamContextData& ParamContextData = AnimNextContext.GetContextData<FAnimNextParamContextData>();
 		TConstArrayView<uint8> SourceDataView(SourceData, ValueProperty->GetSize());
-		ParamContext.GetLayerHandle().SetValueRaw(ParameterId, TypeHandle, SourceDataView);
+		ParamContextData.GetLayerHandle().SetValueRaw(ParameterId, TypeHandle, SourceDataView);
 	}
 }
 
diff --git a/Engine/Plugins/Experimental/AnimNext/Source/AnimNext/Private/Scheduler/AnimNextScheduleGraphTask.cpp b/Engine/Plugins/Experimental/AnimNext/Source/AnimNext/Private/Scheduler/AnimNextScheduleGraphTask.cpp
--- a/Engine/Plugins/Experimental/AnimNext/Source/AnimNext/Private/Scheduler/AnimNextScheduleGraphTask.cpp
+++ b/Engine/Plugins/Experimental/AnimNext/Source/AnimNext/Private/Scheduler/AnimNextScheduleGraphTask.cpp
@@ -8,7 +8,6 @@
 #include "TraitInterfaces/IEvaluate.h"
 #include "TraitInterfaces/IUpdate.h"
 #include "EvaluationVM/EvaluationVM.h"
-#include "Graph/AnimNextExecuteContext.h"
 #include "Graph/AnimNext_LODPose.h"
 #include "AnimNextStats.h"
 #include "Logging/StructuredLog.h"
diff --git a/Engine/Plugins/Experimental/AnimNext/Source/AnimNextEditor/Private/Common/SActionMenu.cpp b/Engine/Plugins/Experimental/AnimNext/Source/AnimNextEditor/Private/Common/SActionMenu.cpp
--- a/Engine/Plugins/Experimental/AnimNext/Source/AnimNextEditor/Private/Common/SActionMenu.cpp
+++ b/Engine/Plugins/Experimental/AnimNext/Source/AnimNextEditor/Private/Common/SActionMenu.cpp
@@ -10,12 +10,11 @@
 #include "Widgets/Images/SImage.h"
 #include "Widgets/Layout/SBox.h"
 #include "Widgets/SBoxPanel.h"
-#include "Widgets/Text/STextBlock.h"
 #include "SGraphPalette.h"
 #include "RigVMCore/RigVMRegistry.h"
+#include "RigVMModel/RigVMClient.h"
+#include "RigVMModel/RigVMSchema.h"
 #include "Units/RigUnit.h"
-#include "Widgets/SToolTip.h"
-#include "Graph/AnimNextExecuteContext.h"
 
 #define LOCTEXT_NAMESPACE "AnimNextEditor"
 
@@ -24,10 +23,12 @@ namespace UE::AnimNext::Editor
 
 void SActionMenu::CollectAllAnimNextGraphActions(FGraphContextMenuBuilder& MenuBuilder) const
 {
+	// Disable reporting as the schema's SupportsX functions will output errors
+	RigVMController->EnableReporting(false);
+
 	for(const FRigVMFunction& Function : FRigVMRegistry::Get().GetFunctions())
 	{
-		const UScriptStruct* FunctionContext = Function.GetExecuteContextStruct();
-		if (FunctionContext == nullptr || !AllowedExecuteContexts.Contains(FunctionContext))
+		if (RigVMSchema == nullptr || !RigVMSchema->SupportsUnitFunction(RigVMController, &Function))
 		{
 			continue;
 		}
@@ -60,18 +61,7 @@ void SActionMenu::CollectAllAnimNextGraphActions(FGraphContextMenuBuilder& MenuB
 
 	for (const FRigVMDispatchFactory* Factory : FRigVMRegistry::Get().GetFactories())
 	{
-		// See if the factory allows our execute contexts
-		bool bAllowed = true;
-		for(const UScriptStruct* AllowedExecuteContext : AllowedExecuteContexts)
-		{
-			if (!Factory->SupportsExecuteContextStruct(AllowedExecuteContext))
-			{
-				bAllowed = false;
-				break;
-			}
-		}
-
-		if(bAllowed)
+		if (RigVMSchema == nullptr || !RigVMSchema->SupportsDispatchFactory(RigVMController, Factory))
 		{
 			continue;
 		}
@@ -88,6 +78,8 @@ void SActionMenu::CollectAllAnimNextGraphActions(FGraphContextMenuBuilder& MenuB
 
 		MenuBuilder.AddAction(MakeShared<FAnimNextSchemaAction_DispatchFactory>(Template->GetNotation(), NodeCategory, MenuDesc, ToolTip));
 	};
+
+	RigVMController->EnableReporting(true);
 }
 
 SActionMenu::~SActionMenu()
@@ -96,15 +88,23 @@ SActionMenu::~SActionMenu()
 	OnCloseReasonCallback.ExecuteIfBound(bActionExecuted, false, !DraggedFromPins.IsEmpty());
 }
 
-void SActionMenu::Construct(const FArguments& InArgs)
+void SActionMenu::Construct(const FArguments& InArgs, UEdGraph* InGraph)
 {
-	Graph = InArgs._Graph;
+	check(InGraph);
+
+	Graph = InGraph;
 	DraggedFromPins = InArgs._DraggedFromPins;
 	NewNodePosition = InArgs._NewNodePosition;
 	OnClosedCallback = InArgs._OnClosedCallback;
 	bAutoExpandActionMenu = InArgs._AutoExpandActionMenu;
 	OnCloseReasonCallback = InArgs._OnCloseReason;
-	AllowedExecuteContexts = InArgs._AllowedExecuteContexts;
+
+	RigVMClientHost = Graph->GetImplementingOuter<IRigVMClientHost>();
+	check(RigVMClientHost);
+	RigVMController = RigVMClientHost->GetRigVMClient()->GetController(Graph);
+	check(RigVMController);
+	RigVMSchema = RigVMController->GetGraph()->GetSchema();
+	check(RigVMSchema);
 
 	SBorder::Construct(SBorder::FArguments()
 		.BorderImage(FAppStyle::Get().GetBrush("Menu.Background"))
diff --git a/Engine/Plugins/Experimental/AnimNext/Source/AnimNextEditor/Private/Common/SActionMenu.h b/Engine/Plugins/Experimental/AnimNext/Source/AnimNextEditor/Private/Common/SActionMenu.h
--- a/Engine/Plugins/Experimental/AnimNext/Source/AnimNextEditor/Private/Common/SActionMenu.h
+++ b/Engine/Plugins/Experimental/AnimNext/Source/AnimNextEditor/Private/Common/SActionMenu.h
@@ -5,18 +5,19 @@
 #include "EdGraph/EdGraph.h"
 #include "EdGraph/EdGraphPin.h"
 #include "EdGraph/EdGraphSchema.h"
-#include "Layout/Margin.h"
 #include "Layout/Visibility.h"
 #include "SGraphActionMenu.h"
 #include "Widgets/DeclarativeSyntaxSupport.h"
 #include "Widgets/Input/SEditableTextBox.h"
 #include "Widgets/Layout/SBorder.h"
 #include "Widgets/SWidget.h"
-#include "Widgets/Views/SExpanderArrow.h"
 
+class URigVMSchema;
 class SEditableTextBox;
 class SGraphActionMenu;
 class UEdGraph;
+class IRigVMClientHost;
+class URigVMController;
 
 namespace UE::AnimNext::Editor
 {
@@ -28,22 +29,19 @@ public:
 	DECLARE_DELEGATE_ThreeParams(FClosedReason, bool /*bActionExecuted*/, bool /*bContextSensitiveChecked*/, bool /*bGraphPinContext*/);
 
 	SLATE_BEGIN_ARGS(SActionMenu)
-		: _Graph(static_cast<UEdGraph*>(nullptr))
-		, _NewNodePosition(FVector2D::ZeroVector)
+		: _NewNodePosition(FVector2D::ZeroVector)
 		, _AutoExpandActionMenu(false)
-	{
-	}
-		SLATE_ARGUMENT(UEdGraph*, Graph)
-		SLATE_ARGUMENT(FVector2D, NewNodePosition)
-		SLATE_ARGUMENT(TArray<UEdGraphPin*>, DraggedFromPins)
-		SLATE_ARGUMENT(SGraphEditor::FActionMenuClosed, OnClosedCallback)
-		SLATE_ARGUMENT(bool, AutoExpandActionMenu)
-		SLATE_EVENT(FClosedReason, OnCloseReason)
-		SLATE_ARGUMENT(TArray<UScriptStruct*>, AllowedExecuteContexts)
+	{}
+	
+	SLATE_ARGUMENT(FVector2D, NewNodePosition)
+	SLATE_ARGUMENT(TArray<UEdGraphPin*>, DraggedFromPins)
+	SLATE_ARGUMENT(SGraphEditor::FActionMenuClosed, OnClosedCallback)
+	SLATE_ARGUMENT(bool, AutoExpandActionMenu)
+	SLATE_EVENT(FClosedReason, OnCloseReason)
 
 	SLATE_END_ARGS()
 
-	void Construct(const FArguments& InArgs);
+	void Construct(const FArguments& InArgs, UEdGraph* InGraph);
 
 	~SActionMenu();
 
@@ -61,6 +59,9 @@ protected:
 	
 private:
 	UEdGraph* Graph = nullptr;
+	const URigVMSchema* RigVMSchema = nullptr;
+	IRigVMClientHost* RigVMClientHost = nullptr;
+	URigVMController* RigVMController = nullptr;
 	bool bAutoExpandActionMenu = false;
 	bool bActionExecuted = false;
 
@@ -71,7 +72,6 @@ private:
 	FClosedReason OnCloseReasonCallback;
 
 	TSharedPtr<SGraphActionMenu> GraphActionMenu;
-	TArray<UScriptStruct*> AllowedExecuteContexts;
 };
 
 }
diff --git a/Engine/Plugins/Experimental/AnimNext/Source/AnimNextEditor/Private/Graph/AnimNextGraphDocumentSummoner.cpp b/Engine/Plugins/Experimental/AnimNext/Source/AnimNextEditor/Private/Graph/AnimNextGraphDocumentSummoner.cpp
--- a/Engine/Plugins/Experimental/AnimNext/Source/AnimNextEditor/Private/Graph/AnimNextGraphDocumentSummoner.cpp
+++ b/Engine/Plugins/Experimental/AnimNext/Source/AnimNextEditor/Private/Graph/AnimNextGraphDocumentSummoner.cpp
@@ -2,7 +2,6 @@
 
 #include "AnimNextGraphDocumentSummoner.h"
 #include "Common/SActionMenu.h"
-#include "Graph/AnimNextExecuteContext.h"
 #include "Graph/AnimNextGraph_EdGraph.h"
 #include "Workspace/AnimNextWorkspaceEditor.h"
 
@@ -16,13 +15,11 @@ FAnimNextGraphDocumentSummoner::FAnimNextGraphDocumentSummoner(FName InIdentifie
 
 FActionMenuContent FAnimNextGraphDocumentSummoner::OnCreateGraphActionMenu(UEdGraph* InGraph, const FVector2D& InNodePosition, const TArray<UEdGraphPin*>& InDraggedPins, bool bAutoExpand, SGraphEditor::FActionMenuClosed InOnMenuClosed) const
 {
-	TSharedRef<SActionMenu> ActionMenu = SNew(SActionMenu)
+	TSharedRef<SActionMenu> ActionMenu = SNew(SActionMenu, InGraph)
 		.AutoExpandActionMenu(bAutoExpand)
-		.Graph(InGraph)
 		.NewNodePosition(InNodePosition)
 		.DraggedFromPins(InDraggedPins)
-		.OnClosedCallback(InOnMenuClosed)
-		.AllowedExecuteContexts( { FRigVMExecuteContext::StaticStruct(), FAnimNextExecuteContext::StaticStruct() });
+		.OnClosedCallback(InOnMenuClosed);
 
 	TSharedPtr<SWidget> FilterTextBox = StaticCastSharedRef<SWidget>(ActionMenu->GetFilterTextBox());
 	return FActionMenuContent(StaticCastSharedRef<SWidget>(ActionMenu), FilterTextBox);
diff --git a/Engine/Plugins/Experimental/AnimNext/Source/AnimNextEditor/Private/Graph/AnimNextGraphEditor.cpp b/Engine/Plugins/Experimental/AnimNext/Source/AnimNextEditor/Private/Graph/AnimNextGraphEditor.cpp
--- a/Engine/Plugins/Experimental/AnimNext/Source/AnimNextEditor/Private/Graph/AnimNextGraphEditor.cpp
+++ b/Engine/Plugins/Experimental/AnimNext/Source/AnimNextEditor/Private/Graph/AnimNextGraphEditor.cpp
@@ -171,13 +171,11 @@ TSharedRef<SGraphEditor> FGraphEditor::CreateGraphEditorWidget(TSharedRef<FTabIn
 
 FActionMenuContent FGraphEditor::OnCreateGraphActionMenu(UEdGraph* InGraph, const FVector2D& InNodePosition, const TArray<UEdGraphPin*>& InDraggedPins, bool bAutoExpand, SGraphEditor::FActionMenuClosed InOnMenuClosed)
 {
-	TSharedRef<SActionMenu> ActionMenu = SNew(SActionMenu)
+	TSharedRef<SActionMenu> ActionMenu = SNew(SActionMenu, InGraph)
 		.AutoExpandActionMenu(bAutoExpand)
-		.Graph(InGraph)
 		.NewNodePosition(InNodePosition)
 		.DraggedFromPins(InDraggedPins)
-		.OnClosedCallback(InOnMenuClosed)
-		.AllowedExecuteContexts( { FRigVMExecuteContext::StaticStruct(), FAnimNextExecuteContext::StaticStruct() });
+		.OnClosedCallback(InOnMenuClosed);
 
 	TSharedPtr<SWidget> FilterTextBox = StaticCastSharedRef<SWidget>(ActionMenu->GetFilterTextBox());
 	return FActionMenuContent(StaticCastSharedRef<SWidget>(ActionMenu), FilterTextBox);
diff --git a/Engine/Plugins/Experimental/AnimNext/Source/AnimNextEditor/Private/Param/AnimNextParameterBlockEditor.cpp b/Engine/Plugins/Experimental/AnimNext/Source/AnimNextEditor/Private/Param/AnimNextParameterBlockEditor.cpp
--- a/Engine/Plugins/Experimental/AnimNext/Source/AnimNextEditor/Private/Param/AnimNextParameterBlockEditor.cpp
+++ b/Engine/Plugins/Experimental/AnimNext/Source/AnimNextEditor/Private/Param/AnimNextParameterBlockEditor.cpp
@@ -14,7 +14,6 @@
 #include "Kismet2/BlueprintEditorUtils.h"
 #include "ScopedTransaction.h"
 #include "Param/AnimNextParameterSettings.h"
-#include "Param/AnimNextParameterExecuteContext.h"
 #include "Widgets/Docking/SDockTab.h"
 
 #define LOCTEXT_NAMESPACE "AnimNextParameterBlockEditor"
@@ -224,13 +223,11 @@ TSharedRef<SGraphEditor> FParameterBlockEditor::CreateGraphEditorWidget(TSharedR
 
 FActionMenuContent FParameterBlockEditor::OnCreateGraphActionMenu(UEdGraph* InGraph, const FVector2D& InNodePosition, const TArray<UEdGraphPin*>& InDraggedPins, bool bAutoExpand, SGraphEditor::FActionMenuClosed InOnMenuClosed)
 {
-	TSharedRef<SActionMenu> ActionMenu = SNew(SActionMenu)
+	TSharedRef<SActionMenu> ActionMenu = SNew(SActionMenu, InGraph)
 		.AutoExpandActionMenu(bAutoExpand)
-		.Graph(InGraph)
 		.NewNodePosition(InNodePosition)
 		.DraggedFromPins(InDraggedPins)
-		.OnClosedCallback(InOnMenuClosed)
-		.AllowedExecuteContexts( { FRigVMExecuteContext::StaticStruct(), FAnimNextParameterExecuteContext::StaticStruct() });
+		.OnClosedCallback(InOnMenuClosed);
 
 	TSharedPtr<SWidget> FilterTextBox = StaticCastSharedRef<SWidget>(ActionMenu->GetFilterTextBox());
 	return FActionMenuContent(StaticCastSharedRef<SWidget>(ActionMenu), FilterTextBox);
diff --git a/Engine/Plugins/Experimental/AnimNext/Source/AnimNextEditor/Private/Param/ParameterBlockFactory.cpp b/Engine/Plugins/Experimental/AnimNext/Source/AnimNextEditor/Private/Param/ParameterBlockFactory.cpp
--- a/Engine/Plugins/Experimental/AnimNext/Source/AnimNextEditor/Private/Param/ParameterBlockFactory.cpp
+++ b/Engine/Plugins/Experimental/AnimNext/Source/AnimNextEditor/Private/Param/ParameterBlockFactory.cpp
@@ -3,7 +3,7 @@
 #include "Param/ParameterBlockFactory.h"
 #include "Param/AnimNextParameterBlock.h"
 #include "Param/AnimNextParameterBlock_EditorData.h"
-#include "Param/AnimNextParameterExecuteContext.h"
+#include "AnimNextExecuteContext.h"
 #include "UncookedOnlyUtils.h"
 
 UAnimNextParameterBlockFactory::UAnimNextParameterBlockFactory()
diff --git a/Engine/Plugins/Experimental/AnimNext/Source/AnimNextEditor/Private/Param/ParameterBlockGraphDocumentSummoner.cpp b/Engine/Plugins/Experimental/AnimNext/Source/AnimNextEditor/Private/Param/ParameterBlockGraphDocumentSummoner.cpp
--- a/Engine/Plugins/Experimental/AnimNext/Source/AnimNextEditor/Private/Param/ParameterBlockGraphDocumentSummoner.cpp
+++ b/Engine/Plugins/Experimental/AnimNext/Source/AnimNextEditor/Private/Param/ParameterBlockGraphDocumentSummoner.cpp
@@ -4,7 +4,7 @@
 #include "Common/SActionMenu.h"
 #include "Graph/AnimNextGraph.h"
 #include "Param/AnimNextParameterBlock_EdGraph.h"
-#include "Param/AnimNextParameterExecuteContext.h"
+
 #include "Workspace/AnimNextWorkspaceEditor.h"
 
 namespace UE::AnimNext::Editor
@@ -17,13 +17,11 @@ FParameterBlockGraphDocumentSummoner::FParameterBlockGraphDocumentSummoner(FName
 
 FActionMenuContent FParameterBlockGraphDocumentSummoner::OnCreateGraphActionMenu(UEdGraph* InGraph, const FVector2D& InNodePosition, const TArray<UEdGraphPin*>& InDraggedPins, bool bAutoExpand, SGraphEditor::FActionMenuClosed InOnMenuClosed) const
 {
-	TSharedRef<SActionMenu> ActionMenu = SNew(SActionMenu)
+	TSharedRef<SActionMenu> ActionMenu = SNew(SActionMenu, InGraph)
 		.AutoExpandActionMenu(bAutoExpand)
-		.Graph(InGraph)
 		.NewNodePosition(InNodePosition)
 		.DraggedFromPins(InDraggedPins)
-		.OnClosedCallback(InOnMenuClosed)
-		.AllowedExecuteContexts( { FRigVMExecuteContext::StaticStruct(), FAnimNextParameterExecuteContext::StaticStruct() });
+		.OnClosedCallback(InOnMenuClosed);
 
 	TSharedPtr<SWidget> FilterTextBox = StaticCastSharedRef<SWidget>(ActionMenu->GetFilterTextBox());
 	return FActionMenuContent(StaticCastSharedRef<SWidget>(ActionMenu), FilterTextBox);
diff --git a/Engine/Plugins/Experimental/AnimNext/Source/AnimNextTestSuite/Private/AnimNextTraitGraphTest.cpp b/Engine/Plugins/Experimental/AnimNext/Source/AnimNextTestSuite/Private/AnimNextTraitGraphTest.cpp
--- a/Engine/Plugins/Experimental/AnimNext/Source/AnimNextTestSuite/Private/AnimNextTraitGraphTest.cpp
+++ b/Engine/Plugins/Experimental/AnimNext/Source/AnimNextTestSuite/Private/AnimNextTraitGraphTest.cpp
@@ -11,7 +11,7 @@
 #include "TraitInterfaces/IEvaluate.h"
 #include "TraitInterfaces/IUpdate.h"
 #include "Editor/Transactor.h"
-#include "Graph/AnimNextExecuteContext.h"
+#include "AnimNextExecuteContext.h"
 #include "Graph/AnimNextGraph.h"
 #include "Graph/AnimNextGraph_EditorData.h"
 #include "Graph/GraphFactory.h"
diff --git a/Engine/Plugins/Experimental/AnimNext/Source/AnimNextUncookedOnly/Internal/AnimNextRigVMAssetEditorData.h b/Engine/Plugins/Experimental/AnimNext/Source/AnimNextUncookedOnly/Internal/AnimNextRigVMAssetEditorData.h
--- a/Engine/Plugins/Experimental/AnimNext/Source/AnimNextUncookedOnly/Internal/AnimNextRigVMAssetEditorData.h
+++ b/Engine/Plugins/Experimental/AnimNext/Source/AnimNextUncookedOnly/Internal/AnimNextRigVMAssetEditorData.h
@@ -137,9 +137,6 @@ protected:
 	// Override called during initialization to determine what RigVM controller class is used
 	virtual TSubclassOf<URigVMController> GetControllerClass() const { return URigVMController::StaticClass(); }
 
-	// Override called during initialization to determine what RigVM schema class is used
-	virtual TSubclassOf<URigVMSchema> GetRigVMSchemaClass() const PURE_VIRTUAL(UAnimNextRigVMAssetEditorData::GetRigVMSchemaClass, return nullptr;)
-
 	// Override called during initialization to determine what RigVM execute struct is used
 	virtual UScriptStruct* GetExecuteContextStruct() const PURE_VIRTUAL(UAnimNextRigVMAssetEditorData::GetExecuteContextStruct, return nullptr;)
 
diff --git a/Engine/Plugins/Experimental/AnimNext/Source/AnimNextUncookedOnly/Internal/Graph/AnimNextGraph_EditorData.h b/Engine/Plugins/Experimental/AnimNext/Source/AnimNextUncookedOnly/Internal/Graph/AnimNextGraph_EditorData.h
--- a/Engine/Plugins/Experimental/AnimNext/Source/AnimNextUncookedOnly/Internal/Graph/AnimNextGraph_EditorData.h
+++ b/Engine/Plugins/Experimental/AnimNext/Source/AnimNextUncookedOnly/Internal/Graph/AnimNextGraph_EditorData.h
@@ -7,7 +7,7 @@
 #include "RigVMModel/RigVMGraph.h"
 #include "AnimNextGraph_EdGraph.h"
 #include "AnimNextRigVMAssetEditorData.h"
-#include "Graph/AnimNextExecuteContext.h"
+#include "AnimNextExecuteContext.h"
 #include "Kismet/BlueprintFunctionLibrary.h"
 #include "AnimNextGraph_EditorData.generated.h"
 
@@ -29,12 +29,8 @@ namespace UE::AnimNext::Editor
 	struct FUtils;
 }
 
-/**
- * The Schema is used to determine which actions are allowed
- * on a graph. This includes any topological change.
- */
 UCLASS()
-class UAnimNextGraph_Schema : public URigVMSchema
+class UAnimNextTraitGraph_Schema : public URigVMSchema
 {
 	GENERATED_BODY()
 };
@@ -82,7 +78,6 @@ private:
 
 	// UAnimNextRigVMAssetEditorData interface
 	virtual TSubclassOf<URigVMController> GetControllerClass() const override { return UAnimNextGraph_Controller::StaticClass(); }
-	virtual TSubclassOf<URigVMSchema> GetRigVMSchemaClass() const override { return UAnimNextGraph_Schema::StaticClass(); }
 	virtual UScriptStruct* GetExecuteContextStruct() const override { return FAnimNextExecuteContext::StaticStruct(); }
 	virtual UEdGraph* CreateEdGraph(URigVMGraph* InRigVMGraph, bool bForce) override;
 	virtual bool RemoveEdGraph(URigVMGraph* InModel) override;
diff --git a/Engine/Plugins/Experimental/AnimNext/Source/AnimNextUncookedOnly/Internal/Graph/AnimNextGraph_TraitSchema.h b/Engine/Plugins/Experimental/AnimNext/Source/AnimNextUncookedOnly/Internal/Graph/AnimNextGraph_TraitSchema.h
new file mode 100644
--- /dev/null
+++ b/Engine/Plugins/Experimental/AnimNext/Source/AnimNextUncookedOnly/Internal/Graph/AnimNextGraph_TraitSchema.h
@@ -0,0 +1,17 @@
+﻿// Copyright Epic Games, Inc. All Rights Reserved.
+
+#pragma once
+
+#include "CoreMinimal.h"
+#include "AnimNextRigVMAssetSchema.h"
+#include "AnimNextGraph_TraitSchema.generated.h"
+
+UCLASS()
+class UAnimNextGraph_TraitSchema : public UAnimNextRigVMAssetSchema
+{
+	GENERATED_BODY()
+
+	// URigVMSchema interface
+	virtual bool SupportsUnitFunction(URigVMController* InController, const FRigVMFunction* InUnitFunction) const override;
+	virtual bool SupportsDispatchFactory(URigVMController* InController, const FRigVMDispatchFactory* InDispatchFactory) const override;
+};
diff --git a/Engine/Plugins/Experimental/AnimNext/Source/AnimNextUncookedOnly/Internal/Param/AnimNextParameterBlock_EditorData.h b/Engine/Plugins/Experimental/AnimNext/Source/AnimNextUncookedOnly/Internal/Param/AnimNextParameterBlock_EditorData.h
--- a/Engine/Plugins/Experimental/AnimNext/Source/AnimNextUncookedOnly/Internal/Param/AnimNextParameterBlock_EditorData.h
+++ b/Engine/Plugins/Experimental/AnimNext/Source/AnimNextUncookedOnly/Internal/Param/AnimNextParameterBlock_EditorData.h
@@ -6,7 +6,7 @@
 #include "AnimNextParameterBlock_EdGraph.h"
 #include "AnimNextRigVMAssetEditorData.h"
 #include "Kismet/BlueprintFunctionLibrary.h"
-#include "Param/AnimNextParameterExecuteContext.h"
+#include "AnimNextExecuteContext.h"
 #include "Param/ParamType.h"
 #include "AnimNextParameterBlock_EditorData.generated.h"
 
@@ -39,12 +39,6 @@ namespace UE::AnimNext::Tests
 	class FEditor_Parameters_ParameterBlock;
 }
 
-UCLASS()
-class UAnimNextParameterBlockLibrary_Schema : public URigVMSchema
-{
-	GENERATED_BODY()
-};
-
 // Script-callable editor API hoisted onto UAnimNextParameterBlock
 UCLASS()
 class ANIMNEXTUNCOOKEDONLY_API UAnimNextParameterBlockLibrary : public UBlueprintFunctionLibrary
@@ -94,8 +88,7 @@ class UAnimNextParameterBlock_EditorData : public UAnimNextRigVMAssetEditorData
 	virtual void RecompileVM() override;
 
 	// UAnimNextRigVMAssetEditorData interface
-	virtual TSubclassOf<URigVMSchema> GetRigVMSchemaClass() const override { return UAnimNextParameterBlockLibrary_Schema::StaticClass(); }
-	virtual UScriptStruct* GetExecuteContextStruct() const override { return FAnimNextParameterExecuteContext::StaticStruct(); }
+	virtual UScriptStruct* GetExecuteContextStruct() const override { return FAnimNextExecuteContext::StaticStruct(); }
 	virtual UEdGraph* CreateEdGraph(URigVMGraph* InRigVMGraph, bool bForce) override;
 	virtual bool RemoveEdGraph(URigVMGraph* InModel) override;
 	virtual void CreateEdGraphForCollapseNode(URigVMCollapseNode* InNode) override;
diff --git a/Engine/Plugins/Experimental/AnimNext/Source/AnimNextUncookedOnly/Internal/Param/AnimNextParameterBlock_ParamSchema.h b/Engine/Plugins/Experimental/AnimNext/Source/AnimNextUncookedOnly/Internal/Param/AnimNextParameterBlock_ParamSchema.h
new file mode 100644
--- /dev/null
+++ b/Engine/Plugins/Experimental/AnimNext/Source/AnimNextUncookedOnly/Internal/Param/AnimNextParameterBlock_ParamSchema.h
@@ -0,0 +1,17 @@
+﻿// Copyright Epic Games, Inc. All Rights Reserved.
+
+#pragma once
+
+#include "CoreMinimal.h"
+#include "AnimNextRigVMAssetSchema.h"
+#include "AnimNextParameterBlock_ParamSchema.generated.h"
+
+UCLASS()
+class UAnimNextParameterBlock_ParamSchema : public UAnimNextRigVMAssetSchema
+{
+	GENERATED_BODY()
+
+	// URigVMSchema interface
+	virtual bool SupportsUnitFunction(URigVMController* InController, const FRigVMFunction* InUnitFunction) const override;
+};
+
diff --git a/Engine/Plugins/Experimental/AnimNext/Source/AnimNextUncookedOnly/Private/AnimNextRigVMAssetEditorData.cpp b/Engine/Plugins/Experimental/AnimNext/Source/AnimNextUncookedOnly/Private/AnimNextRigVMAssetEditorData.cpp
--- a/Engine/Plugins/Experimental/AnimNext/Source/AnimNextUncookedOnly/Private/AnimNextRigVMAssetEditorData.cpp
+++ b/Engine/Plugins/Experimental/AnimNext/Source/AnimNextUncookedOnly/Private/AnimNextRigVMAssetEditorData.cpp
@@ -4,6 +4,7 @@
 
 #include "AnimNextRigVMAsset.h"
 #include "AnimNextRigVMAssetEntry.h"
+#include "AnimNextRigVMAssetSchema.h"
 #include "ControlRigDefines.h"
 #include "ExternalPackageHelper.h"
 #include "IAnimNextRigVMGraphInterface.h"
@@ -65,6 +66,7 @@ void UAnimNextRigVMAssetEditorData::ReconstructAllNodes()
 
 void UAnimNextRigVMAssetEditorData::Serialize(FArchive& Ar)
 {
+	RigVMClient.SetDefaultSchemaClass(UAnimNextRigVMAssetSchema::StaticClass());
 	RigVMClient.SetOuterClientHost(this, GET_MEMBER_NAME_CHECKED(UAnimNextRigVMAssetEditorData, RigVMClient));
 
 	Super::Serialize(Ar);
@@ -74,8 +76,6 @@ void UAnimNextRigVMAssetEditorData::Initialize(bool bRecompileVM)
 {
 	RigVMClient.bDefaultModelCanBeRemoved = true;
 	RigVMClient.SetControllerClass(GetControllerClass());
-	RigVMClient.SetSchemaClass(GetRigVMSchemaClass());
-	RigVMClient.SetExecuteContextStruct(GetExecuteContextStruct());
 	RigVMClient.SetOuterClientHost(this, GET_MEMBER_NAME_CHECKED(UAnimNextRigVMAssetEditorData, RigVMClient));
 	{
 		TGuardValue<bool> DisableClientNotifs(RigVMClient.bSuspendNotifications, true);
@@ -420,7 +420,7 @@ void UAnimNextRigVMAssetEditorData::HandleRigVMGraphAdded(const FRigVMClient* In
 #if WITH_EDITOR
 		if(!bSuspendPythonMessagesForRigVMClient)
 		{
-			const FString AssetName = InClient->GetSchema()->GetSanitizedName(GetName(), true, false);
+			const FString AssetName = RigVMGraph->GetSchema()->GetSanitizedName(GetName(), true, false);
 			RigVMPythonUtils::Print(AssetName, FString::Printf(TEXT("asset.add_graph('%s')"), *RigVMGraph->GetName()));
 		}
 #endif
@@ -437,7 +437,7 @@ void UAnimNextRigVMAssetEditorData::HandleRigVMGraphRemoved(const FRigVMClient*
 #if WITH_EDITOR
 		if(!bSuspendPythonMessagesForRigVMClient)
 		{
-			const FString AssetName = InClient->GetSchema()->GetSanitizedName(GetName(), true, false);
+			const FString AssetName = RigVMGraph->GetSchema()->GetSanitizedName(GetName(), true, false);
 			RigVMPythonUtils::Print(AssetName, FString::Printf(TEXT("asset.add_graph('%s')"), *RigVMGraph->GetName()));
 		}
 #endif
diff --git a/Engine/Plugins/Experimental/AnimNext/Source/AnimNextUncookedOnly/Private/AnimNextRigVMAssetSchema.cpp b/Engine/Plugins/Experimental/AnimNext/Source/AnimNextUncookedOnly/Private/AnimNextRigVMAssetSchema.cpp
new file mode 100644
--- /dev/null
+++ b/Engine/Plugins/Experimental/AnimNext/Source/AnimNextUncookedOnly/Private/AnimNextRigVMAssetSchema.cpp
@@ -0,0 +1,9 @@
+﻿// Copyright Epic Games, Inc. All Rights Reserved.
+
+#include "AnimNextRigVMAssetSchema.h"
+#include "AnimNextExecuteContext.h"
+
+UAnimNextRigVMAssetSchema::UAnimNextRigVMAssetSchema()
+{
+	SetExecuteContextStruct(FAnimNextExecuteContext::StaticStruct());
+}
diff --git a/Engine/Plugins/Experimental/AnimNext/Source/AnimNextUncookedOnly/Private/AnimNextRigVMAssetSchema.h b/Engine/Plugins/Experimental/AnimNext/Source/AnimNextUncookedOnly/Private/AnimNextRigVMAssetSchema.h
new file mode 100644
--- /dev/null
+++ b/Engine/Plugins/Experimental/AnimNext/Source/AnimNextUncookedOnly/Private/AnimNextRigVMAssetSchema.h
@@ -0,0 +1,17 @@
+﻿// Copyright Epic Games, Inc. All Rights Reserved.
+
+#pragma once
+
+#include "CoreMinimal.h"
+#include "RigVMModel/RigVMSchema.h"
+#include "AnimNextRigVMAssetSchema.generated.h"
+
+UCLASS()
+class UAnimNextRigVMAssetSchema : public URigVMSchema
+{
+	GENERATED_BODY()
+
+protected:
+	UAnimNextRigVMAssetSchema();
+};
+
diff --git a/Engine/Plugins/Experimental/AnimNext/Source/AnimNextUncookedOnly/Private/Graph/AnimNextGraph_EditorData.cpp b/Engine/Plugins/Experimental/AnimNext/Source/AnimNextUncookedOnly/Private/Graph/AnimNextGraph_EditorData.cpp
--- a/Engine/Plugins/Experimental/AnimNext/Source/AnimNextUncookedOnly/Private/Graph/AnimNextGraph_EditorData.cpp
+++ b/Engine/Plugins/Experimental/AnimNext/Source/AnimNextUncookedOnly/Private/Graph/AnimNextGraph_EditorData.cpp
@@ -7,6 +7,7 @@
 #include "Graph/AnimNextGraph.h"
 #include "Graph/AnimNextGraphEntry.h"
 #include "Graph/AnimNextGraph_EdGraphSchema.h"
+#include "Graph/AnimNextGraph_TraitSchema.h"
 #include "RigVMModel/RigVMFunctionLibrary.h"
 #include "RigVMModel/RigVMNotifications.h"
 #include "RigVMModel/Nodes/RigVMCollapseNode.h"
@@ -15,8 +16,18 @@
 
 void UAnimNextGraph_EditorData::PostLoad()
 {
-	Super::PostLoad();
+	if(GetLinkerCustomVersion(FFortniteMainBranchObjectVersion::GUID) < FFortniteMainBranchObjectVersion::AnimNextCombineGraphContexts)
+	{
+		TArray<URigVMGraph*> AllModels = RigVMClient.GetAllModels(false, true);
+		for(URigVMGraph* Graph : AllModels)
+		{
+			Graph->SetExecuteContextStruct(FAnimNextExecuteContext::StaticStruct());
+			Graph->SetSchemaClass(UAnimNextGraph_TraitSchema::StaticClass());
+		}
+	}
 
+	Super::PostLoad();
+	
 	if(GetLinkerCustomVersion(FFortniteMainBranchObjectVersion::GUID) < FFortniteMainBranchObjectVersion::AnimNextMoveGraphsToEntries)
 	{
 		for(TObjectPtr<UAnimNextGraph_EdGraph> Graph : Graphs_DEPRECATED)
@@ -220,7 +231,7 @@ UAnimNextGraphEntry* UAnimNextGraph_EditorData::AddGraph(FName InName, bool bSet
 	{
 		TGuardValue<bool> EnablePythonPrint(bSuspendPythonMessagesForRigVMClient, !bPrintPythonCommand);
 		TGuardValue<bool> DisableAutoCompile(bAutoRecompileVM, false);
-		URigVMGraph* NewGraph = RigVMClient.AddModel(URigVMGraph::StaticClass()->GetFName(), bSetupUndoRedo);
+		URigVMGraph* NewGraph = RigVMClient.AddModel(URigVMGraph::StaticClass()->GetFName(), UAnimNextGraph_TraitSchema::StaticClass(), bSetupUndoRedo);
 		ensure(NewGraph);
 		NewEntry->Graph = NewGraph;
 
diff --git a/Engine/Plugins/Experimental/AnimNext/Source/AnimNextUncookedOnly/Private/Graph/AnimNextGraph_TraitSchema.cpp b/Engine/Plugins/Experimental/AnimNext/Source/AnimNextUncookedOnly/Private/Graph/AnimNextGraph_TraitSchema.cpp
new file mode 100644
--- /dev/null
+++ b/Engine/Plugins/Experimental/AnimNext/Source/AnimNextUncookedOnly/Private/Graph/AnimNextGraph_TraitSchema.cpp
@@ -0,0 +1,44 @@
+﻿// Copyright Epic Games, Inc. All Rights Reserved.
+
+#include "Graph/AnimNextGraph_TraitSchema.h"
+
+#include "AnimNextExecuteContext.h"
+#include "Graph/RigUnit_AnimNextBase.h"
+#include "Param/RigVMDispatch_GetParameter.h"
+
+bool UAnimNextGraph_TraitSchema::SupportsUnitFunction(URigVMController* InController, const FRigVMFunction* InUnitFunction) const
+{
+	if(InUnitFunction)
+	{
+		if(const UScriptStruct* FunctionExecuteContextStruct = InUnitFunction->GetExecuteContextStruct())
+		{
+			if(FunctionExecuteContextStruct == FAnimNextExecuteContext::StaticStruct())
+			{
+				// Only allow nodes that are children of FRigUnit_AnimNextBase
+				if(InUnitFunction->Struct)
+				{
+					return InUnitFunction->Struct->IsChildOf(FRigUnit_AnimNextBase::StaticStruct());
+				}
+			}
+		}
+	}
+
+	return Super::SupportsUnitFunction(InController, InUnitFunction);
+}
+
+bool UAnimNextGraph_TraitSchema::SupportsDispatchFactory(URigVMController* InController, const FRigVMDispatchFactory* InDispatchFactory) const
+{
+	if(InDispatchFactory)
+	{
+		if(const UScriptStruct* DispatchExecuteContextStruct = InDispatchFactory->GetExecuteContextStruct())
+		{
+			if(DispatchExecuteContextStruct == FAnimNextExecuteContext::StaticStruct())
+			{
+				// We only support the FRigVMDispatch_GetParameter at the moment.
+				return (InDispatchFactory->GetScriptStruct() == FRigVMDispatch_GetParameter::StaticStruct());
+			}
+		}
+	}
+	
+	return Super::SupportsDispatchFactory(InController, InDispatchFactory);
+}
\ No newline at end of file
diff --git a/Engine/Plugins/Experimental/AnimNext/Source/AnimNextUncookedOnly/Private/Param/AnimNextParameterBlock_EditorData.cpp b/Engine/Plugins/Experimental/AnimNext/Source/AnimNextUncookedOnly/Private/Param/AnimNextParameterBlock_EditorData.cpp
--- a/Engine/Plugins/Experimental/AnimNext/Source/AnimNextUncookedOnly/Private/Param/AnimNextParameterBlock_EditorData.cpp
+++ b/Engine/Plugins/Experimental/AnimNext/Source/AnimNextUncookedOnly/Private/Param/AnimNextParameterBlock_EditorData.cpp
@@ -2,6 +2,7 @@
 
 #include "Param/AnimNextParameterBlock_EditorData.h"
 
+#include "AnimNextExecuteContext.h"
 #include "UncookedOnlyUtils.h"
 #include "Curves/CurveFloat.h"
 #include "Param/AnimNextParameterBlock.h"
@@ -9,6 +10,7 @@
 #include "Param/AnimNextParameterBlockParameter.h"
 #include "Param/AnimNextParameterBlock_EdGraph.h"
 #include "Param/AnimNextParameterBlock_EdGraphSchema.h"
+#include "Param/AnimNextParameterBlock_ParamSchema.h"
 #include "RigVMModel/RigVMFunctionLibrary.h"
 #include "RigVMModel/Nodes/RigVMCollapseNode.h"
 #include "UObject/LinkerLoad.h"
@@ -106,7 +108,7 @@ UAnimNextParameterBlockGraph* UAnimNextParameterBlock_EditorData::AddGraph(FName
 	{
 		TGuardValue<bool> EnablePythonPrint(bSuspendPythonMessagesForRigVMClient, !bPrintPythonCommand);
 		TGuardValue<bool> DisableAutoCompile(bAutoRecompileVM, false);
-		URigVMGraph* NewGraph = RigVMClient.AddModel(URigVMGraph::StaticClass()->GetFName(), bSetupUndoRedo);
+		URigVMGraph* NewGraph = RigVMClient.AddModel(URigVMGraph::StaticClass()->GetFName(), UAnimNextParameterBlock_ParamSchema::StaticClass(), bSetupUndoRedo);
 		ensure(NewGraph);
 		NewEntry->Graph = NewGraph;
 
@@ -132,6 +134,16 @@ TConstArrayView<TSubclassOf<UAnimNextRigVMAssetEntry>> UAnimNextParameterBlock_E
 
 void UAnimNextParameterBlock_EditorData::PostLoad()
 {
+	if(GetLinkerCustomVersion(FFortniteMainBranchObjectVersion::GUID) < FFortniteMainBranchObjectVersion::AnimNextCombineGraphContexts)
+	{
+		TArray<URigVMGraph*> AllModels = RigVMClient.GetAllModels(false, true);
+		for(URigVMGraph* Graph : AllModels)
+		{
+			Graph->SetExecuteContextStruct(FAnimNextExecuteContext::StaticStruct());
+			Graph->SetSchemaClass(UAnimNextParameterBlock_ParamSchema::StaticClass());
+		}
+	}
+
 	Super::PostLoad();
 
 	if(GetLinkerCustomVersion(FFortniteMainBranchObjectVersion::GUID) < FFortniteMainBranchObjectVersion::AnimNextMoveGraphsToEntries)
diff --git a/Engine/Plugins/Experimental/AnimNext/Source/AnimNextUncookedOnly/Private/Param/AnimNextParameterBlock_ParamSchema.cpp b/Engine/Plugins/Experimental/AnimNext/Source/AnimNextUncookedOnly/Private/Param/AnimNextParameterBlock_ParamSchema.cpp
new file mode 100644
--- /dev/null
+++ b/Engine/Plugins/Experimental/AnimNext/Source/AnimNextUncookedOnly/Private/Param/AnimNextParameterBlock_ParamSchema.cpp
@@ -0,0 +1,23 @@
+﻿// Copyright Epic Games, Inc. All Rights Reserved.
+
+#include "PAram/AnimNextParameterBlock_ParamSchema.h"
+
+#include "AnimNextExecuteContext.h"
+#include "Param/RigUnit_AnimNextParameterBase.h"
+
+bool UAnimNextParameterBlock_ParamSchema::SupportsUnitFunction(URigVMController* InController, const FRigVMFunction* InUnitFunction) const
+{
+	if(const UScriptStruct* FunctionExecuteContextStruct = InUnitFunction->GetExecuteContextStruct())
+	{
+		if(FunctionExecuteContextStruct == FAnimNextExecuteContext::StaticStruct())
+		{
+			// Only allow nodes that are children of FRigUnit_AnimNextParameterBase
+			if(InUnitFunction->Struct)
+			{
+				return InUnitFunction->Struct->IsChildOf(FRigUnit_AnimNextParameterBase::StaticStruct());
+			}
+		}
+	}
+
+	return Super::SupportsUnitFunction(InController, InUnitFunction);
+}
diff --git a/Engine/Plugins/Experimental/AnimNext/Source/AnimNextUncookedOnly/Private/UncookedOnlyUtils.cpp b/Engine/Plugins/Experimental/AnimNext/Source/AnimNextUncookedOnly/Private/UncookedOnlyUtils.cpp
--- a/Engine/Plugins/Experimental/AnimNext/Source/AnimNextUncookedOnly/Private/UncookedOnlyUtils.cpp
+++ b/Engine/Plugins/Experimental/AnimNext/Source/AnimNextUncookedOnly/Private/UncookedOnlyUtils.cpp
@@ -634,7 +634,7 @@ void FUtils::Compile(UAnimNextGraph* InGraph)
 	}
 
 	URigVMCompiler* Compiler = URigVMCompiler::StaticClass()->GetDefaultObject<URigVMCompiler>();
-	EditorData->VMCompileSettings.SetExecuteContextStruct(EditorData->RigVMClient.GetExecuteContextStruct());
+	EditorData->VMCompileSettings.SetExecuteContextStruct(FAnimNextExecuteContext::StaticStruct());
 	const FRigVMCompileSettings Settings = (EditorData->bCompileInDebugMode) ? FRigVMCompileSettings::Fast(EditorData->VMCompileSettings.GetExecuteContextStruct()) : EditorData->VMCompileSettings;
 	Compiler->Compile(Settings, VMTempGraphs, TempController, InGraph->VM, InGraph->ExtendedExecuteContext, TArray<FRigVMExternalVariable>(), & EditorData->PinToOperandMap);
 
@@ -847,7 +847,7 @@ void FUtils::CompileVM(UAnimNextParameterBlock* InParameterBlock)
 	}
 
 	URigVMCompiler* Compiler = URigVMCompiler::StaticClass()->GetDefaultObject<URigVMCompiler>();
-	EditorData->VMCompileSettings.SetExecuteContextStruct(EditorData->RigVMClient.GetExecuteContextStruct());
+	EditorData->VMCompileSettings.SetExecuteContextStruct(FAnimNextExecuteContext::StaticStruct());
 	FRigVMExtendedExecuteContext& CDOContext = InParameterBlock->GetRigVMExtendedExecuteContext();
 	const FRigVMCompileSettings Settings = (EditorData->bCompileInDebugMode) ? FRigVMCompileSettings::Fast(EditorData->VMCompileSettings.GetExecuteContextStruct()) : EditorData->VMCompileSettings;
 	URigVMController* RootController = VMClient->GetOrCreateController(RootGraph);
diff --git a/Engine/Plugins/Runtime/RigVM/Source/RigVMDeveloper/Private/RigVMBlueprint.cpp b/Engine/Plugins/Runtime/RigVM/Source/RigVMDeveloper/Private/RigVMBlueprint.cpp
--- a/Engine/Plugins/Runtime/RigVM/Source/RigVMDeveloper/Private/RigVMBlueprint.cpp
+++ b/Engine/Plugins/Runtime/RigVM/Source/RigVMDeveloper/Private/RigVMBlueprint.cpp
@@ -214,10 +214,10 @@ URigVMBlueprint::URigVMBlueprint(const FObjectInitializer& ObjectInitializer)
 void URigVMBlueprint::CommonInitialization(const FObjectInitializer& ObjectInitializer)
 {
 	// guard against this running multiple times
-	check(GetRigVMClient()->GetSchema() == nullptr);
-	
-	RigVMClient.SetSchemaClass(GetRigVMSchemaClass());
-	RigVMClient.SetExecuteContextStruct(GetRigVMExecuteContextStruct());
+	check(GetRigVMClient()->GetDefaultSchemaClass() == nullptr);
+
+	RigVMClient.SetDefaultSchemaClass(GetRigVMSchemaClass());
+	RigVMClient.SetDefaultExecuteContextStruct(GetRigVMExecuteContextStruct());
 
 	for(UEdGraph* UberGraph : UbergraphPages)
 	{
@@ -498,7 +498,7 @@ void URigVMBlueprint::HandleRigVMGraphAdded(const FRigVMClient* InClient, const
 #if WITH_EDITOR
 		if(!bSuspendPythonMessagesForRigVMClient)
 		{
-			const FString BlueprintName = InClient->GetSchema()->GetSanitizedName(GetName(), true, false);
+			const FString BlueprintName = InClient->GetDefaultSchema()->GetSanitizedName(GetName(), true, false);
 			RigVMPythonUtils::Print(BlueprintName, 
 				FString::Printf(TEXT("blueprint.add_model('%s')"),
 					*Model->GetName()));
@@ -517,7 +517,7 @@ void URigVMBlueprint::HandleRigVMGraphRemoved(const FRigVMClient* InClient, cons
 #if WITH_EDITOR
 		if(!bSuspendPythonMessagesForRigVMClient)
 		{
-			const FString BlueprintName = InClient->GetSchema()->GetSanitizedName(GetName(), true, false);
+			const FString BlueprintName = InClient->GetDefaultSchema()->GetSanitizedName(GetName(), true, false);
 			RigVMPythonUtils::Print(BlueprintName, 
 				FString::Printf(TEXT("blueprint.remove_model('%s')"),
 					*Model->GetName()));
@@ -1280,7 +1280,7 @@ void URigVMBlueprint::RecompileVM()
 		}
 
 		URigVMCompiler* Compiler = URigVMCompiler::StaticClass()->GetDefaultObject<URigVMCompiler>();
-		VMCompileSettings.SetExecuteContextStruct(RigVMClient.GetExecuteContextStruct());
+		VMCompileSettings.SetExecuteContextStruct(RigVMClient.GetDefaultExecuteContextStruct());
 
 	    FRigVMExtendedExecuteContext& CDOContext = CDO->GetRigVMExtendedExecuteContext();
 		const FRigVMCompileSettings Settings = (bCompileInDebugMode) ? FRigVMCompileSettings::Fast(VMCompileSettings.GetExecuteContextStruct()) : VMCompileSettings;
@@ -2303,7 +2303,7 @@ URigVMGraph* URigVMBlueprint::GetTemplateModel(bool bIsFunctionLibrary)
 			TemplateModel = NewObject<URigVMGraph>(this, TEXT("TemplateModel"));
 		}
 		TemplateModel->SetFlags(RF_Transient);
-		TemplateModel->SetExecuteContextStruct(RigVMClient.GetExecuteContextStruct());
+		TemplateModel->SetExecuteContextStruct(RigVMClient.GetDefaultExecuteContextStruct());
 	}
 	return TemplateModel;
 #else
@@ -2320,7 +2320,7 @@ URigVMController* URigVMBlueprint::GetTemplateController(bool bIsFunctionLibrary
 		TemplateController->SetGraph(GetTemplateModel(bIsFunctionLibrary));
 		TemplateController->EnableReporting(false);
 		TemplateController->SetFlags(RF_Transient);
-		TemplateController->SetSchema(RigVMClient.GetOrCreateSchema());
+		TemplateController->SetSchemaClass(RigVMClient.GetDefaultSchemaClass());
 	}
 	return TemplateController;
 #else
diff --git a/Engine/Plugins/Runtime/RigVM/Source/RigVMDeveloper/Private/RigVMModel/RigVMClient.cpp b/Engine/Plugins/Runtime/RigVM/Source/RigVMDeveloper/Private/RigVMModel/RigVMClient.cpp
--- a/Engine/Plugins/Runtime/RigVM/Source/RigVMDeveloper/Private/RigVMModel/RigVMClient.cpp
+++ b/Engine/Plugins/Runtime/RigVM/Source/RigVMDeveloper/Private/RigVMModel/RigVMClient.cpp
@@ -20,30 +20,20 @@ UObject* IRigVMClientHost::ResolveUserDefinedTypeById(const FString& InTypeName)
 	return nullptr;
 }
 
-void FRigVMClient::SetSchemaClass(TSubclassOf<URigVMSchema> InSchemaClass)
+void FRigVMClient::SetDefaultSchemaClass(TSubclassOf<URigVMSchema> InSchemaClass)
 {
 	check(InSchemaClass);
-	
-	if(InSchemaClass == SchemaClass)
+
+	if(InSchemaClass == DefaultSchemaClass)
 	{
 		return;
 	}
-	
-	if(SchemaPtr)
-	{
-		if(!SchemaPtr->HasAnyFlags(RF_ClassDefaultObject))
-		{
-			DestroyObject(SchemaPtr);
-		}
-		SchemaPtr = nullptr;
-	}
 
-	SchemaClass = InSchemaClass;
+	DefaultSchemaClass = InSchemaClass;
 
-	URigVMSchema* Schema = GetOrCreateSchema();
-	for(auto& Pair : Controllers)
+	for(TPair<FSoftObjectPath, TObjectPtr<URigVMController>>& ObjectControllerPair : Controllers)
 	{
-		Pair.Value->SetSchema(Schema);
+		ObjectControllerPair.Value->SetSchemaClass(InSchemaClass);
 	}
 }
 
@@ -126,33 +116,12 @@ void FRigVMClient::Reset()
 	FunctionLibrary = nullptr;
 
 	ResetActionStack();
-
-	if(SchemaPtr)
-	{
-		if(!SchemaPtr->HasAnyFlags(RF_ClassDefaultObject))
-		{
-			DestroyObject(SchemaPtr);
-		}
-		SchemaPtr = nullptr;
-	}
 }
 
-URigVMSchema* FRigVMClient::GetOrCreateSchema()
+URigVMSchema* FRigVMClient::GetDefaultSchema() const
 {
-	check(SchemaClass);
-	if(SchemaPtr && !IsValid(SchemaPtr))
-	{
-		SchemaPtr = nullptr;
-	}
-	if(SchemaPtr && SchemaPtr->GetClass() != SchemaClass)
-	{
-		SchemaPtr = nullptr;
-	}
-	if(SchemaPtr == nullptr)
-	{
-		SchemaPtr = CastChecked<URigVMSchema>(SchemaClass->GetDefaultObject(true));
-	}
-	return SchemaPtr;
+	check(DefaultSchemaClass);
+	return DefaultSchemaClass->GetDefaultObject<URigVMSchema>();
 }
 
 URigVMGraph* FRigVMClient::GetDefaultModel() const
@@ -381,6 +350,14 @@ URigVMGraph* FRigVMClient::AddModel(const FString InName, bool bSetupUndoRedo, b
 
 URigVMGraph* FRigVMClient::AddModel(const FName& InName, bool bSetupUndoRedo, const FObjectInitializer* ObjectInitializer, bool bCreateController)
 {
+	check(DefaultSchemaClass);
+	return AddModel(InName, DefaultSchemaClass, bSetupUndoRedo, ObjectInitializer, bCreateController);
+}
+
+URigVMGraph* FRigVMClient::AddModel(const FName& InName, TSubclassOf<URigVMSchema> InSchemaClass, bool bSetupUndoRedo, const FObjectInitializer* ObjectInitializer, bool bCreateController)
+{
+	check(InSchemaClass);
+
 #if WITH_EDITOR
 	TSharedPtr<FScopedTransaction> Transaction;
 	if(bSetupUndoRedo)
@@ -400,6 +377,9 @@ URigVMGraph* FRigVMClient::AddModel(const FName& InName, bool bSetupUndoRedo, co
 	{
 		Model = NewObject<URigVMGraph>(GetOuter(), SafeGraphName);
 	}
+
+	Model->SetSchemaClass(InSchemaClass);
+
 	AddModel(Model, bCreateController);
 
 	if(bSetupUndoRedo)
@@ -430,7 +410,7 @@ void FRigVMClient::AddModel(URigVMGraph* InModel, bool bCreateController)
 		Models.Add(InModel);
 	}
 
-	InModel->SetExecuteContextStruct(GetOrCreateSchema()->GetExecuteContextStruct());
+	InModel->SetExecuteContextStruct(InModel->GetSchema()->GetExecuteContextStruct());
 
 	if(bCreateController)
 	{
@@ -461,6 +441,12 @@ void FRigVMClient::AddModel(URigVMGraph* InModel, bool bCreateController)
 }
 
 URigVMFunctionLibrary* FRigVMClient::GetOrCreateFunctionLibrary(bool bSetupUndoRedo, const FObjectInitializer* ObjectInitializer, bool bCreateController)
+{
+	check(DefaultSchemaClass);
+	return GetOrCreateFunctionLibrary(DefaultSchemaClass, bSetupUndoRedo, ObjectInitializer, bCreateController);
+}
+
+URigVMFunctionLibrary* FRigVMClient::GetOrCreateFunctionLibrary(TSubclassOf<URigVMSchema> InSchemaClass, bool bSetupUndoRedo, const FObjectInitializer* ObjectInitializer, bool bCreateController)
 {
 	if(FunctionLibrary)
 	{
@@ -488,6 +474,8 @@ URigVMFunctionLibrary* FRigVMClient::GetOrCreateFunctionLibrary(bool bSetupUndoR
 		NewFunctionLibrary = NewObject<URigVMFunctionLibrary>(GetOuter(), SafeGraphName);
 	}
 
+	NewFunctionLibrary->SetSchemaClass(InSchemaClass);
+
 	NewFunctionLibrary->GetFunctionHostObjectPathDelegate.BindLambda([this]() -> const FSoftObjectPath 
 		{
 			if (GetOuter()->Implements<URigVMClientHost>())
@@ -519,15 +507,15 @@ TArray<FName> FRigVMClient::GetEntryNames() const
 	return EntryNames;
 }
 
-UScriptStruct* FRigVMClient::GetExecuteContextStruct() const
+UScriptStruct* FRigVMClient::GetDefaultExecuteContextStruct() const
 {
-	check(SchemaPtr);
-	return SchemaPtr->GetExecuteContextStruct();
+	check(GetDefaultSchema());
+	return GetDefaultSchema()->GetExecuteContextStruct();
 }
 
-void FRigVMClient::SetExecuteContextStruct(UScriptStruct* InExecuteContextStruct)
+void FRigVMClient::SetDefaultExecuteContextStruct(UScriptStruct* InExecuteContextStruct)
 {
-	GetOrCreateSchema()->SetExecuteContextStruct(InExecuteContextStruct);
+	GetDefaultSchema()->SetExecuteContextStruct(InExecuteContextStruct);
 }
 
 URigVMGraph* FRigVMClient::GetFocusedModel() const
@@ -838,7 +826,16 @@ URigVMController* FRigVMClient::CreateController(const URigVMGraph* InModel)
 	const FName SafeControllerName = GetUniqueName(*FString::Printf(TEXT("%s_Controller"), *ModelName));
 	URigVMController* Controller = NewObject<URigVMController>(GetOuter(), ControllerClass, SafeControllerName);
 	Controllers.Add(InModel, Controller);
-	Controller->SetSchema(GetOrCreateSchema());
+	if(InModel && InModel->GetSchemaClass() != nullptr)
+	{
+		Controller->SetSchemaClass(InModel->GetSchemaClass());
+	}
+	else
+	{
+		check(DefaultSchemaClass);
+		Controller->SetSchemaClass(DefaultSchemaClass);
+	}
+
 	Controller->SetActionStack(GetOrCreateActionStack());
 	if(InModel)
 	{
@@ -885,13 +882,6 @@ void FRigVMClient::ResetActionStack()
 	ActionStack = nullptr;
 }
 
-void FRigVMClient::SetSchema(URigVMSchema* InSchema)
-{
-	check(InSchema);
-	SchemaPtr = InSchema;
-	SchemaClass = InSchema->GetClass();
-}
-
 FName FRigVMClient::GetUniqueName(const FName& InDesiredName) const
 {
 	return GetUniqueName(GetOuter(), InDesiredName);
@@ -988,14 +978,10 @@ void FRigVMClient::ProcessDetachedLinks()
 
 void FRigVMClient::PostDuplicateHost(const FString& InOldPathName, const FString& InNewPathName)
 {
-	// this also takes care of wiping invalid schemas on the client
-	URigVMSchema* Schema = GetOrCreateSchema();
-	
 	TArray<URigVMGraph*> AllModels = GetAllModels(true, true);
 	for(URigVMGraph* Model : AllModels)
 	{
 		URigVMController* Controller = GetOrCreateController(Model);
-		Controller->SetSchema(Schema);
 		Controller->PostDuplicateHost(InOldPathName, InNewPathName);
 	}
 }
diff --git a/Engine/Plugins/Runtime/RigVM/Source/RigVMDeveloper/Private/RigVMModel/RigVMController.cpp b/Engine/Plugins/Runtime/RigVM/Source/RigVMDeveloper/Private/RigVMModel/RigVMController.cpp
--- a/Engine/Plugins/Runtime/RigVM/Source/RigVMDeveloper/Private/RigVMModel/RigVMController.cpp
+++ b/Engine/Plugins/Runtime/RigVM/Source/RigVMDeveloper/Private/RigVMModel/RigVMController.cpp
@@ -84,7 +84,6 @@ void FRigVMClientPatchResult::Merge(const FRigVMClientPatchResult& InOther)
 
 URigVMController::URigVMController()
 	: bValidatePinDefaults(true)
-	, SchemaPtr(nullptr)
 	, bSuspendNotifications(false)
 	, bReportWarningsAndErrors(true)
 	, bIgnoreRerouteCompactnessChanges(false)
@@ -105,7 +104,6 @@ URigVMController::URigVMController()
 URigVMController::URigVMController(const FObjectInitializer& ObjectInitializer)
 	: Super(ObjectInitializer)
 	, bValidatePinDefaults(true)
-	, SchemaPtr(nullptr)
 	, bSuspendNotifications(false)
 	, bReportWarningsAndErrors(true)
 	, bIgnoreRerouteCompactnessChanges(false)
@@ -184,9 +182,9 @@ void URigVMController::SetGraph(URigVMGraph* InGraph)
 	HandleModifiedEvent(ERigVMGraphNotifType::GraphChanged, GetGraph(), nullptr);
 }
 
-void URigVMController::SetSchema(URigVMSchema* InSchema)
+URigVMSchema* URigVMController::GetSchema() const
 {
-	SchemaPtr = InSchema;
+	return SchemaClass->GetDefaultObject<URigVMSchema>(); 
 }
 
 bool URigVMController::PushGraph(URigVMGraph* InGraph, bool bSetupUndoRedo)
@@ -13958,6 +13956,7 @@ bool URigVMController::IsValidGraph() const
 
 bool URigVMController::IsValidSchema() const
 {
+	const URigVMSchema* SchemaPtr = GetSchema();
 	if (SchemaPtr == nullptr)
 	{
 		ReportError(TEXT("Controller does not have a schema associated."));
diff --git a/Engine/Plugins/Runtime/RigVM/Source/RigVMDeveloper/Private/RigVMModel/RigVMGraph.cpp b/Engine/Plugins/Runtime/RigVM/Source/RigVMDeveloper/Private/RigVMModel/RigVMGraph.cpp
--- a/Engine/Plugins/Runtime/RigVM/Source/RigVMDeveloper/Private/RigVMModel/RigVMGraph.cpp
+++ b/Engine/Plugins/Runtime/RigVM/Source/RigVMDeveloper/Private/RigVMModel/RigVMGraph.cpp
@@ -9,6 +9,7 @@
 #include "UObject/Package.h"
 #include "UObject/ObjectSaveContext.h"
 #include "RigVMTypeUtils.h"
+#include "RigVMModel/RigVMClient.h"
 
 #include UE_INLINE_GENERATED_CPP_BY_NAME(RigVMGraph)
 
@@ -16,6 +17,7 @@ URigVMGraph::URigVMGraph()
 : DiagnosticsAST(nullptr)
 , RuntimeAST(nullptr)
 , bEditable(true)
+, SchemaClass(nullptr)
 {
 	SetExecuteContextStruct(FRigVMExecuteContext::StaticStruct());
 }
@@ -403,6 +405,18 @@ TArray<FRigVMGraphVariableDescription> URigVMGraph::GetOutputArguments() const
 	return Outputs;
 }
 
+URigVMSchema* URigVMGraph::GetSchema() const
+{
+	check(SchemaClass);
+	return SchemaClass->GetDefaultObject<URigVMSchema>();
+}
+
+void URigVMGraph::SetSchemaClass(TSubclassOf<URigVMSchema> InSchemaClass)
+{
+	check(InSchemaClass);
+	SchemaClass = InSchemaClass;
+}
+
 FRigVMGraphModifiedEvent& URigVMGraph::OnModified()
 {
 	return ModifiedEvent;
diff --git a/Engine/Plugins/Runtime/RigVM/Source/RigVMDeveloper/Public/RigVMModel/RigVMClient.h b/Engine/Plugins/Runtime/RigVM/Source/RigVMDeveloper/Public/RigVMModel/RigVMClient.h
--- a/Engine/Plugins/Runtime/RigVM/Source/RigVMDeveloper/Public/RigVMModel/RigVMClient.h
+++ b/Engine/Plugins/Runtime/RigVM/Source/RigVMDeveloper/Public/RigVMModel/RigVMClient.h
@@ -156,8 +156,7 @@ public:
 	GENERATED_BODY()
 
 	FRigVMClient()
-		: SchemaPtr(nullptr)
-		, SchemaClass(URigVMSchema::StaticClass())
+		: DefaultSchemaClass(nullptr)
 		, ControllerClass(URigVMController::StaticClass())
 		, FunctionLibrary(nullptr)
 		, ActionStack(nullptr)
@@ -169,15 +168,21 @@ public:
 	{
 	}
 
-	void SetSchemaClass(TSubclassOf<URigVMSchema> InSchemaClass);
+	UE_DEPRECATED(5.5, "Please use SetDefaultSchemaClass or set a schema per controller/graph.")
+	void SetSchemaClass(TSubclassOf<URigVMSchema> InSchemaClass) { SetDefaultSchemaClass(InSchemaClass); }
+	void SetDefaultSchemaClass(TSubclassOf<URigVMSchema> InSchemaClass);
 	void SetControllerClass(TSubclassOf<URigVMController> InControllerClass);
 	void SetOuterClientHost(UObject* InOuterClientHost, const FName& InOuterClientHostPropertyName);
 	void SetFromDeprecatedData(URigVMGraph* InDefaultGraph, URigVMFunctionLibrary* InFunctionLibrary);
 
 	void Reset();
 	int32 Num() const { return Models.Num(); }
-	const URigVMSchema* GetSchema() const { return SchemaPtr; }
-	URigVMSchema* GetOrCreateSchema();
+	UE_DEPRECATED(5.5, "Please use GetDefaultSchema or get a schema per controller/graph.")
+	URigVMSchema* GetSchema() { return GetDefaultSchema(); }
+	URigVMSchema* GetDefaultSchema() const;
+	TSubclassOf<URigVMSchema> GetDefaultSchemaClass() const { return DefaultSchemaClass; }
+	UE_DEPRECATED(5.5, "Please use GetSchema or get a schema per controller/graph.")
+	URigVMSchema* GetOrCreateSchema() { return GetDefaultSchema(); }
 	URigVMGraph* GetDefaultModel() const;
 	URigVMGraph* GetModel(int32 InIndex) const { return Models.IsValidIndex(InIndex) ? Models[InIndex] : nullptr; }
 	URigVMGraph* GetModel(const UEdGraph* InEdGraph = nullptr) const;
@@ -197,9 +202,14 @@ public:
 	bool RemoveController(const URigVMGraph* InModel);
 	URigVMFunctionLibrary* GetFunctionLibrary() const { return FunctionLibrary; }
 	URigVMFunctionLibrary* GetOrCreateFunctionLibrary(bool bSetupUndoRedo, const FObjectInitializer* ObjectInitializer = nullptr, bool bCreateController = true);
+	URigVMFunctionLibrary* GetOrCreateFunctionLibrary(TSubclassOf<URigVMSchema> InSchemaClass, bool bSetupUndoRedo, const FObjectInitializer* ObjectInitializer = nullptr, bool bCreateController = true);
 	TArray<FName> GetEntryNames() const;
-	UScriptStruct* GetExecuteContextStruct() const;
-	void SetExecuteContextStruct(UScriptStruct* InExecuteContextStruct);
+	UE_DEPRECATED(5.5, "Please use GetDefaultExecuteContextStruct or get an execute context from a graph/controller schema.")
+	UScriptStruct* GetExecuteContextStruct() const { return GetDefaultExecuteContextStruct(); }
+	UScriptStruct* GetDefaultExecuteContextStruct() const;
+	UE_DEPRECATED(5.5, "Please use SetDefaultExecuteContextStruct or set an execute context on a graph/controller schema.")
+	void SetExecuteContextStruct(UScriptStruct* InExecuteContextStruct) { SetDefaultExecuteContextStruct(InExecuteContextStruct); }
+	void SetDefaultExecuteContextStruct(UScriptStruct* InExecuteContextStruct);
 
 	FRigVMGetFocusedGraph& OnGetFocusedGraph() { return OnGetFocusedGraphDelegate;}
 	const FRigVMGetFocusedGraph& OnGetFocusedGraph() const { return OnGetFocusedGraphDelegate; }
@@ -207,6 +217,7 @@ public:
 
 	URigVMGraph* AddModel(const FString InName, bool bSetupUndoRedo, bool bPrintPythonCommand);
 	URigVMGraph* AddModel(const FName& InName, bool bSetupUndoRedo, const FObjectInitializer* ObjectInitializer = nullptr, bool bCreateController = true);
+	URigVMGraph* AddModel(const FName& InName, TSubclassOf<URigVMSchema> InSchemaClass, bool bSetupUndoRedo, const FObjectInitializer* ObjectInitializer = nullptr, bool bCreateController = true);
 	void AddModel(URigVMGraph* InModel, bool bCreateController);
 	bool RemoveModel(FString InName, bool bSetupUndoRedo, bool bPrintPythonCommand);
 	bool RemoveModel(const FString& InNodePathOrName, bool bSetupUndoRedo);
@@ -278,15 +289,10 @@ private:
 	URigVMActionStack* GetOrCreateActionStack();
 	void ResetActionStack();
 
-	void SetSchema(URigVMSchema* InSchema);
-
 	FRigVMGetFocusedGraph OnGetFocusedGraphDelegate;
 
 	UPROPERTY(transient)
-	TObjectPtr<URigVMSchema> SchemaPtr;
-
-	UPROPERTY(transient)
-	TSubclassOf<URigVMSchema> SchemaClass;
+	TSubclassOf<URigVMSchema> DefaultSchemaClass;
 
 	UPROPERTY(transient)
 	TSubclassOf<URigVMController> ControllerClass;
diff --git a/Engine/Plugins/Runtime/RigVM/Source/RigVMDeveloper/Public/RigVMModel/RigVMController.h b/Engine/Plugins/Runtime/RigVM/Source/RigVMDeveloper/Public/RigVMModel/RigVMController.h
--- a/Engine/Plugins/Runtime/RigVM/Source/RigVMDeveloper/Public/RigVMModel/RigVMController.h
+++ b/Engine/Plugins/Runtime/RigVM/Source/RigVMDeveloper/Public/RigVMModel/RigVMController.h
@@ -265,11 +265,15 @@ public:
 
 	// Returns the schema used by this controller
 	UFUNCTION(BlueprintCallable, Category = RigVMController)
-	URigVMSchema* GetSchema() const { return SchemaPtr; }
+	URigVMSchema* GetSchema() const;
 
-	// Sets the schema on the controller
+	UE_DEPRECATED(5.5, "Please use SetSchemaClass instead.")
+	UFUNCTION(BlueprintCallable, Category = RigVMController, meta=(DeprecatedFunction, DeprecationMessage="Function has been deprecated, please use SetSchemaClass instead."))
+	void SetSchema(URigVMSchema* InSchema) { check(InSchema); SetSchemaClass(InSchema->GetClass()); }
+
+	// Sets the schema class on the controller
 	UFUNCTION(BlueprintCallable, Category = RigVMController)
-	void SetSchema(URigVMSchema* InSchema);
+	void SetSchemaClass(TSubclassOf<URigVMSchema> InSchemaClass) { SchemaClass = InSchemaClass; }
 
 	// Pushes a new graph to the stack
 	// This causes a GraphChanged modified event.
@@ -1407,7 +1411,7 @@ private:
 	TArray<TObjectPtr<URigVMGraph>> Graphs;
 
 	UPROPERTY(transient, DuplicateTransient)
-	TObjectPtr<URigVMSchema> SchemaPtr;
+	TSubclassOf<URigVMSchema> SchemaClass;
 
 	mutable TWeakObjectPtr<URigVMActionStack> WeakActionStack;
 	mutable FDelegateHandle ActionStackHandle;
diff --git a/Engine/Plugins/Runtime/RigVM/Source/RigVMDeveloper/Public/RigVMModel/RigVMGraph.h b/Engine/Plugins/Runtime/RigVM/Source/RigVMDeveloper/Public/RigVMModel/RigVMGraph.h
--- a/Engine/Plugins/Runtime/RigVM/Source/RigVMDeveloper/Public/RigVMModel/RigVMGraph.h
+++ b/Engine/Plugins/Runtime/RigVM/Source/RigVMDeveloper/Public/RigVMModel/RigVMGraph.h
@@ -144,6 +144,18 @@ public:
     UFUNCTION(BlueprintCallable, Category = RigVMGraph)
     TArray<FRigVMGraphVariableDescription> GetOutputArguments() const;
 
+	// Returns the schema used by this graph
+	UFUNCTION(BlueprintCallable, Category = RigVMGraph)
+	URigVMSchema* GetSchema() const;
+
+	// Returns the schema class used by this graph
+	UFUNCTION(BlueprintCallable, Category = RigVMGraph)
+	TSubclassOf<URigVMSchema> GetSchemaClass() const { return SchemaClass; }
+
+	// Sets the schema class on the graph
+	UFUNCTION(BlueprintCallable, Category = RigVMGraph)
+	void SetSchemaClass(TSubclassOf<URigVMSchema> InSchemaClass);
+
 	// Returns the modified event, which can be used to 
 	// subscribe to changes happening within the Graph.
 	FRigVMGraphModifiedEvent& OnModified();
@@ -205,6 +217,9 @@ private:
 	UPROPERTY()
 	TArray<FRigVMGraphVariableDescription> LocalVariables;
 
+	UPROPERTY(transient, DuplicateTransient)
+	TSubclassOf<URigVMSchema> SchemaClass;
+
 	bool IsNameAvailable(const FString& InName) const;
 
 	friend class URigVMController;
diff --git a/Engine/Source/Runtime/Core/Public/UObject/FortniteMainBranchObjectVersion.h b/Engine/Source/Runtime/Core/Public/UObject/FortniteMainBranchObjectVersion.h
--- a/Engine/Source/Runtime/Core/Public/UObject/FortniteMainBranchObjectVersion.h
+++ b/Engine/Source/Runtime/Core/Public/UObject/FortniteMainBranchObjectVersion.h
@@ -416,7 +416,10 @@ struct FFortniteMainBranchObjectVersion
 
 		// Skeletal Mesh optionally cooks half edge data per lod
 		SkeletalHalfEdgeData,
-		
+
+		// Combine graph contexts for AnimNext graphs
+		AnimNextCombineGraphContexts,
+
 		// -----<new versions can be added above this line>-------------------------------------------------
 		VersionPlusOne,
 		LatestVersion = VersionPlusOne - 1
