Control Rig: FIx crash when updating profiling info
#rnx

diff --git a/Engine/Plugins/Runtime/RigVM/Source/RigVM/Public/RigVMCore/RigVMProfilingInfo.h b/Engine/Plugins/Runtime/RigVM/Source/RigVM/Public/RigVMCore/RigVMProfilingInfo.h
--- a/Engine/Plugins/Runtime/RigVM/Source/RigVM/Public/RigVMCore/RigVMProfilingInfo.h
+++ b/Engine/Plugins/Runtime/RigVM/Source/RigVM/Public/RigVMCore/RigVMProfilingInfo.h
@@ -103,9 +103,24 @@ struct RIGVM_API FRigVMProfilingInfo
 	inline void AddOverallCycles(uint64 Cycles) { OverallCycles += Cycles; }
 
 	inline void ResetInstructionCyclesDuringLastRun(int32 NewSize = 0) { InstructionCyclesDuringLastRun.Reset(NewSize); }
-	inline uint64 GetInstructionCyclesDuringLastRun(int32 InstructionIndex) const { return InstructionCyclesDuringLastRun.IsValidIndex(InstructionIndex) ? InstructionCyclesDuringLastRun[InstructionIndex] : UINT64_MAX; }
-	inline void SetInstructionCyclesDuringLastRun(int32 InstructionIndex, uint64 CyclesDuringLastRun) { InstructionCyclesDuringLastRun[InstructionIndex] = CyclesDuringLastRun; }
-	inline void AddInstructionCyclesDuringLastRun(int32 InstructionIndex, uint64 CyclesDuringLastRun) { InstructionCyclesDuringLastRun[InstructionIndex] += CyclesDuringLastRun; }
+	inline uint64 GetInstructionCyclesDuringLastRun(int32 InstructionIndex) const
+	{
+		return InstructionCyclesDuringLastRun.IsValidIndex(InstructionIndex) ? InstructionCyclesDuringLastRun[InstructionIndex] : UINT64_MAX;
+	}
+	inline void SetInstructionCyclesDuringLastRun(int32 InstructionIndex, uint64 CyclesDuringLastRun)
+	{
+		if (InstructionCyclesDuringLastRun.IsValidIndex(InstructionIndex))
+		{
+			InstructionCyclesDuringLastRun[InstructionIndex] = CyclesDuringLastRun;
+		}
+	}
+	inline void AddInstructionCyclesDuringLastRun(int32 InstructionIndex, uint64 CyclesDuringLastRun)
+	{
+		if (InstructionCyclesDuringLastRun.IsValidIndex(InstructionIndex))
+		{
+			InstructionCyclesDuringLastRun[InstructionIndex] += CyclesDuringLastRun;
+		}
+	}
 	inline void InitInstructionCyclesDuringLastRunValues(int32 NewSize, uint64 DefaultValue)
 	{
 		InstructionCyclesDuringLastRun.SetNumUninitialized(NewSize);
