Property Editor: Control Rig binding row support

#onbehalfof helge.mathee
#rnx

diff --git a/Engine/Plugins/Animation/ControlRig/Source/ControlRigEditor/Private/ControlRigModuleDetails.cpp b/Engine/Plugins/Animation/ControlRig/Source/ControlRigEditor/Private/ControlRigModuleDetails.cpp
--- a/Engine/Plugins/Animation/ControlRig/Source/ControlRigEditor/Private/ControlRigModuleDetails.cpp
+++ b/Engine/Plugins/Animation/ControlRig/Source/ControlRigEditor/Private/ControlRigModuleDetails.cpp
@@ -627,22 +627,22 @@ void FRigModuleInstanceDetails::CustomizeDetails(IDetailLayoutBuilder& DetailBui
 					})
 				);
 
-				// todo: remove the original row.
+				TSharedPtr<SWidget> ValueWidget = DefaultProperty->CreatePropertyValueWidgetWithCustomization(DetailBuilder.GetDetailsView());
 
-				Category.AddCustomRow(DefaultProperty->GetPropertyDisplayName())
+				const bool bShowChildren = true;
+				Category.AddProperty(DefaultProperty).CustomWidget(bShowChildren)
 				.NameContent()
 				[
 					DefaultProperty->CreatePropertyNameWidget()
 				]
 
-				// note: this doesn't work for some reason. seeking help from the editor team
 				.ValueContent()
 				[
-					DefaultProperty->CreatePropertyValueWidget()
+					ValueWidget ? ValueWidget.ToSharedRef() : SNullWidget::NullWidget
 					// todo: if the property is bound / or partially bound
 					// mark the property value widget as disabled / read only.
 				]
-				
+
 				.ExtensionContent()
 				[
 					PropertyAccessEditor.MakePropertyBindingWidget(nullptr, BindingArgs)
diff --git a/Engine/Plugins/Runtime/RigVM/Source/RigVMEditor/Private/Editor/RigVMGraphDetailCustomization.cpp b/Engine/Plugins/Runtime/RigVM/Source/RigVMEditor/Private/Editor/RigVMGraphDetailCustomization.cpp
--- a/Engine/Plugins/Runtime/RigVM/Source/RigVMEditor/Private/Editor/RigVMGraphDetailCustomization.cpp
+++ b/Engine/Plugins/Runtime/RigVM/Source/RigVMEditor/Private/Editor/RigVMGraphDetailCustomization.cpp
@@ -1830,6 +1830,620 @@ void FRigVMWrappedNodeDetailCustomization::CustomizeLiveValues(IDetailLayoutBuil
 	*/
 }
 
+template<typename VectorType, int32 NumberOfComponents>
+void FRigVMGraphMathTypeDetailCustomization::MakeVectorHeaderRow(TSharedRef<class IPropertyHandle> InPropertyHandle, FDetailWidgetRow& HeaderRow, IPropertyTypeCustomizationUtils& StructCustomizationUtils)
+{
+	typedef typename VectorType::FReal NumericType;
+	typedef SNumericVectorInputBox<NumericType, VectorType, NumberOfComponents> SLocalVectorInputBox;
+
+	FEditPropertyChain PropertyChain;
+	TArray<int32> PropertyArrayIndices;
+	bool bEnabled;
+	if (!GetPropertyChain(InPropertyHandle, PropertyChain, PropertyArrayIndices, bEnabled))
+	{
+		return;
+	}
+
+	typename SLocalVectorInputBox::FArguments Args;
+	Args.Font(IDetailLayoutBuilder::GetDetailFont());
+	Args.IsEnabled(bEnabled);
+	Args.AllowSpin(true);
+	Args.SpinDelta(0.01f);
+	Args.bColorAxisLabels(true);
+	Args.X_Lambda([this, InPropertyHandle]()
+		{
+			return GetVectorComponent<VectorType, NumericType>(InPropertyHandle, 0);
+		});
+	Args.OnXChanged_Lambda([this, InPropertyHandle](NumericType Value)
+		{
+			OnVectorComponentChanged<VectorType, NumericType>(InPropertyHandle, 0, Value, false);
+		});
+	Args.OnXCommitted_Lambda([this, InPropertyHandle](NumericType Value, ETextCommit::Type CommitType)
+		{
+			OnVectorComponentChanged<VectorType, NumericType>(InPropertyHandle, 0, Value, true, CommitType);
+		});
+	Args.Y_Lambda([this, InPropertyHandle]()
+		{
+			return GetVectorComponent<VectorType, NumericType>(InPropertyHandle, 1);
+		});
+	Args.OnYChanged_Lambda([this, InPropertyHandle](NumericType Value)
+		{
+			OnVectorComponentChanged<VectorType, NumericType>(InPropertyHandle, 1, Value, false);
+		});
+	Args.OnYCommitted_Lambda([this, InPropertyHandle](NumericType Value, ETextCommit::Type CommitType)
+		{
+			OnVectorComponentChanged<VectorType, NumericType>(InPropertyHandle, 1, Value, true, CommitType);
+		});
+
+	ExtendVectorArgs<VectorType>(InPropertyHandle, &Args);
+
+	HeaderRow
+		.IsEnabled(bEnabled)
+		.NameContent()
+		[
+			InPropertyHandle->CreatePropertyNameWidget()
+		]
+	.ValueContent()
+		.MinDesiredWidth(375.f)
+		.MaxDesiredWidth(375.f)
+		.HAlign(HAlign_Left)
+		[
+			SArgumentNew(Args, SLocalVectorInputBox)
+		];
+}
+
+template<typename RotationType>
+void FRigVMGraphMathTypeDetailCustomization::MakeRotationHeaderRow(TSharedRef<class IPropertyHandle> InPropertyHandle, FDetailWidgetRow& HeaderRow, IPropertyTypeCustomizationUtils& StructCustomizationUtils)
+{
+	FEditPropertyChain PropertyChain;
+	TArray<int32> PropertyArrayIndices;
+	bool bEnabled;
+	if (!GetPropertyChain(InPropertyHandle, PropertyChain, PropertyArrayIndices, bEnabled))
+	{
+		return;
+	}
+		
+	typedef typename RotationType::FReal NumericType;
+	typedef SAdvancedRotationInputBox<NumericType> SLocalRotationInputBox;
+	typename SLocalRotationInputBox::FArguments Args;
+	Args.Font(IDetailLayoutBuilder::GetDetailFont());
+	Args.IsEnabled(bEnabled);
+	Args.AllowSpin(true);
+	Args.bColorAxisLabels(true);
+
+	ExtendRotationArgs<RotationType>(InPropertyHandle, &Args);
+
+	HeaderRow
+		.IsEnabled(bEnabled)
+		.NameContent()
+		[
+			InPropertyHandle->CreatePropertyNameWidget()
+		]
+	.ValueContent()
+		.MinDesiredWidth(375.f)
+		.MaxDesiredWidth(375.f)
+		.HAlign(HAlign_Left)
+		[
+			SArgumentNew(Args, SLocalRotationInputBox)
+		];
+}
+
+template <typename TransformType>
+void FRigVMGraphMathTypeDetailCustomization::ConfigureTransformWidgetArgs(TSharedRef<IPropertyHandle> InPropertyHandle, typename SAdvancedTransformInputBox<TransformType>::FArguments& WidgetArgs)
+{
+	FEditPropertyChain PropertyChain;
+	TArray<int32> PropertyArrayIndices;
+	bool bEnabled;
+	if (!GetPropertyChain(InPropertyHandle, PropertyChain, PropertyArrayIndices, bEnabled))
+	{
+		return;
+	}
+	
+	typedef typename TransformType::FReal FReal;
+	WidgetArgs.IsEnabled(bEnabled);
+	WidgetArgs.AllowEditRotationRepresentation(true);
+	WidgetArgs.UseQuaternionForRotation(IsQuaternionBasedRotation<TransformType>());
+
+	static TransformType Identity = TransformType::Identity;
+
+	UObject* DefaultObject = !ObjectsBeingCustomized.IsEmpty() ? ObjectsBeingCustomized[0]->GetClass()->GetDefaultObject() :
+	!StructsBeingCustomized.IsEmpty() ? StructsBeingCustomized[0]->GetStruct()->GetClass() : nullptr;
+
+	if (!DefaultObject)
+	{
+		return;
+	}
+	
+	TransformType DefaultValue = ContainerMemoryBlockToValueRef<TransformType>((uint8*)DefaultObject, Identity, PropertyChain, PropertyArrayIndices);
+	
+	WidgetArgs.DiffersFromDefault_Lambda([this, InPropertyHandle, DefaultValue](ESlateTransformComponent::Type InTransformComponent) -> bool
+	{
+		if (ObjectsBeingCustomized.IsEmpty())
+		{
+			return false;
+		}
+		
+		FEditPropertyChain PropertyChain;
+		TArray<int32> PropertyArrayIndices;
+		bool bEnabled;
+		if (!GetPropertyChain(InPropertyHandle, PropertyChain, PropertyArrayIndices, bEnabled))
+		{
+			return false;
+		}
+		
+		for(const TWeakObjectPtr<UObject>& Object : ObjectsBeingCustomized)
+		{
+			if(Object.Get() && InPropertyHandle->IsValidHandle())
+			{
+				const TransformType& Transform = ContainerMemoryBlockToValueRef<TransformType>((uint8*)Object.Get(), Identity, PropertyChain, PropertyArrayIndices);
+
+				switch(InTransformComponent)
+				{
+					case ESlateTransformComponent::Location:
+					{
+						if(!Transform.GetLocation().Equals(DefaultValue.GetLocation()))
+						{
+							return true;
+						}
+						break;
+					}
+					case ESlateTransformComponent::Rotation:
+					{
+						if(!Transform.Rotator().Equals(DefaultValue.Rotator()))
+						{
+							return true;
+						}
+						break;
+					}
+					case ESlateTransformComponent::Scale:
+					{
+						if(!Transform.GetScale3D().Equals(DefaultValue.GetScale3D()))
+						{
+							return true;
+						}
+						break;
+					}
+					default:
+					{
+						break;
+					}
+				}
+			}
+		}
+		return false;
+	});
+
+	WidgetArgs.OnGetNumericValue_Lambda([this, InPropertyHandle](
+		ESlateTransformComponent::Type InTransformComponent,
+		ESlateRotationRepresentation::Type InRotationRepresentation,
+		ESlateTransformSubComponent::Type InTransformSubComponent) -> TOptional<FReal>
+	{
+		TOptional<FReal> Result;
+		FEditPropertyChain PropertyChain;
+		TArray<int32> PropertyArrayIndices;
+		bool bEnabled;
+		if (!GetPropertyChain(InPropertyHandle, PropertyChain, PropertyArrayIndices, bEnabled))
+		{
+			return Result;
+		}
+
+		const TArray<uint8*> MemoryBlocks = GetMemoryBeingCustomized();
+		for(uint8* MemoryBlock: MemoryBlocks)
+		{
+			if(MemoryBlock)
+			{
+				const TransformType& Transform = ContainerMemoryBlockToValueRef<TransformType>(MemoryBlock, Identity, PropertyChain, PropertyArrayIndices);
+				
+				TOptional<FReal> Value = SAdvancedTransformInputBox<TransformType>::GetNumericValueFromTransform(
+					Transform,
+					InTransformComponent,
+					InRotationRepresentation,
+					InTransformSubComponent
+					);
+
+				if(Value.IsSet())
+				{
+					if(Result.IsSet())
+					{
+						if(!FMath::IsNearlyEqual(Result.GetValue(), Value.GetValue()))
+						{
+							return TOptional<FReal>();
+						}
+					}
+					else
+					{
+						Result = Value;
+					}
+				}
+			}
+		}
+		return Result;
+	});
+	
+
+	auto OnNumericValueChanged = [this, InPropertyHandle](
+		ESlateTransformComponent::Type InTransformComponent, 
+		ESlateRotationRepresentation::Type InRotationRepresentation, 
+		ESlateTransformSubComponent::Type InSubComponent,
+		FReal InValue,
+		bool bIsCommit,
+		ETextCommit::Type InCommitType = ETextCommit::Default)
+	{
+		FEditPropertyChain PropertyChain;
+		TArray<int32> PropertyArrayIndices;
+		bool bEnabled;
+		if (!GetPropertyChain(InPropertyHandle, PropertyChain, PropertyArrayIndices, bEnabled))
+		{
+			return;
+		}
+		
+		TArray<UObject*> ObjectsView;
+		for(int32 Index = 0; Index < ObjectsBeingCustomized.Num(); Index++)
+		{
+			const TWeakObjectPtr<UObject>& Object = ObjectsBeingCustomized[Index];
+			if (Object.Get())
+			{
+				ObjectsView.Add(Object.Get());
+			}
+		}
+		FPropertyChangedEvent PropertyChangedEvent(InPropertyHandle->GetProperty(), bIsCommit ? EPropertyChangeType::ValueSet : EPropertyChangeType::Interactive, ObjectsView);
+		FPropertyChangedChainEvent PropertyChangedChainEvent(PropertyChain, PropertyChangedEvent);
+
+		URigVMController* Controller = nullptr;
+		if(BlueprintBeingCustomized && GraphBeingCustomized)
+		{
+			Controller = BlueprintBeingCustomized->GetController(GraphBeingCustomized);
+			if(bIsCommit)
+			{
+				Controller->OpenUndoBracket(FString::Printf(TEXT("Set %s"), *InPropertyHandle->GetProperty()->GetName()));
+			}
+		}
+
+		for(int32 Index = 0; Index < ObjectsBeingCustomized.Num(); Index++)
+		{
+			const TWeakObjectPtr<UObject>& Object = ObjectsBeingCustomized[Index];
+			if(Object.Get() && InPropertyHandle->IsValidHandle())
+			{
+				TransformType& Transform = ContainerMemoryBlockToValueRef<TransformType>((uint8*)Object.Get(), Identity, PropertyChain, PropertyArrayIndices);
+				TransformType PreviousTransform = Transform;
+				
+				SAdvancedTransformInputBox<TransformType>::ApplyNumericValueChange(
+					Transform,
+					InValue,
+					InTransformComponent,
+					InRotationRepresentation,
+					InSubComponent);
+
+				if(!PreviousTransform.Equals(Transform))
+				{
+					Object->PostEditChangeChainProperty(PropertyChangedChainEvent);
+					InPropertyHandle->NotifyPostChange(PropertyChangedEvent.ChangeType);
+				}
+			}
+		}
+
+		if(Controller && bIsCommit)
+		{
+			Controller->CloseUndoBracket();
+		}
+	};
+
+	WidgetArgs.OnNumericValueChanged_Lambda([OnNumericValueChanged](
+		ESlateTransformComponent::Type InTransformComponent, 
+		ESlateRotationRepresentation::Type InRotationRepresentation, 
+		ESlateTransformSubComponent::Type InSubComponent,
+		FReal InValue)
+	{
+		OnNumericValueChanged(InTransformComponent, InRotationRepresentation, InSubComponent, InValue, false);
+	});
+
+	WidgetArgs.OnNumericValueCommitted_Lambda([OnNumericValueChanged](
+		ESlateTransformComponent::Type InTransformComponent, 
+		ESlateRotationRepresentation::Type InRotationRepresentation, 
+		ESlateTransformSubComponent::Type InSubComponent,
+		FReal InValue, 
+		ETextCommit::Type InCommitType)
+	{
+		OnNumericValueChanged(InTransformComponent, InRotationRepresentation, InSubComponent, InValue, true, InCommitType);
+	});
+
+	WidgetArgs.OnResetToDefault_Lambda([this, DefaultValue, InPropertyHandle](ESlateTransformComponent::Type InTransformComponent)
+	{
+		if (ObjectsBeingCustomized.IsEmpty())
+		{
+			return;
+		}
+		
+		FEditPropertyChain PropertyChain;
+		TArray<int32> PropertyArrayIndices;
+		bool bEnabled;
+		if (!GetPropertyChain(InPropertyHandle, PropertyChain, PropertyArrayIndices, bEnabled))
+		{
+			return;
+		}
+		
+		URigVMController* Controller = nullptr;
+		if(BlueprintBeingCustomized && GraphBeingCustomized)
+		{
+			Controller = BlueprintBeingCustomized->GetController(GraphBeingCustomized);
+			if(Controller)
+			{
+				Controller->OpenUndoBracket(FString::Printf(TEXT("Reset %s to Default"), *InPropertyHandle->GetProperty()->GetName()));
+			}
+		}
+
+		TArray<UObject*> ObjectsView;
+		for(int32 Index = 0; Index < ObjectsBeingCustomized.Num(); Index++)
+		{
+			const TWeakObjectPtr<UObject>& Object = ObjectsBeingCustomized[Index];
+			if (Object.Get())
+			{
+				ObjectsView.Add(Object.Get());
+			}
+		}
+		FPropertyChangedEvent PropertyChangedEvent(InPropertyHandle->GetProperty(), EPropertyChangeType::ValueSet, ObjectsView);
+		FPropertyChangedChainEvent PropertyChangedChainEvent(PropertyChain, PropertyChangedEvent);
+		
+		for(int32 Index = 0; Index < ObjectsBeingCustomized.Num(); Index++)
+		{
+			const TWeakObjectPtr<UObject>& Object = ObjectsBeingCustomized[Index];
+			if(Object.Get() && InPropertyHandle->IsValidHandle())
+			{
+				TransformType& Transform = ContainerMemoryBlockToValueRef<TransformType>((uint8*)Object.Get(), Identity, PropertyChain, PropertyArrayIndices);
+				TransformType PreviousTransform = Transform;
+
+				switch(InTransformComponent)
+				{
+					case ESlateTransformComponent::Location:
+					{
+						Transform.SetLocation(DefaultValue.GetLocation());
+						break;
+					}
+					case ESlateTransformComponent::Rotation:
+					{
+						Transform.SetRotation(DefaultValue.GetRotation());
+						break;
+					}
+					case ESlateTransformComponent::Scale:
+					{
+						Transform.SetScale3D(DefaultValue.GetScale3D());
+						break;
+					}
+					case ESlateTransformComponent::Max:
+					default:
+					{
+						Transform.SetLocation(DefaultValue.GetLocation());
+						break;
+					}
+				}
+
+				
+				if(!PreviousTransform.Equals(Transform))
+				{
+					Object->PostEditChangeChainProperty(PropertyChangedChainEvent);
+					InPropertyHandle->NotifyPostChange(PropertyChangedEvent.ChangeType);
+				}
+			}
+		}
+
+		if(Controller)
+		{
+			Controller->CloseUndoBracket();
+		}
+	});
+
+	WidgetArgs.OnCopyToClipboard_Lambda([this, InPropertyHandle](
+		ESlateTransformComponent::Type InComponent
+		)
+	{
+		TOptional<FReal> Result;
+		FEditPropertyChain PropertyChain;
+		TArray<int32> PropertyArrayIndices;
+		bool bEnabled;
+		if (!GetPropertyChain(InPropertyHandle, PropertyChain, PropertyArrayIndices, bEnabled))
+		{
+			return;
+		}
+
+		const TArray<uint8*> MemoryBlocks = GetMemoryBeingCustomized();
+		for(uint8* MemoryBlock: MemoryBlocks)
+		{
+			if(MemoryBlock)
+			{
+				const TransformType& Transform = ContainerMemoryBlockToValueRef<TransformType>(MemoryBlock, Identity, PropertyChain, PropertyArrayIndices);
+				FString Content;
+				switch(InComponent)
+				{
+					case ESlateTransformComponent::Location:
+					{
+						const FVector Data = Transform.GetLocation();
+						TBaseStructure<FVector>::Get()->ExportText(Content, &Data, &Data, nullptr, PPF_None, nullptr);
+						break;
+					}
+					case ESlateTransformComponent::Rotation:
+					{
+						const FRotator Data = Transform.Rotator();
+						TBaseStructure<FRotator>::Get()->ExportText(Content, &Data, &Data, nullptr, PPF_None, nullptr);
+						break;
+					}
+					case ESlateTransformComponent::Scale:
+					{
+						const FVector Data = Transform.GetScale3D();
+						TBaseStructure<FVector>::Get()->ExportText(Content, &Data, &Data, nullptr, PPF_None, nullptr);
+						break;
+					}
+					case ESlateTransformComponent::Max:
+					default:
+					{
+						TBaseStructure<TransformType>::Get()->ExportText(Content, &Transform, &Transform, nullptr, PPF_None, nullptr);
+						break;
+					}
+				}
+
+				if(!Content.IsEmpty())
+				{
+					FPlatformApplicationMisc::ClipboardCopy(*Content);
+				}
+			}
+		}
+	});
+
+	WidgetArgs.OnPasteFromClipboard_Lambda([this, InPropertyHandle, OnNumericValueChanged](
+		ESlateTransformComponent::Type InComponent
+		)
+	{
+		FString Content;
+		FPlatformApplicationMisc::ClipboardPaste(Content);
+	
+		if(Content.IsEmpty())
+		{
+			return;
+		}
+
+		if (ObjectsBeingCustomized.IsEmpty())
+		{
+			return;
+		}
+
+		TOptional<FReal> Result;
+		FEditPropertyChain PropertyChain;
+		TArray<int32> PropertyArrayIndices;
+		bool bEnabled;
+		if (!GetPropertyChain(InPropertyHandle, PropertyChain, PropertyArrayIndices, bEnabled))
+		{
+			return;
+		}
+
+		TArray<UObject*> ObjectsView;
+		for(int32 Index = 0; Index < ObjectsBeingCustomized.Num(); Index++)
+		{
+			const TWeakObjectPtr<UObject>& Object = ObjectsBeingCustomized[Index];
+			if (Object.Get())
+			{
+				ObjectsView.Add(Object.Get());
+			}
+		}
+		FPropertyChangedEvent PropertyChangedEvent(InPropertyHandle->GetProperty(), EPropertyChangeType::ValueSet, ObjectsView);
+		FPropertyChangedChainEvent PropertyChangedChainEvent(PropertyChain, PropertyChangedEvent);
+
+		URigVMController* Controller = nullptr;
+		if(BlueprintBeingCustomized && GraphBeingCustomized)
+		{
+			Controller = BlueprintBeingCustomized->GetController(GraphBeingCustomized);
+			Controller->OpenUndoBracket(FString::Printf(TEXT("Set %s"), *InPropertyHandle->GetProperty()->GetName()));
+		}
+		
+		for(const TWeakObjectPtr<UObject>& Object : ObjectsBeingCustomized)
+		{
+			if(Object.Get() && InPropertyHandle->IsValidHandle())
+			{
+				TransformType& Transform = ContainerMemoryBlockToValueRef<TransformType>((uint8*)Object.Get(), Identity, PropertyChain, PropertyArrayIndices);
+				const TransformType PreviousTransform = Transform;
+
+				// Apply the new value
+				{
+					class FRigPasteTransformWidgetErrorPipe : public FOutputDevice
+					{
+					public:
+				
+						int32 NumErrors;
+				
+						FRigPasteTransformWidgetErrorPipe()
+							: FOutputDevice()
+							, NumErrors(0)
+						{
+						}
+				
+						virtual void Serialize(const TCHAR* V, ELogVerbosity::Type Verbosity, const class FName& Category) override
+						{
+							UE_LOG(LogRigVM, Error, TEXT("Error Pasting to Widget: %s"), V);
+							NumErrors++;
+						}
+					};
+				
+					FRigPasteTransformWidgetErrorPipe ErrorPipe;
+					
+					switch(InComponent)
+					{
+						case ESlateTransformComponent::Location:
+						{
+							FVector Data = Transform.GetLocation();
+							TBaseStructure<FVector>::Get()->ImportText(*Content, &Data, nullptr, PPF_None, &ErrorPipe, TBaseStructure<FVector>::Get()->GetName(), true);
+							Transform.SetLocation(Data);
+							break;
+						}
+						case ESlateTransformComponent::Rotation:
+						{
+							FRotator Data = Transform.Rotator();
+							TBaseStructure<FRotator>::Get()->ImportText(*Content, &Data, nullptr, PPF_None, &ErrorPipe, TBaseStructure<FRotator>::Get()->GetName(), true);
+							FQuat Quat = Data.Quaternion();
+							Transform.SetRotation(Quat);
+							
+							break;
+						}
+						case ESlateTransformComponent::Scale:
+						{
+							FVector Data = Transform.GetScale3D();
+							TBaseStructure<FVector>::Get()->ImportText(*Content, &Data, nullptr, PPF_None, &ErrorPipe, TBaseStructure<FVector>::Get()->GetName(), true);
+							Transform.SetScale3D(Data);
+							break;
+						}
+						case ESlateTransformComponent::Max:
+						default:
+						{
+							TBaseStructure<TransformType>::Get()->ImportText(*Content, &Transform, nullptr, PPF_None, &ErrorPipe, TBaseStructure<TransformType>::Get()->GetName(), true);
+							break;
+						}
+					}
+				
+					if(ErrorPipe.NumErrors == 0 && !PreviousTransform.Equals(Transform))
+					{
+						Object->PostEditChangeChainProperty(PropertyChangedChainEvent);
+						InPropertyHandle->NotifyPostChange(PropertyChangedEvent.ChangeType);
+					}
+				}
+			}
+		}
+		
+		if (Controller)
+		{
+			Controller->CloseUndoBracket();
+		}
+	});
+}
+
+template<typename TransformType>
+void FRigVMGraphMathTypeDetailCustomization::MakeTransformHeaderRow(TSharedRef<class IPropertyHandle> InPropertyHandle, FDetailWidgetRow& HeaderRow, IPropertyTypeCustomizationUtils& StructCustomizationUtils)
+{
+	typename SAdvancedTransformInputBox<TransformType>::FArguments WidgetArgs;
+	ConfigureTransformWidgetArgs<TransformType>(InPropertyHandle, WidgetArgs);
+
+	SAdvancedTransformInputBox<TransformType>::ConfigureHeader(HeaderRow, InPropertyHandle->GetPropertyDisplayName(), InPropertyHandle->GetToolTipText(), WidgetArgs);
+	SAdvancedTransformInputBox<TransformType>::ConfigureComponentWidgetRow(HeaderRow, ESlateTransformComponent::Max, WidgetArgs);
+}
+	
+template<typename TransformType>
+void FRigVMGraphMathTypeDetailCustomization::MakeTransformChildren(TSharedRef<class IPropertyHandle> InPropertyHandle, IDetailChildrenBuilder& StructBuilder, IPropertyTypeCustomizationUtils& StructCustomizationUtils)
+{
+	FDetailWidgetRow* LocationRow = nullptr;
+	FDetailWidgetRow* RotationRow = nullptr;
+	FDetailWidgetRow* ScaleRow = nullptr;
+
+	typename SAdvancedTransformInputBox<TransformType>::FArguments WidgetArgs;
+	ConfigureTransformWidgetArgs<TransformType>(InPropertyHandle, WidgetArgs);
+
+	const FString StandaloneWidgetMetadata = InPropertyHandle->GetMetaData(TEXT("SStandaloneCustomizedValueWidget"));
+	const bool IsUsingStandaloneWidget = StandaloneWidgetMetadata.Equals(TEXT("True"));
+	
+	LocationRow = &StructBuilder.AddCustomRow(LOCTEXT("TransformLocation", "Location"));
+	RotationRow = &StructBuilder.AddCustomRow(LOCTEXT("TransformRotation", "Rotation"));
+	ScaleRow = &StructBuilder.AddCustomRow(LOCTEXT("TransformScale", "Scale"));
+	
+	SAdvancedTransformInputBox<TransformType>::ConfigureComponentWidgetRow(*LocationRow, ESlateTransformComponent::Location, WidgetArgs);
+	SAdvancedTransformInputBox<TransformType>::ConfigureComponentWidgetRow(*RotationRow, ESlateTransformComponent::Rotation, WidgetArgs);
+	SAdvancedTransformInputBox<TransformType>::ConfigureComponentWidgetRow(*ScaleRow, ESlateTransformComponent::Scale, WidgetArgs);
+}
+
 FRigVMGraphMathTypeDetailCustomization::FRigVMGraphMathTypeDetailCustomization()
 : ScriptStruct(nullptr)
 , BlueprintBeingCustomized(nullptr)
@@ -1864,43 +2478,77 @@ void FRigVMGraphMathTypeDetailCustomization::CustomizeHeader(TSharedRef<IPropert
 	FProperty* Property = InPropertyHandle->GetProperty();
 	const FStructProperty* StructProperty = CastField<FStructProperty>(Property);
 	ScriptStruct = StructProperty->Struct;
-}
-
-void FRigVMGraphMathTypeDetailCustomization::CustomizeChildren(TSharedRef<IPropertyHandle> InPropertyHandle,
-	IDetailChildrenBuilder& StructBuilder, IPropertyTypeCustomizationUtils& StructCustomizationUtils)
-{
-	if(!InPropertyHandle->IsValidHandle())
-	{
-		return;
-	}
 
 	if(ScriptStruct == TBaseStructure<FVector>::Get())
 	{
-		CustomizeVector<FVector, 3>(InPropertyHandle, StructBuilder, StructCustomizationUtils);
+		MakeVectorHeaderRow<FVector, 3>(InPropertyHandle, HeaderRow, StructCustomizationUtils);
 	}
 	else if(ScriptStruct == TBaseStructure<FVector2D>::Get())
 	{
-		CustomizeVector<FVector2D, 2>(InPropertyHandle, StructBuilder, StructCustomizationUtils);
+		MakeVectorHeaderRow<FVector2D, 2>(InPropertyHandle, HeaderRow, StructCustomizationUtils);
 	}
 	else if(ScriptStruct == TBaseStructure<FVector4>::Get())
 	{
-		CustomizeVector<FVector4, 4>(InPropertyHandle, StructBuilder, StructCustomizationUtils);
+		MakeVectorHeaderRow<FVector4, 4>(InPropertyHandle, HeaderRow, StructCustomizationUtils);
 	}
 	else if(ScriptStruct == TBaseStructure<FRotator>::Get())
 	{
-		CustomizeRotation<FRotator>(InPropertyHandle, StructBuilder, StructCustomizationUtils);
+		MakeRotationHeaderRow<FRotator>(InPropertyHandle, HeaderRow, StructCustomizationUtils);
 	}
 	else if(ScriptStruct == TBaseStructure<FQuat>::Get())
 	{
-		CustomizeRotation<FQuat>(InPropertyHandle, StructBuilder, StructCustomizationUtils);
+		MakeRotationHeaderRow<FQuat>(InPropertyHandle, HeaderRow, StructCustomizationUtils);
 	}
 	else if(ScriptStruct == TBaseStructure<FTransform>::Get())
 	{
-		CustomizeTransform<FTransform>(InPropertyHandle, StructBuilder, StructCustomizationUtils);
+		MakeTransformHeaderRow<FTransform>(InPropertyHandle, HeaderRow, StructCustomizationUtils);
+	}
+	else if(ScriptStruct == TBaseStructure<FEulerTransform>::Get())
+	{
+		MakeTransformHeaderRow<FEulerTransform>(InPropertyHandle, HeaderRow, StructCustomizationUtils);
+	}
+}
+
+void FRigVMGraphMathTypeDetailCustomization::CustomizeChildren(TSharedRef<IPropertyHandle> InPropertyHandle,
+	IDetailChildrenBuilder& StructBuilder, IPropertyTypeCustomizationUtils& StructCustomizationUtils)
+{
+	if(!InPropertyHandle->IsValidHandle())
+	{
+		return;
+	}
+	
+	TArray<UObject*> Objects;
+	InPropertyHandle->GetOuterObjects(Objects);
+
+	for (UObject* Object : Objects)
+	{
+		ObjectsBeingCustomized.Add(Object);
+
+		if(BlueprintBeingCustomized == nullptr)
+		{
+			BlueprintBeingCustomized = Object->GetTypedOuter<URigVMBlueprint>();
+		}
+
+		if(GraphBeingCustomized == nullptr)
+		{
+			GraphBeingCustomized = Object->GetTypedOuter<URigVMGraph>();
+		}
+	}
+
+	StructsBeingCustomized.Reset();
+	InPropertyHandle->GetOuterStructs(StructsBeingCustomized);
+
+	FProperty* Property = InPropertyHandle->GetProperty();
+	const FStructProperty* StructProperty = CastField<FStructProperty>(Property);
+	ScriptStruct = StructProperty->Struct;
+
+	if(ScriptStruct == TBaseStructure<FTransform>::Get())
+	{
+		MakeTransformChildren<FTransform>(InPropertyHandle, StructBuilder, StructCustomizationUtils);
 	}
 	else if(ScriptStruct == TBaseStructure<FEulerTransform>::Get())
 	{
-		CustomizeTransform<FEulerTransform>(InPropertyHandle, StructBuilder, StructCustomizationUtils);
+		MakeTransformChildren<FEulerTransform>(InPropertyHandle, StructBuilder, StructCustomizationUtils);
 	}
 }
 
diff --git a/Engine/Plugins/Runtime/RigVM/Source/RigVMEditor/Public/Editor/RigVMGraphDetailCustomization.h b/Engine/Plugins/Runtime/RigVM/Source/RigVMEditor/Public/Editor/RigVMGraphDetailCustomization.h
--- a/Engine/Plugins/Runtime/RigVM/Source/RigVMEditor/Public/Editor/RigVMGraphDetailCustomization.h
+++ b/Engine/Plugins/Runtime/RigVM/Source/RigVMEditor/Public/Editor/RigVMGraphDetailCustomization.h
@@ -471,66 +471,19 @@ protected:
 	void ExtendVectorArgs(TSharedRef<class IPropertyHandle> InPropertyHandle, void* ArgumentsPtr) {}
 
 	template<typename VectorType, int32 NumberOfComponents>
-	void CustomizeVector(TSharedRef<class IPropertyHandle> InPropertyHandle, class IDetailChildrenBuilder& StructBuilder, IPropertyTypeCustomizationUtils& StructCustomizationUtils)
-	{
-		typedef typename VectorType::FReal NumericType;
-		typedef SNumericVectorInputBox<NumericType, VectorType, NumberOfComponents> SLocalVectorInputBox;
+	void MakeVectorHeaderRow(TSharedRef<class IPropertyHandle> InPropertyHandle, FDetailWidgetRow& HeaderRow, IPropertyTypeCustomizationUtils& StructCustomizationUtils);
 
-		FEditPropertyChain PropertyChain;
-        TArray<int32> PropertyArrayIndices;
-        bool bEnabled;
-        if (!GetPropertyChain(InPropertyHandle, PropertyChain, PropertyArrayIndices, bEnabled))
-        {
-        	return;
-        }
+	template<typename RotationType>
+	void MakeRotationHeaderRow(TSharedRef<class IPropertyHandle> InPropertyHandle, FDetailWidgetRow& HeaderRow, IPropertyTypeCustomizationUtils& StructCustomizationUtils);
 
-		typename SLocalVectorInputBox::FArguments Args;
-		Args.Font(IDetailLayoutBuilder::GetDetailFont());
-		Args.IsEnabled(bEnabled);
-		Args.AllowSpin(true);
-		Args.SpinDelta(0.01f);
-		Args.bColorAxisLabels(true);
-		Args.X_Lambda([this, InPropertyHandle]()
-		{
-			return GetVectorComponent<VectorType, NumericType>(InPropertyHandle, 0);
-		});
-		Args.OnXChanged_Lambda([this, InPropertyHandle](NumericType Value)
-		{
-			OnVectorComponentChanged<VectorType, NumericType>(InPropertyHandle, 0, Value, false);
-		});
-		Args.OnXCommitted_Lambda([this, InPropertyHandle](NumericType Value, ETextCommit::Type CommitType)
-		{
-			OnVectorComponentChanged<VectorType, NumericType>(InPropertyHandle, 0, Value, true, CommitType);
-		});
-		Args.Y_Lambda([this, InPropertyHandle]()
-		{
-			return GetVectorComponent<VectorType, NumericType>(InPropertyHandle, 1);
-		});
-		Args.OnYChanged_Lambda([this, InPropertyHandle](NumericType Value)
-		{
-			OnVectorComponentChanged<VectorType, NumericType>(InPropertyHandle, 1, Value, false);
-		});
-		Args.OnYCommitted_Lambda([this, InPropertyHandle](NumericType Value, ETextCommit::Type CommitType)
-		{
-			OnVectorComponentChanged<VectorType, NumericType>(InPropertyHandle, 1, Value, true, CommitType);
-		});
-
-		ExtendVectorArgs<VectorType>(InPropertyHandle, &Args);
-
-		StructBuilder.AddProperty(InPropertyHandle).CustomWidget()
-		.IsEnabled(bEnabled)
-		.NameContent()
-		[
-			InPropertyHandle->CreatePropertyNameWidget()
-		]
-		.ValueContent()
-		.MinDesiredWidth(375.f)
-		.MaxDesiredWidth(375.f)
-		.HAlign(HAlign_Left)
-		[
-			SArgumentNew(Args, SLocalVectorInputBox)
-		];
-	}
+	template<typename TransformType>
+	void ConfigureTransformWidgetArgs(TSharedRef<class IPropertyHandle> InPropertyHandle, typename SAdvancedTransformInputBox<TransformType>::FArguments& WidgetArgs);
+
+	template<typename TransformType>
+	void MakeTransformHeaderRow(TSharedRef<class IPropertyHandle> InPropertyHandle, FDetailWidgetRow& HeaderRow, IPropertyTypeCustomizationUtils& StructCustomizationUtils);
+
+	template<typename TransformType>
+	void MakeTransformChildren(TSharedRef<class IPropertyHandle> InPropertyHandle, IDetailChildrenBuilder& StructBuilder, IPropertyTypeCustomizationUtils& StructCustomizationUtils);
 
 	// returns the rotation for rotator or quaternions (or empty optional for varying values)
 	template<typename RotationType>
@@ -634,535 +587,6 @@ protected:
 	// specializations for FRotator and FQuat at the end of this file
 	template<typename RotationType>
 	void ExtendRotationArgs(TSharedRef<class IPropertyHandle> InPropertyHandle, void* ArgumentsPtr) {}
-
-	// add the widget for a rotation (rotator or quat)
-	template<typename RotationType>
-	void CustomizeRotation(TSharedRef<class IPropertyHandle> InPropertyHandle, class IDetailChildrenBuilder& StructBuilder, IPropertyTypeCustomizationUtils& StructCustomizationUtils)
-	{
-		FEditPropertyChain PropertyChain;
-		TArray<int32> PropertyArrayIndices;
-		bool bEnabled;
-		if (!GetPropertyChain(InPropertyHandle, PropertyChain, PropertyArrayIndices, bEnabled))
-		{
-			return;
-		}
-		
-		typedef typename RotationType::FReal NumericType;
-		typedef SAdvancedRotationInputBox<NumericType> SLocalRotationInputBox;
-		typename SLocalRotationInputBox::FArguments Args;
-		Args.Font(IDetailLayoutBuilder::GetDetailFont());
-		Args.IsEnabled(bEnabled);
-		Args.AllowSpin(true);
-		Args.bColorAxisLabels(true);
-
-		ExtendRotationArgs<RotationType>(InPropertyHandle, &Args);
-		
-		StructBuilder.AddProperty(InPropertyHandle).CustomWidget()
-		.IsEnabled(bEnabled)
-		.NameContent()
-		[
-			InPropertyHandle->CreatePropertyNameWidget()
-		]
-		.ValueContent()
-		.MinDesiredWidth(375.f)
-		.MaxDesiredWidth(375.f)
-		.HAlign(HAlign_Left)
-		[
-			SArgumentNew(Args, SLocalRotationInputBox)
-		];
-	}
-
-	// add the widget for a transform / euler transform
-	template<typename TransformType>
-	void CustomizeTransform(TSharedRef<class IPropertyHandle> InPropertyHandle, class IDetailChildrenBuilder& StructBuilder, IPropertyTypeCustomizationUtils& StructCustomizationUtils)
-	{
-		FEditPropertyChain PropertyChain;
-		TArray<int32> PropertyArrayIndices;
-		bool bEnabled;
-		if (!GetPropertyChain(InPropertyHandle, PropertyChain, PropertyArrayIndices, bEnabled))
-		{
-			return;
-		}
-		
-		typedef typename TransformType::FReal FReal;
-		typename SAdvancedTransformInputBox<TransformType>::FArguments WidgetArgs;
-		WidgetArgs.IsEnabled(bEnabled);
-		WidgetArgs.AllowEditRotationRepresentation(true);
-		WidgetArgs.UseQuaternionForRotation(IsQuaternionBasedRotation<TransformType>());
-
-		static TransformType Identity = TransformType::Identity;
-
-		UObject* DefaultObject = !ObjectsBeingCustomized.IsEmpty() ? ObjectsBeingCustomized[0]->GetClass()->GetDefaultObject() :
-		!StructsBeingCustomized.IsEmpty() ? StructsBeingCustomized[0]->GetStruct()->GetClass() : nullptr;
-
-		if (!DefaultObject)
-		{
-			return;
-		}
-		
-		TransformType DefaultValue = ContainerMemoryBlockToValueRef<TransformType>((uint8*)DefaultObject, Identity, PropertyChain, PropertyArrayIndices);
-		
-		WidgetArgs.DiffersFromDefault_Lambda([this, InPropertyHandle, DefaultValue](ESlateTransformComponent::Type InTransformComponent) -> bool
-		{
-			if (ObjectsBeingCustomized.IsEmpty())
-			{
-				return false;
-			}
-			
-			FEditPropertyChain PropertyChain;
-			TArray<int32> PropertyArrayIndices;
-			bool bEnabled;
-			if (!GetPropertyChain(InPropertyHandle, PropertyChain, PropertyArrayIndices, bEnabled))
-			{
-				return false;
-			}
-			
-			for(const TWeakObjectPtr<UObject>& Object : ObjectsBeingCustomized)
-			{
-				if(Object.Get() && InPropertyHandle->IsValidHandle())
-				{
-					const TransformType& Transform = ContainerMemoryBlockToValueRef<TransformType>((uint8*)Object.Get(), Identity, PropertyChain, PropertyArrayIndices);
-
-					switch(InTransformComponent)
-					{
-						case ESlateTransformComponent::Location:
-						{
-							if(!Transform.GetLocation().Equals(DefaultValue.GetLocation()))
-							{
-								return true;
-							}
-							break;
-						}
-						case ESlateTransformComponent::Rotation:
-						{
-							if(!Transform.Rotator().Equals(DefaultValue.Rotator()))
-							{
-								return true;
-							}
-							break;
-						}
-						case ESlateTransformComponent::Scale:
-						{
-							if(!Transform.GetScale3D().Equals(DefaultValue.GetScale3D()))
-							{
-								return true;
-							}
-							break;
-						}
-						default:
-						{
-							break;
-						}
-					}
-				}
-			}
-			return false;
-		});
-
-		WidgetArgs.OnGetNumericValue_Lambda([this, InPropertyHandle](
-			ESlateTransformComponent::Type InTransformComponent,
-			ESlateRotationRepresentation::Type InRotationRepresentation,
-			ESlateTransformSubComponent::Type InTransformSubComponent) -> TOptional<FReal>
-		{
-			TOptional<FReal> Result;
-			FEditPropertyChain PropertyChain;
-			TArray<int32> PropertyArrayIndices;
-			bool bEnabled;
-			if (!GetPropertyChain(InPropertyHandle, PropertyChain, PropertyArrayIndices, bEnabled))
-			{
-				return Result;
-			}
-
-			const TArray<uint8*> MemoryBlocks = GetMemoryBeingCustomized();
-			for(uint8* MemoryBlock: MemoryBlocks)
-			{
-				if(MemoryBlock)
-				{
-					const TransformType& Transform = ContainerMemoryBlockToValueRef<TransformType>(MemoryBlock, Identity, PropertyChain, PropertyArrayIndices);
-					
-					TOptional<FReal> Value = SAdvancedTransformInputBox<TransformType>::GetNumericValueFromTransform(
-						Transform,
-						InTransformComponent,
-						InRotationRepresentation,
-						InTransformSubComponent
-						);
-
-					if(Value.IsSet())
-					{
-						if(Result.IsSet())
-						{
-							if(!FMath::IsNearlyEqual(Result.GetValue(), Value.GetValue()))
-							{
-								return TOptional<FReal>();
-							}
-						}
-						else
-						{
-							Result = Value;
-						}
-					}
-				}
-			}
-			return Result;
-		});
-		
-
-		auto OnNumericValueChanged = [this, InPropertyHandle](
-			ESlateTransformComponent::Type InTransformComponent, 
-			ESlateRotationRepresentation::Type InRotationRepresentation, 
-			ESlateTransformSubComponent::Type InSubComponent,
-			FReal InValue,
-			bool bIsCommit,
-			ETextCommit::Type InCommitType = ETextCommit::Default)
-		{
-			FEditPropertyChain PropertyChain;
-			TArray<int32> PropertyArrayIndices;
-			bool bEnabled;
-			if (!GetPropertyChain(InPropertyHandle, PropertyChain, PropertyArrayIndices, bEnabled))
-			{
-				return;
-			}
-			
-			TArray<UObject*> ObjectsView;
-			for(int32 Index = 0; Index < ObjectsBeingCustomized.Num(); Index++)
-			{
-				const TWeakObjectPtr<UObject>& Object = ObjectsBeingCustomized[Index];
-				if (Object.Get())
-				{
-					ObjectsView.Add(Object.Get());
-				}
-			}
-			FPropertyChangedEvent PropertyChangedEvent(InPropertyHandle->GetProperty(), bIsCommit ? EPropertyChangeType::ValueSet : EPropertyChangeType::Interactive, ObjectsView);
-			FPropertyChangedChainEvent PropertyChangedChainEvent(PropertyChain, PropertyChangedEvent);
-
-			URigVMController* Controller = nullptr;
-			if(BlueprintBeingCustomized && GraphBeingCustomized)
-			{
-				Controller = BlueprintBeingCustomized->GetController(GraphBeingCustomized);
-				if(bIsCommit)
-				{
-					Controller->OpenUndoBracket(FString::Printf(TEXT("Set %s"), *InPropertyHandle->GetProperty()->GetName()));
-				}
-			}
-
-			for(int32 Index = 0; Index < ObjectsBeingCustomized.Num(); Index++)
-			{
-				const TWeakObjectPtr<UObject>& Object = ObjectsBeingCustomized[Index];
-				if(Object.Get() && InPropertyHandle->IsValidHandle())
-				{
-					TransformType& Transform = ContainerMemoryBlockToValueRef<TransformType>((uint8*)Object.Get(), Identity, PropertyChain, PropertyArrayIndices);
-					TransformType PreviousTransform = Transform;
-					
-					SAdvancedTransformInputBox<TransformType>::ApplyNumericValueChange(
-						Transform,
-						InValue,
-						InTransformComponent,
-						InRotationRepresentation,
-						InSubComponent);
-
-					if(!PreviousTransform.Equals(Transform))
-					{
-						Object->PostEditChangeChainProperty(PropertyChangedChainEvent);
-						InPropertyHandle->NotifyPostChange(PropertyChangedEvent.ChangeType);
-					}
-				}
-			}
-
-			if(Controller && bIsCommit)
-			{
-				Controller->CloseUndoBracket();
-			}
-		};
-
-		WidgetArgs.OnNumericValueChanged_Lambda([OnNumericValueChanged](
-			ESlateTransformComponent::Type InTransformComponent, 
-			ESlateRotationRepresentation::Type InRotationRepresentation, 
-			ESlateTransformSubComponent::Type InSubComponent,
-			FReal InValue)
-		{
-			OnNumericValueChanged(InTransformComponent, InRotationRepresentation, InSubComponent, InValue, false);
-		});
-
-		WidgetArgs.OnNumericValueCommitted_Lambda([OnNumericValueChanged](
-			ESlateTransformComponent::Type InTransformComponent, 
-			ESlateRotationRepresentation::Type InRotationRepresentation, 
-			ESlateTransformSubComponent::Type InSubComponent,
-			FReal InValue, 
-			ETextCommit::Type InCommitType)
-		{
-			OnNumericValueChanged(InTransformComponent, InRotationRepresentation, InSubComponent, InValue, true, InCommitType);
-		});
-
-		WidgetArgs.OnResetToDefault_Lambda([this, DefaultValue, InPropertyHandle](ESlateTransformComponent::Type InTransformComponent)
-		{
-			if (ObjectsBeingCustomized.IsEmpty())
-			{
-				return;
-			}
-			
-			FEditPropertyChain PropertyChain;
-			TArray<int32> PropertyArrayIndices;
-			bool bEnabled;
-			if (!GetPropertyChain(InPropertyHandle, PropertyChain, PropertyArrayIndices, bEnabled))
-			{
-				return;
-			}
-			
-			URigVMController* Controller = nullptr;
-			if(BlueprintBeingCustomized && GraphBeingCustomized)
-			{
-				Controller = BlueprintBeingCustomized->GetController(GraphBeingCustomized);
-				if(Controller)
-				{
-					Controller->OpenUndoBracket(FString::Printf(TEXT("Reset %s to Default"), *InPropertyHandle->GetProperty()->GetName()));
-				}
-			}
-
-			TArray<UObject*> ObjectsView;
-			for(int32 Index = 0; Index < ObjectsBeingCustomized.Num(); Index++)
-			{
-				const TWeakObjectPtr<UObject>& Object = ObjectsBeingCustomized[Index];
-				if (Object.Get())
-				{
-					ObjectsView.Add(Object.Get());
-				}
-			}
-			FPropertyChangedEvent PropertyChangedEvent(InPropertyHandle->GetProperty(), EPropertyChangeType::ValueSet, ObjectsView);
-			FPropertyChangedChainEvent PropertyChangedChainEvent(PropertyChain, PropertyChangedEvent);
-			
-			for(int32 Index = 0; Index < ObjectsBeingCustomized.Num(); Index++)
-			{
-				const TWeakObjectPtr<UObject>& Object = ObjectsBeingCustomized[Index];
-				if(Object.Get() && InPropertyHandle->IsValidHandle())
-				{
-					TransformType& Transform = ContainerMemoryBlockToValueRef<TransformType>((uint8*)Object.Get(), Identity, PropertyChain, PropertyArrayIndices);
-					TransformType PreviousTransform = Transform;
-
-					switch(InTransformComponent)
-					{
-						case ESlateTransformComponent::Location:
-						{
-							Transform.SetLocation(DefaultValue.GetLocation());
-							break;
-						}
-						case ESlateTransformComponent::Rotation:
-						{
-							Transform.SetRotation(DefaultValue.GetRotation());
-							break;
-						}
-						case ESlateTransformComponent::Scale:
-						{
-							Transform.SetScale3D(DefaultValue.GetScale3D());
-							break;
-						}
-						case ESlateTransformComponent::Max:
-						default:
-						{
-							Transform.SetLocation(DefaultValue.GetLocation());
-							break;
-						}
-					}
-
-					
-					if(!PreviousTransform.Equals(Transform))
-					{
-						Object->PostEditChangeChainProperty(PropertyChangedChainEvent);
-						InPropertyHandle->NotifyPostChange(PropertyChangedEvent.ChangeType);
-					}
-				}
-			}
-
-			if(Controller)
-			{
-				Controller->CloseUndoBracket();
-			}
-		});
-
-		WidgetArgs.OnCopyToClipboard_Lambda([this, InPropertyHandle](
-			ESlateTransformComponent::Type InComponent
-			)
-		{
-			TOptional<FReal> Result;
-			FEditPropertyChain PropertyChain;
-			TArray<int32> PropertyArrayIndices;
-			bool bEnabled;
-			if (!GetPropertyChain(InPropertyHandle, PropertyChain, PropertyArrayIndices, bEnabled))
-			{
-				return;
-			}
-
-			const TArray<uint8*> MemoryBlocks = GetMemoryBeingCustomized();
-			for(uint8* MemoryBlock: MemoryBlocks)
-			{
-				if(MemoryBlock)
-				{
-					const TransformType& Transform = ContainerMemoryBlockToValueRef<TransformType>(MemoryBlock, Identity, PropertyChain, PropertyArrayIndices);
-					FString Content;
-					switch(InComponent)
-					{
-						case ESlateTransformComponent::Location:
-						{
-							const FVector Data = Transform.GetLocation();
-							TBaseStructure<FVector>::Get()->ExportText(Content, &Data, &Data, nullptr, PPF_None, nullptr);
-							break;
-						}
-						case ESlateTransformComponent::Rotation:
-						{
-							const FRotator Data = Transform.Rotator();
-							TBaseStructure<FRotator>::Get()->ExportText(Content, &Data, &Data, nullptr, PPF_None, nullptr);
-							break;
-						}
-						case ESlateTransformComponent::Scale:
-						{
-							const FVector Data = Transform.GetScale3D();
-							TBaseStructure<FVector>::Get()->ExportText(Content, &Data, &Data, nullptr, PPF_None, nullptr);
-							break;
-						}
-						case ESlateTransformComponent::Max:
-						default:
-						{
-							TBaseStructure<TransformType>::Get()->ExportText(Content, &Transform, &Transform, nullptr, PPF_None, nullptr);
-							break;
-						}
-					}
-
-					if(!Content.IsEmpty())
-					{
-						FPlatformApplicationMisc::ClipboardCopy(*Content);
-					}
-				}
-			}
-		});
-
-		WidgetArgs.OnPasteFromClipboard_Lambda([this, InPropertyHandle, OnNumericValueChanged](
-			ESlateTransformComponent::Type InComponent
-			)
-		{
-			FString Content;
-			FPlatformApplicationMisc::ClipboardPaste(Content);
-		
-			if(Content.IsEmpty())
-			{
-				return;
-			}
-
-			if (ObjectsBeingCustomized.IsEmpty())
-			{
-				return;
-			}
-
-			TOptional<FReal> Result;
-			FEditPropertyChain PropertyChain;
-			TArray<int32> PropertyArrayIndices;
-			bool bEnabled;
-			if (!GetPropertyChain(InPropertyHandle, PropertyChain, PropertyArrayIndices, bEnabled))
-			{
-				return;
-			}
-
-			TArray<UObject*> ObjectsView;
-			for(int32 Index = 0; Index < ObjectsBeingCustomized.Num(); Index++)
-			{
-				const TWeakObjectPtr<UObject>& Object = ObjectsBeingCustomized[Index];
-				if (Object.Get())
-				{
-					ObjectsView.Add(Object.Get());
-				}
-			}
-			FPropertyChangedEvent PropertyChangedEvent(InPropertyHandle->GetProperty(), EPropertyChangeType::ValueSet, ObjectsView);
-			FPropertyChangedChainEvent PropertyChangedChainEvent(PropertyChain, PropertyChangedEvent);
-
-			URigVMController* Controller = nullptr;
-			if(BlueprintBeingCustomized && GraphBeingCustomized)
-			{
-				Controller = BlueprintBeingCustomized->GetController(GraphBeingCustomized);
-				Controller->OpenUndoBracket(FString::Printf(TEXT("Set %s"), *InPropertyHandle->GetProperty()->GetName()));
-			}
-			
-			for(const TWeakObjectPtr<UObject>& Object : ObjectsBeingCustomized)
-			{
-				if(Object.Get() && InPropertyHandle->IsValidHandle())
-				{
-					TransformType& Transform = ContainerMemoryBlockToValueRef<TransformType>((uint8*)Object.Get(), Identity, PropertyChain, PropertyArrayIndices);
-					const TransformType PreviousTransform = Transform;
-
-					// Apply the new value
-					{
-						class FRigPasteTransformWidgetErrorPipe : public FOutputDevice
-						{
-						public:
-					
-							int32 NumErrors;
-					
-							FRigPasteTransformWidgetErrorPipe()
-								: FOutputDevice()
-								, NumErrors(0)
-							{
-							}
-					
-							virtual void Serialize(const TCHAR* V, ELogVerbosity::Type Verbosity, const class FName& Category) override
-							{
-								UE_LOG(LogRigVM, Error, TEXT("Error Pasting to Widget: %s"), V);
-								NumErrors++;
-							}
-						};
-					
-						FRigPasteTransformWidgetErrorPipe ErrorPipe;
-						
-						switch(InComponent)
-						{
-							case ESlateTransformComponent::Location:
-							{
-								FVector Data = Transform.GetLocation();
-								TBaseStructure<FVector>::Get()->ImportText(*Content, &Data, nullptr, PPF_None, &ErrorPipe, TBaseStructure<FVector>::Get()->GetName(), true);
-								Transform.SetLocation(Data);
-								break;
-							}
-							case ESlateTransformComponent::Rotation:
-							{
-								FRotator Data = Transform.Rotator();
-								TBaseStructure<FRotator>::Get()->ImportText(*Content, &Data, nullptr, PPF_None, &ErrorPipe, TBaseStructure<FRotator>::Get()->GetName(), true);
-								FQuat Quat = Data.Quaternion();
-								Transform.SetRotation(Quat);
-								
-								break;
-							}
-							case ESlateTransformComponent::Scale:
-							{
-								FVector Data = Transform.GetScale3D();
-								TBaseStructure<FVector>::Get()->ImportText(*Content, &Data, nullptr, PPF_None, &ErrorPipe, TBaseStructure<FVector>::Get()->GetName(), true);
-								Transform.SetScale3D(Data);
-								break;
-							}
-							case ESlateTransformComponent::Max:
-							default:
-							{
-								TBaseStructure<TransformType>::Get()->ImportText(*Content, &Transform, nullptr, PPF_None, &ErrorPipe, TBaseStructure<TransformType>::Get()->GetName(), true);
-								break;
-							}
-						}
-					
-						if(ErrorPipe.NumErrors == 0 && !PreviousTransform.Equals(Transform))
-						{
-							Object->PostEditChangeChainProperty(PropertyChangedChainEvent);
-							InPropertyHandle->NotifyPostChange(PropertyChangedEvent.ChangeType);
-						}
-					}
-				}
-			}
-			
-			if (Controller)
-			{
-				Controller->CloseUndoBracket();
-			}
-		});
-
-		SAdvancedTransformInputBox<TransformType>::ConstructGroupedTransformRows(
-			StructBuilder,
-			InPropertyHandle->GetPropertyDisplayName(),
-			InPropertyHandle->GetToolTipText(),
-			WidgetArgs);
-	}
 		
 	UScriptStruct* ScriptStruct;
 	URigVMBlueprint* BlueprintBeingCustomized;
diff --git a/Engine/Source/Developer/AnimationWidgets/Public/SAdvancedTransformInputBox.h b/Engine/Source/Developer/AnimationWidgets/Public/SAdvancedTransformInputBox.h
--- a/Engine/Source/Developer/AnimationWidgets/Public/SAdvancedTransformInputBox.h
+++ b/Engine/Source/Developer/AnimationWidgets/Public/SAdvancedTransformInputBox.h
@@ -1196,81 +1196,81 @@ public:
 		return InBuilder.AddGroup(GroupName, LocalizedDisplayName);
 	}
 
-	template<typename BuilderType = IDetailCategoryBuilder>
-	static FDetailWidgetRow& ConstructGroupedTransformRows(
-		BuilderType& InBuilder,
-		const FText& InLabel,
-		const FText& InTooltip,
-		typename SAdvancedTransformInputBox<TransformType>::FArguments WidgetArgs,
-		TSharedPtr<SWidget> NameContent = TSharedPtr<SWidget>()
-		)
+	static void ConfigureComponentWidgetRow(
+		FDetailWidgetRow& WidgetRow, 
+		ESlateTransformComponent::Type InComponent,
+		typename SAdvancedTransformInputBox<TransformType>::FArguments WidgetArgs)
 	{
-		auto ConstructComponentWidgetRow = [WidgetArgs](FDetailWidgetRow& WidgetRow, ESlateTransformComponent::Type InComponent)
+		if(WidgetArgs._OnCopyToClipboard.IsBound() && WidgetArgs._OnPasteFromClipboard.IsBound())
 		{
-			if(WidgetArgs._OnCopyToClipboard.IsBound() && WidgetArgs._OnPasteFromClipboard.IsBound())
-			{
-				WidgetRow
-				.CopyAction(FUIAction(
-					FExecuteAction::CreateLambda([WidgetArgs, InComponent]()
-					{
-						WidgetArgs._OnCopyToClipboard.ExecuteIfBound(InComponent);
-					}),
-					FCanExecuteAction())
-				);
-
-				WidgetRow.PasteAction(FUIAction(
+			WidgetRow
+			.CopyAction(FUIAction(
 				FExecuteAction::CreateLambda([WidgetArgs, InComponent]()
 				{
-					WidgetArgs._OnPasteFromClipboard.ExecuteIfBound(InComponent);
+					WidgetArgs._OnCopyToClipboard.ExecuteIfBound(InComponent);
 				}),
-				FCanExecuteAction::CreateLambda([WidgetArgs]() -> bool
-				{
-					return WidgetArgs._IsEnabled.Get();
-				})));
-			}
+				FCanExecuteAction())
+			);
 
-			if(WidgetArgs._OnResetToDefault.IsBound() && WidgetArgs._DiffersFromDefault.IsBound())
+			WidgetRow.PasteAction(FUIAction(
+			FExecuteAction::CreateLambda([WidgetArgs, InComponent]()
 			{
-				WidgetRow.OverrideResetToDefault(FResetToDefaultOverride::Create(
-					TAttribute<bool>::CreateLambda([WidgetArgs, InComponent]() -> bool
-					{
-						if(!WidgetArgs._IsEnabled.Get())
-						{
-							return false;
-						}
-						return WidgetArgs._DiffersFromDefault.Execute(InComponent);
-					}),
-					FSimpleDelegate::CreateLambda([WidgetArgs, InComponent]()
+				WidgetArgs._OnPasteFromClipboard.ExecuteIfBound(InComponent);
+			}),
+			FCanExecuteAction::CreateLambda([WidgetArgs]() -> bool
+			{
+				return WidgetArgs._IsEnabled.Get();
+			})));
+		}
+
+		if(WidgetArgs._OnResetToDefault.IsBound() && WidgetArgs._DiffersFromDefault.IsBound())
+		{
+			WidgetRow.OverrideResetToDefault(FResetToDefaultOverride::Create(
+				TAttribute<bool>::CreateLambda([WidgetArgs, InComponent]() -> bool
+				{
+					if(!WidgetArgs._IsEnabled.Get())
 					{
-						WidgetArgs._OnResetToDefault.Execute(InComponent);
-					})
-				));
-			}
+						return false;
+					}
+					return WidgetArgs._DiffersFromDefault.Execute(InComponent);
+				}),
+				FSimpleDelegate::CreateLambda([WidgetArgs, InComponent]()
+				{
+					WidgetArgs._OnResetToDefault.Execute(InComponent);
+				})
+			));
+		}
 			
-			if(InComponent != ESlateTransformComponent::Max)
-			{
-				WidgetRow
-				.NameContent()
-				.HAlign(HAlign_Fill)
+		if(InComponent != ESlateTransformComponent::Max)
+		{
+			WidgetRow
+			.NameContent()
+			.HAlign(HAlign_Fill)
+			[
+				SAdvancedTransformInputBox<TransformType>::ConstructLabel(WidgetArgs, InComponent)
+			]
+			.ValueContent()
+			.MinDesiredWidth(375.f)
+			.MaxDesiredWidth(375.f)
+			[
+				SNew(SHorizontalBox)
+				+SHorizontalBox::Slot()
 				[
-					SAdvancedTransformInputBox<TransformType>::ConstructLabel(WidgetArgs, InComponent)
+					SAdvancedTransformInputBox<TransformType>::ConstructWidget(WidgetArgs, InComponent)
 				]
-				.ValueContent()
-				.MinDesiredWidth(375.f)
-				.MaxDesiredWidth(375.f)
-				[
-					SNew(SHorizontalBox)
-					+SHorizontalBox::Slot()
-					[
-						SAdvancedTransformInputBox<TransformType>::ConstructWidget(WidgetArgs, InComponent)
-					]
-				];
-			}			
-		};
-		
-		IDetailGroup& Group = ConstructDetailGroup(InBuilder, *InLabel.ToString(), InLabel);
-		FDetailWidgetRow& HeaderRow = Group.HeaderRow();
-		ConstructComponentWidgetRow(HeaderRow, ESlateTransformComponent::Max);
+			];
+		}			
+	};
+
+	static void ConfigureHeader(
+		FDetailWidgetRow& HeaderRow,
+		const FText& InLabel,
+		const FText& InTooltip,
+		typename SAdvancedTransformInputBox<TransformType>::FArguments WidgetArgs,
+		TSharedPtr<SWidget> NameContent = TSharedPtr<SWidget>()
+	)
+	{
+		ConfigureComponentWidgetRow(HeaderRow, ESlateTransformComponent::Max, WidgetArgs);
 
 		if(!NameContent.IsValid())
 		{
@@ -1286,23 +1286,37 @@ public:
 		[
 			NameContent.ToSharedRef()
 		];
+	}
+
+	template<typename BuilderType = IDetailCategoryBuilder>
+	static FDetailWidgetRow& ConstructGroupedTransformRows(
+		BuilderType& InBuilder,
+		const FText& InLabel,
+		const FText& InTooltip,
+		typename SAdvancedTransformInputBox<TransformType>::FArguments WidgetArgs,
+		TSharedPtr<SWidget> NameContent = TSharedPtr<SWidget>()
+		)
+	{
+		IDetailGroup& Group = ConstructDetailGroup(InBuilder, *InLabel.ToString(), InLabel);
+		FDetailWidgetRow& HeaderRow = Group.HeaderRow();
+		ConfigureHeader(HeaderRow, InLabel, InTooltip, WidgetArgs, NameContent);
 
 		if(WidgetArgs._ConstructLocation)
 		{
 			FDetailWidgetRow& WidgetRow = Group.AddWidgetRow();
-			ConstructComponentWidgetRow(WidgetRow, ESlateTransformComponent::Location);
+			ConfigureComponentWidgetRow(WidgetRow, ESlateTransformComponent::Location, WidgetArgs);
 		}
 
 		if(WidgetArgs._ConstructRotation)
 		{
 			FDetailWidgetRow& WidgetRow = Group.AddWidgetRow();
-			ConstructComponentWidgetRow(WidgetRow, ESlateTransformComponent::Rotation);
+			ConfigureComponentWidgetRow(WidgetRow, ESlateTransformComponent::Rotation, WidgetArgs);
 		}
 
 		if(WidgetArgs._ConstructScale)
 		{
 			FDetailWidgetRow& WidgetRow = Group.AddWidgetRow();
-			ConstructComponentWidgetRow(WidgetRow, ESlateTransformComponent::Scale);
+			ConfigureComponentWidgetRow(WidgetRow, ESlateTransformComponent::Scale, WidgetArgs);
 		}
 
 		return HeaderRow;
