RigVM: Fix old code that assumes single model

#rnx

diff --git a/Engine/Plugins/Runtime/RigVM/Source/RigVMDeveloper/Private/RigVMBlueprint.cpp b/Engine/Plugins/Runtime/RigVM/Source/RigVMDeveloper/Private/RigVMBlueprint.cpp
--- a/Engine/Plugins/Runtime/RigVM/Source/RigVMDeveloper/Private/RigVMBlueprint.cpp
+++ b/Engine/Plugins/Runtime/RigVM/Source/RigVMDeveloper/Private/RigVMBlueprint.cpp
@@ -250,11 +250,11 @@ void URigVMBlueprint::InitializeModelIfRequired(bool bRecompileVM)
 
 	if (RigVMClient.GetController(0) == nullptr)
 	{
-		check(RigVMClient.Num() == 1);
-		check(RigVMClient.GetFunctionLibrary());
-		
-		RigVMClient.GetOrCreateController(RigVMClient.GetDefaultModel());
-		RigVMClient.GetOrCreateController(RigVMClient.GetFunctionLibrary());
+		const TArray<URigVMGraph*> Models = RigVMClient.GetAllModels(true, false);
+		for(const URigVMGraph* Model : Models)
+		{
+			RigVMClient.GetOrCreateController(Model);
+		}
 
 		bool bRecompileRequired = false;
 		for (int32 i = 0; i < UbergraphPages.Num(); ++i)
