Control Rig: FIx accessing invalid weak objects
#rnx

diff --git a/Engine/Plugins/Runtime/RigVM/Source/RigVMDeveloper/Private/RigVMCompiler/RigVMASTProxy.cpp b/Engine/Plugins/Runtime/RigVM/Source/RigVMDeveloper/Private/RigVMCompiler/RigVMASTProxy.cpp
--- a/Engine/Plugins/Runtime/RigVM/Source/RigVMDeveloper/Private/RigVMCompiler/RigVMASTProxy.cpp
+++ b/Engine/Plugins/Runtime/RigVM/Source/RigVMDeveloper/Private/RigVMCompiler/RigVMASTProxy.cpp
@@ -85,7 +85,14 @@ const UObject* FRigVMCallstack::operator[](int32 InIndex) const
 
 bool FRigVMCallstack::Contains(const UObject* InEntry) const
 {
-	return Stack.Contains(InEntry);
+	return Stack.ContainsByPredicate([InEntry](const TWeakObjectPtr<UObject>& Object)
+	{
+		if (Object.IsValid())
+		{
+			return Object.Get() == InEntry;
+		}
+		return false;
+	});
 }
 
 FRigVMCallstack FRigVMCallstack::GetCallStackUpTo(int32 InIndex) const
