Fix editor crash when adding a new trait in AnimNext

diff --git a/Engine/Plugins/Runtime/RigVM/Source/RigVMDeveloper/Private/RigVMModel/RigVMPin.cpp b/Engine/Plugins/Runtime/RigVM/Source/RigVMDeveloper/Private/RigVMModel/RigVMPin.cpp
--- a/Engine/Plugins/Runtime/RigVM/Source/RigVMDeveloper/Private/RigVMModel/RigVMPin.cpp
+++ b/Engine/Plugins/Runtime/RigVM/Source/RigVMDeveloper/Private/RigVMModel/RigVMPin.cpp
@@ -1164,7 +1164,7 @@ FString URigVMPin::GetMetaData(FName InKey) const
 					{
 						const FRigVMDecorator* VMDecorator = (FRigVMDecorator*)DecoratorScope->GetStructMemory();
 						Struct = VMDecorator->GetDecoratorSharedDataStruct();
-						Property = Struct->FindPropertyByName(GetFName());
+						Property = Struct != nullptr ? Struct->FindPropertyByName(GetFName()) : nullptr;
 						if(Property)
 						{
 							const FString MetaData = Property->GetMetaData(InKey);
