Control Rig: Fix bad initialization of PredicateCount in FRigVMExecuteOp
#rnx

diff --git a/Engine/Plugins/Runtime/RigVM/Source/RigVM/Public/RigVMCore/RigVMByteCode.h b/Engine/Plugins/Runtime/RigVM/Source/RigVM/Public/RigVMCore/RigVMByteCode.h
--- a/Engine/Plugins/Runtime/RigVM/Source/RigVM/Public/RigVMCore/RigVMByteCode.h
+++ b/Engine/Plugins/Runtime/RigVM/Source/RigVM/Public/RigVMCore/RigVMByteCode.h
@@ -262,7 +262,7 @@ struct RIGVM_API FRigVMExecuteOp : public FRigVMBaseOp
 	, FunctionIndex(INDEX_NONE)
 	, ArgumentCount(0)
 	, FirstPredicateIndex(INDEX_NONE)
-	, PredicateCount(INDEX_NONE)
+	, PredicateCount(0)
 	{
 	}
 
@@ -270,6 +270,8 @@ struct RIGVM_API FRigVMExecuteOp : public FRigVMBaseOp
 	: FRigVMBaseOp(ERigVMOpCode::Execute)
 	, FunctionIndex(InFunctionIndex)
 	, ArgumentCount(InArgumentCount)
+	, FirstPredicateIndex(INDEX_NONE)
+	, PredicateCount(0)
 	{
 	}
 
