RigVM: When changing pin values don't assume pin exists

#rnx

diff --git a/Engine/Plugins/Runtime/RigVM/Source/RigVMEditor/Private/Editor/RigVMEditor.cpp b/Engine/Plugins/Runtime/RigVM/Source/RigVMEditor/Private/Editor/RigVMEditor.cpp
--- a/Engine/Plugins/Runtime/RigVM/Source/RigVMEditor/Private/Editor/RigVMEditor.cpp
+++ b/Engine/Plugins/Runtime/RigVM/Source/RigVMEditor/Private/Editor/RigVMEditor.cpp
@@ -2058,7 +2058,9 @@ void FRigVMEditor::OnWrappedPropertyChangedChainEvent(URigVMDetailsViewWrapperOb
 		const FName RootPinName = InPropertyChangedChainEvent.PropertyChain.GetHead()->GetValue()->GetFName();
 		const FString RootPinNameString = RootPinName.ToString();
 		FString PinPath = URigVMPin::JoinPinPath(Node->GetName(), RootPinNameString);
-		
+		URigVMController* Controller = GetRigVMBlueprint()->GetController(Node->GetGraph());
+		check(Controller);
+
 		const FProperty* Property = WrapperObjects[0]->GetClass()->FindPropertyByName(RootPinName);
 		uint8* PropertyStorage = nullptr;
 		if (Property)
@@ -2071,14 +2073,32 @@ void FRigVMEditor::OnWrappedPropertyChangedChainEvent(URigVMDetailsViewWrapperOb
 				check(InPropertyPath.StartsWith(RootPinNameString));
 				FString RemainingPropertyPath = InPropertyPath.Mid(RootPinNameString.Len());
 				RemainingPropertyPath.RemoveFromStart(TEXT("->"));
-				const FString SegmentPath = RemainingPropertyPath.Replace(TEXT("->"), TEXT("."));
-			
-				const FRigVMPropertyPath PropertyTraverser(Property, SegmentPath);
-				PropertyStorage = PropertyTraverser.GetData<uint8>(PropertyStorage, Property);
-				Property = PropertyTraverser.GetTailProperty();
-				PinPath = URigVMPin::JoinPinPath(PinPath, SegmentPath);
-				PinPath.ReplaceInline(TEXT("["), TEXT(""));
-				PinPath.ReplaceInline(TEXT("]"), TEXT(""));
+				RemainingPropertyPath.ReplaceInline(TEXT("->"), TEXT("."));
+
+				// traverse each property one by one to make sure the expected pin exists.
+				// this may not be the case for an array element.
+				while(!RemainingPropertyPath.IsEmpty())
+				{
+					FString Left = RemainingPropertyPath, Right;
+					(void)URigVMPin::SplitPinPathAtStart(RemainingPropertyPath, Left, Right);
+
+					FString NewPinPath = URigVMPin::JoinPinPath(PinPath, Left);
+					NewPinPath.ReplaceInline(TEXT("["), TEXT(""));
+					NewPinPath.ReplaceInline(TEXT("]"), TEXT(""));
+
+					// this may be an array pin which doesn't exist yet
+					if(!Controller->GetGraph()->FindPin(NewPinPath))
+					{
+						break;
+					}
+
+					const FRigVMPropertyPath PropertyTraverser(Property, Left);
+					PropertyStorage = PropertyTraverser.GetData<uint8>(PropertyStorage, Property);
+					Property = PropertyTraverser.GetTailProperty();
+					PinPath = NewPinPath;
+
+					RemainingPropertyPath = Right;
+				}
 			}
 		}
 
@@ -2092,8 +2112,6 @@ void FRigVMEditor::OnWrappedPropertyChangedChainEvent(URigVMDetailsViewWrapperOb
 			if (!DefaultValue.IsEmpty())
 			{
 				const bool bInteractive = InPropertyChangedChainEvent.ChangeType == EPropertyChangeType::Interactive;
-				URigVMController* Controller = GetRigVMBlueprint()->GetController(Node->GetGraph());
-				check(Controller);
 				Controller->SetPinDefaultValue(PinPath, DefaultValue, true, !bInteractive, true, !bInteractive);
 			}
 		}
