RigVM: Ensure to not re-use node names during patching to avoid linker errors

#rnx

diff --git a/Engine/Plugins/Runtime/RigVM/Source/RigVMDeveloper/Private/RigVMModel/RigVMController.cpp b/Engine/Plugins/Runtime/RigVM/Source/RigVMDeveloper/Private/RigVMModel/RigVMController.cpp
--- a/Engine/Plugins/Runtime/RigVM/Source/RigVMDeveloper/Private/RigVMModel/RigVMController.cpp
+++ b/Engine/Plugins/Runtime/RigVM/Source/RigVMDeveloper/Private/RigVMModel/RigVMController.cpp
@@ -18842,27 +18842,32 @@ URigVMNode* URigVMController::ConvertRerouteNodeToDispatch(URigVMRerouteNode* In
 	FString PinToResolveName;
 
 	// constant is empty
+	FString NewNodeNameSuffix;
 	if(InTemplateNotation == FRigVMDispatch_Constant().GetTemplateNotation())
 	{
 		PinToResolveName = FRigVMDispatch_Constant::ValueName.ToString();
+		NewNodeNameSuffix = TEXT("Constant");
 	}
 	else if(InTemplateNotation == FRigVMDispatch_MakeStruct().GetTemplateNotation())
 	{
 		InputRedirects = &ElementsRedirects;
 		OutputRedirects = &StructRedirects;
 		PinToResolveName = FRigVMDispatch_MakeStruct::StructName.ToString();
+		NewNodeNameSuffix = TEXT("MakeStruct");
 	}
 	else if(InTemplateNotation == FRigVMDispatch_BreakStruct().GetTemplateNotation())
 	{
 		InputRedirects = &StructRedirects;
 		OutputRedirects = &ElementsRedirects;
 		PinToResolveName = FRigVMDispatch_MakeStruct::StructName.ToString();
+		NewNodeNameSuffix = TEXT("BreakStruct");
 	}
 	else if(InTemplateNotation == FRigVMDispatch_ArrayMake().GetTemplateNotation())
 	{
 		InputRedirects = &ValuesRedirects;
 		OutputRedirects = &ArrayRedirects;
 		PinToResolveName = FRigVMDispatch_ArrayMake::ArrayName.ToString();
+		NewNodeNameSuffix = TEXT("ArrayMake");
 	}
 	else
 	{
@@ -18900,8 +18905,11 @@ URigVMNode* URigVMController::ConvertRerouteNodeToDispatch(URigVMRerouteNode* In
 	// old node.
 	const FString DeletedName = GetSchema()->GetValidNodeName(GetGraph(), FString::Printf(TEXT("%s_Deleted"), *NodeName));
 	RenameNode(InRerouteNode, *DeletedName, false);
-
-	URigVMNode* NewNode = AddTemplateNode(InTemplateNotation, NodePosition, NodeName, bSetupUndoRedo, bPrintPythonCommand);
+	const FString NewNodeName = NodeName + TEXT("_") + NewNodeNameSuffix;
+	FRestoreLinkedPathSettings RestoreLinkedPathSettings;
+	RestoreLinkedPathSettings.NodeNameMap.Add(NodeName, NewNodeName);
+	
+	URigVMNode* NewNode = AddTemplateNode(InTemplateNotation, NodePosition, NewNodeName, bSetupUndoRedo, bPrintPythonCommand);
 	if(NewNode)
 	{
 		URigVMPin* PinToResolve = NewNode->FindPin(PinToResolveName);
@@ -18926,9 +18934,8 @@ URigVMNode* URigVMController::ConvertRerouteNodeToDispatch(URigVMRerouteNode* In
 		}
 		ApplyPinStates(NewNode, RemappedPinStates, {}, bSetupUndoRedo);
 
-		const FString NodeNamePrefix = URigVMPin::JoinPinPath({NodeName, FString()});
-		FRestoreLinkedPathSettings RestoreSettings;
-		RestoreSettings.RemapDelegates.Add(NodeName,
+		const FString NodeNamePrefix = URigVMPin::JoinPinPath({NewNodeName, FString()});
+		RestoreLinkedPathSettings.RemapDelegates.Add(NewNodeName,
 			FRigVMController_PinPathRemapDelegate::CreateLambda([NodeNamePrefix, InputRedirects, OutputRedirects](const FString& InPinPath, bool bIsInput) -> FString
 			{
 				TArray<FString> Parts;
@@ -18944,7 +18951,7 @@ URigVMNode* URigVMController::ConvertRerouteNodeToDispatch(URigVMRerouteNode* In
 				return InPinPath;
 			})
 		);
-		RestoreLinkedPaths(LinkedPaths, RestoreSettings, bSetupUndoRedo);
+		RestoreLinkedPaths(LinkedPaths, RestoreLinkedPathSettings, bSetupUndoRedo);
 
 		if(!RemoveNode(InRerouteNode, bSetupUndoRedo, bPrintPythonCommand))
 		{
@@ -19345,10 +19352,14 @@ FRigVMClientPatchResult URigVMController::PatchUnitNodesOnLoad()
 			const FString DeletedName = GetSchema()->GetValidNodeName(Graph, FString::Printf(TEXT("%s_Deleted"), *NodeName));
 			RenameNode(UnitNode, *DeletedName, false, false);
 			
+			const FString NewNodeName = NodeName + TEXT("_") + URigVMUnitNode::StaticClass()->GetName();
+			FRestoreLinkedPathSettings RestoreLinkedPathSettings;
+			RestoreLinkedPathSettings.NodeNameMap.Add(NodeName, NewNodeName);
+			
 			URigVMTemplateNode* NewNode = AddTemplateNode(
 				Template->GetNotation(),
 				NodePosition, 
-				NodeName, 
+				NewNodeName, 
 				false, 
 				false);
 
@@ -19375,8 +19386,7 @@ FRigVMClientPatchResult URigVMController::PatchUnitNodesOnLoad()
 
 			RemoveNode(UnitNode, false, false);
 
-			FRestoreLinkedPathSettings Settings;
-			RestoreLinkedPaths(LinkedPaths, Settings);
+			RestoreLinkedPaths(LinkedPaths, RestoreLinkedPathSettings);
 		}
 	}
 
@@ -19535,10 +19545,13 @@ FRigVMClientPatchResult URigVMController::PatchIfSelectNodesOnLoad()
 				bIsIfNode ? FRigVMDispatch_If::StaticStruct() : FRigVMDispatch_SelectInt32::StaticStruct());
 
 			FRigVMTemplate* Template = const_cast<FRigVMTemplate*>(Factory->GetTemplate());
+			const FString NewNodeName = NodeName + TEXT("_") + Factory->GetFactoryName().ToString();
+			FRestoreLinkedPathSettings RestoreLinkedPathSettings;
+			RestoreLinkedPathSettings.NodeNameMap.Add(NodeName, NewNodeName);
 			URigVMTemplateNode* NewNode = AddTemplateNode(
 				Template->GetNotation(),
 				NodePosition, 
-				NodeName, 
+				NewNodeName, 
 				false, 
 				false);
 
@@ -19572,8 +19585,7 @@ FRigVMClientPatchResult URigVMController::PatchIfSelectNodesOnLoad()
 
 			ApplyPinStates(NewNode, PinStates, {}, false);
 
-			FRestoreLinkedPathSettings Settings;
-			RestoreLinkedPaths(LinkedPaths, Settings);
+			RestoreLinkedPaths(LinkedPaths, RestoreLinkedPathSettings);
 
 			RemoveNode(IfOrSelectNode, false, false);
 		}
@@ -19616,14 +19628,16 @@ FRigVMClientPatchResult URigVMController::PatchArrayNodesOnLoad()
 			const FString DeletedName = GetSchema()->GetValidNodeName(Graph, FString::Printf(TEXT("%s_Deleted"), *NodeName));
 			RenameNode(ArrayNode, *DeletedName, false);
 			
-			URigVMNode* NewNode = AddArrayNode(OpCode, CPPType, CPPTypeObject, NodePosition, NodeName, false, false, true);
+			const FString NewNodeName = NodeName + TEXT("_") + StaticEnum<ERigVMOpCode>()->GetDisplayNameTextByValue((int64)OpCode).ToString();
+			FRestoreLinkedPathSettings RestoreLinkedPathSettings;
+			RestoreLinkedPathSettings.NodeNameMap.Add(NodeName, NewNodeName);
+			URigVMNode* NewNode = AddArrayNode(OpCode, CPPType, CPPTypeObject, NodePosition, NewNodeName, false, false, true);
 			ApplyPinStates(NewNode, PinStates, {}, false);
 			Result.AddedNodes.Add(NewNode);
 			
 			RemoveNode(ArrayNode, false, false);
 
-			FRestoreLinkedPathSettings Settings;
-			RestoreLinkedPaths(LinkedPaths, Settings);
+			RestoreLinkedPaths(LinkedPaths, RestoreLinkedPathSettings);
 		}
 	}
 
