RigVM: Follow up to: Ensure to not re-use node names during patching to avoid linker errors

#rnx

diff --git a/Engine/Plugins/Runtime/RigVM/Source/RigVMDeveloper/Private/RigVMModel/RigVMController.cpp b/Engine/Plugins/Runtime/RigVM/Source/RigVMDeveloper/Private/RigVMModel/RigVMController.cpp
--- a/Engine/Plugins/Runtime/RigVM/Source/RigVMDeveloper/Private/RigVMModel/RigVMController.cpp
+++ b/Engine/Plugins/Runtime/RigVM/Source/RigVMDeveloper/Private/RigVMModel/RigVMController.cpp
@@ -18934,13 +18934,16 @@ URigVMNode* URigVMController::ConvertRerouteNodeToDispatch(URigVMRerouteNode* In
 		}
 		ApplyPinStates(NewNode, RemappedPinStates, {}, bSetupUndoRedo);
 
-		const FString NodeNamePrefix = URigVMPin::JoinPinPath({NewNodeName, FString()});
-		RestoreLinkedPathSettings.RemapDelegates.Add(NewNodeName,
-			FRigVMController_PinPathRemapDelegate::CreateLambda([NodeNamePrefix, InputRedirects, OutputRedirects](const FString& InPinPath, bool bIsInput) -> FString
+		RestoreLinkedPathSettings.RemapDelegates.Add(NodeName,
+			FRigVMController_PinPathRemapDelegate::CreateLambda([NodeName, NewNodeName, InputRedirects, OutputRedirects](const FString& InPinPath, bool bIsInput) -> FString
 			{
 				TArray<FString> Parts;
 				if(URigVMPin::SplitPinPath(InPinPath, Parts))
 				{
+					if(Parts[0].Equals(NodeName, ESearchCase::CaseSensitive))
+					{
+						Parts[0] = NewNodeName;
+					}
 					const TMap<FString, FString>* Redirects = bIsInput ? InputRedirects : OutputRedirects;
 					if(const FString* RedirectedPart = Redirects->Find(Parts[1]))
 					{
