Control Rig: Fix crash when modifying input/output pins on public functions
#rnx

diff --git a/Engine/Plugins/Runtime/RigVM/Source/RigVMDeveloper/Private/RigVMModel/Nodes/RigVMFunctionReferenceNode.cpp b/Engine/Plugins/Runtime/RigVM/Source/RigVMDeveloper/Private/RigVMModel/Nodes/RigVMFunctionReferenceNode.cpp
--- a/Engine/Plugins/Runtime/RigVM/Source/RigVMDeveloper/Private/RigVMModel/Nodes/RigVMFunctionReferenceNode.cpp
+++ b/Engine/Plugins/Runtime/RigVM/Source/RigVMDeveloper/Private/RigVMModel/Nodes/RigVMFunctionReferenceNode.cpp
@@ -166,6 +166,14 @@ uint32 URigVMFunctionReferenceNode::GetStructureHash() const
 	return Hash;
 }
 
+void URigVMFunctionReferenceNode::UpdateFunctionHeaderFromHost()
+{
+	if (const FRigVMGraphFunctionData* Data = GetReferencedFunctionData())
+	{
+		ReferencedFunctionHeader = Data->Header;
+	}
+}
+
 const FRigVMGraphFunctionData* URigVMFunctionReferenceNode::GetReferencedFunctionData(bool bLoadIfNecessary) const
 {
 	if (IRigVMGraphFunctionHost* Host = ReferencedFunctionHeader.GetFunctionHost(bLoadIfNecessary))
diff --git a/Engine/Plugins/Runtime/RigVM/Source/RigVMDeveloper/Private/RigVMModel/RigVMBuildData.cpp b/Engine/Plugins/Runtime/RigVM/Source/RigVMDeveloper/Private/RigVMModel/RigVMBuildData.cpp
--- a/Engine/Plugins/Runtime/RigVM/Source/RigVMDeveloper/Private/RigVMModel/RigVMBuildData.cpp
+++ b/Engine/Plugins/Runtime/RigVM/Source/RigVMDeveloper/Private/RigVMModel/RigVMBuildData.cpp
@@ -172,14 +172,15 @@ const FRigVMFunctionReferenceArray* URigVMBuildData::FindFunctionReferences(cons
 }
 
 void URigVMBuildData::ForEachFunctionReference(const FRigVMGraphFunctionIdentifier& InFunction,
-                                               TFunction<void(URigVMFunctionReferenceNode*)> PerReferenceFunction) const
+                                               TFunction<void(URigVMFunctionReferenceNode*)> PerReferenceFunction,
+                                               bool bLoadIfNecessary) const
 {
 	if (const FRigVMFunctionReferenceArray* ReferencesEntry = FindFunctionReferences(InFunction))
 	{
 		for (int32 ReferenceIndex = 0; ReferenceIndex < ReferencesEntry->Num(); ReferenceIndex++)
 		{
 			const TSoftObjectPtr<URigVMFunctionReferenceNode>& Reference = ReferencesEntry->operator [](ReferenceIndex);
-			if (!Reference.IsValid())
+			if (bLoadIfNecessary && !Reference.IsValid())
 			{
 				Reference.LoadSynchronous();
 			}
diff --git a/Engine/Plugins/Runtime/RigVM/Source/RigVMDeveloper/Private/RigVMModel/RigVMController.cpp b/Engine/Plugins/Runtime/RigVM/Source/RigVMDeveloper/Private/RigVMModel/RigVMController.cpp
--- a/Engine/Plugins/Runtime/RigVM/Source/RigVMDeveloper/Private/RigVMModel/RigVMController.cpp
+++ b/Engine/Plugins/Runtime/RigVM/Source/RigVMDeveloper/Private/RigVMModel/RigVMController.cpp
@@ -10590,7 +10590,7 @@ FName URigVMController::AddExposedPin(const FName& InPinName, ERigVMPinDirection
 	RefreshFunctionPins(EntryNode);
 	RefreshFunctionPins(ReturnNode);
 
-	RefreshFunctionReferences(LibraryNode, bSetupUndoRedo);
+	RefreshFunctionReferences(LibraryNode, bSetupUndoRedo, false);
 
 	if (bSetupUndoRedo)
 	{
@@ -10698,7 +10698,7 @@ bool URigVMController::RemoveExposedPin(const FName& InPinName, bool bSetupUndoR
 
 		RefreshFunctionPins(Graph->GetEntryNode(), bSetupUndoRedo);
 		RefreshFunctionPins(Graph->GetReturnNode(), bSetupUndoRedo);
-		RefreshFunctionReferences(LibraryNode, bSetupUndoRedo);
+		RefreshFunctionReferences(LibraryNode, bSetupUndoRedo, false);
 	}
 
 	if (bSetupUndoRedo)
@@ -11123,7 +11123,7 @@ bool URigVMController::ChangeExposedPinType(const FName& InPinName, const FStrin
 	// Change pin type on function references
 	if (URigVMFunctionLibrary* FunctionLibrary = Cast<URigVMFunctionLibrary>(LibraryNode->GetGraph()))
 	{
-		RefreshFunctionReferences(LibraryNode, bSetupUndoRedo);
+		RefreshFunctionReferences(LibraryNode, bSetupUndoRedo, false);
 	}
 
 	// Change pin types on input variable nodes
@@ -11226,7 +11226,7 @@ bool URigVMController::SetExposedPinIndex(const FName& InPinName, int32 InNewInd
 
 	RefreshFunctionPins(LibraryNode->GetEntryNode());
 	RefreshFunctionPins(LibraryNode->GetReturnNode());
-	RefreshFunctionReferences(LibraryNode, false);
+	RefreshFunctionReferences(LibraryNode, false, false);
 	
 	if (bSetupUndoRedo)
 	{
@@ -15287,6 +15287,7 @@ bool URigVMController::GenerateNewPinInfos(const FRigVMRegistry& Registry, URigV
 	}
 	else if (FunctionRefNode)
 	{
+		FunctionRefNode->UpdateFunctionHeaderFromHost();
 		const FRigVMGraphFunctionHeader& FunctionHeader = FunctionRefNode->GetReferencedFunctionHeader();
 		if (FunctionHeader.IsValid())
 		{
@@ -18159,23 +18160,28 @@ const FRigVMByteCode* URigVMController::GetCurrentByteCode() const
 	return nullptr;
 }
 
-void URigVMController::RefreshFunctionReferences(URigVMLibraryNode* InFunctionDefinition, bool bSetupUndoRedo)
+void URigVMController::RefreshFunctionReferences(URigVMLibraryNode* InFunctionDefinition, bool bSetupUndoRedo, bool bLoadIfNecessary)
 {
 	check(InFunctionDefinition);
 
 	if (const URigVMFunctionLibrary* FunctionLibrary = Cast<URigVMFunctionLibrary>(InFunctionDefinition->GetGraph()))
 	{
-		FunctionLibrary->ForEachReference(InFunctionDefinition->GetFName(), [this, bSetupUndoRedo](URigVMFunctionReferenceNode* ReferenceNode)
+		TMap<URigVMController*,TSharedPtr<FRigVMControllerCompileBracketScope>> CompilationBrackets;
+		FunctionLibrary->ForEachReference(InFunctionDefinition->GetFName(), [this, bSetupUndoRedo, &CompilationBrackets](URigVMFunctionReferenceNode* ReferenceNode)
 		{
 			if(URigVMController* ReferenceController = GetControllerForGraph(ReferenceNode->GetGraph()))
 			{
+				if (!CompilationBrackets.Contains(ReferenceController))
+				{
+					CompilationBrackets.FindOrAdd(ReferenceController) = MakeShared<FRigVMControllerCompileBracketScope>(ReferenceController);
+				}
 				const TArray<FLinkedPath> LinkedPaths = GetLinkedPaths(ReferenceNode->GetLinks());
 				ReferenceController->FastBreakLinkedPaths(LinkedPaths, bSetupUndoRedo);
 				ReferenceController->RepopulatePinsOnNode(ReferenceNode, false, false, true);
 				TGuardValue<bool> ReportGuard(ReferenceController->bReportWarningsAndErrors, false);
 				ReferenceController->RestoreLinkedPaths(LinkedPaths, FRestoreLinkedPathSettings(), bSetupUndoRedo);
 			}
-		});
+		}, bLoadIfNecessary);
 	}
 }
 
diff --git a/Engine/Plugins/Runtime/RigVM/Source/RigVMDeveloper/Private/RigVMModel/RigVMFunctionLibrary.cpp b/Engine/Plugins/Runtime/RigVM/Source/RigVMDeveloper/Private/RigVMModel/RigVMFunctionLibrary.cpp
--- a/Engine/Plugins/Runtime/RigVM/Source/RigVMDeveloper/Private/RigVMModel/RigVMFunctionLibrary.cpp
+++ b/Engine/Plugins/Runtime/RigVM/Source/RigVMDeveloper/Private/RigVMModel/RigVMFunctionLibrary.cpp
@@ -96,13 +96,14 @@ TArray< FString > URigVMFunctionLibrary::GetReferencePathsForFunction(const FNam
 }
 
 void URigVMFunctionLibrary::ForEachReference(const FName& InFunctionName,
-	TFunction<void(URigVMFunctionReferenceNode*)> PerReferenceFunction) const
+	TFunction<void(URigVMFunctionReferenceNode*)> PerReferenceFunction,
+	bool bLoadIfNecessary) const
 {
 	if (URigVMLibraryNode* Function = FindFunction(InFunctionName))
 	{
 		if(const URigVMBuildData* BuildData = URigVMBuildData::Get())
 		{
-			BuildData->ForEachFunctionReference(Function->GetFunctionIdentifier(),PerReferenceFunction);
+			BuildData->ForEachFunctionReference(Function->GetFunctionIdentifier(),PerReferenceFunction, bLoadIfNecessary);
 		}
 	}
 }
diff --git a/Engine/Plugins/Runtime/RigVM/Source/RigVMDeveloper/Public/RigVMModel/Nodes/RigVMFunctionReferenceNode.h b/Engine/Plugins/Runtime/RigVM/Source/RigVMDeveloper/Public/RigVMModel/Nodes/RigVMFunctionReferenceNode.h
--- a/Engine/Plugins/Runtime/RigVM/Source/RigVMDeveloper/Public/RigVMModel/Nodes/RigVMFunctionReferenceNode.h
+++ b/Engine/Plugins/Runtime/RigVM/Source/RigVMDeveloper/Public/RigVMModel/Nodes/RigVMFunctionReferenceNode.h
@@ -53,6 +53,8 @@ public:
 
 	const FRigVMGraphFunctionHeader& GetReferencedFunctionHeader() const { return ReferencedFunctionHeader; }
 
+	void UpdateFunctionHeaderFromHost();
+
 	const FRigVMGraphFunctionData* GetReferencedFunctionData(bool bLoadIfNecessary = true) const;
 
 	
diff --git a/Engine/Plugins/Runtime/RigVM/Source/RigVMDeveloper/Public/RigVMModel/RigVMBuildData.h b/Engine/Plugins/Runtime/RigVM/Source/RigVMDeveloper/Public/RigVMModel/RigVMBuildData.h
--- a/Engine/Plugins/Runtime/RigVM/Source/RigVMDeveloper/Public/RigVMModel/RigVMBuildData.h
+++ b/Engine/Plugins/Runtime/RigVM/Source/RigVMDeveloper/Public/RigVMModel/RigVMBuildData.h
@@ -86,8 +86,9 @@ public:
 	 * Iterator function to invoke a lambda / TFunction for each reference of a function
 	 * @param InFunction The function to iterate all references for
 	 * @param PerReferenceFunction The function to invoke for each reference
+	 * @param bLoadIfNecessary If true, will load packages if necessary
 	 */
-	void ForEachFunctionReference(const FRigVMGraphFunctionIdentifier& InFunction, TFunction<void(URigVMFunctionReferenceNode*)> PerReferenceFunction) const;
+	void ForEachFunctionReference(const FRigVMGraphFunctionIdentifier& InFunction, TFunction<void(URigVMFunctionReferenceNode*)> PerReferenceFunction, bool bLoadIfNecessary = true) const;
 
 	/**
 	* Iterator function to invoke a lambda / TFunction for each reference of a function
diff --git a/Engine/Plugins/Runtime/RigVM/Source/RigVMDeveloper/Public/RigVMModel/RigVMController.h b/Engine/Plugins/Runtime/RigVM/Source/RigVMDeveloper/Public/RigVMModel/RigVMController.h
--- a/Engine/Plugins/Runtime/RigVM/Source/RigVMDeveloper/Public/RigVMModel/RigVMController.h
+++ b/Engine/Plugins/Runtime/RigVM/Source/RigVMDeveloper/Public/RigVMModel/RigVMController.h
@@ -1249,7 +1249,7 @@ private:
 	FRigVMExternalVariable GetVariableByName(const FName& InExternalVariableName, const bool bIncludeInputArguments = false) const;
 	TArray<FRigVMExternalVariable> GetAllVariables(const bool bIncludeInputArguments = false) const;
 
-	void RefreshFunctionReferences(URigVMLibraryNode* InFunctionDefinition, bool bSetupUndoRedo);
+	void RefreshFunctionReferences(URigVMLibraryNode* InFunctionDefinition, bool bSetupUndoRedo, bool bLoadIfNecessary);
 
 public:
 
diff --git a/Engine/Plugins/Runtime/RigVM/Source/RigVMDeveloper/Public/RigVMModel/RigVMFunctionLibrary.h b/Engine/Plugins/Runtime/RigVM/Source/RigVMDeveloper/Public/RigVMModel/RigVMFunctionLibrary.h
--- a/Engine/Plugins/Runtime/RigVM/Source/RigVMDeveloper/Public/RigVMModel/RigVMFunctionLibrary.h
+++ b/Engine/Plugins/Runtime/RigVM/Source/RigVMDeveloper/Public/RigVMModel/RigVMFunctionLibrary.h
@@ -54,8 +54,9 @@ public:
 	* Iterator function to invoke a lambda / TFunction for each reference of a function
 	* @param InFunctionName The function name to iterate all references for
 	* @param PerReferenceFunction The function to invoke for each reference
+	* @param bLoadIfNecessary If true, will load packages when needed
 	*/
-	void ForEachReference(const FName& InFunctionName, TFunction<void(URigVMFunctionReferenceNode*)> PerReferenceFunction) const;
+	void ForEachReference(const FName& InFunctionName, TFunction<void(URigVMFunctionReferenceNode*)> PerReferenceFunction, bool bLoadIfNecessary = true) const;
 
 	/**
 	* Iterator function to invoke a lambda / TFunction for each reference of a function
