RigVM: First steps torwards external dependency gathering mechanism

#rnx

diff --git a/Engine/Plugins/Runtime/RigVM/Source/RigVMDeveloper/Private/RigVMBlueprint.cpp b/Engine/Plugins/Runtime/RigVM/Source/RigVMDeveloper/Private/RigVMBlueprint.cpp
--- a/Engine/Plugins/Runtime/RigVM/Source/RigVMDeveloper/Private/RigVMBlueprint.cpp
+++ b/Engine/Plugins/Runtime/RigVM/Source/RigVMDeveloper/Private/RigVMBlueprint.cpp
@@ -2288,6 +2288,30 @@ TArray<FString> URigVMBlueprint::GeneratePythonCommands(const FString InNewBluep
 	return Commands;
 }
 
+TArray<FRigVMExternalDependency> URigVMBlueprint::GetExternalDependenciesForCategory(const FName& InCategory) const
+{
+	TArray<FRigVMExternalDependency> Dependencies;
+	if(const FRigVMClient* Client = GetRigVMClient())
+	{
+		CollectExternalDependencies(Dependencies, InCategory, Client);
+	}
+	if(const IRigVMGraphFunctionHost* FunctionHost = GetRigVMGraphFunctionHost())
+	{
+		if(const FRigVMGraphFunctionStore* FunctionStore = FunctionHost->GetRigVMGraphFunctionStore())
+		{
+			CollectExternalDependencies(Dependencies, InCategory, FunctionStore);
+		}
+	}
+
+#if WITH_EDITOR
+	const TArray<FRigVMGraphVariableDescription> MemberVariables = GetMemberVariables();
+	for(const FRigVMGraphVariableDescription& MemberVariable : MemberVariables)
+	{
+		CollectExternalDependenciesForCPPTypeObject(Dependencies, InCategory, MemberVariable.CPPTypeObject.Get());
+	}
+#endif
+	return Dependencies;
+}
 
 URigVMGraph* URigVMBlueprint::GetTemplateModel(bool bIsFunctionLibrary)
 {
diff --git a/Engine/Plugins/Runtime/RigVM/Source/RigVMDeveloper/Private/RigVMModel/RigVMExternalDependency.cpp b/Engine/Plugins/Runtime/RigVM/Source/RigVMDeveloper/Private/RigVMModel/RigVMExternalDependency.cpp
new file mode 100644
--- /dev/null
+++ b/Engine/Plugins/Runtime/RigVM/Source/RigVMDeveloper/Private/RigVMModel/RigVMExternalDependency.cpp
@@ -0,0 +1,208 @@
+// Copyright Epic Games, Inc. All Rights Reserved.
+
+#include "RigVMModel/RigVMExternalDependency.h"
+
+#include UE_INLINE_GENERATED_CPP_BY_NAME(RigVMExternalDependency)
+
+const FName IRigVMExternalDependencyManager::UserDefinedEnumCategory = TEXT("UserDefinedEnum");
+const FName IRigVMExternalDependencyManager::UserDefinedStructCategory = TEXT("UserDefinedStruct");
+const FName IRigVMExternalDependencyManager::RigVMGraphFunctionCategory = TEXT("RigVMGraphFunction");
+
+TArray<FName> IRigVMExternalDependencyManager::DependencyCategories = {
+	UserDefinedEnumCategory,
+	UserDefinedStructCategory,
+	RigVMGraphFunctionCategory
+};
+
+const TArray<FName>& IRigVMExternalDependencyManager::GetExternalDependencyCategories() const
+{
+	return IRigVMExternalDependencyManager::DependencyCategories;
+}
+
+TArray<FRigVMExternalDependency> IRigVMExternalDependencyManager::GetAllExternalDependencies() const
+{
+	const TArray<FName>& Categories = GetExternalDependencyCategories();
+	TArray<FRigVMExternalDependency> Dependencies;
+	for(const FName& Category : Categories)
+	{
+		Dependencies.Append(GetExternalDependenciesForCategory(Category));
+	}
+	return Dependencies;
+}
+
+void IRigVMExternalDependencyManager::CollectExternalDependencies(TArray<FRigVMExternalDependency>& OutDependencies, const FName& InCategory, const FRigVMClient* InClient) const
+{
+	check(InClient);
+
+	const TArray<URigVMGraph*> Graphs = InClient->GetAllModels(true, true);
+	for(const URigVMGraph* Graph : Graphs)
+	{
+		CollectExternalDependencies(OutDependencies, InCategory, Graph);
+	}
+}
+
+void IRigVMExternalDependencyManager::CollectExternalDependencies(TArray<FRigVMExternalDependency>& OutDependencies, const FName& InCategory, const FRigVMGraphFunctionStore* InFunctionStore) const
+{
+	check(InFunctionStore);
+	for(const FRigVMGraphFunctionData& Function : InFunctionStore->PublicFunctions)
+	{
+		CollectExternalDependencies(OutDependencies, InCategory, &Function);
+	}
+	for(const FRigVMGraphFunctionData& Function : InFunctionStore->PrivateFunctions)
+	{
+		CollectExternalDependencies(OutDependencies, InCategory, &Function);
+	}
+}
+
+void IRigVMExternalDependencyManager::CollectExternalDependencies(TArray<FRigVMExternalDependency>& OutDependencies, const FName& InCategory, const FRigVMGraphFunctionData* InFunction) const
+{
+	check(InFunction);
+	CollectExternalDependencies(OutDependencies, InCategory, &InFunction->Header);
+	CollectExternalDependencies(OutDependencies, InCategory, &InFunction->CompilationData);
+}
+
+void IRigVMExternalDependencyManager::CollectExternalDependencies(TArray<FRigVMExternalDependency>& OutDependencies, const FName& InCategory, const FRigVMGraphFunctionHeader* InHeader) const
+{
+	check(InHeader);
+	if(InCategory == RigVMGraphFunctionCategory)
+	{
+		OutDependencies.AddUnique({InHeader->LibraryPointer.LibraryNode.ToString(), InCategory});
+	}
+	for(const FRigVMGraphFunctionArgument& Argument : InHeader->Arguments)
+	{
+		CollectExternalDependenciesForCPPTypeObject(OutDependencies, InCategory, Argument.CPPTypeObject.Get());
+	}
+}
+
+void IRigVMExternalDependencyManager::CollectExternalDependencies(TArray<FRigVMExternalDependency>& OutDependencies, const FName& InCategory, const FRigVMFunctionCompilationData* InFunction) const
+{
+	check(InFunction);
+
+	for(const FRigVMFunctionCompilationPropertyDescription& Property : InFunction->LiteralPropertyDescriptions)
+	{
+		CollectExternalDependenciesForCPPTypeObject(OutDependencies, InCategory, Property.CPPTypeObject.Get());
+	}
+	for(const FRigVMFunctionCompilationPropertyDescription& Property : InFunction->WorkPropertyDescriptions)
+	{
+		CollectExternalDependenciesForCPPTypeObject(OutDependencies, InCategory, Property.CPPTypeObject.Get());
+	}
+
+	const FRigVMRegistry& Registry = FRigVMRegistry::Get();
+	for(const FName& FunctionName : InFunction->FunctionNames)
+	{
+		if(const FRigVMFunction* Function = Registry.FindFunction(*FunctionName.ToString()))
+		{
+			for(const FRigVMFunctionArgument& Argument :  Function->Arguments)
+			{
+				const FRigVMTemplateArgumentType& ArgumentType = Registry.FindTypeFromCPPType(Argument.Type);
+				if(ArgumentType.IsValid())
+				{
+					CollectExternalDependenciesForCPPTypeObject(OutDependencies, InCategory, ArgumentType.CPPTypeObject.Get());
+				}
+			}
+		}
+	}
+}
+
+void IRigVMExternalDependencyManager::CollectExternalDependencies(TArray<FRigVMExternalDependency>& OutDependencies, const FName& InCategory, const URigVMGraph* InGraph) const
+{
+	for(const URigVMNode* Node : InGraph->GetNodes())
+	{
+		CollectExternalDependencies(OutDependencies, InCategory, Node);
+	}
+
+	const TArray<FRigVMGraphVariableDescription> LocalVariables = InGraph->GetLocalVariables();
+	for(const FRigVMGraphVariableDescription& LocalVariable : LocalVariables)
+	{
+		CollectExternalDependenciesForCPPTypeObject(OutDependencies, InCategory, LocalVariable.CPPTypeObject.Get());
+	}
+}
+
+void IRigVMExternalDependencyManager::CollectExternalDependencies(TArray<FRigVMExternalDependency>& OutDependencies, const FName& InCategory, const URigVMNode* InNode) const
+{
+	for(const URigVMPin* Pin : InNode->GetPins())
+	{
+		CollectExternalDependencies(OutDependencies, InCategory, Pin);
+	}
+	if(const URigVMFunctionReferenceNode* FunctionReferenceNode = Cast<URigVMFunctionReferenceNode>(InNode))
+	{
+		if(const FRigVMGraphFunctionData* Function = FunctionReferenceNode->GetReferencedFunctionData(true))
+		{
+			CollectExternalDependencies(OutDependencies, InCategory, Function);
+		}
+	}
+}
+
+void IRigVMExternalDependencyManager::CollectExternalDependencies(TArray<FRigVMExternalDependency>& OutDependencies, const FName& InCategory, const URigVMPin* InPin) const
+{
+	CollectExternalDependenciesForCPPTypeObject(OutDependencies, InCategory, InPin->GetCPPTypeObject());
+	for(const URigVMPin* SubPin : InPin->GetSubPins())
+	{
+		CollectExternalDependencies(OutDependencies, InCategory, SubPin);
+	}
+}
+
+void IRigVMExternalDependencyManager::CollectExternalDependencies(TArray<FRigVMExternalDependency>& OutDependencies, const FName& InCategory, const UStruct* InStruct) const
+{
+	check(InStruct);
+	if(InCategory == UserDefinedStructCategory)
+	{
+		if(const UUserDefinedStruct* UserDefinedStruct = Cast<UUserDefinedStruct>(InStruct))
+		{
+			OutDependencies.AddUnique({UserDefinedStruct->GetPathName(), InCategory});
+		}
+	}
+	for (TFieldIterator<FProperty> PropertyIt(InStruct); PropertyIt; ++PropertyIt)
+	{
+		const FProperty* Property = *PropertyIt;
+		while(const FArrayProperty* ArrayProperty = CastField<FArrayProperty>(Property))
+		{
+			Property = ArrayProperty->Inner;
+		}
+		if(const FStructProperty* StructProperty = CastField<FStructProperty>(Property))
+		{
+			CollectExternalDependencies(OutDependencies, InCategory, StructProperty->Struct);
+		}
+		else if(const FEnumProperty* EnumProperty = CastField<FEnumProperty>(Property))
+		{
+			if(const UEnum* Enum = EnumProperty->GetEnum())
+			{
+				CollectExternalDependencies(OutDependencies, InCategory, Enum);
+			}
+		}
+		else if(const FByteProperty* ByteProperty = CastField<FByteProperty>(Property))
+		{
+			if(const UEnum* Enum = ByteProperty->Enum)
+			{
+				CollectExternalDependencies(OutDependencies, InCategory, Enum);
+			}
+		}
+	}
+}
+
+void IRigVMExternalDependencyManager::CollectExternalDependencies(TArray<FRigVMExternalDependency>& OutDependencies, const FName& InCategory, const UEnum* InEnum) const
+{
+	check(InEnum);
+	if(InCategory == UserDefinedEnumCategory)
+	{
+		if(const UUserDefinedEnum* UserDefinedEnum = Cast<UUserDefinedEnum>(InEnum))
+		{
+			OutDependencies.AddUnique({UserDefinedEnum->GetPathName(), InCategory});
+		}
+	}
+}
+
+void IRigVMExternalDependencyManager::CollectExternalDependenciesForCPPTypeObject(TArray<FRigVMExternalDependency>& OutDependencies, const FName& InCategory, const UObject* InObject) const
+{
+	if(InObject)
+	{
+		if(const UEnum* Enum = Cast<UEnum>(InObject))
+		{
+			CollectExternalDependencies(OutDependencies, InCategory, Enum);
+		}
+		else if(const UStruct* Struct = Cast<UStruct>(InObject))
+		{
+			CollectExternalDependencies(OutDependencies, InCategory, Struct);
+		}
+	}
+}
diff --git a/Engine/Plugins/Runtime/RigVM/Source/RigVMDeveloper/Public/RigVMBlueprint.h b/Engine/Plugins/Runtime/RigVM/Source/RigVMDeveloper/Public/RigVMBlueprint.h
--- a/Engine/Plugins/Runtime/RigVM/Source/RigVMDeveloper/Public/RigVMBlueprint.h
+++ b/Engine/Plugins/Runtime/RigVM/Source/RigVMDeveloper/Public/RigVMBlueprint.h
@@ -8,6 +8,7 @@
 #include "RigVMCore/RigVM.h"
 #include "RigVMHost.h"
 #include "RigVMModel/RigVMClient.h"
+#include "RigVMModel/RigVMExternalDependency.h"
 #include "RigVMCompiler/RigVMCompiler.h"
 #include "RigVMCore/RigVMGraphFunctionDefinition.h"
 #include "EdGraph/RigVMEdGraph.h"
@@ -214,7 +215,7 @@ struct RIGVMDEVELOPER_API FRigVMOldPublicFunctionData
 
 
 UCLASS(BlueprintType, meta=(IgnoreClassThumbnail))
-class RIGVMDEVELOPER_API URigVMBlueprint : public UBlueprint, public IRigVMClientHost
+class RIGVMDEVELOPER_API URigVMBlueprint : public UBlueprint, public IRigVMClientHost, public IRigVMExternalDependencyManager
 {
 	GENERATED_UCLASS_BODY()
 
@@ -401,6 +402,12 @@ public:
 
 	//  --- IRigVMClientHost interface End ---
 
+	//  --- IRigVMExternalDependencyManager interface Start ---
+
+	virtual TArray<FRigVMExternalDependency> GetExternalDependenciesForCategory(const FName& InCategory) const override;
+	
+	//  --- IRigVMExternalDependencyManager interface End ---
+
 
 	FOnRigVMRequestInspectObject& OnRequestInspectObject() { return OnRequestInspectObjectEvent; }
 	void RequestInspectObject(const TArray<UObject*>& InObjects) { OnRequestInspectObjectEvent.Broadcast(InObjects); }
diff --git a/Engine/Plugins/Runtime/RigVM/Source/RigVMDeveloper/Public/RigVMModel/RigVMExternalDependency.h b/Engine/Plugins/Runtime/RigVM/Source/RigVMDeveloper/Public/RigVMModel/RigVMExternalDependency.h
new file mode 100644
--- /dev/null
+++ b/Engine/Plugins/Runtime/RigVM/Source/RigVMDeveloper/Public/RigVMModel/RigVMExternalDependency.h
@@ -0,0 +1,91 @@
+// Copyright Epic Games, Inc. All Rights Reserved.
+
+#pragma once
+
+#include "RigVMClient.h"
+#include "RigVMCore/RigVMGraphFunctionHost.h"
+#include "RigVMExternalDependency.generated.h"
+
+USTRUCT()
+struct RIGVMDEVELOPER_API FRigVMExternalDependency
+{
+	GENERATED_BODY()
+
+public:
+
+	FRigVMExternalDependency()
+		: ExternalPath()
+		, InternalPath()
+		, Category(NAME_None)
+	{}
+
+	FRigVMExternalDependency(const FString& InPath, const FName& InCategory)
+		: ExternalPath(InPath)
+		, InternalPath()
+		, Category(InCategory)
+	{}
+
+	bool IsExternal() const { return InternalPath.IsEmpty(); }
+	bool IsInternal() const { return !IsExternal(); }
+
+	const FString& GetExternalPath() const { return ExternalPath; }
+	const FString& GetInternalPath() const { return InternalPath; }
+	const FName& GetCategory() const { return Category; }
+
+	bool operator==(const FRigVMExternalDependency& InOther) const
+	{
+		return ExternalPath == InOther.ExternalPath &&
+			InternalPath == InOther.InternalPath &&
+			Category == InOther.Category;
+	}
+
+private:
+
+	UPROPERTY()
+	FString ExternalPath;
+
+	UPROPERTY()
+	FString InternalPath;
+
+	UPROPERTY()
+	FName Category;
+};
+
+UINTERFACE()
+class RIGVMDEVELOPER_API URigVMExternalDependencyManager : public UInterface
+{
+	GENERATED_BODY()
+};
+
+// Interface to deal with mapping external dependencies
+class RIGVMDEVELOPER_API IRigVMExternalDependencyManager
+{
+	GENERATED_BODY()
+
+public:
+	
+	virtual const TArray<FName>& GetExternalDependencyCategories() const;
+	virtual TArray<FRigVMExternalDependency> GetExternalDependenciesForCategory(const FName& InCategory) const = 0;
+	TArray<FRigVMExternalDependency> GetAllExternalDependencies() const;
+
+	static const FName UserDefinedEnumCategory;
+	static const FName UserDefinedStructCategory;
+	static const FName RigVMGraphFunctionCategory;
+
+protected:
+
+	virtual void CollectExternalDependencies(TArray<FRigVMExternalDependency>& OutDependencies, const FName& InCategory, const FRigVMClient* InClient) const;
+	virtual void CollectExternalDependencies(TArray<FRigVMExternalDependency>& OutDependencies, const FName& InCategory, const FRigVMGraphFunctionStore* InFunctionStore) const;
+	virtual void CollectExternalDependencies(TArray<FRigVMExternalDependency>& OutDependencies, const FName& InCategory, const FRigVMGraphFunctionData* InFunction) const;
+	virtual void CollectExternalDependencies(TArray<FRigVMExternalDependency>& OutDependencies, const FName& InCategory, const FRigVMGraphFunctionHeader* InHeader) const;
+	virtual void CollectExternalDependencies(TArray<FRigVMExternalDependency>& OutDependencies, const FName& InCategory, const FRigVMFunctionCompilationData* InCompilationData) const;
+	virtual void CollectExternalDependencies(TArray<FRigVMExternalDependency>& OutDependencies, const FName& InCategory, const URigVMGraph* InGraph) const;
+	virtual void CollectExternalDependencies(TArray<FRigVMExternalDependency>& OutDependencies, const FName& InCategory, const URigVMNode* InNode) const;
+	virtual void CollectExternalDependencies(TArray<FRigVMExternalDependency>& OutDependencies, const FName& InCategory, const URigVMPin* InPin) const;
+	virtual void CollectExternalDependencies(TArray<FRigVMExternalDependency>& OutDependencies, const FName& InCategory, const UStruct* InStruct) const;
+	virtual void CollectExternalDependencies(TArray<FRigVMExternalDependency>& OutDependencies, const FName& InCategory, const UEnum* InEnum) const;
+	virtual void CollectExternalDependenciesForCPPTypeObject(TArray<FRigVMExternalDependency>& OutDependencies, const FName& InCategory, const UObject* InObject) const;
+
+	static TArray<FName> DependencyCategories;
+};
+
