ControlRig / RigVM: Make sure to only apply pin factory on matching assets

#rnx

diff --git a/Engine/Plugins/Animation/ControlRig/Source/ControlRigEditor/Private/Graph/ControlRigGraphPanelPinFactory.cpp b/Engine/Plugins/Animation/ControlRig/Source/ControlRigEditor/Private/Graph/ControlRigGraphPanelPinFactory.cpp
--- a/Engine/Plugins/Animation/ControlRig/Source/ControlRigEditor/Private/Graph/ControlRigGraphPanelPinFactory.cpp
+++ b/Engine/Plugins/Animation/ControlRig/Source/ControlRigEditor/Private/Graph/ControlRigGraphPanelPinFactory.cpp
@@ -25,11 +25,26 @@
 
 TSharedPtr<SGraphPin> FControlRigGraphPanelPinFactory::CreatePin(UEdGraphPin* InPin) const
 {
+	if(InPin == nullptr)
+	{
+		return nullptr;
+	}
+	
 	TSharedPtr<SGraphPin> InternalResult = CreatePin_Internal(InPin);
 	if(InternalResult.IsValid())
 	{
 		return InternalResult;
 	}
+
+	// if the graph we are looking at is not a control rig graph - let's not do this
+	if (const UEdGraphNode* OwningNode = InPin->GetOwningNode())
+	{
+		// only create pins within control rig graphs
+		if (Cast<UControlRigGraph>(OwningNode->GetGraph()) == nullptr)
+		{
+			return nullptr;
+		}
+	}
 	
 	TSharedPtr<SGraphPin> K2PinWidget = FNodeFactory::CreateK2PinWidget(InPin);
 	if(K2PinWidget.IsValid())
diff --git a/Engine/Plugins/Runtime/RigVM/Source/RigVMEditor/Private/EdGraph/RigVMEdGraphPanelPinFactory.cpp b/Engine/Plugins/Runtime/RigVM/Source/RigVMEditor/Private/EdGraph/RigVMEdGraphPanelPinFactory.cpp
--- a/Engine/Plugins/Runtime/RigVM/Source/RigVMEditor/Private/EdGraph/RigVMEdGraphPanelPinFactory.cpp
+++ b/Engine/Plugins/Runtime/RigVM/Source/RigVMEditor/Private/EdGraph/RigVMEdGraphPanelPinFactory.cpp
@@ -27,6 +27,11 @@ FName FRigVMEdGraphPanelPinFactory::GetFactoryName() const
 
 TSharedPtr<SGraphPin> FRigVMEdGraphPanelPinFactory::CreatePin(UEdGraphPin* InPin) const
 {
+	if(InPin == nullptr)
+	{
+		return nullptr;
+	}
+	
 	// we need to check if this is the right factory for the implementation
 	if(const UEdGraphNode* EdGraphNode = InPin->GetOuter())
 	{
@@ -45,6 +50,16 @@ TSharedPtr<SGraphPin> FRigVMEdGraphPanelPinFactory::CreatePin(UEdGraphPin* InPin
 		return InternalResult;
 	}
 
+	// if the graph we are looking at is not a rig vm graph - let's not do this
+	if (const UEdGraphNode* OwningNode = InPin->GetOwningNode())
+	{
+		// only create pins within rig vm graphs
+		if (Cast<URigVMEdGraph>(OwningNode->GetGraph()) == nullptr)
+		{
+			return nullptr;
+		}
+	}
+
 	return FNodeFactory::CreateK2PinWidget(InPin);
 }
 
