Control Rig: Localize functions not working
#rnx

diff --git a/Engine/Plugins/Runtime/RigVM/Source/RigVMDeveloper/Private/RigVMModel/RigVMController.cpp b/Engine/Plugins/Runtime/RigVM/Source/RigVMDeveloper/Private/RigVMModel/RigVMController.cpp
--- a/Engine/Plugins/Runtime/RigVM/Source/RigVMDeveloper/Private/RigVMModel/RigVMController.cpp
+++ b/Engine/Plugins/Runtime/RigVM/Source/RigVMDeveloper/Private/RigVMModel/RigVMController.cpp
@@ -4460,7 +4460,7 @@ TMap<FRigVMGraphFunctionIdentifier, URigVMLibraryNode*> URigVMController::Locali
 	// find all functions to localize
 	for(int32 NodeToVisitIndex=0; NodeToVisitIndex<NodesToVisit.Num(); NodeToVisitIndex++)
 	{
-		FRigVMGraphFunctionIdentifier& NodeToVisit = NodesToVisit[NodeToVisitIndex];
+		const FRigVMGraphFunctionIdentifier NodeToVisit = NodesToVisit[NodeToVisitIndex];
 
 		// Already local
 		if (NodeToVisit.HostObject == ThisFunctionHost)
diff --git a/Engine/Plugins/Runtime/RigVM/Source/RigVMEditor/Private/Widgets/SRigVMGraphFunctionLocalizationWidget.cpp b/Engine/Plugins/Runtime/RigVM/Source/RigVMEditor/Private/Widgets/SRigVMGraphFunctionLocalizationWidget.cpp
--- a/Engine/Plugins/Runtime/RigVM/Source/RigVMEditor/Private/Widgets/SRigVMGraphFunctionLocalizationWidget.cpp
+++ b/Engine/Plugins/Runtime/RigVM/Source/RigVMEditor/Private/Widgets/SRigVMGraphFunctionLocalizationWidget.cpp
@@ -102,7 +102,7 @@ void SRigVMGraphFunctionLocalizationWidget::Construct(const FArguments& InArgs,
 
 	for(int32 NodeToVisitIndex=0; NodeToVisitIndex<NodesToVisit.Num(); NodeToVisitIndex++)
 	{
-		const FRigVMGraphFunctionIdentifier& NodeToVisit = NodesToVisit[NodeToVisitIndex];
+		const FRigVMGraphFunctionIdentifier NodeToVisit = NodesToVisit[NodeToVisitIndex];
 
 		bool bIsPublic;
 		FRigVMGraphFunctionData* FunctionData = FRigVMGraphFunctionData::FindFunctionData(NodeToVisit, &bIsPublic);
