RigVM: Make sure to refresh engine types also when running -game

#rnx

diff --git a/Engine/Plugins/Runtime/RigVM/Source/RigVM/Private/RigVMCore/RigVMRegistry.cpp b/Engine/Plugins/Runtime/RigVM/Source/RigVM/Private/RigVMCore/RigVMRegistry.cpp
--- a/Engine/Plugins/Runtime/RigVM/Source/RigVM/Private/RigVMCore/RigVMRegistry.cpp
+++ b/Engine/Plugins/Runtime/RigVM/Source/RigVM/Private/RigVMCore/RigVMRegistry.cpp
@@ -319,6 +319,21 @@ void FRigVMRegistry::InitializeIfNeeded()
 void FRigVMRegistry::RefreshEngineTypes()
 {
 	FScopeLock RefreshTypesScopeLock(&RefreshTypesMutex);
+	RefreshEngineTypes_NoLock();
+}
+
+void FRigVMRegistry::RefreshEngineTypesIfRequired()
+{
+	FScopeLock RefreshTypesScopeLock(&RefreshTypesMutex);
+	if(bEverRefreshedEngineTypes)
+	{
+		return;
+	}
+	RefreshEngineTypes_NoLock();
+}
+
+void FRigVMRegistry::RefreshEngineTypes_NoLock()
+{
 	TGuardValue<bool> EnableGuardRefresh(bIsRefreshingEngineTypes, true);
 
 	const int32 NumTypesBefore = Types.Num(); 
@@ -390,6 +405,8 @@ void FRigVMRegistry::RefreshEngineTypes()
 			}
 		}
 	}
+	
+	bEverRefreshedEngineTypes = true;
 }
 
 void FRigVMRegistry::OnAssetRenamed(const FAssetData& InAssetData, const FString& InOldObjectPath)
diff --git a/Engine/Plugins/Runtime/RigVM/Source/RigVM/Public/RigVMCore/RigVMRegistry.h b/Engine/Plugins/Runtime/RigVM/Source/RigVM/Public/RigVMCore/RigVMRegistry.h
--- a/Engine/Plugins/Runtime/RigVM/Source/RigVM/Public/RigVMCore/RigVMRegistry.h
+++ b/Engine/Plugins/Runtime/RigVM/Source/RigVM/Public/RigVMCore/RigVMRegistry.h
@@ -82,6 +82,10 @@ public:
 	// based on the names.
 	void RefreshEngineTypes();
 
+	// Refreshes the list and finds the function pointers
+	// based on the names.
+	void RefreshEngineTypesIfRequired();
+
 	// Update the registry when types are renamed
 	void OnAssetRenamed(const FAssetData& InAssetData, const FString& InOldObjectPath);
 	
@@ -298,6 +302,7 @@ private:
 	// disable default constructor
 	FRigVMRegistry()
 		: bIsRefreshingEngineTypes(false)
+		, bEverRefreshedEngineTypes(false)
 	{
 	}
 	
@@ -327,6 +332,8 @@ private:
 
 	TRigVMTypeIndex FindOrAddType_Internal(const FRigVMTemplateArgumentType& InType, bool bForce);
 
+	void RefreshEngineTypes_NoLock();
+
 	static EObjectFlags DisallowedFlags()
 	{
 		return RF_BeginDestroyed | RF_FinishDestroyed;
@@ -404,6 +411,9 @@ private:
 	// If this is true the registry is currently refreshing all types
 	bool bIsRefreshingEngineTypes;
 
+	// This is true if the engine has ever refreshed the engine types
+	bool bEverRefreshedEngineTypes;
+
 	static FCriticalSection RefreshTypesMutex;
 	static FCriticalSection FunctionRegistryMutex;
 	static FCriticalSection FactoryRegistryMutex;
diff --git a/Engine/Plugins/Runtime/RigVM/Source/RigVMDeveloper/Private/RigVMBlueprint.cpp b/Engine/Plugins/Runtime/RigVM/Source/RigVMDeveloper/Private/RigVMBlueprint.cpp
--- a/Engine/Plugins/Runtime/RigVM/Source/RigVMDeveloper/Private/RigVMBlueprint.cpp
+++ b/Engine/Plugins/Runtime/RigVM/Source/RigVMDeveloper/Private/RigVMBlueprint.cpp
@@ -921,6 +921,8 @@ void URigVMBlueprint::PostLoad()
 {
 	Super::PostLoad();
 
+	FRigVMRegistry::Get().RefreshEngineTypesIfRequired();
+
 	bVMRecompilationRequired = true;
 	{
 		TGuardValue<bool> IsCompilingGuard(bIsCompiling, true);
