RigVM: Don't allow to move any pin prior to the execute pin in the UI

#rnx

diff --git a/Engine/Plugins/Runtime/RigVM/Source/RigVMEditor/Private/Editor/RigVMGraphDetailCustomization.cpp b/Engine/Plugins/Runtime/RigVM/Source/RigVMEditor/Private/Editor/RigVMGraphDetailCustomization.cpp
--- a/Engine/Plugins/Runtime/RigVM/Source/RigVMEditor/Private/Editor/RigVMGraphDetailCustomization.cpp
+++ b/Engine/Plugins/Runtime/RigVM/Source/RigVMEditor/Private/Editor/RigVMGraphDetailCustomization.cpp
@@ -473,6 +473,10 @@ bool FRigVMFunctionArgumentLayout::CanArgumentBeMoved(bool bMoveUp) const
 				for(int32 Index = 0; Index < Node->GetPins().Num(); Index++)
 				{
 					URigVMPin* OtherPin = Node->GetPins()[Index];
+					if(OtherPin->IsExecuteContext())
+					{
+						continue;
+					}
 					if(IsInput(OtherPin) == bLookForInput)
 					{
 						return OtherPin != Pin;
