Control Rig: When unloading a plugin, remove all references to user-defined structs and enums that are a part of the the plugin content.

#rnx

diff --git a/Engine/Plugins/Runtime/RigVM/Source/RigVM/Private/RigVMCore/RigVMRegistry.cpp b/Engine/Plugins/Runtime/RigVM/Source/RigVM/Private/RigVMCore/RigVMRegistry.cpp
--- a/Engine/Plugins/Runtime/RigVM/Source/RigVM/Private/RigVMCore/RigVMRegistry.cpp
+++ b/Engine/Plugins/Runtime/RigVM/Source/RigVM/Private/RigVMCore/RigVMRegistry.cpp
@@ -14,6 +14,7 @@
 #include "Misc/CoreDelegates.h"
 #include "Misc/DelayedAutoRegister.h"
 #include "RigVMFunctions/RigVMDispatch_Core.h"
+#include "Interfaces/IPluginManager.h"
 
 const FName FRigVMRegistry::TemplateNameMetaName = TEXT("TemplateName");
 
@@ -306,6 +307,8 @@ void FRigVMRegistry::InitializeIfNeeded()
 			}
 		}
 
+		IPluginManager::Get().OnPluginUnmounted().RemoveAll(this);
+
 		UE::Anim::AttributeTypes::GetOnAttributeTypesChanged().RemoveAll(this);
 	});
 	
@@ -313,6 +316,8 @@ void FRigVMRegistry::InitializeIfNeeded()
 	AssetRegistryModule.Get().OnAssetRemoved().AddRaw(this, &FRigVMRegistry::OnAssetRemoved);
 	AssetRegistryModule.Get().OnAssetRenamed().AddRaw(this, &FRigVMRegistry::OnAssetRenamed);
 
+	IPluginManager::Get().OnPluginUnmounted().AddRaw(this, &FRigVMRegistry::OnPluginUnloaded);
+	
 	UE::Anim::AttributeTypes::GetOnAttributeTypesChanged().AddRaw(this, &FRigVMRegistry::OnAnimationAttributeTypesChanged);
 }
 
@@ -423,7 +428,44 @@ void FRigVMRegistry::OnAssetRenamed(const FAssetData& InAssetData, const FString
 
 void FRigVMRegistry::OnAssetRemoved(const FAssetData& InAssetData)
 {
-	if (RemoveType(InAssetData))
+	if (RemoveType(InAssetData.ToSoftObjectPath(), InAssetData.GetClass()))
+	{
+		OnRigVMRegistryChangedDelegate.Broadcast();
+	}
+}
+
+void FRigVMRegistry::OnPluginUnloaded(IPlugin& InPlugin)
+{
+	const FString PluginContentPath = InPlugin.GetMountedAssetPath();
+
+	TSet<FSoftObjectPath> PathsToRemove;
+	for (const TPair<FSoftObjectPath, TRigVMTypeIndex>& Item: UserDefinedTypeToIndex)
+	{
+		const FSoftObjectPath ObjectPath = Item.Key;
+		const FString PackageName = ObjectPath.GetLongPackageName();
+		
+		if (PackageName.StartsWith(PluginContentPath))
+		{
+			PathsToRemove.Add(ObjectPath);
+		}
+	}
+
+	bool bRegistryChanged = false;
+	for (FSoftObjectPath ObjectPath: PathsToRemove)
+	{
+		const UClass* ObjectClass = nullptr;
+		if (const UObject* TypeObject = ObjectPath.ResolveObject())
+		{
+			ObjectClass = TypeObject->GetClass();
+		}
+		
+		if (RemoveType(ObjectPath, ObjectClass))
+		{
+			bRegistryChanged = true;
+		}
+	}
+
+	if (bRegistryChanged)
 	{
 		OnRigVMRegistryChangedDelegate.Broadcast();
 	}
@@ -768,16 +810,13 @@ void FRigVMRegistry::PropagateTypeAddedToCategory(const FRigVMTemplateArgument::
 	}
 }
 
-bool FRigVMRegistry::RemoveType(const FAssetData& InAssetData)
+bool FRigVMRegistry::RemoveType(const FSoftObjectPath& InObjectPath, const UClass* InObjectClass)
 {
-	const FSoftObjectPath AssetPath = InAssetData.ToSoftObjectPath();
-	const UClass* TypeClass = InAssetData.GetClass();
-	
-	if (const TRigVMTypeIndex* TypeIndexPtr = UserDefinedTypeToIndex.Find(AssetPath))
+	if (const TRigVMTypeIndex* TypeIndexPtr = UserDefinedTypeToIndex.Find(InObjectPath))
 	{
 		const TRigVMTypeIndex TypeIndex = *TypeIndexPtr;
 		
-		UserDefinedTypeToIndex.Remove(AssetPath);
+		UserDefinedTypeToIndex.Remove(InObjectPath);
 		
 		if(TypeIndex == INDEX_NONE)
 		{
@@ -806,7 +845,7 @@ bool FRigVMRegistry::RemoveType(const FAssetData& InAssetData)
 				break;
 			}
 			
-			if(TypeClass == UUserDefinedEnum::StaticClass())
+			if(InObjectClass == UUserDefinedEnum::StaticClass())
 			{
 				switch(ArrayDimension)
 				{
@@ -831,7 +870,7 @@ bool FRigVMRegistry::RemoveType(const FAssetData& InAssetData)
 					}
 				}
 			}
-			else if(TypeClass == UUserDefinedStruct::StaticClass())
+			else if(InObjectClass == UUserDefinedStruct::StaticClass())
 			{
 				switch(ArrayDimension)
 				{
@@ -1820,8 +1859,7 @@ FString FRigVMRegistry::FindOrAddSingletonDispatchFunction(UScriptStruct* InFact
 }
 
 const TArray<FRigVMDispatchFactory*>& FRigVMRegistry::GetFactories() const
-{
-	return Factories;
+{	return Factories;
 }
 
 const TArray<FRigVMFunction>* FRigVMRegistry::GetPredicatesForStruct(const FName& InStructName) const
diff --git a/Engine/Plugins/Runtime/RigVM/Source/RigVM/Public/RigVMCore/RigVMRegistry.h b/Engine/Plugins/Runtime/RigVM/Source/RigVM/Public/RigVMCore/RigVMRegistry.h
--- a/Engine/Plugins/Runtime/RigVM/Source/RigVM/Public/RigVMCore/RigVMRegistry.h
+++ b/Engine/Plugins/Runtime/RigVM/Source/RigVM/Public/RigVMCore/RigVMRegistry.h
@@ -7,12 +7,9 @@
 #include "Containers/Map.h"
 #include "Containers/UnrealString.h"
 #include "CoreMinimal.h"
-#include "CoreTypes.h"
-#include "HAL/PlatformCrt.h"
 #include "RigVMCore/RigVMTraits.h"
 #include "RigVMDispatchFactory.h"
 #include "RigVMFunction.h"
-#include "RigVMMemory.h"
 #include "RigVMTemplate.h"
 #include "RigVMTypeIndex.h"
 #include "Templates/EnableIf.h"
@@ -25,6 +22,7 @@
 #include "UObject/GCObject.h"
 
 class FProperty;
+class IPlugin;
 class UObject;
 struct FRigVMDispatchFactory;
 
@@ -92,6 +90,9 @@ public:
 	// Update the registry when old types are removed
     void OnAssetRemoved(const FAssetData& InAssetData);
 
+	// Removes all types associated with a plugin that's being unloaded. 
+	void OnPluginUnloaded(IPlugin& InPlugin);
+	
 	// Update the registry when new types are added to the attribute system so that they can be selected
 	// on Attribute Nodes
 	void OnAnimationAttributeTypesChanged(const UScriptStruct* InStruct, bool bIsAdded);
@@ -108,7 +109,7 @@ public:
 
 	// Removes a type from the registry, and updates all dependent templates
 	// which also creates invalid permutations in templates that we should ignore
-	bool RemoveType(const FAssetData& InAssetData);
+	bool RemoveType(const FSoftObjectPath& InObjectPath, const UClass* InObjectClass);
 
 	// Returns the type index given a type
 	TRigVMTypeIndex GetTypeIndex(const FRigVMTemplateArgumentType& InType) const;
diff --git a/Engine/Plugins/Runtime/RigVM/Source/RigVM/RigVM.Build.cs b/Engine/Plugins/Runtime/RigVM/Source/RigVM/RigVM.Build.cs
--- a/Engine/Plugins/Runtime/RigVM/Source/RigVM/RigVM.Build.cs
+++ b/Engine/Plugins/Runtime/RigVM/Source/RigVM/RigVM.Build.cs
@@ -11,12 +11,13 @@ public class RigVM : ModuleRules
 
         PublicDependencyModuleNames.AddRange(
             new string[] {
-                "Core",
-                "CoreUObject",
-                "Engine",
-                "AnimationCore",
-                "AnimGraphRuntime",
-                "DeveloperSettings",
+				"AnimGraphRuntime",
+				"AnimationCore",
+				"Core",
+				"CoreUObject",
+				"DeveloperSettings",
+				"Engine",
+				"Projects",
 				"StructUtils",
 			}
 		);
