[Control Rig] Make sure to refresh engine types for rigvmhost to avoid missing type causing crash during game launch

diff --git a/Engine/Plugins/Runtime/RigVM/Source/RigVM/Private/RigVMHost.cpp b/Engine/Plugins/Runtime/RigVM/Source/RigVM/Private/RigVMHost.cpp
--- a/Engine/Plugins/Runtime/RigVM/Source/RigVM/Private/RigVMHost.cpp
+++ b/Engine/Plugins/Runtime/RigVM/Source/RigVM/Private/RigVMHost.cpp
@@ -147,7 +147,9 @@ void URigVMHost::Serialize(FArchive& Ar)
 void URigVMHost::PostLoad()
 {
 	Super::PostLoad();
-
+	
+	FRigVMRegistry::Get().RefreshEngineTypesIfRequired();
+	
 	FRigVMExtendedExecuteContext& ExtendedExecuteContext = GetRigVMExtendedExecuteContext();
 
 	ExtendedExecuteContext.InvalidateCachedMemory();
