Control Rig: Add casting from enum to int
#rnx

diff --git a/Engine/Plugins/Animation/ControlRig/Source/ControlRig/Private/Units/Execution/RigUnit_DynamicHierarchy.cpp b/Engine/Plugins/Animation/ControlRig/Source/ControlRig/Private/Units/Execution/RigUnit_DynamicHierarchy.cpp
--- a/Engine/Plugins/Animation/ControlRig/Source/ControlRig/Private/Units/Execution/RigUnit_DynamicHierarchy.cpp
+++ b/Engine/Plugins/Animation/ControlRig/Source/ControlRig/Private/Units/Execution/RigUnit_DynamicHierarchy.cpp
@@ -1074,9 +1074,17 @@ FRigUnit_HierarchyAddAnimationChannelInteger_Execute()
 		ControlSettings.SetupLimitArrayForType(true, true, true);
 		ControlSettings.LimitEnabled[0] = LimitsEnabled.Enabled;
 		ControlSettings.MinimumValue = FRigControlValue::Make<int32>(MinimumValue);
-		ControlSettings.MaximumValue = FRigControlValue::Make<int32>(MaximumValue);
 		ControlSettings.DisplayName = Controller->GetHierarchy()->GetSafeNewDisplayName(Parent, Name);
 		ControlSettings.ControlEnum = ControlEnum;
+
+		if (ControlEnum)
+		{
+			ControlSettings.MaximumValue = FRigControlValue::Make<int32>(ControlEnum->GetMaxEnumValue());
+		}
+		else
+		{
+			ControlSettings.MaximumValue = FRigControlValue::Make<int32>(MaximumValue);
+		}
 		const FRigControlValue Value = FRigControlValue::Make<int32>(InitialValue);
 		
 		FRigHierarchyControllerInstructionBracket InstructionBracket(Controller, ExecuteContext.GetInstructionIndex());
diff --git a/Engine/Plugins/Runtime/RigVM/Source/RigVM/Private/RigVMFunctions/RigVMDispatch_CastEnum.cpp b/Engine/Plugins/Runtime/RigVM/Source/RigVM/Private/RigVMFunctions/RigVMDispatch_CastEnum.cpp
new file mode 100644
--- /dev/null
+++ b/Engine/Plugins/Runtime/RigVM/Source/RigVM/Private/RigVMFunctions/RigVMDispatch_CastEnum.cpp
@@ -0,0 +1,105 @@
+﻿// Copyright Epic Games, Inc. All Rights Reserved.
+
+#include "RigVMFunctions/RigVMDispatch_CastEnum.h"
+#include "RigVMCore/RigVMRegistry.h"
+#include "RigVMFunctions/RigVMFunctionDefines.h"
+
+#define LOCTEXT_NAMESPACE "RigVMDispatch_CastEnum"
+
+const FName FRigVMDispatch_CastEnum::ValueName = TEXT("Value");
+const FName FRigVMDispatch_CastEnum::ResultName = TEXT("Result");
+
+const TArray<FRigVMTemplateArgumentInfo>& FRigVMDispatch_CastEnum::GetArgumentInfos() const
+{
+	static TArray<FRigVMTemplateArgumentInfo> OutInfos;
+	if (OutInfos.IsEmpty())
+	{
+		static const TArray<FRigVMTemplateArgument::ETypeCategory> ElementCategories =
+		{
+			FRigVMTemplateArgument::ETypeCategory_SingleEnumValue
+		};
+		
+		static TArray<FRigVMTemplateArgumentInfo> Infos;
+		Infos.Emplace(ValueName, ERigVMPinDirection::Input, ElementCategories);
+		Infos.Emplace(ResultName, ERigVMPinDirection::Output, RigVMTypeUtils::TypeIndex::Int32);
+		
+		OutInfos = BuildArgumentListFromPrimaryArgument(Infos, ValueName);
+	}
+
+	return OutInfos;
+}
+
+bool FRigVMDispatch_CastEnum::GetPermutationsFromArgumentType(const FName& InArgumentName, const TRigVMTypeIndex& InTypeIndex, TArray<FRigVMTemplateTypeMap, TInlineAllocator<1>>& OutPermutations) const
+{
+	if (InArgumentName == ValueName)
+	{
+		OutPermutations.Add(
+	{
+			{ ValueName, InTypeIndex },
+			{ ResultName, RigVMTypeUtils::TypeIndex::Int32 }
+		});
+	}
+	else if (InArgumentName == ResultName && InTypeIndex == RigVMTypeUtils::TypeIndex::Int32)
+	{
+		const TArray<TRigVMTypeIndex>& EnumTypes = FRigVMRegistry::Get().GetTypesForCategory(FRigVMTemplateArgument::ETypeCategory_SingleEnumValue);
+		for (const TRigVMTypeIndex& Type : EnumTypes)
+		{
+			OutPermutations.Add(
+	{
+				{ ValueName, Type },
+				{ ResultName, InTypeIndex }
+			});
+		}
+	}
+	return !OutPermutations.IsEmpty();
+}
+
+#if WITH_EDITOR
+
+FString FRigVMDispatch_CastEnum::GetNodeTitle(const FRigVMTemplateTypeMap& InTypes) const
+{
+	return TEXT("Cast to int");
+}
+
+FText FRigVMDispatch_CastEnum::GetNodeTooltip(const FRigVMTemplateTypeMap& InTypes) const
+{
+	return LOCTEXT("CastToolTip", "Casts from enum to int");
+}
+
+#endif
+
+void FRigVMDispatch_CastEnum::Execute(FRigVMExtendedExecuteContext& InContext, FRigVMMemoryHandleArray Handles, FRigVMPredicateBranchArray RigVMBranches)
+{
+	const FEnumProperty* ValueProperty = CastFieldChecked<FEnumProperty>(Handles[0].GetProperty());
+	const FProperty* ResultProperty = CastFieldChecked<FProperty>(Handles[1].GetProperty());
+
+	if (!ResultProperty || !ValueProperty)
+	{
+		return;
+	}
+
+	uint8* ValuePtr = Handles[0].GetData();
+	int32* ResultPtr = (int32*)Handles[1].GetData();
+	if(ValuePtr && ResultPtr)
+	{
+		*ResultPtr = 0;
+		ValueProperty->CopyCompleteValue(ResultPtr, ValuePtr);
+	}
+	else
+	{
+		*ResultPtr = INDEX_NONE;
+	}
+
+#if WITH_EDITOR
+	if (*ResultPtr == INDEX_NONE)
+	{
+		const FRigVMExecuteContext& ExecuteContext = InContext.GetPublicData<>();
+		if(ExecuteContext.GetLog() != nullptr)
+		{
+			ExecuteContext.Report(EMessageSeverity::Error, InContext.GetPublicData<>().GetFunctionName(), InContext.GetPublicData<>().GetInstructionIndex(), TEXT("Enum value invalid"));
+		}
+	}
+#endif
+}
+
+#undef LOCTEXT_NAMESPACE
diff --git a/Engine/Plugins/Runtime/RigVM/Source/RigVM/Public/RigVMFunctions/RigVMDispatch_CastEnum.h b/Engine/Plugins/Runtime/RigVM/Source/RigVM/Public/RigVMFunctions/RigVMDispatch_CastEnum.h
new file mode 100644
--- /dev/null
+++ b/Engine/Plugins/Runtime/RigVM/Source/RigVM/Public/RigVMFunctions/RigVMDispatch_CastEnum.h
@@ -0,0 +1,33 @@
+﻿// Copyright Epic Games, Inc. All Rights Reserved.
+
+#pragma once
+
+#include "RigVMCore/RigVMDispatchFactory.h"
+#include "RigVMDispatch_CastEnum.generated.h"
+
+USTRUCT(meta=(DisplayName = "Cast", Category = "Enum", Keywords = "As", NodeColor = "1,1,1,1"))
+struct RIGVM_API FRigVMDispatch_CastEnum : public FRigVMDispatchFactory
+{
+	GENERATED_BODY()
+
+public:
+	FRigVMDispatch_CastEnum()
+	{
+		FactoryScriptStruct = StaticStruct();
+	}
+	
+	virtual const TArray<FRigVMTemplateArgumentInfo>& GetArgumentInfos() const override;
+	virtual bool GetPermutationsFromArgumentType(const FName& InArgumentName, const TRigVMTypeIndex& InTypeIndex, TArray<FRigVMTemplateTypeMap, TInlineAllocator<1>>& OutPermutations) const override;
+#if WITH_EDITOR
+	virtual FString GetNodeTitle(const FRigVMTemplateTypeMap& InTypes) const override;
+	virtual FText GetNodeTooltip(const FRigVMTemplateTypeMap& InTypes) const override;
+#endif
+	virtual bool IsSingleton() const override { return true; }
+	
+protected:
+	virtual FRigVMFunctionPtr GetDispatchFunctionImpl(const FRigVMTemplateTypeMap& InTypes) const override { return &FRigVMDispatch_CastEnum::Execute; }
+	static void Execute(FRigVMExtendedExecuteContext& InContext, FRigVMMemoryHandleArray Handles, FRigVMPredicateBranchArray RigVMBranches);
+
+	static const FName ValueName;
+	static const FName ResultName;
+};
diff --git a/Engine/Plugins/Runtime/RigVM/Source/RigVMDeveloper/Private/EdGraph/RigVMEdGraphNode.cpp b/Engine/Plugins/Runtime/RigVM/Source/RigVMDeveloper/Private/EdGraph/RigVMEdGraphNode.cpp
--- a/Engine/Plugins/Runtime/RigVM/Source/RigVMDeveloper/Private/EdGraph/RigVMEdGraphNode.cpp
+++ b/Engine/Plugins/Runtime/RigVM/Source/RigVMDeveloper/Private/EdGraph/RigVMEdGraphNode.cpp
@@ -16,6 +16,7 @@
 #include "RigVMBlueprintUtils.h"
 #include "../../../RigVMEditor/Public/RigVMEditorModule.h"
 #include "RigVMCore/RigVMExecuteContext.h"
+#include "RigVMFunctions/RigVMDispatch_CastEnum.h"
 #include "RigVMFunctions/RigVMDispatch_CastObject.h"
 #include "RigVMModel/Nodes/RigVMAggregateNode.h"
 #include "RigVMModel/Nodes/RigVMFunctionReferenceNode.h"
@@ -1059,7 +1060,8 @@ bool URigVMEdGraphNode::DrawAsCompactNode() const
 		if(const URigVMTemplateNode* TemplateModelNode = Cast<URigVMTemplateNode>(GetModelNode()))
 		{
 			if(TemplateModelNode->GetNotation() == RigVMTypeUtils::GetCastTemplateNotation() ||
-				TemplateModelNode->GetNotation() == FRigVMDispatch_CastObject().GetTemplateNotation())
+				TemplateModelNode->GetNotation() == FRigVMDispatch_CastObject().GetTemplateNotation() ||
+				TemplateModelNode->GetNotation() == FRigVMDispatch_CastEnum().GetTemplateNotation())
 			{
 				DrawAsCompactNodeCache = true;
 				
diff --git a/Engine/Plugins/Runtime/RigVM/Source/RigVMDeveloper/Private/RigVMDeveloperTypeUtils.cpp b/Engine/Plugins/Runtime/RigVM/Source/RigVMDeveloper/Private/RigVMDeveloperTypeUtils.cpp
--- a/Engine/Plugins/Runtime/RigVM/Source/RigVMDeveloper/Private/RigVMDeveloperTypeUtils.cpp
+++ b/Engine/Plugins/Runtime/RigVM/Source/RigVMDeveloper/Private/RigVMDeveloperTypeUtils.cpp
@@ -3,6 +3,7 @@
 #include "RigVMDeveloperTypeUtils.h"
 
 #include "Internationalization/StringTableCore.h"
+#include "RigVMFunctions/RigVMDispatch_CastEnum.h"
 #include "RigVMFunctions/RigVMDispatch_CastObject.h"
 #include "RigVMModel/RigVMController.h"
 #include "RigVMModel/RigVMVariableDescription.h"
@@ -638,7 +639,8 @@ const FRigVMFunction* RigVMTypeUtils::GetCastForTypeIndices(const TRigVMTypeInde
 	const FRigVMTemplate* CastTemplates[] =
 	{
 		Registry.FindTemplate(CastTemplateNotation),
-		Registry.FindTemplate(FRigVMDispatch_CastObject().GetTemplateNotation())
+		Registry.FindTemplate(FRigVMDispatch_CastObject().GetTemplateNotation()),
+		Registry.FindTemplate(FRigVMDispatch_CastEnum().GetTemplateNotation())
 	};
 
 	for(const FRigVMTemplate* CastTemplate : CastTemplates)
