RigVM: Make sure to react to user defined struct / enum being loaded later

#rnx

diff --git a/Engine/Plugins/Runtime/RigVM/Source/RigVM/Private/RigVMCore/RigVMRegistry.cpp b/Engine/Plugins/Runtime/RigVM/Source/RigVM/Private/RigVMCore/RigVMRegistry.cpp
--- a/Engine/Plugins/Runtime/RigVM/Source/RigVM/Private/RigVMCore/RigVMRegistry.cpp
+++ b/Engine/Plugins/Runtime/RigVM/Source/RigVM/Private/RigVMCore/RigVMRegistry.cpp
@@ -302,6 +302,7 @@ void FRigVMRegistry::InitializeIfNeeded()
 		{
 			if (AssetRegistryModule->TryGet())
 			{
+				AssetRegistryModule->Get().OnAssetAdded().RemoveAll(this);
 				AssetRegistryModule->Get().OnAssetRemoved().RemoveAll(this);
 				AssetRegistryModule->Get().OnAssetRenamed().RemoveAll(this);
 			}
@@ -313,6 +314,7 @@ void FRigVMRegistry::InitializeIfNeeded()
 	});
 	
 	FAssetRegistryModule& AssetRegistryModule = FModuleManager::LoadModuleChecked<FAssetRegistryModule>(TEXT("AssetRegistry"));
+	AssetRegistryModule.Get().OnAssetAdded().AddRaw(this, &FRigVMRegistry::OnAssetAdded);
 	AssetRegistryModule.Get().OnAssetRemoved().AddRaw(this, &FRigVMRegistry::OnAssetRemoved);
 	AssetRegistryModule.Get().OnAssetRenamed().AddRaw(this, &FRigVMRegistry::OnAssetRenamed);
 
@@ -414,6 +416,29 @@ void FRigVMRegistry::RefreshEngineTypes_NoLock()
 	bEverRefreshedEngineTypes = true;
 }
 
+void FRigVMRegistry::OnAssetAdded(const FAssetData& InAssetData)
+{
+	if (InAssetData.GetClass())
+	{
+		if(InAssetData.GetClass()->IsChildOf(UUserDefinedStruct::StaticClass()))
+		{
+			UScriptStruct* CPPTypeObject = Cast<UScriptStruct>(InAssetData.GetAsset());
+			if (FindOrAddType(CPPTypeObject))
+			{
+				OnRigVMRegistryChangedDelegate.Broadcast();
+			}
+		}
+		else if(InAssetData.GetClass()->IsChildOf(UUserDefinedEnum::StaticClass()))
+		{
+			UEnum* CPPTypeObject = Cast<UEnum>(InAssetData.GetAsset());
+			if (FindOrAddType(CPPTypeObject))
+			{
+				OnRigVMRegistryChangedDelegate.Broadcast();
+			}
+		}
+	}
+}
+
 void FRigVMRegistry::OnAssetRenamed(const FAssetData& InAssetData, const FString& InOldObjectPath)
 {
 	const FSoftObjectPath OldPath(InOldObjectPath);
diff --git a/Engine/Plugins/Runtime/RigVM/Source/RigVM/Public/RigVMCore/RigVMRegistry.h b/Engine/Plugins/Runtime/RigVM/Source/RigVM/Public/RigVMCore/RigVMRegistry.h
--- a/Engine/Plugins/Runtime/RigVM/Source/RigVM/Public/RigVMCore/RigVMRegistry.h
+++ b/Engine/Plugins/Runtime/RigVM/Source/RigVM/Public/RigVMCore/RigVMRegistry.h
@@ -84,6 +84,9 @@ public:
 	// based on the names.
 	void RefreshEngineTypesIfRequired();
 
+	// Update the registry when new types are added
+	void OnAssetAdded(const FAssetData& InAssetData);
+
 	// Update the registry when types are renamed
 	void OnAssetRenamed(const FAssetData& InAssetData, const FString& InOldObjectPath);
 	
diff --git a/Engine/Plugins/Runtime/RigVM/Source/RigVMDeveloper/Private/RigVMModel/RigVMController.cpp b/Engine/Plugins/Runtime/RigVM/Source/RigVMDeveloper/Private/RigVMModel/RigVMController.cpp
--- a/Engine/Plugins/Runtime/RigVM/Source/RigVMDeveloper/Private/RigVMModel/RigVMController.cpp
+++ b/Engine/Plugins/Runtime/RigVM/Source/RigVMDeveloper/Private/RigVMModel/RigVMController.cpp
@@ -15558,20 +15558,30 @@ void URigVMController::RepopulatePinsOnNode(const FRigVMRegistry& Registry, cons
 		}
 	}
 	// add missing pins
-	for (int32 Index = 0; Index < NodeData.NewPinsToAdd.Num(); Index++)
 	{
-		const FString& PinPath = NodeData.NewPinInfos.GetPinPath(NodeData.NewPinsToAdd[Index]);
-		FString ParentPinPath, PinName;
-		UObject* OuterForPin = InNode;
-		if (URigVMPin::SplitPinPathAtEnd(PinPath, ParentPinPath, PinName))
+		TArray<bool> PinIndexAdded;
+		PinIndexAdded.AddZeroed(NodeData.NewPinInfos.Num());
+		for (int32 Index = 0; Index < NodeData.NewPinsToAdd.Num(); Index++)
 		{
-			OuterForPin = InNode->FindPin(ParentPinPath);
-		}
+			const int32 PinIndexToAdd = NodeData.NewPinsToAdd[Index];
+			if(PinIndexAdded[PinIndexToAdd])
+			{
+				continue;
+			}
+			const FString& PinPath = NodeData.NewPinInfos.GetPinPath(PinIndexToAdd);
+			FString ParentPinPath, PinName;
+			UObject* OuterForPin = InNode;
+			if (URigVMPin::SplitPinPathAtEnd(PinPath, ParentPinPath, PinName))
+			{
+				OuterForPin = InNode->FindPin(ParentPinPath);
+			}
 
-		CreatePinFromPinInfo(Registry, NodeData.PreviousPinInfos, NodeData.NewPinInfos[NodeData.NewPinsToAdd[Index]], PinPath, OuterForPin);
+			CreatePinFromPinInfo(Registry, NodeData.PreviousPinInfos, NodeData.NewPinInfos[NodeData.NewPinsToAdd[Index]], PinPath, OuterForPin);
 #if UE_RIGVMCONTROLLER_VERBOSE_REPOPULATE
-		UE_LOG(LogRigVMDeveloper, Display, TEXT("Adding new pin '%s'."), *PinPath);
+			UE_LOG(LogRigVMDeveloper, Display, TEXT("Adding new pin '%s'."), *PinPath);
 #endif
+			PinIndexAdded[PinIndexToAdd] = true;
+		}
 	}
 	// update existing pins
 	for (int32 Index = 0; Index < NodeData.PreviousPinsToUpdate.Num(); Index++)
