Control Rig: Add dispatch to cast from int to enum
#rnx

diff --git a/Engine/Plugins/Animation/ControlRig/Source/ControlRig/Private/Units/Execution/RigUnit_DynamicHierarchy.cpp b/Engine/Plugins/Animation/ControlRig/Source/ControlRig/Private/Units/Execution/RigUnit_DynamicHierarchy.cpp
--- a/Engine/Plugins/Animation/ControlRig/Source/ControlRig/Private/Units/Execution/RigUnit_DynamicHierarchy.cpp
+++ b/Engine/Plugins/Animation/ControlRig/Source/ControlRig/Private/Units/Execution/RigUnit_DynamicHierarchy.cpp
@@ -505,7 +505,14 @@ void FRigUnit_HierarchyAddControlInteger_LimitSettings::Configure(FRigControlSet
 	OutSettings.SetupLimitArrayForType(false, false, false);
 	OutSettings.LimitEnabled[0] = Limit;
 	OutSettings.MinimumValue = FRigControlValue::Make<int32>(MinValue);
-	OutSettings.MaximumValue = FRigControlValue::Make<int32>(MaxValue);
+	if (OutSettings.ControlEnum)
+	{
+		OutSettings.MaximumValue = FRigControlValue::Make<int32>(OutSettings.ControlEnum->GetMaxEnumValue());
+	}
+	else
+	{
+		OutSettings.MaximumValue = FRigControlValue::Make<int32>(MaxValue);
+	}
 	OutSettings.bDrawLimits = bDrawLimits;
 }
 
diff --git a/Engine/Plugins/Runtime/RigVM/Source/RigVM/Private/RigVMFunctions/RigVMDispatch_CastEnum.cpp b/Engine/Plugins/Runtime/RigVM/Source/RigVM/Private/RigVMFunctions/RigVMDispatch_CastEnum.cpp
--- a/Engine/Plugins/Runtime/RigVM/Source/RigVM/Private/RigVMFunctions/RigVMDispatch_CastEnum.cpp
+++ b/Engine/Plugins/Runtime/RigVM/Source/RigVM/Private/RigVMFunctions/RigVMDispatch_CastEnum.cpp
@@ -6,8 +6,8 @@
 
 #define LOCTEXT_NAMESPACE "RigVMDispatch_CastEnum"
 
-const FName FRigVMDispatch_CastEnum::ValueName = TEXT("Value");
-const FName FRigVMDispatch_CastEnum::ResultName = TEXT("Result");
+const FName FRigVMDispatch_CastEnumBase::ValueName = TEXT("Value");
+const FName FRigVMDispatch_CastEnumBase::ResultName = TEXT("Result");
 
 const TArray<FRigVMTemplateArgumentInfo>& FRigVMDispatch_CastEnum::GetArgumentInfos() const
 {
@@ -102,4 +102,100 @@ void FRigVMDispatch_CastEnum::Execute(FRigVMExtendedExecuteContext& InContext, F
 #endif
 }
 
+
+
+const TArray<FRigVMTemplateArgumentInfo>& FRigVMDispatch_CastEnumToInt::GetArgumentInfos() const
+{
+	static TArray<FRigVMTemplateArgumentInfo> OutInfos;
+	if (OutInfos.IsEmpty())
+	{
+		static const TArray<FRigVMTemplateArgument::ETypeCategory> ElementCategories =
+		{
+			FRigVMTemplateArgument::ETypeCategory_SingleEnumValue
+		};
+		
+		static TArray<FRigVMTemplateArgumentInfo> Infos;
+		Infos.Emplace(ValueName, ERigVMPinDirection::Input, RigVMTypeUtils::TypeIndex::Int32);
+		Infos.Emplace(ResultName, ERigVMPinDirection::Output, ElementCategories);
+		
+		OutInfos = BuildArgumentListFromPrimaryArgument(Infos, ResultName);
+	}
+
+	return OutInfos;
+}
+
+bool FRigVMDispatch_CastEnumToInt::GetPermutationsFromArgumentType(const FName& InArgumentName, const TRigVMTypeIndex& InTypeIndex, TArray<FRigVMTemplateTypeMap, TInlineAllocator<1>>& OutPermutations) const
+{
+	if (InArgumentName == ResultName)
+	{
+		OutPermutations.Add(
+	{
+			{ ValueName, RigVMTypeUtils::TypeIndex::Int32 },
+			{ ResultName, InTypeIndex }
+		});
+	}
+	else if (InArgumentName == ValueName && InTypeIndex == RigVMTypeUtils::TypeIndex::Int32)
+	{
+		const TArray<TRigVMTypeIndex>& EnumTypes = FRigVMRegistry::Get().GetTypesForCategory(FRigVMTemplateArgument::ETypeCategory_SingleEnumValue);
+		for (const TRigVMTypeIndex& Type : EnumTypes)
+		{
+			OutPermutations.Add(
+	{
+				{ ValueName, InTypeIndex },
+				{ ResultName, Type }
+			});
+		}
+	}
+	return !OutPermutations.IsEmpty();
+}
+
+#if WITH_EDITOR
+
+FString FRigVMDispatch_CastEnumToInt::GetNodeTitle(const FRigVMTemplateTypeMap& InTypes) const
+{
+	return TEXT("Cast to enum");
+}
+
+FText FRigVMDispatch_CastEnumToInt::GetNodeTooltip(const FRigVMTemplateTypeMap& InTypes) const
+{
+	return LOCTEXT("CastToolTip", "Casts from int to enum");
+}
+
+#endif
+
+void FRigVMDispatch_CastEnumToInt::Execute(FRigVMExtendedExecuteContext& InContext, FRigVMMemoryHandleArray Handles, FRigVMPredicateBranchArray RigVMBranches)
+{
+	const FProperty* ValueProperty = CastFieldChecked<FProperty>(Handles[0].GetProperty());
+	const FEnumProperty* ResultProperty = CastFieldChecked<FEnumProperty>(Handles[1].GetProperty());
+
+	if (!ResultProperty || !ValueProperty)
+	{
+		return;
+	}
+
+	UEnum* Enum = ResultProperty->GetEnum();
+	int32* ValuePtr = (int32*)Handles[0].GetData();
+	uint8* ResultPtr = Handles[1].GetData();
+	if(ValuePtr && ResultPtr)
+	{
+		*ResultPtr = 0;
+		ResultProperty->CopyCompleteValue(ResultPtr, ValuePtr);
+	}
+	else
+	{
+		*ResultPtr = INDEX_NONE;
+	}
+
+#if WITH_EDITOR
+	if (*ResultPtr == INDEX_NONE)
+	{
+		const FRigVMExecuteContext& ExecuteContext = InContext.GetPublicData<>();
+		if(ExecuteContext.GetLog() != nullptr)
+		{
+			ExecuteContext.Report(EMessageSeverity::Error, InContext.GetPublicData<>().GetFunctionName(), InContext.GetPublicData<>().GetInstructionIndex(), TEXT("Enum value invalid"));
+		}
+	}
+#endif
+}
+
 #undef LOCTEXT_NAMESPACE
diff --git a/Engine/Plugins/Runtime/RigVM/Source/RigVM/Public/RigVMFunctions/RigVMDispatch_CastEnum.h b/Engine/Plugins/Runtime/RigVM/Source/RigVM/Public/RigVMFunctions/RigVMDispatch_CastEnum.h
--- a/Engine/Plugins/Runtime/RigVM/Source/RigVM/Public/RigVMFunctions/RigVMDispatch_CastEnum.h
+++ b/Engine/Plugins/Runtime/RigVM/Source/RigVM/Public/RigVMFunctions/RigVMDispatch_CastEnum.h
@@ -5,8 +5,14 @@
 #include "RigVMCore/RigVMDispatchFactory.h"
 #include "RigVMDispatch_CastEnum.generated.h"
 
+struct FRigVMDispatch_CastEnumBase
+{
+	static const FName ValueName;
+	static const FName ResultName;
+};
+
 USTRUCT(meta=(DisplayName = "Cast", Category = "Enum", Keywords = "As", NodeColor = "1,1,1,1"))
-struct RIGVM_API FRigVMDispatch_CastEnum : public FRigVMDispatchFactory
+struct RIGVM_API FRigVMDispatch_CastEnum : public FRigVMDispatchFactory, public FRigVMDispatch_CastEnumBase
 {
 	GENERATED_BODY()
 
@@ -27,7 +33,28 @@ public:
 protected:
 	virtual FRigVMFunctionPtr GetDispatchFunctionImpl(const FRigVMTemplateTypeMap& InTypes) const override { return &FRigVMDispatch_CastEnum::Execute; }
 	static void Execute(FRigVMExtendedExecuteContext& InContext, FRigVMMemoryHandleArray Handles, FRigVMPredicateBranchArray RigVMBranches);
+};
 
-	static const FName ValueName;
-	static const FName ResultName;
+USTRUCT(meta=(DisplayName = "Cast", Category = "Enum", Keywords = "As", NodeColor = "1,1,1,1"))
+struct RIGVM_API FRigVMDispatch_CastEnumToInt : public FRigVMDispatchFactory, public FRigVMDispatch_CastEnumBase
+{
+	GENERATED_BODY()
+
+public:
+	FRigVMDispatch_CastEnumToInt()
+	{
+		FactoryScriptStruct = StaticStruct();
+	}
+	
+	virtual const TArray<FRigVMTemplateArgumentInfo>& GetArgumentInfos() const override;
+	virtual bool GetPermutationsFromArgumentType(const FName& InArgumentName, const TRigVMTypeIndex& InTypeIndex, TArray<FRigVMTemplateTypeMap, TInlineAllocator<1>>& OutPermutations) const override;
+#if WITH_EDITOR
+	virtual FString GetNodeTitle(const FRigVMTemplateTypeMap& InTypes) const override;
+	virtual FText GetNodeTooltip(const FRigVMTemplateTypeMap& InTypes) const override;
+#endif
+	virtual bool IsSingleton() const override { return true; }
+	
+protected:
+	virtual FRigVMFunctionPtr GetDispatchFunctionImpl(const FRigVMTemplateTypeMap& InTypes) const override { return &FRigVMDispatch_CastEnumToInt::Execute; }
+	static void Execute(FRigVMExtendedExecuteContext& InContext, FRigVMMemoryHandleArray Handles, FRigVMPredicateBranchArray RigVMBranches);
 };
diff --git a/Engine/Plugins/Runtime/RigVM/Source/RigVMDeveloper/Private/EdGraph/RigVMEdGraphNode.cpp b/Engine/Plugins/Runtime/RigVM/Source/RigVMDeveloper/Private/EdGraph/RigVMEdGraphNode.cpp
--- a/Engine/Plugins/Runtime/RigVM/Source/RigVMDeveloper/Private/EdGraph/RigVMEdGraphNode.cpp
+++ b/Engine/Plugins/Runtime/RigVM/Source/RigVMDeveloper/Private/EdGraph/RigVMEdGraphNode.cpp
@@ -1061,7 +1061,8 @@ bool URigVMEdGraphNode::DrawAsCompactNode() const
 		{
 			if(TemplateModelNode->GetNotation() == RigVMTypeUtils::GetCastTemplateNotation() ||
 				TemplateModelNode->GetNotation() == FRigVMDispatch_CastObject().GetTemplateNotation() ||
-				TemplateModelNode->GetNotation() == FRigVMDispatch_CastEnum().GetTemplateNotation())
+				TemplateModelNode->GetNotation() == FRigVMDispatch_CastEnum().GetTemplateNotation() ||
+				TemplateModelNode->GetNotation() == FRigVMDispatch_CastEnumToInt().GetTemplateNotation())
 			{
 				DrawAsCompactNodeCache = true;
 				
diff --git a/Engine/Plugins/Runtime/RigVM/Source/RigVMDeveloper/Private/RigVMDeveloperTypeUtils.cpp b/Engine/Plugins/Runtime/RigVM/Source/RigVMDeveloper/Private/RigVMDeveloperTypeUtils.cpp
--- a/Engine/Plugins/Runtime/RigVM/Source/RigVMDeveloper/Private/RigVMDeveloperTypeUtils.cpp
+++ b/Engine/Plugins/Runtime/RigVM/Source/RigVMDeveloper/Private/RigVMDeveloperTypeUtils.cpp
@@ -640,7 +640,8 @@ const FRigVMFunction* RigVMTypeUtils::GetCastForTypeIndices(const TRigVMTypeInde
 	{
 		Registry.FindTemplate(CastTemplateNotation),
 		Registry.FindTemplate(FRigVMDispatch_CastObject().GetTemplateNotation()),
-		Registry.FindTemplate(FRigVMDispatch_CastEnum().GetTemplateNotation())
+		Registry.FindTemplate(FRigVMDispatch_CastEnum().GetTemplateNotation()),
+		Registry.FindTemplate(FRigVMDispatch_CastEnumToInt().GetTemplateNotation())
 	};
 
 	for(const FRigVMTemplate* CastTemplate : CastTemplates)
