Control Rig: Fix static analysis warnings
#rnx

diff --git a/Engine/Plugins/Runtime/RigVM/Source/RigVM/Private/RigVMFunctions/RigVMDispatch_CastEnum.cpp b/Engine/Plugins/Runtime/RigVM/Source/RigVM/Private/RigVMFunctions/RigVMDispatch_CastEnum.cpp
--- a/Engine/Plugins/Runtime/RigVM/Source/RigVM/Private/RigVMFunctions/RigVMDispatch_CastEnum.cpp
+++ b/Engine/Plugins/Runtime/RigVM/Source/RigVM/Private/RigVMFunctions/RigVMDispatch_CastEnum.cpp
@@ -80,15 +80,14 @@ void FRigVMDispatch_CastEnum::Execute(FRigVMExtendedExecuteContext& InContext, F
 
 	uint8* ValuePtr = Handles[0].GetData();
 	int32* ResultPtr = (int32*)Handles[1].GetData();
-	if(ValuePtr && ResultPtr)
-	{
-		*ResultPtr = 0;
-		ValueProperty->CopyCompleteValue(ResultPtr, ValuePtr);
-	}
-	else
+
+	if (ValuePtr == nullptr || ResultPtr == nullptr)
 	{
-		*ResultPtr = INDEX_NONE;
+		return;
 	}
+	
+	*ResultPtr = 0;
+	ValueProperty->CopyCompleteValue(ResultPtr, ValuePtr);
 
 #if WITH_EDITOR
 	if (*ResultPtr == INDEX_NONE)
@@ -173,19 +172,17 @@ void FRigVMDispatch_CastEnumToInt::Execute(FRigVMExtendedExecuteContext& InConte
 		return;
 	}
 
-	UEnum* Enum = ResultProperty->GetEnum();
 	int32* ValuePtr = (int32*)Handles[0].GetData();
 	uint8* ResultPtr = Handles[1].GetData();
-	if(ValuePtr && ResultPtr)
-	{
-		*ResultPtr = 0;
-		ResultProperty->CopyCompleteValue(ResultPtr, ValuePtr);
-	}
-	else
+
+	if (ValuePtr == nullptr || ResultPtr == nullptr)
 	{
-		*ResultPtr = INDEX_NONE;
+		return;
 	}
-
+	
+	*ResultPtr = 0;
+	ResultProperty->CopyCompleteValue(ResultPtr, ValuePtr);
+	
 #if WITH_EDITOR
 	if (*ResultPtr == INDEX_NONE)
 	{
