Control Rig: Fix refresh hierarchy closing the editor
#rnx

diff --git a/Engine/Plugins/Runtime/RigVM/Source/RigVMEditor/Private/Editor/RigVMEditor.cpp b/Engine/Plugins/Runtime/RigVM/Source/RigVMEditor/Private/Editor/RigVMEditor.cpp
--- a/Engine/Plugins/Runtime/RigVM/Source/RigVMEditor/Private/Editor/RigVMEditor.cpp
+++ b/Engine/Plugins/Runtime/RigVM/Source/RigVMEditor/Private/Editor/RigVMEditor.cpp
@@ -108,6 +108,20 @@ FRigVMEditor::~FRigVMEditor()
 		RigVMBlueprint->OnGetFocusedGraph().Unbind();
 #endif
 	}
+
+	if (bRequestedReopen)
+	{
+		// Sometimes FPersonaToolkit::SetPreviewMesh will request an asset editor close and reopen. If
+		// SetPreviewMesh is called from within the editor, the close will not fully take effect until
+		// the callback finishes, so the open editor action will fail. In that case, let's make sure we
+		// detect a reopen requested, and open the editor again on the next tick.
+		bRequestedReopen = false;
+		FSoftObjectPath AssetToReopen = RigVMBlueprint;
+		GEditor->GetTimerManager()->SetTimerForNextTick([AssetToReopen]()
+		{
+			GEditor->GetEditorSubsystem<UAssetEditorSubsystem>()->OpenEditorForAsset(AssetToReopen);
+		});
+	}
 }
 
 void FRigVMEditor::InitRigVMEditor(const EToolkitMode::Type Mode, const TSharedPtr< class IToolkitHost >& InitToolkitHost, class URigVMBlueprint* InRigVMBlueprint)
@@ -138,6 +152,8 @@ void FRigVMEditor::InitRigVMEditor(const EToolkitMode::Type Mode, const TSharedP
 	const bool bCreateDefaultStandaloneMenu = true;
 	const bool bCreateDefaultToolbar = true;
 	InitAssetEditor(Mode, InitToolkitHost, GetEditorAppName(), FTabManager::FLayout::NullLayout, bCreateDefaultStandaloneMenu, bCreateDefaultToolbar, ObjectsBeingEdited);
+	GEditor->GetEditorSubsystem<UAssetEditorSubsystem>()->OnAssetEditorRequestedOpen().AddSP(this, &FRigVMEditor::HandleAssetRequestedOpen);
+	GEditor->GetEditorSubsystem<UAssetEditorSubsystem>()->OnAssetEditorRequestClose().AddSP(this, &FRigVMEditor::HandleAssetRequestClose);
 
 	CreateDefaultCommands();
 
@@ -381,6 +397,22 @@ void FRigVMEditor::InitRigVMEditor(const EToolkitMode::Type Mode, const TSharedP
 	PropertyChangedHandle = FCoreUObjectDelegates::OnObjectPropertyChanged.AddSP(this, &FRigVMEditor::OnPropertyChanged);
 }
 
+void FRigVMEditor::HandleAssetRequestedOpen(UObject* InObject)
+{
+	if (InObject == GetRigVMBlueprint())
+	{
+		bRequestedReopen = true;
+	}
+}
+
+void FRigVMEditor::HandleAssetRequestClose(UObject* InObject, EAssetEditorCloseReason InReason)
+{
+	if (InObject == GetRigVMBlueprint())
+	{
+		bRequestedReopen = false;
+	}
+}
+
 const FName FRigVMEditor::GetEditorAppName() const
 {
 	static const FName AppName(TEXT("RigVMEditorApp"));
@@ -449,6 +481,14 @@ void FRigVMEditor::Tick(float DeltaTime)
 	}
 }
 
+void FRigVMEditor::BringToolkitToFront()
+{
+	if (ToolkitHost.IsValid())
+	{
+		FBlueprintEditor::BringToolkitToFront();
+	}
+}
+
 FName FRigVMEditor::GetToolkitFName() const
 {
 	return FName("RigVMEditor");
diff --git a/Engine/Plugins/Runtime/RigVM/Source/RigVMEditor/Public/Editor/RigVMEditor.h b/Engine/Plugins/Runtime/RigVM/Source/RigVMEditor/Public/Editor/RigVMEditor.h
--- a/Engine/Plugins/Runtime/RigVM/Source/RigVMEditor/Public/Editor/RigVMEditor.h
+++ b/Engine/Plugins/Runtime/RigVM/Source/RigVMEditor/Public/Editor/RigVMEditor.h
@@ -50,6 +50,10 @@ public:
 	 */
 	virtual void InitRigVMEditor(const EToolkitMode::Type Mode, const TSharedPtr< class IToolkitHost >& InitToolkitHost, class URigVMBlueprint* InRigVMBlueprint);
 
+	virtual void HandleAssetRequestedOpen(UObject* InObject);
+	virtual void HandleAssetRequestClose(UObject* InObject, EAssetEditorCloseReason InReason);
+	bool bRequestedReopen = false;
+
 	virtual const FName GetEditorAppName() const;
 	virtual const FName GetEditorModeName() const;
 	virtual TSharedPtr<FApplicationMode> CreateEditorMode();
@@ -63,6 +67,7 @@ public:
 	virtual void Tick(float DeltaTime) override;
 
 	// IToolkit Interface
+	virtual void BringToolkitToFront() override;
 	virtual FName GetToolkitFName() const override;
 	virtual FName GetToolkitContextFName() const override;
 	virtual FText GetBaseToolkitName() const override;
