Control Rig: Fix loctext conflicts in FRigVMDispatch_CastEnum
#rnx

diff --git a/Engine/Plugins/Runtime/RigVM/Source/RigVM/Private/RigVMFunctions/RigVMDispatch_CastEnum.cpp b/Engine/Plugins/Runtime/RigVM/Source/RigVM/Private/RigVMFunctions/RigVMDispatch_CastEnum.cpp
--- a/Engine/Plugins/Runtime/RigVM/Source/RigVM/Private/RigVMFunctions/RigVMDispatch_CastEnum.cpp
+++ b/Engine/Plugins/Runtime/RigVM/Source/RigVM/Private/RigVMFunctions/RigVMDispatch_CastEnum.cpp
@@ -9,7 +9,7 @@
 const FName FRigVMDispatch_CastEnumBase::ValueName = TEXT("Value");
 const FName FRigVMDispatch_CastEnumBase::ResultName = TEXT("Result");
 
-const TArray<FRigVMTemplateArgumentInfo>& FRigVMDispatch_CastEnum::GetArgumentInfos() const
+const TArray<FRigVMTemplateArgumentInfo>& FRigVMDispatch_CastEnumToInt::GetArgumentInfos() const
 {
 	static TArray<FRigVMTemplateArgumentInfo> OutInfos;
 	if (OutInfos.IsEmpty())
@@ -29,7 +29,7 @@ const TArray<FRigVMTemplateArgumentInfo>& FRigVMDispatch_CastEnum::GetArgumentIn
 	return OutInfos;
 }
 
-bool FRigVMDispatch_CastEnum::GetPermutationsFromArgumentType(const FName& InArgumentName, const TRigVMTypeIndex& InTypeIndex, TArray<FRigVMTemplateTypeMap, TInlineAllocator<1>>& OutPermutations) const
+bool FRigVMDispatch_CastEnumToInt::GetPermutationsFromArgumentType(const FName& InArgumentName, const TRigVMTypeIndex& InTypeIndex, TArray<FRigVMTemplateTypeMap, TInlineAllocator<1>>& OutPermutations) const
 {
 	if (InArgumentName == ValueName)
 	{
@@ -56,19 +56,19 @@ bool FRigVMDispatch_CastEnum::GetPermutationsFromArgumentType(const FName& InArg
 
 #if WITH_EDITOR
 
-FString FRigVMDispatch_CastEnum::GetNodeTitle(const FRigVMTemplateTypeMap& InTypes) const
+FString FRigVMDispatch_CastEnumToInt::GetNodeTitle(const FRigVMTemplateTypeMap& InTypes) const
 {
 	return TEXT("Cast to int");
 }
 
-FText FRigVMDispatch_CastEnum::GetNodeTooltip(const FRigVMTemplateTypeMap& InTypes) const
+FText FRigVMDispatch_CastEnumToInt::GetNodeTooltip(const FRigVMTemplateTypeMap& InTypes) const
 {
-	return LOCTEXT("CastToolTip", "Casts from enum to int");
+	return LOCTEXT("CastEnumToolTip", "Casts from enum to int");
 }
 
 #endif
 
-void FRigVMDispatch_CastEnum::Execute(FRigVMExtendedExecuteContext& InContext, FRigVMMemoryHandleArray Handles, FRigVMPredicateBranchArray RigVMBranches)
+void FRigVMDispatch_CastEnumToInt::Execute(FRigVMExtendedExecuteContext& InContext, FRigVMMemoryHandleArray Handles, FRigVMPredicateBranchArray RigVMBranches)
 {
 	const FEnumProperty* ValueProperty = CastFieldChecked<FEnumProperty>(Handles[0].GetProperty());
 	const FProperty* ResultProperty = CastFieldChecked<FProperty>(Handles[1].GetProperty());
@@ -103,7 +103,7 @@ void FRigVMDispatch_CastEnum::Execute(FRigVMExtendedExecuteContext& InContext, F
 
 
 
-const TArray<FRigVMTemplateArgumentInfo>& FRigVMDispatch_CastEnumToInt::GetArgumentInfos() const
+const TArray<FRigVMTemplateArgumentInfo>& FRigVMDispatch_CastIntToEnum::GetArgumentInfos() const
 {
 	static TArray<FRigVMTemplateArgumentInfo> OutInfos;
 	if (OutInfos.IsEmpty())
@@ -123,7 +123,7 @@ const TArray<FRigVMTemplateArgumentInfo>& FRigVMDispatch_CastEnumToInt::GetArgum
 	return OutInfos;
 }
 
-bool FRigVMDispatch_CastEnumToInt::GetPermutationsFromArgumentType(const FName& InArgumentName, const TRigVMTypeIndex& InTypeIndex, TArray<FRigVMTemplateTypeMap, TInlineAllocator<1>>& OutPermutations) const
+bool FRigVMDispatch_CastIntToEnum::GetPermutationsFromArgumentType(const FName& InArgumentName, const TRigVMTypeIndex& InTypeIndex, TArray<FRigVMTemplateTypeMap, TInlineAllocator<1>>& OutPermutations) const
 {
 	if (InArgumentName == ResultName)
 	{
@@ -150,19 +150,19 @@ bool FRigVMDispatch_CastEnumToInt::GetPermutationsFromArgumentType(const FName&
 
 #if WITH_EDITOR
 
-FString FRigVMDispatch_CastEnumToInt::GetNodeTitle(const FRigVMTemplateTypeMap& InTypes) const
+FString FRigVMDispatch_CastIntToEnum::GetNodeTitle(const FRigVMTemplateTypeMap& InTypes) const
 {
 	return TEXT("Cast to enum");
 }
 
-FText FRigVMDispatch_CastEnumToInt::GetNodeTooltip(const FRigVMTemplateTypeMap& InTypes) const
+FText FRigVMDispatch_CastIntToEnum::GetNodeTooltip(const FRigVMTemplateTypeMap& InTypes) const
 {
-	return LOCTEXT("CastToolTip", "Casts from int to enum");
+	return LOCTEXT("CastIntToEnumToolTip", "Casts from int to enum");
 }
 
 #endif
 
-void FRigVMDispatch_CastEnumToInt::Execute(FRigVMExtendedExecuteContext& InContext, FRigVMMemoryHandleArray Handles, FRigVMPredicateBranchArray RigVMBranches)
+void FRigVMDispatch_CastIntToEnum::Execute(FRigVMExtendedExecuteContext& InContext, FRigVMMemoryHandleArray Handles, FRigVMPredicateBranchArray RigVMBranches)
 {
 	const FProperty* ValueProperty = CastFieldChecked<FProperty>(Handles[0].GetProperty());
 	const FEnumProperty* ResultProperty = CastFieldChecked<FEnumProperty>(Handles[1].GetProperty());
diff --git a/Engine/Plugins/Runtime/RigVM/Source/RigVM/Public/RigVMFunctions/RigVMDispatch_CastEnum.h b/Engine/Plugins/Runtime/RigVM/Source/RigVM/Public/RigVMFunctions/RigVMDispatch_CastEnum.h
--- a/Engine/Plugins/Runtime/RigVM/Source/RigVM/Public/RigVMFunctions/RigVMDispatch_CastEnum.h
+++ b/Engine/Plugins/Runtime/RigVM/Source/RigVM/Public/RigVMFunctions/RigVMDispatch_CastEnum.h
@@ -12,12 +12,12 @@ struct FRigVMDispatch_CastEnumBase
 };
 
 USTRUCT(meta=(DisplayName = "Cast", Category = "Enum", Keywords = "As", NodeColor = "1,1,1,1"))
-struct RIGVM_API FRigVMDispatch_CastEnum : public FRigVMDispatchFactory, public FRigVMDispatch_CastEnumBase
+struct RIGVM_API FRigVMDispatch_CastEnumToInt : public FRigVMDispatchFactory, public FRigVMDispatch_CastEnumBase
 {
 	GENERATED_BODY()
 
 public:
-	FRigVMDispatch_CastEnum()
+	FRigVMDispatch_CastEnumToInt()
 	{
 		FactoryScriptStruct = StaticStruct();
 	}
@@ -31,17 +31,17 @@ public:
 	virtual bool IsSingleton() const override { return true; }
 	
 protected:
-	virtual FRigVMFunctionPtr GetDispatchFunctionImpl(const FRigVMTemplateTypeMap& InTypes) const override { return &FRigVMDispatch_CastEnum::Execute; }
+	virtual FRigVMFunctionPtr GetDispatchFunctionImpl(const FRigVMTemplateTypeMap& InTypes) const override { return &FRigVMDispatch_CastEnumToInt::Execute; }
 	static void Execute(FRigVMExtendedExecuteContext& InContext, FRigVMMemoryHandleArray Handles, FRigVMPredicateBranchArray RigVMBranches);
 };
 
 USTRUCT(meta=(DisplayName = "Cast", Category = "Enum", Keywords = "As", NodeColor = "1,1,1,1"))
-struct RIGVM_API FRigVMDispatch_CastEnumToInt : public FRigVMDispatchFactory, public FRigVMDispatch_CastEnumBase
+struct RIGVM_API FRigVMDispatch_CastIntToEnum : public FRigVMDispatchFactory, public FRigVMDispatch_CastEnumBase
 {
 	GENERATED_BODY()
 
 public:
-	FRigVMDispatch_CastEnumToInt()
+	FRigVMDispatch_CastIntToEnum()
 	{
 		FactoryScriptStruct = StaticStruct();
 	}
@@ -55,6 +55,6 @@ public:
 	virtual bool IsSingleton() const override { return true; }
 	
 protected:
-	virtual FRigVMFunctionPtr GetDispatchFunctionImpl(const FRigVMTemplateTypeMap& InTypes) const override { return &FRigVMDispatch_CastEnumToInt::Execute; }
+	virtual FRigVMFunctionPtr GetDispatchFunctionImpl(const FRigVMTemplateTypeMap& InTypes) const override { return &FRigVMDispatch_CastIntToEnum::Execute; }
 	static void Execute(FRigVMExtendedExecuteContext& InContext, FRigVMMemoryHandleArray Handles, FRigVMPredicateBranchArray RigVMBranches);
 };
diff --git a/Engine/Plugins/Runtime/RigVM/Source/RigVMDeveloper/Private/EdGraph/RigVMEdGraphNode.cpp b/Engine/Plugins/Runtime/RigVM/Source/RigVMDeveloper/Private/EdGraph/RigVMEdGraphNode.cpp
--- a/Engine/Plugins/Runtime/RigVM/Source/RigVMDeveloper/Private/EdGraph/RigVMEdGraphNode.cpp
+++ b/Engine/Plugins/Runtime/RigVM/Source/RigVMDeveloper/Private/EdGraph/RigVMEdGraphNode.cpp
@@ -1061,7 +1061,7 @@ bool URigVMEdGraphNode::DrawAsCompactNode() const
 		{
 			if(TemplateModelNode->GetNotation() == RigVMTypeUtils::GetCastTemplateNotation() ||
 				TemplateModelNode->GetNotation() == FRigVMDispatch_CastObject().GetTemplateNotation() ||
-				TemplateModelNode->GetNotation() == FRigVMDispatch_CastEnum().GetTemplateNotation() ||
+				TemplateModelNode->GetNotation() == FRigVMDispatch_CastIntToEnum().GetTemplateNotation() ||
 				TemplateModelNode->GetNotation() == FRigVMDispatch_CastEnumToInt().GetTemplateNotation())
 			{
 				DrawAsCompactNodeCache = true;
diff --git a/Engine/Plugins/Runtime/RigVM/Source/RigVMDeveloper/Private/RigVMDeveloperTypeUtils.cpp b/Engine/Plugins/Runtime/RigVM/Source/RigVMDeveloper/Private/RigVMDeveloperTypeUtils.cpp
--- a/Engine/Plugins/Runtime/RigVM/Source/RigVMDeveloper/Private/RigVMDeveloperTypeUtils.cpp
+++ b/Engine/Plugins/Runtime/RigVM/Source/RigVMDeveloper/Private/RigVMDeveloperTypeUtils.cpp
@@ -640,7 +640,7 @@ const FRigVMFunction* RigVMTypeUtils::GetCastForTypeIndices(const TRigVMTypeInde
 	{
 		Registry.FindTemplate(CastTemplateNotation),
 		Registry.FindTemplate(FRigVMDispatch_CastObject().GetTemplateNotation()),
-		Registry.FindTemplate(FRigVMDispatch_CastEnum().GetTemplateNotation()),
+		Registry.FindTemplate(FRigVMDispatch_CastIntToEnum().GetTemplateNotation()),
 		Registry.FindTemplate(FRigVMDispatch_CastEnumToInt().GetTemplateNotation())
 	};
 
