Control Rig: Fix crash when opening rigvm with old cast enum nodes
#rnx

diff --git a/Engine/Plugins/Runtime/RigVM/Source/RigVMDeveloper/Private/EdGraph/RigVMEdGraphNode.cpp b/Engine/Plugins/Runtime/RigVM/Source/RigVMDeveloper/Private/EdGraph/RigVMEdGraphNode.cpp
--- a/Engine/Plugins/Runtime/RigVM/Source/RigVMDeveloper/Private/EdGraph/RigVMEdGraphNode.cpp
+++ b/Engine/Plugins/Runtime/RigVM/Source/RigVMDeveloper/Private/EdGraph/RigVMEdGraphNode.cpp
@@ -1392,7 +1392,10 @@ FSlateIcon URigVMEdGraphNode::GetIconAndTint(FLinearColor& OutColor) const
 		}
 		else if (const URigVMDispatchNode* DispatchNode = Cast<URigVMDispatchNode>(ModelNode))
 		{
-			MetadataScriptStruct = DispatchNode->GetFactory()->GetScriptStruct();
+			if (const FRigVMDispatchFactory* Factory = DispatchNode->GetFactory())
+			{
+				MetadataScriptStruct = Factory->GetScriptStruct();
+			}
 		}
 
 		if(MetadataScriptStruct && MetadataScriptStruct->HasMetaDataHierarchical(FRigVMStruct::IconMetaName))
diff --git a/Engine/Plugins/Runtime/RigVM/Source/RigVMDeveloper/Private/RigVMModel/RigVMController.cpp b/Engine/Plugins/Runtime/RigVM/Source/RigVMDeveloper/Private/RigVMModel/RigVMController.cpp
--- a/Engine/Plugins/Runtime/RigVM/Source/RigVMDeveloper/Private/RigVMModel/RigVMController.cpp
+++ b/Engine/Plugins/Runtime/RigVM/Source/RigVMDeveloper/Private/RigVMModel/RigVMController.cpp
@@ -15038,6 +15038,10 @@ bool URigVMController::GenerateNewPinInfos(const FRigVMRegistry& Registry, URigV
 		}
 
 		const FRigVMTemplate* Template = DispatchNode->GetTemplate();
+		if (!Template)
+		{
+			return false;
+		}
 
 		FRigVMDispatchContext DispatchContext = DispatchNode->GetDispatchContext();
 		auto AddExecutePins = [Template, DispatchNode, &Registry, &DispatchContext, &NewPinInfos, &PreviousPinInfos, this](ERigVMPinDirection InPinDirection)
