RigVM: Make sure to react to user defined struct / enum being loaded later

#rnx

diff --git a/Engine/Plugins/Runtime/RigVM/Source/RigVM/Private/RigVMCore/RigVMRegistry.cpp b/Engine/Plugins/Runtime/RigVM/Source/RigVM/Private/RigVMCore/RigVMRegistry.cpp
--- a/Engine/Plugins/Runtime/RigVM/Source/RigVM/Private/RigVMCore/RigVMRegistry.cpp
+++ b/Engine/Plugins/Runtime/RigVM/Source/RigVM/Private/RigVMCore/RigVMRegistry.cpp
@@ -305,6 +305,9 @@ void FRigVMRegistry::InitializeIfNeeded()
 				AssetRegistryModule->Get().OnAssetRemoved().RemoveAll(this);
 				AssetRegistryModule->Get().OnAssetRenamed().RemoveAll(this);
 			}
+
+			UUserDefinedStruct::OnStructLoaded().RemoveAll(this);
+			UUserDefinedEnum::OnEnumLoaded().RemoveAll(this);
 		}
 
 		IPluginManager::Get().OnPluginUnmounted().RemoveAll(this);
@@ -316,6 +319,9 @@ void FRigVMRegistry::InitializeIfNeeded()
 	AssetRegistryModule.Get().OnAssetRemoved().AddRaw(this, &FRigVMRegistry::OnAssetRemoved);
 	AssetRegistryModule.Get().OnAssetRenamed().AddRaw(this, &FRigVMRegistry::OnAssetRenamed);
 
+	UUserDefinedStruct::OnStructLoaded().AddRaw(this, &FRigVMRegistry::OnUserDefinedStructLoaded);
+	UUserDefinedEnum::OnEnumLoaded().AddRaw(this, &FRigVMRegistry::OnUserDefinedEnumLoaded);
+
 	IPluginManager::Get().OnPluginUnmounted().AddRaw(this, &FRigVMRegistry::OnPluginUnloaded);
 	
 	UE::Anim::AttributeTypes::GetOnAttributeTypesChanged().AddRaw(this, &FRigVMRegistry::OnAnimationAttributeTypesChanged);
@@ -434,6 +440,40 @@ void FRigVMRegistry::OnAssetRemoved(const FAssetData& InAssetData)
 	}
 }
 
+void FRigVMRegistry::OnUserDefinedStructLoaded(UUserDefinedStruct* InLoadedStruct)
+{
+	OnUserDefinedTypeLoaded(InLoadedStruct);
+}
+
+void FRigVMRegistry::OnUserDefinedEnumLoaded(UUserDefinedEnum* InLoadedEnum)
+{
+	OnUserDefinedTypeLoaded(InLoadedEnum);
+}
+
+void FRigVMRegistry::OnUserDefinedTypeLoaded(UObject* InLoadedTypeObject)
+{
+	check(InLoadedTypeObject);
+	
+	FScopeLock RefreshTypesScopeLock(&RefreshTypesMutex);
+	
+	const FSoftObjectPath Path(InLoadedTypeObject);
+	if(!UserDefinedTypeToIndex.Contains(Path))
+	{
+		if(UScriptStruct* Struct = Cast<UScriptStruct>(InLoadedTypeObject))
+		{
+			(void)FindOrAddType(Struct);
+		}
+		else if(UEnum* Enum = Cast<UEnum>(InLoadedTypeObject))
+		{
+			(void)FindOrAddType(Enum);
+		}
+		else
+		{
+			checkNoEntry();
+		}
+	}
+}
+
 void FRigVMRegistry::OnPluginUnloaded(IPlugin& InPlugin)
 {
 	const FString PluginContentPath = InPlugin.GetMountedAssetPath();
@@ -925,7 +965,16 @@ TRigVMTypeIndex FRigVMRegistry::GetTypeIndex(const FRigVMTemplateArgumentType& I
 {
 	if(const TRigVMTypeIndex* Index = TypeToIndex.Find(InType))
 	{
+#if UE_RIGVM_DEBUG_TYPEINDEX
+		TRigVMTypeIndex TypeIndex = *Index;
+		if(TypeIndex != INDEX_NONE)
+		{
+			TypeIndex.Name = Types[TypeIndex].Type.CPPType;
+		}
+		return TypeIndex;
+#else
 		return *Index;
+#endif
 	}
 	return INDEX_NONE;
 }
@@ -988,6 +1037,12 @@ TRigVMTypeIndex FRigVMRegistry::GetTypeIndexFromCPPType(const FString& InCPPType
 			});
 		}
 	}
+#if UE_RIGVM_DEBUG_TYPEINDEX
+	if(Result != INDEX_NONE)
+	{
+		Result.Name = Types[Result].Type.CPPType;
+	}
+#endif
 	return Result;
 }
 
diff --git a/Engine/Plugins/Runtime/RigVM/Source/RigVM/Public/RigVMCore/RigVMRegistry.h b/Engine/Plugins/Runtime/RigVM/Source/RigVM/Public/RigVMCore/RigVMRegistry.h
--- a/Engine/Plugins/Runtime/RigVM/Source/RigVM/Public/RigVMCore/RigVMRegistry.h
+++ b/Engine/Plugins/Runtime/RigVM/Source/RigVM/Public/RigVMCore/RigVMRegistry.h
@@ -90,6 +90,15 @@ public:
 	// Update the registry when old types are removed
     void OnAssetRemoved(const FAssetData& InAssetData);
 
+	// Update the registry a user defined struct is loaded
+	void OnUserDefinedStructLoaded(UUserDefinedStruct* InLoadedStruct);
+
+	// Update the registry a user defined enum is loaded
+	void OnUserDefinedEnumLoaded(UUserDefinedEnum* InLoadedEnum);
+
+	// Update the registry a user defined type is loaded
+	void OnUserDefinedTypeLoaded(UObject* InLoadedTypeObject);
+
 	// Removes all types associated with a plugin that's being unloaded. 
 	void OnPluginUnloaded(IPlugin& InPlugin);
 	
diff --git a/Engine/Plugins/Runtime/RigVM/Source/RigVMDeveloper/Private/RigVMModel/RigVMController.cpp b/Engine/Plugins/Runtime/RigVM/Source/RigVMDeveloper/Private/RigVMModel/RigVMController.cpp
--- a/Engine/Plugins/Runtime/RigVM/Source/RigVMDeveloper/Private/RigVMModel/RigVMController.cpp
+++ b/Engine/Plugins/Runtime/RigVM/Source/RigVMDeveloper/Private/RigVMModel/RigVMController.cpp
@@ -15562,20 +15562,30 @@ void URigVMController::RepopulatePinsOnNode(const FRigVMRegistry& Registry, cons
 		}
 	}
 	// add missing pins
-	for (int32 Index = 0; Index < NodeData.NewPinsToAdd.Num(); Index++)
 	{
-		const FString& PinPath = NodeData.NewPinInfos.GetPinPath(NodeData.NewPinsToAdd[Index]);
-		FString ParentPinPath, PinName;
-		UObject* OuterForPin = InNode;
-		if (URigVMPin::SplitPinPathAtEnd(PinPath, ParentPinPath, PinName))
+		TArray<bool> PinIndexAdded;
+		PinIndexAdded.AddZeroed(NodeData.NewPinInfos.Num());
+		for (int32 Index = 0; Index < NodeData.NewPinsToAdd.Num(); Index++)
 		{
-			OuterForPin = InNode->FindPin(ParentPinPath);
-		}
+			const int32 PinIndexToAdd = NodeData.NewPinsToAdd[Index];
+			if(PinIndexAdded[PinIndexToAdd])
+			{
+				continue;
+			}
+			const FString& PinPath = NodeData.NewPinInfos.GetPinPath(PinIndexToAdd);
+			FString ParentPinPath, PinName;
+			UObject* OuterForPin = InNode;
+			if (URigVMPin::SplitPinPathAtEnd(PinPath, ParentPinPath, PinName))
+			{
+				OuterForPin = InNode->FindPin(ParentPinPath);
+			}
 
-		CreatePinFromPinInfo(Registry, NodeData.PreviousPinInfos, NodeData.NewPinInfos[NodeData.NewPinsToAdd[Index]], PinPath, OuterForPin);
+			CreatePinFromPinInfo(Registry, NodeData.PreviousPinInfos, NodeData.NewPinInfos[NodeData.NewPinsToAdd[Index]], PinPath, OuterForPin);
 #if UE_RIGVMCONTROLLER_VERBOSE_REPOPULATE
-		UE_LOG(LogRigVMDeveloper, Display, TEXT("Adding new pin '%s'."), *PinPath);
+			UE_LOG(LogRigVMDeveloper, Display, TEXT("Adding new pin '%s'."), *PinPath);
 #endif
+			PinIndexAdded[PinIndexToAdd] = true;
+		}
 	}
 	// update existing pins
 	for (int32 Index = 0; Index < NodeData.PreviousPinsToUpdate.Num(); Index++)
diff --git a/Engine/Source/Runtime/Engine/Classes/Engine/UserDefinedEnum.h b/Engine/Source/Runtime/Engine/Classes/Engine/UserDefinedEnum.h
--- a/Engine/Source/Runtime/Engine/Classes/Engine/UserDefinedEnum.h
+++ b/Engine/Source/Runtime/Engine/Classes/Engine/UserDefinedEnum.h
@@ -16,6 +16,10 @@
 #include "UserDefinedEnum.generated.h"
 
 class UEnumCookedMetaData;
+class UUserDefinedEnum;
+
+// this is fired when the user defined type has finished loading in the serialize method (not after postload)
+DECLARE_MULTICAST_DELEGATE_OneParam(FOnEnumLoaded, UUserDefinedEnum*);
 
 /** 
  *	An Enumeration is a list of named values.
@@ -87,6 +91,15 @@ public:
 	ENGINE_API FString GenerateNewEnumeratorName();
 #endif	// WITH_EDITOR
 
+public:
+	ENGINE_API FORCEINLINE static FOnEnumLoaded& OnEnumLoaded() { return LoadedEvent; }
+
+protected:
+
+	ENGINE_API static FOnEnumLoaded LoadedEvent;
+
+public:
+
 #if WITH_EDITORONLY_DATA
 protected:
 	ENGINE_API virtual TSubclassOf<UEnumCookedMetaData> GetCookedMetaDataClass() const;
diff --git a/Engine/Source/Runtime/Engine/Classes/Engine/UserDefinedStruct.h b/Engine/Source/Runtime/Engine/Classes/Engine/UserDefinedStruct.h
--- a/Engine/Source/Runtime/Engine/Classes/Engine/UserDefinedStruct.h
+++ b/Engine/Source/Runtime/Engine/Classes/Engine/UserDefinedStruct.h
@@ -17,6 +17,9 @@ class UStructCookedMetaData;
 
 DECLARE_MULTICAST_DELEGATE_OneParam(FOnStructChanged, UUserDefinedStruct*);
 
+// this is fired when the user defined type has finished loading in the serialize method (not after postload)
+DECLARE_MULTICAST_DELEGATE_OneParam(FOnStructLoaded, UUserDefinedStruct*);
+
 UENUM()
 enum EUserDefinedStructureStatus : int
 {
@@ -91,19 +94,18 @@ public:
 	ENGINE_API virtual void GetAssetRegistryTags(FAssetRegistryTagsContext Context) const override;
 	UE_DEPRECATED(5.4, "Implement the version that takes FAssetRegistryTagsContext instead.")
 	ENGINE_API virtual void GetAssetRegistryTags(TArray<FAssetRegistryTag>& OutTags) const override;
-	ENGINE_API virtual void PostLoad() override;
 	ENGINE_API virtual void PreSaveRoot(FObjectPreSaveRootContext ObjectSaveContext) override;
 	ENGINE_API virtual void PostSaveRoot(FObjectPostSaveRootContext ObjectSaveContext) override;
 	// End of UObject interface.
 
-	/** Creates a new guid if needed */
-	ENGINE_API void ValidateGuid();
-
 	ENGINE_API virtual void OnChanged();
 
 	friend UUserDefinedStructEditorData;
 #endif	// WITH_EDITOR
 
+	/** Creates a new guid if needed */
+	ENGINE_API void ValidateGuid();
+
 	// UObject interface.
 	ENGINE_API virtual void Serialize(FStructuredArchive::FRecord Record) override;
 	ENGINE_API virtual void SerializeTaggedProperties(FStructuredArchive::FSlot Slot, uint8* Data, UStruct* DefaultsStruct, uint8* Defaults, const UObject* BreakRecursionIfFullyLoad = nullptr) const override;
@@ -135,6 +137,15 @@ public:
 	/** Inspects properties and default values, setting appropriate StructFlags */
 	ENGINE_API void UpdateStructFlags();
 
+public:
+	ENGINE_API FORCEINLINE static FOnStructLoaded& OnStructLoaded() { return LoadedEvent; }
+
+protected:
+
+	ENGINE_API static FOnStructLoaded LoadedEvent;
+
+public:
+
 #if WITH_EDITORONLY_DATA
 	FORCEINLINE FOnStructChanged& OnStructChanged() { return ChangedEvent; }
 public:
diff --git a/Engine/Source/Runtime/Engine/Private/UserDefinedEnum.cpp b/Engine/Source/Runtime/Engine/Private/UserDefinedEnum.cpp
--- a/Engine/Source/Runtime/Engine/Private/UserDefinedEnum.cpp
+++ b/Engine/Source/Runtime/Engine/Private/UserDefinedEnum.cpp
@@ -16,6 +16,8 @@
 #include "UObject/Package.h"
 #endif	// WITH_EDITOR
 
+FOnEnumLoaded UUserDefinedEnum::LoadedEvent;
+
 UUserDefinedEnum::UUserDefinedEnum(const FObjectInitializer& ObjectInitializer)
 	: Super(ObjectInitializer)
 {
@@ -43,6 +45,11 @@ void UUserDefinedEnum::Serialize(FArchive& Ar)
 		}
 	}
 #endif // WITH_EDITOR
+
+	if(Ar.IsLoading() && Ar.IsPersistent())
+	{
+		LoadedEvent.Broadcast(this);
+	}
 }
 
 FString UUserDefinedEnum::GenerateFullEnumName(const TCHAR* InEnumName) const
diff --git a/Engine/Source/Runtime/Engine/Private/UserDefinedStruct.cpp b/Engine/Source/Runtime/Engine/Private/UserDefinedStruct.cpp
--- a/Engine/Source/Runtime/Engine/Private/UserDefinedStruct.cpp
+++ b/Engine/Source/Runtime/Engine/Private/UserDefinedStruct.cpp
@@ -47,6 +47,8 @@ void FUserStructOnScopeIgnoreDefaults::Initialize()
 	}
 }
 
+FOnStructLoaded UUserDefinedStruct::LoadedEvent;
+
 UUserDefinedStruct::UUserDefinedStruct(const FObjectInitializer& ObjectInitializer)
 	: Super(ObjectInitializer)
 {
@@ -140,6 +142,12 @@ void UUserDefinedStruct::Serialize(FStructuredArchive::FRecord Record)
 		}
 	}
 #endif
+
+	if(UnderlyingArchive.IsLoading() && UnderlyingArchive.IsPersistent())
+	{
+		ValidateGuid();
+		LoadedEvent.Broadcast(this);
+	}
 }
 
 #if WITH_EDITOR
@@ -158,13 +166,6 @@ void UUserDefinedStruct::PostDuplicate(bool bDuplicateForPIE)
 	}
 }
 
-void UUserDefinedStruct::PostLoad()
-{
-	Super::PostLoad();
-
-	ValidateGuid();
-}
-
 void UUserDefinedStruct::PreSaveRoot(FObjectPreSaveRootContext ObjectSaveContext)
 {
 	Super::PreSaveRoot(ObjectSaveContext);
@@ -206,6 +207,13 @@ void UUserDefinedStruct::GetAssetRegistryTags(FAssetRegistryTagsContext Context)
 	Context.AddTag(FAssetRegistryTag(TEXT("Tooltip"), FStructureEditorUtils::GetTooltip(this), FAssetRegistryTag::TT_Hidden));
 }
 
+void UUserDefinedStruct::OnChanged()
+{
+	ChangedEvent.Broadcast(this);
+}
+
+#endif	// WITH_EDITOR
+
 void UUserDefinedStruct::ValidateGuid()
 {
 	// Backward compatibility:
@@ -222,13 +230,6 @@ void UUserDefinedStruct::ValidateGuid()
 	}
 }
 
-void UUserDefinedStruct::OnChanged()
-{
-	ChangedEvent.Broadcast(this);
-}
-
-#endif	// WITH_EDITOR
-
 FProperty* UUserDefinedStruct::CustomFindProperty(const FName Name) const
 {
 #if WITH_EDITOR
