RigVM: Thread-safe initialization of RigVMRegistry

#rnx

diff --git a/Engine/Plugins/Runtime/RigVM/Source/RigVM/Private/RigVMCore/RigVMRegistry.cpp b/Engine/Plugins/Runtime/RigVM/Source/RigVM/Private/RigVMCore/RigVMRegistry.cpp
--- a/Engine/Plugins/Runtime/RigVM/Source/RigVM/Private/RigVMCore/RigVMRegistry.cpp
+++ b/Engine/Plugins/Runtime/RigVM/Source/RigVM/Private/RigVMCore/RigVMRegistry.cpp
@@ -35,6 +35,13 @@ static FDelayedAutoRegisterHelper GRigVMRegistrySingletonHelper(EDelayedRegister
 });
 
 
+FRigVMRegistry::FRigVMRegistry() :
+	bIsRefreshingEngineTypes(false),
+	bEverRefreshedEngineTypes(false)
+{
+	Initialize();
+}
+
 FRigVMRegistry::~FRigVMRegistry()
 {
 	Reset();
@@ -45,7 +52,6 @@ FRigVMRegistry& FRigVMRegistry::Get()
 	// static in a function scope ensures that the GC system is initiated before 
 	// the registry constructor is called
 	static FRigVMRegistry s_RigVMRegistry;
-	s_RigVMRegistry.InitializeIfNeeded();
 	return s_RigVMRegistry;
 }
 
@@ -211,16 +217,9 @@ uint32 FRigVMRegistry::GetHashForProperty(const FProperty* InProperty) const
 	return Hash;
 }
 
-void FRigVMRegistry::InitializeIfNeeded()
-{
-	// we don't need to use a mutex here since this is called on
-	// the main thread during engine startup for the first time
-	
-	if(!Types.IsEmpty())
-	{
-		return;
-	}
 
+void FRigVMRegistry::Initialize()
+{
 	Types.Reserve(512);
 	TypeToIndex.Reserve(512);
 	TypesPerCategory.Reserve(19);
@@ -793,7 +792,7 @@ TRigVMTypeIndex FRigVMRegistry::FindOrAddType_Internal(const FRigVMTemplateArgum
 				{
 					// by creating a template argument for the child property
 					// the type will be added by calling ::FindOrAddType_Internal recursively.
-					FRigVMTemplateArgument DummyArgument(Property);
+					FRigVMTemplateArgument DummyArgument(Property, *this);
 				}
 #if WITH_EDITOR
 				else
diff --git a/Engine/Plugins/Runtime/RigVM/Source/RigVM/Private/RigVMCore/RigVMTemplate.cpp b/Engine/Plugins/Runtime/RigVM/Source/RigVM/Private/RigVMCore/RigVMTemplate.cpp
--- a/Engine/Plugins/Runtime/RigVM/Source/RigVM/Private/RigVMCore/RigVMTemplate.cpp
+++ b/Engine/Plugins/Runtime/RigVM/Source/RigVM/Private/RigVMCore/RigVMTemplate.cpp
@@ -39,7 +39,13 @@ FRigVMTemplateArgument::FRigVMTemplateArgument(const FName& InName, ERigVMPinDir
 	, Direction(InDirection)
 {}
 
-FRigVMTemplateArgument::FRigVMTemplateArgument(FProperty* InProperty)
+FRigVMTemplateArgument::FRigVMTemplateArgument(FProperty* InProperty):
+	FRigVMTemplateArgument(InProperty, FRigVMRegistry::Get())
+{
+	
+}
+
+FRigVMTemplateArgument::FRigVMTemplateArgument(FProperty* InProperty, FRigVMRegistry& InRegistry)
 	: Name(InProperty->GetFName())
 {
 #if WITH_EDITOR
@@ -77,10 +83,10 @@ FRigVMTemplateArgument::FRigVMTemplateArgument(FProperty* InProperty)
 	}
 	
 	const FRigVMTemplateArgumentType Type(CPPTypeName, CPPTypeObject);
-	const TRigVMTypeIndex TypeIndex = FRigVMRegistry::Get().FindOrAddType_Internal(Type, true); 
+	const TRigVMTypeIndex TypeIndex = InRegistry.FindOrAddType_Internal(Type, true); 
 
 	TypeIndices.Add(TypeIndex);
-	EnsureValidExecuteType();
+	EnsureValidExecuteType(InRegistry);
 	UpdateTypeToPermutations();
 }
 
@@ -89,7 +95,7 @@ FRigVMTemplateArgument::FRigVMTemplateArgument(const FName& InName, ERigVMPinDir
 	, Direction(InDirection)
 	, TypeIndices({InTypeIndex})
 {
-	EnsureValidExecuteType();
+	EnsureValidExecuteType(FRigVMRegistry::Get());
 	UpdateTypeToPermutations();
 }
 
@@ -99,7 +105,7 @@ FRigVMTemplateArgument::FRigVMTemplateArgument(const FName& InName, ERigVMPinDir
 	, TypeIndices(InTypeIndices)
 {
 	check(TypeIndices.Num() > 0);
-	EnsureValidExecuteType();
+	EnsureValidExecuteType(FRigVMRegistry::Get());
 	UpdateTypeToPermutations();
 }
 
@@ -143,19 +149,18 @@ FRigVMTemplateArgument::FRigVMTemplateArgument(const FName& InName, ERigVMPinDir
 		{
 			TArray<TRigVMTypeIndex> Indices = AllTypes.Array();
 			TypeIndices = MoveTemp(Indices);
-			EnsureValidExecuteType();
+			EnsureValidExecuteType(FRigVMRegistry::Get());
 		}
 
 		UpdateTypeToPermutations();
 	}
 }
 
-void FRigVMTemplateArgument::EnsureValidExecuteType()
+void FRigVMTemplateArgument::EnsureValidExecuteType(FRigVMRegistry& InRegistry)
 {
-	const FRigVMRegistry& Registry = FRigVMRegistry::Get();
 	for(TRigVMTypeIndex& TypeIndex : TypeIndices)
 	{
-		Registry.ConvertExecuteContextToBaseType(TypeIndex);
+		InRegistry.ConvertExecuteContextToBaseType(TypeIndex);
 	}
 }
 
diff --git a/Engine/Plugins/Runtime/RigVM/Source/RigVM/Public/RigVMCore/RigVMRegistry.h b/Engine/Plugins/Runtime/RigVM/Source/RigVM/Public/RigVMCore/RigVMRegistry.h
--- a/Engine/Plugins/Runtime/RigVM/Source/RigVM/Public/RigVMCore/RigVMRegistry.h
+++ b/Engine/Plugins/Runtime/RigVM/Source/RigVM/Public/RigVMCore/RigVMRegistry.h
@@ -39,8 +39,8 @@ struct RIGVM_API FRigVMRegistry : public FGCObject
 public:
 
 	DECLARE_MULTICAST_DELEGATE(FOnRigVMRegistryChanged);
-	
-	~FRigVMRegistry();
+
+	virtual ~FRigVMRegistry() override;
 	
 	// Returns the singleton registry
 	static FRigVMRegistry& Get();
@@ -73,9 +73,6 @@ public:
 	// Register a set of allowed object types
 	void RegisterObjectTypes(TConstArrayView<TPair<UClass*, ERegisterObjectOperation>> InClasses);
 
-	// Initializes the registry by storing the defaults
-	void InitializeIfNeeded();
-	
 	// Refreshes the list and finds the function pointers
 	// based on the names.
 	void RefreshEngineTypes();
@@ -309,13 +306,8 @@ private:
 
 	static const FName TemplateNameMetaName;
 
-	// disable default constructor
-	FRigVMRegistry()
-		: bIsRefreshingEngineTypes(false)
-		, bEverRefreshedEngineTypes(false)
-	{
-	}
-	
+	FRigVMRegistry();
+
 	// disable copy constructor
 	FRigVMRegistry(const FRigVMRegistry&) = delete;
 	// disable assignment operator
@@ -340,6 +332,9 @@ private:
 		uint32 Hash;
 	};
 
+	// Initialize the base types
+	void Initialize();
+	
 	TRigVMTypeIndex FindOrAddType_Internal(const FRigVMTemplateArgumentType& InType, bool bForce);
 
 	void RefreshEngineTypes_NoLock();
diff --git a/Engine/Plugins/Runtime/RigVM/Source/RigVM/Public/RigVMCore/RigVMTemplate.h b/Engine/Plugins/Runtime/RigVM/Source/RigVM/Public/RigVMCore/RigVMTemplate.h
--- a/Engine/Plugins/Runtime/RigVM/Source/RigVM/Public/RigVMCore/RigVMTemplate.h
+++ b/Engine/Plugins/Runtime/RigVM/Source/RigVM/Public/RigVMCore/RigVMTemplate.h
@@ -4,7 +4,6 @@
 
 #include "Containers/Array.h"
 #include "Containers/Map.h"
-#include "Containers/SparseArray.h"
 #include "Containers/UnrealString.h"
 #include "CoreMinimal.h"
 #include "Delegates/Delegate.h"
@@ -13,8 +12,6 @@
 #include "Math/Color.h"
 #include "Misc/AssertionMacros.h"
 #include "RigVMFunction.h"
-#include "RigVMModule.h"
-#include "RigVMTraits.h"
 #include "RigVMTypeIndex.h"
 #include "RigVMTypeUtils.h"
 #include "UObject/Class.h"
@@ -31,6 +28,7 @@ struct FRigVMDispatchFactory;
 struct FRigVMTemplate;
 struct FRigVMDispatchContext;
 struct FRigVMUserDefinedTypeResolver;
+struct FRigVMRegistry;
 
 typedef TMap<FName, TRigVMTypeIndex> FRigVMTemplateTypeMap;
 
@@ -264,10 +262,12 @@ protected:
 	bool bUseCategories = false;
 	TArray<ETypeCategory> TypeCategories;
 
-	// constructor from a property. this forces the type to be created
 	FRigVMTemplateArgument(FProperty* InProperty);
 
-	void EnsureValidExecuteType();
+	// constructor from a property. this forces the type to be created
+	FRigVMTemplateArgument(FProperty* InProperty, FRigVMRegistry& InRegistry);
+
+	void EnsureValidExecuteType(FRigVMRegistry& InRegistry);
 	void UpdateTypeToPermutations();
 
 	friend struct FRigVMTemplate;
