RigVM: Make FindOrAddType threadsafe

29.01 side CL: 32270943

#rnx

diff --git a/Engine/Plugins/Runtime/RigVM/Source/RigVM/Private/RigVMCore/RigVMRegistry.cpp b/Engine/Plugins/Runtime/RigVM/Source/RigVM/Private/RigVMCore/RigVMRegistry.cpp
--- a/Engine/Plugins/Runtime/RigVM/Source/RigVM/Private/RigVMCore/RigVMRegistry.cpp
+++ b/Engine/Plugins/Runtime/RigVM/Source/RigVM/Private/RigVMCore/RigVMRegistry.cpp
@@ -18,7 +18,7 @@
 
 const FName FRigVMRegistry::TemplateNameMetaName = TEXT("TemplateName");
 
-FCriticalSection FRigVMRegistry::RefreshTypesMutex;
+FCriticalSection FRigVMRegistry::FindOrAddTypeMutex;
 
 FCriticalSection FRigVMRegistry::FunctionRegistryMutex;
 FCriticalSection FRigVMRegistry::FactoryRegistryMutex;
@@ -220,6 +220,11 @@ uint32 FRigVMRegistry::GetHashForProperty(const FProperty* InProperty) const
 
 void FRigVMRegistry::Initialize()
 {
+	// this should not be necessary since the initialize is used only
+	// on a constructor on a static variable (thread safe)
+	// but in case code paths change in the future we'll also lock here.
+	const FScopeLock FindOrAddTypeLock(&FindOrAddTypeMutex);
+	
 	Types.Reserve(512);
 	TypeToIndex.Reserve(512);
 	TypesPerCategory.Reserve(19);
@@ -265,31 +270,31 @@ void FRigVMRegistry::Initialize()
 	TemplatesPerCategory.Add(FRigVMTemplateArgument::ETypeCategory_ArrayObjectValue, TArray<int32>()).Reserve(64);
 	TemplatesPerCategory.Add(FRigVMTemplateArgument::ETypeCategory_ArrayArrayObjectValue, TArray<int32>()).Reserve(64);
 
-	RigVMTypeUtils::TypeIndex::Execute = FindOrAddType(FRigVMTemplateArgumentType(FRigVMExecuteContext::StaticStruct()));
-	RigVMTypeUtils::TypeIndex::ExecuteArray = FindOrAddType(FRigVMTemplateArgumentType(FRigVMExecuteContext::StaticStruct()).ConvertToArray());
-	RigVMTypeUtils::TypeIndex::Bool = FindOrAddType(FRigVMTemplateArgumentType(RigVMTypeUtils::BoolTypeName, nullptr));
-	RigVMTypeUtils::TypeIndex::Float = FindOrAddType(FRigVMTemplateArgumentType(RigVMTypeUtils::FloatTypeName, nullptr));
-	RigVMTypeUtils::TypeIndex::Double = FindOrAddType(FRigVMTemplateArgumentType(RigVMTypeUtils::DoubleTypeName, nullptr));
-	RigVMTypeUtils::TypeIndex::Int32 = FindOrAddType(FRigVMTemplateArgumentType(RigVMTypeUtils::Int32TypeName, nullptr));
-	RigVMTypeUtils::TypeIndex::UInt32 = FindOrAddType(FRigVMTemplateArgumentType(RigVMTypeUtils::UInt32TypeName, nullptr));
-	RigVMTypeUtils::TypeIndex::UInt8 = FindOrAddType(FRigVMTemplateArgumentType(RigVMTypeUtils::UInt8TypeName, nullptr));
-	RigVMTypeUtils::TypeIndex::FName = FindOrAddType(FRigVMTemplateArgumentType(RigVMTypeUtils::FNameTypeName, nullptr));
-	RigVMTypeUtils::TypeIndex::FString = FindOrAddType(FRigVMTemplateArgumentType(RigVMTypeUtils::FStringTypeName, nullptr));
-	RigVMTypeUtils::TypeIndex::WildCard = FindOrAddType(FRigVMTemplateArgumentType(RigVMTypeUtils::GetWildCardCPPTypeName(), RigVMTypeUtils::GetWildCardCPPTypeObject()));
-	RigVMTypeUtils::TypeIndex::BoolArray = FindOrAddType(FRigVMTemplateArgumentType(RigVMTypeUtils::BoolArrayTypeName, nullptr));
-	RigVMTypeUtils::TypeIndex::FloatArray = FindOrAddType(FRigVMTemplateArgumentType(RigVMTypeUtils::FloatArrayTypeName, nullptr));
-	RigVMTypeUtils::TypeIndex::DoubleArray = FindOrAddType(FRigVMTemplateArgumentType(RigVMTypeUtils::DoubleArrayTypeName, nullptr));
-	RigVMTypeUtils::TypeIndex::Int32Array = FindOrAddType(FRigVMTemplateArgumentType(RigVMTypeUtils::Int32ArrayTypeName, nullptr));
-	RigVMTypeUtils::TypeIndex::UInt32Array = FindOrAddType(FRigVMTemplateArgumentType(RigVMTypeUtils::UInt32ArrayTypeName, nullptr));
-	RigVMTypeUtils::TypeIndex::UInt8Array = FindOrAddType(FRigVMTemplateArgumentType(RigVMTypeUtils::UInt8ArrayTypeName, nullptr));
-	RigVMTypeUtils::TypeIndex::FNameArray = FindOrAddType(FRigVMTemplateArgumentType(RigVMTypeUtils::FNameArrayTypeName, nullptr));
-	RigVMTypeUtils::TypeIndex::FStringArray = FindOrAddType(FRigVMTemplateArgumentType(RigVMTypeUtils::FStringArrayTypeName, nullptr));
-	RigVMTypeUtils::TypeIndex::WildCardArray = FindOrAddType(FRigVMTemplateArgumentType(RigVMTypeUtils::GetWildCardArrayCPPTypeName(), RigVMTypeUtils::GetWildCardCPPTypeObject()));
+	RigVMTypeUtils::TypeIndex::Execute = FindOrAddType_NoLock(FRigVMTemplateArgumentType(FRigVMExecuteContext::StaticStruct()), false);
+	RigVMTypeUtils::TypeIndex::ExecuteArray = FindOrAddType_NoLock(FRigVMTemplateArgumentType(FRigVMExecuteContext::StaticStruct()).ConvertToArray(), false);
+	RigVMTypeUtils::TypeIndex::Bool = FindOrAddType_NoLock(FRigVMTemplateArgumentType(RigVMTypeUtils::BoolTypeName, nullptr), false);
+	RigVMTypeUtils::TypeIndex::Float = FindOrAddType_NoLock(FRigVMTemplateArgumentType(RigVMTypeUtils::FloatTypeName, nullptr), false);
+	RigVMTypeUtils::TypeIndex::Double = FindOrAddType_NoLock(FRigVMTemplateArgumentType(RigVMTypeUtils::DoubleTypeName, nullptr), false);
+	RigVMTypeUtils::TypeIndex::Int32 = FindOrAddType_NoLock(FRigVMTemplateArgumentType(RigVMTypeUtils::Int32TypeName, nullptr), false);
+	RigVMTypeUtils::TypeIndex::UInt32 = FindOrAddType_NoLock(FRigVMTemplateArgumentType(RigVMTypeUtils::UInt32TypeName, nullptr), false);
+	RigVMTypeUtils::TypeIndex::UInt8 = FindOrAddType_NoLock(FRigVMTemplateArgumentType(RigVMTypeUtils::UInt8TypeName, nullptr), false);
+	RigVMTypeUtils::TypeIndex::FName = FindOrAddType_NoLock(FRigVMTemplateArgumentType(RigVMTypeUtils::FNameTypeName, nullptr), false);
+	RigVMTypeUtils::TypeIndex::FString = FindOrAddType_NoLock(FRigVMTemplateArgumentType(RigVMTypeUtils::FStringTypeName, nullptr), false);
+	RigVMTypeUtils::TypeIndex::WildCard = FindOrAddType_NoLock(FRigVMTemplateArgumentType(RigVMTypeUtils::GetWildCardCPPTypeName(), RigVMTypeUtils::GetWildCardCPPTypeObject()), false);
+	RigVMTypeUtils::TypeIndex::BoolArray = FindOrAddType_NoLock(FRigVMTemplateArgumentType(RigVMTypeUtils::BoolArrayTypeName, nullptr), false);
+	RigVMTypeUtils::TypeIndex::FloatArray = FindOrAddType_NoLock(FRigVMTemplateArgumentType(RigVMTypeUtils::FloatArrayTypeName, nullptr), false);
+	RigVMTypeUtils::TypeIndex::DoubleArray = FindOrAddType_NoLock(FRigVMTemplateArgumentType(RigVMTypeUtils::DoubleArrayTypeName, nullptr), false);
+	RigVMTypeUtils::TypeIndex::Int32Array = FindOrAddType_NoLock(FRigVMTemplateArgumentType(RigVMTypeUtils::Int32ArrayTypeName, nullptr), false);
+	RigVMTypeUtils::TypeIndex::UInt32Array = FindOrAddType_NoLock(FRigVMTemplateArgumentType(RigVMTypeUtils::UInt32ArrayTypeName, nullptr), false);
+	RigVMTypeUtils::TypeIndex::UInt8Array = FindOrAddType_NoLock(FRigVMTemplateArgumentType(RigVMTypeUtils::UInt8ArrayTypeName, nullptr), false);
+	RigVMTypeUtils::TypeIndex::FNameArray = FindOrAddType_NoLock(FRigVMTemplateArgumentType(RigVMTypeUtils::FNameArrayTypeName, nullptr), false);
+	RigVMTypeUtils::TypeIndex::FStringArray = FindOrAddType_NoLock(FRigVMTemplateArgumentType(RigVMTypeUtils::FStringArrayTypeName, nullptr), false);
+	RigVMTypeUtils::TypeIndex::WildCardArray = FindOrAddType_NoLock(FRigVMTemplateArgumentType(RigVMTypeUtils::GetWildCardArrayCPPTypeName(), RigVMTypeUtils::GetWildCardCPPTypeObject()), false);
 
 	// register the default math types
 	for(UScriptStruct* MathType : GetMathTypes())
 	{
-		FindOrAddType(FRigVMTemplateArgumentType(MathType));
+		FindOrAddType_NoLock(FRigVMTemplateArgumentType(MathType), false);
 	}
 
 	// hook the registry to prepare for engine shutdown
@@ -328,13 +333,13 @@ void FRigVMRegistry::Initialize()
 
 void FRigVMRegistry::RefreshEngineTypes()
 {
-	FScopeLock RefreshTypesScopeLock(&RefreshTypesMutex);
+	FScopeLock FindOrAddTypeLock(&FindOrAddTypeMutex);
 	RefreshEngineTypes_NoLock();
 }
 
 void FRigVMRegistry::RefreshEngineTypesIfRequired()
 {
-	FScopeLock RefreshTypesScopeLock(&RefreshTypesMutex);
+	FScopeLock FindOrAddTypeLock(&FindOrAddTypeMutex);
 	if(bEverRefreshedEngineTypes)
 	{
 		return;
@@ -360,7 +365,7 @@ void FRigVMRegistry::RefreshEngineTypes_NoLock()
 		// if this is a C++ type - skip it
 		if(ScriptStruct->IsA<UUserDefinedStruct>() || ScriptStruct->IsChildOf(FRigVMExecuteContext::StaticStruct()))
 		{
-			FindOrAddType(FRigVMTemplateArgumentType(ScriptStruct));
+			FindOrAddType_NoLock(FRigVMTemplateArgumentType(ScriptStruct), false);
 		}
 		else if (ScriptStruct != FRigVMDispatchFactory::StaticStruct() &&
 				 ScriptStruct->IsChildOf(FRigVMDispatchFactory::StaticStruct()))
@@ -375,7 +380,7 @@ void FRigVMRegistry::RefreshEngineTypes_NoLock()
 		if(IsAllowedType(Enum))
 		{
 			const FString CPPType = Enum->CppType.IsEmpty() ? Enum->GetName() : Enum->CppType;
-			FindOrAddType(FRigVMTemplateArgumentType(*CPPType, Enum));
+			FindOrAddType_NoLock(FRigVMTemplateArgumentType(*CPPType, Enum), false);
 		}
 	}
 	
@@ -385,8 +390,8 @@ void FRigVMRegistry::RefreshEngineTypes_NoLock()
 		if (IsAllowedType(Class))
 		{
 			// Register both the class and the object type for use
-			FindOrAddType(FRigVMTemplateArgumentType(Class, RigVMTypeUtils::EClassArgType::AsClass));
-			FindOrAddType(FRigVMTemplateArgumentType(Class, RigVMTypeUtils::EClassArgType::AsObject));
+			FindOrAddType_NoLock(FRigVMTemplateArgumentType(Class, RigVMTypeUtils::EClassArgType::AsClass), false);
+			FindOrAddType_NoLock(FRigVMTemplateArgumentType(Class, RigVMTypeUtils::EClassArgType::AsObject), false);
 		}
 	}
 
@@ -453,18 +458,18 @@ void FRigVMRegistry::OnUserDefinedTypeLoaded(UObject* InLoadedTypeObject)
 {
 	check(InLoadedTypeObject);
 	
-	FScopeLock RefreshTypesScopeLock(&RefreshTypesMutex);
+	FScopeLock FindOrAddTypeLock(&FindOrAddTypeMutex);
 	
 	const FSoftObjectPath Path(InLoadedTypeObject);
 	if(!UserDefinedTypeToIndex.Contains(Path))
 	{
 		if(UScriptStruct* Struct = Cast<UScriptStruct>(InLoadedTypeObject))
 		{
-			(void)FindOrAddType(Struct);
+			(void)FindOrAddType_NoLock(Struct, false);
 		}
 		else if(UEnum* Enum = Cast<UEnum>(InLoadedTypeObject))
 		{
-			(void)FindOrAddType(Enum);
+			(void)FindOrAddType_NoLock(Enum, false);
 		}
 		else
 		{
@@ -538,7 +543,13 @@ void FRigVMRegistry::Reset()
 	Factories.Reset();
 }
 
-TRigVMTypeIndex FRigVMRegistry::FindOrAddType_Internal(const FRigVMTemplateArgumentType& InType, bool bForce)
+TRigVMTypeIndex FRigVMRegistry::FindOrAddType(const FRigVMTemplateArgumentType& InType, bool bForce)
+{
+	const FScopeLock FindOrAddTypesLock(&FindOrAddTypeMutex);
+	return FindOrAddType_NoLock(InType, bForce);
+}
+
+TRigVMTypeIndex FRigVMRegistry::FindOrAddType_NoLock(const FRigVMTemplateArgumentType& InType, bool bForce)
 {
 	// we don't use a mutex here since by the time the engine relies on worker
 	// thread for execution or async loading all types will have been registered.
diff --git a/Engine/Plugins/Runtime/RigVM/Source/RigVM/Private/RigVMCore/RigVMTemplate.cpp b/Engine/Plugins/Runtime/RigVM/Source/RigVM/Private/RigVMCore/RigVMTemplate.cpp
--- a/Engine/Plugins/Runtime/RigVM/Source/RigVM/Private/RigVMCore/RigVMTemplate.cpp
+++ b/Engine/Plugins/Runtime/RigVM/Source/RigVM/Private/RigVMCore/RigVMTemplate.cpp
@@ -83,7 +83,7 @@ FRigVMTemplateArgument::FRigVMTemplateArgument(FProperty* InProperty, FRigVMRegi
 	}
 	
 	const FRigVMTemplateArgumentType Type(CPPTypeName, CPPTypeObject);
-	const TRigVMTypeIndex TypeIndex = InRegistry.FindOrAddType_Internal(Type, true); 
+	const TRigVMTypeIndex TypeIndex = InRegistry.FindOrAddType(Type, true); 
 
 	TypeIndices.Add(TypeIndex);
 	EnsureValidExecuteType(InRegistry);
diff --git a/Engine/Plugins/Runtime/RigVM/Source/RigVM/Public/RigVMCore/RigVMRegistry.h b/Engine/Plugins/Runtime/RigVM/Source/RigVM/Public/RigVMCore/RigVMRegistry.h
--- a/Engine/Plugins/Runtime/RigVM/Source/RigVM/Public/RigVMCore/RigVMRegistry.h
+++ b/Engine/Plugins/Runtime/RigVM/Source/RigVM/Public/RigVMCore/RigVMRegistry.h
@@ -107,11 +107,8 @@ public:
 	void Reset();
 
 	// Adds a type if it doesn't exist yet and returns its index.
-	// This function is not thead-safe
-	TRigVMTypeIndex FindOrAddType(const FRigVMTemplateArgumentType& InType)
-	{
-		return FindOrAddType_Internal(InType, false);
-	}
+	// This function is thread-safe
+	TRigVMTypeIndex FindOrAddType(const FRigVMTemplateArgumentType& InType, bool bForce = false);
 
 	// Removes a type from the registry, and updates all dependent templates
 	// which also creates invalid permutations in templates that we should ignore
@@ -335,7 +332,7 @@ private:
 	// Initialize the base types
 	void Initialize();
 	
-	TRigVMTypeIndex FindOrAddType_Internal(const FRigVMTemplateArgumentType& InType, bool bForce);
+	TRigVMTypeIndex FindOrAddType_NoLock(const FRigVMTemplateArgumentType& InType, bool bForce);
 
 	void RefreshEngineTypes_NoLock();
 
@@ -419,7 +416,7 @@ private:
 	// This is true if the engine has ever refreshed the engine types
 	bool bEverRefreshedEngineTypes;
 
-	static FCriticalSection RefreshTypesMutex;
+	static FCriticalSection FindOrAddTypeMutex;
 	static FCriticalSection FunctionRegistryMutex;
 	static FCriticalSection FactoryRegistryMutex;
 	static FCriticalSection TemplateRegistryMutex;
