Switched to using FLazyName for control rig built in types

Avoid static initialization order issue in PGO builds

diff --git a/Engine/Plugins/Runtime/RigVM/Source/RigVM/Private/RigVMCore/RigVMMemoryStorageStruct.cpp b/Engine/Plugins/Runtime/RigVM/Source/RigVM/Private/RigVMCore/RigVMMemoryStorageStruct.cpp
--- a/Engine/Plugins/Runtime/RigVM/Source/RigVM/Private/RigVMCore/RigVMMemoryStorageStruct.cpp
+++ b/Engine/Plugins/Runtime/RigVM/Source/RigVM/Private/RigVMCore/RigVMMemoryStorageStruct.cpp
@@ -585,11 +585,11 @@ void FRigVMMemoryStorageStruct::SetDefaultValues(const TArray<FRigVMPropertyDesc
 {
 	EPropertyBagPropertyType BagPropertyType = EPropertyBagPropertyType::None;
 
-	FString VMType = RigVMDescriptor.CPPType;
+	FString VMTypeString = *RigVMDescriptor.CPPType;
 
 	OutBagContainerTypes.Reset(); // in case it is reused by the caller
 
-	if (RigVMTypeUtils::IsArrayType(VMType))
+	if (RigVMTypeUtils::IsArrayType(VMTypeString))
 	{
 		const int32 NumContainers = RigVMDescriptor.Containers.Num();
 		for (int i= 0; i < NumContainers; ++i)
@@ -600,19 +600,21 @@ void FRigVMMemoryStorageStruct::SetDefaultValues(const TArray<FRigVMPropertyDesc
 				OutBagContainerTypes.Add(EPropertyBagContainerType::Array);
 				break;
 			case EPinContainerType::Set:
-				ensureMsgf(false, TEXT("Unsuported Set type container : %s"), *VMType);
+				ensureMsgf(false, TEXT("Unsuported Set type container : %s"), *VMTypeString);
 				break;
 			case EPinContainerType::Map:
-				ensureMsgf(false, TEXT("Unsuported Map type container : %s"), *VMType);
+				ensureMsgf(false, TEXT("Unsuported Map type container : %s"), *VMTypeString);
 				break;
 			default:
 				break;
 			}
 		}
 
-		VMType = RigVMDescriptor.GetTailCPPType();
+		VMTypeString = RigVMDescriptor.GetTailCPPType();
 	}
 
+	FName VMType = *VMTypeString;
+
 	static const FName IntTypeName(TEXT("int")); // type used by some engine tests
 	static const FName Int64TypeName(TEXT("Int64"));
 	static const FName UInt64TypeName(TEXT("UInt64"));
@@ -667,7 +669,7 @@ void FRigVMMemoryStorageStruct::SetDefaultValues(const TArray<FRigVMPropertyDesc
 	}
 	else
 	{
-		ensureMsgf(false, TEXT("Unsupported type : %s"), *VMType);
+		ensureMsgf(false, TEXT("Unsupported type : %s"), *VMTypeString);
 		OutBagPropertyType = EPropertyBagPropertyType::None;
 	}
 
diff --git a/Engine/Plugins/Runtime/RigVM/Source/RigVM/Public/RigVMTypeUtils.h b/Engine/Plugins/Runtime/RigVM/Source/RigVM/Public/RigVMTypeUtils.h
--- a/Engine/Plugins/Runtime/RigVM/Source/RigVM/Public/RigVMTypeUtils.h
+++ b/Engine/Plugins/Runtime/RigVM/Source/RigVM/Public/RigVMTypeUtils.h
@@ -55,42 +55,42 @@ namespace RigVMTypeUtils
 	const TCHAR TSubclassOfTemplate[] = TEXT("TSubclassOf<%s%s>");
 	const TCHAR TScriptInterfaceTemplate[] = TEXT("TScriptInterface<%s%s>");
 
-	const FString BoolType = TEXT("bool");
-	const FString FloatType = TEXT("float");
-	const FString DoubleType = TEXT("double");
-	const FString Int32Type = TEXT("int32");
-	const FString UInt32Type = TEXT("uint32");
-	const FString UInt8Type = TEXT("uint8");
-	const FString FNameType = TEXT("FName");
-	const FString FStringType = TEXT("FString");
-	const FString FTextType = TEXT("FText");
-	const FString BoolArrayType = TEXT("TArray<bool>");
-	const FString FloatArrayType = TEXT("TArray<float>");
-	const FString DoubleArrayType = TEXT("TArray<double>");
-	const FString Int32ArrayType = TEXT("TArray<int32>");
-	const FString UInt32ArrayType = TEXT("TArray<uint32>");
-	const FString UInt8ArrayType = TEXT("TArray<uint8>");
-	const FString FNameArrayType = TEXT("TArray<FName>");
-	const FString FStringArrayType = TEXT("TArray<FString>");
-	const FString FTextArrayType = TEXT("TArray<FText>");
-
-	const FName BoolTypeName = *BoolType;
-	const FName FloatTypeName = *FloatType;
-	const FName DoubleTypeName = *DoubleType;
-	const FName Int32TypeName = *Int32Type;
-	const FName UInt32TypeName = *UInt32Type;
-	const FName UInt8TypeName = *UInt8Type;
-	const FName FNameTypeName = *FNameType;
-	const FName FStringTypeName = *FStringType;
-	const FName FTextTypeName = *FTextType;
-	const FName BoolArrayTypeName = *BoolArrayType;
-	const FName FloatArrayTypeName = *FloatArrayType;
-	const FName DoubleArrayTypeName = *DoubleArrayType;
-	const FName Int32ArrayTypeName = *Int32ArrayType;
-	const FName UInt32ArrayTypeName = *UInt32ArrayType;
-	const FName UInt8ArrayTypeName = *UInt8ArrayType;
-	const FName FNameArrayTypeName = *FNameArrayType;
-	const FName FStringArrayTypeName = *FStringArrayType;
+	const inline TCHAR* BoolType = TEXT("bool");
+	const inline TCHAR* FloatType = TEXT("float");
+	const inline TCHAR* DoubleType = TEXT("double");
+	const inline TCHAR* Int32Type = TEXT("int32");
+	const inline TCHAR* UInt32Type = TEXT("uint32");
+	const inline TCHAR* UInt8Type = TEXT("uint8");
+	const inline TCHAR* FNameType = TEXT("FName");
+	const inline TCHAR* FStringType = TEXT("FString");
+	const inline TCHAR* FTextType = TEXT("FText");
+	const inline TCHAR* BoolArrayType = TEXT("TArray<bool>");
+	const inline TCHAR* FloatArrayType = TEXT("TArray<float>");
+	const inline TCHAR* DoubleArrayType = TEXT("TArray<double>");
+	const inline TCHAR* Int32ArrayType = TEXT("TArray<int32>");
+	const inline TCHAR* UInt32ArrayType = TEXT("TArray<uint32>");
+	const inline TCHAR* UInt8ArrayType = TEXT("TArray<uint8>");
+	const inline TCHAR* FNameArrayType = TEXT("TArray<FName>");
+	const inline TCHAR* FStringArrayType = TEXT("TArray<FString>");
+	const inline TCHAR* FTextArrayType = TEXT("TArray<FText>");
+
+	const FLazyName BoolTypeName(BoolType);
+	const FLazyName FloatTypeName(FloatType);
+	const FLazyName DoubleTypeName(DoubleType);
+	const FLazyName Int32TypeName(Int32Type);
+	const FLazyName UInt32TypeName(UInt32Type);
+	const FLazyName UInt8TypeName(UInt8Type);
+	const FLazyName FNameTypeName(FNameType);
+	const FLazyName FStringTypeName(FStringType);
+	const FLazyName FTextTypeName(FTextType);
+	const FLazyName BoolArrayTypeName(BoolArrayType);
+	const FLazyName FloatArrayTypeName(FloatArrayType);
+	const FLazyName DoubleArrayTypeName(DoubleArrayType);
+	const FLazyName Int32ArrayTypeName(Int32ArrayType);
+	const FLazyName UInt32ArrayTypeName(UInt32ArrayType);
+	const FLazyName UInt8ArrayTypeName(UInt8ArrayType);
+	const FLazyName FNameArrayTypeName(FNameArrayType);
+	const FLazyName FStringArrayTypeName(FStringArrayType);
 
 	class RIGVM_API TypeIndex
 	{
