Control Rig: Fix CR_LatticeDeformerDemo not loading correctly because UserDefinedStruct not updating the SetAnimAttribute template
#rnx

diff --git a/Engine/Plugins/Animation/ControlRig/Source/ControlRig/Private/Units/Animation/RigUnit_AnimAttribute.cpp b/Engine/Plugins/Animation/ControlRig/Source/ControlRig/Private/Units/Animation/RigUnit_AnimAttribute.cpp
--- a/Engine/Plugins/Animation/ControlRig/Source/ControlRig/Private/Units/Animation/RigUnit_AnimAttribute.cpp
+++ b/Engine/Plugins/Animation/ControlRig/Source/ControlRig/Private/Units/Animation/RigUnit_AnimAttribute.cpp
@@ -378,12 +378,12 @@ const TArray<FRigVMTemplateArgumentInfo>& FRigDispatch_GetAnimAttribute::GetArgu
 	{
 		Infos = Super::GetArgumentInfos(); 
 
-		FRigVMTemplateArgument::FTypeFilter	TypeFilter;
-		TypeFilter.BindStatic(&FRigDispatch_AnimAttributeBase::IsTypeSupported);
-		const TArray<TRigVMTypeIndex> Types = FRigVMTemplateArgumentInfo::GetTypesFromCategories(GetValueTypeCategory(), TypeFilter);
-		
-		DefaultArgIndex = Infos.Emplace(DefaultArgName, ERigVMPinDirection::Input, Types);
-		ValueArgIndex = Infos.Emplace(ValueArgName, ERigVMPinDirection::Output, Types);
+		// Will be open to any category, but will filter the type through our IsTypeSupported callback
+		// If we reduce this to multiple (more targeted) categories, and any of those categories have common types, bUseCategories will not be true
+		// and the template will not receive notifications of newly added types
+		const TArray<FRigVMTemplateArgument::ETypeCategory> Categories = {FRigVMTemplateArgument::ETypeCategory_SingleAnyValue};
+		DefaultArgIndex = Infos.Emplace(DefaultArgName, ERigVMPinDirection::Input, Categories, [](const TRigVMTypeIndex& Type) { return FRigDispatch_AnimAttributeBase::IsTypeSupported(Type); });
+		ValueArgIndex = Infos.Emplace(ValueArgName, ERigVMPinDirection::Output, Categories, [](const TRigVMTypeIndex& Type) { return FRigDispatch_AnimAttributeBase::IsTypeSupported(Type); });
 		
 		FoundArgIndex = Infos.Emplace(FoundArgName, ERigVMPinDirection::Output, RigVMTypeUtils::TypeIndex::Bool);
 	}
@@ -461,11 +461,12 @@ const TArray<FRigVMTemplateArgumentInfo>& FRigDispatch_SetAnimAttribute::GetArgu
 	{
 		Infos = Super::GetArgumentInfos();
 
-		FRigVMTemplateArgument::FTypeFilter	TypeFilter;
-		TypeFilter.BindStatic(&FRigDispatch_AnimAttributeBase::IsTypeSupported);
-		const TArray<TRigVMTypeIndex> Types = FRigVMTemplateArgumentInfo::GetTypesFromCategories(GetValueTypeCategory(), TypeFilter);
+		// Will be open to any category, but will filter the type through our IsTypeSupported callback
+		// If we reduce this to multiple (more targeted) categories, and any of those categories have common types, bUseCategories will not be true
+		// and the template will not receive notifications of newly added types
+		static const TArray<FRigVMTemplateArgument::ETypeCategory> Categories = {FRigVMTemplateArgument::ETypeCategory_SingleAnyValue};
 		
-		ValueArgIndex = Infos.Emplace(ValueArgName, ERigVMPinDirection::Input, Types);
+		ValueArgIndex = Infos.Emplace(ValueArgName, ERigVMPinDirection::Input, Categories, [](const TRigVMTypeIndex& Type) { return FRigDispatch_AnimAttributeBase::IsTypeSupported(Type); });		
 		SuccessArgIndex = Infos.Emplace(SuccessArgName, ERigVMPinDirection::Output, RigVMTypeUtils::TypeIndex::Bool);
 	}
 	
diff --git a/Engine/Plugins/Runtime/RigVM/Source/RigVM/Private/RigVMCore/RigVMDispatchFactory.cpp b/Engine/Plugins/Runtime/RigVM/Source/RigVM/Private/RigVMCore/RigVMDispatchFactory.cpp
--- a/Engine/Plugins/Runtime/RigVM/Source/RigVM/Private/RigVMCore/RigVMDispatchFactory.cpp
+++ b/Engine/Plugins/Runtime/RigVM/Source/RigVM/Private/RigVMCore/RigVMDispatchFactory.cpp
@@ -326,6 +326,7 @@ TArray<FRigVMTemplateArgumentInfo> FRigVMDispatchFactory::BuildArgumentListFromP
 				}
 			}
 		}
+		return true;
 	});
 
 	if (!bFoundArg)
diff --git a/Engine/Plugins/Runtime/RigVM/Source/RigVM/Private/RigVMCore/RigVMTemplate.cpp b/Engine/Plugins/Runtime/RigVM/Source/RigVM/Private/RigVMCore/RigVMTemplate.cpp
--- a/Engine/Plugins/Runtime/RigVM/Source/RigVM/Private/RigVMCore/RigVMTemplate.cpp
+++ b/Engine/Plugins/Runtime/RigVM/Source/RigVM/Private/RigVMCore/RigVMTemplate.cpp
@@ -109,10 +109,11 @@ FRigVMTemplateArgument::FRigVMTemplateArgument(const FName& InName, ERigVMPinDir
 	UpdateTypeToPermutations();
 }
 
-FRigVMTemplateArgument::FRigVMTemplateArgument(const FName& InName, ERigVMPinDirection InDirection, const TArray<ETypeCategory>& InTypeCategories)
+FRigVMTemplateArgument::FRigVMTemplateArgument(const FName& InName, ERigVMPinDirection InDirection, const TArray<ETypeCategory>& InTypeCategories, TFunction<bool(const TRigVMTypeIndex&)> InFilterType)
 	: Name(InName)
 	, Direction(InDirection)
 	, TypeCategories(InTypeCategories)
+	, FilterType(InFilterType)
 {
 	const int32 NumCategories = TypeCategories.Num();
 	if (NumCategories > 0)
@@ -148,6 +149,13 @@ FRigVMTemplateArgument::FRigVMTemplateArgument(const FName& InName, ERigVMPinDir
 		else
 		{
 			TArray<TRigVMTypeIndex> Indices = AllTypes.Array();
+			if (FilterType)
+			{
+				Indices = Indices.FilterByPredicate([this](const TRigVMTypeIndex& Type)
+				{
+					return FilterType(Type);
+				});
+			}
 			TypeIndices = MoveTemp(Indices);
 			EnsureValidExecuteType(FRigVMRegistry::Get());
 		}
@@ -172,7 +180,8 @@ void FRigVMTemplateArgument::UpdateTypeToPermutations()
 	int32 TypeIndex = 0;
 	ForEachType([&](const TRigVMTypeIndex Type)
 	{
-	   TypeToPermutations.FindOrAdd(Type).Add(TypeIndex++);
+		TypeToPermutations.FindOrAdd(Type).Add(TypeIndex++);
+		return true;
 	});
 }
 
@@ -306,6 +315,7 @@ const TArray<int32>& FRigVMTemplateArgument::GetPermutations(const TRigVMTypeInd
 			Permutations.Add(IndexInTypes);
 		}
 		IndexInTypes++;
+		return true;
 	});
 
 	if (!Permutations.IsEmpty())
@@ -333,7 +343,21 @@ void FRigVMTemplateArgument::GetAllTypes(TArray<TRigVMTypeIndex>& OutTypes) cons
 	OutTypes.Reset();
 	for (const ETypeCategory Category: TypeCategories)
 	{
-		OutTypes.Append(FRigVMRegistry::Get().GetTypesForCategory(Category));
+		if (FilterType == nullptr)
+		{
+			OutTypes.Append(FRigVMRegistry::Get().GetTypesForCategory(Category));
+		}
+		else
+		{
+			const TArray<TRigVMTypeIndex>& CategoryTypes = FRigVMRegistry::Get().GetTypesForCategory(Category);
+			for (const TRigVMTypeIndex& Type : CategoryTypes)
+			{
+				if (FilterType(Type))
+				{
+					OutTypes.Add(Type);
+				}
+			}
+		}
 	}
 }
 
@@ -345,6 +369,26 @@ TRigVMTypeIndex FRigVMTemplateArgument::GetTypeIndex(const int32 InIndex) const
 		return TypeIndices.IsValidIndex(InIndex) ? TypeIndices[InIndex] : TypeIndices[0];		
 	}
 
+	if (FilterType)
+	{
+		TRigVMTypeIndex ValidType = INDEX_NONE;
+		int32 ValidIndex = 0;
+		CategoryViews(TypeCategories).ForEachType([this, &ValidIndex, InIndex, &ValidType](const TRigVMTypeIndex& Type) -> bool
+		{
+			if (FilterType(Type))
+			{
+				if (ValidIndex == InIndex)
+				{
+					ValidType = Type;
+					return false;
+				}
+				ValidIndex++;
+			}
+			return true;
+		});
+		return ValidType;
+	}
+
 	return CategoryViews(TypeCategories).GetTypeIndex(InIndex);
 }
 
@@ -355,6 +399,30 @@ int32 FRigVMTemplateArgument::FindTypeIndex(const TRigVMTypeIndex InTypeIndex) c
 		return TypeIndices.IndexOfByKey(InTypeIndex);
 	}
 
+	if (FilterType)
+	{
+		bool bFound = false;
+		int32 ValidIndex = 0;
+		CategoryViews(TypeCategories).ForEachType([this, &ValidIndex, &bFound, InTypeIndex](const TRigVMTypeIndex& Type) -> bool
+		{
+			if (Type == InTypeIndex)
+			{
+				bFound = true;
+				return false;
+			}
+			if (FilterType(Type))
+			{
+				ValidIndex++;
+			}
+			return true;
+		});
+		if (bFound)
+		{
+			return ValidIndex;
+		}
+		return INDEX_NONE;
+	}
+
 	return CategoryViews(TypeCategories).FindIndex(InTypeIndex);
 }
 
@@ -364,6 +432,20 @@ int32 FRigVMTemplateArgument::GetNumTypes() const
 	{
 		return TypeIndices.Num();
 	}
+
+	if (FilterType)
+	{
+		int32 NumTypes = 0;
+		CategoryViews(TypeCategories).ForEachType([this, &NumTypes](const TRigVMTypeIndex& Type) -> bool
+		{
+			if (FilterType(Type))
+			{
+				NumTypes++;
+			}
+			return true;
+		});
+		return NumTypes;
+	}
 	
 	return Algo::Accumulate(TypeCategories, 0, [](int32 Sum, const ETypeCategory Category)
 	{
@@ -383,12 +465,24 @@ void FRigVMTemplateArgument::RemoveType(const int32 InIndex)
 	TypeIndices.RemoveAt(InIndex);
 }
 
-void FRigVMTemplateArgument::ForEachType(TFunction<void(const TRigVMTypeIndex InType)>&& InCallback) const
+void FRigVMTemplateArgument::ForEachType(TFunction<bool(const TRigVMTypeIndex InType)>&& InCallback) const
 {
 	if (!bUseCategories)
 	{
 		return Algo::ForEach(TypeIndices, InCallback);
 	}
+
+	if (FilterType)
+	{
+		CategoryViews(TypeCategories).ForEachType([this, InCallback](const TRigVMTypeIndex Type)
+		{
+			if (FilterType(Type))
+			{
+				return InCallback(Type);
+			}
+			return true;
+		});
+	}
 	
 	return CategoryViews(TypeCategories).ForEachType(MoveTemp(InCallback));
 }
@@ -405,6 +499,7 @@ TArray<TRigVMTypeIndex> FRigVMTemplateArgument::GetSupportedTypeIndices(const TA
 			{
 				SupportedTypes.AddUnique(TypeIndex);
 			}
+			return true;
 		});
 	}
 	else
@@ -438,6 +533,7 @@ TArray<FString> FRigVMTemplateArgument::GetSupportedTypeStrings(const TArray<int
 				const FString TypeString = Registry.GetType(TypeIndex).CPPType.ToString();
 				SupportedTypes.AddUnique(TypeString);
 			}
+			return true;
 		});
 	}
 	else
@@ -468,11 +564,17 @@ FRigVMTemplateArgument::CategoryViews::CategoryViews(const TArray<ETypeCategory>
 	}
 }
 
-void FRigVMTemplateArgument::CategoryViews::ForEachType(TFunction<void(const TRigVMTypeIndex InType)>&& InCallback) const
+void FRigVMTemplateArgument::CategoryViews::ForEachType(TFunction<bool(const TRigVMTypeIndex InType)>&& InCallback) const
 {
 	for (const TArrayView<const TRigVMTypeIndex>& TypeView: Types)
 	{
-		Algo::ForEach(TypeView, InCallback);
+		for (const TRigVMTypeIndex& TypeIndex : TypeView)
+		{
+			if (!InCallback(TypeIndex))
+			{
+				return;
+			}
+		}
 	}
 }
 
@@ -522,10 +624,11 @@ FRigVMTemplateArgumentInfo::FRigVMTemplateArgumentInfo(const FName InName, ERigV
 
 FRigVMTemplateArgumentInfo::FRigVMTemplateArgumentInfo(
 	const FName InName, ERigVMPinDirection InDirection,
-	const TArray<FRigVMTemplateArgument::ETypeCategory>& InTypeCategories)
+	const TArray<FRigVMTemplateArgument::ETypeCategory>& InTypeCategories,
+	TypeFilterCallback InTypeFilter)
 	: Name(InName)
 	, Direction(InDirection)
-	, FactoryCallback([InTypeCategories](const FName InName, ERigVMPinDirection InDirection){ return FRigVMTemplateArgument(InName, InDirection, InTypeCategories); })
+	, FactoryCallback([InTypeCategories, InTypeFilter](const FName InName, ERigVMPinDirection InDirection){ return FRigVMTemplateArgument(InName, InDirection, InTypeCategories, InTypeFilter); })
 {}
 
 FRigVMTemplateArgumentInfo::FRigVMTemplateArgumentInfo(const FName InName, ERigVMPinDirection InDirection)
@@ -1884,7 +1987,7 @@ bool FRigVMTemplate::UpdateArgumentTypes()
 					// Find if these types were already registered
 					if (ContainsPermutation(Types))
 					{
-						return;
+						return true;
 					}
 
 					uint32 TypeHash=0;
@@ -1896,12 +1999,12 @@ bool FRigVMTemplate::UpdateArgumentTypes()
 						if(TypeIndexPtr == nullptr)
 						{
 							bResult = false;
-							return;
+							return true;
 						}
 						if(*TypeIndexPtr == INDEX_NONE)
 						{
 							bResult = false;
-							return;
+							return true;
 						}
 
 						const TRigVMTypeIndex& TypeIndex = *TypeIndexPtr;
@@ -1929,6 +2032,7 @@ bool FRigVMTemplate::UpdateArgumentTypes()
 		{
 			bResult = false;
 		}
+		return true;
 	});
 
 	for(int32 ArgumentIndex = 0; ArgumentIndex < Arguments.Num(); ArgumentIndex++)
@@ -2014,6 +2118,7 @@ uint32 GetTypeHash(const FRigVMTemplateArgument& InArgument)
 	InArgument.ForEachType([&](const TRigVMTypeIndex TypeIndex)
 	{
 		Hash = HashCombine(Hash, Registry.GetHashForType(TypeIndex));
+		return true;
 	});
 	return Hash;
 }
diff --git a/Engine/Plugins/Runtime/RigVM/Source/RigVM/Public/RigVMCore/RigVMTemplate.h b/Engine/Plugins/Runtime/RigVM/Source/RigVM/Public/RigVMCore/RigVMTemplate.h
--- a/Engine/Plugins/Runtime/RigVM/Source/RigVM/Public/RigVMCore/RigVMTemplate.h
+++ b/Engine/Plugins/Runtime/RigVM/Source/RigVM/Public/RigVMCore/RigVMTemplate.h
@@ -196,7 +196,7 @@ struct RIGVM_API FRigVMTemplateArgument
 	FRigVMTemplateArgument(const FName& InName, ERigVMPinDirection InDirection);
 	FRigVMTemplateArgument(const FName& InName, ERigVMPinDirection InDirection, TRigVMTypeIndex InType);
 	FRigVMTemplateArgument(const FName& InName, ERigVMPinDirection InDirection, const TArray<TRigVMTypeIndex>& InTypeIndices);
-	FRigVMTemplateArgument(const FName& InName, ERigVMPinDirection InDirection, const TArray<ETypeCategory>& InTypeCategories);
+	FRigVMTemplateArgument(const FName& InName, ERigVMPinDirection InDirection, const TArray<ETypeCategory>& InTypeCategories, TFunction<bool(const TRigVMTypeIndex&)> InFilterType = nullptr);
 
 	// returns the name of the argument
 	const FName& GetName() const { return Name; }
@@ -214,7 +214,7 @@ struct RIGVM_API FRigVMTemplateArgument
 	int32 GetNumTypes() const;
 	void AddTypeIndex(const TRigVMTypeIndex InTypeIndex);
 	void RemoveType(const int32 InIndex);
-	void ForEachType(TFunction<void(const TRigVMTypeIndex InType)>&& InCallback) const;
+	void ForEachType(TFunction<bool(const TRigVMTypeIndex InType)>&& InCallback) const;
 	int32 FindTypeIndex(const TRigVMTypeIndex InTypeIndex) const;
 
 	template <typename Predicate>
@@ -224,6 +224,29 @@ struct RIGVM_API FRigVMTemplateArgument
 		{
 			return TypeIndices.IndexOfByPredicate(Pred);
 		}
+		if (FilterType)
+		{
+			bool bFound = false;
+			int32 ValidIndex = 0;
+			CategoryViews(TypeCategories).ForEachType([this, &ValidIndex, &bFound, Pred](const TRigVMTypeIndex Type)
+			{
+				if (FilterType(Type))
+				{
+					if (Pred(Type))
+					{
+						bFound = true;
+						return false;
+					}
+					ValidIndex++;
+				}
+				return true;
+			});
+			if (bFound)
+			{
+				return ValidIndex;
+			}
+			return INDEX_NONE;
+		}
 		return CategoryViews(TypeCategories).IndexOfByPredicate(Pred);
 	}
 
@@ -261,6 +284,7 @@ protected:
 
 	bool bUseCategories = false;
 	TArray<ETypeCategory> TypeCategories;
+	TFunction<bool(const TRigVMTypeIndex&)> FilterType;
 
 	FRigVMTemplateArgument(FProperty* InProperty);
 
@@ -283,7 +307,7 @@ private:
 		CategoryViews() = delete;
 		CategoryViews(const TArray<ETypeCategory>& InCategories);
 		
-		void ForEachType(TFunction<void(const TRigVMTypeIndex InType)>&& InCallback) const;
+		void ForEachType(TFunction<bool(const TRigVMTypeIndex InType)>&& InCallback) const;
 
 		TRigVMTypeIndex GetTypeIndex(int32 InIndex) const;
 		
@@ -317,6 +341,7 @@ private:
 struct RIGVM_API FRigVMTemplateArgumentInfo
 {
 	using ArgumentCallback = TFunction<FRigVMTemplateArgument(const FName /*InName*/, ERigVMPinDirection /*InDirection*/)>;
+	using TypeFilterCallback = TFunction<bool(const TRigVMTypeIndex& InNewType)>;
 	
 	FName Name = NAME_None;
 	ERigVMPinDirection Direction = ERigVMPinDirection::Invalid;
@@ -324,7 +349,7 @@ struct RIGVM_API FRigVMTemplateArgumentInfo
 	
 	FRigVMTemplateArgumentInfo(const FName InName, ERigVMPinDirection InDirection, TRigVMTypeIndex InTypeIndex);
 	FRigVMTemplateArgumentInfo(const FName InName, ERigVMPinDirection InDirection, const TArray<TRigVMTypeIndex>& InTypeIndices);
-	FRigVMTemplateArgumentInfo(const FName InName, ERigVMPinDirection InDirection, const TArray<FRigVMTemplateArgument::ETypeCategory>& InTypeCategories);
+	FRigVMTemplateArgumentInfo(const FName InName, ERigVMPinDirection InDirection, const TArray<FRigVMTemplateArgument::ETypeCategory>& InTypeCategories, TypeFilterCallback InTypeFilter = nullptr);
 	FRigVMTemplateArgumentInfo(const FName InName, ERigVMPinDirection InDirection);
 	FRigVMTemplateArgumentInfo(const FName InName, ERigVMPinDirection InDirection, ArgumentCallback&& InCallback);
 
diff --git a/Engine/Plugins/Runtime/RigVM/Source/RigVMDeveloper/Private/RigVMDeveloperTypeUtils.cpp b/Engine/Plugins/Runtime/RigVM/Source/RigVMDeveloper/Private/RigVMDeveloperTypeUtils.cpp
--- a/Engine/Plugins/Runtime/RigVM/Source/RigVMDeveloper/Private/RigVMDeveloperTypeUtils.cpp
+++ b/Engine/Plugins/Runtime/RigVM/Source/RigVMDeveloper/Private/RigVMDeveloperTypeUtils.cpp
@@ -617,6 +617,7 @@ const TArray<TRigVMTypeIndex>& RigVMTypeUtils::GetAvailableCasts(const TRigVMTyp
 					AvailableCasts.Add(ArgumentB->GetTypeIndex(Index));
 				}
 				Index++;
+				return true;
 			});
 
 			return AvailableCasts;
