Control Rig: Fix outdated host being manipulated after saving rig
#rnx

diff --git a/Engine/Plugins/Runtime/RigVM/Source/RigVMEditor/Private/Editor/RigVMEditor.cpp b/Engine/Plugins/Runtime/RigVM/Source/RigVMEditor/Private/Editor/RigVMEditor.cpp
--- a/Engine/Plugins/Runtime/RigVM/Source/RigVMEditor/Private/Editor/RigVMEditor.cpp
+++ b/Engine/Plugins/Runtime/RigVM/Source/RigVMEditor/Private/Editor/RigVMEditor.cpp
@@ -1129,12 +1129,16 @@ void FRigVMEditor::SaveAsset_Execute()
 {
 	LastDebuggedHost = GetCustomDebugObjectLabel(GetBlueprintObj()->GetObjectBeingDebugged());
 	FBlueprintEditor::SaveAsset_Execute();
+
+	UpdateRigVMHost();
 }
 
 void FRigVMEditor::SaveAssetAs_Execute()
 {
 	LastDebuggedHost = GetCustomDebugObjectLabel(GetBlueprintObj()->GetObjectBeingDebugged());
 	FBlueprintEditor::SaveAssetAs_Execute();
+
+	UpdateRigVMHost();
 }
 
 bool FRigVMEditor::IsEditable(UEdGraph* InGraph) const
