Control Rig: Fix crash when processing lazy expressions which haven't been visited during the setup memory stage, while compiling the rigvm
#rnx

diff --git a/Engine/Plugins/Runtime/RigVM/Source/RigVMDeveloper/Private/RigVMCompiler/RigVMCompiler.cpp b/Engine/Plugins/Runtime/RigVM/Source/RigVMDeveloper/Private/RigVMCompiler/RigVMCompiler.cpp
--- a/Engine/Plugins/Runtime/RigVM/Source/RigVMDeveloper/Private/RigVMCompiler/RigVMCompiler.cpp
+++ b/Engine/Plugins/Runtime/RigVM/Source/RigVMDeveloper/Private/RigVMCompiler/RigVMCompiler.cpp
@@ -1259,6 +1259,7 @@ bool URigVMCompiler::Compile(const FRigVMCompileSettings& InSettings, TArray<URi
 		WorkData.VM->GenerateDefaultMemoryType(MemoryType, Properties);
 	}
 
+	const TMap<const FRigVMExprAST*, bool> ExprSetupMemoryComplete = WorkData.ExprComplete;
 	WorkData.bSetupMemory = false;
 	WorkData.ExprComplete.Reset();
 
@@ -1312,7 +1313,10 @@ bool URigVMCompiler::Compile(const FRigVMCompileSettings& InSettings, TArray<URi
 			TGuardValue<TOptional<uint32>> HashGuard(WorkData.CurrentBlockHash, BlockInfo->Hash);
 			for(const FRigVMExprAST* Expression : BlockInfo->Expressions)
 			{
-				TraverseExpression(Expression, WorkData);
+				if (ExprSetupMemoryComplete.Contains(Expression))
+				{
+					TraverseExpression(Expression, WorkData);
+				}
 			}
 
 			BlockInfo->EndInstruction = ByteCode.GetNumInstructions() - 1;
