RigVM: Replace static use of FString / FName

#rnx

diff --git a/Engine/Plugins/Animation/ControlRig/Source/ControlRigDeveloper/Private/ControlRigBlueprint.cpp b/Engine/Plugins/Animation/ControlRig/Source/ControlRigDeveloper/Private/ControlRigBlueprint.cpp
--- a/Engine/Plugins/Animation/ControlRig/Source/ControlRigDeveloper/Private/ControlRigBlueprint.cpp
+++ b/Engine/Plugins/Animation/ControlRig/Source/ControlRigDeveloper/Private/ControlRigBlueprint.cpp
@@ -46,9 +46,6 @@
 #define LOCTEXT_NAMESPACE "ControlRigBlueprint"
 
 TArray<UControlRigBlueprint*> UControlRigBlueprint::sCurrentlyOpenedRigBlueprints;
-#if WITH_EDITOR
-const FName UControlRigBlueprint::ControlRigPanelNodeFactoryName(TEXT("FControlRigGraphPanelPinFactory"));
-#endif
 
 UControlRigBlueprint::UControlRigBlueprint(const FObjectInitializer& ObjectInitializer)
 	: URigVMBlueprint(ObjectInitializer)
@@ -1304,7 +1301,7 @@ void UControlRigBlueprint::GetPreloadDependencies(TArray<UObject*>& OutDeps)
 }
 
 #if WITH_EDITOR
-const FName& UControlRigBlueprint::GetPanelPinFactoryName() const
+const FLazyName& UControlRigBlueprint::GetPanelPinFactoryName() const
 {
 	return ControlRigPanelNodeFactoryName;
 }
diff --git a/Engine/Plugins/Animation/ControlRig/Source/ControlRigDeveloper/Private/Graph/ControlRigGraphSchema.cpp b/Engine/Plugins/Animation/ControlRig/Source/ControlRigDeveloper/Private/Graph/ControlRigGraphSchema.cpp
--- a/Engine/Plugins/Animation/ControlRig/Source/ControlRigDeveloper/Private/Graph/ControlRigGraphSchema.cpp
+++ b/Engine/Plugins/Animation/ControlRig/Source/ControlRigDeveloper/Private/Graph/ControlRigGraphSchema.cpp
@@ -15,8 +15,6 @@
 
 #define LOCTEXT_NAMESPACE "ControlRigGraphSchema"
 
-const FName UControlRigGraphSchema::GraphName_ControlRig(TEXT("Rig"));
-
 UControlRigGraphSchema::UControlRigGraphSchema()
 {
 }
diff --git a/Engine/Plugins/Animation/ControlRig/Source/ControlRigDeveloper/Public/ControlRigBlueprint.h b/Engine/Plugins/Animation/ControlRig/Source/ControlRigDeveloper/Public/ControlRigBlueprint.h
--- a/Engine/Plugins/Animation/ControlRig/Source/ControlRigDeveloper/Public/ControlRigBlueprint.h
+++ b/Engine/Plugins/Animation/ControlRig/Source/ControlRigDeveloper/Public/ControlRigBlueprint.h
@@ -59,8 +59,8 @@ public:
 	virtual UClass* GetRigVMEditorSettingsClass() const override;
 	virtual void GetPreloadDependencies(TArray<UObject*>& OutDeps) override;
 #if WITH_EDITOR
-	virtual const FName& GetPanelPinFactoryName() const override;
-	static const FName ControlRigPanelNodeFactoryName;
+	virtual const FLazyName& GetPanelPinFactoryName() const override;
+	static inline const FLazyName ControlRigPanelNodeFactoryName = FLazyName(TEXT("FControlRigGraphPanelPinFactory"));
 	virtual IRigVMEditorModule* GetEditorModule() const override;
 #endif
 
diff --git a/Engine/Plugins/Animation/ControlRig/Source/ControlRigDeveloper/Public/Graph/ControlRigGraphSchema.h b/Engine/Plugins/Animation/ControlRig/Source/ControlRigDeveloper/Public/Graph/ControlRigGraphSchema.h
--- a/Engine/Plugins/Animation/ControlRig/Source/ControlRigDeveloper/Public/Graph/ControlRigGraphSchema.h
+++ b/Engine/Plugins/Animation/ControlRig/Source/ControlRigDeveloper/Public/Graph/ControlRigGraphSchema.h
@@ -23,7 +23,7 @@ class CONTROLRIGDEVELOPER_API UControlRigGraphSchema : public URigVMEdGraphSchem
 
 public:
 	/** Name constants */
-	static const FName GraphName_ControlRig;
+	static inline const FLazyName GraphName_ControlRig = FLazyName(TEXT("Rig"));
 
 public:
 	UControlRigGraphSchema();
@@ -32,6 +32,6 @@ public:
 	virtual FLinearColor GetPinTypeColor(const FEdGraphPinType& PinType) const override;
 
 	// URigVMEdGraphSchema interface
-	virtual const FName& GetRootGraphName() const override { return GraphName_ControlRig; }
+	virtual const FLazyName& GetRootGraphName() const override { return GraphName_ControlRig; }
 	virtual bool IsRigVMDefaultEvent(const FName& InEventName) const override;
 };
diff --git a/Engine/Plugins/Runtime/RigVM/Source/RigVM/Private/RigVMCore/RigVMDispatchFactory.cpp b/Engine/Plugins/Runtime/RigVM/Source/RigVM/Private/RigVMCore/RigVMDispatchFactory.cpp
--- a/Engine/Plugins/Runtime/RigVM/Source/RigVM/Private/RigVMCore/RigVMDispatchFactory.cpp
+++ b/Engine/Plugins/Runtime/RigVM/Source/RigVM/Private/RigVMCore/RigVMDispatchFactory.cpp
@@ -7,14 +7,12 @@
 
 #include UE_INLINE_GENERATED_CPP_BY_NAME(RigVMDispatchFactory)
 
-const FString FRigVMDispatchFactory::DispatchPrefix = TEXT("DISPATCH_");
-const FString FRigVMDispatchFactory::TrueString = TEXT("True");
 FCriticalSection FRigVMDispatchFactory::GetTemplateMutex;
 
 FName FRigVMDispatchFactory::GetFactoryName() const
 {
 	static constexpr TCHAR Format[] = TEXT("%s%s");
-	return *FString::Printf(Format, *DispatchPrefix, *GetScriptStruct()->GetName());
+	return *FString::Printf(Format, DispatchPrefix, *GetScriptStruct()->GetName());
 }
 
 #if WITH_EDITOR
diff --git a/Engine/Plugins/Runtime/RigVM/Source/RigVM/Private/RigVMCore/RigVMExecuteContext.cpp b/Engine/Plugins/Runtime/RigVM/Source/RigVM/Private/RigVMCore/RigVMExecuteContext.cpp
--- a/Engine/Plugins/Runtime/RigVM/Source/RigVM/Private/RigVMCore/RigVMExecuteContext.cpp
+++ b/Engine/Plugins/Runtime/RigVM/Source/RigVM/Private/RigVMCore/RigVMExecuteContext.cpp
@@ -36,7 +36,7 @@ void FRigVMExecuteContext::SetWorld(const UWorld* InWorld)
 
 bool FRigVMExecuteContext::SerializeFromMismatchedTag(const FPropertyTag& Tag, FStructuredArchive::FSlot Slot)
 {
-	static const FName ControlRigExecuteContextName("ControlRigExecuteContext");
+	static const FLazyName ControlRigExecuteContextName("ControlRigExecuteContext");
 	if (Tag.GetType().IsStruct(ControlRigExecuteContextName))
 	{
 		static const FString CRExecuteContextPath = TEXT("/Script/ControlRig.ControlRigExecuteContext");
diff --git a/Engine/Plugins/Runtime/RigVM/Source/RigVM/Private/RigVMCore/RigVMGraphFunctionDefinition.cpp b/Engine/Plugins/Runtime/RigVM/Source/RigVM/Private/RigVMCore/RigVMGraphFunctionDefinition.cpp
--- a/Engine/Plugins/Runtime/RigVM/Source/RigVM/Private/RigVMCore/RigVMGraphFunctionDefinition.cpp
+++ b/Engine/Plugins/Runtime/RigVM/Source/RigVM/Private/RigVMCore/RigVMGraphFunctionDefinition.cpp
@@ -6,9 +6,6 @@
 
 #include UE_INLINE_GENERATED_CPP_BY_NAME(RigVMGraphFunctionDefinition)
 
-const FString FRigVMGraphFunctionData::EntryString = TEXT("Entry");
-const FString FRigVMGraphFunctionData::ReturnString = TEXT("Return");
-
 FRigVMExternalVariable FRigVMGraphFunctionArgument::GetExternalVariable() const
 {
 	FRigVMExternalVariable Variable;
diff --git a/Engine/Plugins/Runtime/RigVM/Source/RigVM/Private/RigVMCore/RigVMMemoryStorage.cpp b/Engine/Plugins/Runtime/RigVM/Source/RigVM/Private/RigVMCore/RigVMMemoryStorage.cpp
--- a/Engine/Plugins/Runtime/RigVM/Source/RigVM/Private/RigVMCore/RigVMMemoryStorage.cpp
+++ b/Engine/Plugins/Runtime/RigVM/Source/RigVM/Private/RigVMCore/RigVMMemoryStorage.cpp
@@ -12,12 +12,6 @@
 
 ////////////////////////////////////////////////////////////////////////////////////////////////////////////////////////
 
-const FString FRigVMPropertyDescription::ArrayPrefix = TEXT("TArray<");
-const FString FRigVMPropertyDescription::MapPrefix = TEXT("TMap<");
-const FString FRigVMPropertyDescription::ContainerSuffix = TEXT(">");
-
-////////////////////////////////////////////////////////////////////////////////////////////////////////////////////////
-
 // Generator class should be parented to the asset object, instead of the package
 // because the engine no longer supports multiple 'assets' per package
 static UObject* GetGeneratorClassOuter(UPackage* InPackage)
@@ -851,8 +845,8 @@ FProperty* URigVMMemoryStorageGeneratorClass::AddProperty(URigVMMemoryStorageGen
 
 		// store some additional meta data,
 		// mainly for inspecting things in the details panel
-		static const FName NAME_DisplayName(TEXT("DisplayName"));
-		static const FName NAME_ToolTipName(TEXT("ToolTip"));
+		static const FLazyName NAME_DisplayName(TEXT("DisplayName"));
+		static const FLazyName NAME_ToolTipName(TEXT("ToolTip"));
 
 		FString DisplayName = Result->GetName();
 		while(DisplayName.ReplaceInline(TEXT("__"), TEXT("_")) > 0)
diff --git a/Engine/Plugins/Runtime/RigVM/Source/RigVM/Private/RigVMCore/RigVMMemoryStorageStruct.cpp b/Engine/Plugins/Runtime/RigVM/Source/RigVM/Private/RigVMCore/RigVMMemoryStorageStruct.cpp
--- a/Engine/Plugins/Runtime/RigVM/Source/RigVM/Private/RigVMCore/RigVMMemoryStorageStruct.cpp
+++ b/Engine/Plugins/Runtime/RigVM/Source/RigVM/Private/RigVMCore/RigVMMemoryStorageStruct.cpp
@@ -615,15 +615,11 @@ void FRigVMMemoryStorageStruct::SetDefaultValues(const TArray<FRigVMPropertyDesc
 
 	FName VMType = *VMTypeString;
 
-	static const FName IntTypeName(TEXT("int")); // type used by some engine tests
-	static const FName Int64TypeName(TEXT("Int64"));
-	static const FName UInt64TypeName(TEXT("UInt64"));
-
 	if (VMType == RigVMTypeUtils::BoolTypeName)
 	{
 		OutBagPropertyType = EPropertyBagPropertyType::Bool;
 	}
-	else if (VMType == RigVMTypeUtils::Int32TypeName || VMType == IntTypeName)
+	else if (VMType == RigVMTypeUtils::Int32TypeName || VMType == RigVMTypeUtils::IntTypeName)
 	{
 		OutBagPropertyType = EPropertyBagPropertyType::Int32;
 	}
@@ -631,11 +627,11 @@ void FRigVMMemoryStorageStruct::SetDefaultValues(const TArray<FRigVMPropertyDesc
 	{
 		OutBagPropertyType = EPropertyBagPropertyType::UInt32;
 	}
-	else if (VMType == Int64TypeName)
+	else if (VMType == RigVMTypeUtils::Int64TypeName)
 	{
 		OutBagPropertyType = EPropertyBagPropertyType::Int64;
 	}
-	else if (VMType == UInt64TypeName)
+	else if (VMType == RigVMTypeUtils::UInt64TypeName)
 	{
 		OutBagPropertyType = EPropertyBagPropertyType::UInt64;
 	}
diff --git a/Engine/Plugins/Runtime/RigVM/Source/RigVM/Private/RigVMCore/RigVMRegistry.cpp b/Engine/Plugins/Runtime/RigVM/Source/RigVM/Private/RigVMCore/RigVMRegistry.cpp
--- a/Engine/Plugins/Runtime/RigVM/Source/RigVM/Private/RigVMCore/RigVMRegistry.cpp
+++ b/Engine/Plugins/Runtime/RigVM/Source/RigVM/Private/RigVMCore/RigVMRegistry.cpp
@@ -16,8 +16,6 @@
 #include "RigVMFunctions/RigVMDispatch_Core.h"
 #include "Interfaces/IPluginManager.h"
 
-const FName FRigVMRegistry::TemplateNameMetaName = TEXT("TemplateName");
-
 FCriticalSection FRigVMRegistry::FindOrAddTypeMutex;
 
 FCriticalSection FRigVMRegistry::FunctionRegistryMutex;
@@ -1488,7 +1486,7 @@ const FRigVMFunction* FRigVMRegistry::FindFunction_NoLock(const TCHAR* InName, c
 		const bool bIsDispatchFactory = StructOrFactoryName.StartsWith(FRigVMDispatchFactory::DispatchPrefix, ESearchCase::CaseSensitive);
 		if(bIsDispatchFactory)
 		{
-			StructOrFactoryName = StructOrFactoryName.Mid(FRigVMDispatchFactory::DispatchPrefix.Len());
+			StructOrFactoryName = StructOrFactoryName.Mid(FCString::Strlen(FRigVMDispatchFactory::DispatchPrefix));
 		}
 		else if(StructOrFactoryName.StartsWith(StructPrefix, ESearchCase::CaseSensitive))
 		{
@@ -1607,7 +1605,7 @@ const FRigVMTemplate* FRigVMRegistry::FindTemplate_NoLock(const FName& InNotatio
 	{
 		const FString OriginalDispatchFactoryName = OriginalNotation
 			.Left(OriginalNotation.Find(TEXT("(")))
-			.RightChop(FRigVMDispatchFactory::DispatchPrefix.Len());
+			.RightChop(FCString::Strlen(FRigVMDispatchFactory::DispatchPrefix));
 
 		const FCoreRedirectObjectName OldObjectName(OriginalDispatchFactoryName);
 		TArray<const FCoreRedirect*> Redirects;
@@ -1830,7 +1828,7 @@ FRigVMDispatchFactory* FRigVMRegistry::FindDispatchFactory_NoLock(const FName& I
 	// if the factory has never been registered - we should try to look it up	
 	if(FactoryName.StartsWith(FRigVMDispatchFactory::DispatchPrefix))
 	{
-		const FString ScriptStructName = FactoryName.Mid(FRigVMDispatchFactory::DispatchPrefix.Len());
+		const FString ScriptStructName = FactoryName.Mid(FCString::Strlen(FRigVMDispatchFactory::DispatchPrefix));
 		if(UScriptStruct* FactoryStruct = FindFirstObject<UScriptStruct>(*ScriptStructName, EFindFirstObjectOptions::NativeFirst | EFindFirstObjectOptions::EnsureIfAmbiguous))
 		{
 			FRigVMRegistry* MutableThis = (FRigVMRegistry*)this;
diff --git a/Engine/Plugins/Runtime/RigVM/Source/RigVM/Private/RigVMCore/RigVMStruct.cpp b/Engine/Plugins/Runtime/RigVM/Source/RigVM/Private/RigVMCore/RigVMStruct.cpp
--- a/Engine/Plugins/Runtime/RigVM/Source/RigVM/Private/RigVMCore/RigVMStruct.cpp
+++ b/Engine/Plugins/Runtime/RigVM/Source/RigVM/Private/RigVMCore/RigVMStruct.cpp
@@ -275,47 +275,6 @@ void FRigVMStructUpgradeInfo::SetDefaultValues(const FRigVMStruct* InNewStructMe
 
 ////////////////////////////////////////////////////////////////////////////////////////////////////////////////////////
 
-const FName FRigVMStruct::DeprecatedMetaName("Deprecated");
-const FName FRigVMStruct::InputMetaName("Input");
-const FName FRigVMStruct::OutputMetaName("Output");
-const FName FRigVMStruct::IOMetaName("IO");
-const FName FRigVMStruct::HiddenMetaName("Hidden");
-const FName FRigVMStruct::VisibleMetaName("Visible");
-const FName FRigVMStruct::DetailsOnlyMetaName("DetailsOnly");
-const FName FRigVMStruct::AbstractMetaName("Abstract");
-const FName FRigVMStruct::CategoryMetaName("Category");
-const FName FRigVMStruct::DisplayNameMetaName("DisplayName");
-const FName FRigVMStruct::MenuDescSuffixMetaName("MenuDescSuffix");
-const FName FRigVMStruct::ShowVariableNameInTitleMetaName("ShowVariableNameInTitle");
-const FName FRigVMStruct::CustomWidgetMetaName("CustomWidget");
-const FName FRigVMStruct::ConstantMetaName("Constant");
-const FName FRigVMStruct::TitleColorMetaName("TitleColor");
-const FName FRigVMStruct::NodeColorMetaName("NodeColor");
-const FName FRigVMStruct::IconMetaName("Icon");
-const FName FRigVMStruct::KeywordsMetaName("Keywords");
-const FName FRigVMStruct::TemplateNameMetaName = FRigVMRegistry::TemplateNameMetaName;
-const FName FRigVMStruct::FixedSizeArrayMetaName("FixedSizeArray");
-const FName FRigVMStruct::ShowOnlySubPinsMetaName("ShowOnlySubPins");
-const FName FRigVMStruct::HideSubPinsMetaName("HideSubPins");
-const FName FRigVMStruct::ArraySizeMetaName("ArraySize");
-const FName FRigVMStruct::AggregateMetaName("Aggregate");
-const FName FRigVMStruct::ExpandPinByDefaultMetaName("ExpandByDefault");
-const FName FRigVMStruct::DefaultArraySizeMetaName("DefaultArraySize");
-const FName FRigVMStruct::VaryingMetaName("Varying");
-const FName FRigVMStruct::SingletonMetaName("Singleton");
-const FName FRigVMStruct::SliceContextMetaName("SliceContext");
-const FName FRigVMStruct::ExecuteName = TEXT("Execute");
-const FName FRigVMStruct::ExecuteContextName = TEXT("ExecuteContext");
-const FName FRigVMStruct::ForLoopCountPinName("Count");
-const FName FRigVMStruct::ForLoopContinuePinName("Continue");
-const FName FRigVMStruct::ForLoopCompletedPinName("Completed");
-const FName FRigVMStruct::ForLoopIndexPinName("Index");
-const FName FRigVMStruct::ComputeLazilyMetaName("Lazy");
-const FName FRigVMStruct::ControlFlowBlockToRunName("BlockToRun");
-const FName FRigVMStruct::ControlFlowCompletedName("Completed");
-const FName FRigVMStruct::ControlFlowCountName("Count");
-const FName FRigVMStruct::ControlFlowIndexName("Index");
-
 float FRigVMStruct::GetRatioFromIndex(int32 InIndex, int32 InCount)
 {
 	if (InCount <= 1)
diff --git a/Engine/Plugins/Runtime/RigVM/Source/RigVM/Private/RigVMCore/RigVMTemplate.cpp b/Engine/Plugins/Runtime/RigVM/Source/RigVM/Private/RigVMCore/RigVMTemplate.cpp
--- a/Engine/Plugins/Runtime/RigVM/Source/RigVM/Private/RigVMCore/RigVMTemplate.cpp
+++ b/Engine/Plugins/Runtime/RigVM/Source/RigVM/Private/RigVMCore/RigVMTemplate.cpp
@@ -990,7 +990,7 @@ FLinearColor FRigVMTemplate::GetColor(const TArray<int32>& InPermutationIndices)
 
 	auto VisitPermutation = [&bFirstColorFound, &ResolvedColor, this](int32 InPermutationIndex) -> bool
 	{
-		static const FName NodeColorName = TEXT("NodeColor");
+		static const FLazyName NodeColorName = TEXT("NodeColor");
 		FString NodeColorMetadata;
 
 		// if we can't find one permutation we are not going to find any, so it's ok to return false here
diff --git a/Engine/Plugins/Runtime/RigVM/Source/RigVM/Private/RigVMFunctions/RigVMDispatch_Array.cpp b/Engine/Plugins/Runtime/RigVM/Source/RigVM/Private/RigVMFunctions/RigVMDispatch_Array.cpp
--- a/Engine/Plugins/Runtime/RigVM/Source/RigVM/Private/RigVMFunctions/RigVMDispatch_Array.cpp
+++ b/Engine/Plugins/Runtime/RigVM/Source/RigVM/Private/RigVMFunctions/RigVMDispatch_Array.cpp
@@ -9,20 +9,6 @@
 #include UE_INLINE_GENERATED_CPP_BY_NAME(RigVMDispatch_Array)
 #define LOCTEXT_NAMESPACE "RigVMDispatch_Array"
 
-const FName& FRigVMDispatch_ArrayBase::ExecuteName = FRigVMStruct::ExecuteContextName;
-const FName FRigVMDispatch_ArrayBase::ArrayName = TEXT("Array");
-const FName FRigVMDispatch_ArrayBase::ValuesName = TEXT("Values");
-const FName FRigVMDispatch_ArrayBase::NumName = TEXT("Num");
-const FName FRigVMDispatch_ArrayBase::IndexName = TEXT("Index");
-const FName FRigVMDispatch_ArrayBase::ElementName = TEXT("Element");
-const FName FRigVMDispatch_ArrayBase::SuccessName = TEXT("Success");
-const FName FRigVMDispatch_ArrayBase::OtherName = TEXT("Other");
-const FName FRigVMDispatch_ArrayBase::CloneName = TEXT("Clone");
-const FName FRigVMDispatch_ArrayBase::CountName = TEXT("Count");
-const FName FRigVMDispatch_ArrayBase::RatioName = TEXT("Ratio");
-const FName FRigVMDispatch_ArrayBase::ResultName = TEXT("Result");
-const FName& FRigVMDispatch_ArrayBase::CompletedName = FRigVMStruct::ControlFlowCompletedName;
-
 ////////////////////////////////////////////////////////////////////////////////////////////////////////////////////////
 ///
 #if WITH_EDITOR
@@ -199,7 +185,7 @@ const TArray<FRigVMExecuteArgument>& FRigVMDispatch_ArrayBaseMutable::GetExecute
 	const FRigVMDispatchContext& InContext) const
 {
 	static const TArray<FRigVMExecuteArgument> Arguments = {
-			{ExecuteName, ERigVMPinDirection::IO}
+			{FRigVMStruct::ExecuteContextName, ERigVMPinDirection::IO}
 	};
 	return Arguments;
 }
@@ -365,7 +351,7 @@ FText FRigVMDispatch_ArrayReset::GetArgumentTooltip(const FName& InArgumentName,
 
 FName FRigVMDispatch_ArrayGetNum::GetArgumentNameForOperandIndex(int32 InOperandIndex, int32 InTotalOperands) const
 {
-	static const FName ArgumentNames[] = {
+	static const FLazyName ArgumentNames[] = {
 		ArrayName,
 		NumName
 	};
@@ -427,7 +413,7 @@ FText FRigVMDispatch_ArrayGetNum::GetArgumentTooltip(const FName& InArgumentName
 
 FName FRigVMDispatch_ArraySetNum::GetArgumentNameForOperandIndex(int32 InOperandIndex, int32 InTotalOperands) const
 {
-	static const FName ArgumentNames[] = {
+	static const FLazyName ArgumentNames[] = {
 		ArrayName,
 		NumName
 	};
@@ -492,7 +478,7 @@ FText FRigVMDispatch_ArraySetNum::GetArgumentTooltip(const FName& InArgumentName
 
 FName FRigVMDispatch_ArrayGetAtIndex::GetArgumentNameForOperandIndex(int32 InOperandIndex, int32 InTotalOperands) const
 {
-	static const FName ArgumentNames[] = {
+	static const FLazyName ArgumentNames[] = {
 		ArrayName,
 		IndexName,
 		ElementName
@@ -577,7 +563,7 @@ FText FRigVMDispatch_ArrayGetAtIndex::GetArgumentTooltip(const FName& InArgument
 
 FName FRigVMDispatch_ArraySetAtIndex::GetArgumentNameForOperandIndex(int32 InOperandIndex, int32 InTotalOperands) const
 {
-	static const FName ArgumentNames[] = {
+	static const FLazyName ArgumentNames[] = {
 		ArrayName,
 		IndexName,
 		ElementName
@@ -662,7 +648,7 @@ FText FRigVMDispatch_ArraySetAtIndex::GetArgumentTooltip(const FName& InArgument
 
 FName FRigVMDispatch_ArrayAdd::GetArgumentNameForOperandIndex(int32 InOperandIndex, int32 InTotalOperands) const
 {
-	static const FName ArgumentNames[] = {
+	static const FLazyName ArgumentNames[] = {
 		ArrayName,
 		ElementName,
 		IndexName
@@ -783,7 +769,7 @@ FText FRigVMDispatch_ArrayInsert::GetArgumentTooltip(const FName& InArgumentName
 
 FName FRigVMDispatch_ArrayRemove::GetArgumentNameForOperandIndex(int32 InOperandIndex, int32 InTotalOperands) const
 {
-	static const FName ArgumentNames[] = {
+	static const FLazyName ArgumentNames[] = {
 		ArrayName,
 		IndexName
 	};
@@ -880,7 +866,7 @@ FText FRigVMDispatch_ArrayReverse::GetArgumentTooltip(const FName& InArgumentNam
 
 FName FRigVMDispatch_ArrayFind::GetArgumentNameForOperandIndex(int32 InOperandIndex, int32 InTotalOperands) const
 {
-	static const FName ArgumentNames[] = {
+	static const FLazyName ArgumentNames[] = {
 		ArrayName,
 		ElementName,
 		IndexName,
@@ -994,7 +980,7 @@ FText FRigVMDispatch_ArrayFind::GetArgumentTooltip(const FName& InArgumentName,
 
 FName FRigVMDispatch_ArrayAppend::GetArgumentNameForOperandIndex(int32 InOperandIndex, int32 InTotalOperands) const
 {
-	static const FName ArgumentNames[] = {
+	static const FLazyName ArgumentNames[] = {
 		ArrayName,
 		OtherName
 	};
@@ -1077,7 +1063,7 @@ FText FRigVMDispatch_ArrayAppend::GetArgumentTooltip(const FName& InArgumentName
 
 FName FRigVMDispatch_ArrayClone::GetArgumentNameForOperandIndex(int32 InOperandIndex, int32 InTotalOperands) const
 {
-	static const FName ArgumentNames[] = {
+	static const FLazyName ArgumentNames[] = {
 		ArrayName,
 		CloneName
 	};
@@ -1262,7 +1248,7 @@ FText FRigVMDispatch_ArrayUnion::GetArgumentTooltip(const FName& InArgumentName,
 
 FName FRigVMDispatch_ArrayDifference::GetArgumentNameForOperandIndex(int32 InOperandIndex, int32 InTotalOperands) const
 {
-	static const FName ArgumentNames[] = {
+	static const FLazyName ArgumentNames[] = {
 		ArrayName,
 		OtherName,
 		ResultName
@@ -1430,7 +1416,7 @@ FText FRigVMDispatch_ArrayIntersection::GetArgumentTooltip(const FName& InArgume
 
 FName FRigVMDispatch_ArrayIterator::GetArgumentNameForOperandIndex(int32 InOperandIndex, int32 InTotalOperands) const
 {
-	static const FName ArgumentNames[] = {
+	static const FLazyName ArgumentNames[] = {
 		ArrayName,
 		ElementName,
 		IndexName,
@@ -1464,7 +1450,7 @@ const TArray<FRigVMExecuteArgument>& FRigVMDispatch_ArrayIterator::GetExecuteArg
 	if(ExecuteArguments.IsEmpty())
 	{
 		ExecuteArguments = FRigVMDispatch_ArrayBaseMutable::GetExecuteArguments_Impl(InContext);
-		ExecuteArguments.Emplace(CompletedName, ERigVMPinDirection::Output);
+		ExecuteArguments.Emplace(FRigVMStruct::ControlFlowCompletedName, ERigVMPinDirection::Output);
 	}
 	return ExecuteArguments;
 }
@@ -1472,13 +1458,13 @@ const TArray<FRigVMExecuteArgument>& FRigVMDispatch_ArrayIterator::GetExecuteArg
 const TArray<FName>& FRigVMDispatch_ArrayIterator::GetControlFlowBlocks_Impl(
 	const FRigVMDispatchContext& InContext) const
 {
-	static const TArray<FName> Blocks = {ExecuteName, CompletedName};
+	static const TArray<FName> Blocks = {FRigVMStruct::ExecuteContextName, FRigVMStruct::ControlFlowCompletedName};
 	return Blocks;
 }
 
 const bool FRigVMDispatch_ArrayIterator::IsControlFlowBlockSliced(const FName& InBlockName) const
 {
-	return InBlockName == ExecuteName;
+	return InBlockName == FRigVMStruct::ExecuteContextName;
 }
 
 FRigVMTemplateTypeMap FRigVMDispatch_ArrayIterator::OnNewArgumentType(const FName& InArgumentName,
@@ -1518,7 +1504,7 @@ void FRigVMDispatch_ArrayIterator::Execute(FRigVMExtendedExecuteContext& InConte
 
 	if(Block.IsNone())
 	{
-		Block = ExecuteName;
+		Block = FRigVMStruct::ExecuteContextName;
 		Count = ArrayHelper.Num();
 		Index = -1;;
 		Ratio = 0.f;
@@ -1532,7 +1518,7 @@ void FRigVMDispatch_ArrayIterator::Execute(FRigVMExtendedExecuteContext& InConte
 
 	if(Index >= Count)
 	{
-		Block = CompletedName;
+		Block = FRigVMStruct::ControlFlowCompletedName;
 		return;
 	}
 	
@@ -1570,7 +1556,7 @@ FText FRigVMDispatch_ArrayIterator::GetArgumentTooltip(const FName& InArgumentNa
 	{
 		return LOCTEXT("ArrayIterator_RatioArgumentToolTip", "A float ratio from 0.0 (first element) to 1.0 (last element).");
 	}
-	if(InArgumentName == CompletedName)
+	if(InArgumentName == FRigVMStruct::ControlFlowCompletedName)
 	{
 		return LOCTEXT("ArrayIterator_CompletedArgumentToolTip", "The execute block to run once the loop has completed.");
 	}
diff --git a/Engine/Plugins/Runtime/RigVM/Source/RigVM/Private/RigVMFunctions/RigVMDispatch_CastEnum.cpp b/Engine/Plugins/Runtime/RigVM/Source/RigVM/Private/RigVMFunctions/RigVMDispatch_CastEnum.cpp
--- a/Engine/Plugins/Runtime/RigVM/Source/RigVM/Private/RigVMFunctions/RigVMDispatch_CastEnum.cpp
+++ b/Engine/Plugins/Runtime/RigVM/Source/RigVM/Private/RigVMFunctions/RigVMDispatch_CastEnum.cpp
@@ -6,9 +6,6 @@
 
 #define LOCTEXT_NAMESPACE "RigVMDispatch_CastEnum"
 
-const FName FRigVMDispatch_CastEnumBase::ValueName = TEXT("Value");
-const FName FRigVMDispatch_CastEnumBase::ResultName = TEXT("Result");
-
 const TArray<FRigVMTemplateArgumentInfo>& FRigVMDispatch_CastEnumToInt::GetArgumentInfos() const
 {
 	static TArray<FRigVMTemplateArgumentInfo> OutInfos;
diff --git a/Engine/Plugins/Runtime/RigVM/Source/RigVM/Private/RigVMFunctions/RigVMDispatch_CastObject.cpp b/Engine/Plugins/Runtime/RigVM/Source/RigVM/Private/RigVMFunctions/RigVMDispatch_CastObject.cpp
--- a/Engine/Plugins/Runtime/RigVM/Source/RigVM/Private/RigVMFunctions/RigVMDispatch_CastObject.cpp
+++ b/Engine/Plugins/Runtime/RigVM/Source/RigVM/Private/RigVMFunctions/RigVMDispatch_CastObject.cpp
@@ -5,9 +5,6 @@
 
 #define LOCTEXT_NAMESPACE "RigVMDispatch_CastObject"
 
-const FName FRigVMDispatch_CastObject::ValueName = TEXT("Value");
-const FName FRigVMDispatch_CastObject::ResultName = TEXT("Result");
-
 const TArray<FRigVMTemplateArgumentInfo>& FRigVMDispatch_CastObject::GetArgumentInfos() const
 {
 	static TArray<FRigVMTemplateArgumentInfo> OutInfos;
diff --git a/Engine/Plugins/Runtime/RigVM/Source/RigVM/Private/RigVMFunctions/RigVMDispatch_Constant.cpp b/Engine/Plugins/Runtime/RigVM/Source/RigVM/Private/RigVMFunctions/RigVMDispatch_Constant.cpp
--- a/Engine/Plugins/Runtime/RigVM/Source/RigVM/Private/RigVMFunctions/RigVMDispatch_Constant.cpp
+++ b/Engine/Plugins/Runtime/RigVM/Source/RigVM/Private/RigVMFunctions/RigVMDispatch_Constant.cpp
@@ -8,8 +8,6 @@
 #include UE_INLINE_GENERATED_CPP_BY_NAME(RigVMDispatch_Constant)
 #define LOCTEXT_NAMESPACE "RigVMDispatch_Constant"
 
-const FName FRigVMDispatch_Constant::ValueName = TEXT("Value");
-
 const TArray<FRigVMTemplateArgumentInfo>& FRigVMDispatch_Constant::GetArgumentInfos() const
 {
 	static TArray<FRigVMTemplateArgumentInfo> Infos;
diff --git a/Engine/Plugins/Runtime/RigVM/Source/RigVM/Private/RigVMFunctions/RigVMDispatch_Core.cpp b/Engine/Plugins/Runtime/RigVM/Source/RigVM/Private/RigVMFunctions/RigVMDispatch_Core.cpp
--- a/Engine/Plugins/Runtime/RigVM/Source/RigVM/Private/RigVMFunctions/RigVMDispatch_Core.cpp
+++ b/Engine/Plugins/Runtime/RigVM/Source/RigVM/Private/RigVMFunctions/RigVMDispatch_Core.cpp
@@ -6,15 +6,9 @@
 
 #include UE_INLINE_GENERATED_CPP_BY_NAME(RigVMDispatch_Core)
 
-
-const FName FRigVMDispatch_CoreEquals::AName = TEXT("A");
-const FName FRigVMDispatch_CoreEquals::BName = TEXT("B");
-const FName FRigVMDispatch_CoreEquals::ResultName = TEXT("Result");
-
-
 FName FRigVMDispatch_CoreEquals::GetArgumentNameForOperandIndex(int32 InOperandIndex, int32 InTotalOperands) const
 {
-	static const FName ArgumentNames[] = {
+	static const FLazyName ArgumentNames[] = {
 		AName,
 		BName,
 		ResultName
diff --git a/Engine/Plugins/Runtime/RigVM/Source/RigVM/Private/RigVMFunctions/RigVMDispatch_If.cpp b/Engine/Plugins/Runtime/RigVM/Source/RigVM/Private/RigVMFunctions/RigVMDispatch_If.cpp
--- a/Engine/Plugins/Runtime/RigVM/Source/RigVM/Private/RigVMFunctions/RigVMDispatch_If.cpp
+++ b/Engine/Plugins/Runtime/RigVM/Source/RigVM/Private/RigVMFunctions/RigVMDispatch_If.cpp
@@ -6,14 +6,9 @@
 
 #include UE_INLINE_GENERATED_CPP_BY_NAME(RigVMDispatch_If)
 
-const FName FRigVMDispatch_If::ConditionName = TEXT("Condition");
-const FName FRigVMDispatch_If::TrueName = TEXT("True");
-const FName FRigVMDispatch_If::FalseName = TEXT("False");
-const FName FRigVMDispatch_If::ResultName = TEXT("Result");
-
 FName FRigVMDispatch_If::GetArgumentNameForOperandIndex(int32 InOperandIndex, int32 InTotalOperands) const
 {
-	static const FName ArgumentNames[] = {
+	static const FLazyName ArgumentNames[] = {
 		ConditionName,
 		TrueName,
 		FalseName,
diff --git a/Engine/Plugins/Runtime/RigVM/Source/RigVM/Private/RigVMFunctions/RigVMDispatch_MakeStruct.cpp b/Engine/Plugins/Runtime/RigVM/Source/RigVM/Private/RigVMFunctions/RigVMDispatch_MakeStruct.cpp
--- a/Engine/Plugins/Runtime/RigVM/Source/RigVM/Private/RigVMFunctions/RigVMDispatch_MakeStruct.cpp
+++ b/Engine/Plugins/Runtime/RigVM/Source/RigVM/Private/RigVMFunctions/RigVMDispatch_MakeStruct.cpp
@@ -8,9 +8,6 @@
 #include UE_INLINE_GENERATED_CPP_BY_NAME(RigVMDispatch_MakeStruct)
 #define LOCTEXT_NAMESPACE "RigVMDispatch_MakeStruct"
 
-const FName FRigVMDispatch_MakeStruct::ElementsName = TEXT("Elements");
-const FName FRigVMDispatch_MakeStruct::StructName = TEXT("Struct");
-
 const TArray<FRigVMTemplateArgumentInfo>& FRigVMDispatch_MakeStruct::GetArgumentInfos() const
 {
 	static TArray<FRigVMTemplateArgumentInfo> Infos;
diff --git a/Engine/Plugins/Runtime/RigVM/Source/RigVM/Private/RigVMFunctions/RigVMDispatch_Print.cpp b/Engine/Plugins/Runtime/RigVM/Source/RigVM/Private/RigVMFunctions/RigVMDispatch_Print.cpp
--- a/Engine/Plugins/Runtime/RigVM/Source/RigVM/Private/RigVMFunctions/RigVMDispatch_Print.cpp
+++ b/Engine/Plugins/Runtime/RigVM/Source/RigVM/Private/RigVMFunctions/RigVMDispatch_Print.cpp
@@ -7,17 +7,9 @@
 
 #include UE_INLINE_GENERATED_CPP_BY_NAME(RigVMDispatch_Print)
 
-
-const FName FRigVMDispatch_Print::PrefixName = TEXT("Prefix");
-const FName FRigVMDispatch_Print::ValueName = TEXT("Value");
-const FName FRigVMDispatch_Print::EnabledName = TEXT("Enabled");
-const FName FRigVMDispatch_Print::ScreenDurationName = TEXT("ScreenDuration");
-const FName FRigVMDispatch_Print::ScreenColorName = TEXT("ScreenColor");
-
-
 FName FRigVMDispatch_Print::GetArgumentNameForOperandIndex(int32 InOperandIndex, int32 InTotalOperands) const
 {
-	static const FName ArgumentNames[] = {
+	static const FLazyName ArgumentNames[] = {
 		PrefixName,
 		ValueName,
 		EnabledName,
diff --git a/Engine/Plugins/Runtime/RigVM/Source/RigVM/Private/RigVMFunctions/RigVMDispatch_Select.cpp b/Engine/Plugins/Runtime/RigVM/Source/RigVM/Private/RigVMFunctions/RigVMDispatch_Select.cpp
--- a/Engine/Plugins/Runtime/RigVM/Source/RigVM/Private/RigVMFunctions/RigVMDispatch_Select.cpp
+++ b/Engine/Plugins/Runtime/RigVM/Source/RigVM/Private/RigVMFunctions/RigVMDispatch_Select.cpp
@@ -7,11 +7,6 @@
 
 #include UE_INLINE_GENERATED_CPP_BY_NAME(RigVMDispatch_Select)
 
-const FName FRigVMDispatch_SelectInt32::IndexName = TEXT("Index");
-const FName FRigVMDispatch_SelectInt32::ValuesName = TEXT("Values");
-const FName FRigVMDispatch_SelectInt32::ResultName = TEXT("Result");
-
-
 FName FRigVMDispatch_SelectInt32::GetArgumentNameForOperandIndex(int32 InOperandIndex, int32 InTotalOperands) const
 {
 	if(InOperandIndex > 0 && InOperandIndex < InTotalOperands - 1)
diff --git a/Engine/Plugins/Runtime/RigVM/Source/RigVM/Private/RigVMFunctions/RigVMDispatch_Switch.cpp b/Engine/Plugins/Runtime/RigVM/Source/RigVM/Private/RigVMFunctions/RigVMDispatch_Switch.cpp
--- a/Engine/Plugins/Runtime/RigVM/Source/RigVM/Private/RigVMFunctions/RigVMDispatch_Switch.cpp
+++ b/Engine/Plugins/Runtime/RigVM/Source/RigVM/Private/RigVMFunctions/RigVMDispatch_Switch.cpp
@@ -7,12 +7,9 @@
 
 #include UE_INLINE_GENERATED_CPP_BY_NAME(RigVMDispatch_Switch)
 
-const FName FRigVMDispatch_SwitchInt32::IndexName = TEXT("Index");
-const FName FRigVMDispatch_SwitchInt32::CasesName = TEXT("Cases");
-
 FName FRigVMDispatch_SwitchInt32::GetArgumentNameForOperandIndex(int32 InOperandIndex, int32 InTotalOperands) const
 {
-	static const FName ArgumentNames[] = {
+	static const FLazyName ArgumentNames[] = {
 		IndexName,
 		FRigVMStruct::ControlFlowBlockToRunName
 	};
diff --git a/Engine/Plugins/Runtime/RigVM/Source/RigVM/Private/RigVMFunctions/RigVMFunction_ControlFlow.cpp b/Engine/Plugins/Runtime/RigVM/Source/RigVM/Private/RigVMFunctions/RigVMFunction_ControlFlow.cpp
--- a/Engine/Plugins/Runtime/RigVM/Source/RigVM/Private/RigVMFunctions/RigVMFunction_ControlFlow.cpp
+++ b/Engine/Plugins/Runtime/RigVM/Source/RigVM/Private/RigVMFunctions/RigVMFunction_ControlFlow.cpp
@@ -10,12 +10,12 @@ FRigVMFunction_ControlFlowBranch_Execute()
 	{
 		if(Condition)
 		{
-			static const FName TrueName = GET_MEMBER_NAME_CHECKED(FRigVMFunction_ControlFlowBranch, True); 
+			static const FLazyName TrueName(GET_MEMBER_NAME_CHECKED(FRigVMFunction_ControlFlowBranch, True)); 
 			BlockToRun = TrueName;
 		}
 		else
 		{
-			static const FName FalseName = GET_MEMBER_NAME_CHECKED(FRigVMFunction_ControlFlowBranch, False); 
+			static const FLazyName FalseName(GET_MEMBER_NAME_CHECKED(FRigVMFunction_ControlFlowBranch, False)); 
 			BlockToRun = FalseName;
 		}
 	}
diff --git a/Engine/Plugins/Runtime/RigVM/Source/RigVM/Private/RigVMTypeUtils.cpp b/Engine/Plugins/Runtime/RigVM/Source/RigVM/Private/RigVMTypeUtils.cpp
--- a/Engine/Plugins/Runtime/RigVM/Source/RigVM/Private/RigVMTypeUtils.cpp
+++ b/Engine/Plugins/Runtime/RigVM/Source/RigVM/Private/RigVMTypeUtils.cpp
@@ -74,7 +74,7 @@ UObject* RigVMTypeUtils::UserDefinedTypeFromCPPType(FString& InOutCPPType, const
 
 		if(OriginalTypeName.Contains(TEXT("FUserDefinedStruct_")))
 		{
-			static const FName GuidTag = GET_MEMBER_NAME_CHECKED(UUserDefinedStruct, Guid);
+			static const FLazyName GuidTag(GET_MEMBER_NAME_CHECKED(UUserDefinedStruct, Guid));
 
 			TArray<FAssetData> AssetDataList;
 			AssetRegistry.GetAssetsByClass(UUserDefinedStruct::StaticClass()->GetClassPathName(), AssetDataList, true);
diff --git a/Engine/Plugins/Runtime/RigVM/Source/RigVM/Public/RigVMCore/RigVMDispatchFactory.h b/Engine/Plugins/Runtime/RigVM/Source/RigVM/Public/RigVMCore/RigVMDispatchFactory.h
--- a/Engine/Plugins/Runtime/RigVM/Source/RigVM/Public/RigVMCore/RigVMDispatchFactory.h
+++ b/Engine/Plugins/Runtime/RigVM/Source/RigVM/Public/RigVMCore/RigVMDispatchFactory.h
@@ -248,8 +248,8 @@ protected:
 
 	RIGVM_API virtual const TArray<FName>& GetControlFlowBlocks_Impl(const FRigVMDispatchContext& InContext) const;
 
-	RIGVM_API static const FString DispatchPrefix;
-	RIGVM_API static const FString TrueString;
+	RIGVM_API static const inline TCHAR* DispatchPrefix = TEXT("DISPATCH_");
+	RIGVM_API static const inline TCHAR* TrueString = TEXT("True");
 
 	const TArray<FName>* UpdateArgumentNameCache(int32 InNumberOperands) const;
 	const TArray<FName>* UpdateArgumentNameCache_NoLock(int32 InNumberOperands) const;
diff --git a/Engine/Plugins/Runtime/RigVM/Source/RigVM/Public/RigVMCore/RigVMGraphFunctionDefinition.h b/Engine/Plugins/Runtime/RigVM/Source/RigVM/Public/RigVMCore/RigVMGraphFunctionDefinition.h
--- a/Engine/Plugins/Runtime/RigVM/Source/RigVM/Public/RigVMCore/RigVMGraphFunctionDefinition.h
+++ b/Engine/Plugins/Runtime/RigVM/Source/RigVM/Public/RigVMCore/RigVMGraphFunctionDefinition.h
@@ -566,6 +566,6 @@ struct RIGVM_API FRigVMGraphFunctionData
 
 	bool PatchSharedArgumentOperandsIfRequired();
 
-	static const FString EntryString;
-	static const FString ReturnString;
+	static const inline TCHAR* EntryString = TEXT("Entry");
+	static const inline TCHAR* ReturnString = TEXT("Return");
 };
diff --git a/Engine/Plugins/Runtime/RigVM/Source/RigVM/Public/RigVMCore/RigVMMemoryStorage.h b/Engine/Plugins/Runtime/RigVM/Source/RigVM/Public/RigVMCore/RigVMMemoryStorage.h
--- a/Engine/Plugins/Runtime/RigVM/Source/RigVM/Public/RigVMCore/RigVMMemoryStorage.h
+++ b/Engine/Plugins/Runtime/RigVM/Source/RigVM/Public/RigVMCore/RigVMMemoryStorage.h
@@ -1031,9 +1031,9 @@ public:
 
 private:
 	
-	static const FString ArrayPrefix;
-	static const FString MapPrefix;
-	static const FString ContainerSuffix;
+	static const inline TCHAR* ArrayPrefix = TEXT("TArray<");
+	static const inline TCHAR* MapPrefix = TEXT("TMap<");
+	static const inline TCHAR* ContainerSuffix = TEXT(">");
 };
 
 /**
diff --git a/Engine/Plugins/Runtime/RigVM/Source/RigVM/Public/RigVMCore/RigVMRegistry.h b/Engine/Plugins/Runtime/RigVM/Source/RigVM/Public/RigVMCore/RigVMRegistry.h
--- a/Engine/Plugins/Runtime/RigVM/Source/RigVM/Public/RigVMCore/RigVMRegistry.h
+++ b/Engine/Plugins/Runtime/RigVM/Source/RigVM/Public/RigVMCore/RigVMRegistry.h
@@ -290,9 +290,9 @@ public:
 	uint32 GetHashForEnum(const UEnum* InEnum, bool bCheckTypeIndex = true) const;
 	uint32 GetHashForProperty(const FProperty* InProperty) const;
 
-private:
+	static inline const FLazyName TemplateNameMetaName = FLazyName(TEXT("TemplateName"));
 
-	static const FName TemplateNameMetaName;
+private:
 
 	FRigVMRegistry();
 
diff --git a/Engine/Plugins/Runtime/RigVM/Source/RigVM/Public/RigVMCore/RigVMStruct.h b/Engine/Plugins/Runtime/RigVM/Source/RigVM/Public/RigVMCore/RigVMStruct.h
--- a/Engine/Plugins/Runtime/RigVM/Source/RigVM/Public/RigVMCore/RigVMStruct.h
+++ b/Engine/Plugins/Runtime/RigVM/Source/RigVM/Public/RigVMCore/RigVMStruct.h
@@ -226,49 +226,48 @@ public:
 	virtual FName GetNextAggregateName(const FName& InLastAggregatePinName) const;
 	virtual FRigVMStructUpgradeInfo GetUpgradeInfo() const { return FRigVMStructUpgradeInfo(); }
 
-	static const FName DeprecatedMetaName;
-	static const FName InputMetaName;
-	static const FName OutputMetaName;
-	static const FName IOMetaName;
-	static const FName HiddenMetaName;
-	static const FName VisibleMetaName;
-	static const FName DetailsOnlyMetaName;
-	static const FName AbstractMetaName;
-	static const FName CategoryMetaName;
-	static const FName DisplayNameMetaName;
-	static const FName MenuDescSuffixMetaName;
-	static const FName ShowVariableNameInTitleMetaName;
-	static const FName CustomWidgetMetaName;
-	static const FName ConstantMetaName;
-	static const FName TitleColorMetaName;
-	static const FName NodeColorMetaName;
+	static inline const FLazyName DeprecatedMetaName = FLazyName(TEXT("Deprecated"));
+	static inline const FLazyName InputMetaName = FLazyName(TEXT("Input"));
+	static inline const FLazyName OutputMetaName = FLazyName(TEXT("Output"));
+	static inline const FLazyName IOMetaName = FLazyName(TEXT("IO"));
+	static inline const FLazyName HiddenMetaName = FLazyName(TEXT("Hidden"));
+	static inline const FLazyName VisibleMetaName = FLazyName(TEXT("Visible"));
+	static inline const FLazyName DetailsOnlyMetaName = FLazyName(TEXT("DetailsOnly"));
+	static inline const FLazyName AbstractMetaName = FLazyName(TEXT("Abstract"));
+	static inline const FLazyName CategoryMetaName = FLazyName(TEXT("Category"));
+	static inline const FLazyName DisplayNameMetaName = FLazyName(TEXT("DisplayName"));
+	static inline const FLazyName MenuDescSuffixMetaName = FLazyName(TEXT("MenuDescSuffix"));
+	static inline const FLazyName ShowVariableNameInTitleMetaName = FLazyName(TEXT("ShowVariableNameInTitle"));
+	static inline const FLazyName CustomWidgetMetaName = FLazyName(TEXT("CustomWidget"));
+	static inline const FLazyName ConstantMetaName = FLazyName(TEXT("Constant"));
+	static inline const FLazyName TitleColorMetaName = FLazyName(TEXT("TitleColor"));
+	static inline const FLazyName NodeColorMetaName = FLazyName(TEXT("NodeColor"));
 	// icon meta name format: StyleSetName|StyleName|SmallStyleName|StatusOverlayStyleName
 	// the last two names are optional, see FSlateIcon() for reference
 	// Example: Icon="EditorStyle|GraphEditor.Sequence_16x"
-	static const FName IconMetaName;
-	static const FName KeywordsMetaName;
-	static const FName TemplateNameMetaName;
-	static const FName FixedSizeArrayMetaName;
-	static const FName ShowOnlySubPinsMetaName;
-	static const FName HideSubPinsMetaName;
-	static const FName ArraySizeMetaName;
-	static const FName AggregateMetaName;
-	static const FName ExpandPinByDefaultMetaName;
-	static const FName DefaultArraySizeMetaName;
-	static const FName VaryingMetaName;
-	static const FName SingletonMetaName;
-	static const FName SliceContextMetaName;
-	static const FName ExecuteName;
-	static const FName ExecuteContextName;
-	static const FName ForLoopCountPinName;
-	static const FName ForLoopContinuePinName;
-	static const FName ForLoopCompletedPinName;
-	static const FName ForLoopIndexPinName;
-	static const FName ComputeLazilyMetaName;
-	static const FName ControlFlowBlockToRunName;
-	static const FName ControlFlowCompletedName;
-	static const FName ControlFlowCountName;
-	static const FName ControlFlowIndexName;
+	static inline const FLazyName IconMetaName = FLazyName(TEXT("Icon"));
+	static inline const FLazyName KeywordsMetaName = FLazyName(TEXT("Keywords"));
+	static inline const FLazyName FixedSizeArrayMetaName = FLazyName(TEXT("FixedSizeArray"));
+	static inline const FLazyName ShowOnlySubPinsMetaName = FLazyName(TEXT("ShowOnlySubPins"));
+	static inline const FLazyName HideSubPinsMetaName = FLazyName(TEXT("HideSubPins"));
+	static inline const FLazyName ArraySizeMetaName = FLazyName(TEXT("ArraySize"));
+	static inline const FLazyName AggregateMetaName = FLazyName(TEXT("Aggregate"));
+	static inline const FLazyName ExpandPinByDefaultMetaName = FLazyName(TEXT("ExpandByDefault"));
+	static inline const FLazyName DefaultArraySizeMetaName = FLazyName(TEXT("DefaultArraySize"));
+	static inline const FLazyName VaryingMetaName = FLazyName(TEXT("Varying"));
+	static inline const FLazyName SingletonMetaName = FLazyName(TEXT("Singleton"));
+	static inline const FLazyName SliceContextMetaName = FLazyName(TEXT("SliceContext"));
+	static inline const FLazyName ExecuteName = FLazyName(TEXT("Execute"));
+	static inline const FLazyName ExecuteContextName = FLazyName(TEXT("ExecuteContext"));
+	static inline const FLazyName ForLoopCountPinName = FLazyName(TEXT("Count"));
+	static inline const FLazyName ForLoopContinuePinName = FLazyName(TEXT("Continue"));
+	static inline const FLazyName ForLoopCompletedPinName = FLazyName(TEXT("Completed"));
+	static inline const FLazyName ForLoopIndexPinName = FLazyName(TEXT("Index"));
+	static inline const FLazyName ComputeLazilyMetaName = FLazyName(TEXT("Lazy"));
+	static inline const FLazyName ControlFlowBlockToRunName = FLazyName(TEXT("BlockToRun"));
+	static inline const FLazyName ControlFlowCompletedName = FLazyName(TEXT("Completed"));
+	static inline const FLazyName ControlFlowCountName = FLazyName(TEXT("Count"));
+	static inline const FLazyName ControlFlowIndexName = FLazyName(TEXT("Index"));
 
 protected:
 
diff --git a/Engine/Plugins/Runtime/RigVM/Source/RigVM/Public/RigVMFunctions/RigVMDispatch_Array.h b/Engine/Plugins/Runtime/RigVM/Source/RigVM/Public/RigVMFunctions/RigVMDispatch_Array.h
--- a/Engine/Plugins/Runtime/RigVM/Source/RigVM/Public/RigVMFunctions/RigVMDispatch_Array.h
+++ b/Engine/Plugins/Runtime/RigVM/Source/RigVM/Public/RigVMFunctions/RigVMDispatch_Array.h
@@ -24,20 +24,18 @@ public:
 protected:
 	static FRigVMTemplateArgumentInfo CreateArgumentInfo(const FName& InName, ERigVMPinDirection InDirection);
 	static TMap<uint32, int32> GetArrayHash(FScriptArrayHelper& InArrayHelper, const FArrayProperty* InArrayProperty);
-	
-	static const FName& ExecuteName;
-	static const FName ArrayName;
-	static const FName ValuesName;
-	static const FName NumName;
-	static const FName IndexName;
-	static const FName ElementName;
-	static const FName SuccessName;
-	static const FName OtherName;
-	static const FName CloneName;
-	static const FName CountName;
-	static const FName RatioName;
-	static const FName ResultName;
-	static const FName& CompletedName;
+
+	static inline FLazyName ArrayName = FLazyName(TEXT("Array"));
+	static inline FLazyName ValuesName = FLazyName(TEXT("Values"));
+	static inline FLazyName NumName = FLazyName(TEXT("Num"));
+	static inline FLazyName IndexName = FLazyName(TEXT("Index"));
+	static inline FLazyName ElementName = FLazyName(TEXT("Element"));
+	static inline FLazyName SuccessName = FLazyName(TEXT("Success"));
+	static inline FLazyName OtherName = FLazyName(TEXT("Other"));
+	static inline FLazyName CloneName = FLazyName(TEXT("Clone"));
+	static inline FLazyName CountName = FLazyName(TEXT("Count"));
+	static inline FLazyName RatioName = FLazyName(TEXT("Ratio"));
+	static inline FLazyName ResultName = FLazyName(TEXT("Result"));
 
 	friend class URigVMController;
 };
diff --git a/Engine/Plugins/Runtime/RigVM/Source/RigVM/Public/RigVMFunctions/RigVMDispatch_CastEnum.h b/Engine/Plugins/Runtime/RigVM/Source/RigVM/Public/RigVMFunctions/RigVMDispatch_CastEnum.h
--- a/Engine/Plugins/Runtime/RigVM/Source/RigVM/Public/RigVMFunctions/RigVMDispatch_CastEnum.h
+++ b/Engine/Plugins/Runtime/RigVM/Source/RigVM/Public/RigVMFunctions/RigVMDispatch_CastEnum.h
@@ -7,8 +7,8 @@
 
 struct FRigVMDispatch_CastEnumBase
 {
-	static const FName ValueName;
-	static const FName ResultName;
+	static inline const FLazyName ValueName = FLazyName(TEXT("Value"));
+	static inline const FLazyName ResultName = FLazyName(TEXT("Result"));
 };
 
 USTRUCT(meta=(DisplayName = "Cast", Category = "Enum", Keywords = "As", NodeColor = "1,1,1,1"))
diff --git a/Engine/Plugins/Runtime/RigVM/Source/RigVM/Public/RigVMFunctions/RigVMDispatch_CastObject.h b/Engine/Plugins/Runtime/RigVM/Source/RigVM/Public/RigVMFunctions/RigVMDispatch_CastObject.h
--- a/Engine/Plugins/Runtime/RigVM/Source/RigVM/Public/RigVMFunctions/RigVMDispatch_CastObject.h
+++ b/Engine/Plugins/Runtime/RigVM/Source/RigVM/Public/RigVMFunctions/RigVMDispatch_CastObject.h
@@ -28,6 +28,6 @@ protected:
 	virtual FRigVMFunctionPtr GetDispatchFunctionImpl(const FRigVMTemplateTypeMap& InTypes) const override { return &FRigVMDispatch_CastObject::Execute; }
 	static void Execute(FRigVMExtendedExecuteContext& InContext, FRigVMMemoryHandleArray Handles, FRigVMPredicateBranchArray RigVMBranches);
 
-	static const FName ValueName;
-	static const FName ResultName;
+	static inline const FLazyName ValueName = FLazyName(TEXT("Value"));
+	static inline const FLazyName ResultName = FLazyName(TEXT("Result"));
 };
diff --git a/Engine/Plugins/Runtime/RigVM/Source/RigVM/Public/RigVMFunctions/RigVMDispatch_Constant.h b/Engine/Plugins/Runtime/RigVM/Source/RigVM/Public/RigVMFunctions/RigVMDispatch_Constant.h
--- a/Engine/Plugins/Runtime/RigVM/Source/RigVM/Public/RigVMFunctions/RigVMDispatch_Constant.h
+++ b/Engine/Plugins/Runtime/RigVM/Source/RigVM/Public/RigVMFunctions/RigVMDispatch_Constant.h
@@ -30,7 +30,7 @@ protected:
 	virtual FRigVMFunctionPtr GetDispatchFunctionImpl(const FRigVMTemplateTypeMap& InTypes) const override { return &FRigVMDispatch_Constant::Execute; }
 	static void Execute(FRigVMExtendedExecuteContext& InContext, FRigVMMemoryHandleArray Handles, FRigVMPredicateBranchArray RigVMBranches);
 
-	static const FName ValueName;
+	static inline const FLazyName ValueName = FLazyName(TEXT("Value"));
 
 	friend class URigVMController;
 };
diff --git a/Engine/Plugins/Runtime/RigVM/Source/RigVM/Public/RigVMFunctions/RigVMDispatch_Core.h b/Engine/Plugins/Runtime/RigVM/Source/RigVM/Public/RigVMFunctions/RigVMDispatch_Core.h
--- a/Engine/Plugins/Runtime/RigVM/Source/RigVM/Public/RigVMFunctions/RigVMDispatch_Core.h
+++ b/Engine/Plugins/Runtime/RigVM/Source/RigVM/Public/RigVMFunctions/RigVMDispatch_Core.h
@@ -82,9 +82,9 @@ protected:
 		Result = AdaptResult(A.Equals(B), InContext);
 	}
 
-	static const FName AName;
-	static const FName BName;
-	static const FName ResultName;
+	static inline FLazyName AName = FLazyName(TEXT("A"));
+	static inline FLazyName BName = FLazyName(TEXT("B"));
+	static inline FLazyName ResultName = FLazyName(TEXT("Result"));
 };
 
 /*
diff --git a/Engine/Plugins/Runtime/RigVM/Source/RigVM/Public/RigVMFunctions/RigVMDispatch_If.h b/Engine/Plugins/Runtime/RigVM/Source/RigVM/Public/RigVMFunctions/RigVMDispatch_If.h
--- a/Engine/Plugins/Runtime/RigVM/Source/RigVM/Public/RigVMFunctions/RigVMDispatch_If.h
+++ b/Engine/Plugins/Runtime/RigVM/Source/RigVM/Public/RigVMFunctions/RigVMDispatch_If.h
@@ -35,8 +35,8 @@ protected:
 	virtual FRigVMFunctionPtr GetDispatchFunctionImpl(const FRigVMTemplateTypeMap& InTypes) const override;
 	static void Execute(FRigVMExtendedExecuteContext& InContext, FRigVMMemoryHandleArray Handles, FRigVMPredicateBranchArray Predicates);
 
-	static const FName ConditionName;
-	static const FName TrueName;
-	static const FName FalseName;
-	static const FName ResultName;
+	static inline const FLazyName ConditionName = FLazyName(TEXT("Condition"));
+	static inline const FLazyName TrueName = FLazyName(TEXT("True"));
+	static inline const FLazyName FalseName = FLazyName(TEXT("False"));
+	static inline const FLazyName ResultName = FLazyName(TEXT("Result"));
 };
diff --git a/Engine/Plugins/Runtime/RigVM/Source/RigVM/Public/RigVMFunctions/RigVMDispatch_MakeStruct.h b/Engine/Plugins/Runtime/RigVM/Source/RigVM/Public/RigVMFunctions/RigVMDispatch_MakeStruct.h
--- a/Engine/Plugins/Runtime/RigVM/Source/RigVM/Public/RigVMFunctions/RigVMDispatch_MakeStruct.h
+++ b/Engine/Plugins/Runtime/RigVM/Source/RigVM/Public/RigVMFunctions/RigVMDispatch_MakeStruct.h
@@ -31,8 +31,8 @@ protected:
 	virtual FRigVMFunctionPtr GetDispatchFunctionImpl(const FRigVMTemplateTypeMap& InTypes) const override { return &FRigVMDispatch_MakeStruct::Execute; }
 	static void Execute(FRigVMExtendedExecuteContext& InContext, FRigVMMemoryHandleArray Handles, FRigVMPredicateBranchArray RigVMBranches);
 
-	static const FName ElementsName;
-	static const FName StructName;
+	static inline const FLazyName ElementsName = FLazyName(TEXT("Elements"));
+	static inline const FLazyName StructName = FLazyName(TEXT("Struct"));
 
 	friend struct FRigVMDispatch_BreakStruct;
 	friend class URigVMController;
diff --git a/Engine/Plugins/Runtime/RigVM/Source/RigVM/Public/RigVMFunctions/RigVMDispatch_Print.h b/Engine/Plugins/Runtime/RigVM/Source/RigVM/Public/RigVMFunctions/RigVMDispatch_Print.h
--- a/Engine/Plugins/Runtime/RigVM/Source/RigVM/Public/RigVMFunctions/RigVMDispatch_Print.h
+++ b/Engine/Plugins/Runtime/RigVM/Source/RigVM/Public/RigVMFunctions/RigVMDispatch_Print.h
@@ -39,9 +39,9 @@ protected:
 	}
 	static void Execute(FRigVMExtendedExecuteContext& InContext, FRigVMMemoryHandleArray Handles, FRigVMPredicateBranchArray Predicates);
 
-	static const FName PrefixName;
-	static const FName ValueName;
-	static const FName EnabledName;
-	static const FName ScreenDurationName;
-	static const FName ScreenColorName;
+	static inline const FLazyName PrefixName = FLazyName(TEXT("Prefix"));
+	static inline const FLazyName ValueName = FLazyName(TEXT("Value"));
+	static inline const FLazyName EnabledName = FLazyName(TEXT("Enabled"));
+	static inline const FLazyName ScreenDurationName = FLazyName(TEXT("ScreenDuration"));
+	static inline const FLazyName ScreenColorName = FLazyName(TEXT("ScreenColor"));
 };
diff --git a/Engine/Plugins/Runtime/RigVM/Source/RigVM/Public/RigVMFunctions/RigVMDispatch_Select.h b/Engine/Plugins/Runtime/RigVM/Source/RigVM/Public/RigVMFunctions/RigVMDispatch_Select.h
--- a/Engine/Plugins/Runtime/RigVM/Source/RigVM/Public/RigVMFunctions/RigVMDispatch_Select.h
+++ b/Engine/Plugins/Runtime/RigVM/Source/RigVM/Public/RigVMFunctions/RigVMDispatch_Select.h
@@ -35,7 +35,7 @@ protected:
 	virtual FRigVMFunctionPtr GetDispatchFunctionImpl(const FRigVMTemplateTypeMap& InTypes) const override;
 	static void Execute(FRigVMExtendedExecuteContext& InContext, FRigVMMemoryHandleArray Handles, FRigVMPredicateBranchArray Predicates);
 
-	static const FName IndexName;
-	static const FName ValuesName;
-	static const FName ResultName;
+	static inline FLazyName IndexName = FLazyName(TEXT("Index"));
+	static inline FLazyName ValuesName = FLazyName(TEXT("Values"));
+	static inline FLazyName ResultName = FLazyName(TEXT("Result"));
 };
diff --git a/Engine/Plugins/Runtime/RigVM/Source/RigVM/Public/RigVMFunctions/RigVMDispatch_Switch.h b/Engine/Plugins/Runtime/RigVM/Source/RigVM/Public/RigVMFunctions/RigVMDispatch_Switch.h
--- a/Engine/Plugins/Runtime/RigVM/Source/RigVM/Public/RigVMFunctions/RigVMDispatch_Switch.h
+++ b/Engine/Plugins/Runtime/RigVM/Source/RigVM/Public/RigVMFunctions/RigVMDispatch_Switch.h
@@ -40,7 +40,7 @@ protected:
 	static void Execute(FRigVMExtendedExecuteContext& InContext, FRigVMMemoryHandleArray Handles, FRigVMPredicateBranchArray Predicates);
 	static FName GetCaseName(int32 InIndex);
 	static FName GetCaseDisplayName(int32 InIndex);
-	
-	static const FName IndexName;
-	static const FName CasesName;
+
+	static inline const FLazyName IndexName = FLazyName(TEXT("Index"));
+	static inline const FLazyName CasesName = FLazyName(TEXT("Cases"));
 };
diff --git a/Engine/Plugins/Runtime/RigVM/Source/RigVM/Public/RigVMTypeUtils.h b/Engine/Plugins/Runtime/RigVM/Source/RigVM/Public/RigVMTypeUtils.h
--- a/Engine/Plugins/Runtime/RigVM/Source/RigVM/Public/RigVMTypeUtils.h
+++ b/Engine/Plugins/Runtime/RigVM/Source/RigVM/Public/RigVMTypeUtils.h
@@ -46,21 +46,24 @@ private:
 
 namespace RigVMTypeUtils
 {
-	const TCHAR TArrayPrefix[] = TEXT("TArray<");
-	const TCHAR TObjectPtrPrefix[] = TEXT("TObjectPtr<");
-	const TCHAR TSubclassOfPrefix[] = TEXT("TSubclassOf<");
-	const TCHAR TScriptInterfacePrefix[] = TEXT("TScriptInterface<");
-	const TCHAR TArrayTemplate[] = TEXT("TArray<%s>");
-	const TCHAR TObjectPtrTemplate[] = TEXT("TObjectPtr<%s%s>");
-	const TCHAR TSubclassOfTemplate[] = TEXT("TSubclassOf<%s%s>");
-	const TCHAR TScriptInterfaceTemplate[] = TEXT("TScriptInterface<%s%s>");
+	constexpr TCHAR TArrayPrefix[] = TEXT("TArray<");
+	constexpr TCHAR TObjectPtrPrefix[] = TEXT("TObjectPtr<");
+	constexpr TCHAR TSubclassOfPrefix[] = TEXT("TSubclassOf<");
+	constexpr TCHAR TScriptInterfacePrefix[] = TEXT("TScriptInterface<");
+	constexpr TCHAR TArrayTemplate[] = TEXT("TArray<%s>");
+	constexpr TCHAR TObjectPtrTemplate[] = TEXT("TObjectPtr<%s%s>");
+	constexpr TCHAR TSubclassOfTemplate[] = TEXT("TSubclassOf<%s%s>");
+	constexpr TCHAR TScriptInterfaceTemplate[] = TEXT("TScriptInterface<%s%s>");
 
 	const inline TCHAR* BoolType = TEXT("bool");
 	const inline TCHAR* FloatType = TEXT("float");
 	const inline TCHAR* DoubleType = TEXT("double");
+	const inline TCHAR* IntType = TEXT("int");
 	const inline TCHAR* Int32Type = TEXT("int32");
-	const inline TCHAR* UInt32Type = TEXT("uint32");
+	const inline TCHAR* Int64Type = TEXT("int64");
 	const inline TCHAR* UInt8Type = TEXT("uint8");
+	const inline TCHAR* UInt32Type = TEXT("uint32");
+	const inline TCHAR* UInt64Type = TEXT("uint64");
 	const inline TCHAR* FNameType = TEXT("FName");
 	const inline TCHAR* FStringType = TEXT("FString");
 	const inline TCHAR* FTextType = TEXT("FText");
@@ -77,9 +80,12 @@ namespace RigVMTypeUtils
 	const FLazyName BoolTypeName(BoolType);
 	const FLazyName FloatTypeName(FloatType);
 	const FLazyName DoubleTypeName(DoubleType);
+	const FLazyName IntTypeName(IntType);
 	const FLazyName Int32TypeName(Int32Type);
-	const FLazyName UInt32TypeName(UInt32Type);
+	const FLazyName Int64TypeName(Int64Type);
 	const FLazyName UInt8TypeName(UInt8Type);
+	const FLazyName UInt32TypeName(UInt32Type);
+	const FLazyName UInt64TypeName(UInt64Type);
 	const FLazyName FNameTypeName(FNameType);
 	const FLazyName FStringTypeName(FStringType);
 	const FLazyName FTextTypeName(FTextType);
@@ -185,9 +191,9 @@ namespace RigVMTypeUtils
 		return WildCardCPPType;
 	}
 
-	static const FName& GetWildCardCPPTypeName()
+	static const FLazyName& GetWildCardCPPTypeName()
 	{
-		static const FName WildCardCPPTypeName = *GetWildCardCPPType(); 
+		static const FLazyName WildCardCPPTypeName(*GetWildCardCPPType()); 
 		return WildCardCPPTypeName;
 	}
 
@@ -197,9 +203,9 @@ namespace RigVMTypeUtils
 		return WildCardArrayCPPType;
 	}
 
-	static const FName& GetWildCardArrayCPPTypeName()
+	static const FLazyName& GetWildCardArrayCPPTypeName()
 	{
-		static const FName WildCardArrayCPPTypeName = *GetWildCardArrayCPPType(); 
+		static const FLazyName WildCardArrayCPPTypeName(*GetWildCardArrayCPPType()); 
 		return WildCardArrayCPPTypeName;
 	}
 
diff --git a/Engine/Plugins/Runtime/RigVM/Source/RigVMDeveloper/Private/EdGraph/RigVMEdGraphSchema.cpp b/Engine/Plugins/Runtime/RigVM/Source/RigVMDeveloper/Private/EdGraph/RigVMEdGraphSchema.cpp
--- a/Engine/Plugins/Runtime/RigVM/Source/RigVMDeveloper/Private/EdGraph/RigVMEdGraphSchema.cpp
+++ b/Engine/Plugins/Runtime/RigVM/Source/RigVMDeveloper/Private/EdGraph/RigVMEdGraphSchema.cpp
@@ -40,8 +40,6 @@
 
 #define LOCTEXT_NAMESPACE "CRigVMGraphSchema"
 
-const FName URigVMEdGraphSchema::GraphName_RigVM(TEXT("RigVM"));
-
 FRigVMLocalVariableNameValidator::FRigVMLocalVariableNameValidator(const UBlueprint* Blueprint, const URigVMGraph* Graph, FName InExistingName)
 	: FStringSetNameValidator(InExistingName.ToString())
 {
diff --git a/Engine/Plugins/Runtime/RigVM/Source/RigVMDeveloper/Private/RigVMBlueprint.cpp b/Engine/Plugins/Runtime/RigVM/Source/RigVMDeveloper/Private/RigVMBlueprint.cpp
--- a/Engine/Plugins/Runtime/RigVM/Source/RigVMDeveloper/Private/RigVMBlueprint.cpp
+++ b/Engine/Plugins/Runtime/RigVM/Source/RigVMDeveloper/Private/RigVMBlueprint.cpp
@@ -149,8 +149,6 @@ FSoftObjectPath URigVMBlueprint::PreDuplicateAssetPath;
 FSoftObjectPath URigVMBlueprint::PreDuplicateHostPath;
 TArray<URigVMBlueprint*> URigVMBlueprint::sCurrentlyOpenedRigVMBlueprints;
 #if WITH_EDITOR
-const FName URigVMBlueprint::RigVMPanelNodeFactoryName(TEXT("FRigVMEdGraphPanelNodeFactory"));
-const FName URigVMBlueprint::RigVMPanelPinFactoryName(TEXT("FRigVMEdGraphPanelPinFactory"));
 FCriticalSection URigVMBlueprint::QueuedCompilerMessageDelegatesMutex;
 TArray<FOnRigVMReportCompilerMessage::FDelegate> URigVMBlueprint::QueuedCompilerMessageDelegates;
 #endif
@@ -762,12 +760,12 @@ URigVMEditorSettings* URigVMBlueprint::GetRigVMEditorSettings() const
 }
 
 #if WITH_EDITOR
-const FName& URigVMBlueprint::GetPanelNodeFactoryName() const
+const FLazyName& URigVMBlueprint::GetPanelNodeFactoryName() const
 {
 	return RigVMPanelNodeFactoryName;
 }
 
-const FName& URigVMBlueprint::GetPanelPinFactoryName() const
+const FLazyName& URigVMBlueprint::GetPanelPinFactoryName() const
 {
 	return RigVMPanelPinFactoryName;
 }
diff --git a/Engine/Plugins/Runtime/RigVM/Source/RigVMDeveloper/Private/RigVMCompiler/RigVMCodeGenerator.cpp b/Engine/Plugins/Runtime/RigVM/Source/RigVMDeveloper/Private/RigVMCompiler/RigVMCodeGenerator.cpp
--- a/Engine/Plugins/Runtime/RigVM/Source/RigVMDeveloper/Private/RigVMCompiler/RigVMCodeGenerator.cpp
+++ b/Engine/Plugins/Runtime/RigVM/Source/RigVMDeveloper/Private/RigVMCompiler/RigVMCodeGenerator.cpp
@@ -26,9 +26,9 @@ static constexpr TCHAR RigVM_DeclareExternalVariableFormat[] = TEXT("\t{0}* {1}
 static constexpr TCHAR RigVM_UpdateExternalVariableFormat[] = TEXT("\t{0} = &GetExternalVariableRef<{1}>(Context, TEXT(\"{2}\"), TEXT(\"{1}\"));");
 static constexpr TCHAR RigVM_MemberPropertyFormat[] = TEXT("\t{0} {1} = {2};");
 static constexpr TCHAR RigVM_MemberPropertyFormatNoDefault[] = TEXT("\t{0} {1};");
-static constexpr TCHAR RigVM_DeclareEntryNameFormat[] = TEXT("\tstatic const FName EntryName_{0};");
+static constexpr TCHAR RigVM_DeclareEntryNameFormat[] = TEXT("\tstatic const FLazyName EntryName_{0};");
 static constexpr TCHAR RigVM_DefineEntryNameFormat[] = TEXT("const FName U{0}::EntryName_{1} = TEXT(\"{2}\");");
-static constexpr TCHAR RigVM_DeclareBlockNameFormat[] = TEXT("\tstatic const FName BlockName_{0};");
+static constexpr TCHAR RigVM_DeclareBlockNameFormat[] = TEXT("\tstatic const FLazyName BlockName_{0};");
 static constexpr TCHAR RigVM_DefineBlockNameFormat[] = TEXT("const FName U{0}::BlockName_{1} = TEXT(\"{2}\");");
 static constexpr TCHAR RigVM_DefineConstFormatNoDefault[] = TEXT("\tstatic const {0} {1};");
 static constexpr TCHAR RigVM_StructConstantArrayArrayValue[] = TEXT("URigVMNativized::GetStructArrayArrayConstant<{0}>(TEXT(\"{1}\"))");
diff --git a/Engine/Plugins/Runtime/RigVM/Source/RigVMDeveloper/Private/RigVMCompiler/RigVMCompiler.cpp b/Engine/Plugins/Runtime/RigVM/Source/RigVMDeveloper/Private/RigVMCompiler/RigVMCompiler.cpp
--- a/Engine/Plugins/Runtime/RigVM/Source/RigVMDeveloper/Private/RigVMCompiler/RigVMCompiler.cpp
+++ b/Engine/Plugins/Runtime/RigVM/Source/RigVMDeveloper/Private/RigVMCompiler/RigVMCompiler.cpp
@@ -4022,7 +4022,9 @@ const FRigVMCompilerWorkData::FRigVMASTProxyArray& URigVMCompiler::FindProxiesWi
 
 FString URigVMCompiler::GetPinNameWithDirectionPrefix(const URigVMPin* Pin)
 {
-	const FString & Prefix = (Pin->GetDirection() == ERigVMPinDirection::Input) ? FRigVMGraphFunctionData::EntryString : (Pin->GetDirection() == ERigVMPinDirection::Output) ? FRigVMGraphFunctionData::ReturnString : "";
+	static const FString EntryString = FRigVMGraphFunctionData::EntryString;
+	static const FString ReturnString = FRigVMGraphFunctionData::ReturnString;
+	const FString & Prefix = (Pin->GetDirection() == ERigVMPinDirection::Input) ? EntryString : (Pin->GetDirection() == ERigVMPinDirection::Output) ? ReturnString : "";
 	const FString NameWithDirectionPrefix = Prefix + "_" + FRigVMPropertyDescription::SanitizeName(FName(Pin->GetName())).ToString();
 	return NameWithDirectionPrefix;
 }
diff --git a/Engine/Plugins/Runtime/RigVM/Source/RigVMDeveloper/Private/RigVMDeveloperTypeUtils.cpp b/Engine/Plugins/Runtime/RigVM/Source/RigVMDeveloper/Private/RigVMDeveloperTypeUtils.cpp
--- a/Engine/Plugins/Runtime/RigVM/Source/RigVMDeveloper/Private/RigVMDeveloperTypeUtils.cpp
+++ b/Engine/Plugins/Runtime/RigVM/Source/RigVMDeveloper/Private/RigVMDeveloperTypeUtils.cpp
@@ -581,9 +581,9 @@ bool RigVMTypeUtils::AreCompatible(const FEdGraphPinType& InTypeA, const FEdGrap
 
 namespace RigVMTypeUtils
 {
-	static const FName CastTemplateValueName = TEXT("Value");
-	static const FName CastTemplateResultName = TEXT("Result");
-	static const FName CastTemplateNotation = TEXT("Cast::Execute(in Value,out Result)");
+	static const FLazyName CastTemplateValueName = TEXT("Value");
+	static const FLazyName CastTemplateResultName = TEXT("Result");
+	static const FLazyName CastTemplateNotation = TEXT("Cast::Execute(in Value,out Result)");
 }
 
 const TArray<TRigVMTypeIndex>& RigVMTypeUtils::GetAvailableCasts(const TRigVMTypeIndex& InTypeIndex, bool bAsInput)
@@ -675,17 +675,17 @@ const FRigVMFunction* RigVMTypeUtils::GetCastForTypeIndices(const TRigVMTypeInde
 	return nullptr;
 }
 
-const FName& RigVMTypeUtils::GetCastTemplateValueName()
+const FLazyName& RigVMTypeUtils::GetCastTemplateValueName()
 {
 	return CastTemplateValueName;
 }
 
-const FName& RigVMTypeUtils::GetCastTemplateResultName()
+const FLazyName& RigVMTypeUtils::GetCastTemplateResultName()
 {
 	return CastTemplateResultName;
 }
 
-const FName& RigVMTypeUtils::GetCastTemplateNotation()
+const FLazyName& RigVMTypeUtils::GetCastTemplateNotation()
 {
 	return CastTemplateNotation;
 }
diff --git a/Engine/Plugins/Runtime/RigVM/Source/RigVMDeveloper/Private/RigVMModel/Nodes/RigVMEnumNode.cpp b/Engine/Plugins/Runtime/RigVM/Source/RigVMDeveloper/Private/RigVMModel/Nodes/RigVMEnumNode.cpp
--- a/Engine/Plugins/Runtime/RigVM/Source/RigVMDeveloper/Private/RigVMModel/Nodes/RigVMEnumNode.cpp
+++ b/Engine/Plugins/Runtime/RigVM/Source/RigVMDeveloper/Private/RigVMModel/Nodes/RigVMEnumNode.cpp
@@ -4,9 +4,7 @@
 
 #include UE_INLINE_GENERATED_CPP_BY_NAME(RigVMEnumNode)
 
-const FString URigVMEnumNode::EnumName = TEXT("Enum");
-const FString URigVMEnumNode::EnumValueName = TEXT("EnumValue");
-const FString URigVMEnumNode::EnumIndexName = TEXT("EnumIndex");
+const TCHAR* URigVMEnumNode::EnumName = TEXT("Enum");
 
 URigVMEnumNode::URigVMEnumNode()
 {
diff --git a/Engine/Plugins/Runtime/RigVM/Source/RigVMDeveloper/Private/RigVMModel/Nodes/RigVMIfNode.cpp b/Engine/Plugins/Runtime/RigVM/Source/RigVMDeveloper/Private/RigVMModel/Nodes/RigVMIfNode.cpp
--- a/Engine/Plugins/Runtime/RigVM/Source/RigVMDeveloper/Private/RigVMModel/Nodes/RigVMIfNode.cpp
+++ b/Engine/Plugins/Runtime/RigVM/Source/RigVMDeveloper/Private/RigVMModel/Nodes/RigVMIfNode.cpp
@@ -4,17 +4,10 @@
 
 #include UE_INLINE_GENERATED_CPP_BY_NAME(RigVMIfNode)
 
-const FString UDEPRECATED_RigVMIfNode::IfName = TEXT("If");
-const FString UDEPRECATED_RigVMIfNode::ConditionName = TEXT("Condition");
-const FString UDEPRECATED_RigVMIfNode::TrueName = TEXT("True");
-const FString UDEPRECATED_RigVMIfNode::FalseName = TEXT("False");
-const FString UDEPRECATED_RigVMIfNode::ResultName = TEXT("Result");
-
-
 FName UDEPRECATED_RigVMIfNode::GetNotation() const
 {
 	static constexpr TCHAR Format[] = TEXT("%s(in %s,in %s,in %s,out %s)");
-	static const FName Notation = *FString::Printf(Format, *IfName, *ConditionName, *TrueName, *FalseName, *ResultName);
+	static const FLazyName Notation(*FString::Printf(Format, IfName, ConditionName, TrueName, FalseName, ResultName));
 	return Notation;
 }
 
@@ -33,10 +26,10 @@ const FRigVMTemplate* UDEPRECATED_RigVMIfNode::GetTemplate() const
 			return IfNodeTemplate;
 		}
 
-		static const FName ConditionFName = *ConditionName;
-		static const FName TrueFName = *TrueName;
-		static const FName FalseFName = *FalseName;
-		static const FName ResultFName = *ResultName;
+		static const FLazyName ConditionFName(ConditionName);
+		static const FLazyName TrueFName(TrueName);
+		static const FLazyName FalseFName(FalseName);
+		static const FLazyName ResultFName(ResultName);
 
 		static TArray<FRigVMTemplateArgumentInfo> Infos;
 		if(Infos.IsEmpty())
@@ -69,7 +62,7 @@ const FRigVMTemplate* UDEPRECATED_RigVMIfNode::GetTemplate() const
 				return Types;
 			});
 
-		IfNodeTemplate = CachedTemplate = FRigVMRegistry::Get().GetOrAddTemplateFromArguments(*IfName, Infos, Delegates);
+		IfNodeTemplate = CachedTemplate = FRigVMRegistry::Get().GetOrAddTemplateFromArguments(IfName, Infos, Delegates);
 	}
 	return CachedTemplate;
 }
diff --git a/Engine/Plugins/Runtime/RigVM/Source/RigVMDeveloper/Private/RigVMModel/Nodes/RigVMInvokeEntryNode.cpp b/Engine/Plugins/Runtime/RigVM/Source/RigVMDeveloper/Private/RigVMModel/Nodes/RigVMInvokeEntryNode.cpp
--- a/Engine/Plugins/Runtime/RigVM/Source/RigVMDeveloper/Private/RigVMModel/Nodes/RigVMInvokeEntryNode.cpp
+++ b/Engine/Plugins/Runtime/RigVM/Source/RigVMDeveloper/Private/RigVMModel/Nodes/RigVMInvokeEntryNode.cpp
@@ -6,8 +6,6 @@
 
 #include UE_INLINE_GENERATED_CPP_BY_NAME(RigVMInvokeEntryNode)
 
-const FString URigVMInvokeEntryNode::EntryName = TEXT("Entry");
-
 URigVMInvokeEntryNode::URigVMInvokeEntryNode()
 {
 }
diff --git a/Engine/Plugins/Runtime/RigVM/Source/RigVMDeveloper/Private/RigVMModel/Nodes/RigVMParameterNode.cpp b/Engine/Plugins/Runtime/RigVM/Source/RigVMDeveloper/Private/RigVMModel/Nodes/RigVMParameterNode.cpp
--- a/Engine/Plugins/Runtime/RigVM/Source/RigVMDeveloper/Private/RigVMModel/Nodes/RigVMParameterNode.cpp
+++ b/Engine/Plugins/Runtime/RigVM/Source/RigVMDeveloper/Private/RigVMModel/Nodes/RigVMParameterNode.cpp
@@ -4,10 +4,6 @@
 
 #include UE_INLINE_GENERATED_CPP_BY_NAME(RigVMParameterNode)
 
-const FString URigVMParameterNode::ParameterName = TEXT("Parameter");
-const FString URigVMParameterNode::DefaultName = TEXT("Default");
-const FString URigVMParameterNode::ValueName = TEXT("Value");
-
 URigVMParameterNode::URigVMParameterNode()
 {
 }
diff --git a/Engine/Plugins/Runtime/RigVM/Source/RigVMDeveloper/Private/RigVMModel/Nodes/RigVMRerouteNode.cpp b/Engine/Plugins/Runtime/RigVM/Source/RigVMDeveloper/Private/RigVMModel/Nodes/RigVMRerouteNode.cpp
--- a/Engine/Plugins/Runtime/RigVM/Source/RigVMDeveloper/Private/RigVMModel/Nodes/RigVMRerouteNode.cpp
+++ b/Engine/Plugins/Runtime/RigVM/Source/RigVMDeveloper/Private/RigVMModel/Nodes/RigVMRerouteNode.cpp
@@ -5,9 +5,6 @@
 
 #include UE_INLINE_GENERATED_CPP_BY_NAME(RigVMRerouteNode)
 
-const FString URigVMRerouteNode::RerouteName = TEXT("Reroute");
-const FString URigVMRerouteNode::ValueName = TEXT("Value");
-
 URigVMRerouteNode::URigVMRerouteNode()
 {
 }
diff --git a/Engine/Plugins/Runtime/RigVM/Source/RigVMDeveloper/Private/RigVMModel/Nodes/RigVMSelectNode.cpp b/Engine/Plugins/Runtime/RigVM/Source/RigVMDeveloper/Private/RigVMModel/Nodes/RigVMSelectNode.cpp
--- a/Engine/Plugins/Runtime/RigVM/Source/RigVMDeveloper/Private/RigVMModel/Nodes/RigVMSelectNode.cpp
+++ b/Engine/Plugins/Runtime/RigVM/Source/RigVMDeveloper/Private/RigVMModel/Nodes/RigVMSelectNode.cpp
@@ -4,11 +4,6 @@
 
 #include UE_INLINE_GENERATED_CPP_BY_NAME(RigVMSelectNode)
 
-const FString UDEPRECATED_RigVMSelectNode::SelectName = TEXT("Select");
-const FString UDEPRECATED_RigVMSelectNode::IndexName = TEXT("Index");
-const FString UDEPRECATED_RigVMSelectNode::ValueName = TEXT("Values");
-const FString UDEPRECATED_RigVMSelectNode::ResultName = TEXT("Result");
-
 bool UDEPRECATED_RigVMSelectNode::AllowsLinksOn(const URigVMPin* InPin) const
 {
 	if(InPin->GetRootPin() == InPin)
@@ -25,7 +20,7 @@ bool UDEPRECATED_RigVMSelectNode::AllowsLinksOn(const URigVMPin* InPin) const
 FName UDEPRECATED_RigVMSelectNode::GetNotation() const
 {
 	static constexpr TCHAR Format[] = TEXT("%s(in %s,in %s,out %s)");
-	static const FName Notation = *FString::Printf(Format, *SelectName, *IndexName, *ValueName, *ResultName);
+	static const FLazyName Notation(*FString::Printf(Format, SelectName, IndexName, ValueName, ResultName));
 	return Notation;
 }
 
@@ -44,9 +39,9 @@ const FRigVMTemplate* UDEPRECATED_RigVMSelectNode::GetTemplate() const
 			return SelectNodeTemplate;
 		}
 
-		static const FName IndexFName = *IndexName;
-		static const FName ValueFName = *ValueName;
-		static const FName ResultFName = *ResultName;
+		static const FLazyName IndexFName(IndexName);
+		static const FLazyName ValueFName(ValueName);
+		static const FLazyName ResultFName(ResultName);
 
 		static TArray<FRigVMTemplateArgumentInfo> Infos;
 		if(Infos.IsEmpty())
@@ -95,7 +90,7 @@ const FRigVMTemplate* UDEPRECATED_RigVMSelectNode::GetTemplate() const
 				return Types;
 			});
 
-		SelectNodeTemplate = CachedTemplate = FRigVMRegistry::Get().GetOrAddTemplateFromArguments(*SelectName, Infos, Delegates);
+		SelectNodeTemplate = CachedTemplate = FRigVMRegistry::Get().GetOrAddTemplateFromArguments(SelectName, Infos, Delegates);
 	}
 	return CachedTemplate;
 }
diff --git a/Engine/Plugins/Runtime/RigVM/Source/RigVMDeveloper/Private/RigVMModel/Nodes/RigVMVariableNode.cpp b/Engine/Plugins/Runtime/RigVM/Source/RigVMDeveloper/Private/RigVMModel/Nodes/RigVMVariableNode.cpp
--- a/Engine/Plugins/Runtime/RigVM/Source/RigVMDeveloper/Private/RigVMModel/Nodes/RigVMVariableNode.cpp
+++ b/Engine/Plugins/Runtime/RigVM/Source/RigVMDeveloper/Private/RigVMModel/Nodes/RigVMVariableNode.cpp
@@ -6,9 +6,6 @@
 
 #include UE_INLINE_GENERATED_CPP_BY_NAME(RigVMVariableNode)
 
-const FString URigVMVariableNode::VariableName = TEXT("Variable");
-const FString URigVMVariableNode::ValueName = TEXT("Value");
-
 URigVMVariableNode::URigVMVariableNode()
 {
 }
diff --git a/Engine/Plugins/Runtime/RigVM/Source/RigVMDeveloper/Private/RigVMModel/RigVMBuildData.cpp b/Engine/Plugins/Runtime/RigVM/Source/RigVMDeveloper/Private/RigVMModel/RigVMBuildData.cpp
--- a/Engine/Plugins/Runtime/RigVM/Source/RigVMDeveloper/Private/RigVMModel/RigVMBuildData.cpp
+++ b/Engine/Plugins/Runtime/RigVM/Source/RigVMDeveloper/Private/RigVMModel/RigVMBuildData.cpp
@@ -9,7 +9,7 @@
 
 #include UE_INLINE_GENERATED_CPP_BY_NAME(RigVMBuildData)
 
-static const FName FunctionReferenceNodeDataName = TEXT("FunctionReferenceNodeData");
+static const FLazyName FunctionReferenceNodeDataName = TEXT("FunctionReferenceNodeData");
 
 // When the object system has been completely loaded, collect all the references between RigVM graphs 
 static FDelayedAutoRegisterHelper GRigVMBuildDataSingletonHelper(EDelayedRegisterRunPhase::EndOfEngineInit, []() -> void
diff --git a/Engine/Plugins/Runtime/RigVM/Source/RigVMDeveloper/Private/RigVMModel/RigVMController.cpp b/Engine/Plugins/Runtime/RigVM/Source/RigVMDeveloper/Private/RigVMModel/RigVMController.cpp
--- a/Engine/Plugins/Runtime/RigVM/Source/RigVMDeveloper/Private/RigVMModel/RigVMController.cpp
+++ b/Engine/Plugins/Runtime/RigVM/Source/RigVMDeveloper/Private/RigVMModel/RigVMController.cpp
@@ -1460,7 +1460,7 @@ URigVMUnitNode* URigVMController::AddUnitNode(UScriptStruct* InScriptStruct, TSu
 	Node->NodeTitle = InScriptStruct->GetMetaData(TEXT("DisplayName"));
 	
 	FString NodeColorMetadata;
-	InScriptStruct->GetStringMetaDataHierarchical(*URigVMNode::NodeColorName, &NodeColorMetadata);
+	InScriptStruct->GetStringMetaDataHierarchical(URigVMNode::NodeColorName, &NodeColorMetadata);
 	if (!NodeColorMetadata.IsEmpty())
 	{
 		Node->NodeColor = GetColorFromMetadata(NodeColorMetadata);
@@ -1780,14 +1780,14 @@ URigVMVariableNode* URigVMController::AddVariableNode(const FName& InVariableNam
 		AddNodePin(Node, ExecutePin);
 	}
 
-	URigVMPin* VariablePin = NewObject<URigVMPin>(Node, *URigVMVariableNode::VariableName);
+	URigVMPin* VariablePin = NewObject<URigVMPin>(Node, URigVMVariableNode::VariableName);
 	VariablePin->CPPType = RigVMTypeUtils::FNameType;
 	VariablePin->Direction = ERigVMPinDirection::Hidden;
 	VariablePin->DefaultValue = InVariableName.ToString();
 	VariablePin->CustomWidgetName = TEXT("VariableName");
 	AddNodePin(Node, VariablePin);
 
-	URigVMPin* ValuePin = NewObject<URigVMPin>(Node, *URigVMVariableNode::ValueName);
+	URigVMPin* ValuePin = NewObject<URigVMPin>(Node, URigVMVariableNode::ValueName);
 	ValuePin->CPPType = ExternalVariable.TypeName.ToString();
 	ValuePin->CPPTypeObject = ExternalVariable.TypeObject;
 	if (ValuePin->CPPTypeObject)
@@ -3031,7 +3031,7 @@ URigVMRerouteNode* URigVMController::AddRerouteNodeOnPin(const FString& InPinPat
 	URigVMRerouteNode* Node = NewObject<URigVMRerouteNode>(Graph, *Name);
 	Node->Position = InPosition;
 
-	URigVMPin* ValuePin = NewObject<URigVMPin>(Node, *URigVMRerouteNode::ValueName);
+	URigVMPin* ValuePin = NewObject<URigVMPin>(Node, URigVMRerouteNode::ValueName);
 	ConfigurePinFromPin(ValuePin, Pin);
 	ValuePin->Direction = ERigVMPinDirection::IO;
 	AddNodePin(Node, ValuePin);
@@ -12695,7 +12695,7 @@ URigVMRerouteNode* URigVMController::AddFreeRerouteNode(const FString& InCPPType
 	URigVMRerouteNode* Node = NewObject<URigVMRerouteNode>(Graph, *Name);
 	Node->Position = InPosition;
 
-	URigVMPin* ValuePin = NewObject<URigVMPin>(Node, *URigVMRerouteNode::ValueName);
+	URigVMPin* ValuePin = NewObject<URigVMPin>(Node, URigVMRerouteNode::ValueName);
 	ValuePin->CPPType = InCPPType;
 	ValuePin->CPPTypeObjectPath = InCPPTypeObjectPath;
 	ValuePin->bIsConstant = bIsConstant;
@@ -13408,7 +13408,7 @@ URigVMEnumNode* URigVMController::AddEnumNode(const FName& InCPPTypeObjectPath,
 	URigVMEnumNode* Node = NewObject<URigVMEnumNode>(Graph, *Name);
 	Node->Position = InPosition;
 
-	URigVMPin* EnumValuePin = NewObject<URigVMPin>(Node, *URigVMEnumNode::EnumValueName);
+	URigVMPin* EnumValuePin = NewObject<URigVMPin>(Node, URigVMEnumNode::EnumValueName);
 	EnumValuePin->CPPType = CPPTypeObject->GetName();
 	EnumValuePin->CPPTypeObject = CPPTypeObject;
 	EnumValuePin->CPPTypeObjectPath = InCPPTypeObjectPath;
@@ -13416,7 +13416,7 @@ URigVMEnumNode* URigVMController::AddEnumNode(const FName& InCPPTypeObjectPath,
 	EnumValuePin->DefaultValue = Enum->GetNameStringByValue(0);
 	AddNodePin(Node, EnumValuePin);
 
-	URigVMPin* EnumIndexPin = NewObject<URigVMPin>(Node, *URigVMEnumNode::EnumIndexName);
+	URigVMPin* EnumIndexPin = NewObject<URigVMPin>(Node, URigVMEnumNode::EnumIndexName);
 	EnumIndexPin->CPPType = RigVMTypeUtils::Int32Type;
 	EnumIndexPin->Direction = ERigVMPinDirection::Output;
 	EnumIndexPin->DisplayName = TEXT("Result");
@@ -13618,7 +13618,7 @@ URigVMInvokeEntryNode* URigVMController::AddInvokeEntryNode(const FName& InEntry
 	ExecutePin->Direction = ERigVMPinDirection::IO;
 	AddNodePin(Node, ExecutePin);
 
-	URigVMPin* EntryNamePin = NewObject<URigVMPin>(Node, *URigVMInvokeEntryNode::EntryName);
+	URigVMPin* EntryNamePin = NewObject<URigVMPin>(Node, URigVMInvokeEntryNode::EntryName);
 	EntryNamePin->CPPType = RigVMTypeUtils::FNameType;
 	EntryNamePin->Direction = ERigVMPinDirection::Input;
 	EntryNamePin->bIsConstant = true;
@@ -14948,7 +14948,7 @@ void URigVMController::OrphanPins(const TArray<FRepopulatePinsNodeData>& NodesPi
 #endif
 				check(Pin->IsRootPin());
 
-				const FString OrphanedName = FString::Printf(TEXT("%s%s"), *URigVMPin::OrphanPinPrefix, *Pin->GetName());
+				const FString OrphanedName = FString::Printf(TEXT("%s%s"), URigVMPin::OrphanPinPrefix, *Pin->GetName());
 				if (!NodeData.Node->FindPin(OrphanedName))
 				{
 					URigVMPin* OrphanPin = NewObject<URigVMPin>(NodeData.Node, *OrphanedName);
@@ -15020,7 +15020,7 @@ bool URigVMController::GenerateNewPinInfos(const FRigVMRegistry& Registry, URigV
 		}
 
 		FString NodeColorMetadata;
-		ScriptStruct->GetStringMetaDataHierarchical(*URigVMNode::NodeColorName, &NodeColorMetadata);
+		ScriptStruct->GetStringMetaDataHierarchical(URigVMNode::NodeColorName, &NodeColorMetadata);
 		if (!NodeColorMetadata.IsEmpty())
 		{
 			UnitNode->NodeColor = GetColorFromMetadata(NodeColorMetadata);
@@ -15531,7 +15531,7 @@ void URigVMController::RepopulatePinsOnNode(const FRigVMRegistry& Registry, cons
 #endif
 			check(Pin->IsRootPin());
 
-			const FString OrphanedName = FString::Printf(TEXT("%s%s"), *URigVMPin::OrphanPinPrefix, *Pin->GetName());
+			const FString OrphanedName = FString::Printf(TEXT("%s%s"), URigVMPin::OrphanPinPrefix, *Pin->GetName());
 			if (InNode->FindPin(OrphanedName) == nullptr)
 			{
 				Pin->DisplayName = Pin->GetFName();
@@ -15758,7 +15758,7 @@ void URigVMController::RemovePinsDuringRepopulate(URigVMNode* InNode, TArray<URi
 		if(bSetupOrphanedPins && !Pin->IsExecuteContext())
 		{
 			URigVMPin* RootPin = Pin->GetRootPin();
-			const FString OrphanedName = FString::Printf(TEXT("%s%s"), *URigVMPin::OrphanPinPrefix, *RootPin->GetName());
+			const FString OrphanedName = FString::Printf(TEXT("%s%s"), URigVMPin::OrphanPinPrefix, *RootPin->GetName());
 
 			URigVMPin* OrphanedRootPin = nullptr;
 			
@@ -17597,7 +17597,7 @@ bool URigVMController::ChangePinType(URigVMPin* InPin, TRigVMTypeIndex InTypeInd
 		{
 			FastBreakLinkedPaths(LinkedPaths);
 
-			const FString OrphanedName = FString::Printf(TEXT("%s%s"), *URigVMPin::OrphanPinPrefix, *InPin->GetName());
+			const FString OrphanedName = FString::Printf(TEXT("%s%s"), URigVMPin::OrphanPinPrefix, *InPin->GetName());
 			if(InPin->GetNode()->FindPin(OrphanedName) == nullptr)
 			{
 				URigVMPin* OrphanedPin = NewObject<URigVMPin>(InPin->GetNode(), *OrphanedName);
@@ -18522,7 +18522,7 @@ TArray<URigVMController::FLinkedPath> URigVMController::RemapLinkedPaths(
 						continue;
 					}
 
-					RemainingPinPath = FString::Printf(TEXT("%s%s"), *URigVMPin::OrphanPinPrefix, *RemainingPinPath);
+					RemainingPinPath = FString::Printf(TEXT("%s%s"), URigVMPin::OrphanPinPrefix, *RemainingPinPath);
 					PinToFind = Node->FindPin(RemainingPinPath);
 
 					if(PinToFind != nullptr)
diff --git a/Engine/Plugins/Runtime/RigVM/Source/RigVMDeveloper/Private/RigVMModel/RigVMExternalDependency.cpp b/Engine/Plugins/Runtime/RigVM/Source/RigVMDeveloper/Private/RigVMModel/RigVMExternalDependency.cpp
--- a/Engine/Plugins/Runtime/RigVM/Source/RigVMDeveloper/Private/RigVMModel/RigVMExternalDependency.cpp
+++ b/Engine/Plugins/Runtime/RigVM/Source/RigVMDeveloper/Private/RigVMModel/RigVMExternalDependency.cpp
@@ -4,10 +4,6 @@
 
 #include UE_INLINE_GENERATED_CPP_BY_NAME(RigVMExternalDependency)
 
-const FName IRigVMExternalDependencyManager::UserDefinedEnumCategory = TEXT("UserDefinedEnum");
-const FName IRigVMExternalDependencyManager::UserDefinedStructCategory = TEXT("UserDefinedStruct");
-const FName IRigVMExternalDependencyManager::RigVMGraphFunctionCategory = TEXT("RigVMGraphFunction");
-
 TArray<FName> IRigVMExternalDependencyManager::DependencyCategories = {
 	UserDefinedEnumCategory,
 	UserDefinedStructCategory,
diff --git a/Engine/Plugins/Runtime/RigVM/Source/RigVMDeveloper/Private/RigVMModel/RigVMNode.cpp b/Engine/Plugins/Runtime/RigVM/Source/RigVMDeveloper/Private/RigVMModel/RigVMNode.cpp
--- a/Engine/Plugins/Runtime/RigVM/Source/RigVMDeveloper/Private/RigVMModel/RigVMNode.cpp
+++ b/Engine/Plugins/Runtime/RigVM/Source/RigVMDeveloper/Private/RigVMModel/RigVMNode.cpp
@@ -13,8 +13,6 @@
 
 #include UE_INLINE_GENERATED_CPP_BY_NAME(RigVMNode)
 
-const FString URigVMNode::NodeColorName = TEXT("NodeColor");
-
 #if WITH_EDITOR
 TArray<int32> URigVMNode::EmptyInstructionArray;
 #endif
diff --git a/Engine/Plugins/Runtime/RigVM/Source/RigVMDeveloper/Private/RigVMModel/RigVMPin.cpp b/Engine/Plugins/Runtime/RigVM/Source/RigVMDeveloper/Private/RigVMModel/RigVMPin.cpp
--- a/Engine/Plugins/Runtime/RigVM/Source/RigVMDeveloper/Private/RigVMModel/RigVMPin.cpp
+++ b/Engine/Plugins/Runtime/RigVM/Source/RigVMDeveloper/Private/RigVMModel/RigVMPin.cpp
@@ -54,7 +54,6 @@ URigVMInjectionInfo::FWeakInfo URigVMInjectionInfo::GetWeakInfo() const
 
 const URigVMPin::FPinOverrideMap URigVMPin::EmptyPinOverrideMap;
 const URigVMPin::FPinOverride URigVMPin::EmptyPinOverride = URigVMPin::FPinOverride(FRigVMASTProxy(), EmptyPinOverrideMap);
-const FString URigVMPin::OrphanPinPrefix = TEXT("Orphan::");
 
 bool URigVMPin::SplitPinPathAtStart(const FString& InPinPath, FString& LeftMost, FString& Right)
 {
diff --git a/Engine/Plugins/Runtime/RigVM/Source/RigVMDeveloper/Public/EdGraph/RigVMEdGraphSchema.h b/Engine/Plugins/Runtime/RigVM/Source/RigVMDeveloper/Public/EdGraph/RigVMEdGraphSchema.h
--- a/Engine/Plugins/Runtime/RigVM/Source/RigVMDeveloper/Public/EdGraph/RigVMEdGraphSchema.h
+++ b/Engine/Plugins/Runtime/RigVM/Source/RigVMDeveloper/Public/EdGraph/RigVMEdGraphSchema.h
@@ -80,7 +80,11 @@ struct RIGVMDEVELOPER_API FRigVMEdGraphSchemaAction_LocalVar : public FEdGraphSc
 public:
 
 	// Simple type info
-	static FName StaticGetTypeId() {static FName Type("FRigVMEdGraphSchemaAction_LocalVar"); return Type;}
+	static FName StaticGetTypeId()
+	{
+		static const FLazyName Type("FRigVMEdGraphSchemaAction_LocalVar");
+		return Type;
+	}
 	virtual FName GetTypeId() const override { return StaticGetTypeId(); } 
 
 	FRigVMEdGraphSchemaAction_LocalVar()
@@ -117,7 +121,11 @@ struct RIGVMDEVELOPER_API FRigVMEdGraphSchemaAction_PromoteToVariable : public F
 public:
 
 	// Simple type info
-	static FName StaticGetTypeId() {static FName Type("FRigVMEdGraphSchemaAction_PromoteToVariable"); return Type;}
+	static FName StaticGetTypeId()
+	{
+		static const FLazyName Type("FRigVMEdGraphSchemaAction_PromoteToVariable");
+		return Type;
+	}
 	virtual FName GetTypeId() const override { return StaticGetTypeId(); } 
 
 	FRigVMEdGraphSchemaAction_PromoteToVariable()
@@ -149,7 +157,11 @@ struct RIGVMDEVELOPER_API FRigVMEdGraphSchemaAction_PromoteToExposedPin : public
 public:
 
 	// Simple type info
-	static FName StaticGetTypeId() {static FName Type("FRigVMEdGraphSchemaAction_PromoteToExposedPin"); return Type;}
+	static FName StaticGetTypeId()
+	{
+		static const FLazyName Type("FRigVMEdGraphSchemaAction_PromoteToExposedPin");
+		return Type;
+	}
 	virtual FName GetTypeId() const override { return StaticGetTypeId(); } 
 
 	FRigVMEdGraphSchemaAction_PromoteToExposedPin()
@@ -179,7 +191,11 @@ struct RIGVMDEVELOPER_API FRigVMEdGraphSchemaAction_Event : public FEdGraphSchem
 public:
 
 	// Simple type info
-	static FName StaticGetTypeId() {static FName Type("FRigVMEdGraphSchemaAction_Event"); return Type;}
+	static FName StaticGetTypeId()
+	{
+		static const FLazyName Type("FRigVMEdGraphSchemaAction_Event");
+		return Type;
+	}
 	virtual FName GetTypeId() const override { return StaticGetTypeId(); } 
 
 	FRigVMEdGraphSchemaAction_Event()
@@ -247,12 +263,12 @@ class RIGVMDEVELOPER_API URigVMEdGraphSchema : public UEdGraphSchema
 
 public:
 	/** Name constants */
-	static const FName GraphName_RigVM;
+	static inline const FLazyName GraphName_RigVM = FLazyName(TEXT("RigVM"));
 
 public:
 	URigVMEdGraphSchema();
 
-	virtual const FName& GetRootGraphName() const { return GraphName_RigVM; }
+	virtual const FLazyName& GetRootGraphName() const { return GraphName_RigVM; }
 
 	// UEdGraphSchema interface
 	virtual void GetGraphContextActions(FGraphContextMenuBuilder& ContextMenuBuilder) const override;
diff --git a/Engine/Plugins/Runtime/RigVM/Source/RigVMDeveloper/Public/RigVMBlueprint.h b/Engine/Plugins/Runtime/RigVM/Source/RigVMDeveloper/Public/RigVMBlueprint.h
--- a/Engine/Plugins/Runtime/RigVM/Source/RigVMDeveloper/Public/RigVMBlueprint.h
+++ b/Engine/Plugins/Runtime/RigVM/Source/RigVMDeveloper/Public/RigVMBlueprint.h
@@ -258,13 +258,13 @@ public:
 
 #if WITH_EDITOR
 	/** Returns true if a given panel node factory is compatible this blueprint */
-	virtual const FName& GetPanelNodeFactoryName() const;
+	virtual const FLazyName& GetPanelNodeFactoryName() const;
 
 	/** Returns true if a given panel pin factory is compatible this blueprint */
-	virtual const FName& GetPanelPinFactoryName() const;
+	virtual const FLazyName& GetPanelPinFactoryName() const;
 
-	static const FName RigVMPanelNodeFactoryName;
-	static const FName RigVMPanelPinFactoryName;
+	static inline const FLazyName RigVMPanelNodeFactoryName = FLazyName(TEXT("FRigVMEdGraphPanelNodeFactory"));
+	static inline const FLazyName RigVMPanelPinFactoryName = FLazyName(TEXT("FRigVMEdGraphPanelPinFactory"));
 
 	/** Returns the editor module to be used for this blueprint */
 	virtual IRigVMEditorModule* GetEditorModule() const;
diff --git a/Engine/Plugins/Runtime/RigVM/Source/RigVMDeveloper/Public/RigVMDeveloperTypeUtils.h b/Engine/Plugins/Runtime/RigVM/Source/RigVMDeveloper/Public/RigVMDeveloperTypeUtils.h
--- a/Engine/Plugins/Runtime/RigVM/Source/RigVMDeveloper/Public/RigVMDeveloperTypeUtils.h
+++ b/Engine/Plugins/Runtime/RigVM/Source/RigVMDeveloper/Public/RigVMDeveloperTypeUtils.h
@@ -54,11 +54,11 @@ namespace RigVMTypeUtils
 	
 	RIGVMDEVELOPER_API const FRigVMFunction* GetCastForTypeIndices(const TRigVMTypeIndex& InSourceTypeIndex, const TRigVMTypeIndex& InTargetTypeIndex);
 
-	RIGVMDEVELOPER_API const FName& GetCastTemplateValueName();
+	RIGVMDEVELOPER_API const FLazyName& GetCastTemplateValueName();
 
-	RIGVMDEVELOPER_API const FName& GetCastTemplateResultName();
+	RIGVMDEVELOPER_API const FLazyName& GetCastTemplateResultName();
 
-	RIGVMDEVELOPER_API const FName& GetCastTemplateNotation();
+	RIGVMDEVELOPER_API const FLazyName& GetCastTemplateNotation();
 
 	// Get a user-facing name for an argument type
 	RIGVMDEVELOPER_API FText GetDisplayTextForArgumentType(const FRigVMTemplateArgumentType& InType);
diff --git a/Engine/Plugins/Runtime/RigVM/Source/RigVMDeveloper/Public/RigVMModel/Nodes/RigVMEnumNode.h b/Engine/Plugins/Runtime/RigVM/Source/RigVMDeveloper/Public/RigVMModel/Nodes/RigVMEnumNode.h
--- a/Engine/Plugins/Runtime/RigVM/Source/RigVMDeveloper/Public/RigVMModel/Nodes/RigVMEnumNode.h
+++ b/Engine/Plugins/Runtime/RigVM/Source/RigVMDeveloper/Public/RigVMModel/Nodes/RigVMEnumNode.h
@@ -53,9 +53,9 @@ public:
 
 private:
 
-	static const FString EnumName;
-	static const FString EnumValueName;
-	static const FString EnumIndexName;
+	static const TCHAR* EnumName;
+	static const inline TCHAR* EnumValueName = TEXT("EnumValue");
+	static const inline TCHAR* EnumIndexName = TEXT("EnumIndex");
 	
 	friend class URigVMController;
 	friend class URigVMCompiler;
diff --git a/Engine/Plugins/Runtime/RigVM/Source/RigVMDeveloper/Public/RigVMModel/Nodes/RigVMIfNode.h b/Engine/Plugins/Runtime/RigVM/Source/RigVMDeveloper/Public/RigVMModel/Nodes/RigVMIfNode.h
--- a/Engine/Plugins/Runtime/RigVM/Source/RigVMDeveloper/Public/RigVMModel/Nodes/RigVMIfNode.h
+++ b/Engine/Plugins/Runtime/RigVM/Source/RigVMDeveloper/Public/RigVMModel/Nodes/RigVMIfNode.h
@@ -34,11 +34,11 @@ public:
 	virtual FLinearColor GetNodeColor() const override { return FLinearColor::Black; }
 private:
 
-	static const FString IfName;
-	static const FString ConditionName;
-	static const FString TrueName;
-	static const FString FalseName;
-	static const FString ResultName;
+	static const inline TCHAR* IfName = TEXT("If");
+	static const inline TCHAR* ConditionName = TEXT("Condition");
+	static const inline TCHAR* TrueName = TEXT("True");
+	static const inline TCHAR* FalseName = TEXT("False");
+	static const inline TCHAR* ResultName = TEXT("Result");
 
 	friend class URigVMController;
 	friend class URigVMCompiler;
diff --git a/Engine/Plugins/Runtime/RigVM/Source/RigVMDeveloper/Public/RigVMModel/Nodes/RigVMInvokeEntryNode.h b/Engine/Plugins/Runtime/RigVM/Source/RigVMDeveloper/Public/RigVMModel/Nodes/RigVMInvokeEntryNode.h
--- a/Engine/Plugins/Runtime/RigVM/Source/RigVMDeveloper/Public/RigVMModel/Nodes/RigVMInvokeEntryNode.h
+++ b/Engine/Plugins/Runtime/RigVM/Source/RigVMDeveloper/Public/RigVMModel/Nodes/RigVMInvokeEntryNode.h
@@ -32,7 +32,7 @@ public:
 
 private:
 
-	static const FString EntryName;
+	static const inline TCHAR* EntryName = TEXT("Entry");
 
 	URigVMPin* GetEntryNamePin() const;
 	
diff --git a/Engine/Plugins/Runtime/RigVM/Source/RigVMDeveloper/Public/RigVMModel/Nodes/RigVMParameterNode.h b/Engine/Plugins/Runtime/RigVM/Source/RigVMDeveloper/Public/RigVMModel/Nodes/RigVMParameterNode.h
--- a/Engine/Plugins/Runtime/RigVM/Source/RigVMDeveloper/Public/RigVMModel/Nodes/RigVMParameterNode.h
+++ b/Engine/Plugins/Runtime/RigVM/Source/RigVMDeveloper/Public/RigVMModel/Nodes/RigVMParameterNode.h
@@ -132,9 +132,9 @@ private:
 
 	virtual bool ContributesToResult() const override { return !IsInput(); }
 
-	static const FString ParameterName;
-	static const FString DefaultName;
-	static const FString ValueName;
+	static const inline TCHAR* ParameterName = TEXT("Parameter");
+	static const inline TCHAR* DefaultName = TEXT("Default");
+	static const inline TCHAR* ValueName = TEXT("Value");
 
 	friend class URigVMController;
 	friend class URigVMCompiler;
diff --git a/Engine/Plugins/Runtime/RigVM/Source/RigVMDeveloper/Public/RigVMModel/Nodes/RigVMRerouteNode.h b/Engine/Plugins/Runtime/RigVM/Source/RigVMDeveloper/Public/RigVMModel/Nodes/RigVMRerouteNode.h
--- a/Engine/Plugins/Runtime/RigVM/Source/RigVMDeveloper/Public/RigVMModel/Nodes/RigVMRerouteNode.h
+++ b/Engine/Plugins/Runtime/RigVM/Source/RigVMDeveloper/Public/RigVMModel/Nodes/RigVMRerouteNode.h
@@ -32,8 +32,8 @@ public:
 
 private:
 
-	static const FString RerouteName;
-	static const FString ValueName;
+	static const inline TCHAR* RerouteName = TEXT("Reroute");
+	static const inline TCHAR* ValueName = TEXT("Value");
 
 	friend class URigVMController;
 	friend class URigVMCompiler;
diff --git a/Engine/Plugins/Runtime/RigVM/Source/RigVMDeveloper/Public/RigVMModel/Nodes/RigVMSelectNode.h b/Engine/Plugins/Runtime/RigVM/Source/RigVMDeveloper/Public/RigVMModel/Nodes/RigVMSelectNode.h
--- a/Engine/Plugins/Runtime/RigVM/Source/RigVMDeveloper/Public/RigVMModel/Nodes/RigVMSelectNode.h
+++ b/Engine/Plugins/Runtime/RigVM/Source/RigVMDeveloper/Public/RigVMModel/Nodes/RigVMSelectNode.h
@@ -40,10 +40,10 @@ protected:
 
 private:
 
-	static const FString SelectName;
-	static const FString IndexName;
-	static const FString ValueName;
-	static const FString ResultName;
+	static const inline TCHAR* SelectName = TEXT("Select");
+	static const inline TCHAR* IndexName = TEXT("Index");
+	static const inline TCHAR* ValueName = TEXT("Values");
+	static const inline TCHAR* ResultName = TEXT("Result");
 
 	friend class URigVMController;
 	friend class URigVMCompiler;
diff --git a/Engine/Plugins/Runtime/RigVM/Source/RigVMDeveloper/Public/RigVMModel/Nodes/RigVMVariableNode.h b/Engine/Plugins/Runtime/RigVM/Source/RigVMDeveloper/Public/RigVMModel/Nodes/RigVMVariableNode.h
--- a/Engine/Plugins/Runtime/RigVM/Source/RigVMDeveloper/Public/RigVMModel/Nodes/RigVMVariableNode.h
+++ b/Engine/Plugins/Runtime/RigVM/Source/RigVMDeveloper/Public/RigVMModel/Nodes/RigVMVariableNode.h
@@ -68,8 +68,8 @@ public:
 
 private:
 
-	static const FString VariableName;
-	static const FString ValueName;
+	static const inline TCHAR* VariableName = TEXT("Variable");
+	static const inline TCHAR* ValueName = TEXT("Value");
 
 	URigVMPin* GetVariableNamePin() const;
 	URigVMPin* GetValuePin() const;
diff --git a/Engine/Plugins/Runtime/RigVM/Source/RigVMDeveloper/Public/RigVMModel/RigVMControllerActions.h b/Engine/Plugins/Runtime/RigVM/Source/RigVMDeveloper/Public/RigVMModel/RigVMControllerActions.h
--- a/Engine/Plugins/Runtime/RigVM/Source/RigVMDeveloper/Public/RigVMModel/RigVMControllerActions.h
+++ b/Engine/Plugins/Runtime/RigVM/Source/RigVMDeveloper/Public/RigVMModel/RigVMControllerActions.h
@@ -118,12 +118,12 @@ struct FRigVMBaseAction
 
 public:
 	
-	inline static const FString RedoPrefix = TEXT("Redo");
-	inline static const FString UndoPrefix = TEXT("Undo");
-	inline static const FString AddActionPrefix = TEXT("Add Action");
-	inline static const FString BeginActionPrefix = TEXT("Begin Action");
-	inline static const FString EndActionPrefix = TEXT("End Action");
-	inline static const FString CancelActionPrefix = TEXT("Cancel Action");
+	static const inline TCHAR* RedoPrefix = TEXT("Redo");
+	static const inline TCHAR* UndoPrefix = TEXT("Undo");
+	static const inline TCHAR* AddActionPrefix = TEXT("Add Action");
+	static const inline TCHAR* BeginActionPrefix = TEXT("Begin Action");
+	static const inline TCHAR* EndActionPrefix = TEXT("End Action");
+	static const inline TCHAR* CancelActionPrefix = TEXT("Cancel Action");
 
 	// Default constructor
 	FRigVMBaseAction(URigVMController* InController)
diff --git a/Engine/Plugins/Runtime/RigVM/Source/RigVMDeveloper/Public/RigVMModel/RigVMExternalDependency.h b/Engine/Plugins/Runtime/RigVM/Source/RigVMDeveloper/Public/RigVMModel/RigVMExternalDependency.h
--- a/Engine/Plugins/Runtime/RigVM/Source/RigVMDeveloper/Public/RigVMModel/RigVMExternalDependency.h
+++ b/Engine/Plugins/Runtime/RigVM/Source/RigVMDeveloper/Public/RigVMModel/RigVMExternalDependency.h
@@ -68,9 +68,9 @@ public:
 	virtual TArray<FRigVMExternalDependency> GetExternalDependenciesForCategory(const FName& InCategory) const = 0;
 	TArray<FRigVMExternalDependency> GetAllExternalDependencies() const;
 
-	static const FName UserDefinedEnumCategory;
-	static const FName UserDefinedStructCategory;
-	static const FName RigVMGraphFunctionCategory;
+	static inline const FLazyName UserDefinedEnumCategory = FLazyName(TEXT("UserDefinedEnum"));
+	static inline const FLazyName UserDefinedStructCategory = FLazyName(TEXT("UserDefinedStruct"));
+	static inline const FLazyName RigVMGraphFunctionCategory = FLazyName(TEXT("RigVMGraphFunction"));
 
 protected:
 
diff --git a/Engine/Plugins/Runtime/RigVM/Source/RigVMDeveloper/Public/RigVMModel/RigVMNode.h b/Engine/Plugins/Runtime/RigVM/Source/RigVMDeveloper/Public/RigVMModel/RigVMNode.h
--- a/Engine/Plugins/Runtime/RigVM/Source/RigVMDeveloper/Public/RigVMModel/RigVMNode.h
+++ b/Engine/Plugins/Runtime/RigVM/Source/RigVMDeveloper/Public/RigVMModel/RigVMNode.h
@@ -328,7 +328,7 @@ public:
 
 private:
 
-	static const FString NodeColorName;
+	static const inline TCHAR* NodeColorName = TEXT("NodeColor");
 
 	bool IsLinkedToRecursive(URigVMPin* InPin, URigVMNode* InNode) const;
 	void GetLinkedNodesRecursive(URigVMPin* InPin, bool bLookForSources, TArray<URigVMNode*>& OutNodes) const;
diff --git a/Engine/Plugins/Runtime/RigVM/Source/RigVMDeveloper/Public/RigVMModel/RigVMPin.h b/Engine/Plugins/Runtime/RigVM/Source/RigVMDeveloper/Public/RigVMModel/RigVMPin.h
--- a/Engine/Plugins/Runtime/RigVM/Source/RigVMDeveloper/Public/RigVMModel/RigVMPin.h
+++ b/Engine/Plugins/Runtime/RigVM/Source/RigVMDeveloper/Public/RigVMModel/RigVMPin.h
@@ -555,7 +555,7 @@ private:
 	mutable FString LastKnownCPPType;
 	mutable TRigVMTypeIndex LastKnownTypeIndex;
 
-	static const FString OrphanPinPrefix;
+	static const inline TCHAR* OrphanPinPrefix = TEXT("Orphan::");
 
 	friend class URigVMController;
 	friend class URigVMGraph;
diff --git a/Engine/Plugins/Runtime/RigVM/Source/RigVMEditor/Private/EdGraph/NodeSpawners/RigVMEdGraphEnumNodeSpawner.cpp b/Engine/Plugins/Runtime/RigVM/Source/RigVMEditor/Private/EdGraph/NodeSpawners/RigVMEdGraphEnumNodeSpawner.cpp
--- a/Engine/Plugins/Runtime/RigVM/Source/RigVMEditor/Private/EdGraph/NodeSpawners/RigVMEdGraphEnumNodeSpawner.cpp
+++ b/Engine/Plugins/Runtime/RigVM/Source/RigVMEditor/Private/EdGraph/NodeSpawners/RigVMEdGraphEnumNodeSpawner.cpp
@@ -59,7 +59,7 @@ UEdGraphNode* URigVMEdGraphEnumNodeSpawner::Invoke(UEdGraph* ParentGraph, FBindi
 	if (bIsTemplateNode)
 	{
 		TArray<FPinInfo> Pins;
-		Pins.Emplace(*URigVMEnumNode::EnumValueName, ERigVMPinDirection::Output, RigVMTypeUtils::Int32TypeName, nullptr);
+		Pins.Emplace(URigVMEnumNode::EnumValueName, ERigVMPinDirection::Output, RigVMTypeUtils::Int32TypeName, nullptr);
 		return SpawnTemplateNode(ParentGraph, Pins);
 	}
 
@@ -78,11 +78,9 @@ UEdGraphNode* URigVMEdGraphEnumNodeSpawner::Invoke(UEdGraph* ParentGraph, FBindi
 
 	URigVMController* Controller = RigBlueprint->GetController(ParentGraph);
 
-	FName Name = *URigVMEnumNode::EnumName;
+	Controller->OpenUndoBracket(FString::Printf(TEXT("Add '%s' Node"), URigVMEnumNode::EnumName));
 
-	Controller->OpenUndoBracket(FString::Printf(TEXT("Add '%s' Node"), *Name.ToString()));
-
-	if (URigVMEnumNode* ModelNode = Controller->AddEnumNode(*Enum->GetPathName(), Location, Name.ToString(), true, true))
+	if (URigVMEnumNode* ModelNode = Controller->AddEnumNode(*Enum->GetPathName(), Location, URigVMEnumNode::EnumName, true, true))
 	{
 		NewNode = Cast<URigVMEdGraphNode>(RigGraph->FindNodeForModelNodeName(ModelNode->GetFName()));
 
diff --git a/Engine/Plugins/Runtime/RigVM/Source/RigVMEditor/Private/EdGraph/NodeSpawners/RigVMEdGraphInvokeEntryNodeSpawner.cpp b/Engine/Plugins/Runtime/RigVM/Source/RigVMEditor/Private/EdGraph/NodeSpawners/RigVMEdGraphInvokeEntryNodeSpawner.cpp
--- a/Engine/Plugins/Runtime/RigVM/Source/RigVMEditor/Private/EdGraph/NodeSpawners/RigVMEdGraphInvokeEntryNodeSpawner.cpp
+++ b/Engine/Plugins/Runtime/RigVM/Source/RigVMEditor/Private/EdGraph/NodeSpawners/RigVMEdGraphInvokeEntryNodeSpawner.cpp
@@ -90,7 +90,7 @@ UEdGraphNode* URigVMEdGraphInvokeEntryNodeSpawner::Invoke(UEdGraph* ParentGraph,
 		TArray<FPinInfo> Pins;
 
 		static UScriptStruct* ExecuteScriptStruct = FRigVMExecuteContext::StaticStruct();
-		static const FName ExecuteStructName = *ExecuteScriptStruct->GetStructCPPName();
+		static const FLazyName ExecuteStructName(*ExecuteScriptStruct->GetStructCPPName());
 		Pins.Emplace(FRigVMStruct::ExecuteName, ERigVMPinDirection::IO, ExecuteStructName, ExecuteScriptStruct);
 
 		return SpawnTemplateNode(ParentGraph, Pins, *NodeName);
diff --git a/Engine/Plugins/Runtime/RigVM/Source/RigVMEditor/Private/EdGraph/NodeSpawners/RigVMEdGraphTemplateNodeSpawner.cpp b/Engine/Plugins/Runtime/RigVM/Source/RigVMEditor/Private/EdGraph/NodeSpawners/RigVMEdGraphTemplateNodeSpawner.cpp
--- a/Engine/Plugins/Runtime/RigVM/Source/RigVMEditor/Private/EdGraph/NodeSpawners/RigVMEdGraphTemplateNodeSpawner.cpp
+++ b/Engine/Plugins/Runtime/RigVM/Source/RigVMEditor/Private/EdGraph/NodeSpawners/RigVMEdGraphTemplateNodeSpawner.cpp
@@ -174,7 +174,7 @@ UEdGraphNode* URigVMEdGraphTemplateNodeSpawner::Invoke(UEdGraph* ParentGraph, FB
 				const FRigVMExecuteArgument* Argument = Template->GetExecuteArgument(Index, Context);
 				check(Argument);
 				static UScriptStruct* ExecuteScriptStruct = FRigVMExecuteContext::StaticStruct();
-				static const FName ExecuteStructName = *ExecuteScriptStruct->GetStructCPPName();
+				static const FLazyName ExecuteStructName(*ExecuteScriptStruct->GetStructCPPName());
 				Pins.Emplace(Argument->Name, Argument->Direction, ExecuteStructName, ExecuteScriptStruct);
 			}
 
diff --git a/Engine/Plugins/Runtime/RigVM/Source/RigVMEditor/Private/EdGraph/NodeSpawners/RigVMEdGraphUnitNodeSpawner.cpp b/Engine/Plugins/Runtime/RigVM/Source/RigVMEditor/Private/EdGraph/NodeSpawners/RigVMEdGraphUnitNodeSpawner.cpp
--- a/Engine/Plugins/Runtime/RigVM/Source/RigVMEditor/Private/EdGraph/NodeSpawners/RigVMEdGraphUnitNodeSpawner.cpp
+++ b/Engine/Plugins/Runtime/RigVM/Source/RigVMEditor/Private/EdGraph/NodeSpawners/RigVMEdGraphUnitNodeSpawner.cpp
@@ -185,7 +185,7 @@ URigVMEdGraphNode* URigVMEdGraphUnitNodeSpawner::SpawnNode(UEdGraph* ParentGraph
 				if(UScriptStruct* ScriptStruct = Cast<UScriptStruct>(Pin.CPPTypeObject))
 				{
 					static UScriptStruct* ExecuteScriptStruct = FRigVMExecuteContext::StaticStruct();
-					static const FName ExecuteStructName = *ExecuteScriptStruct->GetStructCPPName();
+					static const FLazyName ExecuteStructName(*ExecuteScriptStruct->GetStructCPPName());
 
 					if(ScriptStruct->IsChildOf(ExecuteScriptStruct))
 					{
diff --git a/Engine/Plugins/Runtime/RigVM/Source/RigVMEditor/Private/EdGraph/NodeSpawners/RigVMEdGraphVariableNodeSpawner.cpp b/Engine/Plugins/Runtime/RigVM/Source/RigVMEditor/Private/EdGraph/NodeSpawners/RigVMEdGraphVariableNodeSpawner.cpp
--- a/Engine/Plugins/Runtime/RigVM/Source/RigVMEditor/Private/EdGraph/NodeSpawners/RigVMEdGraphVariableNodeSpawner.cpp
+++ b/Engine/Plugins/Runtime/RigVM/Source/RigVMEditor/Private/EdGraph/NodeSpawners/RigVMEdGraphVariableNodeSpawner.cpp
@@ -156,19 +156,17 @@ UEdGraphNode* URigVMEdGraphVariableNodeSpawner::Invoke(UEdGraph* ParentGraph, FB
 		static const FString SetterPrefix = TEXT("Setter");
 		NodeName = FString::Printf(VariableNodeNameFormat, bIsGetter ? *GetterPrefix : *SetterPrefix, *ExternalVariable.Name.ToString());
 
-		static const FName ValueName = *URigVMVariableNode::ValueName;
-
 		TArray<FPinInfo> Pins;
 
 		if(!bIsGetter)
 		{
 			static UScriptStruct* ExecuteScriptStruct = FRigVMExecuteContext::StaticStruct();
-			static const FName ExecuteStructName = *ExecuteScriptStruct->GetStructCPPName();
+			static const FLazyName ExecuteStructName(*ExecuteScriptStruct->GetStructCPPName());
 			Pins.Emplace(FRigVMStruct::ExecuteName, ERigVMPinDirection::IO, ExecuteStructName, ExecuteScriptStruct);
 		}
 		
 		Pins.Emplace(
-			ValueName,
+			URigVMVariableNode::ValueName,
 			bIsGetter ? ERigVMPinDirection::Output : ERigVMPinDirection::Input,
 			ExternalVariable.TypeName,
 			ExternalVariable.TypeObject);
diff --git a/Engine/Plugins/Runtime/RigVM/Source/RigVMEditor/Private/Editor/RigVMDetailsViewWrapperObject.cpp b/Engine/Plugins/Runtime/RigVM/Source/RigVMEditor/Private/Editor/RigVMDetailsViewWrapperObject.cpp
--- a/Engine/Plugins/Runtime/RigVM/Source/RigVMEditor/Private/Editor/RigVMDetailsViewWrapperObject.cpp
+++ b/Engine/Plugins/Runtime/RigVM/Source/RigVMEditor/Private/Editor/RigVMDetailsViewWrapperObject.cpp
@@ -5,6 +5,7 @@
 #include "RigVMModel/Nodes/RigVMUnitNode.h"
 #include "Modules/ModuleManager.h"
 #include "Algo/Sort.h"
+#include "RigVMTypeUtils.h"
 
 #include UE_INLINE_GENERATED_CPP_BY_NAME(RigVMDetailsViewWrapperObject)
 
@@ -363,14 +364,6 @@ UClass* URigVMDetailsViewWrapperObject::GetClassForNodes(TArray<URigVMNode*> InN
 
 	for(URigVMPin* Pin : PinsToInspect)
 	{
-		static FString BoolString = TEXT("bool");
-		static FString Int32String = TEXT("int32");
-		static FString IntString = TEXT("int");
-		static FString FloatString = TEXT("float");
-		static FString DoubleString = TEXT("double");
-		static FString StringString = TEXT("FString");
-		static FString NameString = TEXT("FName");
-
 		FProperty* Property = nullptr;
 		FProperty** ElementProperty = &Property;
 		FFieldVariant PropertyOwner = WrapperClass;
@@ -388,28 +381,28 @@ UClass* URigVMDetailsViewWrapperObject::GetClassForNodes(TArray<URigVMNode*> InN
 			PropertyOwner = ArrayProperty;
 		}
 
-		if(BaseCPPType.Equals(BoolString, ESearchCase::IgnoreCase))
+		if(BaseCPPType.Equals(RigVMTypeUtils::BoolType, ESearchCase::IgnoreCase))
 		{
 			(*ElementProperty) = new FBoolProperty(PropertyOwner, Pin->GetFName(), RF_Public);;
 		}
-		else if(BaseCPPType.Equals(Int32String, ESearchCase::IgnoreCase) ||
-			BaseCPPType.Equals(IntString, ESearchCase::IgnoreCase))
+		else if(BaseCPPType.Equals(RigVMTypeUtils::Int32Type, ESearchCase::IgnoreCase) ||
+			BaseCPPType.Equals(RigVMTypeUtils::IntType, ESearchCase::IgnoreCase))
 		{
 			(*ElementProperty) = new FIntProperty(PropertyOwner, Pin->GetFName(), RF_Public);;
 		}
-		else if(BaseCPPType.Equals(FloatString, ESearchCase::IgnoreCase))
+		else if(BaseCPPType.Equals(RigVMTypeUtils::FloatType, ESearchCase::IgnoreCase))
 		{
 			(*ElementProperty) = new FFloatProperty(PropertyOwner, Pin->GetFName(), RF_Public);;
 		}
-		else if(BaseCPPType.Equals(DoubleString, ESearchCase::IgnoreCase))
+		else if(BaseCPPType.Equals(RigVMTypeUtils::DoubleType, ESearchCase::IgnoreCase))
 		{
 			(*ElementProperty) = new FDoubleProperty(PropertyOwner, Pin->GetFName(), RF_Public);;
 		}
-		else if(BaseCPPType.Equals(StringString, ESearchCase::IgnoreCase))
+		else if(BaseCPPType.Equals(RigVMTypeUtils::FStringType, ESearchCase::IgnoreCase))
 		{
 			(*ElementProperty) = new FStrProperty(PropertyOwner, Pin->GetFName(), RF_Public);;
 		}
-		else if(BaseCPPType.Equals(NameString, ESearchCase::IgnoreCase))
+		else if(BaseCPPType.Equals(RigVMTypeUtils::FNameType, ESearchCase::IgnoreCase))
 		{
 			(*ElementProperty) = new FNameProperty(PropertyOwner, Pin->GetFName(), RF_Public);;
 		}
diff --git a/Engine/Plugins/Runtime/RigVM/Source/RigVMEditor/Private/Editor/RigVMEditor.cpp b/Engine/Plugins/Runtime/RigVM/Source/RigVMEditor/Private/Editor/RigVMEditor.cpp
--- a/Engine/Plugins/Runtime/RigVM/Source/RigVMEditor/Private/Editor/RigVMEditor.cpp
+++ b/Engine/Plugins/Runtime/RigVM/Source/RigVMEditor/Private/Editor/RigVMEditor.cpp
@@ -48,8 +48,6 @@
 
 #define LOCTEXT_NAMESPACE "RigVMEditor"
 
-const FName FRigVMEditorModes::RigVMEditorMode = TEXT("RigVM");
-
 FRigVMEditor::FRigVMEditor()
 	: bAnyErrorsLeft(false)
 	, KnownInstructionLimitWarnings()
@@ -415,7 +413,7 @@ void FRigVMEditor::HandleAssetRequestClose(UObject* InObject, EAssetEditorCloseR
 
 const FName FRigVMEditor::GetEditorAppName() const
 {
-	static const FName AppName(TEXT("RigVMEditorApp"));
+	static const FLazyName AppName(TEXT("RigVMEditorApp"));
 	return AppName;
 }
 
@@ -2296,8 +2294,8 @@ TRigVMTypeIndex FRigVMEditor::OnRequestPinTypeSelectionDialog(const TArray<TRigV
 										TRigVMTypeIndex TypeIndex = InTypes[TypeNameToIndex.FindChecked(*InItem.Get())];
 										const FRigVMTemplateArgumentType Type = FRigVMRegistry::Get().GetType(TypeIndex);
 										const bool bIsArray = Type.IsArray();
-										static const FName TypeIcon(TEXT("Kismet.VariableList.TypeIcon"));
-										static const FName ArrayTypeIcon(TEXT("Kismet.VariableList.ArrayTypeIcon"));
+										static const FLazyName TypeIcon(TEXT("Kismet.VariableList.TypeIcon"));
+										static const FLazyName ArrayTypeIcon(TEXT("Kismet.VariableList.ArrayTypeIcon"));
 
 										const FEdGraphPinType PinType = RigVMTypeUtils::PinTypeFromTypeIndex(TypeIndex);
 										const URigVMEdGraphSchema* Schema = CastChecked<URigVMEdGraphSchema>(GetRigVMBlueprint()->GetRigVMEdGraphSchemaClass()->GetDefaultObject());
diff --git a/Engine/Plugins/Runtime/RigVM/Source/RigVMEditor/Private/Editor/RigVMExecutionStackTabSummoner.cpp b/Engine/Plugins/Runtime/RigVM/Source/RigVMEditor/Private/Editor/RigVMExecutionStackTabSummoner.cpp
--- a/Engine/Plugins/Runtime/RigVM/Source/RigVMEditor/Private/Editor/RigVMExecutionStackTabSummoner.cpp
+++ b/Engine/Plugins/Runtime/RigVM/Source/RigVMEditor/Private/Editor/RigVMExecutionStackTabSummoner.cpp
@@ -6,8 +6,6 @@
 
 #define LOCTEXT_NAMESPACE "RigVMExecutionStackTabSummoner"
 
-const FName FRigVMExecutionStackTabSummoner::TabID(TEXT("Execution Stack"));
-
 FRigVMExecutionStackTabSummoner::FRigVMExecutionStackTabSummoner(const TSharedRef<FRigVMEditor>& InRigVMEditor)
 	: FWorkflowTabFactory(TabID, InRigVMEditor)
 	, RigVMEditor(InRigVMEditor)
diff --git a/Engine/Plugins/Runtime/RigVM/Source/RigVMEditor/Private/RigVMEditorModule.cpp b/Engine/Plugins/Runtime/RigVM/Source/RigVMEditor/Private/RigVMEditorModule.cpp
--- a/Engine/Plugins/Runtime/RigVM/Source/RigVMEditor/Private/RigVMEditorModule.cpp
+++ b/Engine/Plugins/Runtime/RigVM/Source/RigVMEditor/Private/RigVMEditorModule.cpp
@@ -1443,7 +1443,7 @@ void FRigVMEditorModule::GetPinInjectedNodesContextMenuActions(IRigVMClientHost*
 				FString TemplateName;
 				if (URigVMUnitNode* UnitNode = Cast<URigVMUnitNode>(Injection->Node))
 				{
-					if (UnitNode->GetScriptStruct()->GetStringMetaDataHierarchical(FRigVMStruct::TemplateNameMetaName, &TemplateName))
+					if (UnitNode->GetScriptStruct()->GetStringMetaDataHierarchical(FRigVMRegistry::TemplateNameMetaName, &TemplateName))
 					{
 						if (TemplateName == TEXT("AlphaInterp"))
 						{
@@ -1535,7 +1535,7 @@ void FRigVMEditorModule::GetPinInjectedNodesContextMenuActions(IRigVMClientHost*
 				FString TemplateName;
 				if (URigVMUnitNode* UnitNode = Cast<URigVMUnitNode>(Injection->Node))
 				{
-					if (UnitNode->GetScriptStruct()->GetStringMetaDataHierarchical(FRigVMStruct::TemplateNameMetaName, &TemplateName))
+					if (UnitNode->GetScriptStruct()->GetStringMetaDataHierarchical(FRigVMRegistry::TemplateNameMetaName, &TemplateName))
 					{
 						if (TemplateName == TEXT("VisualDebug"))
 						{
diff --git a/Engine/Plugins/Runtime/RigVM/Source/RigVMEditor/Private/Widgets/SRigVMGraphChangePinType.cpp b/Engine/Plugins/Runtime/RigVM/Source/RigVMEditor/Private/Widgets/SRigVMGraphChangePinType.cpp
--- a/Engine/Plugins/Runtime/RigVM/Source/RigVMEditor/Private/Widgets/SRigVMGraphChangePinType.cpp
+++ b/Engine/Plugins/Runtime/RigVM/Source/RigVMEditor/Private/Widgets/SRigVMGraphChangePinType.cpp
@@ -76,8 +76,8 @@ FText SRigVMGraphChangePinType::GetBindingText() const
 
 const FSlateBrush* SRigVMGraphChangePinType::GetBindingImage() const
 {
-	static FName TypeIcon(TEXT("Kismet.VariableList.TypeIcon"));
-	static FName ArrayTypeIcon(TEXT("Kismet.VariableList.ArrayTypeIcon"));
+	static const FLazyName TypeIcon(TEXT("Kismet.VariableList.TypeIcon"));
+	static const FLazyName ArrayTypeIcon(TEXT("Kismet.VariableList.ArrayTypeIcon"));
 
 	if(Types.Num() > 0)
 	{
@@ -201,7 +201,7 @@ void SRigVMGraphChangePinType::FillPinTypeMenu(FMenuBuilder& MenuBuilder)
 	
 	MenuBuilder.BeginSection("PinTypes", LOCTEXT("PinTypes", "Pin Types"));
 	{
-		static FName PropertyIcon(TEXT("Kismet.VariableList.TypeIcon"));
+		static const FLazyName PropertyIcon(TEXT("Kismet.VariableList.TypeIcon"));
 		const URigVMEdGraphSchema* Schema = GetDefault<URigVMEdGraphSchema>();
 
 		const bool bHasAllTypes =
diff --git a/Engine/Plugins/Runtime/RigVM/Source/RigVMEditor/Private/Widgets/SRigVMGraphNode.cpp b/Engine/Plugins/Runtime/RigVM/Source/RigVMEditor/Private/Widgets/SRigVMGraphNode.cpp
--- a/Engine/Plugins/Runtime/RigVM/Source/RigVMEditor/Private/Widgets/SRigVMGraphNode.cpp
+++ b/Engine/Plugins/Runtime/RigVM/Source/RigVMEditor/Private/Widgets/SRigVMGraphNode.cpp
@@ -54,8 +54,8 @@ void SRigVMGraphNode::Construct( const FArguments& InArgs )
 {
 	if (CachedImg_CR_Pin_Connected == nullptr)
 	{
-		static const FName NAME_CR_Pin_Connected("ControlRig.Bug.Solid");
-		static const FName NAME_CR_Pin_Disconnected("ControlRig.Bug.Open");
+		static const FLazyName NAME_CR_Pin_Connected("ControlRig.Bug.Solid");
+		static const FLazyName NAME_CR_Pin_Disconnected("ControlRig.Bug.Open");
 		CachedImg_CR_Pin_Connected = FSlateIcon(TEXT("ControlRigEditorStyle"), NAME_CR_Pin_Connected).GetIcon();
 		CachedImg_CR_Pin_Disconnected = FSlateIcon(TEXT("ControlRigEditorStyle"), NAME_CR_Pin_Disconnected).GetIcon();
 	}
@@ -934,7 +934,7 @@ TArray<FOverlayWidgetInfo> SRigVMGraphNode::GetOverlayWidgets(bool bSelected, co
 					if (URigVMUnitNode* VisualDebugNode = Cast<URigVMUnitNode>(Injection->Node))
 					{
 						FString TemplateName;
-					   if (VisualDebugNode->GetScriptStruct()->GetStringMetaDataHierarchical(FRigVMStruct::TemplateNameMetaName, &TemplateName))
+					   if (VisualDebugNode->GetScriptStruct()->GetStringMetaDataHierarchical(FRigVMRegistry::TemplateNameMetaName, &TemplateName))
 					   {
 						   if (TemplateName == TEXT("VisualDebug"))
 						   {
diff --git a/Engine/Plugins/Runtime/RigVM/Source/RigVMEditor/Private/Widgets/SRigVMGraphPinUserDataPath.cpp b/Engine/Plugins/Runtime/RigVM/Source/RigVMEditor/Private/Widgets/SRigVMGraphPinUserDataPath.cpp
--- a/Engine/Plugins/Runtime/RigVM/Source/RigVMEditor/Private/Widgets/SRigVMGraphPinUserDataPath.cpp
+++ b/Engine/Plugins/Runtime/RigVM/Source/RigVMEditor/Private/Widgets/SRigVMGraphPinUserDataPath.cpp
@@ -130,8 +130,8 @@ const FSlateBrush* SRigVMUserDataPath::GetUserDataIcon() const
 
 const FSlateBrush* SRigVMUserDataPath::GetUserDataIcon(const UNameSpacedUserData::FUserData* InUserData) const
 {
-	static FName TypeIcon(TEXT("Kismet.VariableList.TypeIcon"));
-	static FName ArrayTypeIcon(TEXT("Kismet.VariableList.ArrayTypeIcon"));
+	static const FLazyName TypeIcon(TEXT("Kismet.VariableList.TypeIcon"));
+	static const FLazyName ArrayTypeIcon(TEXT("Kismet.VariableList.ArrayTypeIcon"));
 
 	if (InUserData)
 	{
@@ -191,7 +191,7 @@ void SRigVMUserDataPath::FillUserDataPathMenu(FMenuBuilder& InMenuBuilder, FStri
 		InMenuBuilder.BeginSection("UserData", LOCTEXT("UserData", "User Data"));
 	}
 	{
-		static FName PropertyIcon(TEXT("Kismet.VariableList.TypeIcon"));
+		static const FLazyName PropertyIcon(TEXT("Kismet.VariableList.TypeIcon"));
 		const URigVMEdGraphSchema* Schema = GetDefault<URigVMEdGraphSchema>();
 
 		const TArray<const UNameSpacedUserData::FUserData*> UserDataArray = UserDataObject->GetUserDataArray(InParentPath);
diff --git a/Engine/Plugins/Runtime/RigVM/Source/RigVMEditor/Private/Widgets/SRigVMGraphPinVariableBinding.cpp b/Engine/Plugins/Runtime/RigVM/Source/RigVMEditor/Private/Widgets/SRigVMGraphPinVariableBinding.cpp
--- a/Engine/Plugins/Runtime/RigVM/Source/RigVMEditor/Private/Widgets/SRigVMGraphPinVariableBinding.cpp
+++ b/Engine/Plugins/Runtime/RigVM/Source/RigVMEditor/Private/Widgets/SRigVMGraphPinVariableBinding.cpp
@@ -100,8 +100,8 @@ FText SRigVMGraphVariableBinding::GetBindingText() const
 
 const FSlateBrush* SRigVMGraphVariableBinding::GetBindingImage() const
 {
-	static FName TypeIcon(TEXT("Kismet.VariableList.TypeIcon"));
-	static FName ArrayTypeIcon(TEXT("Kismet.VariableList.ArrayTypeIcon"));
+	static const FLazyName TypeIcon(TEXT("Kismet.VariableList.TypeIcon"));
+	static const FLazyName ArrayTypeIcon(TEXT("Kismet.VariableList.ArrayTypeIcon"));
 
 	if (ModelPins.Num() > 0)
 	{
@@ -311,7 +311,7 @@ void SRigVMGraphVariableBinding::FillLocalVariableMenu(FMenuBuilder& MenuBuilder
 	
 	MenuBuilder.BeginSection("LocalVariables", LOCTEXT("LocalVariables", "Local Variables"));
 	{
-		static FName PropertyIcon(TEXT("Kismet.VariableList.TypeIcon"));
+		static const FLazyName PropertyIcon(TEXT("Kismet.VariableList.TypeIcon"));
 		const URigVMEdGraphSchema* Schema = GetDefault<URigVMEdGraphSchema>();
 
 		for(const FRigVMGraphVariableDescription& LocalVariable : LocalVariables)
diff --git a/Engine/Plugins/Runtime/RigVM/Source/RigVMEditor/Private/Widgets/SRigVMVariableMappingWidget.cpp b/Engine/Plugins/Runtime/RigVM/Source/RigVMEditor/Private/Widgets/SRigVMVariableMappingWidget.cpp
--- a/Engine/Plugins/Runtime/RigVM/Source/RigVMEditor/Private/Widgets/SRigVMVariableMappingWidget.cpp
+++ b/Engine/Plugins/Runtime/RigVM/Source/RigVMEditor/Private/Widgets/SRigVMVariableMappingWidget.cpp
@@ -8,8 +8,8 @@
 
 #define LOCTEXT_NAMESPACE "SRigVMVariableMappingWidget"
 
-static const FName ColumnId_VarLabel("Variable");
-static const FName ColumnID_MappingLabel("Mapping");
+static const FLazyName ColumnId_VarLabel("Variable");
+static const FLazyName ColumnID_MappingLabel("Mapping");
 
 void SRigVMVariableMappingTreeRow::Construct(const FArguments& InArgs, const TSharedRef<STableViewBase>& InOwnerTableView)
 {
diff --git a/Engine/Plugins/Runtime/RigVM/Source/RigVMEditor/Public/Editor/RigVMEditor.h b/Engine/Plugins/Runtime/RigVM/Source/RigVMEditor/Public/Editor/RigVMEditor.h
--- a/Engine/Plugins/Runtime/RigVM/Source/RigVMEditor/Public/Editor/RigVMEditor.h
+++ b/Engine/Plugins/Runtime/RigVM/Source/RigVMEditor/Public/Editor/RigVMEditor.h
@@ -15,7 +15,7 @@ DECLARE_MULTICAST_DELEGATE_TwoParams(FRigVMEditorClosed, const FRigVMEditor*, UR
 struct FRigVMEditorModes
 {
 	// Mode constants
-	static const FName RigVMEditorMode;
+	RIGVMEDITOR_API static inline const FLazyName RigVMEditorMode = FLazyName(TEXT("RigVM"));
 	static FText GetLocalizedMode(const FName InMode)
 	{
 		static TMap< FName, FText > LocModes;
diff --git a/Engine/Plugins/Runtime/RigVM/Source/RigVMEditor/Public/Editor/RigVMExecutionStackTabSummoner.h b/Engine/Plugins/Runtime/RigVM/Source/RigVMEditor/Public/Editor/RigVMExecutionStackTabSummoner.h
--- a/Engine/Plugins/Runtime/RigVM/Source/RigVMEditor/Public/Editor/RigVMExecutionStackTabSummoner.h
+++ b/Engine/Plugins/Runtime/RigVM/Source/RigVMEditor/Public/Editor/RigVMExecutionStackTabSummoner.h
@@ -10,7 +10,8 @@ class FRigVMEditor;
 struct RIGVMEDITOR_API FRigVMExecutionStackTabSummoner : public FWorkflowTabFactory
 {
 public:
-	static const FName TabID;
+
+	static inline const FLazyName TabID = FLazyName(TEXT("Execution Stack"));
 	
 public:
 	FRigVMExecutionStackTabSummoner(const TSharedRef<FRigVMEditor>& InRigVMEditor);
diff --git a/Engine/Source/Runtime/Core/Public/UObject/NameTypes.h b/Engine/Source/Runtime/Core/Public/UObject/NameTypes.h
--- a/Engine/Source/Runtime/Core/Public/UObject/NameTypes.h
+++ b/Engine/Source/Runtime/Core/Public/UObject/NameTypes.h
@@ -1610,6 +1610,11 @@ public:
 
 	CORE_API FName Resolve() const;
 
+	CORE_API FString ToString() const
+	{
+		return Resolve().ToString();
+	}
+
 private:
 	struct FLiteralOrName
 	{
