Fix issues importing RigVM pin default values for structs with native ImportTextItem implementations

Verbose logging was being picked up as errors due to no check of the log verbosity

diff --git a/Engine/Plugins/Runtime/RigVM/Source/RigVMDeveloper/Public/RigVMModel/RigVMPin.h b/Engine/Plugins/Runtime/RigVM/Source/RigVMDeveloper/Public/RigVMModel/RigVMPin.h
--- a/Engine/Plugins/Runtime/RigVM/Source/RigVMDeveloper/Public/RigVMModel/RigVMPin.h
+++ b/Engine/Plugins/Runtime/RigVM/Source/RigVMDeveloper/Public/RigVMModel/RigVMPin.h
@@ -578,6 +578,9 @@ public:
 
 	virtual void Serialize(const TCHAR* V, ELogVerbosity::Type Verbosity, const class FName& Category) override
 	{
-		NumErrors++;
+		if(Verbosity <= ELogVerbosity::Error)
+		{
+			NumErrors++;
+		}
 	}
 };
