Allow UScriptStructs to be pre-registered for use with FRigVMRegistry

In much the same way as UClass is already implemented

diff --git a/Engine/Plugins/Runtime/RigVM/Source/RigVM/Private/RigVMCore/RigVMRegistry.cpp b/Engine/Plugins/Runtime/RigVM/Source/RigVM/Private/RigVMCore/RigVMRegistry.cpp
--- a/Engine/Plugins/Runtime/RigVM/Source/RigVM/Private/RigVMCore/RigVMRegistry.cpp
+++ b/Engine/Plugins/Runtime/RigVM/Source/RigVM/Private/RigVMCore/RigVMRegistry.cpp
@@ -364,6 +364,10 @@ void FRigVMRegistry::RefreshEngineTypes_NoLock()
 		{
 			DispatchFactoriesToRegister.Add(ScriptStruct);
 		}
+		else if(AllowedStructs.Contains(ScriptStruct))
+		{
+			FindOrAddType(FRigVMTemplateArgumentType(ScriptStruct));
+		}
 	}
 
 	for (TObjectIterator<UEnum> EnumIt; EnumIt; ++EnumIt)
@@ -1233,6 +1237,16 @@ bool FRigVMRegistry::IsAllowedType(const UStruct* InStruct) const
 		return true;
 	}
 
+	// Allow structs we have explicitly opted into
+	// This is on the understanding that if they have invalid sub-members that any pins representing them will need to be hidden
+	if (const UScriptStruct* ScriptStruct = Cast<UScriptStruct>(InStruct))
+	{
+		if(AllowedStructs.Contains(ScriptStruct))
+		{
+			return true;
+		}
+	}
+
 	for (TFieldIterator<FProperty> It(InStruct); It; ++It)
 	{
 		if(!IsAllowedType(*It))
@@ -1441,6 +1455,17 @@ void FRigVMRegistry::RegisterObjectTypes(TConstArrayView<TPair<UClass*, ERegiste
 	}
 }
 
+void FRigVMRegistry::RegisterStructTypes(TConstArrayView<UScriptStruct*> InStructs)
+{
+	for (UScriptStruct* Struct : InStructs)
+	{
+		if(!Struct->IsA<UUserDefinedStruct>())
+		{
+			AllowedStructs.Add(Struct);
+		}
+	}
+}
+
 const FRigVMFunction* FRigVMRegistry::FindFunction(const TCHAR* InName, const FRigVMUserDefinedTypeResolver& InTypeResolver) const
 {
 	FScopeLock FunctionRegistryScopeLock(&FunctionRegistryMutex);
diff --git a/Engine/Plugins/Runtime/RigVM/Source/RigVM/Public/RigVMCore/RigVMRegistry.h b/Engine/Plugins/Runtime/RigVM/Source/RigVM/Public/RigVMCore/RigVMRegistry.h
--- a/Engine/Plugins/Runtime/RigVM/Source/RigVM/Public/RigVMCore/RigVMRegistry.h
+++ b/Engine/Plugins/Runtime/RigVM/Source/RigVM/Public/RigVMCore/RigVMRegistry.h
@@ -73,6 +73,9 @@ public:
 	// Register a set of allowed object types
 	void RegisterObjectTypes(TConstArrayView<TPair<UClass*, ERegisterObjectOperation>> InClasses);
 
+	// Register a set of allowed struct types
+	void RegisterStructTypes(TConstArrayView<UScriptStruct*> InStructs);
+
 	// Refreshes the list and finds the function pointers
 	// based on the names.
 	void RefreshEngineTypes();
@@ -398,6 +401,9 @@ private:
 	// All allowed classes
 	TSet<TObjectPtr<const UClass>> AllowedClasses;
 
+	// All allowed structs
+	TSet<TObjectPtr<const UScriptStruct>> AllowedStructs;
+	
 	// Notifies other system that types have been added/removed, and template permutations have been updated
 	FOnRigVMRegistryChanged OnRigVMRegistryChangedDelegate;
 
