RigVMRegistry: Guard against dangling / to be removed types during GC

#rnx

diff --git a/Engine/Plugins/Runtime/RigVM/Source/RigVM/Private/RigVMCore/RigVMRegistry.cpp b/Engine/Plugins/Runtime/RigVM/Source/RigVM/Private/RigVMCore/RigVMRegistry.cpp
--- a/Engine/Plugins/Runtime/RigVM/Source/RigVM/Private/RigVMCore/RigVMRegistry.cpp
+++ b/Engine/Plugins/Runtime/RigVM/Source/RigVM/Private/RigVMCore/RigVMRegistry.cpp
@@ -60,7 +60,9 @@ void FRigVMRegistry::AddReferencedObjects(FReferenceCollector& Collector)
 	// which can happen during cook time.
 	for (FTypeInfo& Type : Types)
 	{
-		if (Type.Type.CPPTypeObject)
+		// the Object needs to be checked for validity since it may be a user defined type (struct or enum)
+		// which is about to get removed. 
+		if (IsValid(Type.Type.CPPTypeObject))
 		{
 			Collector.AddReferencedObject(Type.Type.CPPTypeObject);
 		}
