RigVM: Add option to treat warnings as errors for RigVMPin

#rnx

diff --git a/Engine/Plugins/Runtime/RigVM/Source/RigVMDeveloper/Private/RigVMModel/RigVMPin.cpp b/Engine/Plugins/Runtime/RigVM/Source/RigVMDeveloper/Private/RigVMModel/RigVMPin.cpp
--- a/Engine/Plugins/Runtime/RigVM/Source/RigVMDeveloper/Private/RigVMModel/RigVMPin.cpp
+++ b/Engine/Plugins/Runtime/RigVM/Source/RigVMDeveloper/Private/RigVMModel/RigVMPin.cpp
@@ -969,7 +969,7 @@ bool URigVMPin::IsValidDefaultValue(const FString& InDefaultValue) const
 			TempStructBuffer.AddUninitialized(ScriptStruct->GetStructureSize());
 			ScriptStruct->InitializeDefaultValue(TempStructBuffer.GetData());
 
-			FRigVMPinDefaultValueImportErrorContext ErrorPipe;
+			FRigVMPinDefaultValueImportErrorContext ErrorPipe(true /* warnings as errors */);
 			{
 				// force logging to the error pipe for error detection
 				LOG_SCOPE_VERBOSITY_OVERRIDE(LogExec, ELogVerbosity::Verbose); 
diff --git a/Engine/Plugins/Runtime/RigVM/Source/RigVMDeveloper/Public/RigVMModel/RigVMPin.h b/Engine/Plugins/Runtime/RigVM/Source/RigVMDeveloper/Public/RigVMModel/RigVMPin.h
--- a/Engine/Plugins/Runtime/RigVM/Source/RigVMDeveloper/Public/RigVMModel/RigVMPin.h
+++ b/Engine/Plugins/Runtime/RigVM/Source/RigVMDeveloper/Public/RigVMModel/RigVMPin.h
@@ -570,17 +570,22 @@ public:
 
 	int32 NumErrors;
 
-	FRigVMPinDefaultValueImportErrorContext()
+	FRigVMPinDefaultValueImportErrorContext(bool InTreatWarningAsErrors = false)
 		: FOutputDevice()
 		, NumErrors(0)
+		, bTreatWarningAsErrors(InTreatWarningAsErrors)
 	{
 	}
 
 	virtual void Serialize(const TCHAR* V, ELogVerbosity::Type Verbosity, const class FName& Category) override
 	{
-		if(Verbosity <= ELogVerbosity::Error)
+		const ELogVerbosity::Type MaxVerbosity = bTreatWarningAsErrors ? ELogVerbosity::Warning : ELogVerbosity::Error;
+		if(Verbosity <= MaxVerbosity)
 		{
 			NumErrors++;
 		}
 	}
+
+private:
+	bool bTreatWarningAsErrors;
 };
