RigVM: RigVMRegistry Thread Safety

#rnx

diff --git a/Engine/Plugins/Animation/ControlRig/Source/ControlRig/Private/Units/Animation/RigUnit_AnimAttribute.cpp b/Engine/Plugins/Animation/ControlRig/Source/ControlRig/Private/Units/Animation/RigUnit_AnimAttribute.cpp
--- a/Engine/Plugins/Animation/ControlRig/Source/ControlRig/Private/Units/Animation/RigUnit_AnimAttribute.cpp
+++ b/Engine/Plugins/Animation/ControlRig/Source/ControlRig/Private/Units/Animation/RigUnit_AnimAttribute.cpp
@@ -22,15 +22,15 @@ FName FRigDispatch_AnimAttributeBase::SuccessArgName = TEXT("Success");
 
 bool FRigDispatch_AnimAttributeBase::IsTypeSupported(const TRigVMTypeIndex& InTypeIndex)
 {
-	const FRigVMRegistry& Registry = FRigVMRegistry::Get();
+	const FRigVMRegistry_NoLock& Registry = FRigVMRegistry_NoLock::GetForRead();
 
 	static const TArray<TRigVMTypeIndex> SpecialTypes = {
 		RigVMTypeUtils::TypeIndex::Float,
 		RigVMTypeUtils::TypeIndex::Int32,
 		RigVMTypeUtils::TypeIndex::FString,
-		Registry.GetTypeIndex<FTransform>(false),
-		Registry.GetTypeIndex<FVector>(false),
-		Registry.GetTypeIndex<FQuat>(false),
+		Registry.GetTypeIndex_NoLock<FTransform>(false),
+		Registry.GetTypeIndex_NoLock<FVector>(false),
+		Registry.GetTypeIndex_NoLock<FQuat>(false),
 	};
 
 	if (SpecialTypes.Contains(InTypeIndex))
@@ -38,7 +38,7 @@ bool FRigDispatch_AnimAttributeBase::IsTypeSupported(const TRigVMTypeIndex& InTy
 		return true;
 	}
 
-	const FRigVMTemplateArgumentType& InType = Registry.GetType(InTypeIndex);
+	const FRigVMTemplateArgumentType& InType = Registry.GetType_NoLock(InTypeIndex);
 
 	// cpp type object can become invalid because users can choose to delete
 	// user defined structs
@@ -46,6 +46,14 @@ bool FRigDispatch_AnimAttributeBase::IsTypeSupported(const TRigVMTypeIndex& InTy
 	{
 		if (UScriptStruct* ScriptStruct = Cast<UScriptStruct>(InType.CPPTypeObject))
 		{
+			if (ScriptStruct->IsA<UUserDefinedStruct>())
+			{
+				// allow all user defined structs because even if a struct is not registered with anim attribute system,
+				// it could be added to or removed from the system easily. allowing all of them as valid permutations
+				// avoids having to create orphan pins.
+				return true;
+			}
+
 			static const TArray<TWeakObjectPtr<const UScriptStruct>> SpecialAttributeTypes = {
 				FFloatAnimationAttribute::StaticStruct(),
 				FIntegerAnimationAttribute::StaticStruct(),
@@ -60,15 +68,7 @@ bool FRigDispatch_AnimAttributeBase::IsTypeSupported(const TRigVMTypeIndex& InTy
 				// these type have been added, above, rejecting them here so we don't have duplicated types
 				return false;
 			}
-
-			if (UUserDefinedStruct* UserDefinedStruct = Cast<UUserDefinedStruct>(ScriptStruct))
-			{
-				// allow all user defined structs because even if a struct is not registered with anim attribute system,
-				// it could be added to or removed from the system easily. allowing all of them as valid permutations
-				// avoids having to create orphan pins.
-				return true;
-			}
-			
+		
 			return UE::Anim::AttributeTypes::IsTypeRegistered(ScriptStruct);
 		}
 	}
@@ -272,11 +272,12 @@ void FRigDispatch_SetAnimAttribute::SetAnimAttributeDispatch(FRigVMExtendedExecu
 	}	
 }
 
-void FRigDispatch_AnimAttributeBase::RegisterDependencyTypes() const
+void FRigDispatch_AnimAttributeBase::RegisterDependencyTypes_NoLock() const
 {
-	Super::RegisterDependencyTypes();
+	Super::RegisterDependencyTypes_NoLock();
 
-	FRigVMRegistry& Registry = FRigVMRegistry::Get();
+	FRigVMRegistry_NoLock& Registry = FRigVMRegistry_NoLock::GetForWrite();
+	Registry.FindOrAddType_NoLock(FRigVMTemplateArgumentType(), true);
 	
 	TArray<TWeakObjectPtr<const UScriptStruct>> AttributeTypes = UE::Anim::AttributeTypes::GetRegisteredTypes();
 	for (TWeakObjectPtr<const UScriptStruct> Type : AttributeTypes)
@@ -284,7 +285,7 @@ void FRigDispatch_AnimAttributeBase::RegisterDependencyTypes() const
 		UScriptStruct* TypePtr = const_cast<UScriptStruct*>(Type.Get());
 		if (TypePtr)
 		{
-			Registry.FindOrAddType(FRigVMTemplateArgumentType(TypePtr));
+			Registry.FindOrAddType_NoLock(FRigVMTemplateArgumentType(TypePtr));
 		}
 	}
 }
@@ -417,7 +418,7 @@ FRigVMTemplateTypeMap FRigDispatch_GetAnimAttribute::OnNewArgumentType(const FNa
 
 FRigVMFunctionPtr FRigDispatch_GetAnimAttribute::GetDispatchFunctionImpl(const FRigVMTemplateTypeMap& InTypes) const
 {
-	const FRigVMRegistry& Registry = FRigVMRegistry::Get();
+	const FRigVMRegistry_NoLock& Registry = FRigVMRegistry_NoLock::GetForRead();
 	const TRigVMTypeIndex& ValueTypeIndex = InTypes.FindChecked(TEXT("Value"));
 
 	if (ValueTypeIndex == RigVMTypeUtils::TypeIndex::Float)
@@ -432,20 +433,20 @@ FRigVMFunctionPtr FRigDispatch_GetAnimAttribute::GetDispatchFunctionImpl(const F
 	{
 		return &FRigDispatch_GetAnimAttribute::GetAnimAttributeDispatch<FString>;
 	}
-	if (ValueTypeIndex == Registry.GetTypeIndex<FTransform>(false))
+	if (ValueTypeIndex == Registry.GetTypeIndex_NoLock<FTransform>(false))
 	{
 		return &FRigDispatch_GetAnimAttribute::GetAnimAttributeDispatch<FTransform>;
 	}
-	if (ValueTypeIndex == Registry.GetTypeIndex<FQuat>(false))
+	if (ValueTypeIndex == Registry.GetTypeIndex_NoLock<FQuat>(false))
 	{
 		return &FRigDispatch_GetAnimAttribute::GetAnimAttributeDispatch<FQuat>;
 	}
-	if (ValueTypeIndex == Registry.GetTypeIndex<FVector>(false))
+	if (ValueTypeIndex == Registry.GetTypeIndex_NoLock<FVector>(false))
 	{
 		return &FRigDispatch_GetAnimAttribute::GetAnimAttributeDispatch<FVector>;
 	}
 	
-	const FRigVMTemplateArgumentType& ValueType = Registry.GetType(ValueTypeIndex);
+	const FRigVMTemplateArgumentType& ValueType = Registry.GetType_NoLock(ValueTypeIndex);
 	
 	if (UScriptStruct* ScriptStruct = Cast<UScriptStruct>(ValueType.CPPTypeObject))
 	{
@@ -506,7 +507,7 @@ FRigVMTemplateTypeMap FRigDispatch_SetAnimAttribute::OnNewArgumentType(const FNa
 
 FRigVMFunctionPtr FRigDispatch_SetAnimAttribute::GetDispatchFunctionImpl(const FRigVMTemplateTypeMap& InTypes) const
 {
-	const FRigVMRegistry& Registry = FRigVMRegistry::Get();
+	const FRigVMRegistry_NoLock& Registry = FRigVMRegistry_NoLock::GetForRead();
 	const TRigVMTypeIndex& ValueTypeIndex = InTypes.FindChecked(TEXT("Value"));
 
 	if (ValueTypeIndex == RigVMTypeUtils::TypeIndex::Float)
@@ -521,20 +522,20 @@ FRigVMFunctionPtr FRigDispatch_SetAnimAttribute::GetDispatchFunctionImpl(const F
 	{
 		return &FRigDispatch_SetAnimAttribute::SetAnimAttributeDispatch<FString>;
 	}
-	if (ValueTypeIndex == Registry.GetTypeIndex<FTransform>(false))
+	if (ValueTypeIndex == Registry.GetTypeIndex_NoLock<FTransform>(false))
 	{
 		return &FRigDispatch_SetAnimAttribute::SetAnimAttributeDispatch<FTransform>;
 	}
-	if (ValueTypeIndex == Registry.GetTypeIndex<FQuat>(false))
+	if (ValueTypeIndex == Registry.GetTypeIndex_NoLock<FQuat>(false))
 	{
 		return &FRigDispatch_SetAnimAttribute::SetAnimAttributeDispatch<FQuat>;
 	}
-	if (ValueTypeIndex == Registry.GetTypeIndex<FVector>(false))
+	if (ValueTypeIndex == Registry.GetTypeIndex_NoLock<FVector>(false))
 	{
 		return &FRigDispatch_SetAnimAttribute::SetAnimAttributeDispatch<FVector>;
 	}
 	
-	const FRigVMTemplateArgumentType& ValueType = Registry.GetType(ValueTypeIndex);
+	const FRigVMTemplateArgumentType& ValueType = Registry.GetType_NoLock(ValueTypeIndex);
 	
 	if (UScriptStruct* ScriptStruct = Cast<UScriptStruct>(ValueType.CPPTypeObject))
 	{
diff --git a/Engine/Plugins/Animation/ControlRig/Source/ControlRig/Private/Units/Hierarchy/RigUnit_Metadata.cpp b/Engine/Plugins/Animation/ControlRig/Source/ControlRig/Private/Units/Hierarchy/RigUnit_Metadata.cpp
--- a/Engine/Plugins/Animation/ControlRig/Source/ControlRig/Private/Units/Hierarchy/RigUnit_Metadata.cpp
+++ b/Engine/Plugins/Animation/ControlRig/Source/ControlRig/Private/Units/Hierarchy/RigUnit_Metadata.cpp
@@ -72,10 +72,11 @@ const TArray<FRigVMTemplateArgumentInfo>& FRigDispatch_MetadataBase::GetArgument
 {
 	if(Infos.IsEmpty())
 	{
-		ItemArgIndex = Infos.Emplace(ItemArgName, ERigVMPinDirection::Input, FRigVMRegistry::Get().GetTypeIndex<FRigElementKey>());
+		const FRigVMRegistry_NoLock& Registry = FRigVMRegistry_NoLock::GetForRead();
+		ItemArgIndex = Infos.Emplace(ItemArgName, ERigVMPinDirection::Input, Registry.GetTypeIndex_NoLock<FRigElementKey>());
 		NameArgIndex = Infos.Emplace(NameArgName, ERigVMPinDirection::Input, RigVMTypeUtils::TypeIndex::FName);
-		NameSpaceArgIndex = Infos.Emplace(NameSpaceArgName, ERigVMPinDirection::Input, FRigVMRegistry::Get().GetTypeIndex<ERigMetaDataNameSpace>());
-		CacheArgIndex = Infos.Emplace(CacheArgName, ERigVMPinDirection::Hidden, FRigVMRegistry::Get().GetTypeIndex<FCachedRigElement>());
+		NameSpaceArgIndex = Infos.Emplace(NameSpaceArgName, ERigVMPinDirection::Input, Registry.GetTypeIndex_NoLock<ERigMetaDataNameSpace>());
+		CacheArgIndex = Infos.Emplace(CacheArgName, ERigVMPinDirection::Hidden, Registry.GetTypeIndex_NoLock<FCachedRigElement>());
 	};
 	return Infos;
 }
@@ -145,28 +146,28 @@ const TArray<TRigVMTypeIndex>& FRigDispatch_MetadataBase::GetValueTypes() const
 	static TArray<TRigVMTypeIndex> Types;
 	if(Types.IsEmpty())
 	{
-		const FRigVMRegistry& Registry = FRigVMRegistry::Get();
+		const FRigVMRegistry_NoLock& Registry = FRigVMRegistry_NoLock::GetForRead();
 		Types = {
 			RigVMTypeUtils::TypeIndex::Bool,
 			RigVMTypeUtils::TypeIndex::Float,
 			RigVMTypeUtils::TypeIndex::Int32,
 			RigVMTypeUtils::TypeIndex::FName,
-			Registry.GetTypeIndex<FVector>(false),
-			Registry.GetTypeIndex<FRotator>(false),
-			Registry.GetTypeIndex<FQuat>(false),
-			Registry.GetTypeIndex<FTransform>(false),
-			Registry.GetTypeIndex<FLinearColor>(false),
-			Registry.GetTypeIndex<FRigElementKey>(false),
+			Registry.GetTypeIndex_NoLock<FVector>(false),
+			Registry.GetTypeIndex_NoLock<FRotator>(false),
+			Registry.GetTypeIndex_NoLock<FQuat>(false),
+			Registry.GetTypeIndex_NoLock<FTransform>(false),
+			Registry.GetTypeIndex_NoLock<FLinearColor>(false),
+			Registry.GetTypeIndex_NoLock<FRigElementKey>(false),
 			RigVMTypeUtils::TypeIndex::BoolArray,
 			RigVMTypeUtils::TypeIndex::FloatArray,
 			RigVMTypeUtils::TypeIndex::Int32Array,
 			RigVMTypeUtils::TypeIndex::FNameArray,
-			Registry.GetTypeIndex<FVector>(true),
-			Registry.GetTypeIndex<FRotator>(true),
-			Registry.GetTypeIndex<FQuat>(true),
-			Registry.GetTypeIndex<FTransform>(true),
-			Registry.GetTypeIndex<FLinearColor>(true),
-			Registry.GetTypeIndex<FRigElementKey>(true)
+			Registry.GetTypeIndex_NoLock<FVector>(true),
+			Registry.GetTypeIndex_NoLock<FRotator>(true),
+			Registry.GetTypeIndex_NoLock<FQuat>(true),
+			Registry.GetTypeIndex_NoLock<FTransform>(true),
+			Registry.GetTypeIndex_NoLock<FLinearColor>(true),
+			Registry.GetTypeIndex_NoLock<FRigElementKey>(true)
 		};
 	}
 	return Types;
@@ -205,7 +206,7 @@ FRigBaseMetadata* FRigDispatch_GetMetadata::FindMetadata(const FRigVMExtendedExe
 
 FRigVMFunctionPtr FRigDispatch_GetMetadata::GetDispatchFunctionImpl(const FRigVMTemplateTypeMap& InTypes) const
 {
-	const FRigVMRegistry& Registry = FRigVMRegistry::Get();
+	const FRigVMRegistry_NoLock& Registry = FRigVMRegistry_NoLock::GetForRead();
 	const TRigVMTypeIndex& ValueTypeIndex = InTypes.FindChecked(TEXT("Value"));
 	
 	if(ValueTypeIndex == RigVMTypeUtils::TypeIndex::Bool)
@@ -224,27 +225,27 @@ FRigVMFunctionPtr FRigDispatch_GetMetadata::GetDispatchFunctionImpl(const FRigVM
 	{
 		return &FRigDispatch_GetMetadata::GetMetadataDispatch<FName, FRigNameMetadata, ERigMetadataType::Name>;
 	}
-	if(ValueTypeIndex == Registry.GetTypeIndex<FVector>(false))
+	if(ValueTypeIndex == Registry.GetTypeIndex_NoLock<FVector>(false))
 	{
 		return &FRigDispatch_GetMetadata::GetMetadataDispatch<FVector, FRigVectorMetadata, ERigMetadataType::Vector>;
 	}
-	if(ValueTypeIndex == Registry.GetTypeIndex<FRotator>(false))
+	if(ValueTypeIndex == Registry.GetTypeIndex_NoLock<FRotator>(false))
 	{
 		return &FRigDispatch_GetMetadata::GetMetadataDispatch<FRotator, FRigRotatorMetadata, ERigMetadataType::Rotator>;
 	}
-	if(ValueTypeIndex == Registry.GetTypeIndex<FQuat>(false))
+	if(ValueTypeIndex == Registry.GetTypeIndex_NoLock<FQuat>(false))
 	{
 		return &FRigDispatch_GetMetadata::GetMetadataDispatch<FQuat, FRigQuatMetadata, ERigMetadataType::Quat>;
 	}
-	if(ValueTypeIndex == Registry.GetTypeIndex<FTransform>(false))
+	if(ValueTypeIndex == Registry.GetTypeIndex_NoLock<FTransform>(false))
 	{
 		return &FRigDispatch_GetMetadata::GetMetadataDispatch<FTransform, FRigTransformMetadata, ERigMetadataType::Transform>;
 	}
-	if(ValueTypeIndex == Registry.GetTypeIndex<FLinearColor>(false))
+	if(ValueTypeIndex == Registry.GetTypeIndex_NoLock<FLinearColor>(false))
 	{
 		return &FRigDispatch_GetMetadata::GetMetadataDispatch<FLinearColor, FRigLinearColorMetadata, ERigMetadataType::LinearColor>;
 	}
-	if(ValueTypeIndex == Registry.GetTypeIndex<FRigElementKey>(false))
+	if(ValueTypeIndex == Registry.GetTypeIndex_NoLock<FRigElementKey>(false))
 	{
 		return &FRigDispatch_GetMetadata::GetMetadataDispatch<FRigElementKey, FRigElementKeyMetadata, ERigMetadataType::RigElementKey>;
 	}
@@ -264,27 +265,27 @@ FRigVMFunctionPtr FRigDispatch_GetMetadata::GetDispatchFunctionImpl(const FRigVM
 	{
 		return &FRigDispatch_GetMetadata::GetMetadataDispatch<TArray<FName>, FRigNameArrayMetadata, ERigMetadataType::NameArray>;
 	}
-	if(ValueTypeIndex == Registry.GetTypeIndex<FVector>(true))
+	if(ValueTypeIndex == Registry.GetTypeIndex_NoLock<FVector>(true))
 	{
 		return &FRigDispatch_GetMetadata::GetMetadataDispatch<TArray<FVector>, FRigVectorArrayMetadata, ERigMetadataType::VectorArray>;
 	}
-	if(ValueTypeIndex == Registry.GetTypeIndex<FRotator>(true))
+	if(ValueTypeIndex == Registry.GetTypeIndex_NoLock<FRotator>(true))
 	{
 		return &FRigDispatch_GetMetadata::GetMetadataDispatch<TArray<FRotator>, FRigRotatorArrayMetadata, ERigMetadataType::RotatorArray>;
 	}
-	if(ValueTypeIndex == Registry.GetTypeIndex<FQuat>(true))
+	if(ValueTypeIndex == Registry.GetTypeIndex_NoLock<FQuat>(true))
 	{
 		return &FRigDispatch_GetMetadata::GetMetadataDispatch<TArray<FQuat>, FRigQuatArrayMetadata, ERigMetadataType::QuatArray>;
 	}
-	if(ValueTypeIndex == Registry.GetTypeIndex<FTransform>(true))
+	if(ValueTypeIndex == Registry.GetTypeIndex_NoLock<FTransform>(true))
 	{
 		return &FRigDispatch_GetMetadata::GetMetadataDispatch<TArray<FTransform>, FRigTransformArrayMetadata, ERigMetadataType::TransformArray>;
 	}
-	if(ValueTypeIndex == Registry.GetTypeIndex<FLinearColor>(true))
+	if(ValueTypeIndex == Registry.GetTypeIndex_NoLock<FLinearColor>(true))
 	{
 		return &FRigDispatch_GetMetadata::GetMetadataDispatch<TArray<FLinearColor>, FRigLinearColorArrayMetadata, ERigMetadataType::LinearColorArray>;
 	}
-	if(ValueTypeIndex == Registry.GetTypeIndex<FRigElementKey>(true))
+	if(ValueTypeIndex == Registry.GetTypeIndex_NoLock<FRigElementKey>(true))
 	{
 		return &FRigDispatch_GetMetadata::GetMetadataDispatch<TArray<FRigElementKey>, FRigElementKeyArrayMetadata, ERigMetadataType::RigElementKeyArray>;
 	}
@@ -331,7 +332,7 @@ FRigBaseMetadata* FRigDispatch_SetMetadata::FindOrAddMetadata(const FControlRigE
 
 FRigVMFunctionPtr FRigDispatch_SetMetadata::GetDispatchFunctionImpl(const FRigVMTemplateTypeMap& InTypes) const
 {
-	const FRigVMRegistry& Registry = FRigVMRegistry::Get();
+	const FRigVMRegistry_NoLock& Registry = FRigVMRegistry_NoLock::GetForRead();
 	const TRigVMTypeIndex& ValueTypeIndex = InTypes.FindChecked(TEXT("Value"));
 	
 	if(ValueTypeIndex == RigVMTypeUtils::TypeIndex::Bool)
@@ -350,27 +351,27 @@ FRigVMFunctionPtr FRigDispatch_SetMetadata::GetDispatchFunctionImpl(const FRigVM
 	{
 		return &FRigDispatch_SetMetadata::SetMetadataDispatch<FName, FRigNameMetadata, ERigMetadataType::Name>;
 	}
-	if(ValueTypeIndex == Registry.GetTypeIndex<FVector>(false))
+	if(ValueTypeIndex == Registry.GetTypeIndex_NoLock<FVector>(false))
 	{
 		return &FRigDispatch_SetMetadata::SetMetadataDispatch<FVector, FRigVectorMetadata, ERigMetadataType::Vector>;
 	}
-	if(ValueTypeIndex == Registry.GetTypeIndex<FRotator>(false))
+	if(ValueTypeIndex == Registry.GetTypeIndex_NoLock<FRotator>(false))
 	{
 		return &FRigDispatch_SetMetadata::SetMetadataDispatch<FRotator, FRigRotatorMetadata, ERigMetadataType::Rotator>;
 	}
-	if(ValueTypeIndex == Registry.GetTypeIndex<FQuat>(false))
+	if(ValueTypeIndex == Registry.GetTypeIndex_NoLock<FQuat>(false))
 	{
 		return &FRigDispatch_SetMetadata::SetMetadataDispatch<FQuat, FRigQuatMetadata, ERigMetadataType::Quat>;
 	}
-	if(ValueTypeIndex == Registry.GetTypeIndex<FTransform>(false))
+	if(ValueTypeIndex == Registry.GetTypeIndex_NoLock<FTransform>(false))
 	{
 		return &FRigDispatch_SetMetadata::SetMetadataDispatch<FTransform, FRigTransformMetadata, ERigMetadataType::Transform>;
 	}
-	if(ValueTypeIndex == Registry.GetTypeIndex<FLinearColor>(false))
+	if(ValueTypeIndex == Registry.GetTypeIndex_NoLock<FLinearColor>(false))
 	{
 		return &FRigDispatch_SetMetadata::SetMetadataDispatch<FLinearColor, FRigLinearColorMetadata, ERigMetadataType::LinearColor>;
 	}
-	if(ValueTypeIndex == Registry.GetTypeIndex<FRigElementKey>(false))
+	if(ValueTypeIndex == Registry.GetTypeIndex_NoLock<FRigElementKey>(false))
 	{
 		return &FRigDispatch_SetMetadata::SetMetadataDispatch<FRigElementKey, FRigElementKeyMetadata, ERigMetadataType::RigElementKey>;
 	}
@@ -390,27 +391,27 @@ FRigVMFunctionPtr FRigDispatch_SetMetadata::GetDispatchFunctionImpl(const FRigVM
 	{
 		return &FRigDispatch_SetMetadata::SetMetadataDispatch<TArray<FName>, FRigNameArrayMetadata, ERigMetadataType::NameArray>;
 	}
-	if(ValueTypeIndex == Registry.GetTypeIndex<FVector>(true))
+	if(ValueTypeIndex == Registry.GetTypeIndex_NoLock<FVector>(true))
 	{
 		return &FRigDispatch_SetMetadata::SetMetadataDispatch<TArray<FVector>, FRigVectorArrayMetadata, ERigMetadataType::VectorArray>;
 	}
-	if(ValueTypeIndex == Registry.GetTypeIndex<FRotator>(true))
+	if(ValueTypeIndex == Registry.GetTypeIndex_NoLock<FRotator>(true))
 	{
 		return &FRigDispatch_SetMetadata::SetMetadataDispatch<TArray<FRotator>, FRigRotatorArrayMetadata, ERigMetadataType::RotatorArray>;
 	}
-	if(ValueTypeIndex == Registry.GetTypeIndex<FQuat>(true))
+	if(ValueTypeIndex == Registry.GetTypeIndex_NoLock<FQuat>(true))
 	{
 		return &FRigDispatch_SetMetadata::SetMetadataDispatch<TArray<FQuat>, FRigQuatArrayMetadata, ERigMetadataType::QuatArray>;
 	}
-	if(ValueTypeIndex == Registry.GetTypeIndex<FTransform>(true))
+	if(ValueTypeIndex == Registry.GetTypeIndex_NoLock<FTransform>(true))
 	{
 		return &FRigDispatch_SetMetadata::SetMetadataDispatch<TArray<FTransform>, FRigTransformArrayMetadata, ERigMetadataType::TransformArray>;
 	}
-	if(ValueTypeIndex == Registry.GetTypeIndex<FLinearColor>(true))
+	if(ValueTypeIndex == Registry.GetTypeIndex_NoLock<FLinearColor>(true))
 	{
 		return &FRigDispatch_SetMetadata::SetMetadataDispatch<TArray<FLinearColor>, FRigLinearColorArrayMetadata, ERigMetadataType::LinearColorArray>;
 	}
-	if(ValueTypeIndex == Registry.GetTypeIndex<FRigElementKey>(true))
+	if(ValueTypeIndex == Registry.GetTypeIndex_NoLock<FRigElementKey>(true))
 	{
 		return &FRigDispatch_SetMetadata::SetMetadataDispatch<TArray<FRigElementKey>, FRigElementKeyArrayMetadata, ERigMetadataType::RigElementKeyArray>;
 	}
@@ -841,7 +842,7 @@ const TArray<FRigVMTemplateArgumentInfo>& FRigDispatch_GetModuleMetadata::GetArg
 	if(ValueArgIndex == INDEX_NONE)
 	{
 		NameArgIndex = Infos.Emplace(NameArgName, ERigVMPinDirection::Input, RigVMTypeUtils::TypeIndex::FName);
-		NameSpaceArgIndex = Infos.Emplace(NameSpaceArgName, ERigVMPinDirection::Input, FRigVMRegistry::Get().GetTypeIndex<ERigMetaDataNameSpace>());
+		NameSpaceArgIndex = Infos.Emplace(NameSpaceArgName, ERigVMPinDirection::Input, FRigVMRegistry_NoLock::GetForRead().GetTypeIndex_NoLock<ERigMetaDataNameSpace>());
 		DefaultArgIndex = Infos.Emplace(DefaultArgName, ERigVMPinDirection::Input, GetValueTypes());
 		ValueArgIndex = Infos.Emplace(ValueArgName, ERigVMPinDirection::Output, GetValueTypes());
 		FoundArgIndex = Infos.Emplace(FoundArgName, ERigVMPinDirection::Output, RigVMTypeUtils::TypeIndex::Bool);
@@ -878,7 +879,7 @@ FRigBaseMetadata* FRigDispatch_GetModuleMetadata::FindMetadata(const FRigVMExten
 
 FRigVMFunctionPtr FRigDispatch_GetModuleMetadata::GetDispatchFunctionImpl(const FRigVMTemplateTypeMap& InTypes) const
 {
-	const FRigVMRegistry& Registry = FRigVMRegistry::Get();
+	const FRigVMRegistry_NoLock& Registry = FRigVMRegistry_NoLock::GetForRead();
 	const TRigVMTypeIndex& ValueTypeIndex = InTypes.FindChecked(TEXT("Value"));
 	
 	if(ValueTypeIndex == RigVMTypeUtils::TypeIndex::Bool)
@@ -897,27 +898,27 @@ FRigVMFunctionPtr FRigDispatch_GetModuleMetadata::GetDispatchFunctionImpl(const
 	{
 		return &FRigDispatch_GetModuleMetadata::GetModuleMetadataDispatch<FName, FRigNameMetadata, ERigMetadataType::Name>;
 	}
-	if(ValueTypeIndex == Registry.GetTypeIndex<FVector>(false))
+	if(ValueTypeIndex == Registry.GetTypeIndex_NoLock<FVector>(false))
 	{
 		return &FRigDispatch_GetModuleMetadata::GetModuleMetadataDispatch<FVector, FRigVectorMetadata, ERigMetadataType::Vector>;
 	}
-	if(ValueTypeIndex == Registry.GetTypeIndex<FRotator>(false))
+	if(ValueTypeIndex == Registry.GetTypeIndex_NoLock<FRotator>(false))
 	{
 		return &FRigDispatch_GetModuleMetadata::GetModuleMetadataDispatch<FRotator, FRigRotatorMetadata, ERigMetadataType::Rotator>;
 	}
-	if(ValueTypeIndex == Registry.GetTypeIndex<FQuat>(false))
+	if(ValueTypeIndex == Registry.GetTypeIndex_NoLock<FQuat>(false))
 	{
 		return &FRigDispatch_GetModuleMetadata::GetModuleMetadataDispatch<FQuat, FRigQuatMetadata, ERigMetadataType::Quat>;
 	}
-	if(ValueTypeIndex == Registry.GetTypeIndex<FTransform>(false))
+	if(ValueTypeIndex == Registry.GetTypeIndex_NoLock<FTransform>(false))
 	{
 		return &FRigDispatch_GetModuleMetadata::GetModuleMetadataDispatch<FTransform, FRigTransformMetadata, ERigMetadataType::Transform>;
 	}
-	if(ValueTypeIndex == Registry.GetTypeIndex<FLinearColor>(false))
+	if(ValueTypeIndex == Registry.GetTypeIndex_NoLock<FLinearColor>(false))
 	{
 		return &FRigDispatch_GetModuleMetadata::GetModuleMetadataDispatch<FLinearColor, FRigLinearColorMetadata, ERigMetadataType::LinearColor>;
 	}
-	if(ValueTypeIndex == Registry.GetTypeIndex<FRigElementKey>(false))
+	if(ValueTypeIndex == Registry.GetTypeIndex_NoLock<FRigElementKey>(false))
 	{
 		return &FRigDispatch_GetModuleMetadata::GetModuleMetadataDispatch<FRigElementKey, FRigElementKeyMetadata, ERigMetadataType::RigElementKey>;
 	}
@@ -937,27 +938,27 @@ FRigVMFunctionPtr FRigDispatch_GetModuleMetadata::GetDispatchFunctionImpl(const
 	{
 		return &FRigDispatch_GetModuleMetadata::GetModuleMetadataDispatch<TArray<FName>, FRigNameArrayMetadata, ERigMetadataType::NameArray>;
 	}
-	if(ValueTypeIndex == Registry.GetTypeIndex<FVector>(true))
+	if(ValueTypeIndex == Registry.GetTypeIndex_NoLock<FVector>(true))
 	{
 		return &FRigDispatch_GetModuleMetadata::GetModuleMetadataDispatch<TArray<FVector>, FRigVectorArrayMetadata, ERigMetadataType::VectorArray>;
 	}
-	if(ValueTypeIndex == Registry.GetTypeIndex<FRotator>(true))
+	if(ValueTypeIndex == Registry.GetTypeIndex_NoLock<FRotator>(true))
 	{
 		return &FRigDispatch_GetModuleMetadata::GetModuleMetadataDispatch<TArray<FRotator>, FRigRotatorArrayMetadata, ERigMetadataType::RotatorArray>;
 	}
-	if(ValueTypeIndex == Registry.GetTypeIndex<FQuat>(true))
+	if(ValueTypeIndex == Registry.GetTypeIndex_NoLock<FQuat>(true))
 	{
 		return &FRigDispatch_GetModuleMetadata::GetModuleMetadataDispatch<TArray<FQuat>, FRigQuatArrayMetadata, ERigMetadataType::QuatArray>;
 	}
-	if(ValueTypeIndex == Registry.GetTypeIndex<FTransform>(true))
+	if(ValueTypeIndex == Registry.GetTypeIndex_NoLock<FTransform>(true))
 	{
 		return &FRigDispatch_GetModuleMetadata::GetModuleMetadataDispatch<TArray<FTransform>, FRigTransformArrayMetadata, ERigMetadataType::TransformArray>;
 	}
-	if(ValueTypeIndex == Registry.GetTypeIndex<FLinearColor>(true))
+	if(ValueTypeIndex == Registry.GetTypeIndex_NoLock<FLinearColor>(true))
 	{
 		return &FRigDispatch_GetModuleMetadata::GetModuleMetadataDispatch<TArray<FLinearColor>, FRigLinearColorArrayMetadata, ERigMetadataType::LinearColorArray>;
 	}
-	if(ValueTypeIndex == Registry.GetTypeIndex<FRigElementKey>(true))
+	if(ValueTypeIndex == Registry.GetTypeIndex_NoLock<FRigElementKey>(true))
 	{
 		return &FRigDispatch_GetModuleMetadata::GetModuleMetadataDispatch<TArray<FRigElementKey>, FRigElementKeyArrayMetadata, ERigMetadataType::RigElementKeyArray>;
 	}
@@ -970,7 +971,7 @@ const TArray<FRigVMTemplateArgumentInfo>& FRigDispatch_SetModuleMetadata::GetArg
 	if(ValueArgIndex == INDEX_NONE)
 	{
 		NameArgIndex = Infos.Emplace(NameArgName, ERigVMPinDirection::Input, RigVMTypeUtils::TypeIndex::FName);
-		NameSpaceArgIndex = Infos.Emplace(NameSpaceArgName, ERigVMPinDirection::Input, FRigVMRegistry::Get().GetTypeIndex<ERigMetaDataNameSpace>());
+		NameSpaceArgIndex = Infos.Emplace(NameSpaceArgName, ERigVMPinDirection::Input, FRigVMRegistry_NoLock::GetForRead().GetTypeIndex_NoLock<ERigMetaDataNameSpace>());
 		ValueArgIndex = Infos.Emplace(ValueArgName, ERigVMPinDirection::Input, GetValueTypes());
 		SuccessArgIndex = Infos.Emplace(SuccessArgName, ERigVMPinDirection::Output, RigVMTypeUtils::TypeIndex::Bool);
 	};
@@ -1006,7 +1007,7 @@ FRigBaseMetadata* FRigDispatch_SetModuleMetadata::FindOrAddMetadata(const FContr
 
 FRigVMFunctionPtr FRigDispatch_SetModuleMetadata::GetDispatchFunctionImpl(const FRigVMTemplateTypeMap& InTypes) const
 {
-	const FRigVMRegistry& Registry = FRigVMRegistry::Get();
+	const FRigVMRegistry_NoLock& Registry = FRigVMRegistry_NoLock::GetForRead();
 	const TRigVMTypeIndex& ValueTypeIndex = InTypes.FindChecked(TEXT("Value"));
 	
 	if(ValueTypeIndex == RigVMTypeUtils::TypeIndex::Bool)
@@ -1025,27 +1026,27 @@ FRigVMFunctionPtr FRigDispatch_SetModuleMetadata::GetDispatchFunctionImpl(const
 	{
 		return &FRigDispatch_SetModuleMetadata::SetModuleMetadataDispatch<FName, FRigNameMetadata, ERigMetadataType::Name>;
 	}
-	if(ValueTypeIndex == Registry.GetTypeIndex<FVector>(false))
+	if(ValueTypeIndex == Registry.GetTypeIndex_NoLock<FVector>(false))
 	{
 		return &FRigDispatch_SetModuleMetadata::SetModuleMetadataDispatch<FVector, FRigVectorMetadata, ERigMetadataType::Vector>;
 	}
-	if(ValueTypeIndex == Registry.GetTypeIndex<FRotator>(false))
+	if(ValueTypeIndex == Registry.GetTypeIndex_NoLock<FRotator>(false))
 	{
 		return &FRigDispatch_SetModuleMetadata::SetModuleMetadataDispatch<FRotator, FRigRotatorMetadata, ERigMetadataType::Rotator>;
 	}
-	if(ValueTypeIndex == Registry.GetTypeIndex<FQuat>(false))
+	if(ValueTypeIndex == Registry.GetTypeIndex_NoLock<FQuat>(false))
 	{
 		return &FRigDispatch_SetModuleMetadata::SetModuleMetadataDispatch<FQuat, FRigQuatMetadata, ERigMetadataType::Quat>;
 	}
-	if(ValueTypeIndex == Registry.GetTypeIndex<FTransform>(false))
+	if(ValueTypeIndex == Registry.GetTypeIndex_NoLock<FTransform>(false))
 	{
 		return &FRigDispatch_SetModuleMetadata::SetModuleMetadataDispatch<FTransform, FRigTransformMetadata, ERigMetadataType::Transform>;
 	}
-	if(ValueTypeIndex == Registry.GetTypeIndex<FLinearColor>(false))
+	if(ValueTypeIndex == Registry.GetTypeIndex_NoLock<FLinearColor>(false))
 	{
 		return &FRigDispatch_SetModuleMetadata::SetModuleMetadataDispatch<FLinearColor, FRigLinearColorMetadata, ERigMetadataType::LinearColor>;
 	}
-	if(ValueTypeIndex == Registry.GetTypeIndex<FRigElementKey>(false))
+	if(ValueTypeIndex == Registry.GetTypeIndex_NoLock<FRigElementKey>(false))
 	{
 		return &FRigDispatch_SetModuleMetadata::SetModuleMetadataDispatch<FRigElementKey, FRigElementKeyMetadata, ERigMetadataType::RigElementKey>;
 	}
@@ -1065,27 +1066,27 @@ FRigVMFunctionPtr FRigDispatch_SetModuleMetadata::GetDispatchFunctionImpl(const
 	{
 		return &FRigDispatch_SetModuleMetadata::SetModuleMetadataDispatch<TArray<FName>, FRigNameArrayMetadata, ERigMetadataType::NameArray>;
 	}
-	if(ValueTypeIndex == Registry.GetTypeIndex<FVector>(true))
+	if(ValueTypeIndex == Registry.GetTypeIndex_NoLock<FVector>(true))
 	{
 		return &FRigDispatch_SetModuleMetadata::SetModuleMetadataDispatch<TArray<FVector>, FRigVectorArrayMetadata, ERigMetadataType::VectorArray>;
 	}
-	if(ValueTypeIndex == Registry.GetTypeIndex<FRotator>(true))
+	if(ValueTypeIndex == Registry.GetTypeIndex_NoLock<FRotator>(true))
 	{
 		return &FRigDispatch_SetModuleMetadata::SetModuleMetadataDispatch<TArray<FRotator>, FRigRotatorArrayMetadata, ERigMetadataType::RotatorArray>;
 	}
-	if(ValueTypeIndex == Registry.GetTypeIndex<FQuat>(true))
+	if(ValueTypeIndex == Registry.GetTypeIndex_NoLock<FQuat>(true))
 	{
 		return &FRigDispatch_SetModuleMetadata::SetModuleMetadataDispatch<TArray<FQuat>, FRigQuatArrayMetadata, ERigMetadataType::QuatArray>;
 	}
-	if(ValueTypeIndex == Registry.GetTypeIndex<FTransform>(true))
+	if(ValueTypeIndex == Registry.GetTypeIndex_NoLock<FTransform>(true))
 	{
 		return &FRigDispatch_SetModuleMetadata::SetModuleMetadataDispatch<TArray<FTransform>, FRigTransformArrayMetadata, ERigMetadataType::TransformArray>;
 	}
-	if(ValueTypeIndex == Registry.GetTypeIndex<FLinearColor>(true))
+	if(ValueTypeIndex == Registry.GetTypeIndex_NoLock<FLinearColor>(true))
 	{
 		return &FRigDispatch_SetModuleMetadata::SetModuleMetadataDispatch<TArray<FLinearColor>, FRigLinearColorArrayMetadata, ERigMetadataType::LinearColorArray>;
 	}
-	if(ValueTypeIndex == Registry.GetTypeIndex<FRigElementKey>(true))
+	if(ValueTypeIndex == Registry.GetTypeIndex_NoLock<FRigElementKey>(true))
 	{
 		return &FRigDispatch_SetModuleMetadata::SetModuleMetadataDispatch<TArray<FRigElementKey>, FRigElementKeyArrayMetadata, ERigMetadataType::RigElementKeyArray>;
 	}
diff --git a/Engine/Plugins/Animation/ControlRig/Source/ControlRig/Public/Units/Animation/RigUnit_AnimAttribute.h b/Engine/Plugins/Animation/ControlRig/Source/ControlRig/Public/Units/Animation/RigUnit_AnimAttribute.h
--- a/Engine/Plugins/Animation/ControlRig/Source/ControlRig/Public/Units/Animation/RigUnit_AnimAttribute.h
+++ b/Engine/Plugins/Animation/ControlRig/Source/ControlRig/Public/Units/Animation/RigUnit_AnimAttribute.h
@@ -139,7 +139,7 @@ struct CONTROLRIG_API FRigDispatch_AnimAttributeBase : public FRigDispatchFactor
 {
 	GENERATED_BODY()
 
-	virtual void RegisterDependencyTypes() const override;
+	virtual void RegisterDependencyTypes_NoLock() const override;
 
 #if WITH_EDITOR
 	virtual FString GetNodeTitle(const FRigVMTemplateTypeMap& InTypes) const override;;
diff --git a/Engine/Plugins/Animation/ControlRig/Source/ControlRig/Public/Units/RigDispatchFactory.h b/Engine/Plugins/Animation/ControlRig/Source/ControlRig/Public/Units/RigDispatchFactory.h
--- a/Engine/Plugins/Animation/ControlRig/Source/ControlRig/Public/Units/RigDispatchFactory.h
+++ b/Engine/Plugins/Animation/ControlRig/Source/ControlRig/Public/Units/RigDispatchFactory.h
@@ -18,11 +18,12 @@ struct CONTROLRIG_API FRigDispatchFactory : public FRigVMDispatchFactory
 		return FControlRigExecuteContext::StaticStruct();
 	}
 
-	virtual void RegisterDependencyTypes() const override
+	virtual void RegisterDependencyTypes_NoLock() const override
 	{
-		FRigVMRegistry::Get().FindOrAddType(FControlRigExecuteContext::StaticStruct());
-		FRigVMRegistry::Get().FindOrAddType(FRigElementKey::StaticStruct());
-    	FRigVMRegistry::Get().FindOrAddType(FCachedRigElement::StaticStruct());
+		FRigVMRegistry_NoLock& Registry = FRigVMRegistry_NoLock::GetForWrite();
+		Registry.FindOrAddType_NoLock(FControlRigExecuteContext::StaticStruct());
+		Registry.FindOrAddType_NoLock(FRigElementKey::StaticStruct());
+    	Registry.FindOrAddType_NoLock(FCachedRigElement::StaticStruct());
 	}
 
 #if WITH_EDITOR
diff --git a/Engine/Plugins/Runtime/RigVM/Source/RigVM/Private/RigVMCore/RigVM.cpp b/Engine/Plugins/Runtime/RigVM/Source/RigVM/Private/RigVMCore/RigVM.cpp
--- a/Engine/Plugins/Runtime/RigVM/Source/RigVM/Private/RigVMCore/RigVM.cpp
+++ b/Engine/Plugins/Runtime/RigVM/Source/RigVM/Private/RigVMCore/RigVM.cpp
@@ -643,7 +643,7 @@ int32 URigVM::AddRigVMFunction(const FString& InFunctionName)
 		return FunctionIndex;
 	}
 
-	const FRigVMFunction* Function = FRigVMRegistry::Get().FindFunction(*InFunctionName);
+	const FRigVMFunction* Function = FRigVMRegistry_RWLock::Get().FindFunction(*InFunctionName);
 	if (Function == nullptr)
 	{
 		return INDEX_NONE;
@@ -984,7 +984,7 @@ bool URigVM::ResolveFunctionsIfRequired()
 		for (int32 FunctionIndex = 0; FunctionIndex < FunctionNames.Num(); FunctionIndex++)
 		{
 			const FString FunctionNameString = FunctionNames[FunctionIndex].ToString();
-			if(const FRigVMFunction* Function = FRigVMRegistry::Get().FindFunction(*FunctionNameString, TypeResolver))
+			if(const FRigVMFunction* Function = FRigVMRegistry_RWLock::Get().FindFunction(*FunctionNameString, TypeResolver))
 			{
 				GetFunctions()[FunctionIndex] = Function;
 				GetFactories()[FunctionIndex] = Function->Factory;
diff --git a/Engine/Plugins/Runtime/RigVM/Source/RigVM/Private/RigVMCore/RigVMDispatchFactory.cpp b/Engine/Plugins/Runtime/RigVM/Source/RigVM/Private/RigVMCore/RigVMDispatchFactory.cpp
--- a/Engine/Plugins/Runtime/RigVM/Source/RigVM/Private/RigVMCore/RigVMDispatchFactory.cpp
+++ b/Engine/Plugins/Runtime/RigVM/Source/RigVM/Private/RigVMCore/RigVMDispatchFactory.cpp
@@ -61,7 +61,7 @@ FText FRigVMDispatchFactory::GetNodeTooltip(const FRigVMTemplateTypeMap& InTypes
 
 FString FRigVMDispatchFactory::GetArgumentDefaultValue(const FName& InArgumentName, TRigVMTypeIndex InTypeIndex) const
 {
-	if(FRigVMRegistry::Get().IsArrayType(InTypeIndex))
+	if(FRigVMRegistry_RWLock::Get().IsArrayType(InTypeIndex))
 	{
 		static const FString EmptyArrayString = TEXT("()");
 		return EmptyArrayString;
@@ -85,32 +85,32 @@ FString FRigVMDispatchFactory::GetArgumentDefaultValue(const FName& InArgumentNa
 	{
 		return FName(NAME_None).ToString();
 	}
-	if(InTypeIndex == FRigVMRegistry::Get().GetTypeIndex<FVector2D>())
+	if(InTypeIndex == FRigVMRegistry_RWLock::Get().GetTypeIndex<FVector2D>())
 	{
 		static FString DefaultValueString = GetDefaultValueForStruct(FVector2D::ZeroVector);
 		return DefaultValueString;
 	}
-	if(InTypeIndex == FRigVMRegistry::Get().GetTypeIndex<FVector>())
+	if(InTypeIndex == FRigVMRegistry_RWLock::Get().GetTypeIndex<FVector>())
 	{
 		static FString DefaultValueString = GetDefaultValueForStruct(FVector::ZeroVector); 
 		return DefaultValueString;
 	}
-	if(InTypeIndex == FRigVMRegistry::Get().GetTypeIndex<FRotator>())
+	if(InTypeIndex == FRigVMRegistry_RWLock::Get().GetTypeIndex<FRotator>())
 	{
 		static FString DefaultValueString = GetDefaultValueForStruct(FRotator::ZeroRotator); 
 		return DefaultValueString;
 	}
-	if(InTypeIndex == FRigVMRegistry::Get().GetTypeIndex<FQuat>())
+	if(InTypeIndex == FRigVMRegistry_RWLock::Get().GetTypeIndex<FQuat>())
 	{
 		static FString DefaultValueString = GetDefaultValueForStruct(FQuat::Identity); 
 		return DefaultValueString;
 	}
-	if(InTypeIndex == FRigVMRegistry::Get().GetTypeIndex<FTransform>())
+	if(InTypeIndex == FRigVMRegistry_RWLock::Get().GetTypeIndex<FTransform>())
 	{
 		static FString DefaultValueString = GetDefaultValueForStruct(FTransform::Identity); 
 		return DefaultValueString;
 	}
-	if(InTypeIndex == FRigVMRegistry::Get().GetTypeIndex<FLinearColor>())
+	if(InTypeIndex == FRigVMRegistry_RWLock::Get().GetTypeIndex<FLinearColor>())
 	{
 		static FString DefaultValueString = GetDefaultValueForStruct(FLinearColor::White); 
 		return DefaultValueString;
@@ -204,12 +204,12 @@ const TArray<FRigVMTemplateArgumentInfo>& FRigVMDispatchFactory::GetArgumentInfo
 	return EmptyArguments;
 }
 
-TArray<FRigVMExecuteArgument> FRigVMDispatchFactory::GetExecuteArguments(const FRigVMDispatchContext& InContext) const
+TArray<FRigVMExecuteArgument> FRigVMDispatchFactory::GetExecuteArguments_NoLock(const FRigVMDispatchContext& InContext) const
 {
 	TArray<FRigVMExecuteArgument> Arguments = GetExecuteArguments_Impl(InContext);
 	for(FRigVMExecuteArgument& Argument : Arguments)
 	{
-		if(Argument.TypeIndex != INDEX_NONE && FRigVMRegistry::Get().IsArrayType(Argument.TypeIndex))
+		if(Argument.TypeIndex != INDEX_NONE && FRigVMRegistry_NoLock::GetForRead().IsArrayType_NoLock(Argument.TypeIndex))
 		{
 			Argument.TypeIndex = RigVMTypeUtils::TypeIndex::ExecuteArray;
 		}
@@ -230,10 +230,14 @@ const TArray<FRigVMExecuteArgument>& FRigVMDispatchFactory::GetExecuteArguments_
 
 FRigVMFunctionPtr FRigVMDispatchFactory::GetOrCreateDispatchFunction(const FRigVMTemplateTypeMap& InTypes) const
 {
-	FScopeLock DispatchFunctionScopeLock(&FRigVMRegistry::DispatchFunctionMutex);
+	const FRigVMRegistryWriteLock _;
+	return GetOrCreateDispatchFunction_NoLock(InTypes);
+}
 
-	const FString PermutationName = GetPermutationNameImpl(InTypes);
-	if(const FRigVMFunction* ExistingFunction = FRigVMRegistry::Get().FindFunction(*PermutationName))
+FRigVMFunctionPtr FRigVMDispatchFactory::GetOrCreateDispatchFunction_NoLock(const FRigVMTemplateTypeMap& InTypes) const
+{
+	const FString PermutationName = GetPermutationNameImpl(InTypes, false);
+	if(const FRigVMFunction* ExistingFunction = FRigVMRegistry_NoLock::GetForRead().FindFunction_NoLock(*PermutationName))
 	{
 		return ExistingFunction->FunctionPtr;
 	}
@@ -241,29 +245,17 @@ FRigVMFunctionPtr FRigVMDispatchFactory::GetOrCreateDispatchFunction(const FRigV
 	return CreateDispatchFunction_NoLock(InTypes);
 }
 
-FRigVMFunctionPtr FRigVMDispatchFactory::CreateDispatchFunction(const FRigVMTemplateTypeMap& InTypes) const
-{
-	FScopeLock DispatchFunctionScopeLock(&FRigVMRegistry::DispatchFunctionMutex);
-	return CreateDispatchFunction_NoLock(InTypes);
-}
-
 FRigVMFunctionPtr FRigVMDispatchFactory::CreateDispatchFunction_NoLock(const FRigVMTemplateTypeMap& InTypes) const
 {
 	return GetDispatchFunctionImpl(InTypes);
 }
 
-TArray<FRigVMFunction> FRigVMDispatchFactory::CreateDispatchPredicates(const FRigVMTemplateTypeMap& InTypes) const
-{
-	FScopeLock DispatchPredicatesScopeLock(&FRigVMRegistry::DispatchPredicatesMutex);
-	return CreateDispatchPredicates_NoLock(InTypes);
-}
-
 TArray<FRigVMFunction> FRigVMDispatchFactory::CreateDispatchPredicates_NoLock(const FRigVMTemplateTypeMap& InTypes) const
 {
 	return GetDispatchPredicatesImpl(InTypes);
 }
 
-FString FRigVMDispatchFactory::GetPermutationName(const FRigVMTemplateTypeMap& InTypes) const
+FString FRigVMDispatchFactory::GetPermutationName(const FRigVMTemplateTypeMap& InTypes, const bool bLockRegistry) const
 {
 #if WITH_EDITOR
 	const TArray<FRigVMTemplateArgumentInfo>& Arguments = GetArgumentInfos();
@@ -275,7 +267,7 @@ FString FRigVMDispatchFactory::GetPermutationName(const FRigVMTemplateTypeMap& I
 		check(InTypes.Contains(Argument.Name));
 	}
 #endif
-	return GetPermutationNameImpl(InTypes);
+	return GetPermutationNameImpl(InTypes, bLockRegistry);
 }
 
 TArray<FRigVMTemplateArgumentInfo> FRigVMDispatchFactory::BuildArgumentListFromPrimaryArgument(const TArray<FRigVMTemplateArgumentInfo>& InInfos, const FName& InPrimaryArgumentName) const
@@ -350,10 +342,10 @@ TArray<FRigVMTemplateArgumentInfo> FRigVMDispatchFactory::BuildArgumentListFromP
 	return NewInfos;
 }
 
-FString FRigVMDispatchFactory::GetPermutationNameImpl(const FRigVMTemplateTypeMap& InTypes) const
+FString FRigVMDispatchFactory::GetPermutationNameImpl(const FRigVMTemplateTypeMap& InTypes, const bool bLockRegistry) const
 {
 	static constexpr TCHAR Format[] = TEXT("%s::%s");
-	const FString TypePairStrings = FRigVMTemplate::GetStringFromArgumentTypes(InTypes);
+	const FString TypePairStrings = FRigVMTemplate::GetStringFromArgumentTypes(InTypes, bLockRegistry);
 	return FString::Printf(Format, *GetFactoryName().ToString(), *TypePairStrings);
 }
 
@@ -363,15 +355,21 @@ bool FRigVMDispatchFactory::CopyProperty(const FProperty* InTargetProperty, uint
 	return URigVMMemoryStorage::CopyProperty(InTargetProperty, InTargetPtr, InSourceProperty, InSourcePtr);
 }
 
-const FRigVMTemplate* FRigVMDispatchFactory::GetTemplate() const
+const FRigVMTemplate* FRigVMDispatchFactory::GetTemplate_RWLock() const
+{
+	const FRigVMRegistryWriteLock _;
+	return GetTemplate_NoLock();
+}
+
+const FRigVMTemplate* FRigVMDispatchFactory::GetTemplate_NoLock() const
 {
 	// make sure to rely on the instance of this factory that's stored under the registry
-	FRigVMRegistry& Registry = FRigVMRegistry::Get();
+	FRigVMRegistry_NoLock& Registry = FRigVMRegistry_NoLock::GetForWrite();
 	const FName FactoryName = GetFactoryName();
-	const FRigVMDispatchFactory* ThisFactory = Registry.FindDispatchFactory(FactoryName);
+	const FRigVMDispatchFactory* ThisFactory = Registry.FindDispatchFactory_NoLock(FactoryName);
 	if(ThisFactory != this)
 	{
-		return ThisFactory->GetTemplate();
+		return ThisFactory->GetTemplate_NoLock();
 	}
 
 	FScopeLock GetTemplateScopeLock(&GetTemplateMutex);
@@ -380,7 +378,6 @@ const FRigVMTemplate* FRigVMDispatchFactory::GetTemplate() const
 	{
 		return CachedTemplate;
 	}
-
 	
 	// we don't allow execute types on arguments	
 	const TArray<FRigVMTemplateArgumentInfo>& Infos = GetArgumentInfos();
@@ -389,7 +386,7 @@ const FRigVMTemplate* FRigVMDispatchFactory::GetTemplate() const
 		const FRigVMTemplateArgument Argument = Info.GetArgument();
 		const int32 Index = Argument.IndexOfByPredicate([&](const TRigVMTypeIndex TypeIndex)
 		{
-			return Registry.IsExecuteType(TypeIndex);
+			return Registry.IsExecuteType_NoLock(TypeIndex);
 		});
 		
 		if (Index != INDEX_NONE)
@@ -404,16 +401,16 @@ const FRigVMTemplate* FRigVMDispatchFactory::GetTemplate() const
 	Delegates.GetDispatchFactoryDelegate = FRigVMTemplate_GetDispatchFactoryDelegate::CreateLambda(
 	[FactoryName]()
 	{
-		return FRigVMRegistry::Get().FindDispatchFactory(FactoryName);
+		return FRigVMRegistry_NoLock::GetForRead().FindDispatchFactory_NoLock(FactoryName);
 	});
 
-	CachedTemplate = Registry.AddTemplateFromArguments(GetFactoryName(), Infos, Delegates); 
+	CachedTemplate = Registry.AddTemplateFromArguments_NoLock(GetFactoryName(), Infos, Delegates); 
 	return CachedTemplate;
 }
 
 FName FRigVMDispatchFactory::GetTemplateNotation() const
 {
-	if(const FRigVMTemplate* Template = GetTemplate())
+	if(const FRigVMTemplate* Template = GetTemplate_RWLock())
 	{
 		return Template->GetNotation();
 	}
diff --git a/Engine/Plugins/Runtime/RigVM/Source/RigVM/Private/RigVMCore/RigVMExternalVariable.cpp b/Engine/Plugins/Runtime/RigVM/Source/RigVM/Private/RigVMCore/RigVMExternalVariable.cpp
--- a/Engine/Plugins/Runtime/RigVM/Source/RigVM/Private/RigVMCore/RigVMExternalVariable.cpp
+++ b/Engine/Plugins/Runtime/RigVM/Source/RigVM/Private/RigVMCore/RigVMExternalVariable.cpp
@@ -285,7 +285,7 @@ TRigVMTypeIndex FRigVMExternalVariable::GetTypeIndex() const
 {
 	if(IsValid(true))
 	{
-		return FRigVMRegistry::Get().GetTypeIndexFromCPPType(GetExtendedCPPType().ToString());
+		return FRigVMRegistry_RWLock::Get().GetTypeIndexFromCPPType(GetExtendedCPPType().ToString());
 	}
 	return INDEX_NONE;
 }
diff --git a/Engine/Plugins/Runtime/RigVM/Source/RigVM/Private/RigVMCore/RigVMFunction.cpp b/Engine/Plugins/Runtime/RigVM/Source/RigVM/Private/RigVMCore/RigVMFunction.cpp
--- a/Engine/Plugins/Runtime/RigVM/Source/RigVM/Private/RigVMCore/RigVMFunction.cpp
+++ b/Engine/Plugins/Runtime/RigVM/Source/RigVM/Private/RigVMCore/RigVMFunction.cpp
@@ -81,6 +81,8 @@ const TArray<TRigVMTypeIndex>& FRigVMFunction::GetArgumentTypeIndices() const
 {
 	if(ArgumentTypeIndices.IsEmpty() && !Arguments.IsEmpty())
 	{
+		const FRigVMRegistryWriteLock _;
+
 		if(Struct)
 		{
 			for(const FRigVMFunctionArgument& Argument : Arguments)
@@ -92,11 +94,11 @@ const TArray<TRigVMTypeIndex>& FRigVMFunction::GetArgumentTypeIndices() const
 					RigVMPropertyUtils::GetTypeFromProperty(Property, CPPType, CPPTypeObject);
 
 					const FRigVMTemplateArgumentType Type(CPPType, CPPTypeObject);
-					ArgumentTypeIndices.Add(FRigVMRegistry::Get().FindOrAddType(Type));
+					ArgumentTypeIndices.Add(FRigVMRegistry_NoLock::GetForWrite().FindOrAddType_NoLock(Type));
 				}
 			}
 		}
-		else if(const FRigVMTemplate* Template = GetTemplate())
+		else if(const FRigVMTemplate* Template = GetTemplate_NoLock())
 		{
 			const int32 PermutationIndex = Template->FindPermutation(this);
 			check(PermutationIndex != INDEX_NONE);
@@ -105,7 +107,7 @@ const TArray<TRigVMTypeIndex>& FRigVMFunction::GetArgumentTypeIndices() const
 			{
 				const FRigVMTemplateArgument* TemplateArgument = Template->FindArgument(FunctionArgument.Name);
 				check(TemplateArgument);
-				ArgumentTypeIndices.Add(TemplateArgument->GetTypeIndex(PermutationIndex));
+				ArgumentTypeIndices.Add(TemplateArgument->GetTypeIndex_NoLock(PermutationIndex));
 			}
 		}
 		else
@@ -117,13 +119,19 @@ const TArray<TRigVMTypeIndex>& FRigVMFunction::GetArgumentTypeIndices() const
 }
 
 const FRigVMTemplate* FRigVMFunction::GetTemplate() const
+{
+	const FRigVMRegistryReadLock _;
+	return GetTemplate_NoLock();
+}
+
+const FRigVMTemplate* FRigVMFunction::GetTemplate_NoLock() const
 {
 	if(TemplateIndex == INDEX_NONE)
 	{
 		return nullptr;
 	}
 
-	const FRigVMTemplate* Template = &FRigVMRegistry::Get().GetTemplates()[TemplateIndex];
+	const FRigVMTemplate* Template = &FRigVMRegistry_NoLock::GetForRead().GetTemplates_NoLock()[TemplateIndex];
 	if(Template->NumPermutations() <= 1)
 	{
 		return nullptr;
@@ -132,7 +140,7 @@ const FRigVMTemplate* FRigVMFunction::GetTemplate() const
 	return Template;
 }
 
-const UScriptStruct* FRigVMFunction::GetExecuteContextStruct() const
+const UScriptStruct* FRigVMFunction::GetExecuteContextStruct(bool bLockRegistry) const
 {
 	if(Factory)
 	{
@@ -144,7 +152,7 @@ const UScriptStruct* FRigVMFunction::GetExecuteContextStruct() const
 		FString ExecuteContextName;
 		if(Struct->GetStringMetaDataHierarchical(FRigVMStruct::ExecuteContextName, &ExecuteContextName))
 		{
-			const FRigVMTemplateArgumentType& Type = FRigVMRegistry::Get().FindTypeFromCPPType(ExecuteContextName);
+			const FRigVMTemplateArgumentType& Type = FRigVMRegistry_RWLock::Get().FindTypeFromCPPType(ExecuteContextName, bLockRegistry);
 			if(const UScriptStruct* ExecuteContextStruct = Cast<UScriptStruct>(Type.CPPTypeObject))
 			{
 				return ExecuteContextStruct;
diff --git a/Engine/Plugins/Runtime/RigVM/Source/RigVM/Private/RigVMCore/RigVMRegistry.cpp b/Engine/Plugins/Runtime/RigVM/Source/RigVM/Private/RigVMCore/RigVMRegistry.cpp
--- a/Engine/Plugins/Runtime/RigVM/Source/RigVM/Private/RigVMCore/RigVMRegistry.cpp
+++ b/Engine/Plugins/Runtime/RigVM/Source/RigVM/Private/RigVMCore/RigVMRegistry.cpp
@@ -16,44 +16,33 @@
 #include "RigVMFunctions/RigVMDispatch_Core.h"
 #include "Interfaces/IPluginManager.h"
 
-FCriticalSection FRigVMRegistry::FindOrAddTypeMutex;
-
-FCriticalSection FRigVMRegistry::FunctionRegistryMutex;
-FCriticalSection FRigVMRegistry::FactoryRegistryMutex;
-FCriticalSection FRigVMRegistry::TemplateRegistryMutex;
-
-FCriticalSection FRigVMRegistry::DispatchFunctionMutex;
-FCriticalSection FRigVMRegistry::DispatchPredicatesMutex;
-
-
 // When the object system has been completely loaded, load in all the engine types that we haven't registered already in InitializeIfNeeded 
 static FDelayedAutoRegisterHelper GRigVMRegistrySingletonHelper(EDelayedRegisterRunPhase::EndOfEngineInit, []() -> void
 {
-	FRigVMRegistry::Get().RefreshEngineTypes();
+	FRigVMRegistry_RWLock::Get().RefreshEngineTypes();
 });
 
 
-FRigVMRegistry::FRigVMRegistry() :
+FRigVMRegistry_NoLock::FRigVMRegistry_NoLock() :
 	bIsRefreshingEngineTypes(false),
 	bEverRefreshedEngineTypes(false)
 {
-	Initialize();
 }
 
-FRigVMRegistry::~FRigVMRegistry()
+FRigVMRegistry_NoLock& FRigVMRegistry_NoLock::Get(ELockType InLockType)
 {
-	Reset();
+#if WITH_EDITOR
+	FRigVMRegistry_RWLock::EnsureLocked(InLockType);
+#endif
+	return FRigVMRegistry_RWLock::Get();
 }
 
-FRigVMRegistry& FRigVMRegistry::Get()
+FRigVMRegistry_NoLock::~FRigVMRegistry_NoLock()
 {
-	// static in a function scope ensures that the GC system is initiated before 
-	// the registry constructor is called
-	static FRigVMRegistry s_RigVMRegistry;
-	return s_RigVMRegistry;
+	FRigVMRegistry_NoLock::Reset_NoLock();
 }
 
-void FRigVMRegistry::AddReferencedObjects(FReferenceCollector& Collector)
+void FRigVMRegistry_NoLock::AddReferencedObjects(FReferenceCollector& Collector)
 {
 	// registry should hold strong references to these type objects
 	// otherwise GC may remove them without the registry known it
@@ -69,12 +58,12 @@ void FRigVMRegistry::AddReferencedObjects(FReferenceCollector& Collector)
 	}
 }
 
-FString FRigVMRegistry::GetReferencerName() const
+FString FRigVMRegistry_NoLock::GetReferencerName() const
 {
 	return TEXT("FRigVMRegistry");
 }
 
-const TArray<UScriptStruct*>& FRigVMRegistry::GetMathTypes()
+const TArray<UScriptStruct*>& FRigVMRegistry_NoLock::GetMathTypes()
 {
 	// The list of base math types to automatically register 
 	static const TArray<UScriptStruct*> MathTypes = { 
@@ -93,14 +82,14 @@ const TArray<UScriptStruct*>& FRigVMRegistry::GetMathTypes()
 	return MathTypes;
 }
 
-uint32 FRigVMRegistry::GetHashForType(TRigVMTypeIndex InTypeIndex) const
+uint32 FRigVMRegistry_NoLock::GetHashForType_NoLock(TRigVMTypeIndex InTypeIndex) const
 {
 	if(!Types.IsValidIndex(InTypeIndex))
 	{
 		return UINT32_MAX;
 	}
 
-	FRigVMRegistry* MutableThis = (FRigVMRegistry*)this; 
+	FRigVMRegistry_NoLock* MutableThis = (FRigVMRegistry_NoLock*)this; 
 	FTypeInfo& TypeInfo = MutableThis->Types[InTypeIndex];
 	
 	if(TypeInfo.Hash != UINT32_MAX)
@@ -108,18 +97,18 @@ uint32 FRigVMRegistry::GetHashForType(TRigVMTypeIndex InTypeIndex) const
 		return TypeInfo.Hash;
 	}
 
-	uint32 Hash = INDEX_NONE;
+	uint32 Hash;
 	if(const UScriptStruct* ScriptStruct = Cast<UScriptStruct>(TypeInfo.Type.CPPTypeObject))
 	{
-		Hash = GetHashForScriptStruct(ScriptStruct, false);
+		Hash = GetHashForScriptStruct_NoLock(ScriptStruct, false);
 	}
 	else if(const UStruct* Struct = Cast<UStruct>(TypeInfo.Type.CPPTypeObject))
 	{
-		Hash = GetHashForStruct(Struct);
+		Hash = GetHashForStruct_NoLock(Struct);
 	}
 	else if(const UEnum* Enum = Cast<UEnum>(TypeInfo.Type.CPPTypeObject))
     {
-    	Hash = GetHashForEnum(Enum, false);
+    	Hash = GetHashForEnum_NoLock(Enum, false);
     }
     else
     {
@@ -136,43 +125,43 @@ uint32 FRigVMRegistry::GetHashForType(TRigVMTypeIndex InTypeIndex) const
 	return Hash;
 }
 
-uint32 FRigVMRegistry::GetHashForScriptStruct(const UScriptStruct* InScriptStruct, bool bCheckTypeIndex) const
+uint32 FRigVMRegistry_NoLock::GetHashForScriptStruct_NoLock(const UScriptStruct* InScriptStruct, bool bCheckTypeIndex) const
 {
 	if(bCheckTypeIndex)
 	{
-		const TRigVMTypeIndex TypeIndex = GetTypeIndex(*InScriptStruct->GetStructCPPName(), (UObject*)InScriptStruct);
+		const TRigVMTypeIndex TypeIndex = GetTypeIndex_NoLock(*InScriptStruct->GetStructCPPName(), (UObject*)InScriptStruct);
 		if(TypeIndex != INDEX_NONE)
 		{
-			return GetHashForType(TypeIndex);
+			return GetHashForType_NoLock(TypeIndex);
 		}
 	}
 	
 	const uint32 NameHash = GetTypeHash(InScriptStruct->GetStructCPPName());
-	return HashCombine(NameHash, GetHashForStruct(InScriptStruct));
+	return HashCombine(NameHash, GetHashForStruct_NoLock(InScriptStruct));
 }
 
-uint32 FRigVMRegistry::GetHashForStruct(const UStruct* InStruct) const
+uint32 FRigVMRegistry_NoLock::GetHashForStruct_NoLock(const UStruct* InStruct) const
 {
 	uint32 Hash = GetTypeHash(InStruct->GetPathName());
 	for (TFieldIterator<FProperty> It(InStruct); It; ++It)
 	{
 		const FProperty* Property = *It;
-		if(IsAllowedType(Property))
+		if(IsAllowedType_NoLock(Property))
 		{
-			Hash = HashCombine(Hash, GetHashForProperty(Property));
+			Hash = HashCombine(Hash, GetHashForProperty_NoLock(Property));
 		}
 	}
 	return Hash;
 }
 
-uint32 FRigVMRegistry::GetHashForEnum(const UEnum* InEnum, bool bCheckTypeIndex) const
+uint32 FRigVMRegistry_NoLock::GetHashForEnum_NoLock(const UEnum* InEnum, bool bCheckTypeIndex) const
 {
 	if(bCheckTypeIndex)
 	{
-		const TRigVMTypeIndex TypeIndex = GetTypeIndex(*InEnum->CppType, (UObject*)InEnum);
+		const TRigVMTypeIndex TypeIndex = GetTypeIndex_NoLock(*InEnum->CppType, (UObject*)InEnum);
 		if(TypeIndex != INDEX_NONE)
 		{
-			return GetHashForType(TypeIndex);
+			return GetHashForType_NoLock(TypeIndex);
 		}
 	}
 	
@@ -185,7 +174,7 @@ uint32 FRigVMRegistry::GetHashForEnum(const UEnum* InEnum, bool bCheckTypeIndex)
 	return Hash;
 }
 
-uint32 FRigVMRegistry::GetHashForProperty(const FProperty* InProperty) const
+uint32 FRigVMRegistry_NoLock::GetHashForProperty_NoLock(const FProperty* InProperty) const
 {
 	uint32 Hash = GetTypeHash(InProperty->GetName());
 
@@ -200,31 +189,47 @@ uint32 FRigVMRegistry::GetHashForProperty(const FProperty* InProperty) const
 	
 	if(const FStructProperty* StructProperty = CastField<FStructProperty>(InProperty))
 	{
-		Hash = HashCombine(Hash, GetHashForStruct(StructProperty->Struct));
+		Hash = HashCombine(Hash, GetHashForStruct_NoLock(StructProperty->Struct));
 	}
 	else if(const FByteProperty* ByteProperty = CastField<FByteProperty>(InProperty))
 	{
 		if(ByteProperty->Enum)
 		{
-			Hash = HashCombine(Hash, GetHashForEnum(ByteProperty->Enum));
+			Hash = HashCombine(Hash, GetHashForEnum_NoLock(ByteProperty->Enum));
 		}
 	}
 	else if(const FEnumProperty* EnumProperty = CastField<FEnumProperty>(InProperty))
 	{
-		Hash = HashCombine(Hash, GetHashForEnum(EnumProperty->GetEnum()));
+		Hash = HashCombine(Hash, GetHashForEnum_NoLock(EnumProperty->GetEnum()));
 	}
 	
 	return Hash;
 }
 
-
-void FRigVMRegistry::Initialize()
+void FRigVMRegistry_NoLock::RebuildRegistry_NoLock()
 {
-	// this should not be necessary since the initialize is used only
-	// on a constructor on a static variable (thread safe)
-	// but in case code paths change in the future we'll also lock here.
-	const FScopeLock FindOrAddTypeLock(&FindOrAddTypeMutex);
+	Reset_NoLock();
 	
+	Types.Reset();
+	TypeToIndex.Reset();
+	Functions.Empty();
+	Templates.Empty();
+	DeprecatedTemplates.Empty();
+	Factories.Reset();
+	FunctionNameToIndex.Reset();
+	StructNameToPredicates.Reset();
+	TemplateNotationToIndex.Reset();
+	DeprecatedTemplateNotationToIndex.Reset();
+	TypesPerCategory.Reset();
+	TemplatesPerCategory.Reset();
+	UserDefinedTypeToIndex.Reset();
+	AllowedClasses.Reset();
+
+	Initialize(false);
+}
+
+void FRigVMRegistry_NoLock::Initialize_NoLock()
+{
 	Types.Reserve(512);
 	TypeToIndex.Reserve(512);
 	TypesPerCategory.Reserve(19);
@@ -300,7 +305,7 @@ void FRigVMRegistry::Initialize()
 	// hook the registry to prepare for engine shutdown
 	FCoreDelegates::OnExit.AddLambda([&]()
 	{
-		Reset();
+		Reset_NoLock();
 
 		if (FAssetRegistryModule* AssetRegistryModule = FModuleManager::GetModulePtr<FAssetRegistryModule>(TEXT("AssetRegistry")))
 		{
@@ -315,25 +320,10 @@ void FRigVMRegistry::Initialize()
 
 		UE::Anim::AttributeTypes::GetOnAttributeTypesChanged().RemoveAll(this);
 	});
-	
-	FAssetRegistryModule& AssetRegistryModule = FModuleManager::LoadModuleChecked<FAssetRegistryModule>(TEXT("AssetRegistry"));
-	AssetRegistryModule.Get().OnAssetRemoved().AddRaw(this, &FRigVMRegistry::OnAssetRemoved);
-	AssetRegistryModule.Get().OnAssetRenamed().AddRaw(this, &FRigVMRegistry::OnAssetRenamed);
-
-	IPluginManager::Get().OnPluginUnmounted().AddRaw(this, &FRigVMRegistry::OnPluginUnloaded);
-	
-	UE::Anim::AttributeTypes::GetOnAttributeTypesChanged().AddRaw(this, &FRigVMRegistry::OnAnimationAttributeTypesChanged);
-}
-
-void FRigVMRegistry::RefreshEngineTypes()
-{
-	FScopeLock FindOrAddTypeLock(&FindOrAddTypeMutex);
-	RefreshEngineTypes_NoLock();
 }
 
-void FRigVMRegistry::RefreshEngineTypesIfRequired()
+void FRigVMRegistry_NoLock::RefreshEngineTypesIfRequired_NoLock()
 {
-	FScopeLock FindOrAddTypeLock(&FindOrAddTypeMutex);
 	if(bEverRefreshedEngineTypes)
 	{
 		return;
@@ -341,7 +331,7 @@ void FRigVMRegistry::RefreshEngineTypesIfRequired()
 	RefreshEngineTypes_NoLock();
 }
 
-void FRigVMRegistry::RefreshEngineTypes_NoLock()
+void FRigVMRegistry_NoLock::RefreshEngineTypes_NoLock()
 {
 	TGuardValue<bool> EnableGuardRefresh(bIsRefreshingEngineTypes, true);
 
@@ -368,14 +358,14 @@ void FRigVMRegistry::RefreshEngineTypes_NoLock()
 		}
 		else if(AllowedStructs.Contains(ScriptStruct))
 		{
-			FindOrAddType(FRigVMTemplateArgumentType(ScriptStruct));
+			FindOrAddType_NoLock(FRigVMTemplateArgumentType(ScriptStruct));
 		}
 	}
 
 	for (TObjectIterator<UEnum> EnumIt; EnumIt; ++EnumIt)
 	{
 		UEnum* Enum = *EnumIt;
-		if(IsAllowedType(Enum))
+		if(IsAllowedType_NoLock(Enum))
 		{
 			const FString CPPType = Enum->CppType.IsEmpty() ? Enum->GetName() : Enum->CppType;
 			FindOrAddType_NoLock(FRigVMTemplateArgumentType(*CPPType, Enum), false);
@@ -385,7 +375,7 @@ void FRigVMRegistry::RefreshEngineTypes_NoLock()
 	for (TObjectIterator<UClass> ClassIt; ClassIt; ++ClassIt)
 	{
 		UClass* Class = *ClassIt;
-		if (IsAllowedType(Class))
+		if (IsAllowedType_NoLock(Class))
 		{
 			// Register both the class and the object type for use
 			FindOrAddType_NoLock(FRigVMTemplateArgumentType(Class, RigVMTypeUtils::EClassArgType::AsClass), false);
@@ -396,7 +386,7 @@ void FRigVMRegistry::RefreshEngineTypes_NoLock()
 	// Register all dispatch factories only after all other types have been registered.
 	for (UScriptStruct* DispatchFactoryStruct: DispatchFactoriesToRegister)
 	{
-		RegisterFactory(DispatchFactoryStruct);
+		RegisterFactory_NoLock(DispatchFactoryStruct);
 	}
 
 	const int32 NumTypesNow = Types.Num();
@@ -422,7 +412,7 @@ void FRigVMRegistry::RefreshEngineTypes_NoLock()
 	bEverRefreshedEngineTypes = true;
 }
 
-void FRigVMRegistry::OnAssetRenamed(const FAssetData& InAssetData, const FString& InOldObjectPath)
+void FRigVMRegistry_NoLock::OnAssetRenamed_NoLock(const FAssetData& InAssetData, const FString& InOldObjectPath)
 {
 	const FSoftObjectPath OldPath(InOldObjectPath);
 	
@@ -434,15 +424,15 @@ void FRigVMRegistry::OnAssetRenamed(const FAssetData& InAssetData, const FString
 	}
 }
 
-void FRigVMRegistry::OnAssetRemoved(const FAssetData& InAssetData)
+void FRigVMRegistry_NoLock::OnAssetRemoved_NoLock(const FAssetData& InAssetData)
 {
-	if (RemoveType(InAssetData.ToSoftObjectPath(), InAssetData.GetClass()))
+	if (RemoveType_NoLock(InAssetData.ToSoftObjectPath(), InAssetData.GetClass()))
 	{
 		OnRigVMRegistryChangedDelegate.Broadcast();
 	}
 }
 
-void FRigVMRegistry::OnPluginUnloaded(IPlugin& InPlugin)
+void FRigVMRegistry_NoLock::OnPluginUnloaded_NoLock(IPlugin& InPlugin)
 {
 	const FString PluginContentPath = InPlugin.GetMountedAssetPath();
 
@@ -467,7 +457,7 @@ void FRigVMRegistry::OnPluginUnloaded(IPlugin& InPlugin)
 			ObjectClass = TypeObject->GetClass();
 		}
 		
-		if (RemoveType(ObjectPath, ObjectClass))
+		if (RemoveType_NoLock(ObjectPath, ObjectClass))
 		{
 			bRegistryChanged = true;
 		}
@@ -479,7 +469,7 @@ void FRigVMRegistry::OnPluginUnloaded(IPlugin& InPlugin)
 	}
 }
 
-void FRigVMRegistry::OnAnimationAttributeTypesChanged(const UScriptStruct* InStruct, bool bIsAdded)
+void FRigVMRegistry_NoLock::OnAnimationAttributeTypesChanged_NoLock(const UScriptStruct* InStruct, bool bIsAdded)
 {
 	if (!ensure(InStruct))
 	{
@@ -488,13 +478,13 @@ void FRigVMRegistry::OnAnimationAttributeTypesChanged(const UScriptStruct* InStr
 
 	if (bIsAdded)
 	{
-		FindOrAddType(FRigVMTemplateArgumentType(const_cast<UScriptStruct*>(InStruct)));
+		FindOrAddType_NoLock(FRigVMTemplateArgumentType(const_cast<UScriptStruct*>(InStruct)));
 		OnRigVMRegistryChangedDelegate.Broadcast();		
 	}
 }
 
 
-void FRigVMRegistry::Reset()
+void FRigVMRegistry_NoLock::Reset_NoLock()
 {
 	for(FRigVMDispatchFactory* Factory : Factories)
 	{
@@ -507,18 +497,12 @@ void FRigVMRegistry::Reset()
 	Factories.Reset();
 }
 
-TRigVMTypeIndex FRigVMRegistry::FindOrAddType(const FRigVMTemplateArgumentType& InType, bool bForce)
-{
-	const FScopeLock FindOrAddTypesLock(&FindOrAddTypeMutex);
-	return FindOrAddType_NoLock(InType, bForce);
-}
-
-TRigVMTypeIndex FRigVMRegistry::FindOrAddType_NoLock(const FRigVMTemplateArgumentType& InType, bool bForce)
+TRigVMTypeIndex FRigVMRegistry_NoLock::FindOrAddType_NoLock(const FRigVMTemplateArgumentType& InType, bool bForce)
 {
 	// we don't use a mutex here since by the time the engine relies on worker
 	// thread for execution or async loading all types will have been registered.
 	
-	TRigVMTypeIndex Index = GetTypeIndex(InType);
+	TRigVMTypeIndex Index = GetTypeIndex_NoLock(InType);
 	if(Index == INDEX_NONE)
 	{
 		FRigVMTemplateArgumentType ElementType = InType;
@@ -532,21 +516,21 @@ TRigVMTypeIndex FRigVMRegistry::FindOrAddType_NoLock(const FRigVMTemplateArgumen
 		{
 			if(const UClass* Class = Cast<UClass>(CPPTypeObject))
 			{
-				if(!IsAllowedType(Class))
+				if(!IsAllowedType_NoLock(Class))
 				{
 					return Index;
 				}	
 			}
 			else if(const UEnum* Enum = Cast<UEnum>(CPPTypeObject))
 			{
-				if(!IsAllowedType(Enum))
+				if(!IsAllowedType_NoLock(Enum))
 				{
 					return Index;
 				}
 			}
 			else if(const UStruct* Struct = Cast<UStruct>(CPPTypeObject))
 			{
-				if(!IsAllowedType(Struct))
+				if(!IsAllowedType_NoLock(Struct))
 				{					
 					return Index;
 				}
@@ -601,7 +585,7 @@ TRigVMTypeIndex FRigVMRegistry::FindOrAddType_NoLock(const FRigVMTemplateArgumen
 		TArray<TPair<FRigVMTemplateArgument::ETypeCategory, int32>> ToPropagate;
 		auto RegisterNewType = [&](FRigVMTemplateArgument::ETypeCategory InCategory, int32 NewIndex)
 		{
-			RegisterTypeInCategory(InCategory, NewIndex);
+			RegisterTypeInCategory_NoLock(InCategory, NewIndex);
 			ToPropagate.Emplace(InCategory, NewIndex);
 		}; 
 
@@ -640,7 +624,7 @@ TRigVMTypeIndex FRigVMRegistry::FindOrAddType_NoLock(const FRigVMTemplateArgumen
 					}
 				}
 			}
-			else if(const UClass* Class = Cast<UClass>(CPPTypeObject))
+			else if(CPPTypeObject->IsA<UClass>())
 			{
 				switch(ArrayDimension)
 				{
@@ -665,7 +649,7 @@ TRigVMTypeIndex FRigVMRegistry::FindOrAddType_NoLock(const FRigVMTemplateArgumen
 					}
 				}
 			}
-			else if(const UEnum* Enum = Cast<UEnum>(CPPTypeObject))
+			else if(CPPTypeObject->IsA<UEnum>())
 			{
 				switch(ArrayDimension)
 				{
@@ -753,7 +737,7 @@ TRigVMTypeIndex FRigVMRegistry::FindOrAddType_NoLock(const FRigVMTemplateArgumen
 		// propagate new type to templates once they have all been added to the categories
 		for (const auto& [Category, NewIndex]: ToPropagate)
 		{
-			PropagateTypeAddedToCategory(Category, NewIndex);	
+			PropagateTypeAddedToCategory_NoLock(Category, NewIndex);	
 		}
 
 		// if the type is a structure
@@ -763,7 +747,7 @@ TRigVMTypeIndex FRigVMRegistry::FindOrAddType_NoLock(const FRigVMTemplateArgumen
 			for (TFieldIterator<FProperty> It(Struct); It; ++It)
 			{
 				FProperty* Property = *It;
-				if(IsAllowedType(Property))
+				if(IsAllowedType_NoLock(Property))
 				{
 					// by creating a template argument for the child property
 					// the type will be added by calling ::FindOrAddType_Internal recursively.
@@ -780,12 +764,12 @@ TRigVMTypeIndex FRigVMRegistry::FindOrAddType_NoLock(const FRigVMTemplateArgumen
 			}			
 		}
 		
-		Index = GetTypeIndex(InType);
+		Index = GetTypeIndex_NoLock(InType);
 		if (IsValid(CPPTypeObject))
 		{
 			if (CPPTypeObject->IsA<UUserDefinedStruct>() || CPPTypeObject->IsA<UUserDefinedEnum>())
 			{
-				TRigVMTypeIndex ElementTypeIndex = GetTypeIndex(ElementType);
+				TRigVMTypeIndex ElementTypeIndex = GetTypeIndex_NoLock(ElementType);
 				// used to track name changes to user defined types, stores the element type index, see RemoveType()
 				UserDefinedTypeToIndex.FindOrAdd(CPPTypeObject) = ElementTypeIndex;
 			}
@@ -797,14 +781,14 @@ TRigVMTypeIndex FRigVMRegistry::FindOrAddType_NoLock(const FRigVMTemplateArgumen
 	return Index;
 }
 
-void FRigVMRegistry::RegisterTypeInCategory(FRigVMTemplateArgument::ETypeCategory InCategory, TRigVMTypeIndex InTypeIndex)
+void FRigVMRegistry_NoLock::RegisterTypeInCategory_NoLock(FRigVMTemplateArgument::ETypeCategory InCategory, TRigVMTypeIndex InTypeIndex)
 {
 	check(InCategory != FRigVMTemplateArgument::ETypeCategory_Invalid);
 
 	TypesPerCategory.FindChecked(InCategory).Add(InTypeIndex);
 }
 
-void FRigVMRegistry::PropagateTypeAddedToCategory(const FRigVMTemplateArgument::ETypeCategory InCategory, const TRigVMTypeIndex InTypeIndex)
+void FRigVMRegistry_NoLock::PropagateTypeAddedToCategory_NoLock(const FRigVMTemplateArgument::ETypeCategory InCategory, const TRigVMTypeIndex InTypeIndex)
 {
 	if(bIsRefreshingEngineTypes)
 	{
@@ -824,7 +808,7 @@ void FRigVMRegistry::PropagateTypeAddedToCategory(const FRigVMTemplateArgument::
 	}
 }
 
-bool FRigVMRegistry::RemoveType(const FSoftObjectPath& InObjectPath, const UClass* InObjectClass)
+bool FRigVMRegistry_NoLock::RemoveType_NoLock(const FSoftObjectPath& InObjectPath, const UClass* InObjectClass)
 {
 	if (const TRigVMTypeIndex* TypeIndexPtr = UserDefinedTypeToIndex.Find(InObjectPath))
 	{
@@ -837,17 +821,17 @@ bool FRigVMRegistry::RemoveType(const FSoftObjectPath& InObjectPath, const UClas
 			return false;
 		}
 
-		check(!IsArrayType(TypeIndex));
+		check(!IsArrayType_NoLock(TypeIndex));
 
 		TArray<TRigVMTypeIndex> Indices;
 		Indices.Init(INDEX_NONE, 3);
 		Indices[0] = TypeIndex;
-		Indices[1] = GetArrayTypeFromBaseTypeIndex(Indices[0]);
+		Indices[1] = GetArrayTypeFromBaseTypeIndex_NoLock(Indices[0]);
 
 		// any type that can be removed should have 3 entries in the registry
 		if (ensure(Indices[1] != INDEX_NONE))
 		{
-			Indices[2] = GetArrayTypeFromBaseTypeIndex(Indices[1]);
+			Indices[2] = GetArrayTypeFromBaseTypeIndex_NoLock(Indices[1]);
 		}
 		
 		for (int32 ArrayDimension=0; ArrayDimension<3; ++ArrayDimension)
@@ -866,20 +850,20 @@ bool FRigVMRegistry::RemoveType(const FSoftObjectPath& InObjectPath, const UClas
 				default:
 				case 0:
 					{
-						RemoveTypeInCategory(FRigVMTemplateArgument::ETypeCategory_SingleEnumValue, Index);
-						RemoveTypeInCategory(FRigVMTemplateArgument::ETypeCategory_SingleAnyValue, Index);
+						RemoveTypeInCategory_NoLock(FRigVMTemplateArgument::ETypeCategory_SingleEnumValue, Index);
+						RemoveTypeInCategory_NoLock(FRigVMTemplateArgument::ETypeCategory_SingleAnyValue, Index);
 						break;
 					}
 				case 1:
 					{
-						RemoveTypeInCategory(FRigVMTemplateArgument::ETypeCategory_ArrayEnumValue, Index);
-						RemoveTypeInCategory(FRigVMTemplateArgument::ETypeCategory_ArrayAnyValue, Index);
+						RemoveTypeInCategory_NoLock(FRigVMTemplateArgument::ETypeCategory_ArrayEnumValue, Index);
+						RemoveTypeInCategory_NoLock(FRigVMTemplateArgument::ETypeCategory_ArrayAnyValue, Index);
 						break;
 					}
 				case 2:
 					{
-						RemoveTypeInCategory(FRigVMTemplateArgument::ETypeCategory_ArrayArrayEnumValue, Index);
-						RemoveTypeInCategory(FRigVMTemplateArgument::ETypeCategory_ArrayArrayAnyValue, Index);
+						RemoveTypeInCategory_NoLock(FRigVMTemplateArgument::ETypeCategory_ArrayArrayEnumValue, Index);
+						RemoveTypeInCategory_NoLock(FRigVMTemplateArgument::ETypeCategory_ArrayArrayAnyValue, Index);
 						break;
 					}
 				}
@@ -891,27 +875,27 @@ bool FRigVMRegistry::RemoveType(const FSoftObjectPath& InObjectPath, const UClas
 				default:
 				case 0:
 					{
-						RemoveTypeInCategory(FRigVMTemplateArgument::ETypeCategory_SingleScriptStructValue, Index);
-						RemoveTypeInCategory(FRigVMTemplateArgument::ETypeCategory_SingleAnyValue, Index);
+						RemoveTypeInCategory_NoLock(FRigVMTemplateArgument::ETypeCategory_SingleScriptStructValue, Index);
+						RemoveTypeInCategory_NoLock(FRigVMTemplateArgument::ETypeCategory_SingleAnyValue, Index);
 						break;
 					}
 				case 1:
 					{
-						RemoveTypeInCategory(FRigVMTemplateArgument::ETypeCategory_ArrayScriptStructValue, Index);
-						RemoveTypeInCategory(FRigVMTemplateArgument::ETypeCategory_ArrayAnyValue, Index);
+						RemoveTypeInCategory_NoLock(FRigVMTemplateArgument::ETypeCategory_ArrayScriptStructValue, Index);
+						RemoveTypeInCategory_NoLock(FRigVMTemplateArgument::ETypeCategory_ArrayAnyValue, Index);
 						break;
 					}
 				case 2:
 					{
-						RemoveTypeInCategory(FRigVMTemplateArgument::ETypeCategory_ArrayArrayScriptStructValue, Index);
-						RemoveTypeInCategory(FRigVMTemplateArgument::ETypeCategory_ArrayArrayAnyValue, Index);
+						RemoveTypeInCategory_NoLock(FRigVMTemplateArgument::ETypeCategory_ArrayArrayScriptStructValue, Index);
+						RemoveTypeInCategory_NoLock(FRigVMTemplateArgument::ETypeCategory_ArrayArrayAnyValue, Index);
 						break;
 					}
 				}
 			}
 
 			// remove the type from the registry entirely
-			TypeToIndex.Remove(GetType(Index));
+			TypeToIndex.Remove(GetType_NoLock(Index));
 			Types[Index] = FTypeInfo();
 		}
 
@@ -921,7 +905,7 @@ bool FRigVMRegistry::RemoveType(const FSoftObjectPath& InObjectPath, const UClas
 	return false;
 }
 
-void FRigVMRegistry::RemoveTypeInCategory(FRigVMTemplateArgument::ETypeCategory InCategory, TRigVMTypeIndex InTypeIndex)
+void FRigVMRegistry_NoLock::RemoveTypeInCategory_NoLock(FRigVMTemplateArgument::ETypeCategory InCategory, TRigVMTypeIndex InTypeIndex)
 {
 	check(InCategory != FRigVMTemplateArgument::ETypeCategory_Invalid);
 
@@ -935,7 +919,7 @@ void FRigVMRegistry::RemoveTypeInCategory(FRigVMTemplateArgument::ETypeCategory
 	}
 }
 
-TRigVMTypeIndex FRigVMRegistry::GetTypeIndex(const FRigVMTemplateArgumentType& InType) const
+TRigVMTypeIndex FRigVMRegistry_NoLock::GetTypeIndex_NoLock(const FRigVMTemplateArgumentType& InType) const
 {
 	if(const TRigVMTypeIndex* Index = TypeToIndex.Find(InType))
 	{
@@ -944,7 +928,7 @@ TRigVMTypeIndex FRigVMRegistry::GetTypeIndex(const FRigVMTemplateArgumentType& I
 	return INDEX_NONE;
 }
 
-const FRigVMTemplateArgumentType& FRigVMRegistry::GetType(TRigVMTypeIndex InTypeIndex) const
+const FRigVMTemplateArgumentType& FRigVMRegistry_NoLock::GetType_NoLock(TRigVMTypeIndex InTypeIndex) const
 {
 	if((Types.IsValidIndex(InTypeIndex)))
 	{
@@ -954,9 +938,9 @@ const FRigVMTemplateArgumentType& FRigVMRegistry::GetType(TRigVMTypeIndex InType
 	return EmptyType;
 }
 
-const FRigVMTemplateArgumentType& FRigVMRegistry::FindTypeFromCPPType(const FString& InCPPType) const
+const FRigVMTemplateArgumentType& FRigVMRegistry_NoLock::FindTypeFromCPPType_NoLock(const FString& InCPPType) const
 {
-	const int32 TypeIndex = GetTypeIndexFromCPPType(InCPPType);
+	const int32 TypeIndex = GetTypeIndexFromCPPType_NoLock(InCPPType);
 	if(ensure(Types.IsValidIndex(TypeIndex)))
 	{
 		return Types[TypeIndex].Type;
@@ -966,7 +950,7 @@ const FRigVMTemplateArgumentType& FRigVMRegistry::FindTypeFromCPPType(const FStr
 	return EmptyType;
 }
 
-TRigVMTypeIndex FRigVMRegistry::GetTypeIndexFromCPPType(const FString& InCPPType) const
+TRigVMTypeIndex FRigVMRegistry_NoLock::GetTypeIndexFromCPPType_NoLock(const FString& InCPPType) const
 {
 	TRigVMTypeIndex Result = INDEX_NONE;
 	if(!InCPPType.IsEmpty())
@@ -986,8 +970,9 @@ TRigVMTypeIndex FRigVMRegistry::GetTypeIndexFromCPPType(const FString& InCPPType
 		// things up to date here. 
 		if(Result == INDEX_NONE)
 		{
-			// we may need ot 
-			FRigVMRegistry::Get().RefreshEngineTypes();
+			// we may need to update the types again to registry potentially
+			// missing predicate types 
+			FRigVMRegistry_NoLock::GetForWrite().RefreshEngineTypes_NoLock();
 			Result = Types.IndexOfByPredicate(Predicate);
 		}
 #endif
@@ -1005,7 +990,7 @@ TRigVMTypeIndex FRigVMRegistry::GetTypeIndexFromCPPType(const FString& InCPPType
 	return Result;
 }
 
-bool FRigVMRegistry::IsArrayType(TRigVMTypeIndex InTypeIndex) const
+bool FRigVMRegistry_NoLock::IsArrayType_NoLock(TRigVMTypeIndex InTypeIndex) const
 {
 	if((Types.IsValidIndex(InTypeIndex)))
 	{
@@ -1014,7 +999,7 @@ bool FRigVMRegistry::IsArrayType(TRigVMTypeIndex InTypeIndex) const
 	return false;
 }
 
-bool FRigVMRegistry::IsExecuteType(TRigVMTypeIndex InTypeIndex) const
+bool FRigVMRegistry_NoLock::IsExecuteType_NoLock(TRigVMTypeIndex InTypeIndex) const
 {
 	if(InTypeIndex == INDEX_NONE)
 	{
@@ -1028,7 +1013,7 @@ bool FRigVMRegistry::IsExecuteType(TRigVMTypeIndex InTypeIndex) const
 	return false;
 }
 
-bool FRigVMRegistry::ConvertExecuteContextToBaseType(TRigVMTypeIndex& InOutTypeIndex) const
+bool FRigVMRegistry_NoLock::ConvertExecuteContextToBaseType_NoLock(TRigVMTypeIndex& InOutTypeIndex) const
 {
 	if(InOutTypeIndex == INDEX_NONE)
 	{
@@ -1040,7 +1025,7 @@ bool FRigVMRegistry::ConvertExecuteContextToBaseType(TRigVMTypeIndex& InOutTypeI
 		return true;
 	}
 
-	if(!IsExecuteType(InOutTypeIndex))
+	if(!IsExecuteType_NoLock(InOutTypeIndex))
 	{
 		return false;
 	}
@@ -1049,9 +1034,9 @@ bool FRigVMRegistry::ConvertExecuteContextToBaseType(TRigVMTypeIndex& InOutTypeI
 	// convert them to the base execute type to make matching types easier later.
 	// this means that the execute argument in every permutations shares 
 	// the same type index of RigVMTypeUtils::TypeIndex::Execute
-	if(IsArrayType(InOutTypeIndex))
+	if(IsArrayType_NoLock(InOutTypeIndex))
 	{
-		InOutTypeIndex = GetArrayTypeFromBaseTypeIndex(RigVMTypeUtils::TypeIndex::Execute);
+		InOutTypeIndex = GetArrayTypeFromBaseTypeIndex_NoLock(RigVMTypeUtils::TypeIndex::Execute);
 	}
 	else
 	{
@@ -1061,26 +1046,26 @@ bool FRigVMRegistry::ConvertExecuteContextToBaseType(TRigVMTypeIndex& InOutTypeI
 	return true;
 }
 
-int32 FRigVMRegistry::GetArrayDimensionsForType(TRigVMTypeIndex InTypeIndex) const
+int32 FRigVMRegistry_NoLock::GetArrayDimensionsForType_NoLock(TRigVMTypeIndex InTypeIndex) const
 {
 	if(ensure(Types.IsValidIndex(InTypeIndex)))
 	{
 		const FTypeInfo& Info = Types[InTypeIndex];
 		if(Info.bIsArray)
 		{
-			return 1 + GetArrayDimensionsForType(Info.BaseTypeIndex);
+			return 1 + GetArrayDimensionsForType_NoLock(Info.BaseTypeIndex);
 		}
 	}
 	return 0;
 }
 
-bool FRigVMRegistry::IsWildCardType(TRigVMTypeIndex InTypeIndex) const
+bool FRigVMRegistry_NoLock::IsWildCardType_NoLock(TRigVMTypeIndex InTypeIndex) const
 {
 	return RigVMTypeUtils::TypeIndex::WildCard == InTypeIndex ||
 		RigVMTypeUtils::TypeIndex::WildCardArray == InTypeIndex;
 }
 
-bool FRigVMRegistry::CanMatchTypes(TRigVMTypeIndex InTypeIndexA, TRigVMTypeIndex InTypeIndexB, bool bAllowFloatingPointCasts) const
+bool FRigVMRegistry_NoLock::CanMatchTypes_NoLock(TRigVMTypeIndex InTypeIndexA, TRigVMTypeIndex InTypeIndexB, bool bAllowFloatingPointCasts) const
 {
 	if(!Types.IsValidIndex(InTypeIndexA) || !Types.IsValidIndex(InTypeIndexB))
 	{
@@ -1093,9 +1078,9 @@ bool FRigVMRegistry::CanMatchTypes(TRigVMTypeIndex InTypeIndexA, TRigVMTypeIndex
 	}
 
 	// execute types can always be connected
-	if(IsExecuteType(InTypeIndexA) && IsExecuteType(InTypeIndexB))
+	if(IsExecuteType_NoLock(InTypeIndexA) && IsExecuteType_NoLock(InTypeIndexB))
 	{
-		return GetArrayDimensionsForType(InTypeIndexA) == GetArrayDimensionsForType(InTypeIndexB);
+		return GetArrayDimensionsForType_NoLock(InTypeIndexA) == GetArrayDimensionsForType_NoLock(InTypeIndexB);
 	}
 
 	if(bAllowFloatingPointCasts)
@@ -1117,7 +1102,7 @@ bool FRigVMRegistry::CanMatchTypes(TRigVMTypeIndex InTypeIndexA, TRigVMTypeIndex
 	return false;
 }
 
-const TArray<TRigVMTypeIndex>& FRigVMRegistry::GetCompatibleTypes(TRigVMTypeIndex InTypeIndex) const
+const TArray<TRigVMTypeIndex>& FRigVMRegistry_NoLock::GetCompatibleTypes_NoLock(TRigVMTypeIndex InTypeIndex) const
 {
 	if(InTypeIndex == RigVMTypeUtils::TypeIndex::Float)
 	{
@@ -1144,13 +1129,13 @@ const TArray<TRigVMTypeIndex>& FRigVMRegistry::GetCompatibleTypes(TRigVMTypeInde
 	return EmptyTypes;
 }
 
-const TArray<TRigVMTypeIndex>& FRigVMRegistry::GetTypesForCategory(FRigVMTemplateArgument::ETypeCategory InCategory) const
+const TArray<TRigVMTypeIndex>& FRigVMRegistry_NoLock::GetTypesForCategory_NoLock(FRigVMTemplateArgument::ETypeCategory InCategory) const
 {
 	check(InCategory != FRigVMTemplateArgument::ETypeCategory_Invalid);
 	return TypesPerCategory.FindChecked(InCategory);
 }
 
-TRigVMTypeIndex FRigVMRegistry::GetArrayTypeFromBaseTypeIndex(TRigVMTypeIndex InTypeIndex) const
+TRigVMTypeIndex FRigVMRegistry_NoLock::GetArrayTypeFromBaseTypeIndex_NoLock(TRigVMTypeIndex InTypeIndex) const
 {
 	if(ensure(Types.IsValidIndex(InTypeIndex)))
 	{
@@ -1159,7 +1144,7 @@ TRigVMTypeIndex FRigVMRegistry::GetArrayTypeFromBaseTypeIndex(TRigVMTypeIndex In
 	return INDEX_NONE;
 }
 
-TRigVMTypeIndex FRigVMRegistry::GetBaseTypeFromArrayTypeIndex(TRigVMTypeIndex InTypeIndex) const
+TRigVMTypeIndex FRigVMRegistry_NoLock::GetBaseTypeFromArrayTypeIndex_NoLock(TRigVMTypeIndex InTypeIndex) const
 {
 	if(ensure(Types.IsValidIndex(InTypeIndex)))
 	{
@@ -1168,7 +1153,7 @@ TRigVMTypeIndex FRigVMRegistry::GetBaseTypeFromArrayTypeIndex(TRigVMTypeIndex In
 	return INDEX_NONE;
 }
 
-bool FRigVMRegistry::IsAllowedType(const FProperty* InProperty) const
+bool FRigVMRegistry_NoLock::IsAllowedType_NoLock(const FProperty* InProperty) const
 {
 	if(InProperty->IsA<FBoolProperty>() ||
 		InProperty->IsA<FUInt32Property>() ||
@@ -1187,37 +1172,37 @@ bool FRigVMRegistry::IsAllowedType(const FProperty* InProperty) const
 
 	if(const FArrayProperty* ArrayProperty  = CastField<FArrayProperty>(InProperty))
 	{
-		return IsAllowedType(ArrayProperty->Inner);
+		return IsAllowedType_NoLock(ArrayProperty->Inner);
 	}
 	if(const FStructProperty* StructProperty = CastField<FStructProperty>(InProperty))
 	{
-		return IsAllowedType(StructProperty->Struct);
+		return IsAllowedType_NoLock(StructProperty->Struct);
 	}
 	if(const FObjectProperty* ObjectProperty = CastField<FObjectProperty>(InProperty))
 	{
-		return IsAllowedType(ObjectProperty->PropertyClass);
+		return IsAllowedType_NoLock(ObjectProperty->PropertyClass);
 	}
 	if(const FEnumProperty* EnumProperty = CastField<FEnumProperty>(InProperty))
 	{
-		return IsAllowedType(EnumProperty->GetEnum());
+		return IsAllowedType_NoLock(EnumProperty->GetEnum());
 	}
 	if(const FByteProperty* ByteProperty = CastField<FByteProperty>(InProperty))
 	{
 		if(const UEnum* Enum = ByteProperty->Enum)
 		{
-			return IsAllowedType(Enum);
+			return IsAllowedType_NoLock(Enum);
 		}
 		return true;
 	}
 	return false;
 }
 
-bool FRigVMRegistry::IsAllowedType(const UEnum* InEnum) const
+bool FRigVMRegistry_NoLock::IsAllowedType_NoLock(const UEnum* InEnum) const
 {
 	return !InEnum->HasAnyFlags(DisallowedFlags()) && InEnum->HasAllFlags(NeededFlags());
 }
 
-bool FRigVMRegistry::IsAllowedType(const UStruct* InStruct) const
+bool FRigVMRegistry_NoLock::IsAllowedType_NoLock(const UStruct* InStruct) const
 {
 	if(InStruct->HasAnyFlags(DisallowedFlags()) || !InStruct->HasAllFlags(NeededFlags()))
 	{
@@ -1234,7 +1219,7 @@ bool FRigVMRegistry::IsAllowedType(const UStruct* InStruct) const
 	}
 
 	// allow all user defined structs since they can always be changed to be compliant with RigVM restrictions
-	if (const UUserDefinedStruct* UserDefinedStruct = Cast<UUserDefinedStruct>(InStruct))
+	if (InStruct->IsA<UUserDefinedStruct>())
 	{
 		return true;
 	}
@@ -1251,7 +1236,7 @@ bool FRigVMRegistry::IsAllowedType(const UStruct* InStruct) const
 
 	for (TFieldIterator<FProperty> It(InStruct); It; ++It)
 	{
-		if(!IsAllowedType(*It))
+		if(!IsAllowedType_NoLock(*It))
 		{
 			return false;
 		}
@@ -1259,7 +1244,7 @@ bool FRigVMRegistry::IsAllowedType(const UStruct* InStruct) const
 	return true;
 }
 
-bool FRigVMRegistry::IsAllowedType(const UClass* InClass) const
+bool FRigVMRegistry_NoLock::IsAllowedType_NoLock(const UClass* InClass) const
 {
 	if(!InClass || InClass->HasAnyClassFlags(CLASS_Hidden))
 	{
@@ -1275,10 +1260,8 @@ bool FRigVMRegistry::IsAllowedType(const UClass* InClass) const
 	return AllowedClasses.Contains(InClass);
 }
 
-void FRigVMRegistry::Register(const TCHAR* InName, FRigVMFunctionPtr InFunctionPtr, UScriptStruct* InStruct, const TArray<FRigVMFunctionArgument>& InArguments)
+void FRigVMRegistry_NoLock::Register_NoLock(const TCHAR* InName, FRigVMFunctionPtr InFunctionPtr, UScriptStruct* InStruct, const TArray<FRigVMFunctionArgument>& InArguments)
 {
-	FScopeLock FunctionRegistryScopeLock(&FunctionRegistryMutex);
-	
 	if (FindFunction_NoLock(InName) != nullptr)
 	{
 		return;
@@ -1357,9 +1340,8 @@ void FRigVMRegistry::Register(const TCHAR* InName, FRigVMFunctionPtr InFunctionP
 #endif
 }
 
-const FRigVMDispatchFactory* FRigVMRegistry::RegisterFactory(UScriptStruct* InFactoryStruct)
+const FRigVMDispatchFactory* FRigVMRegistry_NoLock::RegisterFactory_NoLock(UScriptStruct* InFactoryStruct)
 {
-
 	check(InFactoryStruct);
 	check(InFactoryStruct != FRigVMDispatchFactory::StaticStruct());
 	check(InFactoryStruct->IsChildOf(FRigVMDispatchFactory::StaticStruct()));
@@ -1367,8 +1349,6 @@ const FRigVMDispatchFactory* FRigVMRegistry::RegisterFactory(UScriptStruct* InFa
 	// ensure to register factories only once
 	const FRigVMDispatchFactory* ExistingFactory = nullptr;
 
-	FScopeLock FactoryRegistryScopeLock(&FactoryRegistryMutex);
-
 	const bool bFactoryAlreadyRegistered = Factories.ContainsByPredicate([InFactoryStruct, &ExistingFactory](const FRigVMDispatchFactory* Factory)
 	{
 		if(Factory->GetScriptStruct() == InFactoryStruct)
@@ -1394,11 +1374,11 @@ const FRigVMDispatchFactory* FRigVMRegistry::RegisterFactory(UScriptStruct* InFa
 	InFactoryStruct->InitializeStruct(Factory, 1);
 	Factory->FactoryScriptStruct = InFactoryStruct;
 	Factories.Add(Factory);
-	Factory->RegisterDependencyTypes();
+	Factory->RegisterDependencyTypes_NoLock();
 	return Factory;
 }
 
-void FRigVMRegistry::RegisterPredicate(UScriptStruct* InStruct, const TCHAR* InName, const TArray<FRigVMFunctionArgument>& InArguments)
+void FRigVMRegistry_NoLock::RegisterPredicate_NoLock(UScriptStruct* InStruct, const TCHAR* InName, const TArray<FRigVMFunctionArgument>& InArguments)
 {
 	// Make sure the predicate does not already exist
 	TArray<FRigVMFunction>& Predicates = StructNameToPredicates.FindOrAdd(InStruct->GetFName());
@@ -1415,7 +1395,7 @@ void FRigVMRegistry::RegisterPredicate(UScriptStruct* InStruct, const TCHAR* InN
 	Predicates.Add(Function);
 }
 
-void FRigVMRegistry::RegisterObjectTypes(TConstArrayView<TPair<UClass*, ERegisterObjectOperation>> InClasses)
+void FRigVMRegistry_NoLock::RegisterObjectTypes_NoLock(TConstArrayView<TPair<UClass*, ERegisterObjectOperation>> InClasses)
 {
 	for (TPair<UClass*, ERegisterObjectOperation> ClassOpPair : InClasses)
 	{
@@ -1457,7 +1437,7 @@ void FRigVMRegistry::RegisterObjectTypes(TConstArrayView<TPair<UClass*, ERegiste
 	}
 }
 
-void FRigVMRegistry::RegisterStructTypes(TConstArrayView<UScriptStruct*> InStructs)
+void FRigVMRegistry_NoLock::RegisterStructTypes_NoLock(TConstArrayView<UScriptStruct*> InStructs)
 {
 	for (UScriptStruct* Struct : InStructs)
 	{
@@ -1468,13 +1448,7 @@ void FRigVMRegistry::RegisterStructTypes(TConstArrayView<UScriptStruct*> InStruc
 	}
 }
 
-const FRigVMFunction* FRigVMRegistry::FindFunction(const TCHAR* InName, const FRigVMUserDefinedTypeResolver& InTypeResolver) const
-{
-	FScopeLock FunctionRegistryScopeLock(&FunctionRegistryMutex);
-	return FindFunction_NoLock(InName, InTypeResolver);
-}
-
-const FRigVMFunction* FRigVMRegistry::FindFunction_NoLock(const TCHAR* InName, const FRigVMUserDefinedTypeResolver& InTypeResolver) const
+const FRigVMFunction* FRigVMRegistry_NoLock::FindFunction_NoLock(const TCHAR* InName, const FRigVMUserDefinedTypeResolver& InTypeResolver) const
 {
 	// Check first if the function is provided by internally registered rig units. 
 	if(const int32* FunctionIndexPtr = FunctionNameToIndex.Find(InName))
@@ -1488,14 +1462,14 @@ const FRigVMFunction* FRigVMRegistry::FindFunction_NoLock(const TCHAR* InName, c
 	if(NameString.Split(TEXT("::"), &StructOrFactoryName, &SuffixString))
 	{
 		// if the factory has never been registered - FindDispatchFactory will try to look it up and register
-		if(const FRigVMDispatchFactory* Factory = FindDispatchFactory(*StructOrFactoryName))
+		if(const FRigVMDispatchFactory* Factory = FindDispatchFactory_NoLock(*StructOrFactoryName))
 		{
-			if(const FRigVMTemplate* Template = Factory->GetTemplate())
+			if(const FRigVMTemplate* Template = Factory->GetTemplate_NoLock())
 			{
-				const FRigVMTemplateTypeMap ArgumentTypes = Template->GetArgumentTypesFromString(SuffixString, &InTypeResolver);
+				const FRigVMTemplateTypeMap ArgumentTypes = Template->GetArgumentTypesFromString_Impl(SuffixString, &InTypeResolver, false);
 				if(ArgumentTypes.Num() == Template->NumArguments())
 				{
-					const int32 PermutationIndex = Template->FindPermutation(ArgumentTypes);
+					const int32 PermutationIndex = Template->FindPermutation(ArgumentTypes, false);
 					if(PermutationIndex != INDEX_NONE)
 					{
 						return ((FRigVMTemplate*)Template)->GetOrCreatePermutation_NoLock(PermutationIndex);
@@ -1538,7 +1512,7 @@ const FRigVMFunction* FRigVMRegistry::FindFunction_NoLock(const TCHAR* InName, c
 				const FRigVMFunction* RedirectedFunction = FindFunction_NoLock(*(NewStructOrFactoryName + TEXT("::") + SuffixString), InTypeResolver);
 				if(RedirectedFunction)
 				{
-					FRigVMRegistry& MutableRegistry = FRigVMRegistry::Get();
+					FRigVMRegistry_NoLock& MutableRegistry = Get(LockType_Write);
 					MutableRegistry.FunctionNameToIndex.Add(InName, RedirectedFunction->Index);
 					return RedirectedFunction;
 				}
@@ -1549,32 +1523,21 @@ const FRigVMFunction* FRigVMRegistry::FindFunction_NoLock(const TCHAR* InName, c
 	return nullptr;
 }
 
-const FRigVMFunction* FRigVMRegistry::FindFunction(UScriptStruct* InStruct, const TCHAR* InName, const FRigVMUserDefinedTypeResolver& InResolvalInfo) const
+const FRigVMFunction* FRigVMRegistry_NoLock::FindFunction_NoLock(UScriptStruct* InStruct, const TCHAR* InName, const FRigVMUserDefinedTypeResolver& InResolvalInfo) const
 {
 	check(InStruct);
 	check(InName);
 	
 	const FString FunctionName = FString::Printf(TEXT("%s::%s"), *InStruct->GetStructCPPName(), InName);
-	return FindFunction(*FunctionName, InResolvalInfo);
+	return FindFunction_NoLock(*FunctionName, InResolvalInfo);
 }
 
-const TChunkedArray<FRigVMFunction>& FRigVMRegistry::GetFunctions() const
+const TChunkedArray<FRigVMFunction>& FRigVMRegistry_NoLock::GetFunctions_NoLock() const
 {
 	return Functions;
 }
 
-const FRigVMTemplate* FRigVMRegistry::FindTemplate(const FName& InNotation, bool bIncludeDeprecated) const
-{
-	if (InNotation.IsNone())
-	{
-		return nullptr;
-	}
-
-	FScopeLock TemplateRegistryScopeLock(&TemplateRegistryMutex);
-	return FindTemplate_NoLock(InNotation, bIncludeDeprecated);
-}
-
-const FRigVMTemplate* FRigVMRegistry::FindTemplate_NoLock(const FName& InNotation, bool bIncludeDeprecated) const
+const FRigVMTemplate* FRigVMRegistry_NoLock::FindTemplate_NoLock(const FName& InNotation, bool bIncludeDeprecated) const
 {
 	if (InNotation.IsNone())
 	{
@@ -1590,18 +1553,18 @@ const FRigVMTemplate* FRigVMRegistry::FindTemplate_NoLock(const FName& InNotatio
 	FString FactoryName, ArgumentsString;
 	if(NotationString.Split(TEXT("("), &FactoryName, &ArgumentsString))
 	{
-		FRigVMRegistry* MutableThis = (FRigVMRegistry*)this;
+		FRigVMRegistry_NoLock* MutableThis = const_cast<FRigVMRegistry_NoLock*>(this);
 		
 		// deal with a couple of custom cases
 		static const TMap<FString, FString> CoreDispatchMap =
 		{
 			{
 				TEXT("Equals::Execute"),
-				MutableThis->FindOrAddDispatchFactory<FRigVMDispatch_CoreEquals>()->GetFactoryName().ToString()
+				MutableThis->FindOrAddDispatchFactory_NoLock<FRigVMDispatch_CoreEquals>()->GetFactoryName().ToString()
 			},
 			{
 				TEXT("NotEquals::Execute"),
-				MutableThis->FindOrAddDispatchFactory<FRigVMDispatch_CoreNotEquals>()->GetFactoryName().ToString()
+				MutableThis->FindOrAddDispatchFactory_NoLock<FRigVMDispatch_CoreNotEquals>()->GetFactoryName().ToString()
 			},
 		};
 
@@ -1610,9 +1573,9 @@ const FRigVMTemplate* FRigVMRegistry::FindTemplate_NoLock(const FName& InNotatio
 			FactoryName = *RemappedDispatch;
 		}
 		
-		if(const FRigVMDispatchFactory* Factory = FindDispatchFactory(*FactoryName))
+		if(const FRigVMDispatchFactory* Factory = FindDispatchFactory_NoLock(*FactoryName))
 		{
-			return Factory->GetTemplate();
+			return Factory->GetTemplate_NoLock();
 		}
 	}
 
@@ -1641,9 +1604,9 @@ const FRigVMTemplate* FRigVMRegistry::FindTemplate_NoLock(const FName& InNotatio
 			for(const FCoreRedirect* Redirect : Redirects)
 			{
 				const FString NewDispatchFactoryName = FRigVMDispatchFactory::DispatchPrefix + Redirect->NewName.ObjectName.ToString();
-				if(const FRigVMDispatchFactory* NewDispatchFactory = FindDispatchFactory(*NewDispatchFactoryName))
+				if(const FRigVMDispatchFactory* NewDispatchFactory = FindDispatchFactory_NoLock(*NewDispatchFactoryName))
 				{
-					return NewDispatchFactory->GetTemplate();
+					return NewDispatchFactory->GetTemplate_NoLock();
 				}
 			}
 		}
@@ -1686,19 +1649,17 @@ const FRigVMTemplate* FRigVMRegistry::FindTemplate_NoLock(const FName& InNotatio
 	return nullptr;
 }
 
-const TChunkedArray<FRigVMTemplate>& FRigVMRegistry::GetTemplates() const
+const TChunkedArray<FRigVMTemplate>& FRigVMRegistry_NoLock::GetTemplates_NoLock() const
 {
 	return Templates;
 }
 
-const FRigVMTemplate* FRigVMRegistry::GetOrAddTemplateFromArguments(const FName& InName, const TArray<FRigVMTemplateArgumentInfo>& InInfos, const FRigVMTemplateDelegates& InDelegates)
+const FRigVMTemplate* FRigVMRegistry_NoLock::GetOrAddTemplateFromArguments_NoLock(const FName& InName, const TArray<FRigVMTemplateArgumentInfo>& InInfos, const FRigVMTemplateDelegates& InDelegates)
 {
-	FScopeLock TemplateRegistryScopeLock(&TemplateRegistryMutex);
-	
 	// avoid reentry in FindTemplate. try to find an existing
 	// template only if we are not yet in ::FindTemplate.
 	const FName Notation = FRigVMTemplateArgumentInfo::ComputeTemplateNotation(InName, InInfos);
-	if(const FRigVMTemplate* ExistingTemplate = FindTemplate(Notation))
+	if(const FRigVMTemplate* ExistingTemplate = FindTemplate_NoLock(Notation))
 	{
 		return ExistingTemplate;
 	}
@@ -1706,21 +1667,15 @@ const FRigVMTemplate* FRigVMRegistry::GetOrAddTemplateFromArguments(const FName&
 	return AddTemplateFromArguments_NoLock(InName, InInfos, InDelegates);
 }
 
-const FRigVMTemplate* FRigVMRegistry::AddTemplateFromArguments(const FName& InName, const TArray<FRigVMTemplateArgumentInfo>& InInfos, const FRigVMTemplateDelegates& InDelegates)
-{
-	FScopeLock TemplateRegistryScopeLock(&TemplateRegistryMutex);
-	return AddTemplateFromArguments_NoLock(InName, InInfos, InDelegates);
-}
-
-const FRigVMTemplate* FRigVMRegistry::AddTemplateFromArguments_NoLock(const FName& InName, const TArray<FRigVMTemplateArgumentInfo>& InInfos, const FRigVMTemplateDelegates& InDelegates)
+const FRigVMTemplate* FRigVMRegistry_NoLock::AddTemplateFromArguments_NoLock(const FName& InName, const TArray<FRigVMTemplateArgumentInfo>& InInfos, const FRigVMTemplateDelegates& InDelegates)
 {
 	// we only support to ask for templates here which provide singleton types
 	int32 NumPermutations = 0;
 	FRigVMTemplate Template(InName, InInfos);
 	for(const FRigVMTemplateArgument& Argument : Template.Arguments)
 	{
-		const int32 NumIndices = Argument.GetNumTypes();
-		if(!Argument.IsSingleton() && NumPermutations > 1)
+		const int32 NumIndices = Argument.GetNumTypes_NoLock();
+		if(!Argument.IsSingleton_NoLock() && NumPermutations > 1)
 		{
 			if(NumIndices != NumPermutations)
 			{
@@ -1734,10 +1689,10 @@ const FRigVMTemplate* FRigVMRegistry::AddTemplateFromArguments_NoLock(const FNam
 	// if any of the arguments are wildcards we'll need to update the types
 	for(FRigVMTemplateArgument& Argument : Template.Arguments)
 	{
-		if(Argument.GetNumTypes() == 1 && IsWildCardType(Argument.GetTypeIndex(0)))
+		if(Argument.GetNumTypes_NoLock() == 1 && IsWildCardType_NoLock(Argument.GetTypeIndex_NoLock(0)))
 		{
-			Argument.InvalidatePermutations(Argument.GetTypeIndex(0));
-			if(IsArrayType(Argument.GetTypeIndex(0)))
+			Argument.InvalidatePermutations(Argument.GetTypeIndex_NoLock(0));
+			if(IsArrayType_NoLock(Argument.GetTypeIndex_NoLock(0)))
 			{
 				Argument.TypeCategories.Add(FRigVMTemplateArgument::ETypeCategory_ArrayAnyValue);
 			}
@@ -1748,7 +1703,7 @@ const FRigVMTemplate* FRigVMRegistry::AddTemplateFromArguments_NoLock(const FNam
 			Argument.bUseCategories = true;
 			Argument.TypeIndices.Reset();
 	
-			NumPermutations = FMath::Max(NumPermutations, Argument.GetNumTypes()); 
+			NumPermutations = FMath::Max(NumPermutations, Argument.GetNumTypes_NoLock()); 
 		}
 	}
 
@@ -1762,7 +1717,7 @@ const FRigVMTemplate* FRigVMRegistry::AddTemplateFromArguments_NoLock(const FNam
 			TSet< TRigVMTypeIndex > PermutationTypes; PermutationTypes.Reserve(NumPermutations);
 			for(int32 Index = 0; Index < NumPermutations; Index++)
 			{
-				const TRigVMTypeIndex ArgType = Template.Arguments[0].GetTypeIndex(Index);
+				const TRigVMTypeIndex ArgType = Template.Arguments[0].GetTypeIndex_NoLock(Index);
 				if (PermutationTypes.Contains(ArgType))
 				{
 					ToRemove.Add(Index);
@@ -1782,7 +1737,7 @@ const FRigVMTemplate* FRigVMRegistry::AddTemplateFromArguments_NoLock(const FNam
 			{
 				for(int32 ArgIndex = 0; ArgIndex < NumArguments; ArgIndex++)
 				{
-					ArgTypes[ArgIndex] = Template.Arguments[ArgIndex].GetTypeIndex(Index);
+					ArgTypes[ArgIndex] = Template.Arguments[ArgIndex].GetTypeIndex_NoLock(Index);
 				}
 				
 				if (PermutationTypes.Contains(ArgTypes))
@@ -1833,13 +1788,7 @@ const FRigVMTemplate* FRigVMRegistry::AddTemplateFromArguments_NoLock(const FNam
 	return &Templates[Index];
 }
 
-FRigVMDispatchFactory* FRigVMRegistry::FindDispatchFactory(const FName& InFactoryName) const
-{
-	FScopeLock FactoryRegistryScopeLock(&FactoryRegistryMutex);
-	return FindDispatchFactory_NoLock(InFactoryName);
-}
-
-FRigVMDispatchFactory* FRigVMRegistry::FindDispatchFactory_NoLock(const FName& InFactoryName) const
+FRigVMDispatchFactory* FRigVMRegistry_NoLock::FindDispatchFactory_NoLock(const FName& InFactoryName) const
 {
 	FRigVMDispatchFactory* const* FactoryPtr = Factories.FindByPredicate([InFactoryName](const FRigVMDispatchFactory* Factory) -> bool
 	{
@@ -1858,32 +1807,32 @@ FRigVMDispatchFactory* FRigVMRegistry::FindDispatchFactory_NoLock(const FName& I
 		const FString ScriptStructName = FactoryName.Mid(FCString::Strlen(FRigVMDispatchFactory::DispatchPrefix));
 		if(UScriptStruct* FactoryStruct = FindFirstObject<UScriptStruct>(*ScriptStructName, EFindFirstObjectOptions::NativeFirst | EFindFirstObjectOptions::EnsureIfAmbiguous))
 		{
-			FRigVMRegistry* MutableThis = (FRigVMRegistry*)this;
-			return (FRigVMDispatchFactory*)MutableThis->RegisterFactory(FactoryStruct);
+			FRigVMRegistry_NoLock* MutableThis = const_cast<FRigVMRegistry_NoLock*>(this);
+			return const_cast<FRigVMDispatchFactory*>(MutableThis->RegisterFactory_NoLock(FactoryStruct));
 		}
 	}	
 	
 	return nullptr;
 }
 
-FRigVMDispatchFactory* FRigVMRegistry::FindOrAddDispatchFactory(UScriptStruct* InFactoryStruct)
+FRigVMDispatchFactory* FRigVMRegistry_NoLock::FindOrAddDispatchFactory_NoLock(UScriptStruct* InFactoryStruct)
 {
-	return (FRigVMDispatchFactory*)RegisterFactory(InFactoryStruct);
+	return const_cast<FRigVMDispatchFactory*>(RegisterFactory_NoLock(InFactoryStruct));
 }
 
-FString FRigVMRegistry::FindOrAddSingletonDispatchFunction(UScriptStruct* InFactoryStruct)
+FString FRigVMRegistry_NoLock::FindOrAddSingletonDispatchFunction_NoLock(UScriptStruct* InFactoryStruct)
 {
-	if(const FRigVMDispatchFactory* Factory = FindOrAddDispatchFactory(InFactoryStruct))
+	if(const FRigVMDispatchFactory* Factory = FindOrAddDispatchFactory_NoLock(InFactoryStruct))
 	{
 		if(Factory->IsSingleton())
 		{
-			if(const FRigVMTemplate* Template = Factory->GetTemplate())
+			if(const FRigVMTemplate* Template = Factory->GetTemplate_NoLock())
 			{
 				// use the types for the first permutation - since we don't care
 				// for a singleton dispatch
-				const FRigVMTemplateTypeMap TypesForPrimaryPermutation = Template->GetTypesForPermutation(0);
-				const FString Name = Factory->GetPermutationName(TypesForPrimaryPermutation);
-				if(const FRigVMFunction* Function = FindFunction(*Name))
+				const FRigVMTemplateTypeMap TypesForPrimaryPermutation = Template->GetTypesForPermutation_NoLock(0);
+				const FString Name = Factory->GetPermutationName(TypesForPrimaryPermutation, false);
+				if(const FRigVMFunction* Function = FindFunction_NoLock(*Name))
 				{
 					return Function->Name;
 				}
@@ -1893,11 +1842,64 @@ FString FRigVMRegistry::FindOrAddSingletonDispatchFunction(UScriptStruct* InFact
 	return FString();
 }
 
-const TArray<FRigVMDispatchFactory*>& FRigVMRegistry::GetFactories() const
+const TArray<FRigVMDispatchFactory*>& FRigVMRegistry_NoLock::GetFactories_NoLock() const
 {	return Factories;
 }
 
-const TArray<FRigVMFunction>* FRigVMRegistry::GetPredicatesForStruct(const FName& InStructName) const
+const TArray<FRigVMFunction>* FRigVMRegistry_NoLock::GetPredicatesForStruct_NoLock(const FName& InStructName) const
 {
 	return StructNameToPredicates.Find(InStructName);
 }
+
+FRigVMRegistry_RWLock::FRigVMRegistry_RWLock()
+	: FRigVMRegistry_NoLock()
+{
+	Initialize(true);
+}
+
+void FRigVMRegistry_RWLock::Initialize(bool bLockRegistry)
+{
+	LockType.store(LockType_Invalid);
+	const FConditionalWriteScopeLock _(*this, bLockRegistry);
+	
+	Initialize_NoLock();
+	
+	const FAssetRegistryModule& AssetRegistryModule = FModuleManager::LoadModuleChecked<FAssetRegistryModule>(TEXT("AssetRegistry"));
+	AssetRegistryModule.Get().OnAssetRemoved().AddRaw(this, &FRigVMRegistry_RWLock::OnAssetRemoved);
+	AssetRegistryModule.Get().OnAssetRenamed().AddRaw(this, &FRigVMRegistry_RWLock::OnAssetRenamed);
+
+	IPluginManager::Get().OnPluginUnmounted().AddRaw(this, &FRigVMRegistry_RWLock::OnPluginUnloaded);
+	
+	UE::Anim::AttributeTypes::GetOnAttributeTypesChanged().AddRaw(this, &FRigVMRegistry_RWLock::OnAnimationAttributeTypesChanged);
+}
+
+void FRigVMRegistry_RWLock::EnsureLocked(ELockType InLockType)
+{
+	check(InLockType != LockType_Invalid);
+
+	const FRigVMRegistry_RWLock& Registry = Get();
+	const ELockType CurrentLockType = Registry.LockType.load();
+
+	switch(InLockType)
+	{
+		case LockType_Read:
+		{
+			ensureMsgf(
+				CurrentLockType == LockType_Read ||
+				CurrentLockType == LockType_Write,
+				TEXT("The Registry is not locked for reading yet - access to the NoLock registry is only possible after locking the RWLock registry (by using its public API calls)."));
+			break;
+		}
+		case LockType_Write:
+		{
+			ensureMsgf(
+				CurrentLockType == LockType_Write,
+				TEXT("The Registry is not locked for writing yet - access to the NoLock registry is only possible after locking the RWLock registry (by using its public API calls)."));
+			break;
+		}
+		default:
+		{
+			break;
+		}
+	}
+}
diff --git a/Engine/Plugins/Runtime/RigVM/Source/RigVM/Private/RigVMCore/RigVMStruct.cpp b/Engine/Plugins/Runtime/RigVM/Source/RigVM/Private/RigVMCore/RigVMStruct.cpp
--- a/Engine/Plugins/Runtime/RigVM/Source/RigVM/Private/RigVMCore/RigVMStruct.cpp
+++ b/Engine/Plugins/Runtime/RigVM/Source/RigVM/Private/RigVMCore/RigVMStruct.cpp
@@ -130,7 +130,7 @@ FRigVMTemplateTypeMap FRigVMStructUpgradeInfo::GetTypeMapFromStruct(UScriptStruc
 		}
 		
 		const FName PropertyName = It->GetFName();
-		const TRigVMTypeIndex TypeIndex = FRigVMTemplateArgument(*It).GetTypeIndex(0);
+		const TRigVMTypeIndex TypeIndex = FRigVMTemplateArgument::Make(*It).GetTypeIndex(0);
 		Types.Add(PropertyName, TypeIndex);
 	}
 #endif
diff --git a/Engine/Plugins/Runtime/RigVM/Source/RigVM/Private/RigVMCore/RigVMTemplate.cpp b/Engine/Plugins/Runtime/RigVM/Source/RigVM/Private/RigVMCore/RigVMTemplate.cpp
--- a/Engine/Plugins/Runtime/RigVM/Source/RigVM/Private/RigVMCore/RigVMTemplate.cpp
+++ b/Engine/Plugins/Runtime/RigVM/Source/RigVM/Private/RigVMCore/RigVMTemplate.cpp
@@ -40,12 +40,12 @@ FRigVMTemplateArgument::FRigVMTemplateArgument(const FName& InName, ERigVMPinDir
 {}
 
 FRigVMTemplateArgument::FRigVMTemplateArgument(FProperty* InProperty):
-	FRigVMTemplateArgument(InProperty, FRigVMRegistry::Get())
+	FRigVMTemplateArgument(InProperty, FRigVMRegistry_NoLock::GetForWrite())
 {
 	
 }
 
-FRigVMTemplateArgument::FRigVMTemplateArgument(FProperty* InProperty, FRigVMRegistry& InRegistry)
+FRigVMTemplateArgument::FRigVMTemplateArgument(FProperty* InProperty, FRigVMRegistry_NoLock& InRegistry)
 	: Name(InProperty->GetFName())
 {
 #if WITH_EDITOR
@@ -83,19 +83,25 @@ FRigVMTemplateArgument::FRigVMTemplateArgument(FProperty* InProperty, FRigVMRegi
 	}
 	
 	const FRigVMTemplateArgumentType Type(CPPTypeName, CPPTypeObject);
-	const TRigVMTypeIndex TypeIndex = InRegistry.FindOrAddType(Type, true); 
+	const TRigVMTypeIndex TypeIndex = InRegistry.FindOrAddType_NoLock(Type, true); 
 
 	TypeIndices.Add(TypeIndex);
-	EnsureValidExecuteType(InRegistry);
+	EnsureValidExecuteType_NoLock(InRegistry);
 	UpdateTypeToPermutations();
 }
 
+FRigVMTemplateArgument FRigVMTemplateArgument::Make(FProperty* InProperty)
+{
+	const FRigVMRegistryWriteLock _;
+	return FRigVMTemplateArgument(InProperty);
+}
+
 FRigVMTemplateArgument::FRigVMTemplateArgument(const FName& InName, ERigVMPinDirection InDirection, TRigVMTypeIndex InTypeIndex)
 	: Name(InName)
 	, Direction(InDirection)
 	, TypeIndices({InTypeIndex})
 {
-	EnsureValidExecuteType(FRigVMRegistry::Get());
+	EnsureValidExecuteType_NoLock(FRigVMRegistry_RWLock::Get());
 	UpdateTypeToPermutations();
 }
 
@@ -105,7 +111,7 @@ FRigVMTemplateArgument::FRigVMTemplateArgument(const FName& InName, ERigVMPinDir
 	, TypeIndices(InTypeIndices)
 {
 	check(TypeIndices.Num() > 0);
-	EnsureValidExecuteType(FRigVMRegistry::Get());
+	EnsureValidExecuteType_NoLock(FRigVMRegistry_RWLock::Get());
 	UpdateTypeToPermutations();
 }
 
@@ -129,7 +135,7 @@ FRigVMTemplateArgument::FRigVMTemplateArgument(const FName& InName, ERigVMPinDir
 		{
 			for (const ETypeCategory TypeCategory : TypeCategories)
 			{
-				const TArray<TRigVMTypeIndex>& Types = FRigVMRegistry::Get().GetTypesForCategory(TypeCategory);
+				const TArray<TRigVMTypeIndex>& Types = FRigVMRegistry_NoLock::GetForRead().GetTypesForCategory_NoLock(TypeCategory);
 				AllTypes.Reserve(AllTypes.Num() + Types.Num());
 				for (const TRigVMTypeIndex Type: Types)
 				{
@@ -157,25 +163,25 @@ FRigVMTemplateArgument::FRigVMTemplateArgument(const FName& InName, ERigVMPinDir
 				});
 			}
 			TypeIndices = MoveTemp(Indices);
-			EnsureValidExecuteType(FRigVMRegistry::Get());
+			EnsureValidExecuteType_NoLock(FRigVMRegistry_RWLock::Get());
 		}
 
 		UpdateTypeToPermutations();
 	}
 }
 
-void FRigVMTemplateArgument::EnsureValidExecuteType(FRigVMRegistry& InRegistry)
+void FRigVMTemplateArgument::EnsureValidExecuteType_NoLock(FRigVMRegistry_NoLock& InRegistry)
 {
 	for(TRigVMTypeIndex& TypeIndex : TypeIndices)
 	{
-		InRegistry.ConvertExecuteContextToBaseType(TypeIndex);
+		InRegistry.ConvertExecuteContextToBaseType_NoLock(TypeIndex);
 	}
 }
 
 void FRigVMTemplateArgument::UpdateTypeToPermutations()
 {
 	TypeToPermutations.Reset();
-	TypeToPermutations.Reserve(GetNumTypes());
+	TypeToPermutations.Reserve(GetNumTypes_NoLock());
 
 	int32 TypeIndex = 0;
 	ForEachType([&](const TRigVMTypeIndex Type)
@@ -185,7 +191,13 @@ void FRigVMTemplateArgument::UpdateTypeToPermutations()
 	});
 }
 
-bool FRigVMTemplateArgument::SupportsTypeIndex(TRigVMTypeIndex InTypeIndex, TRigVMTypeIndex* OutTypeIndex) const
+bool FRigVMTemplateArgument::SupportsTypeIndex(TRigVMTypeIndex InTypeIndex, TRigVMTypeIndex* OutTypeIndex, const bool bLockRegistry) const
+{
+	const FRigVMRegistryReadLock _(bLockRegistry);
+	return SupportsTypeIndex_NoLock(InTypeIndex, OutTypeIndex);
+}
+
+bool FRigVMTemplateArgument::SupportsTypeIndex_NoLock(TRigVMTypeIndex InTypeIndex, TRigVMTypeIndex* OutTypeIndex) const
 {
 	if(InTypeIndex == INDEX_NONE)
 	{
@@ -193,14 +205,14 @@ bool FRigVMTemplateArgument::SupportsTypeIndex(TRigVMTypeIndex InTypeIndex, TRig
 	}
 	
 	// convert any execute type into the base execute
-	const FRigVMRegistry& Registry = FRigVMRegistry::Get();
-	if(Registry.IsExecuteType(InTypeIndex))
+	const FRigVMRegistry_NoLock& Registry = FRigVMRegistry_NoLock::GetForRead();
+	if(Registry.IsExecuteType_NoLock(InTypeIndex))
 	{
-		const bool bIsArray = Registry.IsArrayType(InTypeIndex);
+		const bool bIsArray = Registry.IsArrayType_NoLock(InTypeIndex);
 		InTypeIndex = RigVMTypeUtils::TypeIndex::Execute;
 		if(bIsArray)
 		{
-			InTypeIndex = Registry.GetArrayTypeFromBaseTypeIndex(InTypeIndex);
+			InTypeIndex = Registry.GetArrayTypeFromBaseTypeIndex_NoLock(InTypeIndex);
 		}
 	}
 
@@ -209,13 +221,13 @@ bool FRigVMTemplateArgument::SupportsTypeIndex(TRigVMTypeIndex InTypeIndex, TRig
 	{
 		if(OutTypeIndex)
 		{
-			(*OutTypeIndex) = GetTypeIndex(Permutations[0]);
+			(*OutTypeIndex) = GetTypeIndex_NoLock(Permutations[0]);
 		}
 		return true;
 	}
 
 	// Try to find compatible type
-	const TArray<TRigVMTypeIndex>& CompatibleTypes = Registry.GetCompatibleTypes(InTypeIndex);
+	const TArray<TRigVMTypeIndex>& CompatibleTypes = Registry.GetCompatibleTypes_NoLock(InTypeIndex);
 	for (const TRigVMTypeIndex& CompatibleTypeIndex : CompatibleTypes)
 	{
 		const TArray<int32>& CompatiblePermutations = GetPermutations(CompatibleTypeIndex);
@@ -223,7 +235,7 @@ bool FRigVMTemplateArgument::SupportsTypeIndex(TRigVMTypeIndex InTypeIndex, TRig
 		{
 			if(OutTypeIndex)
 			{
-				(*OutTypeIndex) = GetTypeIndex(CompatiblePermutations[0]);
+				(*OutTypeIndex) = GetTypeIndex_NoLock(CompatiblePermutations[0]);
 			}
 			return true;
 		}
@@ -232,68 +244,96 @@ bool FRigVMTemplateArgument::SupportsTypeIndex(TRigVMTypeIndex InTypeIndex, TRig
 	return false;
 }
 
-bool FRigVMTemplateArgument::IsSingleton(const TArray<int32>& InPermutationIndices) const
+bool FRigVMTemplateArgument::IsSingleton(const TArray<int32>& InPermutationIndices, const bool bLockRegistry) const
+{
+	const FRigVMRegistryReadLock _(bLockRegistry);
+	return IsSingleton_NoLock(InPermutationIndices);
+}
+
+bool FRigVMTemplateArgument::IsSingleton_NoLock(const TArray<int32>& InPermutationIndices) const
 {
 	if (TypeToPermutations.Num() == 1)
 	{
 		return true;
 	}
-	else if(InPermutationIndices.Num() == 0)
+
+	// if a type is using categories it can't be singleton
+	// since categories provide more than one type.
+	if(bUseCategories)
 	{
 		return false;
 	}
 
-	const TRigVMTypeIndex InType0 = GetTypeIndex(InPermutationIndices[0]);
-	for (int32 PermutationIndex = 1; PermutationIndex < InPermutationIndices.Num(); PermutationIndex++)
+	const bool bUsesPermutations = !InPermutationIndices.IsEmpty(); 
+	const int32 NumPermutations = bUsesPermutations ? InPermutationIndices.Num() : GetNumTypes_NoLock();
+	const TRigVMTypeIndex InType0 = GetTypeIndex_NoLock(bUsesPermutations ? InPermutationIndices[0] : 0);
+	for (int32 PermutationIndex = 1; PermutationIndex < NumPermutations; PermutationIndex++)
 	{
-		if (GetTypeIndex(InPermutationIndices[PermutationIndex]) != InType0)
+		if (GetTypeIndex_NoLock(bUsesPermutations ? InPermutationIndices[PermutationIndex] : PermutationIndex) != InType0)
 		{
 			return false;
 		}
 	}
+
 	return true;
 }
 
 bool FRigVMTemplateArgument::IsExecute() const
 {
-	const FRigVMRegistry& Registry = FRigVMRegistry::Get();
+	const FRigVMRegistry_NoLock& Registry = FRigVMRegistry_NoLock::GetForRead();
 	const int32 FoundAnyNotExec = IndexOfByPredicate([&](const TRigVMTypeIndex Type)
 	{
-		return !Registry.IsExecuteType(Type);
+		return !Registry.IsExecuteType_NoLock(Type);
 	});
 	return FoundAnyNotExec == INDEX_NONE;
 }
 
-FRigVMTemplateArgument::EArrayType FRigVMTemplateArgument::GetArrayType() const
+FRigVMTemplateArgument::EArrayType FRigVMTemplateArgument::GetArrayType(const bool bLockRegistry) const
+{
+	const FRigVMRegistryReadLock _(bLockRegistry);
+	return GetArrayType_NoLock();
+}
+
+FRigVMTemplateArgument::EArrayType FRigVMTemplateArgument::GetArrayType_NoLock() const
 {
-	const int32 NumTypes = GetNumTypes();
-	if (GetNumTypes() > 0)
+	if(CachedArrayType.IsSet())
+	{
+		return CachedArrayType.GetValue();
+	}
+	
+	TArray<TRigVMTypeIndex> Types; GetAllTypes_NoLock(Types);
+
+	const int32 NumTypes = Types.Num();
+	if (NumTypes > 0)
 	{
-		const FRigVMRegistry& Registry = FRigVMRegistry::Get();
-		const EArrayType ArrayType = Registry.IsArrayType(GetTypeIndex(0)) ? EArrayType_ArrayValue : EArrayType_SingleValue;
+		const FRigVMRegistry_NoLock& Registry = FRigVMRegistry_NoLock::GetForRead();
+		const EArrayType ArrayType = Registry.IsArrayType_NoLock(Types[0]) ? EArrayType_ArrayValue : EArrayType_SingleValue;
 		
-		if(IsSingleton())
+		if(IsSingleton_NoLock())
 		{
-			return ArrayType;
+			CachedArrayType = ArrayType;
+			return CachedArrayType.GetValue();
 		}
 
 		for(int32 PermutationIndex=1; PermutationIndex<NumTypes;PermutationIndex++)
 		{
-			const TRigVMTypeIndex TypeIndex = GetTypeIndex(PermutationIndex);
+			const TRigVMTypeIndex TypeIndex = Types[PermutationIndex];
 			// INDEX_NONE indicates deleted permutation
 			if (TypeIndex == INDEX_NONE)
 			{
 				continue;
 			}
 			
-			const EArrayType OtherArrayType = Registry.IsArrayType(TypeIndex) ? EArrayType_ArrayValue : EArrayType_SingleValue;
+			const EArrayType OtherArrayType = Registry.IsArrayType_NoLock(TypeIndex) ? EArrayType_ArrayValue : EArrayType_SingleValue;
 			if(OtherArrayType != ArrayType)
 			{
-				return EArrayType_Mixed;
+				CachedArrayType = EArrayType_Mixed;
+				return CachedArrayType.GetValue();
 			}
 		}
 
-		return ArrayType;
+		CachedArrayType = ArrayType;
+		return CachedArrayType.GetValue();
 	}
 
 	return EArrayType_Invalid;
@@ -332,7 +372,13 @@ void FRigVMTemplateArgument::InvalidatePermutations(const TRigVMTypeIndex InType
 	TypeToPermutations.Remove(InType);
 }
 
-void FRigVMTemplateArgument::GetAllTypes(TArray<TRigVMTypeIndex>& OutTypes) const
+void FRigVMTemplateArgument::GetAllTypes(TArray<TRigVMTypeIndex>& OutTypes, const bool bLockRegistry) const
+{
+	const FRigVMRegistryReadLock _(bLockRegistry);
+	GetAllTypes_NoLock(OutTypes);
+}
+
+void FRigVMTemplateArgument::GetAllTypes_NoLock(TArray<TRigVMTypeIndex>& OutTypes) const
 {
 	if (!bUseCategories)
 	{
@@ -345,11 +391,11 @@ void FRigVMTemplateArgument::GetAllTypes(TArray<TRigVMTypeIndex>& OutTypes) cons
 	{
 		if (FilterType == nullptr)
 		{
-			OutTypes.Append(FRigVMRegistry::Get().GetTypesForCategory(Category));
+			OutTypes.Append(FRigVMRegistry_NoLock::GetForRead().GetTypesForCategory_NoLock(Category));
 		}
 		else
 		{
-			const TArray<TRigVMTypeIndex>& CategoryTypes = FRigVMRegistry::Get().GetTypesForCategory(Category);
+			const TArray<TRigVMTypeIndex>& CategoryTypes = FRigVMRegistry_NoLock::GetForRead().GetTypesForCategory_NoLock(Category);
 			for (const TRigVMTypeIndex& Type : CategoryTypes)
 			{
 				if (FilterType(Type))
@@ -361,7 +407,13 @@ void FRigVMTemplateArgument::GetAllTypes(TArray<TRigVMTypeIndex>& OutTypes) cons
 	}
 }
 
-TRigVMTypeIndex FRigVMTemplateArgument::GetTypeIndex(const int32 InIndex) const
+TRigVMTypeIndex FRigVMTemplateArgument::GetTypeIndex(const int32 InIndex, bool bLockRegistry) const
+{
+	const FRigVMRegistryReadLock _(bLockRegistry);
+	return GetTypeIndex_NoLock(InIndex);
+}
+
+TRigVMTypeIndex FRigVMTemplateArgument::GetTypeIndex_NoLock(const int32 InIndex) const
 {
 	if (!bUseCategories)
 	{
@@ -392,6 +444,26 @@ TRigVMTypeIndex FRigVMTemplateArgument::GetTypeIndex(const int32 InIndex) const
 	return CategoryViews(TypeCategories).GetTypeIndex(InIndex);
 }
 
+TOptional<TRigVMTypeIndex> FRigVMTemplateArgument::TryToGetTypeIndex(const int32 InIndex, const bool bLockRegistry) const
+{
+	const FRigVMRegistryReadLock _(bLockRegistry);
+	return TryToGetTypeIndex_NoLock(InIndex);
+}
+
+TOptional<TRigVMTypeIndex> FRigVMTemplateArgument::TryToGetTypeIndex_NoLock(const int32 InIndex) const
+{
+	if(IsSingleton_NoLock())
+	{
+		return GetTypeIndex_NoLock(0);
+	}
+	const TRigVMTypeIndex TypeIndex = GetTypeIndex_NoLock(InIndex);
+	if(TypeIndex != INDEX_NONE)
+	{
+		return TypeIndex;
+	}
+	return TOptional<TRigVMTypeIndex>();
+}
+
 int32 FRigVMTemplateArgument::FindTypeIndex(const TRigVMTypeIndex InTypeIndex) const
 {
 	if (!bUseCategories)
@@ -449,7 +521,20 @@ int32 FRigVMTemplateArgument::GetNumTypes() const
 	
 	return Algo::Accumulate(TypeCategories, 0, [](int32 Sum, const ETypeCategory Category)
 	{
-		return Sum + FRigVMRegistry::Get().GetTypesForCategory(Category).Num();
+		return Sum + FRigVMRegistry_RWLock::Get().GetTypesForCategory(Category).Num();
+	});
+}
+
+int32 FRigVMTemplateArgument::GetNumTypes_NoLock() const
+{
+	if (!bUseCategories)
+	{
+		return TypeIndices.Num();
+	}
+	
+	return Algo::Accumulate(TypeCategories, 0, [](int32 Sum, const ETypeCategory Category)
+	{
+		return Sum + FRigVMRegistry_NoLock::GetForRead().GetTypesForCategory_NoLock(Category).Num();
 	});
 }
 
@@ -489,6 +574,8 @@ void FRigVMTemplateArgument::ForEachType(TFunction<bool(const TRigVMTypeIndex In
 
 TArray<TRigVMTypeIndex> FRigVMTemplateArgument::GetSupportedTypeIndices(const TArray<int32>& InPermutationIndices) const
 {
+	const FRigVMRegistryReadLock _;
+
 	TArray<TRigVMTypeIndex> SupportedTypes;
 	if(InPermutationIndices.IsEmpty())
 	{
@@ -507,7 +594,7 @@ TArray<TRigVMTypeIndex> FRigVMTemplateArgument::GetSupportedTypeIndices(const TA
 		for(const int32 PermutationIndex : InPermutationIndices)
 		{
 			// INDEX_NONE indicates deleted permutation
-			const TRigVMTypeIndex Type = GetTypeIndex(PermutationIndex);
+			const TRigVMTypeIndex Type = GetTypeIndex_NoLock(PermutationIndex);
 			if (Type != INDEX_NONE)
 			{
 				SupportedTypes.AddUnique(Type);
@@ -521,8 +608,10 @@ TArray<TRigVMTypeIndex> FRigVMTemplateArgument::GetSupportedTypeIndices(const TA
 
 TArray<FString> FRigVMTemplateArgument::GetSupportedTypeStrings(const TArray<int32>& InPermutationIndices) const
 {
+	const FRigVMRegistryReadLock _;
+	const FRigVMRegistry_NoLock& Registry = FRigVMRegistry_NoLock::GetForRead();
+
 	TArray<FString> SupportedTypes;
-	const FRigVMRegistry& Registry = FRigVMRegistry::Get();
 	if(InPermutationIndices.IsEmpty())
 	{
 		ForEachType([&](const TRigVMTypeIndex TypeIndex)
@@ -530,7 +619,7 @@ TArray<FString> FRigVMTemplateArgument::GetSupportedTypeStrings(const TArray<int
 			// INDEX_NONE indicates deleted permutation
 			if (TypeIndex != INDEX_NONE)
 			{
-				const FString TypeString = Registry.GetType(TypeIndex).CPPType.ToString();
+				const FString TypeString = Registry.GetType_NoLock(TypeIndex).CPPType.ToString();
 				SupportedTypes.AddUnique(TypeString);
 			}
 			return true;
@@ -540,11 +629,11 @@ TArray<FString> FRigVMTemplateArgument::GetSupportedTypeStrings(const TArray<int
 	{
 		for(const int32 PermutationIndex : InPermutationIndices)
 		{
-			const TRigVMTypeIndex TypeIndex = GetTypeIndex(PermutationIndex);
+			const TRigVMTypeIndex TypeIndex = GetTypeIndex_NoLock(PermutationIndex);
 			// INDEX_NONE indicates deleted permutation
 			if (TypeIndex != INDEX_NONE)
 			{
-				const FString TypeString = Registry.GetType(TypeIndex).CPPType.ToString();
+				const FString TypeString = Registry.GetType_NoLock(TypeIndex).CPPType.ToString();
 				SupportedTypes.AddUnique(TypeString);
 			}
 		}
@@ -560,7 +649,7 @@ FRigVMTemplateArgument::CategoryViews::CategoryViews(const TArray<ETypeCategory>
 	Types.Reserve(InCategories.Num());
 	for (const ETypeCategory Category: InCategories)
 	{
-		Types.Emplace(FRigVMRegistry::Get().GetTypesForCategory(Category));
+		Types.Emplace(FRigVMRegistry_NoLock::GetForRead().GetTypesForCategory_NoLock(Category));
 	}
 }
 
@@ -680,7 +769,7 @@ TArray<TRigVMTypeIndex> FRigVMTemplateArgumentInfo::GetTypesFromCategories(
 	TSet<TRigVMTypeIndex> AllTypes;
 	for (const FRigVMTemplateArgument::ETypeCategory TypeCategory : InTypeCategories)
 	{
-		AllTypes.Append(FRigVMRegistry::Get().GetTypesForCategory(TypeCategory));
+		AllTypes.Append(FRigVMRegistry_NoLock::GetForRead().GetTypesForCategory_NoLock(TypeCategory));
 	}
 
 	TArray<TRigVMTypeIndex> Types;
@@ -895,11 +984,17 @@ TArray<UStruct*> FRigVMTemplate::GetSuperStructs(UStruct* InStruct, bool bInclud
 }
 
 FRigVMTemplate::FTypeMap FRigVMTemplate::GetArgumentTypesFromString(const FString& InTypeString, const FRigVMUserDefinedTypeResolver* InTypeResolver) const
+{
+	return GetArgumentTypesFromString_Impl(InTypeString, InTypeResolver, true);
+}
+
+FRigVMTemplate::FTypeMap FRigVMTemplate::GetArgumentTypesFromString_Impl(const FString& InTypeString, const FRigVMUserDefinedTypeResolver* InTypeResolver, bool bLockRegistry) const
 {
 	FTypeMap Types;
 	if(!InTypeString.IsEmpty())
 	{
-		FRigVMRegistry& Registry = FRigVMRegistry::Get();
+		const FRigVMRegistryWriteLock _(bLockRegistry);
+		FRigVMRegistry_NoLock& Registry = FRigVMRegistry_NoLock::GetForWrite();
 
 		FString Left, Right = InTypeString;
 		while(!Right.IsEmpty())
@@ -915,7 +1010,7 @@ FRigVMTemplate::FTypeMap FRigVMTemplate::GetArgumentTypesFromString(const FStrin
 			{
 				if(const FRigVMTemplateArgument* Argument = FindArgument(*ArgumentName))
 				{
-					TRigVMTypeIndex TypeIndex = Registry.GetTypeIndexFromCPPType(TypeName);
+					TRigVMTypeIndex TypeIndex = Registry.GetTypeIndexFromCPPType_NoLock(TypeName);
 
 					// If the type was not found, check if it's a user-defined type that hasn't been registered yet.
 					if (TypeIndex == INDEX_NONE && RigVMTypeUtils::RequiresCPPTypeObject(TypeName))
@@ -923,7 +1018,7 @@ FRigVMTemplate::FTypeMap FRigVMTemplate::GetArgumentTypesFromString(const FStrin
 						UObject* CPPTypeObject = RigVMTypeUtils::ObjectFromCPPType(TypeName, true, InTypeResolver);
 						
 						FRigVMTemplateArgumentType ArgType(*TypeName, CPPTypeObject);
-						TypeIndex = Registry.FindOrAddType(ArgType);
+						TypeIndex = Registry.FindOrAddType_NoLock(ArgType, false);
 					}
 					
 					if(TypeIndex != INDEX_NONE)
@@ -937,14 +1032,15 @@ FRigVMTemplate::FTypeMap FRigVMTemplate::GetArgumentTypesFromString(const FStrin
 	return Types;
 }
 
-FString FRigVMTemplate::GetStringFromArgumentTypes(const FTypeMap& InTypes)
+FString FRigVMTemplate::GetStringFromArgumentTypes(const FTypeMap& InTypes, bool bLockRegistry)
 {
-	const FRigVMRegistry& Registry = FRigVMRegistry::Get();
+	const FRigVMRegistryReadLock _(bLockRegistry); 
+	const FRigVMRegistry_NoLock& Registry = FRigVMRegistry_NoLock::GetForRead();
 	
 	TArray<FString> TypePairStrings;
 	for(const TPair<FName,TRigVMTypeIndex>& Pair : InTypes)
 	{
-		const FRigVMTemplateArgumentType& Type = Registry.GetType(Pair.Value);
+		const FRigVMTemplateArgumentType& Type = Registry.GetType_NoLock(Pair.Value);
 		static constexpr TCHAR Format[] = TEXT("%s:%s");
 		TypePairStrings.Add(FString::Printf(Format, *Pair.Key.ToString(), *Type.CPPType.ToString()));
 	}
@@ -1249,15 +1345,15 @@ bool FRigVMTemplate::Merge(const FRigVMTemplate& InOther)
 		return false;
 	}
 
-	if(InOther.GetExecuteContextStruct() != GetExecuteContextStruct())
+	if(InOther.GetExecuteContextStruct(false) != GetExecuteContextStruct(false))
 	{
 		// find the previously defined permutation.
 		UE_LOG(LogRigVM, Display, TEXT("RigVMFunction '%s' cannot be merged into the '%s' template. ExecuteContext Types differ ('%s' vs '%s' from '%s')."),
-			*InOther.GetPrimaryPermutation()->Name,
+			*InOther.GetPrimaryPermutation(false)->Name,
 			*GetNotation().ToString(),
-			*InOther.GetExecuteContextStruct()->GetStructCPPName(),
-			*GetExecuteContextStruct()->GetStructCPPName(),
-			*GetPrimaryPermutation()->Name);
+			*InOther.GetExecuteContextStruct(false)->GetStructCPPName(),
+			*GetExecuteContextStruct(false)->GetStructCPPName(),
+			*GetPrimaryPermutation(false)->Name);
 		return false;
 	}
 
@@ -1272,7 +1368,7 @@ bool FRigVMTemplate::Merge(const FRigVMTemplate& InOther)
 		int32 MatchingArguments = 0;
 		for(int32 ArgumentIndex = 0; ArgumentIndex < Arguments.Num(); ArgumentIndex++)
 		{
-			if (Arguments[ArgumentIndex].GetTypeIndex(PermutationIndex) == InOther.Arguments[ArgumentIndex].GetTypeIndex(0))
+			if (Arguments[ArgumentIndex].GetTypeIndex_NoLock(PermutationIndex) == InOther.Arguments[ArgumentIndex].GetTypeIndex_NoLock(0))
 			{
 				MatchingArguments++;
 			}
@@ -1281,9 +1377,9 @@ bool FRigVMTemplate::Merge(const FRigVMTemplate& InOther)
 		{
 			// find the previously defined permutation.
 			UE_LOG(LogRigVM, Display, TEXT("RigVMFunction '%s' cannot be merged into the '%s' template. It collides with '%s'."),
-				*InOther.GetPrimaryPermutation()->Name,
+				*InOther.GetPrimaryPermutation(false)->Name,
 				*GetNotation().ToString(),
-				*GetPermutation(PermutationIndex)->Name);
+				*GetPermutation_NoLock(PermutationIndex)->Name);
 			return false;
 		}
 	}
@@ -1293,7 +1389,7 @@ bool FRigVMTemplate::Merge(const FRigVMTemplate& InOther)
 	for (int32 ArgumentIndex = 0; ArgumentIndex < Arguments.Num(); ArgumentIndex++)
 	{
 		const FRigVMTemplateArgument& OtherArg = InOther.Arguments[ArgumentIndex];
-		if (OtherArg.GetNumTypes() != 1)
+		if (OtherArg.GetNumTypes_NoLock() != 1)
 		{
 			return false;
 		}
@@ -1301,8 +1397,8 @@ bool FRigVMTemplate::Merge(const FRigVMTemplate& InOther)
 		// Add Other argument information into the TypeToPermutations map
 		{
 			FRigVMTemplateArgument& NewArg = NewArgs.Add_GetRef(Arguments[ArgumentIndex]);
-			const TRigVMTypeIndex OtherTypeIndex = OtherArg.GetTypeIndex(0);
-			const int32 NewPermutationIndex = NewArg.GetNumTypes();
+			const TRigVMTypeIndex OtherTypeIndex = OtherArg.GetTypeIndex_NoLock(0);
+			const int32 NewPermutationIndex = NewArg.GetNumTypes_NoLock();
 			if (TArray<int32>* ArgTypePermutations = NewArg.TypeToPermutations.Find(OtherTypeIndex))
 			{
 				ArgTypePermutations->Add(NewPermutationIndex);
@@ -1360,14 +1456,16 @@ const TArray<FRigVMExecuteArgument>& FRigVMTemplate::GetExecuteArguments(const F
 {
 	if(ExecuteArguments.IsEmpty())
 	{
+		const FRigVMRegistryWriteLock _;
+
 		if(UsesDispatch())
 		{
 			const FRigVMDispatchFactory* Factory = Delegates.GetDispatchFactoryDelegate.Execute();
 			check(Factory);
 
-			ExecuteArguments = Factory->GetExecuteArguments(InContext);
+			ExecuteArguments = Factory->GetExecuteArguments_NoLock(InContext);
 		}
-		else if(const FRigVMFunction* PrimaryPermutation = GetPrimaryPermutation())
+		else if(const FRigVMFunction* PrimaryPermutation = GetPrimaryPermutation(false /* lock */))
 		{
 			if(PrimaryPermutation->Struct)
 			{
@@ -1380,7 +1478,7 @@ const TArray<FRigVMExecuteArgument>& FRigVMTemplate::GetExecuteArguments(const F
 						FRigVMTemplateArgument Argument(*It);
 						if(Argument.IsExecute())
 						{
-							ExecuteArguments.Emplace(Argument.Name, Argument.Direction, Argument.GetTypeIndex(0));
+							ExecuteArguments.Emplace(Argument.Name, Argument.Direction, Argument.GetTypeIndex_NoLock(0));
 						}
 					}
 				}
@@ -1390,19 +1488,19 @@ const TArray<FRigVMExecuteArgument>& FRigVMTemplate::GetExecuteArguments(const F
 	return ExecuteArguments;
 }
 
-const UScriptStruct* FRigVMTemplate::GetExecuteContextStruct() const
+const UScriptStruct* FRigVMTemplate::GetExecuteContextStruct(bool bLockRegistry) const
 {
-	if(const FRigVMDispatchFactory* Factory = GetDispatchFactory())
+	if(const FRigVMDispatchFactory* Factory = GetDispatchFactory(bLockRegistry))
 	{
 		return Factory->GetExecuteContextStruct();
 	}
 	check(!Permutations.IsEmpty());
-	return GetPrimaryPermutation()->GetExecuteContextStruct();
+	return GetPrimaryPermutation(bLockRegistry)->GetExecuteContextStruct(bLockRegistry);
 }
 
 bool FRigVMTemplate::SupportsExecuteContextStruct(const UScriptStruct* InExecuteContextStruct) const
 {
-	return InExecuteContextStruct->IsChildOf(GetExecuteContextStruct());
+	return InExecuteContextStruct->IsChildOf(GetExecuteContextStruct(true));
 }
 
 bool FRigVMTemplate::ArgumentSupportsTypeIndex(const FName& InArgumentName, TRigVMTypeIndex InTypeIndex, TRigVMTypeIndex* OutTypeIndex) const
@@ -1414,36 +1512,35 @@ bool FRigVMTemplate::ArgumentSupportsTypeIndex(const FName& InArgumentName, TRig
 	return false;
 }
 
-const FRigVMFunction* FRigVMTemplate::GetPrimaryPermutation() const
+const FRigVMFunction* FRigVMTemplate::GetPrimaryPermutation(bool bLockRegistry) const
 {
 	if (NumPermutations() > 0)
 	{
-		return GetPermutation(0);
+		return GetPermutation(0, bLockRegistry);
 	}
 	return nullptr;
 }
 
-const FRigVMFunction* FRigVMTemplate::GetPermutation(int32 InIndex) const
+const FRigVMFunction* FRigVMTemplate::GetPermutation(int32 InIndex, bool bLockRegistry) const
 {
-	FScopeLock FunctionRegistryScopeLock(&FRigVMRegistry::FunctionRegistryMutex);
+	const FRigVMRegistryReadLock _(bLockRegistry);
 	return GetPermutation_NoLock(InIndex);
 }
 
 const FRigVMFunction* FRigVMTemplate::GetPermutation_NoLock(int32 InIndex) const
 {
-	const FRigVMRegistry& Registry = FRigVMRegistry::Get();
+	const FRigVMRegistry_NoLock& Registry = FRigVMRegistry_NoLock::GetForRead();
 	const int32 FunctionIndex = Permutations[InIndex];
-	if(Registry.GetFunctions().IsValidIndex(FunctionIndex))
+	if(Registry.GetFunctions_NoLock().IsValidIndex(FunctionIndex))
 	{
-		return &Registry.GetFunctions()[Permutations[InIndex]];
+		return &Registry.GetFunctions_NoLock()[Permutations[InIndex]];
 	}
 	return nullptr;
 }
 
-const FRigVMFunction* FRigVMTemplate::GetOrCreatePermutation(int32 InIndex)
+const FRigVMFunction* FRigVMTemplate::GetOrCreatePermutation(int32 InIndex, bool bLockRegistry)
 {
-	FScopeLock FunctionRegistryScopeLock(&FRigVMRegistry::FunctionRegistryMutex);
-
+	const FRigVMRegistryWriteLock _(bLockRegistry);
 	return GetOrCreatePermutation_NoLock(InIndex);
 }
 
@@ -1456,28 +1553,28 @@ const FRigVMFunction* FRigVMTemplate::GetOrCreatePermutation_NoLock(int32 InInde
 
 	if(Permutations[InIndex] == INDEX_NONE && UsesDispatch())
 	{
-		FRigVMRegistry& Registry = FRigVMRegistry::Get();
+		FRigVMRegistry_NoLock& Registry = FRigVMRegistry_NoLock::GetForWrite();
 		
 		FTypeMap Types;
 		for(const FRigVMTemplateArgument& Argument : Arguments)
 		{
-			Types.Add(Argument.GetName(), Argument.GetTypeIndex(InIndex));
+			Types.Add(Argument.GetName(), Argument.GetTypeIndex_NoLock(InIndex));
 		}
 
 		FRigVMDispatchFactory* Factory = Delegates.GetDispatchFactoryDelegate.Execute();
 		if (ensure(Factory))
 		{
-			const FRigVMFunctionPtr DispatchFunction = Factory->CreateDispatchFunction(Types);
+			const FRigVMFunctionPtr DispatchFunction = Factory->CreateDispatchFunction_NoLock(Types);
 
 			TArray<FRigVMFunctionArgument> FunctionArguments;
 			for(const FRigVMTemplateArgument& Argument : Arguments)
 			{
-				const FRigVMTemplateArgumentType& Type = Registry.GetType(Argument.GetTypeIndex(InIndex));
+				const FRigVMTemplateArgumentType& Type = Registry.GetType_NoLock(Argument.GetTypeIndex_NoLock(InIndex));
 				FunctionArguments.Add(FRigVMFunctionArgument(Argument.Name.ToString(), Type.CPPType.ToString()));
 			}
 			
 			static constexpr TCHAR Format[] = TEXT("%s::%s");
-			const FString PermutationName = Factory->GetPermutationNameImpl(Types);
+			const FString PermutationName = Factory->GetPermutationNameImpl(Types, false);
 			const int32 FunctionIndex = Permutations[InIndex] = Registry.Functions.Num();
 			
 			Registry.Functions.AddElement(
@@ -1492,7 +1589,7 @@ const FRigVMFunction* FRigVMTemplate::GetOrCreatePermutation_NoLock(int32 InInde
 			Registry.Functions[FunctionIndex].TemplateIndex = Index;
 			Registry.FunctionNameToIndex.Add(*PermutationName, FunctionIndex);
 
-			TArray<FRigVMFunction> Predicates = Factory->CreateDispatchPredicates(Types);
+			TArray<FRigVMFunction> Predicates = Factory->CreateDispatchPredicates_NoLock(Types);
 			Registry.StructNameToPredicates.Add(*PermutationName, MoveTemp(Predicates));
 
 			return &Registry.Functions[FunctionIndex];
@@ -1513,21 +1610,21 @@ int32 FRigVMTemplate::FindPermutation(const FRigVMFunction* InPermutation) const
 	return Permutations.Find(InPermutation->Index);
 }
 
-int32 FRigVMTemplate::FindPermutation(const FTypeMap& InTypes) const
+int32 FRigVMTemplate::FindPermutation(const FTypeMap& InTypes, bool bLockRegistry) const
 {
 	FTypeMap Types = InTypes;
 	int32 PermutationIndex = INDEX_NONE;
-	if(FullyResolve(Types, PermutationIndex))
+	if(FullyResolve(Types, PermutationIndex, bLockRegistry))
 	{
 		return PermutationIndex;
 	}
 	return INDEX_NONE;
 }
 
-bool FRigVMTemplate::FullyResolve(FRigVMTemplate::FTypeMap& InOutTypes, int32& OutPermutationIndex) const
+bool FRigVMTemplate::FullyResolve(FRigVMTemplate::FTypeMap& InOutTypes, int32& OutPermutationIndex, bool bLockRegistry) const
 {
 	TArray<int32> PermutationIndices;
-	Resolve(InOutTypes, PermutationIndices, false);
+	Resolve(InOutTypes, PermutationIndices, false, bLockRegistry);
 	if(PermutationIndices.Num() == 1)
 	{
 		OutPermutationIndex = PermutationIndices[0];
@@ -1539,9 +1636,10 @@ bool FRigVMTemplate::FullyResolve(FRigVMTemplate::FTypeMap& InOutTypes, int32& O
 	return OutPermutationIndex != INDEX_NONE;
 }
 
-bool FRigVMTemplate::Resolve(FTypeMap& InOutTypes, TArray<int32>& OutPermutationIndices, bool bAllowFloatingPointCasts) const
+bool FRigVMTemplate::Resolve(FTypeMap& InOutTypes, TArray<int32>& OutPermutationIndices, bool bAllowFloatingPointCasts, bool bLockRegistry) const
 {
-	const FRigVMRegistry& Registry = FRigVMRegistry::Get();
+	const FRigVMRegistryReadLock _(bLockRegistry);
+	const FRigVMRegistry_NoLock& Registry = FRigVMRegistry_NoLock::GetForRead();
 	
 	FTypeMap InputTypes = InOutTypes;
 	InOutTypes.Reset();
@@ -1554,13 +1652,15 @@ bool FRigVMTemplate::Resolve(FTypeMap& InOutTypes, TArray<int32>& OutPermutation
 	
 	for (const FRigVMTemplateArgument& Argument : Arguments)
 	{
-		if (Argument.IsSingleton())
+		if (Argument.IsSingleton_NoLock())
 		{
-			InOutTypes.Add(Argument.Name, Argument.GetTypeIndex(0));
+			InOutTypes.Add(Argument.Name, Argument.GetTypeIndex_NoLock(0));
 			continue;
 		}
 		else if (const TRigVMTypeIndex* InputType = InputTypes.Find(Argument.Name))
 		{
+			TArray<TRigVMTypeIndex> AllTypes; Argument.GetAllTypes_NoLock(AllTypes);
+
 			TRigVMTypeIndex MatchedType = *InputType;
 			bool bFoundMatch = false;
 			bool bFoundPerfectMatch = false;
@@ -1569,10 +1669,10 @@ bool FRigVMTemplate::Resolve(FTypeMap& InOutTypes, TArray<int32>& OutPermutation
 			// instead of removing them one by one, which can be costly
 			TMap<int32, bool> PermutationsToKeep;
 
-			for (int32 PermutationIndex = 0; PermutationIndex < Argument.GetNumTypes(); PermutationIndex++)
+			for (int32 PermutationIndex = 0; PermutationIndex < AllTypes.Num(); PermutationIndex++)
 			{
-				const TRigVMTypeIndex Type = Argument.GetTypeIndex(PermutationIndex);
-				if(!Registry.CanMatchTypes(Type, *InputType, bAllowFloatingPointCasts))
+				const TRigVMTypeIndex Type = AllTypes[PermutationIndex];
+				if(!Registry.CanMatchTypes_NoLock(Type, *InputType, bAllowFloatingPointCasts))
 				{
 					PermutationsToKeep.FindOrAdd(PermutationIndex) = false;
 				}
@@ -1616,7 +1716,7 @@ bool FRigVMTemplate::Resolve(FTypeMap& InOutTypes, TArray<int32>& OutPermutation
 
 			if(const TRigVMTypeIndex* InputType = InputTypes.Find(Argument.Name))
 			{
-				if(Registry.IsArrayType(*InputType))
+				if(Registry.IsArrayType_NoLock(*InputType))
 				{
 					InOutTypes.FindChecked(Argument.Name) = RigVMTypeUtils::TypeIndex::WildCardArray;
 				}
@@ -1637,16 +1737,16 @@ bool FRigVMTemplate::Resolve(FTypeMap& InOutTypes, TArray<int32>& OutPermutation
 		InOutTypes.Reset();
 		for (const FRigVMTemplateArgument& Argument : Arguments)
 		{
-			InOutTypes.Add(Argument.Name, Argument.GetTypeIndex(OutPermutationIndices[0]));
+			InOutTypes.Add(Argument.Name, Argument.GetTypeIndex_NoLock(OutPermutationIndices[0]));
 		}
 	}
 	else if (OutPermutationIndices.Num() > 1)
 	{
 		for (const FRigVMTemplateArgument& Argument : Arguments)
 		{
-			if (Argument.IsSingleton(OutPermutationIndices))
+			if (Argument.IsSingleton_NoLock(OutPermutationIndices))
 			{
-				InOutTypes.FindChecked(Argument.Name) = Argument.GetTypeIndex(OutPermutationIndices[0]);
+				InOutTypes.FindChecked(Argument.Name) = Argument.GetTypeIndex_NoLock(OutPermutationIndices[0]);
 			}
 		}
 	}
@@ -1664,10 +1764,10 @@ uint32 FRigVMTemplate::GetTypesHashFromTypes(const FTypeMap& InTypes) const
 		return TypeHash;
 	}
 
-	const FRigVMRegistry& Registry = FRigVMRegistry::Get();
+	const FRigVMRegistry_NoLock& Registry = FRigVMRegistry_NoLock::GetForRead();
 	for (const TPair<FName, TRigVMTypeIndex>& Pair : InTypes)
 	{
-		if (Registry.IsWildCardType(Pair.Value))
+		if (Registry.IsWildCardType_NoLock(Pair.Value))
 		{
 			return TypeHash;
 		}
@@ -1739,16 +1839,17 @@ bool FRigVMTemplate::ContainsPermutation(const FTypeMap& InTypes) const
 }
 
 bool FRigVMTemplate::ResolveArgument(const FName& InArgumentName, const TRigVMTypeIndex InTypeIndex,
-	FTypeMap& InOutTypes) const
+	FTypeMap& InOutTypes, bool bLockRegistry) const
 {
-	const FRigVMRegistry& Registry = FRigVMRegistry::Get();
+	const FRigVMRegistryReadLock _(bLockRegistry);
+	const FRigVMRegistry_NoLock& Registry = FRigVMRegistry_NoLock::GetForRead();
 
 	auto RemoveWildCardTypes = [&Registry](const FTypeMap& InTypes)
 	{
 		FTypeMap FilteredTypes;
 		for(const FTypePair& Pair: InTypes)
 		{
-			if(!Registry.IsWildCardType(Pair.Value))
+			if(!Registry.IsWildCardType_NoLock(Pair.Value))
 			{
 				FilteredTypes.Add(Pair);
 			}
@@ -1765,11 +1866,11 @@ bool FRigVMTemplate::ResolveArgument(const FName& InArgumentName, const TRigVMTy
 	TArray<int32> PermutationIndices;
 	FTypeMap RemainingTypesToResolve;
 	
-	if(Resolve(ResolvedTypes, PermutationIndices, true))
+	if(Resolve(ResolvedTypes, PermutationIndices, true, false))
 	{
 		// let's see if the input argument resolved into the expected type
 		const TRigVMTypeIndex ResolvedInputType = ResolvedTypes.FindChecked(InArgumentName);
-		if(!Registry.CanMatchTypes(ResolvedInputType, InTypeIndex, true))
+		if(!Registry.CanMatchTypes_NoLock(ResolvedInputType, InTypeIndex, true))
 		{
 			return false;
 		}
@@ -1806,7 +1907,7 @@ bool FRigVMTemplate::ResolveArgument(const FName& InArgumentName, const TRigVMTy
 			FTypeMap NewResolvedTypes = RemoveWildCardTypes(ResolvedTypes);
 			NewResolvedTypes.FindOrAdd(TypeToResolve.Key) = TypeToResolve.Value;
 
-			if(Resolve(NewResolvedTypes, PermutationIndices, true))
+			if(Resolve(NewResolvedTypes, PermutationIndices, true, false))
 			{
 				ResolvedTypes = NewResolvedTypes;
 				RemoveResolvedTypesFromRemainingList(InOutTypes, ResolvedTypes, RemainingTypesToResolve);
@@ -1826,19 +1927,22 @@ bool FRigVMTemplate::ResolveArgument(const FName& InArgumentName, const TRigVMTy
 	return false;
 }
 
-FRigVMTemplateTypeMap FRigVMTemplate::GetTypesForPermutation(const int32 InPermutationIndex) const
+FRigVMTemplateTypeMap FRigVMTemplate::GetTypesForPermutation(const int32 InPermutationIndex, const bool bLockRegistry) const
+{
+	const FRigVMRegistryReadLock _;
+	return GetTypesForPermutation_NoLock(InPermutationIndex);
+}
+
+FRigVMTemplateTypeMap FRigVMTemplate::GetTypesForPermutation_NoLock(const int32 InPermutationIndex) const
 {
 	FTypeMap TypeMap;
 	for (int32 ArgIndex = 0; ArgIndex < NumArguments(); ++ArgIndex)
 	{
 		const FRigVMTemplateArgument* Argument = GetArgument(ArgIndex);
-		if (Argument->GetNumTypes() > InPermutationIndex)
-		{
-			TypeMap.Add(Argument->GetName(), Argument->GetTypeIndex(InPermutationIndex));
-		}
-		else if (Argument->IsSingleton())
+		const TOptional<TRigVMTypeIndex> TypeIndex = Argument->TryToGetTypeIndex_NoLock(InPermutationIndex);
+		if(TypeIndex.IsSet())
 		{
-			TypeMap.Add(Argument->GetName(), Argument->GetTypeIndex(0));
+			TypeMap.Add(Argument->GetName(), TypeIndex.GetValue());
 		}
 		else
 		{
@@ -1938,7 +2042,7 @@ bool FRigVMTemplate::UpdateArgumentTypes()
 	for(int32 ArgumentIndex = 0; ArgumentIndex < Arguments.Num(); ArgumentIndex++)
 	{
 		FRigVMTemplateArgument& Argument = Arguments[ArgumentIndex];
-		if(Argument.bUseCategories || Argument.IsSingleton())
+		if(Argument.bUseCategories || Argument.IsSingleton_NoLock())
 		{
 			continue;
 		}
@@ -1947,7 +2051,7 @@ bool FRigVMTemplate::UpdateArgumentTypes()
 		Argument.TypeToPermutations.Reset();
 	}
 
-	const FRigVMRegistry& Registry = FRigVMRegistry::Get();
+	const FRigVMRegistry_NoLock& Registry = FRigVMRegistry_NoLock::GetForRead();
 	const FRigVMTemplateArgument& PrimaryArgument = Arguments[PrimaryArgumentIndex];
 	TArray<FRigVMTemplateTypeMap, TInlineAllocator<1>> TypesArray;
 	bool bResult = true;
@@ -1981,7 +2085,7 @@ bool FRigVMTemplate::UpdateArgumentTypes()
 					for (TPair<FName, TRigVMTypeIndex>& ArgumentAndType : Types)
 					{
 						// similar to FRigVMTemplateArgument::EnsureValidExecuteType
-						Registry.ConvertExecuteContextToBaseType(ArgumentAndType.Value);
+						Registry.ConvertExecuteContextToBaseType_NoLock(ArgumentAndType.Value);
 					}
 
 					// Find if these types were already registered
@@ -2010,7 +2114,7 @@ bool FRigVMTemplate::UpdateArgumentTypes()
 						const TRigVMTypeIndex& TypeIndex = *TypeIndexPtr;
 						TypeHash = HashCombine(TypeHash, GetTypeHash(TypeIndex));
 
-						if(Argument.bUseCategories || Argument.IsSingleton())
+						if(Argument.bUseCategories || Argument.IsSingleton_NoLock())
 						{
 							continue;
 						}
@@ -2051,7 +2155,7 @@ void FRigVMTemplate::HandleTypeRemoval(TRigVMTypeIndex InTypeIndex)
 	TArray<int32> PermutationsToRemove;
 	for (int32 PermutationIndex = 0; PermutationIndex < NumPermutations(); PermutationIndex++)
 	{
-		FRigVMTemplateTypeMap TypeMap = GetTypesForPermutation(PermutationIndex);
+		FRigVMTemplateTypeMap TypeMap = GetTypesForPermutation_NoLock(PermutationIndex);
 
 		TArray<TRigVMTypeIndex> Types;
 		TypeMap.GenerateValueArray(Types);
@@ -2081,6 +2185,12 @@ void FRigVMTemplate::HandleTypeRemoval(TRigVMTypeIndex InTypeIndex)
 	}
 }
 
+const FRigVMDispatchFactory* FRigVMTemplate::GetDispatchFactory(const bool bLockRegistry) const
+{
+	FRigVMRegistryReadLock _(bLockRegistry);
+	return GetDispatchFactory_NoLock();
+}
+
 void FRigVMTemplate::RecomputeTypesHashToPermutations()
 {
 	TypesHashToPermutation.Reset();
@@ -2089,7 +2199,7 @@ void FRigVMTemplate::RecomputeTypesHashToPermutations()
 		uint32 TypesHash=0;
 		for (int32 ArgIndex=0; ArgIndex<NumArguments(); ++ArgIndex)
 		{
-			TypesHash = HashCombine(Hash, GetTypeHash(Arguments[ArgIndex].GetTypeIndex(PermutationIndex)));
+			TypesHash = HashCombine(Hash, GetTypeHash(Arguments[ArgIndex].GetTypeIndex_NoLock(PermutationIndex)));
 		}
 		TypesHashToPermutation.Add(TypesHash, PermutationIndex);
 	}
@@ -2105,25 +2215,30 @@ void FRigVMTemplate::UpdateTypesHashToPermutation(const int32 InPermutation)
 	uint32 TypeHash=0;
 	for (const FRigVMTemplateArgument& Argument : Arguments)
 	{
-		TypeHash = HashCombine(TypeHash, GetTypeHash(Argument.GetTypeIndex(InPermutation)));
+		TypeHash = HashCombine(TypeHash, GetTypeHash(Argument.GetTypeIndex_NoLock(InPermutation)));
 	}
 	TypesHashToPermutation.Add(TypeHash, InPermutation);
 }
 
-uint32 GetTypeHash(const FRigVMTemplateArgument& InArgument)
+uint32 GetTypeHash_NoLock(const FRigVMTemplateArgument& InArgument)
 {
-	const FRigVMRegistry& Registry = FRigVMRegistry::Get();
+	const FRigVMRegistry_NoLock& Registry = FRigVMRegistry_NoLock::GetForRead();
 	uint32 Hash = GetTypeHash(InArgument.Name.ToString());
 	Hash = HashCombine(Hash, GetTypeHash((int32)InArgument.Direction));
 	InArgument.ForEachType([&](const TRigVMTypeIndex TypeIndex)
 	{
-		Hash = HashCombine(Hash, Registry.GetHashForType(TypeIndex));
+		Hash = HashCombine(Hash, Registry.GetHashForType_NoLock(TypeIndex));
 		return true;
 	});
 	return Hash;
 }
 
 uint32 GetTypeHash(const FRigVMTemplate& InTemplate)
+{
+	return InTemplate.ComputeTypeHash();
+}
+
+uint32 GetTypeHash_NoLock(const FRigVMTemplate& InTemplate)
 {
 	if(InTemplate.Hash != UINT32_MAX)
 	{
@@ -2133,12 +2248,12 @@ uint32 GetTypeHash(const FRigVMTemplate& InTemplate)
 	uint32 Hash = GetTypeHash(InTemplate.GetNotation().ToString());
 	for(const FRigVMTemplateArgument& Argument : InTemplate.Arguments)
 	{
-		Hash = HashCombine(Hash, GetTypeHash(Argument));
+		Hash = HashCombine(Hash, GetTypeHash_NoLock(Argument));
 	}
 
 	// todo: in Dev-EngineMerge we should add the execute arguments to the hash as well
 
-	if(const FRigVMDispatchFactory* Factory = InTemplate.GetDispatchFactory())
+	if(const FRigVMDispatchFactory* Factory = InTemplate.GetDispatchFactory_NoLock())
 	{
 		Hash = HashCombine(Hash, GetTypeHash(Factory->GetFactoryName().ToString()));
 	}
@@ -2147,3 +2262,8 @@ uint32 GetTypeHash(const FRigVMTemplate& InTemplate)
 	return Hash;
 }
 
+uint32 FRigVMTemplate::ComputeTypeHash() const
+{
+	const FRigVMRegistryReadLock _; 
+	return GetTypeHash_NoLock(*this);
+}
diff --git a/Engine/Plugins/Runtime/RigVM/Source/RigVM/Private/RigVMFunctions/RigVMDispatch_Array.cpp b/Engine/Plugins/Runtime/RigVM/Source/RigVM/Private/RigVMFunctions/RigVMDispatch_Array.cpp
--- a/Engine/Plugins/Runtime/RigVM/Source/RigVM/Private/RigVMFunctions/RigVMDispatch_Array.cpp
+++ b/Engine/Plugins/Runtime/RigVM/Source/RigVM/Private/RigVMFunctions/RigVMDispatch_Array.cpp
@@ -508,12 +508,12 @@ FRigVMTemplateTypeMap FRigVMDispatch_ArrayGetAtIndex::OnNewArgumentType(const FN
 		return {
 			{ArrayName, InTypeIndex},
 			{IndexName, RigVMTypeUtils::TypeIndex::Int32},
-			{ElementName, FRigVMRegistry::Get().GetBaseTypeFromArrayTypeIndex(InTypeIndex)}
+			{ElementName, FRigVMRegistry_NoLock::GetForRead().GetBaseTypeFromArrayTypeIndex_NoLock(InTypeIndex)}
 		};
 	}
 
 	return {
-		{ArrayName, FRigVMRegistry::Get().GetArrayTypeFromBaseTypeIndex(InTypeIndex)},
+		{ArrayName, FRigVMRegistry_NoLock::GetForRead().GetArrayTypeFromBaseTypeIndex_NoLock(InTypeIndex)},
 		{IndexName, RigVMTypeUtils::TypeIndex::Int32},
 		{ElementName, InTypeIndex}
 	};
@@ -593,12 +593,12 @@ FRigVMTemplateTypeMap FRigVMDispatch_ArraySetAtIndex::OnNewArgumentType(const FN
 		return {
 			{ArrayName, InTypeIndex},
 			{IndexName, RigVMTypeUtils::TypeIndex::Int32},
-			{ElementName, FRigVMRegistry::Get().GetBaseTypeFromArrayTypeIndex(InTypeIndex)}
+			{ElementName, FRigVMRegistry_NoLock::GetForRead().GetBaseTypeFromArrayTypeIndex_NoLock(InTypeIndex)}
 		};
 	}
 
 	return {
-		{ArrayName, FRigVMRegistry::Get().GetArrayTypeFromBaseTypeIndex(InTypeIndex)},
+		{ArrayName, FRigVMRegistry_NoLock::GetForRead().GetArrayTypeFromBaseTypeIndex_NoLock(InTypeIndex)},
 		{IndexName, RigVMTypeUtils::TypeIndex::Int32},
 		{ElementName, InTypeIndex}
 	};
@@ -897,14 +897,14 @@ FRigVMTemplateTypeMap FRigVMDispatch_ArrayFind::OnNewArgumentType(const FName& I
 	{
 		return {
 			{ArrayName, InTypeIndex},
-			{ElementName, FRigVMRegistry::Get().GetBaseTypeFromArrayTypeIndex(InTypeIndex)},
+			{ElementName, FRigVMRegistry_NoLock::GetForRead().GetBaseTypeFromArrayTypeIndex_NoLock(InTypeIndex)},
 			{IndexName, RigVMTypeUtils::TypeIndex::Int32},
 			{SuccessName, RigVMTypeUtils::TypeIndex::Bool}
 		};
 	}
 
 	return {
-		{ArrayName, FRigVMRegistry::Get().GetArrayTypeFromBaseTypeIndex(InTypeIndex)},
+		{ArrayName, FRigVMRegistry_NoLock::GetForRead().GetArrayTypeFromBaseTypeIndex_NoLock(InTypeIndex)},
 		{ElementName, InTypeIndex},
 		{IndexName, RigVMTypeUtils::TypeIndex::Int32},
 		{SuccessName, RigVMTypeUtils::TypeIndex::Bool}
@@ -1476,7 +1476,7 @@ FRigVMTemplateTypeMap FRigVMDispatch_ArrayIterator::OnNewArgumentType(const FNam
 	{
 		return {
 			{ArrayName, InTypeIndex},
-			{ElementName, FRigVMRegistry::Get().GetBaseTypeFromArrayTypeIndex(InTypeIndex)},
+			{ElementName, FRigVMRegistry_NoLock::GetForRead().GetBaseTypeFromArrayTypeIndex_NoLock(InTypeIndex)},
 			{IndexName, RigVMTypeUtils::TypeIndex::Int32},
 			{CountName, RigVMTypeUtils::TypeIndex::Int32},
 			{RatioName, RigVMTypeUtils::TypeIndex::Float},
@@ -1485,7 +1485,7 @@ FRigVMTemplateTypeMap FRigVMDispatch_ArrayIterator::OnNewArgumentType(const FNam
 	}
 
 	return {
-		{ArrayName, FRigVMRegistry::Get().GetArrayTypeFromBaseTypeIndex(InTypeIndex)},
+		{ArrayName, FRigVMRegistry_NoLock::GetForRead().GetArrayTypeFromBaseTypeIndex_NoLock(InTypeIndex)},
 		{ElementName, InTypeIndex},
 		{IndexName, RigVMTypeUtils::TypeIndex::Int32},
 		{CountName, RigVMTypeUtils::TypeIndex::Int32},
diff --git a/Engine/Plugins/Runtime/RigVM/Source/RigVM/Private/RigVMFunctions/RigVMDispatch_CastEnum.cpp b/Engine/Plugins/Runtime/RigVM/Source/RigVM/Private/RigVMFunctions/RigVMDispatch_CastEnum.cpp
--- a/Engine/Plugins/Runtime/RigVM/Source/RigVM/Private/RigVMFunctions/RigVMDispatch_CastEnum.cpp
+++ b/Engine/Plugins/Runtime/RigVM/Source/RigVM/Private/RigVMFunctions/RigVMDispatch_CastEnum.cpp
@@ -38,7 +38,7 @@ bool FRigVMDispatch_CastEnumToInt::GetPermutationsFromArgumentType(const FName&
 	}
 	else if (InArgumentName == ResultName && InTypeIndex == RigVMTypeUtils::TypeIndex::Int32)
 	{
-		const TArray<TRigVMTypeIndex>& EnumTypes = FRigVMRegistry::Get().GetTypesForCategory(FRigVMTemplateArgument::ETypeCategory_SingleEnumValue);
+		const TArray<TRigVMTypeIndex>& EnumTypes = FRigVMRegistry_NoLock::GetForRead().GetTypesForCategory_NoLock(FRigVMTemplateArgument::ETypeCategory_SingleEnumValue);
 		for (const TRigVMTypeIndex& Type : EnumTypes)
 		{
 			OutPermutations.Add(
@@ -132,7 +132,7 @@ bool FRigVMDispatch_CastIntToEnum::GetPermutationsFromArgumentType(const FName&
 	}
 	else if (InArgumentName == ValueName && InTypeIndex == RigVMTypeUtils::TypeIndex::Int32)
 	{
-		const TArray<TRigVMTypeIndex>& EnumTypes = FRigVMRegistry::Get().GetTypesForCategory(FRigVMTemplateArgument::ETypeCategory_SingleEnumValue);
+		const TArray<TRigVMTypeIndex>& EnumTypes = FRigVMRegistry_NoLock::GetForRead().GetTypesForCategory_NoLock(FRigVMTemplateArgument::ETypeCategory_SingleEnumValue);
 		for (const TRigVMTypeIndex& Type : EnumTypes)
 		{
 			OutPermutations.Add(
diff --git a/Engine/Plugins/Runtime/RigVM/Source/RigVM/Private/RigVMFunctions/RigVMDispatch_CastObject.cpp b/Engine/Plugins/Runtime/RigVM/Source/RigVM/Private/RigVMFunctions/RigVMDispatch_CastObject.cpp
--- a/Engine/Plugins/Runtime/RigVM/Source/RigVM/Private/RigVMFunctions/RigVMDispatch_CastObject.cpp
+++ b/Engine/Plugins/Runtime/RigVM/Source/RigVM/Private/RigVMFunctions/RigVMDispatch_CastObject.cpp
@@ -29,7 +29,7 @@ bool FRigVMDispatch_CastObject::GetPermutationsFromArgumentType(const FName& InA
 {
 	if (InArgumentName == ValueName)
 	{
-		const TArray<TRigVMTypeIndex>& ObjectTypes = FRigVMRegistry::Get().GetTypesForCategory(FRigVMTemplateArgument::ETypeCategory_SingleObjectValue);
+		const TArray<TRigVMTypeIndex>& ObjectTypes = FRigVMRegistry_NoLock::GetForRead().GetTypesForCategory_NoLock(FRigVMTemplateArgument::ETypeCategory_SingleObjectValue);
 		for (const TRigVMTypeIndex& Type : ObjectTypes)
 		{
 			OutPermutations.Add(
@@ -41,7 +41,7 @@ bool FRigVMDispatch_CastObject::GetPermutationsFromArgumentType(const FName& InA
 	}
 	else if (InArgumentName == ResultName)
 	{
-		const TArray<TRigVMTypeIndex>& ObjectTypes = FRigVMRegistry::Get().GetTypesForCategory(FRigVMTemplateArgument::ETypeCategory_SingleObjectValue);
+		const TArray<TRigVMTypeIndex>& ObjectTypes = FRigVMRegistry_NoLock::GetForRead().GetTypesForCategory_NoLock(FRigVMTemplateArgument::ETypeCategory_SingleObjectValue);
 		for (const TRigVMTypeIndex& Type : ObjectTypes)
 		{
 			OutPermutations.Add(
diff --git a/Engine/Plugins/Runtime/RigVM/Source/RigVM/Private/RigVMFunctions/RigVMDispatch_Core.cpp b/Engine/Plugins/Runtime/RigVM/Source/RigVM/Private/RigVMFunctions/RigVMDispatch_Core.cpp
--- a/Engine/Plugins/Runtime/RigVM/Source/RigVM/Private/RigVMFunctions/RigVMDispatch_Core.cpp
+++ b/Engine/Plugins/Runtime/RigVM/Source/RigVM/Private/RigVMFunctions/RigVMDispatch_Core.cpp
@@ -69,27 +69,27 @@ FRigVMFunctionPtr FRigVMDispatch_CoreEquals::GetDispatchFunctionImpl(const FRigV
 	{
 		return &FRigVMDispatch_CoreEquals::StringEquals;
 	}
-	if(TypeIndex == FRigVMRegistry::Get().GetTypeIndex<FVector>())
+	if(TypeIndex == FRigVMRegistry_NoLock::GetForRead().GetTypeIndex_NoLock<FVector>())
 	{
 		return &FRigVMDispatch_CoreEquals::MathTypeEquals<FVector>;
 	}
-	if(TypeIndex == FRigVMRegistry::Get().GetTypeIndex<FVector2D>())
+	if(TypeIndex == FRigVMRegistry_NoLock::GetForRead().GetTypeIndex_NoLock<FVector2D>())
 	{
 		return &FRigVMDispatch_CoreEquals::MathTypeEquals<FVector2D>;
 	}
-	if(TypeIndex == FRigVMRegistry::Get().GetTypeIndex<FRotator>())
+	if(TypeIndex == FRigVMRegistry_NoLock::GetForRead().GetTypeIndex_NoLock<FRotator>())
 	{
 		return &FRigVMDispatch_CoreEquals::MathTypeEquals<FRotator>;
 	}
-	if(TypeIndex == FRigVMRegistry::Get().GetTypeIndex<FQuat>())
+	if(TypeIndex == FRigVMRegistry_NoLock::GetForRead().GetTypeIndex_NoLock<FQuat>())
 	{
 		return &FRigVMDispatch_CoreEquals::MathTypeEquals<FQuat>;
 	}
-	if(TypeIndex == FRigVMRegistry::Get().GetTypeIndex<FTransform>())
+	if(TypeIndex == FRigVMRegistry_NoLock::GetForRead().GetTypeIndex_NoLock<FTransform>())
 	{
 		return &FRigVMDispatch_CoreEquals::MathTypeEquals<FTransform>;
 	}
-	if(TypeIndex == FRigVMRegistry::Get().GetTypeIndex<FLinearColor>())
+	if(TypeIndex == FRigVMRegistry_NoLock::GetForRead().GetTypeIndex_NoLock<FLinearColor>())
 	{
 		return &FRigVMDispatch_CoreEquals::MathTypeEquals<FLinearColor>;
 	}
diff --git a/Engine/Plugins/Runtime/RigVM/Source/RigVM/Private/RigVMFunctions/RigVMDispatch_Print.cpp b/Engine/Plugins/Runtime/RigVM/Source/RigVM/Private/RigVMFunctions/RigVMDispatch_Print.cpp
--- a/Engine/Plugins/Runtime/RigVM/Source/RigVM/Private/RigVMFunctions/RigVMDispatch_Print.cpp
+++ b/Engine/Plugins/Runtime/RigVM/Source/RigVM/Private/RigVMFunctions/RigVMDispatch_Print.cpp
@@ -34,7 +34,7 @@ const TArray<FRigVMTemplateArgumentInfo>& FRigVMDispatch_Print::GetArgumentInfos
 		Infos.Emplace(ValueName, ERigVMPinDirection::Input, ValueCategories);
 		Infos.Emplace(EnabledName, ERigVMPinDirection::Input, RigVMTypeUtils::TypeIndex::Bool);
 		Infos.Emplace(ScreenDurationName, ERigVMPinDirection::Input, RigVMTypeUtils::TypeIndex::Float);
-		Infos.Emplace(ScreenColorName, ERigVMPinDirection::Input, FRigVMRegistry::Get().GetTypeIndex<FLinearColor>());
+		Infos.Emplace(ScreenColorName, ERigVMPinDirection::Input, FRigVMRegistry_NoLock::GetForRead().GetTypeIndex_NoLock<FLinearColor>());
 	}
 	return Infos;
 }
@@ -55,7 +55,7 @@ FRigVMTemplateTypeMap FRigVMDispatch_Print::OnNewArgumentType(const FName& InArg
 	Types.Add(ValueName, InTypeIndex);
 	Types.Add(EnabledName, RigVMTypeUtils::TypeIndex::Bool);
 	Types.Add(ScreenDurationName, RigVMTypeUtils::TypeIndex::Float);
-	Types.Add(ScreenColorName, FRigVMRegistry::Get().GetTypeIndex<FLinearColor>());
+	Types.Add(ScreenColorName, FRigVMRegistry_NoLock::GetForRead().GetTypeIndex_NoLock<FLinearColor>());
 	return Types;
 }
 
diff --git a/Engine/Plugins/Runtime/RigVM/Source/RigVM/Private/RigVMFunctions/RigVMDispatch_Select.cpp b/Engine/Plugins/Runtime/RigVM/Source/RigVM/Private/RigVMFunctions/RigVMDispatch_Select.cpp
--- a/Engine/Plugins/Runtime/RigVM/Source/RigVM/Private/RigVMFunctions/RigVMDispatch_Select.cpp
+++ b/Engine/Plugins/Runtime/RigVM/Source/RigVM/Private/RigVMFunctions/RigVMDispatch_Select.cpp
@@ -48,11 +48,11 @@ FRigVMTemplateTypeMap FRigVMDispatch_SelectInt32::OnNewArgumentType(const FName&
 	if(InArgumentName == ValuesName)
 	{
 		Types.Add(ValuesName, InTypeIndex);
-		Types.Add(ResultName, FRigVMRegistry::Get().GetBaseTypeFromArrayTypeIndex(InTypeIndex));
+		Types.Add(ResultName, FRigVMRegistry_NoLock::GetForRead().GetBaseTypeFromArrayTypeIndex_NoLock(InTypeIndex));
 	}
 	else
 	{
-		Types.Add(ValuesName, FRigVMRegistry::Get().GetArrayTypeFromBaseTypeIndex(InTypeIndex));
+		Types.Add(ValuesName, FRigVMRegistry_NoLock::GetForRead().GetArrayTypeFromBaseTypeIndex_NoLock(InTypeIndex));
 		Types.Add(ResultName, InTypeIndex);
 	}
 	return Types;
diff --git a/Engine/Plugins/Runtime/RigVM/Source/RigVM/Private/RigVMHost.cpp b/Engine/Plugins/Runtime/RigVM/Source/RigVM/Private/RigVMHost.cpp
--- a/Engine/Plugins/Runtime/RigVM/Source/RigVM/Private/RigVMHost.cpp
+++ b/Engine/Plugins/Runtime/RigVM/Source/RigVM/Private/RigVMHost.cpp
@@ -148,7 +148,7 @@ void URigVMHost::PostLoad()
 {
 	Super::PostLoad();
 	
-	FRigVMRegistry::Get().RefreshEngineTypesIfRequired();
+	FRigVMRegistry_RWLock::Get().RefreshEngineTypesIfRequired();
 	
 	FRigVMExtendedExecuteContext& ExtendedExecuteContext = GetRigVMExtendedExecuteContext();
 
@@ -1185,7 +1185,7 @@ TArray<const UObject*> URigVMHost::GetUserDefinedDependencies(const TArray<const
 	const TArray<const FRigVMFunction*>& Functions = VM->GetFunctions();
 	for (const FRigVMFunction* Function : Functions)
 	{
-		const FRigVMRegistry& Registry = FRigVMRegistry::Get();
+		const FRigVMRegistry_RWLock& Registry = FRigVMRegistry_RWLock::Get();
 		const TArray<TRigVMTypeIndex>& TypeIndices = Function->GetArgumentTypeIndices();
 		for (const TRigVMTypeIndex& TypeIndex : TypeIndices)
 		{
diff --git a/Engine/Plugins/Runtime/RigVM/Source/RigVM/Public/RigVMCore/RigVMDispatchFactory.h b/Engine/Plugins/Runtime/RigVM/Source/RigVM/Public/RigVMCore/RigVMDispatchFactory.h
--- a/Engine/Plugins/Runtime/RigVM/Source/RigVM/Public/RigVMCore/RigVMDispatchFactory.h
+++ b/Engine/Plugins/Runtime/RigVM/Source/RigVM/Public/RigVMCore/RigVMDispatchFactory.h
@@ -21,6 +21,7 @@
 
 #include "RigVMDispatchFactory.generated.h"
 
+struct FRigVMRegistry_RWLock;
 struct FRigVMDispatchFactory;
 class URigVMPin;
 
@@ -137,14 +138,16 @@ public:
 	// Returns truf if this factory supports a given executecontext struct
 	RIGVM_API bool SupportsExecuteContextStruct(const UScriptStruct* InExecuteContextStruct) const;
 
-	// registered needed types during registration of the factory
-	virtual void RegisterDependencyTypes() const {}
+	// registered needed types during registration of the factory.
+	// this is called within a NoLock code path - so make sure to access
+	// the FRigVMRegistry_NoLock within as needed.
+	virtual void RegisterDependencyTypes_NoLock() const {}
 
 	// returns the arguments of the template
 	RIGVM_API virtual const TArray<FRigVMTemplateArgumentInfo>& GetArgumentInfos() const;
 
 	// returns the execute arguments of the template
-	RIGVM_API TArray<FRigVMExecuteArgument> GetExecuteArguments(const FRigVMDispatchContext& InContext) const;
+	RIGVM_API TArray<FRigVMExecuteArgument> GetExecuteArguments_NoLock(const FRigVMDispatchContext& InContext) const;
 
 	// this function is deprecated, please use GetPermutationsFromArgumentType
 	// returns the new permutation argument types after a new type is defined for one argument
@@ -170,31 +173,41 @@ public:
 	RIGVM_API FRigVMFunctionPtr GetOrCreateDispatchFunction(const FRigVMTemplateTypeMap& InTypes) const;
 
 	// builds and returns the template
-	RIGVM_API const FRigVMTemplate* GetTemplate() const;
+	RIGVM_API const FRigVMTemplate* GetTemplate() const
+	{
+		return GetTemplate_RWLock();
+	}
 
 	// returns the name of the factory template
 	RIGVM_API FName GetTemplateNotation() const;
 
 	// returns the name of the permutation for a given set of types
-	RIGVM_API FString GetPermutationName(const FRigVMTemplateTypeMap& InTypes) const;
+	RIGVM_API FString GetPermutationName(const FRigVMTemplateTypeMap& InTypes, const bool bLockRegistry = true) const;
 
 	// returns true if the dispatch uses the same function ptr for all permutations
 	virtual bool IsSingleton() const { return false; } 
 
 protected:
 
+	// returns the dispatch function for a given type set
+	RIGVM_API FRigVMFunctionPtr GetOrCreateDispatchFunction_NoLock(const FRigVMTemplateTypeMap& InTypes) const;
+
+	// builds and returns the template
+	RIGVM_API const FRigVMTemplate* GetTemplate_RWLock() const;
+
+	// builds and returns the template
+	RIGVM_API const FRigVMTemplate* GetTemplate_NoLock() const;
+
 	// for each type defined in the primary argument, this function will call GetPermutationsFromArgumentType to construct an array of arguments with the appropiate permutations
 	RIGVM_API TArray<FRigVMTemplateArgumentInfo> BuildArgumentListFromPrimaryArgument(const TArray<FRigVMTemplateArgumentInfo>& InInfos, const FName& InPrimaryArgumentName) const;
 
 	// returns the name of the permutation for a given set of types
-	RIGVM_API FString GetPermutationNameImpl(const FRigVMTemplateTypeMap& InTypes) const;
+	RIGVM_API FString GetPermutationNameImpl(const FRigVMTemplateTypeMap& InTypes, const bool bLockRegistry = true) const;
 
-	RIGVM_API FRigVMFunctionPtr CreateDispatchFunction(const FRigVMTemplateTypeMap& InTypes) const;
 	FRigVMFunctionPtr CreateDispatchFunction_NoLock(const FRigVMTemplateTypeMap& InTypes) const;
 
 	virtual FRigVMFunctionPtr GetDispatchFunctionImpl(const FRigVMTemplateTypeMap& InTypes) const { return nullptr; }
 	
-	RIGVM_API TArray<FRigVMFunction> CreateDispatchPredicates(const FRigVMTemplateTypeMap& InTypes) const;
 	TArray<FRigVMFunction> CreateDispatchPredicates_NoLock(const FRigVMTemplateTypeMap& InTypes) const;
 	
 	virtual TArray<FRigVMFunction> GetDispatchPredicatesImpl(const FRigVMTemplateTypeMap& InTypes) const { return TArray<FRigVMFunction>(); }
@@ -260,7 +273,7 @@ protected:
 	mutable TMap<int32, TSharedPtr<TArray<FName>>> ArgumentNamesMap;
 	FCriticalSection* ArgumentNamesMutex;
 	friend struct FRigVMTemplate;
-	friend struct FRigVMRegistry;
+	friend struct FRigVMRegistry_NoLock;
 	friend struct FRigVMFunction;
 	friend class URigVM;
 };
diff --git a/Engine/Plugins/Runtime/RigVM/Source/RigVM/Public/RigVMCore/RigVMFunction.h b/Engine/Plugins/Runtime/RigVM/Source/RigVM/Public/RigVMCore/RigVMFunction.h
--- a/Engine/Plugins/Runtime/RigVM/Source/RigVM/Public/RigVMCore/RigVMFunction.h
+++ b/Engine/Plugins/Runtime/RigVM/Source/RigVM/Public/RigVMCore/RigVMFunction.h
@@ -171,7 +171,8 @@ struct RIGVM_API FRigVMFunction
 	const TArray<FRigVMFunctionArgument>& GetArguments() const { return Arguments; }
 	const TArray<TRigVMTypeIndex>& GetArgumentTypeIndices() const;
 	const FRigVMTemplate* GetTemplate() const;
-	const UScriptStruct* GetExecuteContextStruct() const;
+	const FRigVMTemplate* GetTemplate_NoLock() const;
+	const UScriptStruct* GetExecuteContextStruct(bool bLockRegistry = true) const;
 	bool SupportsExecuteContextStruct(const UScriptStruct* InExecuteContextStruct) const;
 	const FName& GetArgumentNameForOperandIndex(int32 InOperandIndex, int32 InTotalOperands) const;
 };
diff --git a/Engine/Plugins/Runtime/RigVM/Source/RigVM/Public/RigVMCore/RigVMRegistry.h b/Engine/Plugins/Runtime/RigVM/Source/RigVM/Public/RigVMCore/RigVMRegistry.h
--- a/Engine/Plugins/Runtime/RigVM/Source/RigVM/Public/RigVMCore/RigVMRegistry.h
+++ b/Engine/Plugins/Runtime/RigVM/Source/RigVM/Public/RigVMCore/RigVMRegistry.h
@@ -26,6 +26,9 @@ class IPlugin;
 class UObject;
 struct FRigVMDispatchFactory;
 
+struct FRigVMRegistry_RWLock;
+typedef FRigVMRegistry_RWLock FRigVMRegistry;
+
 /**
  * The FRigVMRegistry is used to manage all known function pointers
  * for use in the RigVM. The Register method is called automatically
@@ -34,17 +37,26 @@ struct FRigVMDispatchFactory;
  * 
  * Inheriting from FGCObject to ensure that all type objects cannot be GCed
  */
-struct RIGVM_API FRigVMRegistry : public FGCObject
+struct RIGVM_API FRigVMRegistry_NoLock : public FGCObject
 {
 public:
 
-	DECLARE_MULTICAST_DELEGATE(FOnRigVMRegistryChanged);
+	enum ELockType
+	{
+		LockType_Read,
+		LockType_Write,
+		LockType_Invalid
+	};
 
-	virtual ~FRigVMRegistry() override;
-	
 	// Returns the singleton registry
-	static FRigVMRegistry& Get();
+	static FRigVMRegistry_NoLock& Get(ELockType InLockType);
+	static const FRigVMRegistry_NoLock& GetForRead() { return Get(LockType_Read); }
+	static FRigVMRegistry_NoLock& GetForWrite() { return Get(LockType_Write); }
+
+	DECLARE_MULTICAST_DELEGATE(FOnRigVMRegistryChanged);
 
+	virtual ~FRigVMRegistry_NoLock() override;
+	
 	// FGCObject overrides
 	virtual void AddReferencedObjects(FReferenceCollector& Collector) override;
 	virtual FString GetReferencerName() const override;
@@ -52,13 +64,13 @@ public:
 	// Registers a function given its name.
 	// The name will be the name of the struct and virtual method,
 	// for example "FMyStruct::MyVirtualMethod"
-	void Register(const TCHAR* InName, FRigVMFunctionPtr InFunctionPtr, UScriptStruct* InStruct = nullptr, const TArray<FRigVMFunctionArgument>& InArguments = TArray<FRigVMFunctionArgument>());
+	virtual void Register_NoLock(const TCHAR* InName, FRigVMFunctionPtr InFunctionPtr, UScriptStruct* InStruct = nullptr, const TArray<FRigVMFunctionArgument>& InArguments = TArray<FRigVMFunctionArgument>());
 
 	// Registers a dispatch factory given its struct.
-	const FRigVMDispatchFactory* RegisterFactory(UScriptStruct* InFactoryStruct);
+	virtual const FRigVMDispatchFactory* RegisterFactory_NoLock(UScriptStruct* InFactoryStruct);
 
 	// Register a predicate contained in the input struct
-	void RegisterPredicate(UScriptStruct* InStruct, const TCHAR* InName, const TArray<FRigVMFunctionArgument>& InArguments);
+	virtual void RegisterPredicate_NoLock(UScriptStruct* InStruct, const TCHAR* InName, const TArray<FRigVMFunctionArgument>& InArguments);
 
 	// How to register an object's class when passed to RegisterObjectTypes
 	enum class ERegisterObjectOperation
@@ -71,50 +83,50 @@ public:
 	};
 
 	// Register a set of allowed object types
-	void RegisterObjectTypes(TConstArrayView<TPair<UClass*, ERegisterObjectOperation>> InClasses);
+	virtual void RegisterObjectTypes_NoLock(TConstArrayView<TPair<UClass*, ERegisterObjectOperation>> InClasses);
 
 	// Register a set of allowed struct types
-	void RegisterStructTypes(TConstArrayView<UScriptStruct*> InStructs);
+	virtual void RegisterStructTypes_NoLock(TConstArrayView<UScriptStruct*> InStructs);
 
 	// Refreshes the list and finds the function pointers
 	// based on the names.
-	void RefreshEngineTypes();
+	virtual void RefreshEngineTypes_NoLock();
 
 	// Refreshes the list and finds the function pointers
 	// based on the names.
-	void RefreshEngineTypesIfRequired();
-
+	virtual void RefreshEngineTypesIfRequired_NoLock();
+	
 	// Update the registry when types are renamed
-	void OnAssetRenamed(const FAssetData& InAssetData, const FString& InOldObjectPath);
+	virtual void OnAssetRenamed_NoLock(const FAssetData& InAssetData, const FString& InOldObjectPath);
 	
 	// Update the registry when old types are removed
-    void OnAssetRemoved(const FAssetData& InAssetData);
+	virtual void OnAssetRemoved_NoLock(const FAssetData& InAssetData);
 
 	// Removes all types associated with a plugin that's being unloaded. 
-	void OnPluginUnloaded(IPlugin& InPlugin);
+	virtual void OnPluginUnloaded_NoLock(IPlugin& InPlugin);
 	
 	// Update the registry when new types are added to the attribute system so that they can be selected
 	// on Attribute Nodes
-	void OnAnimationAttributeTypesChanged(const UScriptStruct* InStruct, bool bIsAdded);
-	
+	virtual void OnAnimationAttributeTypesChanged_NoLock(const UScriptStruct* InStruct, bool bIsAdded);
+
 	// Clear the registry
-	void Reset();
+	virtual void Reset_NoLock();
 
 	// Adds a type if it doesn't exist yet and returns its index.
 	// This function is thread-safe
-	TRigVMTypeIndex FindOrAddType(const FRigVMTemplateArgumentType& InType, bool bForce = false);
+	virtual TRigVMTypeIndex FindOrAddType_NoLock(const FRigVMTemplateArgumentType& InType, bool bForce = false);
 
 	// Removes a type from the registry, and updates all dependent templates
 	// which also creates invalid permutations in templates that we should ignore
-	bool RemoveType(const FSoftObjectPath& InObjectPath, const UClass* InObjectClass);
+	virtual bool RemoveType_NoLock(const FSoftObjectPath& InObjectPath, const UClass* InObjectClass);
 
 	// Returns the type index given a type
-	TRigVMTypeIndex GetTypeIndex(const FRigVMTemplateArgumentType& InType) const;
+	virtual TRigVMTypeIndex GetTypeIndex_NoLock(const FRigVMTemplateArgumentType& InType) const;
 
 	// Returns the type index given a cpp type and a type object
-	TRigVMTypeIndex GetTypeIndex(const FName& InCPPType, UObject* InCPPTypeObject) const
+	virtual TRigVMTypeIndex GetTypeIndex_NoLock(const FName& InCPPType, UObject* InCPPTypeObject) const
 	{
-		return GetTypeIndex(FRigVMTemplateArgumentType(InCPPType, InCPPTypeObject));
+		return GetTypeIndex_NoLock(FRigVMTemplateArgumentType(InCPPType, InCPPTypeObject));
 	}
 
 	// Returns the type index given an enum
@@ -122,14 +134,14 @@ public:
 		typename T,
 		typename TEnableIf<TIsEnum<T>::Value>::Type* = nullptr
 	>
-	TRigVMTypeIndex GetTypeIndex(bool bAsArray = false) const
+	TRigVMTypeIndex GetTypeIndex_NoLock(bool bAsArray = false) const
 	{
 		FRigVMTemplateArgumentType Type(StaticEnum<T>());
 		if(bAsArray)
 		{
 			Type.ConvertToArray();
 		}
-		return GetTypeIndex(Type);
+		return GetTypeIndex_NoLock(Type);
 	}
 
 	// Returns the type index given a struct
@@ -137,14 +149,14 @@ public:
 		typename T,
 		typename TEnableIf<TRigVMIsBaseStructure<T>::Value, T>::Type* = nullptr
 	>
-	TRigVMTypeIndex GetTypeIndex(bool bAsArray = false) const
+	TRigVMTypeIndex GetTypeIndex_NoLock(bool bAsArray = false) const
 	{
 		FRigVMTemplateArgumentType Type(TBaseStructure<T>::Get());
 		if(bAsArray)
 		{
 			Type.ConvertToArray();
 		}
-		return GetTypeIndex(Type);
+		return GetTypeIndex_NoLock(Type);
 	}
 
 	// Returns the type index given a struct
@@ -152,14 +164,14 @@ public:
 		typename T,
 		typename TEnableIf<TModels_V<CRigVMUStruct, T>>::Type * = nullptr
 	>
-	TRigVMTypeIndex GetTypeIndex(bool bAsArray = false) const
+	TRigVMTypeIndex GetTypeIndex_NoLock(bool bAsArray = false) const
 	{
 		FRigVMTemplateArgumentType Type(T::StaticStruct());
 		if(bAsArray)
 		{
 			Type.ConvertToArray();
 		}
-		return GetTypeIndex(Type);
+		return GetTypeIndex_NoLock(Type);
 	}
 
 	// Returns the type index given an object
@@ -167,119 +179,119 @@ public:
 		typename T,
 		typename TEnableIf<TModels_V<CRigVMUClass, T>>::Type * = nullptr
 	>
-	TRigVMTypeIndex GetTypeIndex(bool bAsArray = false) const
+	TRigVMTypeIndex GetTypeIndex_NoLock(bool bAsArray = false) const
 	{
 		FRigVMTemplateArgumentType Type(T::StaticClass(), RigVMTypeUtils::EClassArgType::AsObject);
 		if(bAsArray)
 		{
 			Type.ConvertToArray();
 		}
-		return GetTypeIndex(Type);
+		return GetTypeIndex_NoLock(Type);
 	}
 
 	// Returns the type given its index
-	const FRigVMTemplateArgumentType& GetType(TRigVMTypeIndex InTypeIndex) const;
+	virtual const FRigVMTemplateArgumentType& GetType_NoLock(TRigVMTypeIndex InTypeIndex) const;
 
 	// Returns the number of types
-	int32 NumTypes() const { return Types.Num(); }
+	virtual int32 NumTypes_NoLock() const { return Types.Num(); }
 
 	// Returns the type given only its cpp type
-	const FRigVMTemplateArgumentType& FindTypeFromCPPType(const FString& InCPPType) const;
+	virtual const FRigVMTemplateArgumentType& FindTypeFromCPPType_NoLock(const FString& InCPPType) const;
 
 	// Returns the type index given only its cpp type
-	TRigVMTypeIndex GetTypeIndexFromCPPType(const FString& InCPPType) const;
+	virtual TRigVMTypeIndex GetTypeIndexFromCPPType_NoLock(const FString& InCPPType) const;
 
 	// Returns true if the type is an array
-	bool IsArrayType(TRigVMTypeIndex InTypeIndex) const;
+	virtual bool IsArrayType_NoLock(TRigVMTypeIndex InTypeIndex) const;
 
 	// Returns true if the type is an execute type
-	bool IsExecuteType(TRigVMTypeIndex InTypeIndex) const;
+	virtual bool IsExecuteType_NoLock(TRigVMTypeIndex InTypeIndex) const;
 
 	// Converts the given execute context type to the base execute context type
-	bool ConvertExecuteContextToBaseType(TRigVMTypeIndex& InOutTypeIndex) const;
+	virtual bool ConvertExecuteContextToBaseType_NoLock(TRigVMTypeIndex& InOutTypeIndex) const;
 
 	// Returns the dimensions of the array 
-	int32 GetArrayDimensionsForType(TRigVMTypeIndex InTypeIndex) const;
+	virtual int32 GetArrayDimensionsForType_NoLock(TRigVMTypeIndex InTypeIndex) const;
 
 	// Returns true if the type is a wildcard type
-	bool IsWildCardType(TRigVMTypeIndex InTypeIndex) const;
+	virtual bool IsWildCardType_NoLock(TRigVMTypeIndex InTypeIndex) const;
 
 	// Returns true if the types can be matched.
-	bool CanMatchTypes(TRigVMTypeIndex InTypeIndexA, TRigVMTypeIndex InTypeIndexB, bool bAllowFloatingPointCasts) const;
+	virtual bool CanMatchTypes_NoLock(TRigVMTypeIndex InTypeIndexA, TRigVMTypeIndex InTypeIndexB, bool bAllowFloatingPointCasts) const;
 
 	// Returns the list of compatible types for a given type
-	const TArray<TRigVMTypeIndex>& GetCompatibleTypes(TRigVMTypeIndex InTypeIndex) const;
+	virtual const TArray<TRigVMTypeIndex>& GetCompatibleTypes_NoLock(TRigVMTypeIndex InTypeIndex) const;
 
 	// Returns all compatible types given a category
-	const TArray<TRigVMTypeIndex>& GetTypesForCategory(FRigVMTemplateArgument::ETypeCategory InCategory) const;
+	virtual const TArray<TRigVMTypeIndex>& GetTypesForCategory_NoLock(FRigVMTemplateArgument::ETypeCategory InCategory) const;
 
 	// Returns the type index of the array matching the given element type index
-	TRigVMTypeIndex GetArrayTypeFromBaseTypeIndex(TRigVMTypeIndex InTypeIndex) const;
+	virtual TRigVMTypeIndex GetArrayTypeFromBaseTypeIndex_NoLock(TRigVMTypeIndex InTypeIndex) const;
 
 	// Returns the type index of the element matching the given array type index
-	TRigVMTypeIndex GetBaseTypeFromArrayTypeIndex(TRigVMTypeIndex InTypeIndex) const;
+	virtual TRigVMTypeIndex GetBaseTypeFromArrayTypeIndex_NoLock(TRigVMTypeIndex InTypeIndex) const;
 
 	// Returns the function given its name (or nullptr)
-	const FRigVMFunction* FindFunction(const TCHAR* InName, const FRigVMUserDefinedTypeResolver& InTypeResolver = FRigVMUserDefinedTypeResolver()) const;
+	virtual const FRigVMFunction* FindFunction_NoLock(const TCHAR* InName, const FRigVMUserDefinedTypeResolver& InTypeResolver = FRigVMUserDefinedTypeResolver()) const;
 
 	// Returns the function given its backing up struct and method name
-	const FRigVMFunction* FindFunction(UScriptStruct* InStruct, const TCHAR* InName, const FRigVMUserDefinedTypeResolver& InResolvalInfo = FRigVMUserDefinedTypeResolver()) const;
+	virtual const FRigVMFunction* FindFunction_NoLock(UScriptStruct* InStruct, const TCHAR* InName, const FRigVMUserDefinedTypeResolver& InResolvalInfo = FRigVMUserDefinedTypeResolver()) const;
 
-	// Returns all current rigvm functions
-	const TChunkedArray<FRigVMFunction>& GetFunctions() const;
+	// Returns all current RigVM functions
+	virtual const TChunkedArray<FRigVMFunction>& GetFunctions_NoLock() const;
 
 	// Returns a template pointer given its notation (or nullptr)
-	const FRigVMTemplate* FindTemplate(const FName& InNotation, bool bIncludeDeprecated = false) const;
+	virtual const FRigVMTemplate* FindTemplate_NoLock(const FName& InNotation, bool bIncludeDeprecated = false) const;
 
-	// Returns all current rigvm functions
-	const TChunkedArray<FRigVMTemplate>& GetTemplates() const;
+	// Returns all current RigVM functions
+	virtual const TChunkedArray<FRigVMTemplate>& GetTemplates_NoLock() const;
 
 	// Defines and retrieves a template given its arguments
-	const FRigVMTemplate* GetOrAddTemplateFromArguments(
+	virtual const FRigVMTemplate* GetOrAddTemplateFromArguments_NoLock(
 		const FName& InName,
 		const TArray<FRigVMTemplateArgumentInfo>& InInfos,
 		const FRigVMTemplateDelegates& InDelegates);
 
 	// Adds a new template given its arguments
-	const FRigVMTemplate* AddTemplateFromArguments(
+	virtual const FRigVMTemplate* AddTemplateFromArguments_NoLock(
 		const FName& InName,
 		const TArray<FRigVMTemplateArgumentInfo>& InInfos,
 		const FRigVMTemplateDelegates& InDelegates);
 
 	// Returns a dispatch factory given its name (or nullptr)
-	FRigVMDispatchFactory* FindDispatchFactory(const FName& InFactoryName) const;
+	virtual FRigVMDispatchFactory* FindDispatchFactory_NoLock(const FName& InFactoryName) const;
 
 	// Returns a dispatch factory given its static struct (or nullptr)
-	FRigVMDispatchFactory* FindOrAddDispatchFactory(UScriptStruct* InFactoryStruct);
+	virtual FRigVMDispatchFactory* FindOrAddDispatchFactory_NoLock(UScriptStruct* InFactoryStruct);
 
 	// Returns a dispatch factory given its static struct (or nullptr)
 	template <
 		typename T,
 		typename TEnableIf<TModels_V<CRigVMUStruct, T>>::Type * = nullptr
 	>
-	FRigVMDispatchFactory* FindOrAddDispatchFactory()
+	FRigVMDispatchFactory* FindOrAddDispatchFactory_NoLock()
 	{
-		return FindOrAddDispatchFactory(T::StaticStruct());
+		return FindOrAddDispatchFactory_NoLock(T::StaticStruct());
 	}
 
 	// Returns a dispatch factory's singleton function name if that exists
-	FString FindOrAddSingletonDispatchFunction(UScriptStruct* InFactoryStruct);
+	virtual FString FindOrAddSingletonDispatchFunction_NoLock(UScriptStruct* InFactoryStruct);
 
 	// Returns a dispatch factory's singleton function name if that exists
 	template <
 		typename T,
 		typename TEnableIf<TModels_V<CRigVMUStruct, T>>::Type * = nullptr
 	>
-	FString FindOrAddSingletonDispatchFunction()
+	FString FindOrAddSingletonDispatchFunction_NoLock()
 	{
-		return FindOrAddSingletonDispatchFunction(T::StaticStruct());
+		return FindOrAddSingletonDispatchFunction_NoLock(T::StaticStruct());
 	}
 
 	// Returns all dispatch factories
-	const TArray<FRigVMDispatchFactory*>& GetFactories() const;
+	virtual const TArray<FRigVMDispatchFactory*>& GetFactories_NoLock() const;
 
 	// Given a struct name, return the predicates
-	const TArray<FRigVMFunction>* GetPredicatesForStruct(const FName& InStructName) const;
+	virtual const TArray<FRigVMFunction>* GetPredicatesForStruct_NoLock(const FName& InStructName) const;
 
 	static const TArray<UScriptStruct*>& GetMathTypes();
 
@@ -287,22 +299,25 @@ public:
 	FOnRigVMRegistryChanged& OnRigVMRegistryChanged() { return OnRigVMRegistryChangedDelegate; }
 
 	// Returns a unique hash per type index
-	uint32 GetHashForType(TRigVMTypeIndex InTypeIndex) const;
-	uint32 GetHashForScriptStruct(const UScriptStruct* InScriptStruct, bool bCheckTypeIndex = true) const;
-	uint32 GetHashForStruct(const UStruct* InStruct) const;
-	uint32 GetHashForEnum(const UEnum* InEnum, bool bCheckTypeIndex = true) const;
-	uint32 GetHashForProperty(const FProperty* InProperty) const;
+	virtual uint32 GetHashForType_NoLock(TRigVMTypeIndex InTypeIndex) const;
+	virtual uint32 GetHashForScriptStruct_NoLock(const UScriptStruct* InScriptStruct, bool bCheckTypeIndex = true) const;
+	virtual uint32 GetHashForStruct_NoLock(const UStruct* InStruct) const;
+	virtual uint32 GetHashForEnum_NoLock(const UEnum* InEnum, bool bCheckTypeIndex = true) const;
+	virtual uint32 GetHashForProperty_NoLock(const FProperty* InProperty) const;
+
+	virtual void RebuildRegistry_NoLock(); 
 
 	static inline const FLazyName TemplateNameMetaName = FLazyName(TEXT("TemplateName"));
 
-private:
+protected:
 
-	FRigVMRegistry();
+
+	FRigVMRegistry_NoLock();
 
 	// disable copy constructor
-	FRigVMRegistry(const FRigVMRegistry&) = delete;
+	FRigVMRegistry_NoLock(const FRigVMRegistry_NoLock&) = delete;
 	// disable assignment operator
-	FRigVMRegistry& operator= (const FRigVMRegistry &InOther) = delete;
+	FRigVMRegistry_NoLock& operator= (const FRigVMRegistry_NoLock &InOther) = delete;
 
 	struct FTypeInfo
 	{
@@ -324,11 +339,8 @@ private:
 	};
 
 	// Initialize the base types
-	void Initialize();
-	
-	TRigVMTypeIndex FindOrAddType_NoLock(const FRigVMTemplateArgumentType& InType, bool bForce);
-
-	void RefreshEngineTypes_NoLock();
+	void Initialize_NoLock();
+	virtual void Initialize(bool bLockRegistry) = 0;
 
 	static EObjectFlags DisallowedFlags()
 	{
@@ -340,24 +352,14 @@ private:
 		return RF_Public;
 	}
 
-	bool IsAllowedType(const FProperty* InProperty) const;
-	bool IsAllowedType(const UEnum* InEnum) const;
-	bool IsAllowedType(const UStruct* InStruct) const;
-	bool IsAllowedType(const UClass* InClass) const;
-
-	void RegisterTypeInCategory(const FRigVMTemplateArgument::ETypeCategory InCategory, const TRigVMTypeIndex InTypeIndex);
-	void PropagateTypeAddedToCategory(const FRigVMTemplateArgument::ETypeCategory InCategory, const TRigVMTypeIndex InTypeIndex);
-	
-	void RemoveTypeInCategory(FRigVMTemplateArgument::ETypeCategory InCategory, TRigVMTypeIndex InTypeIndex);
+	bool IsAllowedType_NoLock(const FProperty* InProperty) const;
+	bool IsAllowedType_NoLock(const UEnum* InEnum) const;
+	bool IsAllowedType_NoLock(const UStruct* InStruct) const;
+	bool IsAllowedType_NoLock(const UClass* InClass) const;
 
-	const FRigVMFunction* FindFunction_NoLock(const TCHAR* InName, const FRigVMUserDefinedTypeResolver& InTypeResolver = FRigVMUserDefinedTypeResolver()) const;
-	const FRigVMTemplate* FindTemplate_NoLock(const FName& InNotation, bool bIncludeDeprecated) const;
-	FRigVMDispatchFactory* FindDispatchFactory_NoLock(const FName& InFactoryName) const;
-
-	const FRigVMTemplate* AddTemplateFromArguments_NoLock(
-		const FName& InName,
-		const TArray<FRigVMTemplateArgumentInfo>& InInfos,
-		const FRigVMTemplateDelegates& InDelegates);
+	void RegisterTypeInCategory_NoLock(const FRigVMTemplateArgument::ETypeCategory InCategory, const TRigVMTypeIndex InTypeIndex);
+	void PropagateTypeAddedToCategory_NoLock(const FRigVMTemplateArgument::ETypeCategory InCategory, const TRigVMTypeIndex InTypeIndex);
+	void RemoveTypeInCategory_NoLock(FRigVMTemplateArgument::ETypeCategory InCategory, TRigVMTypeIndex InTypeIndex);
 
 	// memory for all (known) types
 	TArray<FTypeInfo> Types;
@@ -413,16 +415,852 @@ private:
 	// This is true if the engine has ever refreshed the engine types
 	bool bEverRefreshedEngineTypes;
 
-	static FCriticalSection FindOrAddTypeMutex;
-	static FCriticalSection FunctionRegistryMutex;
-	static FCriticalSection FactoryRegistryMutex;
-	static FCriticalSection TemplateRegistryMutex;
+	friend struct FRigVMStruct;
+	friend struct FRigVMTemplate;
+	friend struct FRigVMTemplateArgument;
+	friend struct FRigVMDispatchFactory;
+};
+
+/**
+ * The FRigVMRegistry is used to manage all known function pointers
+ * for use in the RigVM. The Register method is called automatically
+ * when the static struct is initially constructed for each USTRUCT
+ * hosting a RIGVM_METHOD enabled virtual function.
+ * 
+ * Inheriting from FGCObject to ensure that all type objects cannot be GCed
+ */
+struct RIGVM_API FRigVMRegistry_RWLock final : public FRigVMRegistry_NoLock
+{
+	typedef FRigVMRegistry_NoLock Super;
+	
+protected:
+
+	class RIGVM_API FConditionalScopeLock
+	{
+	public:
+		UE_NODISCARD_CTOR explicit FConditionalScopeLock(const FRigVMRegistry_RWLock& InRegistry, ELockType InLockType, bool bInLockEnabled = true)
+		: Registry(InRegistry)
+		, LockType(InLockType)
+		, bLockEnabled(bInLockEnabled)
+		, bResetLockType(false)
+		{
+			if(bLockEnabled)
+			{
+				if(LockType == LockType_Read)
+				{
+					Registry.Lock.ReadLock();
+				}
+				else if(LockType == LockType_Write)
+				{
+					Registry.Lock.WriteLock();
+				}
+
+				ELockType InvalidLockType = LockType_Invalid;
+				if(Registry.LockType.compare_exchange_strong(InvalidLockType, InLockType))
+				{
+					bResetLockType = true;
+				}
+			}
+		}
+
+		~FConditionalScopeLock()
+		{
+			if(bLockEnabled)
+			{
+				if(bResetLockType)
+				{
+					Registry.LockType.store(LockType_Invalid);
+				}
+				if(LockType == LockType_Read)
+				{
+					Registry.Lock.ReadUnlock();
+				}
+				else if(LockType == LockType_Write)
+				{
+					Registry.Lock.WriteUnlock();
+				}
+			}
+		}
+
+	private:
+		const FRigVMRegistry_RWLock& Registry;
+		ELockType LockType;
+		const bool bLockEnabled;
+		bool bResetLockType;
+
+		UE_NONCOPYABLE(FConditionalScopeLock);
+	};
+
+public:
+
+	class RIGVM_API FConditionalReadScopeLock : public FConditionalScopeLock 
+	{
+	public:
+		UE_NODISCARD_CTOR explicit FConditionalReadScopeLock(const FRigVMRegistry_RWLock& InRegistry, bool bInLockEnabled = true)
+		: FConditionalScopeLock(InRegistry, LockType_Read, bInLockEnabled)
+		{
+		}
+
+		UE_NODISCARD_CTOR explicit FConditionalReadScopeLock(bool bInLockEnabled = true)
+		: FConditionalReadScopeLock(FRigVMRegistry_RWLock::Get(), bInLockEnabled)
+		{
+		}
+	};
+
+	class RIGVM_API FConditionalWriteScopeLock : public FConditionalScopeLock
+	{
+	public:
+		UE_NODISCARD_CTOR explicit FConditionalWriteScopeLock(const FRigVMRegistry_RWLock& InRegistry, bool bInLockEnabled = true)
+		: FConditionalScopeLock(InRegistry, LockType_Write, bInLockEnabled)
+		{
+		}
+
+		UE_NODISCARD_CTOR explicit FConditionalWriteScopeLock(bool bInLockEnabled = true)
+		: FConditionalWriteScopeLock(FRigVMRegistry_RWLock::Get(), bInLockEnabled)
+		{
+		}
+	};
+
+	// Returns the singleton registry
+	FORCEINLINE static FRigVMRegistry_RWLock& Get()
+	{
+		// static in a function scope ensures that the GC system is initiated before 
+		// the registry constructor is called
+		static FRigVMRegistry_RWLock s_RigVMRegistry;
+		return s_RigVMRegistry;
+	}
+
+	static FRigVMRegistry_NoLock& Get(ELockType InLockType) = delete;
+	static const FRigVMRegistry_NoLock& GetForRead() = delete;
+	static FRigVMRegistry_NoLock& GetForWrite() = delete;
+
+	// Registers a function given its name.
+	// The name will be the name of the struct and virtual method,
+	// for example "FMyStruct::MyVirtualMethod"
+	void Register(const TCHAR* InName, FRigVMFunctionPtr InFunctionPtr, UScriptStruct* InStruct = nullptr, const TArray<FRigVMFunctionArgument>& InArguments = TArray<FRigVMFunctionArgument>(), bool bLockRegistry = true)
+	{
+		FConditionalWriteScopeLock _(*this, bLockRegistry);
+		Super::Register_NoLock(InName, InFunctionPtr, InStruct, InArguments);
+	}
+
+	// Registers a dispatch factory given its struct.
+	const FRigVMDispatchFactory* RegisterFactory(UScriptStruct* InFactoryStruct, bool bLockRegistry = true)
+	{
+		FConditionalWriteScopeLock _(*this, bLockRegistry);
+		return Super::RegisterFactory_NoLock(InFactoryStruct);
+	}
+
+	// Register a predicate contained in the input struct
+	void RegisterPredicate(UScriptStruct* InStruct, const TCHAR* InName, const TArray<FRigVMFunctionArgument>& InArguments, bool bLockRegistry = true)
+	{
+		FConditionalWriteScopeLock _(*this, bLockRegistry);
+		return Super::RegisterPredicate_NoLock(InStruct, InName, InArguments);
+	}
+
+	// Register a set of allowed object types
+	void RegisterObjectTypes(TConstArrayView<TPair<UClass*, ERegisterObjectOperation>> InClasses, bool bLockRegistry = true)
+	{
+		FConditionalWriteScopeLock _(*this, bLockRegistry);
+		Super::RegisterObjectTypes_NoLock(InClasses);
+	}
+
+	// Refreshes the list and finds the function pointers
+	// based on the names.
+	void RefreshEngineTypes(bool bLockRegistry = true)
+	{
+		FConditionalWriteScopeLock _(*this, bLockRegistry);
+		Super::RefreshEngineTypes_NoLock();
+	}
+
+	// Refreshes the list and finds the function pointers
+	// based on the names.
+	void RefreshEngineTypesIfRequired(bool bLockRegistry = true)
+	{
+		FConditionalWriteScopeLock _(*this, bLockRegistry);
+		Super::RefreshEngineTypesIfRequired_NoLock();
+	}
+
+	// Update the registry when types are renamed
+	void OnAssetRenamed(const FAssetData& InAssetData, const FString& InOldObjectPath)
+	{
+		FWriteScopeLock _(Lock);
+		Super::OnAssetRenamed_NoLock(InAssetData, InOldObjectPath);
+	}
+	
+	// Update the registry when old types are removed
+    void OnAssetRemoved(const FAssetData& InAssetData)
+	{
+		FWriteScopeLock _(Lock);
+		Super::OnAssetRemoved_NoLock(InAssetData);
+	}
+
+	// Removes all types associated with a plugin that's being unloaded. 
+	void OnPluginUnloaded(IPlugin& InPlugin)
+	{
+		FWriteScopeLock _(Lock);
+		Super::OnPluginUnloaded_NoLock(InPlugin);
+	}
+	
+	// Update the registry when new types are added to the attribute system so that they can be selected
+	// on Attribute Nodes
+	void OnAnimationAttributeTypesChanged(const UScriptStruct* InStruct, bool bIsAdded)
+	{
+		FWriteScopeLock _(Lock);
+		Super::OnAnimationAttributeTypesChanged_NoLock(InStruct, bIsAdded);
+	}
+	
+	// Clear the registry
+	void Reset(bool bLockRegistry = true)
+	{
+		FConditionalWriteScopeLock _(*this, bLockRegistry);
+		Super::Reset_NoLock();
+	}
+
+	// Adds a type if it doesn't exist yet and returns its index.
+	// This function is thread-safe
+	TRigVMTypeIndex FindOrAddType(const FRigVMTemplateArgumentType& InType, bool bForce = false, bool bLockRegistry = true)
+	{
+		FConditionalWriteScopeLock _(*this, bLockRegistry);
+		return Super::FindOrAddType_NoLock(InType, bForce);
+	}
+
+	// Removes a type from the registry, and updates all dependent templates
+	// which also creates invalid permutations in templates that we should ignore
+	bool RemoveType(const FSoftObjectPath& InObjectPath, const UClass* InObjectClass, bool bLockRegistry = true)
+	{
+		FConditionalWriteScopeLock _(*this, bLockRegistry);
+		return Super::RemoveType_NoLock(InObjectPath, InObjectClass);
+	}
+
+	// Returns the type index given a type
+	TRigVMTypeIndex GetTypeIndex(const FRigVMTemplateArgumentType& InType, bool bLockRegistry = true) const
+	{
+		FConditionalReadScopeLock _(*this, bLockRegistry);
+		return Super::GetTypeIndex_NoLock(InType);
+	}
+
+	// Returns the type index given a cpp type and a type object
+	TRigVMTypeIndex GetTypeIndex(const FName& InCPPType, UObject* InCPPTypeObject, bool bLockRegistry = true) const
+	{
+		FConditionalReadScopeLock _(*this, bLockRegistry);
+		return GetTypeIndex(FRigVMTemplateArgumentType(InCPPType, InCPPTypeObject));
+	}
+
+	// Returns the type index given an enum, struct, or object
+	template <typename T>
+	TRigVMTypeIndex GetTypeIndex(bool bAsArray = false, bool bLockRegistry = true) const
+	{
+		FConditionalReadScopeLock _(*this, bLockRegistry);
+		return Super::GetTypeIndex_NoLock<T>(bAsArray);
+	}
+
+	// Returns the type given its index
+	const FRigVMTemplateArgumentType& GetType(TRigVMTypeIndex InTypeIndex, bool bLockRegistry = true) const
+	{
+		FConditionalReadScopeLock _(*this, bLockRegistry);
+		return Super::GetType_NoLock(InTypeIndex);
+	}
+
+	// Returns the number of types
+	int32 NumTypes(bool bLockRegistry = true) const
+	{
+		FConditionalReadScopeLock _(*this, bLockRegistry);
+		return Super::NumTypes_NoLock();
+	}
+
+	// Returns the type given only its cpp type
+	const FRigVMTemplateArgumentType& FindTypeFromCPPType(const FString& InCPPType, bool bLockRegistry = true) const
+	{
+		FConditionalReadScopeLock _(*this, bLockRegistry);
+		return Super::FindTypeFromCPPType_NoLock(InCPPType);
+	}
+
+	// Returns the type index given only its cpp type
+	TRigVMTypeIndex GetTypeIndexFromCPPType(const FString& InCPPType, bool bLockRegistry = true) const
+	{
+		FConditionalReadScopeLock _(*this, bLockRegistry);
+		return Super::GetTypeIndexFromCPPType_NoLock(InCPPType);
+	}
+
+	// Returns true if the type is an array
+	bool IsArrayType(TRigVMTypeIndex InTypeIndex, bool bLockRegistry = true) const
+	{
+		FConditionalReadScopeLock _(*this, bLockRegistry);
+		return Super::IsArrayType_NoLock(InTypeIndex);
+	}
+
+	// Returns true if the type is an execute type
+	bool IsExecuteType(TRigVMTypeIndex InTypeIndex, bool bLockRegistry = true) const
+	{
+		FConditionalReadScopeLock _(*this, bLockRegistry);
+		return Super::IsExecuteType_NoLock(InTypeIndex);
+	} 
+
+	// Converts the given execute context type to the base execute context type
+	bool ConvertExecuteContextToBaseType(TRigVMTypeIndex& InOutTypeIndex, bool bLockRegistry = true) const
+	{
+		FConditionalReadScopeLock _(*this, bLockRegistry);
+		return Super::ConvertExecuteContextToBaseType_NoLock(InOutTypeIndex);
+	} 
+
+	// Returns the dimensions of the array 
+	int32 GetArrayDimensionsForType(TRigVMTypeIndex InTypeIndex, bool bLockRegistry = true) const
+	{
+		FConditionalReadScopeLock _(*this, bLockRegistry);
+		return Super::GetArrayDimensionsForType_NoLock(InTypeIndex);
+	} 
+
+	// Returns true if the type is a wildcard type
+	bool IsWildCardType(TRigVMTypeIndex InTypeIndex) const
+	{
+		// no lock required
+		return Super::IsWildCardType_NoLock(InTypeIndex);
+	} 
+
+	// Returns true if the types can be matched.
+	bool CanMatchTypes(TRigVMTypeIndex InTypeIndexA, TRigVMTypeIndex InTypeIndexB, bool bAllowFloatingPointCasts, bool bLockRegistry = true) const
+	{
+		FConditionalReadScopeLock _(*this, bLockRegistry);
+		return Super::CanMatchTypes_NoLock(InTypeIndexA, InTypeIndexB, bAllowFloatingPointCasts);
+	}
+
+	// Returns the list of compatible types for a given type
+	const TArray<TRigVMTypeIndex>& GetCompatibleTypes(TRigVMTypeIndex InTypeIndex) const
+	{
+		// no lock required
+		return Super::GetCompatibleTypes_NoLock(InTypeIndex);
+	}
+
+	// Returns all compatible types given a category
+	const TArray<TRigVMTypeIndex>& GetTypesForCategory(FRigVMTemplateArgument::ETypeCategory InCategory, bool bLockRegistry = true) const
+	{
+		FConditionalReadScopeLock _(*this, bLockRegistry);
+		return Super::GetTypesForCategory_NoLock(InCategory);
+	}
+
+	// Returns the type index of the array matching the given element type index
+	TRigVMTypeIndex GetArrayTypeFromBaseTypeIndex(TRigVMTypeIndex InTypeIndex, bool bLockRegistry = true) const
+	{
+		FConditionalReadScopeLock _(*this, bLockRegistry);
+		return Super::GetArrayTypeFromBaseTypeIndex_NoLock(InTypeIndex);
+	}
+
+	// Returns the type index of the element matching the given array type index
+	TRigVMTypeIndex GetBaseTypeFromArrayTypeIndex(TRigVMTypeIndex InTypeIndex, bool bLockRegistry = true) const
+	{
+		FConditionalReadScopeLock _(*this, bLockRegistry);
+		return Super::GetBaseTypeFromArrayTypeIndex_NoLock(InTypeIndex);
+	}
+
+	// Returns the function given its name (or nullptr)
+	const FRigVMFunction* FindFunction(const TCHAR* InName, const FRigVMUserDefinedTypeResolver& InTypeResolver = FRigVMUserDefinedTypeResolver(), bool bLockRegistry = true) const
+	{
+		FConditionalWriteScopeLock _(*this, bLockRegistry);
+		return Super::FindFunction_NoLock(InName, InTypeResolver);
+	}
+
+	// Returns the function given its backing up struct and method name
+	const FRigVMFunction* FindFunction(UScriptStruct* InStruct, const TCHAR* InName, const FRigVMUserDefinedTypeResolver& InResolvalInfo = FRigVMUserDefinedTypeResolver(), bool bLockRegistry = true) const
+	{
+		FConditionalWriteScopeLock _(*this, bLockRegistry);
+		return Super::FindFunction_NoLock(InStruct, InName, InResolvalInfo);
+	}
+
+	// Returns all current RigVM functions
+	const TChunkedArray<FRigVMFunction>& GetFunctions(bool bLockRegistry = true) const
+	{
+		FConditionalReadScopeLock _(*this, bLockRegistry);
+		return Super::GetFunctions_NoLock();
+	}
+
+	// Returns a template pointer given its notation (or nullptr)
+	const FRigVMTemplate* FindTemplate(const FName& InNotation, bool bIncludeDeprecated = false, bool bLockRegistry = true) const
+	{
+		FConditionalWriteScopeLock _(*this, bLockRegistry);
+		return Super::FindTemplate_NoLock(InNotation, bIncludeDeprecated);
+	}
+
+	// Returns all current RigVM functions
+	const TChunkedArray<FRigVMTemplate>& GetTemplates(bool bLockRegistry = true) const
+	{
+		FConditionalReadScopeLock _(*this, bLockRegistry);
+		return Super::GetTemplates_NoLock();
+	}
+
+	// Defines and retrieves a template given its arguments
+	const FRigVMTemplate* GetOrAddTemplateFromArguments(
+		const FName& InName,
+		const TArray<FRigVMTemplateArgumentInfo>& InInfos,
+		const FRigVMTemplateDelegates& InDelegates,
+		bool bLockRegistry = true)
+	{
+		FConditionalWriteScopeLock _(*this, bLockRegistry);
+		return Super::GetOrAddTemplateFromArguments_NoLock(InName, InInfos, InDelegates);
+	}
+
+	// Adds a new template given its arguments
+	const FRigVMTemplate* AddTemplateFromArguments(
+		const FName& InName,
+		const TArray<FRigVMTemplateArgumentInfo>& InInfos,
+		const FRigVMTemplateDelegates& InDelegates,
+		bool bLockRegistry = true)
+	{
+		FConditionalWriteScopeLock _(*this, bLockRegistry);
+		return Super::AddTemplateFromArguments_NoLock(InName, InInfos, InDelegates);
+	}
+
+	// Returns a dispatch factory given its name (or nullptr)
+	FRigVMDispatchFactory* FindDispatchFactory(const FName& InFactoryName, bool bLockRegistry = true) const
+	{
+		FConditionalReadScopeLock _(*this, bLockRegistry);
+		return Super::FindDispatchFactory_NoLock(InFactoryName);
+	}
+
+	// Returns a dispatch factory given its static struct (or nullptr)
+	FRigVMDispatchFactory* FindOrAddDispatchFactory(UScriptStruct* InFactoryStruct, bool bLockRegistry = true)
+	{
+		FConditionalWriteScopeLock _(*this, bLockRegistry);
+		return Super::FindOrAddDispatchFactory_NoLock(InFactoryStruct);
+	}
+
+	// Returns a dispatch factory given its static struct (or nullptr)
+	template <
+		typename T,
+		typename TEnableIf<TModels_V<CRigVMUStruct, T>>::Type * = nullptr
+	>
+	FRigVMDispatchFactory* FindOrAddDispatchFactory(bool bLockRegistry = true)
+	{
+		FConditionalWriteScopeLock _(*this, bLockRegistry);
+		return Super::FindOrAddDispatchFactory_NoLock<T>();
+	}
+
+	// Returns a dispatch factory's singleton function name if that exists
+	FString FindOrAddSingletonDispatchFunction(UScriptStruct* InFactoryStruct, bool bLockRegistry = true)
+	{
+		FConditionalWriteScopeLock _(*this, bLockRegistry);
+		return Super::FindOrAddSingletonDispatchFunction_NoLock(InFactoryStruct);
+	}
+
+	// Returns a dispatch factory's singleton function name if that exists
+	template <
+		typename T,
+		typename TEnableIf<TModels_V<CRigVMUStruct, T>>::Type * = nullptr
+	>
+	FString FindOrAddSingletonDispatchFunction(bool bLockRegistry = true)
+	{
+		FConditionalWriteScopeLock _(*this, bLockRegistry);
+		return Super::FindOrAddSingletonDispatchFunction_NoLock<T>();
+	}
+
+	// Returns all dispatch factories
+	const TArray<FRigVMDispatchFactory*>& GetFactories(bool bLockRegistry = true) const
+	{
+		FConditionalReadScopeLock _(*this, bLockRegistry);
+		return Super::GetFactories_NoLock();
+	}
+
+	// Given a struct name, return the predicates
+	const TArray<FRigVMFunction>* GetPredicatesForStruct(const FName& InStructName, bool bLockRegistry = true) const
+	{
+		FConditionalReadScopeLock _(*this, bLockRegistry);
+		return Super::GetPredicatesForStruct_NoLock(InStructName);
+	}
+
+	// Returns a unique hash per type index
+	uint32 GetHashForType(TRigVMTypeIndex InTypeIndex, bool bLockRegistry = true) const
+	{
+		FConditionalReadScopeLock _(*this, bLockRegistry);
+		return Super::GetHashForType_NoLock(InTypeIndex);
+	}
+	 
+	uint32 GetHashForScriptStruct(const UScriptStruct* InScriptStruct, bool bCheckTypeIndex = true, bool bLockRegistry = true) const
+	{
+		FConditionalReadScopeLock _(*this, bLockRegistry);
+		return Super::GetHashForScriptStruct_NoLock(InScriptStruct, bCheckTypeIndex);
+	}
+	 
+	uint32 GetHashForStruct(const UStruct* InStruct, bool bLockRegistry = true) const
+	{
+		FConditionalReadScopeLock _(*this, bLockRegistry);
+		return Super::GetHashForStruct_NoLock(InStruct);
+	}
+	 
+	uint32 GetHashForEnum(const UEnum* InEnum, bool bCheckTypeIndex = true, bool bLockRegistry = true) const
+	{
+		const FConditionalReadScopeLock _(*this, bLockRegistry);
+		return Super::GetHashForEnum_NoLock(InEnum, bCheckTypeIndex);
+	}
+	 
+	uint32 GetHashForProperty(const FProperty* InProperty, bool bLockRegistry = true) const
+	{
+		const FConditionalReadScopeLock _(*this, bLockRegistry);
+		return Super::GetHashForProperty_NoLock(InProperty);
+	}
+
+	void RebuildRegistry(bool bLockRegistry = true)
+	{
+		FConditionalWriteScopeLock _(*this, bLockRegistry);
+		Super::RebuildRegistry_NoLock();
+	}
+
+protected:
+
+	// Registers a function given its name.
+	// The name will be the name of the struct and virtual method,
+	// for example "FMyStruct::MyVirtualMethod"
+	// Note: Only call NoLock methods on the FRigVMRegistry_NoLock
+	virtual void Register_NoLock(const TCHAR* InName, FRigVMFunctionPtr InFunctionPtr, UScriptStruct* InStruct = nullptr, const TArray<FRigVMFunctionArgument>& InArguments = TArray<FRigVMFunctionArgument>()) override
+	{
+		Super::Register_NoLock(InName, InFunctionPtr, InStruct, InArguments);
+	}
+
+	// Registers a dispatch factory given its struct.
+	// Note: Only call NoLock methods on the FRigVMRegistry_NoLock
+	virtual const FRigVMDispatchFactory* RegisterFactory_NoLock(UScriptStruct* InFactoryStruct) override
+	{
+		return Super::RegisterFactory_NoLock(InFactoryStruct);
+	}
+
+	// Register a predicate contained in the input struct
+	// Note: Only call NoLock methods on the FRigVMRegistry_NoLock
+	virtual void RegisterPredicate_NoLock(UScriptStruct* InStruct, const TCHAR* InName, const TArray<FRigVMFunctionArgument>& InArguments) override
+	{
+		return Super::RegisterPredicate_NoLock(InStruct, InName, InArguments);
+	}
+
+	// Register a set of allowed object types
+	// Note: Only call NoLock methods on the FRigVMRegistry_NoLock
+	virtual void RegisterObjectTypes_NoLock(TConstArrayView<TPair<UClass*, ERegisterObjectOperation>> InClasses) override
+	{
+		Super::RegisterObjectTypes_NoLock(InClasses);
+	}
+
+	// Refreshes the list and finds the function pointers
+	// based on the names.
+	// Note: Only call NoLock methods on the FRigVMRegistry_NoLock
+	virtual void RefreshEngineTypes_NoLock() override
+	{
+		Super::RefreshEngineTypes_NoLock();
+	}
+
+	// Refreshes the list and finds the function pointers
+	// based on the names.
+	// Note: Only call NoLock methods on the FRigVMRegistry_NoLock
+	virtual void RefreshEngineTypesIfRequired_NoLock() override
+	{
+		Super::RefreshEngineTypesIfRequired_NoLock();
+	}
+
+	// Update the registry when types are renamed
+	// Note: Only call NoLock methods on the FRigVMRegistry_NoLock
+	virtual void OnAssetRenamed_NoLock(const FAssetData& InAssetData, const FString& InOldObjectPath) override
+	{
+		Super::OnAssetRenamed_NoLock(InAssetData, InOldObjectPath);
+	}
+	
+	// Update the registry when old types are removed
+	// Note: Only call NoLock methods on the FRigVMRegistry_NoLock
+	virtual void OnAssetRemoved_NoLock(const FAssetData& InAssetData) override
+	{
+		Super::OnAssetRemoved_NoLock(InAssetData);
+	}
+
+	// Removes all types associated with a plugin that's being unloaded. 
+	// Note: Only call NoLock methods on the FRigVMRegistry_NoLock
+	virtual void OnPluginUnloaded_NoLock(IPlugin& InPlugin) override
+	{
+		Super::OnPluginUnloaded_NoLock(InPlugin);
+	}
+	
+	// Update the registry when new types are added to the attribute system so that they can be selected
+	// on Attribute Nodes
+	// Note: Only call NoLock methods on the FRigVMRegistry_NoLock
+	virtual void OnAnimationAttributeTypesChanged_NoLock(const UScriptStruct* InStruct, bool bIsAdded) override
+	{
+		Super::OnAnimationAttributeTypesChanged_NoLock(InStruct, bIsAdded);
+	}
+	
+	// Clear the registry
+	// Note: Only call NoLock methods on the FRigVMRegistry_NoLock
+	virtual void Reset_NoLock() override
+	{
+		Super::Reset_NoLock();
+	}
+
+	// Adds a type if it doesn't exist yet and returns its index.
+	// This function is thread-safe
+	// Note: Only call NoLock methods on the FRigVMRegistry_NoLock
+	virtual TRigVMTypeIndex FindOrAddType_NoLock(const FRigVMTemplateArgumentType& InType, bool bForce = false) override
+	{
+		return Super::FindOrAddType_NoLock(InType, bForce);
+	}
+
+	// Removes a type from the registry, and updates all dependent templates
+	// which also creates invalid permutations in templates that we should ignore
+	// Note: Only call NoLock methods on the FRigVMRegistry_NoLock
+	virtual bool RemoveType_NoLock(const FSoftObjectPath& InObjectPath, const UClass* InObjectClass) override
+	{
+		return Super::RemoveType_NoLock(InObjectPath, InObjectClass);
+	}
+
+	// Returns the type index given a type
+	// Note: Only call NoLock methods on the FRigVMRegistry_NoLock
+	virtual TRigVMTypeIndex GetTypeIndex_NoLock(const FRigVMTemplateArgumentType& InType) const override
+	{
+		return Super::GetTypeIndex_NoLock(InType);
+	}
+
+	// Returns the type index given a cpp type and a type object
+	// Note: Only call NoLock methods on the FRigVMRegistry_NoLock
+	virtual TRigVMTypeIndex GetTypeIndex_NoLock(const FName& InCPPType, UObject* InCPPTypeObject) const override
+	{
+		return Super::GetTypeIndex_NoLock(FRigVMTemplateArgumentType(InCPPType, InCPPTypeObject));
+	}
+
+	// Returns the type given its index
+	// Note: Only call NoLock methods on the FRigVMRegistry_NoLock
+	virtual const FRigVMTemplateArgumentType& GetType_NoLock(TRigVMTypeIndex InTypeIndex) const override
+	{
+		return Super::GetType_NoLock(InTypeIndex);
+	}
+
+	// Returns the number of types
+	// Note: Only call NoLock methods on the FRigVMRegistry_NoLock
+	virtual int32 NumTypes_NoLock() const override
+	{
+		return Super::NumTypes_NoLock();
+	}
+
+	// Returns the type given only its cpp type
+	// Note: Only call NoLock methods on the FRigVMRegistry_NoLock
+	virtual const FRigVMTemplateArgumentType& FindTypeFromCPPType_NoLock(const FString& InCPPType) const override
+	{
+		return Super::FindTypeFromCPPType_NoLock(InCPPType);
+	}
+
+	// Returns the type index given only its cpp type
+	// Note: Only call NoLock methods on the FRigVMRegistry_NoLock
+	virtual TRigVMTypeIndex GetTypeIndexFromCPPType_NoLock(const FString& InCPPType) const override
+	{
+		return Super::GetTypeIndexFromCPPType_NoLock(InCPPType);
+	}
+
+	// Returns true if the type is an array
+	// Note: Only call NoLock methods on the FRigVMRegistry_NoLock
+	virtual bool IsArrayType_NoLock(TRigVMTypeIndex InTypeIndex) const override
+	{
+		return Super::IsArrayType_NoLock(InTypeIndex);
+	}
+
+	// Returns true if the type is an execute type
+	// Note: Only call NoLock methods on the FRigVMRegistry_NoLock
+	virtual bool IsExecuteType_NoLock(TRigVMTypeIndex InTypeIndex) const override
+	{
+		return Super::IsExecuteType_NoLock(InTypeIndex);
+	} 
+
+	// Converts the given execute context type to the base execute context type
+	// Note: Only call NoLock methods on the FRigVMRegistry_NoLock
+	virtual bool ConvertExecuteContextToBaseType_NoLock(TRigVMTypeIndex& InOutTypeIndex) const override
+	{
+		return Super::ConvertExecuteContextToBaseType_NoLock(InOutTypeIndex);
+	} 
+
+	// Returns the dimensions of the array 
+	// Note: Only call NoLock methods on the FRigVMRegistry_NoLock
+	virtual int32 GetArrayDimensionsForType_NoLock(TRigVMTypeIndex InTypeIndex) const override
+	{
+		return Super::GetArrayDimensionsForType_NoLock(InTypeIndex);
+	} 
+
+	// Returns true if the type is a wildcard type
+	// Note: Only call NoLock methods on the FRigVMRegistry_NoLock
+	virtual bool IsWildCardType_NoLock(TRigVMTypeIndex InTypeIndex) const override
+	{
+		return Super::IsWildCardType_NoLock(InTypeIndex);
+	} 
+
+	// Returns true if the types can be matched.
+	// Note: Only call NoLock methods on the FRigVMRegistry_NoLock
+	virtual bool CanMatchTypes_NoLock(TRigVMTypeIndex InTypeIndexA, TRigVMTypeIndex InTypeIndexB, bool bAllowFloatingPointCasts) const override
+	{
+		return Super::CanMatchTypes_NoLock(InTypeIndexA, InTypeIndexB, bAllowFloatingPointCasts);
+	}
+
+	// Returns the list of compatible types for a given type
+	// Note: Only call NoLock methods on the FRigVMRegistry_NoLock
+	virtual const TArray<TRigVMTypeIndex>& GetCompatibleTypes_NoLock(TRigVMTypeIndex InTypeIndex) const override
+	{
+		return Super::GetCompatibleTypes_NoLock(InTypeIndex);
+	}
+
+	// Returns all compatible types given a category
+	// Note: Only call NoLock methods on the FRigVMRegistry_NoLock
+	virtual const TArray<TRigVMTypeIndex>& GetTypesForCategory_NoLock(FRigVMTemplateArgument::ETypeCategory InCategory) const override
+	{
+		return Super::GetTypesForCategory_NoLock(InCategory);
+	}
+
+	// Returns the type index of the array matching the given element type index
+	// Note: Only call NoLock methods on the FRigVMRegistry_NoLock
+	virtual TRigVMTypeIndex GetArrayTypeFromBaseTypeIndex_NoLock(TRigVMTypeIndex InTypeIndex) const override
+	{
+		return Super::GetArrayTypeFromBaseTypeIndex_NoLock(InTypeIndex);
+	}
+
+	// Returns the type index of the element matching the given array type index
+	// Note: Only call NoLock methods on the FRigVMRegistry_NoLock
+	virtual TRigVMTypeIndex GetBaseTypeFromArrayTypeIndex_NoLock(TRigVMTypeIndex InTypeIndex) const override
+	{
+		return Super::GetBaseTypeFromArrayTypeIndex_NoLock(InTypeIndex);
+	}
+
+	// Returns the function given its name (or nullptr)
+	// Note: Only call NoLock methods on the FRigVMRegistry_NoLock
+	virtual const FRigVMFunction* FindFunction_NoLock(const TCHAR* InName, const FRigVMUserDefinedTypeResolver& InTypeResolver = FRigVMUserDefinedTypeResolver()) const override
+	{
+		return Super::FindFunction_NoLock(InName, InTypeResolver);
+	}
+
+	// Returns the function given its backing up struct and method name
+	// Note: Only call NoLock methods on the FRigVMRegistry_NoLock
+	virtual const FRigVMFunction* FindFunction_NoLock(UScriptStruct* InStruct, const TCHAR* InName, const FRigVMUserDefinedTypeResolver& InResolvalInfo = FRigVMUserDefinedTypeResolver()) const override
+	{
+		return Super::FindFunction_NoLock(InStruct, InName, InResolvalInfo);
+	}
+
+	// Returns all current RigVM functions
+	// Note: Only call NoLock methods on the FRigVMRegistry_NoLock
+	virtual const TChunkedArray<FRigVMFunction>& GetFunctions_NoLock() const override
+	{
+		return Super::GetFunctions_NoLock();
+	}
+
+	// Returns a template pointer given its notation (or nullptr)
+	// Note: Only call NoLock methods on the FRigVMRegistry_NoLock
+	virtual const FRigVMTemplate* FindTemplate_NoLock(const FName& InNotation, bool bIncludeDeprecated = false) const override
+	{
+		return Super::FindTemplate_NoLock(InNotation, bIncludeDeprecated);
+	}
+
+	// Returns all current RigVM functions
+	// Note: Only call NoLock methods on the FRigVMRegistry_NoLock
+	virtual const TChunkedArray<FRigVMTemplate>& GetTemplates_NoLock() const override
+	{
+		return Super::GetTemplates_NoLock();
+	}
+
+	// Defines and retrieves a template given its arguments
+	// Note: Only call NoLock methods on the FRigVMRegistry_NoLock
+	virtual const FRigVMTemplate* GetOrAddTemplateFromArguments_NoLock(
+		const FName& InName,
+		const TArray<FRigVMTemplateArgumentInfo>& InInfos,
+		const FRigVMTemplateDelegates& InDelegates) override
+	{
+		return Super::GetOrAddTemplateFromArguments_NoLock(InName, InInfos, InDelegates);
+	}
+
+	// Adds a new template given its arguments
+	// Note: Only call NoLock methods on the FRigVMRegistry_NoLock
+	virtual const FRigVMTemplate* AddTemplateFromArguments_NoLock(
+		const FName& InName,
+		const TArray<FRigVMTemplateArgumentInfo>& InInfos,
+		const FRigVMTemplateDelegates& InDelegates) override
+	{
+		return Super::AddTemplateFromArguments_NoLock(InName, InInfos, InDelegates);
+	}
+
+	// Returns a dispatch factory given its name (or nullptr)
+	// Note: Only call NoLock methods on the FRigVMRegistry_NoLock
+	virtual FRigVMDispatchFactory* FindDispatchFactory_NoLock(const FName& InFactoryName) const override
+	{
+		return Super::FindDispatchFactory_NoLock(InFactoryName);
+	}
+
+	// Returns a dispatch factory given its static struct (or nullptr)
+	// Note: Only call NoLock methods on the FRigVMRegistry_NoLock
+	virtual FRigVMDispatchFactory* FindOrAddDispatchFactory_NoLock(UScriptStruct* InFactoryStruct) override
+	{
+		return Super::FindOrAddDispatchFactory_NoLock(InFactoryStruct);
+	}
+
+	// Returns a dispatch factory's singleton function name if that exists
+	// Note: Only call NoLock methods on the FRigVMRegistry_NoLock
+	virtual FString FindOrAddSingletonDispatchFunction_NoLock(UScriptStruct* InFactoryStruct) override
+	{
+		return Super::FindOrAddSingletonDispatchFunction_NoLock(InFactoryStruct);
+	}
 
-	static FCriticalSection DispatchFunctionMutex;
-	static FCriticalSection DispatchPredicatesMutex;
+	// Returns all dispatch factories
+	// Note: Only call NoLock methods on the FRigVMRegistry_NoLock
+	virtual const TArray<FRigVMDispatchFactory*>& GetFactories_NoLock() const override
+	{
+		return Super::GetFactories_NoLock();
+	}
+
+	// Given a struct name, return the predicates
+	// Note: Only call NoLock methods on the FRigVMRegistry_NoLock
+	virtual const TArray<FRigVMFunction>* GetPredicatesForStruct_NoLock(const FName& InStructName) const override
+	{
+		return Super::GetPredicatesForStruct_NoLock(InStructName);
+	}
+
+	// Returns a unique hash per type index
+	// Note: Only call NoLock methods on the FRigVMRegistry_NoLock
+	virtual uint32 GetHashForType_NoLock(TRigVMTypeIndex InTypeIndex) const override
+	{
+		return Super::GetHashForType_NoLock(InTypeIndex);
+	}
+	 
+	// Note: Only call NoLock methods on the FRigVMRegistry_NoLock
+	virtual uint32 GetHashForScriptStruct_NoLock(const UScriptStruct* InScriptStruct, bool bCheckTypeIndex = true) const override
+	{
+		return Super::GetHashForScriptStruct_NoLock(InScriptStruct, bCheckTypeIndex);
+	}
+	 
+	// Note: Only call NoLock methods on the FRigVMRegistry_NoLock
+	virtual uint32 GetHashForStruct_NoLock(const UStruct* InStruct) const override
+	{
+		return Super::GetHashForStruct_NoLock(InStruct);
+	}
+	 
+	// Note: Only call NoLock methods on the FRigVMRegistry_NoLock
+	virtual uint32 GetHashForEnum_NoLock(const UEnum* InEnum, bool bCheckTypeIndex = true) const override
+	{
+		return Super::GetHashForEnum_NoLock(InEnum, bCheckTypeIndex);
+	}
+	 
+	// Note: Only call NoLock methods on the FRigVMRegistry_NoLock
+	virtual uint32 GetHashForProperty_NoLock(const FProperty* InProperty) const override
+	{
+		return Super::GetHashForProperty_NoLock(InProperty);
+	}
+	 
+	// Note: Only call NoLock methods on the FRigVMRegistry_NoLock
+	virtual void RebuildRegistry_NoLock() override
+	{
+		Super::RebuildRegistry_NoLock();
+	}
+
+private:
+
+	FRigVMRegistry_RWLock();
+
+	// Initialize the base types
+	virtual void Initialize(bool bLockRegistry) override;
+
+	static void EnsureLocked(ELockType InLockType);
+
+	mutable FRWLock Lock;
+	mutable std::atomic<ELockType> LockType;
 	
 	friend struct FRigVMStruct;
 	friend struct FRigVMTemplate;
+	friend struct FRigVMFunction;
 	friend struct FRigVMTemplateArgument;
 	friend struct FRigVMDispatchFactory;
+	friend struct FRigVMRegistry_NoLock;
 };
+
+typedef FRigVMRegistry_RWLock::FConditionalReadScopeLock FRigVMRegistryReadLock;
+typedef FRigVMRegistry_RWLock::FConditionalWriteScopeLock FRigVMRegistryWriteLock;
diff --git a/Engine/Plugins/Runtime/RigVM/Source/RigVM/Public/RigVMCore/RigVMTemplate.h b/Engine/Plugins/Runtime/RigVM/Source/RigVM/Public/RigVMCore/RigVMTemplate.h
--- a/Engine/Plugins/Runtime/RigVM/Source/RigVM/Public/RigVMCore/RigVMTemplate.h
+++ b/Engine/Plugins/Runtime/RigVM/Source/RigVM/Public/RigVMCore/RigVMTemplate.h
@@ -28,7 +28,7 @@ struct FRigVMDispatchFactory;
 struct FRigVMTemplate;
 struct FRigVMDispatchContext;
 struct FRigVMUserDefinedTypeResolver;
-struct FRigVMRegistry;
+struct FRigVMRegistry_NoLock;
 
 typedef TMap<FName, TRigVMTypeIndex> FRigVMTemplateTypeMap;
 
@@ -205,13 +205,19 @@ struct RIGVM_API FRigVMTemplateArgument
 	ERigVMPinDirection GetDirection() const { return Direction; }
 
 	// returns true if this argument supports a given type across a set of permutations
-	bool SupportsTypeIndex(TRigVMTypeIndex InTypeIndex, TRigVMTypeIndex* OutTypeIndex = nullptr) const;
+	bool SupportsTypeIndex(TRigVMTypeIndex InTypeIndex, TRigVMTypeIndex* OutTypeIndex = nullptr, const bool bLockRegistry = true) const;
+	bool SupportsTypeIndex_NoLock(TRigVMTypeIndex InTypeIndex, TRigVMTypeIndex* OutTypeIndex = nullptr) const;
 
 	// returns the flat list of types (including duplicates) of this argument
-	void GetAllTypes(TArray<TRigVMTypeIndex>& OutTypes) const;
-	
-	TRigVMTypeIndex GetTypeIndex(const int32 InIndex) const;
+	void GetAllTypes(TArray<TRigVMTypeIndex>& OutTypes, const bool bLockRegistry = true) const;
+	void GetAllTypes_NoLock(TArray<TRigVMTypeIndex>& OutTypes) const;
+
+	TRigVMTypeIndex GetTypeIndex(const int32 InIndex, const bool bLockRegistry = true) const;
+	TRigVMTypeIndex GetTypeIndex_NoLock(const int32 InIndex) const;
+	TOptional<TRigVMTypeIndex> TryToGetTypeIndex(const int32 InIndex, const bool bLockRegistry = true) const;
+	TOptional<TRigVMTypeIndex> TryToGetTypeIndex_NoLock(const int32 InIndex) const;
 	int32 GetNumTypes() const;
+	int32 GetNumTypes_NoLock() const;
 	void AddTypeIndex(const TRigVMTypeIndex InTypeIndex);
 	void RemoveType(const int32 InIndex);
 	void ForEachType(TFunction<bool(const TRigVMTypeIndex InType)>&& InCallback) const;
@@ -259,15 +265,17 @@ struct RIGVM_API FRigVMTemplateArgument
 #endif
 	
 	// returns true if an argument is singleton (same type for all variants)
-	bool IsSingleton(const TArray<int32>& InPermutationIndices = TArray<int32>()) const;
+	bool IsSingleton(const TArray<int32>& InPermutationIndices = TArray<int32>(), const bool bLockRegistry = true) const;
+	bool IsSingleton_NoLock(const TArray<int32>& InPermutationIndices = TArray<int32>()) const;
 
 	// returns true if this argument is an execute
 	bool IsExecute() const;
 
 	// returns true if the argument uses an array container
-	EArrayType GetArrayType() const;
+	EArrayType GetArrayType(const bool bLockRegistry = true) const;
+	EArrayType GetArrayType_NoLock() const;
 
-	RIGVM_API friend uint32 GetTypeHash(const FRigVMTemplateArgument& InArgument);
+	RIGVM_API friend uint32 GetTypeHash_NoLock(const FRigVMTemplateArgument& InArgument);
 
 	// Get the map of types to permutation indices
 	const TArray<int32>& GetPermutations(const TRigVMTypeIndex InType) const;
@@ -285,19 +293,23 @@ protected:
 	bool bUseCategories = false;
 	TArray<ETypeCategory> TypeCategories;
 	TFunction<bool(const TRigVMTypeIndex&)> FilterType;
+	mutable TOptional<EArrayType> CachedArrayType;
 
 	FRigVMTemplateArgument(FProperty* InProperty);
 
 	// constructor from a property. this forces the type to be created
-	FRigVMTemplateArgument(FProperty* InProperty, FRigVMRegistry& InRegistry);
+	FRigVMTemplateArgument(FProperty* InProperty, FRigVMRegistry_NoLock& InRegistry);
 
-	void EnsureValidExecuteType(FRigVMRegistry& InRegistry);
+	// static make function using a lock to create it
+	static FRigVMTemplateArgument Make(FProperty* InProperty);
+
+	void EnsureValidExecuteType_NoLock(FRigVMRegistry_NoLock& InRegistry);
 	void UpdateTypeToPermutations();
 
 	friend struct FRigVMTemplate;
 	friend struct FRigVMDispatchFactory;
 	friend class URigVMController;
-	friend struct FRigVMRegistry;
+	friend struct FRigVMRegistry_NoLock;
 	friend struct FRigVMStructUpgradeInfo;
 	friend class URigVMCompiler;
 
@@ -411,7 +423,7 @@ public:
 	const FRigVMExecuteArgument* FindExecuteArgument(const FName& InArgumentName, const FRigVMDispatchContext& InContext) const;
 
 	// returns the top level execute context struct this template uses
-	const UScriptStruct* GetExecuteContextStruct() const;
+	const UScriptStruct* GetExecuteContextStruct(bool bLockRegistry = true) const;
 
 	// returns true if this template supports a given execute context struct
 	bool SupportsExecuteContextStruct(const UScriptStruct* InExecuteContextStruct) const;
@@ -423,13 +435,13 @@ public:
 	int32 NumPermutations() const { return Permutations.Num(); }
 
 	// returns the first / primary permutation of the template
-	const FRigVMFunction* GetPrimaryPermutation() const;
+	const FRigVMFunction* GetPrimaryPermutation(bool bLockRegistry = true) const;
 
 	// returns a permutation given an index
-	const FRigVMFunction* GetPermutation(int32 InIndex) const;
+	const FRigVMFunction* GetPermutation(int32 InIndex, bool bLockRegistry = true) const;
 
 	// returns a permutation given an index and creates it using the backing factory if needed
-	const FRigVMFunction* GetOrCreatePermutation(int32 InIndex);
+	const FRigVMFunction* GetOrCreatePermutation(int32 InIndex, bool bLockRegistry = true);
 
 	// returns true if a given function is a permutation of this template
 	bool ContainsPermutation(const FRigVMFunction* InPermutation) const;
@@ -438,13 +450,13 @@ public:
 	int32 FindPermutation(const FRigVMFunction* InPermutation) const;
 
 	// returns the index of the permutation within the template of a given set of types
-	int32 FindPermutation(const FTypeMap& InTypes) const;
+	int32 FindPermutation(const FTypeMap& InTypes, bool bLockRegistry = true) const;
 
 	// returns true if the template was able to resolve to single permutation
-	bool FullyResolve(FTypeMap& InOutTypes, int32& OutPermutationIndex) const;
+	bool FullyResolve(FTypeMap& InOutTypes, int32& OutPermutationIndex, bool bLockRegistry = true) const;
 
 	// returns true if the template was able to resolve to at least one permutation
-	bool Resolve(FTypeMap& InOutTypes, TArray<int32> & OutPermutationIndices, bool bAllowFloatingPointCasts) const;
+	bool Resolve(FTypeMap& InOutTypes, TArray<int32> & OutPermutationIndices, bool bAllowFloatingPointCasts, bool bLockRegistry = true) const;
 
 	// Will return the hash of the input type map, if it is a valid type map. Otherwise, will return 0.
 	// It is only a valid type map if it includes all arguments, and non of the types is a wildcard.
@@ -454,10 +466,11 @@ public:
 	bool ContainsPermutation(const FTypeMap& InTypes) const;
 
 	// returns true if the template can resolve an argument to a new type
-	bool ResolveArgument(const FName& InArgumentName, const TRigVMTypeIndex InTypeIndex, FTypeMap& InOutTypes) const;
+	bool ResolveArgument(const FName& InArgumentName, const TRigVMTypeIndex InTypeIndex, FTypeMap& InOutTypes, bool bLockRegistry = true) const;
 
 	// returns the types for a specific permutation
-	FRigVMTemplateTypeMap GetTypesForPermutation(const int32 InPermutationIndex) const;
+	FRigVMTemplateTypeMap GetTypesForPermutation(const int32 InPermutationIndex, const bool bLockRegistry = true) const;
+	FRigVMTemplateTypeMap GetTypesForPermutation_NoLock(const int32 InPermutationIndex) const;
 
 	// returns true if a given argument is valid for a template
 	static bool IsValidArgumentForTemplate(const ERigVMPinDirection InDirection);
@@ -478,7 +491,7 @@ public:
 	FTypeMap GetArgumentTypesFromString(const FString& InTypeString, const FRigVMUserDefinedTypeResolver* InTypeResolver = nullptr) const;
 
 	// converts the types provided to a string (like "A:float,B:int32")
-	static FString GetStringFromArgumentTypes(const FTypeMap& InTypes); 
+	static FString GetStringFromArgumentTypes(const FTypeMap& InTypes, bool bLockRegistry = true); 
 
 #if WITH_EDITOR
 
@@ -511,7 +524,10 @@ public:
 	FRigVMTemplate_NewArgumentTypeDelegate& OnNewArgumentType() { return Delegates.NewArgumentTypeDelegate; }
 
 	// Returns the factory this template was created by
-	const FRigVMDispatchFactory* GetDispatchFactory() const
+	const FRigVMDispatchFactory* GetDispatchFactory(const bool bLockRegistry = true) const;
+
+	// Returns the factory this template was created by
+	const FRigVMDispatchFactory* GetDispatchFactory_NoLock() const
 	{
 		return UsesDispatch() ? Delegates.GetDispatchFactoryDelegate.Execute() : nullptr;
 	}
@@ -527,6 +543,7 @@ public:
 	void UpdateTypesHashToPermutation(const int32 InPermutation);
 
 	RIGVM_API friend uint32 GetTypeHash(const FRigVMTemplate& InTemplate);
+	RIGVM_API friend uint32 GetTypeHash_NoLock(const FRigVMTemplate& InTemplate);
 
 private:
 
@@ -544,6 +561,10 @@ private:
 	const FRigVMFunction* GetPermutation_NoLock(int32 InIndex) const;
 	const FRigVMFunction* GetOrCreatePermutation_NoLock(int32 InIndex);
 
+	FTypeMap GetArgumentTypesFromString_Impl(const FString& InTypeString, const FRigVMUserDefinedTypeResolver* InTypeResolver, bool bLockRegistry) const;
+
+	uint32 ComputeTypeHash() const;
+
 	int32 Index;
 	FName Notation;
 	TArray<FRigVMTemplateArgument> Arguments;
@@ -554,7 +575,7 @@ private:
 
 	FRigVMTemplateDelegates Delegates;
 
-	friend struct FRigVMRegistry;
+	friend struct FRigVMRegistry_NoLock;
 	friend class URigVMController;
 	friend class URigVMLibraryNode;
 	friend struct FRigVMDispatchFactory;
diff --git a/Engine/Plugins/Runtime/RigVM/Source/RigVM/Public/RigVMCore/RigVMTypeIndex.h b/Engine/Plugins/Runtime/RigVM/Source/RigVM/Public/RigVMCore/RigVMTypeIndex.h
--- a/Engine/Plugins/Runtime/RigVM/Source/RigVM/Public/RigVMCore/RigVMTypeIndex.h
+++ b/Engine/Plugins/Runtime/RigVM/Source/RigVM/Public/RigVMCore/RigVMTypeIndex.h
@@ -81,7 +81,7 @@ protected:
 	FName Name;
 	int32 Index;
 
-	friend struct FRigVMRegistry;
+	friend struct FRigVMRegistry_NoLock;
 	friend struct FRigVMTemplateArgument;
 };
 #else
diff --git a/Engine/Plugins/Runtime/RigVM/Source/RigVMDeveloper/Private/RigVMCompiler/RigVMCompiler.cpp b/Engine/Plugins/Runtime/RigVM/Source/RigVMDeveloper/Private/RigVMCompiler/RigVMCompiler.cpp
--- a/Engine/Plugins/Runtime/RigVM/Source/RigVMDeveloper/Private/RigVMCompiler/RigVMCompiler.cpp
+++ b/Engine/Plugins/Runtime/RigVM/Source/RigVMDeveloper/Private/RigVMCompiler/RigVMCompiler.cpp
@@ -845,7 +845,7 @@ bool URigVMCompiler::Compile(const FRigVMCompileSettings& InSettings, TArray<URi
 					{
 						for (TFieldIterator<FProperty> It(ScriptStruct, EFieldIterationFlags::None); It; ++It)
 						{
-							const FRigVMTemplateArgument ExpectedArgument(*It);
+							const FRigVMTemplateArgument ExpectedArgument = FRigVMTemplateArgument::Make(*It);
 							const TRigVMTypeIndex ExpectedTypeIndex = ExpectedArgument.GetSupportedTypeIndices()[0];
 							if (URigVMPin* Pin = UnitNode->FindPin(ExpectedArgument.Name.ToString()))
 							{
diff --git a/Engine/Plugins/Runtime/RigVM/Source/RigVMDeveloper/Private/RigVMModel/Nodes/RigVMSelectNode.cpp b/Engine/Plugins/Runtime/RigVM/Source/RigVMDeveloper/Private/RigVMModel/Nodes/RigVMSelectNode.cpp
--- a/Engine/Plugins/Runtime/RigVM/Source/RigVMDeveloper/Private/RigVMModel/Nodes/RigVMSelectNode.cpp
+++ b/Engine/Plugins/Runtime/RigVM/Source/RigVMDeveloper/Private/RigVMModel/Nodes/RigVMSelectNode.cpp
@@ -72,11 +72,11 @@ const FRigVMTemplate* UDEPRECATED_RigVMSelectNode::GetTemplate() const
 				if(InArgumentName == ValueFName)
 				{
 					ValueTypeIndex = InTypeIndex;
-					ResultTypeIndex = FRigVMRegistry::Get().GetBaseTypeFromArrayTypeIndex(InTypeIndex);
+					ResultTypeIndex = FRigVMRegistry_NoLock::GetForRead().GetBaseTypeFromArrayTypeIndex_NoLock(InTypeIndex);
 				}
 				else if(InArgumentName == ResultFName)
 				{
-					ValueTypeIndex = FRigVMRegistry::Get().GetArrayTypeFromBaseTypeIndex(InTypeIndex);;
+					ValueTypeIndex = FRigVMRegistry_NoLock::GetForRead().GetArrayTypeFromBaseTypeIndex_NoLock(InTypeIndex);;
 					ResultTypeIndex = InTypeIndex;
 				}
 				
diff --git a/Engine/Plugins/Runtime/RigVM/Source/RigVMDeveloper/Private/RigVMModel/RigVMController.cpp b/Engine/Plugins/Runtime/RigVM/Source/RigVMDeveloper/Private/RigVMModel/RigVMController.cpp
--- a/Engine/Plugins/Runtime/RigVM/Source/RigVMDeveloper/Private/RigVMModel/RigVMController.cpp
+++ b/Engine/Plugins/Runtime/RigVM/Source/RigVMDeveloper/Private/RigVMModel/RigVMController.cpp
@@ -16598,7 +16598,7 @@ bool URigVMController::FullyResolveTemplateNode(URigVMTemplateNode* InNode, int3
 			{
 				for (TFieldIterator<FProperty> It(StructToVisit, EFieldIterationFlags::None); It; ++It)
 				{
-					const FRigVMTemplateArgument ExpectedArgument(*It);
+					const FRigVMTemplateArgument ExpectedArgument = FRigVMTemplateArgument::Make(*It);
 					const TRigVMTypeIndex ExpectedTypeIndex = ExpectedArgument.GetSupportedTypeIndices()[0];
 
 					if(URigVMPin* Pin = InNode->FindPin(It->GetFName().ToString()))
