RigVM: Correctly adapt default values betwen quat and rotator

#rnx

diff --git a/Engine/Plugins/Experimental/AnimNext/Source/AnimNextEditor/Private/Graph/AnimNextGraph_EdGraphNodeCustomization.cpp b/Engine/Plugins/Experimental/AnimNext/Source/AnimNextEditor/Private/Graph/AnimNextGraph_EdGraphNodeCustomization.cpp
--- a/Engine/Plugins/Experimental/AnimNext/Source/AnimNextEditor/Private/Graph/AnimNextGraph_EdGraphNodeCustomization.cpp
+++ b/Engine/Plugins/Experimental/AnimNext/Source/AnimNextEditor/Private/Graph/AnimNextGraph_EdGraphNodeCustomization.cpp
@@ -102,8 +102,8 @@ void FAnimNextGraph_EdGraphNodeCustomization::GenerateTraitData(UAnimNextGraph_E
 						const int32 ScopedDataIndex = TraitData->ScopedSharedDataInstances.Add(MakeShared<FStructOnScope>(TraitSharedInstanceData));
 
 						// Fill the scoped data with the pin data
-						FRigVMPinDefaultValueImportErrorContext ErrorPipe;
-						LOG_SCOPE_VERBOSITY_OVERRIDE(LogExec, ELogVerbosity::Verbose);
+						FRigVMPinDefaultValueImportErrorContext ErrorPipe(ELogVerbosity::Verbose);
+						LOG_SCOPE_VERBOSITY_OVERRIDE(LogExec, ErrorPipe.GetMaxVerbosity());
 						const FString DefaultValue = TraitPin->GetDefaultValue();
 						TraitSharedInstanceData->ImportText(*DefaultValue, TraitData->ScopedSharedDataInstances[ScopedDataIndex]->GetStructMemory(), nullptr, PPF_SerializedAsImportText, &ErrorPipe, TraitSharedInstanceData->GetName());
 					}
diff --git a/Engine/Plugins/Runtime/RigVM/Source/RigVMDeveloper/Private/RigVMModel/RigVMController.cpp b/Engine/Plugins/Runtime/RigVM/Source/RigVMDeveloper/Private/RigVMModel/RigVMController.cpp
--- a/Engine/Plugins/Runtime/RigVM/Source/RigVMDeveloper/Private/RigVMModel/RigVMController.cpp
+++ b/Engine/Plugins/Runtime/RigVM/Source/RigVMDeveloper/Private/RigVMModel/RigVMController.cpp
@@ -13754,10 +13754,10 @@ FName URigVMController::AddDecorator(URigVMNode* InNode, UScriptStruct* InDecora
 
 	if(!InDefaultValue.IsEmpty())
 	{
-		FRigVMPinDefaultValueImportErrorContext ErrorPipe;
+		FRigVMPinDefaultValueImportErrorContext ErrorPipe(ELogVerbosity::Verbose);
 		{
 			// force logging to the error pipe for error detection
-			LOG_SCOPE_VERBOSITY_OVERRIDE(LogExec, ELogVerbosity::Verbose); 
+			LOG_SCOPE_VERBOSITY_OVERRIDE(LogExec, ErrorPipe.GetMaxVerbosity()); 
 			InDecoratorScriptStruct->ImportText(*InDefaultValue, Decorator, nullptr, PPF_None, &ErrorPipe, InDecoratorScriptStruct->GetName()); 
 		}
 	}
diff --git a/Engine/Plugins/Runtime/RigVM/Source/RigVMDeveloper/Private/RigVMModel/RigVMNode.cpp b/Engine/Plugins/Runtime/RigVM/Source/RigVMDeveloper/Private/RigVMModel/RigVMNode.cpp
--- a/Engine/Plugins/Runtime/RigVM/Source/RigVMDeveloper/Private/RigVMModel/RigVMNode.cpp
+++ b/Engine/Plugins/Runtime/RigVM/Source/RigVMDeveloper/Private/RigVMModel/RigVMNode.cpp
@@ -489,10 +489,10 @@ TSharedPtr<FStructOnScope> URigVMNode::GetDecoratorInstance(const URigVMPin* InD
 			const FString DefaultValue = RootPin->GetDefaultValue();
 			if(!DefaultValue.IsEmpty())
 			{
-				FRigVMPinDefaultValueImportErrorContext ErrorPipe;
+				FRigVMPinDefaultValueImportErrorContext ErrorPipe(ELogVerbosity::Verbose);
 				{
 					// force logging to the error pipe for error detection
-					LOG_SCOPE_VERBOSITY_OVERRIDE(LogExec, ELogVerbosity::Verbose);
+					LOG_SCOPE_VERBOSITY_OVERRIDE(LogExec, ErrorPipe.GetMaxVerbosity());
 					ScriptStruct->ImportText(*DefaultValue, Decorator, nullptr, PPF_SerializedAsImportText, &ErrorPipe, ScriptStruct->GetName());
 				}
 			}
diff --git a/Engine/Plugins/Runtime/RigVM/Source/RigVMDeveloper/Private/RigVMModel/RigVMPin.cpp b/Engine/Plugins/Runtime/RigVM/Source/RigVMDeveloper/Private/RigVMModel/RigVMPin.cpp
--- a/Engine/Plugins/Runtime/RigVM/Source/RigVMDeveloper/Private/RigVMModel/RigVMPin.cpp
+++ b/Engine/Plugins/Runtime/RigVM/Source/RigVMDeveloper/Private/RigVMModel/RigVMPin.cpp
@@ -829,9 +829,9 @@ FString URigVMPin::GetDefaultValue(const URigVMPin::FPinOverride& InOverride, bo
 			if(GetScriptStruct() == TBaseStructure<FQuat>::Get())
 			{
 				// quaternions also allow default values stored as rotators
-				FRigVMPinDefaultValueImportErrorContext ErrorPipe;
+				FRigVMPinDefaultValueImportErrorContext ErrorPipe(ELogVerbosity::Verbose);
 				FRotator Rotator = FRotator::ZeroRotator;
-				LOG_SCOPE_VERBOSITY_OVERRIDE(LogExec, ELogVerbosity::Verbose); 
+				LOG_SCOPE_VERBOSITY_OVERRIDE(LogExec, ErrorPipe.GetMaxVerbosity()); 
 				TBaseStructure<FRotator>::Get()->ImportText(*DefaultValue, &Rotator, nullptr, PPF_None, &ErrorPipe, TBaseStructure<FRotator>::Get()->GetName());
 				if(ErrorPipe.NumErrors == 0)
 				{
@@ -955,9 +955,9 @@ bool URigVMPin::IsValidDefaultValue(const FString& InDefaultValue) const
 			if(ScriptStruct == TBaseStructure<FQuat>::Get())
 			{
 				// quaternions also allow default values stored as rotators
-				FRigVMPinDefaultValueImportErrorContext ErrorPipe;
+				FRigVMPinDefaultValueImportErrorContext ErrorPipe(ELogVerbosity::Verbose);
 				FRotator Rotator = FRotator::ZeroRotator;
-				LOG_SCOPE_VERBOSITY_OVERRIDE(LogExec, ELogVerbosity::Verbose); 
+				LOG_SCOPE_VERBOSITY_OVERRIDE(LogExec, ErrorPipe.GetMaxVerbosity()); 
 				TBaseStructure<FRotator>::Get()->ImportText(*Value, &Rotator, nullptr, PPF_None, &ErrorPipe, TBaseStructure<FRotator>::Get()->GetName());
 				if(ErrorPipe.NumErrors == 0)
 				{
@@ -969,10 +969,10 @@ bool URigVMPin::IsValidDefaultValue(const FString& InDefaultValue) const
 			TempStructBuffer.AddUninitialized(ScriptStruct->GetStructureSize());
 			ScriptStruct->InitializeDefaultValue(TempStructBuffer.GetData());
 
-			FRigVMPinDefaultValueImportErrorContext ErrorPipe(true /* warnings as errors */);
+			FRigVMPinDefaultValueImportErrorContext ErrorPipe(ELogVerbosity::Verbose);
 			{
 				// force logging to the error pipe for error detection
-				LOG_SCOPE_VERBOSITY_OVERRIDE(LogExec, ELogVerbosity::Verbose); 
+				LOG_SCOPE_VERBOSITY_OVERRIDE(LogExec, ErrorPipe.GetMaxVerbosity()); 
 				ScriptStruct->ImportText(*Value, TempStructBuffer.GetData(), nullptr, PPF_None, &ErrorPipe, ScriptStruct->GetName()); 
 			}
 
diff --git a/Engine/Plugins/Runtime/RigVM/Source/RigVMDeveloper/Public/RigVMModel/RigVMPin.h b/Engine/Plugins/Runtime/RigVM/Source/RigVMDeveloper/Public/RigVMModel/RigVMPin.h
--- a/Engine/Plugins/Runtime/RigVM/Source/RigVMDeveloper/Public/RigVMModel/RigVMPin.h
+++ b/Engine/Plugins/Runtime/RigVM/Source/RigVMDeveloper/Public/RigVMModel/RigVMPin.h
@@ -570,22 +570,23 @@ public:
 
 	int32 NumErrors;
 
-	FRigVMPinDefaultValueImportErrorContext(bool InTreatWarningAsErrors = false)
+	FRigVMPinDefaultValueImportErrorContext( ELogVerbosity::Type InMaxVerbosity = ELogVerbosity::Warning )
 		: FOutputDevice()
 		, NumErrors(0)
-		, bTreatWarningAsErrors(InTreatWarningAsErrors)
+		, MaxVerbosity(InMaxVerbosity)
 	{
 	}
 
 	virtual void Serialize(const TCHAR* V, ELogVerbosity::Type Verbosity, const class FName& Category) override
 	{
-		const ELogVerbosity::Type MaxVerbosity = bTreatWarningAsErrors ? ELogVerbosity::Warning : ELogVerbosity::Error;
 		if(Verbosity <= MaxVerbosity)
 		{
 			NumErrors++;
 		}
 	}
 
+	ELogVerbosity::Type GetMaxVerbosity() const { return MaxVerbosity; }
+
 private:
-	bool bTreatWarningAsErrors;
+	ELogVerbosity::Type MaxVerbosity;
 };
diff --git a/Engine/Plugins/Runtime/RigVM/Source/RigVMEditor/Private/Widgets/SRigVMGraphPinQuat.cpp b/Engine/Plugins/Runtime/RigVM/Source/RigVMEditor/Private/Widgets/SRigVMGraphPinQuat.cpp
--- a/Engine/Plugins/Runtime/RigVM/Source/RigVMEditor/Private/Widgets/SRigVMGraphPinQuat.cpp
+++ b/Engine/Plugins/Runtime/RigVM/Source/RigVMEditor/Private/Widgets/SRigVMGraphPinQuat.cpp
@@ -47,9 +47,9 @@ TOptional<FRotator> SRigVMGraphPinQuat::GetRotator() const
 		{
 			// try to import a quaternion
 			{
-				FRigVMPinDefaultValueImportErrorContext ErrorPipe;
+				FRigVMPinDefaultValueImportErrorContext ErrorPipe(ELogVerbosity::Verbose);
 				FQuat Quat = FQuat::Identity;
-				LOG_SCOPE_VERBOSITY_OVERRIDE(LogExec, ELogVerbosity::Verbose); 
+				LOG_SCOPE_VERBOSITY_OVERRIDE(LogExec, ErrorPipe.GetMaxVerbosity()); 
 				TBaseStructure<FQuat>::Get()->ImportText(*DefaultValue, &Quat, nullptr, PPF_None, &ErrorPipe, TBaseStructure<FQuat>::Get()->GetName(), true);
 				if(ErrorPipe.NumErrors == 0)
 				{
@@ -59,9 +59,9 @@ TOptional<FRotator> SRigVMGraphPinQuat::GetRotator() const
 
 			// also try to import a rotator - in case that's how it got stored
 			{
-				FRigVMPinDefaultValueImportErrorContext ErrorPipe; 
+				FRigVMPinDefaultValueImportErrorContext ErrorPipe(ELogVerbosity::Verbose); 
 				FRotator Rotator = FRotator::ZeroRotator;
-				LOG_SCOPE_VERBOSITY_OVERRIDE(LogExec, ELogVerbosity::Verbose); 
+				LOG_SCOPE_VERBOSITY_OVERRIDE(LogExec, ErrorPipe.GetMaxVerbosity()); 
 				TBaseStructure<FRotator>::Get()->ImportText(*DefaultValue, &Rotator, nullptr, PPF_None, &ErrorPipe, TBaseStructure<FQuat>::Get()->GetName(), true);
 				if(ErrorPipe.NumErrors == 0)
 				{
