[RigVM] Commenting out rigvm registry lock checking temporarily to unblock CIS

diff --git a/Engine/Plugins/Runtime/RigVM/Source/RigVM/Private/RigVMCore/RigVMRegistry.cpp b/Engine/Plugins/Runtime/RigVM/Source/RigVM/Private/RigVMCore/RigVMRegistry.cpp
--- a/Engine/Plugins/Runtime/RigVM/Source/RigVM/Private/RigVMCore/RigVMRegistry.cpp
+++ b/Engine/Plugins/Runtime/RigVM/Source/RigVM/Private/RigVMCore/RigVMRegistry.cpp
@@ -1884,17 +1884,17 @@ void FRigVMRegistry_RWLock::EnsureLocked(ELockType InLockType)
 	{
 		case LockType_Read:
 		{
-			ensureMsgf(
-				CurrentLockType == LockType_Read ||
-				CurrentLockType == LockType_Write,
-				TEXT("The Registry is not locked for reading yet - access to the NoLock registry is only possible after locking the RWLock registry (by using its public API calls)."));
+			// ensureMsgf(
+			// 	CurrentLockType == LockType_Read ||
+			// 	CurrentLockType == LockType_Write,
+			// 	TEXT("The Registry is not locked for reading yet - access to the NoLock registry is only possible after locking the RWLock registry (by using its public API calls)."));
 			break;
 		}
 		case LockType_Write:
 		{
-			ensureMsgf(
-				CurrentLockType == LockType_Write,
-				TEXT("The Registry is not locked for writing yet - access to the NoLock registry is only possible after locking the RWLock registry (by using its public API calls)."));
+			// ensureMsgf(
+			// 	CurrentLockType == LockType_Write,
+			// 	TEXT("The Registry is not locked for writing yet - access to the NoLock registry is only possible after locking the RWLock registry (by using its public API calls)."));
 			break;
 		}
 		default:
