[RigVM Registry] Moving implementation of locked registry's get from header to cpp to avoid certain compiler creating muiltiple instance of the static registry
 + uncommented lock checking ensure messages

diff --git a/Engine/Plugins/Runtime/RigVM/Source/RigVM/Private/RigVMCore/RigVMRegistry.cpp b/Engine/Plugins/Runtime/RigVM/Source/RigVM/Private/RigVMCore/RigVMRegistry.cpp
--- a/Engine/Plugins/Runtime/RigVM/Source/RigVM/Private/RigVMCore/RigVMRegistry.cpp
+++ b/Engine/Plugins/Runtime/RigVM/Source/RigVM/Private/RigVMCore/RigVMRegistry.cpp
@@ -919,6 +919,16 @@ void FRigVMRegistry_NoLock::RemoveTypeInCategory_NoLock(FRigVMTemplateArgument::
 	}
 }
 
+// This function needs to be in cpp file instead of header file
+// to avoid confusing certain compilers into creating multiple copies of the registry
+FRigVMRegistry_RWLock& FRigVMRegistry_RWLock::Get()
+{
+	// static in a function scope ensures that the GC system is initiated before 
+	// the registry constructor is called
+	static FRigVMRegistry_RWLock s_RigVMRegistry;
+	return s_RigVMRegistry;
+}
+
 TRigVMTypeIndex FRigVMRegistry_NoLock::GetTypeIndex_NoLock(const FRigVMTemplateArgumentType& InType) const
 {
 	if(const TRigVMTypeIndex* Index = TypeToIndex.Find(InType))
@@ -1884,17 +1894,17 @@ void FRigVMRegistry_RWLock::EnsureLocked(ELockType InLockType)
 	{
 		case LockType_Read:
 		{
-			// ensureMsgf(
-			// 	CurrentLockType == LockType_Read ||
-			// 	CurrentLockType == LockType_Write,
-			// 	TEXT("The Registry is not locked for reading yet - access to the NoLock registry is only possible after locking the RWLock registry (by using its public API calls)."));
+			ensureMsgf(
+				(CurrentLockType == LockType_Read) ||
+				(CurrentLockType == LockType_Write),
+				TEXT("The Registry is not locked for reading yet - access to the NoLock registry is only possible after locking the RWLock registry (by using its public API calls)."));
 			break;
 		}
 		case LockType_Write:
 		{
-			// ensureMsgf(
-			// 	CurrentLockType == LockType_Write,
-			// 	TEXT("The Registry is not locked for writing yet - access to the NoLock registry is only possible after locking the RWLock registry (by using its public API calls)."));
+			ensureMsgf(
+				(CurrentLockType == LockType_Write),
+				TEXT("The Registry is not locked for writing yet - access to the NoLock registry is only possible after locking the RWLock registry (by using its public API calls)."));
 			break;
 		}
 		default:
diff --git a/Engine/Plugins/Runtime/RigVM/Source/RigVM/Public/RigVMCore/RigVMRegistry.h b/Engine/Plugins/Runtime/RigVM/Source/RigVM/Public/RigVMCore/RigVMRegistry.h
--- a/Engine/Plugins/Runtime/RigVM/Source/RigVM/Public/RigVMCore/RigVMRegistry.h
+++ b/Engine/Plugins/Runtime/RigVM/Source/RigVM/Public/RigVMCore/RigVMRegistry.h
@@ -522,13 +522,7 @@ public:
 	};
 
 	// Returns the singleton registry
-	FORCEINLINE static FRigVMRegistry_RWLock& Get()
-	{
-		// static in a function scope ensures that the GC system is initiated before 
-		// the registry constructor is called
-		static FRigVMRegistry_RWLock s_RigVMRegistry;
-		return s_RigVMRegistry;
-	}
+	static FRigVMRegistry_RWLock& Get();
 
 	static FRigVMRegistry_NoLock& Get(ELockType InLockType) = delete;
 	static const FRigVMRegistry_NoLock& GetForRead() = delete;
