RigVM: Avoid accessing registry during function registry at startup

#rnx

diff --git a/Engine/Plugins/Animation/ControlRig/Source/ControlRig/Private/Units/Animation/RigUnit_AnimAttribute.cpp b/Engine/Plugins/Animation/ControlRig/Source/ControlRig/Private/Units/Animation/RigUnit_AnimAttribute.cpp
--- a/Engine/Plugins/Animation/ControlRig/Source/ControlRig/Private/Units/Animation/RigUnit_AnimAttribute.cpp
+++ b/Engine/Plugins/Animation/ControlRig/Source/ControlRig/Private/Units/Animation/RigUnit_AnimAttribute.cpp
@@ -262,12 +262,11 @@ void FRigDispatch_SetAnimAttribute::SetAnimAttributeDispatch(FRigVMExtendedExecu
 	}	
 }
 
-void FRigDispatch_AnimAttributeBase::RegisterDependencyTypes_NoLock() const
+void FRigDispatch_AnimAttributeBase::RegisterDependencyTypes_NoLock(FRigVMRegistry_NoLock& InRegistry) const
 {
-	Super::RegisterDependencyTypes_NoLock();
+	Super::RegisterDependencyTypes_NoLock(InRegistry);
 
-	FRigVMRegistry_NoLock& Registry = FRigVMRegistry_NoLock::GetForWrite();
-	Registry.FindOrAddType_NoLock(FRigVMTemplateArgumentType(), true);
+	InRegistry.FindOrAddType_NoLock(FRigVMTemplateArgumentType(), true);
 	
 	TArray<TWeakObjectPtr<const UScriptStruct>> AttributeTypes = UE::Anim::AttributeTypes::GetRegisteredTypes();
 	for (TWeakObjectPtr<const UScriptStruct> Type : AttributeTypes)
@@ -275,7 +274,7 @@ void FRigDispatch_AnimAttributeBase::RegisterDependencyTypes_NoLock() const
 		UScriptStruct* TypePtr = const_cast<UScriptStruct*>(Type.Get());
 		if (TypePtr)
 		{
-			Registry.FindOrAddType_NoLock(FRigVMTemplateArgumentType(TypePtr));
+			InRegistry.FindOrAddType_NoLock(FRigVMTemplateArgumentType(TypePtr));
 		}
 	}
 }
diff --git a/Engine/Plugins/Animation/ControlRig/Source/ControlRig/Public/Units/Animation/RigUnit_AnimAttribute.h b/Engine/Plugins/Animation/ControlRig/Source/ControlRig/Public/Units/Animation/RigUnit_AnimAttribute.h
--- a/Engine/Plugins/Animation/ControlRig/Source/ControlRig/Public/Units/Animation/RigUnit_AnimAttribute.h
+++ b/Engine/Plugins/Animation/ControlRig/Source/ControlRig/Public/Units/Animation/RigUnit_AnimAttribute.h
@@ -139,7 +139,7 @@ struct CONTROLRIG_API FRigDispatch_AnimAttributeBase : public FRigDispatchFactor
 {
 	GENERATED_BODY()
 
-	virtual void RegisterDependencyTypes_NoLock() const override;
+	virtual void RegisterDependencyTypes_NoLock(FRigVMRegistry_NoLock& InRegistry) const override;
 
 #if WITH_EDITOR
 	virtual FString GetNodeTitle(const FRigVMTemplateTypeMap& InTypes) const override;;
diff --git a/Engine/Plugins/Animation/ControlRig/Source/ControlRig/Public/Units/RigDispatchFactory.h b/Engine/Plugins/Animation/ControlRig/Source/ControlRig/Public/Units/RigDispatchFactory.h
--- a/Engine/Plugins/Animation/ControlRig/Source/ControlRig/Public/Units/RigDispatchFactory.h
+++ b/Engine/Plugins/Animation/ControlRig/Source/ControlRig/Public/Units/RigDispatchFactory.h
@@ -18,12 +18,11 @@ struct CONTROLRIG_API FRigDispatchFactory : public FRigVMDispatchFactory
 		return FControlRigExecuteContext::StaticStruct();
 	}
 
-	virtual void RegisterDependencyTypes_NoLock() const override
+	virtual void RegisterDependencyTypes_NoLock(FRigVMRegistry_NoLock& InRegistry) const override
 	{
-		FRigVMRegistry_NoLock& Registry = FRigVMRegistry_NoLock::GetForWrite();
-		Registry.FindOrAddType_NoLock(FControlRigExecuteContext::StaticStruct());
-		Registry.FindOrAddType_NoLock(FRigElementKey::StaticStruct());
-    	Registry.FindOrAddType_NoLock(FCachedRigElement::StaticStruct());
+		InRegistry.FindOrAddType_NoLock(FControlRigExecuteContext::StaticStruct());
+		InRegistry.FindOrAddType_NoLock(FRigElementKey::StaticStruct());
+    	InRegistry.FindOrAddType_NoLock(FCachedRigElement::StaticStruct());
 	}
 
 #if WITH_EDITOR
diff --git a/Engine/Plugins/Runtime/RigVM/Source/RigVM/Private/RigVMCore/RigVMRegistry.cpp b/Engine/Plugins/Runtime/RigVM/Source/RigVM/Private/RigVMCore/RigVMRegistry.cpp
--- a/Engine/Plugins/Runtime/RigVM/Source/RigVM/Private/RigVMCore/RigVMRegistry.cpp
+++ b/Engine/Plugins/Runtime/RigVM/Source/RigVM/Private/RigVMCore/RigVMRegistry.cpp
@@ -751,7 +751,7 @@ TRigVMTypeIndex FRigVMRegistry_NoLock::FindOrAddType_NoLock(const FRigVMTemplate
 				{
 					// by creating a template argument for the child property
 					// the type will be added by calling ::FindOrAddType_Internal recursively.
-					FRigVMTemplateArgument DummyArgument(Property, *this);
+					(void)FRigVMTemplateArgument::Make_NoLock(Property, *this);
 				}
 #if WITH_EDITOR
 				else
@@ -1294,7 +1294,7 @@ void FRigVMRegistry_NoLock::Register_NoLock(const TCHAR* InName, FRigVMFunctionP
 	for (TFieldIterator<FProperty> It(InStruct); It; ++It)
 	{
 		// creating the argument causes the registration
-		FRigVMTemplateArgument Argument(*It);
+		(void)FRigVMTemplateArgument::Make_NoLock(*It, *this);
 	}
 
 #if WITH_EDITOR
@@ -1384,7 +1384,7 @@ const FRigVMDispatchFactory* FRigVMRegistry_NoLock::RegisterFactory_NoLock(UScri
 	InFactoryStruct->InitializeStruct(Factory, 1);
 	Factory->FactoryScriptStruct = InFactoryStruct;
 	Factories.Add(Factory);
-	Factory->RegisterDependencyTypes_NoLock();
+	Factory->RegisterDependencyTypes_NoLock(*this);
 	return Factory;
 }
 
@@ -1869,7 +1869,9 @@ FRigVMRegistry_RWLock::FRigVMRegistry_RWLock()
 
 void FRigVMRegistry_RWLock::Initialize(bool bLockRegistry)
 {
-	LockType.store(LockType_Invalid);
+	LockType = LockType_Invalid;
+	LockCount = 0;
+	
 	const FConditionalWriteScopeLock _(*this, bLockRegistry);
 	
 	Initialize_NoLock();
diff --git a/Engine/Plugins/Runtime/RigVM/Source/RigVM/Private/RigVMCore/RigVMTemplate.cpp b/Engine/Plugins/Runtime/RigVM/Source/RigVM/Private/RigVMCore/RigVMTemplate.cpp
--- a/Engine/Plugins/Runtime/RigVM/Source/RigVM/Private/RigVMCore/RigVMTemplate.cpp
+++ b/Engine/Plugins/Runtime/RigVM/Source/RigVM/Private/RigVMCore/RigVMTemplate.cpp
@@ -39,12 +39,6 @@ FRigVMTemplateArgument::FRigVMTemplateArgument(const FName& InName, ERigVMPinDir
 	, Direction(InDirection)
 {}
 
-FRigVMTemplateArgument::FRigVMTemplateArgument(FProperty* InProperty):
-	FRigVMTemplateArgument(InProperty, FRigVMRegistry_NoLock::GetForWrite())
-{
-	
-}
-
 FRigVMTemplateArgument::FRigVMTemplateArgument(FProperty* InProperty, FRigVMRegistry_NoLock& InRegistry)
 	: Name(InProperty->GetFName())
 {
@@ -92,8 +86,18 @@ FRigVMTemplateArgument::FRigVMTemplateArgument(FProperty* InProperty, FRigVMRegi
 
 FRigVMTemplateArgument FRigVMTemplateArgument::Make(FProperty* InProperty)
 {
-	const FRigVMRegistryWriteLock _;
-	return FRigVMTemplateArgument(InProperty);
+	FRigVMRegistryWriteLock WriteLock;
+	return Make_NoLock(InProperty, WriteLock.GetRegistry());
+}
+
+FRigVMTemplateArgument FRigVMTemplateArgument::Make_NoLock(FProperty* InProperty)
+{
+	return Make_NoLock(InProperty, FRigVMRegistry_NoLock::GetForWrite());
+}
+
+FRigVMTemplateArgument FRigVMTemplateArgument::Make_NoLock(FProperty* InProperty, FRigVMRegistry_NoLock& InRegistry)
+{
+	return FRigVMTemplateArgument(InProperty, InRegistry);
 }
 
 FRigVMTemplateArgument::FRigVMTemplateArgument(const FName& InName, ERigVMPinDirection InDirection, TRigVMTypeIndex InTypeIndex)
@@ -281,9 +285,14 @@ bool FRigVMTemplateArgument::IsSingleton_NoLock(const TArray<int32>& InPermutati
 bool FRigVMTemplateArgument::IsExecute() const
 {
 	const FRigVMRegistry_NoLock& Registry = FRigVMRegistry_NoLock::GetForRead();
-	const int32 FoundAnyNotExec = IndexOfByPredicate([&](const TRigVMTypeIndex Type)
+	return IsExecute_NoLock(Registry);
+}
+
+bool FRigVMTemplateArgument::IsExecute_NoLock(const FRigVMRegistry_NoLock& InRegistry) const
+{
+	const int32 FoundAnyNotExec = IndexOfByPredicate([&InRegistry](const TRigVMTypeIndex Type)
 	{
-		return !Registry.IsExecuteType_NoLock(Type);
+		return !InRegistry.IsExecuteType_NoLock(Type);
 	});
 	return FoundAnyNotExec == INDEX_NONE;
 }
@@ -807,6 +816,7 @@ FRigVMTemplate::FRigVMTemplate(UScriptStruct* InStruct, const FString& InTemplat
 	, Notation(NAME_None)
 	, Hash(UINT32_MAX)
 {
+	FRigVMRegistry_NoLock& Registry = FRigVMRegistry_NoLock::GetForWrite(); 
 	TArray<FString> ArgumentNotations;
 
 	// create the arguments sorted by super -> child struct.
@@ -816,10 +826,10 @@ FRigVMTemplate::FRigVMTemplate(UScriptStruct* InStruct, const FString& InTemplat
 		// only iterate on this struct's fields, not the super structs'
 		for (TFieldIterator<FProperty> It(Struct, EFieldIterationFlags::None); It; ++It)
 		{
-			FRigVMTemplateArgument Argument(*It);
+			FRigVMTemplateArgument Argument = FRigVMTemplateArgument::Make_NoLock(*It, Registry);
 			Argument.Index = Arguments.Num();
 
-			if(!Argument.IsExecute() && IsValidArgumentForTemplate(Argument.GetDirection()) && Argument.GetDirection() != ERigVMPinDirection::Hidden)
+			if(!Argument.IsExecute_NoLock(Registry) && IsValidArgumentForTemplate(Argument.GetDirection()) && Argument.GetDirection() != ERigVMPinDirection::Hidden)
 			{
 				Arguments.Add(Argument);
 			}
@@ -832,7 +842,7 @@ FRigVMTemplate::FRigVMTemplate(UScriptStruct* InStruct, const FString& InTemplat
 	{
 		if(const FRigVMTemplateArgument* Argument = FindArgument(It->GetFName()))
 		{
-			if(!Argument->IsExecute() && Argument->GetDirection() != ERigVMPinDirection::Hidden)
+			if(!Argument->IsExecute_NoLock(Registry) && Argument->GetDirection() != ERigVMPinDirection::Hidden)
 			{
 				ArgumentNotations.Add(GetArgumentNotation(Argument->Name, Argument->Direction));
 			}
@@ -1456,7 +1466,8 @@ const TArray<FRigVMExecuteArgument>& FRigVMTemplate::GetExecuteArguments(const F
 {
 	if(ExecuteArguments.IsEmpty())
 	{
-		const FRigVMRegistryWriteLock _;
+		FRigVMRegistryWriteLock WriteLock;
+		FRigVMRegistry_NoLock& Registry = WriteLock.GetRegistry();
 
 		if(UsesDispatch())
 		{
@@ -1475,7 +1486,7 @@ const TArray<FRigVMExecuteArgument>& FRigVMTemplate::GetExecuteArguments(const F
 					// only iterate on this struct's fields, not the super structs'
 					for (TFieldIterator<FProperty> It(Struct, EFieldIterationFlags::None); It; ++It)
 					{
-						FRigVMTemplateArgument Argument(*It);
+						FRigVMTemplateArgument Argument = FRigVMTemplateArgument::Make_NoLock(*It, Registry);
 						if(Argument.IsExecute())
 						{
 							ExecuteArguments.Emplace(Argument.Name, Argument.Direction, Argument.GetTypeIndex_NoLock(0));
diff --git a/Engine/Plugins/Runtime/RigVM/Source/RigVM/Public/RigVMCore/RigVMDispatchFactory.h b/Engine/Plugins/Runtime/RigVM/Source/RigVM/Public/RigVMCore/RigVMDispatchFactory.h
--- a/Engine/Plugins/Runtime/RigVM/Source/RigVM/Public/RigVMCore/RigVMDispatchFactory.h
+++ b/Engine/Plugins/Runtime/RigVM/Source/RigVM/Public/RigVMCore/RigVMDispatchFactory.h
@@ -141,7 +141,7 @@ public:
 	// registered needed types during registration of the factory.
 	// this is called within a NoLock code path - so make sure to access
 	// the FRigVMRegistry_NoLock within as needed.
-	virtual void RegisterDependencyTypes_NoLock() const {}
+	virtual void RegisterDependencyTypes_NoLock(FRigVMRegistry_NoLock& InRegistry) const {}
 
 	// returns the arguments of the template
 	RIGVM_API virtual const TArray<FRigVMTemplateArgumentInfo>& GetArgumentInfos() const;
diff --git a/Engine/Plugins/Runtime/RigVM/Source/RigVM/Public/RigVMCore/RigVMRegistry.h b/Engine/Plugins/Runtime/RigVM/Source/RigVM/Public/RigVMCore/RigVMRegistry.h
--- a/Engine/Plugins/Runtime/RigVM/Source/RigVM/Public/RigVMCore/RigVMRegistry.h
+++ b/Engine/Plugins/Runtime/RigVM/Source/RigVM/Public/RigVMCore/RigVMRegistry.h
@@ -41,7 +41,7 @@ struct RIGVM_API FRigVMRegistry_NoLock : public FGCObject
 {
 public:
 
-	enum ELockType
+	enum ELockType : uint8
 	{
 		LockType_Read,
 		LockType_Write,
@@ -439,26 +439,34 @@ protected:
 	{
 	public:
 		UE_NODISCARD_CTOR explicit FConditionalScopeLock(const FRigVMRegistry_RWLock& InRegistry, ELockType InLockType, bool bInLockEnabled = true)
-		: Registry(InRegistry)
-		, LockType(InLockType)
+		: Registry(const_cast<FRigVMRegistry_RWLock*>(&InRegistry))
+		, DesiredLockType(InLockType)
 		, bLockEnabled(bInLockEnabled)
-		, bResetLockType(false)
 		{
 			if(bLockEnabled)
 			{
-				if(LockType == LockType_Read)
+				int32 CurrentLockCount = 0;
+				if(DesiredLockType == LockType_Read)
 				{
-					Registry.Lock.ReadLock();
+					Registry->Lock.ReadLock();
+
+					// fetch_add returns the value preceding the modification
+					// so we have to add one manually to get the current value
+					CurrentLockCount = Registry->LockCount.fetch_add(1) + 1;
 				}
-				else if(LockType == LockType_Write)
+				else if(DesiredLockType == LockType_Write)
 				{
-					Registry.Lock.WriteLock();
+					Registry->Lock.WriteLock();
+
+					// fetch_add returns the value preceding the modification
+					// so we have to add one manually to get the current value
+					CurrentLockCount = Registry->LockCount.fetch_add(1) + 1;
+					ensure(CurrentLockCount == 1);
 				}
 
-				ELockType InvalidLockType = LockType_Invalid;
-				if(Registry.LockType.compare_exchange_strong(InvalidLockType, InLockType))
+				if(CurrentLockCount == 1)
 				{
-					bResetLockType = true;
+					Registry->LockType.store(DesiredLockType);
 				}
 			}
 		}
@@ -467,26 +475,44 @@ protected:
 		{
 			if(bLockEnabled)
 			{
-				if(bResetLockType)
-				{
-					Registry.LockType.store(LockType_Invalid);
-				}
-				if(LockType == LockType_Read)
+				if(DesiredLockType == LockType_Read)
 				{
-					Registry.Lock.ReadUnlock();
+					// fetch_sub returns the value preceding the modification
+					// so we have to subtract one manually to get the current value
+					const int32 CurrentLockCount = Registry->LockCount.fetch_sub(1) - 1;
+					ensure(CurrentLockCount >= 0);
+					if(CurrentLockCount == 0)
+					{
+						Registry->LockType.store(LockType_Invalid);
+					}
+      				Registry->Lock.ReadUnlock();
 				}
-				else if(LockType == LockType_Write)
+				else if(DesiredLockType == LockType_Write)
 				{
-					Registry.Lock.WriteUnlock();
+					// fetch_sub returns the value preceding the modification
+					// so we have to subtract one manually to get the current value
+					const int32 CurrentLockCount = Registry->LockCount.fetch_sub(1) - 1;
+					ensure(CurrentLockCount == 0);
+					Registry->LockType.store(LockType_Invalid);
+					Registry->Lock.WriteUnlock();
 				}
 			}
 		}
 
+		FRigVMRegistry_NoLock& GetRegistry()
+		{
+			return *Registry;
+		}
+
+		const FRigVMRegistry_NoLock& GetRegistry() const
+		{
+			return *Registry;
+		}
+
 	private:
-		const FRigVMRegistry_RWLock& Registry;
-		ELockType LockType;
+		FRigVMRegistry_RWLock* Registry;
+		ELockType DesiredLockType;
 		const bool bLockEnabled;
-		bool bResetLockType;
 
 		UE_NONCOPYABLE(FConditionalScopeLock);
 	};
@@ -560,9 +586,9 @@ public:
 
 	// Refreshes the list and finds the function pointers
 	// based on the names.
-	void RefreshEngineTypes(bool bLockRegistry = true)
+	void RefreshEngineTypes()
 	{
-		FConditionalWriteScopeLock _(*this, bLockRegistry);
+		FConditionalWriteScopeLock _(*this);
 		Super::RefreshEngineTypes_NoLock();
 	}
 
@@ -577,21 +603,21 @@ public:
 	// Update the registry when types are renamed
 	void OnAssetRenamed(const FAssetData& InAssetData, const FString& InOldObjectPath)
 	{
-		FWriteScopeLock _(Lock);
+		FConditionalWriteScopeLock _(*this);
 		Super::OnAssetRenamed_NoLock(InAssetData, InOldObjectPath);
 	}
 	
 	// Update the registry when old types are removed
     void OnAssetRemoved(const FAssetData& InAssetData)
 	{
-		FWriteScopeLock _(Lock);
+		FConditionalWriteScopeLock _(*this);
 		Super::OnAssetRemoved_NoLock(InAssetData);
 	}
 
 	// Removes all types associated with a plugin that's being unloaded. 
 	void OnPluginUnloaded(IPlugin& InPlugin)
 	{
-		FWriteScopeLock _(Lock);
+		FConditionalWriteScopeLock _(*this);
 		Super::OnPluginUnloaded_NoLock(InPlugin);
 	}
 	
@@ -599,7 +625,7 @@ public:
 	// on Attribute Nodes
 	void OnAnimationAttributeTypesChanged(const UScriptStruct* InStruct, bool bIsAdded)
 	{
-		FWriteScopeLock _(Lock);
+		FConditionalWriteScopeLock _(*this);
 		Super::OnAnimationAttributeTypesChanged_NoLock(InStruct, bIsAdded);
 	}
 	
@@ -1247,6 +1273,7 @@ private:
 
 	mutable FRWLock Lock;
 	mutable std::atomic<ELockType> LockType;
+	mutable std::atomic<int32> LockCount;
 	
 	friend struct FRigVMStruct;
 	friend struct FRigVMTemplate;
diff --git a/Engine/Plugins/Runtime/RigVM/Source/RigVM/Public/RigVMCore/RigVMTemplate.h b/Engine/Plugins/Runtime/RigVM/Source/RigVM/Public/RigVMCore/RigVMTemplate.h
--- a/Engine/Plugins/Runtime/RigVM/Source/RigVM/Public/RigVMCore/RigVMTemplate.h
+++ b/Engine/Plugins/Runtime/RigVM/Source/RigVM/Public/RigVMCore/RigVMTemplate.h
@@ -270,7 +270,8 @@ struct RIGVM_API FRigVMTemplateArgument
 
 	// returns true if this argument is an execute
 	bool IsExecute() const;
-
+	bool IsExecute_NoLock(const FRigVMRegistry_NoLock& InRegistry) const;
+	
 	// returns true if the argument uses an array container
 	EArrayType GetArrayType(const bool bLockRegistry = true) const;
 	EArrayType GetArrayType_NoLock() const;
@@ -295,14 +296,16 @@ protected:
 	TFunction<bool(const TRigVMTypeIndex&)> FilterType;
 	mutable TOptional<EArrayType> CachedArrayType;
 
-	FRigVMTemplateArgument(FProperty* InProperty);
-
 	// constructor from a property. this forces the type to be created
 	FRigVMTemplateArgument(FProperty* InProperty, FRigVMRegistry_NoLock& InRegistry);
 
 	// static make function using a lock to create it
 	static FRigVMTemplateArgument Make(FProperty* InProperty);
 
+	// static make function using a lock to create it
+	static FRigVMTemplateArgument Make_NoLock(FProperty* InProperty);
+	static FRigVMTemplateArgument Make_NoLock(FProperty* InProperty, FRigVMRegistry_NoLock& InRegistry);
+
 	void EnsureValidExecuteType_NoLock(FRigVMRegistry_NoLock& InRegistry);
 	void UpdateTypeToPermutations();
 
