RigVM: Controller side work for SwapFunction

#rnx

diff --git a/Engine/Plugins/Runtime/RigVM/Source/RigVM/Private/RigVMCore/RigVMGraphFunctionDefinition.cpp b/Engine/Plugins/Runtime/RigVM/Source/RigVM/Private/RigVMCore/RigVMGraphFunctionDefinition.cpp
--- a/Engine/Plugins/Runtime/RigVM/Source/RigVM/Private/RigVMCore/RigVMGraphFunctionDefinition.cpp
+++ b/Engine/Plugins/Runtime/RigVM/Source/RigVM/Private/RigVMCore/RigVMGraphFunctionDefinition.cpp
@@ -116,6 +116,24 @@ void FRigVMGraphFunctionHeader::PostDuplicateHost(const FString& InOldPathName,
 	}
 }
 
+FRigVMGraphFunctionHeader FRigVMGraphFunctionHeader::FindGraphFunctionHeader(const FString& InHostPath, const FName& InFunctionName, bool* bOutIsPublic, FString* OutErrorMessage)
+{
+	if(const FRigVMGraphFunctionData* FunctionData = FRigVMGraphFunctionData::FindFunctionData(InHostPath, InFunctionName, bOutIsPublic, OutErrorMessage))
+	{
+		return FunctionData->Header;
+	}
+	return FRigVMGraphFunctionHeader(); 
+}
+
+FRigVMGraphFunctionHeader FRigVMGraphFunctionHeader::FindGraphFunctionHeader(const FRigVMGraphFunctionIdentifier& InIdentifier, bool* bOutIsPublic, FString* OutErrorMessage)
+{
+	if(const FRigVMGraphFunctionData* FunctionData = FRigVMGraphFunctionData::FindFunctionData(InIdentifier, bOutIsPublic, OutErrorMessage))
+	{
+		return FunctionData->Header;
+	}
+	return FRigVMGraphFunctionHeader(); 
+}
+
 bool FRigVMGraphFunctionData::IsMutable() const
 {
 	return Header.IsMutable();
@@ -126,26 +144,98 @@ void FRigVMGraphFunctionData::PostDuplicateHost(const FString& InOldHostPathName
 	Header.PostDuplicateHost(InOldHostPathName, InNewHostPathName);
 }
 
-FRigVMGraphFunctionData* FRigVMGraphFunctionData::FindFunctionData(const FRigVMGraphFunctionIdentifier& InIdentifier, bool* bOutIsPublic)
+FRigVMGraphFunctionData* FRigVMGraphFunctionData::FindFunctionData(const FString& InHostPath, const FName& InFunctionName, bool* bOutIsPublic, FString* OutErrorMessage)
+{
+	FRigVMGraphFunctionHeader InvalidHeader;
+
+	UObject* HostObject = StaticLoadObject(UObject::StaticClass(), NULL, *InHostPath, NULL, LOAD_None, NULL);
+	if (!HostObject)
+	{
+		if(OutErrorMessage)
+		{
+			*OutErrorMessage = FString::Printf(TEXT("Failed to load the Host object %s."), *InHostPath);
+		}
+		return nullptr;
+	}
+
+	IRigVMGraphFunctionHost* FunctionHost = Cast<IRigVMGraphFunctionHost>(HostObject);
+	if (!FunctionHost)
+	{
+		if(OutErrorMessage)
+		{
+			*OutErrorMessage = TEXT("Host object is not a IRigVMGraphFunctionHost.");
+		}
+		return nullptr;
+	}
+
+	FRigVMGraphFunctionStore* FunctionStore = FunctionHost->GetRigVMGraphFunctionStore();
+	if (!FunctionStore)
+	{
+		if(OutErrorMessage)
+		{
+			*OutErrorMessage = TEXT("Host object does not contain a function store.");
+		}
+		return nullptr;
+	}
+
+	FRigVMGraphFunctionData* Data = FunctionStore->FindFunctionByName(InFunctionName, bOutIsPublic);
+	if (!Data)
+	{
+		if(OutErrorMessage)
+		{
+			*OutErrorMessage = FString::Printf(TEXT("Function %s not found in host %s."), *InFunctionName.ToString(), *InHostPath);
+		}
+		return nullptr;
+	}
+
+	return Data;
+}
+
+FRigVMGraphFunctionData* FRigVMGraphFunctionData::FindFunctionData(const FRigVMGraphFunctionIdentifier& InIdentifier, bool* bOutIsPublic, FString* OutErrorMessage)
 {
 	IRigVMGraphFunctionHost* FunctionHost = nullptr;
 	if (UObject* FunctionHostObj = InIdentifier.HostObject.TryLoad())
 	{
 		FunctionHost = Cast<IRigVMGraphFunctionHost>(FunctionHostObj);									
 	}
+	else
+	{
+		if(OutErrorMessage)
+		{
+			*OutErrorMessage = FString::Printf(TEXT("Failed to load the Host object %s."), *InIdentifier.HostObject.ToString());
+		}
+		return nullptr;
+	}
 
 	if (!FunctionHost)
 	{
+		if(OutErrorMessage)
+		{
+			*OutErrorMessage = TEXT("Host object is not a IRigVMGraphFunctionHost.");
+		}
 		return nullptr;
 	}
 
 	FRigVMGraphFunctionStore* FunctionStore = FunctionHost->GetRigVMGraphFunctionStore();
 	if (!FunctionStore)
 	{
+		if(OutErrorMessage)
+		{
+			*OutErrorMessage = TEXT("Host object does not contain a function store.");
+		}
 		return nullptr;
 	}
 
-	return FunctionStore->FindFunction(InIdentifier, bOutIsPublic);
+	if(FRigVMGraphFunctionData* FunctionData = FunctionStore->FindFunction(InIdentifier, bOutIsPublic))
+	{
+		return FunctionData;
+	}
+
+	if(OutErrorMessage)
+	{
+		*OutErrorMessage = FString::Printf(TEXT("Function %s not found in host %s."), *InIdentifier.GetFunctionName(), *InIdentifier.HostObject.ToString());
+	}
+	return nullptr;
 }
 
 
diff --git a/Engine/Plugins/Runtime/RigVM/Source/RigVM/Public/RigVMCore/RigVMGraphFunctionDefinition.h b/Engine/Plugins/Runtime/RigVM/Source/RigVM/Public/RigVMCore/RigVMGraphFunctionDefinition.h
--- a/Engine/Plugins/Runtime/RigVM/Source/RigVM/Public/RigVMCore/RigVMGraphFunctionDefinition.h
+++ b/Engine/Plugins/Runtime/RigVM/Source/RigVM/Public/RigVMCore/RigVMGraphFunctionDefinition.h
@@ -393,6 +393,34 @@ struct RIGVM_API FRigVMGraphFunctionIdentifier
 		return HostObject == Other.HostObject && LibraryNode == Other.LibraryNode;
 	}
 
+	bool IsValid() const
+	{
+		return !HostObject.IsNull() && !LibraryNode.IsNull();
+	}
+
+	FString GetFunctionName() const
+	{
+		if(IsValid())
+		{
+			const FString Path = LibraryNode.ToString();
+			FString NodeName;
+			if(Path.Split(TEXT("."), nullptr, &NodeName, ESearchCase::CaseSensitive, ESearchDir::FromEnd))
+			{
+				return NodeName;
+			}
+		}
+		return FString();
+	}
+	
+	FName GetFunctionFName() const
+	{
+		if(!IsValid())
+		{
+			return NAME_None;
+		}
+		return *GetFunctionName();
+	}
+
 	friend FArchive& operator<<(FArchive& Ar, FRigVMGraphFunctionIdentifier& Data)
 	{
 		Ar << Data.LibraryNode;
@@ -446,7 +474,7 @@ struct RIGVM_API FRigVMGraphFunctionHeader
 
 	bool IsMutable() const;
 
-	bool IsValid() const { return !LibraryPointer.HostObject.IsNull(); }
+	bool IsValid() const { return LibraryPointer.IsValid(); }
 
 	FString GetHash() const
 	{
@@ -510,6 +538,10 @@ struct RIGVM_API FRigVMGraphFunctionHeader
 	}
 
 	void PostDuplicateHost(const FString& InOldPathName, const FString& InNewPathName);
+
+	static FRigVMGraphFunctionHeader FindGraphFunctionHeader(const FString& InHostPath, const FName& InFunctionName, bool* bOutIsPublic = nullptr, FString* OutErrorMessage = nullptr);
+
+	static FRigVMGraphFunctionHeader FindGraphFunctionHeader(const FRigVMGraphFunctionIdentifier& InIdentifier, bool* bOutIsPublic = nullptr, FString* OutErrorMessage = nullptr);
 };
 
 USTRUCT(BlueprintType)
@@ -556,7 +588,9 @@ struct RIGVM_API FRigVMGraphFunctionData
 
 	void PostDuplicateHost(const FString& InOldPathName, const FString& InNewPathName);
 
-	static FRigVMGraphFunctionData* FindFunctionData(const FRigVMGraphFunctionIdentifier& InIdentifier, bool* bOutIsPublic = nullptr);	
+	static FRigVMGraphFunctionData* FindFunctionData(const FString& InHostPath, const FName& InFunctionName, bool* bOutIsPublic = nullptr, FString* OutErrorMessage = nullptr);	
+
+	static FRigVMGraphFunctionData* FindFunctionData(const FRigVMGraphFunctionIdentifier& InIdentifier, bool* bOutIsPublic = nullptr, FString* OutErrorMessage = nullptr);	
 
 	static FString GetArgumentNameFromPinHash(const FString& InPinHash);
 	
diff --git a/Engine/Plugins/Runtime/RigVM/Source/RigVMDeveloper/Private/RigVMModel/RigVMController.cpp b/Engine/Plugins/Runtime/RigVM/Source/RigVMDeveloper/Private/RigVMModel/RigVMController.cpp
--- a/Engine/Plugins/Runtime/RigVM/Source/RigVMDeveloper/Private/RigVMModel/RigVMController.cpp
+++ b/Engine/Plugins/Runtime/RigVM/Source/RigVMDeveloper/Private/RigVMModel/RigVMController.cpp
@@ -11247,6 +11247,33 @@ bool URigVMController::SetExposedPinIndex(const FName& InPinName, int32 InNewInd
 	return true;
 }
 
+FRigVMGraphFunctionHeader URigVMController::FindGraphFunctionHeaderByName(FString InHostPath, FName InFunctionName) const
+{
+	FString ErrorMessage;
+	const FRigVMGraphFunctionHeader Header = FRigVMGraphFunctionHeader::FindGraphFunctionHeader(InHostPath, InFunctionName, nullptr, &ErrorMessage);
+	if(!Header.IsValid() && !ErrorMessage.IsEmpty())
+	{
+		ReportError(ErrorMessage);
+	}
+	return Header;
+}
+
+FRigVMGraphFunctionHeader URigVMController::FindGraphFunctionHeader(FRigVMGraphFunctionIdentifier InFunctionIdentifier) const
+{
+	FString ErrorMessage;
+	const FRigVMGraphFunctionHeader Header = FRigVMGraphFunctionHeader::FindGraphFunctionHeader(InFunctionIdentifier, nullptr, &ErrorMessage);
+	if(!Header.IsValid() && !ErrorMessage.IsEmpty())
+	{
+		ReportError(ErrorMessage);
+	}
+	return Header;
+}
+
+FRigVMGraphFunctionIdentifier URigVMController::FindGraphFunctionIdentifier(FString InHostPath, FName InFunctionName) const
+{
+	return FindGraphFunctionHeaderByName(InHostPath, InFunctionName).LibraryPointer;
+}
+
 URigVMFunctionReferenceNode* URigVMController::AddFunctionReferenceNode(URigVMLibraryNode* InFunctionDefinition, const FVector2D& InNodePosition, const FString& InNodeName, bool bSetupUndoRedo, bool bPrintPythonCommand)
 {
 	if (!IsValidGraph())
@@ -11267,6 +11294,238 @@ URigVMFunctionReferenceNode* URigVMController::AddFunctionReferenceNode(URigVMLi
 	return nullptr;
 }
 
+bool URigVMController::SwapFunctionReferenceByName(const FName& InFunctionReferenceNodeName, const FRigVMGraphFunctionIdentifier& InNewFunctionIdentifier, bool bSetupUndoRedo, bool bPrintPythonCommand)
+{
+	if (!IsValidGraph())
+	{
+		return false;
+	}
+
+	if (!bIsTransacting && !IsGraphEditable())
+	{
+		return false;
+	}
+
+	URigVMFunctionReferenceNode* Node = Cast<URigVMFunctionReferenceNode>(GetGraph()->FindNodeByName(InFunctionReferenceNodeName));
+	if(Node == nullptr)
+	{
+		ReportErrorf(TEXT("Cannot find function reference node '%s'."), *InFunctionReferenceNodeName.ToString());
+		return false;
+	}
+
+	return SwapFunctionReference(Node, InNewFunctionIdentifier, bSetupUndoRedo, bPrintPythonCommand);
+}
+
+bool URigVMController::SwapFunctionReference(URigVMFunctionReferenceNode* InFunctionReferenceNode, const FRigVMGraphFunctionIdentifier& InNewFunctionIdentifier, bool bSetupUndoRedo, bool bPrintPythonCommand)
+{
+	if (!IsValidNodeForGraph(InFunctionReferenceNode))
+	{
+		return false;
+	}
+
+	if (!bIsTransacting && !IsGraphEditable())
+	{
+		return false;
+	}
+
+	if(!InNewFunctionIdentifier.IsValid())
+	{
+		ReportWarning(TEXT("The provided new function identifier is not valid."));
+		return false;
+	}
+
+	if(InFunctionReferenceNode->GetReferencedFunctionHeader().LibraryPointer == InNewFunctionIdentifier)
+	{
+		ReportWarning(TEXT("The old and new function definitions are the same. No swap required."));
+		return false;
+	}
+
+	FRigVMReplaceNodesAction Action;
+	if(bSetupUndoRedo)
+	{
+		Action = FRigVMReplaceNodesAction(this, {InFunctionReferenceNode});
+		Action.SetTitle(TEXT("Swap Function Reference"));
+	}
+
+	const TMap<FString, FPinState> PinStates = GetPinStates(InFunctionReferenceNode);
+	const TArray<FLinkedPath> LinkedPaths = GetLinkedPaths(InFunctionReferenceNode);
+	FastBreakLinkedPaths(LinkedPaths, false);
+
+	InFunctionReferenceNode->Modify();
+	InFunctionReferenceNode->ReferencedFunctionHeader.LibraryPointer = InNewFunctionIdentifier;
+	InFunctionReferenceNode->UpdateFunctionHeaderFromHost();
+	RepopulatePinsOnNode(InFunctionReferenceNode, false, false, false);
+
+	ApplyPinStates(InFunctionReferenceNode, PinStates, {}, false);
+	RestoreLinkedPaths(LinkedPaths, FRestoreLinkedPathSettings(), false);
+
+	TArray<FRigVMExternalVariable> ExternalVariables;
+	if (GetExternalVariablesDelegate.IsBound())
+	{
+		ExternalVariables.Append(GetExternalVariablesDelegate.Execute(GetGraph()));
+	}
+
+	TArray<FName> MappedVariablesToRemove;
+	for(const TPair<FName, FName>& Pair : InFunctionReferenceNode->VariableMap)
+	{
+		const FRigVMExternalVariable* OuterExternalVariable = ExternalVariables.FindByPredicate([Pair](const FRigVMExternalVariable& ExternalVariable) -> bool
+		{
+			return ExternalVariable.Name == Pair.Key;
+		});
+
+		const FRigVMExternalVariable* InnerExternalVariable = InFunctionReferenceNode->ReferencedFunctionHeader.ExternalVariables.FindByPredicate([Pair](const FRigVMExternalVariable& ExternalVariable) -> bool
+		{
+			return ExternalVariable.Name == Pair.Value;
+		});
+
+		if(OuterExternalVariable && InnerExternalVariable)
+		{
+			const TRigVMTypeIndex OuterTypeIndex = OuterExternalVariable->GetTypeIndex();
+			const TRigVMTypeIndex InnerTypeIndex = InnerExternalVariable->GetTypeIndex();
+			if(FRigVMRegistry::Get().CanMatchTypes(OuterTypeIndex, InnerTypeIndex, true))
+			{
+				continue;
+			}
+		}
+
+		MappedVariablesToRemove.Add(Pair.Key);
+	}
+	for(const FName& MappedVariableToRemove : MappedVariablesToRemove)
+	{
+		InFunctionReferenceNode->VariableMap.Remove(MappedVariableToRemove);
+	}
+	
+	if (bSetupUndoRedo)
+	{
+		Action.StoreNode(InFunctionReferenceNode, false);
+		GetActionStack()->AddAction(Action);
+	}
+
+	if (bPrintPythonCommand)
+	{
+		const FString GraphName = GetSchema()->GetSanitizedGraphName(GetGraph()->GetGraphName());
+		const FString FunctionRefNodeName = GetSchema()->GetSanitizedNodeName(InFunctionReferenceNode->GetName());
+		const FString NewFunctionDefinitionName = GetSchema()->GetSanitizedNodeName(InNewFunctionIdentifier.GetFunctionName());
+
+		RigVMPythonUtils::Print(GetSchema()->GetGraphOuterName(GetGraph()), 
+			FString::Printf(TEXT("new_definition = blueprint.get_controller_by_name('%s').find_graph_function_header('%s', '%s')"),
+					*GraphName,
+					*InNewFunctionIdentifier.HostObject.ToString(),
+					*NewFunctionDefinitionName));
+
+		RigVMPythonUtils::Print(GetSchema()->GetGraphOuterName(GetGraph()), 
+			FString::Printf(TEXT("blueprint.get_controller_by_name('%s').swap_function_reference_by_name('%s', new_definition)"),
+					*GraphName,
+					*FunctionRefNodeName));
+	}
+
+	return true;
+}
+
+bool URigVMController::SwapAllFunctionReferences(const FRigVMGraphFunctionIdentifier& InOldFunctionIdentifier, const FRigVMGraphFunctionIdentifier& InNewFunctionIdentifier, bool bSetupUndoRedo, bool bPrintPythonCommand)
+{
+	if (!IsValidGraph())
+	{
+		return false;
+	}
+
+	if (!bIsTransacting && !IsGraphEditable())
+	{
+		return false;
+	}
+
+	if(!InOldFunctionIdentifier.IsValid())
+	{
+		ReportWarning(TEXT("The provided old function identifier is not valid."));
+		return false;
+	}
+
+	if(!InNewFunctionIdentifier.IsValid())
+	{
+		ReportWarning(TEXT("The provided new function identifier is not valid."));
+		return false;
+	}
+
+	if(InOldFunctionIdentifier == InNewFunctionIdentifier)
+	{
+		ReportWarning(TEXT("The old and new function identifier are the same. No swap required."));
+		return false;
+	}
+
+	URigVMGraph* Graph = GetGraph();
+	check(Graph);
+
+	if (Graph->IsA<URigVMFunctionLibrary>())
+	{
+		ReportError(TEXT("Cannot swap function reference nodes within function library graphs."));
+		return false;
+	}
+
+	TArray<URigVMFunctionReferenceNode*> FunctionReferenceNodes;
+	for(TObjectPtr<URigVMNode>& Node : Graph->Nodes)
+	{
+		if(URigVMFunctionReferenceNode* FunctionReferenceNode = Cast<URigVMFunctionReferenceNode>(Node))
+		{
+			if(FunctionReferenceNode->GetReferencedFunctionHeader().LibraryPointer == InOldFunctionIdentifier)
+			{
+				FunctionReferenceNodes.Add(FunctionReferenceNode);
+			}
+		}
+	}
+
+	if(FunctionReferenceNodes.IsEmpty())
+	{
+		return false;
+	}
+	
+	if (bSetupUndoRedo && FunctionReferenceNodes.Num() > 1)
+	{
+		OpenUndoBracket(TEXT("Swap All Function References"));
+	}
+
+	for(URigVMFunctionReferenceNode* Node : FunctionReferenceNodes)
+	{
+		if(!SwapFunctionReference(Node, InNewFunctionIdentifier, bSetupUndoRedo, false))
+		{
+			if (bSetupUndoRedo && FunctionReferenceNodes.Num() > 1)
+			{
+				CancelUndoBracket();
+			}
+			return false;
+		}
+	}
+
+	if (bSetupUndoRedo && FunctionReferenceNodes.Num() > 1)
+	{
+		CloseUndoBracket();
+	}
+
+	if (bPrintPythonCommand)
+	{
+		const FString GraphName = GetSchema()->GetSanitizedGraphName(GetGraph()->GetGraphName());
+		const FString OldFunctionDefinitionName = GetSchema()->GetSanitizedNodeName(InOldFunctionIdentifier.GetFunctionName());
+		const FString NewFunctionDefinitionName = GetSchema()->GetSanitizedNodeName(InNewFunctionIdentifier.GetFunctionName());
+
+		RigVMPythonUtils::Print(GetSchema()->GetGraphOuterName(GetGraph()), 
+			FString::Printf(TEXT("old_definition = blueprint.get_controller_by_name('%s').find_graph_function_header('%s', '%s')"),
+					*GraphName,
+					*InOldFunctionIdentifier.HostObject.ToString(),
+					*OldFunctionDefinitionName));
+
+		RigVMPythonUtils::Print(GetSchema()->GetGraphOuterName(GetGraph()), 
+			FString::Printf(TEXT("new_definition = blueprint.get_controller_by_name('%s').find_graph_function_header('%s', '%s')"),
+					*GraphName,
+					*InNewFunctionIdentifier.HostObject.ToString(),
+					*NewFunctionDefinitionName));
+
+		RigVMPythonUtils::Print(GetSchema()->GetGraphOuterName(GetGraph()), 
+			FString::Printf(TEXT("blueprint.get_controller_by_name('%s').swap_all_function_references(old_definition, new_definition)"),
+					*GraphName));
+	}
+	
+	return true;
+}
+
 URigVMFunctionReferenceNode* URigVMController::AddFunctionReferenceNodeFromDescription(const FRigVMGraphFunctionHeader& InFunctionDefinition, const FVector2D& InNodePosition, const FString& InNodeName, bool bSetupUndoRedo, bool bPrintPythonCommand)
 {
 	if (!IsValidGraph())
@@ -11396,28 +11655,13 @@ URigVMFunctionReferenceNode* URigVMController::AddExternalFunctionReferenceNode(
 		return nullptr;
 	}
 
-	UObject* HostObject = StaticLoadObject(UObject::StaticClass(), NULL, *InHostPath, NULL, LOAD_None, NULL);
-	if (!HostObject)
+	const FRigVMGraphFunctionHeader Header = FindGraphFunctionHeaderByName(InHostPath, InFunctionName);
+	if(!Header.IsValid())
 	{
-		ReportErrorf(TEXT("Failed to load the Host object %s."), *InHostPath);
-		return nullptr;
-	}
-
-	IRigVMGraphFunctionHost* FunctionHost = Cast<IRigVMGraphFunctionHost>(HostObject);
-	if (!FunctionHost)
-	{
-		ReportError(TEXT("Host object is not a IRigVMGraphFunctionHost."));
-		return nullptr;
-	}
-
-	FRigVMGraphFunctionData* Data = FunctionHost->GetRigVMGraphFunctionStore()->FindFunctionByName(InFunctionName);
-	if (!Data)
-	{
-		ReportErrorf(TEXT("Function %s not found in host %s."), *InFunctionName.ToString(), *InHostPath);
 		return nullptr;
 	}
 
-	return AddFunctionReferenceNodeFromDescription(Data->Header, InNodePosition, InNodeName, bSetupUndoRedo, bPrintPythonCommand);
+	return AddFunctionReferenceNodeFromDescription(Header, InNodePosition, InNodeName, bSetupUndoRedo, bPrintPythonCommand);
 }
 
 bool URigVMController::SetRemappedVariable(URigVMFunctionReferenceNode* InFunctionRefNode,
diff --git a/Engine/Plugins/Runtime/RigVM/Source/RigVMDeveloper/Public/RigVMModel/RigVMController.h b/Engine/Plugins/Runtime/RigVM/Source/RigVMDeveloper/Public/RigVMModel/RigVMController.h
--- a/Engine/Plugins/Runtime/RigVM/Source/RigVMDeveloper/Public/RigVMModel/RigVMController.h
+++ b/Engine/Plugins/Runtime/RigVM/Source/RigVMDeveloper/Public/RigVMModel/RigVMController.h
@@ -907,6 +907,15 @@ public:
 	UFUNCTION(BlueprintCallable, Category = RigVMController)
 	bool SetExposedPinIndex(const FName& InPinName, int32 InNewIndex, bool bSetupUndoRedo = true, bool bPrintPythonCommand = false);
 
+	UFUNCTION(BlueprintPure, Category = RigVMController)
+	FRigVMGraphFunctionHeader FindGraphFunctionHeaderByName(FString InHostPath, FName InFunctionName) const;
+
+	UFUNCTION(BlueprintPure, Category = RigVMController)
+	FRigVMGraphFunctionHeader FindGraphFunctionHeader(FRigVMGraphFunctionIdentifier InFunctionIdentifier) const;
+
+	UFUNCTION(BlueprintPure, Category = RigVMController)
+	FRigVMGraphFunctionIdentifier FindGraphFunctionIdentifier(FString InHostPath, FName InFunctionName) const;
+
 	// Adds a function reference / invocation to the graph
 	UFUNCTION(BlueprintCallable, Category = RigVMController)
 	URigVMFunctionReferenceNode* AddFunctionReferenceNodeFromDescription(const FRigVMGraphFunctionHeader& InFunctionDefinition, const FVector2D& InNodePosition = FVector2D::ZeroVector, const
@@ -917,6 +926,15 @@ public:
 	UFUNCTION(BlueprintCallable, Category = RigVMController)
 	URigVMFunctionReferenceNode* AddFunctionReferenceNode(URigVMLibraryNode* InFunctionDefinition, const FVector2D& InNodePosition = FVector2D::ZeroVector, const FString& InNodeName = TEXT(""), bool bSetupUndoRedo = true, bool bPrintPythonCommand = false);
 
+	UFUNCTION(BlueprintCallable, Category = RigVMController)
+	bool SwapFunctionReferenceByName(const FName& InFunctionReferenceNodeName, const FRigVMGraphFunctionIdentifier& InNewFunctionIdentifier, bool bSetupUndoRedo = true, bool bPrintPythonCommand = false);
+
+	UFUNCTION(BlueprintCallable, Category = RigVMController)
+	bool SwapFunctionReference(URigVMFunctionReferenceNode* InFunctionReferenceNode, const FRigVMGraphFunctionIdentifier& InNewFunctionIdentifier, bool bSetupUndoRedo = true, bool bPrintPythonCommand = false);
+
+	UFUNCTION(BlueprintCallable, Category = RigVMController)
+	bool SwapAllFunctionReferences(const FRigVMGraphFunctionIdentifier& InOldFunctionIdentifier, const FRigVMGraphFunctionIdentifier& InNewFunctionIdentifier, bool bSetupUndoRedo = true, bool bPrintPythonCommand = false);
+
 	// Sets the remapped variable on a function reference node
 	UFUNCTION(BlueprintCallable, Category = RigVMController)
     bool SetRemappedVariable(URigVMFunctionReferenceNode* InFunctionRefNode, const FName& InInnerVariableName, const FName& InOuterVariableName, bool bSetupUndoRedo = true);
