[ControlRig] Do not force BlueprintActionDatabase to initialize when removing or unmounting assets.
Reduces situations where object references are held by GC during unmounting.
#rnx

diff --git a/Engine/Plugins/Runtime/RigVM/Source/RigVMDeveloper/Private/RigVMBlueprintUtils.cpp b/Engine/Plugins/Runtime/RigVM/Source/RigVMDeveloper/Private/RigVMBlueprintUtils.cpp
--- a/Engine/Plugins/Runtime/RigVM/Source/RigVMDeveloper/Private/RigVMBlueprintUtils.cpp
+++ b/Engine/Plugins/Runtime/RigVM/Source/RigVMDeveloper/Private/RigVMBlueprintUtils.cpp
@@ -144,9 +144,11 @@ void FRigVMBlueprintUtils::HandleAssetDeleted(const FAssetData& InAssetData)
 	if (InAssetData.GetClass() && InAssetData.GetClass()->IsChildOf(URigVMBlueprint::StaticClass()))
 	{
 		// Make sure any RigVMBlueprint removes any TypeActions related to this asset (e.g. public functions)
-		FBlueprintActionDatabase& ActionDatabase = FBlueprintActionDatabase::Get();
-		ActionDatabase.ClearAssetActions(InAssetData.GetClass());
-		ActionDatabase.RefreshClassActions(InAssetData.GetClass());
+		if (FBlueprintActionDatabase* ActionDatabase = FBlueprintActionDatabase::TryGet())
+		{
+			ActionDatabase->ClearAssetActions(InAssetData.GetClass());
+			ActionDatabase->RefreshClassActions(InAssetData.GetClass());
+		}
 	}
 }
 
