Control Rig: Remove soft object path to library node from function identifier
#rnx

diff --git a/Engine/Plugins/Animation/ControlRig/Source/ControlRigDeveloper/Private/ControlRigBlueprint.cpp b/Engine/Plugins/Animation/ControlRig/Source/ControlRigDeveloper/Private/ControlRigBlueprint.cpp
--- a/Engine/Plugins/Animation/ControlRig/Source/ControlRigDeveloper/Private/ControlRigBlueprint.cpp
+++ b/Engine/Plugins/Animation/ControlRig/Source/ControlRigDeveloper/Private/ControlRigBlueprint.cpp
@@ -1838,7 +1838,7 @@ void UControlRigBlueprint::PatchFunctionsOnLoad()
 			for (const FRigVMGraphFunctionData& FunctionData : Store.PublicFunctions)
 			{
 				BackwardsCompatiblePublicFunctions.Add(FunctionData.Header.Name);
-				URigVMLibraryNode* LibraryNode = Cast<URigVMLibraryNode>(FunctionData.Header.LibraryPointer.LibraryNode.ResolveObject());
+				URigVMLibraryNode* LibraryNode = Cast<URigVMLibraryNode>(FunctionData.Header.LibraryPointer.GetNodeSoftPath().ResolveObject());
 				OldHeaders.Add(LibraryNode, FunctionData.Header);
 			}
 		}
diff --git a/Engine/Plugins/Experimental/AnimNext/Source/AnimNextUncookedOnly/Private/AnimNextRigVMAssetEditorData.cpp b/Engine/Plugins/Experimental/AnimNext/Source/AnimNextUncookedOnly/Private/AnimNextRigVMAssetEditorData.cpp
--- a/Engine/Plugins/Experimental/AnimNext/Source/AnimNextUncookedOnly/Private/AnimNextRigVMAssetEditorData.cpp
+++ b/Engine/Plugins/Experimental/AnimNext/Source/AnimNextUncookedOnly/Private/AnimNextRigVMAssetEditorData.cpp
@@ -249,7 +249,7 @@ void UAnimNextRigVMAssetEditorData::RefreshAllModels(ERigVMLoadType InLoadType)
 				{
 					if (URigVMFunctionReferenceNode* FunctionReferenceNode = Cast<URigVMFunctionReferenceNode>(LibraryNode))
 					{
-						if (FunctionReferenceNode->GetReferencedFunctionHeader().LibraryPointer.LibraryNode.GetLongPackageName() != GetPackage()->GetPathName())
+						if (FunctionReferenceNode->GetReferencedFunctionHeader().LibraryPointer.GetNodeSoftPath().GetLongPackageName() != GetPackage()->GetPathName())
 						{
 							continue;
 						}
diff --git a/Engine/Plugins/Runtime/RigVM/Source/RigVM/Private/RigVMCore/RigVMGraphFunctionDefinition.cpp b/Engine/Plugins/Runtime/RigVM/Source/RigVM/Private/RigVMCore/RigVMGraphFunctionDefinition.cpp
--- a/Engine/Plugins/Runtime/RigVM/Source/RigVM/Private/RigVMCore/RigVMGraphFunctionDefinition.cpp
+++ b/Engine/Plugins/Runtime/RigVM/Source/RigVM/Private/RigVMCore/RigVMGraphFunctionDefinition.cpp
@@ -93,26 +93,31 @@ void FRigVMGraphFunctionHeader::PostDuplicateHost(const FString& InOldPathName,
 	const FString OldPathName = InOldPathName + TEXT(":");
 	const FString NewPathName = InNewPathName + TEXT(":");
 
-	auto ReplacePathName = [InOldPathName, InNewPathName, OldPathName, NewPathName](FSoftObjectPath& InOutObjectPath)
+	auto ReplacePathName = [InOldPathName, InNewPathName, OldPathName, NewPathName](FString& InOutObjectPath)
 	{
-		FString PathName = InOutObjectPath.ToString();
-		if(PathName.Equals(InOldPathName, ESearchCase::CaseSensitive))
+		if(InOutObjectPath.Equals(InOldPathName, ESearchCase::CaseSensitive))
 		{
-			InOutObjectPath = FSoftObjectPath(InNewPathName);
+			InOutObjectPath = InNewPathName;
 		}
-		else if(PathName.StartsWith(OldPathName, ESearchCase::CaseSensitive))
+		else if(InOutObjectPath.StartsWith(OldPathName, ESearchCase::CaseSensitive))
 		{
-			PathName = NewPathName + PathName.Mid(OldPathName.Len());
-			InOutObjectPath = FSoftObjectPath(PathName);
+			InOutObjectPath = NewPathName + InOutObjectPath.Mid(OldPathName.Len());
 		}
 	};
+	
+	auto ReplaceSoftPathName = [InOldPathName, InNewPathName, OldPathName, NewPathName, ReplacePathName](FSoftObjectPath& InOutObjectPath)
+	{
+		FString PathName = InOutObjectPath.ToString();
+		ReplacePathName(PathName);
+		InOutObjectPath = FSoftObjectPath(PathName);
+	};
 
-	ReplacePathName(LibraryPointer.LibraryNode);
-	ReplacePathName(LibraryPointer.HostObject);
+	ReplacePathName(LibraryPointer.LibraryNodePath);
+	ReplaceSoftPathName(LibraryPointer.HostObject);
 	for (TPair<FRigVMGraphFunctionIdentifier, uint32>& Pair : Dependencies)
 	{
-		ReplacePathName(Pair.Key.LibraryNode);
-		ReplacePathName(Pair.Key.HostObject);
+		ReplacePathName(Pair.Key.LibraryNodePath);
+		ReplaceSoftPathName(Pair.Key.HostObject);
 	}
 }
 
diff --git a/Engine/Plugins/Runtime/RigVM/Source/RigVM/Public/RigVMCore/RigVMGraphFunctionDefinition.h b/Engine/Plugins/Runtime/RigVM/Source/RigVM/Public/RigVMCore/RigVMGraphFunctionDefinition.h
--- a/Engine/Plugins/Runtime/RigVM/Source/RigVM/Public/RigVMCore/RigVMGraphFunctionDefinition.h
+++ b/Engine/Plugins/Runtime/RigVM/Source/RigVM/Public/RigVMCore/RigVMGraphFunctionDefinition.h
@@ -370,41 +370,43 @@ struct RIGVM_API FRigVMGraphFunctionIdentifier
 {
 	GENERATED_BODY()
 
+	UPROPERTY(meta=(DeprecatedProperty))
+	FSoftObjectPath LibraryNode_DEPRECATED;
+
 	UPROPERTY(BlueprintReadOnly, VisibleAnywhere, Category=FunctionIdentifier)
-	FSoftObjectPath LibraryNode;
+	FString LibraryNodePath;
 
 	// A path to the IRigVMGraphFunctionHost that stores the function information, and compilation data (e.g. RigVMBlueprintGeneratedClass)
 	UPROPERTY(BlueprintReadOnly, VisibleAnywhere, Category=FunctionIdentifier)
 	FSoftObjectPath HostObject;
 
 	FRigVMGraphFunctionIdentifier()
-		: LibraryNode(nullptr), HostObject(nullptr) {}
+		: LibraryNodePath(FString()), HostObject(nullptr) {}
 	
-	FRigVMGraphFunctionIdentifier(FSoftObjectPath InHostObject, FSoftObjectPath InLibraryNode)
-		: LibraryNode(InLibraryNode), HostObject(InHostObject) {}
+	FRigVMGraphFunctionIdentifier(FSoftObjectPath InHostObject, FString InLibraryNodePath)
+		: LibraryNodePath(InLibraryNodePath), HostObject(InHostObject) {}
 
 	friend uint32 GetTypeHash(const FRigVMGraphFunctionIdentifier& Pointer)
 	{
-		return HashCombine(GetTypeHash(Pointer.LibraryNode), GetTypeHash(Pointer.HostObject));
+		return HashCombine(GetTypeHash(Pointer.LibraryNodePath), GetTypeHash(Pointer.HostObject));
 	}
 
 	bool operator==(const FRigVMGraphFunctionIdentifier& Other) const
 	{
-		return HostObject == Other.HostObject && LibraryNode == Other.LibraryNode;
+		return HostObject == Other.HostObject && LibraryNodePath == Other.LibraryNodePath;
 	}
 
 	bool IsValid() const
 	{
-		return !HostObject.IsNull() && !LibraryNode.IsNull();
+		return !HostObject.IsNull() && !LibraryNodePath.IsEmpty();
 	}
 
 	FString GetFunctionName() const
 	{
 		if(IsValid())
 		{
-			const FString Path = LibraryNode.ToString();
 			FString NodeName;
-			if(Path.Split(TEXT("."), nullptr, &NodeName, ESearchCase::CaseSensitive, ESearchDir::FromEnd))
+			if(LibraryNodePath.Split(TEXT("."), nullptr, &NodeName, ESearchCase::CaseSensitive, ESearchDir::FromEnd))
 			{
 				return NodeName;
 			}
@@ -421,9 +423,26 @@ struct RIGVM_API FRigVMGraphFunctionIdentifier
 		return *GetFunctionName();
 	}
 
+	FSoftObjectPath GetNodeSoftPath() const
+	{
+		return LibraryNodePath;
+	}
+
 	friend FArchive& operator<<(FArchive& Ar, FRigVMGraphFunctionIdentifier& Data)
 	{
-		Ar << Data.LibraryNode;
+		Ar.UsingCustomVersion(FRigVMObjectVersion::GUID);
+		
+		if (Ar.IsLoading() && Ar.CustomVer(FRigVMObjectVersion::GUID) < FRigVMObjectVersion::RemoveLibraryNodeReferenceFromFunctionIdentifier)
+		{
+			FSoftObjectPath SoftPath;
+			Ar << SoftPath;
+			Data.LibraryNodePath = SoftPath.ToString();
+		}
+		else
+		{
+			Ar << Data.LibraryNodePath;
+		}
+		
 		Ar << Data.HostObject;
 		return Ar;
 	}
@@ -435,7 +454,7 @@ struct RIGVM_API FRigVMGraphFunctionHeader
 	GENERATED_BODY()
 
 	FRigVMGraphFunctionHeader()
-		: LibraryPointer(nullptr, nullptr)
+		: LibraryPointer(nullptr, FString())
 		, Name(NAME_None)
 	{}
 
@@ -499,7 +518,7 @@ struct RIGVM_API FRigVMGraphFunctionHeader
 	{
 		FString TooltipStr = FString::Printf(TEXT("%s (%s)\n%s"),
 		*NodeTitle,
-		*LibraryPointer.LibraryNode.GetAssetPathString(),
+		*LibraryPointer.GetNodeSoftPath().GetAssetPathString(),
 		*Description);
 		return FText::FromString(TooltipStr);
 	}
diff --git a/Engine/Plugins/Runtime/RigVM/Source/RigVM/Public/RigVMObjectVersion.h b/Engine/Plugins/Runtime/RigVM/Source/RigVM/Public/RigVMObjectVersion.h
--- a/Engine/Plugins/Runtime/RigVM/Source/RigVM/Public/RigVMObjectVersion.h
+++ b/Engine/Plugins/Runtime/RigVM/Source/RigVM/Public/RigVMObjectVersion.h
@@ -39,6 +39,9 @@ struct RIGVM_API FRigVMObjectVersion
 		// Removing unused tooltip property from frunction header
 		VMRemoveTooltipFromFunctionHeader,
 
+		// Removing library node FSoftObjectPath from FRigVMGraphFunctionIdentifier
+		RemoveLibraryNodeReferenceFromFunctionIdentifier,
+
 		// -----<new versions can be added above this line>-------------------------------------------------
 		VersionPlusOne,
 		LatestVersion = VersionPlusOne - 1
diff --git a/Engine/Plugins/Runtime/RigVM/Source/RigVMDeveloper/Private/EdGraph/RigVMEdGraphNode.cpp b/Engine/Plugins/Runtime/RigVM/Source/RigVMDeveloper/Private/EdGraph/RigVMEdGraphNode.cpp
--- a/Engine/Plugins/Runtime/RigVM/Source/RigVMDeveloper/Private/EdGraph/RigVMEdGraphNode.cpp
+++ b/Engine/Plugins/Runtime/RigVM/Source/RigVMDeveloper/Private/EdGraph/RigVMEdGraphNode.cpp
@@ -81,7 +81,7 @@ FText URigVMEdGraphNode::GetNodeTitle(ENodeTitleType::Type TitleType) const
 				const FRigVMGraphFunctionHeader& ReferencedHeader = FunctionReferenceNode->GetReferencedFunctionHeader();
 				{
 					TSoftObjectPtr<URigVMFunctionReferenceNode> RefNodePtr(FunctionReferenceNode);
-					const FString& PackagePath = ReferencedHeader.LibraryPointer.LibraryNode.GetLongPackageName();
+					const FString& PackagePath = ReferencedHeader.LibraryPointer.GetNodeSoftPath().GetLongPackageName();
 					if(PackagePath != RefNodePtr.GetLongPackageName())
 					{
 						SubTitle = FString::Printf(TEXT("From %s"), *PackagePath);
diff --git a/Engine/Plugins/Runtime/RigVM/Source/RigVMDeveloper/Private/RigVMBlueprint.cpp b/Engine/Plugins/Runtime/RigVM/Source/RigVMDeveloper/Private/RigVMBlueprint.cpp
--- a/Engine/Plugins/Runtime/RigVM/Source/RigVMDeveloper/Private/RigVMBlueprint.cpp
+++ b/Engine/Plugins/Runtime/RigVM/Source/RigVMDeveloper/Private/RigVMBlueprint.cpp
@@ -2091,7 +2091,7 @@ TArray<FString> URigVMBlueprint::GeneratePythonCommands(const FString InNewBluep
 							continue;
 						}
 
-						URigVMLibraryNode* LibraryNode = Cast<URigVMLibraryNode>(Reference->GetReferencedFunctionHeader().LibraryPointer.LibraryNode.ResolveObject());
+						URigVMLibraryNode* LibraryNode = Cast<URigVMLibraryNode>(Reference->GetReferencedFunctionHeader().LibraryPointer.GetNodeSoftPath().ResolveObject());
 						if (!ProcessedGraphs.Contains(LibraryNode->GetContainedGraph()))
 						{
 							bFoundUnprocessedReference = true;
@@ -3469,12 +3469,17 @@ void URigVMBlueprint::PatchFunctionReferencesOnLoad()
 
 				if (FunctionReferenceNode->ReferencedNodePtr_DEPRECATED.IsValid())
 				{
-					FunctionReferenceNode->ReferencedFunctionHeader = FunctionReferenceNode->ReferencedNodePtr_DEPRECATED->GetFunctionHeader();					
+					FunctionReferenceNode->ReferencedFunctionHeader = FunctionReferenceNode->ReferencedNodePtr_DEPRECATED->GetFunctionHeader();
 				}
 				else if (!FunctionReferenceNode->ReferencedNodePtr_DEPRECATED.IsNull())
 				{
 					// At least lets make sure we store the path in the header
-					FunctionReferenceNode->ReferencedFunctionHeader.LibraryPointer.LibraryNode = FunctionReferenceNode->ReferencedNodePtr_DEPRECATED.ToSoftObjectPath();
+					FunctionReferenceNode->ReferencedFunctionHeader.LibraryPointer.LibraryNodePath = FunctionReferenceNode->ReferencedNodePtr_DEPRECATED.ToSoftObjectPath().ToString();
+				}
+
+				if (FunctionReferenceNode->ReferencedFunctionHeader.LibraryPointer.LibraryNode_DEPRECATED.IsValid())
+				{
+					FunctionReferenceNode->ReferencedFunctionHeader.LibraryPointer.LibraryNodePath = FunctionReferenceNode->ReferencedFunctionHeader.LibraryPointer.LibraryNode_DEPRECATED.ToString();
 				}
 			}
 
@@ -3654,7 +3659,7 @@ void URigVMBlueprint::PatchFunctionsOnLoad()
 		for (const FRigVMGraphFunctionData& FunctionData : Store.PublicFunctions)
 		{
 			BackwardsCompatiblePublicFunctions.Add(FunctionData.Header.Name);
-			URigVMLibraryNode* LibraryNode = Cast<URigVMLibraryNode>(FunctionData.Header.LibraryPointer.LibraryNode.ResolveObject());
+			URigVMLibraryNode* LibraryNode = Cast<URigVMLibraryNode>(FunctionData.Header.LibraryPointer.GetNodeSoftPath().ResolveObject());
 			OldHeaders.Add(LibraryNode, FunctionData.Header);
 		}
 	}
diff --git a/Engine/Plugins/Runtime/RigVM/Source/RigVMDeveloper/Private/RigVMCompiler/RigVMCompiler.cpp b/Engine/Plugins/Runtime/RigVM/Source/RigVMDeveloper/Private/RigVMCompiler/RigVMCompiler.cpp
--- a/Engine/Plugins/Runtime/RigVM/Source/RigVMDeveloper/Private/RigVMCompiler/RigVMCompiler.cpp
+++ b/Engine/Plugins/Runtime/RigVM/Source/RigVMDeveloper/Private/RigVMCompiler/RigVMCompiler.cpp
@@ -518,7 +518,7 @@ bool URigVMCompiler::Compile(const FRigVMCompileSettings& InSettings, TArray<URi
 							bool bSuccessfullCompilation = false;
 							if (!CompilationData->IsValid() || CompilationData->RequiresRecompilation())
 							{
-								if (URigVMLibraryNode* LibraryNode = Cast<URigVMLibraryNode>(FunctionData->Header.LibraryPointer.LibraryNode.TryLoad()))
+								if (URigVMLibraryNode* LibraryNode = Cast<URigVMLibraryNode>(FunctionData->Header.LibraryPointer.GetNodeSoftPath().TryLoad()))
 								{
 									IRigVMClientHost* ClientHost = LibraryNode->GetImplementingOuter<IRigVMClientHost>();
 									URigVMController* FunctionController = ClientHost->GetRigVMClient()->GetOrCreateController(LibraryNode->GetLibrary());
@@ -2438,7 +2438,7 @@ bool URigVMCompiler::TraverseInlineFunction(const FRigVMInlineFunctionExprAST* I
 	const FRigVMByteCode& FunctionByteCode = FunctionCompilationData->ByteCode;
 
 	// Bytecode to be inlined should never be aligned
-	checkf(!FunctionByteCode.bByteCodeIsAligned, TEXT("Trying to inline aligned function bytecode %s in package %s"), *FunctionReferenceNode->GetFunctionIdentifier().LibraryNode.ToString(), *GetPackage()->GetPathName());
+	checkf(!FunctionByteCode.bByteCodeIsAligned, TEXT("Trying to inline aligned function bytecode %s in package %s"), *FunctionReferenceNode->GetFunctionIdentifier().LibraryNodePath, *GetPackage()->GetPathName());
 	
 	if (WorkData.bSetupMemory)
 	{
diff --git a/Engine/Plugins/Runtime/RigVM/Source/RigVMDeveloper/Private/RigVMModel/Nodes/RigVMFunctionReferenceNode.cpp b/Engine/Plugins/Runtime/RigVM/Source/RigVMDeveloper/Private/RigVMModel/Nodes/RigVMFunctionReferenceNode.cpp
--- a/Engine/Plugins/Runtime/RigVM/Source/RigVMDeveloper/Private/RigVMModel/Nodes/RigVMFunctionReferenceNode.cpp
+++ b/Engine/Plugins/Runtime/RigVM/Source/RigVMDeveloper/Private/RigVMModel/Nodes/RigVMFunctionReferenceNode.cpp
@@ -43,7 +43,7 @@ bool URigVMFunctionReferenceNode::RequiresVariableRemappingInternal(TArray<FRigV
 	bool bHostedInDifferencePackage = false;
 	
 	FRigVMGraphFunctionIdentifier LibraryPointer = ReferencedFunctionHeader.LibraryPointer;
-	const FString& LibraryPackagePath = LibraryPointer.LibraryNode.GetLongPackageName();
+	const FString& LibraryPackagePath = LibraryPointer.GetNodeSoftPath().GetLongPackageName();
 	const FString& ThisPacakgePath = GetPackage()->GetPathName();
 	bHostedInDifferencePackage = LibraryPackagePath != ThisPacakgePath;
 		
@@ -137,7 +137,7 @@ uint32 URigVMFunctionReferenceNode::GetStructureHash() const
 
 	Hash = HashCombine(Hash, GetTypeHash(ReferencedFunctionHeader.Name.ToString()));
 	Hash = HashCombine(Hash, GetTypeHash(ReferencedFunctionHeader.NodeTitle));
-	Hash = HashCombine(Hash, GetTypeHash(ReferencedFunctionHeader.LibraryPointer.LibraryNode.ToString()));
+	Hash = HashCombine(Hash, GetTypeHash(ReferencedFunctionHeader.LibraryPointer.LibraryNodePath));
 	Hash = HashCombine(Hash, GetTypeHash(ReferencedFunctionHeader.Keywords));
 	Hash = HashCombine(Hash, GetTypeHash(ReferencedFunctionHeader.Description));
 	Hash = HashCombine(Hash, GetTypeHash(ReferencedFunctionHeader.NodeColor));
@@ -216,15 +216,16 @@ bool URigVMFunctionReferenceNode::IsReferencedFunctionHostLoaded() const
 
 bool URigVMFunctionReferenceNode::IsReferencedNodeLoaded() const
 {
-	return ReferencedFunctionHeader.LibraryPointer.LibraryNode.ResolveObject() != nullptr;
+	return ReferencedFunctionHeader.LibraryPointer.GetNodeSoftPath().ResolveObject() != nullptr;
 }
 
 URigVMLibraryNode* URigVMFunctionReferenceNode::LoadReferencedNode() const
 {
-	UObject* LibraryNode = ReferencedFunctionHeader.LibraryPointer.LibraryNode.ResolveObject();
+	FSoftObjectPath SoftObjectPath = ReferencedFunctionHeader.LibraryPointer.GetNodeSoftPath();
+	UObject* LibraryNode = SoftObjectPath.ResolveObject();
 	if (!LibraryNode)
 	{
-		LibraryNode = ReferencedFunctionHeader.LibraryPointer.LibraryNode.TryLoad();
+		LibraryNode = SoftObjectPath.TryLoad();
 	}
 	return Cast<URigVMLibraryNode>(LibraryNode);
 	
diff --git a/Engine/Plugins/Runtime/RigVM/Source/RigVMDeveloper/Private/RigVMModel/Nodes/RigVMLibraryNode.cpp b/Engine/Plugins/Runtime/RigVM/Source/RigVMDeveloper/Private/RigVMModel/Nodes/RigVMLibraryNode.cpp
--- a/Engine/Plugins/Runtime/RigVM/Source/RigVMDeveloper/Private/RigVMModel/Nodes/RigVMLibraryNode.cpp
+++ b/Engine/Plugins/Runtime/RigVM/Source/RigVMDeveloper/Private/RigVMModel/Nodes/RigVMLibraryNode.cpp
@@ -148,7 +148,7 @@ bool URigVMLibraryNode::Contains(URigVMLibraryNode* InContainedNode, bool bRecur
 		if(URigVMFunctionReferenceNode* FunctionReferenceNode = Cast<URigVMFunctionReferenceNode>(ContainedNode))
 		{
 			FRigVMGraphFunctionIdentifier ReferencedNode = FunctionReferenceNode->GetReferencedFunctionHeader().LibraryPointer;
-			if(ReferencedNode.LibraryNode == InContainedNode)
+			if(ReferencedNode.GetNodeSoftPath() == InContainedNode)
 			{
 				return true;
 			}
@@ -219,7 +219,7 @@ FRigVMGraphFunctionIdentifier URigVMLibraryNode::GetFunctionIdentifier() const
 	{
 		Identifier.HostObject = Library->GetFunctionHostObjectPath();
 	}
-	Identifier.LibraryNode = this;
+	Identifier.LibraryNodePath = this->GetPathName();
 	return Identifier;
 }
 
@@ -233,7 +233,7 @@ FRigVMGraphFunctionHeader URigVMLibraryNode::GetFunctionHeader(IRigVMGraphFuncti
 	}
 	else
 	{
-	    Header.LibraryPointer.LibraryNode = this;
+	    Header.LibraryPointer.LibraryNodePath = this->GetPathName();
 	    Header.LibraryPointer.HostObject = Cast<UObject>(InHostObject);
 	}	
     
diff --git a/Engine/Plugins/Runtime/RigVM/Source/RigVMDeveloper/Private/RigVMModel/RigVMBuildData.cpp b/Engine/Plugins/Runtime/RigVM/Source/RigVMDeveloper/Private/RigVMModel/RigVMBuildData.cpp
--- a/Engine/Plugins/Runtime/RigVM/Source/RigVMDeveloper/Private/RigVMModel/RigVMBuildData.cpp
+++ b/Engine/Plugins/Runtime/RigVM/Source/RigVMDeveloper/Private/RigVMModel/RigVMBuildData.cpp
@@ -148,7 +148,7 @@ void URigVMBuildData::RegisterReferencesFromAsset(const FAssetData& InAssetData)
 				ReferenceNodeDataProperty->ImportText_Direct(*ReferenceNodeDataString, &ReferenceNodeDatas, nullptr, EPropertyPortFlags::PPF_None);	
 				for(FRigVMReferenceNodeData& ReferenceNodeData : ReferenceNodeDatas)
 				{
-					if (ReferenceNodeData.ReferencedFunctionIdentifier.LibraryNode.IsValid())
+					if (ReferenceNodeData.ReferencedFunctionIdentifier.GetNodeSoftPath().IsValid())
 					{
 						BuildData->RegisterFunctionReference(ReferenceNodeData.ReferencedFunctionIdentifier, ReferenceNodeData.GetReferenceNodeObjectPath());
 					}
@@ -248,31 +248,31 @@ void URigVMBuildData::RegisterFunctionReference(const FRigVMGraphFunctionIdentif
 
 void URigVMBuildData::RegisterFunctionReference(FRigVMReferenceNodeData InReferenceNodeData)
 {
-	if (InReferenceNodeData.ReferencedFunctionIdentifier.LibraryNode.IsValid())
+	if (InReferenceNodeData.ReferencedFunctionIdentifier.GetNodeSoftPath().IsValid())
 	{
 		return RegisterFunctionReference(InReferenceNodeData.ReferencedFunctionIdentifier, InReferenceNodeData.GetReferenceNodeObjectPath());
 	}
 
-	if (!InReferenceNodeData.ReferencedFunctionIdentifier.LibraryNode.IsValid())
+	if (!InReferenceNodeData.ReferencedFunctionIdentifier.GetNodeSoftPath().IsValid())
 	{
 		InReferenceNodeData.ReferencedFunctionIdentifier = InReferenceNodeData.ReferencedHeader_DEPRECATED.LibraryPointer;
 	}
 
-	if (!InReferenceNodeData.ReferencedFunctionIdentifier.LibraryNode.IsValid())
+	if (!InReferenceNodeData.ReferencedFunctionIdentifier.GetNodeSoftPath().IsValid())
 	{
-		InReferenceNodeData.ReferencedFunctionIdentifier.LibraryNode = InReferenceNodeData.ReferencedFunctionPath_DEPRECATED;
+		InReferenceNodeData.ReferencedFunctionIdentifier.LibraryNodePath = InReferenceNodeData.ReferencedFunctionPath_DEPRECATED;
 	}
 	
-	check(InReferenceNodeData.ReferencedFunctionIdentifier.LibraryNode.IsValid());
+	check(InReferenceNodeData.ReferencedFunctionIdentifier.GetNodeSoftPath().IsValid());
 
-	FSoftObjectPath LibraryNodePath = InReferenceNodeData.ReferencedFunctionIdentifier.LibraryNode;
+	FString LibraryNodePath = InReferenceNodeData.ReferencedFunctionIdentifier.LibraryNodePath;
 	TSoftObjectPtr<URigVMLibraryNode> LibraryNodePtr = TSoftObjectPtr<URigVMLibraryNode>(LibraryNodePath);
 
 	// Try to find a FunctionIdentifier with the same LibraryNodePath
 	bool bFound = false;
 	for (TPair< FRigVMGraphFunctionIdentifier, FRigVMFunctionReferenceArray >& Pair : GraphFunctionReferences)
 	{
-		if (Pair.Key.LibraryNode == LibraryNodePath)
+		if (Pair.Key.LibraryNodePath == LibraryNodePath)
 		{
 			Pair.Value.FunctionReferences.Add(InReferenceNodeData.GetReferenceNodeObjectPath());
 			bFound = true;
diff --git a/Engine/Plugins/Runtime/RigVM/Source/RigVMDeveloper/Private/RigVMModel/RigVMClient.cpp b/Engine/Plugins/Runtime/RigVM/Source/RigVMDeveloper/Private/RigVMModel/RigVMClient.cpp
--- a/Engine/Plugins/Runtime/RigVM/Source/RigVMDeveloper/Private/RigVMModel/RigVMClient.cpp
+++ b/Engine/Plugins/Runtime/RigVM/Source/RigVMDeveloper/Private/RigVMModel/RigVMClient.cpp
@@ -1074,7 +1074,7 @@ void FRigVMClient::HandleGraphModifiedEvent(ERigVMGraphNotifType InNotifType, UR
 					{
 						if (FunctionStore && FunctionHost)
 						{
-							FunctionStore->RemoveFunction(FRigVMGraphFunctionIdentifier (Cast<UObject>(FunctionHost), CollapseNode));
+							FunctionStore->RemoveFunction(FRigVMGraphFunctionIdentifier (Cast<UObject>(FunctionHost), CollapseNode->GetPathName()));
 						}
 					}
 				}
diff --git a/Engine/Plugins/Runtime/RigVM/Source/RigVMDeveloper/Private/RigVMModel/RigVMController.cpp b/Engine/Plugins/Runtime/RigVM/Source/RigVMDeveloper/Private/RigVMModel/RigVMController.cpp
--- a/Engine/Plugins/Runtime/RigVM/Source/RigVMDeveloper/Private/RigVMModel/RigVMController.cpp
+++ b/Engine/Plugins/Runtime/RigVM/Source/RigVMDeveloper/Private/RigVMModel/RigVMController.cpp
@@ -576,7 +576,7 @@ FRigVMPinInfoArray::FRigVMPinInfoArray(const FRigVMGraphFunctionHeader& Function
 		const TRigVMTypeIndex TypeIndex = Registry.GetTypeIndexFromCPPType(FunctionArgument.CPPType.ToString());
 		if (TypeIndex == INDEX_NONE)
 		{
-			InController->ReportErrorf( TEXT("Invalid pin type %s for %s in %s"), *FunctionArgument.CPPType.ToString(), *FunctionHeader.LibraryPointer.LibraryNode.ToString(), *InController->GetPackage()->GetPathName());
+			InController->ReportErrorf( TEXT("Invalid pin type %s for %s in %s"), *FunctionArgument.CPPType.ToString(), *FunctionHeader.LibraryPointer.LibraryNodePath, *InController->GetPackage()->GetPathName());
 		}
 		ensureMsgf(TypeIndex != INDEX_NONE, TEXT("Invalid pin type %s in %s"), *FunctionArgument.CPPType.ToString(), *InController->GetPackage()->GetPathName());
 		(void)AddPin(InController, INDEX_NONE, FunctionArgument.Name, FunctionArgument.Direction, TypeIndex, FunctionArgument.DefaultValue, nullptr, InPreviousPinInfos, true);
@@ -4472,7 +4472,7 @@ TMap<FRigVMGraphFunctionIdentifier, URigVMLibraryNode*> URigVMController::Locali
 		FRigVMGraphFunctionData* FunctionData = FRigVMGraphFunctionData::FindFunctionData(NodeToVisit, &bIsPublic);
 		if (!FunctionData)
 		{
-			ReportAndNotifyErrorf(TEXT("Cannot localize function - could not find function %s in host %s."), *NodeToVisit.LibraryNode.ToString(), *NodeToVisit.HostObject.ToString());
+			ReportAndNotifyErrorf(TEXT("Cannot localize function - could not find function %s in host %s."), *NodeToVisit.LibraryNodePath, *NodeToVisit.HostObject.ToString());
 			return LocalizedFunctions;
 		}
 
@@ -4484,7 +4484,7 @@ TMap<FRigVMGraphFunctionIdentifier, URigVMLibraryNode*> URigVMController::Locali
 
 		if (!bLocalizeDependentPrivateFunctions)
 		{
-			ReportAndNotifyErrorf(TEXT("Cannot localize function - dependency %s is private."), *NodeToVisit.LibraryNode.ToString());
+			ReportAndNotifyErrorf(TEXT("Cannot localize function - dependency %s is private."), *NodeToVisit.LibraryNodePath);
 			return LocalizedFunctions;
 		}
 
@@ -4516,7 +4516,7 @@ TMap<FRigVMGraphFunctionIdentifier, URigVMLibraryNode*> URigVMController::Locali
 	{
 		for(FRigVMGraphFunctionData* FunctionToLocalize : FunctionsToLocalize)
 		{
-			if (URigVMLibraryNode* ReferencedFunction = Cast<URigVMLibraryNode>(FunctionToLocalize->Header.LibraryPointer.LibraryNode.TryLoad()))
+			if (URigVMLibraryNode* ReferencedFunction = Cast<URigVMLibraryNode>(FunctionToLocalize->Header.LibraryPointer.GetNodeSoftPath().TryLoad()))
 			{
 				if (IRigVMClientHost* ClientHost = ReferencedFunction->GetImplementingOuter<IRigVMClientHost>())
 				{
@@ -4529,7 +4529,7 @@ TMap<FRigVMGraphFunctionIdentifier, URigVMLibraryNode*> URigVMController::Locali
 			{
 				URigVMLibraryNode* LocalizedFunction = ThisLibrary->FindFunction(NodeNames[0]);
 				LocalizedFunctions.Add(FunctionToLocalize->Header.LibraryPointer, LocalizedFunction);
-				ThisLibrary->LocalizedFunctions.FindOrAdd(FunctionToLocalize->Header.LibraryPointer.LibraryNode.ToString(), LocalizedFunction);
+				ThisLibrary->LocalizedFunctions.FindOrAdd(FunctionToLocalize->Header.LibraryPointer.LibraryNodePath, LocalizedFunction);
 			}
 		}
 	}
diff --git a/Engine/Plugins/Runtime/RigVM/Source/RigVMDeveloper/Private/RigVMModel/RigVMExternalDependency.cpp b/Engine/Plugins/Runtime/RigVM/Source/RigVMDeveloper/Private/RigVMModel/RigVMExternalDependency.cpp
--- a/Engine/Plugins/Runtime/RigVM/Source/RigVMDeveloper/Private/RigVMModel/RigVMExternalDependency.cpp
+++ b/Engine/Plugins/Runtime/RigVM/Source/RigVMDeveloper/Private/RigVMModel/RigVMExternalDependency.cpp
@@ -62,7 +62,7 @@ void IRigVMExternalDependencyManager::CollectExternalDependencies(TArray<FRigVME
 	check(InHeader);
 	if(InCategory == RigVMGraphFunctionCategory)
 	{
-		OutDependencies.AddUnique({InHeader->LibraryPointer.LibraryNode.ToString(), InCategory});
+		OutDependencies.AddUnique({InHeader->LibraryPointer.LibraryNodePath, InCategory});
 	}
 	for(const FRigVMGraphFunctionArgument& Argument : InHeader->Arguments)
 	{
diff --git a/Engine/Plugins/Runtime/RigVM/Source/RigVMDeveloper/Private/RigVMModel/RigVMFunctionLibrary.cpp b/Engine/Plugins/Runtime/RigVM/Source/RigVMDeveloper/Private/RigVMModel/RigVMFunctionLibrary.cpp
--- a/Engine/Plugins/Runtime/RigVM/Source/RigVMDeveloper/Private/RigVMModel/RigVMFunctionLibrary.cpp
+++ b/Engine/Plugins/Runtime/RigVM/Source/RigVMDeveloper/Private/RigVMModel/RigVMFunctionLibrary.cpp
@@ -122,7 +122,7 @@ void URigVMFunctionLibrary::ForEachReferenceSoftPtr(const FName& InFunctionName,
 
 URigVMLibraryNode* URigVMFunctionLibrary::FindPreviouslyLocalizedFunction(FRigVMGraphFunctionIdentifier InFunctionToLocalize)
 {
-	const FString PathName = InFunctionToLocalize.LibraryNode.ToString();
+	const FString& PathName = InFunctionToLocalize.LibraryNodePath;
 	if(!LocalizedFunctions.Contains((PathName)))
 	{
 		return nullptr;
diff --git a/Engine/Plugins/Runtime/RigVM/Source/RigVMDeveloper/Private/RigVMModel/RigVMSchema.cpp b/Engine/Plugins/Runtime/RigVM/Source/RigVMDeveloper/Private/RigVMModel/RigVMSchema.cpp
--- a/Engine/Plugins/Runtime/RigVM/Source/RigVMDeveloper/Private/RigVMModel/RigVMSchema.cpp
+++ b/Engine/Plugins/Runtime/RigVM/Source/RigVMDeveloper/Private/RigVMModel/RigVMSchema.cpp
@@ -188,7 +188,7 @@ bool URigVMSchema::SupportsGraphFunction(URigVMController* InController, const F
 	URigVMLibraryNode* ParentLibraryNode = Cast<URigVMLibraryNode>(Graph->GetOuter());
 	while (ParentLibraryNode)
 	{
-		if (TSoftObjectPtr<UObject>(ParentLibraryNode).ToSoftObjectPath() == InGraphFunction->LibraryPointer.LibraryNode)
+		if (TSoftObjectPtr<UObject>(ParentLibraryNode).ToSoftObjectPath() == InGraphFunction->LibraryPointer.GetNodeSoftPath())
 		{
 			static const FString Message = TEXT("You cannot place functions inside of itself or an indirect recursion.");
 			InController->ReportError(Message);
diff --git a/Engine/Plugins/Runtime/RigVM/Source/RigVMEditor/Private/EdGraph/NodeSpawners/RigVMEdGraphFunctionRefNodeSpawner.cpp b/Engine/Plugins/Runtime/RigVM/Source/RigVMEditor/Private/EdGraph/NodeSpawners/RigVMEdGraphFunctionRefNodeSpawner.cpp
--- a/Engine/Plugins/Runtime/RigVM/Source/RigVMEditor/Private/EdGraph/NodeSpawners/RigVMEdGraphFunctionRefNodeSpawner.cpp
+++ b/Engine/Plugins/Runtime/RigVM/Source/RigVMEditor/Private/EdGraph/NodeSpawners/RigVMEdGraphFunctionRefNodeSpawner.cpp
@@ -378,7 +378,7 @@ bool URigVMEdGraphFunctionRefNodeSpawner::IsTemplateNodeFilteredOut(FBlueprintAc
 			}
 		}
 	}
-	const FString ReferencedAssetObjectPathString = ReferencedPublicFunctionHeader.LibraryPointer.LibraryNode.GetAssetName();
+	const FString ReferencedAssetObjectPathString = ReferencedPublicFunctionHeader.LibraryPointer.GetNodeSoftPath().GetAssetName();
 	for (UBlueprint* Blueprint : Filter.Context.Blueprints)
 	{
 		if(Blueprint->GetPathName() == ReferencedAssetObjectPathString)
diff --git a/Engine/Plugins/Runtime/RigVM/Source/RigVMEditor/Private/RigVMEditorModule.cpp b/Engine/Plugins/Runtime/RigVM/Source/RigVMEditor/Private/RigVMEditorModule.cpp
--- a/Engine/Plugins/Runtime/RigVM/Source/RigVMEditor/Private/RigVMEditorModule.cpp
+++ b/Engine/Plugins/Runtime/RigVM/Source/RigVMEditor/Private/RigVMEditorModule.cpp
@@ -466,7 +466,7 @@ void FRigVMEditorModule::GetInstanceActions(URigVMBlueprint* RigVMBlueprint, FBl
 				// Avoid adding functions that are already added by the GetTypeActions functions (public functions that are already saved into the blueprint tag)
 				if (RigVMBlueprint->PublicGraphFunctions.ContainsByPredicate([LocalLibrarySoftPath, Function](const FRigVMGraphFunctionHeader& Header) -> bool
 				{
-					return FRigVMGraphFunctionIdentifier(LocalLibrarySoftPath, Function) == Header.LibraryPointer;
+					return FRigVMGraphFunctionIdentifier(LocalLibrarySoftPath, Function->GetPathName()) == Header.LibraryPointer;
 				}))
 				{
 					continue;
@@ -785,9 +785,9 @@ void FRigVMEditorModule::GetNodeVariablesContextMenuActions(IRigVMClientHost* Ri
 		if (URigVMFunctionReferenceNode* FunctionReferenceNode = Cast<URigVMFunctionReferenceNode>(EdGraphNode->GetModelNode()))
 		{
 			TSoftObjectPtr<URigVMFunctionReferenceNode> RefPtr(FunctionReferenceNode);
-			if(RefPtr.GetLongPackageName() != FunctionReferenceNode->GetReferencedFunctionHeader().LibraryPointer.LibraryNode.GetLongPackageName())
+			if(RefPtr.GetLongPackageName() != FunctionReferenceNode->GetReferencedFunctionHeader().LibraryPointer.GetNodeSoftPath().GetLongPackageName())
 			{
-				if(!FunctionReferenceNode->IsFullyRemapped() && FunctionReferenceNode->GetReferencedFunctionHeader().LibraryPointer.LibraryNode.ResolveObject())
+				if(!FunctionReferenceNode->IsFullyRemapped() && FunctionReferenceNode->GetReferencedFunctionHeader().LibraryPointer.GetNodeSoftPath().ResolveObject())
 				{
 					FToolMenuSection& VariablesSection = Menu->AddSection("RigVMEditorContextMenuVariables", LOCTEXT("Variables", "Variables"));
 					VariablesSection.AddMenuEntry(
@@ -928,7 +928,7 @@ void FRigVMEditorModule::GetNodeOrganizationContextMenuActions(IRigVMClientHost*
 		if (URigVMFunctionReferenceNode* FunctionReferenceNode = Cast<URigVMFunctionReferenceNode>(EdGraphNode->GetModelNode()))
 		{
 			TSoftObjectPtr<URigVMFunctionReferenceNode> RefPtr(FunctionReferenceNode);
-			if(RefPtr.GetLongPackageName() != FunctionReferenceNode->GetReferencedFunctionHeader().LibraryPointer.LibraryNode.GetLongPackageName())
+			if(RefPtr.GetLongPackageName() != FunctionReferenceNode->GetReferencedFunctionHeader().LibraryPointer.GetNodeSoftPath().GetLongPackageName())
 			{
 				OrganizationSection.AddMenuEntry(
 				   "Localize Function",
diff --git a/Engine/Plugins/Runtime/RigVM/Source/RigVMEditor/Private/Widgets/SRigVMGraphFunctionLocalizationWidget.cpp b/Engine/Plugins/Runtime/RigVM/Source/RigVMEditor/Private/Widgets/SRigVMGraphFunctionLocalizationWidget.cpp
--- a/Engine/Plugins/Runtime/RigVM/Source/RigVMEditor/Private/Widgets/SRigVMGraphFunctionLocalizationWidget.cpp
+++ b/Engine/Plugins/Runtime/RigVM/Source/RigVMEditor/Private/Widgets/SRigVMGraphFunctionLocalizationWidget.cpp
@@ -27,7 +27,7 @@ SRigVMGraphFunctionLocalizationItem::SRigVMGraphFunctionLocalizationItem(const F
 
 	if (FRigVMGraphFunctionData* FunctionData = FRigVMGraphFunctionData::FindFunctionData(InFunction, &bIsPublic))
 	{
-		OuterName = InFunction.LibraryNode.GetAssetPathString();
+		OuterName = InFunction.GetNodeSoftPath().GetAssetPathString();
 		FunctionName = FunctionData->Header.Name.ToString();
 		DisplayText = FText::FromString(FString::Printf(TEXT("%s :: %s"), *OuterName, *FunctionName));
 		ToolTipText = bIsPublic ?
