Adapted RigVM Math Customizations to be able to use them at AnimNext, by changing object iteration for customized memory iteration.
Also removed a explicit call to Object->PostEditChangeChainProperty, as InPropertyHandle->NotifyPostChange is making that work inside the function call

diff --git a/Engine/Plugins/Runtime/RigVM/Source/RigVMEditor/Private/Editor/RigVMGraphDetailCustomization.cpp b/Engine/Plugins/Runtime/RigVM/Source/RigVMEditor/Private/Editor/RigVMGraphDetailCustomization.cpp
--- a/Engine/Plugins/Runtime/RigVM/Source/RigVMEditor/Private/Editor/RigVMGraphDetailCustomization.cpp
+++ b/Engine/Plugins/Runtime/RigVM/Source/RigVMEditor/Private/Editor/RigVMGraphDetailCustomization.cpp
@@ -1919,7 +1919,7 @@ void FRigVMGraphEnumDetailCustomization::CustomizeChildren(TSharedRef<IPropertyH
 
 void FRigVMGraphEnumDetailCustomization::HandleControlEnumChanged(TSharedPtr<FString> InEnumPath, ESelectInfo::Type InSelectType, TSharedRef<IPropertyHandle> InPropertyHandle)
 {
-	if (ObjectsBeingCustomized.IsEmpty())
+	if (ObjectsBeingCustomized.IsEmpty() && StructsBeingCustomized.IsEmpty())
 	{
 		return;
 	}
@@ -1932,18 +1932,6 @@ void FRigVMGraphEnumDetailCustomization::HandleControlEnumChanged(TSharedPtr<FSt
 		return;
 	}
 	
-	TArray<UObject*> ObjectsView;
-	for(int32 Index = 0; Index < ObjectsBeingCustomized.Num(); Index++)
-	{
-		const TWeakObjectPtr<UObject>& Object = ObjectsBeingCustomized[Index];
-		if (Object.Get())
-		{
-			ObjectsView.Add(Object.Get());
-		}
-	}
-	FPropertyChangedEvent PropertyChangedEvent(InPropertyHandle->GetProperty(), EPropertyChangeType::ValueSet, ObjectsView);
-	FPropertyChangedChainEvent PropertyChangedChainEvent(PropertyChain, PropertyChangedEvent);
-
 	URigVMController* Controller = nullptr;
 	if(BlueprintBeingCustomized && GraphBeingCustomized)
 	{
@@ -1951,12 +1939,15 @@ void FRigVMGraphEnumDetailCustomization::HandleControlEnumChanged(TSharedPtr<FSt
 		Controller->OpenUndoBracket(FString::Printf(TEXT("Set %s"), *InPropertyHandle->GetProperty()->GetName()));
 	}
 
-	for(int32 Index = 0; Index < ObjectsBeingCustomized.Num(); Index++)
+	const EPropertyChangeType::Type ChangeType = EPropertyChangeType::ValueSet;
+
+	const TArray<uint8*> AllMemoryBeingCustomized = GetMemoryBeingCustomized();
+	for (int32 Index = 0; Index < AllMemoryBeingCustomized.Num(); Index++)
 	{
-		const TWeakObjectPtr<UObject>& Object = ObjectsBeingCustomized[Index];
-		if(Object.Get() && InPropertyHandle->IsValidHandle())
+		const uint8* Memory = AllMemoryBeingCustomized[Index];
+		if (Memory != nullptr && InPropertyHandle->IsValidHandle())
 		{
-			UEnum** CurrentEnum = ContainerMemoryBlockToEnumPtr((uint8*)Object.Get(), PropertyChain, PropertyArrayIndices);
+			UEnum** CurrentEnum = ContainerMemoryBlockToEnumPtr((uint8*)Memory, PropertyChain, PropertyArrayIndices);
 			if (CurrentEnum)
 			{
 				const UEnum* PreviousEnum = *CurrentEnum;
@@ -1964,8 +1955,7 @@ void FRigVMGraphEnumDetailCustomization::HandleControlEnumChanged(TSharedPtr<FSt
 
 				if (PreviousEnum != *CurrentEnum)
 				{
-					Object->PostEditChangeChainProperty(PropertyChangedChainEvent);
-					InPropertyHandle->NotifyPostChange(PropertyChangedEvent.ChangeType);
+					InPropertyHandle->NotifyPostChange(ChangeType);
 				}
 			}
 		}
@@ -2224,18 +2214,6 @@ void FRigVMGraphMathTypeDetailCustomization::ConfigureTransformWidgetArgs(TShare
 			return;
 		}
 		
-		TArray<UObject*> ObjectsView;
-		for(int32 Index = 0; Index < ObjectsBeingCustomized.Num(); Index++)
-		{
-			const TWeakObjectPtr<UObject>& Object = ObjectsBeingCustomized[Index];
-			if (Object.Get())
-			{
-				ObjectsView.Add(Object.Get());
-			}
-		}
-		FPropertyChangedEvent PropertyChangedEvent(InPropertyHandle->GetProperty(), bIsCommit ? EPropertyChangeType::ValueSet : EPropertyChangeType::Interactive, ObjectsView);
-		FPropertyChangedChainEvent PropertyChangedChainEvent(PropertyChain, PropertyChangedEvent);
-
 		URigVMController* Controller = nullptr;
 		if(BlueprintBeingCustomized && GraphBeingCustomized)
 		{
@@ -2246,14 +2224,17 @@ void FRigVMGraphMathTypeDetailCustomization::ConfigureTransformWidgetArgs(TShare
 			}
 		}
 
-		for(int32 Index = 0; Index < ObjectsBeingCustomized.Num(); Index++)
+		const EPropertyChangeType::Type ChangeType = bIsCommit ? EPropertyChangeType::ValueSet : EPropertyChangeType::Interactive;
+
+		const TArray<uint8*> AllMemoryBeingCustomized = GetMemoryBeingCustomized();
+		for (int32 Index = 0; Index < AllMemoryBeingCustomized.Num(); Index++)
 		{
-			const TWeakObjectPtr<UObject>& Object = ObjectsBeingCustomized[Index];
-			if(Object.Get() && InPropertyHandle->IsValidHandle())
+			const uint8* Memory = AllMemoryBeingCustomized[Index];
+			if (Memory != nullptr && InPropertyHandle->IsValidHandle())
 			{
-				TransformType& Transform = ContainerMemoryBlockToValueRef<TransformType>((uint8*)Object.Get(), Identity, PropertyChain, PropertyArrayIndices);
+				TransformType& Transform = ContainerMemoryBlockToValueRef<TransformType>((uint8*)Memory, Identity, PropertyChain, PropertyArrayIndices);
 				TransformType PreviousTransform = Transform;
-				
+
 				SAdvancedTransformInputBox<TransformType>::ApplyNumericValueChange(
 					Transform,
 					InValue,
@@ -2261,10 +2242,9 @@ void FRigVMGraphMathTypeDetailCustomization::ConfigureTransformWidgetArgs(TShare
 					InRotationRepresentation,
 					InSubComponent);
 
-				if(!PreviousTransform.Equals(Transform))
+				if (!PreviousTransform.Equals(Transform))
 				{
-					Object->PostEditChangeChainProperty(PropertyChangedChainEvent);
-					InPropertyHandle->NotifyPostChange(PropertyChangedEvent.ChangeType);
+					InPropertyHandle->NotifyPostChange(ChangeType);
 				}
 			}
 		}
@@ -2296,7 +2276,7 @@ void FRigVMGraphMathTypeDetailCustomization::ConfigureTransformWidgetArgs(TShare
 
 	WidgetArgs.OnResetToDefault_Lambda([this, DefaultValue, InPropertyHandle](ESlateTransformComponent::Type InTransformComponent)
 	{
-		if (ObjectsBeingCustomized.IsEmpty())
+		if (ObjectsBeingCustomized.IsEmpty() && StructsBeingCustomized.IsEmpty())
 		{
 			return;
 		}
@@ -2319,24 +2299,15 @@ void FRigVMGraphMathTypeDetailCustomization::ConfigureTransformWidgetArgs(TShare
 			}
 		}
 
-		TArray<UObject*> ObjectsView;
-		for(int32 Index = 0; Index < ObjectsBeingCustomized.Num(); Index++)
-		{
-			const TWeakObjectPtr<UObject>& Object = ObjectsBeingCustomized[Index];
-			if (Object.Get())
-			{
-				ObjectsView.Add(Object.Get());
-			}
-		}
-		FPropertyChangedEvent PropertyChangedEvent(InPropertyHandle->GetProperty(), EPropertyChangeType::ValueSet, ObjectsView);
-		FPropertyChangedChainEvent PropertyChangedChainEvent(PropertyChain, PropertyChangedEvent);
-		
-		for(int32 Index = 0; Index < ObjectsBeingCustomized.Num(); Index++)
+		const EPropertyChangeType::Type ChangeType = EPropertyChangeType::ValueSet;
+
+		const TArray<uint8*> AllMemoryBeingCustomized = GetMemoryBeingCustomized();
+		for (int32 Index = 0; Index < AllMemoryBeingCustomized.Num(); Index++)
 		{
-			const TWeakObjectPtr<UObject>& Object = ObjectsBeingCustomized[Index];
-			if(Object.Get() && InPropertyHandle->IsValidHandle())
+			const uint8* Memory = AllMemoryBeingCustomized[Index];
+			if (Memory != nullptr && InPropertyHandle->IsValidHandle())
 			{
-				TransformType& Transform = ContainerMemoryBlockToValueRef<TransformType>((uint8*)Object.Get(), Identity, PropertyChain, PropertyArrayIndices);
+				TransformType& Transform = ContainerMemoryBlockToValueRef<TransformType>((uint8*)Memory, Identity, PropertyChain, PropertyArrayIndices);
 				TransformType PreviousTransform = Transform;
 
 				switch(InTransformComponent)
@@ -2363,12 +2334,10 @@ void FRigVMGraphMathTypeDetailCustomization::ConfigureTransformWidgetArgs(TShare
 						break;
 					}
 				}
-
 				
 				if(!PreviousTransform.Equals(Transform))
 				{
-					Object->PostEditChangeChainProperty(PropertyChangedChainEvent);
-					InPropertyHandle->NotifyPostChange(PropertyChangedEvent.ChangeType);
+					InPropertyHandle->NotifyPostChange(ChangeType);
 				}
 			}
 		}
@@ -2447,7 +2416,7 @@ void FRigVMGraphMathTypeDetailCustomization::ConfigureTransformWidgetArgs(TShare
 			return;
 		}
 
-		if (ObjectsBeingCustomized.IsEmpty())
+		if (ObjectsBeingCustomized.IsEmpty() && StructsBeingCustomized.IsEmpty())
 		{
 			return;
 		}
@@ -2461,18 +2430,6 @@ void FRigVMGraphMathTypeDetailCustomization::ConfigureTransformWidgetArgs(TShare
 			return;
 		}
 
-		TArray<UObject*> ObjectsView;
-		for(int32 Index = 0; Index < ObjectsBeingCustomized.Num(); Index++)
-		{
-			const TWeakObjectPtr<UObject>& Object = ObjectsBeingCustomized[Index];
-			if (Object.Get())
-			{
-				ObjectsView.Add(Object.Get());
-			}
-		}
-		FPropertyChangedEvent PropertyChangedEvent(InPropertyHandle->GetProperty(), EPropertyChangeType::ValueSet, ObjectsView);
-		FPropertyChangedChainEvent PropertyChangedChainEvent(PropertyChain, PropertyChangedEvent);
-
 		URigVMController* Controller = nullptr;
 		if(BlueprintBeingCustomized && GraphBeingCustomized)
 		{
@@ -2480,11 +2437,15 @@ void FRigVMGraphMathTypeDetailCustomization::ConfigureTransformWidgetArgs(TShare
 			Controller->OpenUndoBracket(FString::Printf(TEXT("Set %s"), *InPropertyHandle->GetProperty()->GetName()));
 		}
 		
-		for(const TWeakObjectPtr<UObject>& Object : ObjectsBeingCustomized)
+		const EPropertyChangeType::Type ChangeType = EPropertyChangeType::ValueSet;
+
+		const TArray<uint8*> AllMemoryBeingCustomized = GetMemoryBeingCustomized();
+		for (int32 Index = 0; Index < AllMemoryBeingCustomized.Num(); Index++)
 		{
-			if(Object.Get() && InPropertyHandle->IsValidHandle())
+			const uint8* Memory = AllMemoryBeingCustomized[Index];
+			if (Memory != nullptr && InPropertyHandle->IsValidHandle())
 			{
-				TransformType& Transform = ContainerMemoryBlockToValueRef<TransformType>((uint8*)Object.Get(), Identity, PropertyChain, PropertyArrayIndices);
+				TransformType& Transform = ContainerMemoryBlockToValueRef<TransformType>((uint8*)Memory, Identity, PropertyChain, PropertyArrayIndices);
 				const TransformType PreviousTransform = Transform;
 
 				// Apply the new value
@@ -2545,8 +2506,7 @@ void FRigVMGraphMathTypeDetailCustomization::ConfigureTransformWidgetArgs(TShare
 				
 					if(ErrorPipe.NumErrors == 0 && !PreviousTransform.Equals(Transform))
 					{
-						Object->PostEditChangeChainProperty(PropertyChangedChainEvent);
-						InPropertyHandle->NotifyPostChange(PropertyChangedEvent.ChangeType);
+						InPropertyHandle->NotifyPostChange(ChangeType);
 					}
 				}
 			}
diff --git a/Engine/Plugins/Runtime/RigVM/Source/RigVMEditor/Public/Editor/RigVMGraphDetailCustomization.h b/Engine/Plugins/Runtime/RigVM/Source/RigVMEditor/Public/Editor/RigVMGraphDetailCustomization.h
--- a/Engine/Plugins/Runtime/RigVM/Source/RigVMEditor/Public/Editor/RigVMGraphDetailCustomization.h
+++ b/Engine/Plugins/Runtime/RigVM/Source/RigVMEditor/Public/Editor/RigVMGraphDetailCustomization.h
@@ -536,11 +536,11 @@ protected:
 	template<typename VectorType, typename NumericType>
 	void OnVectorComponentChanged(TSharedRef<class IPropertyHandle> InPropertyHandle, int32 InComponent, NumericType InValue, bool bIsCommit, ETextCommit::Type InCommitType = ETextCommit::Default)
 	{
-		if (ObjectsBeingCustomized.IsEmpty())
+		if (ObjectsBeingCustomized.IsEmpty() && StructsBeingCustomized.IsEmpty())
 		{
 			return;
 		}
-		
+
 		FEditPropertyChain PropertyChain;
 		TArray<int32> PropertyArrayIndices;
 		bool bEnabled;
@@ -548,48 +548,38 @@ protected:
 		{
 			return;
 		}
-	
-		TArray<UObject*> ObjectsView;
-		for(int32 Index = 0; Index < ObjectsBeingCustomized.Num(); Index++)
-		{
-			const TWeakObjectPtr<UObject>& Object = ObjectsBeingCustomized[Index];
-			if (Object.Get())
-			{
-				ObjectsView.Add(Object.Get());
-			}
-		}
-		FPropertyChangedEvent PropertyChangedEvent(InPropertyHandle->GetProperty(), bIsCommit ? EPropertyChangeType::ValueSet : EPropertyChangeType::Interactive, ObjectsView);
-		FPropertyChangedChainEvent PropertyChangedChainEvent(PropertyChain, PropertyChangedEvent);
 
 		URigVMController* Controller = nullptr;
-		if(BlueprintBeingCustomized && GraphBeingCustomized)
+		if (BlueprintBeingCustomized && GraphBeingCustomized)
 		{
 			Controller = BlueprintBeingCustomized->GetController(GraphBeingCustomized);
-			if(bIsCommit)
+			if (bIsCommit)
 			{
 				Controller->OpenUndoBracket(FString::Printf(TEXT("Set %s"), *InPropertyHandle->GetProperty()->GetName()));
 			}
 		}
 
-		for(int32 Index = 0; Index < ObjectsBeingCustomized.Num(); Index++)
+		const EPropertyChangeType::Type ChangeType = bIsCommit ? EPropertyChangeType::ValueSet : EPropertyChangeType::Interactive;
+
+		const TArray<uint8*> AllMemoryBeingCustomized = GetMemoryBeingCustomized();
+		for (int32 Index = 0; Index < AllMemoryBeingCustomized.Num(); Index++)
 		{
-			const TWeakObjectPtr<UObject>& Object = ObjectsBeingCustomized[Index];
-			if(Object.Get() && InPropertyHandle->IsValidHandle())
+			const uint8* Memory = AllMemoryBeingCustomized[Index];
+			if (Memory != nullptr && InPropertyHandle->IsValidHandle())
 			{
 				static VectorType ZeroVector = VectorType();
-				VectorType& Vector = ContainerMemoryBlockToValueRef<VectorType>((uint8*)Object.Get(), ZeroVector, PropertyChain, PropertyArrayIndices);
+				VectorType& Vector = ContainerMemoryBlockToValueRef<VectorType>((uint8*)Memory, ZeroVector, PropertyChain, PropertyArrayIndices);
 				VectorType PreviousVector = Vector;
 				Vector[InComponent] = InValue;
-					
-				if(!PreviousVector.Equals(Vector))
+
+				if (!PreviousVector.Equals(Vector))
 				{
-					Object->PostEditChangeChainProperty(PropertyChangedChainEvent);
-					InPropertyHandle->NotifyPostChange(PropertyChangedEvent.ChangeType);
+					InPropertyHandle->NotifyPostChange(ChangeType);
 				}
 			}
 		}
 
-		if(Controller && bIsCommit)
+		if (Controller && bIsCommit)
 		{
 			Controller->CloseUndoBracket();
 		}
@@ -654,7 +644,7 @@ protected:
 	template<typename RotationType>
 	void OnRotationChanged(TSharedRef<class IPropertyHandle> InPropertyHandle, RotationType InValue, bool bIsCommit, ETextCommit::Type InCommitType = ETextCommit::Default)
 	{
-		if (ObjectsBeingCustomized.IsEmpty())
+		if (ObjectsBeingCustomized.IsEmpty() && StructsBeingCustomized.IsEmpty())
 		{
 			return;
 		}
@@ -667,18 +657,6 @@ protected:
         	return;
         }
 	
-		TArray<UObject*> ObjectsView;
-		for(int32 Index = 0; Index < ObjectsBeingCustomized.Num(); Index++)
-		{
-			const TWeakObjectPtr<UObject>& Object = ObjectsBeingCustomized[Index];
-			if (Object.Get())
-			{
-				ObjectsView.Add(Object.Get());
-			}
-		}
-		FPropertyChangedEvent PropertyChangedEvent(InPropertyHandle->GetProperty(), bIsCommit ? EPropertyChangeType::ValueSet : EPropertyChangeType::Interactive, ObjectsView);
-		FPropertyChangedChainEvent PropertyChangedChainEvent(PropertyChain, PropertyChangedEvent);
-
 		URigVMController* Controller = nullptr;
 		if(BlueprintBeingCustomized && GraphBeingCustomized)
 		{
@@ -689,20 +667,22 @@ protected:
 			}
 		}
 
-		for(int32 Index = 0; Index < ObjectsBeingCustomized.Num(); Index++)
+		const EPropertyChangeType::Type ChangeType = bIsCommit ? EPropertyChangeType::ValueSet : EPropertyChangeType::Interactive;
+
+		const TArray<uint8*> AllMemoryBeingCustomized = GetMemoryBeingCustomized();
+		for (int32 Index = 0; Index < AllMemoryBeingCustomized.Num(); Index++)
 		{
-			const TWeakObjectPtr<UObject>& Object = ObjectsBeingCustomized[Index];
-			if(Object.Get() && InPropertyHandle->IsValidHandle())
+			const uint8* Memory = AllMemoryBeingCustomized[Index];
+			if (Memory != nullptr && InPropertyHandle->IsValidHandle())
 			{
 				static RotationType ZeroRotation = RotationType();
-				RotationType& Rotation = ContainerMemoryBlockToValueRef<RotationType>((uint8*)Object.Get(), ZeroRotation, PropertyChain, PropertyArrayIndices);
+				RotationType& Rotation = ContainerMemoryBlockToValueRef<RotationType>((uint8*)Memory, ZeroRotation, PropertyChain, PropertyArrayIndices);
 				RotationType PreviousRotation = Rotation;
 				Rotation = InValue;
-					
-				if(!PreviousRotation.Equals(Rotation))
+
+				if (!PreviousRotation.Equals(Rotation))
 				{
-					Object->PostEditChangeChainProperty(PropertyChangedChainEvent);
-					InPropertyHandle->NotifyPostChange(PropertyChangedEvent.ChangeType);
+					InPropertyHandle->NotifyPostChange(ChangeType);
 				}
 			}
 		}
