RigVM: Introduce notion of default value type and retrieval for original default value

#rnx

diff --git a/Engine/Plugins/Animation/ControlRig/Source/ControlRigDeveloper/Private/ControlRigBlueprint.cpp b/Engine/Plugins/Animation/ControlRig/Source/ControlRigDeveloper/Private/ControlRigBlueprint.cpp
--- a/Engine/Plugins/Animation/ControlRig/Source/ControlRigDeveloper/Private/ControlRigBlueprint.cpp
+++ b/Engine/Plugins/Animation/ControlRig/Source/ControlRigDeveloper/Private/ControlRigBlueprint.cpp
@@ -1204,6 +1204,7 @@ void UControlRigBlueprint::PostLoad()
 							if(URigVMController* Controller = GetController(GraphToValidate))
 							{
 								FRigVMControllerNotifGuard NotifGuard(Controller, true);
+								FRigVMDefaultValueTypeGuard _(Controller, ERigVMPinDefaultValueType::Override);
 								Controller->SetPinDefaultValue(Pin->GetPinPath(), TEXT("Null"), false, false, false);
 							}
 						}
@@ -1806,6 +1807,7 @@ void UControlRigBlueprint::PatchRigElementKeyCacheOnLoad()
 								FCachedRigElement DefaultValueElement(Key, Hierarchy);
 								FString Result;
 								TBaseStructure<FCachedRigElement>::Get()->ExportText(Result, &DefaultValueElement, nullptr, nullptr, PPF_None, nullptr);								
+								FRigVMDefaultValueTypeGuard _(Controller, ERigVMPinDefaultValueType::Override);
 								Controller->SetPinDefaultValue(Pin->GetPinPath(), Result, true, false, false);
 								MarkDirtyDuringLoad();
 							}							
diff --git a/Engine/Plugins/Experimental/AnimNext/Source/AnimNext/Private/Graph/RigDecorator_AnimNextCppTrait.cpp b/Engine/Plugins/Experimental/AnimNext/Source/AnimNext/Private/Graph/RigDecorator_AnimNextCppTrait.cpp
--- a/Engine/Plugins/Experimental/AnimNext/Source/AnimNext/Private/Graph/RigDecorator_AnimNextCppTrait.cpp
+++ b/Engine/Plugins/Experimental/AnimNext/Source/AnimNext/Private/Graph/RigDecorator_AnimNextCppTrait.cpp
@@ -17,6 +17,7 @@ void FRigDecorator_AnimNextCppDecorator::GetProgrammaticPins(URigVMController* I
 		return;
 	}
 
+	FStructOnScope OriginalValueMemoryScope(DecoratorSharedDataStruct);
 	FStructOnScope DefaultValueMemoryScope(DecoratorSharedDataStruct);
 
 	if (!InDefaultValue.IsEmpty())
@@ -25,8 +26,21 @@ void FRigDecorator_AnimNextCppDecorator::GetProgrammaticPins(URigVMController* I
 		DecoratorSharedDataStruct->ImportText(*InDefaultValue, DefaultValueMemoryScope.GetStructMemory(), nullptr, PPF_SerializedAsImportText, &ErrorPipe, DecoratorSharedDataStruct->GetName());
 	}
 
+	const TFunction<ERigVMPinDefaultValueType(const FName&)> DefaultValueTypeGetter = [&OriginalValueMemoryScope, &DefaultValueMemoryScope](const FName& InPropertyName)
+	{
+		if(const FProperty* Property = OriginalValueMemoryScope.GetStruct()->FindPropertyByName(InPropertyName))
+		{
+			if(Property->Identical_InContainer(OriginalValueMemoryScope.GetStructMemory(), DefaultValueMemoryScope.GetStructMemory()))
+			{
+				return ERigVMPinDefaultValueType::Unset;
+			}
+			return ERigVMPinDefaultValueType::Override;
+		}
+		return ERigVMPinDefaultValueType::AutoDetect;
+	};
+
 	const int32 StartPinIndex = OutPinArray.Num();
-	OutPinArray.AddPins(DecoratorSharedDataStruct, InController, ERigVMPinDirection::Invalid, InParentPinIndex, DefaultValueMemoryScope.GetStructMemory(), true);
+	OutPinArray.AddPins(DecoratorSharedDataStruct, InController, ERigVMPinDirection::Invalid, InParentPinIndex, DefaultValueTypeGetter, DefaultValueMemoryScope.GetStructMemory(), true);
 
 	for (int32 PinIndex = StartPinIndex; PinIndex < OutPinArray.Num(); ++PinIndex)
 	{
diff --git a/Engine/Plugins/Experimental/AnimNext/Source/AnimNextUncookedOnly/Private/Graph/AnimNextGraph_Controller.cpp b/Engine/Plugins/Experimental/AnimNext/Source/AnimNextUncookedOnly/Private/Graph/AnimNextGraph_Controller.cpp
--- a/Engine/Plugins/Experimental/AnimNext/Source/AnimNextUncookedOnly/Private/Graph/AnimNextGraph_Controller.cpp
+++ b/Engine/Plugins/Experimental/AnimNext/Source/AnimNextUncookedOnly/Private/Graph/AnimNextGraph_Controller.cpp
@@ -27,7 +27,7 @@ URigVMUnitNode* UAnimNextGraph_Controller::AddUnitNodeWithPins(UScriptStruct* In
 	if (bHasDynamicPins)
 	{
 		const FRigVMRegistry& Registry = FRigVMRegistry::Get();
-		const FRigVMPinInfoArray PreviousPins(Node);
+		const FRigVMPinInfoArray PreviousPins(Node, this);
 
 		for (int32 PinIndex = 0; PinIndex < PinArray.Num(); ++PinIndex)
 		{
diff --git a/Engine/Plugins/Runtime/RigVM/Source/RigVMDeveloper/Private/RigVMBlueprint.cpp b/Engine/Plugins/Runtime/RigVM/Source/RigVMDeveloper/Private/RigVMBlueprint.cpp
--- a/Engine/Plugins/Runtime/RigVM/Source/RigVMDeveloper/Private/RigVMBlueprint.cpp
+++ b/Engine/Plugins/Runtime/RigVM/Source/RigVMDeveloper/Private/RigVMBlueprint.cpp
@@ -987,6 +987,13 @@ void URigVMBlueprint::PostLoad()
 			PatchParameterNodesOnLoad();
 			PatchLinksWithCast();
 			PatchFunctionsOnLoad();
+
+			const FRigVMClientPatchResult PinDefaultValuePatchResult = GetRigVMClient()->PatchPinDefaultValues();
+			if(PinDefaultValuePatchResult.RequiresToMarkPackageDirty())
+			{
+				(void)MarkPackageDirty();
+				bDirtyDuringLoad = true;
+			}
 		}
 
 #if WITH_EDITOR
diff --git a/Engine/Plugins/Runtime/RigVM/Source/RigVMDeveloper/Private/RigVMModel/Nodes/RigVMAggregateNode.cpp b/Engine/Plugins/Runtime/RigVM/Source/RigVMDeveloper/Private/RigVMModel/Nodes/RigVMAggregateNode.cpp
--- a/Engine/Plugins/Runtime/RigVM/Source/RigVMDeveloper/Private/RigVMModel/Nodes/RigVMAggregateNode.cpp
+++ b/Engine/Plugins/Runtime/RigVM/Source/RigVMDeveloper/Private/RigVMModel/Nodes/RigVMAggregateNode.cpp
@@ -243,6 +243,28 @@ TArray<URigVMPin*> URigVMAggregateNode::GetAggregateOutputs() const
 #endif
 }
 
+FString URigVMAggregateNode::GetOriginalDefaultValueForRootPin(const URigVMPin* InRootPin) const
+{
+	if(InRootPin->CanProvideDefaultValue())
+	{
+		if(const URigVMFunctionEntryNode* EntryNode = GetEntryNode())
+		{
+			if(const URigVMPin* EntryPin = EntryNode->FindPin(InRootPin->GetName()))
+			{
+				if(!EntryPin->GetLinks().IsEmpty())
+				{
+					// linked targets only since only input pins on the outer node are valid for default values
+					if(const URigVMPin* InnerRootPin = EntryPin->GetLinks()[0]->GetTargetPin())
+					{
+						return InnerRootPin->GetOriginalDefaultValue();
+					}
+				}
+			}
+		}
+	}
+	return Super::GetOriginalDefaultValueForRootPin(InRootPin);
+}
+
 void URigVMAggregateNode::InvalidateCache()
 {
 	Super::InvalidateCache();
diff --git a/Engine/Plugins/Runtime/RigVM/Source/RigVMDeveloper/Private/RigVMModel/Nodes/RigVMDispatchNode.cpp b/Engine/Plugins/Runtime/RigVM/Source/RigVMDeveloper/Private/RigVMModel/Nodes/RigVMDispatchNode.cpp
--- a/Engine/Plugins/Runtime/RigVM/Source/RigVMDeveloper/Private/RigVMModel/Nodes/RigVMDispatchNode.cpp
+++ b/Engine/Plugins/Runtime/RigVM/Source/RigVMDeveloper/Private/RigVMModel/Nodes/RigVMDispatchNode.cpp
@@ -218,6 +218,18 @@ bool URigVMDispatchNode::ShouldInputPinComputeLazily(const URigVMPin* InPin) con
 	return Super::ShouldInputPinComputeLazily(InPin);
 }
 
+FString URigVMDispatchNode::GetOriginalDefaultValueForRootPin(const URigVMPin* InRootPin) const
+{
+	if(const FRigVMDispatchFactory* Factory = GetFactory())
+	{
+		if(InRootPin->CanProvideDefaultValue())
+		{
+			return Factory->GetArgumentDefaultValue(InRootPin->GetFName(), InRootPin->GetTypeIndex());
+		}
+	}
+	return Super::GetOriginalDefaultValueForRootPin(InRootPin);
+}
+
 FString URigVMDispatchNode::GetFactoryDefaultValue() const
 {
 	TArray<FString> PinDefaultValues;
diff --git a/Engine/Plugins/Runtime/RigVM/Source/RigVMDeveloper/Private/RigVMModel/Nodes/RigVMFunctionReferenceNode.cpp b/Engine/Plugins/Runtime/RigVM/Source/RigVMDeveloper/Private/RigVMModel/Nodes/RigVMFunctionReferenceNode.cpp
--- a/Engine/Plugins/Runtime/RigVM/Source/RigVMDeveloper/Private/RigVMModel/Nodes/RigVMFunctionReferenceNode.cpp
+++ b/Engine/Plugins/Runtime/RigVM/Source/RigVMDeveloper/Private/RigVMModel/Nodes/RigVMFunctionReferenceNode.cpp
@@ -171,6 +171,7 @@ void URigVMFunctionReferenceNode::UpdateFunctionHeaderFromHost()
 	if (const FRigVMGraphFunctionData* Data = GetReferencedFunctionData())
 	{
 		ReferencedFunctionHeader = Data->Header;
+		InvalidateCache();
 	}
 }
 
@@ -183,6 +184,29 @@ const FRigVMGraphFunctionData* URigVMFunctionReferenceNode::GetReferencedFunctio
 	return nullptr;
 }
 
+FString URigVMFunctionReferenceNode::GetOriginalDefaultValueForRootPin(const URigVMPin* InRootPin) const
+{
+	if(InRootPin->CanProvideDefaultValue())
+	{
+		const FRigVMGraphFunctionHeader& Header = GetReferencedFunctionHeader();
+		if(!Header.IsValid())
+		{
+			// we don't have a valid header yet
+			// maybe still waiting for the function reference to resolve?
+			return FString();
+		}
+		const FRigVMGraphFunctionArgument* Argument = Header.Arguments.FindByPredicate([InRootPin](const FRigVMGraphFunctionArgument& InArgument) -> bool
+		{
+			return InArgument.Name == InRootPin->GetFName();
+		});
+		if(Argument)
+		{
+			return Argument->DefaultValue;
+		}
+	}
+	return Super::GetOriginalDefaultValueForRootPin(InRootPin);
+}
+
 FText URigVMFunctionReferenceNode::GetToolTipTextForPin(const URigVMPin* InPin) const
 {
 	check(InPin);
diff --git a/Engine/Plugins/Runtime/RigVM/Source/RigVMDeveloper/Private/RigVMModel/Nodes/RigVMUnitNode.cpp b/Engine/Plugins/Runtime/RigVM/Source/RigVMDeveloper/Private/RigVMModel/Nodes/RigVMUnitNode.cpp
--- a/Engine/Plugins/Runtime/RigVM/Source/RigVMDeveloper/Private/RigVMModel/Nodes/RigVMUnitNode.cpp
+++ b/Engine/Plugins/Runtime/RigVM/Source/RigVMDeveloper/Private/RigVMModel/Nodes/RigVMUnitNode.cpp
@@ -175,6 +175,37 @@ bool URigVMUnitNode::ShouldInputPinComputeLazily(const URigVMPin* InPin) const
 	return Super::ShouldInputPinComputeLazily(InPin);
 }
 
+FString URigVMUnitNode::GetOriginalDefaultValueForRootPin(const URigVMPin* InRootPin) const
+{
+	if(InRootPin->CanProvideDefaultValue())
+	{
+		const TSharedPtr<FStructOnScope> StructOnScope = ConstructStructInstance(true);
+		if(StructOnScope.IsValid())
+		{
+			if(const UStruct* Struct = StructOnScope->GetStruct())
+			{
+				if (const FProperty* Property = Struct->FindPropertyByName(InRootPin->GetFName()))
+				{
+					const uint8* PropertyMemory = Property->ContainerPtrToValuePtr<uint8>(StructOnScope->GetStructMemory());
+
+					FString DefaultValue;
+					Property->ExportText_Direct(
+						DefaultValue,
+						PropertyMemory,
+						PropertyMemory,
+						nullptr,
+						PPF_None,
+						nullptr
+					);
+
+					return DefaultValue;
+				}
+			}
+		}
+	}
+	return Super::GetOriginalDefaultValueForRootPin(InRootPin);
+}
+
 bool URigVMUnitNode::IsOutDated() const
 {
 	return !GetDeprecatedMetadata().IsEmpty();
diff --git a/Engine/Plugins/Runtime/RigVM/Source/RigVMDeveloper/Private/RigVMModel/RigVMClient.cpp b/Engine/Plugins/Runtime/RigVM/Source/RigVMDeveloper/Private/RigVMModel/RigVMClient.cpp
--- a/Engine/Plugins/Runtime/RigVM/Source/RigVMDeveloper/Private/RigVMModel/RigVMClient.cpp
+++ b/Engine/Plugins/Runtime/RigVM/Source/RigVMDeveloper/Private/RigVMModel/RigVMClient.cpp
@@ -945,22 +945,41 @@ FRigVMClientPatchResult FRigVMClient::PatchModelsOnLoad()
 		URigVMController* Controller = GetOrCreateController(Model);
 		TGuardValue<bool> GuardSuspendTemplateComputation(Controller->bSuspendTemplateComputation, true);
 		TGuardValue<bool> GuardIsTransacting(Controller->bIsTransacting, true);
-
-		Result.Merge(Controller->PatchRerouteNodesOnLoad());
-		Result.Merge(Controller->PatchUnitNodesOnLoad());
-		Result.Merge(Controller->PatchDispatchNodesOnLoad());
-		Result.Merge(Controller->PatchBranchNodesOnLoad());
-		Result.Merge(Controller->PatchIfSelectNodesOnLoad());
-		Result.Merge(Controller->PatchArrayNodesOnLoad());
-		Result.Merge(Controller->PatchReduceArrayFloatDoubleConvertsionsOnLoad());
-		Result.Merge(Controller->PatchInvalidLinksOnWildcards());
-		Result.Merge(Controller->PatchExecutePins());
-		Result.Merge(Controller->PatchLazyPins());
-
-		if (URigVMCollapseNode* CollapseNode = Model->GetTypedOuter<URigVMCollapseNode>())
 		{
-			Result.Merge(Controller->PatchFunctionsWithInvalidReturnPaths());	
+			FRigVMDefaultValueTypeGuard DefaultValueTypeGuard(Controller, ERigVMPinDefaultValueType::KeepValueType);
+			Result.Merge(Controller->PatchRerouteNodesOnLoad());
+			Result.Merge(Controller->PatchUnitNodesOnLoad());
+			Result.Merge(Controller->PatchDispatchNodesOnLoad());
+			Result.Merge(Controller->PatchBranchNodesOnLoad());
+			Result.Merge(Controller->PatchIfSelectNodesOnLoad());
+			Result.Merge(Controller->PatchArrayNodesOnLoad());
+			Result.Merge(Controller->PatchReduceArrayFloatDoubleConvertsionsOnLoad());
+			Result.Merge(Controller->PatchInvalidLinksOnWildcards());
+			Result.Merge(Controller->PatchExecutePins());
+			Result.Merge(Controller->PatchLazyPins());
+		
+			if (URigVMCollapseNode* CollapseNode = Model->GetTypedOuter<URigVMCollapseNode>())
+			{
+				Result.Merge(Controller->PatchFunctionsWithInvalidReturnPaths());	
+			}
 		}
+		Result.Merge(Controller->PatchPinDefaultValues());
+	}
+
+	return Result;
+}
+
+FRigVMClientPatchResult FRigVMClient::PatchPinDefaultValues()
+{
+	FRigVMClientPatchResult Result;
+	
+	TArray<URigVMGraph*> AllModels = GetAllModelsLeavesFirst(true);
+	TGuardValue<bool> ClientIgnoreModificationsGuard(bIgnoreModelNotifications, true);
+	for(URigVMGraph* Model : AllModels)
+	{
+		URigVMController* Controller = GetOrCreateController(Model);
+		TGuardValue<bool> GuardIsTransacting(Controller->bIsTransacting, true);
+		Result.Merge(Controller->PatchPinDefaultValues());
 	}
 
 	return Result;
@@ -1264,6 +1283,7 @@ bool FRigVMClient::UpdateFunctionReferences(const FRigVMGraphFunctionHeader& Hea
 
 				Node->Modify();
 				Node->ReferencedFunctionHeader = Header;
+				Node->InvalidateCache();
 
 				if (bUpdateDependencies || bUpdateExternalVariables)
 				{
diff --git a/Engine/Plugins/Runtime/RigVM/Source/RigVMDeveloper/Private/RigVMModel/RigVMController.cpp b/Engine/Plugins/Runtime/RigVM/Source/RigVMDeveloper/Private/RigVMModel/RigVMController.cpp
--- a/Engine/Plugins/Runtime/RigVM/Source/RigVMDeveloper/Private/RigVMModel/RigVMController.cpp
+++ b/Engine/Plugins/Runtime/RigVM/Source/RigVMDeveloper/Private/RigVMModel/RigVMController.cpp
@@ -356,6 +356,7 @@ FRigVMPinInfo::FRigVMPinInfo()
 	, TypeIndex(INDEX_NONE)
 	, bIsArray(false)
 	, Property(nullptr)
+	, DefaultValueType(ERigVMPinDefaultValueType::AutoDetect)
 	, bIsExpanded(false)
 	, bIsConstant(false)
 	, bIsDynamicArray(false)
@@ -364,13 +365,14 @@ FRigVMPinInfo::FRigVMPinInfo()
 {
 }
 
-FRigVMPinInfo::FRigVMPinInfo(const URigVMPin* InPin, int32 InParentIndex, ERigVMPinDirection InDirection)
+FRigVMPinInfo::FRigVMPinInfo(const URigVMPin* InPin, int32 InParentIndex, ERigVMPinDirection InDirection, ERigVMPinDefaultValueType InDefaultValueType)
 	: ParentIndex(InParentIndex)
 	, Name(InPin->GetName())
 	, Direction(InDirection == ERigVMPinDirection::Invalid ? InPin->GetDirection() : InDirection)
 	, TypeIndex(InPin->GetTypeIndex())
 	, bIsArray(InPin->IsArray())
 	, Property(nullptr)
+	, DefaultValueType(InDefaultValueType)
 	, bIsExpanded(InPin->IsExpanded())
 	, bIsConstant(InPin->IsDefinedAsConstant())
 	, bIsDynamicArray(InPin->IsDynamicArray())
@@ -389,13 +391,14 @@ FRigVMPinInfo::FRigVMPinInfo(const URigVMPin* InPin, int32 InParentIndex, ERigVM
 	}
 }
 
-FRigVMPinInfo::FRigVMPinInfo(FProperty* InProperty, ERigVMPinDirection InDirection, int32 InParentIndex, const uint8* InDefaultValueMemory)
+FRigVMPinInfo::FRigVMPinInfo(FProperty* InProperty, ERigVMPinDirection InDirection, int32 InParentIndex, ERigVMPinDefaultValueType InDefaultValueType, const uint8* InDefaultValueMemory)
 	: ParentIndex(InParentIndex)
 	, Name(InProperty->GetFName())
 	, Direction(InDirection)
 	, TypeIndex(INDEX_NONE)
 	, bIsArray(InProperty->IsA<FArrayProperty>())
 	, Property(InProperty)
+	, DefaultValueType(InDefaultValueType)
 	, bIsExpanded(false)
 	, bIsConstant(false)
 	, bIsDynamicArray(false)
@@ -529,22 +532,23 @@ uint32 GetTypeHash(const FRigVMPinInfo& InPin)
 	return Hash;
 }
 
-FRigVMPinInfoArray::FRigVMPinInfoArray(const URigVMNode* InNode)
+FRigVMPinInfoArray::FRigVMPinInfoArray(const URigVMNode* InNode, URigVMController* InController)
 {
 	// this method adds all pins as currently represented in the model.
 	for(const URigVMPin* Pin : InNode->GetPins())
 	{
-		(void)AddPin(Pin, INDEX_NONE);
+		const ERigVMPinDefaultValueType DefaultValueType = InController->GetDefaultValueType(Pin, Pin->GetDefaultValue());
+		(void)AddPin(Pin, INDEX_NONE, ERigVMPinDirection::Invalid, DefaultValueType);
 	}
 }
 
-int32 FRigVMPinInfoArray::AddPin(const URigVMPin* InPin, int32 InParentIndex, ERigVMPinDirection InDirection)
+int32 FRigVMPinInfoArray::AddPin(const URigVMPin* InPin, int32 InParentIndex, ERigVMPinDirection InDirection, ERigVMPinDefaultValueType InDefaultValueType)
 {
 	// this method adds all pins as currently represented in the model.
-	const int32 Index = Pins.Emplace(InPin, InParentIndex, InDirection);
+	const int32 Index = Pins.Emplace(InPin, InParentIndex, InDirection, InDefaultValueType);
 	for(const URigVMPin* SubPin : InPin->GetSubPins())
 	{
-		const int32 SubPinIndex = AddPin(SubPin, Index, InDirection);
+		const int32 SubPinIndex = AddPin(SubPin, Index, InDirection, InDefaultValueType);
 		Pins[Index].SubPins.Add(SubPinIndex);
 	}
 	return Index;
@@ -558,11 +562,12 @@ FRigVMPinInfoArray::FRigVMPinInfoArray(const URigVMNode* InNode, URigVMControlle
 	for(const URigVMPin* Pin : InNode->GetPins())
 	{
 		const FString DefaultValue = Pin->GetDefaultValue();
+		const ERigVMPinDefaultValueType DefaultValueType = InController->GetDefaultValueType(Pin, DefaultValue);
 		if (Pin->GetTypeIndex() == INDEX_NONE)
 		{
 			InController->ReportErrorf( TEXT("Invalid pin type %s for %s in %s"), *Pin->GetCPPType(), *Pin->GetPathName(), *InNode->GetPackage()->GetPathName());
 		}
-		(void)AddPin(InController, INDEX_NONE, Pin->GetFName(), Pin->GetDirection(), Pin->GetTypeIndex(), DefaultValue, nullptr, InPreviousPinInfos, bAddSubPins);
+		(void)AddPin(InController, INDEX_NONE, Pin->GetFName(), Pin->GetDirection(), Pin->GetTypeIndex(), DefaultValue, DefaultValueType, nullptr, InPreviousPinInfos, bAddSubPins);
 	}
 }
 
@@ -579,18 +584,23 @@ FRigVMPinInfoArray::FRigVMPinInfoArray(const FRigVMGraphFunctionHeader& Function
 			InController->ReportErrorf( TEXT("Invalid pin type %s for %s in %s"), *FunctionArgument.CPPType.ToString(), *FunctionHeader.LibraryPointer.LibraryNodePath, *InController->GetPackage()->GetPathName());
 		}
 		ensureMsgf(TypeIndex != INDEX_NONE, TEXT("Invalid pin type %s in %s"), *FunctionArgument.CPPType.ToString(), *InController->GetPackage()->GetPathName());
-		(void)AddPin(InController, INDEX_NONE, FunctionArgument.Name, FunctionArgument.Direction, TypeIndex, FunctionArgument.DefaultValue, nullptr, InPreviousPinInfos, true);
+		ERigVMPinDefaultValueType DefaultValueType = ERigVMPinDefaultValueType::Unset;
+		if(const FRigVMPinInfo* PinInfo = InPreviousPinInfos->GetPinFromPinPath(FunctionArgument.Name.ToString()))
+		{
+			DefaultValueType = PinInfo->DefaultValueType;
+		}
+		(void)AddPin(InController, INDEX_NONE, FunctionArgument.Name, FunctionArgument.Direction, TypeIndex, FunctionArgument.DefaultValue, DefaultValueType, nullptr, InPreviousPinInfos, true);
 	}
 }
 
 int32 FRigVMPinInfoArray::AddPin(FProperty* InProperty, URigVMController* InController,
-                                              ERigVMPinDirection InDirection, int32 InParentIndex, const uint8* InDefaultValueMemory, bool bAddSubPins)
+    ERigVMPinDirection InDirection, int32 InParentIndex, ERigVMPinDefaultValueType InDefaultValueType, const uint8* InDefaultValueMemory, bool bAddSubPins)
 {
 	// this method adds pins as needed based on the property structure
 
 	check(InDefaultValueMemory);
 	
-	const int32 Index = Pins.Emplace(InProperty, InDirection, InParentIndex, InDefaultValueMemory);
+	const int32 Index = Pins.Emplace(InProperty, InDirection, InParentIndex, InDefaultValueType, InDefaultValueMemory);
 	if(InParentIndex != INDEX_NONE)
 	{
 		Pins[InParentIndex].SubPins.Add(Index);
@@ -600,7 +610,11 @@ int32 FRigVMPinInfoArray::AddPin(FProperty* InProperty, URigVMController* InCont
 	{
 		if (const FStructProperty* StructProperty = CastField<FStructProperty>(InProperty))
 		{
-			(void)AddPins(StructProperty->Struct, InController, Pins[Index].Direction, Index, InDefaultValueMemory, bAddSubPins);
+			const TFunction<ERigVMPinDefaultValueType(const FName&)> DefaultValueTypeGetter = [InDefaultValueType](const FName&)
+			{
+				return InDefaultValueType;
+			};
+			(void)AddPins(StructProperty->Struct, InController, Pins[Index].Direction, Index, DefaultValueTypeGetter, InDefaultValueMemory, bAddSubPins);
 		}
 		else if (const FArrayProperty* ArrayProperty = CastField<FArrayProperty>(InProperty))
 		{
@@ -608,7 +622,7 @@ int32 FRigVMPinInfoArray::AddPin(FProperty* InProperty, URigVMController* InCont
 			for(int32 ElementIndex = 0; ElementIndex < ArrayHelper.Num(); ElementIndex++)
 			{
 				const uint8* ElementDefaultValueMemory = ArrayHelper.GetRawPtr(ElementIndex);
-				const int32 SubIndex = AddPin(ArrayProperty->Inner, InController, Pins[Index].Direction, Index, ElementDefaultValueMemory, bAddSubPins);
+				const int32 SubIndex = AddPin(ArrayProperty->Inner, InController, Pins[Index].Direction, Index, InDefaultValueType, ElementDefaultValueMemory, bAddSubPins);
 				Pins[SubIndex].Name = *FString::FormatAsNumber(ElementIndex);
 			}
 		}
@@ -617,7 +631,7 @@ int32 FRigVMPinInfoArray::AddPin(FProperty* InProperty, URigVMController* InCont
 }
 
 int32 FRigVMPinInfoArray::AddPin(URigVMController* InController, int32 InParentIndex, const FName& InName, ERigVMPinDirection InDirection,
-	TRigVMTypeIndex InTypeIndex, const FString& InDefaultValue, const uint8* InDefaultValueMemory, const FRigVMPinInfoArray* InPreviousPinInfos, bool bAddSubPins)
+	TRigVMTypeIndex InTypeIndex, const FString& InDefaultValue, ERigVMPinDefaultValueType InDefaultValueType, const uint8* InDefaultValueMemory, const FRigVMPinInfoArray* InPreviousPinInfos, bool bAddSubPins)
 {
 	const FRigVMRegistry& Registry = FRigVMRegistry::Get();
 		
@@ -628,6 +642,7 @@ int32 FRigVMPinInfoArray::AddPin(URigVMController* InController, int32 InParentI
 	Info.TypeIndex = InTypeIndex;
 	Info.bIsArray = Registry.IsArrayType(InTypeIndex);
 	Info.DefaultValue = InDefaultValue;
+	Info.DefaultValueType = InDefaultValueType;
 	Info.CorrectExecuteTypeIndex();
 
 	const int32 Index = Pins.Add(Info);
@@ -674,7 +689,7 @@ int32 FRigVMPinInfoArray::AddPin(URigVMController* InController, int32 InParentI
 					ElementDefaultValueMemory = ElementDefaultValueMemoryScope.GetStructMemory();
 				}
 
-				(void)AddPin(InController, Index, *FString::FormatAsNumber(ElementIndex), InDirection, ElementTypeIndex, ElementDefaultValue, ElementDefaultValueMemory, InPreviousPinInfos, bAddSubPins);
+				(void)AddPin(InController, Index, *FString::FormatAsNumber(ElementIndex), InDirection, ElementTypeIndex, ElementDefaultValue, InDefaultValueType, ElementDefaultValueMemory, InPreviousPinInfos, bAddSubPins);
 			}
 		}
 		else if(UScriptStruct* ScriptStruct = Cast<UScriptStruct>(Type.CPPTypeObject))
@@ -688,7 +703,11 @@ int32 FRigVMPinInfoArray::AddPin(URigVMController* InController, int32 InParentI
 				ScriptStruct->ImportText(*Pins[Index].DefaultValue, DefaultValueMemoryScope.GetStructMemory(), nullptr, PPF_None, &ErrorPipe, FString());
 				DefaultValueMemory = DefaultValueMemoryScope.GetStructMemory();
 			}
-			AddPins(ScriptStruct, InController, InDirection, Index, DefaultValueMemory, bAddSubPins);
+			const TFunction<ERigVMPinDefaultValueType(const FName&)> DefaultValueTypeGetter = [InDefaultValueType](const FName&)
+			{
+				return InDefaultValueType;
+			};
+			AddPins(ScriptStruct, InController, InDirection, Index, DefaultValueTypeGetter, DefaultValueMemory, bAddSubPins);
 		}
 	}
 
@@ -701,7 +720,8 @@ int32 FRigVMPinInfoArray::AddPin(URigVMController* InController, int32 InParentI
 }
 
 void FRigVMPinInfoArray::AddPins(UScriptStruct* InScriptStruct, URigVMController* InController,
-	ERigVMPinDirection InDirection, int32 InParentIndex, const uint8* InDefaultValueMemory, bool bAddSubPins)
+	ERigVMPinDirection InDirection, int32 InParentIndex, TFunction<ERigVMPinDefaultValueType(const FName&)> InDefaultValueTypeGetter,
+	const uint8* InDefaultValueMemory, bool bAddSubPins)
 {
 	if (InController->GetSchema()->ShouldUnfoldStruct(InController, InScriptStruct))
 	{
@@ -728,7 +748,12 @@ void FRigVMPinInfoArray::AddPins(UScriptStruct* InScriptStruct, URigVMController
 					}
 				}
 #endif
-				(void)AddPin(*It, InController, InDirection, InParentIndex, DefaultValueMemory, bAddSubPinsForProperty);
+				ERigVMPinDefaultValueType DefaultValueType = ERigVMPinDefaultValueType::AutoDetect;
+				if(InDefaultValueTypeGetter)
+				{
+					DefaultValueType = InDefaultValueTypeGetter(It->GetFName());
+				}
+				(void)AddPin(*It, InController, InDirection, InParentIndex, DefaultValueType, DefaultValueMemory, bAddSubPinsForProperty);
 			}
 		}
 	}
@@ -1394,6 +1419,7 @@ URigVMUnitNode* URigVMController::AddUnitNode(UScriptStruct* InScriptStruct, TSu
 			{
 				if (Pin->IsWildCard())
 				{
+					FRigVMDefaultValueTypeGuard _(this, ERigVMPinDefaultValueType::Unset, true);
 					ResolveWildCardPin(Pin, Pair.Value, bSetupUndoRedo);
 				}
 			}
@@ -1469,6 +1495,7 @@ URigVMUnitNode* URigVMController::AddUnitNode(UScriptStruct* InScriptStruct, TSu
 	FString ExportedDefaultValue;
 	CreateDefaultValueForStructIfRequired(InScriptStruct, ExportedDefaultValue);
 	{
+		FRigVMDefaultValueTypeGuard _(this, ERigVMPinDefaultValueType::Unset, true);
 		TGuardValue<bool> SuspendNotifications(bSuspendNotifications, true);
 		AddPinsForStruct(InScriptStruct, Node, nullptr, ERigVMPinDirection::Invalid, ExportedDefaultValue, true);
 	}
@@ -1665,6 +1692,7 @@ bool URigVMController::SetUnitNodeDefaults(URigVMUnitNode* InNode, const FRigStr
 			const FString NewDefault = FRigVMStruct::ExportToFullyQualifiedText(Property, MemberMemoryPtr);
 			if(NewDefault != Pin->GetDefaultValue())
 			{
+				FRigVMDefaultValueTypeGuard _(this, ERigVMPinDefaultValueType::AutoDetect);
 				SetPinDefaultValue(Pin->GetPinPath(), NewDefault, true, bSetupUndoRedo, false, bPrintPythonCommand);
 			}
 		}
@@ -3039,6 +3067,7 @@ URigVMRerouteNode* URigVMController::AddRerouteNodeOnPin(const FString& InPinPat
 	FString DefaultValue = Pin->GetDefaultValue();
 	if (!DefaultValue.IsEmpty())
 	{
+		FRigVMDefaultValueTypeGuard _(this, ERigVMPinDefaultValueType::Override, true);
 		SetPinDefaultValue(ValuePin, Pin->GetDefaultValue(), true, false, false);
 	}
 
@@ -7810,7 +7839,9 @@ bool URigVMController::SetPinDefaultValue(const FString& InPinPath, const FStrin
 			return SetVariableName(VariableNode, *InDefaultValue, bSetupUndoRedo);
 		}
 	}
-	
+
+	// by default treat pin default value changes as overrides
+	FRigVMDefaultValueTypeGuard _(this, ERigVMPinDefaultValueType::Override);
 	if (!SetPinDefaultValue(Pin, InDefaultValue, bResizeArrays, bSetupUndoRedo, bMergeUndoAction, bSetValueOnLinkedPins))
 	{
 		return false;
@@ -7903,6 +7934,7 @@ bool URigVMController::SetPinDefaultValue(URigVMPin* InPin, const FString& InDef
 					{
 						break;
 					}
+					bSetPinDefaultValueSucceeded = true;
 				}
 				while (Elements.Num() < InPin->SubPins.Num())
 				{
@@ -7910,6 +7942,7 @@ bool URigVMController::SetPinDefaultValue(URigVMPin* InPin, const FString& InDef
 					{
 						break;
 					}
+					bSetPinDefaultValueSucceeded = true;
 				}
 			}
 			else
@@ -7957,15 +7990,28 @@ bool URigVMController::SetPinDefaultValue(URigVMPin* InPin, const FString& InDef
 	if(!bSetPinDefaultValueSucceeded)
 	{
 		// no need to send notifications if not changing the value
-		if (InPin->GetSubPins().IsEmpty() && (InPin->DefaultValue != ClampedDefaultValue))
+		if (InPin->GetSubPins().IsEmpty())
 		{
-			InPin->DefaultValue = ClampedDefaultValue;
-			Notify(ERigVMGraphNotifType::PinDefaultValueChanged, InPin);
-			if (!bSuspendNotifications)
+			// always mark the pin's default value type as user provided
+			// even if the value didn't change. this is done to remember
+			// the value when switching versions of nodes / functions.
+			const ERigVMPinDefaultValueType NewDefaultValueType = GetDefaultValueType(InPin, ClampedDefaultValue);
+			if(InPin->DefaultValueType != NewDefaultValueType)
 			{
-				Graph->MarkPackageDirty();
+				InPin->DefaultValueType = NewDefaultValueType;
+				bSetPinDefaultValueSucceeded = true;
+			}
+
+			if(InPin->DefaultValue != ClampedDefaultValue)
+			{
+				InPin->DefaultValue = ClampedDefaultValue;
+				Notify(ERigVMGraphNotifType::PinDefaultValueChanged, InPin);
+				if (!bSuspendNotifications)
+				{
+					(void)Graph->MarkPackageDirty();
+				}
+				bSetPinDefaultValueSucceeded = true;
 			}
-			bSetPinDefaultValueSucceeded = true;
 		}
 	}
 
@@ -8006,13 +8052,6 @@ bool URigVMController::ResetPinDefaultValue(const FString& InPinPath, bool bSetu
 		return false;
 	}
 
-	URigVMNode* Node = Pin->GetNode();
-	if (!Node->IsA<URigVMUnitNode>() && !Node->IsA<URigVMFunctionReferenceNode>())
-	{
-		ReportErrorf(TEXT("Pin '%s' is neither part of a unit nor a function reference node."), *InPinPath);
-		return false;
-	}
-
 	const bool bSuccess = ResetPinDefaultValue(Pin, bSetupUndoRedo);
 	if (bSuccess && bPrintPythonCommand)
 	{
@@ -8038,29 +8077,14 @@ bool URigVMController::ResetPinDefaultValue(URigVMPin* InPin, bool bSetupUndoRed
 
 	URigVMNode* RigVMNode = InPin->GetNode();
 
-	// unit nodes
-	if (URigVMUnitNode* UnitNode = Cast<URigVMUnitNode>(RigVMNode))
-	{
-		// cut off the first one since it's the node
-		static const uint32 Offset = 1;
-		const FString DefaultValue = GetPinInitialDefaultValueFromStruct(UnitNode->GetScriptStruct(), InPin, Offset);
-		if (!DefaultValue.IsEmpty())
-		{
-			SetPinDefaultValue(InPin, DefaultValue, true, bSetupUndoRedo, false);
-			return true;
-		}
-	}
+	FRigVMDefaultValueTypeGuard _(this, ERigVMPinDefaultValueType::Unset);
 
-	// function reference nodes
-	URigVMFunctionReferenceNode* RefNode = Cast<URigVMFunctionReferenceNode>(RigVMNode);
-	if (RefNode != nullptr)
+	FString DefaultValue = RigVMNode->GetOriginalPinDefaultValue(InPin);
+	PostProcessDefaultValue(InPin, DefaultValue);
+	if(!DefaultValue.IsEmpty() || InPin->IsWildCard())
 	{
-		const FString DefaultValue = GetPinInitialDefaultValue(InPin);
-		if (!DefaultValue.IsEmpty())
-		{
-			SetPinDefaultValue(InPin, DefaultValue, true, bSetupUndoRedo, false);
-			return true;
-		}
+		SetPinDefaultValue(InPin, DefaultValue, true, bSetupUndoRedo, false);
+		return true;
 	}
 
 	return false;
@@ -8891,11 +8915,12 @@ URigVMPin* URigVMController::InsertArrayPin(URigVMPin* ArrayPin, int32 InIndex,
 		InIndex = ArrayPin->GetSubPins().Num();
 	}
 
+	FRigVMDefaultValueTypeGuard _(this, ERigVMPinDefaultValueType::Override);
 	FRigVMControllerCompileBracketScope CompileScope(this);
 	FRigVMInsertArrayPinAction Action;
 	if (bSetupUndoRedo)
 	{
-		Action = FRigVMInsertArrayPinAction(this, ArrayPin, InIndex, InDefaultValue);
+		Action = FRigVMInsertArrayPinAction(this, ArrayPin, InIndex, InDefaultValue, OptionalDefaultValueType.Get(ERigVMPinDefaultValueType::Override));
 		Action.SetTitle(FString::Printf(TEXT("Insert Array Pin")));
 		GetActionStack()->BeginAction(Action);
 	}
@@ -8937,9 +8962,19 @@ URigVMPin* URigVMController::InsertArrayPin(URigVMPin* ArrayPin, int32 InIndex,
 	{
 		FString DefaultValue = InDefaultValue;
 		PostProcessDefaultValue(Pin, DefaultValue);
+		if(Pin->CanProvideDefaultValue())
+		{
+			Pin->DefaultValueType = GetDefaultValueType(Pin, DefaultValue);
+		}
 		Pin->DefaultValue = DefaultValue;
 	}
 
+	// set the array pin's default value type based on the resulting array pin list
+	{
+		const FRigVMDefaultValueTypeGuard AutoDetectGuard(this, ERigVMPinDefaultValueType::AutoDetect, true);
+		ArrayPin->DefaultValueType = GetDefaultValueType(ArrayPin, ArrayPin->GetDefaultValue());
+	}
+
 	Notify(ERigVMGraphNotifType::PinAdded, Pin);
 	Notify(ERigVMGraphNotifType::PinArraySizeChanged, ArrayPin);
 	if (!bSuspendNotifications)
@@ -8983,6 +9018,8 @@ bool URigVMController::RemoveArrayPin(const FString& InArrayElementPinPath, bool
 		return false;
 	}
 
+	FRigVMDefaultValueTypeGuard _(this, ERigVMPinDefaultValueType::Override);
+
 	URigVMPin* ArrayPin = ArrayElementPin->GetParentPin();
 	check(ArrayPin);
 	ensure(ArrayPin->IsArray());
@@ -9032,6 +9069,12 @@ bool URigVMController::RemoveArrayPin(const FString& InArrayElementPinPath, bool
 	{
 		Graph->MarkPackageDirty();
 	}
+
+	// set the array pin's default value type based on the resulting array pin list
+	{
+		const FRigVMDefaultValueTypeGuard AutoDetectGuard(this, ERigVMPinDefaultValueType::AutoDetect, true);
+		ArrayPin->DefaultValueType = GetDefaultValueType(ArrayPin, ArrayPin->GetDefaultValue());
+	}
 	Notify(ERigVMGraphNotifType::PinArraySizeChanged, ArrayPin);
 
 	if (bSetupUndoRedo)
@@ -11587,6 +11630,7 @@ URigVMFunctionReferenceNode* URigVMController::AddFunctionReferenceNodeFromDescr
 			const FString& DefaultValue = Argument.DefaultValue;
 			if (!DefaultValue.IsEmpty())
 			{
+				FRigVMDefaultValueTypeGuard _(this, ERigVMPinDefaultValueType::Unset, true);
 				SetPinDefaultValue(TargetPin, DefaultValue, true, false, false);
 			}
 		}
@@ -14415,6 +14459,10 @@ void URigVMController::AddPinsForStruct(UStruct* InStruct, URigVMNode* InNode, U
 				else if(DefaultValuePtr != nullptr)
 				{
 					Pin->DefaultValue = *DefaultValuePtr;
+					if(Pin->CanProvideDefaultValue() && !Pin->DefaultValue.IsEmpty())
+					{
+						Pin->DefaultValueType = GetDefaultValueType(Pin, Pin->DefaultValue);
+					}
 				}
 			}
 
@@ -14435,6 +14483,10 @@ void URigVMController::AddPinsForStruct(UStruct* InStruct, URigVMNode* InNode, U
 						FString DefaultValue = *DefaultValuePtr;
 						PostProcessDefaultValue(Pin, DefaultValue);
 						Pin->DefaultValue = *DefaultValuePtr;
+						if(Pin->CanProvideDefaultValue() && !Pin->DefaultValue.IsEmpty())
+						{
+							Pin->DefaultValueType = GetDefaultValueType(Pin, Pin->DefaultValue);
+						}
 					}
 				}
 			}
@@ -14444,6 +14496,10 @@ void URigVMController::AddPinsForStruct(UStruct* InStruct, URigVMNode* InNode, U
 				FString DefaultValue = *DefaultValuePtr;
 				PostProcessDefaultValue(Pin, DefaultValue);
 				Pin->DefaultValue = DefaultValue;
+				if(Pin->CanProvideDefaultValue() && !Pin->DefaultValue.IsEmpty())
+				{
+					Pin->DefaultValueType = GetDefaultValueType(Pin, Pin->DefaultValue);
+				}
 			}
 
 			if (!bSuspendNotifications)
@@ -14504,6 +14560,10 @@ void URigVMController::AddPinsForArray(FArrayProperty* InArrayProperty, URigVMNo
 			{
 				PostProcessDefaultValue(Pin, DefaultValue);
 				Pin->DefaultValue = DefaultValue;
+				if(Pin->CanProvideDefaultValue() && !Pin->DefaultValue.IsEmpty())
+				{
+					Pin->DefaultValueType = GetDefaultValueType(Pin, Pin->DefaultValue);
+				}
 			}
 		}
 
@@ -14519,6 +14579,10 @@ void URigVMController::AddPinsForArray(FArrayProperty* InArrayProperty, URigVMNo
 			{
 				PostProcessDefaultValue(Pin, DefaultValue);
 				Pin->DefaultValue = DefaultValue;
+				if(Pin->CanProvideDefaultValue() && !Pin->DefaultValue.IsEmpty())
+				{
+					Pin->DefaultValueType = GetDefaultValueType(Pin, Pin->DefaultValue);
+				}
 			}
 		}
 
@@ -14526,6 +14590,10 @@ void URigVMController::AddPinsForArray(FArrayProperty* InArrayProperty, URigVMNo
 		{
 			PostProcessDefaultValue(Pin, DefaultValue);
 			Pin->DefaultValue = DefaultValue;
+			if(Pin->CanProvideDefaultValue() && !Pin->DefaultValue.IsEmpty())
+			{
+				Pin->DefaultValueType = GetDefaultValueType(Pin, Pin->DefaultValue);
+			}
 		}
 	}
 }
@@ -14533,6 +14601,7 @@ void URigVMController::AddPinsForArray(FArrayProperty* InArrayProperty, URigVMNo
 void URigVMController::AddPinsForTemplate(const FRigVMTemplate* InTemplate, const FRigVMTemplateTypeMap& InPinTypeMap, URigVMNode* InNode)
 {
 	const FRigVMRegistry& Registry = FRigVMRegistry::Get();
+	FRigVMDefaultValueTypeGuard _(this, ERigVMPinDefaultValueType::Unset, true);
 
 	FRigVMDispatchContext DispatchContext;
 	if(const URigVMDispatchNode* DispatchNode = Cast<URigVMDispatchNode>(InNode))
@@ -15135,7 +15204,7 @@ void URigVMController::GenerateRepopulatePinsNodeData(TArray<FRepopulatePinsNode
 
 	const FRigVMRegistry& Registry = FRigVMRegistry::Get();
 
-	NodeData.PreviousPinInfos = FRigVMPinInfoArray(InNode);
+	NodeData.PreviousPinInfos = FRigVMPinInfoArray(InNode, this);
 	NodeData.PreviousPinHash = GetTypeHash(NodeData.PreviousPinInfos);
 	if (!GenerateNewPinInfos(Registry, InNode, NodeData.PreviousPinInfos, NodeData.NewPinInfos, NodeData.bSetupOrphanPinsForThisNode))
 	{
@@ -15271,7 +15340,15 @@ bool URigVMController::GenerateNewPinInfos(const FRigVMRegistry& Registry, URigV
 		}
 
 		const TSharedPtr<FStructOnScope> DefaultValueContent = UnitNode->ConstructStructInstance(false);
-		NewPinInfos.AddPins(ScriptStruct, this, ERigVMPinDirection::Invalid, INDEX_NONE, DefaultValueContent->GetStructMemory(), true);
+		const TFunction<ERigVMPinDefaultValueType(const FName&)> DefaultValueTypeGetter = [InNode, this](const FName& InPropertyName)
+		{
+			if(const URigVMPin* Pin = InNode->FindPin(InPropertyName.ToString()))
+			{
+				return GetDefaultValueType(Pin, Pin->GetDefaultValue());
+			}
+			return ERigVMPinDefaultValueType::AutoDetect;
+		};
+		NewPinInfos.AddPins(ScriptStruct, this, ERigVMPinDirection::Invalid, INDEX_NONE, DefaultValueTypeGetter, DefaultValueContent->GetStructMemory(), true);
 	}
 	else if (DispatchNode)
 	{
@@ -15307,7 +15384,7 @@ bool URigVMController::GenerateNewPinInfos(const FRigVMRegistry& Registry, URigV
 							}
 						}
 
-						(void)NewPinInfos.AddPin(this, INDEX_NONE, Arg->Name, Arg->Direction, TypeIndex, DefaultValue, nullptr, &PreviousPinInfos, true);
+						(void)NewPinInfos.AddPin(this, INDEX_NONE, Arg->Name, Arg->Direction, TypeIndex, DefaultValue, ERigVMPinDefaultValueType::Unset, nullptr, &PreviousPinInfos, true);
 					}
 				}
 			};
@@ -15382,7 +15459,12 @@ bool URigVMController::GenerateNewPinInfos(const FRigVMRegistry& Registry, URigV
 				bAddSubPinsForArgument = false;
 			}
 #endif
-			(void)NewPinInfos.AddPin(this, INDEX_NONE, Arg->Name, Arg->GetDirection(), TypeIndex, DefaultValue, DefaultValueMemory, &PreviousPinInfos, bAddSubPinsForArgument);
+			ERigVMPinDefaultValueType DefaultValueType = ERigVMPinDefaultValueType::AutoDetect;
+			if(const FRigVMPinInfo* PinInfo = PreviousPinInfos.GetPinFromPinPath(Arg->Name.ToString()))
+			{
+				DefaultValueType = PinInfo->DefaultValueType;
+			}
+			(void)NewPinInfos.AddPin(this, INDEX_NONE, Arg->Name, Arg->GetDirection(), TypeIndex, DefaultValue, DefaultValueType, DefaultValueMemory, &PreviousPinInfos, bAddSubPinsForArgument);
 		}
 
 		AddExecutePins(ERigVMPinDirection::Output);
@@ -15520,7 +15602,7 @@ bool URigVMController::GenerateNewPinInfos(const FRigVMRegistry& Registry, URigV
 				}
 
 				ERigVMPinDirection Direction = bIsEntryNode ? ERigVMPinDirection::Output : ERigVMPinDirection::Input;
-				(void)NewPinInfos.AddPin(LibraryPin, INDEX_NONE, Direction);
+				(void)NewPinInfos.AddPin(LibraryPin, INDEX_NONE, Direction, LibraryPin->GetDefaultValueType());
 			}
 		}
 		else
@@ -15559,7 +15641,7 @@ bool URigVMController::GenerateNewPinInfos(const FRigVMRegistry& Registry, URigV
 		const FRigVMPinInfo& PreviousPin = PreviousPinInfos[Index];
 		if (PreviousPin.bIsDecorator)
 		{
-			const int32 NewPinIndex = NewPinInfos.AddPin(this, INDEX_NONE, PreviousPin.Name, PreviousPin.Direction, PreviousPin.TypeIndex, PreviousPin.DefaultValue, nullptr, &PreviousPinInfos, true);
+			const int32 NewPinIndex = NewPinInfos.AddPin(this, INDEX_NONE, PreviousPin.Name, PreviousPin.Direction, PreviousPin.TypeIndex, PreviousPin.DefaultValue, PreviousPin.DefaultValueType, nullptr, &PreviousPinInfos, true);
 			NewPinInfos[NewPinIndex].bIsDecorator = true;
 
 			if (URigVMPin* Pin = InNode->FindPin(PreviousPin.PinPath))
@@ -16242,6 +16324,7 @@ URigVMController::FPinState URigVMController::GetPinState(URigVMPin* InPin, bool
 	State.CPPType = InPin->GetCPPType();
 	State.CPPTypeObject = InPin->GetCPPTypeObject();
 	State.DefaultValue = InPin->GetDefaultValue();
+	State.DefaultValueType = InPin->GetDefaultValueType();
 	State.bIsExpanded = InPin->IsExpanded();
 	State.InjectionInfos = InPin->GetInjectedNodes();
 
@@ -16334,12 +16417,13 @@ void URigVMController::ApplyPinState(URigVMPin* InPin, const FPinState& InPinSta
 		}
 	}
 
-	if (!InPinState.DefaultValue.IsEmpty())
+	if (!InPinState.DefaultValue.IsEmpty() || InPin->IsWildCard())
 	{
 		FString DefaultValue = InPinState.DefaultValue;
 		PostProcessDefaultValue(InPin, DefaultValue);
-		if(!DefaultValue.IsEmpty())
+		if(!DefaultValue.IsEmpty() || InPin->IsWildCard())
 		{
+			FRigVMDefaultValueTypeGuard _(this, InPinState.DefaultValueType, true);
 			SetPinDefaultValue(InPin, DefaultValue, true, bSetupUndoRedo, false);
 		}
 	}
@@ -16379,7 +16463,7 @@ void URigVMController::ApplyPinStates(URigVMNode* InNode, const TMap<FString, UR
 	}
 }
 
-URigVMPin* URigVMController::CreatePinFromPinInfo(const FRigVMRegistry& InRegistry, const FRigVMPinInfoArray& InPreviousPinInfos, const FRigVMPinInfo& InPinInfo, const FString& InPinPath, UObject* InOuter) const
+URigVMPin* URigVMController::CreatePinFromPinInfo(const FRigVMRegistry& InRegistry, const FRigVMPinInfoArray& InPreviousPinInfos, const FRigVMPinInfo& InPinInfo, const FString& InPinPath, UObject* InOuter)
 {
 	check(InOuter);
 	URigVMPin* Pin = NewObject<URigVMPin>(InOuter, InPinInfo.Name);
@@ -16411,6 +16495,7 @@ URigVMPin* URigVMController::CreatePinFromPinInfo(const FRigVMRegistry& InRegist
 
 	Pin->bIsExpanded = InPinInfo.bIsExpanded;
 	Pin->DefaultValue = InPinInfo.DefaultValue;
+	Pin->DefaultValueType = InPinInfo.DefaultValueType;
 
 	// reuse expansion state and default value
 	if (const FRigVMPinInfo* PreviousPin = InPreviousPinInfos.GetPinFromPinPath(InPinPath))
@@ -16419,6 +16504,7 @@ URigVMPin* URigVMController::CreatePinFromPinInfo(const FRigVMRegistry& InRegist
 		{
 			Pin->bIsExpanded = PreviousPin->bIsExpanded;
 			Pin->DefaultValue = PreviousPin->DefaultValue;
+			Pin->DefaultValueType = PreviousPin->DefaultValueType;
 		}
 	}
 
@@ -16626,7 +16712,7 @@ void URigVMController::CreateDefaultValueForStructIfRequired(UScriptStruct* InSt
 	}
 }
 
-void URigVMController::PostProcessDefaultValue(URigVMPin* Pin, FString& OutDefaultValue)
+void URigVMController::PostProcessDefaultValue(const URigVMPin* Pin, FString& OutDefaultValue)
 {
 	static const FString NoneString = FName(NAME_None).ToString();
 	static const FString QuotedNoneString = FString::Printf(TEXT("\"%s\""), *NoneString);
@@ -17900,6 +17986,11 @@ bool URigVMController::ChangePinType(URigVMPin* InPin, TRigVMTypeIndex InTypeInd
 				InPin->DefaultValue = TemplateNode->GetInitialDefaultValueForPin(InPin->GetFName());
 			}
 		}
+
+		if(!InPin->DefaultValue.IsEmpty())
+		{
+			InPin->DefaultValueType = GetDefaultValueType(InPin, InPin->DefaultValue);
+		}
 	}
 
 	if (InPin->IsExecuteContext() && !InPin->GetNode()->IsA<URigVMFunctionEntryNode>() && !InPin->GetNode()->IsA<URigVMFunctionReturnNode>())
@@ -18275,6 +18366,19 @@ void URigVMController::AddNodePin(URigVMNode* InNode, URigVMPin* InPin)
 	{
 		InPin->bIsExpanded = true;
 	}
+
+	if(InPin->DefaultValueType == ERigVMPinDefaultValueType::AutoDetect)
+	{
+		if(InPin->DefaultValue.IsEmpty())
+		{
+			InPin->DefaultValueType = ERigVMPinDefaultValueType::Unset;
+		}
+		else
+		{
+			FRigVMDefaultValueTypeGuard _(this, ERigVMPinDefaultValueType::AutoDetect, true);
+			InPin->DefaultValueType = GetDefaultValueType(InPin, InPin->DefaultValue);
+		}
+	}
 }
 
 void URigVMController::AddSubPin(URigVMPin* InParentPin, URigVMPin* InPin)
@@ -20240,6 +20344,112 @@ FRigVMClientPatchResult URigVMController::PatchLazyPins()
 	return Result;
 }
 
+FRigVMClientPatchResult URigVMController::PatchPinDefaultValues()
+{
+	FRigVMClientPatchResult Result;
+	if(!CVarRigVMEnablePinDefaultTypes.GetValueOnAnyThread())
+	{
+		return Result;
+	}
+
+	FRigVMDefaultValueTypeGuard AutoDetectGuard(this, ERigVMPinDefaultValueType::AutoDetect, true);
+	if (const URigVMGraph* Graph = GetGraph())
+	{
+		for (const URigVMNode* Node : Graph->GetNodes())
+		{
+			const TArray<URigVMPin*> Pins = Node->GetAllPinsRecursively();
+			for(URigVMPin* Pin : Pins)
+			{
+				if(Pin->CanProvideDefaultValue() && Pin->GetSubPins().IsEmpty())
+				{
+					if(Pin->GetDefaultValueType() == ERigVMPinDefaultValueType::AutoDetect)
+					{
+						if(!Pin->DefaultValue.IsEmpty())
+						{
+							// this will determine based on the delta if the default value
+							// is unchanged or overridden by the user
+							const ERigVMPinDefaultValueType DefaultValueType = GetDefaultValueType(Pin, Pin->DefaultValue);
+							if(Pin->DefaultValueType != DefaultValueType)
+							{
+								Pin->DefaultValueType = DefaultValueType;
+								Result.bChangedContent = true;
+							}
+						}
+					}
+					else if(!Pin->GetRootPin()->HasUserProvidedDefaultValue())
+					{
+						if(Pin->GetDefaultValueType() == ERigVMPinDefaultValueType::Unset)
+						{
+							// if the pin is supposed to be the original value of the providing
+							// node (like unit struct, dispatch factory or function ref).
+							// reset the value back to the current / updated default.
+							FString OriginalDefaultValue = Pin->GetOriginalDefaultValue();
+							PostProcessDefaultValue(Pin, OriginalDefaultValue);
+							if(!OriginalDefaultValue.IsEmpty() && (Pin->GetCPPType() != RigVMTypeUtils::FStringType))
+							{
+								if(Pin->IsValidDefaultValue(OriginalDefaultValue))
+								{
+									if(!Pin->DefaultValue.Equals(OriginalDefaultValue, ESearchCase::CaseSensitive))
+									{
+										FRigVMDefaultValueTypeGuard OverrideGuard(this, ERigVMPinDefaultValueType::Unset, true);
+										if(SetPinDefaultValue(Pin, OriginalDefaultValue, true, false, false))
+										{
+											Result.bChangedContent = true;
+										}
+									}
+								}
+							}
+						}
+					}
+				}
+			}
+		}
+	}
+	return Result;
+
+}
+
+ERigVMPinDefaultValueType URigVMController::GetDefaultValueType(const URigVMPin* InPin, const FString& InDefaultValue) const
+{
+	if(!CVarRigVMEnablePinDefaultTypes.GetValueOnAnyThread())
+	{
+		return ERigVMPinDefaultValueType::AutoDetect;
+	}
+	
+	// default to user provided values
+	ERigVMPinDefaultValueType Type = OptionalDefaultValueType.Get(InPin->GetDefaultValueType());
+	
+	if(Type == ERigVMPinDefaultValueType::AutoDetect)
+	{
+		if(InDefaultValue.IsEmpty())
+		{
+			return ERigVMPinDefaultValueType::Unset;
+		}
+		
+		FString OriginalDefaultValue = InPin->GetOriginalDefaultValue();
+		PostProcessDefaultValue(InPin, OriginalDefaultValue);
+		if(OriginalDefaultValue.Equals(InDefaultValue, ESearchCase::CaseSensitive))
+		{
+			return ERigVMPinDefaultValueType::Unset;
+		}
+
+		// if the node didn't provide a valid default value let's keep it at autodetect
+		if(OriginalDefaultValue.IsEmpty() && (InPin->GetCPPType() != RigVMTypeUtils::FStringType))
+		{
+			return ERigVMPinDefaultValueType::AutoDetect;
+		}
+		
+		return ERigVMPinDefaultValueType::Override;
+	}
+
+	if(Type == ERigVMPinDefaultValueType::KeepValueType)
+	{
+		return InPin->GetDefaultValueType();
+	}
+
+	return Type;
+}
+
 void URigVMController::PostDuplicateHost(const FString& InOldPathName, const FString& InNewPathName)
 {
 	if (const URigVMGraph* Graph = GetGraph())
diff --git a/Engine/Plugins/Runtime/RigVM/Source/RigVMDeveloper/Private/RigVMModel/RigVMControllerActions.cpp b/Engine/Plugins/Runtime/RigVM/Source/RigVMDeveloper/Private/RigVMModel/RigVMControllerActions.cpp
--- a/Engine/Plugins/Runtime/RigVM/Source/RigVMDeveloper/Private/RigVMModel/RigVMControllerActions.cpp
+++ b/Engine/Plugins/Runtime/RigVM/Source/RigVMDeveloper/Private/RigVMModel/RigVMControllerActions.cpp
@@ -1069,6 +1069,8 @@ bool FRigVMSetPinWatchAction::Redo()
 
 FRigVMSetPinDefaultValueAction::FRigVMSetPinDefaultValueAction()
 : FRigVMBaseAction(nullptr)
+, OldDefaultValueType(ERigVMPinDefaultValueType::AutoDetect)
+, NewDefaultValueType(ERigVMPinDefaultValueType::AutoDetect)
 {
 }
 
@@ -1077,6 +1079,8 @@ FRigVMSetPinDefaultValueAction::FRigVMSetPinDefaultValueAction(URigVMController*
 , PinPath(InPin->GetPinPath())
 , OldDefaultValue(InPin->GetDefaultValueStoredByUserInterface())
 , NewDefaultValue(InNewDefaultValue)
+, OldDefaultValueType(InPin->GetDefaultValueType())
+, NewDefaultValueType(InController->GetDefaultValueType(InPin, InNewDefaultValue))
 {
 	/* Since for template we are chaning types - it is possible that the
 	 * pin is no longer compliant with the old value
@@ -1105,6 +1109,7 @@ bool FRigVMSetPinDefaultValueAction::Merge(const FRigVMBaseAction* Other)
 	}
 
 	NewDefaultValue = Action->NewDefaultValue;
+	NewDefaultValueType = Action->NewDefaultValueType;
 	return true;
 }
 
@@ -1116,8 +1121,18 @@ bool FRigVMSetPinDefaultValueAction::Undo()
 	}
 	if (OldDefaultValue.IsEmpty())
 	{
-		return true;
+		// strings and wildcards allow to set an empty default
+		if(const URigVMPin* Pin = GetController()->GetGraph()->FindPin(PinPath))
+		{
+			if((Pin->GetCPPType() != RigVMTypeUtils::FStringType) &&
+				!Pin->IsWildCard())
+			{
+				return true;
+			}
+		}
 	}
+
+	FRigVMDefaultValueTypeGuard _(GetController(), OldDefaultValueType, true);
 	return GetController()->SetPinDefaultValue(PinPath, OldDefaultValue, true, false);
 }
 
@@ -1127,8 +1142,25 @@ bool FRigVMSetPinDefaultValueAction::Redo()
 	{
 		return false;
 	}
-	if (!NewDefaultValue.IsEmpty())
+	bool bIsValidDefaultValue = !NewDefaultValue.IsEmpty();
+	if(!bIsValidDefaultValue)
+	{
+		// strings and wildcards allow to set an empty default
+		if(const URigVMPin* Pin = GetController()->GetGraph()->FindPin(PinPath))
+		{
+			if(Pin->GetCPPType() == RigVMTypeUtils::FStringType)
+			{
+				bIsValidDefaultValue = true;
+			}
+			else if(Pin->IsWildCard())
+			{
+				bIsValidDefaultValue = true;
+			}
+		}
+	}
+	if (bIsValidDefaultValue)
 	{
+		FRigVMDefaultValueTypeGuard _(GetController(), NewDefaultValueType, true);
 		if (!GetController()->SetPinDefaultValue(PinPath, NewDefaultValue, true, false))
 		{
 			return false;
@@ -1140,14 +1172,16 @@ bool FRigVMSetPinDefaultValueAction::Redo()
 FRigVMInsertArrayPinAction::FRigVMInsertArrayPinAction()
 : FRigVMBaseAction(nullptr)
 , Index(0)
+, NewDefaultValueType(ERigVMPinDefaultValueType::AutoDetect)
 {
 }
 
-FRigVMInsertArrayPinAction::FRigVMInsertArrayPinAction(URigVMController* InController, URigVMPin* InArrayPin, int32 InIndex, const FString& InNewDefaultValue)
+FRigVMInsertArrayPinAction::FRigVMInsertArrayPinAction(URigVMController* InController, URigVMPin* InArrayPin, int32 InIndex, const FString& InNewDefaultValue, const ERigVMPinDefaultValueType& InNewDefaultValueType)
 : FRigVMBaseAction(InController)
 , ArrayPinPath(InArrayPin->GetPinPath())
 , Index(InIndex)
 , NewDefaultValue(InNewDefaultValue)
+, NewDefaultValueType(InNewDefaultValueType)
 {
 }
 
@@ -1166,6 +1200,8 @@ bool FRigVMInsertArrayPinAction::Redo()
 	{
 		return false;
 	}
+
+	FRigVMDefaultValueTypeGuard _(GetController(), NewDefaultValueType, true);
 	if(GetController()->InsertArrayPin(ArrayPinPath, Index, NewDefaultValue, false).IsEmpty())
 	{
 		return false;
@@ -1176,6 +1212,7 @@ bool FRigVMInsertArrayPinAction::Redo()
 FRigVMRemoveArrayPinAction::FRigVMRemoveArrayPinAction()
 : FRigVMBaseAction(nullptr)
 , Index(0)
+, DefaultValueType(ERigVMPinDefaultValueType::AutoDetect)
 {
 }
 
@@ -1184,6 +1221,7 @@ FRigVMRemoveArrayPinAction::FRigVMRemoveArrayPinAction(URigVMController* InContr
 , ArrayPinPath(InArrayElementPin->GetParentPin()->GetPinPath())
 , Index(InArrayElementPin->GetPinIndex())
 , DefaultValue(InArrayElementPin->GetDefaultValue())
+, DefaultValueType(InArrayElementPin->GetDefaultValueType())
 {
 }
 
@@ -1193,10 +1231,13 @@ bool FRigVMRemoveArrayPinAction::Undo()
 	{
 		return false;
 	}
-	if (GetController()->InsertArrayPin(*ArrayPinPath, Index, DefaultValue, false).IsEmpty())
+
+	FRigVMDefaultValueTypeGuard _(GetController(), DefaultValueType, true);
+	if(GetController()->InsertArrayPin(*ArrayPinPath, Index, DefaultValue, false).IsEmpty())
 	{
 		return false;
 	}
+	
 	return FRigVMBaseAction::Undo();
 }
 
@@ -1206,6 +1247,7 @@ bool FRigVMRemoveArrayPinAction::Redo()
 	{
 		return false;
 	}
+
 	return GetController()->RemoveArrayPin(FString::Printf(TEXT("%s.%d"), *ArrayPinPath, Index), false);
 }
 
diff --git a/Engine/Plugins/Runtime/RigVM/Source/RigVMDeveloper/Private/RigVMModel/RigVMNode.cpp b/Engine/Plugins/Runtime/RigVM/Source/RigVMDeveloper/Private/RigVMModel/RigVMNode.cpp
--- a/Engine/Plugins/Runtime/RigVM/Source/RigVMDeveloper/Private/RigVMModel/RigVMNode.cpp
+++ b/Engine/Plugins/Runtime/RigVM/Source/RigVMDeveloper/Private/RigVMModel/RigVMNode.cpp
@@ -114,6 +114,95 @@ TArray<URigVMPin*> URigVMNode::GetAllPinsRecursively() const
 	return Result;
 }
 
+FString URigVMNode::GetOriginalPinDefaultValue(const URigVMPin* InPin) const
+{
+	const FString CompleteSegmentPath = InPin->GetSegmentPath(true);
+	if(const FString* CachedOriginalPinDefaultValue = CachedOriginalPinDefaultValues.Find(CompleteSegmentPath))
+	{
+		return *CachedOriginalPinDefaultValue;
+	}
+	
+	const URigVMPin* RootPin = InPin->GetRootPin();
+	const FString OriginalDefaultValue = GetOriginalDefaultValueForRootPin(RootPin);
+	if((RootPin != InPin) && !OriginalDefaultValue.IsEmpty())
+	{
+		struct Local
+		{
+			static FString TraverseArrayElement(TMap<FString, FString>& Cache, const URigVMPin* InPin, const FString& InSegmentPath, const FString& InRemainingSegmentPath, const FString& InDefaultValue)
+			{
+				FString Left = InRemainingSegmentPath, Right;
+				(void)URigVMPin::SplitPinPathAtStart(InRemainingSegmentPath, Left, Right);
+
+				if(const URigVMPin* SubPin = InPin->FindSubPin(Left))
+				{
+					const TArray<FString> DefaultValues = URigVMPin::SplitDefaultValue(InDefaultValue);
+					if(DefaultValues.IsValidIndex(SubPin->GetPinIndex()))
+					{
+						const FString SubPinDefaultValue = DefaultValues[SubPin->GetPinIndex()];
+						return Traverse(Cache, SubPin, URigVMPin::JoinPinPath(InSegmentPath, Left), Right, SubPinDefaultValue);
+					}
+				}
+				return FString();
+			}
+
+			static FString TraverseStructMember(TMap<FString, FString>& Cache, const URigVMPin* InPin, const FString& InSegmentPath, const FString& InRemainingSegmentPath, const FString& InDefaultValue)
+			{
+				FString Left = InRemainingSegmentPath, Right;
+				(void)URigVMPin::SplitPinPathAtStart(InRemainingSegmentPath, Left, Right);
+
+				const TArray<FString> DefaultValues = URigVMPin::SplitDefaultValue(InDefaultValue);
+				for(const FString& DefaultValue : DefaultValues)
+				{
+					FString Name, Value;
+					if (DefaultValue.Split(TEXT("="), &Name, &Value))
+					{
+						if(Left.Equals(Name, ESearchCase::CaseSensitive))
+						{
+							if(const URigVMPin* SubPin = InPin->FindSubPin(Left))
+							{
+								return Traverse(Cache, SubPin, URigVMPin::JoinPinPath(InSegmentPath, Left), Right, Value);
+							}
+						}
+					}
+				}
+
+				return FString();
+			}
+
+			static FString Traverse(TMap<FString, FString>& Cache, const URigVMPin* InPin, const FString& InSegmentPath, const FString& InRemainingSegmentPath, const FString& InDefaultValue)
+			{
+				FString DefaultValue = InDefaultValue;
+				if(!InRemainingSegmentPath.IsEmpty())
+				{
+					if(InPin->IsArray())
+					{
+						DefaultValue = TraverseArrayElement(Cache, InPin, InSegmentPath, InRemainingSegmentPath, InDefaultValue);
+					}
+					else if(InPin->IsStruct())
+					{
+						DefaultValue = TraverseStructMember(Cache, InPin, InSegmentPath, InRemainingSegmentPath, InDefaultValue);
+					}
+				}
+
+				if(!InDefaultValue.IsEmpty())
+				{
+					Cache.FindOrAdd(InSegmentPath, InDefaultValue);
+				}
+				return DefaultValue;
+			}
+		};
+
+		const FString SegmentPath = InPin->GetSegmentPath(false);
+		return Local::Traverse(CachedOriginalPinDefaultValues, RootPin, RootPin->GetName(), SegmentPath, OriginalDefaultValue);
+	}
+
+	if(!OriginalDefaultValue.IsEmpty())
+	{
+		CachedOriginalPinDefaultValues.FindOrAdd(CompleteSegmentPath, OriginalDefaultValue);
+	}
+	return OriginalDefaultValue;
+}
+
 URigVMPin* URigVMNode::FindPin(const FString& InPinPath) const
 {
 	FString Left, Right;
@@ -152,6 +241,25 @@ URigVMPin* URigVMNode::FindPin(const FString& InPinPath) const
 	return nullptr;
 }
 
+URigVMPin* URigVMNode::FindRootPinByName(const FName& InPinName) const
+{
+	for(TObjectPtr<URigVMPin> Pin : Pins)
+	{
+		if(Pin->GetFName().IsEqual(InPinName, ENameCase::CaseSensitive))
+		{
+			return Pin;
+		}
+	}
+	for(TObjectPtr<URigVMPin> OrphanedPin : OrphanedPins)
+	{
+		if(OrphanedPin->GetFName().IsEqual(InPinName, ENameCase::CaseSensitive))
+		{
+			return OrphanedPin;
+		}
+	}
+	return nullptr;
+}
+
 const TArray<URigVMPin*>& URigVMNode::GetOrphanedPins() const
 {
 	return OrphanedPins;
@@ -223,6 +331,12 @@ FText URigVMNode::GetToolTipTextForPin(const URigVMPin* InPin) const
 	return FText::FromName(InPin->GetFName());
 }
 
+FString URigVMNode::GetOriginalDefaultValueForRootPin(const URigVMPin* InRootPin) const
+{
+	ensure(InRootPin->IsRootPin());
+	return FString();
+}
+
 void URigVMNode::UpdateDecoratorRootPinNames()
 {
 	TArray<FString> NewDecoratorRootPinNames;
@@ -625,6 +739,11 @@ void URigVMNode::GetLinkedNodesRecursive(URigVMPin* InPin, bool bLookForSources,
 	}
 }
 
+void URigVMNode::InvalidateCache()
+{
+	CachedOriginalPinDefaultValues.Reset();
+}
+
 const TArray<int32>& URigVMNode::GetInstructionsForVM(const FRigVMExtendedExecuteContext& Context, URigVM* InVM, const FRigVMASTProxy& InProxy) const
 {
 	if(const FProfilingCache* Cache = UpdateProfilingCacheIfNeeded(Context, InVM, InProxy))
diff --git a/Engine/Plugins/Runtime/RigVM/Source/RigVMDeveloper/Private/RigVMModel/RigVMPin.cpp b/Engine/Plugins/Runtime/RigVM/Source/RigVMDeveloper/Private/RigVMModel/RigVMPin.cpp
--- a/Engine/Plugins/Runtime/RigVM/Source/RigVMDeveloper/Private/RigVMModel/RigVMPin.cpp
+++ b/Engine/Plugins/Runtime/RigVM/Source/RigVMDeveloper/Private/RigVMModel/RigVMPin.cpp
@@ -28,6 +28,8 @@
 
 #include UE_INLINE_GENERATED_CPP_BY_NAME(RigVMPin)
 
+TAutoConsoleVariable<bool> CVarRigVMEnablePinDefaultTypes(TEXT("RigVM.EnablePinDefaultTypes"), false, TEXT("enables the use of pin default types"));
+
 #if WITH_EDITOR
 #include "UObject/CoreRedirects.h"
 #endif
@@ -112,6 +114,7 @@ URigVMPin::URigVMPin()
 	, CPPTypeObject(nullptr)
 	, CPPTypeObjectPath(NAME_None)
 	, DefaultValue(FString())
+	, DefaultValueType(ERigVMPinDefaultValueType::AutoDetect)
 	, BoundVariablePath_DEPRECATED()
 	, LastKnownTypeIndex(INDEX_NONE)
 {
@@ -744,6 +747,15 @@ bool URigVMPin::ShouldHideSubPins() const
 	return false;
 }
 
+FString URigVMPin::GetOriginalDefaultValue() const
+{
+	if(const URigVMNode* Node = GetNode())
+	{
+		return Node->GetOriginalPinDefaultValue(this);
+	}
+	return FString();
+}
+
 FString URigVMPin::GetDefaultValue() const
 {
 	return GetDefaultValue(EmptyPinOverride, true);
@@ -1039,6 +1051,53 @@ bool URigVMPin::IsValidDefaultValue(const FString& InDefaultValue) const
 	return true;
 }
 
+bool URigVMPin::HasUserProvidedDefaultValue() const
+{
+	if(!CVarRigVMEnablePinDefaultTypes.GetValueOnAnyThread())
+	{
+		return false;
+	}
+	
+	if(!CanProvideDefaultValue())
+	{
+		return false;
+	}
+	
+	if(DefaultValueType == ERigVMPinDefaultValueType::Override)
+	{
+		return true;
+	}
+
+	for(const TObjectPtr<URigVMPin>& SubPin : SubPins)
+	{
+		if(SubPin->HasUserProvidedDefaultValue())
+		{
+			return true;
+		}
+	}
+
+	return false;
+}
+
+bool URigVMPin::CanProvideDefaultValue() const
+{
+	if((GetDirection() != ERigVMPinDirection::Input) &&
+		(GetDirection() != ERigVMPinDirection::IO) &&
+		(GetDirection() != ERigVMPinDirection::Visible))
+	{
+		return false;
+	}
+	if(IsWildCard() && !IsArray())
+	{
+		return false;
+	}
+	if(IsExecuteContext())
+	{
+		return false;
+	}
+	return true;
+}
+
 FString URigVMPin::ClampDefaultValueFromMetaData(const FString& InDefaultValue) const
 {
 	FString RetVal = InDefaultValue;
diff --git a/Engine/Plugins/Runtime/RigVM/Source/RigVMDeveloper/Public/RigVMModel/Nodes/RigVMAggregateNode.h b/Engine/Plugins/Runtime/RigVM/Source/RigVMDeveloper/Public/RigVMModel/Nodes/RigVMAggregateNode.h
--- a/Engine/Plugins/Runtime/RigVM/Source/RigVMDeveloper/Public/RigVMModel/Nodes/RigVMAggregateNode.h
+++ b/Engine/Plugins/Runtime/RigVM/Source/RigVMDeveloper/Public/RigVMModel/Nodes/RigVMAggregateNode.h
@@ -44,6 +44,10 @@ public:
 	virtual TArray<URigVMPin*> GetAggregateInputs() const override;
 	virtual TArray<URigVMPin*> GetAggregateOutputs() const override;
 
+protected:
+
+	virtual FString GetOriginalDefaultValueForRootPin(const URigVMPin* InRootPin) const override;
+
 private:
 
 	virtual void InvalidateCache() override;
diff --git a/Engine/Plugins/Runtime/RigVM/Source/RigVMDeveloper/Public/RigVMModel/Nodes/RigVMDispatchNode.h b/Engine/Plugins/Runtime/RigVM/Source/RigVMDeveloper/Public/RigVMModel/Nodes/RigVMDispatchNode.h
--- a/Engine/Plugins/Runtime/RigVM/Source/RigVMDeveloper/Public/RigVMModel/Nodes/RigVMDispatchNode.h
+++ b/Engine/Plugins/Runtime/RigVM/Source/RigVMDeveloper/Public/RigVMModel/Nodes/RigVMDispatchNode.h
@@ -79,6 +79,7 @@ protected:
 	const UScriptStruct* GetFactoryStruct() const;
 	virtual void InvalidateCache() override;
 	virtual bool ShouldInputPinComputeLazily(const URigVMPin* InPin) const override;
+	virtual FString GetOriginalDefaultValueForRootPin(const URigVMPin* InRootPin) const override;
 	FString GetFactoryDefaultValue() const;
 
 private:
diff --git a/Engine/Plugins/Runtime/RigVM/Source/RigVMDeveloper/Public/RigVMModel/Nodes/RigVMFunctionReferenceNode.h b/Engine/Plugins/Runtime/RigVM/Source/RigVMDeveloper/Public/RigVMModel/Nodes/RigVMFunctionReferenceNode.h
--- a/Engine/Plugins/Runtime/RigVM/Source/RigVMDeveloper/Public/RigVMModel/Nodes/RigVMFunctionReferenceNode.h
+++ b/Engine/Plugins/Runtime/RigVM/Source/RigVMDeveloper/Public/RigVMModel/Nodes/RigVMFunctionReferenceNode.h
@@ -57,6 +57,9 @@ public:
 
 	const FRigVMGraphFunctionData* GetReferencedFunctionData(bool bLoadIfNecessary = true) const;
 
+protected:
+
+	virtual FString GetOriginalDefaultValueForRootPin(const URigVMPin* InRootPin) const override;
 	
 private:
 
diff --git a/Engine/Plugins/Runtime/RigVM/Source/RigVMDeveloper/Public/RigVMModel/Nodes/RigVMUnitNode.h b/Engine/Plugins/Runtime/RigVM/Source/RigVMDeveloper/Public/RigVMModel/Nodes/RigVMUnitNode.h
--- a/Engine/Plugins/Runtime/RigVM/Source/RigVMDeveloper/Public/RigVMModel/Nodes/RigVMUnitNode.h
+++ b/Engine/Plugins/Runtime/RigVM/Source/RigVMDeveloper/Public/RigVMModel/Nodes/RigVMUnitNode.h
@@ -102,6 +102,7 @@ protected:
 
 	virtual FText GetToolTipTextForPin(const URigVMPin* InPin) const override;
 	virtual bool ShouldInputPinComputeLazily(const URigVMPin* InPin) const override;
+	virtual FString GetOriginalDefaultValueForRootPin(const URigVMPin* InRootPin) const override;
 	void EnumeratePropertiesOnHostAndStructInstance(
 		URigVMHost* InHost,
 		TSharedPtr<FStructOnScope> InInstance, 
diff --git a/Engine/Plugins/Runtime/RigVM/Source/RigVMDeveloper/Public/RigVMModel/RigVMClient.h b/Engine/Plugins/Runtime/RigVM/Source/RigVMDeveloper/Public/RigVMModel/RigVMClient.h
--- a/Engine/Plugins/Runtime/RigVM/Source/RigVMDeveloper/Public/RigVMModel/RigVMClient.h
+++ b/Engine/Plugins/Runtime/RigVM/Source/RigVMDeveloper/Public/RigVMModel/RigVMClient.h
@@ -246,6 +246,7 @@ public:
 
 	// backwards compatibility
 	FRigVMClientPatchResult PatchModelsOnLoad();
+	FRigVMClientPatchResult PatchPinDefaultValues();
 
 	// try to reattach detached links and delete remaining ones
 	void ProcessDetachedLinks();
diff --git a/Engine/Plugins/Runtime/RigVM/Source/RigVMDeveloper/Public/RigVMModel/RigVMController.h b/Engine/Plugins/Runtime/RigVM/Source/RigVMDeveloper/Public/RigVMModel/RigVMController.h
--- a/Engine/Plugins/Runtime/RigVM/Source/RigVMDeveloper/Public/RigVMModel/RigVMController.h
+++ b/Engine/Plugins/Runtime/RigVM/Source/RigVMDeveloper/Public/RigVMModel/RigVMController.h
@@ -169,8 +169,8 @@ private:
 struct RIGVMDEVELOPER_API FRigVMPinInfo
 {
 	FRigVMPinInfo();
-	FRigVMPinInfo(const URigVMPin* InPin, int32 InParentIndex, ERigVMPinDirection InDirection = ERigVMPinDirection::Invalid);
-	FRigVMPinInfo(FProperty* InProperty, ERigVMPinDirection InDirection, int32 InParentIndex, const uint8* InDefaultValueMemory);
+	FRigVMPinInfo(const URigVMPin* InPin, int32 InParentIndex, ERigVMPinDirection InDirection, ERigVMPinDefaultValueType InDefaultValueType);
+	FRigVMPinInfo(FProperty* InProperty, ERigVMPinDirection InDirection, int32 InParentIndex, ERigVMPinDefaultValueType InDefaultValueType, const uint8* InDefaultValueMemory);
 
 	void CorrectExecuteTypeIndex();
 
@@ -182,6 +182,7 @@ struct RIGVMDEVELOPER_API FRigVMPinInfo
 	FProperty* Property;
 	FString PinPath;
 	FString DefaultValue;
+	ERigVMPinDefaultValueType DefaultValueType;
 	FString DisplayName;
 	FString CustomWidgetName;
 	bool bIsExpanded;
@@ -197,8 +198,8 @@ struct RIGVMDEVELOPER_API FRigVMPinInfo
 struct RIGVMDEVELOPER_API FRigVMPinInfoArray
 {
 	FRigVMPinInfoArray() {}
-	explicit FRigVMPinInfoArray(const URigVMNode* InNode);
-	FRigVMPinInfoArray(const URigVMNode* InNode, URigVMController* InController, const FRigVMPinInfoArray* InPreviousPinInfos = nullptr);
+	explicit FRigVMPinInfoArray(const URigVMNode* InNode, URigVMController* InController);
+	FRigVMPinInfoArray(const URigVMNode* InNode, URigVMController* InController, const FRigVMPinInfoArray* InPreviousPinInfos);
 	FRigVMPinInfoArray(const FRigVMGraphFunctionHeader& FunctionHeader, URigVMController* InController, const FRigVMPinInfoArray* InPreviousPinInfos = nullptr);
 
 	int32 Num() const { return Pins.Num(); }
@@ -207,10 +208,10 @@ struct RIGVMDEVELOPER_API FRigVMPinInfoArray
 	TArray<FRigVMPinInfo>::RangedForIteratorType begin() const { return Pins.begin(); }
 	TArray<FRigVMPinInfo>::RangedForIteratorType end() const { return Pins.end(); }
 
-	int32 AddPin(const URigVMPin* InPin, int32 InParentIndex, ERigVMPinDirection InDirection = ERigVMPinDirection::Invalid);
-	int32 AddPin(FProperty* InProperty, URigVMController* InController, ERigVMPinDirection InDirection, int32 InParentIndex, const uint8* InDefaultValueMemory, bool bAddSubPins);
-	int32 AddPin(URigVMController* InController, int32 InParentIndex, const FName& InName, ERigVMPinDirection InDirection, TRigVMTypeIndex InTypeIndex, const FString& InDefaultValue, const uint8* InDefaultValueMemory, const FRigVMPinInfoArray* InPreviousPinInfos, bool bAddSubPins);
-	void AddPins(UScriptStruct* InScriptStruct, URigVMController* InController, ERigVMPinDirection InDirection, int32 InParentIndex, const uint8* InDefaultValueMemory, bool bAddSubPins);
+	int32 AddPin(const URigVMPin* InPin, int32 InParentIndex, ERigVMPinDirection InDirection, ERigVMPinDefaultValueType InDefaultValueType);
+	int32 AddPin(FProperty* InProperty, URigVMController* InController, ERigVMPinDirection InDirection, int32 InParentIndex, ERigVMPinDefaultValueType InDefaultValueType, const uint8* InDefaultValueMemory, bool bAddSubPins);
+	int32 AddPin(URigVMController* InController, int32 InParentIndex, const FName& InName, ERigVMPinDirection InDirection, TRigVMTypeIndex InTypeIndex, const FString& InDefaultValue, ERigVMPinDefaultValueType InDefaultValueType, const uint8* InDefaultValueMemory, const FRigVMPinInfoArray* InPreviousPinInfos, bool bAddSubPins);
+	void AddPins(UScriptStruct* InScriptStruct, URigVMController* InController, ERigVMPinDirection InDirection, int32 InParentIndex, TFunction<ERigVMPinDefaultValueType(const FName&)> InDefaultValueTypeGetter, const uint8* InDefaultValueMemory, bool bAddSubPins);
 
 	const FString& GetPinPath(const int32 InIndex) const;
 	int32 GetIndexFromPinPath(const FString& InPinPath) const;
@@ -1219,6 +1220,7 @@ private:
 		FString CPPType;
 		UObject* CPPTypeObject;
 		FString DefaultValue;
+		ERigVMPinDefaultValueType DefaultValueType;
 		bool bIsExpanded;
 		TArray<URigVMInjectionInfo*> InjectionInfos;
 		TArray<URigVMInjectionInfo::FWeakInfo> WeakInjectionInfos;
@@ -1233,7 +1235,7 @@ private:
 
 	static FLinearColor GetColorFromMetadata(const FString& InMetadata);
 	static void CreateDefaultValueForStructIfRequired(UScriptStruct* InStruct, FString& InOutDefaultValue);
-	static void PostProcessDefaultValue(URigVMPin* Pin, FString& OutDefaultValue);
+	static void PostProcessDefaultValue(const URigVMPin* Pin, FString& OutDefaultValue);
 
 	void ResolveTemplateNodeMetaData(URigVMTemplateNode* InNode, bool bSetupUndoRedo);
 
@@ -1255,7 +1257,7 @@ private:
 	static URigVMPin* MakeExecutePin(URigVMNode* InNode, const FName& InName);
 	static bool MakeExecutePin(URigVMPin* InOutPin);
 	bool AddGraphNode(URigVMNode* InNode, bool bNotify);
-	static void AddNodePin(URigVMNode* InNode, URigVMPin* InPin);
+	void AddNodePin(URigVMNode* InNode, URigVMPin* InPin);
 	static void AddSubPin(URigVMPin* InParentPin, URigVMPin* InPin);
 	static bool EnsurePinValidity(URigVMPin* InPin, bool bRecursive);
 	static void ValidatePin(URigVMPin* InPin);
@@ -1348,7 +1350,7 @@ public:
 
 protected:
 
-	URigVMPin* CreatePinFromPinInfo(const FRigVMRegistry& InRegistry, const FRigVMPinInfoArray& InPreviousPinInfos, const FRigVMPinInfo& InPinInfo, const FString& InPinPath, UObject* InOuter) const;
+	URigVMPin* CreatePinFromPinInfo(const FRigVMRegistry& InRegistry, const FRigVMPinInfoArray& InPreviousPinInfos, const FRigVMPinInfo& InPinInfo, const FString& InPinPath, UObject* InOuter);
 
 	// backwards compatibility code
 	FRigVMClientPatchResult PatchRerouteNodesOnLoad();
@@ -1362,6 +1364,9 @@ protected:
 	FRigVMClientPatchResult PatchFunctionsWithInvalidReturnPaths();
 	FRigVMClientPatchResult PatchExecutePins();
 	FRigVMClientPatchResult PatchLazyPins();
+	FRigVMClientPatchResult PatchPinDefaultValues();
+
+	ERigVMPinDefaultValueType GetDefaultValueType(const URigVMPin* InPin, const FString& InDefaultValue) const;
 
 	// work to do after a duplication of the host asset
 	void PostDuplicateHost(const FString& InOldPathName, const FString& InNewPathName);
@@ -1440,6 +1445,7 @@ private:
 	ERigVMPinDirection UserLinkDirection;
 	bool bEnableTypeCasting;
 	bool bAllowPrivateFunctions;
+	TOptional<ERigVMPinDefaultValueType> OptionalDefaultValueType;
 
 	// temporary maps used for pin redirection
 	// only valid between Detach & ReattachLinksToPinObjects
@@ -1504,6 +1510,7 @@ private:
 	friend struct FRigVMBaseAction;
 	friend class URigVMCompiler;
 	friend struct FRigVMControllerObjectFactory;
+	friend struct FRigVMSetPinDefaultValueAction;
 	friend struct FRigVMAddRerouteNodeAction;
 	friend struct FRigVMChangePinTypeAction;
 	friend struct FRigVMInjectNodeIntoPinAction;
@@ -1512,6 +1519,7 @@ private:
 	friend class FRigVMControllerCompileBracketScope;
 	friend struct FRigVMPinInfoArray;
 	friend class FRigVMControllerNotifGuard;
+	friend class FRigVMDefaultValueTypeGuard;
 	friend struct FRigVMClient;
 	friend struct FRigVMActionWrapper;
 	friend class URigVMSchema;
@@ -1540,6 +1548,31 @@ private:
 	bool bPreviousSuspendNotifications;
 };
 
+class FRigVMDefaultValueTypeGuard
+{
+public:
+
+	FRigVMDefaultValueTypeGuard(URigVMController* InController, ERigVMPinDefaultValueType InDefaultValueType = ERigVMPinDefaultValueType::Override, bool bForce = false)
+		: Controller(InController)
+	{
+		bPreviousDefaultValueType = Controller->OptionalDefaultValueType;
+		if(!bPreviousDefaultValueType.IsSet() || bForce)
+		{
+			Controller->OptionalDefaultValueType = InDefaultValueType;
+		}
+	}
+
+	~FRigVMDefaultValueTypeGuard()
+	{
+		Controller->OptionalDefaultValueType = bPreviousDefaultValueType;
+	}
+
+private:
+
+	URigVMController* Controller;
+	TOptional<ERigVMPinDefaultValueType> bPreviousDefaultValueType;
+};
+
 USTRUCT()
 struct FRigVMController_CommonTypePerTemplate
 {
diff --git a/Engine/Plugins/Runtime/RigVM/Source/RigVMDeveloper/Public/RigVMModel/RigVMControllerActions.h b/Engine/Plugins/Runtime/RigVM/Source/RigVMDeveloper/Public/RigVMModel/RigVMControllerActions.h
--- a/Engine/Plugins/Runtime/RigVM/Source/RigVMDeveloper/Public/RigVMModel/RigVMControllerActions.h
+++ b/Engine/Plugins/Runtime/RigVM/Source/RigVMDeveloper/Public/RigVMModel/RigVMControllerActions.h
@@ -858,6 +858,12 @@ public:
 
 	UPROPERTY()
 	FString NewDefaultValue;
+
+	UPROPERTY()
+	ERigVMPinDefaultValueType OldDefaultValueType;
+
+	UPROPERTY()
+	ERigVMPinDefaultValueType NewDefaultValueType;
 };
 
 /**
@@ -871,7 +877,7 @@ struct FRigVMInsertArrayPinAction : public FRigVMBaseAction
 public:
 
 	FRigVMInsertArrayPinAction();
-	FRigVMInsertArrayPinAction(URigVMController* InController, URigVMPin* InArrayPin, int32 InIndex, const FString& InNewDefaultValue);
+	FRigVMInsertArrayPinAction(URigVMController* InController, URigVMPin* InArrayPin, int32 InIndex, const FString& InNewDefaultValue, const ERigVMPinDefaultValueType& InNewDefaultValueType);
 	virtual ~FRigVMInsertArrayPinAction() {};
 	virtual UScriptStruct* GetScriptStruct() const override { return FRigVMInsertArrayPinAction::StaticStruct(); }
 	virtual bool Undo() override;
@@ -885,6 +891,9 @@ public:
 
 	UPROPERTY()
 	FString NewDefaultValue;
+
+	UPROPERTY()
+	ERigVMPinDefaultValueType NewDefaultValueType;
 };
 
 /**
@@ -912,6 +921,9 @@ public:
 
 	UPROPERTY()
 	FString DefaultValue;
+
+	UPROPERTY()
+	ERigVMPinDefaultValueType DefaultValueType;
 };
 
 /**
diff --git a/Engine/Plugins/Runtime/RigVM/Source/RigVMDeveloper/Public/RigVMModel/RigVMNode.h b/Engine/Plugins/Runtime/RigVM/Source/RigVMDeveloper/Public/RigVMModel/RigVMNode.h
--- a/Engine/Plugins/Runtime/RigVM/Source/RigVMDeveloper/Public/RigVMModel/RigVMNode.h
+++ b/Engine/Plugins/Runtime/RigVM/Source/RigVMDeveloper/Public/RigVMModel/RigVMNode.h
@@ -68,11 +68,18 @@ public:
 	UFUNCTION(BlueprintCallable, Category = RigVMNode)
 	TArray<URigVMPin*> GetAllPinsRecursively() const;
 
+	// Returns the default value for a given pin
+	FString GetOriginalPinDefaultValue(const URigVMPin* InPin) const;
+
 	// Returns a Pin given it's partial pin path below
 	// this node (for example: "Color.R")
 	UFUNCTION(BlueprintCallable, Category = RigVMNode)
 	URigVMPin* FindPin(const FString& InPinPath) const;
 
+	// Returns a root pin given its name
+	UFUNCTION(BlueprintCallable, Category = RigVMNode)
+	URigVMPin* FindRootPinByName(const FName& InPinName) const;
+
 	// Returns all of the top-level orphaned Pins of this Node.
 	UFUNCTION(BlueprintCallable, Category = RigVMNode)
     virtual const TArray<URigVMPin*>& GetOrphanedPins() const;
@@ -335,11 +342,13 @@ private:
 
 protected:
 
-	virtual void InvalidateCache() {}
+	virtual void InvalidateCache();
 	virtual TArray<int32> GetInstructionsForVMImpl(const FRigVMExtendedExecuteContext& Context, URigVM* InVM, const FRigVMASTProxy& InProxy = FRigVMASTProxy()) const; 
 	virtual FText GetToolTipTextForPin(const URigVMPin* InPin) const;
 	virtual bool AllowsLinksOn(const URigVMPin* InPin) const { return true; }
 	virtual bool ShouldInputPinComputeLazily(const URigVMPin* InPin) const { return false; }
+	virtual FString GetOriginalDefaultValueForRootPin(const URigVMPin* InRootPin) const;
+
 	void UpdateDecoratorRootPinNames();
 
 	UPROPERTY()
@@ -386,6 +395,8 @@ private:
 	mutable TMap<uint32, TSharedPtr<FProfilingCache>> ProfilingCache;
 	static TArray<int32> EmptyInstructionArray;
 #endif
+
+	mutable TMap<FString, FString> CachedOriginalPinDefaultValues;
 	
 	friend class URigVMController;
 	friend class URigVMGraph;
@@ -394,4 +405,3 @@ private:
 	friend class FRigVMLexer;
 	friend class URigVMSchema;
 };
-
diff --git a/Engine/Plugins/Runtime/RigVM/Source/RigVMDeveloper/Public/RigVMModel/RigVMPin.h b/Engine/Plugins/Runtime/RigVM/Source/RigVMDeveloper/Public/RigVMModel/RigVMPin.h
--- a/Engine/Plugins/Runtime/RigVM/Source/RigVMDeveloper/Public/RigVMModel/RigVMPin.h
+++ b/Engine/Plugins/Runtime/RigVM/Source/RigVMDeveloper/Public/RigVMModel/RigVMPin.h
@@ -20,6 +20,8 @@ class URigVMPin;
 class URigVMLink;
 class URigVMVariableNode;
 
+extern RIGVMDEVELOPER_API TAutoConsoleVariable<bool> CVarRigVMEnablePinDefaultTypes;
+
 /**
  * The Injected Info is used for injecting a node on a pin.
  * Injected nodes are not visible to the user, but they are normal
@@ -73,6 +75,15 @@ public:
 	FWeakInfo GetWeakInfo() const;
 };
 
+UENUM()
+enum class ERigVMPinDefaultValueType : uint8
+{
+	AutoDetect, // Detect if this is an unchanged or overridden value based on the delta
+	Unset, // The value is unchanged and will remain the original default
+	Override, // The value is overridden by the user and should stay like this no matter what
+	KeepValueType, // The value type should be kept as well as the value itself (don't touch this)
+};
+
 /**
  * The Pin represents a single connector / pin on a node in the RigVM model.
  * Pins can be connected based on rules. Pins also provide access to a 'PinPath',
@@ -284,6 +295,9 @@ public:
 	// Returns true if this pin's subpins should be hidden in the UI
 	UFUNCTION(BlueprintCallable, Category = RigVMPin)
 	bool ShouldHideSubPins() const;
+
+	UFUNCTION(BlueprintCallable, Category = RigVMPin)
+	FString GetOriginalDefaultValue() const;
 	
 	// Returns the default value of the Pin as a string.
 	// Note that this value is computed based on the Pin's
@@ -300,8 +314,21 @@ public:
 	FString GetDefaultValueStoredByUserInterface() const;
 
 	// Returns true if the default value provided is valid
+	UFUNCTION(BlueprintPure, Category = RigVMPin)
 	bool IsValidDefaultValue(const FString& InDefaultValue) const;
 
+	// Returns true if the default value was ever changed by the user
+	UFUNCTION(BlueprintPure, Category = RigVMPin)
+	bool HasUserProvidedDefaultValue() const;
+
+	// Returns true if the pin can / may provide a default value 
+	UFUNCTION(BlueprintPure, Category = RigVMPin)
+	ERigVMPinDefaultValueType GetDefaultValueType() const { return DefaultValueType; }
+
+	// Returns true if the pin can / may provide a default value 
+	UFUNCTION(BlueprintPure, Category = RigVMPin)
+	bool CanProvideDefaultValue() const;
+
 	// Returns the default value clamped with the limit meta values defined by the UPROPERTY in URigVMUnitNodes 
 	FString ClampDefaultValueFromMetaData(const FString& InDefaultValue) const;
 
@@ -535,6 +562,9 @@ private:
 	UPROPERTY()
 	FString DefaultValue;
 
+	UPROPERTY()
+	ERigVMPinDefaultValueType DefaultValueType;
+
 	UPROPERTY()
 	FName CustomWidgetName;
 
diff --git a/Engine/Plugins/Runtime/RigVM/Source/RigVMEditor/Private/Editor/RigVMDetailsViewWrapperObject.cpp b/Engine/Plugins/Runtime/RigVM/Source/RigVMEditor/Private/Editor/RigVMDetailsViewWrapperObject.cpp
--- a/Engine/Plugins/Runtime/RigVM/Source/RigVMEditor/Private/Editor/RigVMDetailsViewWrapperObject.cpp
+++ b/Engine/Plugins/Runtime/RigVM/Source/RigVMEditor/Private/Editor/RigVMDetailsViewWrapperObject.cpp
@@ -3,6 +3,7 @@
 #include "Editor/RigVMDetailsViewWrapperObject.h"
 #include "RigVMCore/RigVMStruct.h"
 #include "RigVMModel/Nodes/RigVMUnitNode.h"
+#include "RigVMModel/Nodes/RigVMDispatchNode.h"
 #include "Modules/ModuleManager.h"
 #include "Algo/Sort.h"
 #include "RigVMTypeUtils.h"
@@ -220,6 +221,11 @@ UClass* URigVMDetailsViewWrapperObject::GetClassForNodes(TArray<URigVMNode*> InN
 
 	// determine if all nodes are unit nodes and match their script struct
 	TArray<UScriptStruct*> UnitStructs;
+	TArray<const FRigVMDispatchFactory*> DispatchFactories;
+	TArray<const FRigVMGraphFunctionIdentifier> FunctionIdentifiers;
+	
+	const UClass* NodeClass = InNodes[0]->GetClass();
+	bool bMatchingNodes = true;
 	for(URigVMNode* Node : InNodes)
 	{
 		if(URigVMUnitNode* UnitNode = Cast<URigVMUnitNode>(Node))
@@ -229,6 +235,38 @@ UClass* URigVMDetailsViewWrapperObject::GetClassForNodes(TArray<URigVMNode*> InN
 				UnitStructs.AddUnique(ScriptStruct);
 			}
 		}
+		else if(URigVMDispatchNode* DispatchNode = Cast<URigVMDispatchNode>(Node))
+		{
+			if(const FRigVMDispatchFactory* Factory = DispatchNode->GetFactory())
+			{
+				DispatchFactories.AddUnique(Factory);
+			}
+		}
+		else if(URigVMFunctionReferenceNode* FunctionReferenceNode = Cast<URigVMFunctionReferenceNode>(Node))
+		{
+			FunctionIdentifiers.AddUnique(FunctionReferenceNode->GetReferencedFunctionHeader().LibraryPointer);
+		}
+
+		if(NodeClass != Node->GetClass())
+		{
+			bMatchingNodes = false;
+		}
+	}
+
+	if(bMatchingNodes)
+	{
+		if(!UnitStructs.IsEmpty())
+		{
+			bMatchingNodes = UnitStructs.Num() == 1;
+		}
+		else if(!DispatchFactories.IsEmpty())
+		{
+			bMatchingNodes = DispatchFactories.Num() == 1;
+		}
+		else if(!FunctionIdentifiers.IsEmpty())
+		{
+			bMatchingNodes = FunctionIdentifiers.Num() == 1;
+		}
 	}
 
 	// find the matching pins on all passed in nodes
@@ -295,6 +333,36 @@ UClass* URigVMDetailsViewWrapperObject::GetClassForNodes(TArray<URigVMNode*> InN
 		}
 		Notation = FString::Printf(TEXT("%s(%s)"), *FString::Join(StructNames, TEXT("|")), *Notation);
 	}
+	if(DispatchFactories.Num() > 0)
+	{
+		// sort the factories to ensure we get the same notation each time
+		Algo::Sort(DispatchFactories, [](const FRigVMDispatchFactory* A, const FRigVMDispatchFactory* B) -> bool
+		{
+			return B->GetFactoryName().LexicalLess(A->GetFactoryName());
+		});
+
+		TArray<FString> FactoryNames;
+		for(const FRigVMDispatchFactory* Factory : DispatchFactories)
+		{
+			FactoryNames.Add(Factory->GetFactoryName().ToString());
+		}
+		Notation = FString::Printf(TEXT("%s(%s)"), *FString::Join(FactoryNames, TEXT("|")), *Notation);
+	}
+	if(FunctionIdentifiers.Num() > 0)
+	{
+		// sort the functions to ensure we get the same notation each time
+		Algo::Sort(FunctionIdentifiers, [](const FRigVMGraphFunctionIdentifier& A, const FRigVMGraphFunctionIdentifier& B) -> bool
+		{
+			return A.LibraryNodePath > B.LibraryNodePath;
+		});
+
+		TArray<FString> IdentifierStrings;
+		for(const FRigVMGraphFunctionIdentifier& Identifier : FunctionIdentifiers)
+		{
+			IdentifierStrings.Add(Identifier.LibraryNodePath);
+		}
+		Notation = FString::Printf(TEXT("%s(%s)"), *FString::Join(IdentifierStrings, TEXT("|")), *Notation);
+	}
 
 	const FPerClassInfo PerClassInfo(Notation);
 	if(UClass** ExistingClass = InfoToClass.Find(PerClassInfo))
@@ -485,6 +553,59 @@ UClass* URigVMDetailsViewWrapperObject::GetClassForNodes(TArray<URigVMNode*> InN
 			}
 		}
 	}
+	
+	if(bMatchingNodes)
+	{
+		if(DispatchFactories.Num() > 0)
+		{
+			const FRigVMDispatchFactory* Factory = DispatchFactories[0];
+
+			for(URigVMPin* Pin : PinsToInspect)
+			{
+				const FString DefaultValue = Factory->GetArgumentDefaultValue(Pin->GetFName(), Pin->GetTypeIndex());
+				if(!DefaultValue.IsEmpty())
+				{
+					if(FProperty* Property = WrapperClass->FindPropertyByName(Pin->GetFName()))
+					{
+						uint8* TargetMemory = Property->ContainerPtrToValuePtr<uint8>(CDO);
+						
+						// use error pipe to ignore errors
+						FRigVMPinDefaultValueImportErrorContext ErrorPipe;
+						Property->ImportText_Direct(*DefaultValue, TargetMemory, CDO, PPF_None, &ErrorPipe);
+					}
+				}
+			}
+		}
+
+		if(DispatchFactories.Num() > 0)
+		{
+			const FRigVMGraphFunctionIdentifier Identifier = FunctionIdentifiers[0];
+			const FRigVMGraphFunctionHeader Header = FRigVMGraphFunctionHeader::FindGraphFunctionHeader(Identifier);
+
+			for(URigVMPin* Pin : PinsToInspect)
+			{
+				const FRigVMGraphFunctionArgument* Argument = Header.Arguments.FindByPredicate([Pin](const FRigVMGraphFunctionArgument& InArgument)-> bool
+				{
+					return InArgument.Name == Pin->GetFName();
+				});
+				if(Argument)
+				{
+					if(!Argument->DefaultValue.IsEmpty())
+					{
+						if(FProperty* Property = WrapperClass->FindPropertyByName(Pin->GetFName()))
+						{
+							uint8* TargetMemory = Property->ContainerPtrToValuePtr<uint8>(CDO);
+							
+							// use error pipe to ignore errors
+							FRigVMPinDefaultValueImportErrorContext ErrorPipe;
+							Property->ImportText_Direct(*Argument->DefaultValue, TargetMemory, CDO, PPF_None, &ErrorPipe);
+						}
+					}
+				}
+			}
+		}
+
+	}
 
 #if WITH_EDITOR
 	FPropertyEditorModule& EditModule = FModuleManager::Get().GetModuleChecked<FPropertyEditorModule>("PropertyEditor");
diff --git a/Engine/Plugins/Runtime/RigVM/Source/RigVMEditor/Private/RigVMEditorModule.cpp b/Engine/Plugins/Runtime/RigVM/Source/RigVMEditor/Private/RigVMEditorModule.cpp
--- a/Engine/Plugins/Runtime/RigVM/Source/RigVMEditor/Private/RigVMEditorModule.cpp
+++ b/Engine/Plugins/Runtime/RigVM/Source/RigVMEditor/Private/RigVMEditorModule.cpp
@@ -547,6 +547,7 @@ void FRigVMEditorModule::GetNodeContextMenuActions(IRigVMClientHost* RigVMClient
 
 	GetNodeWorkflowContextMenuActions(RigVMClientHost, EdGraphNode, ModelNode, Menu);
 	GetNodeEventsContextMenuActions(RigVMClientHost, EdGraphNode, ModelNode, Menu);
+	GetNodeDefaultValueContextMenuActions(RigVMClientHost, EdGraphNode, ModelNode, Menu);
 	GetNodeConversionContextMenuActions(RigVMClientHost, EdGraphNode, ModelNode, Menu);
 	GetNodeDebugContextMenuActions(RigVMClientHost, EdGraphNode, ModelNode, Menu);
 	GetNodeVariablesContextMenuActions(RigVMClientHost, EdGraphNode, ModelNode, Menu);
@@ -646,6 +647,97 @@ void FRigVMEditorModule::GetNodeEventsContextMenuActions(IRigVMClientHost* RigVM
 	}
 }
 
+void FRigVMEditorModule::GetNodeDefaultValueContextMenuActions(IRigVMClientHost* RigVMClientHost, const URigVMEdGraphNode* EdGraphNode, URigVMNode* ModelNode, UToolMenu* Menu) const
+{
+	if(!CVarRigVMEnablePinDefaultTypes.GetValueOnAnyThread())
+	{
+		return;
+	}
+	
+	URigVMController* Controller = RigVMClientHost->GetRigVMClient()->GetController(ModelNode->GetGraph());
+	FToolMenuSection& DefaultValuesSection = Menu->AddSection("RigVMEditorContextMenuDefaultValues", LOCTEXT("DefaultValuesHeader", "Pin Values"));
+
+	DefaultValuesSection.AddMenuEntry(
+		"Override All",
+		LOCTEXT("OverrideAll", "Override All"),
+		LOCTEXT("OverrideAll_Tooltip", "Overrides all values"),
+		FSlateIcon(),
+		FUIAction(FExecuteAction::CreateLambda([Controller]()
+		{
+			const TArray<FName> SelectedNodeNames = Controller->GetGraph()->GetSelectNodes();
+			if(!SelectedNodeNames.IsEmpty())
+			{
+				FRigVMDefaultValueTypeGuard _(Controller, ERigVMPinDefaultValueType::Override);
+				Controller->OpenUndoBracket(TEXT("Override all pin values"));
+				for(const FName& SelectedNodeName : SelectedNodeNames)
+				{
+					if(const URigVMNode* Node = Controller->GetGraph()->FindNodeByName(SelectedNodeName))
+					{
+						for(const URigVMPin* Pin : Node->GetPins())
+						{
+							if(Pin->CanProvideDefaultValue())
+							{
+								FString DefaultValue = Pin->GetDefaultValue();
+								if(DefaultValue.IsEmpty())
+								{
+									DefaultValue = Pin->GetOriginalDefaultValue();
+								}
+								if(!DefaultValue.IsEmpty())
+								{
+									Controller->SetPinDefaultValue(Pin->GetPinPath(), DefaultValue);
+								}
+							}
+						}
+					}
+				}
+				Controller->CloseUndoBracket();
+			}
+		}))
+	);
+
+	DefaultValuesSection.AddMenuEntry(
+		"Reset Unchanged",
+		LOCTEXT("ResetUnchanged", "Reset Unchanged"),
+		LOCTEXT("ResetUnchanged_Tooltip", "Resets pin values that match the default"),
+		FSlateIcon(),
+		FUIAction(FExecuteAction::CreateLambda([Controller]()
+		{
+			const TArray<FName> SelectedNodeNames = Controller->GetGraph()->GetSelectNodes();
+			if(!SelectedNodeNames.IsEmpty())
+			{
+				FRigVMDefaultValueTypeGuard _(Controller, ERigVMPinDefaultValueType::Unset);
+				Controller->OpenUndoBracket(TEXT("Reset unchanged pin values"));
+				for(const FName& SelectedNodeName : SelectedNodeNames)
+				{
+					if(const URigVMNode* Node = Controller->GetGraph()->FindNodeByName(SelectedNodeName))
+					{
+						for(const URigVMPin* Pin : Node->GetPins())
+						{
+							if(Pin->CanProvideDefaultValue())
+							{
+								FString DefaultValue = Pin->GetDefaultValue();
+								const FString OriginalDefaultValue = Pin->GetOriginalDefaultValue();
+								if(!OriginalDefaultValue.IsEmpty())
+								{
+									if(DefaultValue.IsEmpty())
+									{
+										DefaultValue = OriginalDefaultValue;
+									}
+									if(DefaultValue.Equals(OriginalDefaultValue, ESearchCase::CaseSensitive))
+									{
+										Controller->SetPinDefaultValue(Pin->GetPinPath(), OriginalDefaultValue);
+									}
+								}
+							}
+						}
+					}
+				}
+				Controller->CloseUndoBracket();
+			}
+		}))
+	);
+}
+
 void FRigVMEditorModule::GetNodeConversionContextMenuActions(IRigVMClientHost* RigVMClientHost, const URigVMEdGraphNode* EdGraphNode, URigVMNode* ModelNode, UToolMenu* Menu) const
 {
 	if(URigVMDispatchNode* DispatchNode = Cast<URigVMDispatchNode>(ModelNode))
@@ -1401,6 +1493,17 @@ void FRigVMEditorModule::GetPinResetDefaultContextMenuActions(IRigVMClientHost*
 					Controller->ResetPinDefaultValue(ModelPin->GetPinPath());
 				})
 			));
+
+			Section.AddMenuEntry(
+				"OverridePinDefaultValue",
+				LOCTEXT("OverridePinDefaultValue", "Override Value"),
+				LOCTEXT("OverridePinDefaultValue_Tooltip", "Marks the pin value as an override."),
+				FSlateIcon(),
+				FUIAction(FExecuteAction::CreateLambda([Controller, ModelPin]() {
+					FRigVMDefaultValueTypeGuard _(Controller, ERigVMPinDefaultValueType::Override);
+					Controller->SetPinDefaultValue(ModelPin->GetPinPath(), ModelPin->GetDefaultValue());
+				})
+			));
 		}
 	}
 }
diff --git a/Engine/Plugins/Runtime/RigVM/Source/RigVMEditor/Public/RigVMEditorModule.h b/Engine/Plugins/Runtime/RigVM/Source/RigVMEditor/Public/RigVMEditorModule.h
--- a/Engine/Plugins/Runtime/RigVM/Source/RigVMEditor/Public/RigVMEditorModule.h
+++ b/Engine/Plugins/Runtime/RigVM/Source/RigVMEditor/Public/RigVMEditorModule.h
@@ -84,6 +84,7 @@ protected:
 	/** Specific section callbacks for the context menu */
 	virtual void GetNodeWorkflowContextMenuActions(IRigVMClientHost* RigVMClientHost, const URigVMEdGraphNode* EdGraphNode, URigVMNode* ModelNode, UToolMenu* Menu) const;
 	virtual void GetNodeEventsContextMenuActions(IRigVMClientHost* RigVMClientHost, const URigVMEdGraphNode* EdGraphNode, URigVMNode* ModelNode, UToolMenu* Menu) const;
+	virtual void GetNodeDefaultValueContextMenuActions(IRigVMClientHost* RigVMClientHost, const URigVMEdGraphNode* EdGraphNode, URigVMNode* ModelNode, UToolMenu* Menu) const;
 	virtual void GetNodeConversionContextMenuActions(IRigVMClientHost* RigVMClientHost, const URigVMEdGraphNode* EdGraphNode, URigVMNode* ModelNode, UToolMenu* Menu) const;
 	virtual void GetNodeDebugContextMenuActions(IRigVMClientHost* RigVMClientHost, const URigVMEdGraphNode* EdGraphNode, URigVMNode* ModelNode, UToolMenu* Menu) const;
 	virtual void GetNodeVariablesContextMenuActions(IRigVMClientHost* RigVMClientHost, const URigVMEdGraphNode* EdGraphNode, URigVMNode* ModelNode, UToolMenu* Menu) const;
