RigVM: Support for PinDefaultValueType in user interface - interim

#rnx

diff --git a/Engine/Plugins/Runtime/RigVM/Source/RigVMEditor/Private/Editor/RigVMGraphDetailCustomization.cpp b/Engine/Plugins/Runtime/RigVM/Source/RigVMEditor/Private/Editor/RigVMGraphDetailCustomization.cpp
--- a/Engine/Plugins/Runtime/RigVM/Source/RigVMEditor/Private/Editor/RigVMGraphDetailCustomization.cpp
+++ b/Engine/Plugins/Runtime/RigVM/Source/RigVMEditor/Private/Editor/RigVMGraphDetailCustomization.cpp
@@ -27,6 +27,7 @@
 #include "Widgets/SRigVMGraphPinVariableBinding.h"
 #include "InstancedPropertyBagStructureDataProvider.h"
 #include "Widgets/SRigVMGraphPinEnumPicker.h"
+#include "ScopedTransaction.h"
 
 #define LOCTEXT_NAMESPACE "RigVMGraphDetailCustomization"
 
@@ -1355,13 +1356,100 @@ void FRigVMWrappedNodeDetailCustomization::CustomizeDetails(IDetailLayoutBuilder
 	for (TFieldIterator<FProperty> PropertyIt(WrapperClass); PropertyIt; ++PropertyIt)
 	{
 		FProperty* Property = *PropertyIt;
-		TSharedPtr<IPropertyHandle> PropertyHandle = DetailLayout.GetProperty(Property->GetFName(), WrapperClass);
+		const FName PropertyName = Property->GetFName();
+		TSharedPtr<IPropertyHandle> PropertyHandle = DetailLayout.GetProperty(PropertyName, WrapperClass);
 		if (!PropertyHandle->IsValidHandle())
 		{
 			continue;
 		}
 		PropertiesToVisit.Add(PropertyHandle);
 
+		TAttribute<bool> HasUserProvidedDefaultValue = TAttribute<bool>::CreateLambda([this, PropertyName, PropertyHandle]() -> bool
+		{
+			if(CVarRigVMEnablePinDefaultTypes.GetValueOnAnyThread())
+			{
+				for(const TWeakObjectPtr<URigVMNode>& Node : NodesBeingCustomized)
+				{
+					if(Node.IsValid())
+					{
+						if(const URigVMPin* Pin = Node->FindRootPinByName(PropertyName))
+						{
+							if(Pin->HasUserProvidedDefaultValue())
+							{
+								return true;
+							}
+						}
+					}	
+				}
+			}
+			return PropertyHandle->DiffersFromDefault();
+		});
+
+		FResetToDefaultOverride ResetToDefault = FResetToDefaultOverride::Create(
+			HasUserProvidedDefaultValue,
+			FSimpleDelegate::CreateLambda([this, PropertyHandle, PropertyName]()
+			{
+				FScopedTransaction Transaction(LOCTEXT("ResetValueToDefault", "Reset Value To Default"));
+				const URigVMGraph* Graph = NodesBeingCustomized[0]->GetGraph();
+				URigVMController* Controller = BlueprintBeingCustomized->GetController(Graph);
+				FRigVMDefaultValueTypeGuard _(Controller, ERigVMPinDefaultValueType::Unset);
+
+				Controller->OpenUndoBracket(TEXT("Reset pin default value"));
+				for(const TWeakObjectPtr<URigVMNode> Node : NodesBeingCustomized)
+				{
+					if(const URigVMPin* Pin = Node->FindRootPinByName(PropertyName))
+					{
+						Controller->ResetPinDefaultValue(Pin->GetPinPath());
+					}
+				}
+				Controller->CloseUndoBracket();
+			})
+		);
+
+		static const FSlateFontInfo NameFont = FAppStyle::GetFontStyle( TEXT("PropertyWindow.NormalFont") );
+		
+		TSharedRef<SWidget> LabelWidget = PropertyHandle->CreatePropertyNameWidget();
+
+		/*
+		// in the future we may want some visual alignment of the label widget on top of the 
+		// reset arrow on the right to indicate the state of the default value change
+		TSharedRef<SHorizontalBox> LabelWidget = SNew(SHorizontalBox)
+		+ SHorizontalBox::Slot()
+		.FillWidth(1)
+		[
+			SNew(SBorder)
+			.HAlign(HAlign_Left)
+			.BorderImage_Lambda([HasUserProvidedDefaultValue]() -> const FSlateBrush*
+			{
+				if(CVarRigVMEnablePinDefaultTypes.GetValueOnAnyThread())
+				{
+					if(HasUserProvidedDefaultValue.Get())
+					{
+						static const FSlateBrush* BorderBrush = FAppStyle::Get().GetBrush("FloatingBorder");
+						return BorderBrush;
+					}
+				}
+				return nullptr;
+			})
+			.BorderBackgroundColor_Lambda([HasUserProvidedDefaultValue]() -> FSlateColor
+			{
+				if(CVarRigVMEnablePinDefaultTypes.GetValueOnAnyThread())
+				{
+					if(HasUserProvidedDefaultValue.Get())
+					{
+						return FSlateColor(FLinearColor::Red);
+					}
+				}
+				return FSlateColor(EStyleColor::Background);
+			})
+			[
+				SNew(STextBlock)
+				.Text(PropertyHandle->GetPropertyDisplayName())
+				.Font(NameFont)
+			]
+		];
+		*/
+
 		// check if any / all pins are bound to a variable
 		int32 PinsBoundToVariable = 0;
 		TArray<URigVMPin*> ModelPins;
@@ -1383,7 +1471,7 @@ void FRigVMWrappedNodeDetailCustomization::CustomizeDetails(IDetailLayoutBuilder
 					Row->CustomWidget()
 					.NameContent()
 					[
-						PropertyHandle->CreatePropertyNameWidget()
+						LabelWidget
 					]
 					.ValueContent()
 					[
@@ -1391,6 +1479,7 @@ void FRigVMWrappedNodeDetailCustomization::CustomizeDetails(IDetailLayoutBuilder
 						.ModelPins(ModelPins)
 						.Blueprint(BlueprintBeingCustomized)
 					];
+					continue;
 				}
 			}
 			else // in this case some pins are bound, and some are not - we'll hide the input value widget
@@ -1400,12 +1489,11 @@ void FRigVMWrappedNodeDetailCustomization::CustomizeDetails(IDetailLayoutBuilder
 					Row->CustomWidget()
 					.NameContent()
 					[
-						PropertyHandle->CreatePropertyNameWidget()
+						LabelWidget
 					];
+					continue;
 				}
 			}
-
-			continue;
 		}
 		
 		if (FNameProperty* NameProperty = CastField<FNameProperty>(Property))
@@ -1427,7 +1515,7 @@ void FRigVMWrappedNodeDetailCustomization::CustomizeDetails(IDetailLayoutBuilder
         				Row->CustomWidget()
 						.NameContent()
 						[
-							PropertyHandle->CreatePropertyNameWidget()
+							LabelWidget
 						]
 						.ValueContent()
 						[
@@ -1451,24 +1539,43 @@ void FRigVMWrappedNodeDetailCustomization::CustomizeDetails(IDetailLayoutBuilder
 									return FSlateColor::UseForeground();
 								})
 							]
-        				];
+						]
+        				.OverrideResetToDefault(ResetToDefault);
         			}        			
         			NameListWidgets.Add(Property->GetFName(), NameListWidget);
+       				continue;
         		}
-        		else
+        		
+        		if(IDetailPropertyRow* Row = DetailLayout.EditDefaultProperty(PropertyHandle))
         		{
-        			if(IDetailPropertyRow* Row = DetailLayout.EditDefaultProperty(PropertyHandle))
-        			{
-        				Row->CustomWidget()
-						.NameContent()
-						[
-							PropertyHandle->CreatePropertyNameWidget()
-						];
-        			}
+        			Row->CustomWidget()
+					.NameContent()
+					[
+						LabelWidget
+					]
+        			.OverrideResetToDefault(ResetToDefault);
+        			continue;
         		}
-        		continue;
         	}
         }
+
+		if(IDetailPropertyRow* Row = DetailLayout.EditDefaultProperty(PropertyHandle))
+		{
+			TSharedPtr<SWidget> ValueWidget = PropertyHandle->CreatePropertyValueWidgetWithCustomization(DetailLayout.GetDetailsView());
+
+			constexpr bool bShowChildren = true;
+			Row->CustomWidget(bShowChildren)
+			.NameContent()
+			[
+				LabelWidget
+			]
+			.ValueContent()
+			[
+				ValueWidget ? ValueWidget.ToSharedRef() : SNullWidget::NullWidget
+			]
+			.OverrideResetToDefault(ResetToDefault);
+			continue;
+		}
 	}
 
 	// now loop over all handles and determine expansion states of the corresponding pins
@@ -1576,6 +1683,7 @@ void FRigVMWrappedNodeDetailCustomization::SetNameListText(const FText& NewTypeI
 	{
 		if(URigVMPin* Pin = Node->FindPin(InProperty->GetName()))
 		{
+			FRigVMDefaultValueTypeGuard _(Controller, ERigVMPinDefaultValueType::Override);
 			Controller->SetPinDefaultValue(Pin->GetPinPath(), NewTypeInValue.ToString(), false, true, false, true);
 		}
 	}
diff --git a/Engine/Plugins/Runtime/RigVM/Source/RigVMEditor/Private/Widgets/SRigVMGraphNode.cpp b/Engine/Plugins/Runtime/RigVM/Source/RigVMEditor/Private/Widgets/SRigVMGraphNode.cpp
--- a/Engine/Plugins/Runtime/RigVM/Source/RigVMEditor/Private/Widgets/SRigVMGraphNode.cpp
+++ b/Engine/Plugins/Runtime/RigVM/Source/RigVMEditor/Private/Widgets/SRigVMGraphNode.cpp
@@ -663,26 +663,24 @@ FText SRigVMGraphNode::GetPinLabel(TWeakPtr<SGraphPin> GraphPin) const
 	return FText();
 }
 
-FSlateColor SRigVMGraphNode::GetPinTextColor(TWeakPtr<SGraphPin> GraphPin) const
+TOptional<FSlateColor> SRigVMGraphNode::GetHighlightColor(const SGraphPin* InGraphPin) const
 {
-	if(GraphPin.IsValid())
+	if(HasUserProvidedValue(InGraphPin))
 	{
-		if (GraphPin.Pin()->GetPinObj()->bOrphanedPin)
-		{
-			return FLinearColor::Red;
-		}
-
-		// If there is no schema there is no owning node (or basically this is a deleted node)
-		if (GraphNode)
-		{
-			if(!GraphNode->IsNodeEnabled() || GraphNode->IsDisplayAsDisabledForced() || !GraphPin.Pin()->IsEditingEnabled() || GraphNode->IsNodeUnrelated())
-			{
-				return FLinearColor(1.0f, 1.0f, 1.0f, 0.5f);
-			}
-		}
+		return FSlateColor(FLinearColor::Red);
 	}
+	return SGraphNode::GetHighlightColor(InGraphPin);
+}
 
-	return FLinearColor::White;
+TOptional<const FSlateBrush*> SRigVMGraphNode::GetPinBorder(const SGraphPin* InGraphPin) const
+{
+	if(HasUserProvidedValue(InGraphPin))
+	{
+		static const FName NAME_Pin_BackgroundHovered("Graph.Pin.BackgroundHovered");
+		static const FSlateBrush* CachedImg_Pin_BackgroundHovered = FAppStyle::GetBrush( NAME_Pin_BackgroundHovered );
+		return CachedImg_Pin_BackgroundHovered;
+	}
+	return SGraphNode::GetPinBorder(InGraphPin);
 }
 
 FSlateColor SRigVMGraphNode::GetVariableLabelTextColor(
@@ -2167,4 +2165,24 @@ void SRigVMGraphNode::UpdatePinTreeView()
 	];
 }
 
+bool SRigVMGraphNode::HasUserProvidedValue(const SGraphPin* InGraphPin) const
+{
+	if(!CVarRigVMEnablePinDefaultTypes.GetValueOnAnyThread())
+	{
+		return false;
+	}
+
+	if (const URigVMEdGraphNode* RigVMEdGraphNode = Cast<URigVMEdGraphNode>(GraphNode))
+	{
+		if (const UEdGraphPin* Pin = InGraphPin->GetPinObj())
+		{
+			if(const URigVMPin* ModelPin = RigVMEdGraphNode->FindModelPinFromGraphPin(Pin))
+			{
+				return ModelPin->HasUserProvidedDefaultValue();
+			}
+		}
+	}
+	return false;
+}
+
 #undef LOCTEXT_NAMESPACE
diff --git a/Engine/Plugins/Runtime/RigVM/Source/RigVMEditor/Public/Widgets/SRigVMGraphNode.h b/Engine/Plugins/Runtime/RigVM/Source/RigVMEditor/Public/Widgets/SRigVMGraphNode.h
--- a/Engine/Plugins/Runtime/RigVM/Source/RigVMEditor/Public/Widgets/SRigVMGraphNode.h
+++ b/Engine/Plugins/Runtime/RigVM/Source/RigVMEditor/Public/Widgets/SRigVMGraphNode.h
@@ -73,7 +73,8 @@ protected:
 
 	FText GetPinLabel(TWeakPtr<SGraphPin> GraphPin) const;
 
-	FSlateColor GetPinTextColor(TWeakPtr<SGraphPin> GraphPin) const;
+	virtual TOptional<FSlateColor> GetHighlightColor(const SGraphPin* InGraphPin) const override;
+	virtual TOptional<const FSlateBrush*> GetPinBorder(const SGraphPin* InGraphPin) const override;
 	FSlateColor GetVariableLabelTextColor(TWeakObjectPtr<URigVMFunctionReferenceNode> FunctionReferenceNode, FName InVariableName) const;
 	FText GetVariableLabelTooltipText(TWeakObjectPtr<URigVMBlueprint> InBlueprint, FName InVariableName) const;
 
@@ -94,6 +95,7 @@ protected:
 	FReply OnExpanderArrowClicked(int32 InPinInfoIndex);
 	void HandleModifiedEvent(ERigVMGraphNotifType InNotifType, URigVMGraph* InGraph, UObject* InSubject);
 	virtual void UpdatePinTreeView();
+	bool HasUserProvidedValue(const SGraphPin* InGraphPin) const;
 
 	/** Cached widget title area */
 	TSharedPtr<SOverlay> TitleAreaWidget;
diff --git a/Engine/Source/Editor/GraphEditor/Private/SGraphPin.cpp b/Engine/Source/Editor/GraphEditor/Private/SGraphPin.cpp
--- a/Engine/Source/Editor/GraphEditor/Private/SGraphPin.cpp
+++ b/Engine/Source/Editor/GraphEditor/Private/SGraphPin.cpp
@@ -1128,15 +1128,24 @@ const FSlateBrush* SGraphPin::GetSecondaryPinIcon() const
 const FSlateBrush* SGraphPin::GetPinBorder() const
 {
 	bool bIsMarkedPin = false;
-	TSharedPtr<SGraphPanel> OwnerPanelPtr = OwnerNodePtr.Pin()->GetOwnerPanel();
-	if (OwnerPanelPtr.IsValid())
+	const TSharedPtr<SGraphNode> OwnerNode = OwnerNodePtr.Pin();
+	if(OwnerNode.IsValid())
 	{
-		if (OwnerPanelPtr->MarkedPin.IsValid())
+		const TSharedPtr<SGraphPanel> OwnerPanelPtr = OwnerNodePtr.Pin()->GetOwnerPanel();
+		if (OwnerPanelPtr.IsValid())
 		{
-			bIsMarkedPin = (OwnerPanelPtr->MarkedPin.Pin() == SharedThis(this));
+			if (OwnerPanelPtr->MarkedPin.IsValid())
+			{
+				bIsMarkedPin = (OwnerPanelPtr->MarkedPin.Pin() == SharedThis(this));
+			}
+		}
+
+		TOptional<const FSlateBrush*> OptionalPinBorder = OwnerNode->GetPinBorder(this);
+		if(OptionalPinBorder.IsSet())
+		{
+			return OptionalPinBorder.GetValue();
 		}
 	}
-	UEdGraphPin* GraphPin = GetPinObj();
 	return (IsHovered() || bIsMarkedPin || bIsDiffHighlighted || bOnlyShowDefaultValue) ? CachedImg_Pin_BackgroundHovered : CachedImg_Pin_Background;
 }
 
@@ -1174,6 +1183,17 @@ FSlateColor SGraphPin::GetHighlightColor() const
 	{
 		return PinDiffColor.GetValue();
 	}
+	
+	const TSharedPtr<SGraphNode> OwnerNode = OwnerNodePtr.Pin();
+	if(OwnerNode.IsValid())
+	{
+		TOptional<FSlateColor> OptionalHighlightColor = OwnerNode->GetHighlightColor(this);
+		if(OptionalHighlightColor.IsSet())
+		{
+			return OptionalHighlightColor.GetValue();
+		}
+	}
+	
 	return GetPinColor();
 }
 
@@ -1203,6 +1223,16 @@ FSlateColor SGraphPin::GetSecondaryPinColor() const
 
 FSlateColor SGraphPin::GetPinTextColor() const
 {
+	const TSharedPtr<SGraphNode> OwnerNode = OwnerNodePtr.Pin();
+	if(OwnerNode.IsValid())
+	{
+		TOptional<FSlateColor> OptionalPinTextColor = OwnerNode->GetPinTextColor(this);
+		if(OptionalPinTextColor.IsSet())
+		{
+			return OptionalPinTextColor.GetValue();
+		}
+	}
+
 	UEdGraphPin* GraphPin = GetPinObj();
 	// If there is no schema there is no owning node (or basically this is a deleted node)
 	if (UEdGraphNode* GraphNode = GraphPin ? GraphPin->GetOwningNodeUnchecked() : nullptr)
diff --git a/Engine/Source/Editor/GraphEditor/Public/SGraphNode.h b/Engine/Source/Editor/GraphEditor/Public/SGraphNode.h
--- a/Engine/Source/Editor/GraphEditor/Public/SGraphNode.h
+++ b/Engine/Source/Editor/GraphEditor/Public/SGraphNode.h
@@ -278,7 +278,16 @@ public:
 
 	/** Returns true if the node is hiding its pins */
 	virtual bool IsHidingPinWidgets() const { return false; }
-	
+
+	/** Allows the node to override the pin border image */
+	virtual TOptional<const FSlateBrush*> GetPinBorder(const SGraphPin* InGraphPin) const { return TOptional<const FSlateBrush*>(); }
+
+	/** Allows the node to override the pin highlight color */
+	virtual TOptional<FSlateColor> GetHighlightColor(const SGraphPin* InGraphPin) const { return TOptional<FSlateColor>(); }
+
+	/** Allows the node to override the pin text color */
+	virtual TOptional<FSlateColor> GetPinTextColor(const SGraphPin* InGraphPin) const { return TOptional<FSlateColor>(); }
+
 protected:
 	SGraphNode();
 
