[RigVMRegistry] fixed recursive locking when adding a user defined struct to the animation attribute list in project setting

diff --git a/Engine/Plugins/Runtime/RigVM/Source/RigVM/Private/RigVMCore/RigVMRegistry.cpp b/Engine/Plugins/Runtime/RigVM/Source/RigVM/Private/RigVMCore/RigVMRegistry.cpp
--- a/Engine/Plugins/Runtime/RigVM/Source/RigVM/Private/RigVMCore/RigVMRegistry.cpp
+++ b/Engine/Plugins/Runtime/RigVM/Source/RigVM/Private/RigVMCore/RigVMRegistry.cpp
@@ -424,15 +424,12 @@ void FRigVMRegistry_NoLock::OnAssetRenamed_NoLock(const FAssetData& InAssetData,
 	}
 }
 
-void FRigVMRegistry_NoLock::OnAssetRemoved_NoLock(const FAssetData& InAssetData)
+bool FRigVMRegistry_NoLock::OnAssetRemoved_NoLock(const FAssetData& InAssetData)
 {
-	if (RemoveType_NoLock(InAssetData.ToSoftObjectPath(), InAssetData.GetClass()))
-	{
-		OnRigVMRegistryChangedDelegate.Broadcast();
-	}
+	return RemoveType_NoLock(InAssetData.ToSoftObjectPath(), InAssetData.GetClass());
 }
 
-void FRigVMRegistry_NoLock::OnPluginUnloaded_NoLock(IPlugin& InPlugin)
+bool FRigVMRegistry_NoLock::OnPluginUnloaded_NoLock(IPlugin& InPlugin)
 {
 	const FString PluginContentPath = InPlugin.GetMountedAssetPath();
 
@@ -463,10 +460,7 @@ void FRigVMRegistry_NoLock::OnPluginUnloaded_NoLock(IPlugin& InPlugin)
 		}
 	}
 
-	if (bRegistryChanged)
-	{
-		OnRigVMRegistryChangedDelegate.Broadcast();
-	}
+	return bRegistryChanged;
 }
 
 void FRigVMRegistry_NoLock::OnAnimationAttributeTypesChanged_NoLock(const UScriptStruct* InStruct, bool bIsAdded)
@@ -479,7 +473,6 @@ void FRigVMRegistry_NoLock::OnAnimationAttributeTypesChanged_NoLock(const UScrip
 	if (bIsAdded)
 	{
 		FindOrAddType_NoLock(FRigVMTemplateArgumentType(const_cast<UScriptStruct*>(InStruct)));
-		OnRigVMRegistryChangedDelegate.Broadcast();		
 	}
 }
 
@@ -929,6 +922,47 @@ FRigVMRegistry_RWLock& FRigVMRegistry_RWLock::Get()
 	return s_RigVMRegistry;
 }
 
+void FRigVMRegistry_RWLock::OnAssetRemoved(const FAssetData& InAssetData)
+{
+	bool bAssetRemoved = false;
+	{
+		FConditionalWriteScopeLock _(*this);
+		bAssetRemoved = Super::OnAssetRemoved_NoLock(InAssetData);
+	}
+
+	if (bAssetRemoved)
+	{
+		OnRigVMRegistryChangedDelegate.Broadcast();
+	}
+}
+
+void FRigVMRegistry_RWLock::OnPluginUnloaded(IPlugin& InPlugin)
+{
+	bool bRegistryChanged = false;
+	{
+		FConditionalWriteScopeLock _(*this);
+		bRegistryChanged = Super::OnPluginUnloaded_NoLock(InPlugin);
+	}
+		
+	if (bRegistryChanged)
+	{
+		OnRigVMRegistryChangedDelegate.Broadcast();
+	}
+}
+
+void FRigVMRegistry_RWLock::OnAnimationAttributeTypesChanged(const UScriptStruct* InStruct, bool bIsAdded)
+{
+	{
+		FConditionalWriteScopeLock _(*this);
+		Super::OnAnimationAttributeTypesChanged_NoLock(InStruct, bIsAdded);
+	}
+
+	if (bIsAdded)
+	{
+		OnRigVMRegistryChangedDelegate.Broadcast();		
+	}
+}
+
 TRigVMTypeIndex FRigVMRegistry_NoLock::GetTypeIndex_NoLock(const FRigVMTemplateArgumentType& InType) const
 {
 	if(const TRigVMTypeIndex* Index = TypeToIndex.Find(InType))
diff --git a/Engine/Plugins/Runtime/RigVM/Source/RigVM/Public/RigVMCore/RigVMRegistry.h b/Engine/Plugins/Runtime/RigVM/Source/RigVM/Public/RigVMCore/RigVMRegistry.h
--- a/Engine/Plugins/Runtime/RigVM/Source/RigVM/Public/RigVMCore/RigVMRegistry.h
+++ b/Engine/Plugins/Runtime/RigVM/Source/RigVM/Public/RigVMCore/RigVMRegistry.h
@@ -100,10 +100,10 @@ public:
 	virtual void OnAssetRenamed_NoLock(const FAssetData& InAssetData, const FString& InOldObjectPath);
 	
 	// Update the registry when old types are removed
-	virtual void OnAssetRemoved_NoLock(const FAssetData& InAssetData);
+	virtual bool OnAssetRemoved_NoLock(const FAssetData& InAssetData);
 
 	// Removes all types associated with a plugin that's being unloaded. 
-	virtual void OnPluginUnloaded_NoLock(IPlugin& InPlugin);
+	virtual bool OnPluginUnloaded_NoLock(IPlugin& InPlugin);
 	
 	// Update the registry when new types are added to the attribute system so that they can be selected
 	// on Attribute Nodes
@@ -295,8 +295,6 @@ public:
 
 	static const TArray<UScriptStruct*>& GetMathTypes();
 
-	// Notifies other system that types have been added/removed, and template permutations have been updated
-	FOnRigVMRegistryChanged& OnRigVMRegistryChanged() { return OnRigVMRegistryChangedDelegate; }
 
 	// Returns a unique hash per type index
 	virtual uint32 GetHashForType_NoLock(TRigVMTypeIndex InTypeIndex) const;
@@ -406,8 +404,7 @@ protected:
 	// All allowed structs
 	TSet<TObjectPtr<const UScriptStruct>> AllowedStructs;
 	
-	// Notifies other system that types have been added/removed, and template permutations have been updated
-	FOnRigVMRegistryChanged OnRigVMRegistryChangedDelegate;
+
 
 	// If this is true the registry is currently refreshing all types
 	bool bIsRefreshingEngineTypes;
@@ -608,26 +605,17 @@ public:
 	}
 	
 	// Update the registry when old types are removed
-    void OnAssetRemoved(const FAssetData& InAssetData)
-	{
-		FConditionalWriteScopeLock _(*this);
-		Super::OnAssetRemoved_NoLock(InAssetData);
-	}
+    void OnAssetRemoved(const FAssetData& InAssetData);
 
 	// Removes all types associated with a plugin that's being unloaded. 
-	void OnPluginUnloaded(IPlugin& InPlugin)
-	{
-		FConditionalWriteScopeLock _(*this);
-		Super::OnPluginUnloaded_NoLock(InPlugin);
-	}
-	
+	void OnPluginUnloaded(IPlugin& InPlugin);
+
 	// Update the registry when new types are added to the attribute system so that they can be selected
 	// on Attribute Nodes
-	void OnAnimationAttributeTypesChanged(const UScriptStruct* InStruct, bool bIsAdded)
-	{
-		FConditionalWriteScopeLock _(*this);
-		Super::OnAnimationAttributeTypesChanged_NoLock(InStruct, bIsAdded);
-	}
+	void OnAnimationAttributeTypesChanged(const UScriptStruct* InStruct, bool bIsAdded);
+
+	// Notifies other system that types have been added/removed, and template permutations have been updated
+	FOnRigVMRegistryChanged& OnRigVMRegistryChanged() { return OnRigVMRegistryChangedDelegate; }
 	
 	// Clear the registry
 	void Reset(bool bLockRegistry = true)
@@ -980,16 +968,16 @@ protected:
 	
 	// Update the registry when old types are removed
 	// Note: Only call NoLock methods on the FRigVMRegistry_NoLock
-	virtual void OnAssetRemoved_NoLock(const FAssetData& InAssetData) override
+	virtual bool OnAssetRemoved_NoLock(const FAssetData& InAssetData) override
 	{
-		Super::OnAssetRemoved_NoLock(InAssetData);
+		return Super::OnAssetRemoved_NoLock(InAssetData);
 	}
 
 	// Removes all types associated with a plugin that's being unloaded. 
 	// Note: Only call NoLock methods on the FRigVMRegistry_NoLock
-	virtual void OnPluginUnloaded_NoLock(IPlugin& InPlugin) override
+	virtual bool OnPluginUnloaded_NoLock(IPlugin& InPlugin) override
 	{
-		Super::OnPluginUnloaded_NoLock(InPlugin);
+		return Super::OnPluginUnloaded_NoLock(InPlugin);
 	}
 	
 	// Update the registry when new types are added to the attribute system so that they can be selected
@@ -1274,6 +1262,9 @@ private:
 	mutable FRWLock Lock;
 	mutable std::atomic<ELockType> LockType;
 	mutable std::atomic<int32> LockCount;
+
+	// Notifies other system that types have been added/removed, and template permutations have been updated
+	FOnRigVMRegistryChanged OnRigVMRegistryChangedDelegate;
 	
 	friend struct FRigVMStruct;
 	friend struct FRigVMTemplate;
