RigVM: Follow up fix for 32821686

#rnx

diff --git a/Engine/Plugins/Runtime/RigVM/Source/RigVMEditor/Private/Editor/RigVMDetailsViewWrapperObject.cpp b/Engine/Plugins/Runtime/RigVM/Source/RigVMEditor/Private/Editor/RigVMDetailsViewWrapperObject.cpp
--- a/Engine/Plugins/Runtime/RigVM/Source/RigVMEditor/Private/Editor/RigVMDetailsViewWrapperObject.cpp
+++ b/Engine/Plugins/Runtime/RigVM/Source/RigVMEditor/Private/Editor/RigVMDetailsViewWrapperObject.cpp
@@ -577,7 +577,7 @@ UClass* URigVMDetailsViewWrapperObject::GetClassForNodes(TArray<URigVMNode*> InN
 			}
 		}
 
-		if(DispatchFactories.Num() > 0)
+		if(FunctionIdentifiers.Num() > 0)
 		{
 			const FRigVMGraphFunctionIdentifier Identifier = FunctionIdentifiers[0];
 			const FRigVMGraphFunctionHeader Header = FRigVMGraphFunctionHeader::FindGraphFunctionHeader(Identifier);
