RigVM: Make sure to create controllers within the right client

#rnx

diff --git a/Engine/Plugins/Runtime/RigVM/Source/RigVMDeveloper/Private/RigVMModel/RigVMController.cpp b/Engine/Plugins/Runtime/RigVM/Source/RigVMDeveloper/Private/RigVMModel/RigVMController.cpp
--- a/Engine/Plugins/Runtime/RigVM/Source/RigVMDeveloper/Private/RigVMModel/RigVMController.cpp
+++ b/Engine/Plugins/Runtime/RigVM/Source/RigVMDeveloper/Private/RigVMModel/RigVMController.cpp
@@ -263,8 +263,18 @@ URigVMController* URigVMController::GetControllerForGraph(const URigVMGraph* InG
 	{
 		return const_cast<URigVMController*>(this);
 	}
+
+	// get the client host from the graph rather than this controller,
+	// since the graph may be in a different host / rigvmblueprint.
+	IRigVMClientHost* ClientHost = InGraph->GetImplementingOuter<IRigVMClientHost>();
+
+	// the graph may not be nested under a client in some unit tests
+	if(ClientHost == nullptr)
+	{
+		ClientHost = GetImplementingOuter<IRigVMClientHost>();
+	}
 	
-	if(IRigVMClientHost* ClientHost = GetImplementingOuter<IRigVMClientHost>())
+	if(ClientHost)
 	{
 		if(FRigVMClient* Client = ClientHost->GetRigVMClient())
 		{
