Control Rig: Change RigVMGraphFunctionIdentifier comparison operator to use only the subpath of the library node path
#rnx

diff --git a/Engine/Plugins/Runtime/RigVM/Source/RigVM/Public/RigVMCore/RigVMGraphFunctionDefinition.h b/Engine/Plugins/Runtime/RigVM/Source/RigVM/Public/RigVMCore/RigVMGraphFunctionDefinition.h
--- a/Engine/Plugins/Runtime/RigVM/Source/RigVM/Public/RigVMCore/RigVMGraphFunctionDefinition.h
+++ b/Engine/Plugins/Runtime/RigVM/Source/RigVM/Public/RigVMCore/RigVMGraphFunctionDefinition.h
@@ -393,7 +393,7 @@ struct RIGVM_API FRigVMGraphFunctionIdentifier
 
 	bool operator==(const FRigVMGraphFunctionIdentifier& Other) const
 	{
-		return HostObject == Other.HostObject && LibraryNodePath == Other.LibraryNodePath;
+		return HostObject == Other.HostObject && FSoftObjectPath(LibraryNodePath).GetSubPathString() == FSoftObjectPath(Other.LibraryNodePath).GetSubPathString();
 	}
 
 	bool IsValid() const
