Control Rig: Report all log messages, even from level sequencer rigs where no log has been setup
#rnx

diff --git a/Engine/Plugins/Runtime/RigVM/Source/RigVM/Private/RigVMCore/RigVMExecuteContext.cpp b/Engine/Plugins/Runtime/RigVM/Source/RigVM/Private/RigVMCore/RigVMExecuteContext.cpp
--- a/Engine/Plugins/Runtime/RigVM/Source/RigVM/Private/RigVMCore/RigVMExecuteContext.cpp
+++ b/Engine/Plugins/Runtime/RigVM/Source/RigVM/Private/RigVMCore/RigVMExecuteContext.cpp
@@ -5,6 +5,11 @@
 
 #include UE_INLINE_GENERATED_CPP_BY_NAME(RigVMExecuteContext)
 
+TAutoConsoleVariable<bool> CVarRigVMReportAllMessages(
+	TEXT("RigVM.ReportAllMessages"),
+	false,
+	TEXT("Report all log messages, even when no logging procedure has been setup."));
+
 void FRigVMExecuteContext::SetOwningComponent(const USceneComponent* InOwningComponent)
 {
 	OwningComponent = InOwningComponent;
diff --git a/Engine/Plugins/Runtime/RigVM/Source/RigVM/Public/RigVMCore/RigVMExecuteContext.h b/Engine/Plugins/Runtime/RigVM/Source/RigVM/Public/RigVMCore/RigVMExecuteContext.h
--- a/Engine/Plugins/Runtime/RigVM/Source/RigVM/Public/RigVMCore/RigVMExecuteContext.h
+++ b/Engine/Plugins/Runtime/RigVM/Source/RigVM/Public/RigVMCore/RigVMExecuteContext.h
@@ -39,6 +39,8 @@ struct FRigVMExecuteContext;
 struct FRigVMExtendedExecuteContext;
 struct FRigVMLogSettings;
 
+extern RIGVM_API TAutoConsoleVariable<bool> CVarRigVMReportAllMessages;
+
 USTRUCT()
 struct RIGVM_API FRigVMSlice
 {
@@ -400,6 +402,10 @@ struct FRigVMExecuteContext
 		{
 			Log->Report(InLogSettings, InFunctionName, InInstructionIndex, InMessage);
 		}
+		else if(CVarRigVMReportAllMessages.GetValueOnAnyThread())
+		{
+			Logf(InLogSettings, TEXT("Instruction[%d] '%s': '%s'"), InstructionIndex, *InFunctionName.ToString(), *InMessage);
+		}
 	}
 #endif
 
