RigVM: Improve safeguards for FRigVMRegistry::AddReferencedObjects

#rnx

diff --git a/Engine/Plugins/Runtime/RigVM/Source/RigVM/Private/RigVMCore/RigVMRegistry.cpp b/Engine/Plugins/Runtime/RigVM/Source/RigVM/Private/RigVMCore/RigVMRegistry.cpp
--- a/Engine/Plugins/Runtime/RigVM/Source/RigVM/Private/RigVMCore/RigVMRegistry.cpp
+++ b/Engine/Plugins/Runtime/RigVM/Source/RigVM/Private/RigVMCore/RigVMRegistry.cpp
@@ -51,9 +51,25 @@ void FRigVMRegistry_NoLock::AddReferencedObjects(FReferenceCollector& Collector)
 	{
 		// the Object needs to be checked for validity since it may be a user defined type (struct or enum)
 		// which is about to get removed. 
-		if (IsValid(Type.Type.CPPTypeObject))
-		{
-			Collector.AddReferencedObject(Type.Type.CPPTypeObject);
+		if (!!Type.Type.CPPTypeObject &&
+			IsValid(Type.Type.CPPTypeObject) &&
+			Type.Type.CPPTypeObject->IsValidLowLevelFast() &&
+			!Type.Type.CPPTypeObject->IsNative() &&
+			!Type.Type.CPPTypeObject->IsUnreachable())
+		{
+			// make sure the object is part of the GUObjectArray and can be retrieved
+			// so that GC doesn't crash after receiving the referenced object
+			const int32 ObjectIndex = GUObjectArray.ObjectToIndex(Type.Type.CPPTypeObject);
+			if(ObjectIndex != INDEX_NONE)
+			{
+				if(const FUObjectItem* Item = GUObjectArray.IndexToObject(ObjectIndex))
+				{
+					if(Item->Object == Type.Type.CPPTypeObject)
+					{
+						Collector.AddReferencedObject(Type.Type.CPPTypeObject);
+					}
+				}
+			}
 		}
 	}
 }
