RigVMRegistry fixes

Expose locking version of RegisterStructTypes

diff --git a/Engine/Plugins/Runtime/RigVM/Source/RigVM/Public/RigVMCore/RigVMRegistry.h b/Engine/Plugins/Runtime/RigVM/Source/RigVM/Public/RigVMCore/RigVMRegistry.h
--- a/Engine/Plugins/Runtime/RigVM/Source/RigVM/Public/RigVMCore/RigVMRegistry.h
+++ b/Engine/Plugins/Runtime/RigVM/Source/RigVM/Public/RigVMCore/RigVMRegistry.h
@@ -581,6 +581,13 @@ public:
 		Super::RegisterObjectTypes_NoLock(InClasses);
 	}
 
+	// Register a set of allowed struct types
+	void RegisterStructTypes(TConstArrayView<UScriptStruct*> InStructs, bool bLockRegistry = true)
+	{
+		FConditionalWriteScopeLock _(*this, bLockRegistry);
+		Super::RegisterStructTypes_NoLock(InStructs);
+	}
+
 	// Refreshes the list and finds the function pointers
 	// based on the names.
 	void RefreshEngineTypes()
