RigVM: Creation of variants for functions
#rnx

diff --git a/Engine/Plugins/Animation/ControlRig/Source/ControlRigDeveloper/Private/ControlRigBlueprint.cpp b/Engine/Plugins/Animation/ControlRig/Source/ControlRigDeveloper/Private/ControlRigBlueprint.cpp
--- a/Engine/Plugins/Animation/ControlRig/Source/ControlRigDeveloper/Private/ControlRigBlueprint.cpp
+++ b/Engine/Plugins/Animation/ControlRig/Source/ControlRigDeveloper/Private/ControlRigBlueprint.cpp
@@ -1903,7 +1903,6 @@ void UControlRigBlueprint::PatchFunctionsOnLoad()
 {
 	URigVMBlueprintGeneratedClass* CRGeneratedClass = GetRigVMBlueprintGeneratedClass();
 	FRigVMGraphFunctionStore& Store = CRGeneratedClass->GraphFunctionStore;
-	const URigVMFunctionLibrary* Library = GetLocalFunctionLibrary();
 
 	TMap<URigVMLibraryNode*, FRigVMGraphFunctionHeader> OldHeaders;
 
@@ -1966,13 +1965,24 @@ void UControlRigBlueprint::PatchFunctionsOnLoad()
 				Header.ExternalVariables = OldHeader->ExternalVariables;
 				Header.Dependencies = OldHeader->Dependencies;
 			}
+
+			const FRigVMVariant* Variant = FunctionLibrary->GetFunctionVariant(LibraryNode->GetFName());
+			if (!Variant)
+			{
+				Header.Variant.Guid = FRigVMVariant::GenerateGUID(Header.LibraryPointer.LibraryNodePath);
+			}
+			else
+			{
+				Header.Variant = *Variant;
+			}
+			
 			Store.AddFunction(Header, bIsPublic);
 			
 		}
 	}
 
 	// Update dependencies and external variables if needed
-	for (URigVMLibraryNode* LibraryNode : Library->GetFunctions())
+	for (URigVMLibraryNode* LibraryNode : FunctionLibrary->GetFunctions())
 	{
 		GetRigVMClient()->UpdateExternalVariablesForFunction(LibraryNode);
 		GetRigVMClient()->UpdateDependenciesForFunction(LibraryNode);
@@ -2455,7 +2465,7 @@ void UControlRigBlueprint::RefreshModuleVariables()
 
 	if (UModularRigController* Controller = GetModularRigController())
 	{
-		Controller->RefreshModuleVariables();
+		Controller->RefreshModuleVariables(false);
 	}
 }
 
diff --git a/Engine/Plugins/Runtime/RigVM/Source/RigVM/Public/RigVMCore/RigVMGraphFunctionDefinition.h b/Engine/Plugins/Runtime/RigVM/Source/RigVM/Public/RigVMCore/RigVMGraphFunctionDefinition.h
--- a/Engine/Plugins/Runtime/RigVM/Source/RigVM/Public/RigVMCore/RigVMGraphFunctionDefinition.h
+++ b/Engine/Plugins/Runtime/RigVM/Source/RigVM/Public/RigVMCore/RigVMGraphFunctionDefinition.h
@@ -8,6 +8,7 @@
 #include "UObject/UE5MainStreamObjectVersion.h"
 #include "UObject/FortniteMainBranchObjectVersion.h"
 #include "UObject/UE5ReleaseStreamObjectVersion.h"
+#include "RigVMVariant.h"
 #include "RigVMGraphFunctionDefinition.generated.h"
 
 class IRigVMGraphFunctionHost;
@@ -461,6 +462,9 @@ struct RIGVM_API FRigVMGraphFunctionHeader
 	UPROPERTY(BlueprintReadOnly, VisibleAnywhere, Category=FunctionHeader)
 	FRigVMGraphFunctionIdentifier LibraryPointer;
 
+	UPROPERTY(BlueprintReadOnly, VisibleAnywhere, Category=FunctionIdentifier)
+	FRigVMVariant Variant;
+
 	UPROPERTY(BlueprintReadOnly, VisibleAnywhere, Category=FunctionHeader)
 	FName Name;
 
@@ -528,6 +532,12 @@ struct RIGVM_API FRigVMGraphFunctionHeader
 		Ar.UsingCustomVersion(FRigVMObjectVersion::GUID);
 		
 		Ar << Data.LibraryPointer;
+
+		if (!Ar.IsLoading() || Ar.CustomVer(FRigVMObjectVersion::GUID) >= FRigVMObjectVersion::AddVariantToFunctionIdentifier)
+		{
+			Ar << Data.Variant;
+		}
+		
 		Ar << Data.Name;
 		Ar << Data.NodeTitle;
 		Ar << Data.NodeColor;
diff --git a/Engine/Plugins/Runtime/RigVM/Source/RigVM/Public/RigVMCore/RigVMVariant.h b/Engine/Plugins/Runtime/RigVM/Source/RigVM/Public/RigVMCore/RigVMVariant.h
new file mode 100644
--- /dev/null
+++ b/Engine/Plugins/Runtime/RigVM/Source/RigVM/Public/RigVMCore/RigVMVariant.h
@@ -0,0 +1,78 @@
+// Copyright Epic Games, Inc. All Rights Reserved.
+
+#pragma once
+
+#include "RigVMObjectVersion.h"
+#include "RigVMVariant.generated.h"
+
+struct FRigVMVariantRef;
+
+USTRUCT(BlueprintType)
+struct RIGVM_API FRigVMTag
+{
+	GENERATED_BODY()
+	
+	// User applied tag
+	UPROPERTY(BlueprintReadWrite, VisibleAnywhere, Category=Variant)
+	FString Label;
+	
+	friend FArchive& operator<<(FArchive& Ar, FRigVMTag& Data)
+	{
+		Ar.UsingCustomVersion(FRigVMObjectVersion::GUID);
+		
+		Ar << Data.Label;
+		return Ar;
+	}
+};
+
+USTRUCT(BlueprintType)
+struct RIGVM_API FRigVMVariant
+{
+	GENERATED_BODY()
+
+	FRigVMVariant(const FRigVMVariant& OtherVariant)
+	{
+		Guid = OtherVariant.Guid;
+		Tags = OtherVariant.Tags;
+	}
+
+	// Guid which is shared by all variants of the same element
+	UPROPERTY(BlueprintReadOnly, VisibleAnywhere, Category=Variant)
+	FGuid Guid;
+	
+	// Tags applied to this variant
+	UPROPERTY(BlueprintReadWrite, VisibleAnywhere, Category=Variant)
+	TArray<FRigVMTag> Tags;
+
+	FRigVMVariant()
+		: Guid(FGuid()) {}
+
+	friend FArchive& operator<<(FArchive& Ar, FRigVMVariant& Data)
+	{
+		Ar.UsingCustomVersion(FRigVMObjectVersion::GUID);
+		
+		Ar << Data.Guid;
+		Ar << Data.Tags;
+		return Ar;
+	}
+
+	static FGuid GenerateGUID(const FString InPath = FString())
+	{
+		if (!InPath.IsEmpty())
+		{
+			return FGuid::NewDeterministicGuid(InPath);
+		}
+		return FGuid::NewGuid();
+	}
+};
+
+// This struct should not be serialized.
+// It is generated on demand.
+USTRUCT(BlueprintType)
+struct RIGVM_API FRigVMVariantRef
+{
+	GENERATED_BODY()
+	
+	FSoftObjectPath ObjectPath;
+	FRigVMVariant Variant;
+};
diff --git a/Engine/Plugins/Runtime/RigVM/Source/RigVM/Public/RigVMObjectVersion.h b/Engine/Plugins/Runtime/RigVM/Source/RigVM/Public/RigVMObjectVersion.h
--- a/Engine/Plugins/Runtime/RigVM/Source/RigVM/Public/RigVMObjectVersion.h
+++ b/Engine/Plugins/Runtime/RigVM/Source/RigVM/Public/RigVMObjectVersion.h
@@ -42,6 +42,9 @@ struct RIGVM_API FRigVMObjectVersion
 		// Removing library node FSoftObjectPath from FRigVMGraphFunctionIdentifier
 		RemoveLibraryNodeReferenceFromFunctionIdentifier,
 
+		// Adding variant struct to function identifier
+		AddVariantToFunctionIdentifier,
+
 		// -----<new versions can be added above this line>-------------------------------------------------
 		VersionPlusOne,
 		LatestVersion = VersionPlusOne - 1
diff --git a/Engine/Plugins/Runtime/RigVM/Source/RigVMDeveloper/Private/RigVMBlueprint.cpp b/Engine/Plugins/Runtime/RigVM/Source/RigVMDeveloper/Private/RigVMBlueprint.cpp
--- a/Engine/Plugins/Runtime/RigVM/Source/RigVMDeveloper/Private/RigVMBlueprint.cpp
+++ b/Engine/Plugins/Runtime/RigVM/Source/RigVMDeveloper/Private/RigVMBlueprint.cpp
@@ -856,6 +856,19 @@ void URigVMBlueprint::PreSave(FObjectPreSaveContext ObjectSaveContext)
 		{
 			PublicGraphFunctions[i] = RigClass->GetRigVMGraphFunctionStore()->PublicFunctions[i].Header;
 		}
+
+		URigVMFunctionLibrary* FunctionLibrary = GetLocalFunctionLibrary();
+		FunctionLibrary->FunctionToVariant.Reset();
+		for (int32 Pass=0; Pass<2; ++Pass)
+		{
+			const TArray<FRigVMGraphFunctionData>& Functions = (Pass == 0) ?
+				RigClass->GetRigVMGraphFunctionStore()->PrivateFunctions
+				: RigClass->GetRigVMGraphFunctionStore()->PublicFunctions;
+			for (const FRigVMGraphFunctionData& Function : Functions)
+			{
+				FunctionLibrary->FunctionToVariant.Add(Function.Header.Name, Function.Header.Variant);
+			}
+		}
 	}
 
 #if WITH_EDITORONLY_DATA
@@ -2320,6 +2333,88 @@ TArray<FRigVMExternalDependency> URigVMBlueprint::GetExternalDependenciesForCate
 	return Dependencies;
 }
 
+TArray<FRigVMVariantRef> URigVMBlueprint::GatherFunctionVariantRefsForAsset(const FAssetData& InAssetData)
+{
+	TArray<FRigVMVariantRef> Result;
+
+	if (!InAssetData.IsInstanceOf(URigVMBlueprint::StaticClass()) &&
+		!InAssetData.IsInstanceOf(URigVMBlueprintGeneratedClass::StaticClass()))
+	{
+		return Result;
+	}
+
+	// If the asset is loaded, gather the function variants from the function store
+	// which will include private functions
+	if (InAssetData.IsAssetLoaded())
+	{
+		UObject* AssetObject = InAssetData.GetAsset();
+		URigVMBlueprintGeneratedClass* GeneratedClass = Cast<URigVMBlueprintGeneratedClass>(AssetObject);
+		if (!GeneratedClass)
+		{
+			if (URigVMBlueprint* RigVMBlueprint = Cast<URigVMBlueprint>(AssetObject))
+			{
+				GeneratedClass = RigVMBlueprint->GetRigVMBlueprintGeneratedClass();
+			}
+		}
+
+		if (GeneratedClass)
+		{
+			for (int32 Pass=0; Pass<2; Pass++)
+			{
+				TArray<FRigVMGraphFunctionData>& Functions = (Pass == 0)
+					? GeneratedClass->GraphFunctionStore.PrivateFunctions
+					: GeneratedClass->GraphFunctionStore.PublicFunctions;
+				
+				for (FRigVMGraphFunctionData& Data : Functions)
+				{
+					FRigVMVariantRef& VariantRef = Result.Add_GetRef(FRigVMVariantRef());
+					VariantRef.Variant = Data.Header.Variant;
+					VariantRef.ObjectPath = Data.Header.LibraryPointer.GetNodeSoftPath();
+				}
+			}
+
+			return Result;
+		}
+	}
+
+	// If asset is not loaded, gather function variants from metadata
+	FString PublicGraphFunctionsString = InAssetData.GetTagValueRef<FString>(TEXT("PublicGraphFunctions"));
+	TArray<FRigVMGraphFunctionHeader> PublicFunctions;
+	static const FName FunctionsPropertyName = GET_MEMBER_NAME_CHECKED(URigVMBlueprint, PublicGraphFunctions);
+	FArrayProperty* PublicGraphFunctionsProperty = CastField<FArrayProperty>(URigVMBlueprint::StaticClass()->FindPropertyByName(FunctionsPropertyName));
+	PublicGraphFunctionsProperty->ImportText_Direct(*PublicGraphFunctionsString, &PublicFunctions, nullptr, EPropertyPortFlags::PPF_None);
+
+	for(const FRigVMGraphFunctionHeader& PublicFunction : PublicFunctions)
+	{
+		FRigVMVariantRef& VariantRef = Result.Add_GetRef(FRigVMVariantRef());
+		VariantRef.Variant = PublicFunction.Variant;
+		VariantRef.ObjectPath = PublicFunction.LibraryPointer.GetNodeSoftPath();
+	}
+
+	return Result;
+}
+
+TArray<FRigVMVariantRef> URigVMBlueprint::FindFunctionVariantRefs(const FGuid& InGuid)
+{
+	TArray<FRigVMVariantRef> Result;
+	const FAssetRegistryModule& AssetRegistryModule = FModuleManager::LoadModuleChecked<FAssetRegistryModule>("AssetRegistry");
+	TArray<FAssetData> Assets;
+	FARFilter AssetFilter;
+	AssetFilter.ClassPaths.Add(URigVMBlueprint::StaticClass()->GetClassPathName());
+	AssetFilter.ClassPaths.Add(URigVMBlueprintGeneratedClass::StaticClass()->GetClassPathName());
+	AssetFilter.bRecursiveClasses = true;
+	AssetRegistryModule.Get().GetAssets(AssetFilter, Assets);
+
+	for (const FAssetData& Asset : Assets)
+	{
+		TArray<FRigVMVariantRef> VariantRefs = URigVMBlueprint::GatherFunctionVariantRefsForAsset(Asset);
+		VariantRefs = VariantRefs.FilterByPredicate([InGuid](const FRigVMVariantRef& VariantRef) { return VariantRef.Variant.Guid == InGuid; });
+		Result.Append(VariantRefs);
+	}
+
+	return Result;
+}
+
 URigVMGraph* URigVMBlueprint::GetTemplateModel(bool bIsFunctionLibrary)
 {
 #if WITH_EDITORONLY_DATA
@@ -3708,6 +3803,17 @@ void URigVMBlueprint::PatchFunctionsOnLoad()
 				Header.ExternalVariables = OldHeader->ExternalVariables;
 				Header.Dependencies = OldHeader->Dependencies;
 			}
+
+			const FRigVMVariant* Variant = FunctionLibrary->GetFunctionVariant(LibraryNode->GetFName());
+			if (!Variant)
+			{
+				Header.Variant.Guid = FRigVMVariant::GenerateGUID(Header.LibraryPointer.LibraryNodePath);
+			}
+			else
+			{
+				Header.Variant = *Variant;
+			}
+			
 			Store.AddFunction(Header, bIsPublic);
 			
 		}
diff --git a/Engine/Plugins/Runtime/RigVM/Source/RigVMDeveloper/Private/RigVMModel/Nodes/RigVMLibraryNode.cpp b/Engine/Plugins/Runtime/RigVM/Source/RigVMDeveloper/Private/RigVMModel/Nodes/RigVMLibraryNode.cpp
--- a/Engine/Plugins/Runtime/RigVM/Source/RigVMDeveloper/Private/RigVMModel/Nodes/RigVMLibraryNode.cpp
+++ b/Engine/Plugins/Runtime/RigVM/Source/RigVMDeveloper/Private/RigVMModel/Nodes/RigVMLibraryNode.cpp
@@ -235,8 +235,21 @@ FRigVMGraphFunctionHeader URigVMLibraryNode::GetFunctionHeader(IRigVMGraphFuncti
 	{
 	    Header.LibraryPointer.LibraryNodePath = this->GetPathName();
 	    Header.LibraryPointer.HostObject = Cast<UObject>(InHostObject);
-	}	
-    
+	}
+
+	UObject* HostPtr = Cast<UObject>(Header.LibraryPointer.HostObject.ResolveObject());
+	if (IRigVMGraphFunctionHost* Host = Cast<IRigVMGraphFunctionHost>(HostPtr))
+	{
+		if (FRigVMGraphFunctionData* Data = Host->GetRigVMGraphFunctionStore()->FindFunction(Header.LibraryPointer))
+		{
+			Header.Variant = Data->Header.Variant;
+		}
+	}
+	if (!Header.Variant.Guid.IsValid())
+	{
+		Header.Variant.Guid = FRigVMVariant::GenerateGUID();
+	}
+
     Header.Name = GetFName();
 	Header.Description = GetNodeDescription();
     Header.Category = GetNodeCategory();
@@ -278,7 +291,6 @@ FRigVMGraphFunctionHeader URigVMLibraryNode::GetFunctionHeader(IRigVMGraphFuncti
 	// If the header already exists, try to find it and get its external variables and dependencies
 	TArray<FString> Dependencies;
 	TArray<FRigVMExternalVariable> ExternalVariables;
-	UObject* HostPtr = Cast<UObject>(Header.LibraryPointer.HostObject.ResolveObject());
 	if (IRigVMGraphFunctionHost* Host = Cast<IRigVMGraphFunctionHost>(HostPtr))
 	{
 		if (FRigVMGraphFunctionData* Data = Host->GetRigVMGraphFunctionStore()->FindFunction(Header.LibraryPointer))
diff --git a/Engine/Plugins/Runtime/RigVM/Source/RigVMDeveloper/Private/RigVMModel/RigVMClient.cpp b/Engine/Plugins/Runtime/RigVM/Source/RigVMDeveloper/Private/RigVMModel/RigVMClient.cpp
--- a/Engine/Plugins/Runtime/RigVM/Source/RigVMDeveloper/Private/RigVMModel/RigVMClient.cpp
+++ b/Engine/Plugins/Runtime/RigVM/Source/RigVMDeveloper/Private/RigVMModel/RigVMClient.cpp
@@ -1122,7 +1122,7 @@ void FRigVMClient::HandleGraphModifiedEvent(ERigVMGraphNotifType InNotifType, UR
 			}
 			break;
 		}
-		case ERigVMGraphNotifType::NodeRenamed: // A node has been renamed in the graph (Subject == URigVMNode)
+		case ERigVMGraphNotifType::FunctionRenamed: // A node has been renamed in the graph (Subject == URigVMNode)
 		{
 			if(InSubject->GetOuter()->IsA<URigVMFunctionLibrary>())
 			{
@@ -1134,7 +1134,9 @@ void FRigVMClient::HandleGraphModifiedEvent(ERigVMGraphNotifType InNotifType, UR
 						if (FRigVMGraphFunctionData* Data = Host->GetRigVMGraphFunctionStore()->FindFunctionByName(CollapseNode->GetPreviousFName()))
 						{
 							const FRigVMGraphFunctionIdentifier PreviousFunctionId = Data->Header.LibraryPointer;
+							const FRigVMVariant Variant = Data->Header.Variant;
 							Data->Header = CollapseNode->GetFunctionHeader();
+							Data->Header.Variant = Variant;
 
 							if (const FRigVMFunctionReferenceArray* FunctionReferencesPtr = BuildData->GraphFunctionReferences.Find(PreviousFunctionId))
 							{
diff --git a/Engine/Plugins/Runtime/RigVM/Source/RigVMDeveloper/Private/RigVMModel/RigVMController.cpp b/Engine/Plugins/Runtime/RigVM/Source/RigVMDeveloper/Private/RigVMModel/RigVMController.cpp
--- a/Engine/Plugins/Runtime/RigVM/Source/RigVMDeveloper/Private/RigVMModel/RigVMController.cpp
+++ b/Engine/Plugins/Runtime/RigVM/Source/RigVMDeveloper/Private/RigVMModel/RigVMController.cpp
@@ -1,6 +1,7 @@
 // Copyright Epic Games, Inc. All Rights Reserved.
 
 #include "RigVMModel/RigVMController.h"
+#include "RigVMBlueprint.h"
 #include "RigVMModel/RigVMControllerActions.h"
 #include "RigVMModel/Nodes/RigVMFunctionEntryNode.h"
 #include "RigVMModel/Nodes/RigVMFunctionReturnNode.h"
@@ -6854,6 +6855,7 @@ bool URigVMController::RenameNode(URigVMNode* InNode, const FName& InNewName, bo
 		return false;
 	}
 
+	Notify(ERigVMGraphNotifType::FunctionRenamed, InNode);
 	Notify(ERigVMGraphNotifType::NodeRenamed, InNode);
 
 	// update the links once more
@@ -12101,6 +12103,142 @@ bool URigVMController::IsFunctionPublic(const FName& InFunctionName)
 	return false;
 }
 
+URigVMLibraryNode* URigVMController::CreateFunctionVariant(const FName& InFunctionName, const FName& InVariantName, bool bSetupUndoRedo, bool bPrintPythonCommand)
+{
+	URigVMLibraryNode* Result = nullptr;
+	if (!IsValidGraph())
+	{
+		return Result;
+	}
+
+	if (!bIsTransacting && !IsGraphEditable())
+	{
+		return Result;
+	}
+
+	URigVMGraph* Graph = GetGraph();
+	check(Graph);
+
+	if (!Graph->IsA<URigVMFunctionLibrary>())
+	{
+		ReportError(TEXT("Can only add function variant in library graphs."));
+		return Result;
+	}
+
+	IRigVMClientHost* ClientHost = GetImplementingOuter<IRigVMClientHost>();
+	if (!ClientHost)
+	{
+		ReportErrorf(TEXT("Cannot find client host."));
+		return Result;
+	}
+
+	IRigVMGraphFunctionHost* FunctionHost = ClientHost->GetRigVMGraphFunctionHost();
+	if (!FunctionHost)
+	{
+		ReportErrorf(TEXT("Cannot find function host."));
+		return Result;
+	}
+
+	bool bIsPublic;
+	FRigVMGraphFunctionData* OriginalFunction = FunctionHost->GetRigVMGraphFunctionStore()->FindFunctionByName(InFunctionName, &bIsPublic);
+	if (!OriginalFunction)
+	{
+		ReportErrorf(TEXT("Cannot find function %s."), *InFunctionName.ToString());
+		return Result;
+	}
+
+	URigVMLibraryNode* FunctionToClone = Cast<URigVMLibraryNode>(Graph->FindNodeByName(InFunctionName));
+	if (!FunctionToClone)
+	{
+		ReportErrorf(TEXT("Cannot find function node %s."), *InFunctionName.ToString());
+		return Result;
+	}
+
+	const FName VariantName = *GetSchema()->GetValidNodeName(Graph, InVariantName.IsNone() ? InFunctionName.ToString() : InVariantName.ToString());
+
+	FRigVMBaseAction BaseAction(this);
+	if(bSetupUndoRedo)
+	{
+		BaseAction.SetTitle(FString::Printf(TEXT("Create function variant %s with name %s"), *InFunctionName.ToString(), *InVariantName.ToString()));
+		GetActionStack()->BeginAction(BaseAction);
+		GetActionStack()->AddAction(FRigVMCreateFunctionVariantAction(this, InFunctionName, VariantName));
+	}
+
+	ClientHost->GetRigVMClient()->UpdateGraphFunctionSerializedGraph(FunctionToClone);
+	TArray<FName> NodeNames = ImportNodesFromText(OriginalFunction->SerializedCollapsedNode, false);
+
+	if (!NodeNames.IsEmpty())
+	{
+		if (VariantName != NodeNames[0])
+		{
+			RenameFunction(NodeNames[0], VariantName, false);
+		}
+		FRigVMVariant NewVariant(OriginalFunction->Header.Variant);
+		if (FRigVMGraphFunctionData* NewFunction = FunctionHost->GetRigVMGraphFunctionStore()->FindFunctionByName(VariantName))
+		{
+			const FName NewName = NewFunction->Header.Name;
+			const FRigVMGraphFunctionIdentifier Identifier = NewFunction->Header.LibraryPointer;
+			NewFunction->Header = OriginalFunction->Header;
+			NewFunction->Header.Name = NewName;
+			NewFunction->Header.LibraryPointer = Identifier;
+
+			MarkFunctionAsPublic(NewName, bIsPublic, false);
+		}
+
+		Result = Cast<URigVMLibraryNode>(Graph->FindNodeByName(VariantName));
+	}
+	else if(bSetupUndoRedo)
+	{
+		GetActionStack()->CancelAction(BaseAction);
+	}
+
+	if (bSetupUndoRedo)
+	{
+		GetActionStack()->EndAction(BaseAction);
+	}
+
+	if (bPrintPythonCommand)
+	{
+		const FString GraphName = GetSchema()->GetSanitizedGraphName(GetGraph()->GetGraphName());
+
+		RigVMPythonUtils::Print(GetSchema()->GetGraphOuterName(GetGraph()), 
+			FString::Printf(TEXT("blueprint.get_controller_by_name('%s').create_function_variant('%s', '%s')"),
+				*GraphName,
+				*InFunctionName.ToString(),
+				*InVariantName.ToString()));
+	}
+
+	return Result;
+}
+
+TArray<FRigVMVariantRef> URigVMController::FindVariantsOfFunction(const FName& InFunctionName)
+{
+	TArray<FRigVMVariantRef> Result;
+	
+    IRigVMClientHost* ClientHost = GetImplementingOuter<IRigVMClientHost>();
+    if (!ClientHost)
+    {
+    	ReportErrorf(TEXT("Cannot find client host."));
+    	return Result;
+    }
+
+    IRigVMGraphFunctionHost* FunctionHost = ClientHost->GetRigVMGraphFunctionHost();
+    if (!FunctionHost)
+    {
+    	ReportErrorf(TEXT("Cannot find function host."));
+    	return Result;
+    }
+
+    FRigVMGraphFunctionData* FunctionData = FunctionHost->GetRigVMGraphFunctionStore()->FindFunctionByName(InFunctionName);
+    if (!FunctionData)
+    {
+    	ReportErrorf(TEXT("Cannot find function %s."), *InFunctionName.ToString());
+    	return Result;
+    }
+
+	return URigVMBlueprint::FindFunctionVariantRefs(FunctionData->Header.Variant.Guid);
+}
+
 FRigVMGraphVariableDescription URigVMController::AddLocalVariable(const FName& InVariableName, const FString& InCPPType, UObject* InCPPTypeObject, const FString& InDefaultValue, bool bSetupUndoRedo, bool bPrintPythonCommand)
 {
 	FRigVMGraphVariableDescription NewVariable;
diff --git a/Engine/Plugins/Runtime/RigVM/Source/RigVMDeveloper/Private/RigVMModel/RigVMControllerActions.cpp b/Engine/Plugins/Runtime/RigVM/Source/RigVMDeveloper/Private/RigVMModel/RigVMControllerActions.cpp
--- a/Engine/Plugins/Runtime/RigVM/Source/RigVMDeveloper/Private/RigVMModel/RigVMControllerActions.cpp
+++ b/Engine/Plugins/Runtime/RigVM/Source/RigVMDeveloper/Private/RigVMModel/RigVMControllerActions.cpp
@@ -2032,6 +2032,43 @@ bool FRigVMMarkFunctionPublicAction::Redo()
 	return false;
 }
 
+FRigVMCreateFunctionVariantAction::FRigVMCreateFunctionVariantAction()
+: FRigVMBaseAction(nullptr)
+, FunctionName(NAME_None)
+{
+}
+
+FRigVMCreateFunctionVariantAction::FRigVMCreateFunctionVariantAction(URigVMController* InController, const FName& InFunctionName, const FName& InNewFunctionName)
+: FRigVMBaseAction(InController)
+, FunctionName(InFunctionName)
+, NewFunctionName(InNewFunctionName)
+{
+}
+
+bool FRigVMCreateFunctionVariantAction::Undo()
+{
+	if (!FRigVMBaseAction::Undo())
+	{
+		return false;
+	}
+	return GetController()->RemoveFunctionFromLibrary(NewFunctionName);
+}
+
+bool FRigVMCreateFunctionVariantAction::Redo()
+{
+	if(!CanUndoRedo())
+	{
+		return false;
+	}
+#if WITH_EDITOR
+	if (GetController()->CreateFunctionVariant(FunctionName, NewFunctionName, false))
+	{
+		return FRigVMBaseAction::Redo();
+	}
+#endif
+	return false;
+}
+
 FRigVMImportFromTextAction::FRigVMImportFromTextAction()
 : FRigVMBaseAction(nullptr)
 {
diff --git a/Engine/Plugins/Runtime/RigVM/Source/RigVMDeveloper/Public/RigVMBlueprint.h b/Engine/Plugins/Runtime/RigVM/Source/RigVMDeveloper/Public/RigVMBlueprint.h
--- a/Engine/Plugins/Runtime/RigVM/Source/RigVMDeveloper/Public/RigVMBlueprint.h
+++ b/Engine/Plugins/Runtime/RigVM/Source/RigVMDeveloper/Public/RigVMBlueprint.h
@@ -415,6 +415,9 @@ public:
 	FOnRigVMRequestInspectMemoryStorage& OnRequestInspectMemoryStorage() { return OnRequestInspectMemoryStorageEvent; }
 	void RequestInspectMemoryStorage(const TArray<FRigVMMemoryStorageStruct*>& InMemoryStorageStructs) { OnRequestInspectMemoryStorageEvent.Broadcast(InMemoryStorageStructs); }
 
+	static TArray<FRigVMVariantRef> GatherFunctionVariantRefsForAsset(const FAssetData& InAssetData);
+	static TArray<FRigVMVariantRef> FindFunctionVariantRefs(const FGuid& InGuid);
+
 #endif	// #if WITH_EDITOR
 
 	
diff --git a/Engine/Plugins/Runtime/RigVM/Source/RigVMDeveloper/Public/RigVMModel/RigVMController.h b/Engine/Plugins/Runtime/RigVM/Source/RigVMDeveloper/Public/RigVMModel/RigVMController.h
--- a/Engine/Plugins/Runtime/RigVM/Source/RigVMDeveloper/Public/RigVMModel/RigVMController.h
+++ b/Engine/Plugins/Runtime/RigVM/Source/RigVMDeveloper/Public/RigVMModel/RigVMController.h
@@ -960,6 +960,14 @@ public:
 	UFUNCTION(BlueprintCallable, Category = RigVMController)
 	bool IsFunctionPublic(const FName& InFunctionName);
 
+	// Creates a variant of a function given the name of an existing function variant
+	UFUNCTION(BlueprintCallable, Category = RigVMController)
+	URigVMLibraryNode* CreateFunctionVariant(const FName& InFunctionName, const FName& InVariantName = NAME_None, bool bSetupUndoRedo = true, bool bPrintPythonCommand = false);
+
+	// Returns all variant refs related to the given function
+	UFUNCTION(BlueprintCallable, Category = RigVMController)
+	TArray<FRigVMVariantRef> FindVariantsOfFunction(const FName& InFunctionName);
+
 	// Add a local variable to the graph
 	UFUNCTION(BlueprintCallable, Category = RigVMController)
 	FRigVMGraphVariableDescription AddLocalVariable(const FName& InVariableName, const FString& InCPPType, UObject* InCPPTypeObject, const FString& InDefaultValue, bool bSetupUndoRedo = true, bool bPrintPythonCommand = false);
diff --git a/Engine/Plugins/Runtime/RigVM/Source/RigVMDeveloper/Public/RigVMModel/RigVMControllerActions.h b/Engine/Plugins/Runtime/RigVM/Source/RigVMDeveloper/Public/RigVMModel/RigVMControllerActions.h
--- a/Engine/Plugins/Runtime/RigVM/Source/RigVMDeveloper/Public/RigVMModel/RigVMControllerActions.h
+++ b/Engine/Plugins/Runtime/RigVM/Source/RigVMDeveloper/Public/RigVMModel/RigVMControllerActions.h
@@ -1411,6 +1411,30 @@ public:
 	bool bIsPublic;
 };
 
+/**
+ * An action marking a function as public/private.
+ */
+USTRUCT()
+struct FRigVMCreateFunctionVariantAction : public FRigVMBaseAction
+{
+	GENERATED_BODY()
+
+public:
+
+	FRigVMCreateFunctionVariantAction();
+	FRigVMCreateFunctionVariantAction(URigVMController* InController, const FName& InFunctionName, const FName& InNewFunctionName);
+	virtual ~FRigVMCreateFunctionVariantAction() {};
+	virtual UScriptStruct* GetScriptStruct() const override { return FRigVMCreateFunctionVariantAction::StaticStruct(); }
+	virtual bool Undo() override;
+	virtual bool Redo() override;
+
+	UPROPERTY()
+	FName FunctionName;
+
+	UPROPERTY()
+	FName NewFunctionName;
+};
+
 /**
  * An action importing nodes and links from text
  */
diff --git a/Engine/Plugins/Runtime/RigVM/Source/RigVMDeveloper/Public/RigVMModel/RigVMFunctionLibrary.h b/Engine/Plugins/Runtime/RigVM/Source/RigVMDeveloper/Public/RigVMModel/RigVMFunctionLibrary.h
--- a/Engine/Plugins/Runtime/RigVM/Source/RigVMDeveloper/Public/RigVMModel/RigVMFunctionLibrary.h
+++ b/Engine/Plugins/Runtime/RigVM/Source/RigVMDeveloper/Public/RigVMModel/RigVMFunctionLibrary.h
@@ -75,11 +75,16 @@ public:
 
 	bool IsFunctionPublic(const FName& InFunctionName) const { return PublicFunctionNames.Contains(InFunctionName); }
 
+	const FRigVMVariant* GetFunctionVariant(const FName& InFunctionName) const { return FunctionToVariant.Find(InFunctionName); }
+
 private:
 
 	UPROPERTY()
 	TArray<FName> PublicFunctionNames;
 
+	UPROPERTY()
+	TMap<FName, FRigVMVariant> FunctionToVariant;
+
 #if WITH_EDITORONLY_DATA
 	UPROPERTY()
 	TMap< TObjectPtr<URigVMLibraryNode>, FRigVMFunctionReferenceArray > FunctionReferences_DEPRECATED;
diff --git a/Engine/Plugins/Runtime/RigVM/Source/RigVMDeveloper/Public/RigVMModel/RigVMNotifications.h b/Engine/Plugins/Runtime/RigVM/Source/RigVMDeveloper/Public/RigVMModel/RigVMNotifications.h
--- a/Engine/Plugins/Runtime/RigVM/Source/RigVMDeveloper/Public/RigVMModel/RigVMNotifications.h
+++ b/Engine/Plugins/Runtime/RigVM/Source/RigVMDeveloper/Public/RigVMModel/RigVMNotifications.h
@@ -45,6 +45,7 @@ enum class ERigVMGraphNotifType : uint8
 	InteractionBracketCanceled, // A bracket has been canceled (Subject == nullptr)
 	PinBoundVariableChanged, // A pin has been bound or unbound to / from a variable (Subject == URigVMPin)
 	NodeRenamed, // A node has been renamed in the graph (Subject == URigVMNode)
+	FunctionRenamed, // A function has been renamed in the graph (Subject == URigVMLibraryNode)
 	NodeReferenceChanged, // A node has changed it's referenced function
 	NodeCategoryChanged, // A node's category has changed (Subject == URigVMNode)
 	NodeKeywordsChanged, // A node's keywords have changed (Subject == URigVMNode)
