RigVM: Move gather/find function variants to RigVMBuildData
#rnx

diff --git a/Engine/Plugins/Animation/ControlRig/Source/ControlRigDeveloper/Private/ControlRigBlueprint.cpp b/Engine/Plugins/Animation/ControlRig/Source/ControlRigDeveloper/Private/ControlRigBlueprint.cpp
--- a/Engine/Plugins/Animation/ControlRig/Source/ControlRigDeveloper/Private/ControlRigBlueprint.cpp
+++ b/Engine/Plugins/Animation/ControlRig/Source/ControlRigDeveloper/Private/ControlRigBlueprint.cpp
@@ -1969,7 +1969,7 @@ void UControlRigBlueprint::PatchFunctionsOnLoad()
 			const FRigVMVariant* Variant = FunctionLibrary->GetFunctionVariant(LibraryNode->GetFName());
 			if (!Variant)
 			{
-				Header.Variant.Guid = FRigVMVariant::GenerateGUID(Header.LibraryPointer.LibraryNodePath);
+				Header.Variant.Guid = FRigVMVariant::GenerateGUID(Header.LibraryPointer.GetLibraryNodePath());
 			}
 			else
 			{
diff --git a/Engine/Plugins/Runtime/RigVM/Source/RigVM/Private/RigVMCore/RigVMGraphFunctionDefinition.cpp b/Engine/Plugins/Runtime/RigVM/Source/RigVM/Private/RigVMCore/RigVMGraphFunctionDefinition.cpp
--- a/Engine/Plugins/Runtime/RigVM/Source/RigVM/Private/RigVMCore/RigVMGraphFunctionDefinition.cpp
+++ b/Engine/Plugins/Runtime/RigVM/Source/RigVM/Private/RigVMCore/RigVMGraphFunctionDefinition.cpp
@@ -47,6 +47,13 @@ bool FRigVMGraphFunctionArgument::IsExecuteContext() const
 	return false;
 }
 
+bool FRigVMGraphFunctionIdentifier::IsVariantOf(const FRigVMGraphFunctionIdentifier& InOther) const
+{
+	const FRigVMGraphFunctionHeader& ThisHeader = FRigVMGraphFunctionHeader::FindGraphFunctionHeader(*this);
+	const FRigVMGraphFunctionHeader& OtherHeader = FRigVMGraphFunctionHeader::FindGraphFunctionHeader(InOther);
+	return ThisHeader.Variant.Guid == OtherHeader.Variant.Guid;
+}
+
 bool FRigVMGraphFunctionHeader::IsMutable() const
 {
 	for(const FRigVMGraphFunctionArgument& Arg : Arguments)
@@ -112,11 +119,11 @@ void FRigVMGraphFunctionHeader::PostDuplicateHost(const FString& InOldPathName,
 		InOutObjectPath = FSoftObjectPath(PathName);
 	};
 
-	ReplacePathName(LibraryPointer.LibraryNodePath);
+	ReplacePathName(LibraryPointer.GetLibraryNodePath());
 	ReplaceSoftPathName(LibraryPointer.HostObject);
 	for (TPair<FRigVMGraphFunctionIdentifier, uint32>& Pair : Dependencies)
 	{
-		ReplacePathName(Pair.Key.LibraryNodePath);
+		ReplacePathName(Pair.Key.GetLibraryNodePath());
 		ReplaceSoftPathName(Pair.Key.HostObject);
 	}
 }
diff --git a/Engine/Plugins/Runtime/RigVM/Source/RigVM/Public/RigVMCore/RigVMGraphFunctionDefinition.h b/Engine/Plugins/Runtime/RigVM/Source/RigVM/Public/RigVMCore/RigVMGraphFunctionDefinition.h
--- a/Engine/Plugins/Runtime/RigVM/Source/RigVM/Public/RigVMCore/RigVMGraphFunctionDefinition.h
+++ b/Engine/Plugins/Runtime/RigVM/Source/RigVM/Public/RigVMCore/RigVMGraphFunctionDefinition.h
@@ -374,8 +374,11 @@ struct RIGVM_API FRigVMGraphFunctionIdentifier
 	UPROPERTY(meta=(DeprecatedProperty))
 	FSoftObjectPath LibraryNode_DEPRECATED;
 
-	UPROPERTY(BlueprintReadOnly, VisibleAnywhere, Category=FunctionIdentifier)
-	FString LibraryNodePath;
+private:
+	UPROPERTY(VisibleAnywhere, Category=FunctionIdentifier)
+	mutable FString LibraryNodePath;
+
+public:
 
 	// A path to the IRigVMGraphFunctionHost that stores the function information, and compilation data (e.g. RigVMBlueprintGeneratedClass)
 	UPROPERTY(BlueprintReadOnly, VisibleAnywhere, Category=FunctionIdentifier)
@@ -424,11 +427,25 @@ struct RIGVM_API FRigVMGraphFunctionIdentifier
 		return *GetFunctionName();
 	}
 
-	FSoftObjectPath GetNodeSoftPath() const
+	FString& GetLibraryNodePath() const
 	{
+		if (LibraryNodePath.IsEmpty() && LibraryNode_DEPRECATED.IsValid())
+		{
+			LibraryNodePath = LibraryNode_DEPRECATED.ToString();
+		}
 		return LibraryNodePath;
 	}
 
+	void SetLibraryNodePath(const FString& InPath)
+	{
+		LibraryNodePath = InPath;
+	}
+
+	FSoftObjectPath GetNodeSoftPath() const
+	{
+		return GetLibraryNodePath();
+	}
+
 	friend FArchive& operator<<(FArchive& Ar, FRigVMGraphFunctionIdentifier& Data)
 	{
 		Ar.UsingCustomVersion(FRigVMObjectVersion::GUID);
@@ -447,6 +464,8 @@ struct RIGVM_API FRigVMGraphFunctionIdentifier
 		Ar << Data.HostObject;
 		return Ar;
 	}
+
+	bool IsVariantOf(const FRigVMGraphFunctionIdentifier& InOther) const;
 };
 
 USTRUCT(BlueprintType)
diff --git a/Engine/Plugins/Runtime/RigVM/Source/RigVMDeveloper/Private/RigVMBlueprint.cpp b/Engine/Plugins/Runtime/RigVM/Source/RigVMDeveloper/Private/RigVMBlueprint.cpp
--- a/Engine/Plugins/Runtime/RigVM/Source/RigVMDeveloper/Private/RigVMBlueprint.cpp
+++ b/Engine/Plugins/Runtime/RigVM/Source/RigVMDeveloper/Private/RigVMBlueprint.cpp
@@ -2333,88 +2333,6 @@ TArray<FRigVMExternalDependency> URigVMBlueprint::GetExternalDependenciesForCate
 	return Dependencies;
 }
 
-TArray<FRigVMVariantRef> URigVMBlueprint::GatherFunctionVariantRefsForAsset(const FAssetData& InAssetData)
-{
-	TArray<FRigVMVariantRef> Result;
-
-	if (!InAssetData.IsInstanceOf(URigVMBlueprint::StaticClass()) &&
-		!InAssetData.IsInstanceOf(URigVMBlueprintGeneratedClass::StaticClass()))
-	{
-		return Result;
-	}
-
-	// If the asset is loaded, gather the function variants from the function store
-	// which will include private functions
-	if (InAssetData.IsAssetLoaded())
-	{
-		UObject* AssetObject = InAssetData.GetAsset();
-		URigVMBlueprintGeneratedClass* GeneratedClass = Cast<URigVMBlueprintGeneratedClass>(AssetObject);
-		if (!GeneratedClass)
-		{
-			if (URigVMBlueprint* RigVMBlueprint = Cast<URigVMBlueprint>(AssetObject))
-			{
-				GeneratedClass = RigVMBlueprint->GetRigVMBlueprintGeneratedClass();
-			}
-		}
-
-		if (GeneratedClass)
-		{
-			for (int32 Pass=0; Pass<2; Pass++)
-			{
-				TArray<FRigVMGraphFunctionData>& Functions = (Pass == 0)
-					? GeneratedClass->GraphFunctionStore.PrivateFunctions
-					: GeneratedClass->GraphFunctionStore.PublicFunctions;
-				
-				for (FRigVMGraphFunctionData& Data : Functions)
-				{
-					FRigVMVariantRef& VariantRef = Result.Add_GetRef(FRigVMVariantRef());
-					VariantRef.Variant = Data.Header.Variant;
-					VariantRef.ObjectPath = Data.Header.LibraryPointer.GetNodeSoftPath();
-				}
-			}
-
-			return Result;
-		}
-	}
-
-	// If asset is not loaded, gather function variants from metadata
-	FString PublicGraphFunctionsString = InAssetData.GetTagValueRef<FString>(TEXT("PublicGraphFunctions"));
-	TArray<FRigVMGraphFunctionHeader> PublicFunctions;
-	static const FName FunctionsPropertyName = GET_MEMBER_NAME_CHECKED(URigVMBlueprint, PublicGraphFunctions);
-	FArrayProperty* PublicGraphFunctionsProperty = CastField<FArrayProperty>(URigVMBlueprint::StaticClass()->FindPropertyByName(FunctionsPropertyName));
-	PublicGraphFunctionsProperty->ImportText_Direct(*PublicGraphFunctionsString, &PublicFunctions, nullptr, EPropertyPortFlags::PPF_None);
-
-	for(const FRigVMGraphFunctionHeader& PublicFunction : PublicFunctions)
-	{
-		FRigVMVariantRef& VariantRef = Result.Add_GetRef(FRigVMVariantRef());
-		VariantRef.Variant = PublicFunction.Variant;
-		VariantRef.ObjectPath = PublicFunction.LibraryPointer.GetNodeSoftPath();
-	}
-
-	return Result;
-}
-
-TArray<FRigVMVariantRef> URigVMBlueprint::FindFunctionVariantRefs(const FGuid& InGuid)
-{
-	TArray<FRigVMVariantRef> Result;
-	const FAssetRegistryModule& AssetRegistryModule = FModuleManager::LoadModuleChecked<FAssetRegistryModule>("AssetRegistry");
-	TArray<FAssetData> Assets;
-	FARFilter AssetFilter;
-	AssetFilter.ClassPaths.Add(URigVMBlueprint::StaticClass()->GetClassPathName());
-	AssetFilter.ClassPaths.Add(URigVMBlueprintGeneratedClass::StaticClass()->GetClassPathName());
-	AssetFilter.bRecursiveClasses = true;
-	AssetRegistryModule.Get().GetAssets(AssetFilter, Assets);
-
-	for (const FAssetData& Asset : Assets)
-	{
-		TArray<FRigVMVariantRef> VariantRefs = URigVMBlueprint::GatherFunctionVariantRefsForAsset(Asset);
-		VariantRefs = VariantRefs.FilterByPredicate([InGuid](const FRigVMVariantRef& VariantRef) { return VariantRef.Variant.Guid == InGuid; });
-		Result.Append(VariantRefs);
-	}
-
-	return Result;
-}
-
 URigVMGraph* URigVMBlueprint::GetTemplateModel(bool bIsFunctionLibrary)
 {
 #if WITH_EDITORONLY_DATA
@@ -3576,12 +3494,12 @@ void URigVMBlueprint::PatchFunctionReferencesOnLoad()
 				else if (!FunctionReferenceNode->ReferencedNodePtr_DEPRECATED.IsNull())
 				{
 					// At least lets make sure we store the path in the header
-					FunctionReferenceNode->ReferencedFunctionHeader.LibraryPointer.LibraryNodePath = FunctionReferenceNode->ReferencedNodePtr_DEPRECATED.ToSoftObjectPath().ToString();
+					FunctionReferenceNode->ReferencedFunctionHeader.LibraryPointer.SetLibraryNodePath(FunctionReferenceNode->ReferencedNodePtr_DEPRECATED.ToSoftObjectPath().ToString());
 				}
 
 				if (FunctionReferenceNode->ReferencedFunctionHeader.LibraryPointer.LibraryNode_DEPRECATED.IsValid())
 				{
-					FunctionReferenceNode->ReferencedFunctionHeader.LibraryPointer.LibraryNodePath = FunctionReferenceNode->ReferencedFunctionHeader.LibraryPointer.LibraryNode_DEPRECATED.ToString();
+					FunctionReferenceNode->ReferencedFunctionHeader.LibraryPointer.SetLibraryNodePath(FunctionReferenceNode->ReferencedFunctionHeader.LibraryPointer.LibraryNode_DEPRECATED.ToString());
 				}
 			}
 
@@ -3807,7 +3725,7 @@ void URigVMBlueprint::PatchFunctionsOnLoad()
 			const FRigVMVariant* Variant = FunctionLibrary->GetFunctionVariant(LibraryNode->GetFName());
 			if (!Variant)
 			{
-				Header.Variant.Guid = FRigVMVariant::GenerateGUID(Header.LibraryPointer.LibraryNodePath);
+				Header.Variant.Guid = FRigVMVariant::GenerateGUID(Header.LibraryPointer.GetLibraryNodePath());
 			}
 			else
 			{
diff --git a/Engine/Plugins/Runtime/RigVM/Source/RigVMDeveloper/Private/RigVMCompiler/RigVMCompiler.cpp b/Engine/Plugins/Runtime/RigVM/Source/RigVMDeveloper/Private/RigVMCompiler/RigVMCompiler.cpp
--- a/Engine/Plugins/Runtime/RigVM/Source/RigVMDeveloper/Private/RigVMCompiler/RigVMCompiler.cpp
+++ b/Engine/Plugins/Runtime/RigVM/Source/RigVMDeveloper/Private/RigVMCompiler/RigVMCompiler.cpp
@@ -2438,7 +2438,7 @@ bool URigVMCompiler::TraverseInlineFunction(const FRigVMInlineFunctionExprAST* I
 	const FRigVMByteCode& FunctionByteCode = FunctionCompilationData->ByteCode;
 
 	// Bytecode to be inlined should never be aligned
-	checkf(!FunctionByteCode.bByteCodeIsAligned, TEXT("Trying to inline aligned function bytecode %s in package %s"), *FunctionReferenceNode->GetFunctionIdentifier().LibraryNodePath, *GetPackage()->GetPathName());
+	checkf(!FunctionByteCode.bByteCodeIsAligned, TEXT("Trying to inline aligned function bytecode %s in package %s"), *FunctionReferenceNode->GetFunctionIdentifier().GetLibraryNodePath(), *GetPackage()->GetPathName());
 	
 	if (WorkData.bSetupMemory)
 	{
diff --git a/Engine/Plugins/Runtime/RigVM/Source/RigVMDeveloper/Private/RigVMModel/Nodes/RigVMFunctionReferenceNode.cpp b/Engine/Plugins/Runtime/RigVM/Source/RigVMDeveloper/Private/RigVMModel/Nodes/RigVMFunctionReferenceNode.cpp
--- a/Engine/Plugins/Runtime/RigVM/Source/RigVMDeveloper/Private/RigVMModel/Nodes/RigVMFunctionReferenceNode.cpp
+++ b/Engine/Plugins/Runtime/RigVM/Source/RigVMDeveloper/Private/RigVMModel/Nodes/RigVMFunctionReferenceNode.cpp
@@ -137,7 +137,7 @@ uint32 URigVMFunctionReferenceNode::GetStructureHash() const
 
 	Hash = HashCombine(Hash, GetTypeHash(ReferencedFunctionHeader.Name.ToString()));
 	Hash = HashCombine(Hash, GetTypeHash(ReferencedFunctionHeader.NodeTitle));
-	Hash = HashCombine(Hash, GetTypeHash(ReferencedFunctionHeader.LibraryPointer.LibraryNodePath));
+	Hash = HashCombine(Hash, GetTypeHash(ReferencedFunctionHeader.LibraryPointer.GetLibraryNodePath()));
 	Hash = HashCombine(Hash, GetTypeHash(ReferencedFunctionHeader.Keywords));
 	Hash = HashCombine(Hash, GetTypeHash(ReferencedFunctionHeader.Description));
 	Hash = HashCombine(Hash, GetTypeHash(ReferencedFunctionHeader.NodeColor));
diff --git a/Engine/Plugins/Runtime/RigVM/Source/RigVMDeveloper/Private/RigVMModel/Nodes/RigVMLibraryNode.cpp b/Engine/Plugins/Runtime/RigVM/Source/RigVMDeveloper/Private/RigVMModel/Nodes/RigVMLibraryNode.cpp
--- a/Engine/Plugins/Runtime/RigVM/Source/RigVMDeveloper/Private/RigVMModel/Nodes/RigVMLibraryNode.cpp
+++ b/Engine/Plugins/Runtime/RigVM/Source/RigVMDeveloper/Private/RigVMModel/Nodes/RigVMLibraryNode.cpp
@@ -219,7 +219,7 @@ FRigVMGraphFunctionIdentifier URigVMLibraryNode::GetFunctionIdentifier() const
 	{
 		Identifier.HostObject = Library->GetFunctionHostObjectPath();
 	}
-	Identifier.LibraryNodePath = this->GetPathName();
+	Identifier.SetLibraryNodePath(this->GetPathName());
 	return Identifier;
 }
 
@@ -233,7 +233,7 @@ FRigVMGraphFunctionHeader URigVMLibraryNode::GetFunctionHeader(IRigVMGraphFuncti
 	}
 	else
 	{
-	    Header.LibraryPointer.LibraryNodePath = this->GetPathName();
+	    Header.LibraryPointer.SetLibraryNodePath(this->GetPathName());
 	    Header.LibraryPointer.HostObject = Cast<UObject>(InHostObject);
 	}
 
diff --git a/Engine/Plugins/Runtime/RigVM/Source/RigVMDeveloper/Private/RigVMModel/RigVMBuildData.cpp b/Engine/Plugins/Runtime/RigVM/Source/RigVMDeveloper/Private/RigVMModel/RigVMBuildData.cpp
--- a/Engine/Plugins/Runtime/RigVM/Source/RigVMDeveloper/Private/RigVMModel/RigVMBuildData.cpp
+++ b/Engine/Plugins/Runtime/RigVM/Source/RigVMDeveloper/Private/RigVMModel/RigVMBuildData.cpp
@@ -1,6 +1,8 @@
 // Copyright Epic Games, Inc. All Rights Reserved.
 
 #include "RigVMModel/RigVMBuildData.h"
+#include "RigVMBlueprint.h"
+#include "RigVMBlueprintGeneratedClass.h"
 #include "AssetRegistry/AssetRegistryModule.h"
 #include "RigVMModel/RigVMFunctionLibrary.h"
 #include "UObject/UObjectIterator.h"
@@ -260,19 +262,19 @@ void URigVMBuildData::RegisterFunctionReference(FRigVMReferenceNodeData InRefere
 
 	if (!InReferenceNodeData.ReferencedFunctionIdentifier.GetNodeSoftPath().IsValid())
 	{
-		InReferenceNodeData.ReferencedFunctionIdentifier.LibraryNodePath = InReferenceNodeData.ReferencedFunctionPath_DEPRECATED;
+		InReferenceNodeData.ReferencedFunctionIdentifier.SetLibraryNodePath(InReferenceNodeData.ReferencedFunctionPath_DEPRECATED);
 	}
 	
 	check(InReferenceNodeData.ReferencedFunctionIdentifier.GetNodeSoftPath().IsValid());
 
-	FString LibraryNodePath = InReferenceNodeData.ReferencedFunctionIdentifier.LibraryNodePath;
+	FString LibraryNodePath = InReferenceNodeData.ReferencedFunctionIdentifier.GetLibraryNodePath();
 	TSoftObjectPtr<URigVMLibraryNode> LibraryNodePtr = TSoftObjectPtr<URigVMLibraryNode>(LibraryNodePath);
 
 	// Try to find a FunctionIdentifier with the same LibraryNodePath
 	bool bFound = false;
 	for (TPair< FRigVMGraphFunctionIdentifier, FRigVMFunctionReferenceArray >& Pair : GraphFunctionReferences)
 	{
-		if (Pair.Key.LibraryNodePath == LibraryNodePath)
+		if (Pair.Key.GetLibraryNodePath() == LibraryNodePath)
 		{
 			Pair.Value.FunctionReferences.Add(InReferenceNodeData.GetReferenceNodeObjectPath());
 			bFound = true;
@@ -364,5 +366,97 @@ void URigVMBuildData::ClearInvalidReferences()
 	}
 }
 
+TArray<FRigVMVariantRef> URigVMBuildData::GatherAllFunctionVariantRefs()
+{
+	TArray<FRigVMVariantRef> Result;
+	const FAssetRegistryModule& AssetRegistryModule = FModuleManager::LoadModuleChecked<FAssetRegistryModule>("AssetRegistry");
+	TArray<FAssetData> Assets;
+	FARFilter AssetFilter;
+	AssetFilter.ClassPaths.Add(URigVMBlueprint::StaticClass()->GetClassPathName());
+	AssetFilter.ClassPaths.Add(URigVMBlueprintGeneratedClass::StaticClass()->GetClassPathName());
+	AssetFilter.bRecursiveClasses = true;
+	AssetRegistryModule.Get().GetAssets(AssetFilter, Assets);
+
+	for (const FAssetData& Asset : Assets)
+	{
+		TArray<FRigVMVariantRef> VariantRefs = URigVMBuildData::GatherFunctionVariantRefsForAsset(Asset);
+		Result.Append(VariantRefs);
+	}
+
+	return Result;
+}
+
+TArray<FRigVMVariantRef> URigVMBuildData::GatherFunctionVariantRefsForAsset(const FAssetData& InAssetData)
+{
+	TArray<FRigVMVariantRef> Result;
+
+	if (!InAssetData.IsInstanceOf(URigVMBlueprint::StaticClass()) &&
+		!InAssetData.IsInstanceOf(URigVMBlueprintGeneratedClass::StaticClass()))
+	{
+		return Result;
+	}
+
+	// If the asset is loaded, gather the function variants from the function store
+	// which will include private functions
+	if (InAssetData.IsAssetLoaded())
+	{
+		UObject* AssetObject = InAssetData.GetAsset();
+		URigVMBlueprintGeneratedClass* GeneratedClass = Cast<URigVMBlueprintGeneratedClass>(AssetObject);
+		if (!GeneratedClass)
+		{
+			if (URigVMBlueprint* RigVMBlueprint = Cast<URigVMBlueprint>(AssetObject))
+			{
+				GeneratedClass = RigVMBlueprint->GetRigVMBlueprintGeneratedClass();
+			}
+		}
+
+		if (GeneratedClass)
+		{
+			for (int32 Pass=0; Pass<2; Pass++)
+			{
+				TArray<FRigVMGraphFunctionData>& Functions = (Pass == 0)
+					? GeneratedClass->GraphFunctionStore.PrivateFunctions
+					: GeneratedClass->GraphFunctionStore.PublicFunctions;
+				
+				for (FRigVMGraphFunctionData& Data : Functions)
+				{
+					FRigVMVariantRef& VariantRef = Result.Add_GetRef(FRigVMVariantRef());
+					VariantRef.Variant = Data.Header.Variant;
+					VariantRef.ObjectPath = Data.Header.LibraryPointer.GetNodeSoftPath();
+				}
+			}
+
+			return Result;
+		}
+	}
+
+	// If asset is not loaded, gather function variants from metadata
+	FString PublicGraphFunctionsString = InAssetData.GetTagValueRef<FString>(TEXT("PublicGraphFunctions"));
+	TArray<FRigVMGraphFunctionHeader> PublicFunctions;
+	static const FName FunctionsPropertyName = GET_MEMBER_NAME_CHECKED(URigVMBlueprint, PublicGraphFunctions);
+	FArrayProperty* PublicGraphFunctionsProperty = CastField<FArrayProperty>(URigVMBlueprint::StaticClass()->FindPropertyByName(FunctionsPropertyName));
+	PublicGraphFunctionsProperty->ImportText_Direct(*PublicGraphFunctionsString, &PublicFunctions, nullptr, EPropertyPortFlags::PPF_None);
+
+	for(const FRigVMGraphFunctionHeader& PublicFunction : PublicFunctions)
+	{
+		FRigVMVariantRef& VariantRef = Result.Add_GetRef(FRigVMVariantRef());
+		VariantRef.Variant = PublicFunction.Variant;
+		VariantRef.ObjectPath = PublicFunction.LibraryPointer.GetNodeSoftPath();
+		if (!VariantRef.Variant.Guid.IsValid())
+		{
+			VariantRef.Variant.Guid = FRigVMVariant::GenerateGUID(VariantRef.ObjectPath.ToString());
+		}
+	}
+
+	return Result;
+}
+
+TArray<FRigVMVariantRef> URigVMBuildData::FindFunctionVariantRefs(const FGuid& InGuid)
+{
+	TArray<FRigVMVariantRef> Result = GatherAllFunctionVariantRefs();
+	Result = Result.FilterByPredicate([InGuid](const FRigVMVariantRef& VariantRef) { return VariantRef.Variant.Guid == InGuid; });
+	return Result;
+}
+
 
 
diff --git a/Engine/Plugins/Runtime/RigVM/Source/RigVMDeveloper/Private/RigVMModel/RigVMController.cpp b/Engine/Plugins/Runtime/RigVM/Source/RigVMDeveloper/Private/RigVMModel/RigVMController.cpp
--- a/Engine/Plugins/Runtime/RigVM/Source/RigVMDeveloper/Private/RigVMModel/RigVMController.cpp
+++ b/Engine/Plugins/Runtime/RigVM/Source/RigVMDeveloper/Private/RigVMModel/RigVMController.cpp
@@ -592,7 +592,7 @@ FRigVMPinInfoArray::FRigVMPinInfoArray(const FRigVMGraphFunctionHeader& Function
 		const TRigVMTypeIndex TypeIndex = Registry.GetTypeIndexFromCPPType(FunctionArgument.CPPType.ToString());
 		if (TypeIndex == INDEX_NONE)
 		{
-			InController->ReportErrorf( TEXT("Invalid pin type %s for %s in %s"), *FunctionArgument.CPPType.ToString(), *FunctionHeader.LibraryPointer.LibraryNodePath, *InController->GetPackage()->GetPathName());
+			InController->ReportErrorf( TEXT("Invalid pin type %s for %s in %s"), *FunctionArgument.CPPType.ToString(), *FunctionHeader.LibraryPointer.GetLibraryNodePath(), *InController->GetPackage()->GetPathName());
 		}
 		ensureMsgf(TypeIndex != INDEX_NONE, TEXT("Invalid pin type %s in %s"), *FunctionArgument.CPPType.ToString(), *InController->GetPackage()->GetPathName());
 		ERigVMPinDefaultValueType DefaultValueType = ERigVMPinDefaultValueType::Unset;
@@ -4512,7 +4512,7 @@ TMap<FRigVMGraphFunctionIdentifier, URigVMLibraryNode*> URigVMController::Locali
 		FRigVMGraphFunctionData* FunctionData = FRigVMGraphFunctionData::FindFunctionData(NodeToVisit, &bIsPublic);
 		if (!FunctionData)
 		{
-			ReportAndNotifyErrorf(TEXT("Cannot localize function - could not find function %s in host %s."), *NodeToVisit.LibraryNodePath, *NodeToVisit.HostObject.ToString());
+			ReportAndNotifyErrorf(TEXT("Cannot localize function - could not find function %s in host %s."), *NodeToVisit.GetLibraryNodePath(), *NodeToVisit.HostObject.ToString());
 			return LocalizedFunctions;
 		}
 
@@ -4524,7 +4524,7 @@ TMap<FRigVMGraphFunctionIdentifier, URigVMLibraryNode*> URigVMController::Locali
 
 		if (!bLocalizeDependentPrivateFunctions)
 		{
-			ReportAndNotifyErrorf(TEXT("Cannot localize function - dependency %s is private."), *NodeToVisit.LibraryNodePath);
+			ReportAndNotifyErrorf(TEXT("Cannot localize function - dependency %s is private."), *NodeToVisit.GetLibraryNodePath());
 			return LocalizedFunctions;
 		}
 
@@ -4569,7 +4569,7 @@ TMap<FRigVMGraphFunctionIdentifier, URigVMLibraryNode*> URigVMController::Locali
 			{
 				URigVMLibraryNode* LocalizedFunction = ThisLibrary->FindFunction(NodeNames[0]);
 				LocalizedFunctions.Add(FunctionToLocalize->Header.LibraryPointer, LocalizedFunction);
-				ThisLibrary->LocalizedFunctions.FindOrAdd(FunctionToLocalize->Header.LibraryPointer.LibraryNodePath, LocalizedFunction);
+				ThisLibrary->LocalizedFunctions.FindOrAdd(FunctionToLocalize->Header.LibraryPointer.GetLibraryNodePath(), LocalizedFunction);
 			}
 		}
 	}
@@ -12236,7 +12236,7 @@ TArray<FRigVMVariantRef> URigVMController::FindVariantsOfFunction(const FName& I
     	return Result;
     }
 
-	return URigVMBlueprint::FindFunctionVariantRefs(FunctionData->Header.Variant.Guid);
+	return URigVMBuildData::Get()->FindFunctionVariantRefs(FunctionData->Header.Variant.Guid);
 }
 
 FRigVMGraphVariableDescription URigVMController::AddLocalVariable(const FName& InVariableName, const FString& InCPPType, UObject* InCPPTypeObject, const FString& InDefaultValue, bool bSetupUndoRedo, bool bPrintPythonCommand)
diff --git a/Engine/Plugins/Runtime/RigVM/Source/RigVMDeveloper/Private/RigVMModel/RigVMExternalDependency.cpp b/Engine/Plugins/Runtime/RigVM/Source/RigVMDeveloper/Private/RigVMModel/RigVMExternalDependency.cpp
--- a/Engine/Plugins/Runtime/RigVM/Source/RigVMDeveloper/Private/RigVMModel/RigVMExternalDependency.cpp
+++ b/Engine/Plugins/Runtime/RigVM/Source/RigVMDeveloper/Private/RigVMModel/RigVMExternalDependency.cpp
@@ -62,7 +62,7 @@ void IRigVMExternalDependencyManager::CollectExternalDependencies(TArray<FRigVME
 	check(InHeader);
 	if(InCategory == RigVMGraphFunctionCategory)
 	{
-		OutDependencies.AddUnique({InHeader->LibraryPointer.LibraryNodePath, InCategory});
+		OutDependencies.AddUnique({InHeader->LibraryPointer.GetLibraryNodePath(), InCategory});
 	}
 	for(const FRigVMGraphFunctionArgument& Argument : InHeader->Arguments)
 	{
diff --git a/Engine/Plugins/Runtime/RigVM/Source/RigVMDeveloper/Private/RigVMModel/RigVMFunctionLibrary.cpp b/Engine/Plugins/Runtime/RigVM/Source/RigVMDeveloper/Private/RigVMModel/RigVMFunctionLibrary.cpp
--- a/Engine/Plugins/Runtime/RigVM/Source/RigVMDeveloper/Private/RigVMModel/RigVMFunctionLibrary.cpp
+++ b/Engine/Plugins/Runtime/RigVM/Source/RigVMDeveloper/Private/RigVMModel/RigVMFunctionLibrary.cpp
@@ -122,7 +122,7 @@ void URigVMFunctionLibrary::ForEachReferenceSoftPtr(const FName& InFunctionName,
 
 URigVMLibraryNode* URigVMFunctionLibrary::FindPreviouslyLocalizedFunction(FRigVMGraphFunctionIdentifier InFunctionToLocalize)
 {
-	const FString& PathName = InFunctionToLocalize.LibraryNodePath;
+	const FString& PathName = InFunctionToLocalize.GetLibraryNodePath();
 	if(!LocalizedFunctions.Contains((PathName)))
 	{
 		return nullptr;
diff --git a/Engine/Plugins/Runtime/RigVM/Source/RigVMDeveloper/Public/RigVMBlueprint.h b/Engine/Plugins/Runtime/RigVM/Source/RigVMDeveloper/Public/RigVMBlueprint.h
--- a/Engine/Plugins/Runtime/RigVM/Source/RigVMDeveloper/Public/RigVMBlueprint.h
+++ b/Engine/Plugins/Runtime/RigVM/Source/RigVMDeveloper/Public/RigVMBlueprint.h
@@ -415,9 +415,6 @@ public:
 	FOnRigVMRequestInspectMemoryStorage& OnRequestInspectMemoryStorage() { return OnRequestInspectMemoryStorageEvent; }
 	void RequestInspectMemoryStorage(const TArray<FRigVMMemoryStorageStruct*>& InMemoryStorageStructs) { OnRequestInspectMemoryStorageEvent.Broadcast(InMemoryStorageStructs); }
 
-	static TArray<FRigVMVariantRef> GatherFunctionVariantRefsForAsset(const FAssetData& InAssetData);
-	static TArray<FRigVMVariantRef> FindFunctionVariantRefs(const FGuid& InGuid);
-
 #endif	// #if WITH_EDITOR
 
 	
@@ -822,6 +819,7 @@ private:
 	friend class FRigVMEditorModule;
 	friend class URigVMEdGraphSchema;
 	friend struct FRigVMEdGraphSchemaAction_PromoteToVariable;
+	friend class URigVMBuildData;
 };
 
 class RIGVMDEVELOPER_API FRigVMBlueprintCompileScope
diff --git a/Engine/Plugins/Runtime/RigVM/Source/RigVMDeveloper/Public/RigVMModel/RigVMBuildData.h b/Engine/Plugins/Runtime/RigVM/Source/RigVMDeveloper/Public/RigVMModel/RigVMBuildData.h
--- a/Engine/Plugins/Runtime/RigVM/Source/RigVMDeveloper/Public/RigVMModel/RigVMBuildData.h
+++ b/Engine/Plugins/Runtime/RigVM/Source/RigVMDeveloper/Public/RigVMModel/RigVMBuildData.h
@@ -119,6 +119,15 @@ public:
 	// Helper function to disable clearing transient package references
 	void SetIsRunningUnitTest(bool bIsRunning) { bIsRunningUnitTest = bIsRunning; }
 
+	// Will find all public function variant refs, and private function variant refs from loaded assets
+	TArray<FRigVMVariantRef> GatherAllFunctionVariantRefs();
+
+	// Will find the public function variant refs inside this asset, and private function variant refs if the asset is loaded
+	TArray<FRigVMVariantRef> GatherFunctionVariantRefsForAsset(const FAssetData& InAssetData);
+
+	// Will find all the function variants matching the given variant guid
+	TArray<FRigVMVariantRef> FindFunctionVariantRefs(const FGuid& InGuid);
+
 private:
 
 	// disable default constructor
diff --git a/Engine/Plugins/Runtime/RigVM/Source/RigVMEditor/Private/Editor/RigVMDetailsViewWrapperObject.cpp b/Engine/Plugins/Runtime/RigVM/Source/RigVMEditor/Private/Editor/RigVMDetailsViewWrapperObject.cpp
--- a/Engine/Plugins/Runtime/RigVM/Source/RigVMEditor/Private/Editor/RigVMDetailsViewWrapperObject.cpp
+++ b/Engine/Plugins/Runtime/RigVM/Source/RigVMEditor/Private/Editor/RigVMDetailsViewWrapperObject.cpp
@@ -353,13 +353,13 @@ UClass* URigVMDetailsViewWrapperObject::GetClassForNodes(TArray<URigVMNode*> InN
 		// sort the functions to ensure we get the same notation each time
 		Algo::Sort(FunctionIdentifiers, [](const FRigVMGraphFunctionIdentifier& A, const FRigVMGraphFunctionIdentifier& B) -> bool
 		{
-			return A.LibraryNodePath > B.LibraryNodePath;
+			return A.GetLibraryNodePath() > B.GetLibraryNodePath();
 		});
 
 		TArray<FString> IdentifierStrings;
 		for(const FRigVMGraphFunctionIdentifier& Identifier : FunctionIdentifiers)
 		{
-			IdentifierStrings.Add(Identifier.LibraryNodePath);
+			IdentifierStrings.Add(Identifier.GetLibraryNodePath());
 		}
 		Notation = FString::Printf(TEXT("%s(%s)"), *FString::Join(IdentifierStrings, TEXT("|")), *Notation);
 	}
diff --git a/Engine/Plugins/Runtime/RigVM/Source/RigVMEditor/Private/RigVMEditorModule.cpp b/Engine/Plugins/Runtime/RigVM/Source/RigVMEditor/Private/RigVMEditorModule.cpp
--- a/Engine/Plugins/Runtime/RigVM/Source/RigVMEditor/Private/RigVMEditorModule.cpp
+++ b/Engine/Plugins/Runtime/RigVM/Source/RigVMEditor/Private/RigVMEditorModule.cpp
@@ -396,8 +396,12 @@ void FRigVMEditorModule::GetTypeActions(URigVMBlueprint* RigVMBlueprint, FBluepr
 				{
 					TArray<FRigVMGraphFunctionHeader> PublicFunctions;
 					PublicGraphFunctionsProperty->ImportText_Direct(*PublicGraphFunctionsString, &PublicFunctions, nullptr, EPropertyPortFlags::PPF_None);
-					for(const FRigVMGraphFunctionHeader& PublicFunction : PublicFunctions)
+					for(FRigVMGraphFunctionHeader& PublicFunction : PublicFunctions)
 					{
+						if (PublicFunction.LibraryPointer.GetLibraryNodePath().IsEmpty() && PublicFunction.LibraryPointer.LibraryNode_DEPRECATED.IsValid())
+						{
+							PublicFunction.LibraryPointer.SetLibraryNodePath(PublicFunction.LibraryPointer.LibraryNode_DEPRECATED.ToString());
+						}
 						URigVMEdGraphNodeSpawner* NodeSpawner = URigVMEdGraphFunctionRefNodeSpawner::CreateFromAssetData(ControlRigAssetData, PublicFunction);
 						check(NodeSpawner != nullptr);
 						NodeSpawner->SetRelatedBlueprintClass(BlueprintClass);
@@ -413,8 +417,12 @@ void FRigVMEditorModule::GetTypeActions(URigVMBlueprint* RigVMBlueprint, FBluepr
 					TArray<FRigVMGraphFunctionHeader> PublicFunctions;
 					HeadersArrayProperty->ImportText_Direct(*HeadersString, &PublicFunctions, nullptr, EPropertyPortFlags::PPF_None);
 			
-					for(const FRigVMGraphFunctionHeader& PublicFunction : PublicFunctions)
+					for(FRigVMGraphFunctionHeader& PublicFunction : PublicFunctions)
 					{
+						if (PublicFunction.LibraryPointer.GetLibraryNodePath().IsEmpty() && PublicFunction.LibraryPointer.LibraryNode_DEPRECATED.IsValid())
+						{
+							PublicFunction.LibraryPointer.SetLibraryNodePath(PublicFunction.LibraryPointer.LibraryNode_DEPRECATED.ToString());
+						}
 						URigVMEdGraphNodeSpawner* NodeSpawner = URigVMEdGraphFunctionRefNodeSpawner::CreateFromAssetData(ControlRigAssetData, PublicFunction);
 						check(NodeSpawner != nullptr);
 						NodeSpawner->SetRelatedBlueprintClass(BlueprintClass);
