RigVM: Fix remaining bugs on FRigVMGraphFunctionIdentifier

#rnx

diff --git a/Engine/Plugins/Runtime/RigVM/Source/RigVM/Public/RigVMCore/RigVMGraphFunctionDefinition.h b/Engine/Plugins/Runtime/RigVM/Source/RigVM/Public/RigVMCore/RigVMGraphFunctionDefinition.h
--- a/Engine/Plugins/Runtime/RigVM/Source/RigVM/Public/RigVMCore/RigVMGraphFunctionDefinition.h
+++ b/Engine/Plugins/Runtime/RigVM/Source/RigVM/Public/RigVMCore/RigVMGraphFunctionDefinition.h
@@ -392,25 +392,26 @@ public:
 
 	friend uint32 GetTypeHash(const FRigVMGraphFunctionIdentifier& Pointer)
 	{
-		return HashCombine(GetTypeHash(Pointer.LibraryNodePath), GetTypeHash(Pointer.HostObject));
+		return HashCombine(GetTypeHash(Pointer.GetLibraryNodePath()), GetTypeHash(Pointer.HostObject));
 	}
 
 	bool operator==(const FRigVMGraphFunctionIdentifier& Other) const
 	{
-		return HostObject == Other.HostObject && FSoftObjectPath(LibraryNodePath).GetSubPathString() == FSoftObjectPath(Other.LibraryNodePath).GetSubPathString();
+		return HostObject == Other.HostObject && GetNodeSoftPath().GetSubPathString() == Other.GetNodeSoftPath().GetSubPathString();
 	}
 
 	bool IsValid() const
 	{
-		return !HostObject.IsNull() && !LibraryNodePath.IsEmpty();
+		return !HostObject.IsNull() && (!GetLibraryNodePath().IsEmpty());
 	}
 
 	FString GetFunctionName() const
 	{
 		if(IsValid())
 		{
+			const FString Path = GetLibraryNodePath();
 			FString NodeName;
-			if(LibraryNodePath.Split(TEXT("."), nullptr, &NodeName, ESearchCase::CaseSensitive, ESearchDir::FromEnd))
+			if(Path.Split(TEXT("."), nullptr, &NodeName, ESearchCase::CaseSensitive, ESearchDir::FromEnd))
 			{
 				return NodeName;
 			}
@@ -449,6 +450,14 @@ public:
 	friend FArchive& operator<<(FArchive& Ar, FRigVMGraphFunctionIdentifier& Data)
 	{
 		Ar.UsingCustomVersion(FRigVMObjectVersion::GUID);
+
+		if(Ar.IsSaving())
+		{
+			if(Data.LibraryNodePath.IsEmpty() && Data.LibraryNode_DEPRECATED.IsValid())
+			{
+				Data.LibraryNodePath = Data.GetLibraryNodePath();
+			}
+		}
 		
 		if (Ar.IsLoading() && Ar.CustomVer(FRigVMObjectVersion::GUID) < FRigVMObjectVersion::RemoveLibraryNodeReferenceFromFunctionIdentifier)
 		{
