Control Rig: Fix can't add function references to standard function library
#rnx

diff --git a/Engine/Plugins/Runtime/RigVM/Source/RigVMDeveloper/Private/RigVMModel/RigVMBuildData.cpp b/Engine/Plugins/Runtime/RigVM/Source/RigVMDeveloper/Private/RigVMModel/RigVMBuildData.cpp
--- a/Engine/Plugins/Runtime/RigVM/Source/RigVMDeveloper/Private/RigVMModel/RigVMBuildData.cpp
+++ b/Engine/Plugins/Runtime/RigVM/Source/RigVMDeveloper/Private/RigVMModel/RigVMBuildData.cpp
@@ -437,14 +437,17 @@ TArray<FRigVMVariantRef> URigVMBuildData::GatherFunctionVariantRefsForAsset(cons
 	FArrayProperty* PublicGraphFunctionsProperty = CastField<FArrayProperty>(URigVMBlueprint::StaticClass()->FindPropertyByName(FunctionsPropertyName));
 	PublicGraphFunctionsProperty->ImportText_Direct(*PublicGraphFunctionsString, &PublicFunctions, nullptr, EPropertyPortFlags::PPF_None);
 
-	for(const FRigVMGraphFunctionHeader& PublicFunction : PublicFunctions)
+	for(FRigVMGraphFunctionHeader& PublicFunction : PublicFunctions)
 	{
-		FRigVMVariantRef& VariantRef = Result.Add_GetRef(FRigVMVariantRef());
-		VariantRef.Variant = PublicFunction.Variant;
-		VariantRef.ObjectPath = PublicFunction.LibraryPointer.GetNodeSoftPath();
-		if (!VariantRef.Variant.Guid.IsValid())
+		if (PublicFunction.LibraryPointer.IsValid())
 		{
-			VariantRef.Variant.Guid = FRigVMVariant::GenerateGUID(VariantRef.ObjectPath.ToString());
+			FRigVMVariantRef& VariantRef = Result.Add_GetRef(FRigVMVariantRef());
+			VariantRef.Variant = PublicFunction.Variant;
+			VariantRef.ObjectPath = PublicFunction.LibraryPointer.GetNodeSoftPath();
+			if (!VariantRef.Variant.Guid.IsValid())
+			{
+				VariantRef.Variant.Guid = FRigVMVariant::GenerateGUID(VariantRef.ObjectPath.ToString());
+			}
 		}
 	}
 
diff --git a/Engine/Plugins/Runtime/RigVM/Source/RigVMEditor/Private/RigVMEditorModule.cpp b/Engine/Plugins/Runtime/RigVM/Source/RigVMEditor/Private/RigVMEditorModule.cpp
--- a/Engine/Plugins/Runtime/RigVM/Source/RigVMEditor/Private/RigVMEditorModule.cpp
+++ b/Engine/Plugins/Runtime/RigVM/Source/RigVMEditor/Private/RigVMEditorModule.cpp
@@ -354,80 +354,86 @@ void FRigVMEditorModule::GetTypeActions(URigVMBlueprint* RigVMBlueprint, FBluepr
 			}
 			PackagesProcessed.Add(ControlRigAssetData.PackageName);
 
-			FString PublicGraphFunctionsString;
-			FString PublicFunctionsString;
-			if (PublicGraphFunctionsProperty)
+			TArray<FRigVMGraphFunctionHeader> PublicFunctions;
+			if (ControlRigAssetData.IsAssetLoaded())
 			{
-				PublicGraphFunctionsString = ControlRigAssetData.GetTagValueRef<FString>(PublicGraphFunctionsProperty->GetFName());
-			}
-			// Only look at the deprecated public functions if the PublicGraphFunctionsString is empty
-			if (PublicGraphFunctionsString.IsEmpty() && PublicFunctionsProperty)
-			{
-				PublicFunctionsString = ControlRigAssetData.GetTagValueRef<FString>(PublicFunctionsProperty->GetFName());
-			}
-
-			// For RigVMBlueprintGeneratedClass, the property doesn't exist
-			if (PublicGraphFunctionsString.IsEmpty())
-			{
-				PublicGraphFunctionsString = ControlRigAssetData.GetTagValueRef<FString>(TEXT("PublicGraphFunctions"));
+				UObject* AssetObject = ControlRigAssetData.GetAsset();
+				if (URigVMBlueprint* Blueprint = Cast<URigVMBlueprint>(AssetObject))
+				{
+					PublicFunctions = Blueprint->PublicGraphFunctions;
+				}
+				else if(URigVMBlueprintGeneratedClass* GeneratedClass = Cast<URigVMBlueprintGeneratedClass>(AssetObject))
+				{
+					PublicFunctions.Reserve(GeneratedClass->GraphFunctionStore.PublicFunctions.Num());
+					for (const FRigVMGraphFunctionData& PublicFunction : GeneratedClass->GraphFunctionStore.PublicFunctions)
+					{
+						PublicFunctions.Add(PublicFunction.Header);
+					}
+				}
 			}
-			
-			if(PublicFunctionsString.IsEmpty() && PublicGraphFunctionsString.IsEmpty())
+			else
 			{
-				continue;
-			}
+				FString PublicGraphFunctionsString;
+				FString PublicFunctionsString;
+				if (PublicGraphFunctionsProperty)
+				{
+					PublicGraphFunctionsString = ControlRigAssetData.GetTagValueRef<FString>(PublicGraphFunctionsProperty->GetFName());
+				}
+				// Only look at the deprecated public functions if the PublicGraphFunctionsString is empty
+				if (PublicGraphFunctionsString.IsEmpty() && PublicFunctionsProperty)
+				{
+					PublicFunctionsString = ControlRigAssetData.GetTagValueRef<FString>(PublicFunctionsProperty->GetFName());
+				}
 
-			if (PublicFunctionsProperty && !PublicFunctionsString.IsEmpty())
-			{
-				TArray<FRigVMOldPublicFunctionData> PublicFunctions;
-				PublicFunctionsProperty->ImportText_Direct(*PublicFunctionsString, &PublicFunctions, nullptr, EPropertyPortFlags::PPF_None);
-				for(const FRigVMOldPublicFunctionData& PublicFunction : PublicFunctions)
+				// For RigVMBlueprintGeneratedClass, the property doesn't exist
+				if (PublicGraphFunctionsString.IsEmpty())
 				{
-					URigVMEdGraphNodeSpawner* NodeSpawner = URigVMEdGraphFunctionRefNodeSpawner::CreateFromAssetData(ControlRigAssetData, PublicFunction);
-					check(NodeSpawner != nullptr);
-					NodeSpawner->SetRelatedBlueprintClass(BlueprintClass);
-					ActionRegistrar.AddBlueprintAction(ActionKey, NodeSpawner);
+					PublicGraphFunctionsString = ControlRigAssetData.GetTagValueRef<FString>(TEXT("PublicGraphFunctions"));
+				}
+				
+				if(PublicFunctionsString.IsEmpty() && PublicGraphFunctionsString.IsEmpty())
+				{
+					continue;
 				}
-			}
 
-			if (!PublicGraphFunctionsString.IsEmpty())
-			{
-				if (PublicGraphFunctionsProperty)
+				if (PublicFunctionsProperty && !PublicFunctionsString.IsEmpty())
 				{
-					TArray<FRigVMGraphFunctionHeader> PublicFunctions;
-					PublicGraphFunctionsProperty->ImportText_Direct(*PublicGraphFunctionsString, &PublicFunctions, nullptr, EPropertyPortFlags::PPF_None);
-					for(FRigVMGraphFunctionHeader& PublicFunction : PublicFunctions)
+					TArray<FRigVMOldPublicFunctionData> OldPublicFunctions;
+					PublicFunctionsProperty->ImportText_Direct(*PublicFunctionsString, &PublicFunctions, nullptr, EPropertyPortFlags::PPF_None);
+					for(const FRigVMOldPublicFunctionData& PublicFunction : OldPublicFunctions)
 					{
-						if (PublicFunction.LibraryPointer.GetLibraryNodePath().IsEmpty() && PublicFunction.LibraryPointer.LibraryNode_DEPRECATED.IsValid())
-						{
-							PublicFunction.LibraryPointer.SetLibraryNodePath(PublicFunction.LibraryPointer.LibraryNode_DEPRECATED.ToString());
-						}
 						URigVMEdGraphNodeSpawner* NodeSpawner = URigVMEdGraphFunctionRefNodeSpawner::CreateFromAssetData(ControlRigAssetData, PublicFunction);
 						check(NodeSpawner != nullptr);
 						NodeSpawner->SetRelatedBlueprintClass(BlueprintClass);
 						ActionRegistrar.AddBlueprintAction(ActionKey, NodeSpawner);
 					}
 				}
-				else
+
+				if (!PublicGraphFunctionsString.IsEmpty())
 				{
-					// extract public function headers from generated class
-					const FString& HeadersString = PublicGraphFunctionsString;
-			
-					FArrayProperty* HeadersArrayProperty = CastField<FArrayProperty>(FRigVMGraphFunctionHeaderArray::StaticStruct()->FindPropertyByName(TEXT("Headers")));
-					TArray<FRigVMGraphFunctionHeader> PublicFunctions;
-					HeadersArrayProperty->ImportText_Direct(*HeadersString, &PublicFunctions, nullptr, EPropertyPortFlags::PPF_None);
-			
-					for(FRigVMGraphFunctionHeader& PublicFunction : PublicFunctions)
+					if (PublicGraphFunctionsProperty)
 					{
-						if (PublicFunction.LibraryPointer.GetLibraryNodePath().IsEmpty() && PublicFunction.LibraryPointer.LibraryNode_DEPRECATED.IsValid())
-						{
-							PublicFunction.LibraryPointer.SetLibraryNodePath(PublicFunction.LibraryPointer.LibraryNode_DEPRECATED.ToString());
-						}
-						URigVMEdGraphNodeSpawner* NodeSpawner = URigVMEdGraphFunctionRefNodeSpawner::CreateFromAssetData(ControlRigAssetData, PublicFunction);
-						check(NodeSpawner != nullptr);
-						NodeSpawner->SetRelatedBlueprintClass(BlueprintClass);
-						ActionRegistrar.AddBlueprintAction(ActionKey, NodeSpawner);
+						PublicGraphFunctionsProperty->ImportText_Direct(*PublicGraphFunctionsString, &PublicFunctions, nullptr, EPropertyPortFlags::PPF_None);
 					}
+					else
+					{
+						// extract public function headers from generated class
+						const FString& HeadersString = PublicGraphFunctionsString;
+				
+						FArrayProperty* HeadersArrayProperty = CastField<FArrayProperty>(FRigVMGraphFunctionHeaderArray::StaticStruct()->FindPropertyByName(TEXT("Headers")));
+						HeadersArrayProperty->ImportText_Direct(*HeadersString, &PublicFunctions, nullptr, EPropertyPortFlags::PPF_None);
+					}
+				}
+			}
+
+			for(FRigVMGraphFunctionHeader& PublicFunction : PublicFunctions)
+			{
+				if (PublicFunction.LibraryPointer.IsValid())
+				{
+					URigVMEdGraphNodeSpawner* NodeSpawner = URigVMEdGraphFunctionRefNodeSpawner::CreateFromAssetData(ControlRigAssetData, PublicFunction);
+					check(NodeSpawner != nullptr);
+					NodeSpawner->SetRelatedBlueprintClass(BlueprintClass);
+					ActionRegistrar.AddBlueprintAction(ActionKey, NodeSpawner);
 				}
 			}
 		}
