Control Rig: Fix static analysis warning
#rnx

diff --git a/Engine/Plugins/Runtime/RigVM/Source/RigVMEditor/Private/RigVMEditorModule.cpp b/Engine/Plugins/Runtime/RigVM/Source/RigVMEditor/Private/RigVMEditorModule.cpp
--- a/Engine/Plugins/Runtime/RigVM/Source/RigVMEditor/Private/RigVMEditorModule.cpp
+++ b/Engine/Plugins/Runtime/RigVM/Source/RigVMEditor/Private/RigVMEditorModule.cpp
@@ -399,7 +399,7 @@ void FRigVMEditorModule::GetTypeActions(URigVMBlueprint* RigVMBlueprint, FBluepr
 				if (PublicFunctionsProperty && !PublicFunctionsString.IsEmpty())
 				{
 					TArray<FRigVMOldPublicFunctionData> OldPublicFunctions;
-					PublicFunctionsProperty->ImportText_Direct(*PublicFunctionsString, &PublicFunctions, nullptr, EPropertyPortFlags::PPF_None);
+					PublicFunctionsProperty->ImportText_Direct(*PublicFunctionsString, &OldPublicFunctions, nullptr, EPropertyPortFlags::PPF_None);
 					for(const FRigVMOldPublicFunctionData& PublicFunction : OldPublicFunctions)
 					{
 						URigVMEdGraphNodeSpawner* NodeSpawner = URigVMEdGraphFunctionRefNodeSpawner::CreateFromAssetData(ControlRigAssetData, PublicFunction);
