RigVM: Controller fixes, orphan pins for change data types, default disabled rigvm action stack

#rnx

diff --git a/Engine/Plugins/Runtime/RigVM/Source/RigVMDeveloper/Private/RigVMModel/RigVMController.cpp b/Engine/Plugins/Runtime/RigVM/Source/RigVMDeveloper/Private/RigVMModel/RigVMController.cpp
--- a/Engine/Plugins/Runtime/RigVM/Source/RigVMDeveloper/Private/RigVMModel/RigVMController.cpp
+++ b/Engine/Plugins/Runtime/RigVM/Source/RigVMDeveloper/Private/RigVMModel/RigVMController.cpp
@@ -100,6 +100,7 @@ URigVMController::URigVMController()
 #endif
 	, bEnableSchemaRemoveNodeCheck(true)
 {
+	SetActionStack(URigVMActionStack::GetDisabledActionStack());
 }
 
 URigVMController::URigVMController(const FObjectInitializer& ObjectInitializer)
@@ -120,6 +121,7 @@ URigVMController::URigVMController(const FObjectInitializer& ObjectInitializer)
 #endif
 	, bEnableSchemaRemoveNodeCheck(true)
 {
+	SetActionStack(URigVMActionStack::GetDisabledActionStack());
 }
 
 URigVMController::~URigVMController()
@@ -11349,7 +11351,7 @@ URigVMFunctionReferenceNode* URigVMController::AddFunctionReferenceNode(URigVMLi
 	return nullptr;
 }
 
-bool URigVMController::SwapFunctionReferenceByName(const FName& InFunctionReferenceNodeName, const FRigVMGraphFunctionIdentifier& InNewFunctionIdentifier, bool bSetupUndoRedo, bool bPrintPythonCommand)
+bool URigVMController::SwapFunctionReferenceByName(const FName& InFunctionReferenceNodeName, const FRigVMGraphFunctionIdentifier& InNewFunctionIdentifier, bool bSetupOrphanPins, bool bSetupUndoRedo, bool bPrintPythonCommand)
 {
 	if (!IsValidGraph())
 	{
@@ -11368,10 +11370,10 @@ bool URigVMController::SwapFunctionReferenceByName(const FName& InFunctionRefere
 		return false;
 	}
 
-	return SwapFunctionReference(Node, InNewFunctionIdentifier, bSetupUndoRedo, bPrintPythonCommand);
+	return SwapFunctionReference(Node, InNewFunctionIdentifier, bSetupOrphanPins, bSetupUndoRedo, bPrintPythonCommand);
 }
 
-bool URigVMController::SwapFunctionReference(URigVMFunctionReferenceNode* InFunctionReferenceNode, const FRigVMGraphFunctionIdentifier& InNewFunctionIdentifier, bool bSetupUndoRedo, bool bPrintPythonCommand)
+bool URigVMController::SwapFunctionReference(URigVMFunctionReferenceNode* InFunctionReferenceNode, const FRigVMGraphFunctionIdentifier& InNewFunctionIdentifier, bool bSetupOrphanPins, bool bSetupUndoRedo, bool bPrintPythonCommand)
 {
 	if (!IsValidNodeForGraph(InFunctionReferenceNode))
 	{
@@ -11402,17 +11404,12 @@ bool URigVMController::SwapFunctionReference(URigVMFunctionReferenceNode* InFunc
 		Action.SetTitle(TEXT("Swap Function Reference"));
 	}
 
-	const TMap<FString, FPinState> PinStates = GetPinStates(InFunctionReferenceNode);
-	const TArray<FLinkedPath> LinkedPaths = GetLinkedPaths(InFunctionReferenceNode);
-	FastBreakLinkedPaths(LinkedPaths, false);
-
 	InFunctionReferenceNode->Modify();
+	URigVMBuildData::Get()->UnregisterFunctionReference(InFunctionReferenceNode->ReferencedFunctionHeader.LibraryPointer, InFunctionReferenceNode);
 	InFunctionReferenceNode->ReferencedFunctionHeader.LibraryPointer = InNewFunctionIdentifier;
 	InFunctionReferenceNode->UpdateFunctionHeaderFromHost();
-	RepopulatePinsOnNode(InFunctionReferenceNode, false, false, false);
-
-	ApplyPinStates(InFunctionReferenceNode, PinStates, {}, false);
-	RestoreLinkedPaths(LinkedPaths, FRestoreLinkedPathSettings(), false);
+	URigVMBuildData::Get()->RegisterFunctionReference(InFunctionReferenceNode->ReferencedFunctionHeader.LibraryPointer, InFunctionReferenceNode);
+	RepopulatePinsOnNode(InFunctionReferenceNode, false, bSetupOrphanPins, true);
 
 	TArray<FRigVMExternalVariable> ExternalVariables;
 	if (GetExternalVariablesDelegate.IsBound())
@@ -11469,15 +11466,16 @@ bool URigVMController::SwapFunctionReference(URigVMFunctionReferenceNode* InFunc
 					*NewFunctionDefinitionName));
 
 		RigVMPythonUtils::Print(GetSchema()->GetGraphOuterName(GetGraph()), 
-			FString::Printf(TEXT("blueprint.get_controller_by_name('%s').swap_function_reference_by_name('%s', new_definition)"),
+			FString::Printf(TEXT("blueprint.get_controller_by_name('%s').swap_function_reference_by_name('%s', new_definition, %s)"),
 					*GraphName,
-					*FunctionRefNodeName));
+					*FunctionRefNodeName,
+					(bSetupOrphanPins) ? TEXT("True") : TEXT("False")));
 	}
 
 	return true;
 }
 
-bool URigVMController::SwapAllFunctionReferences(const FRigVMGraphFunctionIdentifier& InOldFunctionIdentifier, const FRigVMGraphFunctionIdentifier& InNewFunctionIdentifier, bool bSetupUndoRedo, bool bPrintPythonCommand)
+bool URigVMController::SwapAllFunctionReferences(const FRigVMGraphFunctionIdentifier& InOldFunctionIdentifier, const FRigVMGraphFunctionIdentifier& InNewFunctionIdentifier, bool bSetupOrphanPins, bool bSetupUndoRedo, bool bPrintPythonCommand)
 {
 	if (!IsValidGraph())
 	{
@@ -11540,7 +11538,7 @@ bool URigVMController::SwapAllFunctionReferences(const FRigVMGraphFunctionIdenti
 
 	for(URigVMFunctionReferenceNode* Node : FunctionReferenceNodes)
 	{
-		if(!SwapFunctionReference(Node, InNewFunctionIdentifier, bSetupUndoRedo, false))
+		if(!SwapFunctionReference(Node, InNewFunctionIdentifier, bSetupOrphanPins, bSetupUndoRedo, false))
 		{
 			if (bSetupUndoRedo && FunctionReferenceNodes.Num() > 1)
 			{
@@ -11574,8 +11572,9 @@ bool URigVMController::SwapAllFunctionReferences(const FRigVMGraphFunctionIdenti
 					*NewFunctionDefinitionName));
 
 		RigVMPythonUtils::Print(GetSchema()->GetGraphOuterName(GetGraph()), 
-			FString::Printf(TEXT("blueprint.get_controller_by_name('%s').swap_all_function_references(old_definition, new_definition)"),
-					*GraphName));
+			FString::Printf(TEXT("blueprint.get_controller_by_name('%s').swap_all_function_references(old_definition, new_definition, %s)"),
+					*GraphName,
+					(bSetupOrphanPins) ? TEXT("True") : TEXT("False")));
 	}
 	
 	return true;
@@ -15813,10 +15812,10 @@ void URigVMController::GenerateRepopulatePinLists(const FRigVMRegistry& Registry
 	{
 		const FString PinPath = NodeData.PreviousPinInfos.GetPinPath(Index);
 		const int32 NewIndex = NodeData.NewPinInfos.GetIndexFromPinPath(PinPath);
+		const int32 RootIndex = NodeData.PreviousPinInfos.GetRootIndex(Index);
 
 		if (NewIndex == INDEX_NONE)
 		{
-			const int32 RootIndex = NodeData.PreviousPinInfos.GetRootIndex(Index);
 			if (NodeData.PreviousPinInfos[Index].Direction != ERigVMPinDirection::Hidden)
 			{
 				if (URigVMPin* Pin = InNode->FindPin(PinPath))
@@ -15877,6 +15876,34 @@ void URigVMController::GenerateRepopulatePinLists(const FRigVMRegistry& Registry
 			{
 				NodeData.PreviousPinsToRemove.Add(Index);
 				NodeData.NewPinsToAdd.Add(NewIndex);
+
+				// we also need to setup orphan pins if a pin CPP type has changed
+				if (URigVMPin* Pin = InNode->FindPin(PinPath))
+				{
+					if (!Pin->GetLinks().IsEmpty())
+					{
+						if (NodeData.bSetupOrphanPinsForThisNode)
+						{
+							if (!NodeData.PreviousPinsToOrphan.Contains(RootIndex))
+							{
+								URigVMPin* RootPin = Pin->GetRootPin();
+
+								if (RootPin->GetSourceLinks(true).Num() > 0 ||
+									RootPin->GetTargetLinks(true).Num() > 0)
+								{
+									NodeData.PreviousPinsToOrphan.Add(RootIndex);
+
+									NodeData.bRequireRecreateLinks = true;
+									NodeData.bRequirePinStates = true;
+#if UE_RIGVMCONTROLLER_VERBOSE_REPOPULATE
+									UE_LOG(LogRigVMDeveloper, Display, TEXT("Previously existing pin '%s' needs to be orphaned."), *RootPin->GetPinPath());
+#endif
+								}
+							}
+						}
+
+					}
+				}
 			}
 
 #if UE_RIGVMCONTROLLER_VERBOSE_REPOPULATE
@@ -15916,7 +15943,7 @@ void URigVMController::GenerateRepopulatePinLists(const FRigVMRegistry& Registry
 			const int32 PreviousRootIndex = NodeData.PreviousPinInfos.GetRootIndex(PreviousIndex);
 			if (NodeData.PreviousPinsToOrphan.Contains(PreviousRootIndex))
 			{
-				NodeData.NewPinsToAdd.Add(Index);
+				NodeData.NewPinsToAdd.AddUnique(Index);
 
 #if UE_RIGVMCONTROLLER_VERBOSE_REPOPULATE
 				UE_LOG(LogRigVMDeveloper, Display, TEXT("Orphaned pin '%s' needs to be re-added."), *PinPath);
diff --git a/Engine/Plugins/Runtime/RigVM/Source/RigVMDeveloper/Private/RigVMModel/RigVMControllerActions.cpp b/Engine/Plugins/Runtime/RigVM/Source/RigVMDeveloper/Private/RigVMModel/RigVMControllerActions.cpp
--- a/Engine/Plugins/Runtime/RigVM/Source/RigVMDeveloper/Private/RigVMModel/RigVMControllerActions.cpp
+++ b/Engine/Plugins/Runtime/RigVM/Source/RigVMDeveloper/Private/RigVMModel/RigVMControllerActions.cpp
@@ -67,8 +67,17 @@ FString FRigVMActionWrapper::ExportText() const
 	return ExportedText;
 }
 
+URigVMActionStack* URigVMActionStack::GetDisabledActionStack()
+{
+	return StaticClass()->GetDefaultObject<URigVMActionStack>();
+}
+
 bool URigVMActionStack::OpenUndoBracket(URigVMController* InController, const FString& InTitle)
 {
+	if(IsDisabled())
+	{
+		return true;
+	}
 	FRigVMBaseAction* Action = new FRigVMBaseAction(InController);
 	Action->Title = InTitle;
 	BracketActions.Add(Action);
@@ -78,6 +87,10 @@ bool URigVMActionStack::OpenUndoBracket(URigVMController* InController, const FS
 
 bool URigVMActionStack::CloseUndoBracket(URigVMController* InController)
 {
+	if(IsDisabled())
+	{
+		return true;
+	}
 	ensure(BracketActions.Num() > 0);
 	if(BracketActions.Last()->IsEmpty())
 	{
@@ -91,6 +104,10 @@ bool URigVMActionStack::CloseUndoBracket(URigVMController* InController)
 
 bool URigVMActionStack::CancelUndoBracket(URigVMController* InController)
 {
+	if(IsDisabled())
+	{
+		return true;
+	}
 	ensure(BracketActions.Num() > 0);
 	FRigVMBaseAction* Action = BracketActions.Pop();
 	CancelAction(*Action);
@@ -205,6 +222,11 @@ void URigVMActionStack::PostTransacted(const FTransactionObjectEvent& Transactio
 	}
 }
 
+bool URigVMActionStack::IsDisabled() const
+{
+	return !HasAnyFlags(RF_ClassDefaultObject);
+}
+
 #if RIGVM_ACTIONSTACK_VERBOSE_LOG		
 
 void URigVMActionStack::LogAction(const UScriptStruct* InActionStruct, const FRigVMBaseAction& InAction, const FString& InPrefix)
diff --git a/Engine/Plugins/Runtime/RigVM/Source/RigVMDeveloper/Public/RigVMModel/RigVMController.h b/Engine/Plugins/Runtime/RigVM/Source/RigVMDeveloper/Public/RigVMModel/RigVMController.h
--- a/Engine/Plugins/Runtime/RigVM/Source/RigVMDeveloper/Public/RigVMModel/RigVMController.h
+++ b/Engine/Plugins/Runtime/RigVM/Source/RigVMDeveloper/Public/RigVMModel/RigVMController.h
@@ -928,13 +928,13 @@ public:
 	URigVMFunctionReferenceNode* AddFunctionReferenceNode(URigVMLibraryNode* InFunctionDefinition, const FVector2D& InNodePosition = FVector2D::ZeroVector, const FString& InNodeName = TEXT(""), bool bSetupUndoRedo = true, bool bPrintPythonCommand = false);
 
 	UFUNCTION(BlueprintCallable, Category = RigVMController)
-	bool SwapFunctionReferenceByName(const FName& InFunctionReferenceNodeName, const FRigVMGraphFunctionIdentifier& InNewFunctionIdentifier, bool bSetupUndoRedo = true, bool bPrintPythonCommand = false);
+	bool SwapFunctionReferenceByName(const FName& InFunctionReferenceNodeName, const FRigVMGraphFunctionIdentifier& InNewFunctionIdentifier, bool bSetupOrphanPins, bool bSetupUndoRedo = true, bool bPrintPythonCommand = false);
 
 	UFUNCTION(BlueprintCallable, Category = RigVMController)
-	bool SwapFunctionReference(URigVMFunctionReferenceNode* InFunctionReferenceNode, const FRigVMGraphFunctionIdentifier& InNewFunctionIdentifier, bool bSetupUndoRedo = true, bool bPrintPythonCommand = false);
+	bool SwapFunctionReference(URigVMFunctionReferenceNode* InFunctionReferenceNode, const FRigVMGraphFunctionIdentifier& InNewFunctionIdentifier, bool bSetupOrphanPins, bool bSetupUndoRedo = true, bool bPrintPythonCommand = false);
 
 	UFUNCTION(BlueprintCallable, Category = RigVMController)
-	bool SwapAllFunctionReferences(const FRigVMGraphFunctionIdentifier& InOldFunctionIdentifier, const FRigVMGraphFunctionIdentifier& InNewFunctionIdentifier, bool bSetupUndoRedo = true, bool bPrintPythonCommand = false);
+	bool SwapAllFunctionReferences(const FRigVMGraphFunctionIdentifier& InOldFunctionIdentifier, const FRigVMGraphFunctionIdentifier& InNewFunctionIdentifier, bool bSetupOrphanPins, bool bSetupUndoRedo = true, bool bPrintPythonCommand = false);
 
 	// Sets the remapped variable on a function reference node
 	UFUNCTION(BlueprintCallable, Category = RigVMController)
diff --git a/Engine/Plugins/Runtime/RigVM/Source/RigVMDeveloper/Public/RigVMModel/RigVMControllerActions.h b/Engine/Plugins/Runtime/RigVM/Source/RigVMDeveloper/Public/RigVMModel/RigVMControllerActions.h
--- a/Engine/Plugins/Runtime/RigVM/Source/RigVMDeveloper/Public/RigVMModel/RigVMControllerActions.h
+++ b/Engine/Plugins/Runtime/RigVM/Source/RigVMDeveloper/Public/RigVMModel/RigVMControllerActions.h
@@ -240,10 +240,16 @@ class RIGVMDEVELOPER_API URigVMActionStack : public UObject
 
 public:
 
+	static URigVMActionStack* GetDisabledActionStack();
+
 	// Begins an action and opens a bracket / scope.
 	template<class ActionType>
 	void BeginAction(ActionType& InAction)
 	{
+		if(IsDisabled())
+		{
+			return;
+		}
 #if RIGVM_ACTIONSTACK_VERBOSE_LOG		
 		TGuardValue<int32> TabDepthGuard(LogActionDepth, CurrentActions.Num());
 		LogAction<ActionType>(InAction, FRigVMBaseAction::BeginActionPrefix);
@@ -263,6 +269,10 @@ public:
 	template<class ActionType>
 	void EndAction(ActionType& InAction, bool bPerformMerge = false)
 	{
+		if(IsDisabled())
+		{
+			return;
+		}
 		ensure(CurrentActions.Num() > 0);
 		ensure((FRigVMBaseAction*)&InAction == CurrentActions.Last());
 		CurrentActions.Pop();
@@ -286,6 +296,10 @@ public:
 	template<class ActionType>
 	void CancelAction(ActionType& InAction)
 	{
+		if(IsDisabled())
+		{
+			return;
+		}
 		ensure(CurrentActions.Num() > 0);
 		ensure((FRigVMBaseAction*)&InAction == CurrentActions.Last());
 		CurrentActions.Pop();
@@ -305,6 +319,10 @@ public:
 	template<class ActionType>
 	void AddAction(const ActionType& InAction, bool bPerformMerge = false)
 	{
+		if(IsDisabled())
+		{
+			return;
+		}
 #if RIGVM_ACTIONSTACK_VERBOSE_LOG		
 		TGuardValue<int32> TabDepthGuard(LogActionDepth, CurrentActions.Num());
 		LogAction<ActionType>(InAction, FRigVMBaseAction::AddActionPrefix);
@@ -416,6 +434,8 @@ private:
 	
 	void LogAction(const UScriptStruct* InActionStruct, const FRigVMBaseAction& InAction, const FString& InPrefix);
 #endif
+
+	bool IsDisabled() const;
 	
 	UPROPERTY()
 	int32 ActionIndex;
