RigVM: Add a nullptr check when verifying types.

#rnx

diff --git a/Engine/Plugins/Runtime/RigVM/Source/RigVM/Private/RigVMCore/RigVMRegistry.cpp b/Engine/Plugins/Runtime/RigVM/Source/RigVM/Private/RigVMCore/RigVMRegistry.cpp
--- a/Engine/Plugins/Runtime/RigVM/Source/RigVM/Private/RigVMCore/RigVMRegistry.cpp
+++ b/Engine/Plugins/Runtime/RigVM/Source/RigVM/Private/RigVMCore/RigVMRegistry.cpp
@@ -1232,7 +1232,10 @@ bool FRigVMRegistry_NoLock::IsAllowedType_NoLock(const FProperty* InProperty) co
 
 	if(const FArrayProperty* ArrayProperty  = CastField<FArrayProperty>(InProperty))
 	{
-		return IsAllowedType_NoLock(ArrayProperty->Inner);
+		if (ArrayProperty->Inner)
+		{
+			return IsAllowedType_NoLock(ArrayProperty->Inner);
+		}
 	}
 	if(const FStructProperty* StructProperty = CastField<FStructProperty>(InProperty))
 	{
@@ -1259,12 +1262,12 @@ bool FRigVMRegistry_NoLock::IsAllowedType_NoLock(const FProperty* InProperty) co
 
 bool FRigVMRegistry_NoLock::IsAllowedType_NoLock(const UEnum* InEnum) const
 {
-	return !InEnum->HasAnyFlags(DisallowedFlags()) && InEnum->HasAllFlags(NeededFlags());
+	return InEnum && !InEnum->HasAnyFlags(DisallowedFlags()) && InEnum->HasAllFlags(NeededFlags());
 }
 
 bool FRigVMRegistry_NoLock::IsAllowedType_NoLock(const UStruct* InStruct) const
 {
-	if(InStruct->HasAnyFlags(DisallowedFlags()) || !InStruct->HasAllFlags(NeededFlags()))
+	if(!InStruct || InStruct->HasAnyFlags(DisallowedFlags()) || !InStruct->HasAllFlags(NeededFlags()))
 	{
 		return false;
 	}
