RigVM: Avoid sharing of internal state when inlining functions

#rnx

diff --git a/Engine/Plugins/Runtime/RigVM/Source/RigVMDeveloper/Private/RigVMCompiler/RigVMCompiler.cpp b/Engine/Plugins/Runtime/RigVM/Source/RigVMDeveloper/Private/RigVMCompiler/RigVMCompiler.cpp
--- a/Engine/Plugins/Runtime/RigVM/Source/RigVMDeveloper/Private/RigVMCompiler/RigVMCompiler.cpp
+++ b/Engine/Plugins/Runtime/RigVM/Source/RigVMDeveloper/Private/RigVMCompiler/RigVMCompiler.cpp
@@ -2505,6 +2505,57 @@ bool URigVMCompiler::TraverseInlineFunction(const FRigVMInlineFunctionExprAST* I
 				}
 			}
 
+			auto FindOrAddProperty = [&WorkData, MemoryType] (const FRigVMFunctionCompilationPropertyDescription& InProperty, const FString& InNewName, const bool bIsExecuteState) -> FRigVMOperand
+			{
+				// Sharing / reusing memory / operands happens as per following contract:
+				// 1. properties are only shared if their CPP type matches
+				// 2. Literal / constant memory is only shared if the constant values match
+				// 3. Work state is only shared if it is not internal work state private to the instruction referring to it
+				// 4. Work state of type FRigVMInstructionSetExecuteState (bIsExecuteState) is never shared either since it is work state private to a lazy branch.
+
+				FString NewName = InNewName;
+				int32 Suffix = 0;
+
+				FRigVMOperand Operand = WorkData.FindProperty(MemoryType, *NewName);
+
+				// look for a free / unused operand
+				while(Operand.IsValid())
+				{
+					const FRigVMPropertyDescription& ExistingProperty = WorkData.PropertyDescriptions[MemoryType][Operand.GetRegisterIndex()];
+					if(!bIsExecuteState &&
+						ExistingProperty.CPPType.Equals(InProperty.CPPType, ESearchCase::CaseSensitive))
+					{
+						if(MemoryType == ERigVMMemoryType::Literal)
+						{
+							if(ExistingProperty.DefaultValue.Equals(InProperty.DefaultValue))
+							{
+								return Operand;
+							}
+						}
+						else
+						{
+							if(MemoryType == ERigVMMemoryType::Work &&
+								RigVMTypeUtils::IsArrayType(InProperty.CPPType))
+							{
+								// for now we don't share operands which have an array type
+								// since they may represent internal state. multiple occurrences of
+								// the same function cannot share internal work state.
+							}
+							else
+							{
+								return Operand;
+							}
+						}
+					}
+
+					NewName = FString::Printf(TEXT("%s_%02d"), *InNewName, ++Suffix);
+					Operand = WorkData.FindProperty(MemoryType, *NewName);
+				}
+						
+				// create a new operand as needed
+				return WorkData.AddProperty(MemoryType, *NewName, InProperty.CPPType, InProperty.CPPTypeObject.Get(), InProperty.DefaultValue);
+			};
+
 			// Add internal properties
 			int32 StartIndex = MemoryType == ERigVMMemoryType::Work ? NumProperties : 0;
 			for (int32 PropertyIndex = StartIndex; PropertyIndex < Properties.Num(); ++PropertyIndex)
@@ -2517,19 +2568,14 @@ bool URigVMCompiler::TraverseInlineFunction(const FRigVMInlineFunctionExprAST* I
 
 				// instantiate function library specific work state as well as
 				// instruction set execute state - which is used for lazy blocks.
-				if (NewName.StartsWith(FunctionLibraryPrefix) ||
-					Description.CPPType.Equals(RigVMInstructionSetExecuteStateName))
+				const bool bIsExecuteState = Description.CPPType.Equals(RigVMInstructionSetExecuteStateName);
+				if (NewName.StartsWith(FunctionLibraryPrefix) || bIsExecuteState)
 				{
 					NewName = FString::Printf(TEXT("%s%s"), *FunctionReferenceNode->GetNodePath(), *NewName.RightChop(FunctionLibraryPrefix.Len()));
 					FRigVMPropertyDescription::SanitizeName(NewName);
 				}
 
-				// Reuse operands if already added to the WorkData
-				FRigVMOperand Operand = WorkData.FindProperty(MemoryType, *NewName);
-				if (!Operand.IsValid())
-				{
-					Operand = WorkData.AddProperty(MemoryType, *NewName, Description.CPPType, Description.CPPTypeObject.Get(), Description.DefaultValue);
-				}
+				const FRigVMOperand Operand = FindOrAddProperty(Description, NewName, bIsExecuteState);
 				FRigVMCompilerWorkData::FFunctionRegisterData Data = {FunctionReferenceNode, MemoryType, PropertyIndex};
 				WorkData.FunctionRegisterToOperand.Add(Data, Operand);
 			}
