Control Rig: Don't request BP action refresh on class deletion

#rnx

diff --git a/Engine/Plugins/Runtime/RigVM/Source/RigVMDeveloper/Private/RigVMBlueprint.cpp b/Engine/Plugins/Runtime/RigVM/Source/RigVMDeveloper/Private/RigVMBlueprint.cpp
--- a/Engine/Plugins/Runtime/RigVM/Source/RigVMDeveloper/Private/RigVMBlueprint.cpp
+++ b/Engine/Plugins/Runtime/RigVM/Source/RigVMDeveloper/Private/RigVMBlueprint.cpp
@@ -923,9 +923,10 @@ void URigVMBlueprint::PostSaveRoot(FObjectPostSaveRootContext ObjectSaveContext)
 	Super::PostSaveRoot(ObjectSaveContext);
 
 	// Make sure all the tags are accounted for in the TypeActions after we save
-	FBlueprintActionDatabase& ActionDatabase = FBlueprintActionDatabase::Get();
-	ActionDatabase.ClearAssetActions(GetClass());
-	ActionDatabase.RefreshClassActions(GetClass());
+	if (FBlueprintActionDatabase* ActionDatabase = FBlueprintActionDatabase::TryGet())
+	{
+		ActionDatabase->RefreshClassActions(GetClass());
+	}
 }
 
 void URigVMBlueprint::PostLoad()
diff --git a/Engine/Plugins/Runtime/RigVM/Source/RigVMDeveloper/Private/RigVMBlueprintUtils.cpp b/Engine/Plugins/Runtime/RigVM/Source/RigVMDeveloper/Private/RigVMBlueprintUtils.cpp
--- a/Engine/Plugins/Runtime/RigVM/Source/RigVMDeveloper/Private/RigVMBlueprintUtils.cpp
+++ b/Engine/Plugins/Runtime/RigVM/Source/RigVMDeveloper/Private/RigVMBlueprintUtils.cpp
@@ -147,7 +147,6 @@ void FRigVMBlueprintUtils::HandleAssetDeleted(const FAssetData& InAssetData)
 		if (FBlueprintActionDatabase* ActionDatabase = FBlueprintActionDatabase::TryGet())
 		{
 			ActionDatabase->ClearAssetActions(InAssetData.GetClass());
-			ActionDatabase->RefreshClassActions(InAssetData.GetClass());
 		}
 	}
 }
