RigVM: Standalone Node preview widget

#rnx

diff --git a/Engine/Plugins/Animation/ControlRig/Source/ControlRigDeveloper/Private/Graph/ControlRigGraph.cpp b/Engine/Plugins/Animation/ControlRig/Source/ControlRigDeveloper/Private/Graph/ControlRigGraph.cpp
--- a/Engine/Plugins/Animation/ControlRig/Source/ControlRigDeveloper/Private/Graph/ControlRigGraph.cpp
+++ b/Engine/Plugins/Animation/ControlRig/Source/ControlRigDeveloper/Private/Graph/ControlRigGraph.cpp
@@ -167,7 +167,7 @@ const TArray<TSharedPtr<FRigVMStringWithTag>>* UControlRigGraph::GetElementNameL
 	
 	if(!ElementNameLists.Contains(InElementType))
 	{
-		const UControlRigBlueprint* Blueprint = CastChecked<UControlRigBlueprint>(GetBlueprint());
+		const UControlRigBlueprint* Blueprint = Cast<UControlRigBlueprint>(GetBlueprint());
 		if(Blueprint == nullptr)
 		{
 			return &EmptyElementNameList;
diff --git a/Engine/Plugins/Animation/ControlRig/Source/ControlRigEditor/Private/Graph/ControlRigGraphPanelPinFactory.cpp b/Engine/Plugins/Animation/ControlRig/Source/ControlRigEditor/Private/Graph/ControlRigGraphPanelPinFactory.cpp
--- a/Engine/Plugins/Animation/ControlRig/Source/ControlRigEditor/Private/Graph/ControlRigGraphPanelPinFactory.cpp
+++ b/Engine/Plugins/Animation/ControlRig/Source/ControlRigEditor/Private/Graph/ControlRigGraphPanelPinFactory.cpp
@@ -40,9 +40,13 @@ TSharedPtr<SGraphPin> FControlRigGraphPanelPinFactory::CreatePin(UEdGraphPin* In
 	if (const UEdGraphNode* OwningNode = InPin->GetOwningNode())
 	{
 		// only create pins within control rig graphs
-		if (Cast<UControlRigGraph>(OwningNode->GetGraph()) == nullptr)
+		if(const URigVMEdGraph* EdGraph = Cast<URigVMEdGraph>(OwningNode->GetGraph()))
 		{
-			return nullptr;
+			if ((Cast<UControlRigGraph>(EdGraph) == nullptr) &&
+				!EdGraph->IsPreviewGraph())
+			{
+				return nullptr;
+			}
 		}
 	}
 	
@@ -141,9 +145,13 @@ TSharedPtr<SGraphPin> FControlRigGraphPanelPinFactory::CreatePin_Internal(UEdGra
 		if (const UEdGraphNode* OwningNode = InPin->GetOwningNode())
 		{
 			// only create pins within control rig graphs
-			if (Cast<UControlRigGraph>(OwningNode->GetGraph()) == nullptr)
+			if(const URigVMEdGraph* EdGraph = Cast<URigVMEdGraph>(OwningNode->GetGraph()))
 			{
-				return nullptr;
+				if ((Cast<UControlRigGraph>(EdGraph) == nullptr) &&
+					!EdGraph->IsPreviewGraph())
+				{
+					return nullptr;
+				}
 			}
 		}
 
diff --git a/Engine/Plugins/Runtime/RigVM/Source/RigVMDeveloper/Private/EdGraph/RigVMEdGraph.cpp b/Engine/Plugins/Runtime/RigVM/Source/RigVMDeveloper/Private/EdGraph/RigVMEdGraph.cpp
--- a/Engine/Plugins/Runtime/RigVM/Source/RigVMDeveloper/Private/EdGraph/RigVMEdGraph.cpp
+++ b/Engine/Plugins/Runtime/RigVM/Source/RigVMDeveloper/Private/EdGraph/RigVMEdGraph.cpp
@@ -103,16 +103,41 @@ void URigVMEdGraph::HandleRigVMGraphRenamed(const FString& InOldNodePath, const
 	}
 }
 
+const URigVMBlueprint* URigVMEdGraph::GetBlueprintDefaultObject() const
+{
+	if(RigVMBlueprintClass)
+	{
+		return RigVMBlueprintClass->GetDefaultObject<URigVMBlueprint>();
+	}
+	return nullptr;
+}
+
+void URigVMEdGraph::SetBlueprintClass(const UClass* InClass)
+{
+	RigVMBlueprintClass = InClass;
+	if(const URigVMBlueprint* Blueprint = GetBlueprintDefaultObject())
+	{
+		Schema = Blueprint->GetRigVMEdGraphSchemaClass();
+	}
+}
+
 void URigVMEdGraph::InitializeFromBlueprint(URigVMBlueprint* InBlueprint)
 {
 	DECLARE_SCOPE_HIERARCHICAL_COUNTER_FUNC()
 
+	SetBlueprintClass(InBlueprint->GetClass());
 	InBlueprint->OnModified().RemoveAll(this);
 	InBlueprint->OnModified().AddUObject(this, &URigVMEdGraph::HandleModifiedEvent);
 	InBlueprint->OnVMCompiled().RemoveAll(this);
 	InBlueprint->OnVMCompiled().AddUObject(this, &URigVMEdGraph::HandleVMCompiledEvent);
 }
 
+bool URigVMEdGraph::IsPreviewGraph() const
+{
+	// if we are not below a client host we are preview
+	return GetImplementingOuter<IRigVMClientHost>() == nullptr;
+}
+
 const URigVMEdGraphSchema* URigVMEdGraph::GetRigVMEdGraphSchema()
 {
 	return CastChecked<const URigVMEdGraphSchema>(GetSchema());
@@ -917,6 +942,14 @@ URigVMGraph* URigVMEdGraph::GetModel() const
 		CachedModelGraph = Model;
 		return Model;
 	}
+
+	// for preview scenarios we'll nest the edgraph under the model graph
+	if(URigVMGraph* Model = GetTypedOuter<URigVMGraph>())
+	{
+		CachedModelGraph = Model;
+		return Model;
+	}
+	
 	return nullptr;
 }
 
@@ -926,6 +959,14 @@ URigVMController* URigVMEdGraph::GetController() const
 	{
 		return Client->GetOrCreateController(GetModel());
 	}
+	
+	// for preview scenarios we'll nest the edgraph under the model graph
+	// and the model graph under the controller
+	if(URigVMController* Controller = GetTypedOuter<URigVMController>())
+	{
+		return Controller;
+	}
+	
 	return nullptr;
 }
 
diff --git a/Engine/Plugins/Runtime/RigVM/Source/RigVMDeveloper/Private/EdGraph/RigVMEdGraphNode.cpp b/Engine/Plugins/Runtime/RigVM/Source/RigVMDeveloper/Private/EdGraph/RigVMEdGraphNode.cpp
--- a/Engine/Plugins/Runtime/RigVM/Source/RigVMDeveloper/Private/EdGraph/RigVMEdGraphNode.cpp
+++ b/Engine/Plugins/Runtime/RigVM/Source/RigVMDeveloper/Private/EdGraph/RigVMEdGraphNode.cpp
@@ -38,6 +38,7 @@ URigVMEdGraphNode::URigVMEdGraphNode()
 : Dimensions(0.0f, 0.0f)
 , NodeTitle(FText::GetEmpty())
 , FullNodeTitle(FText::GetEmpty())
+, bSubTitleEnabled(true)
 , NodeTopologyVersion(INDEX_NONE)
 , CachedTitleColor(FLinearColor(0.f, 0.f, 0.f, 0.f))
 , CachedNodeColor(FLinearColor(0.f, 0.f, 0.f, 0.f))
@@ -183,7 +184,7 @@ FText URigVMEdGraphNode::GetNodeTitle(ENodeTitleType::Type TitleType) const
 
 		FullNodeTitle = NodeTitle;
 
-		if(!SubTitle.IsEmpty())
+		if(!SubTitle.IsEmpty() && bSubTitleEnabled)
 		{
 			FullNodeTitle = FText::FromString(FString::Printf(TEXT("%s\n%s"), *NodeTitle.ToString(), *SubTitle));
 		}
@@ -1104,6 +1105,15 @@ bool URigVMEdGraphNode::DrawAsCompactNode() const
 	return DrawAsCompactNodeCache.GetValue();
 }
 
+void URigVMEdGraphNode::SetModelNode(URigVMNode* InModelNode)
+{
+	ModelNodePath = InModelNode->GetNodePath();
+	CachedModelNode = InModelNode;
+	SyncGraphNodeTitleWithModelNodeTitle();
+	AllocateDefaultPins();
+	PostReconstructNode();
+}
+
 bool URigVMEdGraphNode::ModelPinRemoved_Internal(const URigVMPin* InModelPin)
 {
 	auto RemoveGraphPin = [this](const URigVMPin* InModelPin, bool bAsInput)
diff --git a/Engine/Plugins/Runtime/RigVM/Source/RigVMDeveloper/Public/EdGraph/RigVMEdGraph.h b/Engine/Plugins/Runtime/RigVM/Source/RigVMDeveloper/Public/EdGraph/RigVMEdGraph.h
--- a/Engine/Plugins/Runtime/RigVM/Source/RigVMDeveloper/Public/EdGraph/RigVMEdGraph.h
+++ b/Engine/Plugins/Runtime/RigVM/Source/RigVMDeveloper/Public/EdGraph/RigVMEdGraph.h
@@ -141,7 +141,10 @@ public:
 	virtual void HandleRigVMGraphRenamed(const FString& InOldNodePath, const FString& InNewNodePath) override;
 
 	/** Set up this graph */
+	const URigVMBlueprint* GetBlueprintDefaultObject() const;
+	void SetBlueprintClass(const UClass* InClass);
 	virtual void InitializeFromBlueprint(URigVMBlueprint* InBlueprint);
+	bool IsPreviewGraph() const;
 
 	/** Get the ed graph schema */
 	const URigVMEdGraphSchema* GetRigVMEdGraphSchema();
@@ -200,6 +203,7 @@ private:
 	TMap<FName, UEdGraphNode*> ModelNodePathToEdNode;
 	mutable TWeakObjectPtr<URigVMGraph> CachedModelGraph;
 	TArray<TSharedPtr<FRigVMStringWithTag>> EntryNameList;
+	const UClass* RigVMBlueprintClass;
 
 #endif
 	friend class URigVMEdGraphNode;
diff --git a/Engine/Plugins/Runtime/RigVM/Source/RigVMDeveloper/Public/EdGraph/RigVMEdGraphNode.h b/Engine/Plugins/Runtime/RigVM/Source/RigVMDeveloper/Public/EdGraph/RigVMEdGraphNode.h
--- a/Engine/Plugins/Runtime/RigVM/Source/RigVMDeveloper/Public/EdGraph/RigVMEdGraphNode.h
+++ b/Engine/Plugins/Runtime/RigVM/Source/RigVMDeveloper/Public/EdGraph/RigVMEdGraphNode.h
@@ -71,8 +71,16 @@ private:
 	/** The cached fulol node title */
 	mutable FText FullNodeTitle;
 
+	/** Set this to true to enable the sub title */
+	bool bSubTitleEnabled;
+
 public:
 
+	void SetSubTitleEnabled(bool bEnabled = true)
+	{
+		bSubTitleEnabled = bEnabled;
+	}
+
 	DECLARE_MULTICAST_DELEGATE(FNodeTitleDirtied);
 	DECLARE_MULTICAST_DELEGATE(FNodePinsChanged);
 	DECLARE_MULTICAST_DELEGATE(FNodePinExpansionChanged);
@@ -211,6 +219,9 @@ public:
 	/** Returns true if this node is relying on the cast template */
 	bool DrawAsCompactNode() const;
 
+	/** Sets the node's content from an external client, used for preview nodes without a graph */
+	void SetModelNode(URigVMNode* InModelNode);
+
 protected:
 
 	FLinearColor GetNodeProfilingColor() const;
diff --git a/Engine/Plugins/Runtime/RigVM/Source/RigVMEditor/Private/EdGraph/RigVMEdGraphPanelNodeFactory.cpp b/Engine/Plugins/Runtime/RigVM/Source/RigVMEditor/Private/EdGraph/RigVMEdGraphPanelNodeFactory.cpp
--- a/Engine/Plugins/Runtime/RigVM/Source/RigVMEditor/Private/EdGraph/RigVMEdGraphPanelNodeFactory.cpp
+++ b/Engine/Plugins/Runtime/RigVM/Source/RigVMEditor/Private/EdGraph/RigVMEdGraphPanelNodeFactory.cpp
@@ -19,11 +19,14 @@ FName FRigVMEdGraphPanelNodeFactory::GetFactoryName() const
 TSharedPtr<SGraphNode> FRigVMEdGraphPanelNodeFactory::CreateNode(UEdGraphNode* Node) const
 {
 	// we need to check if this is the right factory for the implementation
-	if(const URigVMBlueprint* Blueprint = Node->GetTypedOuter<URigVMBlueprint>())
+	if(const URigVMEdGraph* EdGraph = Node->GetTypedOuter<URigVMEdGraph>())
 	{
-		if(Blueprint->GetPanelNodeFactoryName() != GetFactoryName())
+		if(const URigVMBlueprint* Blueprint = EdGraph->GetBlueprintDefaultObject())
 		{
-			return nullptr;
+			if(Blueprint->GetPanelNodeFactoryName() != GetFactoryName())
+			{
+				return nullptr;
+			}
 		}
 	}
 
diff --git a/Engine/Plugins/Runtime/RigVM/Source/RigVMEditor/Private/EdGraph/RigVMEdGraphPanelPinFactory.cpp b/Engine/Plugins/Runtime/RigVM/Source/RigVMEditor/Private/EdGraph/RigVMEdGraphPanelPinFactory.cpp
--- a/Engine/Plugins/Runtime/RigVM/Source/RigVMEditor/Private/EdGraph/RigVMEdGraphPanelPinFactory.cpp
+++ b/Engine/Plugins/Runtime/RigVM/Source/RigVMEditor/Private/EdGraph/RigVMEdGraphPanelPinFactory.cpp
@@ -36,11 +36,14 @@ TSharedPtr<SGraphPin> FRigVMEdGraphPanelPinFactory::CreatePin(UEdGraphPin* InPin
 	// we need to check if this is the right factory for the implementation
 	if(const UEdGraphNode* EdGraphNode = InPin->GetOuter())
 	{
-		if(const URigVMBlueprint* Blueprint = EdGraphNode->GetTypedOuter<URigVMBlueprint>())
+		if(const URigVMEdGraph* EdGraph = EdGraphNode->GetTypedOuter<URigVMEdGraph>())
 		{
-			if(Blueprint->GetPanelPinFactoryName() != GetFactoryName())
+			if(const URigVMBlueprint* Blueprint = EdGraph->GetBlueprintDefaultObject())
 			{
-				return nullptr;
+				if(Blueprint->GetPanelPinFactoryName() != GetFactoryName())
+				{
+					return nullptr;
+				}
 			}
 		}
 	}
diff --git a/Engine/Plugins/Runtime/RigVM/Source/RigVMEditor/Private/Editor/RigVMEditorTools.cpp b/Engine/Plugins/Runtime/RigVM/Source/RigVMEditor/Private/Editor/RigVMEditorTools.cpp
--- a/Engine/Plugins/Runtime/RigVM/Source/RigVMEditor/Private/Editor/RigVMEditorTools.cpp
+++ b/Engine/Plugins/Runtime/RigVM/Source/RigVMEditor/Private/Editor/RigVMEditorTools.cpp
@@ -1,13 +1,14 @@
 // Copyright Epic Games, Inc. All Rights Reserved.
 
 #include "Editor/RigVMEditorTools.h"
-
+#include "Editor.h"
 #include "RigVMModel/RigVMController.h"
 #include "RigVMModel/RigVMGraph.h"
 #include "RigVMModel/RigVMFunctionLibrary.h"
 #include "RigVMCore/RigVMGraphFunctionHost.h"
 #include "RigVMCore/RigVMGraphFunctionDefinition.h"
 #include "Widgets/SRigVMGraphFunctionLocalizationWidget.h"
+#include "Subsystems/EditorAssetSubsystem.h"
 
 namespace UE::RigVM::Editor::Tools
 {
@@ -107,5 +108,16 @@ RIGVMEDITOR_API void OnRequestLocalizeFunctionDialog(FRigVMGraphFunctionIdentifi
 		}
 	}
 }
+	
+RIGVMEDITOR_API FAssetData FindAssetFromAnyPath(const FString& InPartialOrFullPath, bool bConvertToRootPath)
+{
+	UEditorAssetSubsystem* EditorAssetSubsystem = GEditor->GetEditorSubsystem<UEditorAssetSubsystem>();
+	if(bConvertToRootPath)
+	{
+		return EditorAssetSubsystem->FindAssetData(FSoftObjectPath(InPartialOrFullPath).GetWithoutSubPath().ToString());
+	}
+	return EditorAssetSubsystem->FindAssetData(InPartialOrFullPath);
+}
+
 
 } // end namespace UE::RigVM::Editor::Tools
diff --git a/Engine/Plugins/Runtime/RigVM/Source/RigVMEditor/Private/Editor/RigVMMinimalEnvironment.cpp b/Engine/Plugins/Runtime/RigVM/Source/RigVMEditor/Private/Editor/RigVMMinimalEnvironment.cpp
new file mode 100644
--- /dev/null
+++ b/Engine/Plugins/Runtime/RigVM/Source/RigVMEditor/Private/Editor/RigVMMinimalEnvironment.cpp
@@ -0,0 +1,176 @@
+// Copyright Epic Games, Inc. All Rights Reserved.
+
+#include "Editor/RigVMMinimalEnvironment.h"
+#include "RigVMBlueprint.h"
+#include "Editor/RigVMEditorTools.h"
+
+#define LOCTEXT_NAMESPACE "RigVMMinimalEnvironment"
+
+FRigVMMinimalEnvironment::FRigVMMinimalEnvironment(const UClass* InRigVMBlueprintClass)
+{
+	ModelController = TStrongObjectPtr<URigVMController>(NewObject<URigVMController>());
+	ModelGraph = TStrongObjectPtr<URigVMGraph>(NewObject<URigVMGraph>(ModelController.Get()));
+	ModelController->SetGraph(ModelGraph.Get());
+
+	EdGraphClass = URigVMEdGraph::StaticClass();
+	EdGraphNodeClass = URigVMEdGraphNode::StaticClass();
+
+	SetSchemata(InRigVMBlueprintClass ? InRigVMBlueprintClass : URigVMBlueprint::StaticClass());
+}
+
+URigVMGraph* FRigVMMinimalEnvironment::GetModel() const
+{
+	return ModelGraph.Get();
+}
+
+URigVMController* FRigVMMinimalEnvironment::GetController() const
+{
+	return ModelController.Get();
+}
+
+URigVMNode* FRigVMMinimalEnvironment::GetNode() const
+{
+	if(ModelNode.IsValid())
+	{
+		return ModelNode.Get();
+	}
+	return nullptr;
+}
+
+void FRigVMMinimalEnvironment::SetNode(URigVMNode* InModelNode)
+{
+	if(!ModelHandle.IsValid())
+	{
+		ModelHandle = ModelGraph->OnModified().AddSP(this, &FRigVMMinimalEnvironment::HandleModified);
+	}
+	if(URigVMNode* PreviousNode = GetNode())
+	{
+		GetController()->RemoveNode(PreviousNode);
+	}
+	
+	ModelNode = InModelNode;
+
+	if(GetEdGraphNode() == nullptr)
+	{
+		if(URigVMEdGraph* MyEdGraph = GetEdGraph())
+		{
+			EdGraphNode = NewObject<URigVMEdGraphNode>(MyEdGraph, EdGraphNodeClass);
+			MyEdGraph->Nodes.Add(EdGraphNode.Get());
+		}
+	}
+
+	HandleModified(ERigVMGraphNotifType::NodeAdded, GetModel(), GetNode());
+}
+
+void FRigVMMinimalEnvironment::SetFunctionNode(const FRigVMGraphFunctionIdentifier& InIdentifier)
+{
+	check(InIdentifier.IsValid());
+
+	const FRigVMGraphFunctionHeader Header = FRigVMGraphFunctionHeader::FindGraphFunctionHeader(InIdentifier);
+	if(Header.IsValid())
+	{
+		const FAssetData AssetData = UE::RigVM::Editor::Tools::FindAssetFromAnyPath(Header.LibraryPointer.GetLibraryNodePath(), true);
+		if(AssetData.IsValid())
+		{
+			if(const UClass* Class = AssetData.GetClass())
+			{
+				if(Class->IsChildOf(URigVMBlueprint::StaticClass()))
+				{
+					SetSchemata(Class);
+				}
+			}
+		}
+
+		if(URigVMFunctionReferenceNode* FunctionReferenceNode = Cast<URigVMFunctionReferenceNode>(GetNode()))
+		{
+			(void)GetController()->SwapFunctionReference(FunctionReferenceNode, InIdentifier, false, false, false);
+		}
+		else
+		{
+			URigVMNode* Node = GetController()->AddFunctionReferenceNodeFromDescription(Header, FVector2D::ZeroVector, FString(), false, false);
+			SetNode(Node);
+		}
+	}
+}
+
+URigVMEdGraph* FRigVMMinimalEnvironment::GetEdGraph() const
+{
+	return EdGraph.Get();
+}
+
+URigVMEdGraphNode* FRigVMMinimalEnvironment::GetEdGraphNode() const
+{
+	if(EdGraphNode.IsValid())
+	{
+		return EdGraphNode.Get();
+	}
+	return nullptr;
+}
+
+void FRigVMMinimalEnvironment::SetSchemata(const UClass* InRigVMBlueprintClass)
+{
+	check(InRigVMBlueprintClass);
+
+	const URigVMBlueprint* CDO = InRigVMBlueprintClass->GetDefaultObject<URigVMBlueprint>();
+	check(CDO);
+
+	EdGraphClass = CDO->GetRigVMEdGraphClass();
+	EdGraphNodeClass = CDO->GetRigVMEdGraphNodeClass();
+
+	if(!EdGraph.IsValid() || EdGraph->GetClass() != EdGraphClass)
+	{
+		EdGraph = TStrongObjectPtr<URigVMEdGraph>(NewObject<URigVMEdGraph>(ModelGraph.Get(), EdGraphClass));
+	}
+
+	ModelController->SetSchemaClass(CDO->GetRigVMSchemaClass());
+	EdGraph->SetBlueprintClass(InRigVMBlueprintClass);
+}
+
+FSimpleDelegate& FRigVMMinimalEnvironment::OnChanged()
+{
+	return ChangedDelegate;
+}
+
+void FRigVMMinimalEnvironment::Tick_GameThead(float InDeltaTime)
+{
+	if(NumModifications.exchange(0) > 0)
+	{
+		// refresh the EdGraphNode
+		if(URigVMNode* MyModelNode = GetNode())
+		{
+			if(URigVMEdGraphNode* MyEdGraphNode = GetEdGraphNode())
+			{
+				MyEdGraphNode->SetSubTitleEnabled(false);
+				MyEdGraphNode->SetModelNode(MyModelNode);
+				(void)ChangedDelegate.ExecuteIfBound();
+			}
+		}
+	}
+}
+
+void FRigVMMinimalEnvironment::HandleModified(ERigVMGraphNotifType InNotification, URigVMGraph* InGraph, UObject* InSubject)
+{
+	if(InGraph != GetModel())
+	{
+		return;
+	}
+
+	switch(InNotification)
+	{
+		case ERigVMGraphNotifType::NodeSelected:
+		case ERigVMGraphNotifType::NodeDeselected:
+		case ERigVMGraphNotifType::PinExpansionChanged:
+		case ERigVMGraphNotifType::InteractionBracketOpened:
+		case ERigVMGraphNotifType::InteractionBracketClosed:
+		case ERigVMGraphNotifType::InteractionBracketCanceled:
+		{
+			break;
+		}
+		default:
+		{
+			(void)NumModifications.fetch_add(1);
+		}
+	}
+}
+
+#undef LOCTEXT_NAMESPACE
diff --git a/Engine/Plugins/Runtime/RigVM/Source/RigVMEditor/Private/Widgets/SRigVMNodePreviewWidget.cpp b/Engine/Plugins/Runtime/RigVM/Source/RigVMEditor/Private/Widgets/SRigVMNodePreviewWidget.cpp
new file mode 100644
--- /dev/null
+++ b/Engine/Plugins/Runtime/RigVM/Source/RigVMEditor/Private/Widgets/SRigVMNodePreviewWidget.cpp
@@ -0,0 +1,37 @@
+// Copyright Epic Games, Inc. All Rights Reserved.
+
+#include "Widgets/SRigVMNodePreviewWidget.h"
+
+#define LOCTEXT_NAMESPACE "SRigVMNodePreviewWidget"
+
+void SRigVMNodePreviewWidget::Construct(const FArguments& InArgs)
+{
+	Environment = InArgs._Environment;
+	check(Environment);
+	Environment->OnChanged().BindSP(this, &SRigVMNodePreviewWidget::UpdateNodeWidget);
+	
+	ChildSlot
+	.Padding(InArgs._Padding);
+	
+	UpdateNodeWidget();
+}
+
+void SRigVMNodePreviewWidget::Tick(const FGeometry& AllottedGeometry, const double InCurrentTime, const float InDeltaTime)
+{
+	SCompoundWidget::Tick(AllottedGeometry, InCurrentTime, InDeltaTime);
+	Environment->Tick_GameThead(InDeltaTime);
+}
+
+void SRigVMNodePreviewWidget::UpdateNodeWidget()
+{
+	ChildSlot.DetachWidget();
+
+	if(URigVMEdGraphNode* EdGraphNode = Environment->GetEdGraphNode())
+	{
+		ChildSlot.AttachWidget(SNew(SRigVMGraphNode)
+			.GraphNodeObj(EdGraphNode)
+		);
+	}
+}
+
+#undef LOCTEXT_NAMESPACE
diff --git a/Engine/Plugins/Runtime/RigVM/Source/RigVMEditor/Public/Editor/RigVMEditorTools.h b/Engine/Plugins/Runtime/RigVM/Source/RigVMEditor/Public/Editor/RigVMEditorTools.h
--- a/Engine/Plugins/Runtime/RigVM/Source/RigVMEditor/Public/Editor/RigVMEditorTools.h
+++ b/Engine/Plugins/Runtime/RigVM/Source/RigVMEditor/Public/Editor/RigVMEditorTools.h
@@ -3,6 +3,7 @@
 #pragma once
 
 #include "CoreMinimal.h"
+#include "AssetRegistry/AssetData.h"
 
 class URigVMController;
 class URigVMGraph;
@@ -24,5 +25,7 @@ RIGVMEDITOR_API void OnRequestLocalizeFunctionDialog(FRigVMGraphFunctionIdentifi
 	, URigVMController* InTargetController
 	, IRigVMGraphFunctionHost* InTargetFunctionHost
 	, bool bForce);
+	
+RIGVMEDITOR_API FAssetData FindAssetFromAnyPath(const FString& InPartialOrFullPath, bool bConvertToRootPath);
 
 } // end namespace UE::RigVM::Editor::Tools
diff --git a/Engine/Plugins/Runtime/RigVM/Source/RigVMEditor/Public/Editor/RigVMMinimalEnvironment.h b/Engine/Plugins/Runtime/RigVM/Source/RigVMEditor/Public/Editor/RigVMMinimalEnvironment.h
new file mode 100644
--- /dev/null
+++ b/Engine/Plugins/Runtime/RigVM/Source/RigVMEditor/Public/Editor/RigVMMinimalEnvironment.h
@@ -0,0 +1,45 @@
+// Copyright Epic Games, Inc. All Rights Reserved.
+
+#pragma once
+
+#include "CoreMinimal.h"
+#include "RigVMModel/RigVMGraph.h"
+#include "RigVMModel/RigVMController.h"
+#include "RigVMModel/RigVMClient.h"
+#include "EdGraph/RigVMEdGraph.h"
+#include "EdGraph/RigVMEdGraphNode.h"
+#include "UObject/StrongObjectPtr.h"
+
+class RIGVMEDITOR_API FRigVMMinimalEnvironment : public TSharedFromThis<FRigVMMinimalEnvironment>
+{
+public:
+	FRigVMMinimalEnvironment(const UClass* InRigVMBlueprintClass = nullptr);
+
+	URigVMGraph* GetModel() const;
+	URigVMController* GetController() const;
+	URigVMNode* GetNode() const;
+	URigVMEdGraph* GetEdGraph() const;
+	URigVMEdGraphNode* GetEdGraphNode() const;
+
+	void SetSchemata(const UClass* InRigVMBlueprintClass);
+	void SetNode(URigVMNode* InModelNode);
+	void SetFunctionNode(const FRigVMGraphFunctionIdentifier& InIdentifier);
+	
+	FSimpleDelegate& OnChanged();
+	void Tick_GameThead(float InDeltaTime);
+
+private:
+
+	void HandleModified(ERigVMGraphNotifType InNotification, URigVMGraph* InGraph, UObject* InSubject);
+	
+	TStrongObjectPtr<URigVMGraph> ModelGraph;
+	TStrongObjectPtr<URigVMController> ModelController;
+	UClass* EdGraphClass;
+	UClass* EdGraphNodeClass;
+	TStrongObjectPtr<URigVMEdGraph> EdGraph;
+	TWeakObjectPtr<URigVMNode> ModelNode;
+	TWeakObjectPtr<URigVMEdGraphNode> EdGraphNode;
+	std::atomic<int32> NumModifications;
+	FSimpleDelegate ChangedDelegate;
+	FDelegateHandle ModelHandle;
+};
diff --git a/Engine/Plugins/Runtime/RigVM/Source/RigVMEditor/Public/Widgets/SRigVMNodePreviewWidget.h b/Engine/Plugins/Runtime/RigVM/Source/RigVMEditor/Public/Widgets/SRigVMNodePreviewWidget.h
new file mode 100644
--- /dev/null
+++ b/Engine/Plugins/Runtime/RigVM/Source/RigVMEditor/Public/Widgets/SRigVMNodePreviewWidget.h
@@ -0,0 +1,28 @@
+// Copyright Epic Games, Inc. All Rights Reserved.
+
+#pragma once
+
+#include "CoreMinimal.h"
+#include "SRigVMGraphNode.h"
+#include "Editor/RigVMMinimalEnvironment.h"
+
+class RIGVMEDITOR_API SRigVMNodePreviewWidget : public SCompoundWidget
+{
+public:
+	
+	SLATE_BEGIN_ARGS(SRigVMNodePreviewWidget)
+		: _Padding(16)
+	{}
+	SLATE_ATTRIBUTE(FMargin, Padding)
+	SLATE_ARGUMENT(TSharedPtr<FRigVMMinimalEnvironment>, Environment)
+	SLATE_END_ARGS()
+	
+	void Construct(const FArguments& InArgs);
+
+	virtual void Tick(const FGeometry& AllottedGeometry, const double InCurrentTime, const float InDeltaTime) override;
+
+private:
+	void UpdateNodeWidget();
+	
+	TSharedPtr<FRigVMMinimalEnvironment> Environment;
+};
diff --git a/Engine/Source/Editor/GraphEditor/Private/SGraphPin.cpp b/Engine/Source/Editor/GraphEditor/Private/SGraphPin.cpp
--- a/Engine/Source/Editor/GraphEditor/Private/SGraphPin.cpp
+++ b/Engine/Source/Editor/GraphEditor/Private/SGraphPin.cpp
@@ -113,6 +113,7 @@ SGraphPin::SGraphPin()
 	, bIsMovingLinks(false)
 	, bUsePinColorForText(false)
 	, bDragAndDropEnabled(true)
+	, bIsDiffHighlighted(false)
 	, bFadeConnections(false)
 {
 	IsEditable = true;
