RigVM Editor: SRigVMLogWidget

#rnx

diff --git a/Engine/Plugins/Runtime/RigVM/Source/RigVMEditor/Private/Widgets/SRigVMLogWidget.cpp b/Engine/Plugins/Runtime/RigVM/Source/RigVMEditor/Private/Widgets/SRigVMLogWidget.cpp
new file mode 100644
--- /dev/null
+++ b/Engine/Plugins/Runtime/RigVM/Source/RigVMEditor/Private/Widgets/SRigVMLogWidget.cpp
@@ -0,0 +1,125 @@
+// Copyright Epic Games, Inc. All Rights Reserved.
+
+#include "Widgets/SRigVMLogWidget.h"
+#include "Modules/ModuleManager.h"
+#include "MessageLogModule.h"
+#include "UserInterface/SMessageLogListing.h"
+#include "Kismet2/CompilerResultsLog.h"
+
+#define LOCTEXT_NAMESPACE "SRigVMLogWidget"
+
+SRigVMLogWidget::SRigVMLogWidget()
+{
+}
+
+SRigVMLogWidget::~SRigVMLogWidget()
+{
+	for(const TWeakPtr<IMessageLogListing>& Listing : BoundListings)
+	{
+		if(Listing.IsValid())
+		{
+			Listing.Pin()->OnDataChanged().RemoveAll(this);
+		}
+	}
+}
+
+void SRigVMLogWidget::Construct(
+	const FArguments& InArgs)
+{
+	ListingModel = FMessageLogListingModel::Create(InArgs._LogName);
+	ListingView = FMessageLogListingViewModel::Create(ListingModel.ToSharedRef(), InArgs._LogLabel);
+	ListingView->SetShowFilters(true);
+	ListingView->SetAllowClear(true);
+	ListingView->SetDiscardDuplicates(false);
+	ListingView->SetScrollToBottom(true);
+
+	ChildSlot
+	[
+		SNew(SBox)
+		.HeightOverride(InArgs._HeightOverride)
+		[
+			SNew(SMessageLogListing, ListingView.ToSharedRef())
+		]
+	];
+}
+
+void SRigVMLogWidget::BindLog(const FName& InLogName)
+{
+	FMessageLogModule& MessageLogModule = FModuleManager::LoadModuleChecked<FMessageLogModule>("MessageLog");
+	if(MessageLogModule.IsRegisteredLogListing(InLogName))
+	{
+		const TSharedRef<IMessageLogListing> Listing = MessageLogModule.GetLogListing(InLogName);
+		BindLog(Listing);
+	}
+}
+
+void SRigVMLogWidget::UnbindLog(const FName& InLogName)
+{
+	FMessageLogModule& MessageLogModule = FModuleManager::LoadModuleChecked<FMessageLogModule>("MessageLog");
+	if(MessageLogModule.IsRegisteredLogListing(InLogName))
+	{
+		const TSharedRef<IMessageLogListing> Listing = MessageLogModule.GetLogListing(InLogName);
+		UnbindLog(Listing);
+	}
+}
+
+void SRigVMLogWidget::BindLog(TSharedPtr<IMessageLogListing> InListing)
+{
+	if(InListing.IsValid())
+	{
+		const TWeakPtr<IMessageLogListing> WeakListing = InListing.ToWeakPtr();
+		if(!BoundListings.Contains(WeakListing))
+		{
+			BoundListings.Add(WeakListing);
+			ListingInfo.FindOrAdd(InListing->GetName()) = FBoundListingInfo(*InListing.Get());
+			(void)InListing->OnDataChanged().AddSP(this, &SRigVMLogWidget::OnBoundListingChanged, WeakListing);
+		}
+	}
+}
+
+void SRigVMLogWidget::UnbindLog(TSharedPtr<IMessageLogListing> InListing)
+{
+	if(InListing.IsValid())
+	{
+		BoundListings.Remove(InListing.ToWeakPtr());
+		ListingInfo.Remove(InListing->GetName());
+		(void)InListing->OnDataChanged().RemoveAll(this);
+	}
+}
+
+void SRigVMLogWidget::BindLog(const URigVMBlueprint* InRigVMBlueprint)
+{
+	TSharedRef<IMessageLogListing> Listing = FCompilerResultsLog::GetBlueprintMessageLog(const_cast<URigVMBlueprint*>(InRigVMBlueprint));
+	BindLog(Listing);
+}
+
+void SRigVMLogWidget::UnbindLog(const URigVMBlueprint* InRigVMBlueprint)
+{
+	TSharedRef<IMessageLogListing> Listing = FCompilerResultsLog::GetBlueprintMessageLog(const_cast<URigVMBlueprint*>(InRigVMBlueprint));
+	UnbindLog(Listing);
+}
+
+void SRigVMLogWidget::OnBoundListingChanged(TWeakPtr<IMessageLogListing> InWeakListing)
+{
+	if(!InWeakListing.IsValid())
+	{
+		return;
+	}
+
+	const IMessageLogListing* Listing = InWeakListing.Pin().Get();
+
+	FBoundListingInfo& Info = ListingInfo.FindChecked(Listing->GetName());
+
+	const TArray< TSharedRef<class FTokenizedMessage> >& Messages = Listing->GetFilteredMessages();
+	while(Messages.Num() > Info.NumMessages)
+	{
+		ListingModel->AddMessage(Messages[Info.NumMessages++]);
+	}
+}
+
+SRigVMLogWidget::FBoundListingInfo::FBoundListingInfo(const IMessageLogListing& InListing)
+	: NumMessages(InListing.GetFilteredMessages().Num())
+{
+}
+
+#undef LOCTEXT_NAMESPACE
diff --git a/Engine/Plugins/Runtime/RigVM/Source/RigVMEditor/Public/Widgets/SRigVMLogWidget.h b/Engine/Plugins/Runtime/RigVM/Source/RigVMEditor/Public/Widgets/SRigVMLogWidget.h
new file mode 100644
--- /dev/null
+++ b/Engine/Plugins/Runtime/RigVM/Source/RigVMEditor/Public/Widgets/SRigVMLogWidget.h
@@ -0,0 +1,69 @@
+// Copyright Epic Games, Inc. All Rights Reserved.
+
+#pragma once
+
+#include "Widgets/Layout/SBox.h"
+#include "Model/MessageLogListingModel.h"
+#include "Presentation/MessageLogListingViewModel.h"
+#include "RigVMBlueprint.h"
+
+class IMessageLogListing;
+
+class SRigVMLogWidget : public SBox
+{
+public:
+	
+	SLATE_BEGIN_ARGS(SRigVMLogWidget)
+		: _LogName(TEXT("RigVMLog"))
+		, _LogLabel(NSLOCTEXT("SRigVMLogWidget", "RigVMLog", "RigVM Log"))
+	{
+	}
+	SLATE_ATTRIBUTE(FOptionalSize, HeightOverride)
+	SLATE_ARGUMENT(FName, LogName)
+	SLATE_ARGUMENT(FText, LogLabel)
+	SLATE_END_ARGS()
+
+	SRigVMLogWidget();
+	virtual ~SRigVMLogWidget() override;
+	
+	void Construct(const FArguments& InArgs);
+
+	void BindLog(const FName& InLogName);
+	void UnbindLog(const FName& InLogName);
+	void BindLog(TSharedPtr<IMessageLogListing> InListing);
+	void UnbindLog(TSharedPtr<IMessageLogListing> InListing);
+	void BindLog(const URigVMBlueprint* InRigVMBlueprint);
+	void UnbindLog(const URigVMBlueprint* InRigVMBlueprint);
+
+	const TSharedRef<FMessageLogListingModel>& GetListing() const
+	{
+		return ListingModel.ToSharedRef();
+	}
+
+	TSharedRef<FMessageLogListingModel> GetListing()
+	{
+		return ListingModel.ToSharedRef();
+	}
+
+private:
+
+	void OnBoundListingChanged(TWeakPtr<IMessageLogListing> InWeakListing);
+	
+	TSharedPtr<FMessageLogListingModel> ListingModel;
+	TSharedPtr<FMessageLogListingViewModel> ListingView;
+	TArray<TWeakPtr<IMessageLogListing>> BoundListings;
+
+	struct FBoundListingInfo
+	{
+		FBoundListingInfo()
+			: NumMessages(0)
+		{
+		}
+
+		FBoundListingInfo(const IMessageLogListing& InListing);
+		
+		int32 NumMessages;
+	};
+
+	TMap<FName, FBoundListingInfo> ListingInfo;
+};
\ No newline at end of file
