RigVM: Fix build breakage on ARM in SRigVMLogWidget.h

#rnx

diff --git a/Engine/Plugins/Runtime/RigVM/Source/RigVMEditor/Public/Widgets/SRigVMLogWidget.h b/Engine/Plugins/Runtime/RigVM/Source/RigVMEditor/Public/Widgets/SRigVMLogWidget.h
--- a/Engine/Plugins/Runtime/RigVM/Source/RigVMEditor/Public/Widgets/SRigVMLogWidget.h
+++ b/Engine/Plugins/Runtime/RigVM/Source/RigVMEditor/Public/Widgets/SRigVMLogWidget.h
@@ -35,7 +35,7 @@ public:
 	void BindLog(const URigVMBlueprint* InRigVMBlueprint);
 	void UnbindLog(const URigVMBlueprint* InRigVMBlueprint);
 
-	const TSharedRef<FMessageLogListingModel>& GetListing() const
+	TSharedRef<FMessageLogListingModel> GetListing() const
 	{
 		return ListingModel.ToSharedRef();
 	}
