RigVM: Fix undo/redo stack not working
#rnx

diff --git a/Engine/Plugins/Runtime/RigVM/Source/RigVMDeveloper/Private/RigVMModel/RigVMControllerActions.cpp b/Engine/Plugins/Runtime/RigVM/Source/RigVMDeveloper/Private/RigVMModel/RigVMControllerActions.cpp
--- a/Engine/Plugins/Runtime/RigVM/Source/RigVMDeveloper/Private/RigVMModel/RigVMControllerActions.cpp
+++ b/Engine/Plugins/Runtime/RigVM/Source/RigVMDeveloper/Private/RigVMModel/RigVMControllerActions.cpp
@@ -224,7 +224,7 @@ void URigVMActionStack::PostTransacted(const FTransactionObjectEvent& Transactio
 
 bool URigVMActionStack::IsDisabled() const
 {
-	return !HasAnyFlags(RF_ClassDefaultObject);
+	return HasAnyFlags(RF_ClassDefaultObject);
 }
 
 #if RIGVM_ACTIONSTACK_VERBOSE_LOG		
