RigVM: Add UI action to add variant to function
#rnx

diff --git a/Engine/Plugins/Runtime/RigVM/Source/RigVMDeveloper/Public/EdGraph/RigVMEdGraphSchema.h b/Engine/Plugins/Runtime/RigVM/Source/RigVMDeveloper/Public/EdGraph/RigVMEdGraphSchema.h
--- a/Engine/Plugins/Runtime/RigVM/Source/RigVMDeveloper/Public/EdGraph/RigVMEdGraphSchema.h
+++ b/Engine/Plugins/Runtime/RigVM/Source/RigVMDeveloper/Public/EdGraph/RigVMEdGraphSchema.h
@@ -312,6 +312,7 @@ public:
 	virtual bool TryRenameGraph(UEdGraph* GraphToRename, const FName& InNewName) const override;
 	virtual bool TryToGetChildEvents(const UEdGraph* Graph, const int32 SectionId, TArray<TSharedPtr<FEdGraphSchemaAction>>& Actions, const FText& ParentCategory) const override;
 	virtual bool CanDuplicateGraph(UEdGraph* InSourceGraph) const { return false; }
+	virtual bool AllowsFunctionVariants() const override { return true; }
 	virtual UEdGraphPin* DropPinOnNode(UEdGraphNode* InTargetNode, const FName& InSourcePinName, const FEdGraphPinType& InSourcePinType, EEdGraphPinDirection InSourcePinDirection) const override;
 	virtual bool SupportsDropPinOnNode(UEdGraphNode* InTargetNode, const FEdGraphPinType& InSourcePinType, EEdGraphPinDirection InSourcePinDirection, FText& OutErrorMessage) const override;
 	virtual void SetPinBeingDroppedOnNode(UEdGraphPin* InSourcePin) const override { PinBeingDropped = InSourcePin; }
diff --git a/Engine/Plugins/Runtime/RigVM/Source/RigVMEditor/Private/Editor/RigVMEditor.cpp b/Engine/Plugins/Runtime/RigVM/Source/RigVMEditor/Private/Editor/RigVMEditor.cpp
--- a/Engine/Plugins/Runtime/RigVM/Source/RigVMEditor/Private/Editor/RigVMEditor.cpp
+++ b/Engine/Plugins/Runtime/RigVM/Source/RigVMEditor/Private/Editor/RigVMEditor.cpp
@@ -728,6 +728,35 @@ void FRigVMEditor::JumpToHyperlink(const UObject* ObjectReference, bool bRequest
 	FBlueprintEditor::JumpToHyperlink(ObjectReference, bRequestRename);
 }
 
+void FRigVMEditor::AddNewFunctionVariant(const UEdGraph* InOriginalFunction)
+{
+#if WITH_EDITOR
+	if (GEditor)
+	{
+		GEditor->CancelTransaction(0);
+	}
+#endif
+	
+	if (const URigVMEdGraph* RigVMEdGraph = Cast<URigVMEdGraph>(InOriginalFunction))
+	{
+		if (URigVMGraph* RigVMGraph = RigVMEdGraph->GetModel())
+		{
+			if (URigVMFunctionLibrary* FunctionLibrary = Cast<URigVMFunctionLibrary>(RigVMGraph->GetParentGraph()))
+			{
+				URigVMBlueprint* RigVMBlueprint = GetRigVMBlueprint();
+				URigVMController* Controller = RigVMBlueprint->GetController(FunctionLibrary);
+				if (const URigVMLibraryNode* VariantNode = Controller->CreateFunctionVariant(RigVMGraph->GetOuter()->GetFName(), NAME_None, true, true))
+				{
+					if (const UEdGraph* NewGraph = RigVMBlueprint->GetEdGraph(VariantNode->GetContainedGraph()))
+					{
+						OpenDocument(NewGraph, FDocumentTracker::OpenNewDocument);
+					}
+				}
+			}
+		}
+	}
+}
+
 void FRigVMEditor::PostUndo(bool bSuccess)
 {
 	const FTransaction* Transaction = GEditor->Trans->GetTransaction(GEditor->Trans->GetQueueLength() - GEditor->Trans->GetUndoCount());
diff --git a/Engine/Plugins/Runtime/RigVM/Source/RigVMEditor/Public/Editor/RigVMEditor.h b/Engine/Plugins/Runtime/RigVM/Source/RigVMEditor/Public/Editor/RigVMEditor.h
--- a/Engine/Plugins/Runtime/RigVM/Source/RigVMEditor/Public/Editor/RigVMEditor.h
+++ b/Engine/Plugins/Runtime/RigVM/Source/RigVMEditor/Public/Editor/RigVMEditor.h
@@ -93,6 +93,7 @@ public:
 	virtual bool ShouldLoadBPLibrariesFromAssetRegistry() override { return false; }
 	virtual void JumpToHyperlink(const UObject* ObjectReference, bool bRequestRename = false) override;
 	virtual bool ShouldOpenGraphByDefault() const { return true; }
+	virtual void AddNewFunctionVariant(const UEdGraph* InOriginalFunction) override;
 
 	// FEditorUndoClient Interface
 	virtual void PostUndo(bool bSuccess) override;
diff --git a/Engine/Source/Editor/Kismet/Private/SMyBlueprint.cpp b/Engine/Source/Editor/Kismet/Private/SMyBlueprint.cpp
--- a/Engine/Source/Editor/Kismet/Private/SMyBlueprint.cpp
+++ b/Engine/Source/Editor/Kismet/Private/SMyBlueprint.cpp
@@ -2540,6 +2540,20 @@ TSharedPtr<SWidget> SMyBlueprint::OnContextMenuOpening()
 					})
 				);
 			}
+
+			if (const UEdGraphSchema* Schema = Graph->EdGraph->GetSchema())
+			{
+				if (Schema->AllowsFunctionVariants())
+				{
+					MenuBuilder.AddMenuEntry(
+						LOCTEXT("MyBlueprint_Add_Func_Variant", "Add Variant"), FText(), FSlateIcon(),
+						FExecuteAction::CreateLambda([BlueprintEditor, Graph]()
+						{
+							BlueprintEditor->AddNewFunctionVariant(Graph->EdGraph);
+						})
+					);
+				}
+			}
 		}
 		// If this is an event, allow us to convert it to a function graph if possible
 		else if( Event )
diff --git a/Engine/Source/Editor/Kismet/Public/BlueprintEditor.h b/Engine/Source/Editor/Kismet/Public/BlueprintEditor.h
--- a/Engine/Source/Editor/Kismet/Public/BlueprintEditor.h
+++ b/Engine/Source/Editor/Kismet/Public/BlueprintEditor.h
@@ -608,6 +608,9 @@ public:
 	/** Pastes a new local variable to the focused function graph */
 	virtual void OnPasteNewLocalVariable(const FBPVariableDescription& VariableDescription);
 
+	/** Create a variant for a function */
+	virtual void AddNewFunctionVariant(const UEdGraph* InOriginalFunction) {}
+
 	// Type of new document/graph being created by a menu item
 	enum ECreatedDocumentType
 	{
diff --git a/Engine/Source/Runtime/Engine/Classes/EdGraph/EdGraphSchema.h b/Engine/Source/Runtime/Engine/Classes/EdGraph/EdGraphSchema.h
--- a/Engine/Source/Runtime/Engine/Classes/EdGraph/EdGraphSchema.h
+++ b/Engine/Source/Runtime/Engine/Classes/EdGraph/EdGraphSchema.h
@@ -1186,6 +1186,9 @@ class UEdGraphSchema : public UObject
 	/* Duplicate a given graph return the duplicate graph */
 	virtual UEdGraph* DuplicateGraph(UEdGraph* GraphToDuplicate) const { return nullptr;}
 
+	/* Indicates whether variants of functions are allowed */
+	virtual bool AllowsFunctionVariants() const { return false; }
+
 	/* returns new FConnectionDrawingPolicy from this schema */
 	virtual class FConnectionDrawingPolicy* CreateConnectionDrawingPolicy(int32 InBackLayerID, int32 InFrontLayerID, float InZoomFactor, const FSlateRect& InClippingRect, class FSlateWindowElementList& InDrawElements, class UEdGraph* InGraphObj) const { return nullptr; }
 
