RigVM Editor: Introduce RigVMTreeToolkit

#rnx

diff --git a/Engine/Plugins/Runtime/RigVM/Source/RigVMDeveloper/Private/EdGraph/RigVMEdGraphNode.cpp b/Engine/Plugins/Runtime/RigVM/Source/RigVMDeveloper/Private/EdGraph/RigVMEdGraphNode.cpp
--- a/Engine/Plugins/Runtime/RigVM/Source/RigVMDeveloper/Private/EdGraph/RigVMEdGraphNode.cpp
+++ b/Engine/Plugins/Runtime/RigVM/Source/RigVMDeveloper/Private/EdGraph/RigVMEdGraphNode.cpp
@@ -1110,8 +1110,9 @@ void URigVMEdGraphNode::SetModelNode(URigVMNode* InModelNode)
 	ModelNodePath = InModelNode->GetNodePath();
 	CachedModelNode = InModelNode;
 	SyncGraphNodeTitleWithModelNodeTitle();
-	AllocateDefaultPins();
-	PostReconstructNode();
+	//AllocateDefaultPins();
+	ReconstructNode();
+	//PostReconstructNode();
 }
 
 bool URigVMEdGraphNode::ModelPinRemoved_Internal(const URigVMPin* InModelPin)
diff --git a/Engine/Plugins/Runtime/RigVM/Source/RigVMEditor/Private/Widgets/RigVMTreeToolkitContext.cpp b/Engine/Plugins/Runtime/RigVM/Source/RigVMEditor/Private/Widgets/RigVMTreeToolkitContext.cpp
new file mode 100644
--- /dev/null
+++ b/Engine/Plugins/Runtime/RigVM/Source/RigVMEditor/Private/Widgets/RigVMTreeToolkitContext.cpp
@@ -0,0 +1,124 @@
+// Copyright Epic Games, Inc. All Rights Reserved.
+
+#include "Widgets/RigVMTreeToolkitContext.h"
+#include "Widgets/RigVMTreeToolkitNode.h"
+#include "Editor/RigVMEditorTools.h"
+
+#define LOCTEXT_NAMESPACE "RigVMTreeToolkitContext"
+
+uint32 FRigVMTreeContext::GetVisibleChildrenHash() const
+{
+	uint32 Hash = HashOffset;
+	for(const TSharedRef<FRigVMTreeFilter>& Filter : Filters)
+	{
+		Hash = HashCombine(Hash, Filter->GetVisibleChildrenHash());
+	}
+	return Hash;
+}
+
+bool FRigVMTreeContext::FiltersNode(TSharedRef<FRigVMTreeNode> InNode) const
+{
+	const TSharedRef<FRigVMTreeContext> SharedRef = ToSharedRef();
+	if(!InNode->GetVisibleChildren(SharedRef).IsEmpty())
+	{
+		return false;
+	}
+	for(const TSharedRef<FRigVMTreeFilter>& Filter : Filters)
+	{
+		if(Filter->IsEnabled())
+		{
+			if(Filter->Filters(InNode, SharedRef))
+			{
+				return true;
+			}
+		}
+	}
+	return false;
+}
+
+FAssetData FRigVMTreeContext::FindAssetFromAnyPath(const FString& InPartialOrFullPath, bool bConvertToRootPath)
+{
+	return UE::RigVM::Editor::Tools::FindAssetFromAnyPath(InPartialOrFullPath, bConvertToRootPath);
+}
+
+void FRigVMTreeContext::LogMessage(const TSharedRef<FTokenizedMessage>& InMessage) const
+{
+	OnLogTokenizedMessage.Broadcast(InMessage);
+}
+
+void FRigVMTreeContext::LogMessage(const FText& InText) const
+{
+	return LogMessage(FTokenizedMessage::Create(EMessageSeverity::Info, InText));
+}
+
+void FRigVMTreeContext::LogMessage(const FString& InString) const
+{
+	return LogMessage(FText::FromString(InString));
+}
+
+void FRigVMTreeContext::LogWarning(const FText& InText) const
+{
+	return LogMessage(FTokenizedMessage::Create(EMessageSeverity::Warning, InText));
+}
+
+void FRigVMTreeContext::LogWarning(const FString& InString) const
+{
+	return LogWarning(FText::FromString(InString));
+}
+
+void FRigVMTreeContext::LogError(const FText& InText) const
+{
+	return LogMessage(FTokenizedMessage::Create(EMessageSeverity::Error, InText));
+}
+
+void FRigVMTreeContext::LogError(const FString& InString) const
+{
+	return LogError(FText::FromString(InString));
+}
+
+FRigVMTreePhase::FRigVMTreePhase(int32 InID, const FString& InName, const TSharedRef<FRigVMTreeContext>& InContext)
+	: ID(InID)
+	, Name(InName)
+	, bIsActive(false)
+	, bAllowsMultiSelection(false)
+	, Context(InContext)
+	, RootNode(FRigVMTreeRootNode::Create())
+{
+}
+
+void FRigVMTreePhase::IncrementContextHash()
+{
+	Context->HashOffset++;
+}
+
+TArray<TSharedRef<FRigVMTreeNode>> FRigVMTreePhase::GetAllNodes() const
+{
+	return RootNode->GetChildren(Context);
+}
+
+const TArray<TSharedRef<FRigVMTreeNode>>& FRigVMTreePhase::GetVisibleNodes() const
+{
+	return RootNode->GetVisibleChildren(Context);
+}
+
+void FRigVMTreePhase::AddNode(const TSharedRef<FRigVMTreeNode>& InNode)
+{
+	RootNode->AddChild(InNode);
+}
+
+void FRigVMTreePhase::RemoveNode(const TSharedRef<FRigVMTreeNode>& InNode)
+{
+	RootNode->RemoveChild(InNode);
+}
+
+void FRigVMTreePhase::SetNodes(const TArray<TSharedRef<FRigVMTreeNode>>& InNodes)
+{
+	RootNode->SetChildren(InNodes);
+}
+
+TSharedPtr<FRigVMTreeNode> FRigVMTreePhase::FindVisibleNode(const FString& InPath) const
+{
+	return FRigVMTreeNode::FindVisibleNodeInSet(GetVisibleNodes(), InPath, Context);
+}
+
+#undef LOCTEXT_NAMESPACE
diff --git a/Engine/Plugins/Runtime/RigVM/Source/RigVMEditor/Private/Widgets/RigVMTreeToolkitFilter.cpp b/Engine/Plugins/Runtime/RigVM/Source/RigVMEditor/Private/Widgets/RigVMTreeToolkitFilter.cpp
new file mode 100644
--- /dev/null
+++ b/Engine/Plugins/Runtime/RigVM/Source/RigVMEditor/Private/Widgets/RigVMTreeToolkitFilter.cpp
@@ -0,0 +1,54 @@
+// Copyright Epic Games, Inc. All Rights Reserved.
+
+#include "Widgets/RigVMTreeToolkitFilter.h"
+#include "Widgets/RigVMTreeToolkitNode.h"
+
+#define LOCTEXT_NAMESPACE "RigVMTreeToolkitFilter"
+
+bool FRigVMTreeFilter::Filters(TSharedRef<FRigVMTreeNode>& InNode, const TSharedRef<FRigVMTreeContext>& InContext)
+{
+	return false;
+}
+
+uint32 FRigVMTreeFilter::GetVisibleChildrenHash() const
+{
+	return HashCombine(GetTypeHash(Type.ToString()), GetTypeHash(bEnabled));
+}
+
+bool FRigVMTreePathFilter::Filters(TSharedRef<FRigVMTreeNode>& InNode, const TSharedRef<FRigVMTreeContext>& InContext)
+{
+	if(FilterText.IsEmpty())
+	{
+		return false;
+	}
+	return !InNode->GetPath().Contains(FilterText, ESearchCase::IgnoreCase);
+}
+
+uint32 FRigVMTreePathFilter::GetVisibleChildrenHash() const
+{
+	return HashCombine(FRigVMTreeFilter::GetVisibleChildrenHash(), GetTypeHash(FilterText));
+}
+
+FText FRigVMTreeEngineContentFilter::GetLabel() const
+{
+	return LOCTEXT("ShowEngineContent", "Show Engine Content");
+}
+
+bool FRigVMTreeEngineContentFilter::Filters(TSharedRef<FRigVMTreeNode>& InNode, const TSharedRef<FRigVMTreeContext>& InContext)
+{
+	static const FString Prefix = TEXT("/Game/");
+	return !InNode->GetPath().StartsWith(Prefix, ESearchCase::IgnoreCase);
+}
+
+FText FRigVMTreeDeveloperContentFilter::GetLabel() const
+{
+	return LOCTEXT("ShowDeveloperContent", "Show Developer Content");
+}
+
+bool FRigVMTreeDeveloperContentFilter::Filters(TSharedRef<FRigVMTreeNode>& InNode, const TSharedRef<FRigVMTreeContext>& InContext)
+{
+	static const FString Prefix = TEXT("/Game/Developers");
+	return InNode->GetPath().StartsWith(Prefix, ESearchCase::IgnoreCase);
+}
+
+#undef LOCTEXT_NAMESPACE
diff --git a/Engine/Plugins/Runtime/RigVM/Source/RigVMEditor/Private/Widgets/RigVMTreeToolkitNode.cpp b/Engine/Plugins/Runtime/RigVM/Source/RigVMEditor/Private/Widgets/RigVMTreeToolkitNode.cpp
new file mode 100644
--- /dev/null
+++ b/Engine/Plugins/Runtime/RigVM/Source/RigVMEditor/Private/Widgets/RigVMTreeToolkitNode.cpp
@@ -0,0 +1,405 @@
+// Copyright Epic Games, Inc. All Rights Reserved.
+
+#include "Widgets/RigVMTreeToolkitNode.h"
+#include "Widgets/RigVMTreeToolkitContext.h"
+#include "Styling/AppStyle.h"
+#include "Styling/StyleColors.h"
+#include "Styling/SlateIconFinder.h"
+
+FAssetData FRigVMTreeNode::GetAssetData() const
+{
+	return FRigVMTreeContext::FindAssetFromAnyPath(GetPath(), false);
+}
+
+URigVMBlueprint* FRigVMTreeNode::GetBlueprint() const
+{
+	const FAssetData AssetData = FRigVMTreeContext::FindAssetFromAnyPath(GetPath(), true);
+	if(AssetData.IsAssetLoaded())
+	{
+		UObject* Asset = AssetData.GetAsset();
+		return Cast<URigVMBlueprint>(Asset);
+	}
+	return nullptr;
+}
+
+FText FRigVMTreeNode::GetLabel() const
+{
+	FString Label = GetPath();
+	if(Parent)
+	{
+		if(Label.StartsWith(Parent->GetPath(), ESearchCase::IgnoreCase))
+		{
+			Label = Label.Mid(Parent->GetPath().Len() + 1);
+		}
+	}
+	return FText::FromString(Label);
+}
+
+ECheckBoxState FRigVMTreeNode::GetCheckState() const
+{
+	if(CheckState.IsSet())
+	{
+		return CheckState.GetValue();
+	}
+
+	if(Parent)
+	{
+		if(Parent->GetCheckState() != ECheckBoxState::Unchecked)
+		{
+			return ECheckBoxState::Undetermined;
+		}
+	}
+	return ECheckBoxState::Unchecked;
+}
+
+void FRigVMTreeNode::SetCheckState(ECheckBoxState InNewState)
+{
+	CheckState = InNewState;
+}
+
+void FRigVMTreeNode::ResetCheckState()
+{
+	CheckState.Reset();
+}
+
+int32 FRigVMTreeNode::GetDepth() const
+{
+	if(Parent && !Parent->IsA<FRigVMTreeRootNode>())
+	{
+		return Parent->GetDepth() + 1;
+	}
+	return 0;
+}
+
+TSharedRef<FRigVMTreeNode> FRigVMTreeNode::GetRoot() const
+{
+	if(Parent && !Parent->IsA<FRigVMTreeRootNode>())
+	{
+		return Parent->GetRoot(); 
+	}
+	return ToSharedRef();
+}
+
+TSharedPtr<FRigVMTreeNode> FRigVMTreeNode::FindVisibleChild(const FString& InFullPath, const TSharedRef<FRigVMTreeContext>& InContext) const
+{
+	return FindVisibleNodeInSet(GetVisibleChildren(InContext), InFullPath, InContext);
+}
+
+void FRigVMTreeNode::DirtyChildren()
+{
+	Children.Reset();
+	DirtyVisibleChildren();
+}
+
+void FRigVMTreeNode::DirtyVisibleChildren()
+{
+	VisibleChildren.Reset();
+	VisibleChildrenHash = UINT32_MAX;
+	DirtyVisibleParent();
+	RequestRefresh();
+}
+
+void FRigVMTreeNode::DirtyVisibleParent() const
+{
+	if(bIsDirtyingParent || (Parent==nullptr))
+	{
+		return;
+	}
+	const TGuardValue<bool> ReEntryGuard(bIsDirtyingParent, true);
+	GetParent()->DirtyVisibleChildren();
+}
+
+void FRigVMTreeNode::RequestRefresh(bool bForce)
+{
+	(void)RefreshDelegate.ExecuteIfBound(bForce);
+}
+
+const FSlateBrush* FRigVMTreeNode::GetBackgroundImage(bool bIsHovered, bool bIsSelected) const
+{
+	return nullptr;
+}
+
+FSlateColor FRigVMTreeNode::GetBackgroundColor(bool bIsHovered, bool bIsSelected) const
+{
+	return FSlateColor(EStyleColor::Background);
+}
+
+bool FRigVMTreeNode::IsLoaded() const
+{
+	if(Parent)
+	{
+		if(!Parent->IsLoaded())
+		{
+			return false;
+		}
+	}
+	return true;
+}
+
+void FRigVMTreeNode::AddChildImpl(const TSharedRef<FRigVMTreeNode>& InChild) const
+{
+	if(!Children.IsSet())
+	{
+		Children = TArray<TSharedRef<FRigVMTreeNode>>();
+	}
+	Children.GetValue().Add(InChild);
+	InChild->Parent = this;
+}
+
+void FRigVMTreeNode::UpdateChildren(const TSharedRef<FRigVMTreeContext>& InContext) const
+{
+	TArray<TSharedRef<FRigVMTreeNode>> NewChildren = GetChildrenImpl(InContext);
+	if(NewChildren.IsEmpty())
+	{
+		Children.Reset();
+		return;
+	}
+
+	bool bRequiresSort = false;
+	bool bIsIdentical = false;
+	const TArray<TSharedRef<FRigVMTreeNode>> OldChildren = Children.Get(TArray<TSharedRef<FRigVMTreeNode>>());
+	
+	// remove missing children
+	if(Children.IsSet())
+	{
+		bIsIdentical = OldChildren.Num() == NewChildren.Num();
+		for(const TSharedRef<FRigVMTreeNode>& OldChild : OldChildren)
+		{
+			if(!NewChildren.Contains(OldChild))
+			{
+				Children.GetValue().Remove(OldChild);
+				bIsIdentical = false;
+			}
+		}
+	}
+
+	// add new children
+	if(!bIsIdentical)
+	{
+		for(TSharedRef<FRigVMTreeNode> NewChild : NewChildren)
+		{
+			if(!OldChildren.Contains(NewChild))
+			{
+				AddChildImpl(NewChild);
+				bRequiresSort = true;
+			}
+		}
+	}
+
+	if(bRequiresSort && Children.IsSet())
+	{
+		Algo::Sort(Children.GetValue(), [](const TSharedPtr<FRigVMTreeNode>& A, const TSharedPtr<FRigVMTreeNode>& B)
+		{
+			return FCString::Strcmp(*A->GetPath(),*B->GetPath()) < 0;
+		});
+	}
+}
+
+void FRigVMTreeNode::UpdateVisibleChildren(const TSharedRef<FRigVMTreeContext>& InContext) const
+{
+	const uint32 ExpectedHash = InContext->GetVisibleChildrenHash();
+	if(VisibleChildrenHash == ExpectedHash)
+	{
+		return;
+	}
+	VisibleChildrenHash = ExpectedHash;
+	
+	const TArray<TSharedRef<FRigVMTreeNode>> NewChildren = GetChildren(InContext);
+	if(NewChildren.IsEmpty())
+	{
+		VisibleChildren.Reset();
+		return;
+	}
+
+	VisibleChildren = NewChildren.FilterByPredicate([InContext](const TSharedRef<FRigVMTreeNode>& Node)
+	{
+		return !InContext->FiltersNode(Node);
+	});
+
+	Algo::Sort(VisibleChildren, [](const TSharedPtr<FRigVMTreeNode>& A, const TSharedPtr<FRigVMTreeNode>& B)
+	{
+		return FCString::Strcmp(*A->GetPath(),*B->GetPath()) < 0;
+	});
+}
+
+bool FRigVMTreeNode::ContainsAnyVisibleCheckedNode() const
+{
+	// operating on the visible children directly here to avoid pulling
+	for(const TSharedRef<FRigVMTreeNode>& VisibleChild : VisibleChildren)
+	{
+		if(VisibleChild->GetCheckState() != ECheckBoxState::Unchecked)
+		{
+			return true;
+		}
+		if(VisibleChild->ContainsAnyVisibleCheckedNode())
+		{
+			return true;
+		}
+	}
+	return false;
+}
+
+TSharedPtr<FRigVMTreeNode> FRigVMTreeNode::FindVisibleNodeInSet(const TArray<TSharedRef<FRigVMTreeNode>>& InNodes, const FString& InPath, const TSharedRef<FRigVMTreeContext>& InContext)
+{
+	if(!InPath.IsEmpty())
+	{
+		for(const TSharedRef<FRigVMTreeNode>& Node : InNodes)
+		{
+			if(Node->GetPath() == InPath)
+			{
+				return Node->ToSharedPtr();
+			}
+		}
+
+		// find the node with the longest matching start path
+		int32 MatchingPathLength = 0;
+		TSharedPtr<FRigVMTreeNode> MatchingNode;
+
+		for(const TSharedRef<FRigVMTreeNode>& Node : InNodes)
+		{
+			if(InPath.StartsWith(Node->GetPath()))
+			{
+				if(Node->GetPath().Len() > MatchingPathLength)
+				{
+					MatchingPathLength = Node->GetPath().Len();
+					MatchingNode = Node->ToSharedPtr();
+				}
+			}
+		}
+
+		if(MatchingNode)
+		{
+			return MatchingNode->FindVisibleChild(InPath, InContext);
+		}
+	}
+	return nullptr;
+}
+
+FRigVMTreeRootNode::FRigVMTreeRootNode()
+	: FRigVMTreeNode(TEXT("Root"))
+{
+}
+
+void FRigVMTreeRootNode::AddChild(const TSharedRef<FRigVMTreeNode>& InNode)
+{
+	check(!InNode->GetParent().IsValid() || InNode->GetParent().Get() == this);
+	AddChildImpl(InNode);
+	DirtyVisibleChildren();
+}
+
+void FRigVMTreeRootNode::RemoveChild(const TSharedRef<FRigVMTreeNode>& InNode)
+{
+	if(Children.IsSet())
+	{
+		if(Children.GetValue().Remove(InNode) > 0)
+		{
+			DirtyVisibleChildren();
+		}
+	}
+}
+
+void FRigVMTreeRootNode::SetChildren(const TArray<TSharedRef<FRigVMTreeNode>>& InNodes)
+{
+	if(InNodes.IsEmpty())
+	{
+		return;
+	}
+	for(const TSharedRef<FRigVMTreeNode>& Node : InNodes)
+	{
+		check(!Node->GetParent().IsValid() || Node->GetParent().Get() == this);
+		AddChildImpl(Node);
+	}
+	DirtyVisibleChildren();
+}
+
+const FSlateBrush* FRigVMTreeCategoryNode::GetBackgroundImage(bool bIsHovered, bool bIsSelected) const
+{
+	return FAppStyle::Get().GetBrush("DetailsView.CategoryMiddle");
+}
+
+FSlateColor FRigVMTreeCategoryNode::GetBackgroundColor(bool bIsHovered, bool bIsSelected) const
+{
+	if(bIsSelected)
+	{
+		return FSlateColor(EStyleColor::Select); 
+	}
+	if(bIsHovered)
+	{
+		return FSlateColor(EStyleColor::Hover); 
+	}
+	return FAppStyle::Get().GetSlateColor("Colors.Header"); 
+}
+
+FRigVMTreePackageNode::FRigVMTreePackageNode(const FAssetData& InAssetData)
+	: FRigVMTreeCategoryNode(InAssetData.GetObjectPathString())
+{
+	SoftObjectPath = InAssetData.GetSoftObjectPath();
+	bIsLoaded = InAssetData.IsAssetLoaded();
+}
+
+void FRigVMTreePackageNode::Initialize()
+{
+	FCoreUObjectDelegates::OnAssetLoaded.AddSP(this, &FRigVMTreePackageNode::HandleAssetLoaded);
+}
+
+FText FRigVMTreePackageNode::GetLabel() const
+{
+	FText Label = FRigVMTreeCategoryNode::GetLabel();
+	const FString LabelString = Label.ToString();
+	FString Right;
+	if(LabelString.Split(TEXT("."), nullptr, &Right, ESearchCase::IgnoreCase, ESearchDir::FromEnd))
+	{
+		if(LabelString.EndsWith(FString::Printf(TEXT("/%s.%s"), *Right, *Right), ESearchCase::CaseSensitive))
+		{
+			Label = FText::FromString(LabelString.LeftChop(Right.Len() + 1));
+		}
+	}
+	return Label;
+}
+
+bool FRigVMTreePackageNode::IsLoaded() const
+{
+	if(!FRigVMTreeCategoryNode::IsLoaded())
+	{
+		return false;
+	}
+
+	if(!bIsLoaded.IsSet())
+	{
+		const TSoftObjectPtr<UObject> SoftObject(SoftObjectPath);
+		return SoftObject.Get() != nullptr;
+
+	}
+	
+	return bIsLoaded.GetValue();
+}
+
+const FSlateBrush* FRigVMTreePackageNode::GetIconAndTint(FLinearColor& OutColor) const
+{
+	if(!IconBrush.IsSet())
+	{
+		const FAssetData AssetData = GetAssetData();
+		if(AssetData.IsValid())
+		{
+			if(UClass* Class = AssetData.GetClass())
+			{
+				Icon = FSlateIconFinder::FindIconForClass(Class);
+				IconBrush = Icon.GetIcon();
+			}
+		}
+	}
+
+	return IconBrush.Get(nullptr);
+}
+
+void FRigVMTreePackageNode::HandleAssetLoaded(UObject* InAsset)
+{
+	if(InAsset)
+	{
+		if(InAsset->GetPathName() == SoftObjectPath.ToString())
+		{
+			DirtyChildren();
+			bIsLoaded = true;
+		}
+	}
+}
diff --git a/Engine/Plugins/Runtime/RigVM/Source/RigVMEditor/Private/Widgets/RigVMTreeToolkitTask.cpp b/Engine/Plugins/Runtime/RigVM/Source/RigVMEditor/Private/Widgets/RigVMTreeToolkitTask.cpp
new file mode 100644
--- /dev/null
+++ b/Engine/Plugins/Runtime/RigVM/Source/RigVMEditor/Private/Widgets/RigVMTreeToolkitTask.cpp
@@ -0,0 +1,74 @@
+// Copyright Epic Games, Inc. All Rights Reserved.
+
+#include "Widgets/RigVMTreeToolkitTask.h"
+#include "Widgets/RigVMTreeToolkitNode.h"
+#include "AssetRegistry/AssetRegistryModule.h"
+#include "Widgets/RigVMTreeToolkitContext.h"
+#include "Kismet2/KismetEditorUtilities.h"
+
+#define LOCTEXT_NAMESPACE "RigVMTreeToolkitTask"
+
+FRigVMTreeLoadPackageForNodeTask::FRigVMTreeLoadPackageForNodeTask(const TSharedRef<FRigVMTreeNode>& InNode)
+{
+	static const IAssetRegistry& AssetRegistry = FModuleManager::LoadModuleChecked<FAssetRegistryModule>(TEXT("AssetRegistry")).Get();
+	AssetData = AssetRegistry.GetAssetByObjectPath(InNode->GetPath());
+}
+
+bool FRigVMTreeLoadPackageForNodeTask::Execute(const TSharedRef<FRigVMTreePhase>& InPhase)
+{
+	if(AssetData.IsValid())
+	{
+		if(!AssetData.IsAssetLoaded())
+		{
+			const TSharedRef<FTokenizedMessage> Message = FTokenizedMessage::Create(EMessageSeverity::Info);
+			Message->AddText(LOCTEXT("Loading", "Loading"));
+			Message->AddToken(FAssetNameToken::Create(AssetData.PackageName.ToString()));
+			InPhase->GetContext()->LogMessage(Message);
+
+			// load synchronously
+			if(AssetData.GetAsset() == nullptr)
+			{
+				InPhase->GetContext()->LogError(FString::Printf(TEXT("Asset '%s' cannot be loaded."), *AssetData.GetObjectPathString()));
+				return false;
+			}
+		}
+		return true;
+	}
+	InPhase->GetContext()->LogError(TEXT("Provided AssetData is not valid."));
+	return false;
+}
+
+FRigVMCompileBlueprintTask::FRigVMCompileBlueprintTask(const TSharedRef<FRigVMTreeNode>& InNode)
+{
+	ObjectPath = InNode->GetPath();
+}
+
+bool FRigVMCompileBlueprintTask::Execute(const TSharedRef<FRigVMTreePhase>& InPhase)
+{
+	if(const UObject* Object = ObjectPath.TryLoad())
+	{
+		if(UBlueprint* Blueprint = Object->GetTypedOuter<UBlueprint>())
+		{
+			const TSharedRef<FTokenizedMessage> Message = FTokenizedMessage::Create(EMessageSeverity::Info);
+			Message->AddText(LOCTEXT("Compiling", "Compiling"));
+			Message->AddToken(FAssetNameToken::Create(Blueprint->GetOutermost()->GetPathName()));
+			InPhase->GetContext()->LogMessage(Message);
+
+			FCompilerResultsLog CompilerResults;
+			CompilerResults.SetSourcePath(Blueprint->GetPathName());
+			CompilerResults.bSilentMode = true;
+			CompilerResults.BeginEvent(TEXT("Compile"));
+			FKismetEditorUtilities::CompileBlueprint(Blueprint, EBlueprintCompileOptions::SkipGarbageCollection | EBlueprintCompileOptions::SkipSave, &CompilerResults);
+			CompilerResults.EndEvent();
+
+			for(const TSharedRef<FTokenizedMessage>& CompilerMessage : CompilerResults.Messages)
+			{
+				InPhase->GetContext()->LogMessage(CompilerMessage);
+			}
+			return true;
+		}
+	}
+	return false;
+}
+
+#undef LOCTEXT_NAMESPACE
diff --git a/Engine/Plugins/Runtime/RigVM/Source/RigVMEditor/Private/Widgets/SRigVMChangesTreeView.cpp b/Engine/Plugins/Runtime/RigVM/Source/RigVMEditor/Private/Widgets/SRigVMChangesTreeView.cpp
new file mode 100644
--- /dev/null
+++ b/Engine/Plugins/Runtime/RigVM/Source/RigVMEditor/Private/Widgets/SRigVMChangesTreeView.cpp
@@ -0,0 +1,688 @@
+// Copyright Epic Games, Inc. All Rights Reserved.
+
+#include "Widgets/SRigVMChangesTreeView.h"
+#include "Widgets/Input/SCheckBox.h"
+#include "Widgets/Input/SButton.h"
+#include "Widgets/Layout/SSpacer.h"
+#include "Styling/AppStyle.h"
+#include "Misc/UObjectToken.h"
+#include "Layout/WidgetPath.h"
+#include "Framework/MultiBox//MultiBoxBuilder.h"
+#include "Framework/Commands/UICommandInfo.h"
+#include "Framework/Application/MenuStack.h"
+#include "Framework/Application/SlateApplication.h"
+#include "ContentBrowserModule.h"
+#include "IContentBrowserSingleton.h"
+#define LOCTEXT_NAMESPACE "SRigVMChangesTreeView"
+
+SRigVMChangesTreeRow::~SRigVMChangesTreeRow()
+{
+	if(Node)
+	{
+		Node->RefreshDelegate.Unbind();
+	}
+}
+
+void SRigVMChangesTreeRow::Construct(const FArguments& InArgs, const TSharedRef<STableViewBase>& InOwnerTableView)
+{
+	Node = InArgs._Node;
+	OwningWidget = InArgs._OwningWidget.Get();
+
+	Node->RefreshDelegate.BindRaw(this, &SRigVMChangesTreeRow::RequestRefresh);
+	
+	this->ChildSlot
+	[
+		SNew(SBorder)
+		.BorderImage(this, &SRigVMChangesTreeRow::GetBackgroundImage)
+		.BorderBackgroundColor(this, &SRigVMChangesTreeRow::GetBackgroundColor)
+		.Padding(3.0f, 2.f)
+		[
+			SNew(SHorizontalBox)
+			+ SHorizontalBox::Slot()
+			.HAlign(HAlign_Left)
+			.VAlign(VAlign_Fill)
+			.AutoWidth()
+			.Padding(0.0f, 0.0f, 2.0f, 0.0f)
+			[
+				SNew(SBox)
+				.WidthOverride(this, &SRigVMChangesTreeRow::GetIndentWidth)
+			]
+			
+			+ SHorizontalBox::Slot()
+			.HAlign(HAlign_Left)
+			.VAlign(VAlign_Center)
+			.Padding(0)
+			.AutoWidth()
+			[
+				SNew( SImage )
+				.Image(this, &SRigVMChangesTreeRow::GetExpanderImage) 
+				.Visibility( this, &SRigVMChangesTreeRow::GetExpanderVisibility )
+				.OnMouseButtonDown(this, &SRigVMChangesTreeRow::OnExpanderMouseButtonDown)
+			]
+
+			+ SHorizontalBox::Slot()
+			.HAlign(HAlign_Left)
+			.VAlign(VAlign_Center)
+			.Padding(2.f, 0.f, 3.f, 0.f)
+			.AutoWidth()
+			[
+				SNew( SImage )
+				.Image(this, &SRigVMChangesTreeRow::GetIcon)
+				.ColorAndOpacity(this, &SRigVMChangesTreeRow::GetIconColor)
+				.DesiredSizeOverride(FVector2D(16.f, 16.f))
+			]
+
+			+ SHorizontalBox::Slot()
+			.AutoWidth()
+			.Padding(0.0f, 4.0f)
+			.HAlign(HAlign_Left)
+			.VAlign(VAlign_Center)
+			[
+				SNew(STextBlock)
+				.Text(Node->GetLabel())
+				.ColorAndOpacity(this, &SRigVMChangesTreeRow::GetTextColor)
+			]
+			+ SHorizontalBox::Slot()
+			.FillWidth(1)
+			[
+				SNew(SSpacer)
+				.Size(FVector2D(10000,0))
+			]
+			+ SHorizontalBox::Slot()
+			.Padding(0.0f, 4.0f)
+			.AutoWidth()
+			[
+				SNew(SCheckBox)
+				.Visibility(this, &SRigVMChangesTreeRow::GetCheckBoxVisibility)
+				.IsChecked(this, &SRigVMChangesTreeRow::GetCheckBoxState)
+				.OnCheckStateChanged(this, &SRigVMChangesTreeRow::OnCheckBoxStateChanged)
+			]
+		]
+	];
+	
+	STableRow< TSharedRef<FRigVMTreeNode> >::ConstructInternal(
+		STableRow<TSharedRef<FRigVMTreeNode>>::FArguments(), InOwnerTableView);
+}
+
+const FSlateBrush* SRigVMChangesTreeRow::GetBackgroundImage() const
+{
+	return Node->GetBackgroundImage(IsHovered(), IsSelected());
+}
+
+FSlateColor SRigVMChangesTreeRow::GetBackgroundColor() const
+{
+	return Node->GetBackgroundColor(IsHovered(), IsSelected());
+}
+
+FOptionalSize SRigVMChangesTreeRow::GetIndentWidth() const
+{
+	return float(Node->GetDepth()) * 16.0f;
+}
+
+FSlateColor SRigVMChangesTreeRow::GetTextColor() const
+{
+	if(!Node->IsLoaded())
+	{
+		return FLinearColor(0.3f, 0.3f, 0.3f, 1.f);
+	}
+	return FSlateColor::UseForeground();
+}
+
+const FSlateBrush* SRigVMChangesTreeRow::GetIcon() const
+{
+	FLinearColor Color = FLinearColor::White;
+	return Node->GetIconAndTint(Color);
+}
+
+FSlateColor SRigVMChangesTreeRow::GetIconColor() const
+{
+	FLinearColor Color = FLinearColor::White;
+	(void)Node->GetIconAndTint(Color);
+	if(!Node->IsLoaded())
+	{
+		Color = Color * FLinearColor(0.1f, 0.1f, 0.1f, 1.f);
+	}
+	return Color;
+}
+
+const FSlateBrush* SRigVMChangesTreeRow::GetExpanderImage() const
+{
+	if(IsItemExpanded())
+	{
+		if(IsHovered())
+		{
+			return FAppStyle::GetBrush("TreeArrow_Expanded_Hovered");
+		}
+		return FAppStyle::GetBrush("TreeArrow_Expanded");
+	}
+	if(IsHovered())
+	{
+		return FAppStyle::GetBrush("TreeArrow_Collapsed_Hovered");
+	}
+	return FAppStyle::GetBrush("TreeArrow_Collapsed");
+}
+
+FReply SRigVMChangesTreeRow::OnExpanderMouseButtonDown(const FGeometry& SenderGeometry, const FPointerEvent& MouseEvent)
+{
+	if(MouseEvent.IsMouseButtonDown(EKeys::LeftMouseButton))
+	{
+		const TSharedRef< ITypedTableView< TSharedRef< FRigVMTreeNode> > > OwnerTable = OwnerTablePtr.Pin().ToSharedRef();
+		OwnerTable->Private_SetItemExpansion(Node.ToSharedRef(), !IsItemExpanded());
+		return FReply::Handled();
+	}
+	return FReply::Unhandled();
+}
+
+EVisibility SRigVMChangesTreeRow::GetExpanderVisibility() const
+{
+	return Node->HasVisibleChildren() ? EVisibility::Visible : EVisibility::Collapsed;
+}
+
+void SRigVMChangesTreeRow::RequestRefresh(bool bForce)
+{
+	if(OwningWidget)
+	{
+		OwningWidget->RequestRefresh_AnyThread(bForce);
+	}
+}
+
+EVisibility SRigVMChangesTreeRow::GetCheckBoxVisibility() const
+{
+	return Node->IsCheckable() ? EVisibility::Visible : EVisibility::Collapsed;
+}
+
+ECheckBoxState SRigVMChangesTreeRow::GetCheckBoxState() const
+{
+	return Node->GetCheckState();
+}
+
+void SRigVMChangesTreeRow::OnCheckBoxStateChanged(ECheckBoxState InNewState)
+{
+	const TArray<TSharedRef<FRigVMTreeNode>>& SelectedNodes = OwningWidget->GetSelectedNodes();
+	if(SelectedNodes.Contains(Node))
+	{
+		for(TSharedRef<FRigVMTreeNode> SelectedNode : SelectedNodes)
+		{
+			SelectedNode->SetCheckState(InNewState);
+		}
+	}
+	else
+	{
+		Node->SetCheckState(InNewState);
+	}
+}
+
+//////////////////////////////////////////////////////////////////////////
+// SRigVMChangesTreeView
+
+SRigVMChangesTreeView::~SRigVMChangesTreeView()
+{
+}
+
+void SRigVMChangesTreeView::Construct(const FArguments& InArgs)
+{
+	RequestRefreshCount = 0;
+	RequestRefreshForceCount = 0;
+	PhaseAttribute = InArgs._Phase;
+
+	OnNodeSelected = InArgs._OnNodeSelected;
+	OnNodeDoubleClicked = InArgs._OnNodeDoubleClicked;
+
+	TSharedRef<SVerticalBox> MainVerticalBox = SNew(SVerticalBox);
+
+	TSharedRef<SHorizontalBox> PathFilterHorizontalBox = SNew(SHorizontalBox);
+	PathFilterHorizontalBox->AddSlot()
+	.FillWidth(1)
+	.HAlign(HAlign_Fill)
+	[
+		SAssignNew(PathFilterBox, SSearchBox)
+		.Visibility(this, &SRigVMChangesTreeView::GetPathFilterVisibility)
+		.SelectAllTextWhenFocused(true)
+		.InitialText(GetPathFilterText())
+		.OnTextChanged(this, &SRigVMChangesTreeView::OnPathFilterTextChanged)
+		.OnTextCommitted(this, &SRigVMChangesTreeView::OnPathFilterTextCommitted)
+	];
+
+	PathFilterHorizontalBox->AddSlot()
+	.AutoWidth()
+	[
+		SNew(SButton)
+		.ContentPadding(FMargin(1, 1))
+		[
+			SNew(SImage)
+			.Image(FAppStyle::Get().GetBrush("Icons.Refresh"))
+			.ToolTipText(LOCTEXT("RefreshToolTip", "Refresh the contents of the view"))
+		]
+		.OnClicked_Lambda([this]()
+		{
+			RequestRefresh_AnyThread(true);
+			return FReply::Handled();
+		})
+	];
+
+	PathFilterHorizontalBox->AddSlot()
+	.AutoWidth()
+	[
+		SNew(SButton)
+		.Visibility(this, &SRigVMChangesTreeView::GetSettingsButtonVisibility)
+		.ContentPadding(FMargin(1, 1))
+		[
+			SNew(SImage)
+			.Image(FAppStyle::Get().GetBrush("Icons.Settings"))
+			.ToolTipText(LOCTEXT("SettingsToolTip", "Change filtering settings here"))
+		]
+		.OnClicked(this, &SRigVMChangesTreeView::OnSettingsButtonClicked)
+	];
+
+	TSharedRef<SHorizontalBox> MainHorizontalBox = SNew(SHorizontalBox);
+	
+	MainHorizontalBox->AddSlot()
+	.FillWidth(1)
+	.HAlign(HAlign_Fill)
+	[
+		SAssignNew(TreeView, STreeView<TSharedRef<FRigVMTreeNode>>)
+		.TreeItemsSource(&FilteredNodes)
+		.OnGenerateRow(this, &SRigVMChangesTreeView::MakeTreeRowWidget)
+		.OnGetChildren(this, &SRigVMChangesTreeView::GetChildrenForNode)
+		.ItemHeight(22.0f)
+		.OnSelectionChanged(this, &SRigVMChangesTreeView::OnSelectionChanged)
+		.OnMouseButtonDoubleClick(this, &SRigVMChangesTreeView::OnTreeElementDoubleClicked)
+		.SelectionMode(this, &SRigVMChangesTreeView::GetSelectionMode)
+		.OnContextMenuOpening(this, &SRigVMChangesTreeView::OnGetNodeContextMenuContent)
+	];
+
+	MainVerticalBox->AddSlot()
+	.AutoHeight()
+	.Padding(0, 2)
+	[
+		PathFilterHorizontalBox
+	];
+
+	MainVerticalBox->AddSlot()
+	.VAlign(VAlign_Fill)
+	.FillHeight(1.0f)		// This is required to make the scrollbar work, as content overflows Slate containers by default
+	[
+		MainHorizontalBox
+	];
+
+	ChildSlot
+	[
+		MainVerticalBox
+	];
+}
+
+void SRigVMChangesTreeView::Tick(const FGeometry& AllottedGeometry, const double InCurrentTime, const float InDeltaTime)
+{
+	SCompoundWidget::Tick(AllottedGeometry, InCurrentTime, InDeltaTime);
+	RefreshFilteredNodesIfRequired();
+}
+
+TSharedRef<ITableRow> SRigVMChangesTreeView::MakeTreeRowWidget(TSharedRef<FRigVMTreeNode> InNode, const TSharedRef<STableViewBase>& OwnerTable)
+{
+	return SNew(SRigVMChangesTreeRow, OwnerTable)
+		.Node(InNode)
+		.OwningWidget(SharedThis(this));
+}
+
+void SRigVMChangesTreeView::GetChildrenForNode(TSharedRef<FRigVMTreeNode> InNode, TArray< TSharedRef<FRigVMTreeNode> >& OutChildren)
+{
+	OutChildren = InNode->GetVisibleChildren(GetContext());
+}
+
+void SRigVMChangesTreeView::OnSelectionChanged(TSharedPtr<FRigVMTreeNode> Selection, ESelectInfo::Type SelectInfo)
+{
+	if(Selection.IsValid())
+	{
+		if(OnNodeSelected.IsBound())
+		{
+			if(OnNodeSelected.Execute(Selection.ToSharedRef()).IsEventHandled())
+			{
+				return;
+			}
+		}
+	}
+}
+
+void SRigVMChangesTreeView::OnTreeElementDoubleClicked(TSharedRef<FRigVMTreeNode> InNode)
+{
+	if(OnNodeDoubleClicked.IsBound())
+	{
+		if(OnNodeDoubleClicked.Execute(InNode).IsEventHandled())
+		{
+			return;
+		}
+	}
+
+	if(const TSharedPtr<FRigVMTreePackageNode> PackageNode = Cast<FRigVMTreePackageNode>(InNode))
+	{
+		const FAssetData AssetData = PackageNode->GetAssetData();
+		if(AssetData.IsValid())
+		{
+			const FContentBrowserModule& ContentBrowserModule = FModuleManager::Get().LoadModuleChecked<FContentBrowserModule>(TEXT("ContentBrowser"));
+			TArray<FAssetData> Assets;
+			Assets.Add(AssetData);
+			ContentBrowserModule.Get().SyncBrowserToAssets(Assets);
+		}			
+	}
+}
+
+ESelectionMode::Type SRigVMChangesTreeView::GetSelectionMode() const
+{
+	return GetPhase()->AllowsMultiSelection() ? ESelectionMode::Multi : ESelectionMode::Single;
+}
+
+TSharedPtr<SWidget> SRigVMChangesTreeView::OnGetNodeContextMenuContent()
+{
+	const bool bShouldCloseWindowAfterMenuSelection = true;
+	FMenuBuilder MenuBuilder( bShouldCloseWindowAfterMenuSelection, TSharedPtr<const FUICommandList>());
+
+	TArray< TSharedRef< FRigVMTreeNode > > SelectedNodes = TreeView->GetSelectedItems();
+	TArray< TSharedRef< FRigVMTreeNode > > RootNodes;
+
+	for(const TSharedRef< FRigVMTreeNode >& Node : SelectedNodes)
+	{
+		if(!Node->IsLoaded())
+		{
+			RootNodes.AddUnique(Node->GetRoot());
+		}
+	}
+	
+	MenuBuilder.BeginSection("Assets", LOCTEXT( "Assets", "Assets" ) );
+	for(const TSharedRef< FRigVMTreeNode >& Node : RootNodes)
+	{
+		if(!Node->IsLoaded())
+		{
+			MenuBuilder.AddMenuEntry(
+				LOCTEXT("LoadAssets", "Load Assets"),
+				LOCTEXT("LoadAssets_ToolTip", "Loads all of the assets backing up the selection"),
+				FSlateIcon(),
+				FUIAction(FExecuteAction::CreateLambda([this, RootNodes]()
+				{
+					TArray<TSharedRef<FRigVMTreeTask>> Tasks;
+					for(const TSharedRef<FRigVMTreeNode>& Node : RootNodes)
+					{
+						Tasks.Add(FRigVMTreeLoadPackageForNodeTask::Create(Node));
+					}
+					GetPhase()->QueueTasks(Tasks);
+				})));
+			break;
+		}
+	}
+	for(const TSharedRef< FRigVMTreeNode >& Node : RootNodes)
+	{
+		if(const TSharedPtr<FRigVMTreePackageNode> PackageNode = Cast<FRigVMTreePackageNode>(Node))
+		{
+			MenuBuilder.AddMenuEntry(
+				LOCTEXT("ShowInContentBrowser", "Locate in Content Browser"),
+				LOCTEXT("ShowInContentBrowser_ToolTip", "Locates the asset in the content browser"),
+				FSlateIcon(),
+				FUIAction(FExecuteAction::CreateLambda([PackageNode]()
+				{
+					const FAssetData AssetData = PackageNode->GetAssetData();
+					if(AssetData.IsValid())
+					{
+						const FContentBrowserModule& ContentBrowserModule = FModuleManager::Get().LoadModuleChecked<FContentBrowserModule>(TEXT("ContentBrowser"));
+						TArray<FAssetData> Assets;
+						Assets.Add(AssetData);
+						ContentBrowserModule.Get().SyncBrowserToAssets(Assets);
+					}
+				})));
+			break;
+		}
+	}
+	MenuBuilder.EndSection();
+
+	for(const TSharedRef<FRigVMTreeNode>& SelectedNode : SelectedNodes)
+	{
+		if(SelectedNode->IsCheckable())
+		{
+			MenuBuilder.BeginSection("Marking", LOCTEXT( "Marking", "Marking" ) );
+			MenuBuilder.AddMenuEntry(
+				LOCTEXT("MarkAllSelected", "Mark Selection"),
+				LOCTEXT("MarkAllSelected_ToolTip", "Marks all checkboes for the selection"),
+				FSlateIcon(),
+				FUIAction(FExecuteAction::CreateLambda([this]()
+				{
+					TArray< TSharedRef< FRigVMTreeNode > > SelectedNodes = TreeView->GetSelectedItems();
+					for(const TSharedRef<FRigVMTreeNode>& SelectedNode : SelectedNodes)
+					{
+						SelectedNode->SetCheckState(ECheckBoxState::Checked);
+					}
+				})));
+			MenuBuilder.AddMenuEntry(
+				LOCTEXT("UnmarkAllSelected", "Unmark Selection"),
+				LOCTEXT("UnmarkAllSelected_ToolTip", "Unmarks all checkboes for the selection"),
+				FSlateIcon(),
+				FUIAction(FExecuteAction::CreateLambda([this]()
+				{
+					TArray< TSharedRef< FRigVMTreeNode > > SelectedNodes = TreeView->GetSelectedItems();
+					for(const TSharedRef<FRigVMTreeNode>& SelectedNode : SelectedNodes)
+					{
+						SelectedNode->ResetCheckState();
+					}
+				})));
+			MenuBuilder.EndSection();
+			break;
+		}
+	}
+
+	if(!SelectedNodes.IsEmpty())
+	{
+		SelectedNodes[0]->GetContextMenu(MenuBuilder);
+	}
+
+	return MenuBuilder.MakeWidget();
+}
+
+EVisibility SRigVMChangesTreeView::GetSettingsButtonVisibility() const
+{
+	const bool bHasFilterSettings = GetContext()->Filters.ContainsByPredicate([](const TSharedPtr<FRigVMTreeFilter>& Filter) -> bool
+	{
+		return Filter->CanBeToggledInUI();
+	});
+	return bHasFilterSettings ? EVisibility::Visible : EVisibility::Collapsed;
+}
+
+FReply SRigVMChangesTreeView::OnSettingsButtonClicked()
+{
+	FMenuBuilder MenuBuilder(true, nullptr);
+
+	for(const TSharedRef<FRigVMTreeFilter>& Filter : GetContext()->Filters)
+	{
+		if(!Filter->CanBeToggledInUI())
+		{
+			 continue;
+		}
+		
+		FUIAction ToggleSettingAction(FExecuteAction::CreateLambda([this, Filter]()
+		{
+			Filter->SetEnabled(!Filter->IsEnabled());
+			RefreshFilteredNodes();
+		}),
+		FCanExecuteAction(),
+		FIsActionChecked::CreateLambda([Filter]()
+		{
+			return Filter->IsInvertedInUI() ? !Filter->IsEnabled() : Filter->IsEnabled();
+		}));
+
+		MenuBuilder.AddMenuEntry(
+			Filter->GetLabel(),
+			Filter->GetToolTip(),
+			FSlateIcon(),
+			ToggleSettingAction,
+			NAME_None,
+			EUserInterfaceActionType::Check);
+	}
+
+	MenuBuilder.MakeWidget();
+
+	FSlateApplication::Get().PushMenu(SharedThis(this), FWidgetPath(), MenuBuilder.MakeWidget(), FSlateApplication::Get().GetCursorPos(), FPopupTransitionEffect(FPopupTransitionEffect::ContextMenu));
+	
+	return FReply::Handled();
+}
+
+EVisibility SRigVMChangesTreeView::GetPathFilterVisibility() const
+{
+	return GetPathFilter().IsValid() ? EVisibility::Visible : EVisibility::Collapsed;
+}
+
+void SRigVMChangesTreeView::OnPathFilterTextChanged(const FText& SearchText)
+{
+	const TSharedPtr<FRigVMTreePathFilter> PathFilter = GetPathFilter();
+	if(PathFilter.IsValid())
+	{
+		if (GetPathFilterText().CompareToCaseIgnored(SearchText) != 0)
+		{
+			PathFilter->SetFilterText(SearchText.ToString());
+			RefreshFilteredNodes();
+		}
+	}
+}
+
+void SRigVMChangesTreeView::OnPathFilterTextCommitted(const FText& SearchText, ETextCommit::Type CommitInfo)
+{
+	OnPathFilterTextChanged(SearchText);
+}
+
+FText SRigVMChangesTreeView::GetPathFilterText()
+{
+	TSharedPtr<FRigVMTreePathFilter> PathFilter = GetPathFilter();
+	if(PathFilter.IsValid())
+	{
+		return FText::FromString(PathFilter->GetFilterText());
+	}
+	static const FText EmptyText;
+	return EmptyText;
+}
+
+void SRigVMChangesTreeView::RequestRefresh_AnyThread(bool bForce)
+{
+	(void)RequestRefreshCount++;
+	if(bForce)
+	{
+		(void)RequestRefreshForceCount++;
+	}
+}
+
+void SRigVMChangesTreeView::RefreshFilteredNodes(bool bForce)
+{
+	const TSharedRef<FRigVMTreePhase> ActivePhase = GetPhase();
+	if(bForce)
+	{
+		ActivePhase->IncrementContextHash();
+	}
+	FilteredNodes = ActivePhase->GetVisibleNodes();
+	TreeView->RebuildList();
+}
+
+void SRigVMChangesTreeView::RefreshFilteredNodesIfRequired()
+{
+	if(RequestRefreshCount.exchange(0) > 0)
+	{
+		const bool bForce = RequestRefreshForceCount.exchange(0) > 0;
+		RefreshFilteredNodes(bForce);
+	}
+}
+
+void SRigVMChangesTreeView::OnPhaseChanged()
+{
+	RefreshFilteredNodes(true);
+
+	if(const TSharedPtr<FRigVMTreePathFilter> PathFilter = GetPathFilter())
+	{
+		PathFilterBox->SetText(FText::FromString(PathFilter->GetFilterText()));
+	}
+
+	for(const TSharedRef<FRigVMTreeNode>& RootNode : FilteredNodes)
+	{
+		if(RootNode->ShouldExpandByDefault())
+		{
+			TreeView->SetItemExpansion(RootNode, true);
+		}
+	}
+}
+
+void SRigVMChangesTreeView::SetSelection(const TSharedPtr<FRigVMTreeNode>& InNode, bool bRequestScrollIntoView)
+{
+	if(InNode.IsValid())
+	{
+		TreeView->SetSelection(InNode.ToSharedRef());
+		if(bRequestScrollIntoView)
+		{
+			TreeView->RequestScrollIntoView(InNode.ToSharedRef());
+		}
+	}
+	else
+	{
+		TreeView->ClearSelection();
+	}
+}
+
+TSharedRef<FRigVMTreePhase> SRigVMChangesTreeView::GetPhase() const
+{
+	if(const TSharedPtr<FRigVMTreePhase> Phase = PhaseAttribute.Get())
+	{
+		return Phase.ToSharedRef();
+	}
+	static const TSharedRef<FRigVMTreeContext> EmptyContext = FRigVMTreeContext::Create();
+	static const TSharedRef<FRigVMTreePhase> EmptyPhase = FRigVMTreePhase::Create(INDEX_NONE, TEXT("Default"), EmptyContext);
+	return EmptyPhase;
+}
+
+TSharedRef<FRigVMTreeContext> SRigVMChangesTreeView::GetContext() const
+{
+	return GetPhase()->GetContext();
+}
+
+TSharedPtr<FRigVMTreePathFilter> SRigVMChangesTreeView::GetPathFilter() const
+{
+	const TSharedRef<FRigVMTreeContext> Context = GetContext();
+	for(const TSharedRef<FRigVMTreeFilter>& Filter : Context->Filters)
+	{
+		if(Filter->IsA<FRigVMTreePathFilter>())
+		{
+			return Cast<FRigVMTreePathFilter>(Filter);
+		}
+	}
+	return nullptr;
+}
+
+TArray<TSharedRef<FRigVMTreeNode>> SRigVMChangesTreeView::GetSelectedNodes() const
+{
+	return TreeView->GetSelectedItems();
+}
+
+bool SRigVMChangesTreeView::HasAnyVisibleCheckedNode() const
+{
+	TArray<TSharedRef<FRigVMTreeNode>> Nodes = GetPhase()->GetVisibleNodes();
+	for(const TSharedRef<FRigVMTreeNode>& Node : Nodes)
+	{
+		if(Node->GetCheckState() != ECheckBoxState::Unchecked)
+		{
+			return true;
+		}
+		if(Node->ContainsAnyVisibleCheckedNode())
+		{
+			return true;
+		}
+	}
+	return false;
+}
+
+TArray<TSharedRef<FRigVMTreeNode>> SRigVMChangesTreeView::GetCheckedNodes() const
+{
+	TArray<TSharedRef<FRigVMTreeNode>> Nodes = GetPhase()->GetVisibleNodes();
+	for(int32 Index = 0; Index < Nodes.Num(); Index++)
+	{
+		Nodes.Append(Nodes[Index]->GetVisibleChildren(GetContext()));
+	}
+
+	Nodes = Nodes.FilterByPredicate([](const TSharedRef<FRigVMTreeNode>& Node)
+	{
+		return Node->GetCheckState() != ECheckBoxState::Unchecked;
+	});
+
+	return Nodes;
+}
+
+#undef LOCTEXT_NAMESPACE
diff --git a/Engine/Plugins/Runtime/RigVM/Source/RigVMEditor/Public/Widgets/RigVMTreeToolkitContext.h b/Engine/Plugins/Runtime/RigVM/Source/RigVMEditor/Public/Widgets/RigVMTreeToolkitContext.h
new file mode 100644
--- /dev/null
+++ b/Engine/Plugins/Runtime/RigVM/Source/RigVMEditor/Public/Widgets/RigVMTreeToolkitContext.h
@@ -0,0 +1,186 @@
+// Copyright Epic Games, Inc. All Rights Reserved.
+
+#pragma once
+
+#include "RigVMTreeToolkitDefines.h"
+#include "RigVMTreeToolkitFilter.h"
+#include "Input/Reply.h"
+#include "Logging/TokenizedMessage.h"
+
+class FRigVMTreeRootNode;
+class FRigVMTreeTask;
+
+/**
+ * The context is the top level object passed to anything
+ * that has to interact with the tree, like visible node traversal,
+ * task execution etc.
+ */
+class FRigVMTreeContext : public FRigVMTreeElement
+{
+public:
+	DEFINE_RIGVMTREETOOLKIT_ELEMENT(FRigVMTreeContext, FRigVMTreeElement)
+
+	virtual uint32 GetVisibleChildrenHash() const;
+	bool FiltersNode(TSharedRef<FRigVMTreeNode> InNode) const;
+
+	TArray<TSharedRef<FRigVMTreeFilter>> Filters;
+
+	static FAssetData FindAssetFromAnyPath(const FString& InPartialOrFullPath, bool bUseRootPath);
+
+	DECLARE_EVENT_OneParam(FRigVMTreeContext, FLogTokenizedMessage, const TSharedRef<FTokenizedMessage>&)
+	void LogMessage(const TSharedRef<FTokenizedMessage>& InMessage) const;
+	void LogMessage(const FText& InText) const;
+	void LogMessage(const FString& InString) const;
+	void LogWarning(const FText& InText) const;
+	void LogWarning(const FString& InString) const;
+	void LogError(const FText& InText) const;
+	void LogError(const FString& InString) const;
+
+private:
+	uint32 HashOffset = 0;
+	FLogTokenizedMessage OnLogTokenizedMessage;
+
+	friend class FRigVMTreePhase;
+	friend class SRigVMBulkEditWidget;
+};
+
+/**
+ * The Phase describes a phase of performing a UI wizard,
+ * like picking a set of inputs, or performing tasks.
+ * This can be also seen as the steps a wizard goes through
+ * when performing a UI process.
+ */
+class FRigVMTreePhase : public FRigVMTreeElement
+{
+public:
+	DEFINE_RIGVMTREETOOLKIT_ELEMENT(FRigVMTreePhase, FRigVMTreeElement)
+
+	DECLARE_DELEGATE_RetVal(FReply, FReplyDelegate)
+	DECLARE_DELEGATE_OneParam(FQueueTasksDelegate, const TArray<TSharedRef<FRigVMTreeTask>>&)
+
+	FRigVMTreePhase(int32 InID, const FString& InName, const TSharedRef<FRigVMTreeContext>& InContext);
+
+	int32 GetID() const
+	{
+		return ID;
+	}
+
+	const FString& GetName() const
+	{
+		return Name;
+	}
+
+	bool IsActive() const 
+	{
+		return bIsActive;
+	}
+
+	virtual bool AllowsMultiSelection() const
+	{
+		return bAllowsMultiSelection;
+	}
+
+	void SetAllowsMultiSelection(bool InAllowsMultiSelection)
+	{
+		bAllowsMultiSelection = InAllowsMultiSelection;
+	}
+
+	TAttribute<bool>& IsCancelButtonVisible()
+	{
+		return IsCancelButtonVisibleAttribute;
+	}
+
+	TAttribute<bool>& IsCancelButtonEnabled()
+	{
+		return IsCancelButtonEnabledAttribute;
+	}
+
+	FReplyDelegate& OnCancel()
+	{
+		return OnCancelDelegate;
+	}
+
+	FReply Cancel() const
+	{
+		if(OnCancelDelegate.IsBound())
+		{
+			return OnCancelDelegate.Execute();
+		}
+		return FReply::Unhandled();
+	}
+
+	TAttribute<bool>& IsPrimaryButtonVisible()
+	{
+		return IsPrimaryButtonVisibleAttribute;
+	}
+
+	TAttribute<bool>& IsPrimaryButtonEnabled()
+	{
+		return IsPrimaryButtonEnabledAttribute;
+	}
+
+	TAttribute<FText>& PrimaryButtonText()
+	{
+		return PrimaryButtonTextAttribute;
+	}
+
+	FReplyDelegate& OnPrimaryAction()
+	{
+		return OnPrimaryActionDelegate;
+	}
+
+	FReply PrimaryAction() const
+	{
+		if(OnPrimaryActionDelegate.IsBound())
+		{
+			return OnPrimaryActionDelegate.Execute();
+		}
+		return FReply::Unhandled();
+	}
+
+	TSharedRef<FRigVMTreeContext> GetContext() const
+	{
+		return Context;
+	}
+
+	void IncrementContextHash();
+
+	TArray<TSharedRef<FRigVMTreeNode>> GetAllNodes() const;
+
+	const TArray<TSharedRef<FRigVMTreeNode>>& GetVisibleNodes() const;
+
+	TSharedPtr<FRigVMTreeNode> FindVisibleNode(const FString& InPath) const;
+
+	void AddNode(const TSharedRef<FRigVMTreeNode>& InNode);
+	void RemoveNode(const TSharedRef<FRigVMTreeNode>& InNode);
+	void SetNodes(const TArray<TSharedRef<FRigVMTreeNode>>& InNodes);
+
+	FQueueTasksDelegate& OnQueueTasks()
+	{
+		return QueueTasksDelegate;
+	}
+	
+	void QueueTasks(const TArray<TSharedRef<FRigVMTreeTask>>& InTasks) const
+	{
+		(void)QueueTasksDelegate.ExecuteIfBound(InTasks);
+	}
+
+private:
+
+	int32 ID;
+	FString Name;
+	bool bIsActive;
+	bool bAllowsMultiSelection;
+	TSharedRef<FRigVMTreeContext> Context;
+	TSharedRef<FRigVMTreeRootNode> RootNode;
+	TAttribute<bool> IsCancelButtonVisibleAttribute;
+	TAttribute<bool> IsCancelButtonEnabledAttribute;
+	FReplyDelegate OnCancelDelegate;
+	TAttribute<bool> IsPrimaryButtonVisibleAttribute;
+	TAttribute<bool> IsPrimaryButtonEnabledAttribute;
+	TAttribute<FText> PrimaryButtonTextAttribute;
+	FReplyDelegate OnPrimaryActionDelegate;
+	FQueueTasksDelegate QueueTasksDelegate;
+
+	friend class SRigVMBulkEditWidget;
+};
diff --git a/Engine/Plugins/Runtime/RigVM/Source/RigVMEditor/Public/Widgets/RigVMTreeToolkitDefines.h b/Engine/Plugins/Runtime/RigVM/Source/RigVMEditor/Public/Widgets/RigVMTreeToolkitDefines.h
new file mode 100644
--- /dev/null
+++ b/Engine/Plugins/Runtime/RigVM/Source/RigVMEditor/Public/Widgets/RigVMTreeToolkitDefines.h
@@ -0,0 +1,91 @@
+// Copyright Epic Games, Inc. All Rights Reserved.
+
+#pragma once
+
+#include "CoreMinimal.h"
+
+#define DEFINE_RIGVMTREETOOLKIT_QUOTE(str) #str
+#define DEFINE_RIGVMTREETOOLKIT_ELEMENT(TypeName, SuperTypeName) \
+typedef SuperTypeName Super; \
+static inline const FLazyName Type = FLazyName(TEXT(DEFINE_RIGVMTREETOOLKIT_QUOTE(TypeName))); \
+virtual FName GetType() const override \
+{ \
+	return Type; \
+} \
+virtual bool IsOfType(const FName& InType) const override \
+{ \
+	if(InType == Type) \
+	{ \
+		return true; \
+	} \
+	return Super::IsOfType(InType); \
+} \
+template <typename... Types> \
+static TSharedRef<TypeName> Create(Types... Args) \
+{ \
+	TSharedRef<TypeName> Ptr = MakeShared<TypeName>(Args...); \
+	Ptr->Initialize(); \
+	return Ptr; \
+} \
+TSharedPtr<TypeName> ToSharedPtr() const \
+{ \
+	return ToSharedRef().ToSharedPtr(); \
+} \
+TSharedRef<TypeName> ToSharedRef() const \
+{ \
+	TypeName* MutableThis = const_cast<TypeName*>(this); \
+	return StaticCastSharedRef<TypeName>(SharedThis(MutableThis)); \
+} \
+template<typename T> \
+friend TSharedPtr<T> Cast(const TSharedRef<TypeName>& InElement) \
+{ \
+	if(InElement->IsA<T>()) \
+	{ \
+		return StaticCastSharedPtr<T>(InElement->ToSharedPtr()); \
+	} \
+	return nullptr; \
+} \
+template<typename T> \
+friend const TSharedRef<T>& CastChecked(const TSharedRef<TypeName>& InElement) \
+{ \
+	check(InElement->IsA<T>()) \
+	return StaticCastSharedRef<T>(InElement); \
+}
+	
+/**
+ * The Tree Element is the base class for anything within a
+ * tree and manages the lifetime as well as type checking
+ */
+class FRigVMTreeElement : public TSharedFromThis<FRigVMTreeElement>
+{
+public:
+	static inline const FLazyName Type = FLazyName(TEXT("FRigVMTreeElement"));
+	virtual FName GetType() const
+	{
+		return Type;
+	}
+
+	virtual bool IsOfType(const FName& InType) const
+	{
+		return InType == Type;
+	}
+
+	template<typename T>
+	bool IsA() const
+	{
+		return IsOfType(T::Type);
+	}
+
+	virtual ~FRigVMTreeElement()
+	{
+	}
+	
+	virtual void Initialize()
+	{
+	}
+
+protected:
+	FRigVMTreeElement()
+	{
+	}
+};
diff --git a/Engine/Plugins/Runtime/RigVM/Source/RigVMEditor/Public/Widgets/RigVMTreeToolkitFilter.h b/Engine/Plugins/Runtime/RigVM/Source/RigVMEditor/Public/Widgets/RigVMTreeToolkitFilter.h
new file mode 100644
--- /dev/null
+++ b/Engine/Plugins/Runtime/RigVM/Source/RigVMEditor/Public/Widgets/RigVMTreeToolkitFilter.h
@@ -0,0 +1,130 @@
+// Copyright Epic Games, Inc. All Rights Reserved.
+
+#pragma once
+
+#include "RigVMTreeToolkitDefines.h"
+
+class FRigVMTreeNode;
+class FRigVMTreeContext;
+
+/**
+ * A Filter is used to hide elements from a tree.
+ * Filters can be used to filter by folder, filter
+ * by tags or anything else.
+ */
+class FRigVMTreeFilter : public FRigVMTreeElement
+{
+public:
+	DEFINE_RIGVMTREETOOLKIT_ELEMENT(FRigVMTreeFilter, FRigVMTreeElement)
+
+	virtual bool CanBeToggledInUI() const
+	{
+		return false;
+	}
+
+	virtual bool IsInvertedInUI() const
+	{
+		return false;
+	}
+
+	virtual FText GetLabel() const
+	{
+		return FText();
+	}
+
+	virtual FText GetToolTip() const
+	{
+		return FText();
+	}
+	
+	virtual bool Filters(TSharedRef<FRigVMTreeNode>& InNode, const TSharedRef<FRigVMTreeContext>& InContext);
+
+	virtual uint32 GetVisibleChildrenHash() const;
+	
+	bool IsEnabled() const
+	{
+		return bEnabled;
+	}
+
+	void SetEnabled(bool InEnabled)
+	{
+		bEnabled = InEnabled;
+	}
+	
+private:
+	bool bEnabled = true;
+};
+
+/**
+ * The Path filter is used to include elements only if they
+ * match a user provided path.
+ */
+class FRigVMTreePathFilter : public FRigVMTreeFilter
+{
+public:
+	DEFINE_RIGVMTREETOOLKIT_ELEMENT(FRigVMTreePathFilter, FRigVMTreeFilter)
+
+	virtual bool Filters(TSharedRef<FRigVMTreeNode>& InNode, const TSharedRef<FRigVMTreeContext>& InContext) override;
+
+	virtual uint32 GetVisibleChildrenHash() const override;
+	
+	const FString& GetFilterText() const
+	{ 
+		return FilterText;
+	}
+	void SetFilterText(const FString& InFilterText)
+	{ 
+		FilterText = InFilterText;
+	}
+
+private:
+	FString FilterText;
+};
+
+/**
+ * The Engine Content Filter can be used to hide elements
+ * which reside on the engine content path.
+ */
+class FRigVMTreeEngineContentFilter : public FRigVMTreeFilter
+{
+public:
+	DEFINE_RIGVMTREETOOLKIT_ELEMENT(FRigVMTreeEngineContentFilter, FRigVMTreeFilter)
+
+	virtual bool CanBeToggledInUI() const override
+	{
+		return true;
+	}
+
+	virtual bool IsInvertedInUI() const override
+	{
+		return true;
+	}
+
+	virtual FText GetLabel() const override;
+
+	virtual bool Filters(TSharedRef<FRigVMTreeNode>& InNode, const TSharedRef<FRigVMTreeContext>& InContext) override;
+};
+
+/**
+ * The Developer Content Filter can be used to hide items which 
+ * reside on a developer content path.
+ */
+class FRigVMTreeDeveloperContentFilter : public FRigVMTreeFilter
+{
+public:
+	DEFINE_RIGVMTREETOOLKIT_ELEMENT(FRigVMTreeDeveloperContentFilter, FRigVMTreeFilter)
+
+	virtual bool CanBeToggledInUI() const override
+	{
+		return true;
+	}
+
+	virtual bool IsInvertedInUI() const override
+	{
+		return true;
+	}
+
+	virtual FText GetLabel() const override;
+	
+	virtual bool Filters(TSharedRef<FRigVMTreeNode>& InNode, const TSharedRef<FRigVMTreeContext>& InContext) override;
+};
diff --git a/Engine/Plugins/Runtime/RigVM/Source/RigVMEditor/Public/Widgets/RigVMTreeToolkitNode.h b/Engine/Plugins/Runtime/RigVM/Source/RigVMEditor/Public/Widgets/RigVMTreeToolkitNode.h
new file mode 100644
--- /dev/null
+++ b/Engine/Plugins/Runtime/RigVM/Source/RigVMEditor/Public/Widgets/RigVMTreeToolkitNode.h
@@ -0,0 +1,189 @@
+// Copyright Epic Games, Inc. All Rights Reserved.
+
+#pragma once
+
+#include "RigVMBlueprint.h"
+#include "RigVMTreeToolkitDefines.h"
+#include "Styling/SlateTypes.h"
+#include "Textures/SlateIcon.h"
+
+class FRigVMTreeContext;
+
+/**
+ * A Node is the base element for anything that is shown
+ * within the tree.
+ */
+class FRigVMTreeNode : public FRigVMTreeElement
+{
+public:
+	DEFINE_RIGVMTREETOOLKIT_ELEMENT(FRigVMTreeNode, FRigVMTreeElement)
+	
+	FRigVMTreeNode(const FString& InPath)
+		: Path(InPath)
+	{
+	}
+
+	virtual const FString& GetPath() const
+	{
+		return Path;
+	}
+
+	virtual FAssetData GetAssetData() const;
+	virtual URigVMBlueprint* GetBlueprint() const;
+
+	virtual FText GetLabel() const;
+
+	virtual void GetContextMenu(FMenuBuilder& InMenuBuilder)
+	{
+	}
+
+	virtual bool IsCheckable() const
+	{
+		return false;
+	}
+
+	ECheckBoxState GetCheckState() const;
+
+	void SetCheckState(ECheckBoxState InNewState);
+	void ResetCheckState();
+
+	virtual bool ShouldExpandByDefault() const
+	{
+		return false;
+	}
+	
+	int32 GetDepth() const;
+
+	TSharedPtr<FRigVMTreeNode> GetParent() const
+	{
+		if(Parent)
+		{
+			return Parent->ToSharedPtr();
+		}
+		return nullptr;
+	}
+
+	TSharedRef<FRigVMTreeNode> GetRoot() const;
+
+	const TArray<TSharedRef<FRigVMTreeNode>>& GetVisibleChildren(const TSharedRef<FRigVMTreeContext>& InContext) const
+	{
+		UpdateVisibleChildren(InContext);
+		return VisibleChildren;
+	}
+
+	TArray<TSharedRef<FRigVMTreeNode>> GetChildren(const TSharedRef<FRigVMTreeContext>& InContext) const
+	{
+		UpdateChildren(InContext);
+		return Children.Get(TArray<TSharedRef<FRigVMTreeNode>>());
+	}
+
+	virtual bool HasVisibleChildren() const
+	{
+		return !VisibleChildren.IsEmpty();
+	}
+
+	TSharedPtr<FRigVMTreeNode> FindVisibleChild(const FString& InFullPath, const TSharedRef<FRigVMTreeContext>& InContext) const;
+
+	virtual void DirtyChildren();
+	virtual void DirtyVisibleChildren();
+	void RequestRefresh(bool bForce = false);
+
+	virtual const FSlateBrush* GetIconAndTint(FLinearColor& OutColor) const
+	{
+		return nullptr;
+	}
+	virtual const FSlateBrush* GetBackgroundImage(bool bIsHovered, bool bIsSelected) const;
+	virtual FSlateColor GetBackgroundColor(bool bIsHovered, bool bIsSelected) const;
+
+	virtual bool IsLoaded() const;
+
+protected:
+	virtual TArray<TSharedRef<FRigVMTreeNode>> GetChildrenImpl(const TSharedRef<FRigVMTreeContext>& InContext) const
+	{
+		return Children.Get(TArray<TSharedRef<FRigVMTreeNode>>());
+	}
+	void AddChildImpl(const TSharedRef<FRigVMTreeNode>& InChild) const;
+	void UpdateChildren(const TSharedRef<FRigVMTreeContext>& InContext) const;
+	void UpdateVisibleChildren(const TSharedRef<FRigVMTreeContext>& InContext) const;
+	void DirtyVisibleParent() const;
+	bool ContainsAnyVisibleCheckedNode() const;
+
+	static TSharedPtr<FRigVMTreeNode> FindVisibleNodeInSet(const TArray<TSharedRef<FRigVMTreeNode>>& InNodes, const FString& InPath, const TSharedRef<FRigVMTreeContext>& InContext);
+
+	FString Path;
+	const FRigVMTreeNode* Parent = nullptr;
+	mutable TOptional<TArray<TSharedRef<FRigVMTreeNode>>> Children;
+	mutable uint32 VisibleChildrenHash = UINT32_MAX;
+	mutable TArray<TSharedRef<FRigVMTreeNode>> VisibleChildren;
+	mutable bool bIsDirtyingParent = false;
+	TOptional<ECheckBoxState> CheckState;
+
+	DECLARE_DELEGATE_OneParam(FRefreshDelegate, bool)
+	FRefreshDelegate RefreshDelegate;
+
+	friend class SRigVMChangesTreeRow;
+	friend class SRigVMChangesTreeView;
+	friend class FRigVMTreePhase;
+	friend class FRemoveNodeTask;
+};
+
+/**
+ * A Root Node is a node which can be placed only at
+ * the root of the tree.
+ */
+class FRigVMTreeRootNode : public FRigVMTreeNode
+{
+public:
+
+	DEFINE_RIGVMTREETOOLKIT_ELEMENT(FRigVMTreeRootNode, FRigVMTreeNode)
+	
+	FRigVMTreeRootNode();
+
+	void AddChild(const TSharedRef<FRigVMTreeNode>& InNode);
+	void RemoveChild(const TSharedRef<FRigVMTreeNode>& InNode);
+	void SetChildren(const TArray<TSharedRef<FRigVMTreeNode>>& InNodes);
+};
+
+/** 
+ * A Category Node is a node which presents itself as a category,
+ * providing a collapsed UI header.
+ */
+class FRigVMTreeCategoryNode : public FRigVMTreeNode
+{
+public:
+	DEFINE_RIGVMTREETOOLKIT_ELEMENT(FRigVMTreeCategoryNode, FRigVMTreeNode)
+	
+	FRigVMTreeCategoryNode(const FString& InPath)
+		: FRigVMTreeNode(InPath)
+	{
+	}
+
+	virtual const FSlateBrush* GetBackgroundImage(bool bIsHovered, bool bIsSelected) const override;
+	virtual FSlateColor GetBackgroundColor(bool bIsHovered, bool bIsSelected) const override;
+};
+
+/**
+ * A Package Node is a special Category Node which reacts to 
+ * the package being loaded and offers to update its content.
+ */
+class FRigVMTreePackageNode : public FRigVMTreeCategoryNode
+{
+public:
+	DEFINE_RIGVMTREETOOLKIT_ELEMENT(FRigVMTreePackageNode, FRigVMTreeCategoryNode)
+	
+	FRigVMTreePackageNode(const FAssetData& InAssetData);
+
+	virtual const FSlateBrush* GetIconAndTint(FLinearColor& OutColor) const override;
+	virtual void Initialize() override;
+
+	virtual FText GetLabel() const override;
+	virtual bool IsLoaded() const override;
+
+protected:
+	void HandleAssetLoaded(UObject* InAsset);
+
+	mutable TOptional<bool> bIsLoaded;
+	FSoftObjectPath SoftObjectPath;
+	mutable FSlateIcon Icon;
+	mutable TOptional<const FSlateBrush*> IconBrush;
+};
diff --git a/Engine/Plugins/Runtime/RigVM/Source/RigVMEditor/Public/Widgets/RigVMTreeToolkitTask.h b/Engine/Plugins/Runtime/RigVM/Source/RigVMEditor/Public/Widgets/RigVMTreeToolkitTask.h
new file mode 100644
--- /dev/null
+++ b/Engine/Plugins/Runtime/RigVM/Source/RigVMEditor/Public/Widgets/RigVMTreeToolkitTask.h
@@ -0,0 +1,103 @@
+// Copyright Epic Games, Inc. All Rights Reserved.
+
+#pragma once
+
+#include "RigVMTreeToolkitDefines.h"
+#include "AssetRegistry/AssetData.h"
+#include "Widgets/SRigVMLogWidget.h"
+
+class FRigVMTreeNode;
+class FRigVMTreePhase;
+class FRigVMTreeContext;
+
+/**
+ * A Task is a small piece of work that needs to be performed.
+ */
+class FRigVMTreeTask : public FRigVMTreeElement
+{
+public:
+	DEFINE_RIGVMTREETOOLKIT_ELEMENT(FRigVMTreeTask, FRigVMTreeElement)
+
+	virtual bool Execute(const TSharedRef<FRigVMTreePhase>& InPhase)
+	{
+		return false;
+	};
+
+	virtual URigVMBlueprint* GetBlueprint(const TSharedRef<FRigVMTreePhase>& InPhase) const
+	{
+		return nullptr;
+	}
+
+	virtual bool RequiresRefresh() const
+	{
+		return false;
+	}
+
+	virtual bool RequiresUndo() const
+	{
+		return false;
+	}
+
+	virtual TArray<FString> GetAffectedNodes() const
+	{
+		return {};
+	}
+	
+	void SetEnableUndo(bool bEnabled)
+	{
+		bEnableUndo = bEnabled;
+	}
+
+	bool IsUndoEnabled() const
+	{
+		return bEnableUndo;
+	}
+
+private:
+	bool bEnableUndo = false;
+};
+
+/**
+ * The Load Package For Node Task will load a package related to a node
+ * given the node's path.
+ */
+class FRigVMTreeLoadPackageForNodeTask : public FRigVMTreeTask
+{
+public:
+	DEFINE_RIGVMTREETOOLKIT_ELEMENT(FRigVMTreeLoadPackageForNodeTask, FRigVMTreeTask)
+
+	FRigVMTreeLoadPackageForNodeTask(const TSharedRef<FRigVMTreeNode>& InNode);
+
+	virtual bool Execute(const TSharedRef<FRigVMTreePhase>& InPhase) override;
+
+	virtual bool RequiresRefresh() const override
+	{
+		return true;
+	}
+private:
+
+	FAssetData AssetData;
+};
+
+/**
+ * The Compile Blueprint Task is used to compile a blueprint related
+ * to a node in the tree (based on the node's path).
+ */
+class FRigVMCompileBlueprintTask : public FRigVMTreeTask
+{
+public:
+	DEFINE_RIGVMTREETOOLKIT_ELEMENT(FRigVMCompileBlueprintTask, FRigVMTreeTask)
+
+	FRigVMCompileBlueprintTask(const TSharedRef<FRigVMTreeNode>& InNode);
+
+	virtual bool Execute(const TSharedRef<FRigVMTreePhase>& InPhase) override;
+
+	virtual bool RequiresRefresh() const override
+	{
+		return true;
+	}
+
+private:
+
+	FSoftObjectPath ObjectPath;
+};
diff --git a/Engine/Plugins/Runtime/RigVM/Source/RigVMEditor/Public/Widgets/SRigVMChangesTreeView.h b/Engine/Plugins/Runtime/RigVM/Source/RigVMEditor/Public/Widgets/SRigVMChangesTreeView.h
new file mode 100644
--- /dev/null
+++ b/Engine/Plugins/Runtime/RigVM/Source/RigVMEditor/Public/Widgets/SRigVMChangesTreeView.h
@@ -0,0 +1,123 @@
+// Copyright Epic Games, Inc. All Rights Reserved.
+
+#pragma once
+
+#include "Widgets/RigVMTreeToolkitNode.h"
+#include "Widgets/RigVMTreeToolkitFilter.h"
+#include "Widgets/RigVMTreeToolkitContext.h"
+#include "Widgets/RigVMTreeToolkitTask.h"
+#include "Widgets/SCompoundWidget.h"
+#include "Widgets/Views/STreeView.h"
+#include "Widgets/Input/SSearchBox.h"
+
+class SRigVMChangesTreeView;
+
+class RIGVMEDITOR_API SRigVMChangesTreeRow
+	: public STableRow<TSharedRef<FRigVMTreeNode>>
+{
+public:
+	
+	SLATE_BEGIN_ARGS(SRigVMChangesTreeRow)
+	{}
+	SLATE_ARGUMENT(TSharedPtr<FRigVMTreeNode>, Node)
+	SLATE_ARGUMENT(TSharedPtr<SRigVMChangesTreeView>, OwningWidget)
+	SLATE_END_ARGS()
+
+	virtual ~SRigVMChangesTreeRow() override;
+	void Construct(const FArguments& InArgs, const TSharedRef<STableViewBase>& OwnerTableView);
+
+private:
+
+	EVisibility GetCheckBoxVisibility() const;
+	ECheckBoxState GetCheckBoxState() const;
+	void OnCheckBoxStateChanged(ECheckBoxState InNewState);
+	const FSlateBrush* GetBackgroundImage() const;
+	FSlateColor GetBackgroundColor() const;
+	FOptionalSize GetIndentWidth() const;
+	FSlateColor GetTextColor() const;
+	const FSlateBrush* GetIcon() const;
+	FSlateColor GetIconColor() const;
+	const FSlateBrush* GetExpanderImage() const;
+	FReply OnExpanderMouseButtonDown(const FGeometry& SenderGeometry, const FPointerEvent& MouseEvent);
+	EVisibility GetExpanderVisibility() const;
+	void RequestRefresh(bool bForce);
+
+	TSharedPtr<FRigVMTreeNode> Node;
+	SRigVMChangesTreeView* OwningWidget = nullptr;
+};
+
+DECLARE_DELEGATE_RetVal_OneParam(FReply,FOnRigVMTreeNodeSelected,TSharedRef<FRigVMTreeNode>);
+DECLARE_DELEGATE_RetVal_OneParam(FReply,FOnRigVMTreeNodeDoubleClicked,TSharedRef<FRigVMTreeNode>);
+
+class RIGVMEDITOR_API SRigVMChangesTreeView : public SCompoundWidget
+{
+public:
+	SLATE_BEGIN_ARGS(SRigVMChangesTreeView)
+	{}
+	SLATE_ATTRIBUTE(TSharedPtr<FRigVMTreePhase>, Phase)
+	SLATE_EVENT(FOnRigVMTreeNodeSelected, OnNodeSelected)
+	SLATE_EVENT(FOnRigVMTreeNodeDoubleClicked, OnNodeDoubleClicked)
+	SLATE_END_ARGS()
+
+	virtual ~SRigVMChangesTreeView() override;
+
+	void Construct(const FArguments& InArgs);
+
+	virtual void Tick(const FGeometry& AllottedGeometry, const double InCurrentTime, const float InDeltaTime) override;
+
+	EVisibility GetPathFilterVisibility() const;
+	void OnPathFilterTextChanged(const FText& SearchText);
+	void OnPathFilterTextCommitted(const FText& SearchText, ETextCommit::Type CommitInfo);
+
+	void RequestRefresh_AnyThread(bool bForce = false);
+	void RefreshFilteredNodes(bool bForce = false);
+	void RefreshFilteredNodesIfRequired();
+
+	void OnPhaseChanged();
+	void SetSelection(const TSharedPtr<FRigVMTreeNode>& InNode, bool bRequestScrollIntoView = false);
+
+	TSharedPtr<STreeView<TSharedRef<FRigVMTreeNode>>> GetTreeView()
+	{
+		return TreeView;
+	}
+
+	TSharedRef<FRigVMTreePhase> GetPhase() const;
+	TSharedRef<FRigVMTreeContext> GetContext() const;
+	TSharedPtr<FRigVMTreePathFilter> GetPathFilter() const;
+
+	TArray<TSharedRef<FRigVMTreeNode>> GetSelectedNodes() const;
+	bool HasAnyVisibleCheckedNode() const;
+	TArray<TSharedRef<FRigVMTreeNode>> GetCheckedNodes() const;
+	
+private:
+
+	TAttribute<TSharedPtr<FRigVMTreePhase>> PhaseAttribute;
+	TArray<TSharedRef<FRigVMTreeNode>> FilteredNodes;
+	
+	FText GetPathFilterText();
+	TSharedPtr<SSearchBox> PathFilterBox;
+
+	TSharedPtr<STreeView<TSharedRef<FRigVMTreeNode>>> TreeView;
+
+	FOnRigVMTreeNodeSelected OnNodeSelected;
+	FOnRigVMTreeNodeDoubleClicked OnNodeDoubleClicked;
+
+	TAttribute<FText> BulkEditTitle;
+	TAttribute<FText> BulkEditConfirmMessage;
+	TAttribute<FString> BulkEditConfirmIniField;
+
+	TSharedRef<ITableRow> MakeTreeRowWidget(TSharedRef<FRigVMTreeNode> InNode, const TSharedRef<STableViewBase>& OwnerTable);
+
+	void GetChildrenForNode(TSharedRef<FRigVMTreeNode> InNode, TArray<TSharedRef<FRigVMTreeNode>>& OutChildren);
+	void OnSelectionChanged(TSharedPtr<FRigVMTreeNode> Selection, ESelectInfo::Type SelectInfo);
+	void OnTreeElementDoubleClicked(TSharedRef<FRigVMTreeNode> InNode);
+	ESelectionMode::Type GetSelectionMode() const;
+
+	TSharedPtr<SWidget> OnGetNodeContextMenuContent();
+
+	EVisibility GetSettingsButtonVisibility() const;
+	FReply OnSettingsButtonClicked();
+
+	std::atomic<int32> RequestRefreshCount;
+	std::atomic<int32> RequestRefreshForceCount;
+};
