RigVM: Avoid combinatorial explosion when looking for source nodes

#rnx

diff --git a/Engine/Plugins/Runtime/RigVM/Source/RigVMDeveloper/Private/RigVMModel/RigVMPin.cpp b/Engine/Plugins/Runtime/RigVM/Source/RigVMDeveloper/Private/RigVMModel/RigVMPin.cpp
--- a/Engine/Plugins/Runtime/RigVM/Source/RigVMDeveloper/Private/RigVMModel/RigVMPin.cpp
+++ b/Engine/Plugins/Runtime/RigVM/Source/RigVMDeveloper/Private/RigVMModel/RigVMPin.cpp
@@ -2275,7 +2275,11 @@ bool URigVMPin::CanLink(const URigVMPin* InSourcePin, const URigVMPin* InTargetP
 
 						return false;
 					}
-					SourceNodes.Append(SourceNodes[SourceNodeIndex]->GetLinkedSourceNodes());
+					const TArray<URigVMNode*> LinkedSourceNodes = SourceNodes[SourceNodeIndex]->GetLinkedSourceNodes();
+					for(URigVMNode* LinkedSourceNode : LinkedSourceNodes)
+					{
+						SourceNodes.AddUnique(LinkedSourceNode);
+					}
 				}
 			}
 		}
