Control Rig: Add category to function variants
#rnx

diff --git a/Engine/Plugins/Runtime/RigVM/Source/RigVMDeveloper/Private/RigVMModel/Nodes/RigVMLibraryNode.cpp b/Engine/Plugins/Runtime/RigVM/Source/RigVMDeveloper/Private/RigVMModel/Nodes/RigVMLibraryNode.cpp
--- a/Engine/Plugins/Runtime/RigVM/Source/RigVMDeveloper/Private/RigVMModel/Nodes/RigVMLibraryNode.cpp
+++ b/Engine/Plugins/Runtime/RigVM/Source/RigVMDeveloper/Private/RigVMModel/Nodes/RigVMLibraryNode.cpp
@@ -243,6 +243,7 @@ FRigVMGraphFunctionHeader URigVMLibraryNode::GetFunctionHeader(IRigVMGraphFuncti
 		if (FRigVMGraphFunctionData* Data = Host->GetRigVMGraphFunctionStore()->FindFunction(Header.LibraryPointer))
 		{
 			Header.Variant = Data->Header.Variant;
+			Header.NodeTitle = Data->Header.NodeTitle;
 		}
 	}
 	if (!Header.Variant.Guid.IsValid())
@@ -259,7 +260,10 @@ FRigVMGraphFunctionHeader URigVMLibraryNode::GetFunctionHeader(IRigVMGraphFuncti
 	//Header.Tooltip = GetToolTipText();
 	
     Header.NodeColor = GetNodeColor();
-    Header.NodeTitle = GetNodeTitle();
+	if (Header.NodeTitle.IsEmpty())
+	{
+		Header.NodeTitle = GetNodeTitle();
+	}
     for(URigVMPin* Pin : GetPins())
     {
     	FRigVMGraphFunctionArgument Arg;
diff --git a/Engine/Plugins/Runtime/RigVM/Source/RigVMDeveloper/Private/RigVMModel/RigVMController.cpp b/Engine/Plugins/Runtime/RigVM/Source/RigVMDeveloper/Private/RigVMModel/RigVMController.cpp
--- a/Engine/Plugins/Runtime/RigVM/Source/RigVMDeveloper/Private/RigVMModel/RigVMController.cpp
+++ b/Engine/Plugins/Runtime/RigVM/Source/RigVMDeveloper/Private/RigVMModel/RigVMController.cpp
@@ -1336,6 +1336,36 @@ TArray<FString> URigVMController::GetAddDecoratorPythonCommands(URigVMNode* Node
 	return Commands;
 }
 
+FRigVMGraphFunctionStore* URigVMController::GetGraphFunctionStore() const
+{
+	IRigVMClientHost* ClientHost = GetImplementingOuter<IRigVMClientHost>();
+	if (!ClientHost)
+	{
+		ReportErrorf(TEXT("Cannot find client host."));
+		return nullptr;
+	}
+
+	IRigVMGraphFunctionHost* FunctionHost = ClientHost->GetRigVMGraphFunctionHost();
+	if (!FunctionHost)
+	{
+		ReportErrorf(TEXT("Cannot find function host."));
+		return nullptr;
+	}
+
+	return FunctionHost->GetRigVMGraphFunctionStore();
+}
+
+FRigVMGraphFunctionData* URigVMController::FindFunctionData(const FName& InFunctionName) const
+{
+	FRigVMGraphFunctionStore* FunctionStore = GetGraphFunctionStore();
+	if (!FunctionStore)
+	{
+		return nullptr;
+	}
+
+	return FunctionStore->FindFunctionByName(InFunctionName);
+}
+
 #if WITH_EDITOR
 
 URigVMUnitNode* URigVMController::AddUnitNode(UScriptStruct* InScriptStruct, const FName& InMethodName, const FVector2D& InPosition, const FString& InNodeName, bool bSetupUndoRedo, bool bPrintPythonCommand)
@@ -6663,6 +6693,11 @@ bool URigVMController::RemoveNodes(TArray<URigVMNode*> InNodes, bool bSetupUndoR
 						GetActionStack()->AddAction(FRigVMMarkFunctionPublicAction(this, LibraryNode->GetFName(), true));
 					}
 				}
+
+				if (FunctionLibrary->FunctionToVariant.Contains(LibraryNode->GetFName()))
+				{
+					FunctionLibrary->FunctionToVariant.Remove(LibraryNode->GetFName());
+				}
 			}
 		}
 
@@ -6886,6 +6921,12 @@ bool URigVMController::RenameNode(URigVMNode* InNode, const FName& InNewName, bo
 				FunctionLibrary->PublicFunctionNames.Remove(InNode->PreviousName);
 				FunctionLibrary->PublicFunctionNames.Add(ValidNewName);
 			}
+
+			if (FunctionLibrary->FunctionToVariant.Contains(InNode->PreviousName))
+			{
+				FunctionLibrary->FunctionToVariant.Remove(InNode->PreviousName);
+				FunctionLibrary->FunctionToVariant.Add(ValidNewName);
+			}
 		}
 	}
 
@@ -7308,7 +7349,28 @@ bool URigVMController::SetNodeCategory(URigVMCollapseNode* InNode, const FString
 		GetActionStack()->BeginAction(Action);
 	}
 
-	InNode->NodeCategory = InCategory;
+	// If this is a function node with variants, always add the function title as a category
+	FString CategoryToSet = InCategory;
+	if (FRigVMGraphFunctionData* FunctionData = FindFunctionData(InNode->GetFName()))
+	{
+		if (FunctionData->Header.Variant.Guid.IsValid())
+		{
+			const FString NodeTitle = FunctionData->Header.NodeTitle;
+			TArray<FString> Categories;
+			InCategory.ParseIntoArray(Categories, TEXT("|"), true);
+			Categories.AddUnique(NodeTitle);
+			Categories.Remove(FString()); // Remove any empty strings
+			if (Categories[0] != NodeTitle)
+			{
+				// Make sure the node title is the first element
+				Categories.Remove(NodeTitle);
+				Categories.Insert(NodeTitle, 0);
+			}
+			CategoryToSet = FString::Join(Categories, TEXT("|"));
+		}
+	}
+
+	InNode->NodeCategory = CategoryToSet;
 	Notify(ERigVMGraphNotifType::NodeCategoryChanged, InNode);
 
 	if (bSetupUndoRedo)
@@ -12181,6 +12243,27 @@ URigVMLibraryNode* URigVMController::CreateFunctionVariant(const FName& InFuncti
 			NewFunction->Header.Name = NewName;
 			NewFunction->Header.LibraryPointer = Identifier;
 
+			TArray<FString> Categories;
+			OriginalFunction->Header.Category.ParseIntoArray(Categories, TEXT("|"), true);
+
+			const int32 NumCategories = Categories.Num();
+			Categories.AddUnique(OriginalFunction->Header.NodeTitle);
+			const bool bAddedCategory = NumCategories != Categories.Num();
+
+			FString NewCategory = FString::Join(Categories, TEXT("|"));
+			if (bAddedCategory)
+			{
+				SetNodeCategory(Cast<URigVMCollapseNode>(FunctionToClone), NewCategory, false);
+			}
+			
+			URigVMLibraryNode* CloneFunction = Cast<URigVMLibraryNode>(Graph->FindNodeByName(NewName));
+			SetNodeCategory(Cast<URigVMCollapseNode>(CloneFunction), NewCategory, false);
+
+			if (URigVMFunctionLibrary* FunctionLibrary = Cast<URigVMFunctionLibrary>(Graph))
+			{
+				FunctionLibrary->FunctionToVariant.Add(NewName, NewVariant);
+			}
+
 			MarkFunctionAsPublic(NewName, bIsPublic, false);
 		}
 
@@ -12214,21 +12297,14 @@ TArray<FRigVMVariantRef> URigVMController::FindVariantsOfFunction(const FName& I
 {
 	TArray<FRigVMVariantRef> Result;
 	
-    IRigVMClientHost* ClientHost = GetImplementingOuter<IRigVMClientHost>();
-    if (!ClientHost)
-    {
-    	ReportErrorf(TEXT("Cannot find client host."));
-    	return Result;
-    }
-
-    IRigVMGraphFunctionHost* FunctionHost = ClientHost->GetRigVMGraphFunctionHost();
-    if (!FunctionHost)
+    FRigVMGraphFunctionStore* FunctionStore = GetGraphFunctionStore();
+    if (!FunctionStore)
     {
-    	ReportErrorf(TEXT("Cannot find function host."));
+    	ReportErrorf(TEXT("Cannot find function store."));
     	return Result;
     }
 
-    FRigVMGraphFunctionData* FunctionData = FunctionHost->GetRigVMGraphFunctionStore()->FindFunctionByName(InFunctionName);
+    FRigVMGraphFunctionData* FunctionData = FunctionStore->FindFunctionByName(InFunctionName);
     if (!FunctionData)
     {
     	ReportErrorf(TEXT("Cannot find function %s."), *InFunctionName.ToString());
diff --git a/Engine/Plugins/Runtime/RigVM/Source/RigVMDeveloper/Public/RigVMModel/RigVMController.h b/Engine/Plugins/Runtime/RigVM/Source/RigVMDeveloper/Public/RigVMModel/RigVMController.h
--- a/Engine/Plugins/Runtime/RigVM/Source/RigVMDeveloper/Public/RigVMModel/RigVMController.h
+++ b/Engine/Plugins/Runtime/RigVM/Source/RigVMDeveloper/Public/RigVMModel/RigVMController.h
@@ -325,6 +325,9 @@ public:
 	TArray<FString> GetAddNodePythonCommands(URigVMNode* Node) const;
 	TArray<FString> GetAddDecoratorPythonCommands(URigVMNode* Node, const FName& DecoratorName) const;
 
+	FRigVMGraphFunctionStore* GetGraphFunctionStore() const;
+	FRigVMGraphFunctionData* FindFunctionData(const FName& InFunctionName) const;
+
 #if WITH_EDITOR
 	// Note: The functions below are scoped with WITH_EDITOR since we are considering
 	// to move this code into the runtime in the future. Right now there's a dependency
diff --git a/Engine/Plugins/Runtime/RigVM/Source/RigVMDeveloper/Public/RigVMModel/RigVMControllerActions.h b/Engine/Plugins/Runtime/RigVM/Source/RigVMDeveloper/Public/RigVMModel/RigVMControllerActions.h
--- a/Engine/Plugins/Runtime/RigVM/Source/RigVMDeveloper/Public/RigVMModel/RigVMControllerActions.h
+++ b/Engine/Plugins/Runtime/RigVM/Source/RigVMDeveloper/Public/RigVMModel/RigVMControllerActions.h
@@ -1432,7 +1432,7 @@ public:
 };
 
 /**
- * An action marking a function as public/private.
+ * An action creating a function variant
  */
 USTRUCT()
 struct FRigVMCreateFunctionVariantAction : public FRigVMBaseAction
