Control Rig: Create variants for RigVMBlueprint assets
#rnx

diff --git a/Engine/Plugins/Runtime/RigVM/Source/RigVMDeveloper/Private/RigVMBlueprint.cpp b/Engine/Plugins/Runtime/RigVM/Source/RigVMDeveloper/Private/RigVMBlueprint.cpp
--- a/Engine/Plugins/Runtime/RigVM/Source/RigVMDeveloper/Private/RigVMBlueprint.cpp
+++ b/Engine/Plugins/Runtime/RigVM/Source/RigVMDeveloper/Private/RigVMBlueprint.cpp
@@ -203,6 +203,8 @@ URigVMBlueprint::URigVMBlueprint(const FObjectInitializer& ObjectInitializer)
 	CompileLog.SetSourcePath(GetPathName());
 #endif
 
+	AssetVariant.Guid = FRigVMVariant::GenerateGUID();
+
 	if(GetClass() == URigVMBlueprint::StaticClass())
 	{
 		CommonInitialization(ObjectInitializer);
@@ -1086,6 +1088,11 @@ void URigVMBlueprint::PostLoad()
 	FCoreUObjectDelegates::OnObjectModified.AddUObject(this, &URigVMBlueprint::OnPreVariableChange);
 	OnChanged().AddUObject(this, &URigVMBlueprint::OnPostVariableChange);
 
+	if (!AssetVariant.Guid.IsValid())
+	{
+		AssetVariant.Guid = FRigVMVariant::GenerateGUID(GetPackage()->GetPathName());
+	}
+
 	if (UPackage* Package = GetOutermost())
 	{
 		Package->SetDirtyFlag(bDirtyDuringLoad);
diff --git a/Engine/Plugins/Runtime/RigVM/Source/RigVMDeveloper/Public/RigVMBlueprint.h b/Engine/Plugins/Runtime/RigVM/Source/RigVMDeveloper/Public/RigVMBlueprint.h
--- a/Engine/Plugins/Runtime/RigVM/Source/RigVMDeveloper/Public/RigVMBlueprint.h
+++ b/Engine/Plugins/Runtime/RigVM/Source/RigVMDeveloper/Public/RigVMBlueprint.h
@@ -504,6 +504,10 @@ protected:
 
 #if WITH_EDITORONLY_DATA
 
+	/** Variant information about this asset */
+	UPROPERTY(AssetRegistrySearchable)
+	FRigVMVariant AssetVariant;
+
 	UPROPERTY(transient, DuplicateTransient)
 	TObjectPtr<URigVMGraph> TemplateModel;
 
diff --git a/Engine/Plugins/Runtime/RigVM/Source/RigVMEditor/Private/RigVMEditorModule.cpp b/Engine/Plugins/Runtime/RigVM/Source/RigVMEditor/Private/RigVMEditorModule.cpp
--- a/Engine/Plugins/Runtime/RigVM/Source/RigVMEditor/Private/RigVMEditorModule.cpp
+++ b/Engine/Plugins/Runtime/RigVM/Source/RigVMEditor/Private/RigVMEditorModule.cpp
@@ -5,6 +5,7 @@
 =============================================================================*/
 
 #include "RigVMEditorModule.h"
+#include "AssetToolsModule.h"
 #include "Modules/ModuleManager.h"
 #include "AssetRegistry/AssetRegistryModule.h"
 #include "Editor/RigVMEditorCommands.h"
@@ -20,6 +21,7 @@
 #include "UserDefinedStructureCompilerUtils.h"
 #include "EdGraph/RigVMEdGraphConnectionDrawingPolicy.h"
 #include "BlueprintActionDatabaseRegistrar.h"
+#include "ContentBrowserMenuContexts.h"
 #include "EdGraph/NodeSpawners/RigVMEdGraphUnitNodeSpawner.h"
 #include "EdGraph/NodeSpawners/RigVMEdGraphVariableNodeSpawner.h"
 #include "EdGraph/NodeSpawners/RigVMEdGraphTemplateNodeSpawner.h"
@@ -75,6 +77,40 @@ void FRigVMEditorModule::StartupModule()
 		{
 			AssetRegistry->OnAssetRemoved().AddStatic(&FRigVMBlueprintUtils::HandleAssetDeleted);
 		}
+
+		if (UToolMenu* Menu = UToolMenus::Get()->ExtendMenu("ContentBrowser.AssetContextMenu"))
+		{
+			if (FToolMenuSection* Section = Menu->FindSection("CommonAssetActions"))
+			{
+				Section->AddDynamicEntry("CreateVariant", FNewToolMenuSectionDelegate::CreateLambda([](FToolMenuSection& InSection)
+				{
+					UContentBrowserAssetContextMenuContext* Context = InSection.FindContext<UContentBrowserAssetContextMenuContext>();
+					if (Context)
+					{
+						TArray<UObject*> SelectedObjects = Context->GetSelectedObjects();
+						if (SelectedObjects.Num() == 1 && SelectedObjects[0]->IsA<URigVMBlueprint>())
+						{
+							InSection.AddMenuEntry("CreateVariant", LOCTEXT("CreateVariant", "Create variant"), LOCTEXT("CreateVariant_ToolTip", "Create a variant for this asset"), FSlateIcon(FRigVMEditorStyle::Get().GetStyleSetName(), "RigVM", "RigVM.Unit"), FExecuteAction::CreateLambda([SelectedObjects]()
+							{
+								FAssetToolsModule& AssetToolsModule = FModuleManager::LoadModuleChecked<FAssetToolsModule>("AssetTools");
+
+								FString PathName = FPackageName::GetLongPackagePath(SelectedObjects[0]->GetPathName());
+								FString ObjectName = SelectedObjects[0]->GetName();
+								FString PackageName;
+								FString BasePackageName = PathName + "/" + ObjectName;
+								AssetToolsModule.Get().CreateUniqueAssetName(BasePackageName, TEXT(""), PackageName, ObjectName);
+
+								UObject* DuplicateAsset = AssetToolsModule.Get().DuplicateAsset(ObjectName, PathName, SelectedObjects[0]);
+								if (URigVMBlueprint* DuplicateBlueprint = Cast<URigVMBlueprint>(DuplicateAsset))
+								{
+									DuplicateBlueprint->AssetVariant = Cast<URigVMBlueprint>(SelectedObjects[0])->AssetVariant;
+								}
+							}));
+						}
+					}
+				}));
+			}
+		}
 	}
 
 	StartupModuleCommon();
diff --git a/Engine/Plugins/Runtime/RigVM/Source/RigVMEditor/RigVMEditor.Build.cs b/Engine/Plugins/Runtime/RigVM/Source/RigVMEditor/RigVMEditor.Build.cs
--- a/Engine/Plugins/Runtime/RigVM/Source/RigVMEditor/RigVMEditor.Build.cs
+++ b/Engine/Plugins/Runtime/RigVM/Source/RigVMEditor/RigVMEditor.Build.cs
@@ -33,6 +33,7 @@ public class RigVMEditor : ModuleRules
                 "MessageLog",
 				"StructUtils",
 				"StructUtilsEditor",
+				"ContentBrowser",
 			}
 		);
 	}
