Slate: Deprecate SListView::ItemHeight and STreeViewItemHeight. ItemHeight and ItemWidth are only used when the panel is is as a tile or both value are assigned. With TreeView and ListView, you can't assigned both. The function are protected, a subclass of those could but it's up to them to expose the ItemHeight and ItemWidth. This behavior changed was introduced with CL 7634094.

diff --git a/Engine/Plugins/Runtime/RigVM/Source/RigVMEditor/Private/Widgets/SRigVMChangesTreeView.cpp b/Engine/Plugins/Runtime/RigVM/Source/RigVMEditor/Private/Widgets/SRigVMChangesTreeView.cpp
--- a/Engine/Plugins/Runtime/RigVM/Source/RigVMEditor/Private/Widgets/SRigVMChangesTreeView.cpp
+++ b/Engine/Plugins/Runtime/RigVM/Source/RigVMEditor/Private/Widgets/SRigVMChangesTreeView.cpp
@@ -284,7 +284,6 @@ void SRigVMChangesTreeView::Construct(const FArguments& InArgs)
 		.TreeItemsSource(&FilteredNodes)
 		.OnGenerateRow(this, &SRigVMChangesTreeView::MakeTreeRowWidget)
 		.OnGetChildren(this, &SRigVMChangesTreeView::GetChildrenForNode)
-		.ItemHeight(22.0f)
 		.OnSelectionChanged(this, &SRigVMChangesTreeView::OnSelectionChanged)
 		.OnMouseButtonDoubleClick(this, &SRigVMChangesTreeView::OnTreeElementDoubleClicked)
 		.SelectionMode(this, &SRigVMChangesTreeView::GetSelectionMode)
