RigVM: Fix action stack related unit tests

#rnx

diff --git a/Engine/Plugins/Runtime/RigVM/Source/RigVMDeveloper/Public/RigVMModel/RigVMController.h b/Engine/Plugins/Runtime/RigVM/Source/RigVMDeveloper/Public/RigVMModel/RigVMController.h
--- a/Engine/Plugins/Runtime/RigVM/Source/RigVMDeveloper/Public/RigVMModel/RigVMController.h
+++ b/Engine/Plugins/Runtime/RigVM/Source/RigVMDeveloper/Public/RigVMModel/RigVMController.h
@@ -1355,7 +1355,10 @@ public:
 	FRigVMTemplate::FTypeMap GetCommonlyUsedTypesForTemplate(const URigVMTemplateNode* InNode) const;
 #endif
 
+	UFUNCTION(BlueprintPure, Category = RigVMController)
 	URigVMActionStack* GetActionStack() const;
+
+	UFUNCTION(BlueprintCallable, Category = RigVMController)
 	void SetActionStack(URigVMActionStack* InActionStack);
 
 	URigVMNode* ConvertRerouteNodeToDispatch(URigVMRerouteNode* InRerouteNode, const FName& InTemplateNotation, bool bSetupUndoRedo, bool bPrintPythonCommand);
