RigVM: Improve access to Function Headers, Data and Variants

#rnx

diff --git a/Engine/Plugins/Runtime/RigVM/Source/RigVM/Private/RigVMCore/RigVMGraphFunctionDefinition.cpp b/Engine/Plugins/Runtime/RigVM/Source/RigVM/Private/RigVMCore/RigVMGraphFunctionDefinition.cpp
--- a/Engine/Plugins/Runtime/RigVM/Source/RigVM/Private/RigVMCore/RigVMGraphFunctionDefinition.cpp
+++ b/Engine/Plugins/Runtime/RigVM/Source/RigVM/Private/RigVMCore/RigVMGraphFunctionDefinition.cpp
@@ -47,6 +47,46 @@ bool FRigVMGraphFunctionArgument::IsExecuteContext() const
 	return false;
 }
 
+TFunction<TArray<FRigVMVariantRef>(const FGuid& InGuid)> FRigVMGraphFunctionIdentifier::GetVariantRefsByGuidFunc;
+
+bool FRigVMGraphFunctionIdentifier::IsVariant() const
+{
+	return !GetVariants(false).IsEmpty();
+}
+
+TArray<FRigVMVariantRef> FRigVMGraphFunctionIdentifier::GetVariants(bool bIncludeSelf) const
+{
+	if(GetVariantRefsByGuidFunc)
+	{
+		const FRigVMGraphFunctionHeader& ThisHeader = FRigVMGraphFunctionHeader::FindGraphFunctionHeader(*this);
+		TArray<FRigVMVariantRef> Result = GetVariantRefsByGuidFunc(ThisHeader.Variant.Guid);
+		if(!bIncludeSelf)
+		{
+			Result.RemoveAll([this](const FRigVMVariantRef& VariantRef) -> bool
+			{
+				return VariantRef.ObjectPath == GetNodeSoftPath();
+			});
+		}
+		return Result;
+	}
+	return TArray<FRigVMVariantRef>();
+}
+
+TArray<FRigVMGraphFunctionIdentifier> FRigVMGraphFunctionIdentifier::GetVariantIdentifiers(bool bIncludeSelf) const
+{
+	const TArray<FRigVMVariantRef> VariantRefs = GetVariants(bIncludeSelf);
+	TArray<FRigVMGraphFunctionIdentifier> Identifiers;
+	for(const FRigVMVariantRef& VariantRef : VariantRefs)
+	{
+		const FRigVMGraphFunctionHeader Header = FRigVMGraphFunctionHeader::FindGraphFunctionHeader(VariantRef.ObjectPath.ToString());
+		if(Header.IsValid())
+		{
+			Identifiers.Add(Header.LibraryPointer);
+		}
+	}
+	return Identifiers;
+}
+
 bool FRigVMGraphFunctionIdentifier::IsVariantOf(const FRigVMGraphFunctionIdentifier& InOther) const
 {
 	const FRigVMGraphFunctionHeader& ThisHeader = FRigVMGraphFunctionHeader::FindGraphFunctionHeader(*this);
@@ -54,6 +94,8 @@ bool FRigVMGraphFunctionIdentifier::IsVariantOf(const FRigVMGraphFunctionIdentif
 	return ThisHeader.Variant.Guid == OtherHeader.Variant.Guid;
 }
 
+TFunction<FRigVMGraphFunctionHeader(const FSoftObjectPath&, const FName&, bool*)> FRigVMGraphFunctionHeader::FindFunctionHeaderFromPathFunc;
+
 bool FRigVMGraphFunctionHeader::IsMutable() const
 {
 	for(const FRigVMGraphFunctionArgument& Arg : Arguments)
@@ -128,9 +170,33 @@ void FRigVMGraphFunctionHeader::PostDuplicateHost(const FString& InOldPathName,
 	}
 }
 
-FRigVMGraphFunctionHeader FRigVMGraphFunctionHeader::FindGraphFunctionHeader(const FString& InHostPath, const FName& InFunctionName, bool* bOutIsPublic, FString* OutErrorMessage)
+FRigVMGraphFunctionHeader FRigVMGraphFunctionHeader::FindGraphFunctionHeader(const FSoftObjectPath& InFunctionObjectPath, bool* bOutIsPublic, FString* OutErrorMessage)
+{
+	return FindGraphFunctionHeader(InFunctionObjectPath, NAME_None, bOutIsPublic, OutErrorMessage);
+}
+
+FRigVMGraphFunctionHeader FRigVMGraphFunctionHeader::FindGraphFunctionHeader(const FSoftObjectPath& InHostObjectPath, const FName& InFunctionName, bool* bOutIsPublic, FString* OutErrorMessage)
 {
-	if(const FRigVMGraphFunctionData* FunctionData = FRigVMGraphFunctionData::FindFunctionData(InHostPath, InFunctionName, bOutIsPublic, OutErrorMessage))
+	const FName FunctionName = GetFunctionNameFromObjectPath(InHostObjectPath.ToString(), InFunctionName);
+	if(FunctionName.IsNone())
+	{
+		return FRigVMGraphFunctionHeader();
+	}
+	
+	if(FindFunctionHeaderFromPathFunc)
+	{
+		if(InHostObjectPath.ResolveObject() == nullptr)
+		{
+			const FRigVMGraphFunctionHeader Header = FindFunctionHeaderFromPathFunc(InHostObjectPath, FunctionName, bOutIsPublic);
+			if(Header.IsValid())
+			{
+				return Header;
+			}
+		}
+	}
+	
+	// relay to the loaded function since the hostpath is loaded
+	if(const FRigVMGraphFunctionData* FunctionData = FRigVMGraphFunctionData::FindFunctionData(InHostObjectPath, FunctionName, bOutIsPublic, OutErrorMessage))
 	{
 		return FunctionData->Header;
 	}
@@ -139,6 +205,19 @@ FRigVMGraphFunctionHeader FRigVMGraphFunctionHeader::FindGraphFunctionHeader(con
 
 FRigVMGraphFunctionHeader FRigVMGraphFunctionHeader::FindGraphFunctionHeader(const FRigVMGraphFunctionIdentifier& InIdentifier, bool* bOutIsPublic, FString* OutErrorMessage)
 {
+	if(FindFunctionHeaderFromPathFunc)
+	{
+		if(InIdentifier.HostObject.ResolveObject() == nullptr)
+		{
+			const FName FunctionName = GetFunctionNameFromObjectPath(InIdentifier.GetLibraryNodePath());
+			const FRigVMGraphFunctionHeader Header = FindFunctionHeaderFromPathFunc(InIdentifier.GetLibraryNodePath(), FunctionName, bOutIsPublic);
+			if(Header.IsValid())
+			{
+				return Header;
+			}
+		}
+	}
+
 	if(const FRigVMGraphFunctionData* FunctionData = FRigVMGraphFunctionData::FindFunctionData(InIdentifier, bOutIsPublic, OutErrorMessage))
 	{
 		return FunctionData->Header;
@@ -146,6 +225,28 @@ FRigVMGraphFunctionHeader FRigVMGraphFunctionHeader::FindGraphFunctionHeader(con
 	return FRigVMGraphFunctionHeader(); 
 }
 
+FName FRigVMGraphFunctionHeader::GetFunctionNameFromObjectPath(const FString& InObjectPath, const FName& InOptionalFunctionName)
+{
+	FName FunctionName = InOptionalFunctionName;
+	if(FunctionName.IsNone())
+	{
+		FString FunctionNameStr;
+		if(!InObjectPath.Split(TEXT("."), nullptr, &FunctionNameStr, ESearchCase::IgnoreCase, ESearchDir::FromEnd))
+		{
+			(void)InObjectPath.Split(TEXT("/"), nullptr, &FunctionNameStr, ESearchCase::IgnoreCase, ESearchDir::FromEnd);
+		}
+
+		if(FunctionNameStr.IsEmpty())
+		{
+			return NAME_None;
+		}
+		FunctionName = *FunctionNameStr;
+	}
+	return FunctionName;
+}
+
+TFunction<IRigVMGraphFunctionHost*(UObject*)> FRigVMGraphFunctionData::GetFunctionHostFromObjectFunc;
+
 bool FRigVMGraphFunctionData::IsMutable() const
 {
 	return Header.IsMutable();
@@ -156,21 +257,31 @@ void FRigVMGraphFunctionData::PostDuplicateHost(const FString& InOldHostPathName
 	Header.PostDuplicateHost(InOldHostPathName, InNewHostPathName);
 }
 
-FRigVMGraphFunctionData* FRigVMGraphFunctionData::FindFunctionData(const FString& InHostPath, const FName& InFunctionName, bool* bOutIsPublic, FString* OutErrorMessage)
+FRigVMGraphFunctionData* FRigVMGraphFunctionData::FindFunctionData(const FSoftObjectPath& InHostObjectPath, const FName& InFunctionName, bool* bOutIsPublic, FString* OutErrorMessage)
 {
 	FRigVMGraphFunctionHeader InvalidHeader;
 
-	UObject* HostObject = StaticLoadObject(UObject::StaticClass(), NULL, *InHostPath, NULL, LOAD_None, NULL);
+	const FName FunctionName = FRigVMGraphFunctionHeader::GetFunctionNameFromObjectPath(InHostObjectPath.ToString(), InFunctionName);
+
+	UObject* HostObject = InHostObjectPath.TryLoad();
 	if (!HostObject)
 	{
 		if(OutErrorMessage)
 		{
-			*OutErrorMessage = FString::Printf(TEXT("Failed to load the Host object %s."), *InHostPath);
+			*OutErrorMessage = FString::Printf(TEXT("Failed to load the Host object %s."), *InHostObjectPath.ToString());
 		}
 		return nullptr;
 	}
 
 	IRigVMGraphFunctionHost* FunctionHost = Cast<IRigVMGraphFunctionHost>(HostObject);
+	if (!FunctionHost)
+	{
+		if(GetFunctionHostFromObjectFunc)
+		{
+			FunctionHost = GetFunctionHostFromObjectFunc(HostObject);
+		}
+	}
+
 	if (!FunctionHost)
 	{
 		if(OutErrorMessage)
@@ -190,12 +301,12 @@ FRigVMGraphFunctionData* FRigVMGraphFunctionData::FindFunctionData(const FString
 		return nullptr;
 	}
 
-	FRigVMGraphFunctionData* Data = FunctionStore->FindFunctionByName(InFunctionName, bOutIsPublic);
+	FRigVMGraphFunctionData* Data = FunctionStore->FindFunctionByName(FunctionName, bOutIsPublic);
 	if (!Data)
 	{
 		if(OutErrorMessage)
 		{
-			*OutErrorMessage = FString::Printf(TEXT("Function %s not found in host %s."), *InFunctionName.ToString(), *InHostPath);
+			*OutErrorMessage = FString::Printf(TEXT("Function %s not found in host %s."), *FunctionName.ToString(), *InHostObjectPath.ToString());
 		}
 		return nullptr;
 	}
diff --git a/Engine/Plugins/Runtime/RigVM/Source/RigVM/Public/RigVMCore/RigVMGraphFunctionDefinition.h b/Engine/Plugins/Runtime/RigVM/Source/RigVM/Public/RigVMCore/RigVMGraphFunctionDefinition.h
--- a/Engine/Plugins/Runtime/RigVM/Source/RigVM/Public/RigVMCore/RigVMGraphFunctionDefinition.h
+++ b/Engine/Plugins/Runtime/RigVM/Source/RigVM/Public/RigVMCore/RigVMGraphFunctionDefinition.h
@@ -474,7 +474,16 @@ public:
 		return Ar;
 	}
 
+	bool IsVariant() const;
+	TArray<FRigVMVariantRef> GetVariants(bool bIncludeSelf = false) const;
+	TArray<FRigVMGraphFunctionIdentifier> GetVariantIdentifiers(bool bIncludeSelf = false) const;
 	bool IsVariantOf(const FRigVMGraphFunctionIdentifier& InOther) const;
+
+protected:
+	
+	static TFunction<TArray<FRigVMVariantRef>(const FGuid& InGuid)> GetVariantRefsByGuidFunc;
+
+	friend class URigVMBuildData;
 };
 
 USTRUCT(BlueprintType)
@@ -596,9 +605,20 @@ struct RIGVM_API FRigVMGraphFunctionHeader
 
 	void PostDuplicateHost(const FString& InOldPathName, const FString& InNewPathName);
 
-	static FRigVMGraphFunctionHeader FindGraphFunctionHeader(const FString& InHostPath, const FName& InFunctionName, bool* bOutIsPublic = nullptr, FString* OutErrorMessage = nullptr);
+	static FRigVMGraphFunctionHeader FindGraphFunctionHeader(const FSoftObjectPath& InFunctionObjectPath, bool* bOutIsPublic = nullptr, FString* OutErrorMessage = nullptr);
+
+	static FRigVMGraphFunctionHeader FindGraphFunctionHeader(const FSoftObjectPath& InHostObjectPath, const FName& InFunctionName, bool* bOutIsPublic = nullptr, FString* OutErrorMessage = nullptr);
 
 	static FRigVMGraphFunctionHeader FindGraphFunctionHeader(const FRigVMGraphFunctionIdentifier& InIdentifier, bool* bOutIsPublic = nullptr, FString* OutErrorMessage = nullptr);
+
+protected:
+
+	static FName GetFunctionNameFromObjectPath(const FString& InObjectPath, const FName& InOptionalFunctionName = NAME_None);
+	
+	static TFunction<FRigVMGraphFunctionHeader(const FSoftObjectPath&, const FName&, bool*)> FindFunctionHeaderFromPathFunc;
+
+	friend class URigVMBuildData;
+	friend struct FRigVMGraphFunctionData;
 };
 
 USTRUCT(BlueprintType)
@@ -645,7 +665,7 @@ struct RIGVM_API FRigVMGraphFunctionData
 
 	void PostDuplicateHost(const FString& InOldPathName, const FString& InNewPathName);
 
-	static FRigVMGraphFunctionData* FindFunctionData(const FString& InHostPath, const FName& InFunctionName, bool* bOutIsPublic = nullptr, FString* OutErrorMessage = nullptr);	
+	static FRigVMGraphFunctionData* FindFunctionData(const FSoftObjectPath& InHostObjectPath, const FName& InFunctionName, bool* bOutIsPublic = nullptr, FString* OutErrorMessage = nullptr);	
 
 	static FRigVMGraphFunctionData* FindFunctionData(const FRigVMGraphFunctionIdentifier& InIdentifier, bool* bOutIsPublic = nullptr, FString* OutErrorMessage = nullptr);	
 
@@ -659,4 +679,7 @@ struct RIGVM_API FRigVMGraphFunctionData
 
 	static const inline TCHAR* EntryString = TEXT("Entry");
 	static const inline TCHAR* ReturnString = TEXT("Return");
+	static TFunction<IRigVMGraphFunctionHost*(UObject*)> GetFunctionHostFromObjectFunc;
+
+	friend class URigVMBuildData;
 };
diff --git a/Engine/Plugins/Runtime/RigVM/Source/RigVMDeveloper/Private/RigVMDeveloperModule.cpp b/Engine/Plugins/Runtime/RigVM/Source/RigVMDeveloper/Private/RigVMDeveloperModule.cpp
--- a/Engine/Plugins/Runtime/RigVM/Source/RigVMDeveloper/Private/RigVMDeveloperModule.cpp
+++ b/Engine/Plugins/Runtime/RigVM/Source/RigVMDeveloper/Private/RigVMDeveloperModule.cpp
@@ -7,6 +7,7 @@
 #include "RigVMDeveloperModule.h"
 #include "RigVMBlueprint.h"
 #include "RigVMBlueprintCompiler.h"
+#include "RigVMModel/RigVMBuildData.h"
 
 DEFINE_LOG_CATEGORY(LogRigVMDeveloper);
 
@@ -29,6 +30,7 @@ void FRigVMDeveloperModule::StartupModule()
 	FKismetCompilerContext::RegisterCompilerForBP(URigVMBlueprint::StaticClass(), &FRigVMDeveloperModule::GetRigVMCompiler);
 	IKismetCompilerInterface& KismetCompilerModule = FModuleManager::LoadModuleChecked<IKismetCompilerInterface>("KismetCompiler");
 	KismetCompilerModule.GetCompilers().Add(&RigVMBlueprintCompiler);
+	URigVMBuildData::Get()->SetupRigVMGraphFunctionPointers();
 }
 
 void FRigVMDeveloperModule::ShutdownModule()
@@ -38,6 +40,7 @@ void FRigVMDeveloperModule::ShutdownModule()
 	{
 		KismetCompilerModule->GetCompilers().Remove(&RigVMBlueprintCompiler);
 	}
+	URigVMBuildData::Get()->TearDownRigVMGraphFunctionPointers();
 }
 
 TSharedPtr<FKismetCompilerContext> FRigVMDeveloperModule::GetRigVMCompiler(UBlueprint* BP, FCompilerResultsLog& InMessageLog, const FKismetCompilerOptions& InCompileOptions)
diff --git a/Engine/Plugins/Runtime/RigVM/Source/RigVMDeveloper/Private/RigVMModel/RigVMBuildData.cpp b/Engine/Plugins/Runtime/RigVM/Source/RigVMDeveloper/Private/RigVMModel/RigVMBuildData.cpp
--- a/Engine/Plugins/Runtime/RigVM/Source/RigVMDeveloper/Private/RigVMModel/RigVMBuildData.cpp
+++ b/Engine/Plugins/Runtime/RigVM/Source/RigVMDeveloper/Private/RigVMModel/RigVMBuildData.cpp
@@ -8,6 +8,7 @@
 #include "UObject/UObjectIterator.h"
 #include "RigVMModel/RigVMClient.h"
 #include "UObject/StrongObjectPtr.h"
+#include "RigVMBlueprint.h"
 
 #include UE_INLINE_GENERATED_CPP_BY_NAME(RigVMBuildData)
 
@@ -61,6 +62,149 @@
 {
 }
 
+TArray<UClass*> URigVMBuildData::FindAllRigVMAssetClasses()
+{
+	// Find all classes which implement IRigVMClientHost
+	TArray<UClass*> ImplementedClasses;
+	for (TObjectIterator<UClass> ClassIterator; ClassIterator; ++ClassIterator)
+	{
+		if (ClassIterator->ImplementsInterface(URigVMClientHost::StaticClass()))
+		{
+			ImplementedClasses.AddUnique(*ClassIterator);
+
+			const UObject* CDO = ClassIterator->GetDefaultObject();
+			if(const URigVMBlueprint* BlueprintCDO = Cast<URigVMBlueprint>(CDO))
+			{
+				if(UClass* GeneratedClass = BlueprintCDO->GetRigVMBlueprintGeneratedClass())
+				{
+					ImplementedClasses.AddUnique(GeneratedClass);
+				}
+			}
+		}
+	}
+	return ImplementedClasses;
+}
+
+void URigVMBuildData::SetupRigVMGraphFunctionPointers()
+{
+	FRigVMGraphFunctionIdentifier::GetVariantRefsByGuidFunc = [this](const FGuid& InGuid)
+	{
+		return FindFunctionVariantRefs(InGuid);
+	};
+	FRigVMGraphFunctionHeader::FindFunctionHeaderFromPathFunc = [this](const FSoftObjectPath& InObjectPath, const FName& InFunctionName, bool *bOutIsPublic)
+	{
+		const FAssetRegistryModule& AssetRegistryModule = FModuleManager::LoadModuleChecked<FAssetRegistryModule>("AssetRegistry");
+		TArray<FAssetData> Assets;
+		AssetRegistryModule.GetRegistry().GetAssetsByPath(*InObjectPath.ToString(), Assets, true);
+		if(!Assets.IsEmpty())
+		{
+			const FString FunctionNameString = InFunctionName.ToString();
+			const TArray<FRigVMGraphFunctionHeader> Headers = GetFunctionHeadersForAsset(Assets[0]);
+			for(const FRigVMGraphFunctionHeader& Header : Headers)
+			{
+				if(Header.LibraryPointer.GetFunctionName() == FunctionNameString)
+				{
+					if(bOutIsPublic)
+					{
+						// if this isn't loaded - it means that the asset is not loaded thus the function can only be public
+						*bOutIsPublic = true;
+						
+						const FSoftObjectPath LibraryNodeSoftPath = Header.LibraryPointer.GetNodeSoftPath();
+						if(const URigVMLibraryNode* LibraryNode = Cast<URigVMLibraryNode>(LibraryNodeSoftPath.ResolveObject()))
+						{
+							if(const URigVMFunctionLibrary* FunctionLibrary = LibraryNode->GetTypedOuter<URigVMFunctionLibrary>())
+							{
+								*bOutIsPublic = FunctionLibrary->IsFunctionPublic(Header.LibraryPointer.GetFunctionFName());
+							}
+						}
+					}
+					return Header;
+				}
+			}
+		}
+		if(bOutIsPublic)
+		{
+			*bOutIsPublic = false;
+		}
+		return FRigVMGraphFunctionHeader();
+	};
+	FRigVMGraphFunctionData::GetFunctionHostFromObjectFunc = [this](UObject* InObject) -> IRigVMGraphFunctionHost*
+	{
+		if(IRigVMGraphFunctionHost* Host = InObject->GetImplementingOuter<IRigVMGraphFunctionHost>())
+		{
+			return Host;
+		}
+		if(URigVMBlueprint* Blueprint = Cast<URigVMBlueprint>(InObject))
+		{
+			return Blueprint->GetRigVMGraphFunctionHost();
+		}
+		if(URigVMBlueprint* Blueprint = InObject->GetTypedOuter<URigVMBlueprint>())
+		{
+			return Blueprint->GetRigVMGraphFunctionHost();
+		}
+		return nullptr;
+	};
+}
+
+void URigVMBuildData::TearDownRigVMGraphFunctionPointers()
+{
+	FRigVMGraphFunctionIdentifier::GetVariantRefsByGuidFunc.Reset();
+	FRigVMGraphFunctionHeader::FindFunctionHeaderFromPathFunc.Reset();
+	FRigVMGraphFunctionData::GetFunctionHostFromObjectFunc.Reset();
+}
+
+TArray<FRigVMGraphFunctionHeader> URigVMBuildData::GetFunctionHeadersForAsset(const FAssetData& InAssetData)
+{
+	TArray<FRigVMGraphFunctionHeader> Result;
+	if (!InAssetData.IsInstanceOf(URigVMBlueprint::StaticClass()) &&
+	!InAssetData.IsInstanceOf(URigVMBlueprintGeneratedClass::StaticClass()))
+	{
+		return Result;
+	}
+
+	// If the asset is loaded, gather the function variants from the function store
+	// which will include private functions
+	if (InAssetData.IsAssetLoaded())
+	{
+		UObject* AssetObject = InAssetData.GetAsset();
+		URigVMBlueprintGeneratedClass* GeneratedClass = Cast<URigVMBlueprintGeneratedClass>(AssetObject);
+		if (!GeneratedClass)
+		{
+			if (URigVMBlueprint* RigVMBlueprint = Cast<URigVMBlueprint>(AssetObject))
+			{
+				GeneratedClass = RigVMBlueprint->GetRigVMBlueprintGeneratedClass();
+			}
+		}
+
+		if (GeneratedClass)
+		{
+			for (int32 Pass=0; Pass<2; Pass++)
+			{
+				TArray<FRigVMGraphFunctionData>& Functions = (Pass == 0)
+					? GeneratedClass->GraphFunctionStore.PrivateFunctions
+					: GeneratedClass->GraphFunctionStore.PublicFunctions;
+				
+				for (FRigVMGraphFunctionData& Data : Functions)
+				{
+					Result.Add(Data.Header);
+				}
+			}
+
+			return Result;
+		}
+	}
+
+	// If asset is not loaded, gather function variants from metadata
+	const FString PublicGraphFunctionsString = InAssetData.GetTagValueRef<FString>(TEXT("PublicGraphFunctions"));
+	if(!PublicGraphFunctionsString.IsEmpty())
+	{
+		static const FName FunctionsPropertyName = GET_MEMBER_NAME_CHECKED(URigVMBlueprint, PublicGraphFunctions);
+		const FArrayProperty* PublicGraphFunctionsProperty = CastField<FArrayProperty>(URigVMBlueprint::StaticClass()->FindPropertyByName(FunctionsPropertyName));
+		PublicGraphFunctionsProperty->ImportText_Direct(*PublicGraphFunctionsString, &Result, nullptr, EPropertyPortFlags::PPF_None);
+	}
+	return Result;
+}
+
 URigVMBuildData* URigVMBuildData::Get()
 {
 	// static in a function scope ensures that the GC system is initiated before 
@@ -85,15 +229,7 @@
 		return;
 	}
 
-	// Find all classes which implement IRigVMClientHost
-	TArray<UClass*> ImplementedClasses;
-	for (TObjectIterator<UClass> ClassIterator; ClassIterator; ++ClassIterator)
-	{
-		if (ClassIterator->ImplementsInterface(URigVMClientHost::StaticClass()))
-		{
-			ImplementedClasses.Add(*ClassIterator);
-		}
-	}
+	TArray<UClass*> ImplementedClasses = FindAllRigVMAssetClasses();
 
 	// Loop the classes
 	for (UClass* Class : ImplementedClasses)
@@ -105,15 +241,15 @@
 			const FAssetRegistryModule& AssetRegistryModule = FModuleManager::LoadModuleChecked<FAssetRegistryModule>("AssetRegistry");
 
 			// find all assets of this class in the project
-			TArray<FAssetData> ControlRigAssetDatas;
-			FARFilter ControlRigAssetFilter;
-			ControlRigAssetFilter.ClassPaths.Add(Class->GetClassPathName());
-			AssetRegistryModule.Get().GetAssets(ControlRigAssetFilter, ControlRigAssetDatas);
+			TArray<FAssetData> AssetDatas;
+			FARFilter AssetFilter;
+			AssetFilter.ClassPaths.Add(Class->GetClassPathName());
+			AssetRegistryModule.Get().GetAssets(AssetFilter, AssetDatas);
 
 			// loop over all found assets
-			for(const FAssetData& ControlRigAssetData : ControlRigAssetDatas)
+			for(const FAssetData& AssetData : AssetDatas)
 			{
-				RegisterReferencesFromAsset(ControlRigAssetData);
+				RegisterReferencesFromAsset(AssetData);
 			}
 		}
 	}
@@ -390,60 +526,14 @@
 {
 	TArray<FRigVMVariantRef> Result;
 
-	if (!InAssetData.IsInstanceOf(URigVMBlueprint::StaticClass()) &&
-		!InAssetData.IsInstanceOf(URigVMBlueprintGeneratedClass::StaticClass()))
+	const TArray<FRigVMGraphFunctionHeader> FunctionHeaders = GetFunctionHeadersForAsset(InAssetData);
+	for(const FRigVMGraphFunctionHeader& FunctionHeader : FunctionHeaders)
 	{
-		return Result;
-	}
-
-	// If the asset is loaded, gather the function variants from the function store
-	// which will include private functions
-	if (InAssetData.IsAssetLoaded())
-	{
-		UObject* AssetObject = InAssetData.GetAsset();
-		URigVMBlueprintGeneratedClass* GeneratedClass = Cast<URigVMBlueprintGeneratedClass>(AssetObject);
-		if (!GeneratedClass)
+		if (FunctionHeader.LibraryPointer.IsValid())
 		{
-			if (URigVMBlueprint* RigVMBlueprint = Cast<URigVMBlueprint>(AssetObject))
-			{
-				GeneratedClass = RigVMBlueprint->GetRigVMBlueprintGeneratedClass();
-			}
-		}
-
-		if (GeneratedClass)
-		{
-			for (int32 Pass=0; Pass<2; Pass++)
-			{
-				TArray<FRigVMGraphFunctionData>& Functions = (Pass == 0)
-					? GeneratedClass->GraphFunctionStore.PrivateFunctions
-					: GeneratedClass->GraphFunctionStore.PublicFunctions;
-				
-				for (FRigVMGraphFunctionData& Data : Functions)
-				{
-					FRigVMVariantRef& VariantRef = Result.Add_GetRef(FRigVMVariantRef());
-					VariantRef.Variant = Data.Header.Variant;
-					VariantRef.ObjectPath = Data.Header.LibraryPointer.GetNodeSoftPath();
-				}
-			}
-
-			return Result;
-		}
-	}
-
-	// If asset is not loaded, gather function variants from metadata
-	FString PublicGraphFunctionsString = InAssetData.GetTagValueRef<FString>(TEXT("PublicGraphFunctions"));
-	TArray<FRigVMGraphFunctionHeader> PublicFunctions;
-	static const FName FunctionsPropertyName = GET_MEMBER_NAME_CHECKED(URigVMBlueprint, PublicGraphFunctions);
-	FArrayProperty* PublicGraphFunctionsProperty = CastField<FArrayProperty>(URigVMBlueprint::StaticClass()->FindPropertyByName(FunctionsPropertyName));
-	PublicGraphFunctionsProperty->ImportText_Direct(*PublicGraphFunctionsString, &PublicFunctions, nullptr, EPropertyPortFlags::PPF_None);
-
-	for(FRigVMGraphFunctionHeader& PublicFunction : PublicFunctions)
-	{
-		if (PublicFunction.LibraryPointer.IsValid())
-		{
 			FRigVMVariantRef& VariantRef = Result.Add_GetRef(FRigVMVariantRef());
-			VariantRef.Variant = PublicFunction.Variant;
-			VariantRef.ObjectPath = PublicFunction.LibraryPointer.GetNodeSoftPath();
+			VariantRef.Variant = FunctionHeader.Variant;
+			VariantRef.ObjectPath = FunctionHeader.LibraryPointer.GetNodeSoftPath();
 			if (!VariantRef.Variant.Guid.IsValid())
 			{
 				VariantRef.Variant.Guid = FRigVMVariant::GenerateGUID(VariantRef.ObjectPath.ToString());
@@ -461,4 +551,128 @@
 	return Result;
 }
 
+#if WITH_EDITOR
 
+TArray<FRigVMGraphFunctionIdentifier> URigVMBuildData::GetAllFunctionIdentifiers(bool bOnlyPublic) const
+{
+	static const FName PublicGraphFunctionsName = TEXT("PublicGraphFunctions");
+	TArray<FRigVMGraphFunctionIdentifier> Identifiers;
+	TArray<UClass*> ImplementedClasses = FindAllRigVMAssetClasses();
+
+	const FAssetRegistryModule& AssetRegistryModule = FModuleManager::LoadModuleChecked<FAssetRegistryModule>("AssetRegistry");
+
+	// find all assets of this class in the project
+	TArray<FAssetData> AssetDatas;
+	FARFilter AssetFilter;
+	for(const UClass* Class : ImplementedClasses)
+	{
+		AssetFilter.ClassPaths.AddUnique(Class->GetClassPathName());
+	}
+	AssetRegistryModule.Get().GetAssets(AssetFilter, AssetDatas);
+
+	// loop over all found assets
+	TSet<FName> PackagesProcessed;
+	for(const FAssetData& AssetData : AssetDatas)
+	{
+		// Avoid duplication
+		if (PackagesProcessed.Contains(AssetData.PackageName))
+		{
+			continue;
+		}
+		PackagesProcessed.Add(AssetData.PackageName);
+
+		if(AssetData.IsAssetLoaded())
+		{
+			IRigVMGraphFunctionHost* FunctionHost = Cast<IRigVMGraphFunctionHost>(AssetData.GetAsset());
+			if(FunctionHost == nullptr)
+			{
+				if(IRigVMClientHost* ClientHost = Cast<IRigVMClientHost>(AssetData.GetAsset()))
+				{
+					FunctionHost = ClientHost->GetRigVMGraphFunctionHost();
+				}
+			}
+			if(FunctionHost)
+			{
+				if(const FRigVMGraphFunctionStore* Store = FunctionHost->GetRigVMGraphFunctionStore())
+				{
+					for(const FRigVMGraphFunctionData& Data : Store->PublicFunctions)
+					{
+						Identifiers.Add(Data.Header.LibraryPointer);
+					}
+					if(!bOnlyPublic)
+					{
+						for(const FRigVMGraphFunctionData& Data : Store->PrivateFunctions)
+						{
+							Identifiers.Add(Data.Header.LibraryPointer);
+						}
+					}
+					continue;
+				}
+			}
+		}
+
+		FString PublicGraphFunctionsString = AssetData.GetTagValueRef<FString>(PublicGraphFunctionsName);
+		if (PublicGraphFunctionsString.IsEmpty())
+		{
+			PublicGraphFunctionsString = AssetData.GetTagValueRef<FString>(PublicGraphFunctionsName);
+		}
+		if (!PublicGraphFunctionsString.IsEmpty())
+		{
+			const FArrayProperty* Property = CastField<FArrayProperty>(AssetData.GetClass()->FindPropertyByName(PublicGraphFunctionsName));
+			if (Property)
+			{
+				TArray<FRigVMGraphFunctionHeader> PublicFunctions;
+				Property->ImportText_Direct(*PublicGraphFunctionsString, &PublicFunctions, nullptr, EPropertyPortFlags::PPF_None);
+				for(const FRigVMGraphFunctionHeader& PublicFunction : PublicFunctions)
+				{
+					Identifiers.Add(PublicFunction.LibraryPointer);
+				}
+			}
+			else
+			{
+				const FString& HeadersString = PublicGraphFunctionsString;
+				const FArrayProperty* HeadersArrayProperty = CastField<FArrayProperty>(FRigVMGraphFunctionHeaderArray::StaticStruct()->FindPropertyByName(TEXT("Headers")));
+				TArray<FRigVMGraphFunctionHeader> PublicFunctions;
+				HeadersArrayProperty->ImportText_Direct(*HeadersString, &PublicFunctions, nullptr, EPropertyPortFlags::PPF_None);
+		
+				for(const FRigVMGraphFunctionHeader& PublicFunction : PublicFunctions)
+				{
+					Identifiers.Add(PublicFunction.LibraryPointer);
+				}
+			}
+		}
+	}
+
+	return Identifiers;
+}
+
+#endif
+
+TArray<FRigVMGraphFunctionIdentifier> URigVMBuildData::GetUsedFunctionIdentifiers(bool bOnlyPublic) const
+{
+	TArray<FRigVMGraphFunctionIdentifier> Identifiers;
+#if WITH_EDITOR
+	Identifiers = GetAllFunctionIdentifiers().FilterByPredicate([this](const FRigVMGraphFunctionIdentifier& Identifier)
+	{
+		const FRigVMFunctionReferenceArray* References = FindFunctionReferences(Identifier); 
+		return References != nullptr && References->Num() > 0;
+	});
+#else
+	// outside of editor we can't do much - just get the keys
+	// of the used references array
+	GraphFunctionReferences.GenerateKeyArray(Identifiers);
+#endif
+	return Identifiers;
+}
+
+FRigVMFunctionReferenceArray URigVMBuildData::GetAllFunctionReferences() const
+{
+	FRigVMFunctionReferenceArray FunctionReferences;
+	for (const TTuple<FRigVMGraphFunctionIdentifier, FRigVMFunctionReferenceArray>& FunctionReferenceInfo : GraphFunctionReferences)
+	{
+		FunctionReferences.FunctionReferences.Append(FunctionReferenceInfo.Value.FunctionReferences);
+	}
+	return FunctionReferences;
+}
+
+
diff --git a/Engine/Plugins/Runtime/RigVM/Source/RigVMDeveloper/Private/RigVMModel/RigVMController.cpp b/Engine/Plugins/Runtime/RigVM/Source/RigVMDeveloper/Private/RigVMModel/RigVMController.cpp
--- a/Engine/Plugins/Runtime/RigVM/Source/RigVMDeveloper/Private/RigVMModel/RigVMController.cpp
+++ b/Engine/Plugins/Runtime/RigVM/Source/RigVMDeveloper/Private/RigVMModel/RigVMController.cpp
@@ -12311,7 +12311,7 @@ TArray<FRigVMVariantRef> URigVMController::FindVariantsOfFunction(const FName& I
     	return Result;
     }
 
-	return URigVMBuildData::Get()->FindFunctionVariantRefs(FunctionData->Header.Variant.Guid);
+	return FunctionData->Header.LibraryPointer.GetVariants(true);
 }
 
 FRigVMGraphVariableDescription URigVMController::AddLocalVariable(const FName& InVariableName, const FString& InCPPType, UObject* InCPPTypeObject, const FString& InDefaultValue, bool bSetupUndoRedo, bool bPrintPythonCommand)
diff --git a/Engine/Plugins/Runtime/RigVM/Source/RigVMDeveloper/Public/RigVMModel/RigVMBuildData.h b/Engine/Plugins/Runtime/RigVM/Source/RigVMDeveloper/Public/RigVMModel/RigVMBuildData.h
--- a/Engine/Plugins/Runtime/RigVM/Source/RigVMDeveloper/Public/RigVMModel/RigVMBuildData.h
+++ b/Engine/Plugins/Runtime/RigVM/Source/RigVMDeveloper/Public/RigVMModel/RigVMBuildData.h
@@ -128,10 +128,27 @@ public:
 	// Will find all the function variants matching the given variant guid
 	TArray<FRigVMVariantRef> FindFunctionVariantRefs(const FGuid& InGuid);
 
+#if WITH_EDITOR
+	// Returns all known public function identifiers used in the project
+	TArray<FRigVMGraphFunctionIdentifier> GetAllFunctionIdentifiers(bool bOnlyPublic = true) const;
+#endif
+
+	// Returns all known public function identifiers used in the project
+	TArray<FRigVMGraphFunctionIdentifier> GetUsedFunctionIdentifiers(bool bOnlyPublic = true) const;
+
+	// Returns all known function references
+	FRigVMFunctionReferenceArray GetAllFunctionReferences() const;
+
 private:
 
 	// disable default constructor
 	URigVMBuildData();
+
+	static TArray<UClass*> FindAllRigVMAssetClasses();
+
+	void SetupRigVMGraphFunctionPointers();
+	void TearDownRigVMGraphFunctionPointers();
+	static TArray<FRigVMGraphFunctionHeader> GetFunctionHeadersForAsset(const FAssetData& InAssetData);
 	
 	static bool bInitialized;
 
@@ -146,5 +163,6 @@ private:
 	friend class URigVMController;
 	friend struct FRigVMClient;
 	friend class URigVMCompiler;
+	friend class FRigVMDeveloperModule;
 };
 
