Control Rig: Fix duplicate and copy/paste of rigvmblueprints not updating the functions and references correctly
#rnx

diff --git a/Engine/Plugins/Runtime/RigVM/Source/RigVM/Private/RigVMCore/RigVMGraphFunctionDefinition.cpp b/Engine/Plugins/Runtime/RigVM/Source/RigVM/Private/RigVMCore/RigVMGraphFunctionDefinition.cpp
--- a/Engine/Plugins/Runtime/RigVM/Source/RigVM/Private/RigVMCore/RigVMGraphFunctionDefinition.cpp
+++ b/Engine/Plugins/Runtime/RigVM/Source/RigVM/Private/RigVMCore/RigVMGraphFunctionDefinition.cpp
@@ -137,39 +137,6 @@ FRigVMGraphFunctionData* FRigVMGraphFunctionHeader::GetFunctionData(bool bLoadIf
 	return nullptr;
 }
 
-void FRigVMGraphFunctionHeader::PostDuplicateHost(const FString& InOldPathName, const FString& InNewPathName)
-{
-	const FString OldPathName = InOldPathName + TEXT(":");
-	const FString NewPathName = InNewPathName + TEXT(":");
-
-	auto ReplacePathName = [InOldPathName, InNewPathName, OldPathName, NewPathName](FString& InOutObjectPath)
-	{
-		if(InOutObjectPath.Equals(InOldPathName, ESearchCase::CaseSensitive))
-		{
-			InOutObjectPath = InNewPathName;
-		}
-		else if(InOutObjectPath.StartsWith(OldPathName, ESearchCase::CaseSensitive))
-		{
-			InOutObjectPath = NewPathName + InOutObjectPath.Mid(OldPathName.Len());
-		}
-	};
-	
-	auto ReplaceSoftPathName = [InOldPathName, InNewPathName, OldPathName, NewPathName, ReplacePathName](FSoftObjectPath& InOutObjectPath)
-	{
-		FString PathName = InOutObjectPath.ToString();
-		ReplacePathName(PathName);
-		InOutObjectPath = FSoftObjectPath(PathName);
-	};
-
-	ReplacePathName(LibraryPointer.GetLibraryNodePath());
-	ReplaceSoftPathName(LibraryPointer.HostObject);
-	for (TPair<FRigVMGraphFunctionIdentifier, uint32>& Pair : Dependencies)
-	{
-		ReplacePathName(Pair.Key.GetLibraryNodePath());
-		ReplaceSoftPathName(Pair.Key.HostObject);
-	}
-}
-
 FRigVMGraphFunctionHeader FRigVMGraphFunctionHeader::FindGraphFunctionHeader(const FSoftObjectPath& InFunctionObjectPath, bool* bOutIsPublic, FString* OutErrorMessage)
 {
 	return FindGraphFunctionHeader(InFunctionObjectPath, NAME_None, bOutIsPublic, OutErrorMessage);
@@ -252,11 +219,6 @@ bool FRigVMGraphFunctionData::IsMutable() const
 	return Header.IsMutable();
 }
 
-void FRigVMGraphFunctionData::PostDuplicateHost(const FString& InOldHostPathName, const FString& InNewHostPathName)
-{
-	Header.PostDuplicateHost(InOldHostPathName, InNewHostPathName);
-}
-
 FRigVMGraphFunctionData* FRigVMGraphFunctionData::FindFunctionData(const FSoftObjectPath& InHostObjectPath, const FName& InFunctionName, bool* bOutIsPublic, FString* OutErrorMessage)
 {
 	FRigVMGraphFunctionHeader InvalidHeader;
diff --git a/Engine/Plugins/Runtime/RigVM/Source/RigVM/Private/RigVMCore/RigVMGraphFunctionHost.cpp b/Engine/Plugins/Runtime/RigVM/Source/RigVM/Private/RigVMCore/RigVMGraphFunctionHost.cpp
--- a/Engine/Plugins/Runtime/RigVM/Source/RigVM/Private/RigVMCore/RigVMGraphFunctionHost.cpp
+++ b/Engine/Plugins/Runtime/RigVM/Source/RigVM/Private/RigVMCore/RigVMGraphFunctionHost.cpp
@@ -308,16 +308,3 @@ bool FRigVMGraphFunctionStore::RemoveAllCompilationData()
 
 	return true;
 }
-
-void FRigVMGraphFunctionStore::PostDuplicateHost(const FString& InOldPathName, const FString& InNewPathName)
-{
-	for (FRigVMGraphFunctionData& Data : PublicFunctions)
-	{
-		Data.Header.PostDuplicateHost(InOldPathName, InNewPathName);
-	}
-
-	for (FRigVMGraphFunctionData& Data : PrivateFunctions)
-	{
-		Data.Header.PostDuplicateHost(InOldPathName, InNewPathName);
-	}
-}
diff --git a/Engine/Plugins/Runtime/RigVM/Source/RigVM/Public/RigVMCore/RigVMGraphFunctionDefinition.h b/Engine/Plugins/Runtime/RigVM/Source/RigVM/Public/RigVMCore/RigVMGraphFunctionDefinition.h
--- a/Engine/Plugins/Runtime/RigVM/Source/RigVM/Public/RigVMCore/RigVMGraphFunctionDefinition.h
+++ b/Engine/Plugins/Runtime/RigVM/Source/RigVM/Public/RigVMCore/RigVMGraphFunctionDefinition.h
@@ -603,8 +603,6 @@ struct RIGVM_API FRigVMGraphFunctionHeader
 		return Ar;
 	}
 
-	void PostDuplicateHost(const FString& InOldPathName, const FString& InNewPathName);
-
 	static FRigVMGraphFunctionHeader FindGraphFunctionHeader(const FSoftObjectPath& InFunctionObjectPath, bool* bOutIsPublic = nullptr, FString* OutErrorMessage = nullptr);
 
 	static FRigVMGraphFunctionHeader FindGraphFunctionHeader(const FSoftObjectPath& InHostObjectPath, const FName& InFunctionName, bool* bOutIsPublic = nullptr, FString* OutErrorMessage = nullptr);
@@ -663,8 +661,6 @@ struct RIGVM_API FRigVMGraphFunctionData
 		return Ar;
 	}
 
-	void PostDuplicateHost(const FString& InOldPathName, const FString& InNewPathName);
-
 	static FRigVMGraphFunctionData* FindFunctionData(const FSoftObjectPath& InHostObjectPath, const FName& InFunctionName, bool* bOutIsPublic = nullptr, FString* OutErrorMessage = nullptr);	
 
 	static FRigVMGraphFunctionData* FindFunctionData(const FRigVMGraphFunctionIdentifier& InIdentifier, bool* bOutIsPublic = nullptr, FString* OutErrorMessage = nullptr);	
diff --git a/Engine/Plugins/Runtime/RigVM/Source/RigVM/Public/RigVMCore/RigVMGraphFunctionHost.h b/Engine/Plugins/Runtime/RigVM/Source/RigVM/Public/RigVMCore/RigVMGraphFunctionHost.h
--- a/Engine/Plugins/Runtime/RigVM/Source/RigVM/Public/RigVMCore/RigVMGraphFunctionHost.h
+++ b/Engine/Plugins/Runtime/RigVM/Source/RigVM/Public/RigVMCore/RigVMGraphFunctionHost.h
@@ -68,8 +68,6 @@ public:
 
 	bool RemoveAllCompilationData();
 
-	void PostDuplicateHost(const FString& InOldPathName, const FString& InNewPathName);
-
 	friend FArchive& operator<<(FArchive& Ar, FRigVMGraphFunctionStore& Host)
 	{
 		Ar << Host.PublicFunctions;
diff --git a/Engine/Plugins/Runtime/RigVM/Source/RigVMDeveloper/Private/RigVMBlueprint.cpp b/Engine/Plugins/Runtime/RigVM/Source/RigVMDeveloper/Private/RigVMBlueprint.cpp
--- a/Engine/Plugins/Runtime/RigVM/Source/RigVMDeveloper/Private/RigVMBlueprint.cpp
+++ b/Engine/Plugins/Runtime/RigVM/Source/RigVMDeveloper/Private/RigVMBlueprint.cpp
@@ -383,6 +383,9 @@ void URigVMBlueprint::PostRename(UObject* OldOuter, const FName OldName)
 			MemoryClass->Rename(nullptr, GetPackage(), REN_ForceNoResetLoaders | REN_DoNotDirty | REN_DontCreateRedirectors | REN_NonTransactional);
 		}
 	}
+
+	const FString OldAssetPath = FString::Printf(TEXT("%s.%s"), *OldOuter->GetPathName(), *OldName.ToString());
+	ReplaceFunctionIdentifiers(OldAssetPath, GetPathName());
 }
 
 void URigVMBlueprint::GetPreloadDependencies(TArray<UObject*>& OutDeps)
@@ -2625,6 +2628,64 @@ void URigVMBlueprint::PreDuplicate(FObjectDuplicationParameters& DupParams)
 
 }
 
+void URigVMBlueprint::ReplaceFunctionIdentifiers(const FString& InOldAssetPath, const FString& InNewAssetPath)
+{
+	if (!InOldAssetPath.Equals(GetPathName()))
+	{
+		const FString OldLibraryPath = InOldAssetPath + TEXT(":");
+		const FString NewLibraryPath = InNewAssetPath + TEXT(":");
+		const FString OldHostPath = InOldAssetPath + TEXT("_C");
+		const FString NewHostPath = InNewAssetPath + TEXT("_C");
+
+		auto ReplaceIdentifier = [OldLibraryPath, NewLibraryPath, OldHostPath, NewHostPath](FRigVMGraphFunctionIdentifier& Identifier)
+		{
+			FString& LibraryNodePath = Identifier.GetLibraryNodePath();
+			FSoftObjectPath& HostPath = Identifier.HostObject;
+			FString HostPathStr = HostPath.ToString();
+			if(LibraryNodePath.StartsWith(OldLibraryPath, ESearchCase::CaseSensitive))
+			{
+				LibraryNodePath = NewLibraryPath + LibraryNodePath.Mid(OldLibraryPath.Len());
+			}
+			if(HostPathStr.StartsWith(OldHostPath, ESearchCase::CaseSensitive))
+			{
+				HostPathStr = NewHostPath + HostPathStr.Mid(OldHostPath.Len());
+				HostPath = HostPathStr;
+			}
+		};
+
+		// Replace identifiers in store
+		if(URigVMBlueprintGeneratedClass* CRGeneratedClass = GetRigVMBlueprintGeneratedClass())
+		{
+			FRigVMGraphFunctionStore& Store = CRGeneratedClass->GraphFunctionStore;
+			for (int32 i=0; i<2; ++i)
+			{
+				TArray<FRigVMGraphFunctionData>& Functions = (i == 0) ? Store.PublicFunctions : Store.PrivateFunctions;
+				for (FRigVMGraphFunctionData& Data : Functions)
+				{
+					ReplaceIdentifier(Data.Header.LibraryPointer);
+					for (TPair<FRigVMGraphFunctionIdentifier, uint32>& Pair : Data.Header.Dependencies)
+					{
+						ReplaceIdentifier(Pair.Key);
+					}
+				}
+			}
+		}
+
+		// Replace identifiers in function references
+		TArray<URigVMGraph*> AllModels = RigVMClient.GetAllModels(true, true);
+		for(URigVMGraph* Model : AllModels)
+		{
+			for (URigVMNode* Node : Model->GetNodes())
+			{
+				if (URigVMFunctionReferenceNode* FunctionReferenceNode = Cast<URigVMFunctionReferenceNode>(Node))
+				{
+					ReplaceIdentifier(FunctionReferenceNode->ReferencedFunctionHeader.LibraryPointer);
+				}
+			}
+		}
+	}
+}
+
 void URigVMBlueprint::PostDuplicate(bool bDuplicateForPIE)
 {
 	// assuming PostDuplicate is always followed by a PostLoad:
@@ -2640,37 +2701,11 @@ void URigVMBlueprint::PostDuplicate(bool bDuplicateForPIE)
 		// it will be filled during PostLoad based on the graph model
 		Super::PostDuplicate(bDuplicateForPIE);
 	}
-	
-	auto UpdateFunctionHeaders = [this](const FString& InOldPath, const FString& InNewPath)
-	{
-		if(InOldPath.IsEmpty() || InNewPath.IsEmpty())
-		{
-			return;
-		}
-		if(!InNewPath.Equals(InOldPath, ESearchCase::CaseSensitive))
-		{
-			if(URigVMBlueprintGeneratedClass* CRGeneratedClass = GetRigVMBlueprintGeneratedClass())
-			{
-				FRigVMGraphFunctionStore& Store = CRGeneratedClass->GraphFunctionStore;
-				// technically not needed, the store should be empty, it will not be populated until PostLoad
-				// this code is kept here just in case things change in the future
-				if (!(ensure(Store.PublicFunctions.Num() == 0) && ensure(Store.PrivateFunctions.Num() == 0)))
-				{
-					Store.PostDuplicateHost(InOldPath, InNewPath);
-				}
-			}
-			RigVMClient.PostDuplicateHost(InOldPath, InNewPath);
-		}
-	};
 
-	// update the paths once for the blueprint and once for the generated class
-	// make sure all function headers pointing to things in the old BP are changed to
-	// point to their duplicates in the new BP
-	UpdateFunctionHeaders(PreDuplicateAssetPath.ToString(), GetPathName());
-	if(const URigVMBlueprintGeneratedClass* CRGeneratedClass = GetRigVMBlueprintGeneratedClass())
-	{
-		UpdateFunctionHeaders(PreDuplicateHostPath.ToString(), CRGeneratedClass->GetPathName());
-	}
+	const FString OldAssetPath = PreDuplicateAssetPath.ToString();
+	const FString NewAssetPath = GetPathName();
+	
+	ReplaceFunctionIdentifiers(OldAssetPath, NewAssetPath);
 
 	PreDuplicateAssetPath.Reset();
 	PreDuplicateHostPath.Reset();
diff --git a/Engine/Plugins/Runtime/RigVM/Source/RigVMDeveloper/Private/RigVMModel/RigVMClient.cpp b/Engine/Plugins/Runtime/RigVM/Source/RigVMDeveloper/Private/RigVMModel/RigVMClient.cpp
--- a/Engine/Plugins/Runtime/RigVM/Source/RigVMDeveloper/Private/RigVMModel/RigVMClient.cpp
+++ b/Engine/Plugins/Runtime/RigVM/Source/RigVMDeveloper/Private/RigVMModel/RigVMClient.cpp
@@ -995,16 +995,6 @@ void FRigVMClient::ProcessDetachedLinks()
 	}
 }
 
-void FRigVMClient::PostDuplicateHost(const FString& InOldPathName, const FString& InNewPathName)
-{
-	TArray<URigVMGraph*> AllModels = GetAllModels(true, true);
-	for(URigVMGraph* Model : AllModels)
-	{
-		URigVMController* Controller = GetOrCreateController(Model);
-		Controller->PostDuplicateHost(InOldPathName, InNewPathName);
-	}
-}
-
 void FRigVMClient::PreSave(FObjectPreSaveContext ObjectSaveContext)
 {
 	if (!ObjectSaveContext.IsCooking())
diff --git a/Engine/Plugins/Runtime/RigVM/Source/RigVMDeveloper/Private/RigVMModel/RigVMController.cpp b/Engine/Plugins/Runtime/RigVM/Source/RigVMDeveloper/Private/RigVMModel/RigVMController.cpp
--- a/Engine/Plugins/Runtime/RigVM/Source/RigVMDeveloper/Private/RigVMModel/RigVMController.cpp
+++ b/Engine/Plugins/Runtime/RigVM/Source/RigVMDeveloper/Private/RigVMModel/RigVMController.cpp
@@ -20701,23 +20701,6 @@ ERigVMPinDefaultValueType URigVMController::GetDefaultValueType(const URigVMPin*
 	return Type;
 }
 
-void URigVMController::PostDuplicateHost(const FString& InOldPathName, const FString& InNewPathName)
-{
-	if (const URigVMGraph* Graph = GetGraph())
-	{
-		TArray<URigVMNode*> FunctionRefNodes = Graph->GetNodes().FilterByPredicate([](URigVMNode* Node)
-		{
-			return Node->IsA<URigVMFunctionReferenceNode>();
-		});
-
-		for(URigVMNode* Node : FunctionRefNodes)
-		{
-			URigVMFunctionReferenceNode* FunctionReferenceNode = CastChecked<URigVMFunctionReferenceNode>(Node);
-			FunctionReferenceNode->ReferencedFunctionHeader.PostDuplicateHost(InOldPathName, InNewPathName);
-		}
-	}
-}
-
 #endif
 
 URigVMActionStack* URigVMController::GetActionStack() const
diff --git a/Engine/Plugins/Runtime/RigVM/Source/RigVMDeveloper/Public/RigVMBlueprint.h b/Engine/Plugins/Runtime/RigVM/Source/RigVMDeveloper/Public/RigVMBlueprint.h
--- a/Engine/Plugins/Runtime/RigVM/Source/RigVMDeveloper/Public/RigVMBlueprint.h
+++ b/Engine/Plugins/Runtime/RigVM/Source/RigVMDeveloper/Public/RigVMBlueprint.h
@@ -584,6 +584,8 @@ private:
 	void Notify(ERigVMGraphNotifType InNotifType, UObject* InSubject);
 	void HandleModifiedEvent(ERigVMGraphNotifType InNotifType, URigVMGraph* InGraph, UObject* InSubject);
 
+	void ReplaceFunctionIdentifiers(const FString& OldAssetPath, const FString& NewAssetPath);
+
 #if WITH_EDITOR
 
 	UFUNCTION(BlueprintCallable, Category = "RigVM Blueprint")
diff --git a/Engine/Plugins/Runtime/RigVM/Source/RigVMDeveloper/Public/RigVMModel/RigVMClient.h b/Engine/Plugins/Runtime/RigVM/Source/RigVMDeveloper/Public/RigVMModel/RigVMClient.h
--- a/Engine/Plugins/Runtime/RigVM/Source/RigVMDeveloper/Public/RigVMModel/RigVMClient.h
+++ b/Engine/Plugins/Runtime/RigVM/Source/RigVMDeveloper/Public/RigVMModel/RigVMClient.h
@@ -251,9 +251,6 @@ public:
 	// try to reattach detached links and delete remaining ones
 	void ProcessDetachedLinks();
 
-	// work to be done after a duplication of the source asset
-	void PostDuplicateHost(const FString& InOldPathName, const FString& InNewPathName);
-
 	// work to be done before saving
 	void PreSave(FObjectPreSaveContext ObjectSaveContext);
 
diff --git a/Engine/Plugins/Runtime/RigVM/Source/RigVMDeveloper/Public/RigVMModel/RigVMController.h b/Engine/Plugins/Runtime/RigVM/Source/RigVMDeveloper/Public/RigVMModel/RigVMController.h
--- a/Engine/Plugins/Runtime/RigVM/Source/RigVMDeveloper/Public/RigVMModel/RigVMController.h
+++ b/Engine/Plugins/Runtime/RigVM/Source/RigVMDeveloper/Public/RigVMModel/RigVMController.h
@@ -1383,9 +1383,6 @@ protected:
 
 	ERigVMPinDefaultValueType GetDefaultValueType(const URigVMPin* InPin, const FString& InDefaultValue) const;
 
-	// work to do after a duplication of the host asset
-	void PostDuplicateHost(const FString& InOldPathName, const FString& InNewPathName);
-
 	template<typename T>
 	static void SortGraphElementsByGraphDepth(TArray<T*>& InOutElements, bool bReverse = false)
 	{
