RigVM: Bulk Edit workflows for Swapping Functions

#rnx

diff --git a/Engine/Plugins/Runtime/RigVM/Source/RigVMEditor/Private/Editor/RigVMEditor.cpp b/Engine/Plugins/Runtime/RigVM/Source/RigVMEditor/Private/Editor/RigVMEditor.cpp
--- a/Engine/Plugins/Runtime/RigVM/Source/RigVMEditor/Private/Editor/RigVMEditor.cpp
+++ b/Engine/Plugins/Runtime/RigVM/Source/RigVMEditor/Private/Editor/RigVMEditor.cpp
@@ -45,6 +45,9 @@
 #include "ScopedTransaction.h"
 #include "Editor/RigVMEditorMode.h"
 #include "InstancedPropertyBagStructureDataProvider.h"
+#include "Widgets/SRigVMSwapFunctionsWidget.h"
+#include "Widgets/SRigVMBulkEditDialog.h"
+#include "AssetRegistry/AssetRegistryModule.h"
 
 #define LOCTEXT_NAMESPACE "RigVMEditor"
 
@@ -1391,6 +1394,16 @@ void FRigVMEditor::BindCommands()
 		FRigVMEditorCommands::Get().FrameSelection,
 		FExecuteAction::CreateSP(this, &FRigVMEditor::FrameSelection),
 		FCanExecuteAction());
+
+	GetToolkitCommands()->MapAction(
+		FRigVMEditorCommands::Get().SwapFunctionWithinAsset,
+		FExecuteAction::CreateSP(this, &FRigVMEditor::SwapFunctionWithinAsset),
+		FCanExecuteAction());
+
+	GetToolkitCommands()->MapAction(
+		FRigVMEditorCommands::Get().SwapFunctionAcrossProject,
+		FExecuteAction::CreateSP(this, &FRigVMEditor::SwapFunctionAcrossProject),
+		FCanExecuteAction());
 }
 
 void FRigVMEditor::ToggleAutoCompileGraph()
@@ -1447,6 +1460,16 @@ TSharedRef<SWidget> FRigVMEditor::GenerateExecutionModeMenuContent()
 	return MenuBuilder.MakeWidget();
 }
 
+TSharedRef<SWidget> FRigVMEditor::GenerateBulkEditMenuContent()
+{
+	FMenuBuilder MenuBuilder(true, GetToolkitCommands());
+	MenuBuilder.BeginSection(TEXT("Functions"), LOCTEXT("Functions", "Functions"));
+	MenuBuilder.AddMenuEntry(FRigVMEditorCommands::Get().SwapFunctionWithinAsset, TEXT("SwapFunctionWithinAsset"), TAttribute<FText>(), TAttribute<FText>(), FSlateIcon());
+	MenuBuilder.AddMenuEntry(FRigVMEditorCommands::Get().SwapFunctionAcrossProject, TEXT("SwapFunctionAcrossProject"), TAttribute<FText>(), TAttribute<FText>(), FSlateIcon());
+	MenuBuilder.EndSection();
+	return MenuBuilder.MakeWidget();
+}
+
 void FRigVMEditor::OnActiveTabChanged( TSharedPtr<SDockTab> PreviouslyActive, TSharedPtr<SDockTab> NewlyActivated )
 {
 	if (!NewlyActivated.IsValid())
@@ -2385,6 +2408,15 @@ void FRigVMEditor::HandleJumpToHyperlink(const UObject* InSubject)
 	{
 		GraphToJumpTo = Node->GetGraph();
 		NodeToJumpTo = Node;
+
+		if(const URigVMCollapseNode* CollapseNode = Cast<URigVMCollapseNode>(Node))
+		{
+			if(CollapseNode->GetGraph()->IsA<URigVMFunctionLibrary>())
+			{
+				GraphToJumpTo = CollapseNode->GetContainedGraph();
+				NodeToJumpTo = CollapseNode->GetEntryNode();
+			}
+		}
 	}
 	else if(const URigVMPin* Pin = Cast<URigVMPin>(InSubject))
 	{
@@ -3102,6 +3134,16 @@ void FRigVMEditor::FillToolbar(FToolBarBuilder& ToolbarBuilder, bool bEndSection
 			NAME_None, TAttribute<FText>(), TAttribute<FText>(), FSlateIcon(FAppStyle::GetAppStyleSetName(), "PlayWorld.StepOut"));
 
 		ToolbarBuilder.EndStyleOverride();
+
+
+		FUIAction DefaultBulkEditAction;
+		ToolbarBuilder.AddComboButton(
+			DefaultBulkEditAction,
+			FOnGetContent::CreateSP(this, &FRigVMEditor::GenerateBulkEditMenuContent),
+			LOCTEXT("BulkEdit_Label", "Bulk Edit"),
+			LOCTEXT("BulkEdit_ToolTip", "Perform changes across many nodes / assets"),
+			FSlateIcon(FAppStyle::GetAppStyleSetName(), "LevelEditor.Recompile"),
+			false);
 	}
 
 	if(bEndSection)
@@ -3396,6 +3438,36 @@ void FRigVMEditor::FrameSelection()
 	}
 }
 
+void FRigVMEditor::SwapFunctionWithinAsset()
+{
+	const FAssetData Asset = UE::RigVM::Editor::Tools::FindAssetFromAnyPath(GetRigVMBlueprint()->GetPathName(), true);
+	SwapFunctionForAssets({Asset}, true);
+}
+
+void FRigVMEditor::SwapFunctionAcrossProject()
+{
+	const IAssetRegistry& AssetRegistry = FModuleManager::LoadModuleChecked<FAssetRegistryModule>(TEXT("AssetRegistry")).Get();
+	TArray<FAssetData> AllAssets;
+	AssetRegistry.GetAssetsByClass(GetRigVMBlueprint()->GetClass()->GetClassPathName(), AllAssets, true);
+	SwapFunctionForAssets(AllAssets, false);
+}
+
+void FRigVMEditor::SwapFunctionForAssets(const TArray<FAssetData>& InAssets, bool bSetupUndo)
+{
+	SRigVMSwapFunctionsWidget::FArguments WidgetArgs;
+	WidgetArgs
+		.Assets(InAssets)
+		.EnableUndo(bSetupUndo)
+		.CloseOnSuccess(true);
+
+	const TSharedRef<SRigVMBulkEditDialog<SRigVMSwapFunctionsWidget>> SwapFunctionsDialog =
+		SNew(SRigVMBulkEditDialog<SRigVMSwapFunctionsWidget>)
+		.WindowSize(FVector2D(800.0f, 640.0f))
+		.WidgetArgs(WidgetArgs);
+
+	SwapFunctionsDialog->ShowNormal();
+}
+
 void FRigVMEditor::UpdateGraphCompilerErrors()
 {
 	DECLARE_SCOPE_HIERARCHICAL_COUNTER_FUNC()
diff --git a/Engine/Plugins/Runtime/RigVM/Source/RigVMEditor/Private/Editor/RigVMEditorCommands.cpp b/Engine/Plugins/Runtime/RigVM/Source/RigVMEditor/Private/Editor/RigVMEditorCommands.cpp
--- a/Engine/Plugins/Runtime/RigVM/Source/RigVMEditor/Private/Editor/RigVMEditorCommands.cpp
+++ b/Engine/Plugins/Runtime/RigVM/Source/RigVMEditor/Private/Editor/RigVMEditorCommands.cpp
@@ -19,6 +19,8 @@ void FRigVMEditorCommands::RegisterCommands()
 	UI_COMMAND(StepOut, "Step Out", "Steps out of the collapsed/function node, when halted at a breakpoint.", EUserInterfaceActionType::Button, FInputChord(EKeys::F11, EModifierKey::Shift));
 	UI_COMMAND(StepOver, "Step Over", "Steps over the node, when halted at a breakpoint.", EUserInterfaceActionType::Button, FInputChord(EKeys::F10));
 	UI_COMMAND(FrameSelection, "Frame Selection", "Frames the selected nodes in the Graph View.", EUserInterfaceActionType::Button, FInputChord(EKeys::F));
+	UI_COMMAND(SwapFunctionWithinAsset, "Swap Function (Asset)", "Swaps a function for all occurrences within this asset.", EUserInterfaceActionType::Button, FInputChord());
+	UI_COMMAND(SwapFunctionAcrossProject, "Swap Function (Project)", "Swaps a function for all occurrences in the project.", EUserInterfaceActionType::Button, FInputChord());
 }
 
 #undef LOCTEXT_NAMESPACE
diff --git a/Engine/Plugins/Runtime/RigVM/Source/RigVMEditor/Private/Editor/RigVMEditorCommands.h b/Engine/Plugins/Runtime/RigVM/Source/RigVMEditor/Private/Editor/RigVMEditorCommands.h
--- a/Engine/Plugins/Runtime/RigVM/Source/RigVMEditor/Private/Editor/RigVMEditorCommands.h
+++ b/Engine/Plugins/Runtime/RigVM/Source/RigVMEditor/Private/Editor/RigVMEditorCommands.h
@@ -57,6 +57,12 @@ public:
 	/** Frames the selected nodes */
 	TSharedPtr< FUICommandInfo > FrameSelection;
 
+	/** Swap Function (Asset) */
+	TSharedPtr< FUICommandInfo > SwapFunctionWithinAsset;
+
+	/** Swap Function (Project) */
+	TSharedPtr< FUICommandInfo > SwapFunctionAcrossProject;
+
 	/**
 	 * Initialize commands
 	 */
diff --git a/Engine/Plugins/Runtime/RigVM/Source/RigVMEditor/Private/Editor/RigVMGraphDetailCustomization.cpp b/Engine/Plugins/Runtime/RigVM/Source/RigVMEditor/Private/Editor/RigVMGraphDetailCustomization.cpp
--- a/Engine/Plugins/Runtime/RigVM/Source/RigVMEditor/Private/Editor/RigVMGraphDetailCustomization.cpp
+++ b/Engine/Plugins/Runtime/RigVM/Source/RigVMEditor/Private/Editor/RigVMGraphDetailCustomization.cpp
@@ -28,6 +28,7 @@
 #include "InstancedPropertyBagStructureDataProvider.h"
 #include "Widgets/SRigVMGraphPinEnumPicker.h"
 #include "ScopedTransaction.h"
+#include "Misc/UObjectToken.h"
 
 #define LOCTEXT_NAMESPACE "RigVMGraphDetailCustomization"
 
@@ -958,6 +959,75 @@ void FRigVMGraphDetailCustomization::CustomizeDetails(IDetailLayoutBuilder& Deta
                     .OnSelectionChanged(this, &FRigVMGraphDetailCustomization::OnAccessSpecifierSelected)
             ]
         ];
+
+		// variant guid
+		SettingsCategory.AddCustomRow(FText::GetEmpty())
+		.Visibility(TAttribute<EVisibility>::CreateLambda([this]()
+		{
+			return IsFunctionVariant() ? EVisibility::Visible : EVisibility::Collapsed;
+		}))
+		.NameContent()
+		[
+			SNew(STextBlock)
+			.Text(FText::FromString(TEXT("Variant Guid")))
+			.Font(IDetailLayoutBuilder::GetDetailFont())
+		]
+		.ValueContent()
+		[
+			SNew(STextBlock)
+			.Font(IDetailLayoutBuilder::GetDetailFont())
+			.Text(this, &FRigVMGraphDetailCustomization::GetVariantGuidText)
+		];
+
+		/*
+		SettingsCategory.AddCustomRow(FText::GetEmpty())
+		.WholeRowContent()
+		[
+			SNew(SHorizontalBox)
+			+ SHorizontalBox::Slot()
+			[
+				SNew(SButton)
+				.Visibility_Lambda([this]()
+				{
+					return IsFunctionVariant() ? EVisibility::Visible : EVisibility::Collapsed;
+				})
+				.Text(LOCTEXT("SplitVariant", "Split Variant"))
+				.ToolTipText(LOCTEXT("SplitVariantToolTip", "Splits this variant from the set and makes the function a unique function"))
+				.OnClicked(this, &FRigVMGraphDetailCustomization::OnSplitVariant)
+			]
+			+ SHorizontalBox::Slot()
+			[
+				SNew(SButton)
+				.Text(LOCTEXT("MergeVariant", "Merge Variant"))
+				.ToolTipText(LOCTEXT("MergeVariantToolTip", "Merges this variant with an existing set"))
+				.OnClicked(this, &FRigVMGraphDetailCustomization::OnMergeVariant)
+			]
+		];
+		*/
+
+		// matching variants
+		SettingsCategory.AddCustomRow(FText::GetEmpty())
+		.Visibility(TAttribute<EVisibility>::CreateLambda([this]()
+		{
+			return IsFunctionVariant() ? EVisibility::Visible : EVisibility::Collapsed;
+		}))
+		.NameContent()
+		[
+			SNew(STextBlock)
+			.Text(FText::FromString(TEXT("Matching Variants")))
+			.Font(IDetailLayoutBuilder::GetDetailFont())
+		]
+		.ValueContent()
+		[
+			SAssignNew(VariantLog, SRigVMLogWidget)
+			.LogLabel(LOCTEXT("Variants", "Variants"))
+			.LogName(TEXT("RigVMFunctionVariants"))
+			.ShowFilters(false)
+			.AllowClear(false)
+			.DiscardDuplicates(true)
+			.ScrollToBottom(false)
+			.HeightOverride(120)
+		];
 	}
 
 	// node color
@@ -988,6 +1058,8 @@ void FRigVMGraphDetailCustomization::CustomizeDetails(IDetailLayoutBuilder& Deta
 		Model,
 		Blueprint));
 	DefaultsCategory.AddCustomBuilder(DefaultsArgumentNode);
+	
+	RefreshVariantLog();
 }
 
 bool FRigVMGraphDetailCustomization::IsAddNewInputOutputEnabled() const
@@ -1313,6 +1385,115 @@ TSharedRef<ITableRow> FRigVMGraphDetailCustomization::HandleGenerateRowAccessSpe
         ];
 }
 
+bool FRigVMGraphDetailCustomization::IsFunctionVariant() const
+{
+	if (GraphPtr.IsValid() && RigVMBlueprintPtr.IsValid())
+	{
+		URigVMBlueprint* Blueprint = RigVMBlueprintPtr.Get();
+		if (const URigVMGraph* Model = Blueprint->GetModel(GraphPtr.Get()))
+		{
+			if (const URigVMLibraryNode* LibraryNode = Cast<URigVMLibraryNode>(Model->GetOuter()))
+			{
+				return LibraryNode->GetFunctionHeader(Blueprint->GetRigVMGraphFunctionHost()).LibraryPointer.IsVariant(); 
+			}
+		}
+	}
+	return false;
+}
+
+FText FRigVMGraphDetailCustomization::GetVariantGuidText() const
+{
+	if (GraphPtr.IsValid() && RigVMBlueprintPtr.IsValid())
+	{
+		URigVMBlueprint* Blueprint = RigVMBlueprintPtr.Get();
+		if (const URigVMGraph* Model = Blueprint->GetModel(GraphPtr.Get()))
+		{
+			if (const URigVMLibraryNode* LibraryNode = Cast<URigVMLibraryNode>(Model->GetOuter()))
+			{
+				const FGuid Guid = LibraryNode->GetFunctionHeader(Blueprint->GetRigVMGraphFunctionHost()).Variant.Guid; 
+				if(Guid.IsValid())
+				{
+					return FText::FromString(Guid.ToString(EGuidFormats::DigitsWithHyphensLower));
+				}
+			}
+		}
+	}
+	return FText();
+}
+
+/*
+FReply FRigVMGraphDetailCustomization::OnSplitVariant()
+{
+	// todo
+	RefreshVariantLog();
+	return FReply::Unhandled();
+}
+
+FReply FRigVMGraphDetailCustomization::OnMergeVariant()
+{
+	// todo
+	RefreshVariantLog();
+	return FReply::Unhandled();
+}
+*/
+
+void FRigVMGraphDetailCustomization::RefreshVariantLog()
+{
+	check(VariantLog);
+	
+	VariantLog->GetListing()->ClearMessages();
+
+	if (GraphPtr.IsValid() && RigVMBlueprintPtr.IsValid())
+	{
+		URigVMBlueprint* Blueprint = RigVMBlueprintPtr.Get();
+		if (const URigVMGraph* Model = Blueprint->GetModel(GraphPtr.Get()))
+		{
+			if (const URigVMLibraryNode* LibraryNode = Cast<URigVMLibraryNode>(Model->GetOuter()))
+			{
+				const FRigVMGraphFunctionHeader Header = LibraryNode->GetFunctionHeader(Blueprint->GetRigVMGraphFunctionHost());
+				TArray<FRigVMGraphFunctionIdentifier> Variants = Header.LibraryPointer.GetVariantIdentifiers(false);
+				Variants.Sort([](const FRigVMGraphFunctionIdentifier& A, const FRigVMGraphFunctionIdentifier& B)
+				{
+					return A.GetLibraryNodePath().Compare(B.GetLibraryNodePath()) < 0;
+				});
+
+				for(const FRigVMGraphFunctionIdentifier& Variant : Variants)
+				{
+					if(URigVMLibraryNode* FunctionVariant = Cast<URigVMLibraryNode>(Variant.GetNodeSoftPath().TryLoad()))
+					{
+						const TSharedRef<FUObjectToken> ObjectToken = FUObjectToken::Create(FunctionVariant);
+
+						TWeakObjectPtr<URigVMLibraryNode> WeakFunctionVariant(FunctionVariant);
+						ObjectToken->OnMessageTokenActivated(FOnMessageTokenActivated::CreateLambda([WeakFunctionVariant](const TSharedRef<IMessageToken>&)
+						{
+							if(WeakFunctionVariant.IsValid())
+							{
+								if(UBlueprint* Blueprint = WeakFunctionVariant.Get()->GetTypedOuter<UBlueprint>())
+								{
+									GEditor->GetEditorSubsystem<UAssetEditorSubsystem>()->OpenEditorForAsset(Blueprint);
+					
+									if(IAssetEditorInstance* Editor = GEditor->GetEditorSubsystem<UAssetEditorSubsystem>()->FindEditorForAsset(Blueprint, true))
+									{
+										if(FRigVMEditor* RigVMEditor = static_cast<FRigVMEditor*>(Editor))
+										{
+											RigVMEditor->HandleJumpToHyperlink(WeakFunctionVariant.Get());
+										}
+									}
+								}
+							}
+						}));
+
+						const TSharedRef<FTokenizedMessage> Message = FTokenizedMessage::Create(EMessageSeverity::Info);
+						Message->AddToken(ObjectToken);
+						VariantLog->GetListing()->AddMessage(Message);
+					}
+				}
+			}
+		}
+	}
+	
+}
+
 FRigVMWrappedNodeDetailCustomization::FRigVMWrappedNodeDetailCustomization()
 : BlueprintBeingCustomized(nullptr)
 {
diff --git a/Engine/Plugins/Runtime/RigVM/Source/RigVMEditor/Private/RigVMEditorModule.cpp b/Engine/Plugins/Runtime/RigVM/Source/RigVMEditor/Private/RigVMEditorModule.cpp
--- a/Engine/Plugins/Runtime/RigVM/Source/RigVMEditor/Private/RigVMEditorModule.cpp
+++ b/Engine/Plugins/Runtime/RigVM/Source/RigVMEditor/Private/RigVMEditorModule.cpp
@@ -39,7 +39,10 @@
 #include "RigVMFunctions/Simulation/RigVMFunction_AlphaInterp.h"
 #include "RigVMFunctions/Debug/RigVMFunction_VisualDebug.h"
 #include "ScopedTransaction.h"
+#include "Editor/RigVMEditorTools.h"
 #include "UObject/UObjectIterator.h"
+#include "Widgets/SRigVMSwapFunctionsWidget.h"
+#include "Widgets/SRigVMBulkEditDialog.h"
 
 DEFINE_LOG_CATEGORY(LogRigVMEditor);
 
@@ -603,7 +606,7 @@ void FRigVMEditorModule::GetNodeContextMenuActions(IRigVMClientHost* RigVMClient
 	GetNodeVariablesContextMenuActions(RigVMClientHost, EdGraphNode, ModelNode, Menu);
 	GetNodeTemplatesContextMenuActions(RigVMClientHost, EdGraphNode, ModelNode, Menu);
 	GetNodeOrganizationContextMenuActions(RigVMClientHost, EdGraphNode, ModelNode, Menu);
-	GetNodeVersioningContextMenuActions(RigVMClientHost, EdGraphNode, ModelNode, Menu);
+	GetNodeVariantContextMenuActions(RigVMClientHost, EdGraphNode, ModelNode, Menu);
 	GetNodeTestContextMenuActions(RigVMClientHost, EdGraphNode, ModelNode, Menu);
 }
 
@@ -1143,12 +1146,13 @@ void FRigVMEditorModule::GetNodeOrganizationContextMenuActions(IRigVMClientHost*
 	}));
 }
 
-void FRigVMEditorModule::GetNodeVersioningContextMenuActions(IRigVMClientHost* RigVMClientHost, const URigVMEdGraphNode* EdGraphNode, URigVMNode* ModelNode, UToolMenu* Menu) const
+void FRigVMEditorModule::GetNodeVariantContextMenuActions(IRigVMClientHost* RigVMClientHost, const URigVMEdGraphNode* EdGraphNode, URigVMNode* ModelNode, UToolMenu* Menu) const
 {
 	const URigVMGraph* Model = ModelNode->GetGraph();
 	URigVMController* Controller = RigVMClientHost->GetRigVMClient()->GetController(Model);
 
 	bool bCanNodeBeUpgraded = false;
+	const bool bIsFunctionReference = ModelNode->IsA<URigVMFunctionReferenceNode>();
 	TArray<FName> SelectedNodeNames = Model->GetSelectNodes();
 	SelectedNodeNames.AddUnique(ModelNode->GetFName());
 
@@ -1160,19 +1164,73 @@ void FRigVMEditorModule::GetNodeVersioningContextMenuActions(IRigVMClientHost* R
 		}
 	}
 	
-	if(bCanNodeBeUpgraded)
+	if(bCanNodeBeUpgraded || bIsFunctionReference)
 	{
-		FToolMenuSection& VersioningSection = Menu->AddSection("RigVMEditorContextMenuVersioning", LOCTEXT("VersioningHeader", "Versioning"));
-		VersioningSection.AddMenuEntry(
-			"Upgrade Nodes",
-			LOCTEXT("UpgradeNodes", "Upgrade Nodes"),
-			LOCTEXT("UpgradeNodes_Tooltip", "Upgrades deprecated nodes to their current implementation"),
-			FSlateIcon(),
-			FUIAction(FExecuteAction::CreateLambda([Model, Controller]() {
-				TArray<FName> Nodes = Model->GetSelectNodes();
-				Controller->UpgradeNodes(Nodes, true, true);
-			})
-		));
+		FToolMenuSection& VariantSection = Menu->AddSection("RigVMEditorContextMenuVariant", LOCTEXT("VariantHeader", "Variants"));
+
+		if(bCanNodeBeUpgraded)
+		{
+			VariantSection.AddMenuEntry(
+				"Upgrade Nodes",
+				LOCTEXT("UpgradeNodes", "Upgrade Nodes"),
+				LOCTEXT("UpgradeNodes_Tooltip", "Upgrades deprecated nodes to their current implementation"),
+				FSlateIcon(),
+				FUIAction(FExecuteAction::CreateLambda([Model, Controller]() {
+					TArray<FName> Nodes = Model->GetSelectNodes();
+					Controller->UpgradeNodes(Nodes, true, true);
+				})
+			));
+		}
+
+		if(bIsFunctionReference)
+		{
+			VariantSection.AddMenuEntry(
+				"Swap function for selected nodes",
+				LOCTEXT("SwapSelectedFunction", "Swap function for selected nodes"),
+				LOCTEXT("SwapSelectedFunction_Tooltip", "Swaps this function for another one for all nodes matching within the selection"),
+				FSlateIcon(),
+				FUIAction(FExecuteAction::CreateLambda([Model]() {
+					TArray<FName> NodeNames = Model->GetSelectNodes();
+					TArray<URigVMFunctionReferenceNode*> FunctionReferenceNodes;
+					FRigVMGraphFunctionIdentifier Identifier;
+					for(const FName& NodeName : NodeNames)
+					{
+						if(URigVMFunctionReferenceNode* FunctionReferenceNode = Cast<URigVMFunctionReferenceNode>(Model->FindNodeByName(NodeName)))
+						{
+							if(Identifier.IsValid())
+							{
+								if(FunctionReferenceNode->GetFunctionIdentifier() != Identifier)
+								{
+									continue;
+								}
+							}
+							else
+							{
+								Identifier = FunctionReferenceNode->GetFunctionIdentifier();
+							}
+							FunctionReferenceNodes.Add(FunctionReferenceNode);
+						}
+					}
+					if(!FunctionReferenceNodes.IsEmpty())
+					{
+						SRigVMSwapFunctionsWidget::FArguments WidgetArgs;
+						WidgetArgs
+							.Source(Identifier)
+							.FunctionReferenceNodes(FunctionReferenceNodes)
+							.SkipPickingFunctionRefs(true)
+							.EnableUndo(true)
+							.CloseOnSuccess(true);
+
+						const TSharedRef<SRigVMBulkEditDialog<SRigVMSwapFunctionsWidget>> SwapFunctionsDialog =
+							SNew(SRigVMBulkEditDialog<SRigVMSwapFunctionsWidget>)
+							.WindowSize(FVector2D(800.0f, 640.0f))
+							.WidgetArgs(WidgetArgs);
+
+						SwapFunctionsDialog->ShowNormal();
+					}
+				})
+			));
+		}
 	}
 }
 
diff --git a/Engine/Plugins/Runtime/RigVM/Source/RigVMEditor/Private/Widgets/SRigVMBulkEditWidget.cpp b/Engine/Plugins/Runtime/RigVM/Source/RigVMEditor/Private/Widgets/SRigVMBulkEditWidget.cpp
new file mode 100644
--- /dev/null
+++ b/Engine/Plugins/Runtime/RigVM/Source/RigVMEditor/Private/Widgets/SRigVMBulkEditWidget.cpp
@@ -0,0 +1,708 @@
+// Copyright Epic Games, Inc. All Rights Reserved.
+
+#include "Widgets/SRigVMBulkEditWidget.h"
+#include "Widgets/SWindow.h"
+#include "Widgets/SRigVMLogWidget.h"
+#include "Widgets/Input/SButton.h"
+#include "SPrimaryButton.h"
+#include "Widgets/Text/SInlineEditableTextBlock.h"
+#include "Widgets/Notifications/SProgressBar.h"
+#include "Styling/AppStyle.h"
+#include "Widgets/Notifications/SNotificationList.h"
+#include "Widgets/Layout/SSpacer.h"
+#include "Framework/Notifications/NotificationManager.h"
+#include "RigVMEditorModule.h"
+#include "Misc/UObjectToken.h"
+#include "Dialogs/Dialogs.h"
+#include "Framework/Application/SlateApplication.h"
+#define LOCTEXT_NAMESPACE "SRigVMBulkEditWidget"
+
+//////////////////////////////////////////////////////////////////////////
+// SRigVMBulkEditWidget
+
+SRigVMBulkEditWidget::~SRigVMBulkEditWidget()
+{
+}
+
+void SRigVMBulkEditWidget::Construct(const FArguments& InArgs)
+{
+	Phases = InArgs._Phases;
+	bEnableUndo = InArgs._EnableUndo;
+	bCloseOnSuccess = InArgs._CloseOnSuccess;
+
+	OnPhaseActivated = InArgs._OnPhaseActivated;
+
+	BulkEditTitle = InArgs._BulkEditTitle;
+	BulkEditConfirmMessage = InArgs._BulkEditConfirmMessage;
+	BulkEditConfirmIniField = InArgs._BulkEditConfirmIniField;
+
+	if(BulkEditTitle.Get().IsEmpty())
+	{
+		BulkEditTitle = LOCTEXT("BulkEdit", "Bulk Edit");
+	}
+	if(BulkEditConfirmMessage.Get().IsEmpty())
+	{
+		BulkEditConfirmMessage = LOCTEXT("ConfirmBulkEditWithoutUndo", "This Bulk Edit will run with support for Undo. Are you sure?");
+	}
+
+	TSharedRef<SVerticalBox> MainVerticalBox = SNew(SVerticalBox);
+	TSharedRef<SHorizontalBox> MainHorizontalBox = SNew(SHorizontalBox);
+	if(InArgs._LeftWidget)
+	{
+		MainHorizontalBox->AddSlot()
+		.AutoWidth()
+		.HAlign(HAlign_Left)
+		.VAlign(VAlign_Fill)
+		[
+			InArgs._LeftWidget.ToSharedRef()
+		];
+	}
+	
+	MainHorizontalBox->AddSlot()
+	.FillWidth(1)
+	.HAlign(HAlign_Fill)
+	[
+		SAssignNew(TreeView, SRigVMChangesTreeView)
+		.Visibility_Lambda([this]()
+		{
+			return !bShowLog ? EVisibility::Visible : EVisibility::Collapsed;
+		})
+		.Phase(this, &SRigVMBulkEditWidget::GetActivePhasePtr)
+		.OnNodeSelected(InArgs._OnNodeSelected)
+		.OnNodeDoubleClicked(InArgs._OnNodeDoubleClicked)
+	];
+
+	MainHorizontalBox->AddSlot()
+	.FillWidth(1)
+	.HAlign(HAlign_Fill)
+	[
+		SAssignNew(BulkEditLogWidget, SRigVMLogWidget)
+		.LogName(TEXT("RigVMBulkEditLog"))
+		.LogLabel(LOCTEXT("RigVMBulkEditLog", "Bulk Edit Log"))
+		.Visibility_Lambda([this]()
+		{
+			return bShowLog ? EVisibility::Visible : EVisibility::Collapsed;
+		})
+	];
+	
+	BulkEditLogWidget->GetListing()->ClearMessages();
+
+	if(InArgs._RightWidget)
+	{
+		MainHorizontalBox->AddSlot()
+		.AutoWidth()
+		.HAlign(HAlign_Right)
+		.VAlign(VAlign_Fill)
+		[
+			InArgs._RightWidget.ToSharedRef()
+		];
+	}
+
+	if(InArgs._HeaderWidget)
+	{
+		MainVerticalBox->AddSlot()
+		.AutoHeight()
+		.HAlign(HAlign_Fill)
+		.VAlign(VAlign_Top)
+		[
+			InArgs._HeaderWidget.ToSharedRef()
+		];
+	}
+	
+	MainVerticalBox->AddSlot()
+	.VAlign(VAlign_Fill)
+	.FillHeight(1.0f)		// This is required to make the scrollbar work, as content overflows Slate containers by default
+	[
+		MainHorizontalBox
+	];
+
+	if(InArgs._FooterWidget)
+	{
+		MainVerticalBox->AddSlot()
+		.AutoHeight()
+		.HAlign(HAlign_Fill)
+		.VAlign(VAlign_Bottom)
+		[
+			InArgs._FooterWidget.ToSharedRef()
+		];
+	}
+
+	MainVerticalBox->AddSlot()
+	.AutoHeight()
+	.Padding(0, 2)
+	[
+		SNew(SHorizontalBox)
+
+		+ SHorizontalBox::Slot()
+		.FillWidth(1)
+		[
+			SNew(SVerticalBox)
+
+			+ SVerticalBox::Slot()
+			.AutoHeight()
+			.Padding(0)
+			[
+				SNew(SProgressBar)
+				.Visibility(this, &SRigVMBulkEditWidget::GetTasksProgressVisibility)
+				.Percent(this, &SRigVMBulkEditWidget::GetTasksProgressPercentage)
+			]
+		
+			+ SVerticalBox::Slot()
+			.AutoHeight()
+			.Padding(0, 4, 0, 0)
+			.HAlign(HAlign_Fill)
+			[
+				SNew(SHorizontalBox)
+
+				+ SHorizontalBox::Slot()
+				.AutoWidth()
+				.HAlign(HAlign_Left)
+				[
+					SNew(SButton)
+					.Visibility(this, &SRigVMBulkEditWidget::GetBackButtonVisibility)
+					.IsEnabled(this, &SRigVMBulkEditWidget::IsBackButtonEnabled)
+					.Text(LOCTEXT("Back", "Back"))
+					.OnClicked(this, &SRigVMBulkEditWidget::OnBackButtonClicked)
+				]
+
+				+ SHorizontalBox::Slot()
+				.HAlign(HAlign_Fill)
+				[
+					SNew(SSpacer)
+				]
+
+				+ SHorizontalBox::Slot()
+				.AutoWidth()
+				.HAlign(HAlign_Right)
+				[
+					SNew(SButton)
+					.Visibility(this, &SRigVMBulkEditWidget::GetCancelButtonVisibility)
+					.IsEnabled(this, &SRigVMBulkEditWidget::IsCancelButtonEnabled)
+					.Text(LOCTEXT("Cancel", "Cancel"))
+					.OnClicked(this, &SRigVMBulkEditWidget::OnCancelButtonClicked)
+				]
+
+				+ SHorizontalBox::Slot()
+				.AutoWidth()
+				.HAlign(HAlign_Right)
+				[
+					SNew(SPrimaryButton)
+					.Visibility(this, &SRigVMBulkEditWidget::GetPrimaryButtonVisibility)
+					.IsEnabled(this, &SRigVMBulkEditWidget::IsPrimaryButtonEnabled)
+					.Text(this, &SRigVMBulkEditWidget::GetPrimaryButtonText)
+					.OnClicked(this, &SRigVMBulkEditWidget::OnPrimaryButtonClicked)
+				]
+			]
+		]
+	];
+
+	ChildSlot
+	[
+		MainVerticalBox
+	];
+
+	BulkEditLogWidget->BindLog(LogRigVMDeveloper.GetCategoryName());
+	BulkEditLogWidget->BindLog(LogRigVMEditor.GetCategoryName());
+	ActivatePhase(InArgs._PhaseToActivate);
+}
+
+void SRigVMBulkEditWidget::Tick(const FGeometry& AllottedGeometry, const double InCurrentTime, const float InDeltaTime)
+{
+	SCompoundWidget::Tick(AllottedGeometry, InCurrentTime, InDeltaTime);
+
+	// update the tree view even if it is hidden
+	if(bShowLog)
+	{
+		TreeView->RefreshFilteredNodesIfRequired();
+	}
+
+	FScopeLock _(&TasksCriticalSection);
+	if(!RemainingTasks.IsEmpty())
+	{
+		if(bEnableUndo && !Transaction.IsValid())
+		{
+			Transaction = MakeShareable(new FScopedTransaction(LOCTEXT("BulkEdit", "Bulk Edit")));
+		}
+		
+		FFunctionGraphTask::CreateAndDispatchWhenReady([this]()
+		{
+			TSharedPtr<FRigVMTreeTask> Task;
+			{
+				FScopeLock _(&TasksCriticalSection);
+				if(RemainingTasks.IsEmpty())
+				{
+					return;
+				}
+				Task = RemainingTasks[0].ToSharedPtr();
+			}
+
+			if(Task.IsValid())
+			{
+				Task->SetEnableUndo(bEnableUndo);
+				
+				FScopedScriptExceptionHandler ScopedScriptExceptionHandler(
+					[this](ELogVerbosity::Type InVerbosity, const TCHAR* InMessage, const TCHAR* InStackMessage)
+					{
+						OnScriptException(InVerbosity, InMessage, InStackMessage);
+					});
+
+				const TSharedRef<FRigVMTreePhase>& Phase = GetActivePhase();
+
+				TSharedPtr<TGuardValue<FRigVMReportDelegate>> ReportDelegateGuard; 
+				if(URigVMBlueprint* Blueprint = Task->GetBlueprint(Phase))
+				{
+					ReportDelegateGuard = MakeShareable(new TGuardValue<FRigVMReportDelegate>(
+						Blueprint->VMCompileSettings.ASTSettings.ReportDelegate,
+						FRigVMReportDelegate::CreateLambda(
+						[this](EMessageSeverity::Type InSeverity, UObject* InSubject, const FString& InMessage)
+						{
+							TSharedRef<FTokenizedMessage> Message = FTokenizedMessage::Create(InSeverity); 
+							if(InSubject)
+							{
+								FString Left, Right;
+								if(InMessage.Split(TEXT("@@"), &Left, &Right))
+								{
+									Left.TrimStartAndEndInline();
+									Right.TrimStartAndEndInline();
+									Message->AddText(FText::FromString(Left));
+
+									const TSharedRef<FUObjectToken> SubjectToken = FUObjectToken::Create(InSubject);
+									TWeakObjectPtr<UObject> WeakSubject(InSubject);
+									SubjectToken->OnMessageTokenActivated(FOnMessageTokenActivated::CreateLambda([WeakSubject](const TSharedRef<IMessageToken>&)
+									{
+										if(WeakSubject.IsValid())
+										{
+											if(UBlueprint* Blueprint = WeakSubject.Get()->GetTypedOuter<UBlueprint>())
+											{
+												GEditor->GetEditorSubsystem<UAssetEditorSubsystem>()->OpenEditorForAsset(Blueprint);
+														
+												if(IAssetEditorInstance* Editor = GEditor->GetEditorSubsystem<UAssetEditorSubsystem>()->FindEditorForAsset(Blueprint, true))
+												{
+													if(FRigVMEditor* RigVMEditor = static_cast<FRigVMEditor*>(Editor))
+													{
+														RigVMEditor->HandleJumpToHyperlink(WeakSubject.Get());
+													}
+												}
+											}
+										}
+									}));
+									Message->AddToken(SubjectToken);
+									Message->AddText(FText::FromString(Right));
+								}
+							}
+							else
+							{
+								Message->AddText(FText::FromString(InMessage));
+							}
+							BulkEditLogWidget->GetListing()->AddMessage(Message);
+						})
+					));
+				}
+
+				if(Task->Execute(Phase))
+				{
+					FScopeLock _(&TasksCriticalSection);
+					RemainingTasks.Remove(Task.ToSharedRef());
+					CompletedTasks.Add(Task.ToSharedRef());
+
+					const TArray<FString> AffectedPaths = Task->GetAffectedNodes();
+					for(const FString& AffectedPath : AffectedPaths)
+					{
+						if(const TSharedPtr<FRigVMTreeNode> Node = Phase->FindVisibleNode(AffectedPath))
+						{
+							if(const TSharedPtr<FRigVMTreeNode> Parent = Node->GetParent())
+							{
+								Parent->DirtyChildren();
+							}
+							else
+							{
+								Node->DirtyChildren();
+							}
+						}
+					}
+
+					if(Task->RequiresRefresh() || !AffectedPaths.IsEmpty())
+					{
+						TreeView->RequestRefresh_AnyThread(true);
+					}
+				}
+                else
+                {
+                    CancelTasks();
+                }
+
+				ReportDelegateGuard.Reset();
+			}
+		}, TStatId(), nullptr, ENamedThreads::GameThread);		
+	}
+	else if(!CompletedTasks.IsEmpty())
+	{
+		CompletedTasks.Reset();
+		Transaction.Reset();
+		bTasksSucceeded = true;
+		TreeView->RequestRefresh_AnyThread(true);
+	}
+}
+
+FText SRigVMBulkEditWidget::GetDialogTitle() const
+{
+	if(bShowLog)
+	{
+		return LOCTEXT("ProcessingTasks", "Performing Bulk Edit...");
+	}
+	return FText::FromString(GetActivePhase()->GetName());
+}
+
+bool SRigVMBulkEditWidget::AreTasksInProgress() const
+{
+	FScopeLock _(&TasksCriticalSection);
+	return !RemainingTasks.IsEmpty();
+}
+
+EVisibility SRigVMBulkEditWidget::GetTasksProgressVisibility() const
+{
+	return AreTasksInProgress() ? EVisibility::Visible : EVisibility::Collapsed;
+}
+
+TOptional<float> SRigVMBulkEditWidget::GetTasksProgressPercentage() const
+{
+	if(!AreTasksInProgress())
+	{
+		return TOptional<float>();
+	}
+
+	FScopeLock _(&TasksCriticalSection);
+	const int32 NumTotalTasks = CompletedTasks.Num() + RemainingTasks.Num();
+	if(NumTotalTasks == 0)
+	{
+		return TOptional<float>();
+	}
+	return float(CompletedTasks.Num()) / float(NumTotalTasks);
+}
+
+void SRigVMBulkEditWidget::OnLogMessage(const TSharedRef<FTokenizedMessage>& InMessage) const
+{
+	BulkEditLogWidget->GetListing()->AddMessage( InMessage );
+}
+
+void SRigVMBulkEditWidget::OnScriptException(ELogVerbosity::Type InVerbosity, const TCHAR* InMessage, const TCHAR* InStackMessage)
+{
+	if(InMessage)
+	{
+		switch(InVerbosity)
+		{
+			case ELogVerbosity::Fatal:
+			case ELogVerbosity::Error:
+			{
+				OnLogMessage(FTokenizedMessage::Create(EMessageSeverity::Error, FText::FromString(InMessage)));
+				break;
+			}
+			case ELogVerbosity::Warning:
+			{
+				OnLogMessage(FTokenizedMessage::Create(EMessageSeverity::Warning, FText::FromString(InMessage)));
+				break;
+			}
+			default:
+			{
+				OnLogMessage(FTokenizedMessage::Create(EMessageSeverity::Info, FText::FromString(InMessage)));
+				break;
+			}
+		}
+	}
+}
+
+void SRigVMBulkEditWidget::CloseDialog()
+{
+	const TSharedPtr<SWindow> OwningWindow = FSlateApplication::Get().FindWidgetWindow(AsShared());
+	if (OwningWindow)
+	{
+		OwningWindow->RequestDestroyWindow();
+	}
+}
+
+EVisibility SRigVMBulkEditWidget::GetBackButtonVisibility() const
+{
+	if(IsReadyToClose())
+	{
+		return EVisibility::Collapsed;
+	}
+	return EVisibility::Visible;
+}
+
+bool SRigVMBulkEditWidget::IsBackButtonEnabled() const
+{
+	if(IsReadyToClose())
+	{
+		return false;
+	}
+	if(bShowLog)
+	{
+		return true;
+	}
+	if(AreTasksInProgress())
+	{
+		return false;
+	}
+	return ActivatedPhaseIDs.Num() > 1;
+}
+
+FReply SRigVMBulkEditWidget::OnBackButtonClicked()
+{
+	if(bShowLog)
+	{
+		bShowLog = false;
+		return FReply::Handled();
+	}
+	
+	if(ActivatedPhaseIDs.Num() > 1)
+	{
+		bShowLog = false;
+		ActivatedPhaseIDs.Pop();
+		const int32 PreviousID = ActivatedPhaseIDs.Pop();
+		if(ActivatePhase(PreviousID))
+		{
+			return FReply::Handled();
+		}
+	}
+	return FReply::Unhandled();
+}
+
+EVisibility SRigVMBulkEditWidget::GetCancelButtonVisibility() const
+{
+	if(IsReadyToClose())
+	{
+		return EVisibility::Collapsed;
+	}
+	if(AreTasksInProgress())
+	{
+		return EVisibility::Visible;
+	}
+	return GetActivePhase()->IsCancelButtonVisible().Get(false) ? EVisibility::Visible : EVisibility::Collapsed;
+}
+
+bool SRigVMBulkEditWidget::IsCancelButtonEnabled() const
+{
+	if(IsReadyToClose())
+	{
+		return false;
+	}
+	if(AreTasksInProgress())
+	{
+		return true;
+	}
+	return GetActivePhase()->IsCancelButtonEnabled().Get(true);
+}
+
+FReply SRigVMBulkEditWidget::OnCancelButtonClicked()
+{
+	if(AreTasksInProgress())
+	{
+		CancelTasks();
+		return FReply::Handled();
+	}
+	if(!GetActivePhase()->OnCancel().IsBound())
+	{
+		CloseDialog();
+		return FReply::Handled();
+	}
+	FReply Reply = GetActivePhase()->Cancel();
+	TreeView->RequestRefresh_AnyThread(true);
+	return Reply;
+}
+
+EVisibility SRigVMBulkEditWidget::GetPrimaryButtonVisibility() const
+{
+	if(IsReadyToClose())
+	{
+		return EVisibility::Visible;
+	}
+	return GetActivePhase()->IsPrimaryButtonVisible().Get(false) ? EVisibility::Visible : EVisibility::Collapsed;
+}
+
+bool SRigVMBulkEditWidget::IsPrimaryButtonEnabled() const
+{
+	if(IsReadyToClose())
+	{
+		return true;
+	}
+	if(AreTasksInProgress())
+	{
+		return false;
+	}
+	return GetActivePhase()->IsPrimaryButtonEnabled().Get(true);
+}
+
+FText SRigVMBulkEditWidget::GetPrimaryButtonText() const
+{
+	if(IsReadyToClose())
+	{
+		return LOCTEXT("Close", "Close");
+	}
+	return GetActivePhase()->PrimaryButtonText().Get(LOCTEXT("Ok", "Ok"));
+}
+
+bool SRigVMBulkEditWidget::IsReadyToClose() const
+{
+	return bTasksSucceeded && bCloseOnSuccess;
+}
+
+FReply SRigVMBulkEditWidget::OnPrimaryButtonClicked()
+{
+	if(IsReadyToClose())
+	{
+		CloseDialog();
+		return FReply::Handled();
+	}
+	
+	FReply Reply = GetActivePhase()->PrimaryAction();
+	if(Reply.IsEventHandled())
+	{
+		if(GetActivePhase() == Phases.Last())
+        {
+			FScopeLock _(&TasksCriticalSection);
+        	bShowLog = !RemainingTasks.IsEmpty();
+        }
+	}
+	TreeView->RequestRefresh_AnyThread(true);
+	return Reply;
+}
+
+void SRigVMBulkEditWidget::QueueTasks(const TArray<TSharedRef<FRigVMTreeTask>>& InTasks)
+{
+	FScopeLock _(&TasksCriticalSection);
+
+	if(!bEnableUndo)
+	{
+		const bool bRequiresUndo = InTasks.ContainsByPredicate([](const TSharedRef<FRigVMTreeTask>& Task)
+		{
+			return Task->RequiresUndo();
+		});
+
+		if(bRequiresUndo)
+		{
+			FSuppressableWarningDialog::FSetupInfo Info(BulkEditConfirmMessage.Get(), BulkEditTitle.Get(), BulkEditConfirmIniField.Get());
+			Info.ConfirmText = LOCTEXT("Yes", "Yes");
+			Info.CancelText = LOCTEXT("No", "No");
+
+			const FSuppressableWarningDialog ConfirmationDialog(Info);
+			if(!ConfirmationDialog.ShowModal())
+			{
+				return;
+			}
+		}
+	}
+	
+	RemainingTasks.Append(InTasks);
+}
+
+void SRigVMBulkEditWidget::CancelTasks()
+{
+	FScopeLock _(&TasksCriticalSection);
+	RemainingTasks.Reset();
+	CompletedTasks.Reset();
+	if(Transaction.IsValid())
+	{
+		Transaction->Cancel();
+		Transaction.Reset();
+	}
+	bTasksSucceeded = false;
+}
+
+TSharedPtr<FRigVMTreePhase> SRigVMBulkEditWidget::GetActivePhasePtr() const
+{
+	for(const TSharedRef<FRigVMTreePhase>& Phase : Phases)
+	{
+		if(Phase->IsActive())
+		{
+			return Phase;
+		}
+	}
+	return nullptr;
+}
+
+TSharedRef<FRigVMTreePhase> SRigVMBulkEditWidget::GetActivePhase() const
+{
+	for(const TSharedRef<FRigVMTreePhase>& Phase : Phases)
+	{
+		if(Phase->IsActive())
+		{
+			return Phase;
+		}
+	}
+	static const TSharedRef<FRigVMTreeContext> EmptyContext = FRigVMTreeContext::Create();
+	static const TSharedRef<FRigVMTreePhase> EmptyPhase = FRigVMTreePhase::Create(INDEX_NONE, TEXT("Default"), EmptyContext);
+	return EmptyPhase;
+}
+
+TSharedPtr<FRigVMTreePhase> SRigVMBulkEditWidget::FindPhase(int32 InID) const
+{
+	for(const TSharedRef<FRigVMTreePhase>& Phase : Phases)
+	{
+		if(Phase->GetID() == InID)
+		{
+			return Phase.ToSharedPtr();
+		}
+	}
+	return nullptr;
+}
+
+bool SRigVMBulkEditWidget::ActivatePhase(int32 InID)
+{
+	FString PreviouslySelectedPath;
+	TArray<TSharedRef<FRigVMTreeNode>> PreviousSelection = GetSelectedNodes();
+	if(!PreviousSelection.IsEmpty())
+	{
+		PreviouslySelectedPath = PreviousSelection[0]->GetPath();
+	}
+	
+	for(const TSharedRef<FRigVMTreePhase>& Phase : Phases)
+	{
+		Phase->bIsActive = false;
+	}
+
+	bool bResult = false;
+	if(const TSharedPtr<FRigVMTreePhase> Phase = FindPhase(InID))
+	{
+		Phase->bIsActive = true;
+		Phase->GetContext()->OnLogTokenizedMessage.RemoveAll(this);
+		Phase->GetContext()->OnLogTokenizedMessage.AddSP(this, &SRigVMBulkEditWidget::OnLogMessage);;
+		OnPhaseActivated.Execute(Phase.ToSharedRef());
+		ActivatedPhaseIDs.Add(Phase->GetID());
+		Phase->OnQueueTasks().BindSP(this, &SRigVMBulkEditWidget::QueueTasks);
+		bResult = true;
+	}
+
+	TreeView->OnPhaseChanged();
+
+	if(!PreviouslySelectedPath.IsEmpty())
+	{
+		if(const TSharedPtr<FRigVMTreeNode> Node = GetActivePhase()->FindVisibleNode(PreviouslySelectedPath))
+		{
+			TreeView->SetSelection(Node, true);
+		}
+	}
+	
+	return bResult;
+}
+
+TSharedRef<FRigVMTreeContext> SRigVMBulkEditWidget::GetContext() const
+{
+	return GetActivePhase()->GetContext();
+}
+
+TArray<TSharedRef<FRigVMTreeNode>> SRigVMBulkEditWidget::GetSelectedNodes() const
+{
+	return TreeView->GetSelectedNodes();
+}
+
+bool SRigVMBulkEditWidget::HasAnyVisibleCheckedNode() const
+{
+	return TreeView->HasAnyVisibleCheckedNode();
+}
+
+TArray<TSharedRef<FRigVMTreeNode>> SRigVMBulkEditWidget::GetCheckedNodes() const
+{
+	return TreeView->GetCheckedNodes();
+}
+
+#undef LOCTEXT_NAMESPACE
diff --git a/Engine/Plugins/Runtime/RigVM/Source/RigVMEditor/Private/Widgets/SRigVMLogWidget.cpp b/Engine/Plugins/Runtime/RigVM/Source/RigVMEditor/Private/Widgets/SRigVMLogWidget.cpp
--- a/Engine/Plugins/Runtime/RigVM/Source/RigVMEditor/Private/Widgets/SRigVMLogWidget.cpp
+++ b/Engine/Plugins/Runtime/RigVM/Source/RigVMEditor/Private/Widgets/SRigVMLogWidget.cpp
@@ -28,10 +28,10 @@ void SRigVMLogWidget::Construct(
 {
 	ListingModel = FMessageLogListingModel::Create(InArgs._LogName);
 	ListingView = FMessageLogListingViewModel::Create(ListingModel.ToSharedRef(), InArgs._LogLabel);
-	ListingView->SetShowFilters(true);
-	ListingView->SetAllowClear(true);
-	ListingView->SetDiscardDuplicates(false);
-	ListingView->SetScrollToBottom(true);
+	ListingView->SetShowFilters(InArgs._ShowFilters);
+	ListingView->SetAllowClear(InArgs._AllowClear);
+	ListingView->SetDiscardDuplicates(InArgs._DiscardDuplicates);
+	ListingView->SetScrollToBottom(InArgs._ScrollToBottom);
 
 	ChildSlot
 	[
diff --git a/Engine/Plugins/Runtime/RigVM/Source/RigVMEditor/Private/Widgets/SRigVMSwapFunctionsWidget.cpp b/Engine/Plugins/Runtime/RigVM/Source/RigVMEditor/Private/Widgets/SRigVMSwapFunctionsWidget.cpp
new file mode 100644
--- /dev/null
+++ b/Engine/Plugins/Runtime/RigVM/Source/RigVMEditor/Private/Widgets/SRigVMSwapFunctionsWidget.cpp
@@ -0,0 +1,859 @@
+// Copyright Epic Games, Inc. All Rights Reserved.
+
+#include "Widgets/SRigVMSwapFunctionsWidget.h"
+#include "Widgets/SRigVMBulkEditDialog.h"
+
+#include "RigVMBlueprint.h"
+#include "Dialogs/Dialogs.h"
+#include "AssetRegistry/AssetRegistryModule.h"
+#include "Styling/AppStyle.h"
+#include "Algo/Count.h"
+#include "BlueprintEditor.h"
+#include "Editor/RigVMEditor.h"
+#include "Misc/UObjectToken.h"
+#include "RigVMModel/RigVMControllerActions.h"
+
+#define LOCTEXT_NAMESPACE "SRigVMSwapFunctionsWidget"
+
+uint32 FRigVMSwapFunctionContext::GetVisibleChildrenHash() const
+{
+	return HashCombine(FRigVMTreeContext::GetVisibleChildrenHash(), GetTypeHash(SourceIdentifier));
+}
+
+const FSlateBrush* FRigVMTreeFunctionRefNode::GetIconAndTint(FLinearColor& OutColor) const
+{
+	static const FSlateIcon Icon = FSlateIcon(FAppStyle::GetAppStyleSetName(), "Kismet.AllClasses.FunctionIcon");
+	return Icon.GetIcon();
+}
+
+FRigVMTreeFunctionRefGraphNode::FRigVMTreeFunctionRefGraphNode(const URigVMGraph* InFunctionGraph)
+	: FRigVMTreeNode(InFunctionGraph->GetPathName())
+	, WeakGraph(InFunctionGraph)
+{
+	if(const URigVMGraph* ParentGraph = InFunctionGraph->GetParentGraph())
+	{
+		if(ParentGraph->IsA<URigVMFunctionLibrary>())
+		{
+			static const FString FunctionPrefix = TEXT("Function ");
+			OptionalLabel = FText::FromString(FunctionPrefix + InFunctionGraph->GetTypedOuter<URigVMNode>()->GetName());
+		}
+	}
+
+	if(InFunctionGraph->IsRootGraph())
+	{
+		// let's see if there is only one event
+		FString EventName;
+		if(Algo::CountIf(InFunctionGraph->GetNodes(), [&EventName](const URigVMNode* NodeToCount) -> bool
+		{
+			if(NodeToCount->IsEvent() && NodeToCount->CanOnlyExistOnce())
+			{
+				if(EventName.IsEmpty())
+				{
+					if(const URigVMUnitNode* UnitNode = Cast<URigVMUnitNode>(NodeToCount))
+					{
+						if(UnitNode->GetScriptStruct())
+						{
+							EventName = UnitNode->GetScriptStruct()->GetDisplayNameText().ToString();
+						}
+					}
+				}
+				return true;
+			}
+			return false;
+		}) == 1)
+		{
+			static constexpr TCHAR EventGraphNameFormat[] = TEXT("%s Graph");
+			const FString DesiredGraphName = FString::Printf(EventGraphNameFormat, *EventName);
+			OptionalLabel = FText::FromString(DesiredGraphName);
+		}
+		else
+		{
+			if(InFunctionGraph->GetName() == FRigVMClient::RigVMModelPrefix)
+			{
+				OptionalLabel = LOCTEXT("MainGraph", "Main Graph");
+			}
+		}
+	}
+}
+
+TArray<TSharedRef<FRigVMTreeNode>> FRigVMTreeFunctionRefGraphNode::GetChildrenImpl(const TSharedRef<FRigVMTreeContext>& InContext) const
+{
+	if(FunctionRefNodes.IsEmpty())
+	{
+		if(const URigVMGraph* Graph = WeakGraph.Get())
+		{
+			for(const URigVMNode* Node : Graph->GetNodes())
+			{
+				if(const URigVMFunctionReferenceNode* FunctionReferenceNode = Cast<URigVMFunctionReferenceNode>(Node))
+				{
+					if(const TSharedPtr<FRigVMSwapFunctionContext> SwapContext = Cast<FRigVMSwapFunctionContext>(InContext))
+					{
+						if(FunctionReferenceNode->GetFunctionIdentifier() != SwapContext->GetSourceIdentifier())
+						{
+							continue;
+						}
+					}
+					TSharedRef<FRigVMTreeFunctionRefNode> FunctionRefNode = FRigVMTreeFunctionRefNode::Create(FunctionReferenceNode);
+					FunctionRefNodes.Add(FunctionRefNode);
+				}
+			}
+		}
+	}
+	return FunctionRefNodes;
+}
+
+FText FRigVMTreeFunctionRefGraphNode::GetLabel() const
+{
+	if(OptionalLabel.IsSet())
+	{
+		return OptionalLabel.GetValue();
+	}
+	
+	static const FString ContainedGraphSuffix = TEXT(".ContainedGraph");
+
+	const FText Label = FRigVMTreeNode::GetLabel();
+	FString LabelString = Label.ToString();
+	if(LabelString.EndsWith(ContainedGraphSuffix))
+	{
+		LabelString = LabelString.LeftChop(ContainedGraphSuffix.Len());
+	}
+
+	return FText::FromString(LabelString);
+}
+
+const FSlateBrush* FRigVMTreeFunctionRefGraphNode::GetIconAndTint(FLinearColor& OutColor) const
+{
+	static const FSlateIcon Icon(FAppStyle::GetAppStyleSetName(), "GraphEditor.EventGraph_16x");
+	return Icon.GetIcon();
+}
+
+void FRigVMTreeFunctionRefGraphNode::DirtyChildren()
+{
+	FRigVMTreeNode::DirtyChildren();
+	FunctionRefNodes.Reset();
+}
+
+bool FRigVMTreeFunctionRefGraphNode::ContainsFunctionReference(const TSharedRef<FRigVMTreeContext>& InContext) const
+{
+	const TArray<TSharedRef<FRigVMTreeNode>> NewChildren = GetVisibleChildren(InContext);
+	for(const TSharedRef<FRigVMTreeNode>& NewChild : NewChildren)
+	{
+		if(NewChild->IsA<FRigVMTreeFunctionRefNode>())
+		{
+			return true;
+		}
+		if(NewChild->IsA<FRigVMTreeFunctionRefGraphNode>())
+		{
+			if(Cast<FRigVMTreeFunctionRefGraphNode>(NewChild)->ContainsFunctionReference(InContext))
+			{
+				return true;
+			}
+		}
+	}
+	return false;
+}
+
+FRigVMTreeFunctionRefAssetNode::FRigVMTreeFunctionRefAssetNode(const FAssetData& InAssetData)
+	: FRigVMTreePackageNode(InAssetData)
+{
+}
+
+void FRigVMTreeFunctionRefAssetNode::DirtyChildren()
+{
+	FRigVMTreePackageNode::DirtyChildren();
+	LoadedGraphNodes.Reset();
+	MetaDataBasedNodes.Reset();
+	ReferenceNodeDatas.Reset();
+}
+
+TArray<TSharedRef<FRigVMTreeNode>> FRigVMTreeFunctionRefAssetNode::GetChildrenImpl(const TSharedRef<FRigVMTreeContext>& InContext) const
+{
+	if(LoadedGraphNodes.IsEmpty())
+	{
+		const TSoftObjectPtr<UObject> SoftObject(SoftObjectPath);
+		if(const UObject* AssetObject = SoftObject.Get())
+		{
+			if(const IRigVMClientHost* ClientHost = Cast<IRigVMClientHost>(AssetObject))
+			{
+				if(const FRigVMClient* Client = ClientHost->GetRigVMClient())
+				{
+					TArray<URigVMGraph*> Models = Client->GetAllModels(true, true);
+					for(const URigVMGraph* Graph : Models)
+					{
+						TSharedRef<FRigVMTreeFunctionRefGraphNode> GraphNode = FRigVMTreeFunctionRefGraphNode::Create(Graph);
+						if(GraphNode->ContainsFunctionReference(InContext))
+						{
+							LoadedGraphNodes.Add(GraphNode);
+						}
+					}
+				}
+			}
+		}
+
+		// if we can't find anything - let's fall back on asset metadata
+		if(LoadedGraphNodes.IsEmpty() && !IsLoaded())
+		{
+			if(MetaDataBasedNodes.IsEmpty())
+			{
+				if(ReferenceNodeDatas.IsEmpty())
+				{
+					const FAssetData AssetData = GetAssetData();
+					if(AssetData.IsValid())
+					{
+						if(const UClass* Class = AssetData.GetClass())
+						{
+							static const FLazyName FunctionReferenceNodeDataName = TEXT("FunctionReferenceNodeData");
+
+							const FArrayProperty* ReferenceNodeDataProperty =
+							  CastField<FArrayProperty>(Class->FindPropertyByName(FunctionReferenceNodeDataName));
+							if(ReferenceNodeDataProperty)
+							{
+								const FAssetDataTagMapSharedView::FFindTagResult FoundValue = AssetData.TagsAndValues.FindTag(FunctionReferenceNodeDataName);
+								if (FoundValue.IsSet())
+								{
+									const FString ReferenceNodeDataString = FoundValue.AsString();
+									if(!ReferenceNodeDataString.IsEmpty())
+									{
+										ReferenceNodeDataProperty->ImportText_Direct(*ReferenceNodeDataString, &ReferenceNodeDatas, nullptr, EPropertyPortFlags::PPF_None);
+									}
+								}
+							}
+						}
+					}						
+				}
+				
+				for(const FRigVMReferenceNodeData& ReferenceNodeData : ReferenceNodeDatas)
+				{
+					if(const TSharedPtr<FRigVMSwapFunctionContext> SwapContext = Cast<FRigVMSwapFunctionContext>(InContext))
+					{
+						if(ReferenceNodeData.ReferencedFunctionIdentifier != SwapContext->GetSourceIdentifier())
+						{
+							continue;
+						}
+					}
+					TSharedRef<FRigVMTreeFunctionRefNode> RefNode = FRigVMTreeFunctionRefNode::Create(ReferenceNodeData);
+					MetaDataBasedNodes.Add(RefNode);
+				}
+			}
+			return MetaDataBasedNodes;
+		}
+	}
+	return LoadedGraphNodes;
+}
+
+FText FRigVMTreeFunctionIdentifierNode::GetLabel() const
+{
+	static const FString RigVMFunctionLibraryPrefix = TEXT("RigVMFunctionLibrary.");
+	FText Label = FRigVMTreeNode::GetLabel();
+	const FString LabelString = Label.ToString();
+	if(LabelString.StartsWith(RigVMFunctionLibraryPrefix))
+	{
+		Label = FText::FromString(LabelString.Mid(RigVMFunctionLibraryPrefix.Len()));
+	}
+	return Label;
+}
+
+const FSlateBrush* FRigVMTreeFunctionIdentifierNode::GetIconAndTint(FLinearColor& OutColor) const
+{
+	static const FSlateIcon Icon = FSlateIcon(FAppStyle::GetAppStyleSetName(), "Kismet.AllClasses.FunctionIcon");
+	return Icon.GetIcon();
+}
+
+void FRigVMTreeFunctionIdentifierAssetNode::AddChildNode(const TSharedRef<FRigVMTreeNode>& InNode)
+{
+	AddChildImpl(InNode);
+}
+
+void FRigVMTreeFunctionIdentifierAssetNode::DirtyChildren()
+{
+	// no need to dirty things here
+	// since the content of this is set up using a push model
+}
+
+bool FRigVMTreeEmptyFunctionRefGraphFilter::Filters(TSharedRef<FRigVMTreeNode>& InNode, const TSharedRef<FRigVMTreeContext>& InContext)
+{
+	if(InNode->IsA<FRigVMTreeFunctionRefGraphNode>())
+	{
+		if(InNode->GetVisibleChildren(InContext).IsEmpty())
+		{
+			return true;
+		}
+	}
+	return FRigVMTreeFilter::Filters(InNode, InContext);
+}
+
+bool FRigVMTreeEmptyFunctionRefAssetFilter::Filters(TSharedRef<FRigVMTreeNode>& InNode, const TSharedRef<FRigVMTreeContext>& InContext)
+{
+	if(InNode->IsA<FRigVMTreeFunctionRefAssetNode>())
+	{
+		if(InNode->GetVisibleChildren(InContext).IsEmpty())
+		{
+			return true;
+		}
+	}
+	return FRigVMTreeFilter::Filters(InNode, InContext);
+}
+
+FText FRigVMTreeFunctionWithNoRefsFilter::GetLabel() const
+{
+	return LOCTEXT("ShowFunctionsWithoutReferences", "Show Unused Functions");
+}
+
+bool FRigVMTreeFunctionWithNoRefsFilter::Filters(TSharedRef<FRigVMTreeNode>& InNode, const TSharedRef<FRigVMTreeContext>& InContext)
+{
+	if(const TSharedPtr<FRigVMTreeFunctionIdentifierNode> FunctionIdentifierNode = Cast<FRigVMTreeFunctionIdentifierNode>(InNode))
+	{
+		const FRigVMFunctionReferenceArray* References = URigVMBuildData::Get()->FindFunctionReferences(FunctionIdentifierNode->GetIdentifier());
+		if((References == nullptr) || (References->Num() == 0))
+		{
+			return true;
+		}
+	}
+	return FRigVMTreeFilter::Filters(InNode, InContext);
+}
+
+bool FRigVMTreeSourceFunctionFilter::Filters(TSharedRef<FRigVMTreeNode>& InNode, const TSharedRef<FRigVMTreeContext>& InContext)
+{
+	if(const TSharedPtr<FRigVMTreeFunctionRefNode> FunctionRefNode = Cast<FRigVMTreeFunctionRefNode>(InNode))
+	{
+		if(const TSharedPtr<FRigVMSwapFunctionContext> FunctionSwapContext = Cast<FRigVMSwapFunctionContext>(InContext))
+		{
+			if(FunctionRefNode->GetIdentifier() != FunctionSwapContext->GetSourceIdentifier())
+			{
+				return true;
+			}
+		}
+	}
+	return FRigVMTreeFilter::Filters(InNode, InContext);
+}
+
+bool FRigVMTreeTargetFunctionFilter::Filters(TSharedRef<FRigVMTreeNode>& InNode, const TSharedRef<FRigVMTreeContext>& InContext)
+{
+	if(const TSharedPtr<FRigVMTreeFunctionIdentifierNode> FunctionIdentifierNode = Cast<FRigVMTreeFunctionIdentifierNode>(InNode))
+	{
+		if(const TSharedPtr<FRigVMSwapFunctionContext> FunctionSwapContext = Cast<FRigVMSwapFunctionContext>(InContext))
+		{
+			if(FunctionIdentifierNode->GetIdentifier() == FunctionSwapContext->GetSourceIdentifier())
+			{
+				return true;
+			}
+		}
+	}
+	return FRigVMTreeFilter::Filters(InNode, InContext);
+}
+
+FText FRigVMTreeFunctionVariantFilter::GetLabel() const
+{
+	return LOCTEXT("OnlyShowVariants", "Only Show Variants");
+}
+
+bool FRigVMTreeFunctionVariantFilter::Filters(TSharedRef<FRigVMTreeNode>& InNode, const TSharedRef<FRigVMTreeContext>& InContext)
+{
+	if(const TSharedPtr<FRigVMTreeFunctionIdentifierNode> FunctionIdentifierNode = Cast<FRigVMTreeFunctionIdentifierNode>(InNode))
+	{
+		if(const TSharedPtr<FRigVMSwapFunctionContext> FunctionSwapContext = Cast<FRigVMSwapFunctionContext>(InContext))
+		{
+			if(IsFunctionVariant(FunctionSwapContext->GetSourceIdentifier()))
+			{
+				return !IsFunctionVariantOf(
+					FunctionIdentifierNode->GetIdentifier(),
+					FunctionSwapContext->GetSourceIdentifier());
+			}
+		}
+	}
+	return FRigVMTreeFilter::Filters(InNode, InContext);
+}
+
+bool FRigVMTreeFunctionVariantFilter::IsFunctionVariant(const FRigVMGraphFunctionIdentifier& InIdentifier) const
+{
+	if(const bool* bResult = LibraryNodePathToIsVariant.Find(InIdentifier.GetLibraryNodePath()))
+	{
+		return *bResult;
+	}
+	const bool bResult = InIdentifier.IsVariant();
+	LibraryNodePathToIsVariant.Add(InIdentifier.GetLibraryNodePath(), bResult);
+	return bResult;
+}
+
+bool FRigVMTreeFunctionVariantFilter::IsFunctionVariantOf(const FRigVMGraphFunctionIdentifier& InIdentifier, const FRigVMGraphFunctionIdentifier& InSourceIdentifier) const
+{
+	if(const TArray<FRigVMGraphFunctionIdentifier>* Identifiers = LibraryNodePathToVariants.Find(InSourceIdentifier.GetLibraryNodePath()))
+	{
+		return Identifiers->Contains(InIdentifier);
+	}
+	const TArray<FRigVMGraphFunctionIdentifier> Identifiers = InSourceIdentifier.GetVariantIdentifiers();
+	LibraryNodePathToVariants.Add(InSourceIdentifier.GetLibraryNodePath(), Identifiers);
+	return Identifiers.Contains(InIdentifier);
+}
+
+bool FRigVMSwapFunctionTask::Execute(const TSharedRef<FRigVMTreePhase>& InPhase)
+{
+	URigVMFunctionReferenceNode* ReferenceNode = GetReferenceNode(InPhase);
+	if(ReferenceNode == nullptr)
+	{
+		return false;
+	}
+
+	URigVMGraph* Graph = ReferenceNode->GetGraph();
+	check(Graph);
+	IRigVMClientHost* RigVMClientHost = Graph->GetImplementingOuter<IRigVMClientHost>();
+	check(RigVMClientHost);
+
+	URigVMController* Controller = RigVMClientHost->GetOrCreateController(Graph);
+	check(Controller);
+
+	const TSharedRef<FUObjectToken> ReferenceNodeToken = FUObjectToken::Create(ReferenceNode);
+
+	TWeakObjectPtr<URigVMFunctionReferenceNode> WeakRefNode(ReferenceNode);
+	ReferenceNodeToken->OnMessageTokenActivated(FOnMessageTokenActivated::CreateLambda([WeakRefNode](const TSharedRef<IMessageToken>&)
+	{
+		if(WeakRefNode.IsValid())
+		{
+			if(UBlueprint* Blueprint = WeakRefNode.Get()->GetTypedOuter<UBlueprint>())
+			{
+				GEditor->GetEditorSubsystem<UAssetEditorSubsystem>()->OpenEditorForAsset(Blueprint);
+				
+				if(IAssetEditorInstance* Editor = GEditor->GetEditorSubsystem<UAssetEditorSubsystem>()->FindEditorForAsset(Blueprint, true))
+				{
+					if(FRigVMEditor* RigVMEditor = static_cast<FRigVMEditor*>(Editor))
+					{
+						RigVMEditor->HandleJumpToHyperlink(WeakRefNode.Get());
+					}
+				}
+			}
+		}
+	}));
+
+	const TSharedRef<FTokenizedMessage> Message = FTokenizedMessage::Create(EMessageSeverity::Info);
+	Message->AddText(LOCTEXT("Swapping", "Swapping"));
+	Message->AddToken(ReferenceNodeToken);
+	InPhase->GetContext()->LogMessage(Message);
+
+	// always assume success
+	(void)Controller->SwapFunctionReference(ReferenceNode, Identifier, true, IsUndoEnabled(), true);
+	(void)ReferenceNode->MarkPackageDirty();
+	return true;
+}
+
+URigVMFunctionReferenceNode* FRigVMSwapFunctionTask::GetReferenceNode(const TSharedRef<FRigVMTreePhase>& InPhase) const
+{
+	// we expect this to be loaded by now
+	UObject* NodeObject = StaticFindObject(UObject::StaticClass(), NULL, *ObjectPath, false);
+	if(NodeObject == nullptr)
+	{
+		InPhase->GetContext()->LogError(FString::Printf(TEXT("Cannot find reference node '%s'."), *ObjectPath));
+		return nullptr;
+	}
+	
+	URigVMFunctionReferenceNode* ReferenceNode = Cast<URigVMFunctionReferenceNode>(NodeObject);
+	if(ReferenceNode == nullptr)
+	{
+		InPhase->GetContext()->LogError(FString::Printf(TEXT("ObjectPath '%s' doesn't refer to a reference node."), *ObjectPath));
+		return nullptr;
+	}
+	return ReferenceNode;
+}
+
+URigVMBlueprint* FRigVMSwapFunctionTask::GetBlueprint(const TSharedRef<FRigVMTreePhase>& InPhase) const
+{
+	URigVMFunctionReferenceNode* ReferenceNode = GetReferenceNode(InPhase);
+	if(ReferenceNode == nullptr)
+	{
+		return nullptr;
+	}
+	return ReferenceNode->GetTypedOuter<URigVMBlueprint>();
+}
+
+void SRigVMSwapFunctionsWidget::Construct(const FArguments& InArgs)
+{
+	static const TArray<TSharedRef<FRigVMTreeFilter>> DefaultFilters = {
+		FRigVMTreeEngineContentFilter::Create(),
+		FRigVMTreeDeveloperContentFilter::Create()
+	};
+	static const TSharedRef<FRigVMTreeFilter> DefaultPathFilter = FRigVMTreePathFilter::Create();
+
+	PickTargetContext = FRigVMSwapFunctionContext::Create();
+	PickFunctionRefsContext = FRigVMSwapFunctionContext::Create();
+	SourcePreviewEnvironment = MakeShareable(new FRigVMMinimalEnvironment());
+	TargetPreviewEnvironment = MakeShareable(new FRigVMMinimalEnvironment());
+	bSkipPickingFunctionRefs = InArgs._SkipPickingFunctionRefs;
+
+	TOptional<int32> PhaseToActivate;
+
+	SetSourceFunction(InArgs._Source);
+	SetTargetFunction(InArgs._Target);
+
+	TArray<TSharedRef<FRigVMTreePhase>> Phases;
+	if(!InArgs._Source.IsValid())
+	{
+		TSharedRef<FRigVMTreePhase> Phase = FRigVMTreePhase::Create(PHASE_PICKSOURCE, TEXT("Pick Source Function"), FRigVMTreeContext::Create());
+		Phase->GetContext()->Filters = DefaultFilters;
+		Phase->GetContext()->Filters.Add(DefaultPathFilter);
+		Phase->GetContext()->Filters.Add(FRigVMTreeFunctionWithNoRefsFilter::Create());
+		Phase->SetNodes(GetFunctionIdentifierNodes(InArgs));
+		Phase->PrimaryButtonText().Set(LOCTEXT("Next", "Next"));
+		Phase->IsPrimaryButtonVisible().BindLambda([this]()
+		{
+			const TArray<TSharedRef<FRigVMTreeNode>> Selection = GetBulkEditWidget()->GetSelectedNodes();
+			return Selection.ContainsByPredicate([](const TSharedRef<FRigVMTreeNode>& Node)
+			{
+				return Node->IsA<FRigVMTreeFunctionIdentifierNode>();
+			});
+		});
+		Phase->OnPrimaryAction().BindLambda([this]()
+		{
+			const TArray<TSharedRef<FRigVMTreeNode>> Selection = GetBulkEditWidget()->GetSelectedNodes();
+			check(Selection.Num() == 1);
+			if(Cast<FRigVMTreeFunctionIdentifierNode>(Selection[0]))
+			{
+				GetBulkEditWidget()->ActivatePhase(PHASE_PICKTARGET);
+				return FReply::Handled();
+			}
+			return FReply::Unhandled();
+		});
+		PhaseToActivate = PhaseToActivate.Get(Phase->GetID());
+		Phases.Add(Phase);
+	}
+	
+	if(!InArgs._Target.IsValid())
+	{
+		TSharedRef<FRigVMTreePhase> Phase = FRigVMTreePhase::Create(PHASE_PICKTARGET, TEXT("Pick Target Function"), PickTargetContext.ToSharedRef());
+		Phase->GetContext()->Filters = DefaultFilters;
+		Phase->GetContext()->Filters.Add(DefaultPathFilter);
+		Phase->GetContext()->Filters.Add(FRigVMTreeTargetFunctionFilter::Create());
+		Phase->GetContext()->Filters.Add(FRigVMTreeFunctionVariantFilter::Create());
+		Phase->SetNodes(GetFunctionIdentifierNodes(InArgs));
+		Phase->PrimaryButtonText().Set(LOCTEXT("Next", "Next"));
+		Phase->IsPrimaryButtonVisible().BindLambda([this]()
+		{
+			const TArray<TSharedRef<FRigVMTreeNode>> Selection = GetBulkEditWidget()->GetSelectedNodes();
+			return Selection.ContainsByPredicate([](const TSharedRef<FRigVMTreeNode>& Node)
+			{
+				return Node->IsA<FRigVMTreeFunctionIdentifierNode>();
+			});
+		});
+		Phase->OnPrimaryAction().BindLambda([this]()
+		{
+			const TArray<TSharedRef<FRigVMTreeNode>> Selection = GetBulkEditWidget()->GetSelectedNodes();
+			check(Selection.Num() == 1);
+			if(TSharedPtr<FRigVMTreeFunctionIdentifierNode> FunctionIdentifierNode = Cast<FRigVMTreeFunctionIdentifierNode>(Selection[0]))
+			{
+				GetBulkEditWidget()->ActivatePhase(PHASE_PICKFUNCTIONREFS);
+				return FReply::Handled();
+			}
+			return FReply::Unhandled();
+		});
+		PhaseToActivate = PhaseToActivate.Get(Phase->GetID());
+		Phases.Add(Phase);
+	}
+
+	{
+		TSharedRef<FRigVMTreePhase> Phase = FRigVMTreePhase::Create(PHASE_PICKFUNCTIONREFS, TEXT("Pick Function References"), PickFunctionRefsContext.ToSharedRef());
+		Phase->SetAllowsMultiSelection(true);
+		Phase->GetContext()->Filters = DefaultFilters;
+		Phase->GetContext()->Filters.Add(FRigVMTreePathFilter::Create());
+		Phase->GetContext()->Filters.Add(FRigVMTreeEmptyFunctionRefGraphFilter::Create());
+		Phase->GetContext()->Filters.Add(FRigVMTreeEmptyFunctionRefAssetFilter::Create());
+		Phase->GetContext()->Filters.Add(FRigVMTreeSourceFunctionFilter::Create());
+
+		Phase->SetNodes(GetFunctionRefNodes(InArgs));
+		Phase->PrimaryButtonText().Set(LOCTEXT("SwapFunctions", "SwapFunctions"));
+		Phase->IsPrimaryButtonVisible().BindLambda([this]()
+		{
+			return true;
+		});
+		Phase->PrimaryButtonText().BindLambda([this]()
+		{
+			if(GetBulkEditWidget()->HasAnyVisibleCheckedNode())
+			{
+				return LOCTEXT("SwapFunction", "Swap Function");
+			}
+			return LOCTEXT("Done", "Done");
+		});
+		Phase->OnPrimaryAction().BindLambda([this, Phase]()
+		{
+			GetBulkEditWidget()->GetTreeView()->GetTreeView()->ClearSelection();
+			
+			const TArray<TSharedRef<FRigVMTreeNode>> AllCheckedNodes = GetBulkEditWidget()->GetCheckedNodes();
+			if(AllCheckedNodes.IsEmpty())
+			{
+				GetBulkEditWidget()->CloseDialog();
+				return FReply::Handled();
+			}
+			
+			const TArray<TSharedRef<FRigVMTreeNode>> FunctionRefs = AllCheckedNodes.FilterByPredicate([](const TSharedRef<FRigVMTreeNode>& Node) {
+				return Node->IsA<FRigVMTreeFunctionRefNode>();
+			});
+
+			if(FunctionRefs.IsEmpty())
+			{
+				return FReply::Handled();
+			}
+
+			const TSharedPtr<FRigVMSwapFunctionContext> Context = Cast<FRigVMSwapFunctionContext>(Phase->GetContext());
+			check(Context);
+
+			TSet<FString> VisitedPackages;
+			TArray<TSharedRef<FRigVMTreeTask>> Tasks;
+
+			for(const TSharedRef<FRigVMTreeNode>& Node : FunctionRefs)
+			{
+				const TSharedPtr<FRigVMTreeFunctionRefNode> FunctionRefNode = Cast<FRigVMTreeFunctionRefNode>(Node);
+				check(FunctionRefNode);
+				if(FunctionRefNode->GetIdentifier() != Context->GetSourceIdentifier())
+				{
+					continue;
+				}
+
+				const FAssetData AssetData = FunctionRefNode->GetAssetData();
+				const FString PackagePath = AssetData.ToSoftObjectPath().GetLongPackageName();
+
+				if(!VisitedPackages.Contains(PackagePath))
+				{
+					Tasks.Add(FRigVMTreeLoadPackageForNodeTask::Create(FunctionRefNode->GetRoot()));
+					VisitedPackages.Add(PackagePath);
+				}
+				Tasks.Add(FRigVMSwapFunctionTask::Create(FunctionRefNode->GetPath(), Context->GetTargetIdentifier()));
+			}
+
+			if(Tasks.IsEmpty())
+			{
+				return FReply::Unhandled();
+			}
+			GetBulkEditWidget()->QueueTasks(Tasks);
+			return FReply::Handled();
+		});
+		PhaseToActivate = PhaseToActivate.Get(Phase->GetID());
+		Phases.Add(Phase);
+	}
+
+	const TSharedRef<SWidget> FunctionPreviewBox = SNew(SHorizontalBox)
+	+ SHorizontalBox::Slot()
+	.HAlign(HAlign_Left)
+	.VAlign(VAlign_Top)
+	.AutoWidth()
+	[
+		SNew(SRigVMNodePreviewWidget)
+		.Environment(SourcePreviewEnvironment)
+		.Visibility_Lambda([this]()
+		{
+			return PickTargetContext->GetSourceIdentifier().IsValid() ? EVisibility::Visible : EVisibility::Collapsed;
+		})
+	]
+	+ SHorizontalBox::Slot()
+	.HAlign(HAlign_Left)
+	.VAlign(VAlign_Top)
+	.AutoWidth()
+	[
+		SNew(SRigVMNodePreviewWidget)
+		.Environment(TargetPreviewEnvironment)
+		.Visibility_Lambda([this]()
+		{
+			return PickFunctionRefsContext->GetTargetIdentifier().IsValid() ? EVisibility::Visible : EVisibility::Collapsed;
+		})
+	];
+		
+	ChildSlot
+	[
+		SNew(SVerticalBox)
+
+		+ SVerticalBox::Slot()
+		.FillHeight(1.0f)
+		.Padding(0)
+		[
+			SNew(SHorizontalBox)
+			+ SHorizontalBox::Slot()
+			.FillWidth(1.0f)
+			.Padding(0)
+			[
+				SAssignNew(BulkEditWidget, SRigVMBulkEditWidget)
+				.Phases(Phases)
+				.PhaseToActivate(PhaseToActivate.GetValue())
+				.OnPhaseActivated(this, &SRigVMSwapFunctionsWidget::OnPhaseActivated)
+				.OnNodeSelected(this, &SRigVMSwapFunctionsWidget::OnNodeSelected)
+				.OnNodeDoubleClicked(this, &SRigVMSwapFunctionsWidget::OnNodeDoubleClicked)
+				.RightWidget(FunctionPreviewBox)
+				.BulkEditTitle(LOCTEXT("SwapFunctions", "Swap Functions"))
+				.BulkEditConfirmMessage(LOCTEXT("SwapFunctionsConfirmMessage", "This edit is going to swap functions without support for undo. Are you sure?"))
+				.BulkEditConfirmIniField(TEXT("RigVMSwapFunctions_Warning"))
+				.EnableUndo(InArgs._EnableUndo)
+				.CloseOnSuccess(InArgs._CloseOnSuccess)
+			]
+		]
+	];
+}
+
+TArray<TSharedRef<FRigVMTreeNode>> SRigVMSwapFunctionsWidget::GetFunctionIdentifierNodes(const FArguments& InArgs)
+{
+	TArray<TSharedRef<FRigVMTreeNode>> Nodes;
+	TMap<FString, TSharedRef<FRigVMTreeFunctionIdentifierAssetNode>> VisitedAssets;
+	
+	const TArray<FRigVMGraphFunctionIdentifier> Identifiers = URigVMBuildData::Get()->GetAllFunctionIdentifiers(false);
+	for(const FRigVMGraphFunctionIdentifier& Identifier : Identifiers)
+	{
+		const FString PackagePath = Identifier.HostObject.GetLongPackageName();
+		TSharedPtr<FRigVMTreeFunctionIdentifierAssetNode> PackageCategory;
+		if(!VisitedAssets.Contains(PackagePath))
+		{
+			const FAssetData AssetData = FRigVMTreeContext::FindAssetFromAnyPath(PackagePath, false);
+			if(AssetData.IsValid())
+			{
+				PackageCategory = FRigVMTreeFunctionIdentifierAssetNode::Create(AssetData);
+				VisitedAssets.Add(PackagePath, PackageCategory.ToSharedRef());
+				Nodes.Add(PackageCategory.ToSharedRef());
+			}
+			else
+			{
+				continue;
+			}
+		}
+		else
+		{
+			PackageCategory = VisitedAssets.FindChecked(PackagePath);
+		}
+		PackageCategory->AddChildNode(FRigVMTreeFunctionIdentifierNode::Create(Identifier));
+	}
+
+	return Nodes;
+}
+
+TArray<TSharedRef<FRigVMTreeNode>> SRigVMSwapFunctionsWidget::GetFunctionRefNodes(const FArguments& InArgs)
+{
+	TArray<TSharedRef<FRigVMTreeNode>> Nodes;
+
+	for(const URigVMFunctionReferenceNode* FunctionReferenceNode : InArgs._FunctionReferenceNodes)
+	{
+		Nodes.Add(FRigVMTreeFunctionRefNode::Create(FunctionReferenceNode));
+		if(InArgs._SkipPickingFunctionRefs)
+		{
+			Nodes.Last()->SetCheckState(ECheckBoxState::Checked);
+		}
+	}
+
+	if(!InArgs._SkipPickingFunctionRefs)
+	{
+		for(const URigVMGraph* Graph : InArgs._Graphs)
+		{
+			Nodes.Add(FRigVMTreeFunctionRefGraphNode::Create(Graph));
+		}
+		for(const FAssetData& AssetData : InArgs._Assets)
+		{
+			Nodes.Add(FRigVMTreeFunctionRefAssetNode::Create(AssetData));
+		}
+	}
+	return Nodes;
+}
+
+void SRigVMSwapFunctionsWidget::OnPhaseActivated(TSharedRef<FRigVMTreePhase> Phase)
+{
+	switch(Phase->GetID())
+	{
+		case PHASE_PICKSOURCE:
+		{
+			if(PickTargetContext->GetSourceIdentifier().IsValid())
+			{
+				if(const TSharedPtr<FRigVMTreeNode> Node = Phase->FindVisibleNode(PickTargetContext->GetSourceIdentifier().GetLibraryNodePath()))
+				{
+					const TSharedPtr<SRigVMChangesTreeView> TreeView = GetBulkEditWidget()->GetTreeView();
+					TreeView->SetSelection(Node.ToSharedRef(), true);
+				}
+			}
+			break;
+		}
+		case PHASE_PICKTARGET:
+		{
+			if(PickFunctionRefsContext->GetTargetIdentifier().IsValid())
+			{
+				if(const TSharedPtr<FRigVMTreeNode> Node = Phase->FindVisibleNode(PickFunctionRefsContext->GetTargetIdentifier().GetLibraryNodePath()))
+				{
+					const TSharedPtr<SRigVMChangesTreeView> TreeView = GetBulkEditWidget()->GetTreeView();
+					TreeView->SetSelection(Node.ToSharedRef(), true);
+				}
+			}
+			break;
+		}
+		case PHASE_PICKFUNCTIONREFS:
+		{
+			if(bSkipPickingFunctionRefs)
+			{
+				// process to next phase
+				GetBulkEditWidget()->OnPrimaryButtonClicked();
+			}
+			break;
+		}
+		default:
+		{
+			return;
+		}
+	}
+}
+
+FReply SRigVMSwapFunctionsWidget::OnNodeSelected(TSharedRef<FRigVMTreeNode> Node)
+{
+	if(const TSharedPtr<FRigVMTreeFunctionIdentifierNode> FunctionIdentifierNode = Cast<FRigVMTreeFunctionIdentifierNode>(Node))
+	{
+		if(BulkEditWidget->GetActivePhase()->GetID() == PHASE_PICKSOURCE)
+		{
+			SetSourceFunction(FunctionIdentifierNode->GetIdentifier());
+			return FReply::Handled();
+		}
+		if(BulkEditWidget->GetActivePhase()->GetID() == PHASE_PICKTARGET)
+		{
+			SetTargetFunction(FunctionIdentifierNode->GetIdentifier());
+			return FReply::Handled();
+		}
+	}
+	return FReply::Unhandled();
+}
+
+FReply SRigVMSwapFunctionsWidget::OnNodeDoubleClicked(TSharedRef<FRigVMTreeNode> Node)
+{
+	if(Node->IsA<FRigVMTreeFunctionIdentifierNode>() || Node->IsA<FRigVMTreeFunctionRefNode>())
+	{
+		const FAssetData AssetData = Node->GetRoot()->GetAssetData();
+		if(AssetData.IsValid())
+		{
+			// force load
+			if(UObject* TopLevelObject = AssetData.GetSoftObjectPath().GetWithoutSubPath().TryLoad())
+			{
+				if(TopLevelObject->GetClass()->IsChildOf(URigVMBlueprint::StaticClass()))
+				{
+					if(const UObject* ObjectReference = StaticLoadObject(UObject::StaticClass(), nullptr, *Node->GetPath(), nullptr))
+					{
+						GEditor->GetEditorSubsystem<UAssetEditorSubsystem>()->OpenEditorForAsset(TopLevelObject);
+						if(IAssetEditorInstance* Editor = GEditor->GetEditorSubsystem<UAssetEditorSubsystem>()->FindEditorForAsset(TopLevelObject, true))
+						{
+							if(FRigVMEditor* RigVMEditor = static_cast<FRigVMEditor*>(Editor))
+							{
+								RigVMEditor->HandleJumpToHyperlink(ObjectReference);
+								return FReply::Handled();
+							}
+						}
+					}
+				}
+			}
+		}
+	}
+	return FReply::Unhandled();
+}
+
+void SRigVMSwapFunctionsWidget::SetSourceFunction(const FRigVMGraphFunctionIdentifier& InIdentifier)
+{
+	PickTargetContext->SetSourceIdentifier(InIdentifier);
+	PickFunctionRefsContext->SetSourceIdentifier(InIdentifier);
+	if(SourcePreviewEnvironment && InIdentifier.IsValid())
+	{
+		SourcePreviewEnvironment->SetFunctionNode(InIdentifier);
+	}
+}
+
+void SRigVMSwapFunctionsWidget::SetTargetFunction(const FRigVMGraphFunctionIdentifier& InIdentifier)
+{
+	PickFunctionRefsContext->SetTargetIdentifier(InIdentifier);
+	if(TargetPreviewEnvironment && InIdentifier.IsValid())
+	{
+		TargetPreviewEnvironment->SetFunctionNode(InIdentifier);
+	}
+}
+
+#undef LOCTEXT_NAMESPACE
diff --git a/Engine/Plugins/Runtime/RigVM/Source/RigVMEditor/Public/Editor/RigVMEditor.h b/Engine/Plugins/Runtime/RigVM/Source/RigVMEditor/Public/Editor/RigVMEditor.h
--- a/Engine/Plugins/Runtime/RigVM/Source/RigVMEditor/Public/Editor/RigVMEditor.h
+++ b/Engine/Plugins/Runtime/RigVM/Source/RigVMEditor/Public/Editor/RigVMEditor.h
@@ -7,6 +7,7 @@
 #include "RigVMHost.h"
 #include "RigVMBlueprint.h"
 #include "Editor/RigVMDetailsViewWrapperObject.h"
+#include "AssetRegistry/AssetData.h"
 
 class FRigVMEditor;
 
@@ -191,7 +192,11 @@ protected:
 	FRigVMController_BulkEditResult OnRequestBulkEditDialog(URigVMBlueprint* InBlueprint, URigVMController* InController, URigVMLibraryNode* InFunction, ERigVMControllerBulkEditType InEditType);
 	bool OnRequestBreakLinksDialog(TArray<URigVMLink*> InLinks);
 	TRigVMTypeIndex OnRequestPinTypeSelectionDialog(const TArray<TRigVMTypeIndex>& InTypes);
+
+public:
 	void HandleJumpToHyperlink(const UObject* InSubject);
+
+protected:
 	bool UpdateDefaultValueForVariable(FBPVariableDescription& InVariable, bool bUseCDO);
 
 	URigVMController* ActiveController;
@@ -215,6 +220,7 @@ protected:
 	void ToggleExecutionMode();
 	TSharedRef<SWidget> GenerateEventQueueMenuContent();
 	TSharedRef<SWidget> GenerateExecutionModeMenuContent();
+	TSharedRef<SWidget> GenerateBulkEditMenuContent();
 	virtual void GenerateEventQueueMenuContent(FMenuBuilder& MenuBuilder);
 
 	/** Wraps the normal blueprint editor's action menu creation callback */
@@ -302,6 +308,9 @@ protected:
 	virtual void HandleBreakpointActionRequested(const ERigVMBreakpointAction BreakpointAction);
 	virtual bool IsHaltedAtBreakpoint() const;
 	virtual void FrameSelection();
+	virtual void SwapFunctionWithinAsset();
+	virtual void SwapFunctionAcrossProject();
+	virtual void SwapFunctionForAssets(const TArray<FAssetData>& InAssets, bool bSetupUndo);
 
 	/** Once the log is collected update the graph */
 	void UpdateGraphCompilerErrors();
diff --git a/Engine/Plugins/Runtime/RigVM/Source/RigVMEditor/Public/Editor/RigVMGraphDetailCustomization.h b/Engine/Plugins/Runtime/RigVM/Source/RigVMEditor/Public/Editor/RigVMGraphDetailCustomization.h
--- a/Engine/Plugins/Runtime/RigVM/Source/RigVMEditor/Public/Editor/RigVMGraphDetailCustomization.h
+++ b/Engine/Plugins/Runtime/RigVM/Source/RigVMEditor/Public/Editor/RigVMGraphDetailCustomization.h
@@ -17,6 +17,7 @@
 #include "EdGraph/RigVMEdGraphSchema.h"
 #include "SAdvancedTransformInputBox.h"
 #include "Widgets/SRigVMGraphPinNameListValueWidget.h"
+#include "Widgets/SRigVMLogWidget.h"
 #include "HAL/PlatformApplicationMisc.h"
 
 class IDetailLayoutBuilder;
@@ -192,6 +193,11 @@ public:
 	FText GetCurrentAccessSpecifierName() const;
 	void OnAccessSpecifierSelected( TSharedPtr<FRigVMStringWithTag> SpecifierName, ESelectInfo::Type SelectInfo );
 	TSharedRef<ITableRow> HandleGenerateRowAccessSpecifier( TSharedPtr<FRigVMStringWithTag> SpecifierName, const TSharedRef<STableViewBase>& OwnerTable );
+	bool IsFunctionVariant() const;
+	FText GetVariantGuidText() const;
+	//FReply OnSplitVariant();
+	//FReply OnMergeVariant();
+	void RefreshVariantLog();
 
 private:
 
@@ -207,6 +213,9 @@ private:
 	/** The color block widget */
 	TSharedPtr<SColorBlock> ColorBlock;
 
+	/** The log widget used for function variants */
+	TSharedPtr<SRigVMLogWidget> VariantLog;
+
 	/** Set to true if the UI is currently picking a color */
 	bool bIsPickingColor;
 
diff --git a/Engine/Plugins/Runtime/RigVM/Source/RigVMEditor/Public/RigVMEditorModule.h b/Engine/Plugins/Runtime/RigVM/Source/RigVMEditor/Public/RigVMEditorModule.h
--- a/Engine/Plugins/Runtime/RigVM/Source/RigVMEditor/Public/RigVMEditorModule.h
+++ b/Engine/Plugins/Runtime/RigVM/Source/RigVMEditor/Public/RigVMEditorModule.h
@@ -90,7 +90,7 @@ protected:
 	virtual void GetNodeVariablesContextMenuActions(IRigVMClientHost* RigVMClientHost, const URigVMEdGraphNode* EdGraphNode, URigVMNode* ModelNode, UToolMenu* Menu) const;
 	virtual void GetNodeTemplatesContextMenuActions(IRigVMClientHost* RigVMClientHost, const URigVMEdGraphNode* EdGraphNode, URigVMNode* ModelNode, UToolMenu* Menu) const;
 	virtual void GetNodeOrganizationContextMenuActions(IRigVMClientHost* RigVMClientHost, const URigVMEdGraphNode* EdGraphNode, URigVMNode* ModelNode, UToolMenu* Menu) const;
-	virtual void GetNodeVersioningContextMenuActions(IRigVMClientHost* RigVMClientHost, const URigVMEdGraphNode* EdGraphNode, URigVMNode* ModelNode, UToolMenu* Menu) const;
+	virtual void GetNodeVariantContextMenuActions(IRigVMClientHost* RigVMClientHost, const URigVMEdGraphNode* EdGraphNode, URigVMNode* ModelNode, UToolMenu* Menu) const;
 	virtual void GetNodeTestContextMenuActions(IRigVMClientHost* RigVMClientHost, const URigVMEdGraphNode* EdGraphNode, URigVMNode* ModelNode, UToolMenu* Menu) const;
 	virtual void GetPinWorkflowContextMenuActions(IRigVMClientHost* RigVMClientHost, const UEdGraphPin* EdGraphPin, URigVMPin* ModelPin, UToolMenu* Menu) const;
 	virtual void GetPinDebugContextMenuActions(IRigVMClientHost* RigVMClientHost, const UEdGraphPin* EdGraphPin, URigVMPin* ModelPin, UToolMenu* Menu) const;
diff --git a/Engine/Plugins/Runtime/RigVM/Source/RigVMEditor/Public/Widgets/SRigVMBulkEditDialog.h b/Engine/Plugins/Runtime/RigVM/Source/RigVMEditor/Public/Widgets/SRigVMBulkEditDialog.h
new file mode 100644
--- /dev/null
+++ b/Engine/Plugins/Runtime/RigVM/Source/RigVMEditor/Public/Widgets/SRigVMBulkEditDialog.h
@@ -0,0 +1,87 @@
+// Copyright Epic Games, Inc. All Rights Reserved.
+
+#pragma once
+
+#include "SRigVMBulkEditWidget.h"
+#include "Widgets/SWindow.h"
+
+template<typename InWidgetType = SRigVMBulkEditWidget>
+class RIGVMEDITOR_API SRigVMBulkEditDialog : public SWindow
+{
+public:
+
+	SLATE_BEGIN_ARGS(SRigVMBulkEditDialog)
+		: _WindowSize(1200.0f, 1000.0f)
+	{}
+	SLATE_ARGUMENT(FVector2D, WindowSize)
+	SLATE_ARGUMENT(typename InWidgetType::FArguments, WidgetArgs);
+	SLATE_END_ARGS()
+
+	TSharedRef<SRigVMBulkEditWidget> GetBulkEditWidget() const
+	{
+		return Widget->GetBulkEditWidget();
+	}
+
+	void Construct(const FArguments& InArgs)
+	{
+		Widget = SArgumentNew(InArgs._WidgetArgs, InWidgetType);
+		SWindow::Construct(SWindow::FArguments()
+			.Title(this, &SRigVMBulkEditDialog::GetDialogTitle)
+			.SizingRule( ESizingRule::UserSized )
+			.ClientSize(InArgs._WindowSize)
+			[
+				SNew(SHorizontalBox)
+
+				+SHorizontalBox::Slot()
+				.FillWidth(1.f)
+				.HAlign(HAlign_Fill)
+				.VAlign(VAlign_Fill)
+				[
+					SNew(SVerticalBox)
+
+					+SVerticalBox::Slot()
+					.FillHeight(1)
+					.HAlign(HAlign_Fill)
+					.VAlign(VAlign_Fill)
+					.Padding(2)
+					[
+						GetBulkEditWidget()
+					]
+				]
+			]
+		);
+	}
+
+	FText GetDialogTitle() const
+	{
+		return GetBulkEditWidget()->GetDialogTitle();
+	}
+
+	void ShowNormal()
+	{
+		TSharedPtr<SWindow> ParentWindow = FSlateApplication::Get().GetActiveTopLevelRegularWindow();
+
+		if( !ParentWindow.IsValid() )
+		{
+			// Parent to the main frame window
+			if(FModuleManager::Get().IsModuleLoaded("MainFrame"))
+			{
+				const IMainFrameModule& MainFrame = FModuleManager::LoadModuleChecked<IMainFrameModule>("MainFrame");
+				ParentWindow = MainFrame.GetParentWindow();
+			}
+		}
+
+		if(ParentWindow.IsValid())
+		{
+			FSlateApplication::Get().AddWindowAsNativeChild(SharedThis(this), ParentWindow.ToSharedRef(), true);
+		}
+		else
+		{
+			FSlateApplication::Get().AddWindow(SharedThis(this), true);
+		}
+	}
+
+private:
+
+	TSharedPtr<InWidgetType> Widget;
+};
diff --git a/Engine/Plugins/Runtime/RigVM/Source/RigVMEditor/Public/Widgets/SRigVMBulkEditWidget.h b/Engine/Plugins/Runtime/RigVM/Source/RigVMEditor/Public/Widgets/SRigVMBulkEditWidget.h
new file mode 100644
--- /dev/null
+++ b/Engine/Plugins/Runtime/RigVM/Source/RigVMEditor/Public/Widgets/SRigVMBulkEditWidget.h
@@ -0,0 +1,120 @@
+// Copyright Epic Games, Inc. All Rights Reserved.
+
+#pragma once
+
+#include "Widgets/SRigVMChangesTreeView.h"
+#include "Widgets/SRigVMLogWidget.h"
+#include "Widgets/SCompoundWidget.h"
+#include "Widgets/Views/STreeView.h"
+#include "Widgets/Input/SSearchBox.h"
+#include "Interfaces/IMainFrameModule.h"
+#include "ScopedTransaction.h"
+
+DECLARE_DELEGATE_OneParam(FOnPhaseActivated, TSharedRef<FRigVMTreePhase>);
+
+class RIGVMEDITOR_API SRigVMBulkEditWidget : public SCompoundWidget
+{
+public:
+	SLATE_BEGIN_ARGS(SRigVMBulkEditWidget)
+		: _EnableUndo(false)
+		, _CloseOnSuccess(false)
+		, _PhaseToActivate(INDEX_NONE)
+		, _BulkEditConfirmIniField(TEXT("BulkEdit_Warning"))
+	{}
+	SLATE_ARGUMENT(bool, EnableUndo)
+	SLATE_ARGUMENT(bool, CloseOnSuccess)
+	SLATE_ARGUMENT(TArray<TSharedRef<FRigVMTreePhase>>, Phases)
+	SLATE_ARGUMENT(int32, PhaseToActivate)
+	SLATE_EVENT(FOnRigVMTreeNodeSelected, OnNodeSelected)
+	SLATE_EVENT(FOnRigVMTreeNodeDoubleClicked, OnNodeDoubleClicked)
+	SLATE_EVENT(FOnPhaseActivated, OnPhaseActivated)
+	SLATE_ARGUMENT(TSharedPtr<SWidget>, HeaderWidget)
+	SLATE_ARGUMENT(TSharedPtr<SWidget>, FooterWidget)
+	SLATE_ARGUMENT(TSharedPtr<SWidget>, LeftWidget)
+	SLATE_ARGUMENT(TSharedPtr<SWidget>, RightWidget)
+	SLATE_ATTRIBUTE(FText, BulkEditTitle)
+	SLATE_ATTRIBUTE(FText, BulkEditConfirmMessage)
+	SLATE_ATTRIBUTE(FString, BulkEditConfirmIniField)
+	SLATE_END_ARGS()
+
+	virtual ~SRigVMBulkEditWidget() override;
+
+	void Construct(const FArguments& InArgs);
+
+	virtual void Tick(const FGeometry& AllottedGeometry, const double InCurrentTime, const float InDeltaTime) override;
+
+	FText GetDialogTitle() const;
+
+	void QueueTasks(const TArray<TSharedRef<FRigVMTreeTask>>& InTasks);
+	void CancelTasks();
+
+	TSharedRef<SRigVMBulkEditWidget> GetBulkEditWidget()
+	{
+		return SharedThis(this);
+	}
+
+	TSharedPtr<SRigVMChangesTreeView> GetTreeView()
+	{
+		return TreeView;
+	}
+
+	TSharedPtr<SRigVMLogWidget> GetLogWidget()
+	{
+		return BulkEditLogWidget;
+	}
+
+	TSharedPtr<FRigVMTreePhase> GetActivePhasePtr() const;
+	TSharedRef<FRigVMTreePhase> GetActivePhase() const;
+	TSharedPtr<FRigVMTreePhase> FindPhase(int32 InID) const;
+	bool ActivatePhase(int32 InID);
+	TSharedRef<FRigVMTreeContext> GetContext() const;
+
+	TArray<TSharedRef<FRigVMTreeNode>> GetSelectedNodes() const;
+	bool HasAnyVisibleCheckedNode() const;
+	TArray<TSharedRef<FRigVMTreeNode>> GetCheckedNodes() const;
+	
+	FReply OnBackButtonClicked();
+	FReply OnCancelButtonClicked();
+	FReply OnPrimaryButtonClicked();
+
+	void CloseDialog();
+
+private:
+
+	TArray<TSharedRef<FRigVMTreePhase>> Phases;
+	TArray<int32> ActivatedPhaseIDs;
+	
+	TSharedPtr<SRigVMChangesTreeView> TreeView;
+	TSharedPtr<SRigVMLogWidget> BulkEditLogWidget;
+
+	FOnPhaseActivated OnPhaseActivated;
+
+	TAttribute<FText> BulkEditTitle;
+	TAttribute<FText> BulkEditConfirmMessage;
+	TAttribute<FString> BulkEditConfirmIniField;
+
+	bool AreTasksInProgress() const;
+	EVisibility GetTasksProgressVisibility() const;
+	TOptional<float> GetTasksProgressPercentage() const;
+
+	void OnLogMessage(const TSharedRef<FTokenizedMessage>& InMessage) const;
+	void OnScriptException(ELogVerbosity::Type InVerbosity, const TCHAR* InMessage, const TCHAR* InStackMessage);
+
+	mutable FCriticalSection TasksCriticalSection;
+	TArray<TSharedRef<FRigVMTreeTask>> RemainingTasks;
+	TArray<TSharedRef<FRigVMTreeTask>> CompletedTasks;
+	bool bTasksSucceeded = false;
+	bool bShowLog = false;
+	bool bEnableUndo = false;
+	bool bCloseOnSuccess = false;
+	TSharedPtr<FScopedTransaction> Transaction;
+
+	EVisibility GetBackButtonVisibility() const;
+	bool IsBackButtonEnabled() const;
+	EVisibility GetCancelButtonVisibility() const;
+	bool IsCancelButtonEnabled() const;
+	EVisibility GetPrimaryButtonVisibility() const;
+	bool IsPrimaryButtonEnabled() const;
+	FText GetPrimaryButtonText() const;
+	bool IsReadyToClose() const;
+};
diff --git a/Engine/Plugins/Runtime/RigVM/Source/RigVMEditor/Public/Widgets/SRigVMLogWidget.h b/Engine/Plugins/Runtime/RigVM/Source/RigVMEditor/Public/Widgets/SRigVMLogWidget.h
--- a/Engine/Plugins/Runtime/RigVM/Source/RigVMEditor/Public/Widgets/SRigVMLogWidget.h
+++ b/Engine/Plugins/Runtime/RigVM/Source/RigVMEditor/Public/Widgets/SRigVMLogWidget.h
@@ -16,11 +16,19 @@ public:
 	SLATE_BEGIN_ARGS(SRigVMLogWidget)
 		: _LogName(TEXT("RigVMLog"))
 		, _LogLabel(NSLOCTEXT("SRigVMLogWidget", "RigVMLog", "RigVM Log"))
+		, _ShowFilters(true)
+		, _AllowClear(true)
+		, _DiscardDuplicates(false)
+		, _ScrollToBottom(true)
 	{
 	}
 	SLATE_ATTRIBUTE(FOptionalSize, HeightOverride)
 	SLATE_ARGUMENT(FName, LogName)
 	SLATE_ARGUMENT(FText, LogLabel)
+	SLATE_ARGUMENT(bool, ShowFilters);
+	SLATE_ARGUMENT(bool, AllowClear);
+	SLATE_ARGUMENT(bool, DiscardDuplicates);
+	SLATE_ARGUMENT(bool, ScrollToBottom);
 	SLATE_END_ARGS()
 
 	SRigVMLogWidget();
diff --git a/Engine/Plugins/Runtime/RigVM/Source/RigVMEditor/Public/Widgets/SRigVMSwapFunctionsWidget.h b/Engine/Plugins/Runtime/RigVM/Source/RigVMEditor/Public/Widgets/SRigVMSwapFunctionsWidget.h
new file mode 100644
--- /dev/null
+++ b/Engine/Plugins/Runtime/RigVM/Source/RigVMEditor/Public/Widgets/SRigVMSwapFunctionsWidget.h
@@ -0,0 +1,327 @@
+// Copyright Epic Games, Inc. All Rights Reserved.
+
+#pragma once
+
+#include "CoreMinimal.h"
+#include "SGraphPin.h"
+#include "SRigVMGraphNode.h"
+#include "RigVMModel/RigVMGraph.h"
+#include "RigVMModel/Nodes/RigVMFunctionReferenceNode.h"
+#include "RigVMModel/RigVMClient.h"
+#include "Widgets/SRigVMBulkEditWidget.h"
+#include "Widgets/SRigVMNodePreviewWidget.h"
+
+class FRigVMSwapFunctionContext : public FRigVMTreeContext
+{
+public:
+	DEFINE_RIGVMTREETOOLKIT_ELEMENT(FRigVMSwapFunctionContext, FRigVMTreeContext)
+
+	const FRigVMGraphFunctionIdentifier& GetSourceIdentifier()
+	{
+		return SourceIdentifier;
+	}
+
+	void SetSourceIdentifier(const FRigVMGraphFunctionIdentifier& InIdentifier)
+	{
+		SourceIdentifier = InIdentifier;
+	}
+
+	const FRigVMGraphFunctionIdentifier& GetTargetIdentifier()
+	{
+		return TargetIdentifier;
+	}
+
+	void SetTargetIdentifier(const FRigVMGraphFunctionIdentifier& InIdentifier)
+	{
+		TargetIdentifier = InIdentifier;
+	}
+
+	virtual uint32 GetVisibleChildrenHash() const override;
+
+private:
+	FRigVMGraphFunctionIdentifier SourceIdentifier;
+	FRigVMGraphFunctionIdentifier TargetIdentifier;
+};
+
+class FRigVMTreeFunctionRefNode : public FRigVMTreeNode
+{
+public:
+	DEFINE_RIGVMTREETOOLKIT_ELEMENT(FRigVMTreeFunctionRefNode, FRigVMTreeNode)
+
+	FRigVMTreeFunctionRefNode(const URigVMFunctionReferenceNode* InFunctionReferenceNode)
+		: FRigVMTreeNode(InFunctionReferenceNode->GetPathName())
+		, Identifier(InFunctionReferenceNode->GetFunctionIdentifier())
+	{
+	}
+
+	FRigVMTreeFunctionRefNode(const FRigVMReferenceNodeData& InRefNodeData)
+		: FRigVMTreeNode(InRefNodeData.ReferenceNodePath)
+		, Identifier(InRefNodeData.ReferencedFunctionIdentifier)
+	{
+	}
+
+	virtual bool IsCheckable() const override
+	{
+		return true;
+	}
+	
+	const FRigVMGraphFunctionIdentifier& GetIdentifier() const
+	{
+		return Identifier;
+	}
+	
+	virtual const FSlateBrush* GetIconAndTint(FLinearColor& OutColor) const override;
+
+private:
+	FRigVMGraphFunctionIdentifier Identifier;
+};
+
+class FRigVMTreeFunctionRefGraphNode : public FRigVMTreeNode
+{
+public:
+	DEFINE_RIGVMTREETOOLKIT_ELEMENT(FRigVMTreeFunctionRefGraphNode, FRigVMTreeNode)
+
+	FRigVMTreeFunctionRefGraphNode(const URigVMGraph* InFunctionGraph);
+
+	virtual FText GetLabel() const override;
+
+	virtual const FSlateBrush* GetIconAndTint(FLinearColor& OutColor) const override;
+
+	virtual bool IsCheckable() const override
+	{
+		return true;
+	}
+
+	virtual void DirtyChildren() override;
+	bool ContainsFunctionReference(const TSharedRef<FRigVMTreeContext>& InContext) const;
+
+protected:
+	virtual TArray<TSharedRef<FRigVMTreeNode>> GetChildrenImpl(const TSharedRef<FRigVMTreeContext>& InContext) const override;
+
+	TWeakObjectPtr<const URigVMGraph> WeakGraph;
+	TOptional<FText> OptionalLabel;
+	mutable TArray<TSharedRef<FRigVMTreeNode>> FunctionRefNodes;
+};
+
+class FRigVMTreeFunctionRefAssetNode : public FRigVMTreePackageNode
+{
+public:
+	DEFINE_RIGVMTREETOOLKIT_ELEMENT(FRigVMTreeFunctionRefAssetNode, FRigVMTreePackageNode)
+	
+	FRigVMTreeFunctionRefAssetNode(const FAssetData& InAssetData);
+
+	virtual bool IsCheckable() const override
+	{
+		return true;
+	}
+
+	virtual void DirtyChildren() override;
+
+protected:
+	virtual TArray<TSharedRef<FRigVMTreeNode>> GetChildrenImpl(const TSharedRef<FRigVMTreeContext>& InContext) const override;
+
+	mutable TArray<TSharedRef<FRigVMTreeNode>> LoadedGraphNodes;
+	mutable TArray<TSharedRef<FRigVMTreeNode>> MetaDataBasedNodes;
+	mutable TArray<FRigVMReferenceNodeData> ReferenceNodeDatas;
+
+	friend class SRigVMSwapFunctionsWidget;
+	friend class FRigVMTreeEmptyFunctionRefAssetFilter;
+};
+
+class FRigVMTreeFunctionIdentifierNode : public FRigVMTreeNode
+{
+public:
+	DEFINE_RIGVMTREETOOLKIT_ELEMENT(FRigVMTreeFunctionIdentifierNode, FRigVMTreeNode)
+
+	FRigVMTreeFunctionIdentifierNode(const FRigVMGraphFunctionIdentifier& InIdentifier)
+		: FRigVMTreeNode(InIdentifier.GetLibraryNodePath())
+		, Identifier(InIdentifier)
+	{
+	}
+
+	virtual FText GetLabel() const override;
+	virtual bool IsCheckable() const override { return false; }
+	virtual const FSlateBrush* GetIconAndTint(FLinearColor& OutColor) const override;
+
+	const FRigVMGraphFunctionIdentifier& GetIdentifier() const
+	{
+		return Identifier;
+	}
+
+private:
+	FRigVMGraphFunctionIdentifier Identifier;
+};
+
+class FRigVMTreeFunctionIdentifierAssetNode : public FRigVMTreePackageNode
+{
+public:
+	DEFINE_RIGVMTREETOOLKIT_ELEMENT(FRigVMTreeFunctionIdentifierAssetNode, FRigVMTreePackageNode)
+
+	FRigVMTreeFunctionIdentifierAssetNode(const FAssetData& InAssetData)
+		: FRigVMTreePackageNode(InAssetData)
+	{
+	}
+
+	virtual bool ShouldExpandByDefault() const override
+	{
+		return true;
+	}
+
+	virtual bool IsCheckable() const override { return false; }
+	void AddChildNode(const TSharedRef<FRigVMTreeNode>& InNode);
+	virtual void DirtyChildren() override;
+};
+
+class FRigVMTreeEmptyFunctionRefGraphFilter : public FRigVMTreeFilter
+{
+public:
+	DEFINE_RIGVMTREETOOLKIT_ELEMENT(FRigVMTreeEmptyFunctionRefGraphFilter, FRigVMTreeFilter)
+	virtual bool Filters(TSharedRef<FRigVMTreeNode>& InNode, const TSharedRef<FRigVMTreeContext>& InContext) override;
+};
+
+class FRigVMTreeEmptyFunctionRefAssetFilter : public FRigVMTreeFilter
+{
+public:
+	DEFINE_RIGVMTREETOOLKIT_ELEMENT(FRigVMTreeEmptyFunctionRefAssetFilter, FRigVMTreeFilter)
+	virtual bool Filters(TSharedRef<FRigVMTreeNode>& InNode, const TSharedRef<FRigVMTreeContext>& InContext) override;
+};
+
+class FRigVMTreeFunctionWithNoRefsFilter : public FRigVMTreeFilter
+{
+public:
+	DEFINE_RIGVMTREETOOLKIT_ELEMENT(FRigVMTreeFunctionWithNoRefsFilter, FRigVMTreeFilter);
+	virtual bool CanBeToggledInUI() const override
+	{
+		return true;
+	}
+	virtual bool IsInvertedInUI() const override
+	{
+		return true;
+	}
+	virtual FText GetLabel() const override;
+	virtual bool Filters(TSharedRef<FRigVMTreeNode>& InNode, const TSharedRef<FRigVMTreeContext>& InContext) override;
+};
+
+class FRigVMTreeSourceFunctionFilter : public FRigVMTreeFilter
+{
+public:
+	DEFINE_RIGVMTREETOOLKIT_ELEMENT(FRigVMTreeSourceFunctionFilter, FRigVMTreeFilter)
+	virtual bool Filters(TSharedRef<FRigVMTreeNode>& InNode, const TSharedRef<FRigVMTreeContext>& InContext) override;
+};
+
+class FRigVMTreeTargetFunctionFilter : public FRigVMTreeFilter
+{
+public:
+	DEFINE_RIGVMTREETOOLKIT_ELEMENT(FRigVMTreeTargetFunctionFilter, FRigVMTreeFilter)
+	virtual bool Filters(TSharedRef<FRigVMTreeNode>& InNode, const TSharedRef<FRigVMTreeContext>& InContext) override;
+};
+
+class FRigVMTreeFunctionVariantFilter : public FRigVMTreeFilter
+{
+public:
+	DEFINE_RIGVMTREETOOLKIT_ELEMENT(FRigVMTreeFunctionVariantFilter, FRigVMTreeFilter);
+	virtual bool CanBeToggledInUI() const override
+	{
+		return true;
+	}
+	virtual bool IsInvertedInUI() const override
+	{
+		return false;
+	}
+	virtual FText GetLabel() const override;
+	virtual bool Filters(TSharedRef<FRigVMTreeNode>& InNode, const TSharedRef<FRigVMTreeContext>& InContext) override;
+
+private:
+	bool IsFunctionVariant(const FRigVMGraphFunctionIdentifier& InIdentifier) const;
+	bool IsFunctionVariantOf(const FRigVMGraphFunctionIdentifier& InIdentifier, const FRigVMGraphFunctionIdentifier& InSourceIdentifier) const;
+	mutable TMap<FString,bool> LibraryNodePathToIsVariant;
+	mutable TMap<FString,TArray<FRigVMGraphFunctionIdentifier>> LibraryNodePathToVariants;
+};
+
+class FRigVMSwapFunctionTask : public FRigVMTreeTask
+{
+public:
+	DEFINE_RIGVMTREETOOLKIT_ELEMENT(FRigVMSwapFunctionTask, FRigVMTreeTask)
+	
+	FRigVMSwapFunctionTask(const FString InObjectPath, const FRigVMGraphFunctionIdentifier& InIdentifier)
+		: ObjectPath(InObjectPath)
+		, Identifier(InIdentifier)
+	{
+	}
+
+	virtual bool Execute(const TSharedRef<FRigVMTreePhase>& InPhase) override;
+
+	virtual bool RequiresRefresh() const override
+	{
+		return true;
+	}
+
+	virtual bool RequiresUndo() const override
+	{
+		return true;
+	}
+
+	virtual TArray<FString> GetAffectedNodes() const override
+	{
+		return { ObjectPath };
+	}
+	
+private:
+
+	URigVMFunctionReferenceNode* GetReferenceNode(const TSharedRef<FRigVMTreePhase>& InPhase) const;
+	virtual URigVMBlueprint* GetBlueprint(const TSharedRef<FRigVMTreePhase>& InPhase) const override;
+	
+	FString ObjectPath;
+	FRigVMGraphFunctionIdentifier Identifier;
+};
+
+class RIGVMEDITOR_API SRigVMSwapFunctionsWidget : public SCompoundWidget
+{
+public:
+
+	SLATE_BEGIN_ARGS(SRigVMSwapFunctionsWidget)
+		: _SkipPickingFunctionRefs(false)
+		, _EnableUndo(false)
+		, _CloseOnSuccess(false)
+	{}
+	SLATE_ARGUMENT(FRigVMGraphFunctionIdentifier, Source)
+	SLATE_ARGUMENT(FRigVMGraphFunctionIdentifier, Target)
+	SLATE_ARGUMENT(TArray<URigVMFunctionReferenceNode*>, FunctionReferenceNodes)
+	SLATE_ARGUMENT(TArray<URigVMGraph*>, Graphs)
+	SLATE_ARGUMENT(TArray<FAssetData>, Assets)
+	SLATE_ARGUMENT(bool, SkipPickingFunctionRefs)
+	SLATE_ARGUMENT(bool, EnableUndo)
+	SLATE_ARGUMENT(bool, CloseOnSuccess)
+	SLATE_END_ARGS()
+
+	void Construct(const FArguments& InArgs);
+
+	TSharedRef<SRigVMBulkEditWidget> GetBulkEditWidget()
+	{
+		return BulkEditWidget.ToSharedRef();
+	}
+
+private:
+
+	TSharedPtr<FRigVMTreePhase> PickSourceFunctionPhase;
+	TSharedPtr<FRigVMTreePhase> PickTargetFunctionPhase;
+	TSharedPtr<FRigVMTreePhase> PickFunctionReferencesPhase;
+	TSharedPtr<SRigVMBulkEditWidget> BulkEditWidget;
+	TSharedPtr<FRigVMSwapFunctionContext> PickTargetContext;
+	TSharedPtr<FRigVMSwapFunctionContext> PickFunctionRefsContext;
+	TSharedPtr<FRigVMMinimalEnvironment> SourcePreviewEnvironment;
+	TSharedPtr<FRigVMMinimalEnvironment> TargetPreviewEnvironment;
+	bool bSkipPickingFunctionRefs = false;;
+
+	static TArray<TSharedRef<FRigVMTreeNode>> GetFunctionIdentifierNodes(const FArguments& InArgs);
+	static TArray<TSharedRef<FRigVMTreeNode>> GetFunctionRefNodes(const FArguments& InArgs);
+	void OnPhaseActivated(TSharedRef<FRigVMTreePhase> Phase);
+	FReply OnNodeSelected(TSharedRef<FRigVMTreeNode> Node);
+	FReply OnNodeDoubleClicked(TSharedRef<FRigVMTreeNode> Node);
+	void SetSourceFunction(const FRigVMGraphFunctionIdentifier& InIdentifier);
+	void SetTargetFunction(const FRigVMGraphFunctionIdentifier& InIdentifier);
+
+	static constexpr int32 PHASE_PICKSOURCE = 0;
+	static constexpr int32 PHASE_PICKTARGET = 1;
+	static constexpr int32 PHASE_PICKFUNCTIONREFS = 2;
+};
