RigVM: Better support for unit created context

#rnx

diff --git a/Engine/Plugins/Runtime/RigVM/Source/RigVM/Public/RigVMCore/RigVMStruct.h b/Engine/Plugins/Runtime/RigVM/Source/RigVM/Public/RigVMCore/RigVMStruct.h
--- a/Engine/Plugins/Runtime/RigVM/Source/RigVM/Public/RigVMCore/RigVMStruct.h
+++ b/Engine/Plugins/Runtime/RigVM/Source/RigVM/Public/RigVMCore/RigVMStruct.h
@@ -35,6 +35,7 @@ DECLARE_DELEGATE_RetVal_TwoParams(FName, FRigVMCreateExternalVariableDelegate, F
 DECLARE_DELEGATE_RetVal_TwoParams(bool, FRigVMBindPinToExternalVariableDelegate, FString, FString)
 
 struct FRigVMStruct;
+class URigVMController;
 
 /** Context as of why the node was created */
 enum class ERigVMNodeCreatedReason : uint8
@@ -79,6 +80,10 @@ public:
 	/** Get the reason why this node was created */
 	ERigVMNodeCreatedReason GetReason() const { return Reason; }
 
+#if WITH_EDITOR
+	URigVMController* GetController() const { return Controller; }
+#endif
+
 	/** Get the name of this node */
 	FName GetNodeName() const { return NodeName; }
 
@@ -142,7 +147,10 @@ public:
 	FRigVMBindPinToExternalVariableDelegate& GetBindPinToExternalVariableDelegate() { return BindPinToExternalVariableDelegate; }
 
 private:
-	
+
+#if WITH_EDITOR
+	URigVMController* Controller;
+#endif
 	FName NodeName = NAME_None;
 	ERigVMNodeCreatedReason Reason = ERigVMNodeCreatedReason::Unknown;
 	FRigVMGetExternalVariablesDelegate AllExternalVariablesDelegate;
diff --git a/Engine/Plugins/Runtime/RigVM/Source/RigVMDeveloper/Private/RigVMModel/RigVMController.cpp b/Engine/Plugins/Runtime/RigVM/Source/RigVMDeveloper/Private/RigVMModel/RigVMController.cpp
--- a/Engine/Plugins/Runtime/RigVM/Source/RigVMDeveloper/Private/RigVMModel/RigVMController.cpp
+++ b/Engine/Plugins/Runtime/RigVM/Source/RigVMDeveloper/Private/RigVMModel/RigVMController.cpp
@@ -16461,6 +16461,9 @@ void URigVMController::SetupDefaultUnitNodeDelegates(TDelegate<FName(FRigVMExter
 {
 	TWeakObjectPtr<URigVMController> WeakThis(this);
 
+#if WITH_EDITOR
+	UnitNodeCreatedContext.Controller = this;
+#endif
 	UnitNodeCreatedContext.GetAllExternalVariablesDelegate().BindLambda([WeakThis]() -> TArray<FRigVMExternalVariable> {
 		if (WeakThis.IsValid())
 		{
