RigVM: Optimize AddTemplateFromArguments

#rnx

diff --git a/Engine/Plugins/Runtime/RigVM/Source/RigVM/Private/RigVMCore/RigVMRegistry.cpp b/Engine/Plugins/Runtime/RigVM/Source/RigVM/Private/RigVMCore/RigVMRegistry.cpp
--- a/Engine/Plugins/Runtime/RigVM/Source/RigVM/Private/RigVMCore/RigVMRegistry.cpp
+++ b/Engine/Plugins/Runtime/RigVM/Source/RigVM/Private/RigVMCore/RigVMRegistry.cpp
@@ -1752,21 +1752,26 @@ const FRigVMTemplate* FRigVMRegistry_NoLock::AddTemplateFromArguments_NoLock(con
 	// if any of the arguments are wildcards we'll need to update the types
 	for(FRigVMTemplateArgument& Argument : Template.Arguments)
 	{
-		if(Argument.GetNumTypes_NoLock() == 1 && IsWildCardType_NoLock(Argument.GetTypeIndex_NoLock(0)))
+		const int32 NumTypes = Argument.GetNumTypes_NoLock();
+		if(NumTypes == 1)
 		{
-			Argument.InvalidatePermutations(Argument.GetTypeIndex_NoLock(0));
-			if(IsArrayType_NoLock(Argument.GetTypeIndex_NoLock(0)))
+			const TRigVMTypeIndex FirstTypeIndex = Argument.GetTypeIndex_NoLock(0);
+			if(IsWildCardType_NoLock(FirstTypeIndex))
 			{
-				Argument.TypeCategories.Add(FRigVMTemplateArgument::ETypeCategory_ArrayAnyValue);
-			}
-			else
-			{
-				Argument.TypeCategories.Add(FRigVMTemplateArgument::ETypeCategory_SingleAnyValue);
+				Argument.InvalidatePermutations(FirstTypeIndex);
+				if(IsArrayType_NoLock(FirstTypeIndex))
+				{
+					Argument.TypeCategories.Add(FRigVMTemplateArgument::ETypeCategory_ArrayAnyValue);
+				}
+				else
+				{
+					Argument.TypeCategories.Add(FRigVMTemplateArgument::ETypeCategory_SingleAnyValue);
+				}
+				Argument.bUseCategories = true;
+				Argument.TypeIndices.Reset();
+		
+				NumPermutations = FMath::Max(NumPermutations, Argument.GetNumTypes_NoLock());
 			}
-			Argument.bUseCategories = true;
-			Argument.TypeIndices.Reset();
-	
-			NumPermutations = FMath::Max(NumPermutations, Argument.GetNumTypes_NoLock()); 
 		}
 	}
 
@@ -1778,9 +1783,10 @@ const FRigVMTemplate* FRigVMRegistry_NoLock::AddTemplateFromArguments_NoLock(con
 		if (NumArguments == 1)
 		{
 			TSet< TRigVMTypeIndex > PermutationTypes; PermutationTypes.Reserve(NumPermutations);
+			const FRigVMTypeCacheScope_NoLock TypeCache(Template.Arguments[0]);
 			for(int32 Index = 0; Index < NumPermutations; Index++)
 			{
-				const TRigVMTypeIndex ArgType = Template.Arguments[0].GetTypeIndex_NoLock(Index);
+				const TRigVMTypeIndex ArgType = TypeCache.GetTypeIndex_NoLock(Index);
 				if (PermutationTypes.Contains(ArgType))
 				{
 					ToRemove.Add(Index);
@@ -1795,12 +1801,15 @@ const FRigVMTemplate* FRigVMRegistry_NoLock::AddTemplateFromArguments_NoLock(con
 		{
 			TSet< TArray<TRigVMTypeIndex> > PermutationTypes;
 			PermutationTypes.Reserve(NumPermutations);
-			TArray<TRigVMTypeIndex> ArgTypes; ArgTypes.SetNum(NumArguments);			
+			TArray<TRigVMTypeIndex> ArgTypes; ArgTypes.SetNum(NumArguments);
+			TArray<FRigVMTypeCacheScope_NoLock> TypeCaches;
+			TypeCaches.SetNum(NumArguments);
 			for(int32 Index = 0; Index < NumPermutations; Index++)
 			{
 				for(int32 ArgIndex = 0; ArgIndex < NumArguments; ArgIndex++)
 				{
-					ArgTypes[ArgIndex] = Template.Arguments[ArgIndex].GetTypeIndex_NoLock(Index);
+					(void)TypeCaches[ArgIndex].UpdateIfRequired(Template.Arguments[ArgIndex]);
+					ArgTypes[ArgIndex] = TypeCaches[ArgIndex].GetTypeIndex_NoLock(Index);
 				}
 				
 				if (PermutationTypes.Contains(ArgTypes))
@@ -1826,7 +1835,7 @@ const FRigVMTemplate* FRigVMRegistry_NoLock::AddTemplateFromArguments_NoLock(con
 		}
 		NumPermutations -= ToRemove.Num();
 	}
-	
+
 	for(FRigVMTemplateArgument& Argument : Template.Arguments)
 	{
 		Argument.UpdateTypeToPermutations();
diff --git a/Engine/Plugins/Runtime/RigVM/Source/RigVM/Private/RigVMCore/RigVMTemplate.cpp b/Engine/Plugins/Runtime/RigVM/Source/RigVM/Private/RigVMCore/RigVMTemplate.cpp
--- a/Engine/Plugins/Runtime/RigVM/Source/RigVM/Private/RigVMCore/RigVMTemplate.cpp
+++ b/Engine/Plugins/Runtime/RigVM/Source/RigVM/Private/RigVMCore/RigVMTemplate.cpp
@@ -2278,3 +2278,82 @@ uint32 FRigVMTemplate::ComputeTypeHash() const
 	const FRigVMRegistryReadLock _; 
 	return GetTypeHash_NoLock(*this);
 }
+
+FRigVMTypeCacheScope_NoLock::FRigVMTypeCacheScope_NoLock()
+	: Argument(nullptr)
+	, bShouldCopyTypes(true)
+	, NumTypes(0)
+	, Types(TArray<TRigVMTypeIndex>())
+{
+}
+
+FRigVMTypeCacheScope_NoLock::FRigVMTypeCacheScope_NoLock(const FRigVMTemplateArgument& InArgument)
+	: Argument(&InArgument)
+	, bShouldCopyTypes(InArgument.FilterType != nullptr)//  ||bUseCategories)??
+{
+}
+
+FRigVMTypeCacheScope_NoLock::~FRigVMTypeCacheScope_NoLock()
+{
+}
+
+const FRigVMTypeCacheScope_NoLock& FRigVMTypeCacheScope_NoLock::UpdateIfRequired(const FRigVMTemplateArgument& InArgument)
+{
+	if(Argument != &InArgument)
+	{
+		*this = FRigVMTypeCacheScope_NoLock(InArgument);
+	}
+	return *this;
+}
+
+int32 FRigVMTypeCacheScope_NoLock::GetNumTypes_NoLock() const
+{
+	if(!NumTypes.IsSet())
+	{
+		if(bShouldCopyTypes)
+		{
+			UpdateTypesIfRequired();
+			NumTypes = Types.GetValue().Num();
+		}
+		else
+		{
+			check(Argument);
+			NumTypes = Argument->GetNumTypes_NoLock();
+		}
+	}
+	return NumTypes.GetValue();
+}
+
+TRigVMTypeIndex FRigVMTypeCacheScope_NoLock::GetTypeIndex_NoLock(int32 InIndex) const
+{
+	if(!Types.IsSet())
+	{
+		if(bShouldCopyTypes)
+		{
+			UpdateTypesIfRequired();
+		}
+		else
+		{
+			return Argument->GetTypeIndex_NoLock(InIndex);
+		}
+	}
+	return Types.GetValue()[InIndex];
+}
+
+void FRigVMTypeCacheScope_NoLock::UpdateTypesIfRequired() const
+{
+	check(bShouldCopyTypes);
+	if(Types.IsSet())
+	{
+		return;
+	}
+	Types = TArray<TRigVMTypeIndex>();
+	TArray<TRigVMTypeIndex>& TypesArray = Types.GetValue();
+	TypesArray.Reserve(Argument->GetNumTypes_NoLock());
+
+	Argument->ForEachType([&TypesArray](const TRigVMTypeIndex Type)
+	{
+		TypesArray.Add(Type);
+		return true;
+	});
+}
diff --git a/Engine/Plugins/Runtime/RigVM/Source/RigVM/Public/RigVMCore/RigVMTemplate.h b/Engine/Plugins/Runtime/RigVM/Source/RigVM/Public/RigVMCore/RigVMTemplate.h
--- a/Engine/Plugins/Runtime/RigVM/Source/RigVM/Public/RigVMCore/RigVMTemplate.h
+++ b/Engine/Plugins/Runtime/RigVM/Source/RigVM/Public/RigVMCore/RigVMTemplate.h
@@ -315,6 +315,7 @@ protected:
 	friend struct FRigVMRegistry_NoLock;
 	friend struct FRigVMStructUpgradeInfo;
 	friend class URigVMCompiler;
+	friend class FRigVMTypeCacheScope_NoLock;
 
 private:
 	struct CategoryViews
@@ -583,3 +584,25 @@ private:
 	friend class URigVMLibraryNode;
 	friend struct FRigVMDispatchFactory;
 };
+
+class FRigVMTypeCacheScope_NoLock
+{
+public:
+	FRigVMTypeCacheScope_NoLock();
+	FRigVMTypeCacheScope_NoLock(const FRigVMTemplateArgument& InArgument);
+	~FRigVMTypeCacheScope_NoLock();
+
+	bool IsValid() const { return Argument != nullptr; }
+	const FRigVMTypeCacheScope_NoLock& UpdateIfRequired(const FRigVMTemplateArgument& InArgument);
+	int32 GetNumTypes_NoLock() const;
+	TRigVMTypeIndex GetTypeIndex_NoLock(int32 InIndex) const;
+	
+private:
+
+	void UpdateTypesIfRequired() const;
+
+	const FRigVMTemplateArgument* Argument;
+	bool bShouldCopyTypes;
+	mutable TOptional<int32> NumTypes;
+	mutable TOptional<TArray<TRigVMTypeIndex>> Types;
+};
