RigVM: Build fix for Mac

#rnx

diff --git a/Engine/Plugins/Runtime/RigVM/Source/RigVMEditor/Public/Widgets/SRigVMBulkEditDialog.h b/Engine/Plugins/Runtime/RigVM/Source/RigVMEditor/Public/Widgets/SRigVMBulkEditDialog.h
--- a/Engine/Plugins/Runtime/RigVM/Source/RigVMEditor/Public/Widgets/SRigVMBulkEditDialog.h
+++ b/Engine/Plugins/Runtime/RigVM/Source/RigVMEditor/Public/Widgets/SRigVMBulkEditDialog.h
@@ -3,6 +3,8 @@
 #pragma once
 
 #include "SRigVMBulkEditWidget.h"
+
+#include "Framework/Application/SlateApplication.h"
 #include "Widgets/SWindow.h"
 
 template<typename InWidgetType = SRigVMBulkEditWidget>
