RIgVM: Show asset variants in the class settings details panel
#rnx

diff --git a/Engine/Plugins/Runtime/RigVM/Source/RigVM/Public/RigVMObjectVersion.h b/Engine/Plugins/Runtime/RigVM/Source/RigVM/Public/RigVMObjectVersion.h
--- a/Engine/Plugins/Runtime/RigVM/Source/RigVM/Public/RigVMObjectVersion.h
+++ b/Engine/Plugins/Runtime/RigVM/Source/RigVM/Public/RigVMObjectVersion.h
@@ -45,6 +45,9 @@ struct RIGVM_API FRigVMObjectVersion
 		// Adding variant struct to function identifier
 		AddVariantToFunctionIdentifier,
 
+		// Adding variant to every RigVM asset
+		AddVariantToRigVMAssets,
+
 		// -----<new versions can be added above this line>-------------------------------------------------
 		VersionPlusOne,
 		LatestVersion = VersionPlusOne - 1
diff --git a/Engine/Plugins/Runtime/RigVM/Source/RigVMDeveloper/Private/RigVMBlueprint.cpp b/Engine/Plugins/Runtime/RigVM/Source/RigVMDeveloper/Private/RigVMBlueprint.cpp
--- a/Engine/Plugins/Runtime/RigVM/Source/RigVMDeveloper/Private/RigVMBlueprint.cpp
+++ b/Engine/Plugins/Runtime/RigVM/Source/RigVMDeveloper/Private/RigVMBlueprint.cpp
@@ -203,8 +203,6 @@ URigVMBlueprint::URigVMBlueprint(const FObjectInitializer& ObjectInitializer)
 	CompileLog.SetSourcePath(GetPathName());
 #endif
 
-	AssetVariant.Guid = FRigVMVariant::GenerateGUID();
-
 	if(GetClass() == URigVMBlueprint::StaticClass())
 	{
 		CommonInitialization(ObjectInitializer);
@@ -783,6 +781,8 @@ IRigVMEditorModule* URigVMBlueprint::GetEditorModule() const
 
 void URigVMBlueprint::Serialize(FArchive& Ar)
 {
+	Ar.UsingCustomVersion(FRigVMObjectVersion::GUID);
+	
 	if(IsValid(this))
 	{
 		RigVMClient.SetOuterClientHost(this, GET_MEMBER_NAME_CHECKED(URigVMBlueprint, RigVMClient));
@@ -831,6 +831,11 @@ void URigVMBlueprint::Serialize(FArchive& Ar)
 		{
 			EdGraph->Schema = GetRigVMEdGraphSchemaClass();
 		}
+
+		if (Ar.CustomVer(FRigVMObjectVersion::GUID) < FRigVMObjectVersion::AddVariantToRigVMAssets)
+		{
+			AssetVariant.Guid = FRigVMVariant::GenerateGUID(GetPackage()->GetPathName());
+		}
 	}
 }
 
@@ -1093,7 +1098,7 @@ void URigVMBlueprint::PostLoad()
 
 	if (!AssetVariant.Guid.IsValid())
 	{
-		AssetVariant.Guid = FRigVMVariant::GenerateGUID(GetPackage()->GetPathName());
+		AssetVariant.Guid = FRigVMVariant::GenerateGUID();
 	}
 
 	if (UPackage* Package = GetOutermost())
diff --git a/Engine/Plugins/Runtime/RigVM/Source/RigVMDeveloper/Private/RigVMModel/RigVMBuildData.cpp b/Engine/Plugins/Runtime/RigVM/Source/RigVMDeveloper/Private/RigVMModel/RigVMBuildData.cpp
--- a/Engine/Plugins/Runtime/RigVM/Source/RigVMDeveloper/Private/RigVMModel/RigVMBuildData.cpp
+++ b/Engine/Plugins/Runtime/RigVM/Source/RigVMDeveloper/Private/RigVMModel/RigVMBuildData.cpp
@@ -551,6 +551,37 @@ TArray<FRigVMVariantRef> URigVMBuildData::FindFunctionVariantRefs(const FGuid& I
 	return Result;
 }
 
+TArray<FRigVMVariantRef> URigVMBuildData::GatherAllAssetVariantRefs()
+{
+	TArray<FRigVMVariantRef> Result;
+	const FAssetRegistryModule& AssetRegistryModule = FModuleManager::LoadModuleChecked<FAssetRegistryModule>("AssetRegistry");
+	TArray<FAssetData> Assets;
+	FARFilter AssetFilter;
+	AssetFilter.ClassPaths.Add(URigVMBlueprint::StaticClass()->GetClassPathName());
+	AssetFilter.bRecursiveClasses = true;
+	AssetRegistryModule.Get().GetAssets(AssetFilter, Assets);
+
+	for (const FAssetData& Asset : Assets)
+	{
+		static const FName AssetVariantPropertyName = GET_MEMBER_NAME_CHECKED(URigVMBlueprint, AssetVariant);
+		const FProperty* AssetVariantProperty = CastField<FProperty>(URigVMBlueprint::StaticClass()->FindPropertyByName(AssetVariantPropertyName));
+		const FString VariantStr = Asset.GetTagValueRef<FString>(AssetVariantPropertyName);
+		FRigVMVariant AssetVariant;
+		AssetVariantProperty->ImportText_Direct(*VariantStr, &AssetVariant, nullptr, EPropertyPortFlags::PPF_None);
+		FRigVMVariantRef VariantRef(Asset.ToSoftObjectPath(), AssetVariant);
+		Result.Add(VariantRef);
+	}
+
+	return Result;
+}
+
+TArray<FRigVMVariantRef> URigVMBuildData::FindAssetVariantRefs(const FGuid& InGuid)
+{
+	TArray<FRigVMVariantRef> Result = GatherAllAssetVariantRefs();
+	Result = Result.FilterByPredicate([InGuid](const FRigVMVariantRef& VariantRef) { return VariantRef.Variant.Guid == InGuid; });
+	return Result;
+}
+
 #if WITH_EDITOR
 
 TArray<FRigVMGraphFunctionIdentifier> URigVMBuildData::GetAllFunctionIdentifiers(bool bOnlyPublic) const
diff --git a/Engine/Plugins/Runtime/RigVM/Source/RigVMDeveloper/Public/RigVMBlueprint.h b/Engine/Plugins/Runtime/RigVM/Source/RigVMDeveloper/Public/RigVMBlueprint.h
--- a/Engine/Plugins/Runtime/RigVM/Source/RigVMDeveloper/Public/RigVMBlueprint.h
+++ b/Engine/Plugins/Runtime/RigVM/Source/RigVMDeveloper/Public/RigVMBlueprint.h
@@ -505,7 +505,7 @@ protected:
 #if WITH_EDITORONLY_DATA
 
 	/** Variant information about this asset */
-	UPROPERTY(AssetRegistrySearchable)
+	UPROPERTY(EditAnywhere, AssetRegistrySearchable, Category = "Variant")
 	FRigVMVariant AssetVariant;
 
 	UPROPERTY(transient, DuplicateTransient)
@@ -826,6 +826,7 @@ private:
 	friend class URigVMEdGraphSchema;
 	friend struct FRigVMEdGraphSchemaAction_PromoteToVariable;
 	friend class URigVMBuildData;
+	friend class FRigVMVariantDetailCustomization;
 };
 
 class RIGVMDEVELOPER_API FRigVMBlueprintCompileScope
diff --git a/Engine/Plugins/Runtime/RigVM/Source/RigVMDeveloper/Public/RigVMModel/RigVMBuildData.h b/Engine/Plugins/Runtime/RigVM/Source/RigVMDeveloper/Public/RigVMModel/RigVMBuildData.h
--- a/Engine/Plugins/Runtime/RigVM/Source/RigVMDeveloper/Public/RigVMModel/RigVMBuildData.h
+++ b/Engine/Plugins/Runtime/RigVM/Source/RigVMDeveloper/Public/RigVMModel/RigVMBuildData.h
@@ -128,6 +128,12 @@ public:
 	// Will find all the function variants matching the given variant guid
 	TArray<FRigVMVariantRef> FindFunctionVariantRefs(const FGuid& InGuid);
 
+	// Will find all asset variant refs
+	TArray<FRigVMVariantRef> GatherAllAssetVariantRefs();
+
+	// Will find all the asset variants matching the given variant guid
+	TArray<FRigVMVariantRef> FindAssetVariantRefs(const FGuid& InGuid);
+
 #if WITH_EDITOR
 	// Returns all known public function identifiers used in the project
 	TArray<FRigVMGraphFunctionIdentifier> GetAllFunctionIdentifiers(bool bOnlyPublic = true) const;
diff --git a/Engine/Plugins/Runtime/RigVM/Source/RigVMEditor/Private/Editor/RigVMVariantDetailCustomization.cpp b/Engine/Plugins/Runtime/RigVM/Source/RigVMEditor/Private/Editor/RigVMVariantDetailCustomization.cpp
new file mode 100644
--- /dev/null
+++ b/Engine/Plugins/Runtime/RigVM/Source/RigVMEditor/Private/Editor/RigVMVariantDetailCustomization.cpp
@@ -0,0 +1,130 @@
+// Copyright Epic Games, Inc. All Rights Reserved.
+
+#include "Editor/RigVMVariantDetailCustomization.h"
+
+#include "DetailLayoutBuilder.h"
+#include "DetailWidgetRow.h"
+#include "IDetailChildrenBuilder.h"
+#include "Misc/UObjectToken.h"
+#include "Widgets/Text/STextBlock.h"
+#include "Widgets/SRigVMLogWidget.h"
+
+#define LOCTEXT_NAMESPACE "RigVMVariantDetailCustomization"
+
+class FUObjectToken;
+
+void FRigVMVariantDetailCustomization::CustomizeHeader(TSharedRef<IPropertyHandle> InStructPropertyHandle, FDetailWidgetRow& HeaderRow, IPropertyTypeCustomizationUtils& StructCustomizationUtils)
+{
+	HeaderRow
+	.NameContent()
+	[
+		InStructPropertyHandle->CreatePropertyNameWidget()
+	]
+	.ValueContent()
+	[
+		InStructPropertyHandle->CreatePropertyValueWidget()
+	];
+
+	TArray<UObject*> Objects;
+	InStructPropertyHandle->GetOuterObjects(Objects);
+	ensure(Objects.Num() == 1); // This is in here to ensure we are only showing the modifier details in the blueprint editor
+
+	for (UObject* Object : Objects)
+	{
+		if (Object->IsA<URigVMBlueprint>())
+		{
+			BlueprintBeingCustomized = Cast<URigVMBlueprint>(Object);
+		}
+	}
+}
+
+void FRigVMVariantDetailCustomization::CustomizeChildren(TSharedRef<IPropertyHandle> InStructPropertyHandle, IDetailChildrenBuilder& StructBuilder, IPropertyTypeCustomizationUtils& StructCustomizationUtils)
+{
+	if (InStructPropertyHandle->IsValidHandle())
+	{
+		uint32 NumChildren = 0;
+		InStructPropertyHandle->GetNumChildren(NumChildren);
+
+		for (uint32 ChildIndex = 0; ChildIndex < NumChildren; ChildIndex++)
+		{
+			StructBuilder.AddProperty(InStructPropertyHandle->GetChildHandle(ChildIndex).ToSharedRef());
+		}
+
+		// matching variants
+		StructBuilder.AddCustomRow(FText::GetEmpty())
+		.Visibility(TAttribute<EVisibility>::CreateLambda([this]()
+		{
+			return IsAssetVariant() ? EVisibility::Visible : EVisibility::Collapsed;
+		}))
+		.NameContent()
+		[
+			SNew(STextBlock)
+			.Text(FText::FromString(TEXT("Matching Variants")))
+			.Font(IDetailLayoutBuilder::GetDetailFont())
+		]
+		.ValueContent()
+		[
+			SAssignNew(VariantLog, SRigVMLogWidget)
+			.LogLabel(LOCTEXT("Variants", "Variants"))
+			.LogName(TEXT("RigVMAssetVariants"))
+			.ShowFilters(false)
+			.AllowClear(false)
+			.DiscardDuplicates(true)
+			.ScrollToBottom(false)
+			.HeightOverride(120)
+		];
+	}
+
+	RefreshVariantLog();
+}
+
+bool FRigVMVariantDetailCustomization::IsAssetVariant() const
+{
+	if (BlueprintBeingCustomized)
+	{
+		const FRigVMVariant& Variant = BlueprintBeingCustomized->AssetVariant;
+		TArray<FRigVMVariantRef> Variants = URigVMBuildData::Get()->FindAssetVariantRefs(Variant.Guid);
+		return Variants.Num() > 1;
+	}
+	return false;
+}
+
+FText FRigVMVariantDetailCustomization::GetVariantGuidText() const
+{
+	if (BlueprintBeingCustomized)
+	{
+		const FGuid Guid = BlueprintBeingCustomized->AssetVariant.Guid;
+		if(Guid.IsValid())
+		{
+			return FText::FromString(Guid.ToString(EGuidFormats::DigitsWithHyphensLower));
+		}
+	}
+	return FText();
+}
+
+void FRigVMVariantDetailCustomization::RefreshVariantLog()
+{
+	check(VariantLog);
+
+	VariantLog->GetListing()->ClearMessages();
+
+	if (BlueprintBeingCustomized)
+	{
+		TArray<FRigVMVariantRef> Variants = URigVMBuildData::Get()->FindAssetVariantRefs(BlueprintBeingCustomized->AssetVariant.Guid);
+		Variants = Variants.FilterByPredicate([this](const FRigVMVariantRef& VariantRef) { return VariantRef.ObjectPath != BlueprintBeingCustomized->GetPathName(); });
+		Variants.Sort([](const FRigVMVariantRef& A, const FRigVMVariantRef& B)
+		{
+			return A.ObjectPath.ToString().Compare(B.ObjectPath.ToString()) < 0;
+		});
+
+		for (FRigVMVariantRef& Variant : Variants)
+		{
+			const TSharedRef<FAssetNameToken> AssetToken = FAssetNameToken::Create(Variant.ObjectPath.ToString());
+			const TSharedRef<FTokenizedMessage> Message = FTokenizedMessage::Create(EMessageSeverity::Info);
+			Message->AddToken(AssetToken);
+			VariantLog->GetListing()->AddMessage(Message);
+		}
+	}
+}
+
+#undef LOCTEXT_NAMESPACE
diff --git a/Engine/Plugins/Runtime/RigVM/Source/RigVMEditor/Private/RigVMEditorModule.cpp b/Engine/Plugins/Runtime/RigVM/Source/RigVMEditor/Private/RigVMEditorModule.cpp
--- a/Engine/Plugins/Runtime/RigVM/Source/RigVMEditor/Private/RigVMEditorModule.cpp
+++ b/Engine/Plugins/Runtime/RigVM/Source/RigVMEditor/Private/RigVMEditorModule.cpp
@@ -40,6 +40,7 @@
 #include "RigVMFunctions/Debug/RigVMFunction_VisualDebug.h"
 #include "ScopedTransaction.h"
 #include "Editor/RigVMEditorTools.h"
+#include "Editor/RigVMVariantDetailCustomization.h"
 #include "UObject/UObjectIterator.h"
 #include "Widgets/SRigVMSwapFunctionsWidget.h"
 #include "Widgets/SRigVMBulkEditDialog.h"
@@ -133,6 +134,12 @@ void FRigVMEditorModule::StartupModuleCommon()
 
 	// Register to fixup newly created BPs
 	FKismetEditorUtilities::RegisterOnBlueprintCreatedCallback(this, URigVMHost::StaticClass(), FKismetEditorUtilities::FOnBlueprintCreated::CreateRaw(this, &FRigVMEditorModule::HandleNewBlueprintCreated));
+	
+	FPropertyEditorModule& PropertyEditorModule = FModuleManager::LoadModuleChecked<FPropertyEditorModule>("PropertyEditor");
+    PropertiesToUnregisterOnShutdown.Reset();
+
+	PropertiesToUnregisterOnShutdown.Add(FRigVMVariant::StaticStruct()->GetFName());
+	PropertyEditorModule.RegisterCustomPropertyTypeLayout(PropertiesToUnregisterOnShutdown.Last(), FOnGetPropertyTypeCustomizationInstance::CreateStatic(&FRigVMVariantDetailCustomization::MakeInstance));
 }
 
 void FRigVMEditorModule::ShutdownModule()
diff --git a/Engine/Plugins/Runtime/RigVM/Source/RigVMEditor/Public/Editor/RigVMVariantDetailCustomization.h b/Engine/Plugins/Runtime/RigVM/Source/RigVMEditor/Public/Editor/RigVMVariantDetailCustomization.h
new file mode 100644
--- /dev/null
+++ b/Engine/Plugins/Runtime/RigVM/Source/RigVMEditor/Public/Editor/RigVMVariantDetailCustomization.h
@@ -0,0 +1,34 @@
+// Copyright Epic Games, Inc. All Rights Reserved.
+
+#pragma once
+
+#include "CoreMinimal.h"
+#include "IPropertyTypeCustomization.h"
+#include "RigVMBlueprint.h"
+#include "Widgets/SRigVMLogWidget.h"
+
+class IPropertyHandle;
+
+class FRigVMVariantDetailCustomization : public IPropertyTypeCustomization
+{
+public:
+
+	static TSharedRef<IPropertyTypeCustomization> MakeInstance()
+	{
+		return MakeShareable(new FRigVMVariantDetailCustomization);
+	}
+
+	/** IPropertyTypeCustomization interface */
+	virtual void CustomizeHeader(TSharedRef<class IPropertyHandle> InStructPropertyHandle, class FDetailWidgetRow& HeaderRow, IPropertyTypeCustomizationUtils& StructCustomizationUtils) override;
+	virtual void CustomizeChildren(TSharedRef<class IPropertyHandle> InStructPropertyHandle, class IDetailChildrenBuilder& StructBuilder, IPropertyTypeCustomizationUtils& StructCustomizationUtils) override;
+
+	bool IsAssetVariant() const;
+	FText GetVariantGuidText() const;
+	void RefreshVariantLog();
+protected:
+
+	URigVMBlueprint* BlueprintBeingCustomized;
+
+	/** The log widget used for function variants */
+	TSharedPtr<SRigVMLogWidget> VariantLog;
+};
diff --git a/Engine/Plugins/Runtime/RigVM/Source/RigVMEditor/Public/RigVMEditorModule.h b/Engine/Plugins/Runtime/RigVM/Source/RigVMEditor/Public/RigVMEditorModule.h
--- a/Engine/Plugins/Runtime/RigVM/Source/RigVMEditor/Public/RigVMEditorModule.h
+++ b/Engine/Plugins/Runtime/RigVM/Source/RigVMEditor/Public/RigVMEditorModule.h
@@ -120,6 +120,9 @@ protected:
 
 private:
 
+	/** StaticStruct is not safe on shutdown, so we cache the name, and use this to unregister on shut down */
+	TArray<FName> PropertiesToUnregisterOnShutdown;
+
 	void HandleNewBlueprintCreated(UBlueprint* InBlueprint);
 
 	bool ShowWorkflowOptionsDialog(URigVMUserWorkflowOptions* InOptions) const;
