RigVM: Fix mac compilation errorr
#rnx

diff --git a/Engine/Plugins/Runtime/RigVM/Source/RigVM/Public/RigVMCore/RigVMVariant.h b/Engine/Plugins/Runtime/RigVM/Source/RigVM/Public/RigVMCore/RigVMVariant.h
--- a/Engine/Plugins/Runtime/RigVM/Source/RigVM/Public/RigVMCore/RigVMVariant.h
+++ b/Engine/Plugins/Runtime/RigVM/Source/RigVM/Public/RigVMCore/RigVMVariant.h
@@ -73,6 +73,11 @@ struct RIGVM_API FRigVMVariantRef
 {
 	GENERATED_BODY()
 	
+	FRigVMVariantRef(){}
+
+	FRigVMVariantRef(const FSoftObjectPath& InPath, const FRigVMVariant& InVariant)
+		: ObjectPath(InPath), Variant(InVariant) {}
+	
 	FSoftObjectPath ObjectPath;
 	FRigVMVariant Variant;
 };
