Control Rig: Fix creating variant of function library not creating variants of the functions properly when library is not saved first
#rnx

diff --git a/Engine/Plugins/Animation/ControlRig/Source/ControlRigDeveloper/Private/ControlRigBlueprint.cpp b/Engine/Plugins/Animation/ControlRig/Source/ControlRigDeveloper/Private/ControlRigBlueprint.cpp
--- a/Engine/Plugins/Animation/ControlRig/Source/ControlRigDeveloper/Private/ControlRigBlueprint.cpp
+++ b/Engine/Plugins/Animation/ControlRig/Source/ControlRigDeveloper/Private/ControlRigBlueprint.cpp
@@ -1970,6 +1970,7 @@ void UControlRigBlueprint::PatchFunctionsOnLoad()
 			if (!Variant)
 			{
 				Header.Variant.Guid = FRigVMVariant::GenerateGUID(Header.LibraryPointer.GetLibraryNodePath());
+				FunctionLibrary->FunctionToVariant.FindOrAdd(Header.Name) = Header.Variant;
 			}
 			else
 			{
diff --git a/Engine/Plugins/Runtime/RigVM/Source/RigVMDeveloper/Private/RigVMModel/Nodes/RigVMLibraryNode.cpp b/Engine/Plugins/Runtime/RigVM/Source/RigVMDeveloper/Private/RigVMModel/Nodes/RigVMLibraryNode.cpp
--- a/Engine/Plugins/Runtime/RigVM/Source/RigVMDeveloper/Private/RigVMModel/Nodes/RigVMLibraryNode.cpp
+++ b/Engine/Plugins/Runtime/RigVM/Source/RigVMDeveloper/Private/RigVMModel/Nodes/RigVMLibraryNode.cpp
@@ -247,6 +247,16 @@ FRigVMGraphFunctionHeader URigVMLibraryNode::GetFunctionHeader(IRigVMGraphFuncti
 		}
 	}
 	if (!Header.Variant.Guid.IsValid())
+	{
+		if (URigVMFunctionLibrary* Library = GetLibrary())
+		{
+			if (const FRigVMVariant* Variant = Library->GetFunctionVariant(GetFName()))
+			{
+				Header.Variant = *Variant;
+			}
+		}
+	}
+	if (!Header.Variant.Guid.IsValid())
 	{
 		Header.Variant.Guid = FRigVMVariant::GenerateGUID();
 	}
