RigVM: introduce SRigVMVariantWidget

#rnx

diff --git a/Engine/Plugins/Runtime/RigVM/Source/RigVM/Public/RigVMCore/RigVMVariant.h b/Engine/Plugins/Runtime/RigVM/Source/RigVM/Public/RigVMCore/RigVMVariant.h
--- a/Engine/Plugins/Runtime/RigVM/Source/RigVM/Public/RigVMCore/RigVMVariant.h
+++ b/Engine/Plugins/Runtime/RigVM/Source/RigVM/Public/RigVMCore/RigVMVariant.h
@@ -77,7 +77,17 @@ struct RIGVM_API FRigVMVariantRef
 
 	FRigVMVariantRef(const FSoftObjectPath& InPath, const FRigVMVariant& InVariant)
 		: ObjectPath(InPath), Variant(InVariant) {}
-	
+
+	bool operator == (const FRigVMVariantRef& Other) const
+	{
+		return Variant.Guid == Other.Variant.Guid && ObjectPath == Other.ObjectPath; 
+	}
+
+	friend uint32 GetTypeHash(const FRigVMVariantRef& InVariantRef)
+	{
+		return HashCombine(GetTypeHash(InVariantRef.ObjectPath), GetTypeHash(InVariantRef.Variant.Guid));
+	}
+
 	FSoftObjectPath ObjectPath;
 	FRigVMVariant Variant;
 };
diff --git a/Engine/Plugins/Runtime/RigVM/Source/RigVMEditor/Private/Editor/RigVMGraphDetailCustomization.cpp b/Engine/Plugins/Runtime/RigVM/Source/RigVMEditor/Private/Editor/RigVMGraphDetailCustomization.cpp
--- a/Engine/Plugins/Runtime/RigVM/Source/RigVMEditor/Private/Editor/RigVMGraphDetailCustomization.cpp
+++ b/Engine/Plugins/Runtime/RigVM/Source/RigVMEditor/Private/Editor/RigVMGraphDetailCustomization.cpp
@@ -27,6 +27,7 @@
 #include "Widgets/SRigVMGraphPinVariableBinding.h"
 #include "InstancedPropertyBagStructureDataProvider.h"
 #include "Widgets/SRigVMGraphPinEnumPicker.h"
+#include "Widgets/SRigVMVariantWidget.h"
 #include "ScopedTransaction.h"
 #include "Misc/UObjectToken.h"
 
@@ -960,8 +961,9 @@ void FRigVMGraphDetailCustomization::CustomizeDetails(IDetailLayoutBuilder& Deta
             ]
         ];
 
-		// variant guid
+		// variant
 		SettingsCategory.AddCustomRow(FText::GetEmpty())
+		.OverrideResetToDefault(FResetToDefaultOverride::Hide())
 		.Visibility(TAttribute<EVisibility>::CreateLambda([this]()
 		{
 			return IsFunctionVariant() ? EVisibility::Visible : EVisibility::Collapsed;
@@ -969,64 +971,16 @@ void FRigVMGraphDetailCustomization::CustomizeDetails(IDetailLayoutBuilder& Deta
 		.NameContent()
 		[
 			SNew(STextBlock)
-			.Text(FText::FromString(TEXT("Variant Guid")))
+			.Text(FText::FromString(TEXT("Variant")))
 			.Font(IDetailLayoutBuilder::GetDetailFont())
 		]
 		.ValueContent()
 		[
-			SNew(STextBlock)
-			.Font(IDetailLayoutBuilder::GetDetailFont())
-			.Text(this, &FRigVMGraphDetailCustomization::GetVariantGuidText)
-		];
-
-		/*
-		SettingsCategory.AddCustomRow(FText::GetEmpty())
-		.WholeRowContent()
-		[
-			SNew(SHorizontalBox)
-			+ SHorizontalBox::Slot()
-			[
-				SNew(SButton)
-				.Visibility_Lambda([this]()
-				{
-					return IsFunctionVariant() ? EVisibility::Visible : EVisibility::Collapsed;
-				})
-				.Text(LOCTEXT("SplitVariant", "Split Variant"))
-				.ToolTipText(LOCTEXT("SplitVariantToolTip", "Splits this variant from the set and makes the function a unique function"))
-				.OnClicked(this, &FRigVMGraphDetailCustomization::OnSplitVariant)
-			]
-			+ SHorizontalBox::Slot()
-			[
-				SNew(SButton)
-				.Text(LOCTEXT("MergeVariant", "Merge Variant"))
-				.ToolTipText(LOCTEXT("MergeVariantToolTip", "Merges this variant with an existing set"))
-				.OnClicked(this, &FRigVMGraphDetailCustomization::OnMergeVariant)
-			]
-		];
-		*/
-
-		// matching variants
-		SettingsCategory.AddCustomRow(FText::GetEmpty())
-		.Visibility(TAttribute<EVisibility>::CreateLambda([this]()
-		{
-			return IsFunctionVariant() ? EVisibility::Visible : EVisibility::Collapsed;
-		}))
-		.NameContent()
-		[
-			SNew(STextBlock)
-			.Text(FText::FromString(TEXT("Matching Variants")))
-			.Font(IDetailLayoutBuilder::GetDetailFont())
-		]
-		.ValueContent()
-		[
-			SAssignNew(VariantLog, SRigVMLogWidget)
-			.LogLabel(LOCTEXT("Variants", "Variants"))
-			.LogName(TEXT("RigVMFunctionVariants"))
-			.ShowFilters(false)
-			.AllowClear(false)
-			.DiscardDuplicates(true)
-			.ScrollToBottom(false)
-			.HeightOverride(120)
+			SNew(SRigVMVariantWidget)
+			.Variant(this, &FRigVMGraphDetailCustomization::GetVariant)
+			.VariantRefs(this, &FRigVMGraphDetailCustomization::GetVariantRefs)
+			.OnVariantChanged(this, &FRigVMGraphDetailCustomization::OnVariantChanged)
+			.OnBrowseVariantRef(this, &FRigVMGraphDetailCustomization::OnBrowseVariantRef)
 		];
 	}
 
@@ -1058,8 +1012,6 @@ void FRigVMGraphDetailCustomization::CustomizeDetails(IDetailLayoutBuilder& Deta
 		Model,
 		Blueprint));
 	DefaultsCategory.AddCustomBuilder(DefaultsArgumentNode);
-	
-	RefreshVariantLog();
 }
 
 bool FRigVMGraphDetailCustomization::IsAddNewInputOutputEnabled() const
@@ -1401,7 +1353,7 @@ bool FRigVMGraphDetailCustomization::IsFunctionVariant() const
 	return false;
 }
 
-FText FRigVMGraphDetailCustomization::GetVariantGuidText() const
+FRigVMVariant FRigVMGraphDetailCustomization::GetVariant() const
 {
 	if (GraphPtr.IsValid() && RigVMBlueprintPtr.IsValid())
 	{
@@ -1410,39 +1362,15 @@ FText FRigVMGraphDetailCustomization::GetVariantGuidText() const
 		{
 			if (const URigVMLibraryNode* LibraryNode = Cast<URigVMLibraryNode>(Model->GetOuter()))
 			{
-				const FGuid Guid = LibraryNode->GetFunctionHeader(Blueprint->GetRigVMGraphFunctionHost()).Variant.Guid; 
-				if(Guid.IsValid())
-				{
-					return FText::FromString(Guid.ToString(EGuidFormats::DigitsWithHyphensLower));
-				}
+				return LibraryNode->GetFunctionHeader(Blueprint->GetRigVMGraphFunctionHost()).Variant; 
 			}
 		}
 	}
-	return FText();
-}
-
-/*
-FReply FRigVMGraphDetailCustomization::OnSplitVariant()
-{
-	// todo
-	RefreshVariantLog();
-	return FReply::Unhandled();
-}
-
-FReply FRigVMGraphDetailCustomization::OnMergeVariant()
-{
-	// todo
-	RefreshVariantLog();
-	return FReply::Unhandled();
+	return FRigVMVariant();
 }
-*/
 
-void FRigVMGraphDetailCustomization::RefreshVariantLog()
+TArray<FRigVMVariantRef> FRigVMGraphDetailCustomization::GetVariantRefs() const
 {
-	check(VariantLog);
-	
-	VariantLog->GetListing()->ClearMessages();
-
 	if (GraphPtr.IsValid() && RigVMBlueprintPtr.IsValid())
 	{
 		URigVMBlueprint* Blueprint = RigVMBlueprintPtr.Get();
@@ -1451,47 +1379,39 @@ void FRigVMGraphDetailCustomization::RefreshVariantLog()
 			if (const URigVMLibraryNode* LibraryNode = Cast<URigVMLibraryNode>(Model->GetOuter()))
 			{
 				const FRigVMGraphFunctionHeader Header = LibraryNode->GetFunctionHeader(Blueprint->GetRigVMGraphFunctionHost());
-				TArray<FRigVMGraphFunctionIdentifier> Variants = Header.LibraryPointer.GetVariantIdentifiers(false);
-				Variants.Sort([](const FRigVMGraphFunctionIdentifier& A, const FRigVMGraphFunctionIdentifier& B)
-				{
-					return A.GetLibraryNodePath().Compare(B.GetLibraryNodePath()) < 0;
-				});
+				return Header.LibraryPointer.GetVariants(false);
+			}
+		}
+	}
+	return TArray<FRigVMVariantRef>();
+}
 
-				for(const FRigVMGraphFunctionIdentifier& Variant : Variants)
+void FRigVMGraphDetailCustomization::OnVariantChanged(const FRigVMVariant& InVariant)
+{
+	// todo: update the function's variant info
+}
+
+void FRigVMGraphDetailCustomization::OnBrowseVariantRef(const FRigVMVariantRef& InVariantRef)
+{
+	const FRigVMGraphFunctionHeader Header = FRigVMGraphFunctionHeader::FindGraphFunctionHeader(InVariantRef.ObjectPath);
+	if(Header.IsValid())
+	{
+		if(const URigVMLibraryNode* LibraryNode = Cast<URigVMLibraryNode>(Header.LibraryPointer.GetNodeSoftPath().TryLoad()))
+		{
+			if(UBlueprint* Blueprint = LibraryNode->GetTypedOuter<UBlueprint>())
+			{
+				GEditor->GetEditorSubsystem<UAssetEditorSubsystem>()->OpenEditorForAsset(Blueprint);
+			
+				if(IAssetEditorInstance* Editor = GEditor->GetEditorSubsystem<UAssetEditorSubsystem>()->FindEditorForAsset(Blueprint, true))
 				{
-					if(URigVMLibraryNode* FunctionVariant = Cast<URigVMLibraryNode>(Variant.GetNodeSoftPath().TryLoad()))
+					if(FRigVMEditor* RigVMEditor = static_cast<FRigVMEditor*>(Editor))
 					{
-						const TSharedRef<FUObjectToken> ObjectToken = FUObjectToken::Create(FunctionVariant);
-
-						TWeakObjectPtr<URigVMLibraryNode> WeakFunctionVariant(FunctionVariant);
-						ObjectToken->OnMessageTokenActivated(FOnMessageTokenActivated::CreateLambda([WeakFunctionVariant](const TSharedRef<IMessageToken>&)
-						{
-							if(WeakFunctionVariant.IsValid())
-							{
-								if(UBlueprint* Blueprint = WeakFunctionVariant.Get()->GetTypedOuter<UBlueprint>())
-								{
-									GEditor->GetEditorSubsystem<UAssetEditorSubsystem>()->OpenEditorForAsset(Blueprint);
-					
-									if(IAssetEditorInstance* Editor = GEditor->GetEditorSubsystem<UAssetEditorSubsystem>()->FindEditorForAsset(Blueprint, true))
-									{
-										if(FRigVMEditor* RigVMEditor = static_cast<FRigVMEditor*>(Editor))
-										{
-											RigVMEditor->HandleJumpToHyperlink(WeakFunctionVariant.Get());
-										}
-									}
-								}
-							}
-						}));
-
-						const TSharedRef<FTokenizedMessage> Message = FTokenizedMessage::Create(EMessageSeverity::Info);
-						Message->AddToken(ObjectToken);
-						VariantLog->GetListing()->AddMessage(Message);
+						RigVMEditor->HandleJumpToHyperlink(LibraryNode);
 					}
 				}
 			}
 		}
 	}
-	
 }
 
 FRigVMWrappedNodeDetailCustomization::FRigVMWrappedNodeDetailCustomization()
diff --git a/Engine/Plugins/Runtime/RigVM/Source/RigVMEditor/Private/Editor/RigVMVariantDetailCustomization.cpp b/Engine/Plugins/Runtime/RigVM/Source/RigVMEditor/Private/Editor/RigVMVariantDetailCustomization.cpp
--- a/Engine/Plugins/Runtime/RigVM/Source/RigVMEditor/Private/Editor/RigVMVariantDetailCustomization.cpp
+++ b/Engine/Plugins/Runtime/RigVM/Source/RigVMEditor/Private/Editor/RigVMVariantDetailCustomization.cpp
@@ -5,9 +5,13 @@
 #include "DetailLayoutBuilder.h"
 #include "DetailWidgetRow.h"
 #include "IDetailChildrenBuilder.h"
+#include "Editor/RigVMEditorTools.h"
 #include "Misc/UObjectToken.h"
 #include "Widgets/Text/STextBlock.h"
-#include "Widgets/SRigVMLogWidget.h"
+#include "Widgets/SRigVMVariantWidget.h"
+#include "ContentBrowserModule.h"
+#include "IContentBrowserSingleton.h"
+#include "ScopedTransaction.h"
 
 #define LOCTEXT_NAMESPACE "RigVMVariantDetailCustomization"
 
@@ -15,16 +19,6 @@ class FUObjectToken;
 
 void FRigVMVariantDetailCustomization::CustomizeHeader(TSharedRef<IPropertyHandle> InStructPropertyHandle, FDetailWidgetRow& HeaderRow, IPropertyTypeCustomizationUtils& StructCustomizationUtils)
 {
-	HeaderRow
-	.NameContent()
-	[
-		InStructPropertyHandle->CreatePropertyNameWidget()
-	]
-	.ValueContent()
-	[
-		InStructPropertyHandle->CreatePropertyValueWidget()
-	];
-
 	TArray<UObject*> Objects;
 	InStructPropertyHandle->GetOuterObjects(Objects);
 	ensure(Objects.Num() == 1); // This is in here to ensure we are only showing the modifier details in the blueprint editor
@@ -34,96 +28,73 @@ void FRigVMVariantDetailCustomization::CustomizeHeader(TSharedRef<IPropertyHandl
 		if (Object->IsA<URigVMBlueprint>())
 		{
 			BlueprintBeingCustomized = Cast<URigVMBlueprint>(Object);
+			break;
 		}
 	}
+
+	HeaderRow
+	.OverrideResetToDefault(FResetToDefaultOverride::Hide())
+	.NameContent()
+	[
+		InStructPropertyHandle->CreatePropertyNameWidget()
+	]
+	.ValueContent()
+	[
+		SNew(SRigVMVariantWidget)
+		.Variant(this, &FRigVMVariantDetailCustomization::GetVariant)
+		.VariantRefs(this, &FRigVMVariantDetailCustomization::GetVariantRefs)
+		.OnVariantChanged(this, &FRigVMVariantDetailCustomization::OnVariantChanged)
+		.OnBrowseVariantRef(this, &FRigVMVariantDetailCustomization::OnBrowseVariantRef)
+	];
 }
 
 void FRigVMVariantDetailCustomization::CustomizeChildren(TSharedRef<IPropertyHandle> InStructPropertyHandle, IDetailChildrenBuilder& StructBuilder, IPropertyTypeCustomizationUtils& StructCustomizationUtils)
 {
-	if (InStructPropertyHandle->IsValidHandle())
-	{
-		uint32 NumChildren = 0;
-		InStructPropertyHandle->GetNumChildren(NumChildren);
-
-		for (uint32 ChildIndex = 0; ChildIndex < NumChildren; ChildIndex++)
-		{
-			StructBuilder.AddProperty(InStructPropertyHandle->GetChildHandle(ChildIndex).ToSharedRef());
-		}
+	// nothing to do here
+}
 
-		// matching variants
-		StructBuilder.AddCustomRow(FText::GetEmpty())
-		.Visibility(TAttribute<EVisibility>::CreateLambda([this]()
-		{
-			return IsAssetVariant() ? EVisibility::Visible : EVisibility::Collapsed;
-		}))
-		.NameContent()
-		[
-			SNew(STextBlock)
-			.Text(FText::FromString(TEXT("Matching Variants")))
-			.Font(IDetailLayoutBuilder::GetDetailFont())
-		]
-		.ValueContent()
-		[
-			SAssignNew(VariantLog, SRigVMLogWidget)
-			.LogLabel(LOCTEXT("Variants", "Variants"))
-			.LogName(TEXT("RigVMAssetVariants"))
-			.ShowFilters(false)
-			.AllowClear(false)
-			.DiscardDuplicates(true)
-			.ScrollToBottom(false)
-			.HeightOverride(120)
-		];
+FRigVMVariant FRigVMVariantDetailCustomization::GetVariant() const
+{
+	if (BlueprintBeingCustomized)
+	{
+		return BlueprintBeingCustomized->AssetVariant;
 	}
-
-	RefreshVariantLog();
+	return FRigVMVariant();
 }
 
-bool FRigVMVariantDetailCustomization::IsAssetVariant() const
+TArray<FRigVMVariantRef> FRigVMVariantDetailCustomization::GetVariantRefs() const
 {
 	if (BlueprintBeingCustomized)
 	{
 		const FRigVMVariant& Variant = BlueprintBeingCustomized->AssetVariant;
 		TArray<FRigVMVariantRef> Variants = URigVMBuildData::Get()->FindAssetVariantRefs(Variant.Guid);
-		return Variants.Num() > 1;
+		const FRigVMVariantRef MyVariantRef = FRigVMVariantRef(BlueprintBeingCustomized->GetPathName(), BlueprintBeingCustomized->AssetVariant);
+		Variants.RemoveAll([MyVariantRef](const FRigVMVariantRef& VariantRef) -> bool
+		{
+			return VariantRef == MyVariantRef;
+		});
+		return Variants;
 	}
-	return false;
+	return TArray<FRigVMVariantRef>();
 }
 
-FText FRigVMVariantDetailCustomization::GetVariantGuidText() const
+void FRigVMVariantDetailCustomization::OnVariantChanged(const FRigVMVariant& InNewVariant)
 {
-	if (BlueprintBeingCustomized)
+	if(BlueprintBeingCustomized)
 	{
-		const FGuid Guid = BlueprintBeingCustomized->AssetVariant.Guid;
-		if(Guid.IsValid())
-		{
-			return FText::FromString(Guid.ToString(EGuidFormats::DigitsWithHyphensLower));
-		}
+		FScopedTransaction Transaction(LOCTEXT("ChangedVariantInfo", "Changed Blueprint Variant Information"));
+		BlueprintBeingCustomized->Modify();
+		BlueprintBeingCustomized->AssetVariant = InNewVariant;
 	}
-	return FText();
 }
 
-void FRigVMVariantDetailCustomization::RefreshVariantLog()
+void FRigVMVariantDetailCustomization::OnBrowseVariantRef(const FRigVMVariantRef& InVariantRef)
 {
-	check(VariantLog);
-
-	VariantLog->GetListing()->ClearMessages();
-
-	if (BlueprintBeingCustomized)
+	const FAssetData AssetData = UE::RigVM::Editor::Tools::FindAssetFromAnyPath(InVariantRef.ObjectPath.ToString(), true);
+	if(AssetData.IsValid())
 	{
-		TArray<FRigVMVariantRef> Variants = URigVMBuildData::Get()->FindAssetVariantRefs(BlueprintBeingCustomized->AssetVariant.Guid);
-		Variants = Variants.FilterByPredicate([this](const FRigVMVariantRef& VariantRef) { return VariantRef.ObjectPath != BlueprintBeingCustomized->GetPathName(); });
-		Variants.Sort([](const FRigVMVariantRef& A, const FRigVMVariantRef& B)
-		{
-			return A.ObjectPath.ToString().Compare(B.ObjectPath.ToString()) < 0;
-		});
-
-		for (FRigVMVariantRef& Variant : Variants)
-		{
-			const TSharedRef<FAssetNameToken> AssetToken = FAssetNameToken::Create(Variant.ObjectPath.ToString());
-			const TSharedRef<FTokenizedMessage> Message = FTokenizedMessage::Create(EMessageSeverity::Info);
-			Message->AddToken(AssetToken);
-			VariantLog->GetListing()->AddMessage(Message);
-		}
+		const FContentBrowserModule& ContentBrowserModule = FModuleManager::Get().LoadModuleChecked<FContentBrowserModule>(TEXT("ContentBrowser"));
+		ContentBrowserModule.Get().SyncBrowserToAssets({AssetData});
 	}
 }
 
diff --git a/Engine/Plugins/Runtime/RigVM/Source/RigVMEditor/Private/Widgets/SRigVMVariantWidget.cpp b/Engine/Plugins/Runtime/RigVM/Source/RigVMEditor/Private/Widgets/SRigVMVariantWidget.cpp
new file mode 100644
--- /dev/null
+++ b/Engine/Plugins/Runtime/RigVM/Source/RigVMEditor/Private/Widgets/SRigVMVariantWidget.cpp
@@ -0,0 +1,124 @@
+// Copyright Epic Games, Inc. All Rights Reserved.
+
+#include "Widgets/SRigVMVariantWidget.h"
+#include "Widgets/Text/STextBlock.h"
+#include "Widgets/Layout/SScrollBox.h"
+#include "DetailLayoutBuilder.h"
+
+#define LOCTEXT_NAMESPACE "SRigVMVariantWidget"
+
+SRigVMVariantWidget::SRigVMVariantWidget()
+	: VariantRefHash(UINT32_MAX)
+{
+}
+
+SRigVMVariantWidget::~SRigVMVariantWidget()
+{
+}
+
+void SRigVMVariantWidget::Construct(
+	const FArguments& InArgs)
+{
+	VariantAttribute = InArgs._Variant;
+	VariantRefsAttribute = InArgs._VariantRefs;
+	OnVariantChanged = InArgs._OnVariantChanged;
+	OnCreateVariantRefRow = InArgs._OnCreateVariantRefRow;
+	OnBrowseVariantRef = InArgs._OnBrowseVariantRef;
+
+	if(!OnCreateVariantRefRow.IsBound())
+	{
+		OnCreateVariantRefRow.BindSP(this, &SRigVMVariantWidget::CreateDefaultVariantRefRow);
+	}
+
+	ChildSlot
+	[
+		SNew(SVerticalBox)
+		+ SVerticalBox::Slot()
+		.AutoHeight()
+		.HAlign(HAlign_Left)
+		[
+			SNew(STextBlock)
+			.Font(IDetailLayoutBuilder::GetDetailFont())
+			.Text_Lambda([this]()
+			{
+				const FRigVMVariant Variant = VariantAttribute.Get();
+				return FText::FromString(Variant.Guid.ToString(EGuidFormats::DigitsWithHyphensLower));
+			})
+		]
+		+ SVerticalBox::Slot()
+		.AutoHeight()
+		.HAlign(HAlign_Fill)
+		[
+			SNew(SScrollBox)
+			.Visibility(this, &SRigVMVariantWidget::GetVariantRefListVisibility)
+			+ SScrollBox::Slot()
+			.MaxSize(InArgs._MaxVariantRefListHeight)
+			[
+				SAssignNew(VariantRefListBox, SVerticalBox)
+			]
+		]
+	];
+
+	SetCanTick(true);
+}
+
+void SRigVMVariantWidget::Tick(const FGeometry& AllottedGeometry, const double InCurrentTime, const float InDeltaTime)
+{
+	SBox::Tick(AllottedGeometry, InCurrentTime, InDeltaTime);
+
+	uint32 NewHash = 0;
+	const TArray<FRigVMVariantRef> NewVariantRefs = VariantRefsAttribute.Get();
+	for(const FRigVMVariantRef& NewVariantRef : NewVariantRefs)
+	{
+		NewHash = HashCombine(NewHash, GetTypeHash(NewVariantRef));
+	}
+
+	if(NewHash != VariantRefHash)
+	{
+		VariantRefHash = NewHash;
+		VariantRefs = NewVariantRefs;
+		VariantRefs.Sort([](const FRigVMVariantRef& A, const FRigVMVariantRef& B)
+		{
+			return A.ObjectPath.ToString().Compare(B.ObjectPath.ToString()) < 0;
+		});
+		
+		RebuildVariantRefList();
+	}
+}
+
+EVisibility SRigVMVariantWidget::GetVariantRefListVisibility() const
+{
+	return VariantRefs.IsEmpty() ? EVisibility::Collapsed : EVisibility::Visible;
+}
+
+TSharedPtr<SWidget> SRigVMVariantWidget::CreateDefaultVariantRefRow(const FRigVMVariantRef& InVariantRef) const
+{
+	return SNew(STextBlock)
+	.Text(FText::FromString(InVariantRef.ObjectPath.ToString()));
+}
+
+void SRigVMVariantWidget::RebuildVariantRefList()
+{
+	VariantRefListBox->ClearChildren();
+
+	for(FRigVMVariantRef VariantRef : VariantRefs)
+	{
+		TSharedPtr<SWidget> Widget = OnCreateVariantRefRow.Execute(VariantRef);
+		Widget->SetOnMouseDoubleClick(
+			FPointerEventHandler::CreateLambda(
+				[VariantRef, this](const FGeometry& InGeometry, const FPointerEvent& InMouseEvent) -> FReply
+				{
+					(void)OnBrowseVariantRef.ExecuteIfBound(VariantRef);
+					return FReply::Handled();
+				}
+			)
+		);
+		VariantRefListBox->AddSlot()
+		.AutoHeight()
+		.HAlign(HAlign_Fill)
+		.Padding(0, 4, 0, 0)
+		[
+			Widget.ToSharedRef()
+		];
+	}
+}
diff --git a/Engine/Plugins/Runtime/RigVM/Source/RigVMEditor/Public/Editor/RigVMGraphDetailCustomization.h b/Engine/Plugins/Runtime/RigVM/Source/RigVMEditor/Public/Editor/RigVMGraphDetailCustomization.h
--- a/Engine/Plugins/Runtime/RigVM/Source/RigVMEditor/Public/Editor/RigVMGraphDetailCustomization.h
+++ b/Engine/Plugins/Runtime/RigVM/Source/RigVMEditor/Public/Editor/RigVMGraphDetailCustomization.h
@@ -194,13 +194,14 @@ public:
 	void OnAccessSpecifierSelected( TSharedPtr<FRigVMStringWithTag> SpecifierName, ESelectInfo::Type SelectInfo );
 	TSharedRef<ITableRow> HandleGenerateRowAccessSpecifier( TSharedPtr<FRigVMStringWithTag> SpecifierName, const TSharedRef<STableViewBase>& OwnerTable );
 	bool IsFunctionVariant() const;
-	FText GetVariantGuidText() const;
-	//FReply OnSplitVariant();
-	//FReply OnMergeVariant();
-	void RefreshVariantLog();
+	FRigVMVariant GetVariant() const;
+	TArray<FRigVMVariantRef> GetVariantRefs() const;
 
 private:
 
+	void OnVariantChanged(const FRigVMVariant& InVariant);
+	void OnBrowseVariantRef(const FRigVMVariantRef& InVariantRef);
+
 	/** The Blueprint editor we are embedded in */
 	TWeakPtr<FRigVMEditor> RigVMEditorPtr;
 
@@ -213,9 +214,6 @@ private:
 	/** The color block widget */
 	TSharedPtr<SColorBlock> ColorBlock;
 
-	/** The log widget used for function variants */
-	TSharedPtr<SRigVMLogWidget> VariantLog;
-
 	/** Set to true if the UI is currently picking a color */
 	bool bIsPickingColor;
 
diff --git a/Engine/Plugins/Runtime/RigVM/Source/RigVMEditor/Public/Editor/RigVMVariantDetailCustomization.h b/Engine/Plugins/Runtime/RigVM/Source/RigVMEditor/Public/Editor/RigVMVariantDetailCustomization.h
--- a/Engine/Plugins/Runtime/RigVM/Source/RigVMEditor/Public/Editor/RigVMVariantDetailCustomization.h
+++ b/Engine/Plugins/Runtime/RigVM/Source/RigVMEditor/Public/Editor/RigVMVariantDetailCustomization.h
@@ -22,11 +22,14 @@ public:
 	virtual void CustomizeHeader(TSharedRef<class IPropertyHandle> InStructPropertyHandle, class FDetailWidgetRow& HeaderRow, IPropertyTypeCustomizationUtils& StructCustomizationUtils) override;
 	virtual void CustomizeChildren(TSharedRef<class IPropertyHandle> InStructPropertyHandle, class IDetailChildrenBuilder& StructBuilder, IPropertyTypeCustomizationUtils& StructCustomizationUtils) override;
 
-	bool IsAssetVariant() const;
-	FText GetVariantGuidText() const;
-	void RefreshVariantLog();
+	FRigVMVariant GetVariant() const;
+	TArray<FRigVMVariantRef> GetVariantRefs() const;
+	
 protected:
 
+	void OnVariantChanged(const FRigVMVariant& InNewVariant);
+	void OnBrowseVariantRef(const FRigVMVariantRef& InVariantRef);
+
 	URigVMBlueprint* BlueprintBeingCustomized;
 
 	/** The log widget used for function variants */
diff --git a/Engine/Plugins/Runtime/RigVM/Source/RigVMEditor/Public/Widgets/SRigVMVariantWidget.h b/Engine/Plugins/Runtime/RigVM/Source/RigVMEditor/Public/Widgets/SRigVMVariantWidget.h
new file mode 100644
--- /dev/null
+++ b/Engine/Plugins/Runtime/RigVM/Source/RigVMEditor/Public/Widgets/SRigVMVariantWidget.h
@@ -0,0 +1,51 @@
+// Copyright Epic Games, Inc. All Rights Reserved.
+
+#pragma once
+
+#include "Widgets/Layout/SBox.h"
+#include "Widgets/SBoxPanel.h"
+#include "RigVMCore/RigVMVariant.h"
+
+DECLARE_DELEGATE_OneParam(FRigVMVariantWidget_OnVariantChanged, const FRigVMVariant&);
+DECLARE_DELEGATE_RetVal_OneParam(TSharedPtr<SWidget>, FRigVMVariantWidget_OnCreateVariantRefRow, const FRigVMVariantRef&);
+DECLARE_DELEGATE_OneParam(FRigVMVariantWidget_OnBrowseVariantRef, const FRigVMVariantRef&);
+
+class SRigVMVariantWidget : public SBox
+{
+public:
+	
+	SLATE_BEGIN_ARGS(SRigVMVariantWidget)
+		: _MaxVariantRefListHeight(200.f)
+	{
+	}
+	SLATE_ATTRIBUTE(FRigVMVariant, Variant)
+	SLATE_ATTRIBUTE(TArray<FRigVMVariantRef>, VariantRefs)
+	SLATE_EVENT(FRigVMVariantWidget_OnVariantChanged, OnVariantChanged)
+	SLATE_EVENT(FRigVMVariantWidget_OnCreateVariantRefRow, OnCreateVariantRefRow);
+	SLATE_EVENT(FRigVMVariantWidget_OnBrowseVariantRef, OnBrowseVariantRef)
+	SLATE_ATTRIBUTE(float, MaxVariantRefListHeight)
+	SLATE_END_ARGS()
+
+	SRigVMVariantWidget();
+	virtual ~SRigVMVariantWidget() override;
+	
+	void Construct(const FArguments& InArgs);
+
+	virtual void Tick(const FGeometry& AllottedGeometry, const double InCurrentTime, const float InDeltaTime) override;
+
+private:
+
+	EVisibility GetVariantRefListVisibility() const;
+	TSharedPtr<SWidget> CreateDefaultVariantRefRow(const FRigVMVariantRef& InVariantRef) const;
+	void RebuildVariantRefList();
+
+	TAttribute<FRigVMVariant> VariantAttribute;
+	TAttribute<TArray<FRigVMVariantRef>> VariantRefsAttribute;
+	FRigVMVariantWidget_OnVariantChanged OnVariantChanged;
+	FRigVMVariantWidget_OnCreateVariantRefRow OnCreateVariantRefRow;
+    FRigVMVariantWidget_OnBrowseVariantRef OnBrowseVariantRef;
+
+	TArray<FRigVMVariantRef> VariantRefs;
+	uint32 VariantRefHash;
+	TSharedPtr<SVerticalBox> VariantRefListBox;
+};
\ No newline at end of file
