[RigVM Editor] fixed missing namespace_loctext

#rnx

diff --git a/Engine/Plugins/Runtime/RigVM/Source/RigVMEditor/Private/Widgets/SRigVMVariantWidget.cpp b/Engine/Plugins/Runtime/RigVM/Source/RigVMEditor/Private/Widgets/SRigVMVariantWidget.cpp
--- a/Engine/Plugins/Runtime/RigVM/Source/RigVMEditor/Private/Widgets/SRigVMVariantWidget.cpp
+++ b/Engine/Plugins/Runtime/RigVM/Source/RigVMEditor/Private/Widgets/SRigVMVariantWidget.cpp
@@ -122,3 +122,5 @@ void SRigVMVariantWidget::RebuildVariantRefList()
 		];
 	}
 }
+
+#undef LOCTEXT_NAMESPACE // SRigVMVariantWidget
\ No newline at end of file
