Control Rig: Fix crash when constructing a double array property in the details wrapper object
#rnx

diff --git a/Engine/Plugins/Runtime/RigVM/Source/RigVMEditor/Private/Editor/RigVMDetailsViewWrapperObject.cpp b/Engine/Plugins/Runtime/RigVM/Source/RigVMEditor/Private/Editor/RigVMDetailsViewWrapperObject.cpp
--- a/Engine/Plugins/Runtime/RigVM/Source/RigVMEditor/Private/Editor/RigVMDetailsViewWrapperObject.cpp
+++ b/Engine/Plugins/Runtime/RigVM/Source/RigVMEditor/Private/Editor/RigVMDetailsViewWrapperObject.cpp
@@ -447,6 +447,18 @@ UClass* URigVMDetailsViewWrapperObject::GetClassForNodes(TArray<URigVMNode*> InN
 			Property = ArrayProperty;
 			ElementProperty = &ArrayProperty->Inner;
 			PropertyOwner = ArrayProperty;
+
+			// In case the base type is also an array (RigVM supports up to 2 nested levels of arrays)
+			if (RigVMTypeUtils::IsArrayType(BaseCPPType))
+			{
+				BaseCPPType = RigVMTypeUtils::BaseTypeFromArrayType(BaseCPPType);
+				
+				// create an array property as a container for the tail
+				FArrayProperty* InnerArrayProperty = new FArrayProperty(PropertyOwner, Pin->GetFName(), RF_Public);
+				(*ElementProperty) = InnerArrayProperty;
+				ElementProperty = &InnerArrayProperty->Inner;
+				PropertyOwner = InnerArrayProperty;
+			}
 		}
 
 		if(BaseCPPType.Equals(RigVMTypeUtils::BoolType, ESearchCase::IgnoreCase))
@@ -500,6 +512,10 @@ UClass* URigVMDetailsViewWrapperObject::GetClassForNodes(TArray<URigVMNode*> InN
 		{
 			return nullptr;
 		}
+		if (ElementProperty && *ElementProperty == nullptr)
+		{
+			return nullptr;
+		}
 		Property->SetPropertyFlags(CPF_Edit);
 
 		for(UScriptStruct* UnitStruct : UnitStructs)
