Control Rig: Fix crash when copy/pasting function
#rnx

diff --git a/Engine/Plugins/Runtime/RigVM/Source/RigVMDeveloper/Private/RigVMCompiler/RigVMAST.cpp b/Engine/Plugins/Runtime/RigVM/Source/RigVMDeveloper/Private/RigVMCompiler/RigVMAST.cpp
--- a/Engine/Plugins/Runtime/RigVM/Source/RigVMDeveloper/Private/RigVMCompiler/RigVMAST.cpp
+++ b/Engine/Plugins/Runtime/RigVM/Source/RigVMDeveloper/Private/RigVMCompiler/RigVMAST.cpp
@@ -2761,7 +2761,10 @@ void FRigVMParserAST::Inline(const TArray<URigVMGraph*>& InGraphs)
 	{
 		for (URigVMNode* LocalNode : Graph->GetNodes())
 		{
-			LocalNodeProxies.Add(FRigVMASTProxy::MakeFromUObject(LocalNode));
+			if (IsValid(LocalNode))
+			{
+				LocalNodeProxies.Add(FRigVMASTProxy::MakeFromUObject(LocalNode));
+			}
 		}
 	}
 	Inline(InGraphs, LocalNodeProxies);
diff --git a/Engine/Plugins/Runtime/RigVM/Source/RigVMDeveloper/Private/RigVMModel/RigVMController.cpp b/Engine/Plugins/Runtime/RigVM/Source/RigVMDeveloper/Private/RigVMModel/RigVMController.cpp
--- a/Engine/Plugins/Runtime/RigVM/Source/RigVMDeveloper/Private/RigVMModel/RigVMController.cpp
+++ b/Engine/Plugins/Runtime/RigVM/Source/RigVMDeveloper/Private/RigVMModel/RigVMController.cpp
@@ -4212,6 +4212,15 @@ TArray<FName> URigVMController::ImportNodesFromText(const FString& InText, bool
 		Swap(Factory.CreatedNodes, FilteredNodes);
 	}
 
+	// Some nodes inside these graphs might have been deleted, we need to clean up the graph
+	for (URigVMGraph* CreatedGraph : Factory.CreatedGraphs)
+	{
+		CreatedGraph->Nodes = CreatedGraph->Nodes.FilterByPredicate([](const URigVMNode* Node)
+		{
+			return IsValid(Node);
+		});
+	}
+
 	// the links may already be in the graph's links property
 	// due to the serialization - remove them since we need.
 	// but we also want to maintain the order of the links as they
