Control Rig: Fix PatchFunctionsWithInvalidReturnPaths when return node is not present. Also remove ability to delete entry or return nodes
#rnx

diff --git a/Engine/Plugins/Runtime/RigVM/Source/RigVMDeveloper/Private/RigVMModel/RigVMController.cpp b/Engine/Plugins/Runtime/RigVM/Source/RigVMDeveloper/Private/RigVMModel/RigVMController.cpp
--- a/Engine/Plugins/Runtime/RigVM/Source/RigVMDeveloper/Private/RigVMModel/RigVMController.cpp
+++ b/Engine/Plugins/Runtime/RigVM/Source/RigVMDeveloper/Private/RigVMModel/RigVMController.cpp
@@ -6484,7 +6484,8 @@ bool URigVMController::RemoveNodes(TArray<URigVMNode*> InNodes, bool bSetupUndoR
 		{
 			if(!GetSchema()->CanRemoveNode(this, InNode))
 			{
-				continue;
+				ReportErrorf(TEXT("Our schema does not allow to remove node %s"), *InNode->GetNodePath());
+				return false;
 			}
 		}
 
@@ -20475,6 +20476,26 @@ FRigVMClientPatchResult URigVMController::PatchFunctionsWithInvalidReturnPaths()
 			return Result;
 		}
 
+		// Make sure at this point that both interface nodes exist
+		if (!EntryNode)
+		{
+			EntryNode = NewObject<URigVMFunctionEntryNode>(Graph, TEXT("Entry"));
+			if(AddGraphNode(EntryNode, false))
+			{
+				RefreshFunctionPins(EntryNode);
+				bEntryIsMutable = EntryNode->IsMutable();
+			}
+		}
+		if (!ReturnNode)
+		{
+			ReturnNode = NewObject<URigVMFunctionReturnNode>(Graph, TEXT("Return"));
+			if(AddGraphNode(ReturnNode, false))
+			{
+				RefreshFunctionPins(ReturnNode);
+				bReturnIsMutable = ReturnNode->IsMutable();
+			}
+		}
+		
 		bool bReturnExecuteIsLinked = false;
 		if (bReturnIsMutable)
 		{
