Control Rig: Add tags to variants
#rnx

diff --git a/Engine/Plugins/Runtime/RigVM/Source/RigVM/Public/RigVMCore/RigVMVariant.h b/Engine/Plugins/Runtime/RigVM/Source/RigVM/Public/RigVMCore/RigVMVariant.h
--- a/Engine/Plugins/Runtime/RigVM/Source/RigVM/Public/RigVMCore/RigVMVariant.h
+++ b/Engine/Plugins/Runtime/RigVM/Source/RigVM/Public/RigVMCore/RigVMVariant.h
@@ -7,22 +7,52 @@
 
 struct FRigVMVariantRef;
 
+// User applied tag
 USTRUCT(BlueprintType)
 struct RIGVM_API FRigVMTag
 {
 	GENERATED_BODY()
+
+	UPROPERTY(BlueprintReadWrite, EditAnywhere, Category=Variant)
+	FName Name;
 	
-	// User applied tag
-	UPROPERTY(BlueprintReadWrite, VisibleAnywhere, Category=Variant)
+	UPROPERTY(BlueprintReadWrite, EditAnywhere, Category=Variant)
 	FString Label;
+
+	UPROPERTY(BlueprintReadWrite, EditAnywhere, Category=Variant)
+	FText ToolTip;
+
+	UPROPERTY(BlueprintReadWrite, EditAnywhere, Category=Variant)
+	FLinearColor Color;
+
+	UPROPERTY(BlueprintReadWrite, EditAnywhere, Category=Variant)
+	bool bShowInUserInterface = true;
+
+	UPROPERTY(BlueprintReadWrite, EditAnywhere, Category=Variant)
+	bool bMarksSubjectAsInvalid = false;
 	
 	friend FArchive& operator<<(FArchive& Ar, FRigVMTag& Data)
 	{
 		Ar.UsingCustomVersion(FRigVMObjectVersion::GUID);
-		
+
+		Ar << Data.Name;
 		Ar << Data.Label;
+		Ar << Data.ToolTip;
+		Ar << Data.Color;
+		Ar << Data.bShowInUserInterface;
+		Ar << Data.bMarksSubjectAsInvalid;
 		return Ar;
 	}
+
+	friend bool operator==(const FRigVMTag& A, const FRigVMTag& B)
+	{
+		return A.Name == B.Name &&
+			A.Label == B.Label &&
+			A.ToolTip.EqualTo(B.ToolTip) &&
+			A.Color == B.Color &&
+			A.bShowInUserInterface == B.bShowInUserInterface &&
+			A.bMarksSubjectAsInvalid == B.bMarksSubjectAsInvalid;
+	}
 };
 
 USTRUCT(BlueprintType)
diff --git a/Engine/Plugins/Runtime/RigVM/Source/RigVM/Public/RigVMSettings.h b/Engine/Plugins/Runtime/RigVM/Source/RigVM/Public/RigVMSettings.h
--- a/Engine/Plugins/Runtime/RigVM/Source/RigVM/Public/RigVMSettings.h
+++ b/Engine/Plugins/Runtime/RigVM/Source/RigVM/Public/RigVMSettings.h
@@ -8,6 +8,7 @@
 
 #include "CoreMinimal.h"
 #include "Engine/DeveloperSettings.h"
+#include "RigVMCore/RigVMVariant.h"
 
 #include "RigVMSettings.generated.h"
 
@@ -32,6 +33,9 @@ public:
 	UPROPERTY(EditAnywhere, config, Category = Interaction)
 	bool bAutoLinkMutableNodes;
 
+	UPROPERTY(EditAnywhere, config, Category = Variants)
+	TArray<FRigVMTag> VariantTags;
+
 #endif
 };
 
diff --git a/Engine/Plugins/Runtime/RigVM/Source/RigVMDeveloper/Private/RigVMModel/Nodes/RigVMFunctionReferenceNode.cpp b/Engine/Plugins/Runtime/RigVM/Source/RigVMDeveloper/Private/RigVMModel/Nodes/RigVMFunctionReferenceNode.cpp
--- a/Engine/Plugins/Runtime/RigVM/Source/RigVMDeveloper/Private/RigVMModel/Nodes/RigVMFunctionReferenceNode.cpp
+++ b/Engine/Plugins/Runtime/RigVM/Source/RigVMDeveloper/Private/RigVMModel/Nodes/RigVMFunctionReferenceNode.cpp
@@ -184,6 +184,15 @@ const FRigVMGraphFunctionData* URigVMFunctionReferenceNode::GetReferencedFunctio
 	return nullptr;
 }
 
+TArray<FRigVMTag> URigVMFunctionReferenceNode::GetVariantTags() const
+{
+	if (const FRigVMGraphFunctionData* Data = GetReferencedFunctionData(false))
+	{
+		return Data->Header.Variant.Tags;
+	}
+	return GetReferencedFunctionHeader().Variant.Tags;
+}
+
 FString URigVMFunctionReferenceNode::GetOriginalDefaultValueForRootPin(const URigVMPin* InRootPin) const
 {
 	if(InRootPin->CanProvideDefaultValue())
diff --git a/Engine/Plugins/Runtime/RigVM/Source/RigVMDeveloper/Private/RigVMModel/Nodes/RigVMLibraryNode.cpp b/Engine/Plugins/Runtime/RigVM/Source/RigVMDeveloper/Private/RigVMModel/Nodes/RigVMLibraryNode.cpp
--- a/Engine/Plugins/Runtime/RigVM/Source/RigVMDeveloper/Private/RigVMModel/Nodes/RigVMLibraryNode.cpp
+++ b/Engine/Plugins/Runtime/RigVM/Source/RigVMDeveloper/Private/RigVMModel/Nodes/RigVMLibraryNode.cpp
@@ -237,25 +237,27 @@ FRigVMGraphFunctionHeader URigVMLibraryNode::GetFunctionHeader(IRigVMGraphFuncti
 	    Header.LibraryPointer.HostObject = Cast<UObject>(InHostObject);
 	}
 
-	UObject* HostPtr = Cast<UObject>(Header.LibraryPointer.HostObject.ResolveObject());
-	if (IRigVMGraphFunctionHost* Host = Cast<IRigVMGraphFunctionHost>(HostPtr))
+	if (URigVMFunctionLibrary* Library = GetLibrary())
 	{
-		if (FRigVMGraphFunctionData* Data = Host->GetRigVMGraphFunctionStore()->FindFunction(Header.LibraryPointer))
+		if (const FRigVMVariant* Variant = Library->GetFunctionVariant(GetFName()))
 		{
-			Header.Variant = Data->Header.Variant;
-			Header.NodeTitle = Data->Header.NodeTitle;
+			Header.Variant = *Variant;
 		}
 	}
-	if (!Header.Variant.Guid.IsValid())
+
+	UObject* HostPtr = Cast<UObject>(Header.LibraryPointer.HostObject.ResolveObject());
+	if (IRigVMGraphFunctionHost* Host = Cast<IRigVMGraphFunctionHost>(HostPtr))
 	{
-		if (URigVMFunctionLibrary* Library = GetLibrary())
+		if (FRigVMGraphFunctionData* Data = Host->GetRigVMGraphFunctionStore()->FindFunction(Header.LibraryPointer))
 		{
-			if (const FRigVMVariant* Variant = Library->GetFunctionVariant(GetFName()))
+			if (!Header.Variant.Guid.IsValid())
 			{
-				Header.Variant = *Variant;
+				Header.Variant = Data->Header.Variant;
 			}
+			Header.NodeTitle = Data->Header.NodeTitle;
 		}
 	}
+
 	if (!Header.Variant.Guid.IsValid())
 	{
 		Header.Variant.Guid = FRigVMVariant::GenerateGUID();
diff --git a/Engine/Plugins/Runtime/RigVM/Source/RigVMDeveloper/Private/RigVMModel/RigVMClient.cpp b/Engine/Plugins/Runtime/RigVM/Source/RigVMDeveloper/Private/RigVMModel/RigVMClient.cpp
--- a/Engine/Plugins/Runtime/RigVM/Source/RigVMDeveloper/Private/RigVMModel/RigVMClient.cpp
+++ b/Engine/Plugins/Runtime/RigVM/Source/RigVMDeveloper/Private/RigVMModel/RigVMClient.cpp
@@ -1162,6 +1162,7 @@ void FRigVMClient::HandleGraphModifiedEvent(ERigVMGraphNotifType InNotifType, UR
 		case ERigVMGraphNotifType::NodeCategoryChanged: // A node's category has changed (Subject == URigVMNode)
 		case ERigVMGraphNotifType::NodeKeywordsChanged: // A node's keywords have changed (Subject == URigVMNode)
 		case ERigVMGraphNotifType::NodeDescriptionChanged: // A node's description has changed (Subject == URigVMNode)
+		case ERigVMGraphNotifType::VariantTagsChanged: // The tags in the header of this function variant have changed (Subject == URigVMLibraryNode)
 		{
 			if(InSubject->GetOuter()->IsA<URigVMFunctionLibrary>())
 			{
diff --git a/Engine/Plugins/Runtime/RigVM/Source/RigVMDeveloper/Private/RigVMModel/RigVMController.cpp b/Engine/Plugins/Runtime/RigVM/Source/RigVMDeveloper/Private/RigVMModel/RigVMController.cpp
--- a/Engine/Plugins/Runtime/RigVM/Source/RigVMDeveloper/Private/RigVMModel/RigVMController.cpp
+++ b/Engine/Plugins/Runtime/RigVM/Source/RigVMDeveloper/Private/RigVMModel/RigVMController.cpp
@@ -11981,6 +11981,13 @@ URigVMLibraryNode* URigVMController::AddFunctionToLibrary(const FName& InFunctio
 		}
 	}
 
+	if (URigVMFunctionLibrary* Library = Cast<URigVMFunctionLibrary>(Graph))
+	{
+		FRigVMVariant Variant;
+		Variant.Guid = FRigVMVariant::GenerateGUID();
+		Library->FunctionToVariant.Add(CollapseNode->GetFName(), Variant);
+	}
+
 	if (bSetupUndoRedo)
 	{
 		FRigVMImportFromTextAction Action(this, CollapseNode);
@@ -12303,6 +12310,191 @@ URigVMLibraryNode* URigVMController::CreateFunctionVariant(const FName& InFuncti
 	return Result;
 }
 
+bool URigVMController::AddDefaultTagToFunctionVariant(const FName& InFunctionName, const FName& InTagName, bool bSetupUndoRedo, bool bPrintPythonCommand)
+{
+	URigVMEditorSettings* Settings = GetMutableDefault<URigVMEditorSettings>(URigVMEditorSettings::StaticClass());
+	if (!Settings)
+	{
+		return false;
+	}
+
+	const FRigVMTag* Tag = Settings->VariantTags.FindByPredicate([InTagName](const FRigVMTag& Tag)
+	{
+		return InTagName == Tag.Name;
+	});
+	if (!Tag)
+	{
+		ReportErrorf(TEXT("Could not find default tag with name %s."), *InTagName.ToString());
+		return false;
+	}
+
+	if (AddTagToFunctionVariant(InFunctionName, *Tag, bSetupUndoRedo, false))
+	{
+		if (bPrintPythonCommand)
+		{
+			const FString GraphName = GetSchema()->GetSanitizedGraphName(GetGraph()->GetGraphName());
+			RigVMPythonUtils::Print(GetSchema()->GetGraphOuterName(GetGraph()), 
+			FString::Printf(TEXT("blueprint.get_controller_by_name('%s').add_default_tag_to_function_variant('%s', %s)"),
+				*GraphName,
+				*InFunctionName.ToString(),
+				*InTagName.ToString()));
+		}
+		return true;
+	}
+	return false;
+}
+
+bool URigVMController::AddTagToFunctionVariant(const FName& InFunctionName, const FRigVMTag& InTag, bool bSetupUndoRedo, bool bPrintPythonCommand)
+{
+	if (!IsValidGraph())
+	{
+		return false;
+	}
+
+	if (!bIsTransacting && !IsGraphEditable())
+	{
+		return false;
+	}
+
+	URigVMGraph* Graph = GetGraph();
+	check(Graph);
+
+	URigVMFunctionLibrary* FunctionLibrary = Cast<URigVMFunctionLibrary>(Graph);
+	if (!FunctionLibrary)
+	{
+		ReportError(TEXT("Can only add tag to variant in library graphs."));
+		return false;
+	}
+
+	URigVMLibraryNode* LibraryNode = Cast<URigVMLibraryNode>(FunctionLibrary->FindFunction(InFunctionName));
+	if (!LibraryNode)
+	{
+		ReportErrorf(TEXT("Could not find library node for function %s."), *InFunctionName.ToString());
+		return false;
+	}
+
+	FRigVMVariant* Variant = FunctionLibrary->FunctionToVariant.Find(InFunctionName);
+	if (!Variant)
+	{
+		ReportErrorf(TEXT("Could not find function variant for function %s."), *InFunctionName.ToString());
+		return false;
+	}
+
+	FRigVMBaseAction Action(this);
+	if (bSetupUndoRedo)
+	{
+		Action.SetTitle(FString::Printf(TEXT("Add tag to function variant")));
+		GetActionStack()->BeginAction(Action);
+		GetActionStack()->AddAction(FRigVMAddFunctionVariantTagAction(this, InFunctionName, InTag));
+	}
+
+	Variant->Tags.Add(InTag);
+
+	Notify(ERigVMGraphNotifType::VariantTagsChanged, LibraryNode);
+
+	if (bSetupUndoRedo)
+	{
+		GetActionStack()->EndAction(Action);
+	}
+
+	if (bPrintPythonCommand)
+	{
+		const FString GraphName = GetSchema()->GetSanitizedGraphName(GetGraph()->GetGraphName());
+	
+		const FString TagPythonString = FString::Printf(TEXT("unreal.RigVMTag(name=\"%s\", label=\"%s\", tool_tip=\"%s\", color=unreal.LinearColor(r=%f, g=%f, b=%f, a=%f), show_in_user_interface=%s, marks_subject_as_invalid=%s)"),
+			*InTag.Name.ToString(),
+			*InTag.Label,
+			*InTag.ToolTip.ToString(),
+			InTag.Color.R, InTag.Color.G, InTag.Color.B, InTag.Color.A,
+			(InTag.bShowInUserInterface) ? TEXT("True") : TEXT("False"),
+			(InTag.bMarksSubjectAsInvalid) ? TEXT("True") : TEXT("False"));
+		RigVMPythonUtils::Print(GetSchema()->GetGraphOuterName(GetGraph()), 
+			FString::Printf(TEXT("blueprint.get_controller_by_name('%s').add_tag_to_function_variant('%s', %s)"),
+				*GraphName,
+				*InFunctionName.ToString(),
+				*TagPythonString));
+	}
+
+	return true;
+}
+
+bool URigVMController::RemoveTagFromFunctionVariant(const FName& InFunctionName, const FName& InTagName, bool bSetupUndoRedo, bool bPrintPythonCommand)
+{
+	if (!IsValidGraph())
+	{
+		return false;
+	}
+
+	if (!bIsTransacting && !IsGraphEditable())
+	{
+		return false;
+	}
+
+	URigVMGraph* Graph = GetGraph();
+	check(Graph);
+
+	URigVMFunctionLibrary* FunctionLibrary = Cast<URigVMFunctionLibrary>(Graph);
+	if (!FunctionLibrary)
+	{
+		ReportError(TEXT("Can only add tag to variant in library graphs."));
+		return false;
+	}
+
+	URigVMLibraryNode* LibraryNode = Cast<URigVMLibraryNode>(FunctionLibrary->FindFunction(InFunctionName));
+	if (!LibraryNode)
+	{
+		ReportErrorf(TEXT("Could not find library node for function %s."), *InFunctionName.ToString());
+		return false;
+	}
+
+	FRigVMVariant* Variant = FunctionLibrary->FunctionToVariant.Find(InFunctionName);
+	if (!Variant)
+	{
+		ReportErrorf(TEXT("Could not find function variant for function %s."), *InFunctionName.ToString());
+		return false;
+	}
+
+	const int32 Index = Variant->Tags.IndexOfByPredicate([InTagName](const FRigVMTag& InTag)
+	{
+		return InTag.Name == InTagName;
+	});
+	if (Index == INDEX_NONE)
+	{
+		ReportErrorf(TEXT("Could not find tag %s for function %s."), *InTagName.ToString(), *InFunctionName.ToString());
+		return false;
+	}
+	
+	FRigVMBaseAction Action(this);
+	if (bSetupUndoRedo)
+	{
+		Action.SetTitle(FString::Printf(TEXT("Remove tag from function variant")));
+		GetActionStack()->BeginAction(Action);
+		GetActionStack()->AddAction(FRigVMRemoveFunctionVariantTagAction(this, InFunctionName, InTagName));
+	}
+
+	Variant->Tags.RemoveAt(Index);
+
+	Notify(ERigVMGraphNotifType::VariantTagsChanged, LibraryNode);
+
+	if (bSetupUndoRedo)
+	{
+		GetActionStack()->EndAction(Action);
+	}
+
+	if (bPrintPythonCommand)
+	{
+		const FString GraphName = GetSchema()->GetSanitizedGraphName(GetGraph()->GetGraphName());
+	
+		RigVMPythonUtils::Print(GetSchema()->GetGraphOuterName(GetGraph()), 
+			FString::Printf(TEXT("blueprint.get_controller_by_name('%s').remove_tag_from_function_variant('%s', '%s')"),
+				*GraphName,
+				*InFunctionName.ToString(),
+				*InTagName.ToString()));
+	}
+
+	return true;
+}
+
 TArray<FRigVMVariantRef> URigVMController::FindVariantsOfFunction(const FName& InFunctionName)
 {
 	TArray<FRigVMVariantRef> Result;
diff --git a/Engine/Plugins/Runtime/RigVM/Source/RigVMDeveloper/Private/RigVMModel/RigVMControllerActions.cpp b/Engine/Plugins/Runtime/RigVM/Source/RigVMDeveloper/Private/RigVMModel/RigVMControllerActions.cpp
--- a/Engine/Plugins/Runtime/RigVM/Source/RigVMDeveloper/Private/RigVMModel/RigVMControllerActions.cpp
+++ b/Engine/Plugins/Runtime/RigVM/Source/RigVMDeveloper/Private/RigVMModel/RigVMControllerActions.cpp
@@ -2091,6 +2091,94 @@ bool FRigVMCreateFunctionVariantAction::Redo()
 	return false;
 }
 
+FRigVMAddFunctionVariantTagAction::FRigVMAddFunctionVariantTagAction()
+: FRigVMBaseAction(nullptr)
+, FunctionName(NAME_None)
+{
+	
+}
+
+FRigVMAddFunctionVariantTagAction::FRigVMAddFunctionVariantTagAction(URigVMController* InController, const FName& InFunctionName, const FRigVMTag& InTag)
+: FRigVMBaseAction(InController)
+, FunctionName(InFunctionName)
+, FunctionTag(InTag)
+{
+}
+
+bool FRigVMAddFunctionVariantTagAction::Undo()
+{
+	if (!FRigVMBaseAction::Undo())
+	{
+		return false;
+	}
+	return GetController()->RemoveTagFromFunctionVariant(FunctionName, FunctionTag.Name, false);
+}
+
+bool FRigVMAddFunctionVariantTagAction::Redo()
+{
+	if(!CanUndoRedo())
+	{
+		return false;
+	}
+#if WITH_EDITOR
+	if (GetController()->AddTagToFunctionVariant(FunctionName, FunctionTag, false))
+	{
+		return FRigVMBaseAction::Redo();
+	}
+#endif
+	return false;
+}
+
+FRigVMRemoveFunctionVariantTagAction::FRigVMRemoveFunctionVariantTagAction()
+: FRigVMBaseAction(nullptr)
+, FunctionName(NAME_None)
+{
+	
+}
+
+FRigVMRemoveFunctionVariantTagAction::FRigVMRemoveFunctionVariantTagAction(URigVMController* InController, const FName& InFunctionName, const FName& InTagName)
+: FRigVMBaseAction(InController)
+, FunctionName(InFunctionName)
+{
+	if (URigVMFunctionLibrary* Library = Cast<URigVMFunctionLibrary>(InController->GetGraph()))
+	{
+		if (const FRigVMVariant* Variant = Library->GetFunctionVariant(InFunctionName))
+		{
+			if (const FRigVMTag* Tag = Variant->Tags.FindByPredicate([InTagName](const FRigVMTag& InTag)
+			{
+				return InTag.Name == InTagName;
+			}))
+			{
+				FunctionTag = *Tag;
+			}
+		}
+	}
+}
+
+bool FRigVMRemoveFunctionVariantTagAction::Undo()
+{
+	if (!FRigVMBaseAction::Undo())
+	{
+		return false;
+	}
+	return GetController()->AddTagToFunctionVariant(FunctionName, FunctionTag, false);
+}
+
+bool FRigVMRemoveFunctionVariantTagAction::Redo()
+{
+	if(!CanUndoRedo())
+	{
+		return false;
+	}
+#if WITH_EDITOR
+	if (GetController()->RemoveTagFromFunctionVariant(FunctionName, FunctionTag.Name, false))
+	{
+		return FRigVMBaseAction::Redo();
+	}
+#endif
+	return false;
+}
+
 FRigVMImportFromTextAction::FRigVMImportFromTextAction()
 : FRigVMBaseAction(nullptr)
 {
diff --git a/Engine/Plugins/Runtime/RigVM/Source/RigVMDeveloper/Public/RigVMModel/Nodes/RigVMFunctionReferenceNode.h b/Engine/Plugins/Runtime/RigVM/Source/RigVMDeveloper/Public/RigVMModel/Nodes/RigVMFunctionReferenceNode.h
--- a/Engine/Plugins/Runtime/RigVM/Source/RigVMDeveloper/Public/RigVMModel/Nodes/RigVMFunctionReferenceNode.h
+++ b/Engine/Plugins/Runtime/RigVM/Source/RigVMDeveloper/Public/RigVMModel/Nodes/RigVMFunctionReferenceNode.h
@@ -57,6 +57,8 @@ public:
 
 	const FRigVMGraphFunctionData* GetReferencedFunctionData(bool bLoadIfNecessary = true) const;
 
+	TArray<FRigVMTag> GetVariantTags() const;
+
 protected:
 
 	virtual FString GetOriginalDefaultValueForRootPin(const URigVMPin* InRootPin) const override;
diff --git a/Engine/Plugins/Runtime/RigVM/Source/RigVMDeveloper/Public/RigVMModel/RigVMController.h b/Engine/Plugins/Runtime/RigVM/Source/RigVMDeveloper/Public/RigVMModel/RigVMController.h
--- a/Engine/Plugins/Runtime/RigVM/Source/RigVMDeveloper/Public/RigVMModel/RigVMController.h
+++ b/Engine/Plugins/Runtime/RigVM/Source/RigVMDeveloper/Public/RigVMModel/RigVMController.h
@@ -967,6 +967,19 @@ public:
 	UFUNCTION(BlueprintCallable, Category = RigVMController)
 	URigVMLibraryNode* CreateFunctionVariant(const FName& InFunctionName, const FName& InVariantName = NAME_None, bool bSetupUndoRedo = true, bool bPrintPythonCommand = false);
 
+	// Adds a default tag to a function variant
+	UFUNCTION(BlueprintCallable, Category = RigVMController)
+	bool AddDefaultTagToFunctionVariant(const FName& InFunctionName, const FName& InTagName, bool bSetupUndoRedo = true, bool bPrintPythonCommand = false);
+	
+	// Adds a tag to a function variant
+	UFUNCTION(BlueprintCallable, Category = RigVMController)
+	bool AddTagToFunctionVariant(const FName& InFunctionName, const FRigVMTag& InTag, bool bSetupUndoRedo = true, bool bPrintPythonCommand = false);
+
+
+	// Adds a tag to a function variant
+	UFUNCTION(BlueprintCallable, Category = RigVMController)
+	bool RemoveTagFromFunctionVariant(const FName& InFunctionName, const FName& InTagName, bool bSetupUndoRedo = true, bool bPrintPythonCommand = false);
+
 	// Returns all variant refs related to the given function
 	UFUNCTION(BlueprintCallable, Category = RigVMController)
 	TArray<FRigVMVariantRef> FindVariantsOfFunction(const FName& InFunctionName);
diff --git a/Engine/Plugins/Runtime/RigVM/Source/RigVMDeveloper/Public/RigVMModel/RigVMControllerActions.h b/Engine/Plugins/Runtime/RigVM/Source/RigVMDeveloper/Public/RigVMModel/RigVMControllerActions.h
--- a/Engine/Plugins/Runtime/RigVM/Source/RigVMDeveloper/Public/RigVMModel/RigVMControllerActions.h
+++ b/Engine/Plugins/Runtime/RigVM/Source/RigVMDeveloper/Public/RigVMModel/RigVMControllerActions.h
@@ -1455,6 +1455,54 @@ public:
 	FName NewFunctionName;
 };
 
+/**
+ * An action adding a tag to a function variant
+ */
+USTRUCT()
+struct FRigVMAddFunctionVariantTagAction : public FRigVMBaseAction
+{
+	GENERATED_BODY()
+
+public:
+
+	FRigVMAddFunctionVariantTagAction();
+	FRigVMAddFunctionVariantTagAction(URigVMController* InController, const FName& InFunctionName, const FRigVMTag& InTag);
+	virtual ~FRigVMAddFunctionVariantTagAction() {};
+	virtual UScriptStruct* GetScriptStruct() const override { return FRigVMCreateFunctionVariantAction::StaticStruct(); }
+	virtual bool Undo() override;
+	virtual bool Redo() override;
+
+	UPROPERTY()
+	FName FunctionName;
+
+	UPROPERTY()
+	FRigVMTag FunctionTag;
+};
+
+/**
+ * An action removing a tag from a function variant
+ */
+USTRUCT()
+struct FRigVMRemoveFunctionVariantTagAction : public FRigVMBaseAction
+{
+	GENERATED_BODY()
+
+public:
+
+	FRigVMRemoveFunctionVariantTagAction();
+	FRigVMRemoveFunctionVariantTagAction(URigVMController* InController, const FName& InFunctionName, const FName& InTagName);
+	virtual ~FRigVMRemoveFunctionVariantTagAction() {};
+	virtual UScriptStruct* GetScriptStruct() const override { return FRigVMCreateFunctionVariantAction::StaticStruct(); }
+	virtual bool Undo() override;
+	virtual bool Redo() override;
+
+	UPROPERTY()
+	FName FunctionName;
+
+	UPROPERTY()
+	FRigVMTag FunctionTag;
+};
+
 /**
  * An action importing nodes and links from text
  */
diff --git a/Engine/Plugins/Runtime/RigVM/Source/RigVMDeveloper/Public/RigVMModel/RigVMNotifications.h b/Engine/Plugins/Runtime/RigVM/Source/RigVMDeveloper/Public/RigVMModel/RigVMNotifications.h
--- a/Engine/Plugins/Runtime/RigVM/Source/RigVMDeveloper/Public/RigVMModel/RigVMNotifications.h
+++ b/Engine/Plugins/Runtime/RigVM/Source/RigVMDeveloper/Public/RigVMModel/RigVMNotifications.h
@@ -53,6 +53,7 @@ enum class ERigVMGraphNotifType : uint8
 	VariableRemappingChanged, // A function reference node's remapping has changed (Subject == URigVMFunctionReferenceNode)
 	LibraryTemplateChanged, // The definition of a library node's template has changed (Subject == URigVMLibraryNode)
 	FunctionAccessChanged, // The function has been made public/private (Subject == URigVMLibraryNode)
+	VariantTagsChanged, // The tags in the header of this function variant have changed (Subject == URigVMLibraryNode)
 	Invalid // The max for this enum (used for guarding)
 };
 
