RigVM: Decorator compiler support and unit tests

#rnx

diff --git a/Engine/Plugins/Experimental/AnimNext/Source/AnimNextEditor/Private/Graph/STraitEditorView.cpp b/Engine/Plugins/Experimental/AnimNext/Source/AnimNextEditor/Private/Graph/STraitEditorView.cpp
--- a/Engine/Plugins/Experimental/AnimNext/Source/AnimNextEditor/Private/Graph/STraitEditorView.cpp
+++ b/Engine/Plugins/Experimental/AnimNext/Source/AnimNextEditor/Private/Graph/STraitEditorView.cpp
@@ -579,7 +579,7 @@ int32 STraitEditorView::GetTraitPinIndex(UAnimNextGraph_EdGraphNode* InEdGraphNo
 									TConstArrayView<FTraitInterfaceUID> TraitImplementedInterfaces = Trait->GetTraitInterfaces();
 									TConstArrayView<FTraitInterfaceUID> TraitRequiredInterfaces = Trait->GetTraitRequiredInterfaces();
 									
-									TSharedPtr<FTraitDataEditorDef>& TraitData = TraitsData.Add_GetRef(MakeShared<FTraitDataEditorDef>(RigVMTrait->GetName()
+									TSharedPtr<FTraitDataEditorDef>& TraitData = TraitsData.Add_GetRef(MakeShared<FTraitDataEditorDef>(*RigVMTrait->GetName()
 																																		, TraitStruct->GetDisplayNameText()
 																																		, Trait->GetTraitUID()
 																																		, Trait->GetTraitMode()
diff --git a/Engine/Plugins/Experimental/AnimNext/Source/AnimNextUncookedOnly/Private/UncookedOnlyUtils.cpp b/Engine/Plugins/Experimental/AnimNext/Source/AnimNextUncookedOnly/Private/UncookedOnlyUtils.cpp
--- a/Engine/Plugins/Experimental/AnimNext/Source/AnimNextUncookedOnly/Private/UncookedOnlyUtils.cpp
+++ b/Engine/Plugins/Experimental/AnimNext/Source/AnimNextUncookedOnly/Private/UncookedOnlyUtils.cpp
@@ -661,6 +661,7 @@ void FUtils::CompileVM(UAnimNextGraph* InGraph)
 	URigVMCompiler* Compiler = URigVMCompiler::StaticClass()->GetDefaultObject<URigVMCompiler>();
 	EditorData->VMCompileSettings.SetExecuteContextStruct(FAnimNextExecuteContext::StaticStruct());
 	FRigVMCompileSettings Settings = (EditorData->bCompileInDebugMode) ? FRigVMCompileSettings::Fast(EditorData->VMCompileSettings.GetExecuteContextStruct()) : EditorData->VMCompileSettings;
+	Settings.ASTSettings.bSetupDecorators = false; // disable the default implementation of decorators for now
 	Settings.ASTSettings.ReportDelegate.BindLambda([InGraph](EMessageSeverity::Type InType, UObject* InObject, const FString& InString)
 	{
 		FMessageLog("AnimNextCompilerResults").Message(InType, FText::FromString(InString));
diff --git a/Engine/Plugins/Runtime/RigVM/Source/RigVM/Private/RigVMCore/RigVM.cpp b/Engine/Plugins/Runtime/RigVM/Source/RigVM/Private/RigVMCore/RigVM.cpp
--- a/Engine/Plugins/Runtime/RigVM/Source/RigVM/Private/RigVMCore/RigVM.cpp
+++ b/Engine/Plugins/Runtime/RigVM/Source/RigVM/Private/RigVMCore/RigVM.cpp
@@ -18,6 +18,7 @@
 #include "UObject/UObjectIterator.h"
 #include "UObject/ObjectSaveContext.h"
 #include "RigVMHost.h"
+#include "RigVMCore/RigVMDecorator.h"
 
 #include UE_INLINE_GENERATED_CPP_BY_NAME(RigVM)
 
@@ -555,6 +556,12 @@ bool URigVM::ValidateBytecode()
 				CheckOperandValidity(Op.Arg);
 				break;
 			}
+			case ERigVMOpCode::SetupDecorators:
+			{
+				const FRigVMSetupDecoratorsOp& Op = ByteCodeStorage.GetOpAt<FRigVMSetupDecoratorsOp>(ByteCodeInstruction);
+				CheckOperandValidity(Op.Arg);
+				break;
+			}
 			case ERigVMOpCode::Invalid:
 			{
 				ensure(false);
@@ -1145,6 +1152,13 @@ void URigVM::InstructionOpEval(FRigVMExtendedExecuteContext& Context, int32 Inst
 			InOpFunc(Context, InHandleBaseIndex, {}, Arg);
 			break;
 		}
+		case ERigVMOpCode::SetupDecorators:
+		{
+			const FRigVMSetupDecoratorsOp& Op = ByteCode.GetOpAt<FRigVMSetupDecoratorsOp>(Instructions[InstructionIndex]);
+			const FRigVMOperand& Arg = Op.Arg;
+			InOpFunc(Context, InHandleBaseIndex, {}, Arg);
+			break;
+		}
 		case ERigVMOpCode::Invalid:
 		default:
 		{
@@ -2125,6 +2139,31 @@ ERigVMExecuteResult URigVM::ExecuteInstructions(FRigVMExtendedExecuteContext& Co
 				ContextPublicData.InstructionIndex++;
 				break;
 			}
+			case ERigVMOpCode::SetupDecorators:
+			{
+				ContextPublicData.Decorators.Reset();
+				
+				const TArray<int32>& DecoratorList = *(TArray<int32>*)Context.CachedMemoryHandles[FirstHandleForInstruction[ContextPublicData.InstructionIndex]].GetData();
+				for(const int32 DecoratorPropertyIndex : DecoratorList)
+				{
+					if(Context.WorkMemoryStorage.GetProperties().IsValidIndex(DecoratorPropertyIndex))
+					{
+						const FProperty* Property = Context.WorkMemoryStorage.GetProperties()[DecoratorPropertyIndex];
+						if(const FStructProperty* StructProperty = CastField<FStructProperty>(Property))
+						{
+							if(StructProperty->Struct && StructProperty->Struct->IsChildOf(FRigVMDecorator::StaticStruct()))
+							{
+								ContextPublicData.Decorators.Emplace(
+									Context.WorkMemoryStorage.GetData<FRigVMDecorator>(StructProperty),
+									Cast<UScriptStruct>(StructProperty->Struct));
+							}
+						}
+					}
+				}
+
+				ContextPublicData.InstructionIndex++;
+				break;
+			}
 			case ERigVMOpCode::Invalid:
 			{
 				ensure(false);
@@ -2461,6 +2500,12 @@ TArray<FString> URigVM::DumpByteCodeAsTextArray(FRigVMExtendedExecuteContext& Co
 				ResultLine = FString::Printf(TEXT("Run Instructions %d-%d (%s)"), Op.StartInstruction, Op.EndInstruction, *GetOperandLabel(Context, Op.Arg, OperandFormatFunction));
 				break;
 			}
+			case ERigVMOpCode::SetupDecorators:
+			{
+				const FRigVMSetupDecoratorsOp& Op = ByteCode.GetOpAt<FRigVMSetupDecoratorsOp>(Instructions[InstructionIndex]);
+				ResultLine = FString::Printf(TEXT("Setup Decorators (%s)"), *GetOperandLabel(Context, Op.Arg, OperandFormatFunction));
+				break;
+			}
 			default:
 			{
 				ensure(false);
diff --git a/Engine/Plugins/Runtime/RigVM/Source/RigVM/Private/RigVMCore/RigVMByteCode.cpp b/Engine/Plugins/Runtime/RigVM/Source/RigVM/Private/RigVMCore/RigVMByteCode.cpp
--- a/Engine/Plugins/Runtime/RigVM/Source/RigVM/Private/RigVMCore/RigVMByteCode.cpp
+++ b/Engine/Plugins/Runtime/RigVM/Source/RigVM/Private/RigVMCore/RigVMByteCode.cpp
@@ -5,6 +5,8 @@
 #include "UObject/UE5MainStreamObjectVersion.h"
 #include "UObject/FortniteMainBranchObjectVersion.h"
 #include "RigVMCore/RigVM.h"
+#include "RigVMCore/RigVMDecorator.h"
+#include "RigVMCore/RigVMDecoratorScope.h"
 #include "RigVMObjectVersion.h"
 
 #include UE_INLINE_GENERATED_CPP_BY_NAME(RigVMByteCode)
@@ -432,6 +434,12 @@ void FRigVMByteCode::Save(FArchive& Ar)
 				Ar << Op;
 				break;
 			}
+			case ERigVMOpCode::SetupDecorators:
+			{
+				FRigVMSetupDecoratorsOp Op = GetOpAt<FRigVMSetupDecoratorsOp>(Instruction.ByteCodeIndex);
+				Ar << Op;
+				break;
+			}
 			default:
 			{
 				ensure(false);
@@ -624,6 +632,13 @@ void FRigVMByteCode::Load(FArchive& Ar)
 				AddOp<FRigVMRunInstructionsOp>(Op);
 				break;
 			}
+			case ERigVMOpCode::SetupDecorators:
+			{
+				FRigVMSetupDecoratorsOp Op;
+				Ar << Op;
+				AddOp<FRigVMSetupDecoratorsOp>(Op);
+				break;
+			}
 			default:
 			{
 				ensure(false);
@@ -835,6 +850,10 @@ uint32 FRigVMByteCode::GetOperatorHash(const FRigVMInstruction& InInstruction) c
 		{
 			return GetTypeHash(GetOpAt<FRigVMRunInstructionsOp>(InInstruction));
 		}
+		case ERigVMOpCode::SetupDecorators:
+		{
+			return GetTypeHash(GetOpAt<FRigVMSetupDecoratorsOp>(InInstruction));
+		}
 		case ERigVMOpCode::Invalid:
 		{
 			ensure(false);
@@ -1051,6 +1070,10 @@ uint64 FRigVMByteCode::GetOpNumBytesAt(uint64 InByteCodeIndex, bool bIncludeOper
 		{
 			return (uint64)sizeof(FRigVMRunInstructionsOp);
 		}
+		case ERigVMOpCode::SetupDecorators:
+		{
+			return (uint64)sizeof(FRigVMSetupDecoratorsOp);
+		}
 		case ERigVMOpCode::Invalid:
 		{
 			ensure(false);
@@ -1341,6 +1364,14 @@ FString FRigVMByteCode::DumpToText() const
 				Line += FString::FromInt(Op.EndInstruction);
 				break;
 			}
+			case ERigVMOpCode::SetupDecorators:
+			{
+				const FRigVMSetupDecoratorsOp& Op = GetOpAt<FRigVMSetupDecoratorsOp>(Instruction.ByteCodeIndex);
+				FString Arg;
+				FRigVMOperand::StaticStruct()->ExportText(Arg, &Op.Arg, &Op.Arg, nullptr, PPF_None, nullptr);
+				Line += Arg;
+				break;
+			}
 			case ERigVMOpCode::Invalid:
 			{
 				ensure(false);
@@ -1389,6 +1420,11 @@ uint64 FRigVMByteCode::AddRunInstructionsOp(FRigVMOperand InExecuteStateArg, int
 	return AddOp(FRigVMRunInstructionsOp(InExecuteStateArg, InStartInstruction, InEndInstruction));
 }
 
+uint64 FRigVMByteCode::AddSetupDecoratorsOp(FRigVMOperand InDecoratorListArg)
+{
+	return AddOp(FRigVMSetupDecoratorsOp(InDecoratorListArg));
+}
+
 int32 FRigVMByteCode::AddBranchInfo(const FRigVMBranchInfo& InBranchInfo)
 {
 	FRigVMBranchInfo BranchInfo = InBranchInfo;
@@ -1493,6 +1529,11 @@ FRigVMOperandArray FRigVMByteCode::GetOperandsForOp(const FRigVMInstruction& InI
 			const FRigVMRunInstructionsOp& Op = GetOpAt<FRigVMRunInstructionsOp>(InInstruction.ByteCodeIndex);
 			return FRigVMOperandArray(&Op.Arg, 1);
 		}
+		case ERigVMOpCode::SetupDecorators:
+		{
+			const FRigVMSetupDecoratorsOp& Op = GetOpAt<FRigVMSetupDecoratorsOp>(InInstruction.ByteCodeIndex);
+			return FRigVMOperandArray(&Op.Arg, 1);
+		}
 		case ERigVMOpCode::JumpAbsolute:
 		case ERigVMOpCode::JumpForward:
 		case ERigVMOpCode::JumpBackward:
@@ -1591,6 +1632,11 @@ uint64 FRigVMByteCode::GetFirstOperandByteIndex(const FRigVMInstruction& InInstr
 			const FRigVMRunInstructionsOp& Op = GetOpAt<FRigVMRunInstructionsOp>(InInstruction.ByteCodeIndex);
 			return InInstruction.ByteCodeIndex + ((uint64)&Op.Arg - (uint64)&Op);
 		}
+		case ERigVMOpCode::SetupDecorators:
+		{
+			const FRigVMSetupDecoratorsOp& Op = GetOpAt<FRigVMSetupDecoratorsOp>(InInstruction.ByteCodeIndex);
+			return InInstruction.ByteCodeIndex + ((uint64)&Op.Arg - (uint64)&Op);
+		}
 		case ERigVMOpCode::JumpAbsolute:
 		case ERigVMOpCode::JumpForward:
 		case ERigVMOpCode::JumpBackward:
@@ -1729,6 +1775,11 @@ uint64 FRigVMByteCode::GetOpAlignment(ERigVMOpCode InOpCode) const
 			static const uint64 Alignment = FRigVMRunInstructionsOp::StaticStruct()->GetCppStructOps()->GetAlignment();
 			return Alignment;
 		}
+		case ERigVMOpCode::SetupDecorators:
+		{
+			static const uint64 Alignment = FRigVMSetupDecoratorsOp::StaticStruct()->GetCppStructOps()->GetAlignment();
+			return Alignment;
+		}
 		case ERigVMOpCode::Invalid:
 		{
 			ensure(false);
@@ -2076,6 +2127,74 @@ void FRigVMByteCode::SetOperandsForInstruction(int32 InInstructionIndex, const F
 
 #endif
 
+TMap<int32, TArray<FRigVMDecoratorScope>> FRigVMByteCode::GetDecorators(
+	FRigVMMemoryStorageStruct& InLiteralMemory,
+	FRigVMMemoryStorageStruct& InWorkMemory,
+	const UScriptStruct* InScriptStruct) const
+{
+	TMap<int32, TArray<FRigVMDecoratorScope>> Result;
+	const FRigVMInstructionArray Instructions = GetInstructions();
+	for(int32 InstructionIndex = 0; InstructionIndex < Instructions.Num(); InstructionIndex++)
+	{
+		const FRigVMInstruction& Instruction = Instructions[InstructionIndex];
+		if(Instruction.OpCode == ERigVMOpCode::SetupDecorators)
+		{
+			const TArray<FRigVMDecoratorScope> Decorators = GetDecoratorsForInstruction(Instruction, InLiteralMemory, InWorkMemory, InScriptStruct);
+			if(!Decorators.IsEmpty())
+			{
+				Result.Add(InstructionIndex, Decorators);
+			}
+		}
+	}
+	return Result;
+}
+
+TArray<FRigVMDecoratorScope> FRigVMByteCode::GetDecoratorsForInstruction(
+	const FRigVMInstruction& InInstruction,
+	FRigVMMemoryStorageStruct& InLiteralMemory,
+	FRigVMMemoryStorageStruct& InWorkMemory,
+	const UScriptStruct* InScriptStruct) const
+{
+	TArray<FRigVMDecoratorScope> Decorators;
+
+	if(InInstruction.OpCode == ERigVMOpCode::SetupDecorators)
+	{
+		if(InScriptStruct == nullptr)
+		{
+			InScriptStruct = FRigVMDecorator::StaticStruct();
+		}
+
+		const FRigVMSetupDecoratorsOp& Op = GetOpAt<FRigVMSetupDecoratorsOp>(InInstruction);
+		check(Op.Arg.GetMemoryType() == ERigVMMemoryType::Literal);
+
+		if(InLiteralMemory.GetProperties().IsValidIndex(Op.Arg.GetRegisterIndex()))
+		{
+			if(const FArrayProperty* DecoratorIndicesProperty = CastField<FArrayProperty>(InLiteralMemory.GetProperty(Op.Arg.GetRegisterIndex())))
+			{
+				if(DecoratorIndicesProperty->Inner->IsA<FIntProperty>())
+				{
+					const TArray<int32>& DecoratorIndices = *InLiteralMemory.GetData<TArray<int32>>(DecoratorIndicesProperty);
+					for(const int32 DecoratorIndex : DecoratorIndices)
+					{
+						const FProperty* Property = InWorkMemory.GetProperties()[DecoratorIndex];
+						if(const FStructProperty* StructProperty = CastField<FStructProperty>(Property))
+						{
+							if(StructProperty->Struct && StructProperty->Struct->IsChildOf(InScriptStruct))
+							{
+								Decorators.Emplace(
+									InWorkMemory.GetData<FRigVMDecorator>(StructProperty),
+									Cast<UScriptStruct>(StructProperty->Struct));
+							}
+						}
+					}
+				}
+			}
+		}
+	}
+
+	return Decorators;
+}
+
 const FRigVMBranchInfo* FRigVMByteCode::GetBranchInfo(const FRigVMBranchInfoKey& InBranchInfoKey) const
 {
 	if(BranchInfos.IsEmpty())
diff --git a/Engine/Plugins/Runtime/RigVM/Source/RigVM/Private/RigVMCore/RigVMStruct.cpp b/Engine/Plugins/Runtime/RigVM/Source/RigVM/Private/RigVMCore/RigVMStruct.cpp
--- a/Engine/Plugins/Runtime/RigVM/Source/RigVM/Private/RigVMCore/RigVMStruct.cpp
+++ b/Engine/Plugins/Runtime/RigVM/Source/RigVM/Private/RigVMCore/RigVMStruct.cpp
@@ -543,6 +543,11 @@ bool FRigVMStruct::CheckFunctionExists(UScriptStruct* InStruct, const FName& Fun
 	return true;
 }
 
+bool FRigVMStruct::ShouldCreatePinForProperty(const FProperty* InProperty) const
+{
+	return true;
+}
+
 ERigVMPinDirection FRigVMStruct::GetPinDirectionFromProperty(FProperty* InProperty)
 {
 	bool bIsInput = InProperty->HasMetaData(InputMetaName);
diff --git a/Engine/Plugins/Runtime/RigVM/Source/RigVM/Public/RigVMCore/RigVM.h b/Engine/Plugins/Runtime/RigVM/Source/RigVM/Public/RigVMCore/RigVM.h
--- a/Engine/Plugins/Runtime/RigVM/Source/RigVM/Public/RigVMCore/RigVM.h
+++ b/Engine/Plugins/Runtime/RigVM/Source/RigVM/Public/RigVMCore/RigVM.h
@@ -428,6 +428,32 @@ public:
 	// returns false if an entry can not be executed
 	bool CanExecuteEntry(const FRigVMExtendedExecuteContext& Context, const FName& InEntryName, bool bLogErrorForMissingEntry = true) const;
 
+	// returns the decorators for this VM's bytecode and Context
+	TMap<int32, TArray<FRigVMDecoratorScope>> GetDecorators(FRigVMExtendedExecuteContext& InContext, const UScriptStruct* InScriptStruct = nullptr)
+	{
+		return GetByteCode().GetDecorators(*GetLiteralMemory(), InContext.WorkMemoryStorage, InScriptStruct);
+	}
+
+	// returns the decorators of a given type for this VM's bytecode and Context
+	template<typename T>
+	TMap<int32, TArray<FRigVMDecoratorScope>> GetDecorators(FRigVMExtendedExecuteContext& InContext)
+	{
+		return GetByteCode().GetDecorators<T>(*GetLiteralMemory(), InContext.WorkMemoryStorage);
+	}
+
+	// returns the decorators for the provided memory for a single instruction
+	TArray<FRigVMDecoratorScope> GetDecoratorsForInstruction(const FRigVMInstruction& InInstruction, FRigVMExtendedExecuteContext& InContext, const UScriptStruct* InScriptStruct = nullptr)
+	{
+		return GetByteCode().GetDecoratorsForInstruction(InInstruction, *GetLiteralMemory(), InContext.WorkMemoryStorage, InScriptStruct);
+	}
+
+	// returns the decorators of a given type for the provided memory for a single instruction
+	template<typename T>
+	TArray<FRigVMDecoratorScope> GetDecoratorsForInstruction(const FRigVMInstruction& InInstruction, FRigVMExtendedExecuteContext& InContext)
+	{
+		return GetByteCode().GetDecoratorsForInstruction<T>(InInstruction, *GetLiteralMemory(), InContext.WorkMemoryStorage);
+	}
+
 #if WITH_EDITOR
 	
 	UE_DEPRECATED(5.3, "Please, use WasInstructionVisitedDuringLastRun with Context param")
diff --git a/Engine/Plugins/Runtime/RigVM/Source/RigVM/Public/RigVMCore/RigVMByteCode.h b/Engine/Plugins/Runtime/RigVM/Source/RigVM/Public/RigVMCore/RigVMByteCode.h
--- a/Engine/Plugins/Runtime/RigVM/Source/RigVM/Public/RigVMCore/RigVMByteCode.h
+++ b/Engine/Plugins/Runtime/RigVM/Source/RigVM/Public/RigVMCore/RigVMByteCode.h
@@ -221,6 +221,7 @@ enum class ERigVMOpCode : uint8
 	JumpToBranch, // jumps to a branch based on a name operand
 	Execute, // single execute op (formerly Execute_0_Operands to Execute_64_Operands)
 	RunInstructions, // runs a set of instructions lazily
+	SetupDecorators, // sets up a list of decorators on executecontext
 	Invalid,
 	FirstArrayOpCode = ArrayReset,
 	LastArrayOpCode = ArrayReverse,
@@ -334,7 +335,8 @@ struct RIGVM_API FRigVMUnaryOp : public FRigVMBaseOp
 			uint8(InOpCode) == uint8(ERigVMOpCode::JumpBackwardIf) ||
 			uint8(InOpCode) == uint8(ERigVMOpCode::ChangeType) ||
 			uint8(InOpCode) == uint8(ERigVMOpCode::JumpToBranch) ||
-			uint8(InOpCode) == uint8(ERigVMOpCode::RunInstructions)
+			uint8(InOpCode) == uint8(ERigVMOpCode::RunInstructions) ||
+			uint8(InOpCode) == uint8(ERigVMOpCode::SetupDecorators)
 		);
 	}
 
@@ -1014,6 +1016,23 @@ struct RIGVM_API FRigVMRunInstructionsOp : public FRigVMUnaryOp
 	}
 };
 
+// sets up a list of decorators in the execute context
+USTRUCT()
+struct RIGVM_API FRigVMSetupDecoratorsOp : public FRigVMUnaryOp
+{
+	GENERATED_USTRUCT_BODY()
+
+	FRigVMSetupDecoratorsOp()
+		: FRigVMUnaryOp()
+	{
+	}
+
+	FRigVMSetupDecoratorsOp(FRigVMOperand InDecoratorListArg)
+		: FRigVMUnaryOp(ERigVMOpCode::SetupDecorators, InDecoratorListArg)
+	{
+	}
+};
+
 /**
  * The FRigVMInstruction represents
  * a single instruction within the VM.
@@ -1216,6 +1235,9 @@ public:
 	// adds a run instructions op
 	uint64 AddRunInstructionsOp(FRigVMOperand InExecuteStateArg, int32 InStartInstruction, int32 InEndInstruction);
 
+	// adds a setup decorators op
+	uint64 AddSetupDecoratorsOp(FRigVMOperand InDecoratorListArg);
+
 	// adds information about a branch for an instruction's argument
 	int32 AddBranchInfo(const FRigVMBranchInfo& InBranchInfo);
 	int32 AddBranchInfo(const FName& InBranchLabel, int32 InInstructionIndex, int32 InArgumentIndex, int32 InFirstBranchInstruction, int32 InLastBranchInstruction);
@@ -1408,6 +1430,26 @@ public:
 
 #endif
 
+	// returns the decorators for the provided memory
+	TMap<int32, TArray<FRigVMDecoratorScope>> GetDecorators(FRigVMMemoryStorageStruct& InLiteralMemory, FRigVMMemoryStorageStruct& InWorkMemory, const UScriptStruct* InScriptStruct = nullptr) const;
+
+	// returns the decorators of a given type for the provided memory
+	template<typename T>
+	TMap<int32, TArray<FRigVMDecoratorScope>> GetDecorators(FRigVMMemoryStorageStruct& InLiteralMemory, FRigVMMemoryStorageStruct& InWorkMemory) const
+	{
+		return GetDecorators(InLiteralMemory, InWorkMemory, T::StaticStruct());
+	}
+
+	// returns the decorators for the provided memory for a single instruction
+	TArray<FRigVMDecoratorScope> GetDecoratorsForInstruction(const FRigVMInstruction& InInstruction, FRigVMMemoryStorageStruct& InLiteralMemory, FRigVMMemoryStorageStruct& InWorkMemory, const UScriptStruct* InScriptStruct = nullptr) const;
+
+	// returns the decorators of a given type for the provided memory for a single instruction
+	template<typename T>
+	TArray<FRigVMDecoratorScope> GetDecoratorsForInstruction(const FRigVMInstruction& InInstruction, FRigVMMemoryStorageStruct& InLiteralMemory, FRigVMMemoryStorageStruct& InWorkMemory) const
+	{
+		return GetDecoratorsForInstruction(InInstruction, InLiteralMemory, InWorkMemory, T::StaticStruct());
+	}
+
 private:
 
 	template<class OpType>
diff --git a/Engine/Plugins/Runtime/RigVM/Source/RigVM/Public/RigVMCore/RigVMDecorator.h b/Engine/Plugins/Runtime/RigVM/Source/RigVM/Public/RigVMCore/RigVMDecorator.h
--- a/Engine/Plugins/Runtime/RigVM/Source/RigVM/Public/RigVMCore/RigVMDecorator.h
+++ b/Engine/Plugins/Runtime/RigVM/Source/RigVM/Public/RigVMCore/RigVMDecorator.h
@@ -25,14 +25,12 @@ struct RIGVM_API FRigVMDecorator : public FRigVMStruct
 public:
 
 	FRigVMDecorator()
-	: Name(NAME_None)
-	, DecoratorStruct(nullptr)
 	{}
 
 	virtual ~FRigVMDecorator() {}
 
 	// returns the name of the decorator (the instance of it on the node)
-	FName GetName() const { return Name; }
+	FString GetName() const { return Name; }
 
 	// returns the display name of the decorator
 	virtual FString GetDisplayName() const
@@ -40,20 +38,6 @@ public:
 		return FString();
 	}
 
-	// returns the struct of this decorator
-	const UScriptStruct* GetScriptStruct() const { return DecoratorStruct; }
-
-	// returns true if the given decorator is of a certain type (or super type)
-	template<typename T>
-	bool IsA() const
-	{
-		if(const UScriptStruct* ScriptStruct = GetScriptStruct())
-		{
-			return ScriptStruct->IsChildOf(T::StaticStruct());
-		}
-		return false;
-	}
-
 #if WITH_EDITOR
 
 	// returns true if this decorator can be added to a given node
@@ -67,15 +51,24 @@ public:
 
 	virtual UScriptStruct* GetDecoratorSharedDataStruct() const { return nullptr; }
 
+	virtual bool ShouldCreatePinForProperty(const FProperty* InProperty) const override
+	{
+		if(!Super::ShouldCreatePinForProperty(InProperty))
+		{
+			return false;
+		}
+		return InProperty->GetFName() != GET_MEMBER_NAME_CHECKED(FRigVMDecorator, Name);
+	}
+
 #endif
 
 private:
 
 	// The name of the decorator on the node
-	FName Name;
-	// The struct backing up the decorator
-	UScriptStruct* DecoratorStruct;
-
+	UPROPERTY()
+	FString Name;
+	
 	friend class URigVMNode;
 	friend class URigVMController;
+	friend class FRigVMDecoratorScope;
 };
diff --git a/Engine/Plugins/Runtime/RigVM/Source/RigVM/Public/RigVMCore/RigVMDecoratorScope.h b/Engine/Plugins/Runtime/RigVM/Source/RigVM/Public/RigVMCore/RigVMDecoratorScope.h
new file mode 100644
--- /dev/null
+++ b/Engine/Plugins/Runtime/RigVM/Source/RigVM/Public/RigVMCore/RigVMDecoratorScope.h
@@ -0,0 +1,80 @@
+// Copyright Epic Games, Inc. All Rights Reserved.
+
+#pragma once
+
+#include "CoreMinimal.h"
+#include "UObject/Class.h"
+
+struct FRigVMDecorator;
+class UScriptStruct;
+
+class RIGVM_API FRigVMDecoratorScope
+{
+public:
+	FRigVMDecoratorScope()
+		: Decorator(nullptr)
+		, ScriptStruct(nullptr)
+	{
+	}
+	
+	FRigVMDecoratorScope(FRigVMDecorator* InDecorator, const UScriptStruct* InScriptStruct)
+		: Decorator(InDecorator)
+		, ScriptStruct(InScriptStruct)
+	{
+	}
+
+	bool IsValid() const
+	{
+		return (Decorator != nullptr) && (ScriptStruct != nullptr);
+	}
+
+	template<typename T>
+	bool IsA() const
+	{
+		return ScriptStruct->IsChildOf(T::StaticStruct());
+	}
+
+	template<typename T = FRigVMDecorator>
+	const T* GetDecorator() const
+	{
+		if(IsA<T>())
+		{
+			return static_cast<T*>(Decorator);
+		}
+		return nullptr;
+	}
+
+	template<typename T = FRigVMDecorator>
+	const T* GetDecoratorChecked() const
+	{
+		check(IsA<T>());
+		return static_cast<T*>(Decorator);
+	}
+
+	template<typename T = FRigVMDecorator>
+	T* GetDecorator()
+	{
+		if(IsA<T>())
+		{
+			return static_cast<T*>(Decorator);
+		}
+		return nullptr;
+	}
+
+	template<typename T = FRigVMDecorator>
+	T* GetDecoratorChecked()
+	{
+		check(IsA<T>());
+		return static_cast<T*>(Decorator);
+	}
+
+	const UScriptStruct* GetScriptStruct() const
+	{
+		return ScriptStruct;
+	}
+
+private:
+
+	FRigVMDecorator* Decorator;
+	const UScriptStruct* ScriptStruct;
+};
\ No newline at end of file
diff --git a/Engine/Plugins/Runtime/RigVM/Source/RigVM/Public/RigVMCore/RigVMExecuteContext.h b/Engine/Plugins/Runtime/RigVM/Source/RigVM/Public/RigVMCore/RigVMExecuteContext.h
--- a/Engine/Plugins/Runtime/RigVM/Source/RigVM/Public/RigVMCore/RigVMExecuteContext.h
+++ b/Engine/Plugins/Runtime/RigVM/Source/RigVM/Public/RigVMCore/RigVMExecuteContext.h
@@ -19,6 +19,7 @@
 #include "RigVMCore/RigVMProfilingInfo.h"
 #include "RigVMCore/RigVMNameCache.h"
 #include "RigVMCore/RigVMMemoryStorageStruct.h"
+#include "RigVMCore/RigVMDecoratorScope.h"
 #include "RigVMLog.h"
 #include "RigVMDrawInterface.h"
 #include "RigVMDrawContainer.h"
@@ -271,6 +272,7 @@ struct FRigVMExecuteContext
 		, OwningComponent(nullptr)
 		, OwningActor(nullptr)
 		, World(nullptr)
+		, Decorators()
 	{
 	}
 
@@ -416,6 +418,20 @@ struct FRigVMExecuteContext
 	FRigVMDrawContainer* GetDrawContainer() { return DrawContainerPtr; }
 	void SetDrawContainer(FRigVMDrawContainer* InDrawContainer) { DrawContainerPtr = InDrawContainer; }
 
+	TArrayView<const FRigVMDecoratorScope> GetDecorators() const
+	{
+		if(Decorators.IsEmpty())
+		{
+			return TArrayView<const FRigVMDecoratorScope>();
+		}
+		return TArrayView<const FRigVMDecoratorScope>(Decorators.GetData(), Decorators.Num());
+	}
+
+	TArrayView<FRigVMDecoratorScope> GetDecorators()
+	{
+		return TArrayView<FRigVMDecoratorScope>(Decorators.GetData(), Decorators.Num());
+	}
+
 	virtual void Initialize()
 	{
 		if(NameCache == nullptr)
@@ -483,6 +499,8 @@ protected:
 	/** The world this VM is running in */
 	const UWorld* World;
 
+	/** The decorators accessible to the current instruction */
+	TArray<FRigVMDecoratorScope> Decorators;
 
 #if UE_RIGVM_DEBUG_EXECUTION
 public:
diff --git a/Engine/Plugins/Runtime/RigVM/Source/RigVM/Public/RigVMCore/RigVMStruct.h b/Engine/Plugins/Runtime/RigVM/Source/RigVM/Public/RigVMCore/RigVMStruct.h
--- a/Engine/Plugins/Runtime/RigVM/Source/RigVM/Public/RigVMCore/RigVMStruct.h
+++ b/Engine/Plugins/Runtime/RigVM/Source/RigVM/Public/RigVMCore/RigVMStruct.h
@@ -207,6 +207,7 @@ public:
 	static bool CheckPinExists(UScriptStruct* InStruct, const FName& PinName, const FString& ExpectedType = FString(), FString* OutErrorMessage = nullptr);
 	static bool CheckMetadata(UScriptStruct* InStruct, const FName& PinName, const FName& InMetadataKey, FString* OutErrorMessage = nullptr);
 	static bool CheckFunctionExists(UScriptStruct* InStruct, const FName& FunctionName, FString* OutErrorMessage = nullptr);
+	virtual bool ShouldCreatePinForProperty(const FProperty* InProperty) const; 
 #endif
 	static FString ExportToFullyQualifiedText(const FProperty* InMemberProperty, const uint8* InMemberMemoryPtr, bool bUseQuotes = true);
 	static FString ExportToFullyQualifiedText(const UScriptStruct* InStruct, const uint8* InStructMemoryPtr, bool bUseQuotes = true);
diff --git a/Engine/Plugins/Runtime/RigVM/Source/RigVMDeveloper/Private/RigVMCompiler/RigVMAST.cpp b/Engine/Plugins/Runtime/RigVM/Source/RigVMDeveloper/Private/RigVMCompiler/RigVMAST.cpp
--- a/Engine/Plugins/Runtime/RigVM/Source/RigVMDeveloper/Private/RigVMCompiler/RigVMAST.cpp
+++ b/Engine/Plugins/Runtime/RigVM/Source/RigVMDeveloper/Private/RigVMCompiler/RigVMAST.cpp
@@ -1260,7 +1260,7 @@ TArray<FRigVMExprAST*> FRigVMParserAST::TraversePins(const FRigVMASTProxy& InNod
 			for (URigVMPin* Pin : Node->GetPins())
 			{
 				// We skip decorator pins as we don't want to traverse them
-				if (!Pin->IsDecoratorPin())
+				if (!Pin->IsDecoratorPin()  || Settings.bSetupDecorators)
 				{
 					Pins.AddUnique(Pin);
 				}
@@ -1352,6 +1352,7 @@ FRigVMExprAST* FRigVMParserAST::TraversePin(const FRigVMASTProxy& InPinProxy, FR
 
 	if ((Pin->GetDirection() == ERigVMPinDirection::Input ||
 		Pin->GetDirection() == ERigVMPinDirection::Visible) &&
+		!Pin->IsDecoratorPin() &&
 		LinkIndices.Num() == 0)
 	{
 		if (Cast<URigVMVariableNode>(Pin->GetNode()) ||
diff --git a/Engine/Plugins/Runtime/RigVM/Source/RigVMDeveloper/Private/RigVMCompiler/RigVMCodeGenerator.cpp b/Engine/Plugins/Runtime/RigVM/Source/RigVMDeveloper/Private/RigVMCompiler/RigVMCodeGenerator.cpp
--- a/Engine/Plugins/Runtime/RigVM/Source/RigVMDeveloper/Private/RigVMCompiler/RigVMCodeGenerator.cpp
+++ b/Engine/Plugins/Runtime/RigVM/Source/RigVMDeveloper/Private/RigVMCompiler/RigVMCodeGenerator.cpp
@@ -1574,6 +1574,12 @@ void FRigVMCodeGenerator::ParseInstructionGroups(const FRigVMExtendedExecuteCont
 					// todo
 					break;
 				}
+				case ERigVMOpCode::SetupDecorators:
+				{
+					const FRigVMSetupDecoratorsOp& Op = ByteCode.GetOpAt<FRigVMSetupDecoratorsOp>(Instruction);
+					// todo
+					break;
+				}
 				default:
 				{
 					break;
diff --git a/Engine/Plugins/Runtime/RigVM/Source/RigVMDeveloper/Private/RigVMCompiler/RigVMCompiler.cpp b/Engine/Plugins/Runtime/RigVM/Source/RigVMDeveloper/Private/RigVMCompiler/RigVMCompiler.cpp
--- a/Engine/Plugins/Runtime/RigVM/Source/RigVMDeveloper/Private/RigVMCompiler/RigVMCompiler.cpp
+++ b/Engine/Plugins/Runtime/RigVM/Source/RigVMDeveloper/Private/RigVMCompiler/RigVMCompiler.cpp
@@ -139,9 +139,9 @@ FRigVMOperand FRigVMCompilerWorkData::AddProperty(
 	return FRigVMOperand(InMemoryType, PropertyIndex);
 }
 
-FRigVMOperand FRigVMCompilerWorkData::FindProperty(ERigVMMemoryType InMemoryType, const FName& InName)
+FRigVMOperand FRigVMCompilerWorkData::FindProperty(ERigVMMemoryType InMemoryType, const FName& InName) const
 {
-	TArray<FRigVMPropertyDescription>* PropertyArray = PropertyDescriptions.Find(InMemoryType);
+	const TArray<FRigVMPropertyDescription>* PropertyArray = PropertyDescriptions.Find(InMemoryType);
 	if(PropertyArray)
 	{
 		for(int32 Index=0;Index<PropertyArray->Num();Index++)
@@ -279,6 +279,18 @@ TRigVMTypeIndex FRigVMCompilerWorkData::GetTypeIndexForOperand(const FRigVMOpera
 	return FRigVMRegistry::Get().GetTypeIndex(CPPTypeName, CPPTypeObject);
 }
 
+FName FRigVMCompilerWorkData::GetUniquePropertyName(ERigVMMemoryType InMemoryType, const FName& InDesiredName) const
+{
+	const FString Prefix = InDesiredName.ToString(); 
+	FString Name =  Prefix;
+	int32 Suffix = 1;
+	while(FindProperty(ERigVMMemoryType::Literal, *Name).IsValid())
+	{
+		Name = FString::Printf(TEXT("%s_%d"), *Prefix, Suffix++);
+	}
+	return *Name;
+}
+
 void FRigVMCompilerWorkData::ReportInfo(const FString& InMessage) const
 {
 	Settings.ReportInfo(InMessage);
@@ -2068,6 +2080,44 @@ bool URigVMCompiler::TraverseCallExtern(const FRigVMCallExternExprAST* InExpr, F
 		{
 			return false;
 		}
+
+		if(WorkData.Settings.ASTSettings.bSetupDecorators)
+		{
+			const TArray<URigVMPin*> DecoratorPins = Node->GetDecoratorPins(); 
+			if(!DecoratorPins.IsEmpty())
+			{
+				// also take care of the empty decorator list
+				if(!WorkData.DecoratorListLiterals.Contains(nullptr))
+				{
+					const FName Name = WorkData.GetUniquePropertyName(ERigVMMemoryType::Literal, TEXT("EmptyDecoratorList"));
+					const FRigVMOperand& ListOperand = WorkData.AddProperty(ERigVMMemoryType::Literal, Name, RigVMTypeUtils::ArrayTypeFromBaseType(RigVMTypeUtils::Int32Type), nullptr, TEXT("()"));
+					WorkData.DecoratorListLiterals.Add(nullptr, ListOperand);
+				}
+
+				if(!WorkData.DecoratorListLiterals.Contains(Node))
+				{
+					TArray<FString> DefaultValues;
+					for(const URigVMPin* DecoratorPin : DecoratorPins)
+					{
+						if(const FRigVMExprAST* DecoratorExpr = InExpr->FindVarWithPinName(DecoratorPin->GetFName()))
+						{
+							check(DecoratorExpr->IsVar());
+							const FRigVMOperand DecoratorOperand = FindOrAddRegister(DecoratorExpr->To<FRigVMVarExprAST>(), WorkData);
+							check(DecoratorOperand.GetMemoryType() == ERigVMMemoryType::Work);
+							DefaultValues.Add(FString::FromInt(DecoratorOperand.GetRegisterIndex()));
+						}
+					}
+
+					if(!DefaultValues.IsEmpty())
+					{
+						const FName Name = WorkData.GetUniquePropertyName(ERigVMMemoryType::Literal, TEXT("DecoratorList"));
+						const FString DefaultValue = FString::Printf(TEXT("(%s)"), *FString::Join(DefaultValues, TEXT(",")));
+						const FRigVMOperand& ListOperand = WorkData.AddProperty(ERigVMMemoryType::Literal, Name, RigVMTypeUtils::ArrayTypeFromBaseType(RigVMTypeUtils::Int32Type), nullptr, DefaultValue);
+						WorkData.DecoratorListLiterals.Add(Node, ListOperand);
+					}
+				}
+			}
+		}
 	}
 	else
 	{
@@ -2256,6 +2306,18 @@ bool URigVMCompiler::TraverseCallExtern(const FRigVMCallExternExprAST* InExpr, F
 			return false;
 		}
 
+		// setup the decorator list for the context
+		bool bSetupDecorators = false;
+		if(const FRigVMOperand* DecoratorListOperand = WorkData.DecoratorListLiterals.Find(Node))
+		{
+			if (WorkData.Settings.SetupNodeInstructionIndex)
+			{
+				WorkData.VM->GetByteCode().SetSubject(WorkData.VM->GetByteCode().GetNumInstructions(), Callstack.GetCallPath(), Callstack.GetStack());
+			}
+			WorkData.VM->GetByteCode().AddSetupDecoratorsOp(*DecoratorListOperand);
+			bSetupDecorators = true;
+		}
+
 		// setup the instruction
 		const int32 FunctionIndex = WorkData.VM->AddRigVMFunction(Function->GetName());
 		check(FunctionIndex != INDEX_NONE);
@@ -2411,6 +2473,19 @@ bool URigVMCompiler::TraverseCallExtern(const FRigVMCallExternExprAST* InExpr, F
 				WorkData.VM->GetByteCode().BranchInfos[BranchIndex].LastInstruction = WorkData.VM->GetByteCode().GetNumInstructions() - 1;
 			}
 		}
+
+		if(bSetupDecorators)
+		{
+			// passing nullptr retrieves the empty decorator list
+			if(const FRigVMOperand* EmptyDecoratorListOperand = WorkData.DecoratorListLiterals.Find(nullptr))
+			{
+				if (WorkData.Settings.SetupNodeInstructionIndex)
+				{
+					WorkData.VM->GetByteCode().SetSubject(WorkData.VM->GetByteCode().GetNumInstructions(), Callstack.GetCallPath(), Callstack.GetStack());
+				}
+				WorkData.VM->GetByteCode().AddSetupDecoratorsOp(*EmptyDecoratorListOperand);
+			}
+		}
 	}
 
 	return true;
@@ -2447,6 +2522,24 @@ bool URigVMCompiler::TraverseInlineFunction(const FRigVMInlineFunctionExprAST* I
 			return false;
 		}
 
+		// create a map of all of the decorator setup lists (indices to decorator properties)
+		TMap<int32, FRigVMOperand> LiteralValueToSetupDecoratorArg;
+		const FRigVMInstructionArray FunctionInstructions = FunctionByteCode.GetInstructions();
+		for(const FRigVMInstruction& Instruction : FunctionInstructions)
+		{
+			if(Instruction.OpCode == ERigVMOpCode::SetupDecorators)
+			{
+				const FRigVMSetupDecoratorsOp& Op = FunctionByteCode.GetOpAt<FRigVMSetupDecoratorsOp>(Instruction);
+				check(Op.Arg.GetMemoryType() == ERigVMMemoryType::Literal);
+				check(FunctionCompilationData->LiteralPropertyDescriptions.IsValidIndex(Op.Arg.GetRegisterIndex()));
+				const FString& DefaultValue = FunctionCompilationData->LiteralPropertyDescriptions[Op.Arg.GetRegisterIndex()].DefaultValue;
+				if(!DefaultValue.IsEmpty() && DefaultValue != TEXT("()"))
+				{
+					LiteralValueToSetupDecoratorArg.Add(Op.Arg.GetRegisterIndex(), FRigVMOperand());
+				}
+			}
+		}
+
 		// Add internal operands (not the ones represented by interface pins)
 		for (uint8 MemoryIndex=0; MemoryIndex< (uint8)ERigVMMemoryType::Invalid; ++MemoryIndex)
 		{
@@ -2454,22 +2547,22 @@ bool URigVMCompiler::TraverseInlineFunction(const FRigVMInlineFunctionExprAST* I
 			TArray<FRigVMFunctionCompilationPropertyDescription> Properties;
 			switch (MemoryType)
 			{
-				case ERigVMMemoryType::Work:
+			case ERigVMMemoryType::Work:
 				{
 					Properties = FunctionCompilationData->WorkPropertyDescriptions;
 					break;
 				}
-				case ERigVMMemoryType::Literal:
+			case ERigVMMemoryType::Literal:
 				{
 					Properties = FunctionCompilationData->LiteralPropertyDescriptions;
 					break;
 				}
-				case ERigVMMemoryType::External:
+			case ERigVMMemoryType::External:
 				{
 					Properties = FunctionCompilationData->ExternalPropertyDescriptions;
 					break;
 				}
-				case ERigVMMemoryType::Debug:
+			case ERigVMMemoryType::Debug:
 				{
 					Properties = FunctionCompilationData->DebugPropertyDescriptions;
 					break;
@@ -2505,11 +2598,12 @@ bool URigVMCompiler::TraverseInlineFunction(const FRigVMInlineFunctionExprAST* I
 				}
 			}
 
-			auto FindOrAddProperty = [&WorkData, MemoryType] (const FRigVMFunctionCompilationPropertyDescription& InProperty, const FString& InNewName, const bool bIsExecuteState) -> FRigVMOperand
+			auto FindOrAddProperty = [&WorkData, MemoryType, LiteralValueToSetupDecoratorArg]
+			(const FRigVMFunctionCompilationPropertyDescription& InProperty, const FString& InNewName, const bool bIsExecuteState) -> FRigVMOperand
 			{
 				// Sharing / reusing memory / operands happens as per following contract:
 				// 1. properties are only shared if their CPP type matches
-				// 2. Literal / constant memory is only shared if the constant values match
+				// 2. Literal / constant memory is only shared if the constant values match (and it is not a decorator list)
 				// 3. Work state is only shared if it is not internal work state private to the instruction referring to it
 				// 4. Work state of type FRigVMInstructionSetExecuteState (bIsExecuteState) is never shared either since it is work state private to a lazy branch.
 
@@ -2527,7 +2621,9 @@ bool URigVMCompiler::TraverseInlineFunction(const FRigVMInlineFunctionExprAST* I
 					{
 						if(MemoryType == ERigVMMemoryType::Literal)
 						{
-							if(ExistingProperty.DefaultValue.Equals(InProperty.DefaultValue))
+							// if the value is the same and this is not a decorator setup list
+							if(ExistingProperty.DefaultValue.Equals(InProperty.DefaultValue) &&
+								!LiteralValueToSetupDecoratorArg.Contains(Operand.GetRegisterIndex()))
 							{
 								return Operand;
 							}
@@ -2566,10 +2662,13 @@ bool URigVMCompiler::TraverseInlineFunction(const FRigVMInlineFunctionExprAST* I
 				FString NewName = Description.Name.ToString();
 				static const FString FunctionLibraryPrefix = TEXT("FunctionLibrary");
 
+				const UScriptStruct* ScriptStruct = Cast<UScriptStruct>(Description.CPPTypeObject.Get());
+
 				// instantiate function library specific work state as well as
 				// instruction set execute state - which is used for lazy blocks.
 				const bool bIsExecuteState = Description.CPPType.Equals(RigVMInstructionSetExecuteStateName);
-				if (NewName.StartsWith(FunctionLibraryPrefix) || bIsExecuteState)
+				const bool bIsDecorator = ScriptStruct && ScriptStruct->IsChildOf(FRigVMDecorator::StaticStruct());
+				if (NewName.StartsWith(FunctionLibraryPrefix) || bIsExecuteState || bIsDecorator)
 				{
 					NewName = FString::Printf(TEXT("%s%s"), *FunctionReferenceNode->GetNodePath(), *NewName.RightChop(FunctionLibraryPrefix.Len()));
 					FRigVMPropertyDescription::SanitizeName(NewName);
@@ -2578,6 +2677,49 @@ bool URigVMCompiler::TraverseInlineFunction(const FRigVMInlineFunctionExprAST* I
 				const FRigVMOperand Operand = FindOrAddProperty(Description, NewName, bIsExecuteState);
 				FRigVMCompilerWorkData::FFunctionRegisterData Data = {FunctionReferenceNode, MemoryType, PropertyIndex};
 				WorkData.FunctionRegisterToOperand.Add(Data, Operand);
+
+				if(MemoryType == ERigVMMemoryType::Literal && LiteralValueToSetupDecoratorArg.Contains(PropertyIndex))
+				{
+					LiteralValueToSetupDecoratorArg.FindChecked(PropertyIndex) = Operand;
+				}
+			}
+		}
+		// For decorator setup lists we need to update the integer values (pointing to decorator property indices)
+		for(const TPair<int32, FRigVMOperand>& Pair : LiteralValueToSetupDecoratorArg)
+		{
+			const FRigVMOperand DecoratorIndicesOperand = Pair.Value;
+			const int32 LiteralPropertyIndex = DecoratorIndicesOperand.GetRegisterIndex();
+			FString OriginalDecoratorIndicesString = WorkData.PropertyDescriptions[ERigVMMemoryType::Literal][LiteralPropertyIndex].DefaultValue;
+			if(!OriginalDecoratorIndicesString.IsEmpty() && OriginalDecoratorIndicesString != TEXT("()"))
+			{
+				OriginalDecoratorIndicesString = OriginalDecoratorIndicesString.TrimChar(TEXT('('));
+				OriginalDecoratorIndicesString = OriginalDecoratorIndicesString.TrimChar(TEXT(')'));
+				FString PinPathRemaining = OriginalDecoratorIndicesString;
+				FString Left, Right;
+				TArray<FString> IndexStrings;
+				while(PinPathRemaining.Split(TEXT(","), &Left, &Right))
+				{
+					IndexStrings.Add(Left.TrimStartAndEnd());
+					Left.Empty();
+					PinPathRemaining = Right;
+				}
+				if (!Right.IsEmpty())
+				{
+					IndexStrings.Add(Right.TrimStartAndEnd());
+				}
+
+				for(int32 PartIndex = 0; PartIndex < IndexStrings.Num(); PartIndex++)
+				{
+					const int32 OriginalIndex = FCString::Atoi(*IndexStrings[PartIndex]);
+					FRigVMCompilerWorkData::FFunctionRegisterData Data = {FunctionReferenceNode, ERigVMMemoryType::Work, OriginalIndex};
+					const FRigVMOperand& NewOperand = WorkData.FunctionRegisterToOperand.FindChecked(Data);
+					check(NewOperand.GetMemoryType() == ERigVMMemoryType::Work);
+					IndexStrings[PartIndex] = FString::FromInt(NewOperand.GetRegisterIndex());
+				}
+
+				IndexStrings.Remove(FString());
+				WorkData.PropertyDescriptions[ERigVMMemoryType::Literal][LiteralPropertyIndex].DefaultValue =
+					FString::Printf(TEXT("(%s)"), *FString::Join(IndexStrings, TEXT(",")));
 			}
 		}
 	}
diff --git a/Engine/Plugins/Runtime/RigVM/Source/RigVMDeveloper/Private/RigVMModel/RigVMController.cpp b/Engine/Plugins/Runtime/RigVM/Source/RigVMDeveloper/Private/RigVMModel/RigVMController.cpp
--- a/Engine/Plugins/Runtime/RigVM/Source/RigVMDeveloper/Private/RigVMModel/RigVMController.cpp
+++ b/Engine/Plugins/Runtime/RigVM/Source/RigVMDeveloper/Private/RigVMModel/RigVMController.cpp
@@ -736,6 +736,14 @@ void FRigVMPinInfoArray::AddPins(UScriptStruct* InScriptStruct, URigVMController
 	ERigVMPinDirection InDirection, int32 InParentIndex, TFunction<ERigVMPinDefaultValueType(const FName&)> InDefaultValueTypeGetter,
 	const uint8* InDefaultValueMemory, bool bAddSubPins)
 {
+	TSharedPtr<FStructOnScope> StructOnScope;
+	const FRigVMStruct* RigVMStructInstance = nullptr;
+	if(InScriptStruct->IsChildOf(FRigVMStruct::StaticStruct()))
+	{
+		StructOnScope = MakeShareable(new FStructOnScope(InScriptStruct));
+		RigVMStructInstance = reinterpret_cast<FRigVMStruct*>(StructOnScope->GetStructMemory()); 
+	}
+
 	if (InController->GetSchema()->ShouldUnfoldStruct(InController, InScriptStruct))
 	{
 		TArray<UStruct*> StructsToVisit = FRigVMTemplate::GetSuperStructs(InScriptStruct, true);
@@ -745,6 +753,15 @@ void FRigVMPinInfoArray::AddPins(UScriptStruct* InScriptStruct, URigVMController
 			// properties of the super struct in this iterator.
 			for (TFieldIterator<FProperty> It(StructToVisit, EFieldIterationFlags::None); It; ++It)
 			{
+				// filter pins that structs want to hide
+				if(RigVMStructInstance)
+				{
+					if(!RigVMStructInstance->ShouldCreatePinForProperty(*It))
+					{
+						continue;
+					}
+				}
+
 				const uint8* DefaultValueMemory = nullptr;
 				if(InDefaultValueMemory)
 				{
@@ -14475,8 +14492,7 @@ FName URigVMController::AddDecorator(URigVMNode* InNode, UScriptStruct* InDecora
 		}
 	}
 
-	Decorator->Name = ValidDecoratorName;
-	Decorator->DecoratorStruct = InDecoratorScriptStruct;
+	Decorator->Name = ValidDecoratorName.ToString();
 
 	FString FailureReason;
 	if(!Decorator->CanBeAddedToNode(InNode, &FailureReason))
@@ -14842,6 +14858,14 @@ void URigVMController::AddPinsForStruct(UStruct* InStruct, URigVMNode* InNode, U
 		}
 	}
 
+	TSharedPtr<FStructOnScope> StructOnScope;
+	const FRigVMStruct* RigVMStructInstance = nullptr;
+	if(InStruct->IsChildOf(FRigVMStruct::StaticStruct()))
+	{
+		StructOnScope = MakeShareable(new FStructOnScope(InStruct));
+		RigVMStructInstance = reinterpret_cast<FRigVMStruct*>(StructOnScope->GetStructMemory()); 
+	}
+
 	TArray<UStruct*> StructsToVisit = FRigVMTemplate::GetSuperStructs(InStruct, true);
 	for(UStruct* StructToVisit : StructsToVisit)
 	{
@@ -14849,6 +14873,15 @@ void URigVMController::AddPinsForStruct(UStruct* InStruct, URigVMNode* InNode, U
 		// properties of the super struct in this iterator.
 		for (TFieldIterator<FProperty> It(StructToVisit, EFieldIterationFlags::None); It; ++It)
 		{
+			// filter pins that structs want to hide
+			if(RigVMStructInstance)
+			{
+				if(!RigVMStructInstance->ShouldCreatePinForProperty(*It))
+				{
+					continue;
+				}
+			}
+			
 			FName PropertyName = It->GetFName();
 
 			URigVMPin* Pin = NewObject<URigVMPin>(InParentPin == nullptr ? Cast<UObject>(InNode) : Cast<UObject>(InParentPin), PropertyName);
diff --git a/Engine/Plugins/Runtime/RigVM/Source/RigVMDeveloper/Private/RigVMModel/RigVMNode.cpp b/Engine/Plugins/Runtime/RigVM/Source/RigVMDeveloper/Private/RigVMModel/RigVMNode.cpp
--- a/Engine/Plugins/Runtime/RigVM/Source/RigVMDeveloper/Private/RigVMModel/RigVMNode.cpp
+++ b/Engine/Plugins/Runtime/RigVM/Source/RigVMDeveloper/Private/RigVMModel/RigVMNode.cpp
@@ -344,6 +344,10 @@ void URigVMNode::UpdateDecoratorRootPinNames()
 	{
 		if(Pin->IsDecoratorPin())
 		{
+			if(URigVMPin* NamePin = Pin->FindSubPin(TEXT("Name")))
+			{
+				NamePin->DefaultValue = Pin->GetName();
+			}
 			NewDecoratorRootPinNames.Add(Pin->GetName());
 		}
 	}
@@ -555,14 +559,18 @@ bool URigVMNode::IsDecoratorPin(const URigVMPin* InDecoratorPin) const
 	return FindDecorator(InDecoratorPin) != nullptr;
 }
 
-URigVMPin* URigVMNode::FindDecorator(const FName& InName) const
+URigVMPin* URigVMNode::FindDecorator(const FName& InName, const FString& InSubPinPath) const
 {
 	const FString NameString = InName.ToString();
 	for(const FString& DecoratorRootPinName : DecoratorRootPinNames)
 	{
 		if(DecoratorRootPinName.Equals(NameString, ESearchCase::CaseSensitive))
 		{
-			return FindPin(DecoratorRootPinName);
+			if(InSubPinPath.IsEmpty())
+			{
+				return FindPin(DecoratorRootPinName);
+			}
+			return FindPin(URigVMPin::JoinPinPath(DecoratorRootPinName, InSubPinPath));
 		}
 	}
 	return nullptr;
@@ -612,8 +620,7 @@ TSharedPtr<FStructOnScope> URigVMNode::GetDecoratorInstance(const URigVMPin* InD
 			}
 		}
 
-		Decorator->Name = RootPin->GetFName();
-		Decorator->DecoratorStruct = ScriptStruct;
+		Decorator->Name = RootPin->GetName();
 		
 		return Scope;
 	}
diff --git a/Engine/Plugins/Runtime/RigVM/Source/RigVMDeveloper/Private/RigVMModel/RigVMPin.cpp b/Engine/Plugins/Runtime/RigVM/Source/RigVMDeveloper/Private/RigVMModel/RigVMPin.cpp
--- a/Engine/Plugins/Runtime/RigVM/Source/RigVMDeveloper/Private/RigVMModel/RigVMPin.cpp
+++ b/Engine/Plugins/Runtime/RigVM/Source/RigVMDeveloper/Private/RigVMModel/RigVMPin.cpp
@@ -810,9 +810,13 @@ FString URigVMPin::GetDefaultValue(const URigVMPin::FPinOverride& InOverride, bo
 	}
 	else if (IsStruct())
 	{
-		if (SubPins.Num() > 0)
+		if (SubPins.Num() > 0 || IsDecoratorPin())
 		{
 			TArray<FString> MemberDefaultValues;
+			if(IsDecoratorPin())
+			{
+				MemberDefaultValues.Add(FString::Printf(TEXT("Name=\"%s\""), *GetName()));
+			}
 			for (const URigVMPin* SubPin : SubPins)
 			{
 				FString MemberDefaultValue = SubPin->GetDefaultValue(InOverride, bAdaptValueForPinType);
diff --git a/Engine/Plugins/Runtime/RigVM/Source/RigVMDeveloper/Public/RigVMCompiler/RigVMAST.h b/Engine/Plugins/Runtime/RigVM/Source/RigVMDeveloper/Public/RigVMCompiler/RigVMAST.h
--- a/Engine/Plugins/Runtime/RigVM/Source/RigVMDeveloper/Public/RigVMCompiler/RigVMAST.h
+++ b/Engine/Plugins/Runtime/RigVM/Source/RigVMDeveloper/Public/RigVMCompiler/RigVMAST.h
@@ -1148,6 +1148,9 @@ struct RIGVMDEVELOPER_API FRigVMParserASTSettings
 	UPROPERTY(EditAnywhere, Category = "AST")
 	bool bFoldLiterals = false;
 
+	UPROPERTY(EditAnywhere, Transient, BlueprintReadWrite, Category = "AST")
+	bool bSetupDecorators = true;
+
 	// links to be ignored during the parse
 	UPROPERTY()
 	TArray<TObjectPtr<URigVMLink>> LinksToSkip;
diff --git a/Engine/Plugins/Runtime/RigVM/Source/RigVMDeveloper/Public/RigVMCompiler/RigVMCompiler.h b/Engine/Plugins/Runtime/RigVM/Source/RigVMDeveloper/Public/RigVMCompiler/RigVMCompiler.h
--- a/Engine/Plugins/Runtime/RigVM/Source/RigVMDeveloper/Public/RigVMCompiler/RigVMCompiler.h
+++ b/Engine/Plugins/Runtime/RigVM/Source/RigVMDeveloper/Public/RigVMCompiler/RigVMCompiler.h
@@ -141,7 +141,7 @@ public:
 	TArray<const FRigVMExprAST*> ExprToSkip;
 	TArray<const FRigVMExprAST*> TraversalExpressions;
 	TMap<FString, int32> ProcessedLinks;
-	TMap<int32, FRigVMOperand> IntegerLiterals;
+	TMap<const URigVMNode*, FRigVMOperand> DecoratorListLiterals;
 	FRigVMOperand ComparisonOperand;
 
 	using FRigVMASTProxyArray = TArray<FRigVMASTProxy, TInlineAllocator<3>>; 
@@ -178,11 +178,12 @@ public:
 	TMap<ERigVMMemoryType, TArray<FRigVMPropertyDescription>> PropertyDescriptions;
 
 	FRigVMOperand AddProperty(ERigVMMemoryType InMemoryType, const FName& InName, const FString& InCPPType, UObject* InCPPTypeObject, const FString& InDefaultValue = FString());
-	FRigVMOperand FindProperty(ERigVMMemoryType InMemoryType, const FName& InName);
+	FRigVMOperand FindProperty(ERigVMMemoryType InMemoryType, const FName& InName) const;
 	FRigVMPropertyDescription GetProperty(const FRigVMOperand& InOperand);
 	int32 FindOrAddPropertyPath(const FRigVMOperand& InOperand, const FString& InHeadCPPType, const FString& InSegmentPath);
 	const FProperty* GetPropertyForOperand(const FRigVMOperand& InOperand) const;
 	TRigVMTypeIndex GetTypeIndexForOperand(const FRigVMOperand& InOperand) const;
+	FName GetUniquePropertyName(ERigVMMemoryType InMemoryType, const FName& InDesiredName) const;
 
 	TSharedPtr<FRigVMParserAST> AST;
 	
@@ -268,7 +269,6 @@ public:
 		ExprComplete.Reset();
 		ExprToSkip.Reset();
 		ProcessedLinks.Reset();
-		IntegerLiterals.Reset();
 		CachedProxiesWithSharedOperand.Reset();
 		ProxyTargets.Reset();
 		BranchInfos.Reset();
diff --git a/Engine/Plugins/Runtime/RigVM/Source/RigVMDeveloper/Public/RigVMModel/Nodes/RigVMUnitNode.h b/Engine/Plugins/Runtime/RigVM/Source/RigVMDeveloper/Public/RigVMModel/Nodes/RigVMUnitNode.h
--- a/Engine/Plugins/Runtime/RigVM/Source/RigVMDeveloper/Public/RigVMModel/Nodes/RigVMUnitNode.h
+++ b/Engine/Plugins/Runtime/RigVM/Source/RigVMDeveloper/Public/RigVMModel/Nodes/RigVMUnitNode.h
@@ -95,9 +95,6 @@ public:
 
 	virtual uint32 GetStructureHash() const override;
 
-	// allows the node to support non-native pins
-	virtual bool HasNonNativePins() const { return false; }
-
 protected:
 
 	virtual FText GetToolTipTextForPin(const URigVMPin* InPin) const override;
diff --git a/Engine/Plugins/Runtime/RigVM/Source/RigVMDeveloper/Public/RigVMModel/RigVMNode.h b/Engine/Plugins/Runtime/RigVM/Source/RigVMDeveloper/Public/RigVMModel/RigVMNode.h
--- a/Engine/Plugins/Runtime/RigVM/Source/RigVMDeveloper/Public/RigVMModel/RigVMNode.h
+++ b/Engine/Plugins/Runtime/RigVM/Source/RigVMDeveloper/Public/RigVMModel/RigVMNode.h
@@ -301,6 +301,9 @@ public:
 
 	virtual uint32 GetStructureHash() const;
 
+	// allows the node to support non-native pins
+	virtual bool HasNonNativePins() const { return !DecoratorRootPinNames.IsEmpty(); }
+
 	UFUNCTION(BlueprintPure, Category = RigVMNode)
 	TArray<URigVMPin*> GetDecoratorPins() const;
 
@@ -311,7 +314,7 @@ public:
 
 	bool IsDecoratorPin(const URigVMPin* InDecoratorPin) const;
 
-	URigVMPin* FindDecorator(const FName& InName) const;
+	URigVMPin* FindDecorator(const FName& InName, const FString& InSubPinPath = FString()) const;
 
 	URigVMPin* FindDecorator(const URigVMPin* InDecoratorPin) const;
 
diff --git a/Engine/Plugins/Runtime/RigVM/Source/RigVMEditor/Private/Widgets/SRigVMExecutionStackView.cpp b/Engine/Plugins/Runtime/RigVM/Source/RigVMEditor/Private/Widgets/SRigVMExecutionStackView.cpp
--- a/Engine/Plugins/Runtime/RigVM/Source/RigVMEditor/Private/Widgets/SRigVMExecutionStackView.cpp
+++ b/Engine/Plugins/Runtime/RigVM/Source/RigVMEditor/Private/Widgets/SRigVMExecutionStackView.cpp
@@ -585,6 +585,11 @@ void SRigVMExecutionStackView::PopulateStackView(URigVM* InVM, FRigVMExtendedExe
 						Label = FString::Printf(TEXT("Run Instructions %d-%d"), Op.StartInstruction, Op.EndInstruction);
 						break;
 					}
+					case ERigVMOpCode::SetupDecorators:
+					{
+						Label = TEXT("Setup Decorators");
+						break;
+					}
 					case ERigVMOpCode::Exit:
 					{
 						Label = TEXT("Exit");
