RigVM: Fix unit test breakage

#rnx

diff --git a/Engine/Plugins/Runtime/RigVM/Source/RigVMDeveloper/Private/RigVMModel/RigVMController.cpp b/Engine/Plugins/Runtime/RigVM/Source/RigVMDeveloper/Private/RigVMModel/RigVMController.cpp
--- a/Engine/Plugins/Runtime/RigVM/Source/RigVMDeveloper/Private/RigVMModel/RigVMController.cpp
+++ b/Engine/Plugins/Runtime/RigVM/Source/RigVMDeveloper/Private/RigVMModel/RigVMController.cpp
@@ -6753,7 +6753,8 @@ bool URigVMController::RemoveNodes(TArray<URigVMNode*> InNodes, bool bSetupUndoR
 			{
 				TArray<URigVMNode*> ContainedNodes = SubGraph->GetNodes();
 				TGuardValue<bool> SuspendTemplateComputation(bSuspendTemplateComputation, true);
-				TGuardValue<bool> DisableSchemaRemoveNodeCheck(bEnableSchemaRemoveNodeCheck, false);
+				TGuardValue<bool> SubGraphSuspendTemplateComputation(SubGraphController->bSuspendTemplateComputation, true);
+				TGuardValue<bool> SubGraphDisableSchemaRemoveNodeCheck(SubGraphController->bEnableSchemaRemoveNodeCheck, false);
 				SubGraphController->RemoveNodes(ContainedNodes, false, false);
 
 				if(IRigVMClientHost* ClientHost = GetImplementingOuter<IRigVMClientHost>())
