RigVM: remove need for copy initialization

#rnx

diff --git a/Engine/Plugins/Runtime/RigVM/Source/RigVMDeveloper/Private/RigVMCompiler/RigVMAST.cpp b/Engine/Plugins/Runtime/RigVM/Source/RigVMDeveloper/Private/RigVMCompiler/RigVMAST.cpp
--- a/Engine/Plugins/Runtime/RigVM/Source/RigVMDeveloper/Private/RigVMCompiler/RigVMAST.cpp
+++ b/Engine/Plugins/Runtime/RigVM/Source/RigVMDeveloper/Private/RigVMCompiler/RigVMAST.cpp
@@ -1350,9 +1350,12 @@ FRigVMExprAST* FRigVMParserAST::TraversePin(const FRigVMASTProxy& InPinProxy, FR
 		}
 	}
 
-	if ((Pin->GetDirection() == ERigVMPinDirection::Input ||
+	if(Pin->IsDecoratorPin())
+	{
+		PinExpr = MakeExpr<FRigVMVarExprAST>(FRigVMExprAST::EType::Var, InPinProxy);
+	}
+	else if ((Pin->GetDirection() == ERigVMPinDirection::Input ||
 		Pin->GetDirection() == ERigVMPinDirection::Visible) &&
-		!Pin->IsDecoratorPin() &&
 		LinkIndices.Num() == 0)
 	{
 		if (Cast<URigVMVariableNode>(Pin->GetNode()) ||
@@ -1411,7 +1414,8 @@ FRigVMExprAST* FRigVMParserAST::TraversePin(const FRigVMASTProxy& InPinProxy, FR
 
 		if (!bHasSourceLinkToRoot && 
 			GetSourceLinkIndices(InPinProxy, false).Num() == 0 &&
-			(Pin->GetDirection() == ERigVMPinDirection::IO || LinkIndices.Num() > 0))
+			(Pin->GetDirection() == ERigVMPinDirection::IO ||
+			(LinkIndices.Num() > 0 && (Pin->IsArray() || Pin->GetNode()->IsA<URigVMVariableNode>()))))
 		{
 			FRigVMLiteralExprAST* LiteralExpr = MakeExpr<FRigVMLiteralExprAST>(InPinProxy);
 			const FRigVMASTLinkDescription LiteralLink(InPinProxy, InPinProxy, FString());
