Control Rig: Fix crash when loading transforming robot asset

#rnx

diff --git a/Engine/Plugins/Runtime/RigVM/Source/RigVMDeveloper/Private/RigVMModel/RigVMController.cpp b/Engine/Plugins/Runtime/RigVM/Source/RigVMDeveloper/Private/RigVMModel/RigVMController.cpp
--- a/Engine/Plugins/Runtime/RigVM/Source/RigVMDeveloper/Private/RigVMModel/RigVMController.cpp
+++ b/Engine/Plugins/Runtime/RigVM/Source/RigVMDeveloper/Private/RigVMModel/RigVMController.cpp
@@ -16826,6 +16826,17 @@ URigVMController::FPinState URigVMController::GetPinState(URigVMPin* InPin, bool
 		}
 		State.InjectionInfos.Reset();
 	}
+	else
+	{
+		// move the injection info under the graph so that they don't get lost
+		// when the pins get destroyed. the injection info will get renamed / reparented
+		// under the pin again once it is re-populated - or destroyed if the pin
+		// no longer exists during ApplyPinState.
+		for(URigVMInjectionInfo* InjectionInfo : State.InjectionInfos)
+		{
+			RenameObject(InjectionInfo, nullptr, InPin->GetGraph());
+		}
+	}
 
 	return State;
 }
