Control Rig: Swap modules widget
#rnx

diff --git a/Engine/Plugins/Animation/ControlRig/Source/ControlRigEditor/Private/ControlRigModularRigCommands.cpp b/Engine/Plugins/Animation/ControlRig/Source/ControlRigEditor/Private/ControlRigModularRigCommands.cpp
--- a/Engine/Plugins/Animation/ControlRig/Source/ControlRigEditor/Private/ControlRigModularRigCommands.cpp
+++ b/Engine/Plugins/Animation/ControlRig/Source/ControlRigEditor/Private/ControlRigModularRigCommands.cpp
@@ -11,6 +11,7 @@ void FControlRigModularRigCommands::RegisterCommands()
 	UI_COMMAND(DeleteModuleItem, "Delete", "Delete module from the rig.", EUserInterfaceActionType::Button, FInputChord(EKeys::Delete));
 	UI_COMMAND(MirrorModuleItem, "Mirror", "Mirror module from the rig.", EUserInterfaceActionType::Button, FInputChord());
 	UI_COMMAND(ReresolveModuleItem, "Auto Resolve", "Auto Resolve secondary connectors.", EUserInterfaceActionType::Button, FInputChord());
+	UI_COMMAND(SwapModuleClassItem, "Swap Module", "Swap module class.", EUserInterfaceActionType::Button, FInputChord());
 }
 
 #undef LOCTEXT_NAMESPACE
diff --git a/Engine/Plugins/Animation/ControlRig/Source/ControlRigEditor/Private/ControlRigModularRigCommands.h b/Engine/Plugins/Animation/ControlRig/Source/ControlRigEditor/Private/ControlRigModularRigCommands.h
--- a/Engine/Plugins/Animation/ControlRig/Source/ControlRigEditor/Private/ControlRigModularRigCommands.h
+++ b/Engine/Plugins/Animation/ControlRig/Source/ControlRigEditor/Private/ControlRigModularRigCommands.h
@@ -33,6 +33,9 @@ public:
 	/** Reresolve Module */
 	TSharedPtr< FUICommandInfo > ReresolveModuleItem;
 
+	/** Swap Module Class */
+	TSharedPtr< FUICommandInfo > SwapModuleClassItem;
+
 	/**
 	 * Initialize commands
 	 */
diff --git a/Engine/Plugins/Animation/ControlRig/Source/ControlRigEditor/Private/Editor/ControlRigEditor.cpp b/Engine/Plugins/Animation/ControlRig/Source/ControlRigEditor/Private/Editor/ControlRigEditor.cpp
--- a/Engine/Plugins/Animation/ControlRig/Source/ControlRigEditor/Private/Editor/ControlRigEditor.cpp
+++ b/Engine/Plugins/Animation/ControlRig/Source/ControlRigEditor/Private/Editor/ControlRigEditor.cpp
@@ -111,12 +111,148 @@
 #include "SchematicGraphPanel/SSchematicGraphPanel.h"
 #include "RigVMCore/RigVMExecuteContext.h"
 #include "Editor/RigVMGraphDetailCustomization.h"
+#include "Widgets/SRigVMSwapAssetReferencesWidget.h"
+#include "Widgets/SRigVMBulkEditDialog.h"
 
 #define LOCTEXT_NAMESPACE "ControlRigEditor"
 
 TAutoConsoleVariable<bool> CVarControlRigShowTestingToolbar(TEXT("ControlRig.Test.EnableTestingToolbar"), false, TEXT("When true we'll show the testing toolbar in Control Rig Editor."));
 TAutoConsoleVariable<bool> CVarShowSchematicPanelOverlay(TEXT("ControlRig.Preview.ShowSchematicPanelOverlay"), true, TEXT("When true we'll add an overlay to the persona viewport to show modular rig information."));
 
+FAutoConsoleCommand FCmdModularRigSwapModuleWidget
+(
+	TEXT("ControlRig.SwapModules"),
+	TEXT("Create swap modules wizard."),
+	FConsoleCommandDelegate::CreateLambda([]()
+	{
+		SRigVMSwapAssetReferencesWidget::FArguments WidgetArgs;
+		WidgetArgs
+			.EnableUndo(true)
+			.CloseOnSuccess(true)
+			.OnGetReferences_Lambda([](const FAssetData& ReferencedAsset) -> TArray<FSoftObjectPath>
+			{
+				TArray<FSoftObjectPath> Result;
+				FAssetRegistryModule& AssetRegistryModule = FModuleManager::LoadModuleChecked<FAssetRegistryModule>(TEXT("AssetRegistry"));
+				IAssetRegistry& AssetRegistry = AssetRegistryModule.GetRegistry();
+
+				UClass* ReferencedClass = nullptr;
+				if (UControlRigBlueprint* ReferencedBlueprint = Cast<UControlRigBlueprint>(ReferencedAsset.GetAsset()))
+				{
+					ReferencedClass = ReferencedBlueprint->GetRigVMBlueprintGeneratedClass();
+				}
+				
+				TArray<FName> PackageDependencies;
+				AssetRegistry.GetReferencers(ReferencedAsset.PackageName, PackageDependencies);
+
+				for (FName& DependencyPath : PackageDependencies)
+				{
+					TArray<FAssetData> Assets;
+					AssetRegistry.GetAssetsByPackageName(DependencyPath, Assets);
+
+					for (const FAssetData& DependencyData : Assets)
+					{
+						if (DependencyData.IsAssetLoaded())
+						{
+							if (UControlRigBlueprint* Blueprint = Cast<UControlRigBlueprint>(DependencyData.GetAsset()))
+							{
+								if (Blueprint->IsModularRig())
+								{
+									TArray<const FRigModuleReference*> Modules = Blueprint->ModularRigModel.FindModuleInstancesOfClass(ReferencedClass);
+									for (const FRigModuleReference* Module : Modules)
+									{
+										FSoftObjectPath ModulePath = DependencyData.GetSoftObjectPath();
+										ModulePath.SetSubPathString(Module->GetPath());
+										Result.Add(ModulePath);
+									}
+								}
+							}
+						}
+						else
+						{
+							// Check only modular rigs
+							{
+								static const FLazyName ControlRigTypeName(GET_MEMBER_NAME_CHECKED(UControlRigBlueprint, ControlRigType));
+								FProperty* ControlRigTypeProperty = CastField<FProperty>(UControlRigBlueprint::StaticClass()->FindPropertyByName(ControlRigTypeName));
+								const FString ControlRigTypeString = DependencyData.GetTagValueRef<FString>(ControlRigTypeName);
+								if (ControlRigTypeString.IsEmpty())
+								{
+									continue;
+								}
+
+								EControlRigType RigType;
+								ControlRigTypeProperty->ImportText_Direct(*ControlRigTypeString, &RigType, nullptr, EPropertyPortFlags::PPF_None);
+								if (RigType != EControlRigType::ModularRig)
+								{
+									continue;
+								}
+							}
+
+							static const FLazyName ModuleReferenceDataName(GET_MEMBER_NAME_CHECKED(UControlRigBlueprint, ModuleReferenceData));
+							FArrayProperty* ModuleReferenceDataProperty = CastField<FArrayProperty>(UControlRigBlueprint::StaticClass()->FindPropertyByName(ModuleReferenceDataName));
+							const FString ModularRigDataString = DependencyData.GetTagValueRef<FString>(ModuleReferenceDataName);
+							if (ModularRigDataString.IsEmpty())
+							{
+								continue;
+							}
+
+							TArray<FModuleReferenceData> Modules;
+							ModuleReferenceDataProperty->ImportText_Direct(*ModularRigDataString, &Modules, nullptr, EPropertyPortFlags::PPF_None);
+
+							for (const FModuleReferenceData& Module : Modules)
+							{
+								if (Module.ReferencedModule == ReferencedClass)
+								{
+									FSoftObjectPath ModulePath = DependencyData.GetSoftObjectPath();
+									ModulePath.SetSubPathString(Module.ModulePath);
+									Result.Add(ModulePath);
+								}
+							}
+						}
+					}
+				}
+				
+				return Result;
+			})
+			.OnSwapReference_Lambda([](const FSoftObjectPath& ModulePath, const FAssetData& NewModuleAsset) -> bool
+			{
+				TSubclassOf<UControlRig> NewModuleClass = nullptr;
+				if (const UControlRigBlueprint* ModuleBlueprint = Cast<UControlRigBlueprint>(NewModuleAsset.GetAsset()))
+				{
+					NewModuleClass = ModuleBlueprint->GetRigVMBlueprintGeneratedClass();
+				}
+				if (NewModuleClass)
+				{
+					if (UControlRigBlueprint* RigBlueprint = Cast<UControlRigBlueprint>(ModulePath.GetWithoutSubPath().ResolveObject()))
+					{
+						return RigBlueprint->GetModularRigController()->SwapModuleClass(ModulePath.GetSubPathString(), NewModuleClass);
+					}
+				}
+				return false;
+			})
+			.ExtraAssetFilter_Lambda([](const FAssetData& AssetData) -> bool
+			{
+				static const FLazyName ControlRigTypeName(GET_MEMBER_NAME_CHECKED(UControlRigBlueprint, ControlRigType));
+				FProperty* ControlRigTypeProperty = CastField<FProperty>(UControlRigBlueprint::StaticClass()->FindPropertyByName(ControlRigTypeName));
+				const FString ControlRigTypeString = AssetData.GetTagValueRef<FString>(ControlRigTypeName);
+				if (ControlRigTypeString.IsEmpty())
+				{
+					return false;
+				}
+
+				EControlRigType RigType;
+				ControlRigTypeProperty->ImportText_Direct(*ControlRigTypeString, &RigType, nullptr, EPropertyPortFlags::PPF_None);
+				return RigType == EControlRigType::RigModule;
+			});
+
+		const TSharedRef<SRigVMBulkEditDialog<SRigVMSwapAssetReferencesWidget>> SwapModulesDialog =
+		 	SNew(SRigVMBulkEditDialog<SRigVMSwapAssetReferencesWidget>)
+		 	.WindowSize(FVector2D(800.0f, 640.0f))
+		 	.WidgetArgs(WidgetArgs);
+		
+		SwapModulesDialog->ShowNormal();
+	})
+);
+
 const FName FControlRigEditorModes::ControlRigEditorMode = TEXT("Rigging");
 const TArray<FName> FControlRigEditor::ForwardsSolveEventQueue = {FRigUnit_BeginExecution::EventName};
 const TArray<FName> FControlRigEditor::BackwardsSolveEventQueue = {FRigUnit_InverseExecution::EventName};
diff --git a/Engine/Plugins/Animation/ControlRig/Source/ControlRigEditor/Private/Editor/SModularRigModel.cpp b/Engine/Plugins/Animation/ControlRig/Source/ControlRigEditor/Private/Editor/SModularRigModel.cpp
--- a/Engine/Plugins/Animation/ControlRig/Source/ControlRigEditor/Private/Editor/SModularRigModel.cpp
+++ b/Engine/Plugins/Animation/ControlRig/Source/ControlRigEditor/Private/Editor/SModularRigModel.cpp
@@ -54,6 +54,8 @@
 #include "Kismet2/SClassPickerDialog.h"
 #include "RigVMFunctions/Math/RigVMMathLibrary.h"
 #include "Preferences/PersonaOptions.h"
+#include "Widgets/SRigVMBulkEditDialog.h"
+#include "Widgets/SRigVMSwapAssetReferencesWidget.h"
 
 #define LOCTEXT_NAMESPACE "SModularRigModel"
 
@@ -229,6 +231,10 @@ void SModularRigModel::BindCommands()
 	CommandList->MapAction(Commands.ReresolveModuleItem,
 		FExecuteAction::CreateSP(this, &SModularRigModel::HandleReresolveModules),
 		FCanExecuteAction());
+
+	CommandList->MapAction(Commands.SwapModuleClassItem,
+		FExecuteAction::CreateSP(this, &SModularRigModel::HandleSwapClassForModules),
+		FCanExecuteAction::CreateSP(this, &SModularRigModel::CanSwapModules));
 }
 
 FReply SModularRigModel::OnKeyDown(const FGeometry& MyGeometry, const FKeyEvent& InKeyEvent)
@@ -395,6 +401,7 @@ void SModularRigModel::CreateContextMenu()
 					ModulesSection.AddMenuEntry(Commands.DeleteModuleItem);
 					ModulesSection.AddMenuEntry(Commands.MirrorModuleItem);
 					ModulesSection.AddMenuEntry(Commands.ReresolveModuleItem);
+					ModulesSection.AddMenuEntry(Commands.SwapModuleClassItem);
 				}
 			})
 		);
@@ -860,6 +867,145 @@ void SModularRigModel::HandleReresolveModules(const TArray<FString>& InPaths)
 	}
 }
 
+bool SModularRigModel::CanSwapModules() const
+{
+	// Only if all modules selected have the same module class
+	if(!ControlRigEditor.IsValid())
+	{
+		return false;
+	}
+
+	UModularRig* Rig = GetDefaultModularRig();
+	if (Rig)
+	{
+		TSoftClassPtr<UControlRig> CommonClass = nullptr;
+		TArray<TSharedPtr<FModularRigTreeElement>> SelectedItems = TreeView->GetSelectedItems();
+		for (TSharedPtr<FModularRigTreeElement>& SelectedItem : SelectedItems)
+		{
+			TSoftClassPtr<UControlRig> ModuleClass;
+			if (const FRigModuleReference* Module = ControlRigBlueprint->ModularRigModel.FindModule(SelectedItem->ModulePath))
+			{
+				if (Module->Class.IsValid())
+				{
+					ModuleClass = Module->Class;
+				}
+			}
+			if(!ModuleClass)
+			{
+				return false;
+			}
+			if(!CommonClass)
+			{
+				CommonClass = ModuleClass;
+			}
+			if(ModuleClass != CommonClass)
+			{
+				return false;
+			}
+		}
+		return true;
+	}
+	return false;
+}
+
+void SModularRigModel::HandleSwapClassForModules()
+{
+	if(!ControlRigEditor.IsValid())
+	{
+		return;
+	}
+
+	UModularRig* Rig = GetDefaultModularRig();
+	if (Rig)
+	{
+		TArray<TSharedPtr<FModularRigTreeElement>> SelectedItems = TreeView->GetSelectedItems();
+		TArray<FString> SelectedPaths;
+		Algo::Transform(SelectedItems, SelectedPaths, [](const TSharedPtr<FModularRigTreeElement>& Element)
+		{
+			if (Element.IsValid())
+			{
+				return Element->ModulePath;
+			}
+			return FString();
+		});
+		HandleSwapClassForModules(SelectedPaths);
+	}
+}
+
+void SModularRigModel::HandleSwapClassForModules(const TArray<FString>& InPaths)
+{
+	TArray<FSoftObjectPath> ModulePaths;
+	Algo::Transform(InPaths, ModulePaths, [this](const FString& Path)
+	{
+		FSoftObjectPath ModulePath(ControlRigBlueprint->GetPathName());
+		ModulePath.SetSubPathString(Path);
+		return ModulePath;
+	});
+
+	TSoftClassPtr<UControlRig> SourceClass = nullptr;
+	if (FRigModuleReference* Module = ControlRigBlueprint->ModularRigModel.FindModule(InPaths[0]))
+	{
+		SourceClass = Module->Class;
+	}
+
+	if (!SourceClass)
+	{
+		return;
+	}
+
+	TArray<FAssetData> SourceAssets;
+	IAssetRegistry::Get()->GetAssetsByPackageName(*SourceClass->GetPackage()->GetPathName(),SourceAssets);
+	if (SourceAssets.IsEmpty())
+	{
+		return;
+	}
+	
+	SRigVMSwapAssetReferencesWidget::FArguments WidgetArgs;
+	WidgetArgs
+		.EnableUndo(true)
+		.CloseOnSuccess(true)
+		.Source(SourceAssets[0])
+		.ReferencePaths(ModulePaths)
+		.SkipPickingRefs(true)
+		.OnSwapReference_Lambda([](const FSoftObjectPath& ModulePath, const FAssetData& NewModuleAsset) -> bool
+		{
+			TSubclassOf<UControlRig> NewModuleClass = nullptr;
+			if (const UControlRigBlueprint* ModuleBlueprint = Cast<UControlRigBlueprint>(NewModuleAsset.GetAsset()))
+			{
+				NewModuleClass = ModuleBlueprint->GetRigVMBlueprintGeneratedClass();
+			}
+			if (NewModuleClass)
+			{
+				if (UControlRigBlueprint* RigBlueprint = Cast<UControlRigBlueprint>(ModulePath.GetWithoutSubPath().ResolveObject()))
+				{
+					return RigBlueprint->GetModularRigController()->SwapModuleClass(ModulePath.GetSubPathString(), NewModuleClass);
+				}
+			}
+			return false;
+		})
+		.ExtraAssetFilter_Lambda([](const FAssetData& AssetData) -> bool
+		{
+			static const FLazyName ControlRigTypeName(GET_MEMBER_NAME_CHECKED(UControlRigBlueprint, ControlRigType));
+			FProperty* ControlRigTypeProperty = CastField<FProperty>(UControlRigBlueprint::StaticClass()->FindPropertyByName(ControlRigTypeName));
+			const FString ControlRigTypeString = AssetData.GetTagValueRef<FString>(ControlRigTypeName);
+			if (ControlRigTypeString.IsEmpty())
+			{
+				return false;
+			}
+
+			EControlRigType RigType;
+			ControlRigTypeProperty->ImportText_Direct(*ControlRigTypeString, &RigType, nullptr, EPropertyPortFlags::PPF_None);
+			return RigType == EControlRigType::RigModule;
+		});
+
+	const TSharedRef<SRigVMBulkEditDialog<SRigVMSwapAssetReferencesWidget>> SwapModulesDialog =
+		SNew(SRigVMBulkEditDialog<SRigVMSwapAssetReferencesWidget>)
+		.WindowSize(FVector2D(800.0f, 640.0f))
+		.WidgetArgs(WidgetArgs);
+	
+	SwapModulesDialog->ShowNormal();
+}
+
 void SModularRigModel::HandleConnectorResolved(const FRigElementKey& InConnector, const FRigElementKey& InTarget)
 {
 	if (ControlRigBlueprint.IsValid())
@@ -988,7 +1134,13 @@ UModularRig* SModularRigModel::GetDefaultModularRig() const
 {
 	if (ControlRigBlueprint.IsValid())
 	{
-		if (UControlRig* DebuggedRig = ControlRigBeingDebuggedPtr.Get())
+		UControlRig* DebuggedRig = ControlRigBeingDebuggedPtr.Get();
+		if (!DebuggedRig)
+		{
+			DebuggedRig = ControlRigBlueprint->GetDebuggedControlRig();
+		}
+		
+		if (DebuggedRig)
 		{
 			return Cast<UModularRig>(DebuggedRig);
 		}
diff --git a/Engine/Plugins/Animation/ControlRig/Source/ControlRigEditor/Private/Editor/SModularRigModel.h b/Engine/Plugins/Animation/ControlRig/Source/ControlRigEditor/Private/Editor/SModularRigModel.h
--- a/Engine/Plugins/Animation/ControlRig/Source/ControlRigEditor/Private/Editor/SModularRigModel.h
+++ b/Engine/Plugins/Animation/ControlRig/Source/ControlRigEditor/Private/Editor/SModularRigModel.h
@@ -82,6 +82,11 @@ private:
 	void HandleReresolveModules();
 	void HandleReresolveModules(const TArray<FString>& InPaths);
 
+	/** Swap module class for items */
+	bool CanSwapModules() const;
+	void HandleSwapClassForModules();
+	void HandleSwapClassForModules(const TArray<FString>& InPaths);
+
 	/** Resolve connector */
 	void HandleConnectorResolved(const FRigElementKey& InConnector, const FRigElementKey& InTarget);
 
diff --git a/Engine/Plugins/Runtime/RigVM/Source/RigVM/Public/RigVMCore/RigVMGraphFunctionDefinition.h b/Engine/Plugins/Runtime/RigVM/Source/RigVM/Public/RigVMCore/RigVMGraphFunctionDefinition.h
--- a/Engine/Plugins/Runtime/RigVM/Source/RigVM/Public/RigVMCore/RigVMGraphFunctionDefinition.h
+++ b/Engine/Plugins/Runtime/RigVM/Source/RigVM/Public/RigVMCore/RigVMGraphFunctionDefinition.h
@@ -558,7 +558,7 @@ struct RIGVM_API FRigVMGraphFunctionHeader
 	FText GetTooltip() const
 	{
 		FString TooltipStr = FString::Printf(TEXT("%s (%s)\n%s"),
-		*NodeTitle,
+		*Name.ToString(),
 		*LibraryPointer.GetNodeSoftPath().GetAssetPathString(),
 		*Description);
 		return FText::FromString(TooltipStr);
diff --git a/Engine/Plugins/Runtime/RigVM/Source/RigVMDeveloper/Public/RigVMBlueprint.h b/Engine/Plugins/Runtime/RigVM/Source/RigVMDeveloper/Public/RigVMBlueprint.h
--- a/Engine/Plugins/Runtime/RigVM/Source/RigVMDeveloper/Public/RigVMBlueprint.h
+++ b/Engine/Plugins/Runtime/RigVM/Source/RigVMDeveloper/Public/RigVMBlueprint.h
@@ -827,6 +827,7 @@ private:
 	friend struct FRigVMEdGraphSchemaAction_PromoteToVariable;
 	friend class URigVMBuildData;
 	friend class FRigVMVariantDetailCustomization;
+	friend class FRigVMTreeAssetVariantFilter;
 };
 
 class RIGVMDEVELOPER_API FRigVMBlueprintCompileScope
diff --git a/Engine/Plugins/Runtime/RigVM/Source/RigVMEditor/Private/Editor/RigVMEditor.cpp b/Engine/Plugins/Runtime/RigVM/Source/RigVMEditor/Private/Editor/RigVMEditor.cpp
--- a/Engine/Plugins/Runtime/RigVM/Source/RigVMEditor/Private/Editor/RigVMEditor.cpp
+++ b/Engine/Plugins/Runtime/RigVM/Source/RigVMEditor/Private/Editor/RigVMEditor.cpp
@@ -46,6 +46,7 @@
 #include "Editor/RigVMEditorMode.h"
 #include "InstancedPropertyBagStructureDataProvider.h"
 #include "Widgets/SRigVMSwapFunctionsWidget.h"
+#include "Widgets/SRigVMSwapAssetReferencesWidget.h"
 #include "Widgets/SRigVMBulkEditDialog.h"
 #include "AssetRegistry/AssetRegistryModule.h"
 
@@ -1404,6 +1405,11 @@ void FRigVMEditor::BindCommands()
 		FRigVMEditorCommands::Get().SwapFunctionAcrossProject,
 		FExecuteAction::CreateSP(this, &FRigVMEditor::SwapFunctionAcrossProject),
 		FCanExecuteAction());
+
+	GetToolkitCommands()->MapAction(
+		FRigVMEditorCommands::Get().SwapAssetReferences,
+		FExecuteAction::CreateSP(this, &FRigVMEditor::SwapAssetReferences),
+		FCanExecuteAction());
 }
 
 void FRigVMEditor::ToggleAutoCompileGraph()
@@ -1467,6 +1473,9 @@ TSharedRef<SWidget> FRigVMEditor::GenerateBulkEditMenuContent()
 	MenuBuilder.AddMenuEntry(FRigVMEditorCommands::Get().SwapFunctionWithinAsset, TEXT("SwapFunctionWithinAsset"), TAttribute<FText>(), TAttribute<FText>(), FSlateIcon());
 	MenuBuilder.AddMenuEntry(FRigVMEditorCommands::Get().SwapFunctionAcrossProject, TEXT("SwapFunctionAcrossProject"), TAttribute<FText>(), TAttribute<FText>(), FSlateIcon());
 	MenuBuilder.EndSection();
+	MenuBuilder.BeginSection(TEXT("Asset"), LOCTEXT("Asset", "Asset"));
+	MenuBuilder.AddMenuEntry(FRigVMEditorCommands::Get().SwapAssetReferences, TEXT("SwapAssetReferences"), TAttribute<FText>(), TAttribute<FText>(), FSlateIcon());
+	MenuBuilder.EndSection();
 	return MenuBuilder.MakeWidget();
 }
 
@@ -3468,6 +3477,29 @@ void FRigVMEditor::SwapFunctionForAssets(const TArray<FAssetData>& InAssets, boo
 	SwapFunctionsDialog->ShowNormal();
 }
 
+void FRigVMEditor::SwapAssetReferences()
+{
+	TArray<FAssetData> SourceAssets;
+	IAssetRegistry::Get()->GetAssetsByPackageName(*GetRigVMBlueprint()->GetPackage()->GetPathName(),SourceAssets);
+	if (SourceAssets.IsEmpty())
+	{
+		return;
+	}
+	
+	SRigVMSwapAssetReferencesWidget::FArguments WidgetArgs;
+	WidgetArgs
+		.Source(SourceAssets[0])
+		.EnableUndo(false)
+		.CloseOnSuccess(true);
+
+	const TSharedRef<SRigVMBulkEditDialog<SRigVMSwapAssetReferencesWidget>> SwapFunctionsDialog =
+		SNew(SRigVMBulkEditDialog<SRigVMSwapAssetReferencesWidget>)
+		.WindowSize(FVector2D(800.0f, 640.0f))
+		.WidgetArgs(WidgetArgs);
+
+	SwapFunctionsDialog->ShowNormal();
+}
+
 void FRigVMEditor::UpdateGraphCompilerErrors()
 {
 	DECLARE_SCOPE_HIERARCHICAL_COUNTER_FUNC()
diff --git a/Engine/Plugins/Runtime/RigVM/Source/RigVMEditor/Private/Editor/RigVMEditorCommands.cpp b/Engine/Plugins/Runtime/RigVM/Source/RigVMEditor/Private/Editor/RigVMEditorCommands.cpp
--- a/Engine/Plugins/Runtime/RigVM/Source/RigVMEditor/Private/Editor/RigVMEditorCommands.cpp
+++ b/Engine/Plugins/Runtime/RigVM/Source/RigVMEditor/Private/Editor/RigVMEditorCommands.cpp
@@ -21,6 +21,7 @@ void FRigVMEditorCommands::RegisterCommands()
 	UI_COMMAND(FrameSelection, "Frame Selection", "Frames the selected nodes in the Graph View.", EUserInterfaceActionType::Button, FInputChord(EKeys::F));
 	UI_COMMAND(SwapFunctionWithinAsset, "Swap Function (Asset)", "Swaps a function for all occurrences within this asset.", EUserInterfaceActionType::Button, FInputChord());
 	UI_COMMAND(SwapFunctionAcrossProject, "Swap Function (Project)", "Swaps a function for all occurrences in the project.", EUserInterfaceActionType::Button, FInputChord());
+	UI_COMMAND(SwapAssetReferences, "Swap Asset References", "Swaps an asset reference for another asset.", EUserInterfaceActionType::Button, FInputChord());
 }
 
 #undef LOCTEXT_NAMESPACE
diff --git a/Engine/Plugins/Runtime/RigVM/Source/RigVMEditor/Private/Editor/RigVMEditorCommands.h b/Engine/Plugins/Runtime/RigVM/Source/RigVMEditor/Private/Editor/RigVMEditorCommands.h
--- a/Engine/Plugins/Runtime/RigVM/Source/RigVMEditor/Private/Editor/RigVMEditorCommands.h
+++ b/Engine/Plugins/Runtime/RigVM/Source/RigVMEditor/Private/Editor/RigVMEditorCommands.h
@@ -63,6 +63,9 @@ public:
 	/** Swap Function (Project) */
 	TSharedPtr< FUICommandInfo > SwapFunctionAcrossProject;
 
+	/** Swap Asset References */
+	TSharedPtr< FUICommandInfo > SwapAssetReferences;
+
 	/**
 	 * Initialize commands
 	 */
diff --git a/Engine/Plugins/Runtime/RigVM/Source/RigVMEditor/Private/Widgets/SRigVMSwapAssetReferencesWidget.cpp b/Engine/Plugins/Runtime/RigVM/Source/RigVMEditor/Private/Widgets/SRigVMSwapAssetReferencesWidget.cpp
new file mode 100644
--- /dev/null
+++ b/Engine/Plugins/Runtime/RigVM/Source/RigVMEditor/Private/Widgets/SRigVMSwapAssetReferencesWidget.cpp
@@ -0,0 +1,454 @@
+// Copyright Epic Games, Inc. All Rights Reserved.
+
+#include "Widgets/SRigVMSwapAssetReferencesWidget.h"
+#include "Widgets/SRigVMBulkEditDialog.h"
+
+#include "RigVMBlueprint.h"
+#include "Dialogs/Dialogs.h"
+#include "AssetRegistry/AssetRegistryModule.h"
+#include "Styling/AppStyle.h"
+#include "Algo/Count.h"
+#include "BlueprintEditor.h"
+#include "RigVMEditorBlueprintLibrary.h"
+#include "Editor/RigVMEditor.h"
+#include "Misc/UObjectToken.h"
+#include "RigVMModel/RigVMControllerActions.h"
+
+#define LOCTEXT_NAMESPACE "SRigVMSwapAssetReferencesWidget"
+
+uint32 FRigVMSwapAssetReferencesContext::GetVisibleChildrenHash() const
+{
+	return HashCombine(FRigVMTreeContext::GetVisibleChildrenHash(), GetTypeHash(SourceAsset));
+}
+
+TArray<TSharedRef<FRigVMTreeNode>> FRigVMTreeAssetRefAssetNode::GetChildrenImpl(const TSharedRef<FRigVMTreeContext>& InContext) const
+{
+	if (AssetRefNodes.IsEmpty())
+	{
+		if(const TSharedPtr<FRigVMSwapAssetReferencesContext> SwapContext = Cast<FRigVMSwapAssetReferencesContext>(InContext))
+		{
+			for (const FSoftObjectPath& Reference : SwapContext->GetReferences())
+			{
+				TSharedRef<FRigVMTreeReferenceNode> RefNode = FRigVMTreeReferenceNode::Create(Reference);
+				AssetRefNodes.Add(RefNode);
+			}
+		}
+	}
+	return AssetRefNodes;
+}
+
+bool FRigVMTreeTargetAssetFilter::Filters(TSharedRef<FRigVMTreeNode>& InNode, const TSharedRef<FRigVMTreeContext>& InContext)
+{
+	if(const TSharedPtr<FRigVMTreePackageNode> AssetNode = Cast<FRigVMTreePackageNode>(InNode))
+	{
+		if(const TSharedPtr<FRigVMSwapAssetReferencesContext> AssetSwapContext = Cast<FRigVMSwapAssetReferencesContext>(InContext))
+		{
+			return AssetNode->GetAssetData() == AssetSwapContext->GetSourceAsset();
+		}
+	}
+	return FRigVMTreeFilter::Filters(InNode, InContext);
+}
+
+FText FRigVMTreeAssetVariantFilter::GetLabel() const
+{
+	return LOCTEXT("OnlyShowVariants", "Only Show Variants");
+}
+
+bool FRigVMTreeAssetVariantFilter::Filters(TSharedRef<FRigVMTreeNode>& InNode, const TSharedRef<FRigVMTreeContext>& InContext)
+{
+	if(const TSharedPtr<FRigVMTreeCategoryNode> TargetAssetNode = Cast<FRigVMTreeCategoryNode>(InNode))
+	{
+		if(const TSharedPtr<FRigVMSwapAssetReferencesContext> SwapContext = Cast<FRigVMSwapAssetReferencesContext>(InContext))
+		{
+			if (!SourceVariants.Contains(SwapContext->GetSourceAsset().GetFullName()))
+			{
+				static const FName AssetVariantPropertyName = GET_MEMBER_NAME_CHECKED(URigVMBlueprint, AssetVariant);
+				const FProperty* AssetVariantProperty = CastField<FProperty>(URigVMBlueprint::StaticClass()->FindPropertyByName(AssetVariantPropertyName));
+				const FString VariantStr = SwapContext->GetSourceAsset().GetTagValueRef<FString>(AssetVariantPropertyName);
+				FRigVMVariant AssetVariant;
+				AssetVariantProperty->ImportText_Direct(*VariantStr, &AssetVariant, nullptr, EPropertyPortFlags::PPF_None);
+
+				if (!AssetVariant.Guid.IsValid())
+				{
+					AssetVariant.Guid = FRigVMVariant::GenerateGUID(SwapContext->GetSourceAsset().GetPackage()->GetPathName());
+				}
+
+				SourceVariants.FindOrAdd(SwapContext->GetSourceAsset().GetFullName()) = URigVMBuildData::Get()->FindAssetVariantRefs(AssetVariant.Guid);
+			}
+
+			TArray<FRigVMVariantRef>& VariantRefs = SourceVariants.FindChecked(SwapContext->GetSourceAsset().GetFullName());
+			if(!VariantRefs.IsEmpty())
+			{
+				return !VariantRefs.ContainsByPredicate([InNode](const FRigVMVariantRef& VariantRef)
+				{
+					return VariantRef.ObjectPath.GetWithoutSubPath() == InNode->GetAssetData().GetSoftObjectPath();
+				});
+			}
+		}
+	}
+	return true;
+}
+
+bool FRigVMSwapAssetReferenceTask::Execute(const TSharedRef<FRigVMTreePhase>& InPhase)
+{
+	if (SwapFunction.IsBound())
+	{
+		const TSharedRef<FTokenizedMessage> Message = FTokenizedMessage::Create(EMessageSeverity::Info);
+		Message->AddText(LOCTEXT("Swap", "Swapping"));
+		Message->AddToken(FAssetNameToken::Create(ReferencePath.GetWithoutSubPath().ToString()));
+		Message->AddText(FText::FromString(ReferencePath.GetSubPathString()));
+		InPhase->GetContext()->LogMessage(Message);
+
+		return SwapFunction.Execute(ReferencePath, NewAsset);
+	}
+	return false;
+}
+
+void SRigVMSwapAssetReferencesWidget::Construct(const FArguments& InArgs)
+{
+	static const TArray<TSharedRef<FRigVMTreeFilter>> DefaultFilters = {
+		FRigVMTreeEngineContentFilter::Create(),
+		FRigVMTreeDeveloperContentFilter::Create()
+	};
+	static const TSharedRef<FRigVMTreeFilter> DefaultPathFilter = FRigVMTreePathFilter::Create();
+
+	OnGetReferences = InArgs._OnGetReferences;
+	OnSwapReference = InArgs._OnSwapReference;
+	ExtraAssetFilter = InArgs._ExtraAssetFilter;
+	PickTargetContext = FRigVMSwapAssetReferencesContext::Create();
+	PickAssetRefsContext = FRigVMSwapAssetReferencesContext::Create();
+	bSkipPickingRefs = InArgs._SkipPickingRefs;
+
+	TOptional<int32> PhaseToActivate;
+
+	SetSourceAsset(InArgs._Source);
+	SetTargetAsset(InArgs._Target);
+	PickAssetRefsContext->SetReferences(InArgs._ReferencePaths);
+
+	TArray<TSharedRef<FRigVMTreePhase>> Phases;
+	if(!InArgs._Source.IsValid())
+	{
+		TSharedRef<FRigVMTreePhase> Phase = FRigVMTreePhase::Create(PHASE_PICKSOURCE, TEXT("Pick Source Asset"), FRigVMTreeContext::Create());
+		Phase->GetContext()->Filters = DefaultFilters;
+		Phase->GetContext()->Filters.Add(DefaultPathFilter);
+		Phase->SetNodes(GetAssetNodes(InArgs));
+		Phase->PrimaryButtonText().Set(LOCTEXT("Next", "Next"));
+		Phase->IsPrimaryButtonVisible().BindLambda([this]()
+		{
+			const TArray<TSharedRef<FRigVMTreeNode>> Selection = GetBulkEditWidget()->GetSelectedNodes();
+			return Selection.ContainsByPredicate([](const TSharedRef<FRigVMTreeNode>& Node)
+			{
+				return Node->IsA<FRigVMTreePackageNode>();
+			});
+		});
+		Phase->OnPrimaryAction().BindLambda([this]()
+		{
+			const TArray<TSharedRef<FRigVMTreeNode>> Selection = GetBulkEditWidget()->GetSelectedNodes();
+			check(Selection.Num() == 1);
+			if(Cast<FRigVMTreePackageNode>(Selection[0]))
+			{
+				GetBulkEditWidget()->ActivatePhase(PHASE_PICKTARGET);
+				return FReply::Handled();
+			}
+			return FReply::Unhandled();
+		});
+		PhaseToActivate = PhaseToActivate.Get(Phase->GetID());
+		Phases.Add(Phase);
+	}
+	
+	if(!InArgs._Target.IsValid())
+	{
+		TSharedRef<FRigVMTreePhase> Phase = FRigVMTreePhase::Create(PHASE_PICKTARGET, TEXT("Pick Target Asset"), PickTargetContext.ToSharedRef());
+		Phase->GetContext()->Filters = DefaultFilters;
+		Phase->GetContext()->Filters.Add(DefaultPathFilter);
+		Phase->GetContext()->Filters.Add(FRigVMTreeTargetAssetFilter::Create());
+		Phase->GetContext()->Filters.Add(FRigVMTreeAssetVariantFilter::Create());
+		Phase->SetNodes(GetAssetNodes(InArgs));
+		Phase->PrimaryButtonText().Set(LOCTEXT("Next", "Next"));
+		Phase->IsPrimaryButtonVisible().BindLambda([this]()
+		{
+			const TArray<TSharedRef<FRigVMTreeNode>> Selection = GetBulkEditWidget()->GetSelectedNodes();
+			return Selection.ContainsByPredicate([](const TSharedRef<FRigVMTreeNode>& Node)
+			{
+				return Node->IsA<FRigVMTreePackageNode>();
+			});
+		});
+		Phase->OnPrimaryAction().BindLambda([this]()
+		{
+			const TArray<TSharedRef<FRigVMTreeNode>> Selection = GetBulkEditWidget()->GetSelectedNodes();
+			check(Selection.Num() == 1);
+			if(Cast<FRigVMTreePackageNode>(Selection[0]))
+			{
+				GetBulkEditWidget()->ActivatePhase(PHASE_PICKASSETREFS);
+				return FReply::Handled();
+			}
+			return FReply::Unhandled();
+		});
+		PhaseToActivate = PhaseToActivate.Get(Phase->GetID());
+		Phases.Add(Phase);
+	}
+	
+	{
+		TSharedRef<FRigVMTreePhase> Phase = FRigVMTreePhase::Create(PHASE_PICKASSETREFS, TEXT("Pick Asset References"), PickAssetRefsContext.ToSharedRef());
+		Phase->SetAllowsMultiSelection(true);
+		Phase->GetContext()->Filters = DefaultFilters;
+		Phase->GetContext()->Filters.Add(FRigVMTreePathFilter::Create());
+		
+		Phase->PrimaryButtonText().Set(LOCTEXT("SwapAssetRefs", "Swap Asset References"));
+		Phase->IsPrimaryButtonVisible().BindLambda([this]()
+		{
+			return true;
+		});
+		Phase->PrimaryButtonText().BindLambda([this]()
+		{
+			if(GetBulkEditWidget()->HasAnyVisibleCheckedNode())
+			{
+				return LOCTEXT("SwapAssetRefs", "Swap Asset References");
+			}
+			return LOCTEXT("Done", "Done");
+		});
+		Phase->OnPrimaryAction().BindLambda([this, Phase]()
+		{
+			GetBulkEditWidget()->GetTreeView()->GetTreeView()->ClearSelection();
+			
+			const TArray<TSharedRef<FRigVMTreeNode>> AllCheckedNodes = GetBulkEditWidget()->GetCheckedNodes();
+			if(AllCheckedNodes.IsEmpty())
+			{
+				GetBulkEditWidget()->CloseDialog();
+				return FReply::Handled();
+			}
+			
+			const TArray<TSharedRef<FRigVMTreeNode>> ModuleRefs = AllCheckedNodes.FilterByPredicate([](const TSharedRef<FRigVMTreeNode>& Node) {
+				return Node->IsA<FRigVMTreeReferenceNode>();
+			});
+			
+			if(ModuleRefs.IsEmpty())
+			{
+				return FReply::Handled();
+			}
+			
+			const TSharedPtr<FRigVMSwapAssetReferencesContext> Context = Cast<FRigVMSwapAssetReferencesContext>(Phase->GetContext());
+			check(Context);
+			
+			TSet<FSoftObjectPath> VisitedPackages;
+			TArray<TSharedRef<FRigVMTreeTask>> Tasks;
+			
+			for(const TSharedRef<FRigVMTreeNode>& Node : ModuleRefs)
+			{
+				const TSharedPtr<FRigVMTreeReferenceNode> ModuleRefNode = Cast<FRigVMTreeReferenceNode>(Node);
+				check(ModuleRefNode);
+				
+				TSharedPtr<FRigVMTreeAssetRefAssetNode> AssetNode = Cast<FRigVMTreeAssetRefAssetNode>(ModuleRefNode->GetRoot());
+			
+				if(!VisitedPackages.Contains(AssetNode->GetPackagePath()))
+				{
+					Tasks.Add(FRigVMTreeLoadPackageForNodeTask::Create(AssetNode.ToSharedRef()));
+					VisitedPackages.Add(AssetNode->GetPackagePath());
+				}
+				Tasks.Add(FRigVMSwapAssetReferenceTask::Create(ModuleRefNode->GetReferencePath(), Context->GetTargetAsset(), OnSwapReference));
+			}
+	
+			if(Tasks.IsEmpty())
+			{
+				return FReply::Unhandled();
+			}
+			GetBulkEditWidget()->QueueTasks(Tasks);
+			return FReply::Handled();
+		});
+		PhaseToActivate = PhaseToActivate.Get(Phase->GetID());
+		Phases.Add(Phase);
+	}
+
+	ChildSlot
+	[
+		SNew(SVerticalBox)
+
+		+ SVerticalBox::Slot()
+		.FillHeight(1.0f)
+		.Padding(0)
+		[
+			SNew(SHorizontalBox)
+			+ SHorizontalBox::Slot()
+			.FillWidth(1.0f)
+			.Padding(0)
+			[
+				SAssignNew(BulkEditWidget, SRigVMBulkEditWidget)
+				.Phases(Phases)
+				.PhaseToActivate(PhaseToActivate.GetValue())
+				.OnPhaseActivated(this, &SRigVMSwapAssetReferencesWidget::OnPhaseActivated)
+				.OnNodeSelected(this, &SRigVMSwapAssetReferencesWidget::OnNodeSelected)
+				.OnNodeDoubleClicked(this, &SRigVMSwapAssetReferencesWidget::OnNodeDoubleClicked)
+				//.RightWidget(FunctionPreviewBox)
+				.BulkEditTitle(LOCTEXT("SwapAssetReference", "Swap Asset Reference"))
+				.BulkEditConfirmMessage(LOCTEXT("SwapAssetReferencesConfirmMessage", "This edit is going to swap asset references without support for undo. Are you sure?"))
+				.BulkEditConfirmIniField(TEXT("RigVMSwapAssetReferences_Warning"))
+				.EnableUndo(InArgs._EnableUndo)
+				.CloseOnSuccess(InArgs._CloseOnSuccess)
+			]
+		]
+	];
+}
+
+TArray<TSharedRef<FRigVMTreeNode>> SRigVMSwapAssetReferencesWidget::GetAssetNodes(const FArguments& InArgs)
+{
+	TArray<TSharedRef<FRigVMTreeNode>> Nodes;
+
+	const TArray<FAssetData> Assets = URigVMEditorBlueprintLibrary::GetAssetsWithFilter(URigVMBlueprint::StaticClass(),
+		FRigVMAssetDataFilter::CreateLambda([this] (const FAssetData& AssetData)
+		{
+			if (ExtraAssetFilter.IsBound())
+			{
+				return ExtraAssetFilter.Execute(AssetData);
+			}
+			return true;
+		}));
+	for(const FAssetData& Asset : Assets)
+	{
+		const FString PackagePath = Asset.PackageName.ToString();
+		const FAssetData AssetData = FRigVMTreeContext::FindAssetFromAnyPath(PackagePath, false);
+		if(AssetData.IsValid())
+		{
+			Nodes.Add(FRigVMTreePackageNode::Create(AssetData)->ToSharedRef());
+		}
+		else
+		{
+			continue;
+		}
+	}
+
+	return Nodes;
+}
+
+void SRigVMSwapAssetReferencesWidget::OnPhaseActivated(TSharedRef<FRigVMTreePhase> Phase)
+{
+	switch(Phase->GetID())
+	{
+		case PHASE_PICKSOURCE:
+		{
+			if(PickTargetContext->GetSourceAsset().IsValid())
+			{
+				if(const TSharedPtr<FRigVMTreeNode> Node = Phase->FindVisibleNode(PickTargetContext->GetSourceAsset().PackageName.ToString()))
+				{
+					const TSharedPtr<SRigVMChangesTreeView> TreeView = GetBulkEditWidget()->GetTreeView();
+					TreeView->SetSelection(Node.ToSharedRef(), true);
+				}
+			}
+			break;
+		}
+		case PHASE_PICKTARGET:
+		{
+			if(PickTargetContext->GetSourceAsset().IsValid())
+			{
+				if(const TSharedPtr<FRigVMTreeNode> Node = Phase->FindVisibleNode(PickTargetContext->GetSourceAsset().PackageName.ToString()))
+				{
+					const TSharedPtr<SRigVMChangesTreeView> TreeView = GetBulkEditWidget()->GetTreeView();
+					TreeView->SetSelection(Node.ToSharedRef(), true);
+				}
+			}
+			break;
+		}
+		case PHASE_PICKASSETREFS:
+		{
+			TArray<FSoftObjectPath> References = PickAssetRefsContext->GetReferences();
+			if (OnGetReferences.IsBound())
+			{
+				References = OnGetReferences.Execute(PickAssetRefsContext->GetSourceAsset());
+				PickAssetRefsContext->SetReferences(References);
+			}
+
+			TSet<FAssetData> Assets;
+			TArray<TSharedRef<FRigVMTreeNode>> Nodes;
+			for (const FSoftObjectPath& Reference : References)
+			{
+				const FAssetData ReferenceAssetData = IAssetRegistry::Get()->GetAssetByObjectPath(Reference.GetWithoutSubPath());
+				if (!Assets.Contains(ReferenceAssetData))
+				{
+					TSharedRef<FRigVMTreeAssetRefAssetNode> NewAssetNode = FRigVMTreeAssetRefAssetNode::Create(ReferenceAssetData);
+					if (bSkipPickingRefs)
+					{
+						NewAssetNode->SetCheckState(ECheckBoxState::Checked);
+					}
+					Nodes.Add(NewAssetNode);
+					Assets.Add(ReferenceAssetData);
+				}
+			}
+			
+			Phase->SetNodes(Nodes);
+
+			if(bSkipPickingRefs)
+			{
+				// process to next phase
+				GetBulkEditWidget()->OnPrimaryButtonClicked();
+			}
+			
+			break;
+		}
+		default:
+		{
+			return;
+		}
+	}
+}
+
+FReply SRigVMSwapAssetReferencesWidget::OnNodeSelected(TSharedRef<FRigVMTreeNode> Node)
+{
+	if(const TSharedPtr<FRigVMTreePackageNode> AssetNode = Cast<FRigVMTreePackageNode>(Node))
+	{
+		if(BulkEditWidget->GetActivePhase()->GetID() == PHASE_PICKSOURCE)
+		{
+			SetSourceAsset(AssetNode->GetAssetData());
+			return FReply::Handled();
+		}
+		if(BulkEditWidget->GetActivePhase()->GetID() == PHASE_PICKTARGET)
+		{
+			SetTargetAsset(AssetNode->GetAssetData());
+			return FReply::Handled();
+		}
+	}
+	return FReply::Unhandled();
+}
+
+FReply SRigVMSwapAssetReferencesWidget::OnNodeDoubleClicked(TSharedRef<FRigVMTreeNode> Node)
+{
+	if(Node->IsA<FRigVMTreePackageNode>())
+	{
+		const FAssetData AssetData = Node->GetAssetData();
+		if(AssetData.IsValid())
+		{
+			// force load
+			if(UObject* TopLevelObject = AssetData.GetSoftObjectPath().GetWithoutSubPath().TryLoad())
+			{
+				if(TopLevelObject->GetClass()->IsChildOf(URigVMBlueprint::StaticClass()))
+				{
+					if(const UObject* ObjectReference = StaticLoadObject(UObject::StaticClass(), nullptr, *Node->GetPath(), nullptr))
+					{
+						GEditor->GetEditorSubsystem<UAssetEditorSubsystem>()->OpenEditorForAsset(TopLevelObject);
+						if(IAssetEditorInstance* Editor = GEditor->GetEditorSubsystem<UAssetEditorSubsystem>()->FindEditorForAsset(TopLevelObject, true))
+						{
+							if(FRigVMEditor* RigVMEditor = static_cast<FRigVMEditor*>(Editor))
+							{
+								RigVMEditor->HandleJumpToHyperlink(ObjectReference);
+								return FReply::Handled();
+							}
+						}
+					}
+				}
+			}
+		}
+	}
+	return FReply::Unhandled();
+}
+
+void SRigVMSwapAssetReferencesWidget::SetSourceAsset(const FAssetData& InAsset)
+{
+	PickTargetContext->SetSourceAsset(InAsset);
+	PickAssetRefsContext->SetSourceAsset(InAsset);
+	PickAssetRefsContext->ClearReferences(); // If the source asset has changed, we should clear any existing references
+}
+
+void SRigVMSwapAssetReferencesWidget::SetTargetAsset(const FAssetData& InAsset)
+{
+	PickAssetRefsContext->SetTargetAsset(InAsset);
+}
+
+#undef LOCTEXT_NAMESPACE
diff --git a/Engine/Plugins/Runtime/RigVM/Source/RigVMEditor/Public/Editor/RigVMEditor.h b/Engine/Plugins/Runtime/RigVM/Source/RigVMEditor/Public/Editor/RigVMEditor.h
--- a/Engine/Plugins/Runtime/RigVM/Source/RigVMEditor/Public/Editor/RigVMEditor.h
+++ b/Engine/Plugins/Runtime/RigVM/Source/RigVMEditor/Public/Editor/RigVMEditor.h
@@ -311,6 +311,7 @@ protected:
 	virtual void SwapFunctionWithinAsset();
 	virtual void SwapFunctionAcrossProject();
 	virtual void SwapFunctionForAssets(const TArray<FAssetData>& InAssets, bool bSetupUndo);
+	virtual void SwapAssetReferences();
 
 	/** Once the log is collected update the graph */
 	void UpdateGraphCompilerErrors();
diff --git a/Engine/Plugins/Runtime/RigVM/Source/RigVMEditor/Public/Widgets/RigVMTreeToolkitNode.h b/Engine/Plugins/Runtime/RigVM/Source/RigVMEditor/Public/Widgets/RigVMTreeToolkitNode.h
--- a/Engine/Plugins/Runtime/RigVM/Source/RigVMEditor/Public/Widgets/RigVMTreeToolkitNode.h
+++ b/Engine/Plugins/Runtime/RigVM/Source/RigVMEditor/Public/Widgets/RigVMTreeToolkitNode.h
@@ -179,6 +179,8 @@ public:
 	virtual FText GetLabel() const override;
 	virtual bool IsLoaded() const override;
 
+	const FSoftObjectPath& GetPackagePath() const { return SoftObjectPath; }
+
 protected:
 	void HandleAssetLoaded(UObject* InAsset);
 
diff --git a/Engine/Plugins/Runtime/RigVM/Source/RigVMEditor/Public/Widgets/SRigVMSwapAssetReferencesWidget.h b/Engine/Plugins/Runtime/RigVM/Source/RigVMEditor/Public/Widgets/SRigVMSwapAssetReferencesWidget.h
new file mode 100644
--- /dev/null
+++ b/Engine/Plugins/Runtime/RigVM/Source/RigVMEditor/Public/Widgets/SRigVMSwapAssetReferencesWidget.h
@@ -0,0 +1,218 @@
+// Copyright Epic Games, Inc. All Rights Reserved.
+
+#pragma once
+
+#include "CoreMinimal.h"
+#include "Widgets/SRigVMBulkEditWidget.h"
+#include "RigVMEditorBlueprintLibrary.h"
+
+DECLARE_DELEGATE_RetVal_OneParam(TArray<FSoftObjectPath>, FOnGetReferences, FAssetData);
+DECLARE_DELEGATE_RetVal_TwoParams(bool, FOnSwapReference, FSoftObjectPath, FAssetData);
+
+class FRigVMSwapAssetReferencesContext : public FRigVMTreeContext
+{
+public:
+	DEFINE_RIGVMTREETOOLKIT_ELEMENT(FRigVMSwapAssetReferencesContext, FRigVMTreeContext)
+
+	const FAssetData& GetSourceAsset()
+	{
+		return SourceAsset;
+	}
+
+	void SetSourceAsset(const FAssetData& InAsset)
+	{
+		SourceAsset = InAsset;
+	}
+
+	const FAssetData& GetTargetAsset()
+	{
+		return TargetAsset;
+	}
+
+	void SetTargetAsset(const FAssetData& InAsset)
+	{
+		TargetAsset = InAsset;
+	}
+
+	const TArray<FSoftObjectPath>& GetReferences()
+	{
+		return References;
+	}
+
+	void SetReferences(const TArray<FSoftObjectPath>& InReferences)
+	{
+		References = InReferences;
+	}
+
+	void ClearReferences()
+	{
+		References.Reset();
+	}
+
+	virtual uint32 GetVisibleChildrenHash() const override;
+
+private:
+	FAssetData SourceAsset;
+	FAssetData TargetAsset;
+	TArray<FSoftObjectPath> References;
+};
+
+// A single reference node
+class FRigVMTreeReferenceNode : public FRigVMTreeNode
+{
+public:
+	DEFINE_RIGVMTREETOOLKIT_ELEMENT(FRigVMTreeReferenceNode, FRigVMTreeNode)
+
+	FRigVMTreeReferenceNode(const FSoftObjectPath& InModulePath)
+		: FRigVMTreeNode(InModulePath.GetSubPathString())
+		, ModulePath(InModulePath)
+	{
+		
+	}
+
+	virtual bool IsCheckable() const override
+	{
+		return true;
+	}
+	
+	const FSoftObjectPath& GetReferencePath() const
+	{
+		return ModulePath;
+	}
+
+private:
+	FSoftObjectPath ModulePath;
+};
+
+// The asset node which contains multiple reference nodes
+class FRigVMTreeAssetRefAssetNode : public FRigVMTreePackageNode
+{
+public:
+	DEFINE_RIGVMTREETOOLKIT_ELEMENT(FRigVMTreeAssetRefAssetNode, FRigVMTreePackageNode)
+	
+	FRigVMTreeAssetRefAssetNode(const FAssetData& InAssetData)
+		: FRigVMTreePackageNode(InAssetData)
+	{
+		
+	}
+
+	virtual bool IsCheckable() const override
+	{
+		return true;
+	}
+
+protected:
+	virtual TArray<TSharedRef<FRigVMTreeNode>> GetChildrenImpl(const TSharedRef<FRigVMTreeContext>& InContext) const override;
+
+	mutable TArray<TSharedRef<FRigVMTreeNode>> AssetRefNodes;
+};
+
+// Filters the source asset from the target asset list
+class FRigVMTreeTargetAssetFilter : public FRigVMTreeFilter
+{
+public:
+	DEFINE_RIGVMTREETOOLKIT_ELEMENT(FRigVMTreeTargetAssetFilter, FRigVMTreeFilter)
+	virtual bool Filters(TSharedRef<FRigVMTreeNode>& InNode, const TSharedRef<FRigVMTreeContext>& InContext) override;
+};
+
+// Filters target assets to show only variants of the source asset
+class FRigVMTreeAssetVariantFilter : public FRigVMTreeFilter
+{
+public:
+	DEFINE_RIGVMTREETOOLKIT_ELEMENT(FRigVMTreeAssetVariantFilter, FRigVMTreeFilter);
+	virtual bool CanBeToggledInUI() const override
+	{
+		return true;
+	}
+	virtual bool IsInvertedInUI() const override
+	{
+		return false;
+	}
+	virtual FText GetLabel() const override;
+
+	virtual bool Filters(TSharedRef<FRigVMTreeNode>& InNode, const TSharedRef<FRigVMTreeContext>& InContext) override;
+
+private:
+	mutable TMap<FString, TArray<FRigVMVariantRef>> SourceVariants;
+};
+
+class FRigVMSwapAssetReferenceTask : public FRigVMTreeTask
+{
+public:
+	DEFINE_RIGVMTREETOOLKIT_ELEMENT(FRigVMSwapAssetReferenceTask, FRigVMTreeTask)
+	
+	FRigVMSwapAssetReferenceTask(const FSoftObjectPath InReferencePath, const FAssetData& InNewAsset, const FOnSwapReference& InSwapFunction)
+		: ReferencePath(InReferencePath)
+		, NewAsset(InNewAsset)
+		, SwapFunction(InSwapFunction)
+	{
+	}
+
+	virtual bool Execute(const TSharedRef<FRigVMTreePhase>& InPhase) override;
+
+	virtual bool RequiresRefresh() const override
+	{
+		return true;
+	}
+
+	virtual bool RequiresUndo() const override
+	{
+		return true;
+	}
+	
+private:
+	
+	FSoftObjectPath ReferencePath;
+	FAssetData NewAsset;
+	FOnSwapReference SwapFunction;
+};
+
+class RIGVMEDITOR_API SRigVMSwapAssetReferencesWidget : public SCompoundWidget
+{
+public:
+
+	SLATE_BEGIN_ARGS(SRigVMSwapAssetReferencesWidget)
+		: _SkipPickingRefs(false)
+		, _EnableUndo(false)
+		, _CloseOnSuccess(false)
+	{}
+	SLATE_ARGUMENT(FAssetData, Source)
+	SLATE_ARGUMENT(FAssetData, Target)
+	SLATE_ARGUMENT(TArray<FSoftObjectPath>, ReferencePaths)
+	SLATE_ARGUMENT(bool, SkipPickingRefs)
+	SLATE_ARGUMENT(bool, EnableUndo)
+	SLATE_ARGUMENT(bool, CloseOnSuccess)
+	SLATE_EVENT(FRigVMAssetDataFilter, ExtraAssetFilter)
+	SLATE_EVENT(FOnGetReferences, OnGetReferences)
+	SLATE_EVENT(FOnSwapReference, OnSwapReference)
+	SLATE_END_ARGS()
+
+	void Construct(const FArguments& InArgs);
+
+	TSharedRef<SRigVMBulkEditWidget> GetBulkEditWidget()
+	{
+		return BulkEditWidget.ToSharedRef();
+	}
+
+private:
+
+	TSharedPtr<SRigVMBulkEditWidget> BulkEditWidget;
+	TSharedPtr<FRigVMSwapAssetReferencesContext> PickTargetContext;
+	TSharedPtr<FRigVMSwapAssetReferencesContext> PickAssetRefsContext;
+	bool bSkipPickingRefs = false;
+
+	FRigVMAssetDataFilter ExtraAssetFilter;
+	FOnGetReferences OnGetReferences;
+	FOnSwapReference OnSwapReference;
+
+	TArray<TSharedRef<FRigVMTreeNode>> GetAssetNodes(const FArguments& InArgs);
+	void OnPhaseActivated(TSharedRef<FRigVMTreePhase> Phase);
+	FReply OnNodeSelected(TSharedRef<FRigVMTreeNode> Node);
+	FReply OnNodeDoubleClicked(TSharedRef<FRigVMTreeNode> Node);
+	void SetSourceAsset(const FAssetData& InAsset);
+	void SetTargetAsset(const FAssetData& InAsset);
+
+	static constexpr int32 PHASE_PICKSOURCE = 0;
+	static constexpr int32 PHASE_PICKTARGET = 1;
+	static constexpr int32 PHASE_PICKASSETREFS = 2;
+};
diff --git a/Engine/Plugins/Runtime/RigVM/Source/RigVMEditor/Public/Widgets/SRigVMSwapFunctionsWidget.h b/Engine/Plugins/Runtime/RigVM/Source/RigVMEditor/Public/Widgets/SRigVMSwapFunctionsWidget.h
--- a/Engine/Plugins/Runtime/RigVM/Source/RigVMEditor/Public/Widgets/SRigVMSwapFunctionsWidget.h
+++ b/Engine/Plugins/Runtime/RigVM/Source/RigVMEditor/Public/Widgets/SRigVMSwapFunctionsWidget.h
@@ -303,9 +303,6 @@ public:
 
 private:
 
-	TSharedPtr<FRigVMTreePhase> PickSourceFunctionPhase;
-	TSharedPtr<FRigVMTreePhase> PickTargetFunctionPhase;
-	TSharedPtr<FRigVMTreePhase> PickFunctionReferencesPhase;
 	TSharedPtr<SRigVMBulkEditWidget> BulkEditWidget;
 	TSharedPtr<FRigVMSwapFunctionContext> PickTargetContext;
 	TSharedPtr<FRigVMSwapFunctionContext> PickFunctionRefsContext;
