RigVM: Ignore obsolete block for block combination hash

#rnx

diff --git a/Engine/Plugins/Runtime/RigVM/Source/RigVMDeveloper/Private/RigVMCompiler/RigVMAST.cpp b/Engine/Plugins/Runtime/RigVM/Source/RigVMDeveloper/Private/RigVMCompiler/RigVMAST.cpp
--- a/Engine/Plugins/Runtime/RigVM/Source/RigVMDeveloper/Private/RigVMCompiler/RigVMAST.cpp
+++ b/Engine/Plugins/Runtime/RigVM/Source/RigVMDeveloper/Private/RigVMCompiler/RigVMAST.cpp
@@ -224,6 +224,13 @@ FRigVMExprAST::FRigVMBlockArray FRigVMExprAST::GetBlocks(bool bSortByDepth) cons
 		Algo::Sort(Blocks, SortByDepth);
 	}
 
+	// remove the obsolete block to avoid combining expressions on
+	// a branch which are supposed to be obsolete.
+	if(const FRigVMParserAST* Parser = GetParser())
+	{
+		Blocks.Remove(Parser->GetObsoleteBlock(false));
+	}
+
 	return Blocks;
 }
 
