Control Rig: Swap ifdef for if preprocessor conditionals when using WITH_EDITORONLY_DATA
#rnx

diff --git a/Engine/Plugins/Runtime/RigVM/Source/RigVMDeveloper/Private/RigVMBlueprintUtils.cpp b/Engine/Plugins/Runtime/RigVM/Source/RigVMDeveloper/Private/RigVMBlueprintUtils.cpp
--- a/Engine/Plugins/Runtime/RigVM/Source/RigVMDeveloper/Private/RigVMBlueprintUtils.cpp
+++ b/Engine/Plugins/Runtime/RigVM/Source/RigVMDeveloper/Private/RigVMBlueprintUtils.cpp
@@ -104,7 +104,7 @@ void FRigVMBlueprintUtils::HandleReconstructAllNodes(UBlueprint* InBlueprint)
 void FRigVMBlueprintUtils::HandleRefreshAllNodes(UBlueprint* InBlueprint)
 {
 	DECLARE_SCOPE_HIERARCHICAL_COUNTER_FUNC()
-#ifdef WITH_EDITORONLY_DATA
+#if WITH_EDITORONLY_DATA
 	// Avoid refreshing EdGraph nodes during cook
 	if (GIsCookerLoadingPackage)
 	{
