Control Rig: Do not categorize function with function title if it is not a variant
#rnx

diff --git a/Engine/Plugins/Runtime/RigVM/Source/RigVMDeveloper/Private/RigVMModel/RigVMController.cpp b/Engine/Plugins/Runtime/RigVM/Source/RigVMDeveloper/Private/RigVMModel/RigVMController.cpp
--- a/Engine/Plugins/Runtime/RigVM/Source/RigVMDeveloper/Private/RigVMModel/RigVMController.cpp
+++ b/Engine/Plugins/Runtime/RigVM/Source/RigVMDeveloper/Private/RigVMModel/RigVMController.cpp
@@ -7381,20 +7381,23 @@ bool URigVMController::SetNodeCategory(URigVMCollapseNode* InNode, const FString
 	FString CategoryToSet = InCategory;
 	if (FRigVMGraphFunctionData* FunctionData = FindFunctionData(InNode->GetFName()))
 	{
-		if (FunctionData->Header.Variant.Guid.IsValid())
+		if (FunctionData->Header.LibraryPointer.IsVariant())
 		{
-			const FString NodeTitle = FunctionData->Header.NodeTitle;
-			TArray<FString> Categories;
-			InCategory.ParseIntoArray(Categories, TEXT("|"), true);
-			Categories.AddUnique(NodeTitle);
-			Categories.Remove(FString()); // Remove any empty strings
-			if (Categories[0] != NodeTitle)
+			if (FunctionData->Header.Variant.Guid.IsValid())
 			{
-				// Make sure the node title is the first element
-				Categories.Remove(NodeTitle);
-				Categories.Insert(NodeTitle, 0);
+				const FString NodeTitle = FunctionData->Header.NodeTitle;
+				TArray<FString> Categories;
+				InCategory.ParseIntoArray(Categories, TEXT("|"), true);
+				Categories.AddUnique(NodeTitle);
+				Categories.Remove(FString()); // Remove any empty strings
+				if (Categories[0] != NodeTitle)
+				{
+					// Make sure the node title is the first element
+					Categories.Remove(NodeTitle);
+					Categories.Insert(NodeTitle, 0);
+				}
+				CategoryToSet = FString::Join(Categories, TEXT("|"));
 			}
-			CategoryToSet = FString::Join(Categories, TEXT("|"));
 		}
 	}
 
