Control Rig: Fix nodes appearing multiple times in the context menu
#rnx

diff --git a/Engine/Plugins/Runtime/RigVM/Source/RigVMDeveloper/Private/RigVMBlueprint.cpp b/Engine/Plugins/Runtime/RigVM/Source/RigVMDeveloper/Private/RigVMBlueprint.cpp
--- a/Engine/Plugins/Runtime/RigVM/Source/RigVMDeveloper/Private/RigVMBlueprint.cpp
+++ b/Engine/Plugins/Runtime/RigVM/Source/RigVMDeveloper/Private/RigVMBlueprint.cpp
@@ -935,6 +935,7 @@ void URigVMBlueprint::PostSaveRoot(FObjectPostSaveRootContext ObjectSaveContext)
 	// Make sure all the tags are accounted for in the TypeActions after we save
 	if (FBlueprintActionDatabase* ActionDatabase = FBlueprintActionDatabase::TryGet())
 	{
+		ActionDatabase->ClearAssetActions(GetClass());
 		ActionDatabase->RefreshClassActions(GetClass());
 	}
 }
