Control Rig: Fix crash when deleting rigvm graph
#rnx

diff --git a/Engine/Plugins/Runtime/RigVM/Source/RigVMDeveloper/Private/RigVMModel/RigVMClient.cpp b/Engine/Plugins/Runtime/RigVM/Source/RigVMDeveloper/Private/RigVMModel/RigVMClient.cpp
--- a/Engine/Plugins/Runtime/RigVM/Source/RigVMDeveloper/Private/RigVMModel/RigVMClient.cpp
+++ b/Engine/Plugins/Runtime/RigVM/Source/RigVMDeveloper/Private/RigVMModel/RigVMClient.cpp
@@ -566,9 +566,6 @@ bool FRigVMClient::RemoveModel(const FString& InNodePathOrName, bool bSetupUndoR
 		}
 #endif
 
-		// remove the controller
-		verify(RemoveController(Model));
-
 		if (GetOuter()->Implements<URigVMClientHost>())
 		{
 			IRigVMClientHost* ClientHost = Cast<IRigVMClientHost>(GetOuter());
@@ -589,6 +586,9 @@ bool FRigVMClient::RemoveModel(const FString& InNodePathOrName, bool bSetupUndoR
 
 		// clean up the model
 		Models.Remove(Model);
+		
+		// remove the controller
+		verify(RemoveController(Model));
 
 		NotifyOuterOfPropertyChange();
 		return true;
@@ -622,10 +622,15 @@ FName FRigVMClient::RenameModel(const FString& InNodePathOrName, const FName& In
 		}
 #endif
 
+		TObjectPtr<URigVMController>* Controller = Controllers.Find(Model);
 		const FString OldNodePath = Model->GetNodePath();
 		const FName SafeNewName = GetUniqueName(InNewName);
 		Model->Rename(*SafeNewName.ToString(), nullptr, REN_ForceNoResetLoaders | REN_DontCreateRedirectors);
 		const FString NewNodePath = Model->GetNodePath();
+		if (Controller)
+		{
+			Controllers.Add(Model, *Controller);
+		}
 
 		if(bSetupUndoRedo)
 		{
