RigVM Variant: Fix member initialization

#rnx

diff --git a/Engine/Plugins/Runtime/RigVM/Source/RigVM/Public/RigVMCore/RigVMVariant.h b/Engine/Plugins/Runtime/RigVM/Source/RigVM/Public/RigVMCore/RigVMVariant.h
--- a/Engine/Plugins/Runtime/RigVM/Source/RigVM/Public/RigVMCore/RigVMVariant.h
+++ b/Engine/Plugins/Runtime/RigVM/Source/RigVM/Public/RigVMCore/RigVMVariant.h
@@ -13,6 +13,24 @@ struct RIGVM_API FRigVMTag
 {
 	GENERATED_BODY()
 
+	FRigVMTag()
+		: Name(NAME_None)
+		, Label()
+		, ToolTip()
+		, Color(FLinearColor::White)
+		, bShowInUserInterface(true)
+		, bMarksSubjectAsInvalid(false)
+	{}
+
+	FRigVMTag(const FName&  InName, const FString& InLabel, const FText& InToolTip, const FLinearColor& InColor, bool InShowInUserInterface = true, bool InMarksSubjectAsInvalid = false)
+	: Name(InName)
+	, Label(InLabel)
+	, ToolTip(InToolTip)
+	, Color(InColor)
+	, bShowInUserInterface(InShowInUserInterface)
+	, bMarksSubjectAsInvalid(InMarksSubjectAsInvalid)
+	{}
+
 	UPROPERTY(BlueprintReadWrite, EditAnywhere, Category=Variant)
 	FName Name;
 	
@@ -26,10 +44,35 @@ struct RIGVM_API FRigVMTag
 	FLinearColor Color;
 
 	UPROPERTY(BlueprintReadWrite, EditAnywhere, Category=Variant)
-	bool bShowInUserInterface = true;
+	bool bShowInUserInterface;
 
 	UPROPERTY(BlueprintReadWrite, EditAnywhere, Category=Variant)
-	bool bMarksSubjectAsInvalid = false;
+	bool bMarksSubjectAsInvalid;
+
+	bool IsValid() const
+	{
+		return !Name.IsNone();
+	}
+
+	FString GetLabel() const
+	{
+		if(Label.IsEmpty())
+		{
+			return Name.ToString();
+		}
+		return Label;
+	}
+
+	friend uint32 GetTypeHash(const FRigVMTag& InTag)
+	{
+		uint32 Hash = GetTypeHash(InTag.Name);
+		Hash = HashCombine(Hash, GetTypeHash(InTag.Label));
+		Hash = HashCombine(Hash, GetTypeHash(InTag.ToolTip.ToString()));
+		Hash = HashCombine(Hash, GetTypeHash(InTag.Color));
+		Hash = HashCombine(Hash, GetTypeHash(InTag.bShowInUserInterface));
+		Hash = HashCombine(Hash, GetTypeHash(InTag.bMarksSubjectAsInvalid));
+		return Hash;
+	}
 	
 	friend FArchive& operator<<(FArchive& Ar, FRigVMTag& Data)
 	{
