RigVM Editor: Implement variant tag front end

#rnx

diff --git a/Engine/Plugins/Runtime/RigVM/Source/RigVM/Private/RigVMCore/RigVMGraphFunctionHost.cpp b/Engine/Plugins/Runtime/RigVM/Source/RigVM/Private/RigVMCore/RigVMGraphFunctionHost.cpp
--- a/Engine/Plugins/Runtime/RigVM/Source/RigVM/Private/RigVMCore/RigVMGraphFunctionHost.cpp
+++ b/Engine/Plugins/Runtime/RigVM/Source/RigVM/Private/RigVMCore/RigVMGraphFunctionHost.cpp
@@ -180,6 +180,7 @@ FRigVMGraphFunctionData* FRigVMGraphFunctionStore::UpdateFunctionInterface(const
 		{
 			if (FRigVMGraphFunctionData* NewData = AddFunction(Header, bIsPublic))
 			{
+				NewData->Header.Variant = Header.Variant;
 				NewData->Header.Dependencies = Dependencies;
 				NewData->Header.ExternalVariables = ExternalVariables;
 				return NewData;
diff --git a/Engine/Plugins/Runtime/RigVM/Source/RigVM/Private/RigVMSettings.cpp b/Engine/Plugins/Runtime/RigVM/Source/RigVM/Private/RigVMSettings.cpp
--- a/Engine/Plugins/Runtime/RigVM/Source/RigVM/Private/RigVMSettings.cpp
+++ b/Engine/Plugins/Runtime/RigVM/Source/RigVM/Private/RigVMSettings.cpp
@@ -4,6 +4,8 @@
 
 #include UE_INLINE_GENERATED_CPP_BY_NAME(RigVMSettings)
 
+#define LOCTEXT_NAMESPACE "RigVMSettings"
+
 URigVMEditorSettings::URigVMEditorSettings(const FObjectInitializer& ObjectInitializer)
     : Super(ObjectInitializer)
 {
@@ -12,3 +14,27 @@ URigVMEditorSettings::URigVMEditorSettings(const FObjectInitializer& ObjectIniti
 #endif
 }
 
+URigVMProjectSettings::URigVMProjectSettings(const FObjectInitializer& ObjectInitializer)
+	: Super(ObjectInitializer)
+{
+	VariantTags.Emplace(TEXT("Outdated"), TEXT("Outdated"), LOCTEXT("OutdatedToolTip", "This item is outdated and should no longer be used."), FLinearColor::Red, true, true);
+	VariantTags.Emplace(TEXT("Stable"), TEXT("Stable"), LOCTEXT("StableToolTip", "This item is stable and ready to use."), FLinearColor::Green, true);
+}
+
+#undef LOCTEXT_NAMESPACE // RigVMSettings
+FRigVMTag URigVMProjectSettings::GetTag(FName InTagName) const
+{
+	if(const FRigVMTag* FoundTag = FindTag(InTagName))
+	{
+		return *FoundTag;
+	}
+	return FRigVMTag();
+}
+
+const FRigVMTag* URigVMProjectSettings::FindTag(FName InTagName) const
+{
+	return VariantTags.FindByPredicate([InTagName](const FRigVMTag& Tag)
+	{
+		return InTagName == Tag.Name;
+	});
+}
diff --git a/Engine/Plugins/Runtime/RigVM/Source/RigVM/Public/RigVMSettings.h b/Engine/Plugins/Runtime/RigVM/Source/RigVM/Public/RigVMSettings.h
--- a/Engine/Plugins/Runtime/RigVM/Source/RigVM/Public/RigVMSettings.h
+++ b/Engine/Plugins/Runtime/RigVM/Source/RigVM/Public/RigVMSettings.h
@@ -33,10 +33,21 @@ public:
 	UPROPERTY(EditAnywhere, config, Category = Interaction)
 	bool bAutoLinkMutableNodes;
 
-	UPROPERTY(EditAnywhere, config, Category = Variants)
-	TArray<FRigVMTag> VariantTags;
-
 #endif
 };
 
 
+UCLASS(config = Editor, meta=(DisplayName="RigVM Project Settings"))
+class RIGVM_API URigVMProjectSettings : public UDeveloperSettings
+{
+	GENERATED_UCLASS_BODY()
+
+public:
+
+	UPROPERTY(EditAnywhere, config, Category = Variants)
+	TArray<FRigVMTag> VariantTags;
+
+	UFUNCTION(BlueprintPure, Category= Variants)
+	FRigVMTag GetTag(FName InTagName) const;
+	const FRigVMTag* FindTag(FName InTagName) const;
+};
\ No newline at end of file
diff --git a/Engine/Plugins/Runtime/RigVM/Source/RigVMDeveloper/Private/RigVMModel/Nodes/RigVMLibraryNode.cpp b/Engine/Plugins/Runtime/RigVM/Source/RigVMDeveloper/Private/RigVMModel/Nodes/RigVMLibraryNode.cpp
--- a/Engine/Plugins/Runtime/RigVM/Source/RigVMDeveloper/Private/RigVMModel/Nodes/RigVMLibraryNode.cpp
+++ b/Engine/Plugins/Runtime/RigVM/Source/RigVMDeveloper/Private/RigVMModel/Nodes/RigVMLibraryNode.cpp
@@ -252,7 +252,19 @@ FRigVMGraphFunctionHeader URigVMLibraryNode::GetFunctionHeader(IRigVMGraphFuncti
 		{
 			if (!Header.Variant.Guid.IsValid())
 			{
+				if (!Data->Header.Variant.Guid.IsValid())
+				{
+					Data->Header.Variant.Guid = FRigVMVariant::GenerateGUID();
+				}
 				Header.Variant = Data->Header.Variant;
+
+				if (URigVMFunctionLibrary* Library = GetLibrary())
+				{
+					if (FRigVMVariant* Variant = Library->GetFunctionVariant(GetFName()))
+					{
+						Variant->Guid = Header.Variant.Guid;
+					}
+				}
 			}
 			Header.NodeTitle = Data->Header.NodeTitle;
 		}
diff --git a/Engine/Plugins/Runtime/RigVM/Source/RigVMDeveloper/Private/RigVMModel/RigVMController.cpp b/Engine/Plugins/Runtime/RigVM/Source/RigVMDeveloper/Private/RigVMModel/RigVMController.cpp
--- a/Engine/Plugins/Runtime/RigVM/Source/RigVMDeveloper/Private/RigVMModel/RigVMController.cpp
+++ b/Engine/Plugins/Runtime/RigVM/Source/RigVMDeveloper/Private/RigVMModel/RigVMController.cpp
@@ -12333,16 +12333,13 @@ URigVMLibraryNode* URigVMController::CreateFunctionVariant(const FName& InFuncti
 
 bool URigVMController::AddDefaultTagToFunctionVariant(const FName& InFunctionName, const FName& InTagName, bool bSetupUndoRedo, bool bPrintPythonCommand)
 {
-	URigVMEditorSettings* Settings = GetMutableDefault<URigVMEditorSettings>(URigVMEditorSettings::StaticClass());
+	URigVMProjectSettings* Settings = GetMutableDefault<URigVMProjectSettings>(URigVMProjectSettings::StaticClass());
 	if (!Settings)
 	{
 		return false;
 	}
 
-	const FRigVMTag* Tag = Settings->VariantTags.FindByPredicate([InTagName](const FRigVMTag& Tag)
-	{
-		return InTagName == Tag.Name;
-	});
+	const FRigVMTag* Tag = Settings->FindTag(InTagName);
 	if (!Tag)
 	{
 		ReportErrorf(TEXT("Could not find default tag with name %s."), *InTagName.ToString());
@@ -12394,13 +12391,22 @@ bool URigVMController::AddTagToFunctionVariant(const FName& InFunctionName, cons
 		return false;
 	}
 
-	FRigVMVariant* Variant = FunctionLibrary->FunctionToVariant.Find(InFunctionName);
+	FRigVMVariant* Variant = FunctionLibrary->GetFunctionVariant(InFunctionName);
 	if (!Variant)
 	{
 		ReportErrorf(TEXT("Could not find function variant for function %s."), *InFunctionName.ToString());
 		return false;
 	}
 
+	if(Variant->Tags.ContainsByPredicate([InTag](const FRigVMTag& Tag) -> bool
+	{
+		return Tag.Name == InTag.Name; 
+	}))
+	{
+		ReportErrorf(TEXT("Cannot add tag %s to function %s - tag already applied."), *InFunctionName.ToString());
+		return false;
+	}
+
 	FRigVMBaseAction Action(this);
 	if (bSetupUndoRedo)
 	{
@@ -12468,7 +12474,7 @@ bool URigVMController::RemoveTagFromFunctionVariant(const FName& InFunctionName,
 		return false;
 	}
 
-	FRigVMVariant* Variant = FunctionLibrary->FunctionToVariant.Find(InFunctionName);
+	FRigVMVariant* Variant = FunctionLibrary->GetFunctionVariant(InFunctionName);
 	if (!Variant)
 	{
 		ReportErrorf(TEXT("Could not find function variant for function %s."), *InFunctionName.ToString());
@@ -12493,7 +12499,10 @@ bool URigVMController::RemoveTagFromFunctionVariant(const FName& InFunctionName,
 		GetActionStack()->AddAction(FRigVMRemoveFunctionVariantTagAction(this, InFunctionName, InTagName));
 	}
 
-	Variant->Tags.RemoveAt(Index);
+	Variant->Tags.RemoveAll([InTagName](const FRigVMTag& Tag) -> bool
+	{
+		return Tag.Name == InTagName;
+	});
 
 	Notify(ERigVMGraphNotifType::VariantTagsChanged, LibraryNode);
 
diff --git a/Engine/Plugins/Runtime/RigVM/Source/RigVMDeveloper/Private/RigVMModel/RigVMFunctionLibrary.cpp b/Engine/Plugins/Runtime/RigVM/Source/RigVMDeveloper/Private/RigVMModel/RigVMFunctionLibrary.cpp
--- a/Engine/Plugins/Runtime/RigVM/Source/RigVMDeveloper/Private/RigVMModel/RigVMFunctionLibrary.cpp
+++ b/Engine/Plugins/Runtime/RigVM/Source/RigVMDeveloper/Private/RigVMModel/RigVMFunctionLibrary.cpp
@@ -168,4 +168,14 @@ const FSoftObjectPath URigVMFunctionLibrary::GetFunctionHostObjectPath() const
 	return FSoftObjectPath();
 }
 
+const FRigVMVariant* URigVMFunctionLibrary::GetFunctionVariant(const FName& InFunctionName) const
+{
+	return const_cast<URigVMFunctionLibrary*>(this)->GetFunctionVariant(InFunctionName);
+}
+
+FRigVMVariant* URigVMFunctionLibrary::GetFunctionVariant(const FName& InFunctionName)
+{
+	return FunctionToVariant.Find(InFunctionName);
+}
+
 
diff --git a/Engine/Plugins/Runtime/RigVM/Source/RigVMDeveloper/Public/RigVMBlueprint.h b/Engine/Plugins/Runtime/RigVM/Source/RigVMDeveloper/Public/RigVMBlueprint.h
--- a/Engine/Plugins/Runtime/RigVM/Source/RigVMDeveloper/Public/RigVMBlueprint.h
+++ b/Engine/Plugins/Runtime/RigVM/Source/RigVMDeveloper/Public/RigVMBlueprint.h
@@ -828,6 +828,7 @@ private:
 	friend class URigVMBuildData;
 	friend class FRigVMVariantDetailCustomization;
 	friend class FRigVMTreeAssetVariantFilter;
+	friend class FRigVMTreePackageNode;
 };
 
 class RIGVMDEVELOPER_API FRigVMBlueprintCompileScope
diff --git a/Engine/Plugins/Runtime/RigVM/Source/RigVMDeveloper/Public/RigVMModel/RigVMFunctionLibrary.h b/Engine/Plugins/Runtime/RigVM/Source/RigVMDeveloper/Public/RigVMModel/RigVMFunctionLibrary.h
--- a/Engine/Plugins/Runtime/RigVM/Source/RigVMDeveloper/Public/RigVMModel/RigVMFunctionLibrary.h
+++ b/Engine/Plugins/Runtime/RigVM/Source/RigVMDeveloper/Public/RigVMModel/RigVMFunctionLibrary.h
@@ -75,7 +75,8 @@ public:
 
 	bool IsFunctionPublic(const FName& InFunctionName) const { return PublicFunctionNames.Contains(InFunctionName); }
 
-	const FRigVMVariant* GetFunctionVariant(const FName& InFunctionName) const { return FunctionToVariant.Find(InFunctionName); }
+	const FRigVMVariant* GetFunctionVariant(const FName& InFunctionName) const;
+	FRigVMVariant* GetFunctionVariant(const FName& InFunctionName);
 
 private:
 
diff --git a/Engine/Plugins/Runtime/RigVM/Source/RigVMEditor/Private/Editor/RigVMGraphDetailCustomization.cpp b/Engine/Plugins/Runtime/RigVM/Source/RigVMEditor/Private/Editor/RigVMGraphDetailCustomization.cpp
--- a/Engine/Plugins/Runtime/RigVM/Source/RigVMEditor/Private/Editor/RigVMGraphDetailCustomization.cpp
+++ b/Engine/Plugins/Runtime/RigVM/Source/RigVMEditor/Private/Editor/RigVMGraphDetailCustomization.cpp
@@ -966,7 +966,7 @@ void FRigVMGraphDetailCustomization::CustomizeDetails(IDetailLayoutBuilder& Deta
 		.OverrideResetToDefault(FResetToDefaultOverride::Hide())
 		.Visibility(TAttribute<EVisibility>::CreateLambda([this]()
 		{
-			return IsFunctionVariant() ? EVisibility::Visible : EVisibility::Collapsed;
+			return IsValidFunction() ? EVisibility::Visible : EVisibility::Collapsed;
 		}))
 		.NameContent()
 		[
@@ -981,6 +981,11 @@ void FRigVMGraphDetailCustomization::CustomizeDetails(IDetailLayoutBuilder& Deta
 			.VariantRefs(this, &FRigVMGraphDetailCustomization::GetVariantRefs)
 			.OnVariantChanged(this, &FRigVMGraphDetailCustomization::OnVariantChanged)
 			.OnBrowseVariantRef(this, &FRigVMGraphDetailCustomization::OnBrowseVariantRef)
+			.OnGetTags(this, &FRigVMGraphDetailCustomization::OnGetAssignedTags)
+			.OnAddTag(this, &FRigVMGraphDetailCustomization::OnAddAssignedTag)
+			.OnRemoveTag(this, &FRigVMGraphDetailCustomization::OnRemoveAssignedTag)
+			.CanAddTags(true)
+			.EnableTagContextMenu(true)
 		];
 	}
 
@@ -1337,7 +1342,7 @@ TSharedRef<ITableRow> FRigVMGraphDetailCustomization::HandleGenerateRowAccessSpe
         ];
 }
 
-bool FRigVMGraphDetailCustomization::IsFunctionVariant() const
+bool FRigVMGraphDetailCustomization::IsValidFunction() const
 {
 	if (GraphPtr.IsValid() && RigVMBlueprintPtr.IsValid())
 	{
@@ -1346,7 +1351,7 @@ bool FRigVMGraphDetailCustomization::IsFunctionVariant() const
 		{
 			if (const URigVMLibraryNode* LibraryNode = Cast<URigVMLibraryNode>(Model->GetOuter()))
 			{
-				return LibraryNode->GetFunctionHeader(Blueprint->GetRigVMGraphFunctionHost()).LibraryPointer.IsVariant(); 
+				return LibraryNode->GetFunctionHeader(Blueprint->GetRigVMGraphFunctionHost()).IsValid(); 
 			}
 		}
 	}
@@ -1414,6 +1419,51 @@ void FRigVMGraphDetailCustomization::OnBrowseVariantRef(const FRigVMVariantRef&
 	}
 }
 
+TArray<FRigVMTag> FRigVMGraphDetailCustomization::OnGetAssignedTags() const
+{
+	return GetVariant().Tags;
+}
+
+void FRigVMGraphDetailCustomization::OnAddAssignedTag(const FName& InTagName)
+{
+	if (GraphPtr.IsValid() && RigVMBlueprintPtr.IsValid())
+	{
+		URigVMBlueprint* Blueprint = RigVMBlueprintPtr.Get();
+		if (const URigVMGraph* Model = Blueprint->GetModel(GraphPtr.Get()))
+		{
+			if (URigVMGraph* FunctionLibrary = Blueprint->GetLocalFunctionLibrary())
+			{
+				if (const URigVMLibraryNode* LibraryNode = Cast<URigVMLibraryNode>(Model->GetOuter()))
+				{
+					const FString& FunctionName = LibraryNode->GetFunctionHeader().LibraryPointer.GetFunctionName();
+					URigVMController* FunctionLibraryController = Blueprint->GetOrCreateController(FunctionLibrary);
+					FunctionLibraryController->AddDefaultTagToFunctionVariant(*FunctionName, InTagName);
+				}
+			}
+		}
+	}
+}
+
+void FRigVMGraphDetailCustomization::OnRemoveAssignedTag(const FName& InTagName)
+{
+	if (GraphPtr.IsValid() && RigVMBlueprintPtr.IsValid())
+	{
+		URigVMBlueprint* Blueprint = RigVMBlueprintPtr.Get();
+		if (const URigVMGraph* Model = Blueprint->GetModel(GraphPtr.Get()))
+		{
+			if (URigVMGraph* FunctionLibrary = Blueprint->GetLocalFunctionLibrary())
+			{
+				if (const URigVMLibraryNode* LibraryNode = Cast<URigVMLibraryNode>(Model->GetOuter()))
+				{
+					const FString& FunctionName = LibraryNode->GetFunctionHeader().LibraryPointer.GetFunctionName();
+					URigVMController* FunctionLibraryController = Blueprint->GetOrCreateController(FunctionLibrary);
+					FunctionLibraryController->RemoveTagFromFunctionVariant(*FunctionName, InTagName);
+				}
+			}
+		}
+	}
+}
+
 FRigVMWrappedNodeDetailCustomization::FRigVMWrappedNodeDetailCustomization()
 : BlueprintBeingCustomized(nullptr)
 {
diff --git a/Engine/Plugins/Runtime/RigVM/Source/RigVMEditor/Private/Editor/RigVMVariantDetailCustomization.cpp b/Engine/Plugins/Runtime/RigVM/Source/RigVMEditor/Private/Editor/RigVMVariantDetailCustomization.cpp
--- a/Engine/Plugins/Runtime/RigVM/Source/RigVMEditor/Private/Editor/RigVMVariantDetailCustomization.cpp
+++ b/Engine/Plugins/Runtime/RigVM/Source/RigVMEditor/Private/Editor/RigVMVariantDetailCustomization.cpp
@@ -12,6 +12,7 @@
 #include "ContentBrowserModule.h"
 #include "IContentBrowserSingleton.h"
 #include "ScopedTransaction.h"
+#include "RigVMSettings.h"
 
 #define LOCTEXT_NAMESPACE "RigVMVariantDetailCustomization"
 
@@ -45,6 +46,11 @@ void FRigVMVariantDetailCustomization::CustomizeHeader(TSharedRef<IPropertyHandl
 		.VariantRefs(this, &FRigVMVariantDetailCustomization::GetVariantRefs)
 		.OnVariantChanged(this, &FRigVMVariantDetailCustomization::OnVariantChanged)
 		.OnBrowseVariantRef(this, &FRigVMVariantDetailCustomization::OnBrowseVariantRef)
+		.OnGetTags(this, &FRigVMVariantDetailCustomization::OnGetTags)
+		.OnAddTag(this, &FRigVMVariantDetailCustomization::OnAddTag)
+		.OnRemoveTag(this, &FRigVMVariantDetailCustomization::OnRemoveTag)
+		.CanAddTags(true)
+		.EnableTagContextMenu(true)
 	];
 }
 
@@ -98,4 +104,55 @@ void FRigVMVariantDetailCustomization::OnBrowseVariantRef(const FRigVMVariantRef
 	}
 }
 
+TArray<FRigVMTag> FRigVMVariantDetailCustomization::OnGetTags() const
+{
+	if(BlueprintBeingCustomized)
+	{
+		return BlueprintBeingCustomized->AssetVariant.Tags;
+	}
+	return {};
+}
+
+void FRigVMVariantDetailCustomization::OnAddTag(const FName& InTagName)
+{
+	if(BlueprintBeingCustomized)
+	{
+		const URigVMProjectSettings* Settings = GetMutableDefault<URigVMProjectSettings>(URigVMProjectSettings::StaticClass());
+		if (Settings)
+		{
+			if (const FRigVMTag* Tag = Settings->FindTag(InTagName))
+			{
+				if(!BlueprintBeingCustomized->AssetVariant.Tags.ContainsByPredicate([InTagName](const FRigVMTag& Tag) -> bool
+				{
+					return Tag.Name == InTagName;
+				}))
+				{
+					FScopedTransaction Transaction(LOCTEXT("AddedBlueprintVariantTag", "Added Blueprint Variant Tag"));
+					BlueprintBeingCustomized->Modify();
+					BlueprintBeingCustomized->AssetVariant.Tags.Add(*Tag);
+				}
+			}
+		}
+	}
+}
+
+void FRigVMVariantDetailCustomization::OnRemoveTag(const FName& InTagName)
+{
+	if(BlueprintBeingCustomized)
+	{
+		if(BlueprintBeingCustomized->AssetVariant.Tags.ContainsByPredicate([InTagName](const FRigVMTag& Tag) -> bool
+			{
+				return Tag.Name == InTagName;
+			}))
+		{
+			FScopedTransaction Transaction(LOCTEXT("RemovedBlueprintVariantTag", "Removed Blueprint Variant Tag"));
+			BlueprintBeingCustomized->Modify();
+			BlueprintBeingCustomized->AssetVariant.Tags.RemoveAll([InTagName](const FRigVMTag& Tag) -> bool
+			{
+				return Tag.Name == InTagName;
+			});
+		}
+	}
+}
+
 #undef LOCTEXT_NAMESPACE
diff --git a/Engine/Plugins/Runtime/RigVM/Source/RigVMEditor/Private/Widgets/RigVMTreeToolkitNode.cpp b/Engine/Plugins/Runtime/RigVM/Source/RigVMEditor/Private/Widgets/RigVMTreeToolkitNode.cpp
--- a/Engine/Plugins/Runtime/RigVM/Source/RigVMEditor/Private/Widgets/RigVMTreeToolkitNode.cpp
+++ b/Engine/Plugins/Runtime/RigVM/Source/RigVMEditor/Private/Widgets/RigVMTreeToolkitNode.cpp
@@ -332,6 +332,7 @@ FSlateColor FRigVMTreeCategoryNode::GetBackgroundColor(bool bIsHovered, bool bIs
 
 FRigVMTreePackageNode::FRigVMTreePackageNode(const FAssetData& InAssetData)
 	: FRigVMTreeCategoryNode(InAssetData.GetObjectPathString())
+	, bRetrievedTags(false)
 {
 	SoftObjectPath = InAssetData.GetSoftObjectPath();
 	bIsLoaded = InAssetData.IsAssetLoaded();
@@ -374,6 +375,31 @@ bool FRigVMTreePackageNode::IsLoaded() const
 	return bIsLoaded.GetValue();
 }
 
+const TArray<FRigVMTag>& FRigVMTreePackageNode::GetTags() const
+{
+	if(!bRetrievedTags)
+	{
+		if(IsLoaded())
+		{
+			if(const URigVMBlueprint* Blueprint = GetBlueprint())
+			{
+				Tags = Blueprint->AssetVariant.Tags;
+			}
+			else
+			{
+				const FAssetData AssetData = GetAssetData();
+				if(AssetData.FindTag(GET_MEMBER_NAME_CHECKED(URigVMBlueprint, AssetVariant)))
+				{
+					//const FRigVMVariant& Variant = AssetData.GetTagValueRef<FRigVMVariant>(GET_MEMBER_NAME_CHECKED(URigVMBlueprint, AssetVariant));
+					//Tags = Variant.Tags;
+				}
+			}
+		}
+		bRetrievedTags = true;
+	}
+	return FRigVMTreeCategoryNode::GetTags();
+}
+
 const FSlateBrush* FRigVMTreePackageNode::GetIconAndTint(FLinearColor& OutColor) const
 {
 	if(!IconBrush.IsSet())
diff --git a/Engine/Plugins/Runtime/RigVM/Source/RigVMEditor/Private/Widgets/SRigVMChangesTreeView.cpp b/Engine/Plugins/Runtime/RigVM/Source/RigVMEditor/Private/Widgets/SRigVMChangesTreeView.cpp
--- a/Engine/Plugins/Runtime/RigVM/Source/RigVMEditor/Private/Widgets/SRigVMChangesTreeView.cpp
+++ b/Engine/Plugins/Runtime/RigVM/Source/RigVMEditor/Private/Widgets/SRigVMChangesTreeView.cpp
@@ -13,6 +13,7 @@
 #include "Framework/Application/SlateApplication.h"
 #include "ContentBrowserModule.h"
 #include "IContentBrowserSingleton.h"
+#include "Widgets/SRigVMVariantWidget.h"
 #define LOCTEXT_NAMESPACE "SRigVMChangesTreeView"
 
 SRigVMChangesTreeRow::~SRigVMChangesTreeRow()
@@ -82,6 +83,21 @@ void SRigVMChangesTreeRow::Construct(const FArguments& InArgs, const TSharedRef<
 				.Text(Node->GetLabel())
 				.ColorAndOpacity(this, &SRigVMChangesTreeRow::GetTextColor)
 			]
+			+ SHorizontalBox::Slot()
+			.AutoWidth()
+			.Padding(8, 0, 0, 0)
+			.HAlign(HAlign_Left)
+			.VAlign(VAlign_Center)
+			[
+				SNew(SRigVMVariantTagWidget)
+				.Visibility(GetVariantTags().IsEmpty() ? EVisibility::Collapsed : EVisibility::Visible)
+				.CanAddTags(false)
+				.EnableContextMenu(false)
+				.EnableTick(false)
+				.Orientation(EOrientation::Orient_Horizontal)
+				.OnGetTags(this, &SRigVMChangesTreeRow::GetVariantTags)
+			]
+
 			+ SHorizontalBox::Slot()
 			.FillWidth(1)
 			[
@@ -178,6 +194,15 @@ EVisibility SRigVMChangesTreeRow::GetExpanderVisibility() const
 	return Node->HasVisibleChildren() ? EVisibility::Visible : EVisibility::Collapsed;
 }
 
+TArray<FRigVMTag> SRigVMChangesTreeRow::GetVariantTags() const
+{
+	if(!Tags.IsSet())
+	{
+		Tags = Node->GetTags();
+	}
+	return Tags.GetValue();
+}
+
 void SRigVMChangesTreeRow::RequestRefresh(bool bForce)
 {
 	if(OwningWidget)
diff --git a/Engine/Plugins/Runtime/RigVM/Source/RigVMEditor/Private/Widgets/SRigVMGraphNode.cpp b/Engine/Plugins/Runtime/RigVM/Source/RigVMEditor/Private/Widgets/SRigVMGraphNode.cpp
--- a/Engine/Plugins/Runtime/RigVM/Source/RigVMEditor/Private/Widgets/SRigVMGraphNode.cpp
+++ b/Engine/Plugins/Runtime/RigVM/Source/RigVMEditor/Private/Widgets/SRigVMGraphNode.cpp
@@ -39,6 +39,7 @@
 #include "RigVMFunctions/RigVMDispatch_Select.h"
 #include "Algo/Copy.h"
 #include "Brushes/SlateColorBrush.h"
+#include "Widgets/SRigVMVariantTagWidget.h"
 
 #if WITH_EDITOR
 #include "Editor.h"
@@ -1884,6 +1885,52 @@ void SRigVMGraphNode::UpdatePinTreeView()
 		.Size(FVector2D(1.f, 2.f))
 	];
 
+	// add the tags this node potentially has
+	if(ModelNode.IsValid())
+	{
+		if(URigVMFunctionReferenceNode* FunctionRefNode = Cast<URigVMFunctionReferenceNode>(ModelNode.Get()))
+		{
+			const FRigVMVariant Variant = FunctionRefNode->GetReferencedFunctionHeader().Variant;
+			if(!Variant.Tags.IsEmpty())
+			{
+				TWeakObjectPtr<URigVMFunctionReferenceNode> WeakFunctionRefNode = FunctionRefNode;
+				
+				LeftNodeBox->AddSlot()
+				.HAlign(HAlign_Fill)
+				.VAlign(VAlign_Center)
+				.AutoHeight()
+				.Padding(2, 0, 2, 2)
+				[
+					SNew(SRigVMVariantTagWidget)
+					.Visibility_Lambda([WeakFunctionRefNode]() -> EVisibility
+					{
+						if(WeakFunctionRefNode.IsValid())
+						{
+							const FRigVMVariant Variant = WeakFunctionRefNode->GetReferencedFunctionHeader().Variant;
+							if(!Variant.Tags.IsEmpty())
+							{
+								return EVisibility::Visible;
+							}
+						}
+						return EVisibility::Collapsed;
+					})
+					.Orientation(EOrientation::Orient_Horizontal)
+					.CanAddTags(false)
+					.EnableContextMenu(false)
+					.OnGetTags_Lambda([WeakFunctionRefNode]() -> TArray<FRigVMTag>
+					{
+						if(WeakFunctionRefNode.IsValid())
+						{
+							const FRigVMVariant Variant = WeakFunctionRefNode->GetReferencedFunctionHeader().Variant;
+							return Variant.Tags;
+						}
+						return {};
+					})
+				];
+			}
+		}
+	}
+
 	auto AddArrayPlusButtonLambda = [this](URigVMPin* InModelPin, TSharedPtr<SHorizontalBox> InSlotLayout, const float InEmptySidePadding)
 	{
 		// add array plus button
diff --git a/Engine/Plugins/Runtime/RigVM/Source/RigVMEditor/Private/Widgets/SRigVMSwapFunctionsWidget.cpp b/Engine/Plugins/Runtime/RigVM/Source/RigVMEditor/Private/Widgets/SRigVMSwapFunctionsWidget.cpp
--- a/Engine/Plugins/Runtime/RigVM/Source/RigVMEditor/Private/Widgets/SRigVMSwapFunctionsWidget.cpp
+++ b/Engine/Plugins/Runtime/RigVM/Source/RigVMEditor/Private/Widgets/SRigVMSwapFunctionsWidget.cpp
@@ -259,6 +259,16 @@ const FSlateBrush* FRigVMTreeFunctionIdentifierNode::GetIconAndTint(FLinearColor
 	return Icon.GetIcon();
 }
 
+const TArray<FRigVMTag>& FRigVMTreeFunctionIdentifierNode::GetTags() const
+{
+	const FRigVMGraphFunctionHeader Header = FRigVMGraphFunctionHeader::FindGraphFunctionHeader(Identifier);
+	if(Header.IsValid())
+	{
+		Tags = Header.Variant.Tags;
+	}
+	return FRigVMTreeNode::GetTags();
+}
+
 void FRigVMTreeFunctionIdentifierAssetNode::AddChildNode(const TSharedRef<FRigVMTreeNode>& InNode)
 {
 	AddChildImpl(InNode);
diff --git a/Engine/Plugins/Runtime/RigVM/Source/RigVMEditor/Private/Widgets/SRigVMVariantTagWidget.cpp b/Engine/Plugins/Runtime/RigVM/Source/RigVMEditor/Private/Widgets/SRigVMVariantTagWidget.cpp
new file mode 100644
--- /dev/null
+++ b/Engine/Plugins/Runtime/RigVM/Source/RigVMEditor/Private/Widgets/SRigVMVariantTagWidget.cpp
@@ -0,0 +1,388 @@
+// Copyright Epic Games, Inc. All Rights Reserved.
+
+#include "Widgets/SRigVMVariantTagWidget.h"
+#include "Widgets/Input/SComboButton.h"
+#include "Widgets/Images/SImage.h"
+#include "Framework/MultiBox/MultiBoxBuilder.h"
+#include "Framework/Application/SlateApplication.h"
+#include "RigVMSettings.h"
+#include "Editor/RigVMEditorStyle.h"
+#include "Styling/AppStyle.h"
+
+#define LOCTEXT_NAMESPACE "SRigVMVariantTagWidget"
+
+SRigVMVariantCapsule::SRigVMVariantCapsule()
+{
+}
+
+SRigVMVariantCapsule::~SRigVMVariantCapsule()
+{
+}
+
+void SRigVMVariantCapsule::Construct(
+	const FArguments& InArgs)
+{
+	NameAttribute = InArgs._Name;
+	ColorAttribute = InArgs._Color;
+	OnRemoveTag = InArgs._OnRemoveTag;
+	EnableContextMenu = InArgs._EnableContextMenu;
+
+	TAttribute<float> MinDesiredWithAttribute;
+	if(InArgs._MinDesiredLabelWidth > SMALL_NUMBER)
+	{
+		MinDesiredWithAttribute = InArgs._MinDesiredLabelWidth;
+	}
+	
+	SButton::FArguments ButtonArgs;
+	ButtonArgs
+	.ButtonStyle( &FRigVMEditorStyle::Get().GetWidgetStyle< FButtonStyle >( "TagButton" ) )
+	.ContentPadding(0)
+	.HAlign(HAlign_Fill)
+	.OnClicked(InArgs._OnClicked)
+	.Content()
+	[
+		SNew(SBorder)
+		.Padding(0.0f)
+		.HAlign(HAlign_Fill)
+		.BorderImage(FAppStyle::Get().GetBrush("FilterBar.FilterBackground"))
+		[
+			SNew(SHorizontalBox)
+			+SHorizontalBox::Slot()
+			.VAlign(VAlign_Center)
+			.HAlign(HAlign_Left)
+			.Padding(0)
+			.AutoWidth()
+			[
+				SNew(SImage)
+				.Image(FAppStyle::Get().GetBrush("FilterBar.FilterImage"))
+				.ToolTipText(InArgs._ToolTipText)
+				.ColorAndOpacity(this, &SRigVMVariantCapsule::GetColor)
+			]
+			+SHorizontalBox::Slot()
+			.HAlign(HAlign_Fill)
+			.VAlign(VAlign_Center)
+			.Padding(FMargin(4,1,4,1))
+			.AutoWidth()
+			[
+				SNew(STextBlock)
+				.MinDesiredWidth(MinDesiredWithAttribute)
+				.Text(InArgs._Label)
+				.ToolTipText(InArgs._ToolTipText)
+				.ColorAndOpacity(this, &SRigVMVariantCapsule::GetLabelColor)
+			]
+		]
+	];
+	
+	SButton::Construct(ButtonArgs);
+
+	SetPadding(InArgs._Padding);
+}
+
+FSlateColor SRigVMVariantCapsule::GetColor() const
+{
+	if(IsEnabled())
+	{
+		return ColorAttribute.Get();
+	}
+	return FLinearColor::Black;
+}
+
+FSlateColor SRigVMVariantCapsule::GetLabelColor() const
+{
+	if(IsEnabled())
+	{
+		return FSlateColor::UseForeground();
+	}
+	return FSlateColor::UseSubduedForeground();
+}
+
+FReply SRigVMVariantCapsule::OnMouseButtonDown(const FGeometry& MyGeometry, const FPointerEvent& MouseEvent)
+{
+	if(EnableContextMenu.Get() && MouseEvent.IsMouseButtonDown(EKeys::RightMouseButton))
+	{
+		FMenuBuilder MenuBuilder( true, NULL );
+		if(OnRemoveTag.IsBound())
+		{
+			FUIAction RemoveTagAction( FExecuteAction::CreateSP( this, &SRigVMVariantCapsule::HandleRemoveTag) );
+			MenuBuilder.AddMenuEntry( LOCTEXT("RemoveTag", "RemoveTag"), FText::GetEmpty(), FSlateIcon(), RemoveTagAction, NAME_None, EUserInterfaceActionType::Button);
+		}
+		const FWidgetPath WidgetPath = MouseEvent.GetEventPath() != nullptr ? *MouseEvent.GetEventPath() : FWidgetPath();
+		FSlateApplication::Get().PushMenu(SharedThis(this), WidgetPath, MenuBuilder.MakeWidget(), FSlateApplication::Get().GetCursorPos(), FPopupTransitionEffect::ContextMenu);
+		return FReply::Handled();
+	}
+	return SButton::OnMouseButtonDown(MyGeometry, MouseEvent);
+}
+
+void SRigVMVariantCapsule::HandleRemoveTag()
+{
+	if(OnRemoveTag.IsBound() && (NameAttribute.IsBound() || NameAttribute.IsSet()))
+	{
+		(void)OnRemoveTag.Execute(NameAttribute.Get());
+	}
+}
+
+SRigVMVariantTagWidget::SRigVMVariantTagWidget()
+	: LastTagHash(0)
+{
+}
+
+SRigVMVariantTagWidget::~SRigVMVariantTagWidget()
+{
+}
+
+void SRigVMVariantTagWidget::Construct(
+	const FArguments& InArgs)
+{
+	OnGetTags = InArgs._OnGetTags;
+	OnAddTag = InArgs._OnAddTag;
+	OnRemoveTag = InArgs._OnRemoveTag;
+	CanAddTags = InArgs._CanAddTags;
+	EnableContextMenu = InArgs._EnableContextMenu;
+	LastTagHash = 0;
+	MinDesiredLabelWidth = InArgs._MinDesiredLabelWidth;
+
+	const TAttribute<EVisibility> BoxVisibilityAttribute = TAttribute<EVisibility>::CreateLambda([this]() -> EVisibility
+	{
+		if(!OnGetTags.IsBound())
+		{
+			return EVisibility::Collapsed;
+		}
+		return OnGetTags.Execute().IsEmpty() ? EVisibility::Collapsed : EVisibility::Visible;
+	});
+
+	if(InArgs._Orientation == EOrientation::Orient_Vertical)
+	{
+		SAssignNew(VerticalCapsuleBox, SVerticalBox)
+		.Visibility(BoxVisibilityAttribute);
+	}
+	else
+	{
+		SAssignNew(HorizontalCapsuleBox, SHorizontalBox)
+		.Visibility(BoxVisibilityAttribute);
+	}
+
+	const TSharedRef<SWidget> CapsuleBox = 
+		HorizontalCapsuleBox.IsValid() ? 
+		StaticCastSharedRef<SWidget>(HorizontalCapsuleBox.ToSharedRef()) :
+		StaticCastSharedRef<SWidget>(VerticalCapsuleBox.ToSharedRef());
+
+	ChildSlot
+	[
+		SNew(SVerticalBox)
+		+SVerticalBox::Slot()
+		.HAlign(HAlign_Left)
+		.AutoHeight()
+		.Padding( 0,0,0,0 )
+		[
+			SNew(SHorizontalBox)
+
+			+SHorizontalBox::Slot()
+			.AutoWidth()
+			.VAlign(VAlign_Center)
+			.Padding(0)
+			[
+				CapsuleBox
+			]
+
+			+SHorizontalBox::Slot()
+			.AutoWidth()
+			.VAlign(VAlign_Center)
+			.Padding(0)
+			[
+				SNew(STextBlock)
+				.MinDesiredWidth(50)
+				.Text(LOCTEXT("NoTags", "No tags"))
+				.ToolTipText(LOCTEXT("NoTagsToolTip", "There are no tags currently applied. Press the plus button to add a tag."))
+				.Visibility_Lambda([this]() -> EVisibility
+				{
+					if(!OnGetTags.IsBound())
+					{
+						return EVisibility::Collapsed;
+					}
+					return OnGetTags.Execute().IsEmpty() ? EVisibility::Visible : EVisibility::Collapsed;
+				})
+			]
+
+			+SHorizontalBox::Slot()
+			.AutoWidth()
+			.VAlign(VAlign_Top)
+			.Padding( 7,5,0,0 )
+			[
+				SNew(SComboButton)
+				.Visibility_Lambda([this]() -> EVisibility
+				{
+					if(!CanAddTags.IsBound() && !CanAddTags.IsSet())
+					{
+						return EVisibility::Collapsed;
+					}
+					return CanAddTags.Get() ? EVisibility::Visible : EVisibility::Collapsed;
+				})
+				.Method(EPopupMethod::UseCurrentWindow)
+				.ComboButtonStyle(FAppStyle::Get(), "SimpleComboButton")
+				.HasDownArrow(false)
+				.ButtonContent()
+				[
+					SNew(SImage)
+					.Image(FAppStyle::GetBrush(TEXT("Icons.PlusCircle")))
+				]
+				.OnGetMenuContent(this, &SRigVMVariantTagWidget::OnBuildAddTagMenuContent)
+			]
+		]
+	];
+
+	UpdateCapsules();
+	SetCanTick(InArgs._EnableTick);
+}
+
+void SRigVMVariantTagWidget::Tick(const FGeometry& AllottedGeometry, const double InCurrentTime, const float InDeltaTime)
+{
+	SBox::Tick(AllottedGeometry, InCurrentTime, InDeltaTime);
+	UpdateCapsules();
+}
+
+void SRigVMVariantTagWidget::UpdateCapsules()
+{
+	if(!OnGetTags.IsBound())
+	{
+		return;
+	}
+
+	TArray<FRigVMTag> AssignedTags = OnGetTags.Execute();
+	AssignedTags.Sort([](const FRigVMTag& A, const FRigVMTag& B) -> bool
+	{
+		return FCString::Strcmp(*A.GetLabel(), *B.GetLabel()) < 0;
+	});
+	
+	uint32 Hash = 0;
+	for(const FRigVMTag& AssignedTag : AssignedTags)
+	{
+		Hash = HashCombine(Hash, GetTypeHash(AssignedTag));
+	}
+
+	if(Hash == LastTagHash)
+	{
+		return;
+	}
+	LastTagHash = Hash;
+
+	if(HorizontalCapsuleBox)
+	{
+		HorizontalCapsuleBox->ClearChildren();
+	}
+	else
+	{
+		VerticalCapsuleBox->ClearChildren();
+	}
+
+	for(const FRigVMTag& AssignedTag : AssignedTags)
+	{
+		if(!AssignedTag.bShowInUserInterface)
+		{
+			continue;
+		}
+
+		TSharedRef<SRigVMVariantCapsule> Capsule = SNew(SRigVMVariantCapsule)
+		.Name(AssignedTag.Name)
+		.Color(AssignedTag.Color)
+		.Label(FText::FromString(AssignedTag.GetLabel()))
+		.ToolTipText(AssignedTag.ToolTip)
+		.EnableContextMenu(EnableContextMenu)
+		.MinDesiredLabelWidth(MinDesiredLabelWidth)
+		.OnRemoveTag(OnRemoveTag);
+
+		if(HorizontalCapsuleBox)
+		{
+			HorizontalCapsuleBox->AddSlot()
+			.Padding(0, 0, 3, 0)
+			[
+				Capsule
+			];
+		}
+		else
+		{
+			VerticalCapsuleBox->AddSlot()
+			.Padding(0, 3, 0, 0)
+			[
+				Capsule
+			];
+		}
+	}
+}
+
+FReply SRigVMVariantTagWidget::OnAddTagClicked(const FName& InTagName) const
+{
+	if(OnAddTag.IsBound())
+	{
+		(void)OnAddTag.Execute(InTagName);
+		return FReply::Handled();
+	}
+	return FReply::Unhandled();
+}
+
+TSharedRef<SWidget> SRigVMVariantTagWidget::OnBuildAddTagMenuContent() const
+{
+	FMenuBuilder MenuBuilder(true, nullptr);
+
+	MenuBuilder.BeginSection(FName("AddTag"), LOCTEXT("AddTag", "Add Tag"));
+
+	const URigVMProjectSettings* Settings = GetDefault<URigVMProjectSettings>(URigVMProjectSettings::StaticClass());
+	TArray<FRigVMTag> AvailableTags = Settings->VariantTags;
+	AvailableTags.Sort([](const FRigVMTag& A, const FRigVMTag& B) -> bool
+	{
+		return FCString::Strcmp(*A.GetLabel(), *B.GetLabel()) < 0;
+	});
+
+	TArray<FRigVMTag> AssignedTags;
+	if(OnGetTags.IsBound())
+	{
+		AssignedTags = OnGetTags.Execute();
+	}
+
+	for(const FRigVMTag& AvailableTag
+		: AvailableTags)
+	{
+		if(!AvailableTag.IsValid())
+		{
+			continue;
+		}
+		if(!AvailableTag.bShowInUserInterface)
+		{
+			continue;
+		}
+
+		const bool bAlreadyHasTag = AssignedTags.ContainsByPredicate([AvailableTag](const FRigVMTag& ExistingTag) -> bool
+		{
+			return AvailableTag.Name == ExistingTag.Name;
+		});
+		
+		const FName TagName = AvailableTag.Name;
+		MenuBuilder.AddWidget(
+			SNew(SRigVMVariantCapsule)
+			.IsEnabled(!bAlreadyHasTag)
+			.Name(AvailableTag.Name)
+			.Color(AvailableTag.Color)
+			.Label(FText::FromString(AvailableTag.GetLabel()))
+			.ToolTipText(AvailableTag.ToolTip)
+			.EnableContextMenu(false)
+			.MinDesiredLabelWidth(MinDesiredLabelWidth)
+			.OnClicked_Lambda([this, TagName]()
+			{
+				const FReply Reply = OnAddTagClicked(TagName);
+				if(WeakAddTagMenuWidget.IsValid())
+				{
+					FSlateApplication::Get().DismissMenuByWidget(WeakAddTagMenuWidget.Pin().ToSharedRef());
+				}
+				return Reply;
+			})
+			, FText(), false, true, AvailableTag.ToolTip);
+	}
+
+	MenuBuilder.EndSection();
+	
+	TSharedRef<SWidget> MenuWidget = MenuBuilder.MakeWidget();
+	WeakAddTagMenuWidget = MenuWidget.ToWeakPtr();
+	
+	return MenuWidget;
+}
+
+#undef LOCTEXT_NAMESPACE // SRigVMVariantTagWidget
diff --git a/Engine/Plugins/Runtime/RigVM/Source/RigVMEditor/Private/Widgets/SRigVMVariantWidget.cpp b/Engine/Plugins/Runtime/RigVM/Source/RigVMEditor/Private/Widgets/SRigVMVariantWidget.cpp
--- a/Engine/Plugins/Runtime/RigVM/Source/RigVMEditor/Private/Widgets/SRigVMVariantWidget.cpp
+++ b/Engine/Plugins/Runtime/RigVM/Source/RigVMEditor/Private/Widgets/SRigVMVariantWidget.cpp
@@ -20,8 +20,10 @@ void SRigVMVariantWidget::Construct(
 	const FArguments& InArgs)
 {
 	VariantAttribute = InArgs._Variant;
-	VariantRefsAttribute = InArgs._VariantRefs;
+	
 	OnVariantChanged = InArgs._OnVariantChanged;
+
+	VariantRefsAttribute = InArgs._VariantRefs;
 	OnCreateVariantRefRow = InArgs._OnCreateVariantRefRow;
 	OnBrowseVariantRef = InArgs._OnBrowseVariantRef;
 
@@ -35,6 +37,7 @@ void SRigVMVariantWidget::Construct(
 		SNew(SVerticalBox)
 		+ SVerticalBox::Slot()
 		.AutoHeight()
+		.Padding(0, 8, 0, 0)
 		.HAlign(HAlign_Left)
 		[
 			SNew(STextBlock)
@@ -48,6 +51,20 @@ void SRigVMVariantWidget::Construct(
 		+ SVerticalBox::Slot()
 		.AutoHeight()
 		.HAlign(HAlign_Fill)
+		.Padding(0, 8, 0, 0)
+		[
+			SAssignNew(TagWidget, SRigVMVariantTagWidget)
+			.OnGetTags(InArgs._OnGetTags)
+			.OnAddTag(InArgs._OnAddTag)
+			.OnRemoveTag(InArgs._OnRemoveTag)
+			.CanAddTags(InArgs._CanAddTags)
+			.EnableContextMenu(InArgs._EnableTagContextMenu)
+			.MinDesiredLabelWidth(50.f)
+		]
+		+ SVerticalBox::Slot()
+		.AutoHeight()
+		.HAlign(HAlign_Fill)
+		.Padding(0, 8, 0, 0)
 		[
 			SNew(SScrollBox)
 			.Visibility(this, &SRigVMVariantWidget::GetVariantRefListVisibility)
diff --git a/Engine/Plugins/Runtime/RigVM/Source/RigVMEditor/Public/Editor/RigVMEditorStyle.h b/Engine/Plugins/Runtime/RigVM/Source/RigVMEditor/Public/Editor/RigVMEditorStyle.h
--- a/Engine/Plugins/Runtime/RigVM/Source/RigVMEditor/Public/Editor/RigVMEditorStyle.h
+++ b/Engine/Plugins/Runtime/RigVM/Source/RigVMEditor/Public/Editor/RigVMEditorStyle.h
@@ -2,9 +2,11 @@
 
 #pragma once
 
+#include "Styling/CoreStyle.h"
 #include "Styling/SlateStyle.h"
 #include "Styling/SlateStyleRegistry.h"
 #include "Styling/SlateStyleMacros.h"
+#include "Styling/SlateTypes.h"
 #include "Styling/StyleColors.h"
 #include "Brushes/SlateRoundedBoxBrush.h"
 
@@ -88,6 +90,18 @@ class RIGVMEDITOR_API FRigVMEditorStyle final
 			Set("RigVM.Node.PinTree.Arrow_Expanded_Right", new IMAGE_BRUSH("Slate/RigVM_TreeArrow_Expanded_Right", Icon10x10, DefaultForeground));
 			Set("RigVM.Node.PinTree.Arrow_Expanded_Hovered_Right", new IMAGE_BRUSH("Slate/RigVM_TreeArrow_Expanded_Hovered_Right", Icon10x10, DefaultForeground));
 		}
+
+		// Button styles
+		{
+			FButtonStyle TagButton = FCoreStyle::Get().GetWidgetStyle<FButtonStyle>("Button");
+
+			// disable padding on the tag button
+			TagButton
+			.SetNormalPadding(0)
+			.SetPressedPadding(0);
+			
+			Set( "TagButton", TagButton );
+		}
 	}
 public:
 
diff --git a/Engine/Plugins/Runtime/RigVM/Source/RigVMEditor/Public/Editor/RigVMGraphDetailCustomization.h b/Engine/Plugins/Runtime/RigVM/Source/RigVMEditor/Public/Editor/RigVMGraphDetailCustomization.h
--- a/Engine/Plugins/Runtime/RigVM/Source/RigVMEditor/Public/Editor/RigVMGraphDetailCustomization.h
+++ b/Engine/Plugins/Runtime/RigVM/Source/RigVMEditor/Public/Editor/RigVMGraphDetailCustomization.h
@@ -193,7 +193,7 @@ public:
 	FText GetCurrentAccessSpecifierName() const;
 	void OnAccessSpecifierSelected( TSharedPtr<FRigVMStringWithTag> SpecifierName, ESelectInfo::Type SelectInfo );
 	TSharedRef<ITableRow> HandleGenerateRowAccessSpecifier( TSharedPtr<FRigVMStringWithTag> SpecifierName, const TSharedRef<STableViewBase>& OwnerTable );
-	bool IsFunctionVariant() const;
+	bool IsValidFunction() const;
 	FRigVMVariant GetVariant() const;
 	TArray<FRigVMVariantRef> GetVariantRefs() const;
 
@@ -201,6 +201,9 @@ private:
 
 	void OnVariantChanged(const FRigVMVariant& InVariant);
 	void OnBrowseVariantRef(const FRigVMVariantRef& InVariantRef);
+	TArray<FRigVMTag> OnGetAssignedTags() const;
+	void OnAddAssignedTag(const FName& InTagName);
+	void OnRemoveAssignedTag(const FName& InTagName);
 
 	/** The Blueprint editor we are embedded in */
 	TWeakPtr<FRigVMEditor> RigVMEditorPtr;
diff --git a/Engine/Plugins/Runtime/RigVM/Source/RigVMEditor/Public/Editor/RigVMVariantDetailCustomization.h b/Engine/Plugins/Runtime/RigVM/Source/RigVMEditor/Public/Editor/RigVMVariantDetailCustomization.h
--- a/Engine/Plugins/Runtime/RigVM/Source/RigVMEditor/Public/Editor/RigVMVariantDetailCustomization.h
+++ b/Engine/Plugins/Runtime/RigVM/Source/RigVMEditor/Public/Editor/RigVMVariantDetailCustomization.h
@@ -5,6 +5,7 @@
 #include "CoreMinimal.h"
 #include "IPropertyTypeCustomization.h"
 #include "RigVMBlueprint.h"
+#include "RigVMCore/RigVMVariant.h"
 #include "Widgets/SRigVMLogWidget.h"
 
 class IPropertyHandle;
@@ -29,6 +30,9 @@ protected:
 
 	void OnVariantChanged(const FRigVMVariant& InNewVariant);
 	void OnBrowseVariantRef(const FRigVMVariantRef& InVariantRef);
+	TArray<FRigVMTag> OnGetTags() const;
+	void OnAddTag(const FName& InTagName);
+	void OnRemoveTag(const FName& InTagName);
 
 	URigVMBlueprint* BlueprintBeingCustomized;
 
diff --git a/Engine/Plugins/Runtime/RigVM/Source/RigVMEditor/Public/Widgets/RigVMTreeToolkitNode.h b/Engine/Plugins/Runtime/RigVM/Source/RigVMEditor/Public/Widgets/RigVMTreeToolkitNode.h
--- a/Engine/Plugins/Runtime/RigVM/Source/RigVMEditor/Public/Widgets/RigVMTreeToolkitNode.h
+++ b/Engine/Plugins/Runtime/RigVM/Source/RigVMEditor/Public/Widgets/RigVMTreeToolkitNode.h
@@ -4,6 +4,7 @@
 
 #include "RigVMBlueprint.h"
 #include "RigVMTreeToolkitDefines.h"
+#include "RigVMCore/RigVMVariant.h"
 #include "Styling/SlateTypes.h"
 #include "Textures/SlateIcon.h"
 
@@ -52,6 +53,11 @@ public:
 		return false;
 	}
 	
+	virtual const TArray<FRigVMTag>& GetTags() const
+	{
+		return Tags;
+	}
+	
 	int32 GetDepth() const;
 
 	TSharedPtr<FRigVMTreeNode> GetParent() const
@@ -117,6 +123,7 @@ protected:
 	mutable TArray<TSharedRef<FRigVMTreeNode>> VisibleChildren;
 	mutable bool bIsDirtyingParent = false;
 	TOptional<ECheckBoxState> CheckState;
+	mutable TArray<FRigVMTag> Tags;
 
 	DECLARE_DELEGATE_OneParam(FRefreshDelegate, bool)
 	FRefreshDelegate RefreshDelegate;
@@ -178,6 +185,7 @@ public:
 
 	virtual FText GetLabel() const override;
 	virtual bool IsLoaded() const override;
+	virtual const TArray<FRigVMTag>& GetTags() const override;
 
 	const FSoftObjectPath& GetPackagePath() const { return SoftObjectPath; }
 
@@ -188,4 +196,5 @@ protected:
 	FSoftObjectPath SoftObjectPath;
 	mutable FSlateIcon Icon;
 	mutable TOptional<const FSlateBrush*> IconBrush;
+	mutable bool bRetrievedTags;
 };
diff --git a/Engine/Plugins/Runtime/RigVM/Source/RigVMEditor/Public/Widgets/SRigVMChangesTreeView.h b/Engine/Plugins/Runtime/RigVM/Source/RigVMEditor/Public/Widgets/SRigVMChangesTreeView.h
--- a/Engine/Plugins/Runtime/RigVM/Source/RigVMEditor/Public/Widgets/SRigVMChangesTreeView.h
+++ b/Engine/Plugins/Runtime/RigVM/Source/RigVMEditor/Public/Widgets/SRigVMChangesTreeView.h
@@ -40,10 +40,12 @@ private:
 	const FSlateBrush* GetExpanderImage() const;
 	FReply OnExpanderMouseButtonDown(const FGeometry& SenderGeometry, const FPointerEvent& MouseEvent);
 	EVisibility GetExpanderVisibility() const;
+	TArray<FRigVMTag> GetVariantTags() const;
 	void RequestRefresh(bool bForce);
 
 	TSharedPtr<FRigVMTreeNode> Node;
 	SRigVMChangesTreeView* OwningWidget = nullptr;
+	mutable TOptional<TArray<FRigVMTag>> Tags;
 };
 
 DECLARE_DELEGATE_RetVal_OneParam(FReply,FOnRigVMTreeNodeSelected,TSharedRef<FRigVMTreeNode>);
diff --git a/Engine/Plugins/Runtime/RigVM/Source/RigVMEditor/Public/Widgets/SRigVMSwapFunctionsWidget.h b/Engine/Plugins/Runtime/RigVM/Source/RigVMEditor/Public/Widgets/SRigVMSwapFunctionsWidget.h
--- a/Engine/Plugins/Runtime/RigVM/Source/RigVMEditor/Public/Widgets/SRigVMSwapFunctionsWidget.h
+++ b/Engine/Plugins/Runtime/RigVM/Source/RigVMEditor/Public/Widgets/SRigVMSwapFunctionsWidget.h
@@ -142,6 +142,7 @@ public:
 	virtual FText GetLabel() const override;
 	virtual bool IsCheckable() const override { return false; }
 	virtual const FSlateBrush* GetIconAndTint(FLinearColor& OutColor) const override;
+	virtual const TArray<FRigVMTag>& GetTags() const override;
 
 	const FRigVMGraphFunctionIdentifier& GetIdentifier() const
 	{
diff --git a/Engine/Plugins/Runtime/RigVM/Source/RigVMEditor/Public/Widgets/SRigVMVariantTagWidget.h b/Engine/Plugins/Runtime/RigVM/Source/RigVMEditor/Public/Widgets/SRigVMVariantTagWidget.h
new file mode 100644
--- /dev/null
+++ b/Engine/Plugins/Runtime/RigVM/Source/RigVMEditor/Public/Widgets/SRigVMVariantTagWidget.h
@@ -0,0 +1,101 @@
+// Copyright Epic Games, Inc. All Rights Reserved.
+
+#pragma once
+
+#include "Widgets/Layout/SBox.h"
+#include "Widgets/Input/SButton.h"
+#include "Widgets/SBoxPanel.h"
+#include "RigVMCore/RigVMVariant.h"
+
+DECLARE_DELEGATE_RetVal(TArray<FRigVMTag>, FRigVMVariant_OnGetTags);
+DECLARE_DELEGATE_OneParam(FRigVMVariant_OnAddTag, const FName&);
+DECLARE_DELEGATE_OneParam(FRigVMVariant_OnRemoveTag, const FName&);
+
+class SRigVMVariantCapsule : public SButton
+{
+public:
+	
+	SLATE_BEGIN_ARGS(SRigVMVariantCapsule)
+		: _Padding(0)
+		, _EnableContextMenu(false)
+		, _MinDesiredLabelWidth(0.f)
+	{
+	}
+	SLATE_ATTRIBUTE(FName, Name)
+	SLATE_ATTRIBUTE(FText, Label)
+	SLATE_ATTRIBUTE(FText, ToolTipText)
+	SLATE_ATTRIBUTE(FSlateColor, Color)
+	SLATE_EVENT(FOnClicked, OnClicked)
+	SLATE_ATTRIBUTE(FMargin, Padding)
+	SLATE_EVENT(FRigVMVariant_OnRemoveTag, OnRemoveTag)
+	SLATE_ATTRIBUTE(bool, EnableContextMenu)
+	SLATE_ARGUMENT(float, MinDesiredLabelWidth)
+	SLATE_END_ARGS()
+
+	SRigVMVariantCapsule();
+	virtual ~SRigVMVariantCapsule() override;
+
+	void Construct(const FArguments& InArgs);
+
+	FSlateColor GetColor() const;
+	FSlateColor GetLabelColor() const;
+
+	virtual FReply OnMouseButtonDown( const FGeometry& MyGeometry, const FPointerEvent& MouseEvent ) override;
+
+private:
+
+	void HandleRemoveTag();
+
+	TAttribute<FName> NameAttribute;
+	TAttribute<FSlateColor> ColorAttribute;
+	FRigVMVariant_OnRemoveTag OnRemoveTag;
+	TAttribute<bool> EnableContextMenu;
+};
+
+class SRigVMVariantTagWidget : public SBox
+{
+public:
+
+	SLATE_BEGIN_ARGS(SRigVMVariantTagWidget)
+		: _Orientation(EOrientation::Orient_Vertical)
+		, _CanAddTags(false)
+		, _EnableContextMenu(false)
+		, _MinDesiredLabelWidth(0.f)
+		, _EnableTick(true)
+	{
+	}
+	SLATE_ARGUMENT(EOrientation, Orientation)
+	SLATE_ATTRIBUTE(bool, CanAddTags)
+	SLATE_ATTRIBUTE(bool, EnableContextMenu)
+	SLATE_EVENT(FRigVMVariant_OnGetTags, OnGetTags)
+	SLATE_EVENT(FRigVMVariant_OnAddTag, OnAddTag)
+	SLATE_EVENT(FRigVMVariant_OnRemoveTag, OnRemoveTag)
+	SLATE_ARGUMENT(float, MinDesiredLabelWidth)
+	SLATE_ARGUMENT(bool, EnableTick)
+	SLATE_END_ARGS()
+
+	SRigVMVariantTagWidget();
+	virtual ~SRigVMVariantTagWidget() override;
+	
+	void Construct(const FArguments& InArgs);
+
+	virtual void Tick(const FGeometry& AllottedGeometry, const double InCurrentTime, const float InDeltaTime) override;
+
+private:
+
+	void UpdateCapsules();
+	FReply OnAddTagClicked(const FName& InTagName) const;
+	TSharedRef<SWidget> OnBuildAddTagMenuContent() const;
+
+	TSharedPtr<SVerticalBox> VerticalCapsuleBox;
+	TSharedPtr<SHorizontalBox> HorizontalCapsuleBox;
+	mutable TWeakPtr<SWidget> WeakAddTagMenuWidget;
+
+	FRigVMVariant_OnGetTags OnGetTags;
+	FRigVMVariant_OnAddTag OnAddTag;
+	FRigVMVariant_OnRemoveTag OnRemoveTag;
+	TAttribute<bool> CanAddTags;
+	TAttribute<bool> EnableContextMenu;
+	uint32 LastTagHash;
+	float MinDesiredLabelWidth;
+};
diff --git a/Engine/Plugins/Runtime/RigVM/Source/RigVMEditor/Public/Widgets/SRigVMVariantWidget.h b/Engine/Plugins/Runtime/RigVM/Source/RigVMEditor/Public/Widgets/SRigVMVariantWidget.h
--- a/Engine/Plugins/Runtime/RigVM/Source/RigVMEditor/Public/Widgets/SRigVMVariantWidget.h
+++ b/Engine/Plugins/Runtime/RigVM/Source/RigVMEditor/Public/Widgets/SRigVMVariantWidget.h
@@ -5,6 +5,7 @@
 #include "Widgets/Layout/SBox.h"
 #include "Widgets/SBoxPanel.h"
 #include "RigVMCore/RigVMVariant.h"
+#include "Widgets/SRigVMVariantTagWidget.h"
 
 DECLARE_DELEGATE_OneParam(FRigVMVariantWidget_OnVariantChanged, const FRigVMVariant&);
 DECLARE_DELEGATE_RetVal_OneParam(TSharedPtr<SWidget>, FRigVMVariantWidget_OnCreateVariantRefRow, const FRigVMVariantRef&);
@@ -16,6 +17,8 @@ public:
 	
 	SLATE_BEGIN_ARGS(SRigVMVariantWidget)
 		: _MaxVariantRefListHeight(200.f)
+		, _CanAddTags(false)
+		, _EnableTagContextMenu(false)
 	{
 	}
 	SLATE_ATTRIBUTE(FRigVMVariant, Variant)
@@ -24,6 +27,11 @@ public:
 	SLATE_EVENT(FRigVMVariantWidget_OnCreateVariantRefRow, OnCreateVariantRefRow);
 	SLATE_EVENT(FRigVMVariantWidget_OnBrowseVariantRef, OnBrowseVariantRef)
 	SLATE_ATTRIBUTE(float, MaxVariantRefListHeight)
+	SLATE_EVENT(FRigVMVariant_OnGetTags, OnGetTags)
+	SLATE_EVENT(FRigVMVariant_OnAddTag, OnAddTag)
+	SLATE_EVENT(FRigVMVariant_OnRemoveTag, OnRemoveTag)
+	SLATE_ATTRIBUTE(bool, CanAddTags)
+	SLATE_ATTRIBUTE(bool, EnableTagContextMenu)
 	SLATE_END_ARGS()
 
 	SRigVMVariantWidget();
@@ -40,11 +48,13 @@ private:
 	void RebuildVariantRefList();
 
 	TAttribute<FRigVMVariant> VariantAttribute;
-	TAttribute<TArray<FRigVMVariantRef>> VariantRefsAttribute;
 	FRigVMVariantWidget_OnVariantChanged OnVariantChanged;
-	FRigVMVariantWidget_OnCreateVariantRefRow OnCreateVariantRefRow;
-    FRigVMVariantWidget_OnBrowseVariantRef OnBrowseVariantRef;
 
+	TSharedPtr<SRigVMVariantTagWidget> TagWidget;
+
+	TAttribute<TArray<FRigVMVariantRef>> VariantRefsAttribute;
+	FRigVMVariantWidget_OnCreateVariantRefRow OnCreateVariantRefRow;
+	FRigVMVariantWidget_OnBrowseVariantRef OnBrowseVariantRef;
 	TArray<FRigVMVariantRef> VariantRefs;
 	uint32 VariantRefHash;
 	TSharedPtr<SVerticalBox> VariantRefListBox;
