RigVM: Variant Tag Front End - fix horizontal tag list layout

#rnx

diff --git a/Engine/Plugins/Runtime/RigVM/Source/RigVMEditor/Private/Widgets/SRigVMVariantTagWidget.cpp b/Engine/Plugins/Runtime/RigVM/Source/RigVMEditor/Private/Widgets/SRigVMVariantTagWidget.cpp
--- a/Engine/Plugins/Runtime/RigVM/Source/RigVMEditor/Private/Widgets/SRigVMVariantTagWidget.cpp
+++ b/Engine/Plugins/Runtime/RigVM/Source/RigVMEditor/Private/Widgets/SRigVMVariantTagWidget.cpp
@@ -59,7 +59,7 @@ void SRigVMVariantCapsule::Construct(
 				.ColorAndOpacity(this, &SRigVMVariantCapsule::GetColor)
 			]
 			+SHorizontalBox::Slot()
-			.HAlign(HAlign_Fill)
+			.HAlign(HAlign_Left)
 			.VAlign(VAlign_Center)
 			.Padding(FMargin(4,1,4,1))
 			.AutoWidth()
@@ -294,6 +294,7 @@ void SRigVMVariantTagWidget::UpdateCapsules()
 		{
 			HorizontalCapsuleBox->AddSlot()
 			.Padding(0, 0, 3, 0)
+			.AutoWidth()
 			[
 				Capsule
 			];
