Control Rig: Fix context menu not repopulating after deleting a RigVMBlueprint
#rnx

diff --git a/Engine/Plugins/Runtime/RigVM/Source/RigVMEditor/Private/Editor/RigVMEditor.cpp b/Engine/Plugins/Runtime/RigVM/Source/RigVMEditor/Private/Editor/RigVMEditor.cpp
--- a/Engine/Plugins/Runtime/RigVM/Source/RigVMEditor/Private/Editor/RigVMEditor.cpp
+++ b/Engine/Plugins/Runtime/RigVM/Source/RigVMEditor/Private/Editor/RigVMEditor.cpp
@@ -164,6 +164,18 @@ void FRigVMEditor::InitRigVMEditor(const EToolkitMode::Type Mode, const TSharedP
 	InRigVMBlueprint->InitializeModelIfRequired();
 
 	CommonInitialization(Blueprints, false);
+
+	// If the class actions have not been populated, refresh them
+	{
+		UClass* ActionKey = InRigVMBlueprint->GetClass();
+		FBlueprintActionDatabase& ActionDatabase = FBlueprintActionDatabase::Get();
+		FBlueprintActionDatabase::FActionRegistry const& ActionRegistry = ActionDatabase.GetAllActions();
+		if (!ActionRegistry.Contains(ActionKey) ||
+			ActionRegistry.FindChecked(ActionKey).IsEmpty())
+		{
+			ActionDatabase.RefreshClassActions(ActionKey);
+		}
+	}
 	
 	// user-defined-struct can change even after load
 	// refresh the models such that pins are updated to match
