RigVM: Exclude Verse Enums and Structs from registry

#rnx

diff --git a/Engine/Plugins/Runtime/RigVM/Source/RigVM/Private/RigVMCore/RigVMRegistry.cpp b/Engine/Plugins/Runtime/RigVM/Source/RigVM/Private/RigVMCore/RigVMRegistry.cpp
--- a/Engine/Plugins/Runtime/RigVM/Source/RigVM/Private/RigVMCore/RigVMRegistry.cpp
+++ b/Engine/Plugins/Runtime/RigVM/Source/RigVM/Private/RigVMCore/RigVMRegistry.cpp
@@ -372,7 +372,11 @@ void FRigVMRegistry_NoLock::RefreshEngineTypes_NoLock()
 		// if this is a C++ type - skip it
 		if(ScriptStruct->IsA<UUserDefinedStruct>() || ScriptStruct->IsChildOf(FRigVMExecuteContext::StaticStruct()))
 		{
-			FindOrAddType_NoLock(FRigVMTemplateArgumentType(ScriptStruct), false);
+			// this check for example makes sure we don't add structs defined in verse
+			if(IsAllowedType_NoLock(ScriptStruct))
+			{
+				FindOrAddType_NoLock(FRigVMTemplateArgumentType(ScriptStruct), false);
+			}
 		}
 		else if (ScriptStruct != FRigVMDispatchFactory::StaticStruct() &&
 				 ScriptStruct->IsChildOf(FRigVMDispatchFactory::StaticStruct()))
@@ -1269,7 +1273,24 @@ bool FRigVMRegistry_NoLock::IsAllowedType_NoLock(const FProperty* InProperty) co
 
 bool FRigVMRegistry_NoLock::IsAllowedType_NoLock(const UEnum* InEnum) const
 {
-	return InEnum && !InEnum->HasAnyFlags(DisallowedFlags()) && InEnum->HasAllFlags(NeededFlags());
+	if(!InEnum)
+	{
+		return false;
+	}
+	
+	// disallow verse based enums for now
+	if (FPackageName::IsVersePackage(InEnum->GetPackage()->GetName()))
+	{
+		return false;
+	}
+
+	static const FName VerseEnumName(TEXT("VerseEnum"));
+	if(IsTypeOfByName(InEnum, VerseEnumName))
+	{
+		return false;
+	}
+
+	return !InEnum->HasAnyFlags(DisallowedFlags()) && InEnum->HasAllFlags(NeededFlags());
 }
 
 bool FRigVMRegistry_NoLock::IsAllowedType_NoLock(const UStruct* InStruct) const
@@ -1287,6 +1308,18 @@ bool FRigVMRegistry_NoLock::IsAllowedType_NoLock(const UStruct* InStruct) const
 	{
 		return false;
 	}
+	
+	// disallow verse data structures for now
+	if (FPackageName::IsVersePackage(InStruct->GetPackage()->GetName()))
+	{
+		return false;
+	}
+	
+	static const FName VerseStructName(TEXT("VerseStruct"));
+	if(IsTypeOfByName(InStruct, VerseStructName))
+	{
+		return false;
+	}
 
 	// allow all user defined structs since they can always be changed to be compliant with RigVM restrictions
 	if (InStruct->IsA<UUserDefinedStruct>())
@@ -1327,9 +1360,41 @@ bool FRigVMRegistry_NoLock::IsAllowedType_NoLock(const UClass* InClass) const
 		return false;
 	}
 
+	// disallow verse based classes for now
+	if (FPackageName::IsVersePackage(InClass->GetPackage()->GetName()))
+	{
+		return false;
+	}
+
+	static const FName VerseClassName(TEXT("VerseClass"));
+	if(IsTypeOfByName(InClass, VerseClassName))
+	{
+		return false;
+	}
+
 	return AllowedClasses.Contains(InClass);
 }
 
+bool FRigVMRegistry_NoLock::IsTypeOfByName(const UObject* InObject, const FName& InName)
+{
+	if(!InObject || InName.IsNone())
+	{
+		return false;
+	}
+	
+	const UClass* Class = InObject->GetClass();
+	while(Class)
+	{
+		if(Class->GetFName().IsEqual(InName, ENameCase::CaseSensitive))
+		{
+			return true;
+		}
+		Class = Class->GetSuperClass();
+	}
+	
+	return false;
+}
+
 void FRigVMRegistry_NoLock::Register_NoLock(const TCHAR* InName, FRigVMFunctionPtr InFunctionPtr, UScriptStruct* InStruct, const TArray<FRigVMFunctionArgument>& InArguments)
 {
 	if (FindFunction_NoLock(InName) != nullptr)
diff --git a/Engine/Plugins/Runtime/RigVM/Source/RigVM/Public/RigVMCore/RigVMRegistry.h b/Engine/Plugins/Runtime/RigVM/Source/RigVM/Public/RigVMCore/RigVMRegistry.h
--- a/Engine/Plugins/Runtime/RigVM/Source/RigVM/Public/RigVMCore/RigVMRegistry.h
+++ b/Engine/Plugins/Runtime/RigVM/Source/RigVM/Public/RigVMCore/RigVMRegistry.h
@@ -354,6 +354,7 @@ protected:
 	bool IsAllowedType_NoLock(const UEnum* InEnum) const;
 	bool IsAllowedType_NoLock(const UStruct* InStruct) const;
 	bool IsAllowedType_NoLock(const UClass* InClass) const;
+	static bool IsTypeOfByName(const UObject* InObject, const FName& InName);
 
 	void RegisterTypeInCategory_NoLock(const FRigVMTemplateArgument::ETypeCategory InCategory, const TRigVMTypeIndex InTypeIndex);
 	void PropagateTypeAddedToCategory_NoLock(const FRigVMTemplateArgument::ETypeCategory InCategory, const TRigVMTypeIndex InTypeIndex);
