RigVM: support enum to int32 as a cast property

Also added lots of code to debug types again - when type indices get copied throughout registry and other services we lose the name - which can make it hard to debug.

#rnx

diff --git a/Engine/Plugins/Runtime/RigVM/Source/RigVM/Private/RigVMCore/RigVMRegistry.cpp b/Engine/Plugins/Runtime/RigVM/Source/RigVM/Private/RigVMCore/RigVMRegistry.cpp
--- a/Engine/Plugins/Runtime/RigVM/Source/RigVM/Private/RigVMCore/RigVMRegistry.cpp
+++ b/Engine/Plugins/Runtime/RigVM/Source/RigVM/Private/RigVMCore/RigVMRegistry.cpp
@@ -603,7 +603,7 @@ TRigVMTypeIndex FRigVMRegistry_NoLock::FindOrAddType_NoLock(const FRigVMTemplate
 
 		// update the categories first then propagate to TemplatesPerCategory once all categories up to date
 		TArray<TPair<FRigVMTemplateArgument::ETypeCategory, int32>> ToPropagate;
-		auto RegisterNewType = [&](FRigVMTemplateArgument::ETypeCategory InCategory, int32 NewIndex)
+		auto RegisterNewType = [&](FRigVMTemplateArgument::ETypeCategory InCategory, TRigVMTypeIndex NewIndex)
 		{
 			RegisterTypeInCategory_NoLock(InCategory, NewIndex);
 			ToPropagate.Emplace(InCategory, NewIndex);
@@ -1210,7 +1210,16 @@ TRigVMTypeIndex FRigVMRegistry_NoLock::GetArrayTypeFromBaseTypeIndex_NoLock(TRig
 {
 	if(ensure(Types.IsValidIndex(InTypeIndex)))
 	{
+#if UE_RIGVM_DEBUG_TYPEINDEX
+		TRigVMTypeIndex Result = Types[InTypeIndex].ArrayTypeIndex;
+		if(!InTypeIndex.Name.IsNone())
+		{
+			Result.Name = *RigVMTypeUtils::ArrayTypeFromBaseType(InTypeIndex.Name.ToString());
+		}
+		return Result;
+#else
 		return Types[InTypeIndex].ArrayTypeIndex;
+#endif
 	}
 	return INDEX_NONE;
 }
@@ -1219,7 +1228,16 @@ TRigVMTypeIndex FRigVMRegistry_NoLock::GetBaseTypeFromArrayTypeIndex_NoLock(TRig
 {
 	if(ensure(Types.IsValidIndex(InTypeIndex)))
 	{
+#if UE_RIGVM_DEBUG_TYPEINDEX
+		TRigVMTypeIndex Result = Types[InTypeIndex].BaseTypeIndex;
+		if(!InTypeIndex.Name.IsNone())
+		{
+			Result.Name = *RigVMTypeUtils::BaseTypeFromArrayType(InTypeIndex.Name.ToString());
+		}
+		return Result;
+#else
 		return Types[InTypeIndex].BaseTypeIndex;
+#endif
 	}
 	return INDEX_NONE;
 }
diff --git a/Engine/Plugins/Runtime/RigVM/Source/RigVM/Private/RigVMCore/RigVMTemplate.cpp b/Engine/Plugins/Runtime/RigVM/Source/RigVM/Private/RigVMCore/RigVMTemplate.cpp
--- a/Engine/Plugins/Runtime/RigVM/Source/RigVM/Private/RigVMCore/RigVMTemplate.cpp
+++ b/Engine/Plugins/Runtime/RigVM/Source/RigVM/Private/RigVMCore/RigVMTemplate.cpp
@@ -220,7 +220,7 @@ bool FRigVMTemplateArgument::SupportsTypeIndex_NoLock(TRigVMTypeIndex InTypeInde
 		}
 	}
 
-	const TArray<int32>& Permutations = GetPermutations(InTypeIndex);
+	const TArray<int32>& Permutations = GetPermutations_NoLock(InTypeIndex);
 	if (!Permutations.IsEmpty())
 	{
 		if(OutTypeIndex)
@@ -234,7 +234,7 @@ bool FRigVMTemplateArgument::SupportsTypeIndex_NoLock(TRigVMTypeIndex InTypeInde
 	const TArray<TRigVMTypeIndex>& CompatibleTypes = Registry.GetCompatibleTypes_NoLock(InTypeIndex);
 	for (const TRigVMTypeIndex& CompatibleTypeIndex : CompatibleTypes)
 	{
-		const TArray<int32>& CompatiblePermutations = GetPermutations(CompatibleTypeIndex);
+		const TArray<int32>& CompatiblePermutations = GetPermutations_NoLock(CompatibleTypeIndex);
 		if (!CompatiblePermutations.IsEmpty())
 		{
 			if(OutTypeIndex)
@@ -348,7 +348,13 @@ FRigVMTemplateArgument::EArrayType FRigVMTemplateArgument::GetArrayType_NoLock()
 	return EArrayType_Invalid;
 }
 
-const TArray<int32>& FRigVMTemplateArgument::GetPermutations(const TRigVMTypeIndex InType) const
+const TArray<int32>& FRigVMTemplateArgument::GetPermutations(const TRigVMTypeIndex InType, bool bLockRegistry) const
+{
+	FRigVMRegistryReadLock _(bLockRegistry);
+	return GetPermutations_NoLock(InType);
+}
+
+const TArray<int32>& FRigVMTemplateArgument::GetPermutations_NoLock(const TRigVMTypeIndex InType) const
 {
 	if (const TArray<int32>* Found = TypeToPermutations.Find(InType))
 	{
@@ -585,6 +591,18 @@ TArray<TRigVMTypeIndex> FRigVMTemplateArgument::GetSupportedTypeIndices(const TA
 {
 	const FRigVMRegistryReadLock _;
 
+#if UE_RIGVM_DEBUG_TYPEINDEX
+	const FRigVMRegistry_NoLock& Registry = _.GetRegistry();
+	auto UpdateTypeIndex = [&Registry](const TRigVMTypeIndex& TypeIndex) -> TRigVMTypeIndex
+	{
+		if(TypeIndex.Name.IsNone())
+		{
+			return Registry.GetTypeIndex_NoLock(Registry.GetType_NoLock(TypeIndex));
+		}
+		return TypeIndex;
+	};
+#endif
+
 	TArray<TRigVMTypeIndex> SupportedTypes;
 	if(InPermutationIndices.IsEmpty())
 	{
@@ -593,7 +611,11 @@ TArray<TRigVMTypeIndex> FRigVMTemplateArgument::GetSupportedTypeIndices(const TA
 			// INDEX_NONE indicates deleted permutation
 			if (TypeIndex != INDEX_NONE)
 			{
+#if UE_RIGVM_DEBUG_TYPEINDEX
+				SupportedTypes.AddUnique(UpdateTypeIndex(TypeIndex));
+#else
 				SupportedTypes.AddUnique(TypeIndex);
+#endif
 			}
 			return true;
 		});
@@ -603,10 +625,14 @@ TArray<TRigVMTypeIndex> FRigVMTemplateArgument::GetSupportedTypeIndices(const TA
 		for(const int32 PermutationIndex : InPermutationIndices)
 		{
 			// INDEX_NONE indicates deleted permutation
-			const TRigVMTypeIndex Type = GetTypeIndex_NoLock(PermutationIndex);
-			if (Type != INDEX_NONE)
+			const TRigVMTypeIndex TypeIndex = GetTypeIndex_NoLock(PermutationIndex);
+			if (TypeIndex != INDEX_NONE)
 			{
-				SupportedTypes.AddUnique(Type);
+#if UE_RIGVM_DEBUG_TYPEINDEX
+				SupportedTypes.AddUnique(UpdateTypeIndex(TypeIndex));
+#else
+				SupportedTypes.AddUnique(TypeIndex);
+#endif
 			}
 		}
 	}
@@ -1610,7 +1636,13 @@ const FRigVMFunction* FRigVMTemplate::GetOrCreatePermutation_NoLock(int32 InInde
 	return nullptr;
 }
 
-bool FRigVMTemplate::ContainsPermutation(const FRigVMFunction* InPermutation) const
+bool FRigVMTemplate::ContainsPermutation(const FRigVMFunction* InPermutation, bool bLockRegistry) const
+{
+	FRigVMRegistryReadLock _(bLockRegistry);
+	return ContainsPermutation_NoLock(InPermutation);
+}
+
+bool FRigVMTemplate::ContainsPermutation_NoLock(const FRigVMFunction* InPermutation) const
 {
 	return FindPermutation(InPermutation) != INDEX_NONE;
 }
@@ -1796,7 +1828,13 @@ uint32 FRigVMTemplate::GetTypesHashFromTypes(const FTypeMap& InTypes) const
 	return TypeHash;
 }
 
-bool FRigVMTemplate::ContainsPermutation(const FTypeMap& InTypes) const
+bool FRigVMTemplate::ContainsPermutation(const FTypeMap& InTypes, bool bLockRegistry) const
+{
+	const FRigVMRegistryReadLock Readlock(bLockRegistry);
+	return ContainsPermutation_NoLock((InTypes));
+}
+
+bool FRigVMTemplate::ContainsPermutation_NoLock(const FTypeMap& InTypes) const
 {
 	// If they type map is valid (full description of arguments), then we can rely on
 	// the TypesHashToPermutation cache. Otherwise, we will have to search for a specific permutation
@@ -1812,7 +1850,7 @@ bool FRigVMTemplate::ContainsPermutation(const FTypeMap& InTypes) const
 	{
 		if (const FRigVMTemplateArgument* Argument = FindArgument(Pair.Key))
 		{
-			const TArray<int32>& ArgumentPermutations = Argument->GetPermutations(Pair.Value);
+			const TArray<int32>& ArgumentPermutations = Argument->GetPermutations_NoLock(Pair.Value);
 			if (!ArgumentPermutations.IsEmpty())
 			{
 				// If possible permutations is empty, initialize it
@@ -2100,7 +2138,7 @@ bool FRigVMTemplate::UpdateArgumentTypes()
 					}
 
 					// Find if these types were already registered
-					if (ContainsPermutation(Types))
+					if (ContainsPermutation_NoLock(Types))
 					{
 						return true;
 					}
diff --git a/Engine/Plugins/Runtime/RigVM/Source/RigVM/Private/RigVMFunctions/RigVMDispatch_CastEnum.cpp b/Engine/Plugins/Runtime/RigVM/Source/RigVM/Private/RigVMFunctions/RigVMDispatch_CastEnum.cpp
--- a/Engine/Plugins/Runtime/RigVM/Source/RigVM/Private/RigVMFunctions/RigVMDispatch_CastEnum.cpp
+++ b/Engine/Plugins/Runtime/RigVM/Source/RigVM/Private/RigVMFunctions/RigVMDispatch_CastEnum.cpp
@@ -16,11 +16,8 @@ const TArray<FRigVMTemplateArgumentInfo>& FRigVMDispatch_CastEnumToInt::GetArgum
 			FRigVMTemplateArgument::ETypeCategory_SingleEnumValue
 		};
 		
-		static TArray<FRigVMTemplateArgumentInfo> Infos;
-		Infos.Emplace(ValueName, ERigVMPinDirection::Input, ElementCategories);
-		Infos.Emplace(ResultName, ERigVMPinDirection::Output, RigVMTypeUtils::TypeIndex::Int32);
-		
-		OutInfos = BuildArgumentListFromPrimaryArgument(Infos, ValueName);
+		OutInfos.Emplace(ValueName, ERigVMPinDirection::Input, ElementCategories);
+		OutInfos.Emplace(ResultName, ERigVMPinDirection::Output, RigVMTypeUtils::TypeIndex::Int32);
 	}
 
 	return OutInfos;
@@ -36,18 +33,6 @@ bool FRigVMDispatch_CastEnumToInt::GetPermutationsFromArgumentType(const FName&
 			{ ResultName, RigVMTypeUtils::TypeIndex::Int32 }
 		});
 	}
-	else if (InArgumentName == ResultName && InTypeIndex == RigVMTypeUtils::TypeIndex::Int32)
-	{
-		const TArray<TRigVMTypeIndex>& EnumTypes = FRigVMRegistry_NoLock::GetForRead().GetTypesForCategory_NoLock(FRigVMTemplateArgument::ETypeCategory_SingleEnumValue);
-		for (const TRigVMTypeIndex& Type : EnumTypes)
-		{
-			OutPermutations.Add(
-	{
-				{ ValueName, Type },
-				{ ResultName, InTypeIndex }
-			});
-		}
-	}
 	return !OutPermutations.IsEmpty();
 }
 
@@ -110,11 +95,8 @@ const TArray<FRigVMTemplateArgumentInfo>& FRigVMDispatch_CastIntToEnum::GetArgum
 			FRigVMTemplateArgument::ETypeCategory_SingleEnumValue
 		};
 		
-		static TArray<FRigVMTemplateArgumentInfo> Infos;
-		Infos.Emplace(ValueName, ERigVMPinDirection::Input, RigVMTypeUtils::TypeIndex::Int32);
-		Infos.Emplace(ResultName, ERigVMPinDirection::Output, ElementCategories);
-		
-		OutInfos = BuildArgumentListFromPrimaryArgument(Infos, ResultName);
+		OutInfos.Emplace(ValueName, ERigVMPinDirection::Input, RigVMTypeUtils::TypeIndex::Int32);
+		OutInfos.Emplace(ResultName, ERigVMPinDirection::Output, ElementCategories);
 	}
 
 	return OutInfos;
@@ -130,18 +112,6 @@ bool FRigVMDispatch_CastIntToEnum::GetPermutationsFromArgumentType(const FName&
 			{ ResultName, InTypeIndex }
 		});
 	}
-	else if (InArgumentName == ValueName && InTypeIndex == RigVMTypeUtils::TypeIndex::Int32)
-	{
-		const TArray<TRigVMTypeIndex>& EnumTypes = FRigVMRegistry_NoLock::GetForRead().GetTypesForCategory_NoLock(FRigVMTemplateArgument::ETypeCategory_SingleEnumValue);
-		for (const TRigVMTypeIndex& Type : EnumTypes)
-		{
-			OutPermutations.Add(
-	{
-				{ ValueName, InTypeIndex },
-				{ ResultName, Type }
-			});
-		}
-	}
 	return !OutPermutations.IsEmpty();
 }
 
diff --git a/Engine/Plugins/Runtime/RigVM/Source/RigVM/Public/RigVMCore/RigVMTemplate.h b/Engine/Plugins/Runtime/RigVM/Source/RigVM/Public/RigVMCore/RigVMTemplate.h
--- a/Engine/Plugins/Runtime/RigVM/Source/RigVM/Public/RigVMCore/RigVMTemplate.h
+++ b/Engine/Plugins/Runtime/RigVM/Source/RigVM/Public/RigVMCore/RigVMTemplate.h
@@ -279,7 +279,8 @@ struct RIGVM_API FRigVMTemplateArgument
 	RIGVM_API friend uint32 GetTypeHash_NoLock(const FRigVMTemplateArgument& InArgument);
 
 	// Get the map of types to permutation indices
-	const TArray<int32>& GetPermutations(const TRigVMTypeIndex InType) const;
+	const TArray<int32>& GetPermutations(const TRigVMTypeIndex InType, bool bLockRegistry = true) const;
+	const TArray<int32>& GetPermutations_NoLock(const TRigVMTypeIndex InType) const;
 	void InvalidatePermutations(const TRigVMTypeIndex InType);
 
 protected:
@@ -448,7 +449,8 @@ public:
 	const FRigVMFunction* GetOrCreatePermutation(int32 InIndex, bool bLockRegistry = true);
 
 	// returns true if a given function is a permutation of this template
-	bool ContainsPermutation(const FRigVMFunction* InPermutation) const;
+	bool ContainsPermutation(const FRigVMFunction* InPermutation, bool bLockRegistry = true) const;
+	bool ContainsPermutation_NoLock(const FRigVMFunction* InPermutation) const;
 
 	// returns the index of the permutation within the template of a given function (or INDEX_NONE)
 	int32 FindPermutation(const FRigVMFunction* InPermutation) const;
@@ -467,7 +469,8 @@ public:
 	uint32 GetTypesHashFromTypes(const FTypeMap& InTypes) const;
 
 	// returns true if the template was able to resolve to at least one permutation
-	bool ContainsPermutation(const FTypeMap& InTypes) const;
+	bool ContainsPermutation(const FTypeMap& InTypes, bool bLockRegistry = true) const;
+	bool ContainsPermutation_NoLock(const FTypeMap& InTypes) const;
 
 	// returns true if the template can resolve an argument to a new type
 	bool ResolveArgument(const FName& InArgumentName, const TRigVMTypeIndex InTypeIndex, FTypeMap& InOutTypes, bool bLockRegistry = true) const;
diff --git a/Engine/Plugins/Runtime/RigVM/Source/RigVM/Public/RigVMDefines.h b/Engine/Plugins/Runtime/RigVM/Source/RigVM/Public/RigVMDefines.h
--- a/Engine/Plugins/Runtime/RigVM/Source/RigVM/Public/RigVMDefines.h
+++ b/Engine/Plugins/Runtime/RigVM/Source/RigVM/Public/RigVMDefines.h
@@ -7,29 +7,37 @@
 #pragma once
 
 #ifndef UE_RIGVM_UCLASS_BASED_STORAGE_DISABLED
-#define UE_RIGVM_UCLASS_BASED_STORAGE_DISABLED 0
+	#define UE_RIGVM_UCLASS_BASED_STORAGE_DISABLED 0
 #endif
 
 #ifndef UE_RIGVM_UOBJECT_PROPERTIES_ENABLED
-#define UE_RIGVM_UOBJECT_PROPERTIES_ENABLED 1
+	#define UE_RIGVM_UOBJECT_PROPERTIES_ENABLED 1
 #endif
 
 #ifndef UE_RIGVM_UINTERFACE_PROPERTIES_ENABLED
-#define UE_RIGVM_UINTERFACE_PROPERTIES_ENABLED 1
+	#define UE_RIGVM_UINTERFACE_PROPERTIES_ENABLED 1
 #endif
 
 #ifndef UE_RIGVM_AGGREGATE_NODES_ENABLED
-#define UE_RIGVM_AGGREGATE_NODES_ENABLED 1
+	#define UE_RIGVM_AGGREGATE_NODES_ENABLED 1
 #endif
 
 #ifndef UE_RIGVM_DEBUG_TYPEINDEX
-#define UE_RIGVM_DEBUG_TYPEINDEX 1
+	#define UE_RIGVM_DEBUG_TYPEINDEX 1
+#endif
+
+// in shipping always disable debugging the type index
+#ifdef UE_BUILD_SHIPPING
+	#if UE_BUILD_SHIPPING
+		#undef UE_RIGVM_DEBUG_TYPEINDEX
+		#define UE_RIGVM_DEBUG_TYPEINDEX 0
+	#endif
 #endif
 
 #ifndef UE_RIGVM_DEBUG_EXECUTION
-#define UE_RIGVM_DEBUG_EXECUTION 0
+	#define UE_RIGVM_DEBUG_EXECUTION 0
 #endif
 
 #ifndef UE_RIGVMCONTROLLER_VERBOSE_REPOPULATE
-#define UE_RIGVMCONTROLLER_VERBOSE_REPOPULATE 0
+	#define UE_RIGVMCONTROLLER_VERBOSE_REPOPULATE 0
 #endif
diff --git a/Engine/Plugins/Runtime/RigVM/Source/RigVMDeveloper/Private/RigVMDeveloperTypeUtils.cpp b/Engine/Plugins/Runtime/RigVM/Source/RigVMDeveloper/Private/RigVMDeveloperTypeUtils.cpp
--- a/Engine/Plugins/Runtime/RigVM/Source/RigVMDeveloper/Private/RigVMDeveloperTypeUtils.cpp
+++ b/Engine/Plugins/Runtime/RigVM/Source/RigVMDeveloper/Private/RigVMDeveloperTypeUtils.cpp
@@ -654,20 +654,14 @@ const FRigVMFunction* RigVMTypeUtils::GetCastForTypeIndices(const TRigVMTypeInde
 
 			if(SourceArgument && TargetArgument)
 			{
-				const TArray<int32>& SourcePermutations = SourceArgument->GetPermutations(InSourceTypeIndex);
-				const TArray<int32>& TargetPermutations = TargetArgument->GetPermutations(InTargetTypeIndex);
-				if(!SourcePermutations.IsEmpty() && !TargetPermutations.IsEmpty())
+				FRigVMTemplateTypeMap Types;
+				Types.Add(SourceArgument->GetName(), InSourceTypeIndex);
+				Types.Add(TargetArgument->GetName(), InTargetTypeIndex);
+
+				int32 PermutationIndex = INDEX_NONE;
+				if(CastTemplate->FullyResolve(Types, PermutationIndex))
 				{
-					for(int32 SourceIndex = 0, SourceCount = SourcePermutations.Num(); SourceIndex < SourceCount; ++SourceIndex)
-					{
-						for(int32 TargetIndex = 0, TargetCount = TargetPermutations.Num(); TargetIndex < TargetCount; ++TargetIndex)
-						{
-							if(SourcePermutations[SourceIndex] == TargetPermutations[TargetIndex])
-							{
-								return const_cast<FRigVMTemplate*>(CastTemplate)->GetOrCreatePermutation(SourcePermutations[SourceIndex]);
-							}
-						}
-					}
+					return const_cast<FRigVMTemplate*>(CastTemplate)->GetOrCreatePermutation(PermutationIndex);
 				}
 			}
 		}
