Control Rig: Fix crash when displaying transform metadata in details panel
#rnx

diff --git a/Engine/Plugins/Runtime/RigVM/Source/RigVMEditor/Private/Editor/RigVMGraphDetailCustomization.cpp b/Engine/Plugins/Runtime/RigVM/Source/RigVMEditor/Private/Editor/RigVMGraphDetailCustomization.cpp
--- a/Engine/Plugins/Runtime/RigVM/Source/RigVMEditor/Private/Editor/RigVMGraphDetailCustomization.cpp
+++ b/Engine/Plugins/Runtime/RigVM/Source/RigVMEditor/Private/Editor/RigVMGraphDetailCustomization.cpp
@@ -2342,15 +2342,22 @@ void FRigVMGraphMathTypeDetailCustomization::ConfigureTransformWidgetArgs(TShare
 
 	static TransformType Identity = TransformType::Identity;
 
-	UObject* DefaultObject = !ObjectsBeingCustomized.IsEmpty() ? ObjectsBeingCustomized[0]->GetClass()->GetDefaultObject() :
-	!StructsBeingCustomized.IsEmpty() ? StructsBeingCustomized[0]->GetStruct()->GetClass() : nullptr;
-
-	if (!DefaultObject)
+	uint8* ContainerMemory = nullptr;
+	TSharedPtr<FStructOnScope> DefaultStruct = nullptr;
+	if (!ObjectsBeingCustomized.IsEmpty())
+	{
+		ContainerMemory = (uint8*) ObjectsBeingCustomized[0]->GetClass()->GetDefaultObject();
+	}
+	else if(!StructsBeingCustomized.IsEmpty())
+	{
+		DefaultStruct = MakeShareable(new FStructOnScope(StructsBeingCustomized[0]->GetStruct()));
+		ContainerMemory = DefaultStruct->GetStructMemory();
+	}
+	if (!ContainerMemory)
 	{
 		return;
 	}
-	
-	TransformType DefaultValue = ContainerMemoryBlockToValueRef<TransformType>((uint8*)DefaultObject, Identity, PropertyChain, PropertyArrayIndices);
+	TransformType DefaultValue = ContainerMemoryBlockToValueRef<TransformType>(ContainerMemory, Identity, PropertyChain, PropertyArrayIndices);
 	
 	WidgetArgs.DiffersFromDefault_Lambda([this, InPropertyHandle, DefaultValue](ESlateTransformComponent::Type InTransformComponent) -> bool
 	{
