RigVM: Enum dispatches to deal with both FEnumProperty and FByteProperty

#rnx

diff --git a/Engine/Plugins/Runtime/RigVM/Source/RigVM/Private/RigVMFunctions/RigVMDispatch_CastEnum.cpp b/Engine/Plugins/Runtime/RigVM/Source/RigVM/Private/RigVMFunctions/RigVMDispatch_CastEnum.cpp
--- a/Engine/Plugins/Runtime/RigVM/Source/RigVM/Private/RigVMFunctions/RigVMDispatch_CastEnum.cpp
+++ b/Engine/Plugins/Runtime/RigVM/Source/RigVM/Private/RigVMFunctions/RigVMDispatch_CastEnum.cpp
@@ -52,24 +52,40 @@ FText FRigVMDispatch_CastEnumToInt::GetNodeTooltip(const FRigVMTemplateTypeMap&
 
 void FRigVMDispatch_CastEnumToInt::Execute(FRigVMExtendedExecuteContext& InContext, FRigVMMemoryHandleArray Handles, FRigVMPredicateBranchArray RigVMBranches)
 {
-	const FEnumProperty* ValueProperty = CastFieldChecked<FEnumProperty>(Handles[0].GetProperty());
+	const FProperty* ValueProperty = CastFieldChecked<FProperty>(Handles[0].GetProperty());
 	const FProperty* ResultProperty = CastFieldChecked<FProperty>(Handles[1].GetProperty());
-
 	if (!ResultProperty || !ValueProperty)
 	{
 		return;
 	}
 
-	uint8* ValuePtr = Handles[0].GetData();
+	const uint8* ValuePtr = Handles[0].GetData();
 	int32* ResultPtr = (int32*)Handles[1].GetData();
-
 	if (ValuePtr == nullptr || ResultPtr == nullptr)
 	{
 		return;
 	}
 	
-	*ResultPtr = 0;
-	ValueProperty->CopyCompleteValue(ResultPtr, ValuePtr);
+	const FNumericProperty* NumericProperty = nullptr;
+	if(const FEnumProperty* EnumProperty = CastField<FEnumProperty>(ValueProperty))
+	{
+		NumericProperty = EnumProperty->GetUnderlyingProperty();
+	}
+	else if(const FByteProperty* ByteProperty = CastField<FByteProperty>(ValueProperty))
+	{
+		NumericProperty = ByteProperty;
+	}
+
+	if(NumericProperty == nullptr)
+	{
+		*ResultPtr = 0;
+		checkNoEntry();
+		return;
+	}
+	check(NumericProperty->IsInteger())
+
+	const int64 Value = NumericProperty->GetSignedIntPropertyValue(ValuePtr);
+	*ResultPtr = (int32)Value;
 
 #if WITH_EDITOR
 	if (*ResultPtr == INDEX_NONE)
@@ -132,23 +148,39 @@ FText FRigVMDispatch_CastIntToEnum::GetNodeTooltip(const FRigVMTemplateTypeMap&
 void FRigVMDispatch_CastIntToEnum::Execute(FRigVMExtendedExecuteContext& InContext, FRigVMMemoryHandleArray Handles, FRigVMPredicateBranchArray RigVMBranches)
 {
 	const FProperty* ValueProperty = CastFieldChecked<FProperty>(Handles[0].GetProperty());
-	const FEnumProperty* ResultProperty = CastFieldChecked<FEnumProperty>(Handles[1].GetProperty());
-
+	const FProperty* ResultProperty = CastFieldChecked<FProperty>(Handles[1].GetProperty());
 	if (!ResultProperty || !ValueProperty)
 	{
 		return;
 	}
 
-	int32* ValuePtr = (int32*)Handles[0].GetData();
+	const int32* ValuePtr = (int32*)Handles[0].GetData();
 	uint8* ResultPtr = Handles[1].GetData();
-
 	if (ValuePtr == nullptr || ResultPtr == nullptr)
 	{
 		return;
 	}
-	
-	*ResultPtr = 0;
-	ResultProperty->CopyCompleteValue(ResultPtr, ValuePtr);
+
+	const FNumericProperty* NumericProperty = nullptr;
+	if(const FEnumProperty* EnumProperty = CastField<FEnumProperty>(ResultProperty))
+	{
+		NumericProperty = EnumProperty->GetUnderlyingProperty();
+	}
+	else if(const FByteProperty* ByteProperty = CastField<FByteProperty>(ResultProperty))
+	{
+		NumericProperty = ByteProperty;
+	}
+
+	if(NumericProperty == nullptr)
+	{
+		*ResultPtr = 0;
+		checkNoEntry();
+		return;
+	}
+	check(NumericProperty->IsInteger());
+
+	const int64 Value = (int32)*ValuePtr;
+	NumericProperty->SetIntPropertyValue(ResultPtr, Value);
 	
 #if WITH_EDITOR
 	if (*ResultPtr == INDEX_NONE)
