Control Rig: Fix crash when copy pasting a function
#rnx

diff --git a/Engine/Plugins/Runtime/RigVM/Source/RigVMDeveloper/Private/RigVMModel/RigVMController.cpp b/Engine/Plugins/Runtime/RigVM/Source/RigVMDeveloper/Private/RigVMModel/RigVMController.cpp
--- a/Engine/Plugins/Runtime/RigVM/Source/RigVMDeveloper/Private/RigVMModel/RigVMController.cpp
+++ b/Engine/Plugins/Runtime/RigVM/Source/RigVMDeveloper/Private/RigVMModel/RigVMController.cpp
@@ -4133,14 +4133,16 @@ protected:
 		}
 		else if (URigVMInjectionInfo* Injection = Cast<URigVMInjectionInfo>(CreatedObject))
 		{
-			URigVMNode* InjectedNode = Injection->Node;
-			ProcessConstructedObject(InjectedNode);
-
-			FName NewName = InjectedNode->GetFName();
-			UpdateObjectName(URigVMNode::StaticClass(), Injection->GetGraph(), NewName);
-			Controller->RenameObject(InjectedNode, *NewName.ToString(), nullptr);
-			Injection->InputPin = Injection->InputPin ? Injection->Node->FindPin(Injection->InputPin->GetName()) : nullptr;
-			Injection->OutputPin = Injection->OutputPin ? Injection->Node->FindPin(Injection->OutputPin->GetName()) : nullptr;
+			if (URigVMNode* InjectedNode = Injection->Node)
+			{
+				ProcessConstructedObject(InjectedNode);
+
+				FName NewName = InjectedNode->GetFName();
+				UpdateObjectName(URigVMNode::StaticClass(), Injection->GetGraph(), NewName);
+				Controller->RenameObject(InjectedNode, *NewName.ToString(), nullptr);
+				Injection->InputPin = Injection->InputPin ? Injection->Node->FindPin(Injection->InputPin->GetName()) : nullptr;
+				Injection->OutputPin = Injection->OutputPin ? Injection->Node->FindPin(Injection->OutputPin->GetName()) : nullptr;
+			}
 		}
 	}
 };
