RigVM: Fix template argument discovery for filtered types

#rnx

diff --git a/Engine/Plugins/Runtime/RigVM/Source/RigVM/Private/RigVMCore/RigVMTemplate.cpp b/Engine/Plugins/Runtime/RigVM/Source/RigVM/Private/RigVMCore/RigVMTemplate.cpp
--- a/Engine/Plugins/Runtime/RigVM/Source/RigVM/Private/RigVMCore/RigVMTemplate.cpp
+++ b/Engine/Plugins/Runtime/RigVM/Source/RigVM/Private/RigVMCore/RigVMTemplate.cpp
@@ -547,6 +547,20 @@ int32 FRigVMTemplateArgument::GetNumTypes_NoLock() const
 		return TypeIndices.Num();
 	}
 	
+	if (FilterType)
+	{
+		int32 NumTypes = 0;
+		CategoryViews(TypeCategories).ForEachType([this, &NumTypes](const TRigVMTypeIndex& Type) -> bool
+		{
+			if (FilterType(Type))
+			{
+				NumTypes++;
+			}
+			return true;
+		});
+		return NumTypes;
+	}
+
 	return Algo::Accumulate(TypeCategories, 0, [](int32 Sum, const ETypeCategory Category)
 	{
 		return Sum + FRigVMRegistry_NoLock::GetForRead().GetTypesForCategory_NoLock(Category).Num();
@@ -582,6 +596,7 @@ void FRigVMTemplateArgument::ForEachType(TFunction<bool(const TRigVMTypeIndex In
 			}
 			return true;
 		});
+		return;
 	}
 	
 	return CategoryViews(TypeCategories).ForEachType(MoveTemp(InCallback));
