RigVM: Don't clear subpins in ed graph in certain scenarios

#rnx

diff --git a/Engine/Plugins/Runtime/RigVM/Source/RigVMDeveloper/Private/EdGraph/RigVMEdGraphNode.cpp b/Engine/Plugins/Runtime/RigVM/Source/RigVMDeveloper/Private/EdGraph/RigVMEdGraphNode.cpp
--- a/Engine/Plugins/Runtime/RigVM/Source/RigVMDeveloper/Private/EdGraph/RigVMEdGraphNode.cpp
+++ b/Engine/Plugins/Runtime/RigVM/Source/RigVMDeveloper/Private/EdGraph/RigVMEdGraphNode.cpp
@@ -848,7 +848,11 @@ void URigVMEdGraphNode::RemoveGraphSubPins(UEdGraphPin* InParentPin, const TArra
 
 		CachedPins.Remove(ModelSubPin);
 	}
-	InParentPin->SubPins.Reset();
+	
+	InParentPin->SubPins.RemoveAll([InPinsToKeep](const UEdGraphPin* Pin) -> bool
+	{
+		return !InPinsToKeep.Contains(Pin);
+	});
 }
 
 bool URigVMEdGraphNode::ModelPinsChanged(bool bForce)
