Control Rig: FIx serialized public function graphs not getting cooked
#rnx

diff --git a/Engine/Plugins/Animation/ControlRig/Source/ControlRigDeveloper/Private/ControlRigBlueprint.cpp b/Engine/Plugins/Animation/ControlRig/Source/ControlRigDeveloper/Private/ControlRigBlueprint.cpp
--- a/Engine/Plugins/Animation/ControlRig/Source/ControlRigDeveloper/Private/ControlRigBlueprint.cpp
+++ b/Engine/Plugins/Animation/ControlRig/Source/ControlRigDeveloper/Private/ControlRigBlueprint.cpp
@@ -1974,15 +1974,10 @@ void UControlRigBlueprint::PatchPropagateToChildren()
 	}
 }
 
-void UControlRigBlueprint::PatchFunctionsOnLoad()
+void UControlRigBlueprint::GetBackwardsCompatibilityPublicFunctions(TArray<FName>& BackwardsCompatiblePublicFunctions, TMap<URigVMLibraryNode*, FRigVMGraphFunctionHeader>& OldHeaders)
 {
 	URigVMBlueprintGeneratedClass* CRGeneratedClass = GetRigVMBlueprintGeneratedClass();
 	FRigVMGraphFunctionStore& Store = CRGeneratedClass->GraphFunctionStore;
-
-	TMap<URigVMLibraryNode*, FRigVMGraphFunctionHeader> OldHeaders;
-
-	// Backwards compatibility. Store public access in the model
-	TArray<FName> BackwardsCompatiblePublicFunctions;
 	if (GetLinkerCustomVersion(FControlRigObjectVersion::GUID) < FControlRigObjectVersion::StoreFunctionsInGeneratedClass)
 	{
 		for (const FRigVMOldPublicFunctionData& OldPublicFunction : PublicFunctions_DEPRECATED)
@@ -2015,55 +2010,6 @@ void UControlRigBlueprint::PatchFunctionsOnLoad()
 			}
 		}
 	}
-
-	// Lets rebuild the FunctionStore from the model
-	if (FunctionLibrary)
-	{
-		Store.PublicFunctions.Reset();
-		Store.PrivateFunctions.Reset();
-
-		for (URigVMLibraryNode* LibraryNode : FunctionLibrary->GetFunctions())
-		{
-			bool bIsPublic = FunctionLibrary->IsFunctionPublic(LibraryNode->GetFName());
-			if (!bIsPublic)
-			{
-				bIsPublic = BackwardsCompatiblePublicFunctions.Contains(LibraryNode->GetFName());
-				if (bIsPublic)
-				{
-					FunctionLibrary->PublicFunctionNames.Add(LibraryNode->GetFName());
-				}
-			}
-
-			FRigVMGraphFunctionHeader Header = LibraryNode->GetFunctionHeader(CRGeneratedClass);
-			if (FRigVMGraphFunctionHeader* OldHeader = OldHeaders.Find(LibraryNode))
-			{				
-				Header.ExternalVariables = OldHeader->ExternalVariables;
-				Header.Dependencies = OldHeader->Dependencies;
-			}
-
-			const FRigVMVariant* Variant = FunctionLibrary->GetFunctionVariant(LibraryNode->GetFName());
-			if (!Variant)
-			{
-				Header.Variant.Guid = FRigVMVariant::GenerateGUID(Header.LibraryPointer.GetLibraryNodePath());
-				FunctionLibrary->FunctionToVariant.FindOrAdd(Header.Name) = Header.Variant;
-			}
-			else
-			{
-				Header.Variant = *Variant;
-			}
-			
-			Store.AddFunction(Header, bIsPublic);
-			
-		}
-
-		// Update dependencies and external variables if needed
-		for (URigVMLibraryNode* LibraryNode : FunctionLibrary->GetFunctions())
-		{
-			GetRigVMClient()->UpdateExternalVariablesForFunction(LibraryNode);
-			GetRigVMClient()->UpdateDependenciesForFunction(LibraryNode);
-		}
-	}
-
 }
 
 void UControlRigBlueprint::CreateMemberVariablesOnLoad()
diff --git a/Engine/Plugins/Animation/ControlRig/Source/ControlRigDeveloper/Public/ControlRigBlueprint.h b/Engine/Plugins/Animation/ControlRig/Source/ControlRigDeveloper/Public/ControlRigBlueprint.h
--- a/Engine/Plugins/Animation/ControlRig/Source/ControlRigDeveloper/Public/ControlRigBlueprint.h
+++ b/Engine/Plugins/Animation/ControlRig/Source/ControlRigDeveloper/Public/ControlRigBlueprint.h
@@ -326,7 +326,7 @@ private:
 	static TArray<UControlRigBlueprint*> sCurrentlyOpenedRigBlueprints;
 
 	virtual void PathDomainSpecificContentOnLoad() override;
-	virtual void PatchFunctionsOnLoad() override;
+	virtual void GetBackwardsCompatibilityPublicFunctions(TArray<FName>& BackwardsCompatiblePublicFunctions, TMap<URigVMLibraryNode*, FRigVMGraphFunctionHeader>& OldHeaders) override;
 	void PatchRigElementKeyCacheOnLoad();
 	void PatchPropagateToChildren();
 
diff --git a/Engine/Plugins/Runtime/RigVM/Source/RigVMDeveloper/Private/RigVMBlueprint.cpp b/Engine/Plugins/Runtime/RigVM/Source/RigVMDeveloper/Private/RigVMBlueprint.cpp
--- a/Engine/Plugins/Runtime/RigVM/Source/RigVMDeveloper/Private/RigVMBlueprint.cpp
+++ b/Engine/Plugins/Runtime/RigVM/Source/RigVMDeveloper/Private/RigVMBlueprint.cpp
@@ -3717,34 +3717,13 @@ void URigVMBlueprint::PatchFunctionsOnLoad()
 {
 	URigVMBlueprintGeneratedClass* CRGeneratedClass = GetRigVMBlueprintGeneratedClass();
 	FRigVMGraphFunctionStore& Store = CRGeneratedClass->GraphFunctionStore;
-	const URigVMFunctionLibrary* Library = GetLocalFunctionLibrary();
+	URigVMFunctionLibrary* FunctionLibrary = GetLocalFunctionLibrary();
 
 	TMap<URigVMLibraryNode*, FRigVMGraphFunctionHeader> OldHeaders;
 
 	// Backwards compatibility. Store public access in the model
 	TArray<FName> BackwardsCompatiblePublicFunctions;
-	if (GetLinkerCustomVersion(FUE5MainStreamObjectVersion::GUID) < FUE5MainStreamObjectVersion::RigVMSaveFunctionAccessInModel)
-	{
-		for (const FRigVMGraphFunctionData& FunctionData : Store.PublicFunctions)
-		{
-			BackwardsCompatiblePublicFunctions.Add(FunctionData.Header.Name);
-			URigVMLibraryNode* LibraryNode = Cast<URigVMLibraryNode>(FunctionData.Header.LibraryPointer.GetNodeSoftPath().ResolveObject());
-			OldHeaders.Add(LibraryNode, FunctionData.Header);
-		}
-	}
-
-	// Addressing issue where PublicGraphFunctions is populated, but the model PublicFunctionNames is not
-	URigVMFunctionLibrary* FunctionLibrary = GetLocalFunctionLibrary();
-	if (FunctionLibrary)
-	{
-		if (PublicGraphFunctions.Num() > FunctionLibrary->PublicFunctionNames.Num())
-		{
-			for (const FRigVMGraphFunctionHeader& PublicHeader : PublicGraphFunctions)
-			{
-				BackwardsCompatiblePublicFunctions.Add(PublicHeader.Name);
-			}
-		}
-	}
+	GetBackwardsCompatibilityPublicFunctions(BackwardsCompatiblePublicFunctions, OldHeaders);
 
 	// Lets rebuild the FunctionStore from the model
 	if (FunctionLibrary)
@@ -3775,6 +3754,7 @@ void URigVMBlueprint::PatchFunctionsOnLoad()
 			if (!Variant)
 			{
 				Header.Variant.Guid = FRigVMVariant::GenerateGUID(Header.LibraryPointer.GetLibraryNodePath());
+				FunctionLibrary->FunctionToVariant.FindOrAdd(Header.Name) = Header.Variant;
 			}
 			else
 			{
@@ -3782,15 +3762,46 @@ void URigVMBlueprint::PatchFunctionsOnLoad()
 			}
 			
 			Store.AddFunction(Header, bIsPublic);
-			
+			if (bIsPublic)
+			{
+				GetRigVMClient()->UpdateGraphFunctionSerializedGraph(LibraryNode);
+			}
+		}
+
+		// Update dependencies and external variables if needed
+		for (URigVMLibraryNode* LibraryNode : FunctionLibrary->GetFunctions())
+		{
+			GetRigVMClient()->UpdateExternalVariablesForFunction(LibraryNode);
+			GetRigVMClient()->UpdateDependenciesForFunction(LibraryNode);
 		}
 	}
+}
 
-	// Update dependencies and external variables if needed
-	for (URigVMLibraryNode* LibraryNode : Library->GetFunctions())
+void URigVMBlueprint::GetBackwardsCompatibilityPublicFunctions(TArray<FName>& BackwardsCompatiblePublicFunctions, TMap<URigVMLibraryNode*, FRigVMGraphFunctionHeader>& OldHeaders)
+{
+	URigVMBlueprintGeneratedClass* CRGeneratedClass = GetRigVMBlueprintGeneratedClass();
+	FRigVMGraphFunctionStore& Store = CRGeneratedClass->GraphFunctionStore;
+	if (GetLinkerCustomVersion(FUE5MainStreamObjectVersion::GUID) < FUE5MainStreamObjectVersion::RigVMSaveFunctionAccessInModel)
 	{
-		GetRigVMClient()->UpdateExternalVariablesForFunction(LibraryNode);
-		GetRigVMClient()->UpdateDependenciesForFunction(LibraryNode);
+		for (const FRigVMGraphFunctionData& FunctionData : Store.PublicFunctions)
+		{
+			BackwardsCompatiblePublicFunctions.Add(FunctionData.Header.Name);
+			URigVMLibraryNode* LibraryNode = Cast<URigVMLibraryNode>(FunctionData.Header.LibraryPointer.GetNodeSoftPath().ResolveObject());
+			OldHeaders.Add(LibraryNode, FunctionData.Header);
+		}
+	}
+
+	// Addressing issue where PublicGraphFunctions is populated, but the model PublicFunctionNames is not
+	URigVMFunctionLibrary* FunctionLibrary = GetLocalFunctionLibrary();
+	if (FunctionLibrary)
+	{
+		if (PublicGraphFunctions.Num() > FunctionLibrary->PublicFunctionNames.Num())
+		{
+			for (const FRigVMGraphFunctionHeader& PublicHeader : PublicGraphFunctions)
+			{
+				BackwardsCompatiblePublicFunctions.Add(PublicHeader.Name);
+			}
+		}
 	}
 }
 
diff --git a/Engine/Plugins/Runtime/RigVM/Source/RigVMDeveloper/Public/RigVMBlueprint.h b/Engine/Plugins/Runtime/RigVM/Source/RigVMDeveloper/Public/RigVMBlueprint.h
--- a/Engine/Plugins/Runtime/RigVM/Source/RigVMDeveloper/Public/RigVMBlueprint.h
+++ b/Engine/Plugins/Runtime/RigVM/Source/RigVMDeveloper/Public/RigVMBlueprint.h
@@ -617,6 +617,7 @@ private:
 	virtual void PatchParameterNodesOnLoad() {}
 	virtual void PatchLinksWithCast();
 	virtual void PatchFunctionsOnLoad();
+	virtual void GetBackwardsCompatibilityPublicFunctions(TArray<FName> &BackwardsCompatiblePublicFunctions, TMap<URigVMLibraryNode*, FRigVMGraphFunctionHeader>& OldHeaders);
 
 protected:
 
