Control Rig: Do not show compile button on viewport when moving nodes or expanding pins
#rnx

diff --git a/Engine/Plugins/Runtime/RigVM/Source/RigVMDeveloper/Private/EdGraph/RigVMEdGraphSchema.cpp b/Engine/Plugins/Runtime/RigVM/Source/RigVMDeveloper/Private/EdGraph/RigVMEdGraphSchema.cpp
--- a/Engine/Plugins/Runtime/RigVM/Source/RigVMDeveloper/Private/EdGraph/RigVMEdGraphSchema.cpp
+++ b/Engine/Plugins/Runtime/RigVM/Source/RigVMDeveloper/Private/EdGraph/RigVMEdGraphSchema.cpp
@@ -2160,6 +2160,13 @@ void URigVMEdGraphSchema::EndGraphNodeInteraction(UEdGraphNode* InNode) const
 	
 	bool bMovedSomething = false;
 
+	URigVMBlueprint* Blueprint = Graph->GetBlueprint();
+	if (Blueprint == nullptr)
+	{
+		return;
+	}
+	TGuardValue<bool> GuardDirtyBlueprintStatus(Blueprint->bSkipDirtyBlueprintStatus, true);
+	
 	Graph->GetController()->OpenUndoBracket(TEXT("Move Nodes"));
 
 	for (const UEdGraphNode* NodeToMove : NodesToMove)
diff --git a/Engine/Plugins/Runtime/RigVM/Source/RigVMDeveloper/Private/RigVMBlueprint.cpp b/Engine/Plugins/Runtime/RigVM/Source/RigVMDeveloper/Private/RigVMBlueprint.cpp
--- a/Engine/Plugins/Runtime/RigVM/Source/RigVMDeveloper/Private/RigVMBlueprint.cpp
+++ b/Engine/Plugins/Runtime/RigVM/Source/RigVMDeveloper/Private/RigVMBlueprint.cpp
@@ -174,6 +174,7 @@ URigVMBlueprint::URigVMBlueprint(const FObjectInitializer& ObjectInitializer)
 	bVMRecompilationRequired = false;
 	bIsCompiling = false;
 	VMRecompilationBracket = 0;
+	bSkipDirtyBlueprintStatus = false;
 
 	bUpdatingExternalVariables = false;
 	
@@ -3024,8 +3025,13 @@ void URigVMBlueprint::HandleModifiedEvent(ERigVMGraphNotifType InNotifType, URig
 		{
 			if(bMarkBlueprintAsStructurallyModifiedPending)
 			{
+				const TEnumAsByte<EBlueprintStatus> OldStatus = Status;
 				bMarkBlueprintAsStructurallyModifiedPending = false;
 				FBlueprintEditorUtils::MarkBlueprintAsStructurallyModified(this);
+				if (bSkipDirtyBlueprintStatus)
+				{
+					Status = OldStatus;
+				}
 			}
 		}
 		else
diff --git a/Engine/Plugins/Runtime/RigVM/Source/RigVMDeveloper/Public/RigVMBlueprint.h b/Engine/Plugins/Runtime/RigVM/Source/RigVMDeveloper/Public/RigVMBlueprint.h
--- a/Engine/Plugins/Runtime/RigVM/Source/RigVMDeveloper/Public/RigVMBlueprint.h
+++ b/Engine/Plugins/Runtime/RigVM/Source/RigVMDeveloper/Public/RigVMBlueprint.h
@@ -581,6 +581,9 @@ private:
 
 	UPROPERTY(transient, DuplicateTransient)
 	int32 VMRecompilationBracket;
+	
+	UPROPERTY(transient, DuplicateTransient)
+    bool bSkipDirtyBlueprintStatus;
 
 	FRigVMGraphModifiedEvent ModifiedEvent;
 	void Notify(ERigVMGraphNotifType InNotifType, UObject* InSubject);
@@ -832,6 +835,7 @@ private:
 	friend class FRigVMVariantDetailCustomization;
 	friend class FRigVMTreeAssetVariantFilter;
 	friend class FRigVMTreePackageNode;
+	friend class SRigVMGraphNode;
 };
 
 class RIGVMDEVELOPER_API FRigVMBlueprintCompileScope
diff --git a/Engine/Plugins/Runtime/RigVM/Source/RigVMEditor/Private/Widgets/SRigVMGraphNode.cpp b/Engine/Plugins/Runtime/RigVM/Source/RigVMEditor/Private/Widgets/SRigVMGraphNode.cpp
--- a/Engine/Plugins/Runtime/RigVM/Source/RigVMEditor/Private/Widgets/SRigVMGraphNode.cpp
+++ b/Engine/Plugins/Runtime/RigVM/Source/RigVMEditor/Private/Widgets/SRigVMGraphNode.cpp
@@ -1473,6 +1473,8 @@ FReply SRigVMGraphNode::OnExpanderArrowClicked(int32 InPinInfoIndex)
 				}
 			}
 
+			TGuardValue<bool> GuardDirtyBlueprintStatus(Blueprint->bSkipDirtyBlueprintStatus, true);
+
 			Controller->OpenUndoBracket(PinInfo.bExpanded ? TEXT("Collapsing Pin") : TEXT("Expanding Pin"));
 			for(const FString& PinPathToModify : PinPathsToModify)
 			{
