[Control Rig] Need to remove the deprecated rigvm memory generator class during postload instead of after postload(endload) to avoid an assert from asset registry, which expects that objects are ready to be used after postload

diff --git a/Engine/Plugins/Runtime/RigVM/Source/RigVMDeveloper/Private/RigVMBlueprint.cpp b/Engine/Plugins/Runtime/RigVM/Source/RigVMDeveloper/Private/RigVMBlueprint.cpp
--- a/Engine/Plugins/Runtime/RigVM/Source/RigVMDeveloper/Private/RigVMBlueprint.cpp
+++ b/Engine/Plugins/Runtime/RigVM/Source/RigVMDeveloper/Private/RigVMBlueprint.cpp
@@ -1082,6 +1082,9 @@ void URigVMBlueprint::PostLoad()
 #endif
 	}
 
+	// remove invalid class objects that were parented to the rigvmbp object
+	RemoveDeprecatedVMMemoryClass();
+	
 #if WITH_EDITOR
 	if(GIsEditor)
 	{
@@ -1199,7 +1202,6 @@ void URigVMBlueprint::HandlePackageDone()
 		BuildData->ClearInvalidReferences();
 	}
 	
-	RemoveDeprecatedVMMemoryClass();
 	{
 		const FRigVMCompileSettingsDuringLoadGuard Guard(VMCompileSettings);
 		RecompileVM();
@@ -1238,6 +1240,9 @@ void URigVMBlueprint::RemoveDeprecatedVMMemoryClass()
 	{
 		if (URigVMMemoryStorageGeneratorClass* DeprecatedClass = Cast<URigVMMemoryStorageGeneratorClass>(Object))
 		{
+			// Making sure it is fully loaded before removing it to avoid ambiguity regarding load order
+			DeprecatedClass->ConditionalPostLoad();
+			
 			DeprecatedClass->Rename(nullptr, GetTransientPackage(), REN_ForceNoResetLoaders | REN_DoNotDirty | REN_DontCreateRedirectors | REN_NonTransactional);
 			OldMemoryStorageGeneratorClasses.Add(DeprecatedClass);
 		}
