Control Rig: Fix crash when pasting+undoing node with bound pin
#rnx

diff --git a/Engine/Plugins/Runtime/RigVM/Source/RigVMDeveloper/Private/RigVMModel/RigVMController.cpp b/Engine/Plugins/Runtime/RigVM/Source/RigVMDeveloper/Private/RigVMModel/RigVMController.cpp
--- a/Engine/Plugins/Runtime/RigVM/Source/RigVMDeveloper/Private/RigVMModel/RigVMController.cpp
+++ b/Engine/Plugins/Runtime/RigVM/Source/RigVMDeveloper/Private/RigVMModel/RigVMController.cpp
@@ -3474,7 +3474,10 @@ bool URigVMController::RemoveInjectedNode(const FString& InPinPath, bool bAsInpu
 	URigVMNode* NodeEjected = EjectNodeFromPin(InPinPath, bSetupUndoRedo);
 	if (!NodeEjected)
 	{
-		GetActionStack()->CancelAction(Action);
+		if (bSetupUndoRedo)
+		{
+			GetActionStack()->CancelAction(Action);
+		}
 		return false;
 	}
 
@@ -3498,7 +3501,10 @@ bool URigVMController::RemoveInjectedNode(const FString& InPinPath, bool bAsInpu
 	// 3.- Remove node
 	if (!RemoveNode(NodeEjected, bSetupUndoRedo, false))
 	{
-		GetActionStack()->CancelAction(Action);
+		if (bSetupUndoRedo)
+		{
+			GetActionStack()->CancelAction(Action);
+		}
 		return false;
 	}
 
@@ -18859,6 +18865,15 @@ bool URigVMController::AddGraphNode(URigVMNode* InNode, bool bNotify)
 
 	if(!GetSchema()->CanAddNode(this, InNode))
 	{
+		if (InNode->IsInjected())
+		{
+			URigVMInjectionInfo* Info = InNode->GetInjectionInfo();
+			if (URigVMPin* Pin = Info->GetPin())
+			{
+				RemoveInjectedNode(Pin->GetPinPath(), Info->bInjectedAsInput, false);
+			}
+		}
+	
 		Graph->Nodes.Remove(InNode);
 		DestroyObject(InNode);
 		return false;
