RigVM: Correctly scope variable nodes to avoid re-entry

#rnx

diff --git a/Engine/Plugins/Runtime/RigVM/Source/RigVMDeveloper/Private/RigVMCompiler/RigVMAST.cpp b/Engine/Plugins/Runtime/RigVM/Source/RigVMDeveloper/Private/RigVMCompiler/RigVMAST.cpp
--- a/Engine/Plugins/Runtime/RigVM/Source/RigVMDeveloper/Private/RigVMCompiler/RigVMAST.cpp
+++ b/Engine/Plugins/Runtime/RigVM/Source/RigVMDeveloper/Private/RigVMCompiler/RigVMAST.cpp
@@ -228,7 +228,10 @@ FRigVMExprAST::FRigVMBlockArray FRigVMExprAST::GetBlocks(bool bSortByDepth) cons
 	// a branch which are supposed to be obsolete.
 	if(const FRigVMParserAST* Parser = GetParser())
 	{
-		Blocks.Remove(Parser->GetObsoleteBlock(false));
+		Blocks.RemoveAll([](const FRigVMBlockExprAST* Block) -> bool
+		{
+			return Block->IsObsolete();
+		});
 	}
 
 	return Blocks;
@@ -260,6 +263,7 @@ TOptional<uint32> FRigVMExprAST::GetBlockCombinationHash() const
 	{
 		const FRigVMVarExprAST* VarExpr = To<FRigVMVarExprAST>();
 		check(VarExpr);
+		
 		if(const URigVMPin* Pin = VarExpr->GetPin())
 		{
 			if(const URigVMNode* Node = Pin->GetNode())
@@ -268,6 +272,19 @@ TOptional<uint32> FRigVMExprAST::GetBlockCombinationHash() const
 				{
 					return TOptional<uint32>();
 				}
+
+				// variable node output pins may be shared across many blocks,
+				// since they represent data which is not considered scoped.
+				if(Node->IsA<URigVMVariableNode>())
+				{
+					if(Pin->GetName() == URigVMVariableNode::ValueName)
+					{
+						if(Pin->GetDirection() == ERigVMPinDirection::Output)
+						{
+							return TOptional<uint32>();
+						}
+					}
+				}
 			}
 		}
 	}
@@ -671,6 +688,46 @@ bool FRigVMBlockExprAST::Contains(const FRigVMExprAST* InExpression, TMap<const
 	return false;
 }
 
+bool FRigVMBlockExprAST::IsObsolete() const
+{
+	if(bIsObsolete)
+	{
+		return true;
+	}
+
+	struct Local
+	{
+		static bool HasOnlyObsoleteParents(const FRigVMExprAST* InExpr)
+		{
+			for(int32 ParentIndex = 0; ParentIndex < InExpr->NumParents(); ParentIndex++)
+			{
+				const FRigVMExprAST* ParentExpr = InExpr->ParentAt(ParentIndex);
+				if(!IsObsoleteExpr(ParentExpr))
+				{
+					return false;
+				}
+			}
+			return true;
+		}
+		
+		static bool IsObsoleteExpr(const FRigVMExprAST* InExpr)
+		{
+			if(InExpr->IsA(EType::Block))
+			{
+				const FRigVMBlockExprAST* Block = InExpr->To<FRigVMBlockExprAST>();
+				if(!Block->bIsObsolete) // avoid further recursion here
+				{
+					return false;
+				}
+			}
+
+			return HasOnlyObsoleteParents(InExpr);
+		}
+	};
+
+	return Local::HasOnlyObsoleteParents(this);
+}
+
 bool FRigVMNodeExprAST::IsConstant() const
 {
 	if (URigVMNode* CurrentNode = GetNode())
diff --git a/Engine/Plugins/Runtime/RigVM/Source/RigVMDeveloper/Public/RigVMCompiler/RigVMAST.h b/Engine/Plugins/Runtime/RigVM/Source/RigVMDeveloper/Public/RigVMCompiler/RigVMAST.h
--- a/Engine/Plugins/Runtime/RigVM/Source/RigVMDeveloper/Public/RigVMCompiler/RigVMAST.h
+++ b/Engine/Plugins/Runtime/RigVM/Source/RigVMDeveloper/Public/RigVMCompiler/RigVMAST.h
@@ -543,7 +543,7 @@ public:
 	};
 
 	// returns true in case this block should not be executed
-	bool IsObsolete() const { return bIsObsolete; }
+	bool IsObsolete() const;
 
 protected:
 
diff --git a/Engine/Plugins/Runtime/RigVM/Source/RigVMDeveloper/Public/RigVMModel/Nodes/RigVMVariableNode.h b/Engine/Plugins/Runtime/RigVM/Source/RigVMDeveloper/Public/RigVMModel/Nodes/RigVMVariableNode.h
--- a/Engine/Plugins/Runtime/RigVM/Source/RigVMDeveloper/Public/RigVMModel/Nodes/RigVMVariableNode.h
+++ b/Engine/Plugins/Runtime/RigVM/Source/RigVMDeveloper/Public/RigVMModel/Nodes/RigVMVariableNode.h
@@ -81,6 +81,7 @@ private:
 	friend class URigVMPin;
 	friend class URigVMCompiler;
 	friend class FRigVMVarExprAST;
+	friend class FRigVMExprAST;
 	friend class FRigVMParserAST;
 	friend class URigVMSchema;
 	friend class URigVMEdGraphVariableNodeSpawner;
