[Backout] - CL34474076 - CIS Issue 500550
Original CL Desc
-----------------------------------------------------------------
RigVM: Fix template argument discovery for filtered types

#rnx

diff --git a/Engine/Plugins/Runtime/RigVM/Source/RigVM/Private/RigVMCore/RigVMTemplate.cpp b/Engine/Plugins/Runtime/RigVM/Source/RigVM/Private/RigVMCore/RigVMTemplate.cpp
--- a/Engine/Plugins/Runtime/RigVM/Source/RigVM/Private/RigVMCore/RigVMTemplate.cpp
+++ b/Engine/Plugins/Runtime/RigVM/Source/RigVM/Private/RigVMCore/RigVMTemplate.cpp
@@ -596,7 +596,6 @@ void FRigVMTemplateArgument::ForEachType(TFunction<bool(const TRigVMTypeIndex In
 			}
 			return true;
 		});
-		return;
 	}
 	
 	return CategoryViews(TypeCategories).ForEachType(MoveTemp(InCallback));
