[RigVMRegistry] remove extra safety checks in non-shipping build, hoping that it helps catch invalid types entering the registry more often via a direct crash

diff --git a/Engine/Plugins/Runtime/RigVM/Source/RigVM/Private/RigVMCore/RigVMRegistry.cpp b/Engine/Plugins/Runtime/RigVM/Source/RigVM/Private/RigVMCore/RigVMRegistry.cpp
--- a/Engine/Plugins/Runtime/RigVM/Source/RigVM/Private/RigVMCore/RigVMRegistry.cpp
+++ b/Engine/Plugins/Runtime/RigVM/Source/RigVM/Private/RigVMCore/RigVMRegistry.cpp
@@ -51,8 +51,21 @@ void FRigVMRegistry_NoLock::AddReferencedObjects(FReferenceCollector& Collector)
 	{
 		// the Object needs to be checked for validity since it may be a user defined type (struct or enum)
 		// which is about to get removed. 
-		if (!!Type.Type.CPPTypeObject)
+		if (Type.Type.CPPTypeObject)
 		{
+#if !UE_BUILD_SHIPPING
+			// in non shipping builds, immediately run IsValidLowLevelFast such that
+			// we can catch invalid types earlier via a direct crash more often
+			if (Type.Type.CPPTypeObject->IsValidLowLevelFast())
+			{
+				// By design, hold strong references only to non-native types
+				if (!Type.Type.CPPTypeObject->IsNative())
+				{
+					Collector.AddReferencedObject(Type.Type.CPPTypeObject);	
+				}
+			}
+#else
+			// in shipping builds, try to be as safe as possible
 			if(IsValid(Type.Type.CPPTypeObject))
 			{
 				if(Type.Type.CPPTypeObject->GetClass())
@@ -77,6 +90,7 @@ void FRigVMRegistry_NoLock::AddReferencedObjects(FReferenceCollector& Collector)
 					}
 				}
 			}
+#endif
 		}
 	}
 }
