Control Rig: Fix num types calculation in template argument ignoring the filter callback
#rnx

diff --git a/Engine/Plugins/Runtime/RigVM/Source/RigVM/Private/RigVMCore/RigVMTemplate.cpp b/Engine/Plugins/Runtime/RigVM/Source/RigVM/Private/RigVMCore/RigVMTemplate.cpp
--- a/Engine/Plugins/Runtime/RigVM/Source/RigVM/Private/RigVMCore/RigVMTemplate.cpp
+++ b/Engine/Plugins/Runtime/RigVM/Source/RigVM/Private/RigVMCore/RigVMTemplate.cpp
@@ -546,6 +546,7 @@ int32 FRigVMTemplateArgument::GetNumTypes_NoLock() const
 	{
 		return TypeIndices.Num();
 	}
+
 	
 	if (FilterType)
 	{
@@ -596,6 +597,7 @@ void FRigVMTemplateArgument::ForEachType(TFunction<bool(const TRigVMTypeIndex In
 			}
 			return true;
 		});
+		return;
 	}
 	
 	return CategoryViews(TypeCategories).ForEachType(MoveTemp(InCallback));
