Control Rig: Allow curves as function inputs
#rnx

diff --git a/Engine/Plugins/Runtime/RigVM/Source/RigVM/Public/RigVMFunctions/Animation/RigVMFunction_AnimEvalRichCurve.h b/Engine/Plugins/Runtime/RigVM/Source/RigVM/Public/RigVMFunctions/Animation/RigVMFunction_AnimEvalRichCurve.h
--- a/Engine/Plugins/Runtime/RigVM/Source/RigVM/Public/RigVMFunctions/Animation/RigVMFunction_AnimEvalRichCurve.h
+++ b/Engine/Plugins/Runtime/RigVM/Source/RigVM/Public/RigVMFunctions/Animation/RigVMFunction_AnimEvalRichCurve.h
@@ -30,7 +30,7 @@ struct RIGVM_API FRigVMFunction_AnimEvalRichCurve : public FRigVMFunction_AnimBa
 	UPROPERTY(meta=(Input))
 	float Value;
 
-	UPROPERTY(meta=(Input, Constant))
+	UPROPERTY(meta=(Input))
 	FRuntimeFloatCurve Curve;
 
 	UPROPERTY(meta=(Input))
diff --git a/Engine/Plugins/Runtime/RigVM/Source/RigVMEditor/Private/Widgets/SRigVMGraphPinCurveFloat.cpp b/Engine/Plugins/Runtime/RigVM/Source/RigVMEditor/Private/Widgets/SRigVMGraphPinCurveFloat.cpp
--- a/Engine/Plugins/Runtime/RigVM/Source/RigVMEditor/Private/Widgets/SRigVMGraphPinCurveFloat.cpp
+++ b/Engine/Plugins/Runtime/RigVM/Source/RigVMEditor/Private/Widgets/SRigVMGraphPinCurveFloat.cpp
@@ -46,6 +46,12 @@ TSharedRef<SWidget>	SRigVMGraphPinCurveFloat::GetDefaultValueWidget()
 	return Widget;
 }
 
+void SRigVMGraphPinCurveFloat::Tick(const FGeometry& AllottedGeometry, const double InCurrentTime, const float InDeltaTime)
+{
+	UpdateAndGetCurve();
+	SGraphPin::Tick(AllottedGeometry, InCurrentTime, InDeltaTime);
+}
+
 TArray<FRichCurveEditInfoConst> SRigVMGraphPinCurveFloat::GetCurves() const
 {
 	TArray<FRichCurveEditInfoConst> Curves;
diff --git a/Engine/Plugins/Runtime/RigVM/Source/RigVMEditor/Public/Widgets/SRigVMGraphPinCurveFloat.h b/Engine/Plugins/Runtime/RigVM/Source/RigVMEditor/Public/Widgets/SRigVMGraphPinCurveFloat.h
--- a/Engine/Plugins/Runtime/RigVM/Source/RigVMEditor/Public/Widgets/SRigVMGraphPinCurveFloat.h
+++ b/Engine/Plugins/Runtime/RigVM/Source/RigVMEditor/Public/Widgets/SRigVMGraphPinCurveFloat.h
@@ -30,6 +30,7 @@ protected:
 
 	//~ Begin SGraphPin Interface
 	virtual TSharedRef<SWidget>	GetDefaultValueWidget() override;
+	virtual void Tick(const FGeometry& AllottedGeometry, const double InCurrentTime, const float InDeltaTime) override;
 	//~ End SGraphPin Interface
 
 	FRuntimeFloatCurve& UpdateAndGetCurve();
