Control RIg: Removing bad design about maintaining node title and categories shared for variants. It is confusing for the users.
#rnx

diff --git a/Engine/Plugins/Runtime/RigVM/Source/RigVMDeveloper/Private/RigVMModel/Nodes/RigVMLibraryNode.cpp b/Engine/Plugins/Runtime/RigVM/Source/RigVMDeveloper/Private/RigVMModel/Nodes/RigVMLibraryNode.cpp
--- a/Engine/Plugins/Runtime/RigVM/Source/RigVMDeveloper/Private/RigVMModel/Nodes/RigVMLibraryNode.cpp
+++ b/Engine/Plugins/Runtime/RigVM/Source/RigVMDeveloper/Private/RigVMModel/Nodes/RigVMLibraryNode.cpp
@@ -266,7 +266,6 @@ FRigVMGraphFunctionHeader URigVMLibraryNode::GetFunctionHeader(IRigVMGraphFuncti
 					}
 				}
 			}
-			Header.NodeTitle = Data->Header.NodeTitle;
 		}
 	}
 
@@ -284,10 +283,7 @@ FRigVMGraphFunctionHeader URigVMLibraryNode::GetFunctionHeader(IRigVMGraphFuncti
 	//Header.Tooltip = GetToolTipText();
 	
     Header.NodeColor = GetNodeColor();
-	if (Header.NodeTitle.IsEmpty())
-	{
-		Header.NodeTitle = GetNodeTitle();
-	}
+	Header.NodeTitle = GetNodeTitle();
     for(URigVMPin* Pin : GetPins())
     {
     	FRigVMGraphFunctionArgument Arg;
diff --git a/Engine/Plugins/Runtime/RigVM/Source/RigVMDeveloper/Private/RigVMModel/RigVMController.cpp b/Engine/Plugins/Runtime/RigVM/Source/RigVMDeveloper/Private/RigVMModel/RigVMController.cpp
--- a/Engine/Plugins/Runtime/RigVM/Source/RigVMDeveloper/Private/RigVMModel/RigVMController.cpp
+++ b/Engine/Plugins/Runtime/RigVM/Source/RigVMDeveloper/Private/RigVMModel/RigVMController.cpp
@@ -7385,31 +7385,7 @@ bool URigVMController::SetNodeCategory(URigVMCollapseNode* InNode, const FString
 		GetActionStack()->BeginAction(Action);
 	}
 
-	// If this is a function node with variants, always add the function title as a category
-	FString CategoryToSet = InCategory;
-	if (FRigVMGraphFunctionData* FunctionData = FindFunctionData(InNode->GetFName()))
-	{
-		if (FunctionData->Header.LibraryPointer.IsVariant())
-		{
-			if (FunctionData->Header.Variant.Guid.IsValid())
-			{
-				const FString NodeTitle = FunctionData->Header.NodeTitle;
-				TArray<FString> Categories;
-				InCategory.ParseIntoArray(Categories, TEXT("|"), true);
-				Categories.AddUnique(NodeTitle);
-				Categories.Remove(FString()); // Remove any empty strings
-				if (Categories[0] != NodeTitle)
-				{
-					// Make sure the node title is the first element
-					Categories.Remove(NodeTitle);
-					Categories.Insert(NodeTitle, 0);
-				}
-				CategoryToSet = FString::Join(Categories, TEXT("|"));
-			}
-		}
-	}
-
-	InNode->NodeCategory = CategoryToSet;
+	InNode->NodeCategory = InCategory;
 	Notify(ERigVMGraphNotifType::NodeCategoryChanged, InNode);
 
 	if (bSetupUndoRedo)
@@ -11547,6 +11523,8 @@ bool URigVMController::SwapFunctionReference(URigVMFunctionReferenceNode* InFunc
 	{
 		InFunctionReferenceNode->VariableMap.Remove(MappedVariableToRemove);
 	}
+
+	Notify(ERigVMGraphNotifType::NodeReferenceChanged, InFunctionReferenceNode);
 	
 	if (bSetupUndoRedo)
 	{
@@ -12013,7 +11991,19 @@ URigVMLibraryNode* URigVMController::AddFunctionToLibrary(const FName& InFunctio
 	if (URigVMFunctionLibrary* Library = Cast<URigVMFunctionLibrary>(Graph))
 	{
 		FRigVMVariant Variant;
-		Variant.Guid = FRigVMVariant::GenerateGUID();
+		if (FRigVMGraphFunctionStore* FunctionStore = GetGraphFunctionStore())
+		{
+			if (FRigVMGraphFunctionData* FunctionData = FunctionStore->FindFunctionByName(CollapseNode->GetFName()))
+			{
+				Variant = FunctionData->Header.Variant;
+			}
+		}
+		
+		if (!Variant.Guid.IsValid())
+		{
+			Variant.Guid = FRigVMVariant::GenerateGUID();
+		}
+		
 		Library->FunctionToVariant.Add(CollapseNode->GetFName(), Variant);
 	}
 
@@ -12288,23 +12278,7 @@ URigVMLibraryNode* URigVMController::CreateFunctionVariant(const FName& InFuncti
 			NewFunction->Header = OriginalFunction->Header;
 			NewFunction->Header.Name = NewName;
 			NewFunction->Header.LibraryPointer = Identifier;
-
-			TArray<FString> Categories;
-			OriginalFunction->Header.Category.ParseIntoArray(Categories, TEXT("|"), true);
-
-			const int32 NumCategories = Categories.Num();
-			Categories.AddUnique(OriginalFunction->Header.NodeTitle);
-			const bool bAddedCategory = NumCategories != Categories.Num();
-
-			FString NewCategory = FString::Join(Categories, TEXT("|"));
-			if (bAddedCategory)
-			{
-				SetNodeCategory(Cast<URigVMCollapseNode>(FunctionToClone), NewCategory, false);
-			}
 			
-			URigVMLibraryNode* CloneFunction = Cast<URigVMLibraryNode>(Graph->FindNodeByName(NewName));
-			SetNodeCategory(Cast<URigVMCollapseNode>(CloneFunction), NewCategory, false);
-
 			if (URigVMFunctionLibrary* FunctionLibrary = Cast<URigVMFunctionLibrary>(Graph))
 			{
 				FunctionLibrary->FunctionToVariant.Add(NewName, NewVariant);
