[RigVM] added additional search pass for user defined struct classes now that they have been moved from engine to coreuobject and needed a redirect

-This is a workaround for UE-168245
-all unit tests passed

diff --git a/Engine/Plugins/Runtime/RigVM/Source/RigVM/Private/RigVMTypeUtils.cpp b/Engine/Plugins/Runtime/RigVM/Source/RigVM/Private/RigVMTypeUtils.cpp
--- a/Engine/Plugins/Runtime/RigVM/Source/RigVM/Private/RigVMTypeUtils.cpp
+++ b/Engine/Plugins/Runtime/RigVM/Source/RigVM/Private/RigVMTypeUtils.cpp
@@ -9,6 +9,7 @@
 #include "AssetRegistry/AssetRegistryModule.h"
 #include "Engine/UserDefinedEnum.h"
 #include "StructUtils/UserDefinedStruct.h"
+#include "UObject/LinkerLoad.h"
 
 namespace RigVMTypeUtils
 {
@@ -79,6 +80,18 @@ UObject* RigVMTypeUtils::UserDefinedTypeFromCPPType(FString& InOutCPPType, const
 			TArray<FAssetData> AssetDataList;
 			AssetRegistry.GetAssetsByClass(UUserDefinedStruct::StaticClass()->GetClassPathName(), AssetDataList, true);
 
+			// Temporary workaround:
+			// User Defined Struct has been moved from Engine to CoreUObject at CL34495787 for UE-216472
+			// But currently the AssetRegistry is not applying CoreRedirects to the ClassPath, so searching using the latest class path
+			// will not give the complete list of assets.
+			// To get the complete list, we need to search again using old names of the class. 
+			// This can be removed once the AssetRegistry issue (UE-168245) is fixed
+			TArray<FString> OldPathNames = FLinkerLoad::FindPreviousPathNamesForClass(UUserDefinedStruct::StaticClass()->GetClassPathName().ToString(), false);
+			for (const FString& OldPathName : OldPathNames)
+			{
+				AssetRegistry.GetAssetsByClass( FTopLevelAssetPath(OldPathName), AssetDataList, true);
+			}
+
 			// first pass - try to find it using the tag
 			if(CPPTypeObject == nullptr)
 			{
