Control RIg: Allow to edit curves on function reference nodes
#rnx

diff --git a/Engine/Plugins/Runtime/RigVM/Source/RigVMEditor/Private/Widgets/SRigVMGraphPinCurveFloat.cpp b/Engine/Plugins/Runtime/RigVM/Source/RigVMEditor/Private/Widgets/SRigVMGraphPinCurveFloat.cpp
--- a/Engine/Plugins/Runtime/RigVM/Source/RigVMEditor/Private/Widgets/SRigVMGraphPinCurveFloat.cpp
+++ b/Engine/Plugins/Runtime/RigVM/Source/RigVMEditor/Private/Widgets/SRigVMGraphPinCurveFloat.cpp
@@ -124,6 +124,11 @@ bool SRigVMGraphPinCurveFloat::IsValidCurve(FRichCurveEditInfo CurveInfo)
 					return true;
 				}
 			}
+			else
+			{
+				// for any other node, the curve is valid
+				return true;
+			}
 		}
 	}
 	return false;
