Removes expensive call to the deprecated function GetSelectedObjects which may load the selected objects. Previously right-clicking items in the Content Browser -- even if they were not types derived from URigVMBlueprint -- would result in the selected asset being loaded so we may add a menu context item to create a asset variant.

The new approach leverages the FAssetData which provides the answer to most questions needed to conditionally create the menu context item while avoiding loading the entire selected asset. Since we need the asset loaded for duplication, we only perform the load just before performing the duplication. This also avoids possible issues with passing around UObject ptrs that may have later been cleaned by the GC.

#rnx

diff --git a/Engine/Plugins/Runtime/RigVM/Source/RigVMEditor/Private/RigVMEditorModule.cpp b/Engine/Plugins/Runtime/RigVM/Source/RigVMEditor/Private/RigVMEditorModule.cpp
--- a/Engine/Plugins/Runtime/RigVM/Source/RigVMEditor/Private/RigVMEditorModule.cpp
+++ b/Engine/Plugins/Runtime/RigVM/Source/RigVMEditor/Private/RigVMEditorModule.cpp
@@ -92,26 +92,43 @@ void FRigVMEditorModule::StartupModule()
 					UContentBrowserAssetContextMenuContext* Context = InSection.FindContext<UContentBrowserAssetContextMenuContext>();
 					if (Context)
 					{
-						TArray<UObject*> SelectedObjects = Context->GetSelectedObjects();
-						if (SelectedObjects.Num() == 1 && SelectedObjects[0]->IsA<URigVMBlueprint>())
+						TArray<FAssetData> SelectedAssets = Context->SelectedAssets;
+						if (SelectedAssets.Num() != 1)
 						{
-							InSection.AddMenuEntry("CreateVariant", LOCTEXT("CreateVariant", "Create variant"), LOCTEXT("CreateVariant_ToolTip", "Create a variant for this asset"), FSlateIcon(FRigVMEditorStyle::Get().GetStyleSetName(), "RigVM", "RigVM.Unit"), FExecuteAction::CreateLambda([SelectedObjects]()
+							// We only expect a single asset
+							return;
+						}
+
+						const FAssetData& SelectedAssetData = SelectedAssets[0];
+						if (!SelectedAssetData.IsInstanceOf<URigVMBlueprint>())
+						{
+							// We aren't dealing with a RigVMBlueprint derived type
+							return;
+						}
+
+						FSoftObjectPath SoftObjectPath = SelectedAssetData.GetSoftObjectPath();
+						InSection.AddMenuEntry("CreateVariant", LOCTEXT("CreateVariant", "Create variant"), LOCTEXT("CreateVariant_ToolTip", "Create a variant for this asset"), FSlateIcon(FRigVMEditorStyle::Get().GetStyleSetName(), "RigVM", "RigVM.Unit"), FExecuteAction::CreateLambda([SoftObjectPath]()
 							{
-								FAssetToolsModule& AssetToolsModule = FModuleManager::LoadModuleChecked<FAssetToolsModule>("AssetTools");
+								// Perform the load from within our lambda since this can be expensive, and should not be done speculatively
+								UObject* SelectedObject = SoftObjectPath.TryLoad();
+								if (!SelectedObject)
+								{
+									return;
+								}
 
-								FString PathName = FPackageName::GetLongPackagePath(SelectedObjects[0]->GetPathName());
-								FString ObjectName = SelectedObjects[0]->GetName();
+								FAssetToolsModule& AssetToolsModule = FModuleManager::LoadModuleChecked<FAssetToolsModule>("AssetTools");
+								FString PathName = SoftObjectPath.GetLongPackageName();
+								FString ObjectName = SelectedObject->GetName();
 								FString PackageName;
 								FString BasePackageName = PathName + "/" + ObjectName;
 								AssetToolsModule.Get().CreateUniqueAssetName(BasePackageName, TEXT(""), PackageName, ObjectName);
 
-								UObject* DuplicateAsset = AssetToolsModule.Get().DuplicateAsset(ObjectName, PathName, SelectedObjects[0]);
+								UObject* DuplicateAsset = AssetToolsModule.Get().DuplicateAsset(ObjectName, PathName, SelectedObject);
 								if (URigVMBlueprint* DuplicateBlueprint = Cast<URigVMBlueprint>(DuplicateAsset))
 								{
-									DuplicateBlueprint->AssetVariant = Cast<URigVMBlueprint>(SelectedObjects[0])->AssetVariant;
+									DuplicateBlueprint->AssetVariant = Cast<URigVMBlueprint>(SelectedObject)->AssetVariant;
 								}
 							}));
-						}
 					}
 				}));
 			}
