RigVM: Introducing Pin Categories

#rnx

diff --git a/Engine/Plugins/Runtime/RigVM/Source/RigVMDeveloper/Private/EdGraph/RigVMEdGraph.cpp b/Engine/Plugins/Runtime/RigVM/Source/RigVMDeveloper/Private/EdGraph/RigVMEdGraph.cpp
--- a/Engine/Plugins/Runtime/RigVM/Source/RigVMDeveloper/Private/EdGraph/RigVMEdGraph.cpp
+++ b/Engine/Plugins/Runtime/RigVM/Source/RigVMDeveloper/Private/EdGraph/RigVMEdGraph.cpp
@@ -591,6 +591,7 @@ bool URigVMEdGraph::HandleModifiedEvent_Internal(ERigVMGraphNotifType InNotifTyp
 		case ERigVMGraphNotifType::PinDirectionChanged:
 		case ERigVMGraphNotifType::PinIndexChanged:
 		case ERigVMGraphNotifType::PinBoundVariableChanged:
+		case ERigVMGraphNotifType::PinCategoryChanged:
 		{
 			if (URigVMPin* ModelPin = Cast<URigVMPin>(InSubject))
 			{
@@ -602,6 +603,7 @@ bool URigVMEdGraph::HandleModifiedEvent_Internal(ERigVMGraphNotifType InNotifTyp
 			break;
 		}
 		case ERigVMGraphNotifType::LibraryTemplateChanged:
+		case ERigVMGraphNotifType::PinDisplayNameChanged:
 		{
 			if (URigVMNode* LibraryNode = Cast<URigVMNode>(InSubject))
 			{
diff --git a/Engine/Plugins/Runtime/RigVM/Source/RigVMDeveloper/Private/RigVMModel/Nodes/RigVMTemplateNode.cpp b/Engine/Plugins/Runtime/RigVM/Source/RigVMDeveloper/Private/RigVMModel/Nodes/RigVMTemplateNode.cpp
--- a/Engine/Plugins/Runtime/RigVM/Source/RigVMDeveloper/Private/RigVMModel/Nodes/RigVMTemplateNode.cpp
+++ b/Engine/Plugins/Runtime/RigVM/Source/RigVMDeveloper/Private/RigVMModel/Nodes/RigVMTemplateNode.cpp
@@ -667,42 +667,39 @@ FString URigVMTemplateNode::GetInitialDefaultValueForPin(const FName& InRootPinN
 	return DefaultValue;
 }
 
-FName URigVMTemplateNode::GetDisplayNameForPin(const FName& InRootPinName,
-	const TArray<int32>& InPermutationIndices) const
+FName URigVMTemplateNode::GetDisplayNameForPin(const FString& InPinPath) const
 {
 #if WITH_EDITOR
 	if(const FRigVMTemplate* Template = GetTemplate())
 	{
-		const TArray<int32>* PermutationIndicesPtr = &InPermutationIndices;
-		TArray<int32> AllPermutations;
-		if(PermutationIndicesPtr->IsEmpty())
+		if(!InPinPath.Contains(TEXT(".")))
 		{
-			AllPermutations = GetResolvedPermutationIndices(false);
-			PermutationIndicesPtr = &AllPermutations;
-		}
+			const FName RootPinName = *InPinPath;
+			const TArray<int32> AllPermutations = GetResolvedPermutationIndices(false);
 
-		const FText DisplayNameText = Template->GetDisplayNameForArgument(InRootPinName, *PermutationIndicesPtr);
-		if(DisplayNameText.IsEmpty())
-		{
-			return NAME_None;
-		}
+			const FText DisplayNameText = Template->GetDisplayNameForArgument(RootPinName, AllPermutations);
+			if(DisplayNameText.IsEmpty())
+			{
+				return NAME_None;
+			}
 
-		FString DefaultDisplayName = InRootPinName.ToString();
-		FString Left, Right;
-		if(RigVMStringUtils::SplitPinPathAtEnd(DefaultDisplayName, Left, Right))
-		{
-			DefaultDisplayName = Right;
-		}
+			FString DefaultDisplayName = RootPinName.ToString();
+			FString Left, Right;
+			if(RigVMStringUtils::SplitPinPathAtEnd(DefaultDisplayName, Left, Right))
+			{
+				DefaultDisplayName = Right;
+			}
 
-		const FName DisplayName = *DisplayNameText.ToString();
-		if(DisplayName.IsEqual(*DefaultDisplayName))
-		{
-			return NAME_None;
+			const FName DisplayName = *DisplayNameText.ToString();
+			if(DisplayName.IsEqual(*DefaultDisplayName))
+			{
+				return NAME_None;
+			}
+			return DisplayName;
 		}
-		return DisplayName;
 	}
 #endif
-	return NAME_None;
+	return Super::GetDisplayNameForPin(InPinPath);
 }
 
 TRigVMTypeIndex URigVMTemplateNode::TryReduceTypesToSingle(const TArray<TRigVMTypeIndex>& InTypes, const TRigVMTypeIndex PreferredType) const
diff --git a/Engine/Plugins/Runtime/RigVM/Source/RigVMDeveloper/Private/RigVMModel/RigVMController.cpp b/Engine/Plugins/Runtime/RigVM/Source/RigVMDeveloper/Private/RigVMModel/RigVMController.cpp
--- a/Engine/Plugins/Runtime/RigVM/Source/RigVMDeveloper/Private/RigVMModel/RigVMController.cpp
+++ b/Engine/Plugins/Runtime/RigVM/Source/RigVMDeveloper/Private/RigVMModel/RigVMController.cpp
@@ -26,6 +26,7 @@
 #include "Algo/Transform.h"
 #include "RigVMPythonUtils.h"
 #include "RigVMTypeUtils.h"
+#include "RigVMStringUtils.h"
 #include "StructUtils/UserDefinedStruct.h"
 #include "RigVMFunctions/RigVMDispatch_If.h"
 #include "RigVMFunctions/RigVMDispatch_Select.h"
@@ -1567,7 +1568,7 @@ URigVMUnitNode* URigVMController::AddUnitNode(UScriptStruct* InScriptStruct, TSu
 	
 	if (!bSuspendNotifications)
 	{
-		Graph->MarkPackageDirty();
+		(void)Graph->MarkPackageDirty();
 	}
 
 	FRigVMControllerCompileBracketScope CompileScope(this);
@@ -1917,7 +1918,7 @@ URigVMVariableNode* URigVMController::AddVariableNode(const FName& InVariableNam
 
 	if (!bSuspendNotifications)
 	{
-		Graph->MarkPackageDirty();
+		(void)Graph->MarkPackageDirty();
 	}
 
 	FRigVMControllerCompileBracketScope CompileScope(this);
@@ -2954,7 +2955,7 @@ URigVMCommentNode* URigVMController::AddCommentNode(const FString& InCommentText
 	
 	if (!bSuspendNotifications)
 	{
-		Graph->MarkPackageDirty();
+		(void)Graph->MarkPackageDirty();
 	}
 
 	FRigVMControllerCompileBracketScope CompileScope(this);
@@ -3142,7 +3143,7 @@ URigVMRerouteNode* URigVMController::AddRerouteNodeOnPin(const FString& InPinPat
 	
 	if (!bSuspendNotifications)
 	{
-		Graph->MarkPackageDirty();
+		(void)Graph->MarkPackageDirty();
 	}
 
 	if (bSetupUndoRedo)
@@ -3371,7 +3372,7 @@ URigVMInjectionInfo* URigVMController::AddInjectedNode(const FString& InPinPath,
 
 	if (!bSuspendNotifications)
 	{
-		Graph->MarkPackageDirty();
+		(void)Graph->MarkPackageDirty();
 	}
 	
 	if (bSetupUndoRedo)
@@ -3510,7 +3511,7 @@ bool URigVMController::RemoveInjectedNode(const FString& InPinPath, bool bAsInpu
 
 	if (!bSuspendNotifications)
 	{
-		Graph->MarkPackageDirty();
+		(void)Graph->MarkPackageDirty();
 	}
 	
 	if (bSetupUndoRedo)
@@ -3662,7 +3663,7 @@ URigVMInjectionInfo* URigVMController::InjectNodeIntoPin(URigVMPin* InPin, bool
 
 	if (!bSuspendNotifications)
 	{
-		Graph->MarkPackageDirty();
+		(void)Graph->MarkPackageDirty();
 	}
 
 	if (bSetupUndoRedo)
@@ -3796,7 +3797,7 @@ URigVMNode* URigVMController::EjectNodeFromPin(URigVMPin* InPin, bool bSetupUndo
 		
 	if (!bSuspendNotifications)
 	{
-		Graph->MarkPackageDirty();
+		(void)Graph->MarkPackageDirty();
 	}
 
 	if (bSetupUndoRedo)
@@ -6781,7 +6782,7 @@ bool URigVMController::RemoveNodes(TArray<URigVMNode*> InNodes, bool bSetupUndoR
 
 	if (!bSuspendNotifications)
 	{
-		Graph->MarkPackageDirty();
+		(void)Graph->MarkPackageDirty();
 	}
 
 	if (bPrintPythonCommand)
@@ -7703,7 +7704,7 @@ bool URigVMController::RenameVariable(const FName& InOldName, const FName& InNew
 		Notify(ERigVMGraphNotifType::VariableRenamed, RenamedNode);
 		if (!bSuspendNotifications)
 		{
-			Graph->MarkPackageDirty();
+			(void)Graph->MarkPackageDirty();
 		}
 	}
 
@@ -7794,7 +7795,7 @@ bool URigVMController::SetPinExpansion(URigVMPin* InPin, bool bIsExpanded, bool
 	Notify(ERigVMGraphNotifType::PinExpansionChanged, InPin);
 	if (!bSuspendNotifications)
 	{
-		Graph->MarkPackageDirty();
+		(void)Graph->MarkPackageDirty();
 	}
 
 	if (bSetupUndoRedo)
@@ -7869,7 +7870,214 @@ bool URigVMController::SetPinIsWatched(URigVMPin* InPin, bool bIsWatched, bool b
 	Notify(ERigVMGraphNotifType::PinWatchedChanged, InPin);
 	if (!bSuspendNotifications)
 	{
-		Graph->MarkPackageDirty();
+		(void)Graph->MarkPackageDirty();
+	}
+
+	if (bSetupUndoRedo)
+	{
+		GetActionStack()->EndAction(Action);
+	}
+
+	return true;
+}
+
+bool URigVMController::SetPinDisplayName(const FString& InPinPath, const FString& InDisplayName, bool bSetupUndoRedo, bool bPrintPythonCommand)
+{
+	if (!IsValidGraph())
+	{
+		return false;
+	}
+
+	if (!bIsTransacting && !IsGraphEditable())
+	{
+		return false;
+	}
+
+	const URigVMGraph* Graph = GetGraph();
+	check(Graph);
+
+	URigVMPin* Pin = Graph->FindPin(InPinPath);
+	if (Pin == nullptr)
+	{
+		ReportErrorf(TEXT("Cannot find pin '%s'."), *InPinPath);
+		return false;
+	}
+
+	const bool bSuccess = SetPinDisplayName(Pin, InDisplayName, bSetupUndoRedo);
+	if (bSuccess && bPrintPythonCommand)
+	{
+		const FString GraphName = GetSchema()->GetSanitizedGraphName(GetGraph()->GetGraphName());
+
+		RigVMPythonUtils::Print(GetSchema()->GetGraphOuterName(GetGraph()),
+			FString::Printf(TEXT("blueprint.get_controller_by_name('%s').set_pin_display_name('%s', '%s')"),
+			*GraphName,
+			*GetSchema()->GetSanitizedPinPath(InPinPath),
+			*InDisplayName));
+	}
+
+	return bSuccess;
+}
+
+bool URigVMController::SetPinDisplayName(URigVMPin* InPin, const FString& InDisplayName, bool bSetupUndoRedo)
+{
+	if(!IsValidPinForGraph(InPin))
+	{
+		return false;
+	}
+
+	const FName OldDisplayFName = InPin->GetDisplayName();
+	const FString OldDisplayName = OldDisplayFName.IsNone() ? FString() : OldDisplayFName.ToString();
+
+	if(InDisplayName.Equals(OldDisplayName, ESearchCase::CaseSensitive))
+	{
+		return false;
+	}
+
+	FRigVMControllerCompileBracketScope CompileScope(this);
+	FRigVMSetPinDisplayNameAction Action;
+	if (bSetupUndoRedo)
+	{
+		Action = FRigVMSetPinDisplayNameAction(this, InPin, InDisplayName);
+		Action.SetTitle(TEXT("Set Pin Display Name"));
+		GetActionStack()->BeginAction(Action);
+	}
+
+	InPin->DisplayName = InDisplayName.IsEmpty() ? FName(NAME_None) : FName(*InDisplayName);
+	Notify(ERigVMGraphNotifType::PinDisplayNameChanged, InPin);
+
+	if (!bSuspendNotifications)
+	{
+		const URigVMGraph* Graph = GetGraph();
+		check(Graph);
+		(void)Graph->MarkPackageDirty();
+	}
+
+	if (bSetupUndoRedo)
+	{
+		GetActionStack()->EndAction(Action);
+	}
+
+	return true;
+}
+
+bool URigVMController::SetPinCategory(const FString& InPinPath, const FString& InCategory, bool bSetupUndoRedo, bool bPrintPythonCommand)
+{
+	if (!IsValidGraph())
+	{
+		return false;
+	}
+
+	if (!bIsTransacting && !IsGraphEditable())
+	{
+		return false;
+	}
+
+	const URigVMGraph* Graph = GetGraph();
+	check(Graph);
+
+	URigVMPin* Pin = Graph->FindPin(InPinPath);
+	if (Pin == nullptr)
+	{
+		ReportErrorf(TEXT("Cannot find pin '%s'."), *InPinPath);
+		return false;
+	}
+
+	const bool bSuccess = SetPinCategory(Pin, InCategory, bSetupUndoRedo);
+	if (bSuccess && bPrintPythonCommand)
+	{
+		const FString GraphName = GetSchema()->GetSanitizedGraphName(GetGraph()->GetGraphName());
+
+		RigVMPythonUtils::Print(GetSchema()->GetGraphOuterName(GetGraph()),
+			FString::Printf(TEXT("blueprint.get_controller_by_name('%s').set_pin_category('%s', '%s')"),
+			*GraphName,
+			*GetSchema()->GetSanitizedPinPath(InPinPath),
+			*InCategory));
+	}
+
+	return bSuccess;
+}
+
+bool URigVMController::ClearPinCategory(const FString& InPinPath, bool bSetupUndoRedo, bool bPrintPythonCommand)
+{
+	return SetPinCategory(InPinPath, FString(), bSetupUndoRedo, bPrintPythonCommand);
+}
+
+bool URigVMController::SetPinCategory(URigVMPin* InPin, const FString& InCategory, bool bSetupUndoRedo)
+{
+	if(!IsValidPinForGraph(InPin))
+	{
+		return false;
+	}
+
+	TArray<FString> Categories;
+	RigVMStringUtils::SplitNodePath(InCategory, Categories);
+
+	// trim each category name
+	for(FString& Category : Categories)
+	{
+		Category.TrimStartAndEndInline();
+	}
+
+	// remove empty categories
+	Categories.Remove(FString());
+
+	// re-join the category path
+	FString NewCategory;
+	if(!Categories.IsEmpty())
+	{
+		NewCategory = RigVMStringUtils::JoinNodePath(Categories);
+	}
+
+	if(NewCategory.Equals(InPin->UserDefinedCategory, ESearchCase::CaseSensitive))
+	{
+		return false;
+	}
+
+	// we only allow one category per pin path - so we need to check all parents and children for set categories
+	if(!NewCategory.IsEmpty())
+	{
+		TArray<URigVMPin*> PinsToCheck;
+		URigVMPin* ParentPin = InPin->GetParentPin();
+		while(ParentPin)
+		{
+			PinsToCheck.Add(ParentPin);
+			ParentPin = ParentPin->GetParentPin();
+		}
+		PinsToCheck.Append(InPin->GetAllSubPinsRecursively());
+
+		for(const URigVMPin* ParentOrSubPin : PinsToCheck)
+		{
+			if(!ParentOrSubPin->UserDefinedCategory.IsEmpty())
+			{
+				ReportErrorf(
+					TEXT("Cannot set category (%s) on pin '%s' since pin '%s' already has a category set (%s)."),
+					*NewCategory,
+					*InPin->GetPinPath(),
+					*ParentOrSubPin->GetPinPath(),
+					*ParentOrSubPin->UserDefinedCategory
+				);
+				return false;
+			}
+		}
+	}
+
+	FRigVMControllerCompileBracketScope CompileScope(this);
+	FRigVMSetPinCategoryAction Action;
+	if (bSetupUndoRedo)
+	{
+		Action = FRigVMSetPinCategoryAction(this, InPin, NewCategory);
+		Action.SetTitle(TEXT("Set Pin Category"));
+		GetActionStack()->BeginAction(Action);
+	}
+
+	InPin->UserDefinedCategory = NewCategory;
+	Notify(ERigVMGraphNotifType::PinCategoryChanged, InPin);
+
+	if (!bSuspendNotifications)
+	{
+		const URigVMGraph* Graph = GetGraph();
+		check(Graph);
+		(void)Graph->MarkPackageDirty();
 	}
 
 	if (bSetupUndoRedo)
@@ -9070,7 +9278,7 @@ URigVMPin* URigVMController::InsertArrayPin(URigVMPin* ArrayPin, int32 InIndex,
 	Notify(ERigVMGraphNotifType::PinArraySizeChanged, ArrayPin);
 	if (!bSuspendNotifications)
 	{
-		Graph->MarkPackageDirty();
+		(void)Graph->MarkPackageDirty();
 	}
 
 	if (bSetupUndoRedo)
@@ -9158,7 +9366,7 @@ bool URigVMController::RemoveArrayPin(const FString& InArrayElementPinPath, bool
 
 	if (!bSuspendNotifications)
 	{
-		Graph->MarkPackageDirty();
+		(void)Graph->MarkPackageDirty();
 	}
 
 	// set the array pin's default value type based on the resulting array pin list
@@ -10191,7 +10399,7 @@ bool URigVMController::AddLink(URigVMPin* OutputPin, URigVMPin* InputPin, bool b
 
 	if (!bSuspendNotifications)
 	{
-		Graph->MarkPackageDirty();
+		(void)Graph->MarkPackageDirty();
 	}
 	Notify(ERigVMGraphNotifType::LinkAdded, Link);
 
@@ -10406,7 +10614,7 @@ bool URigVMController::BreakLink(URigVMPin* OutputPin, URigVMPin* InputPin, bool
 			
 			if (!bSuspendNotifications)
 			{
-				Graph->MarkPackageDirty();
+				(void)Graph->MarkPackageDirty();
 			}
 			Notify(ERigVMGraphNotifType::LinkRemoved, Link);
 
@@ -12149,7 +12357,7 @@ bool URigVMController::MarkFunctionAsPublic(const FName& InFunctionName, bool bI
 
 	if (!bSuspendNotifications)
 	{
-		Graph->MarkPackageDirty();
+		(void)Graph->MarkPackageDirty();
 	}
 
 	if (bPrintPythonCommand)
@@ -12599,7 +12807,7 @@ FRigVMGraphVariableDescription URigVMController::AddLocalVariable(const FName& I
 
 	if (!bSuspendNotifications)
 	{
-		Graph->MarkPackageDirty();
+		(void)Graph->MarkPackageDirty();
 	}
 
 	if (bPrintPythonCommand)
@@ -12754,7 +12962,7 @@ bool URigVMController::RemoveLocalVariable(const FName& InVariableName, bool bSe
 
 		if (!bSuspendNotifications)
 		{
-			Graph->MarkPackageDirty();
+			(void)Graph->MarkPackageDirty();
 		}
 
 		if (bSetupUndoRedo)
@@ -12824,7 +13032,7 @@ bool URigVMController::RenameLocalVariable(const FName& InVariableName, const FN
 	
 	if (!bSuspendNotifications)
 	{
-		Graph->MarkPackageDirty();
+		(void)Graph->MarkPackageDirty();
 	}
 
 	FRigVMControllerCompileBracketScope CompileScope(this);
@@ -12858,7 +13066,7 @@ bool URigVMController::RenameLocalVariable(const FName& InVariableName, const FN
 		Notify(ERigVMGraphNotifType::VariableRenamed, RenamedNode);
 		if (!bSuspendNotifications)
 		{
-			Graph->MarkPackageDirty();
+			(void)Graph->MarkPackageDirty();
 		}
 	}
 
@@ -12910,7 +13118,7 @@ bool URigVMController::SetLocalVariableType(const FName& InVariableName, const F
 
 	if (!bSuspendNotifications)
 	{
-		Graph->MarkPackageDirty();
+		(void)Graph->MarkPackageDirty();
 	}
 
 	FRigVMControllerCompileBracketScope CompileScope(this);
@@ -13045,7 +13253,7 @@ bool URigVMController::SetLocalVariableDefaultValue(const FName& InVariableName,
 	
 	if (!bSuspendNotifications)
 	{
-		Graph->MarkPackageDirty();
+		(void)Graph->MarkPackageDirty();
 	}
 
 	if (bPrintPythonCommand)
@@ -14358,7 +14566,7 @@ URigVMInvokeEntryNode* URigVMController::AddInvokeEntryNode(const FName& InEntry
 
 	if (!bSuspendNotifications)
 	{
-		Graph->MarkPackageDirty();
+		(void)Graph->MarkPackageDirty();
 	}
 
 	FRigVMControllerCompileBracketScope CompileScope(this);
@@ -17290,7 +17498,7 @@ void URigVMController::ResolveTemplateNodeMetaData(URigVMTemplateNode* InNode, b
 
 	for(URigVMPin* Pin : InNode->GetPins())
 	{
-		const FName DisplayName = InNode->GetDisplayNameForPin(Pin->GetFName());
+		const FName DisplayName = InNode->GetDisplayNameForPin(Pin->GetName());
 		if(Pin->DisplayName != DisplayName)
 		{
 			Pin->DisplayName = DisplayName;
diff --git a/Engine/Plugins/Runtime/RigVM/Source/RigVMDeveloper/Private/RigVMModel/RigVMControllerActions.cpp b/Engine/Plugins/Runtime/RigVM/Source/RigVMDeveloper/Private/RigVMModel/RigVMControllerActions.cpp
--- a/Engine/Plugins/Runtime/RigVM/Source/RigVMDeveloper/Private/RigVMModel/RigVMControllerActions.cpp
+++ b/Engine/Plugins/Runtime/RigVM/Source/RigVMDeveloper/Private/RigVMModel/RigVMControllerActions.cpp
@@ -1052,6 +1052,82 @@ bool FRigVMSetPinExpansionAction::Redo()
 	return FRigVMBaseAction::Redo();
 }
 
+FRigVMSetPinDisplayNameAction::FRigVMSetPinDisplayNameAction()
+: FRigVMBaseAction(nullptr)
+, PinPath()
+, OldDisplayName()
+, NewDisplayName()
+{
+}
+
+FRigVMSetPinDisplayNameAction::FRigVMSetPinDisplayNameAction(URigVMController* InController, URigVMPin* InPin, const FString& InNewDisplayName)
+: FRigVMBaseAction(InController)
+, PinPath(InPin->GetPinPath())
+, OldDisplayName(InPin->DisplayName.IsNone() ? FString() : InPin->DisplayName.ToString())
+, NewDisplayName(InNewDisplayName)
+{
+}
+
+bool FRigVMSetPinDisplayNameAction::Undo()
+{
+	if(!FRigVMBaseAction::Undo())
+	{
+		return false;
+	}
+	return GetController()->SetPinDisplayName(PinPath, OldDisplayName, false);
+}
+
+bool FRigVMSetPinDisplayNameAction::Redo()
+{
+	if(!CanUndoRedo())
+	{
+		return false;
+	}
+	if(!GetController()->SetPinDisplayName(PinPath, NewDisplayName, false))
+	{
+		return false;
+	}
+	return FRigVMBaseAction::Redo();
+}
+
+FRigVMSetPinCategoryAction::FRigVMSetPinCategoryAction()
+: FRigVMBaseAction(nullptr)
+, PinPath()
+, OldCategory()
+, NewCategory()
+{
+}
+
+FRigVMSetPinCategoryAction::FRigVMSetPinCategoryAction(URigVMController* InController, URigVMPin* InPin, const FString& InNewCategory)
+: FRigVMBaseAction(InController)
+, PinPath(InPin->GetPinPath())
+, OldCategory(InPin->UserDefinedCategory)
+, NewCategory(InNewCategory)
+{
+}
+
+bool FRigVMSetPinCategoryAction::Undo()
+{
+	if(!FRigVMBaseAction::Undo())
+	{
+		return false;
+	}
+	return GetController()->SetPinCategory(PinPath, OldCategory, false);
+}
+
+bool FRigVMSetPinCategoryAction::Redo()
+{
+	if(!CanUndoRedo())
+	{
+		return false;
+	}
+	if(!GetController()->SetPinCategory(PinPath, NewCategory, false))
+	{
+		return false;
+	}
+	return FRigVMBaseAction::Redo();
+}
+
 FRigVMSetPinWatchAction::FRigVMSetPinWatchAction()
 : FRigVMBaseAction(nullptr)
 , OldIsWatched(false)
diff --git a/Engine/Plugins/Runtime/RigVM/Source/RigVMDeveloper/Private/RigVMModel/RigVMNode.cpp b/Engine/Plugins/Runtime/RigVM/Source/RigVMDeveloper/Private/RigVMModel/RigVMNode.cpp
--- a/Engine/Plugins/Runtime/RigVM/Source/RigVMDeveloper/Private/RigVMModel/RigVMNode.cpp
+++ b/Engine/Plugins/Runtime/RigVM/Source/RigVMDeveloper/Private/RigVMModel/RigVMNode.cpp
@@ -114,6 +114,42 @@ TArray<URigVMPin*> URigVMNode::GetAllPinsRecursively() const
 	return Result;
 }
 
+TArray<FString> URigVMNode::GetPinCategories() const
+{
+	const TArray<URigVMPin*> AllPins = GetAllPinsRecursively();
+
+	TArray<FString> PinCategories;
+	for(const URigVMPin* Pin : AllPins)
+	{
+		if(!Pin->UserDefinedCategory.IsEmpty())
+		{
+			PinCategories.Add(Pin->UserDefinedCategory);
+		}
+	}
+	return PinCategories;
+}
+
+TArray<URigVMPin*> URigVMNode::GetPinsForCategory(FString InCategory) const
+{
+	InCategory.TrimStartAndEndInline();
+	if(InCategory.IsEmpty())
+	{
+		return {};
+	}
+	
+	const TArray<URigVMPin*> AllPins = GetAllPinsRecursively();
+
+	TArray<URigVMPin*> PinsInCategory;
+	for(URigVMPin* Pin : AllPins)
+	{
+		if(Pin->GetCategory().Equals(InCategory))
+		{
+			PinsInCategory.Add(Pin);
+		}
+	}
+	return PinsInCategory;
+}
+
 FString URigVMNode::GetOriginalPinDefaultValue(const URigVMPin* InPin) const
 {
 	const FString CompleteSegmentPath = InPin->GetSegmentPath(true);
@@ -648,6 +684,16 @@ UScriptStruct* URigVMNode::GetTraitScriptStruct(const URigVMPin* InTraitPin) con
 	return nullptr;
 }
 
+FName URigVMNode::GetDisplayNameForPin(const FString& InPinPath) const
+{
+	return NAME_None;
+}
+
+FString URigVMNode::GetCategoryForPin(const FString& InPinPath) const
+{
+	return FString();
+}
+
 URigVMLibraryNode* URigVMNode::FindFunctionForNode() const  
 {
 	const UObject* Subject = this;
diff --git a/Engine/Plugins/Runtime/RigVM/Source/RigVMDeveloper/Private/RigVMModel/RigVMPin.cpp b/Engine/Plugins/Runtime/RigVM/Source/RigVMDeveloper/Private/RigVMModel/RigVMPin.cpp
--- a/Engine/Plugins/Runtime/RigVM/Source/RigVMDeveloper/Private/RigVMModel/RigVMPin.cpp
+++ b/Engine/Plugins/Runtime/RigVM/Source/RigVMDeveloper/Private/RigVMModel/RigVMPin.cpp
@@ -224,6 +224,22 @@ FString URigVMPin::GetSubPinPath(const URigVMPin* InParentPin, bool bIncludePare
 	return GetName();
 }
 
+FString URigVMPin::GetCategory() const
+{
+	if (UserDefinedCategory.IsEmpty())
+	{
+		if(const URigVMNode* Node = GetNode())
+		{
+			const FString CategoryFromNode = Node->GetCategoryForPin(this->GetSegmentPath(true));
+			if(!CategoryFromNode.IsEmpty())
+			{
+				return CategoryFromNode;
+			}
+		}
+	}
+	return UserDefinedCategory;
+}
+
 FString URigVMPin::GetSegmentPath(bool bIncludeRootPin) const
 {
 	URigVMPin* ParentPin = GetParentPin();
@@ -366,12 +382,12 @@ FName URigVMPin::GetDisplayName() const
 {
 	if (DisplayName == NAME_None)
 	{
-		if(const URigVMTemplateNode* Node = Cast<URigVMTemplateNode>(GetNode()))
+		if(const URigVMNode* Node = GetNode())
 		{
-			const FName DisplayNameForArgument = Node->GetDisplayNameForPin(*this->GetSegmentPath(true));
-			if(!DisplayNameForArgument.IsNone())
+			const FName DisplayNameFromNode = Node->GetDisplayNameForPin(this->GetSegmentPath(true));
+			if(!DisplayNameFromNode.IsNone())
 			{
-				return DisplayNameForArgument;
+				return DisplayNameFromNode;
 			}
 		}
 		return GetFName();
@@ -1702,6 +1718,17 @@ const TArray<URigVMPin*>& URigVMPin::GetSubPins() const
 	return SubPins;
 }
 
+TArray<URigVMPin*> URigVMPin::GetAllSubPinsRecursively() const
+{
+	TArray<URigVMPin*> AllSubPins;
+	AllSubPins.Append(SubPins);
+	for(const TObjectPtr<URigVMPin>& SubPin : SubPins)
+	{
+		AllSubPins.Append(SubPin->GetAllSubPinsRecursively());
+	}
+	return AllSubPins;
+}
+
 URigVMPin* URigVMPin::FindSubPin(const FString& InPinPath) const
 {
 	FString Left, Right;
diff --git a/Engine/Plugins/Runtime/RigVM/Source/RigVMDeveloper/Public/RigVMModel/Nodes/RigVMTemplateNode.h b/Engine/Plugins/Runtime/RigVM/Source/RigVMDeveloper/Public/RigVMModel/Nodes/RigVMTemplateNode.h
--- a/Engine/Plugins/Runtime/RigVM/Source/RigVMDeveloper/Public/RigVMModel/Nodes/RigVMTemplateNode.h
+++ b/Engine/Plugins/Runtime/RigVM/Source/RigVMDeveloper/Public/RigVMModel/Nodes/RigVMTemplateNode.h
@@ -132,7 +132,7 @@ public:
 	FString GetInitialDefaultValueForPin(const FName& InRootPinName, const TArray<int32>& InPermutationIndices = TArray<int32>()) const;
 
 	// returns the display name for a pin
-	FName GetDisplayNameForPin(const FName& InRootPinName, const TArray<int32>& InPermutationIndices = TArray<int32>()) const;
+	virtual FName GetDisplayNameForPin(const FString& InPinPath) const override;
 
 	// Tries to reduce the input types to a single type, if all are compatible
 	// Will prioritize the InPreferredType if available
diff --git a/Engine/Plugins/Runtime/RigVM/Source/RigVMDeveloper/Public/RigVMModel/RigVMController.h b/Engine/Plugins/Runtime/RigVM/Source/RigVMDeveloper/Public/RigVMModel/RigVMController.h
--- a/Engine/Plugins/Runtime/RigVM/Source/RigVMDeveloper/Public/RigVMModel/RigVMController.h
+++ b/Engine/Plugins/Runtime/RigVM/Source/RigVMDeveloper/Public/RigVMModel/RigVMController.h
@@ -799,6 +799,20 @@ public:
 	UFUNCTION(BlueprintCallable, Category = RigVMController)
 	bool SetPinIsWatched(const FString& InPinPath, bool bIsWatched, bool bSetupUndoRedo = true);
 
+	// Sets the pin display name. The display name is UI relevant only.
+	UFUNCTION(BlueprintCallable, Category = RigVMController)
+	bool SetPinDisplayName(const FString& InPinPath, const FString& InDisplayName, bool bSetupUndoRedo = true, bool bPrintPythonCommand = false);
+
+	// Sets the pin category. The category is UI relevant only and used
+	// to order pins in the user interface of the node as well as on the details panel.
+	UFUNCTION(BlueprintCallable, Category = RigVMController)
+	bool SetPinCategory(const FString& InPinPath, const FString& InCategory, bool bSetupUndoRedo = true, bool bPrintPythonCommand = false);
+
+	// Clears the pin category. The category is UI relevant only and used
+	// to order pins in the user interface of the node as well as on the details panel.
+	UFUNCTION(BlueprintCallable, Category = RigVMController)
+	bool ClearPinCategory(const FString& InPinPath, bool bSetupUndoRedo = true, bool bPrintPythonCommand = false);
+
 	// Returns the default value of a pin given its pinpath.
 	UFUNCTION(BlueprintCallable, Category = RigVMController)
 	FString GetPinDefaultValue(const FString& InPinPath);
@@ -1226,6 +1240,8 @@ private:
 	bool SetPinExpansion(URigVMPin* InPin, bool bIsExpanded, bool bSetupUndoRedo = true);
 	void ExpandPinRecursively(URigVMPin* InPin, bool bSetupUndoRedo);
 	bool SetPinIsWatched(URigVMPin* InPin, bool bIsWatched, bool bSetupUndoRedo);
+	bool SetPinDisplayName(URigVMPin* InPin, const FString& InDisplayName, bool bSetupUndoRedo);
+	bool SetPinCategory(URigVMPin* InPin, const FString& InCategory, bool bSetupUndoRedo);
 	bool SetVariableName(URigVMVariableNode* InVariableNode, const FName& InVariableName, bool bSetupUndoRedo);
 	static void ForEveryPinRecursively(URigVMPin* InPin, TFunction<void(URigVMPin*)> OnEachPinFunction);
 	static void ForEveryPinRecursively(URigVMNode* InNode, TFunction<void(URigVMPin*)> OnEachPinFunction);
diff --git a/Engine/Plugins/Runtime/RigVM/Source/RigVMDeveloper/Public/RigVMModel/RigVMControllerActions.h b/Engine/Plugins/Runtime/RigVM/Source/RigVMDeveloper/Public/RigVMModel/RigVMControllerActions.h
--- a/Engine/Plugins/Runtime/RigVM/Source/RigVMDeveloper/Public/RigVMModel/RigVMControllerActions.h
+++ b/Engine/Plugins/Runtime/RigVM/Source/RigVMDeveloper/Public/RigVMModel/RigVMControllerActions.h
@@ -852,6 +852,60 @@ public:
 	bool NewIsWatched;
 };
 
+/**
+ * An action setting a pin's display name in the graph.
+ */
+USTRUCT()
+struct FRigVMSetPinDisplayNameAction : public FRigVMBaseAction
+{
+	GENERATED_BODY()
+
+public:
+
+	FRigVMSetPinDisplayNameAction();
+	FRigVMSetPinDisplayNameAction(URigVMController* InController, URigVMPin* InPin, const FString& InNewDisplayName);
+	virtual ~FRigVMSetPinDisplayNameAction() {};
+	virtual UScriptStruct* GetScriptStruct() const override { return FRigVMSetPinDisplayNameAction::StaticStruct(); }
+	virtual bool Undo() override;
+	virtual bool Redo() override;
+
+	UPROPERTY()
+	FString PinPath;
+
+	UPROPERTY()
+	FString OldDisplayName;
+
+	UPROPERTY()
+	FString NewDisplayName;
+};
+
+/**
+ * An action setting a pin's category in the graph.
+ */
+USTRUCT()
+struct FRigVMSetPinCategoryAction : public FRigVMBaseAction
+{
+	GENERATED_BODY()
+
+public:
+
+	FRigVMSetPinCategoryAction();
+	FRigVMSetPinCategoryAction(URigVMController* InController, URigVMPin* InPin, const FString& InNewCategory);
+	virtual ~FRigVMSetPinCategoryAction() {};
+	virtual UScriptStruct* GetScriptStruct() const override { return FRigVMSetPinCategoryAction::StaticStruct(); }
+	virtual bool Undo() override;
+	virtual bool Redo() override;
+
+	UPROPERTY()
+	FString PinPath;
+
+	UPROPERTY()
+	FString OldCategory;
+
+	UPROPERTY()
+	FString NewCategory;
+};
+
 /**
  * An action setting a pin's default value in the graph.
  */
diff --git a/Engine/Plugins/Runtime/RigVM/Source/RigVMDeveloper/Public/RigVMModel/RigVMNode.h b/Engine/Plugins/Runtime/RigVM/Source/RigVMDeveloper/Public/RigVMModel/RigVMNode.h
--- a/Engine/Plugins/Runtime/RigVM/Source/RigVMDeveloper/Public/RigVMModel/RigVMNode.h
+++ b/Engine/Plugins/Runtime/RigVM/Source/RigVMDeveloper/Public/RigVMModel/RigVMNode.h
@@ -68,6 +68,14 @@ public:
 	UFUNCTION(BlueprintCallable, Category = RigVMNode)
 	TArray<URigVMPin*> GetAllPinsRecursively() const;
 
+	// Returns all user defined categories on this node
+	UFUNCTION(BlueprintCallable, Category = RigVMNode)
+	TArray<FString> GetPinCategories() const;
+
+	// Returns all pins for a given category
+	UFUNCTION(BlueprintCallable, Category = RigVMNode)
+	TArray<URigVMPin*> GetPinsForCategory(FString InCategory) const;
+
 	// Returns the default value for a given pin
 	FString GetOriginalPinDefaultValue(const URigVMPin* InPin) const;
 
@@ -336,6 +344,12 @@ public:
 		return FString();
 	}
 
+	// returns the display name for a pin
+	virtual FName GetDisplayNameForPin(const FString& InPinPath) const;
+
+	// returns the category for a pin
+	virtual FString GetCategoryForPin(const FString& InPinPath) const;
+
 private:
 
 	static const inline TCHAR* NodeColorName = TEXT("NodeColor");
diff --git a/Engine/Plugins/Runtime/RigVM/Source/RigVMDeveloper/Public/RigVMModel/RigVMNotifications.h b/Engine/Plugins/Runtime/RigVM/Source/RigVMDeveloper/Public/RigVMModel/RigVMNotifications.h
--- a/Engine/Plugins/Runtime/RigVM/Source/RigVMDeveloper/Public/RigVMModel/RigVMNotifications.h
+++ b/Engine/Plugins/Runtime/RigVM/Source/RigVMDeveloper/Public/RigVMModel/RigVMNotifications.h
@@ -54,6 +54,8 @@ enum class ERigVMGraphNotifType : uint8
 	LibraryTemplateChanged, // The definition of a library node's template has changed (Subject == URigVMLibraryNode)
 	FunctionAccessChanged, // The function has been made public/private (Subject == URigVMLibraryNode)
 	VariantTagsChanged, // The tags in the header of this function variant have changed (Subject == URigVMLibraryNode)
+	PinDisplayNameChanged, // The display name of a pin has changed - requiring a rebuild of the node user interface (Subject == URigVMPin) 
+	PinCategoryChanged, // The category of a pin has changed - requiring a rebuild of the node user interface (Subject == URigVMPin)
 	Invalid // The max for this enum (used for guarding)
 };
 
diff --git a/Engine/Plugins/Runtime/RigVM/Source/RigVMDeveloper/Public/RigVMModel/RigVMPin.h b/Engine/Plugins/Runtime/RigVM/Source/RigVMDeveloper/Public/RigVMModel/RigVMPin.h
--- a/Engine/Plugins/Runtime/RigVM/Source/RigVMDeveloper/Public/RigVMModel/RigVMPin.h
+++ b/Engine/Plugins/Runtime/RigVM/Source/RigVMDeveloper/Public/RigVMModel/RigVMPin.h
@@ -168,6 +168,11 @@ public:
 	UFUNCTION(BlueprintCallable, Category = RigVMPin)
 	FString GetSubPinPath(const URigVMPin* InParentPin, bool bIncludeParentPinName = false) const;
 
+	// Returns the category on a pin. The category is UI relevant only and used
+	// to order pins in the user interface of the node as well as on the details panel.
+	UFUNCTION(BlueprintCallable, Category = RigVMPin)
+	FString GetCategory() const;
+
 	// Returns a . separated path containing all names of the pin within its main
 	// memory owner / storage. This is typically used to create an offset pointer
 	// within memory (FRigVMRegisterOffset).
@@ -397,6 +402,10 @@ public:
 	UFUNCTION(BlueprintCallable, Category = RigVMPin)
 	const TArray<URigVMPin*>& GetSubPins() const;
 
+	// Returns all of the SubPins of this one including sub-sub-pins
+	UFUNCTION(BlueprintCallable, Category = RigVMPin)
+	TArray<URigVMPin*> GetAllSubPinsRecursively() const;
+
 	// Returns a SubPin given a name / path or nullptr.
 	UFUNCTION(BlueprintCallable, Category = RigVMPin)
 	URigVMPin* FindSubPin(const FString& InPinPath) const;
@@ -577,6 +586,9 @@ private:
 	UPROPERTY()
 	TArray<TObjectPtr<URigVMInjectionInfo>> InjectionInfos;
 
+	UPROPERTY()
+	FString UserDefinedCategory;
+
 #if WITH_EDITORONLY_DATA
 	UPROPERTY()
 	FString BoundVariablePath_DEPRECATED;
@@ -592,6 +604,8 @@ private:
 	friend class URigVMNode;
 	friend class URigVMLink;
 	friend class FRigVMParserAST;
+	friend struct FRigVMSetPinDisplayNameAction;
+	friend struct FRigVMSetPinCategoryAction;
 };
 
 class RIGVMDEVELOPER_API FRigVMPinDefaultValueImportErrorContext : public FOutputDevice
