RigVM: Support for layout on function headers

#rnx

diff --git a/Engine/Plugins/Runtime/RigVM/Source/RigVM/Private/RigVMCore/RigVMGraphFunctionDefinition.cpp b/Engine/Plugins/Runtime/RigVM/Source/RigVM/Private/RigVMCore/RigVMGraphFunctionDefinition.cpp
--- a/Engine/Plugins/Runtime/RigVM/Source/RigVM/Private/RigVMCore/RigVMGraphFunctionDefinition.cpp
+++ b/Engine/Plugins/Runtime/RigVM/Source/RigVM/Private/RigVMCore/RigVMGraphFunctionDefinition.cpp
@@ -94,6 +94,23 @@ bool FRigVMGraphFunctionIdentifier::IsVariantOf(const FRigVMGraphFunctionIdentif
 	return ThisHeader.Variant.Guid == OtherHeader.Variant.Guid;
 }
 
+const FString* FRigVMGraphFunctionLayout::FindCategory(const FString& InElement) const
+{
+	for(const FRigVMGraphFunctionCategory& Category : Categories)
+	{
+		if(Category.Elements.Contains(InElement))
+		{
+			return &Category.Path;
+		}
+	}
+	return nullptr;
+}
+
+const FString* FRigVMGraphFunctionLayout::FindDisplayName(const FString& InElement) const
+{
+	return DisplayNames.Find(InElement);
+}
+
 TFunction<FRigVMGraphFunctionHeader(const FSoftObjectPath&, const FName&, bool*)> FRigVMGraphFunctionHeader::FindFunctionHeaderFromPathFunc;
 
 bool FRigVMGraphFunctionHeader::IsMutable() const
diff --git a/Engine/Plugins/Runtime/RigVM/Source/RigVM/Public/RigVMCore/RigVMGraphFunctionDefinition.h b/Engine/Plugins/Runtime/RigVM/Source/RigVM/Public/RigVMCore/RigVMGraphFunctionDefinition.h
--- a/Engine/Plugins/Runtime/RigVM/Source/RigVM/Public/RigVMCore/RigVMGraphFunctionDefinition.h
+++ b/Engine/Plugins/Runtime/RigVM/Source/RigVM/Public/RigVMCore/RigVMGraphFunctionDefinition.h
@@ -486,6 +486,93 @@ protected:
 	friend class URigVMBuildData;
 };
 
+USTRUCT(BlueprintType)
+struct RIGVM_API FRigVMGraphFunctionCategory
+{
+	GENERATED_BODY()
+	
+	FRigVMGraphFunctionCategory()
+	: Path()
+	, Elements()
+	{}
+
+	UPROPERTY(BlueprintReadOnly, VisibleAnywhere, Category=FunctionArgument)
+	FString Path;
+
+	UPROPERTY(BlueprintReadOnly, VisibleAnywhere, Category=FunctionArgument)
+	TArray<FString> Elements;
+
+	friend uint32 GetTypeHash(const FRigVMGraphFunctionCategory& Category)
+	{
+		uint32 Hash = GetTypeHash(Category.Path);
+		for (const FString& Element : Category.Elements)
+		{
+			Hash = HashCombine(Hash, GetTypeHash(Element));
+		}
+		return Hash;
+	}
+
+	bool operator < (const FRigVMGraphFunctionCategory& Other) const
+	{
+		return FCString::Strcmp(*Path, *Other.Path) < 0;
+	}
+
+	friend FArchive& operator<<(FArchive& Ar, FRigVMGraphFunctionCategory& Category)
+	{
+		Ar << Category.Path;
+		Ar << Category.Elements;
+		return Ar;
+	}
+};
+
+USTRUCT(BlueprintType)
+struct RIGVM_API FRigVMGraphFunctionLayout
+{
+	GENERATED_BODY()
+	
+	FRigVMGraphFunctionLayout()
+	: Categories()
+	, DisplayNames()
+	{}
+
+	UPROPERTY(BlueprintReadOnly, VisibleAnywhere, Category=FunctionArgument)
+	TArray<FRigVMGraphFunctionCategory> Categories;
+
+	UPROPERTY(BlueprintReadOnly, VisibleAnywhere, Category=FunctionArgument)
+	TMap<FString, FString> DisplayNames;
+
+	void Reset()
+	{
+		Categories.Reset();
+		DisplayNames.Reset();
+	}
+
+	friend uint32 GetTypeHash(const FRigVMGraphFunctionLayout& Layout)
+	{
+		uint32 Hash = 0;;
+		for (const FRigVMGraphFunctionCategory& Category : Layout.Categories)
+		{
+			Hash = HashCombine(Hash, GetTypeHash(Category));
+		}
+		for(const TPair<FString, FString>& Pair : Layout.DisplayNames)
+		{
+			Hash = HashCombine(Hash, GetTypeHash(Pair));
+		}
+		return Hash;
+	}
+
+	friend FArchive& operator<<(FArchive& Ar, FRigVMGraphFunctionLayout& Layout)
+	{
+		Ar << Layout.Categories;
+		Ar << Layout.DisplayNames;
+		return Ar;
+	}
+
+	const FString* FindCategory(const FString& InElement) const;
+	const FString* FindDisplayName(const FString& InElement) const;
+};
+
+
 USTRUCT(BlueprintType)
 struct RIGVM_API FRigVMGraphFunctionHeader
 {
@@ -532,6 +619,9 @@ struct RIGVM_API FRigVMGraphFunctionHeader
 	UPROPERTY()
 	TArray<FRigVMExternalVariable> ExternalVariables;
 
+	UPROPERTY()
+	FRigVMGraphFunctionLayout Layout;
+
 	bool IsMutable() const;
 
 	bool IsValid() const { return LibraryPointer.IsValid(); }
@@ -600,6 +690,23 @@ struct RIGVM_API FRigVMGraphFunctionHeader
 		Ar << Data.Arguments;
 		Ar << Data.Dependencies;
 		Ar << Data.ExternalVariables;
+
+		if (Ar.IsLoading())
+		{
+			if (Ar.CustomVer(FRigVMObjectVersion::GUID) >= FRigVMObjectVersion::FunctionHeaderStoresLayout)
+			{
+				Ar << Data.Layout;
+			}
+			else
+			{
+				Data.Layout.Reset();
+			}
+		}
+		else
+		{
+			Ar << Data.Layout;
+		}
+		
 		return Ar;
 	}
 
diff --git a/Engine/Plugins/Runtime/RigVM/Source/RigVM/Public/RigVMObjectVersion.h b/Engine/Plugins/Runtime/RigVM/Source/RigVM/Public/RigVMObjectVersion.h
--- a/Engine/Plugins/Runtime/RigVM/Source/RigVM/Public/RigVMObjectVersion.h
+++ b/Engine/Plugins/Runtime/RigVM/Source/RigVM/Public/RigVMObjectVersion.h
@@ -48,6 +48,9 @@ struct RIGVM_API FRigVMObjectVersion
 		// Adding variant to every RigVM asset
 		AddVariantToRigVMAssets,
 
+		// Storing user interface layout within function header
+		FunctionHeaderStoresLayout,
+
 		// -----<new versions can be added above this line>-------------------------------------------------
 		VersionPlusOne,
 		LatestVersion = VersionPlusOne - 1
diff --git a/Engine/Plugins/Runtime/RigVM/Source/RigVMDeveloper/Private/EdGraph/RigVMEdGraph.cpp b/Engine/Plugins/Runtime/RigVM/Source/RigVMDeveloper/Private/EdGraph/RigVMEdGraph.cpp
--- a/Engine/Plugins/Runtime/RigVM/Source/RigVMDeveloper/Private/EdGraph/RigVMEdGraph.cpp
+++ b/Engine/Plugins/Runtime/RigVM/Source/RigVMDeveloper/Private/EdGraph/RigVMEdGraph.cpp
@@ -602,6 +602,17 @@ bool URigVMEdGraph::HandleModifiedEvent_Internal(ERigVMGraphNotifType InNotifTyp
 			}
 			break;
 		}
+		case ERigVMGraphNotifType::PinCategoriesChanged:
+		{
+			if (URigVMNode* ModelNode = Cast<URigVMNode>(InSubject))
+			{
+				if (URigVMEdGraphNode* RigNode = Cast<URigVMEdGraphNode>(FindNodeForModelNodeName(ModelNode->GetFName())))
+				{
+					RigNode->ModelPinsChanged();
+				}
+			}
+			break;
+		}
 		case ERigVMGraphNotifType::LibraryTemplateChanged:
 		case ERigVMGraphNotifType::PinDisplayNameChanged:
 		{
diff --git a/Engine/Plugins/Runtime/RigVM/Source/RigVMDeveloper/Private/EdGraph/RigVMEdGraphNode.cpp b/Engine/Plugins/Runtime/RigVM/Source/RigVMDeveloper/Private/EdGraph/RigVMEdGraphNode.cpp
--- a/Engine/Plugins/Runtime/RigVM/Source/RigVMDeveloper/Private/EdGraph/RigVMEdGraphNode.cpp
+++ b/Engine/Plugins/Runtime/RigVM/Source/RigVMDeveloper/Private/EdGraph/RigVMEdGraphNode.cpp
@@ -1338,8 +1338,7 @@ bool URigVMEdGraphNode::ShowPaletteIconOnNode() const
 	{
 		return ModelNode->IsEvent() ||
 			ModelNode->IsA<URigVMInvokeEntryNode>() ||
-			ModelNode->IsA<URigVMFunctionEntryNode>() ||
-			ModelNode->IsA<URigVMFunctionReturnNode>() ||
+			ModelNode->IsA<URigVMFunctionInterfaceNode>() ||
 			ModelNode->IsA<URigVMFunctionReferenceNode>() ||
 			ModelNode->IsA<URigVMCollapseNode>() ||
 			ModelNode->IsA<URigVMUnitNode>() ||
@@ -1394,8 +1393,7 @@ FSlateIcon URigVMEdGraphNode::GetIconAndTint(FLinearColor& OutColor) const
 			return CollapsedNodeIcon;
 		}
 
-		if (ModelNode->IsA<URigVMFunctionEntryNode>() || 
-            ModelNode->IsA<URigVMFunctionReturnNode>())
+		if (ModelNode->IsA<URigVMFunctionInterfaceNode>())
 		{
 			return EntryReturnIcon;
 		}
diff --git a/Engine/Plugins/Runtime/RigVM/Source/RigVMDeveloper/Private/EdGraph/RigVMEdGraphSchema.cpp b/Engine/Plugins/Runtime/RigVM/Source/RigVMDeveloper/Private/EdGraph/RigVMEdGraphSchema.cpp
--- a/Engine/Plugins/Runtime/RigVM/Source/RigVMDeveloper/Private/EdGraph/RigVMEdGraphSchema.cpp
+++ b/Engine/Plugins/Runtime/RigVM/Source/RigVMDeveloper/Private/EdGraph/RigVMEdGraphSchema.cpp
@@ -1729,8 +1729,7 @@ UEdGraphPin* URigVMEdGraphSchema::DropPinOnNode(UEdGraphNode* InTargetNode, cons
 					Model = CollapseNode->GetContainedGraph();
 					PinDirection = PinDirection == ERigVMPinDirection::Output ? ERigVMPinDirection::Input : ERigVMPinDirection::Output;
 				}
-				else if (ModelNode->IsA<URigVMFunctionEntryNode>() ||
-					ModelNode->IsA<URigVMFunctionReturnNode>())
+				else if (ModelNode->IsA<URigVMFunctionInterfaceNode>())
 				{
 					Model = ModelNode->GetGraph();
 				}
diff --git a/Engine/Plugins/Runtime/RigVM/Source/RigVMDeveloper/Private/RigVMBlueprint.cpp b/Engine/Plugins/Runtime/RigVM/Source/RigVMDeveloper/Private/RigVMBlueprint.cpp
--- a/Engine/Plugins/Runtime/RigVM/Source/RigVMDeveloper/Private/RigVMBlueprint.cpp
+++ b/Engine/Plugins/Runtime/RigVM/Source/RigVMDeveloper/Private/RigVMBlueprint.cpp
@@ -3762,6 +3762,7 @@ void URigVMBlueprint::PatchFunctionsOnLoad()
 			{				
 				Header.ExternalVariables = OldHeader->ExternalVariables;
 				Header.Dependencies = OldHeader->Dependencies;
+				Header.Layout = OldHeader->Layout;
 			}
 
 			const FRigVMVariant* Variant = FunctionLibrary->GetFunctionVariant(LibraryNode->GetFName());
diff --git a/Engine/Plugins/Runtime/RigVM/Source/RigVMDeveloper/Private/RigVMCompiler/RigVMAST.cpp b/Engine/Plugins/Runtime/RigVM/Source/RigVMDeveloper/Private/RigVMCompiler/RigVMAST.cpp
--- a/Engine/Plugins/Runtime/RigVM/Source/RigVMDeveloper/Private/RigVMCompiler/RigVMAST.cpp
+++ b/Engine/Plugins/Runtime/RigVM/Source/RigVMDeveloper/Private/RigVMCompiler/RigVMAST.cpp
@@ -1263,8 +1263,7 @@ FRigVMExprAST* FRigVMParserAST::CreateExpressionForNode(const FRigVMASTProxy& In
 			InNodeProxy.IsA<URigVMVariableNode>() ||
 			InNodeProxy.IsA<URigVMEnumNode>() ||
 			InNodeProxy.IsA<URigVMLibraryNode>() ||
-			InNodeProxy.IsA<URigVMFunctionEntryNode>() ||
-			InNodeProxy.IsA<URigVMFunctionReturnNode>())
+			InNodeProxy.IsA<URigVMFunctionInterfaceNode>())
 		{
 			NodeExpr = MakeExpr<FRigVMNoOpExprAST>(InNodeProxy);
 		}
@@ -2866,8 +2865,7 @@ void FRigVMParserAST::Inline(const TArray<URigVMGraph*>& InGraphs, const TArray<
 			// If its an interface node of the library we are compiling, don't recurse
 			if (OutTraversalInfo.LibraryNodeBeingCompiled != nullptr)
 			{
-				if (Node->IsA<URigVMFunctionEntryNode>() ||
-					Node->IsA<URigVMFunctionReturnNode>())
+				if (Node->IsA<URigVMFunctionInterfaceNode>())
 				{
 					if (Node->GetTypedOuter<URigVMLibraryNode>() == OutTraversalInfo.LibraryNodeBeingCompiled)
 					{
@@ -2888,8 +2886,7 @@ void FRigVMParserAST::Inline(const TArray<URigVMGraph*>& InGraphs, const TArray<
 			}
 			
 			return Node->IsA<URigVMLibraryNode>() ||
-				Node->IsA<URigVMFunctionEntryNode>() ||
-				Node->IsA<URigVMFunctionReturnNode>();
+				Node->IsA<URigVMFunctionInterfaceNode>();
 		}
 
 		static bool IsValidPinForAST(const FRigVMASTProxy& InPinProxy, LocalPinTraversalInfo& OutTraversalInfo)
@@ -3287,8 +3284,7 @@ void FRigVMParserAST::Inline(const TArray<URigVMGraph*>& InGraphs, const TArray<
 			const bool bIsCompilingFunction = OutTraversalInfo.LibraryNodeBeingCompiled != nullptr;
 			if (bIsCompilingFunction)
 			{
-				if (InNodeProxy.IsA<URigVMFunctionEntryNode>() ||
-				   InNodeProxy.IsA<URigVMFunctionReturnNode>())
+				if (InNodeProxy.IsA<URigVMFunctionInterfaceNode>())
 				{
 					URigVMNode* Node = InNodeProxy.GetSubjectChecked<URigVMNode>();
 					if (Node->GetTypedOuter<URigVMLibraryNode>() != OutTraversalInfo.LibraryNodeBeingCompiled)
@@ -3299,8 +3295,7 @@ void FRigVMParserAST::Inline(const TArray<URigVMGraph*>& InGraphs, const TArray<
 			}
 			else
 			{
-				if (InNodeProxy.IsA<URigVMFunctionEntryNode>() ||
-				   InNodeProxy.IsA<URigVMFunctionReturnNode>())
+				if (InNodeProxy.IsA<URigVMFunctionInterfaceNode>())
 				{
 					return;
 				}
diff --git a/Engine/Plugins/Runtime/RigVM/Source/RigVMDeveloper/Private/RigVMCompiler/RigVMCompiler.cpp b/Engine/Plugins/Runtime/RigVM/Source/RigVMDeveloper/Private/RigVMCompiler/RigVMCompiler.cpp
--- a/Engine/Plugins/Runtime/RigVM/Source/RigVMDeveloper/Private/RigVMCompiler/RigVMCompiler.cpp
+++ b/Engine/Plugins/Runtime/RigVM/Source/RigVMDeveloper/Private/RigVMCompiler/RigVMCompiler.cpp
@@ -743,7 +743,7 @@ bool URigVMCompiler::Compile(const FRigVMCompileSettings& InSettings, TArray<URi
 				}
 			}
 
-			if(ModelNode->IsA<URigVMFunctionEntryNode>() || ModelNode->IsA<URigVMFunctionReturnNode>())
+			if(ModelNode->IsA<URigVMFunctionInterfaceNode>())
 			{
 				for(URigVMPin* ExecutePin : ModelNode->Pins)
 				{
@@ -3543,7 +3543,7 @@ FString URigVMCompiler::GetPinHashImpl(const URigVMPin* InPin, const FRigVMVarEx
 					bUseFullNodePath = false;
 				}
 			}
-			else if(Node->IsA<URigVMFunctionEntryNode>() || Node->IsA<URigVMFunctionReturnNode>())
+			else if(Node->IsA<URigVMFunctionInterfaceNode>())
 			{
 				const FString FullPath = InPinProxy.GetCallstack().GetCallPath(true);
 				return FString::Printf(TEXT("%s%s%s"), *Prefix, *FullPath, *Suffix);
diff --git a/Engine/Plugins/Runtime/RigVM/Source/RigVMDeveloper/Private/RigVMModel/Nodes/RigVMFunctionEntryNode.cpp b/Engine/Plugins/Runtime/RigVM/Source/RigVMDeveloper/Private/RigVMModel/Nodes/RigVMFunctionEntryNode.cpp
--- a/Engine/Plugins/Runtime/RigVM/Source/RigVMDeveloper/Private/RigVMModel/Nodes/RigVMFunctionEntryNode.cpp
+++ b/Engine/Plugins/Runtime/RigVM/Source/RigVMDeveloper/Private/RigVMModel/Nodes/RigVMFunctionEntryNode.cpp
@@ -6,30 +6,6 @@
 
 #include UE_INLINE_GENERATED_CPP_BY_NAME(RigVMFunctionEntryNode)
 
-uint32 URigVMFunctionEntryNode::GetStructureHash() const
-{
-	// Avoid hashing the template for library nodes
-	return URigVMNode::GetStructureHash();
-}
-
-FLinearColor URigVMFunctionEntryNode::GetNodeColor() const
-{
-	if(URigVMGraph* RootGraph = GetRootGraph())
-	{
-		if(RootGraph->IsA<URigVMFunctionLibrary>())
-		{
-			return FLinearColor(FColor::FromHex("CB00FFFF"));
-		}
-	}
-	return FLinearColor(FColor::FromHex("005DFFFF"));
-}
-
-bool URigVMFunctionEntryNode::IsDefinedAsVarying() const
-{ 
-	// todo
-	return true; 
-}
-
 bool URigVMFunctionEntryNode::IsWithinLoop() const
 {
 	if(const URigVMCollapseNode* CollapseNode = Cast<URigVMCollapseNode>(GetGraph()->GetOuter()))
@@ -43,13 +19,3 @@ FString URigVMFunctionEntryNode::GetNodeTitle() const
 {
 	return TEXT("Entry");
 }
-
-FText URigVMFunctionEntryNode::GetToolTipText() const
-{
-	return FText::FromName(GetGraph()->GetOuter()->GetFName());
-}
-
-FText URigVMFunctionEntryNode::GetToolTipTextForPin(const URigVMPin* InPin) const
-{
-	return Super::GetToolTipTextForPin(InPin);
-}
diff --git a/Engine/Plugins/Runtime/RigVM/Source/RigVMDeveloper/Private/RigVMModel/Nodes/RigVMFunctionInterfaceNode.cpp b/Engine/Plugins/Runtime/RigVM/Source/RigVMDeveloper/Private/RigVMModel/Nodes/RigVMFunctionInterfaceNode.cpp
new file mode 100644
--- /dev/null
+++ b/Engine/Plugins/Runtime/RigVM/Source/RigVMDeveloper/Private/RigVMModel/Nodes/RigVMFunctionInterfaceNode.cpp
@@ -0,0 +1,85 @@
+// Copyright Epic Games, Inc. All Rights Reserved.
+
+#include "RigVMModel/Nodes/RigVMFunctionInterfaceNode.h"
+#include "RigVMModel/RigVMFunctionLibrary.h"
+#include "RigVMModel/Nodes/RigVMCollapseNode.h"
+
+#include UE_INLINE_GENERATED_CPP_BY_NAME(RigVMFunctionInterfaceNode)
+
+uint32 URigVMFunctionInterfaceNode::GetStructureHash() const
+{
+	// Avoid hashing the template for library nodes
+	return URigVMNode::GetStructureHash();
+}
+
+FLinearColor URigVMFunctionInterfaceNode::GetNodeColor() const
+{
+	if(URigVMGraph* RootGraph = GetRootGraph())
+	{
+		if(RootGraph->IsA<URigVMFunctionLibrary>())
+		{
+			return FLinearColor(FColor::FromHex("CB00FFFF"));
+		}
+	}
+	return FLinearColor(FColor::FromHex("005DFFFF"));
+}
+
+bool URigVMFunctionInterfaceNode::IsDefinedAsVarying() const
+{ 
+	return true; 
+}
+
+FText URigVMFunctionInterfaceNode::GetToolTipText() const
+{
+	return FText::FromName(GetGraph()->GetOuter()->GetFName());
+}
+
+FText URigVMFunctionInterfaceNode::GetToolTipTextForPin(const URigVMPin* InPin) const
+{
+	if(const URigVMPin* OuterPin = FindReferencedPin(InPin))
+	{
+		return OuterPin->GetToolTipText();
+	}
+	return Super::GetToolTipTextForPin(InPin);
+}
+
+FName URigVMFunctionInterfaceNode::GetDisplayNameForPin(const FString& InPinPath) const
+{
+	if(const URigVMPin* OuterPin = FindReferencedPin(InPinPath))
+	{
+		return OuterPin->GetDisplayName();
+	}
+	return Super::GetDisplayNameForPin(InPinPath);
+}
+
+FString URigVMFunctionInterfaceNode::GetCategoryForPin(const FString& InPinPath) const
+{
+	if(const URigVMPin* OuterPin = FindReferencedPin(InPinPath))
+	{
+		return OuterPin->GetCategory();
+	}
+	return Super::GetCategoryForPin(InPinPath);
+}
+
+TArray<FString> URigVMFunctionInterfaceNode::GetPinCategories() const
+{
+	if(const URigVMCollapseNode* OuterNode = Cast<URigVMCollapseNode>(GetGraph()->GetOuter()))
+	{
+		return OuterNode->GetPinCategories();
+	}
+	return PinCategories;
+}
+
+const URigVMPin* URigVMFunctionInterfaceNode::FindReferencedPin(const URigVMPin* InPin) const
+{
+	return FindReferencedPin(InPin->GetSegmentPath(true));
+}
+
+const URigVMPin* URigVMFunctionInterfaceNode::FindReferencedPin(const FString& InPinPath) const
+{
+	if(const URigVMCollapseNode* OuterNode = Cast<URigVMCollapseNode>(GetGraph()->GetOuter()))
+	{
+		return OuterNode->FindPin(InPinPath);
+	}
+	return nullptr;
+}
diff --git a/Engine/Plugins/Runtime/RigVM/Source/RigVMDeveloper/Private/RigVMModel/Nodes/RigVMFunctionReferenceNode.cpp b/Engine/Plugins/Runtime/RigVM/Source/RigVMDeveloper/Private/RigVMModel/Nodes/RigVMFunctionReferenceNode.cpp
--- a/Engine/Plugins/Runtime/RigVM/Source/RigVMDeveloper/Private/RigVMModel/Nodes/RigVMFunctionReferenceNode.cpp
+++ b/Engine/Plugins/Runtime/RigVM/Source/RigVMDeveloper/Private/RigVMModel/Nodes/RigVMFunctionReferenceNode.cpp
@@ -22,6 +22,58 @@ FText URigVMFunctionReferenceNode::GetToolTipText() const
 	return ReferencedFunctionHeader.GetTooltip();
 }
 
+FName URigVMFunctionReferenceNode::GetDisplayNameForPin(const FString& InPinPath) const
+{
+	if(const FString* DisplayName = ReferencedFunctionHeader.Layout.FindDisplayName(InPinPath))
+	{
+		if(!DisplayName->IsEmpty())
+		{
+			return *(*DisplayName);
+		}
+	}
+
+	if(const URigVMPin* Pin = FindPin(InPinPath))
+	{
+		if(Pin->IsRootPin())
+		{
+			const FRigVMGraphFunctionArgument* Argument = ReferencedFunctionHeader.Arguments.FindByPredicate([Pin](const FRigVMGraphFunctionArgument& Argument)
+			{
+				return Argument.Name == Pin->GetFName();
+			});
+
+			if (Argument && !Argument->DisplayName.IsNone())
+			{
+				return Argument->DisplayName;
+			}
+		}
+	}
+
+	if(const URigVMPin* ReferencedPin = FindReferencedPin(InPinPath))
+	{
+		return ReferencedPin->GetDisplayName();
+	}
+
+	return Super::GetDisplayNameForPin(InPinPath);
+}
+
+FString URigVMFunctionReferenceNode::GetCategoryForPin(const FString& InPinPath) const
+{
+	if(const FString* Category = ReferencedFunctionHeader.Layout.FindCategory(InPinPath))
+	{
+		if(!Category->IsEmpty())
+		{
+			return *Category;
+		}
+	}
+
+	if(const URigVMPin* ReferencedPin = FindReferencedPin(InPinPath))
+	{
+		return ReferencedPin->GetCategory();
+	}
+
+	return Super::GetCategoryForPin(InPinPath);
+}
+
 FString URigVMFunctionReferenceNode::GetNodeCategory() const
 {
 	return ReferencedFunctionHeader.Category;
@@ -233,10 +285,25 @@ FText URigVMFunctionReferenceNode::GetToolTipTextForPin(const URigVMPin* InPin)
 			return *Tooltip;
 		}
 	}
-	
+
+	if(const URigVMPin* ReferencedPin = FindReferencedPin(InPin))
+	{
+		return ReferencedPin->GetToolTipText();
+	}
+
 	return Super::GetToolTipTextForPin(InPin);
 }
 
+TArray<FString> URigVMFunctionReferenceNode::GetPinCategories() const
+{
+	TArray<FString> TransientPinCategories;
+	for(const FRigVMGraphFunctionCategory& Category : ReferencedFunctionHeader.Layout.Categories)
+	{
+		TransientPinCategories.Add(Category.Path);
+	}
+	return TransientPinCategories;
+}
+
 FRigVMGraphFunctionIdentifier URigVMFunctionReferenceNode::GetFunctionIdentifier() const
 {
 	return GetReferencedFunctionHeader().LibraryPointer;
@@ -275,4 +342,18 @@ TArray<int32> URigVMFunctionReferenceNode::GetInstructionsForVMImpl(const FRigVM
 	}
 	
 	return Instructions;
-}
\ No newline at end of file
+}
+
+const URigVMPin* URigVMFunctionReferenceNode::FindReferencedPin(const URigVMPin* InPin) const
+{
+	return FindReferencedPin(InPin->GetSegmentPath(true));
+}
+
+const URigVMPin* URigVMFunctionReferenceNode::FindReferencedPin(const FString& InPinPath) const
+{
+	if(const URigVMLibraryNode* LibraryNode = LoadReferencedNode())
+	{
+		return LibraryNode->FindPin(InPinPath);
+	}
+	return nullptr;
+}
diff --git a/Engine/Plugins/Runtime/RigVM/Source/RigVMDeveloper/Private/RigVMModel/Nodes/RigVMFunctionReturnNode.cpp b/Engine/Plugins/Runtime/RigVM/Source/RigVMDeveloper/Private/RigVMModel/Nodes/RigVMFunctionReturnNode.cpp
--- a/Engine/Plugins/Runtime/RigVM/Source/RigVMDeveloper/Private/RigVMModel/Nodes/RigVMFunctionReturnNode.cpp
+++ b/Engine/Plugins/Runtime/RigVM/Source/RigVMDeveloper/Private/RigVMModel/Nodes/RigVMFunctionReturnNode.cpp
@@ -2,44 +2,11 @@
 
 #include "RigVMModel/Nodes/RigVMFunctionReturnNode.h"
 #include "RigVMModel/RigVMFunctionLibrary.h"
+#include "RigVMModel/Nodes/RigVMCollapseNode.h"
 
 #include UE_INLINE_GENERATED_CPP_BY_NAME(RigVMFunctionReturnNode)
 
-uint32 URigVMFunctionReturnNode::GetStructureHash() const
-{
-	// Avoid hashing the template for library nodes
-	return URigVMNode::GetStructureHash();
-}
-
-FLinearColor URigVMFunctionReturnNode::GetNodeColor() const
-{
-	if(URigVMGraph* RootGraph = GetRootGraph())
-	{
-		if(RootGraph->IsA<URigVMFunctionLibrary>())
-		{
-			return FLinearColor(FColor::FromHex("CB00FFFF"));
-		}
-	}
-	return FLinearColor(FColor::FromHex("005DFFFF"));
-}
-
-bool URigVMFunctionReturnNode::IsDefinedAsVarying() const
-{ 
-	// todo
-	return true; 
-}
-
 FString URigVMFunctionReturnNode::GetNodeTitle() const
 {
 	return TEXT("Return");
 }
-
-FText URigVMFunctionReturnNode::GetToolTipText() const
-{
-	return FText::FromName(GetGraph()->GetOuter()->GetFName());
-}
-
-FText URigVMFunctionReturnNode::GetToolTipTextForPin(const URigVMPin* InPin) const
-{
-	return Super::GetToolTipTextForPin(InPin);
-}
diff --git a/Engine/Plugins/Runtime/RigVM/Source/RigVMDeveloper/Private/RigVMModel/Nodes/RigVMLibraryNode.cpp b/Engine/Plugins/Runtime/RigVM/Source/RigVMDeveloper/Private/RigVMModel/Nodes/RigVMLibraryNode.cpp
--- a/Engine/Plugins/Runtime/RigVM/Source/RigVMDeveloper/Private/RigVMModel/Nodes/RigVMLibraryNode.cpp
+++ b/Engine/Plugins/Runtime/RigVM/Source/RigVMDeveloper/Private/RigVMModel/Nodes/RigVMLibraryNode.cpp
@@ -312,7 +312,7 @@ FRigVMGraphFunctionHeader URigVMLibraryNode::GetFunctionHeader(IRigVMGraphFuncti
     	Header.Arguments.Add(Arg);
     }
 
-	// If the header already exists, try to find it and get its external variables and dependencies
+	// If the header already exists, try to find it and get its external variables, dependencies and layout
 	TArray<FString> Dependencies;
 	TArray<FRigVMExternalVariable> ExternalVariables;
 	if (IRigVMGraphFunctionHost* Host = Cast<IRigVMGraphFunctionHost>(HostPtr))
@@ -327,6 +327,37 @@ FRigVMGraphFunctionHeader URigVMLibraryNode::GetFunctionHeader(IRigVMGraphFuncti
 			Header.ExternalVariables = GetExternalVariables();
 		}
 	}
+
+	// fill in the pin categories based on the data stored on the pins themselves
+	const TArray<URigVMPin*> AllPins = GetAllPinsRecursively();
+	TMap<FString, FRigVMGraphFunctionCategory> CategoryMap;
+	for(const URigVMPin* Pin : AllPins)
+	{
+		if(!Pin->UserDefinedCategory.IsEmpty())
+		{
+			FRigVMGraphFunctionCategory& Category = CategoryMap.FindOrAdd(Pin->UserDefinedCategory);
+			Category.Path = Pin->UserDefinedCategory;
+			Category.Elements.Add(Pin->GetSegmentPath(true));
+		}
+	}
+
+	// add the categories in the order they have been added
+	for(const FString& PinCategory : PinCategories)
+	{
+		if(FRigVMGraphFunctionCategory* Category = CategoryMap.Find(PinCategory))
+		{
+			Header.Layout.Categories.Add(*Category);
+		}
+	}
+
+	// fill in all user provided display names
+	for(const URigVMPin* Pin : AllPins)
+	{
+		if(!Pin->DisplayName.IsNone())
+		{
+			Header.Layout.DisplayNames.Add(Pin->GetSegmentPath(true), Pin->DisplayName.ToString());
+		}
+	}
 	
 	return Header;
 }
diff --git a/Engine/Plugins/Runtime/RigVM/Source/RigVMDeveloper/Private/RigVMModel/Nodes/RigVMTemplateNode.cpp b/Engine/Plugins/Runtime/RigVM/Source/RigVMDeveloper/Private/RigVMModel/Nodes/RigVMTemplateNode.cpp
--- a/Engine/Plugins/Runtime/RigVM/Source/RigVMDeveloper/Private/RigVMModel/Nodes/RigVMTemplateNode.cpp
+++ b/Engine/Plugins/Runtime/RigVM/Source/RigVMDeveloper/Private/RigVMModel/Nodes/RigVMTemplateNode.cpp
@@ -390,7 +390,7 @@ bool URigVMTemplateNode::SupportsType(const URigVMPin* InPin, TRigVMTypeIndex In
 				}
 				return true;
 			}
-			else if (IsA<URigVMFunctionEntryNode>() || IsA<URigVMFunctionReturnNode>())
+			else if (IsA<URigVMFunctionInterfaceNode>())
 			{
 				if(OutTypeIndex)
 				{
@@ -416,7 +416,7 @@ bool URigVMTemplateNode::SupportsType(const URigVMPin* InPin, TRigVMTypeIndex In
 		}
 
 		// an entry/return node that does not contain an argument for the pin will always support the connections
-		if (IsA<URigVMFunctionEntryNode>() || IsA<URigVMFunctionReturnNode>())
+		if (IsA<URigVMFunctionInterfaceNode>())
 		{
 			if (Template->FindArgument(RootPin->GetFName()) == nullptr)
 			{
diff --git a/Engine/Plugins/Runtime/RigVM/Source/RigVMDeveloper/Private/RigVMModel/RigVMClient.cpp b/Engine/Plugins/Runtime/RigVM/Source/RigVMDeveloper/Private/RigVMModel/RigVMClient.cpp
--- a/Engine/Plugins/Runtime/RigVM/Source/RigVMDeveloper/Private/RigVMModel/RigVMClient.cpp
+++ b/Engine/Plugins/Runtime/RigVM/Source/RigVMDeveloper/Private/RigVMModel/RigVMClient.cpp
@@ -1188,6 +1188,8 @@ void FRigVMClient::HandleGraphModifiedEvent(ERigVMGraphNotifType InNotifType, UR
 		case ERigVMGraphNotifType::PinTypeChanged: // A pin's data type has changed (Subject == URigVMPin)
 		case ERigVMGraphNotifType::PinIndexChanged: // A pin's index has changed (Subject == URigVMPin)
 		case ERigVMGraphNotifType::PinWatchedChanged: // A pin's watch state has changed (Subject == URigVMPin)
+		case ERigVMGraphNotifType::PinDisplayNameChanged: // A pin's display name / UI label has changed (Subject == URigVMPin)
+		case ERigVMGraphNotifType::PinCategoryChanged: // A pin's category has changed (Subject == URigVMPin)
 		{
 			if (URigVMPin* Pin = Cast<URigVMPin>(InSubject))
 			{
@@ -1199,7 +1201,7 @@ void FRigVMClient::HandleGraphModifiedEvent(ERigVMGraphNotifType InNotifType, UR
 					}
 					if(Node->GetOuter()->IsA<URigVMFunctionLibrary>())
 					{
-						if (URigVMCollapseNode* CollapseNode = Cast<URigVMCollapseNode>(Node))
+						if (const URigVMCollapseNode* CollapseNode = Cast<URigVMCollapseNode>(Node))
 						{
 							UpdateGraphFunctionData(CollapseNode);
 						}
@@ -1208,7 +1210,20 @@ void FRigVMClient::HandleGraphModifiedEvent(ERigVMGraphNotifType InNotifType, UR
 			}
 			break;
 		}
-
+		case ERigVMGraphNotifType::PinCategoriesChanged: // A node's category list has changed (Subject == URigVMNode)
+		{
+			if (URigVMNode* Node = Cast<URigVMNode>(InSubject))
+			{
+				if(Node->GetOuter()->IsA<URigVMFunctionLibrary>())
+				{
+					if (const URigVMCollapseNode* CollapseNode = Cast<URigVMCollapseNode>(Node))
+					{
+						UpdateGraphFunctionData(CollapseNode);
+					}
+				}
+			}
+			break;
+		}
 		case ERigVMGraphNotifType::LinkAdded: // A link has been added (Subject == URigVMLink)
 		case ERigVMGraphNotifType::LinkRemoved: // A link has been removed (Subject == URigVMLink)
 		{
diff --git a/Engine/Plugins/Runtime/RigVM/Source/RigVMDeveloper/Private/RigVMModel/RigVMController.cpp b/Engine/Plugins/Runtime/RigVM/Source/RigVMDeveloper/Private/RigVMModel/RigVMController.cpp
--- a/Engine/Plugins/Runtime/RigVM/Source/RigVMDeveloper/Private/RigVMModel/RigVMController.cpp
+++ b/Engine/Plugins/Runtime/RigVM/Source/RigVMDeveloper/Private/RigVMModel/RigVMController.cpp
@@ -1260,7 +1260,7 @@ TArray<FString> URigVMController::GetAddNodePythonCommands(URigVMNode* Node) con
 						*RigVMPythonUtils::Vector2DToPythonString(InvokeEntryNode->GetPosition()),
 						*NodeName));
 	}
-	else if (Node->IsA<URigVMFunctionEntryNode>() || Node->IsA<URigVMFunctionReturnNode>())
+	else if (Node->IsA<URigVMFunctionInterfaceNode>())
 	{
 		
 		
@@ -4857,8 +4857,7 @@ URigVMCollapseNode* URigVMController::CollapseNodes(const TArray<URigVMNode*>& I
 			continue;
 		}
 
-		if (Node->IsA<URigVMFunctionEntryNode>() ||
-			Node->IsA<URigVMFunctionReturnNode>())
+		if (Node->IsA<URigVMFunctionInterfaceNode>())
 		{
 			continue;
 		}
@@ -7925,6 +7924,22 @@ bool URigVMController::SetPinDisplayName(URigVMPin* InPin, const FString& InDisp
 		return false;
 	}
 
+	// for entry or return nodes we relay to the outer pins
+	if(const URigVMFunctionInterfaceNode* InterfaceNode = Cast<URigVMFunctionInterfaceNode>(InPin->GetNode()))
+	{
+		if(const URigVMPin* OuterPin = InterfaceNode->FindReferencedPin(InPin))
+		{
+			if(URigVMController* OuterController = GetControllerForGraph(OuterPin->GetGraph()))
+			{
+				return OuterController->SetPinDisplayName(OuterPin->GetPinPath(), InDisplayName, bSetupUndoRedo);
+			}
+		}
+	}
+	else if(InPin->GetNode()->IsA<URigVMFunctionReferenceNode>())
+	{
+		return false;
+	}
+
 	const FName OldDisplayFName = InPin->GetDisplayName();
 	const FString OldDisplayName = OldDisplayFName.IsNone() ? FString() : OldDisplayFName.ToString();
 
@@ -8009,6 +8024,22 @@ bool URigVMController::SetPinCategory(URigVMPin* InPin, const FString& InCategor
 		return false;
 	}
 
+	// for entry or return nodes we relay to the outer pins
+	if(const URigVMFunctionInterfaceNode* InterfaceNode = Cast<URigVMFunctionInterfaceNode>(InPin->GetNode()))
+	{
+		if(const URigVMPin* OuterPin = InterfaceNode->FindReferencedPin(InPin))
+		{
+			if(URigVMController* OuterController = GetControllerForGraph(OuterPin->GetGraph()))
+			{
+				return OuterController->SetPinCategory(OuterPin->GetPinPath(), InCategory, bSetupUndoRedo);
+			}
+		}
+	}
+	else if(InPin->GetNode()->IsA<URigVMFunctionReferenceNode>())
+	{
+		return false;
+	}
+
 	TArray<FString> Categories;
 	RigVMStringUtils::SplitNodePath(InCategory, Categories);
 
@@ -8072,6 +8103,435 @@ bool URigVMController::SetPinCategory(URigVMPin* InPin, const FString& InCategor
 
 	InPin->UserDefinedCategory = NewCategory;
 	Notify(ERigVMGraphNotifType::PinCategoryChanged, InPin);
+	if(!NewCategory.IsEmpty())
+	{
+		if(bSetupUndoRedo)
+		{
+			if(!InPin->GetNode()->PinCategories.Contains(NewCategory))
+			{
+				TArray<FString> NewCategories = InPin->GetNode()->PinCategories;
+				NewCategories.Add(NewCategory);
+				(void)SetPinCategories(InPin->GetNode(), NewCategories, bSetupUndoRedo);
+			}
+		}
+	}
+
+	if (!bSuspendNotifications)
+	{
+		const URigVMGraph* Graph = GetGraph();
+		check(Graph);
+		(void)Graph->MarkPackageDirty();
+	}
+
+	if (bSetupUndoRedo)
+	{
+		GetActionStack()->EndAction(Action);
+	}
+
+	return true;
+}
+
+bool URigVMController::RemovePinCategory(const FName& InNodeName, const FString& InPinCategory, bool bSetupUndoRedo, bool bPrintPythonCommand)
+{
+	if (!IsValidGraph())
+	{
+		return false;
+	}
+
+	if (!bIsTransacting && !IsGraphEditable())
+	{
+		return false;
+	}
+
+	const URigVMGraph* Graph = GetGraph();
+	check(Graph);
+
+	const URigVMNode* Node = Graph->FindNodeByName(InNodeName);
+	if (Node == nullptr)
+	{
+		ReportErrorf(TEXT("Cannot find node '%s'."), *InNodeName.ToString());
+		return false;
+	}
+
+	const bool bSuccess = RemovePinCategory(Node, InPinCategory, bSetupUndoRedo);
+	if (bSuccess && bPrintPythonCommand)
+	{
+		const FString GraphName = GetSchema()->GetSanitizedGraphName(GetGraph()->GetGraphName());
+
+		RigVMPythonUtils::Print(GetSchema()->GetGraphOuterName(GetGraph()),
+			FString::Printf(TEXT("blueprint.get_controller_by_name('%s').remove_pin_category('%s', '%s')"),
+			*GraphName,
+			*GetSchema()->GetSanitizedNodeName(InNodeName.ToString()),
+			*InPinCategory));
+	}
+
+	return bSuccess;
+}
+
+bool URigVMController::RemovePinCategory(const URigVMNode* InNode, const FString& InPinCategory, bool bSetupUndoRedo)
+{
+	if(!IsValidNodeForGraph(InNode))
+	{
+		return false;
+	}
+
+	// for entry or return nodes we relay to the outer pins
+	if(const URigVMFunctionInterfaceNode* InterfaceNode = Cast<URigVMFunctionInterfaceNode>(InNode))
+	{
+		if(const URigVMCollapseNode* OuterNode = Cast<URigVMCollapseNode>(InNode->GetGraph()->GetOuter()))
+		{
+			if(URigVMController* OuterController = GetControllerForGraph(OuterNode->GetGraph()))
+			{
+				return OuterController->RemovePinCategory(OuterNode, InPinCategory, bSetupUndoRedo);
+			}
+		}
+	}
+	else if(InNode->IsA<URigVMFunctionReferenceNode>())
+	{
+		return false;
+	}
+
+	const TArray<URigVMPin*> PinsForCategory = InNode->GetPinsForCategory(InPinCategory);
+	if(!InNode->GetPinCategories().Contains(InPinCategory) && PinsForCategory.IsEmpty())
+	{
+		return false;
+	}
+
+	FRigVMControllerCompileBracketScope CompileScope(this);
+	FRigVMChangeNodePinCategoriesAction Action;
+	if (bSetupUndoRedo)
+	{
+		Action = FRigVMChangeNodePinCategoriesAction(this, InNode);
+		Action.SetTitle(TEXT("Remove Pin Category"));
+		GetActionStack()->BeginAction(Action);
+	}
+
+	for(URigVMPin* PinInCategory : PinsForCategory)
+	{
+		(void)SetPinCategory(PinInCategory, FString(), bSetupUndoRedo);
+	}
+
+	const_cast<URigVMNode*>(InNode)->PinCategories.Remove(InPinCategory);
+	Notify(ERigVMGraphNotifType::PinCategoriesChanged, const_cast<URigVMNode*>(InNode));
+
+	if (!bSuspendNotifications)
+	{
+		const URigVMGraph* Graph = GetGraph();
+		check(Graph);
+		(void)Graph->MarkPackageDirty();
+	}
+
+	if (bSetupUndoRedo)
+	{
+		Action.UpdateAfterModification(InNode);
+		GetActionStack()->EndAction(Action);
+	}
+
+	return true;
+}
+
+bool URigVMController::RenamePinCategory(const FName& InNodeName, const FString& InOldPinCategory, const FString& InNewPinCategory, bool bSetupUndoRedo, bool bPrintPythonCommand)
+{
+	if (!IsValidGraph())
+	{
+		return false;
+	}
+
+	if (!bIsTransacting && !IsGraphEditable())
+	{
+		return false;
+	}
+
+	const URigVMGraph* Graph = GetGraph();
+	check(Graph);
+
+	const URigVMNode* Node = Graph->FindNodeByName(InNodeName);
+	if (Node == nullptr)
+	{
+		ReportErrorf(TEXT("Cannot find node '%s'."), *InNodeName.ToString());
+		return false;
+	}
+
+	const bool bSuccess = RenamePinCategory(Node, InOldPinCategory, InNewPinCategory, bSetupUndoRedo);
+	if (bSuccess && bPrintPythonCommand)
+	{
+		const FString GraphName = GetSchema()->GetSanitizedGraphName(GetGraph()->GetGraphName());
+
+		RigVMPythonUtils::Print(GetSchema()->GetGraphOuterName(GetGraph()),
+			FString::Printf(TEXT("blueprint.get_controller_by_name('%s').rename_pin_category('%s', '%s', '%s')"),
+			*GraphName,
+			*GetSchema()->GetSanitizedNodeName(InNodeName.ToString()),
+			*InOldPinCategory,
+			*InNewPinCategory));
+	}
+
+	return bSuccess;
+}
+
+bool URigVMController::RenamePinCategory(const URigVMNode* InNode, const FString& InOldPinCategory, const FString& InNewPinCategory, bool bSetupUndoRedo)
+{
+	if(!IsValidNodeForGraph(InNode))
+	{
+		return false;
+	}
+
+	// for entry or return nodes we relay to the outer pins
+	if(const URigVMFunctionInterfaceNode* InterfaceNode = Cast<URigVMFunctionInterfaceNode>(InNode))
+	{
+		if(const URigVMCollapseNode* OuterNode = Cast<URigVMCollapseNode>(InNode->GetGraph()->GetOuter()))
+		{
+			if(URigVMController* OuterController = GetControllerForGraph(OuterNode->GetGraph()))
+			{
+				return OuterController->RenamePinCategory(OuterNode, InOldPinCategory, InNewPinCategory, bSetupUndoRedo);
+			}
+		}
+	}
+	else if(InNode->IsA<URigVMFunctionReferenceNode>())
+	{
+		return false;
+	}
+
+	const FString OldPrefix = InOldPinCategory + TEXT("|");
+	const FString NewPrefix = InNewPinCategory + TEXT("|");
+
+	TArray<FString> CategoriesToRename;
+	for(const FString& ExistingCategory : InNode->GetPinCategories())
+	{
+		if(ExistingCategory.Equals(InOldPinCategory, ESearchCase::CaseSensitive) ||
+			ExistingCategory.StartsWith(OldPrefix, ESearchCase::CaseSensitive))
+		{
+			CategoriesToRename.Add(ExistingCategory);
+		}
+	}
+
+	if(CategoriesToRename.IsEmpty())
+	{
+		return false;
+	}
+
+	FRigVMControllerCompileBracketScope CompileScope(this);
+	FRigVMChangeNodePinCategoriesAction Action;
+	if (bSetupUndoRedo)
+	{
+		Action = FRigVMChangeNodePinCategoriesAction(this, InNode);
+		Action.SetTitle(TEXT("Rename Pin Category"));
+		GetActionStack()->BeginAction(Action);
+	}
+
+	TArray<URigVMPin*> PinsToUpdate;
+	
+	// first update the category list on the node
+	for(const FString& CategoryToRename : CategoriesToRename)
+	{
+		const int32 Index = InNode->PinCategories.Find(CategoryToRename);
+		if(CategoryToRename.Equals(InOldPinCategory, ESearchCase::CaseSensitive))
+        {
+			const_cast<URigVMNode*>(InNode)->PinCategories[Index] = InNewPinCategory;
+        }
+        else if(CategoryToRename.StartsWith(OldPrefix, ESearchCase::CaseSensitive))
+        {
+        	const FString NewCategory = NewPrefix + CategoryToRename.RightChop(OldPrefix.Len());
+			const_cast<URigVMNode*>(InNode)->PinCategories[Index] = NewCategory;
+        }
+		PinsToUpdate.Append(InNode->GetPinsForCategory(CategoryToRename));
+	}
+
+	// then update the pins and their categories
+	for(URigVMPin* PinToUpdate : PinsToUpdate)
+	{
+		const FString CategoryToRename = PinToUpdate->GetCategory();
+		if(CategoryToRename.Equals(InOldPinCategory, ESearchCase::CaseSensitive))
+		{
+			(void)SetPinCategory(PinToUpdate, InNewPinCategory, bSetupUndoRedo);
+		}
+		else if(CategoryToRename.StartsWith(OldPrefix, ESearchCase::CaseSensitive))
+		{
+			const FString NewCategory = NewPrefix + CategoryToRename.RightChop(OldPrefix.Len());
+			(void)SetPinCategory(PinToUpdate, NewCategory, bSetupUndoRedo);
+		}
+	}
+	
+	Notify(ERigVMGraphNotifType::PinCategoriesChanged, const_cast<URigVMNode*>(InNode));
+
+	if (!bSuspendNotifications)
+	{
+		const URigVMGraph* Graph = GetGraph();
+		check(Graph);
+		(void)Graph->MarkPackageDirty();
+	}
+
+	if (bSetupUndoRedo)
+	{
+		Action.UpdateAfterModification(InNode);
+		GetActionStack()->EndAction(Action);
+	}
+
+	return true;
+}
+
+bool URigVMController::SetPinCategoryIndex(const FName& InNodeName, const FString& InPinCategory, int32 InNewIndex, bool bSetupUndoRedo, bool bPrintPythonCommand)
+{
+	if (!IsValidGraph())
+	{
+		return false;
+	}
+
+	if (!bIsTransacting && !IsGraphEditable())
+	{
+		return false;
+	}
+
+	const URigVMGraph* Graph = GetGraph();
+	check(Graph);
+
+	const URigVMNode* Node = Graph->FindNodeByName(InNodeName);
+	if (Node == nullptr)
+	{
+		ReportErrorf(TEXT("Cannot find node '%s'."), *InNodeName.ToString());
+		return false;
+	}
+
+	const bool bSuccess = SetPinCategoryIndex(Node, InPinCategory, InNewIndex, bSetupUndoRedo);
+	if (bSuccess && bPrintPythonCommand)
+	{
+		const FString GraphName = GetSchema()->GetSanitizedGraphName(GetGraph()->GetGraphName());
+
+		RigVMPythonUtils::Print(GetSchema()->GetGraphOuterName(GetGraph()),
+			FString::Printf(TEXT("blueprint.get_controller_by_name('%s').set_pin_category_index('%s', '%s', %d)"),
+			*GraphName,
+			*GetSchema()->GetSanitizedNodeName(InNodeName.ToString()),
+			*InPinCategory,
+			InNewIndex));
+	}
+
+	return bSuccess;
+}
+
+bool URigVMController::SetPinCategoryIndex(const URigVMNode* InNode, const FString& InPinCategory, int32 InNewIndex, bool bSetupUndoRedo)
+{
+	if(!IsValidNodeForGraph(InNode))
+	{
+		return false;
+	}
+
+	// for entry or return nodes we relay to the outer pins
+	if(const URigVMFunctionInterfaceNode* InterfaceNode = Cast<URigVMFunctionInterfaceNode>(InNode))
+	{
+		if(const URigVMCollapseNode* OuterNode = Cast<URigVMCollapseNode>(InNode->GetGraph()->GetOuter()))
+		{
+			if(URigVMController* OuterController = GetControllerForGraph(OuterNode->GetGraph()))
+			{
+				return OuterController->SetPinCategoryIndex(OuterNode, InPinCategory, InNewIndex, bSetupUndoRedo);
+			}
+		}
+	}
+	else if(InNode->IsA<URigVMFunctionReferenceNode>())
+	{
+		return false;
+	}
+
+	TArray<FString> Categories = InNode->GetPinCategories();
+	if(!Categories.Contains(InPinCategory))
+	{
+		return false;
+	}
+
+	InNewIndex = FMath::Clamp(InNewIndex, 0, Categories.Num() - 1);
+	if(Categories.Find(InPinCategory) == InNewIndex)
+	{
+		return false;
+	}
+
+	Categories.Remove(InPinCategory);
+	if(InNewIndex >= Categories.Num())
+	{
+		Categories.Add(InPinCategory);
+	}
+	else
+	{
+		Categories.Insert(InPinCategory, InNewIndex);
+	}
+
+	return SetPinCategories(InNode, Categories, bSetupUndoRedo);
+}
+
+bool URigVMController::SetPinCategories(const FName& InNodeName, const TArray<FString>& InCategories, bool bSetupUndoRedo)
+{
+	if (!IsValidGraph())
+	{
+		return false;
+	}
+
+	if (!bIsTransacting && !IsGraphEditable())
+	{
+		return false;
+	}
+
+	const URigVMGraph* Graph = GetGraph();
+	check(Graph);
+
+	const URigVMNode* Node = Graph->FindNodeByName(InNodeName);
+	if (Node == nullptr)
+	{
+		ReportErrorf(TEXT("Cannot find node '%s'."), *InNodeName.ToString());
+		return false;
+	}
+
+	return SetPinCategories(Node, InCategories, bSetupUndoRedo);
+}
+
+bool URigVMController::SetPinCategories(const URigVMNode* InNode, const TArray<FString>& InCategories, bool bSetupUndoRedo)
+{
+	if(!IsValidNodeForGraph(InNode))
+	{
+		return false;
+	}
+
+	// for entry or return nodes we relay to the outer pins
+	if(const URigVMFunctionInterfaceNode* InterfaceNode = Cast<URigVMFunctionInterfaceNode>(InNode))
+	{
+		if(const URigVMCollapseNode* OuterNode = Cast<URigVMCollapseNode>(InNode->GetGraph()->GetOuter()))
+		{
+			if(URigVMController* OuterController = GetControllerForGraph(OuterNode->GetGraph()))
+			{
+				return OuterController->SetPinCategories(OuterNode, InCategories, bSetupUndoRedo);
+			}
+		}
+	}
+	else if(InNode->IsA<URigVMFunctionReferenceNode>())
+	{
+		return false;
+	}
+
+	const TArray<FString> PreviousCategories = InNode->GetPinCategories();
+	if(PreviousCategories.Num() == InCategories.Num())
+	{
+		bool bFoundDifference = false;
+		for(int32 Index = 0; Index < PreviousCategories.Num(); Index++)
+		{
+			if(!PreviousCategories[Index].Equals(InCategories[Index], ESearchCase::CaseSensitive))
+			{
+				bFoundDifference = true;
+				break;
+			}
+		}
+		if(!bFoundDifference)
+		{
+			return false;
+		}
+	}
+
+	FRigVMControllerCompileBracketScope CompileScope(this);
+	FRigVMChangeNodePinCategoriesAction Action;
+	if (bSetupUndoRedo)
+	{
+		Action = FRigVMChangeNodePinCategoriesAction(this, InNode);
+		Action.SetTitle(TEXT("Set Pin Categories"));
+		GetActionStack()->BeginAction(Action);
+	}
+
+	const_cast<URigVMNode*>(InNode)->PinCategories = InCategories;
+	Notify(ERigVMGraphNotifType::PinCategoriesChanged, const_cast<URigVMNode*>(InNode));
 
 	if (!bSuspendNotifications)
 	{
@@ -8082,6 +8542,7 @@ bool URigVMController::SetPinCategory(URigVMPin* InPin, const FString& InCategor
 
 	if (bSetupUndoRedo)
 	{
+		Action.UpdateAfterModification(InNode);
 		GetActionStack()->EndAction(Action);
 	}
 
@@ -11393,7 +11854,7 @@ bool URigVMController::ChangeExposedPinType(const FName& InPinName, const FStrin
 		TArray<URigVMPin*> ExtendedInterfacePins;
 		for (URigVMNode* Node : Graph->GetNodes())
 		{
-			if (Node->IsA<URigVMFunctionEntryNode>() || Node->IsA<URigVMFunctionReturnNode>())
+			if (Node->IsA<URigVMFunctionInterfaceNode>())
 			{
 				if (URigVMPin* InterfacePin = Node->FindPin(Pin->GetName()))
 				{
@@ -18006,7 +18467,7 @@ bool URigVMController::UpdateTemplateNodePinTypes(URigVMTemplateNode* InNode, bo
 
 	const FRigVMRegistry& Registry = FRigVMRegistry::Get();
 
-	if (InNode->IsA<URigVMFunctionEntryNode>() || InNode->IsA<URigVMFunctionReturnNode>())
+	if (InNode->IsA<URigVMFunctionInterfaceNode>())
 	{
 		return false;
 	}
diff --git a/Engine/Plugins/Runtime/RigVM/Source/RigVMDeveloper/Private/RigVMModel/RigVMControllerActions.cpp b/Engine/Plugins/Runtime/RigVM/Source/RigVMDeveloper/Private/RigVMModel/RigVMControllerActions.cpp
--- a/Engine/Plugins/Runtime/RigVM/Source/RigVMDeveloper/Private/RigVMModel/RigVMControllerActions.cpp
+++ b/Engine/Plugins/Runtime/RigVM/Source/RigVMDeveloper/Private/RigVMModel/RigVMControllerActions.cpp
@@ -1128,6 +1128,49 @@ bool FRigVMSetPinCategoryAction::Redo()
 	return FRigVMBaseAction::Redo();
 }
 
+FRigVMChangeNodePinCategoriesAction::FRigVMChangeNodePinCategoriesAction()
+: FRigVMBaseAction(nullptr)
+, NodeName()
+, OldCategories()
+, NewCategories()
+{
+}
+
+FRigVMChangeNodePinCategoriesAction::FRigVMChangeNodePinCategoriesAction(URigVMController* InController, const URigVMNode* InNode)
+: FRigVMBaseAction(InController)
+, NodeName(InNode->GetName())
+, OldCategories(InNode->GetPinCategories())
+, NewCategories()
+{
+}
+
+bool FRigVMChangeNodePinCategoriesAction::Undo()
+{
+	if(!FRigVMBaseAction::Undo())
+	{
+		return false;
+	}
+	return GetController()->SetPinCategories(*NodeName, OldCategories, false);
+}
+
+bool FRigVMChangeNodePinCategoriesAction::Redo()
+{
+	if(!CanUndoRedo())
+	{
+		return false;
+	}
+	if(!GetController()->SetPinCategories(*NodeName, NewCategories, false))
+	{
+		return false;
+	}
+	return FRigVMBaseAction::Redo();
+}
+
+void FRigVMChangeNodePinCategoriesAction::UpdateAfterModification(const URigVMNode* InNode)
+{
+	NewCategories = InNode->GetPinCategories();
+}
+
 FRigVMSetPinWatchAction::FRigVMSetPinWatchAction()
 : FRigVMBaseAction(nullptr)
 , OldIsWatched(false)
diff --git a/Engine/Plugins/Runtime/RigVM/Source/RigVMDeveloper/Private/RigVMModel/RigVMNode.cpp b/Engine/Plugins/Runtime/RigVM/Source/RigVMDeveloper/Private/RigVMModel/RigVMNode.cpp
--- a/Engine/Plugins/Runtime/RigVM/Source/RigVMDeveloper/Private/RigVMModel/RigVMNode.cpp
+++ b/Engine/Plugins/Runtime/RigVM/Source/RigVMDeveloper/Private/RigVMModel/RigVMNode.cpp
@@ -116,16 +116,6 @@ TArray<URigVMPin*> URigVMNode::GetAllPinsRecursively() const
 
 TArray<FString> URigVMNode::GetPinCategories() const
 {
-	const TArray<URigVMPin*> AllPins = GetAllPinsRecursively();
-
-	TArray<FString> PinCategories;
-	for(const URigVMPin* Pin : AllPins)
-	{
-		if(!Pin->UserDefinedCategory.IsEmpty())
-		{
-			PinCategories.Add(Pin->UserDefinedCategory);
-		}
-	}
 	return PinCategories;
 }
 
diff --git a/Engine/Plugins/Runtime/RigVM/Source/RigVMDeveloper/Private/RigVMModel/RigVMSchema.cpp b/Engine/Plugins/Runtime/RigVM/Source/RigVMDeveloper/Private/RigVMModel/RigVMSchema.cpp
--- a/Engine/Plugins/Runtime/RigVM/Source/RigVMDeveloper/Private/RigVMModel/RigVMSchema.cpp
+++ b/Engine/Plugins/Runtime/RigVM/Source/RigVMDeveloper/Private/RigVMModel/RigVMSchema.cpp
@@ -299,8 +299,7 @@ bool URigVMSchema::CanAddNode(URigVMController* InController, const URigVMNode*
 			}
 		}
 	}
-	else if(InNode->IsA<URigVMFunctionEntryNode>() ||
-		InNode->IsA<URigVMFunctionReturnNode>())
+	else if(InNode->IsA<URigVMFunctionInterfaceNode>())
 	{
 		// only allow entry / return nodes on sub graphs
 		if(Graph->IsRootGraph())
@@ -412,8 +411,7 @@ bool URigVMSchema::CanRemoveNode(URigVMController* InController, const URigVMNod
 {
 	RIGVMSCHEMA_DEFAULT_FUNCTION_BODY
 
-	if(InNode->IsA<URigVMFunctionEntryNode>() ||
-		InNode->IsA<URigVMFunctionReturnNode>())
+	if(InNode->IsA<URigVMFunctionInterfaceNode>())
 	{
 		return false;
 	}
diff --git a/Engine/Plugins/Runtime/RigVM/Source/RigVMDeveloper/Public/RigVMModel/Nodes/RigVMFunctionEntryNode.h b/Engine/Plugins/Runtime/RigVM/Source/RigVMDeveloper/Public/RigVMModel/Nodes/RigVMFunctionEntryNode.h
--- a/Engine/Plugins/Runtime/RigVM/Source/RigVMDeveloper/Public/RigVMModel/Nodes/RigVMFunctionEntryNode.h
+++ b/Engine/Plugins/Runtime/RigVM/Source/RigVMDeveloper/Public/RigVMModel/Nodes/RigVMFunctionEntryNode.h
@@ -2,7 +2,7 @@
 
 #pragma once
 
-#include "RigVMModel/Nodes/RigVMTemplateNode.h"
+#include "RigVMModel/Nodes/RigVMFunctionInterfaceNode.h"
 #include "RigVMFunctionEntryNode.generated.h"
 
 /**
@@ -10,7 +10,7 @@
  * input pins of the library node for links within.
  */
 UCLASS(BlueprintType)
-class RIGVMDEVELOPER_API URigVMFunctionEntryNode : public URigVMTemplateNode
+class RIGVMDEVELOPER_API URigVMFunctionEntryNode : public URigVMFunctionInterfaceNode
 {
 	GENERATED_BODY()
 
@@ -20,17 +20,12 @@ public:
 	virtual UScriptStruct* GetScriptStruct() const override { return nullptr; }
 	virtual const FRigVMTemplate* GetTemplate() const override { return nullptr; }
 	virtual FName GetNotation() const override { return NAME_None; }
-	virtual uint32 GetStructureHash() const override;
 
 	// Override node functions
-	virtual FLinearColor GetNodeColor() const override;
-	virtual bool IsDefinedAsVarying() const override;
 	virtual bool IsWithinLoop() const override;
 
 	// URigVMNode interface
 	virtual FString GetNodeTitle() const override;
-	virtual  FText GetToolTipText() const override;
-	virtual FText GetToolTipTextForPin(const URigVMPin* InPin) const override;
 
 private:
 
diff --git a/Engine/Plugins/Runtime/RigVM/Source/RigVMDeveloper/Public/RigVMModel/Nodes/RigVMFunctionInterfaceNode.h b/Engine/Plugins/Runtime/RigVM/Source/RigVMDeveloper/Public/RigVMModel/Nodes/RigVMFunctionInterfaceNode.h
new file mode 100644
--- /dev/null
+++ b/Engine/Plugins/Runtime/RigVM/Source/RigVMDeveloper/Public/RigVMModel/Nodes/RigVMFunctionInterfaceNode.h
@@ -0,0 +1,42 @@
+// Copyright Epic Games, Inc. All Rights Reserved.
+
+#pragma once
+
+#include "RigVMModel/Nodes/RigVMTemplateNode.h"
+#include "RigVMFunctionInterfaceNode.generated.h"
+
+/**
+ * The Function Interface node is is used as the base class for
+ * both the entry and return nodes.
+ */
+UCLASS(BlueprintType)
+class RIGVMDEVELOPER_API URigVMFunctionInterfaceNode : public URigVMTemplateNode
+{
+	GENERATED_BODY()
+
+public:
+
+	// Override template node functions
+	virtual uint32 GetStructureHash() const override;
+
+	// Override node functions
+	virtual FLinearColor GetNodeColor() const override;
+	virtual bool IsDefinedAsVarying() const override;
+	
+	// URigVMNode interface
+	virtual  FText GetToolTipText() const override;
+	virtual FText GetToolTipTextForPin(const URigVMPin* InPin) const override;	
+	virtual FName GetDisplayNameForPin(const FString& InPinPath) const override;
+	virtual FString GetCategoryForPin(const FString& InPinPath) const override;
+	virtual TArray<FString> GetPinCategories() const override;
+
+protected:
+
+	const URigVMPin* FindReferencedPin(const URigVMPin* InPin) const;
+	const URigVMPin* FindReferencedPin(const FString& InPinPath) const;
+
+private:
+
+	friend class URigVMController;
+};
+
diff --git a/Engine/Plugins/Runtime/RigVM/Source/RigVMDeveloper/Public/RigVMModel/Nodes/RigVMFunctionReferenceNode.h b/Engine/Plugins/Runtime/RigVM/Source/RigVMDeveloper/Public/RigVMModel/Nodes/RigVMFunctionReferenceNode.h
--- a/Engine/Plugins/Runtime/RigVM/Source/RigVMDeveloper/Public/RigVMModel/Nodes/RigVMFunctionReferenceNode.h
+++ b/Engine/Plugins/Runtime/RigVM/Source/RigVMDeveloper/Public/RigVMModel/Nodes/RigVMFunctionReferenceNode.h
@@ -24,6 +24,10 @@ public:
 	virtual FString GetNodeTitle() const override;
 	virtual FLinearColor GetNodeColor() const override;
 	virtual FText GetToolTipText() const override;
+	virtual FName GetDisplayNameForPin(const FString& InPinPath) const override;
+	virtual FString GetCategoryForPin(const FString& InPinPath) const override;
+	virtual FText GetToolTipTextForPin(const URigVMPin* InPin) const override;
+	virtual TArray<FString> GetPinCategories() const override;
 	// end URigVMNode interface
 
 	// URigVMLibraryNode interface
@@ -65,9 +69,10 @@ protected:
 	
 private:
 
-	virtual FText GetToolTipTextForPin(const URigVMPin* InPin) const override;
 	bool RequiresVariableRemappingInternal(TArray<FRigVMExternalVariable>& InnerVariables) const;
-	virtual TArray<int32> GetInstructionsForVMImpl(const FRigVMExtendedExecuteContext& Context, URigVM* InVM, const FRigVMASTProxy& InProxy = FRigVMASTProxy()) const override; 
+	virtual TArray<int32> GetInstructionsForVMImpl(const FRigVMExtendedExecuteContext& Context, URigVM* InVM, const FRigVMASTProxy& InProxy = FRigVMASTProxy()) const override;
+	const URigVMPin* FindReferencedPin(const URigVMPin* InPin) const;
+	const URigVMPin* FindReferencedPin(const FString& InPinPath) const;
 
 	//void SetReferencedFunctionData(FRigVMGraphFunctionData* Data);
 
diff --git a/Engine/Plugins/Runtime/RigVM/Source/RigVMDeveloper/Public/RigVMModel/Nodes/RigVMFunctionReturnNode.h b/Engine/Plugins/Runtime/RigVM/Source/RigVMDeveloper/Public/RigVMModel/Nodes/RigVMFunctionReturnNode.h
--- a/Engine/Plugins/Runtime/RigVM/Source/RigVMDeveloper/Public/RigVMModel/Nodes/RigVMFunctionReturnNode.h
+++ b/Engine/Plugins/Runtime/RigVM/Source/RigVMDeveloper/Public/RigVMModel/Nodes/RigVMFunctionReturnNode.h
@@ -2,7 +2,7 @@
 
 #pragma once
 
-#include "RigVMModel/Nodes/RigVMTemplateNode.h"
+#include "RigVMModel/Nodes/RigVMFunctionInterfaceNode.h"
 #include "RigVMFunctionReturnNode.generated.h"
 
 /**
@@ -10,7 +10,7 @@
  * output pins of the library node for links within.
  */
 UCLASS(BlueprintType)
-class RIGVMDEVELOPER_API URigVMFunctionReturnNode : public URigVMTemplateNode
+class RIGVMDEVELOPER_API URigVMFunctionReturnNode : public URigVMFunctionInterfaceNode
 {
 	GENERATED_BODY()
 
@@ -20,16 +20,9 @@ public:
 	virtual UScriptStruct* GetScriptStruct() const override { return nullptr; }
 	virtual const FRigVMTemplate* GetTemplate() const override { return nullptr; }
 	virtual FName GetNotation() const override { return NAME_None; }
-	virtual uint32 GetStructureHash() const override;
 
-	// Override node functions
-	virtual FLinearColor GetNodeColor() const override;
-	virtual bool IsDefinedAsVarying() const override;
-	
 	// URigVMNode interface
 	virtual FString GetNodeTitle() const override;
-	virtual  FText GetToolTipText() const override;
-	virtual FText GetToolTipTextForPin(const URigVMPin* InPin) const override;	
 
 private:
 
diff --git a/Engine/Plugins/Runtime/RigVM/Source/RigVMDeveloper/Public/RigVMModel/RigVMController.h b/Engine/Plugins/Runtime/RigVM/Source/RigVMDeveloper/Public/RigVMModel/RigVMController.h
--- a/Engine/Plugins/Runtime/RigVM/Source/RigVMDeveloper/Public/RigVMModel/RigVMController.h
+++ b/Engine/Plugins/Runtime/RigVM/Source/RigVMDeveloper/Public/RigVMModel/RigVMController.h
@@ -813,6 +813,18 @@ public:
 	UFUNCTION(BlueprintCallable, Category = RigVMController)
 	bool ClearPinCategory(const FString& InPinPath, bool bSetupUndoRedo = true, bool bPrintPythonCommand = false);
 
+	// Removes a pin category. The category is UI relevant only and used
+	// to order pins in the user interface of the node as well as on the details panel.
+	bool RemovePinCategory(const FName& InNodeName, const FString& InPinCategory, bool bSetupUndoRedo = true, bool bPrintPythonCommand = false);
+
+	// Renames a pin category. The category is UI relevant only and used
+	// to order pins in the user interface of the node as well as on the details panel.
+	bool RenamePinCategory(const FName& InNodeName, const FString& InOldPinCategory, const FString& InNewPinCategory, bool bSetupUndoRedo = true, bool bPrintPythonCommand = false);
+
+	// Changes a pin category's index. The category is UI relevant only and used
+	// to order pins in the user interface of the node as well as on the details panel.
+	bool SetPinCategoryIndex(const FName& InNodeName, const FString& InPinCategory, int32 InNewIndex, bool bSetupUndoRedo = true, bool bPrintPythonCommand = false);
+
 	// Returns the default value of a pin given its pinpath.
 	UFUNCTION(BlueprintCallable, Category = RigVMController)
 	FString GetPinDefaultValue(const FString& InPinPath);
@@ -1242,6 +1254,11 @@ private:
 	bool SetPinIsWatched(URigVMPin* InPin, bool bIsWatched, bool bSetupUndoRedo);
 	bool SetPinDisplayName(URigVMPin* InPin, const FString& InDisplayName, bool bSetupUndoRedo);
 	bool SetPinCategory(URigVMPin* InPin, const FString& InCategory, bool bSetupUndoRedo);
+	bool RemovePinCategory(const URigVMNode* InNode, const FString& InPinCategory, bool bSetupUndoRedo);
+	bool RenamePinCategory(const URigVMNode* InNode, const FString& InOldPinCategory, const FString& InNewPinCategory, bool bSetupUndoRedo);
+	bool SetPinCategoryIndex(const URigVMNode* InNode, const FString& InPinCategory, int32 InNewIndex, bool bSetupUndoRedo);
+	bool SetPinCategories(const FName& InNodeName, const TArray<FString>& InCategories, bool bSetupUndoRedo);
+	bool SetPinCategories(const URigVMNode* InNode, const TArray<FString>& InCategories, bool bSetupUndoRedo);
 	bool SetVariableName(URigVMVariableNode* InVariableNode, const FName& InVariableName, bool bSetupUndoRedo);
 	static void ForEveryPinRecursively(URigVMPin* InPin, TFunction<void(URigVMPin*)> OnEachPinFunction);
 	static void ForEveryPinRecursively(URigVMNode* InNode, TFunction<void(URigVMPin*)> OnEachPinFunction);
@@ -1555,6 +1572,7 @@ private:
 	friend struct FRigVMChangePinTypeAction;
 	friend struct FRigVMInjectNodeIntoPinAction;
 	friend struct FRigVMEjectNodeFromPinAction;
+	friend struct FRigVMChangeNodePinCategoriesAction;
 	friend class FRigVMParserAST;
 	friend class FRigVMControllerCompileBracketScope;
 	friend struct FRigVMPinInfoArray;
diff --git a/Engine/Plugins/Runtime/RigVM/Source/RigVMDeveloper/Public/RigVMModel/RigVMControllerActions.h b/Engine/Plugins/Runtime/RigVM/Source/RigVMDeveloper/Public/RigVMModel/RigVMControllerActions.h
--- a/Engine/Plugins/Runtime/RigVM/Source/RigVMDeveloper/Public/RigVMModel/RigVMControllerActions.h
+++ b/Engine/Plugins/Runtime/RigVM/Source/RigVMDeveloper/Public/RigVMModel/RigVMControllerActions.h
@@ -906,6 +906,35 @@ public:
 	FString NewCategory;
 };
 
+/**
+ * An action setting a pin's category in the graph.
+ */
+USTRUCT()
+struct FRigVMChangeNodePinCategoriesAction : public FRigVMBaseAction
+{
+	GENERATED_BODY()
+
+public:
+
+	FRigVMChangeNodePinCategoriesAction();
+	FRigVMChangeNodePinCategoriesAction(URigVMController* InController, const URigVMNode* InNode);
+	virtual ~FRigVMChangeNodePinCategoriesAction() {};
+	virtual UScriptStruct* GetScriptStruct() const override { return FRigVMSetPinCategoryAction::StaticStruct(); }
+	virtual bool Undo() override;
+	virtual bool Redo() override;
+	
+	void UpdateAfterModification(const URigVMNode* InNode);
+
+	UPROPERTY()
+	FString NodeName;
+
+	UPROPERTY()
+	TArray<FString> OldCategories;
+
+	UPROPERTY()
+	TArray<FString> NewCategories;
+};
+
 /**
  * An action setting a pin's default value in the graph.
  */
diff --git a/Engine/Plugins/Runtime/RigVM/Source/RigVMDeveloper/Public/RigVMModel/RigVMNode.h b/Engine/Plugins/Runtime/RigVM/Source/RigVMDeveloper/Public/RigVMModel/RigVMNode.h
--- a/Engine/Plugins/Runtime/RigVM/Source/RigVMDeveloper/Public/RigVMModel/RigVMNode.h
+++ b/Engine/Plugins/Runtime/RigVM/Source/RigVMDeveloper/Public/RigVMModel/RigVMNode.h
@@ -70,7 +70,7 @@ public:
 
 	// Returns all user defined categories on this node
 	UFUNCTION(BlueprintCallable, Category = RigVMNode)
-	TArray<FString> GetPinCategories() const;
+	virtual TArray<FString> GetPinCategories() const;
 
 	// Returns all pins for a given category
 	UFUNCTION(BlueprintCallable, Category = RigVMNode)
@@ -400,6 +400,12 @@ private:
 	UPROPERTY()
 	TArray<TObjectPtr<URigVMPin>> OrphanedPins;
 
+protected:
+	UPROPERTY()
+	TArray<FString> PinCategories;
+
+private:
+
 #if WITH_EDITOR
 	struct FProfilingCache
 	{
diff --git a/Engine/Plugins/Runtime/RigVM/Source/RigVMDeveloper/Public/RigVMModel/RigVMNotifications.h b/Engine/Plugins/Runtime/RigVM/Source/RigVMDeveloper/Public/RigVMModel/RigVMNotifications.h
--- a/Engine/Plugins/Runtime/RigVM/Source/RigVMDeveloper/Public/RigVMModel/RigVMNotifications.h
+++ b/Engine/Plugins/Runtime/RigVM/Source/RigVMDeveloper/Public/RigVMModel/RigVMNotifications.h
@@ -56,6 +56,7 @@ enum class ERigVMGraphNotifType : uint8
 	VariantTagsChanged, // The tags in the header of this function variant have changed (Subject == URigVMLibraryNode)
 	PinDisplayNameChanged, // The display name of a pin has changed - requiring a rebuild of the node user interface (Subject == URigVMPin) 
 	PinCategoryChanged, // The category of a pin has changed - requiring a rebuild of the node user interface (Subject == URigVMPin)
+	PinCategoriesChanged, // The category list of a node has changed - requiring a rebuild of the node user interface (Subject == URigVMNode)
 	Invalid // The max for this enum (used for guarding)
 };
 
diff --git a/Engine/Plugins/Runtime/RigVM/Source/RigVMDeveloper/Public/RigVMModel/RigVMPin.h b/Engine/Plugins/Runtime/RigVM/Source/RigVMDeveloper/Public/RigVMModel/RigVMPin.h
--- a/Engine/Plugins/Runtime/RigVM/Source/RigVMDeveloper/Public/RigVMModel/RigVMPin.h
+++ b/Engine/Plugins/Runtime/RigVM/Source/RigVMDeveloper/Public/RigVMModel/RigVMPin.h
@@ -606,6 +606,7 @@ private:
 	friend class FRigVMParserAST;
 	friend struct FRigVMSetPinDisplayNameAction;
 	friend struct FRigVMSetPinCategoryAction;
+	friend class URigVMLibraryNode;
 };
 
 class RIGVMDEVELOPER_API FRigVMPinDefaultValueImportErrorContext : public FOutputDevice
