Control Rig: Adding Init array node to resize and set a default value for elements
#rnx

diff --git a/Engine/Plugins/Runtime/RigVM/Source/RigVM/Private/RigVMFunctions/RigVMDispatch_Array.cpp b/Engine/Plugins/Runtime/RigVM/Source/RigVM/Private/RigVMFunctions/RigVMDispatch_Array.cpp
--- a/Engine/Plugins/Runtime/RigVM/Source/RigVM/Private/RigVMFunctions/RigVMDispatch_Array.cpp
+++ b/Engine/Plugins/Runtime/RigVM/Source/RigVM/Private/RigVMFunctions/RigVMDispatch_Array.cpp
@@ -476,6 +476,85 @@ FText FRigVMDispatch_ArraySetNum::GetArgumentTooltip(const FName& InArgumentName
 
 ////////////////////////////////////////////////////////////////////////////////////////////////////////////////////////
 
+FName FRigVMDispatch_ArrayInit::GetArgumentNameForOperandIndex(int32 InOperandIndex, int32 InTotalOperands) const
+{
+	static const FLazyName ArgumentNames[] = {
+		ArrayName,
+		NumName,
+		ElementName
+	};
+	check(InTotalOperands == UE_ARRAY_COUNT(ArgumentNames));
+	return ArgumentNames[InOperandIndex];
+}
+
+const TArray<FRigVMTemplateArgumentInfo>& FRigVMDispatch_ArrayInit::GetArgumentInfos() const
+{
+	static TArray<FRigVMTemplateArgumentInfo> Infos;
+	if (Infos.IsEmpty())
+	{
+		Infos.Emplace(CreateArgumentInfo(ArrayName, ERigVMPinDirection::IO));
+		Infos.Emplace(CreateArgumentInfo(NumName, ERigVMPinDirection::Input));
+		Infos.Emplace(CreateArgumentInfo(ElementName, ERigVMPinDirection::Input));
+	};
+	return Infos;
+}
+
+FRigVMTemplateTypeMap FRigVMDispatch_ArrayInit::OnNewArgumentType(const FName& InArgumentName, TRigVMTypeIndex InTypeIndex) const
+{
+	check(InArgumentName == ArrayName);
+	return {
+		{ArrayName, InTypeIndex},
+		{NumName, RigVMTypeUtils::TypeIndex::Int32},
+		{ElementName, FRigVMRegistry_NoLock::GetForRead().GetBaseTypeFromArrayTypeIndex_NoLock(InTypeIndex)}
+	};
+}
+
+void FRigVMDispatch_ArrayInit::Execute(FRigVMExtendedExecuteContext& InContext, FRigVMMemoryHandleArray Handles, FRigVMPredicateBranchArray Predicates)
+{
+	const FArrayProperty* ArrayProperty = CastFieldChecked<FArrayProperty>(Handles[0].GetProperty());
+	FScriptArrayHelper ArrayHelper(ArrayProperty, Handles[0].GetData());
+	const int32* Num = (int32*)Handles[1].GetData();
+	const FProperty* ValueProperty = Handles[2].GetProperty();
+	const uint8* ValueData = Handles[2].GetData();
+
+	if(InContext.IsValidArraySize(*Num))
+	{
+		ArrayHelper.EmptyAndAddValues(*Num);
+		for (int32 Index=0; Index<ArrayHelper.Num(); ++Index)
+		{
+			URigVMMemoryStorage::CopyProperty(ArrayProperty->Inner, ArrayHelper.GetElementPtr(Index), ValueProperty, ValueData);
+		}
+	}
+}
+
+#if WITH_EDITOR
+
+FText FRigVMDispatch_ArrayInit::GetNodeTooltip(const FRigVMTemplateTypeMap& InTypes) const
+{
+	return LOCTEXT("ArrayInitToolTip", "Sets the size of the array, initializing all elements to the given value.\nModifies the input array.");
+}
+
+FText FRigVMDispatch_ArrayInit::GetArgumentTooltip(const FName& InArgumentName, TRigVMTypeIndex InTypeIndex) const
+{
+	if(InArgumentName == ArrayName)
+	{
+		return LOCTEXT("ArrayInit_ArrayArgumentToolTip", "The array to initialize.");
+	}
+	if(InArgumentName == NumName)
+	{
+		return LOCTEXT("ArraySetNum_NumArgumentToolTip", "The new size of the array.");
+	}
+	if(InArgumentName == ElementName)
+	{
+		return LOCTEXT("ArrayInit_ElementArgumentToolTip", "The value that is set to all elements.");
+	}
+	return FText();
+}
+
+#endif
+
+////////////////////////////////////////////////////////////////////////////////////////////////////////////////////////
+
 FName FRigVMDispatch_ArrayGetAtIndex::GetArgumentNameForOperandIndex(int32 InOperandIndex, int32 InTotalOperands) const
 {
 	static const FLazyName ArgumentNames[] = {
diff --git a/Engine/Plugins/Runtime/RigVM/Source/RigVM/Public/RigVMFunctions/RigVMDispatch_Array.h b/Engine/Plugins/Runtime/RigVM/Source/RigVM/Public/RigVMFunctions/RigVMDispatch_Array.h
--- a/Engine/Plugins/Runtime/RigVM/Source/RigVM/Public/RigVMFunctions/RigVMDispatch_Array.h
+++ b/Engine/Plugins/Runtime/RigVM/Source/RigVM/Public/RigVMFunctions/RigVMDispatch_Array.h
@@ -150,6 +150,30 @@ protected:
 	static void Execute(FRigVMExtendedExecuteContext& InContext, FRigVMMemoryHandleArray Handles, FRigVMPredicateBranchArray Predicates);
 };
 
+USTRUCT(meta=(DisplayName = "Init"))
+struct RIGVM_API FRigVMDispatch_ArrayInit : public FRigVMDispatch_ArrayBaseMutable
+{
+	GENERATED_BODY()
+
+public:
+	FRigVMDispatch_ArrayInit()
+	{
+		FactoryScriptStruct = StaticStruct();
+	}
+
+	virtual FName GetArgumentNameForOperandIndex(int32 InOperandIndex, int32 InTotalOperands) const override;
+	virtual const TArray<FRigVMTemplateArgumentInfo>& GetArgumentInfos() const override;
+	virtual FRigVMTemplateTypeMap OnNewArgumentType(const FName& InArgumentName, TRigVMTypeIndex InTypeIndex) const override;
+#if WITH_EDITOR
+	virtual FText GetNodeTooltip(const FRigVMTemplateTypeMap& InTypes) const override;
+	virtual FText GetArgumentTooltip(const FName& InArgumentName, TRigVMTypeIndex InTypeIndex) const override;
+#endif
+
+protected:
+	virtual FRigVMFunctionPtr GetDispatchFunctionImpl(const FRigVMTemplateTypeMap& InTypes) const override { return &FRigVMDispatch_ArrayInit::Execute; }
+	static void Execute(FRigVMExtendedExecuteContext& InContext, FRigVMMemoryHandleArray Handles, FRigVMPredicateBranchArray Predicates);
+};
+
 USTRUCT(meta=(DisplayName = "At", Keywords = "Get Index,At Index,[]"))
 struct RIGVM_API FRigVMDispatch_ArrayGetAtIndex : public FRigVMDispatch_ArrayBase
 {
