Fixed incorrect TSoftObjectPtr type.

diff --git a/Engine/Plugins/Runtime/RigVM/Source/RigVMDeveloper/Private/RigVMBlueprint.cpp b/Engine/Plugins/Runtime/RigVM/Source/RigVMDeveloper/Private/RigVMBlueprint.cpp
--- a/Engine/Plugins/Runtime/RigVM/Source/RigVMDeveloper/Private/RigVMBlueprint.cpp
+++ b/Engine/Plugins/Runtime/RigVM/Source/RigVMDeveloper/Private/RigVMBlueprint.cpp
@@ -2489,7 +2489,7 @@ TArray<FAssetData> URigVMBlueprint::GetDependentAssets() const
 					{
 						if (const URigVMBlueprint* ControlRigBlueprint = ReferencePtr->GetTypedOuter<URigVMBlueprint>())
 						{
-							const TSoftObjectPtr<UPackage> Blueprint = ControlRigBlueprint;
+							const TSoftObjectPtr<const URigVMBlueprint> Blueprint = ControlRigBlueprint;
 							const FSoftObjectPath AssetPath = Blueprint.ToSoftObjectPath();
 							if(AssetPath.GetLongPackageName().StartsWith(TEXT("/Engine/Transient")))
 							{
