RigVM: support for pin categories - first steps on UI work

#rnx

diff --git a/Engine/Plugins/Runtime/RigVM/Source/RigVM/Private/RigVMCore/RigVMGraphFunctionDefinition.cpp b/Engine/Plugins/Runtime/RigVM/Source/RigVM/Private/RigVMCore/RigVMGraphFunctionDefinition.cpp
--- a/Engine/Plugins/Runtime/RigVM/Source/RigVM/Private/RigVMCore/RigVMGraphFunctionDefinition.cpp
+++ b/Engine/Plugins/Runtime/RigVM/Source/RigVM/Private/RigVMCore/RigVMGraphFunctionDefinition.cpp
@@ -563,3 +563,27 @@ bool FRigVMGraphFunctionData::PatchSharedArgumentOperandsIfRequired()
 	
 	return true;
 }
+
+FArchive& operator<<(FArchive& Ar, FRigVMGraphFunctionLayout& Layout)
+{
+	Ar.UsingCustomVersion(FRigVMObjectVersion::GUID);
+
+	Ar << Layout.Categories;
+	if(Ar.IsLoading())
+	{
+		if (Ar.CustomVer(FRigVMObjectVersion::GUID) < FRigVMObjectVersion::FunctionHeaderLayoutStoresPinIndexInCategory)
+		{
+			Layout.PinIndexInCategory.Reset();
+		}
+		else
+		{
+			Ar << Layout.PinIndexInCategory;
+		}
+	}
+	else
+	{
+		Ar << Layout.PinIndexInCategory;
+	}
+	Ar << Layout.DisplayNames;
+	return Ar;
+}
diff --git a/Engine/Plugins/Runtime/RigVM/Source/RigVM/Public/RigVMCore/RigVMGraphFunctionDefinition.h b/Engine/Plugins/Runtime/RigVM/Source/RigVM/Public/RigVMCore/RigVMGraphFunctionDefinition.h
--- a/Engine/Plugins/Runtime/RigVM/Source/RigVM/Public/RigVMCore/RigVMGraphFunctionDefinition.h
+++ b/Engine/Plugins/Runtime/RigVM/Source/RigVM/Public/RigVMCore/RigVMGraphFunctionDefinition.h
@@ -538,12 +538,16 @@ struct RIGVM_API FRigVMGraphFunctionLayout
 	UPROPERTY(BlueprintReadOnly, VisibleAnywhere, Category=FunctionArgument)
 	TArray<FRigVMGraphFunctionCategory> Categories;
 
+	UPROPERTY(BlueprintReadOnly, VisibleAnywhere, Category=FunctionArgument)
+	TMap<FString, int32> PinIndexInCategory;
+
 	UPROPERTY(BlueprintReadOnly, VisibleAnywhere, Category=FunctionArgument)
 	TMap<FString, FString> DisplayNames;
 
 	void Reset()
 	{
 		Categories.Reset();
+		PinIndexInCategory.Reset();
 		DisplayNames.Reset();
 	}
 
@@ -554,6 +558,10 @@ struct RIGVM_API FRigVMGraphFunctionLayout
 		{
 			Hash = HashCombine(Hash, GetTypeHash(Category));
 		}
+		for(const TPair<FString, int32>& Pair : Layout.PinIndexInCategory)
+		{
+			Hash = HashCombine(Hash, GetTypeHash(Pair));
+		}
 		for(const TPair<FString, FString>& Pair : Layout.DisplayNames)
 		{
 			Hash = HashCombine(Hash, GetTypeHash(Pair));
@@ -561,12 +569,7 @@ struct RIGVM_API FRigVMGraphFunctionLayout
 		return Hash;
 	}
 
-	friend FArchive& operator<<(FArchive& Ar, FRigVMGraphFunctionLayout& Layout)
-	{
-		Ar << Layout.Categories;
-		Ar << Layout.DisplayNames;
-		return Ar;
-	}
+	friend RIGVM_API FArchive& operator<<(FArchive& Ar, FRigVMGraphFunctionLayout& Layout);
 
 	const FString* FindCategory(const FString& InElement) const;
 	const FString* FindDisplayName(const FString& InElement) const;
diff --git a/Engine/Plugins/Runtime/RigVM/Source/RigVM/Public/RigVMObjectVersion.h b/Engine/Plugins/Runtime/RigVM/Source/RigVM/Public/RigVMObjectVersion.h
--- a/Engine/Plugins/Runtime/RigVM/Source/RigVM/Public/RigVMObjectVersion.h
+++ b/Engine/Plugins/Runtime/RigVM/Source/RigVM/Public/RigVMObjectVersion.h
@@ -51,6 +51,9 @@ struct RIGVM_API FRigVMObjectVersion
 		// Storing user interface layout within function header
 		FunctionHeaderStoresLayout,
 
+		// Storing user interface relevant pin index in category
+		FunctionHeaderLayoutStoresPinIndexInCategory,
+
 		// -----<new versions can be added above this line>-------------------------------------------------
 		VersionPlusOne,
 		LatestVersion = VersionPlusOne - 1
diff --git a/Engine/Plugins/Runtime/RigVM/Source/RigVMDeveloper/Private/EdGraph/RigVMEdGraph.cpp b/Engine/Plugins/Runtime/RigVM/Source/RigVMDeveloper/Private/EdGraph/RigVMEdGraph.cpp
--- a/Engine/Plugins/Runtime/RigVM/Source/RigVMDeveloper/Private/EdGraph/RigVMEdGraph.cpp
+++ b/Engine/Plugins/Runtime/RigVM/Source/RigVMDeveloper/Private/EdGraph/RigVMEdGraph.cpp
@@ -591,36 +591,55 @@ bool URigVMEdGraph::HandleModifiedEvent_Internal(ERigVMGraphNotifType InNotifTyp
 		case ERigVMGraphNotifType::PinDirectionChanged:
 		case ERigVMGraphNotifType::PinIndexChanged:
 		case ERigVMGraphNotifType::PinBoundVariableChanged:
+		{
+        	if (URigVMPin* ModelPin = Cast<URigVMPin>(InSubject))
+        	{
+        		if (URigVMEdGraphNode* RigNode = Cast<URigVMEdGraphNode>(FindNodeForModelNodeName(ModelPin->GetNode()->GetFName())))
+        		{
+        			RigNode->ModelPinsChanged();
+        		}
+        	}
+        	break;
+        }
 		case ERigVMGraphNotifType::PinCategoryChanged:
+		case ERigVMGraphNotifType::PinCategoriesChanged:
 		{
+			URigVMNode* ModelNode = nullptr;
 			if (URigVMPin* ModelPin = Cast<URigVMPin>(InSubject))
 			{
-				if (URigVMEdGraphNode* RigNode = Cast<URigVMEdGraphNode>(FindNodeForModelNodeName(ModelPin->GetNode()->GetFName())))
+				ModelNode = ModelPin->GetNode();
+			}
+			else
+			{
+				ModelNode = Cast<URigVMNode>(InSubject);
+			}
+			if(ModelNode)
+			{
+				if (URigVMEdGraphNode* RigNode = Cast<URigVMEdGraphNode>(FindNodeForModelNodeName(ModelNode->GetFName())))
 				{
-					RigNode->ModelPinsChanged();
+					RigNode->ModelPinsChanged(true);
 				}
 			}
 			break;
 		}
-		case ERigVMGraphNotifType::PinCategoriesChanged:
+		case ERigVMGraphNotifType::LibraryTemplateChanged:
 		{
-			if (URigVMNode* ModelNode = Cast<URigVMNode>(InSubject))
+			if (URigVMNode* LibraryNode = Cast<URigVMNode>(InSubject))
 			{
-				if (URigVMEdGraphNode* RigNode = Cast<URigVMEdGraphNode>(FindNodeForModelNodeName(ModelNode->GetFName())))
+				if (URigVMEdGraphNode* RigNode = Cast<URigVMEdGraphNode>(FindNodeForModelNodeName(LibraryNode->GetFName())))
 				{
-					RigNode->ModelPinsChanged();
+					RigNode->ModelPinsChanged(true);
 				}
 			}
 			break;
 		}
-		case ERigVMGraphNotifType::LibraryTemplateChanged:
 		case ERigVMGraphNotifType::PinDisplayNameChanged:
 		{
-			if (URigVMNode* LibraryNode = Cast<URigVMNode>(InSubject))
+			if (URigVMPin* ModelPin = Cast<URigVMPin>(InSubject))
 			{
-				if (URigVMEdGraphNode* RigNode = Cast<URigVMEdGraphNode>(FindNodeForModelNodeName(LibraryNode->GetFName())))
+				if (URigVMEdGraphNode* RigNode = Cast<URigVMEdGraphNode>(FindNodeForModelNodeName(ModelPin->GetNode()->GetFName())))
 				{
-					RigNode->ModelPinsChanged(true);
+					RigNode->SynchronizeGraphPinNameWithModelPin(ModelPin);
 				}
 			}
 			break;
@@ -711,6 +730,17 @@ bool URigVMEdGraph::HandleModifiedEvent_Internal(ERigVMGraphNotifType InNotifTyp
 			}
 			break;
 		}
+		case ERigVMGraphNotifType::PinCategoryExpansionChanged:
+		{
+			if (URigVMNode* ModelNode = Cast<URigVMNode>(InSubject))
+			{
+				if (URigVMEdGraphNode* RigNode = Cast<URigVMEdGraphNode>(FindNodeForModelNodeName(ModelNode->GetFName())))
+				{
+					RigNode->OnNodePinExpansionChanged().Broadcast();
+				}
+			}
+			break;
+		}
 		default:
 		{
 			break;
diff --git a/Engine/Plugins/Runtime/RigVM/Source/RigVMDeveloper/Private/EdGraph/RigVMEdGraphNode.cpp b/Engine/Plugins/Runtime/RigVM/Source/RigVMDeveloper/Private/EdGraph/RigVMEdGraphNode.cpp
--- a/Engine/Plugins/Runtime/RigVM/Source/RigVMDeveloper/Private/EdGraph/RigVMEdGraphNode.cpp
+++ b/Engine/Plugins/Runtime/RigVM/Source/RigVMDeveloper/Private/EdGraph/RigVMEdGraphNode.cpp
@@ -14,6 +14,7 @@
 #include "RigVMBlueprint.h"
 #include "PropertyPathHelpers.h"
 #include "RigVMBlueprintUtils.h"
+#include "RigVMStringUtils.h"
 #include "../../../RigVMEditor/Public/RigVMEditorModule.h"
 #include "RigVMCore/RigVMExecuteContext.h"
 #include "RigVMFunctions/RigVMDispatch_CastEnum.h"
@@ -570,6 +571,15 @@ UEdGraphPin* URigVMEdGraphNode::FindGraphPinFromModelPin(const URigVMPin* InMode
 	return nullptr;
 }
 
+UEdGraphPin* URigVMEdGraphNode::FindGraphPinFromCategory(const FString& InCategory, bool bAsInput) const
+{
+	if(const FPinPair* Pair = CachedCategoryPins.Find(InCategory))
+	{
+		return bAsInput ? Pair->InputPin : Pair->OutputPin;
+	}
+	return nullptr;
+}
+
 void URigVMEdGraphNode::SynchronizeGraphPinNameWithModelPin(const URigVMPin* InModelPin, bool bNotify)
 {
 	auto SyncGraphPinLambda = [this](const URigVMPin* InModelPin, bool bAsInput) -> bool
@@ -578,17 +588,24 @@ void URigVMEdGraphNode::SynchronizeGraphPinNameWithModelPin(const URigVMPin* InM
 		{
 			const FString OldPinName = GraphPin->PinName.ToString();
 			const FString NewPinName = InModelPin->GetPinPath();
+			const FText DisplayName = FText::FromName(InModelPin->GetDisplayName());
 
 			if(OldPinName != NewPinName)
 			{
 				PinPathToModelPin.Remove(OldPinName);
 				
 				GraphPin->PinName = *NewPinName;
-				GraphPin->PinFriendlyName = FText::FromName(InModelPin->GetDisplayName());
+				GraphPin->PinFriendlyName = DisplayName;
 
 				PinPathToModelPin.Add(NewPinName, (URigVMPin*)InModelPin);
 				return true;
 			}
+
+			if(!GraphPin->PinFriendlyName.EqualTo(DisplayName))
+			{
+				GraphPin->PinFriendlyName = DisplayName;
+				return true;
+			}
 		}
 
 		return false;
@@ -797,11 +814,97 @@ bool URigVMEdGraphNode::CreateGraphPinFromModelPin(const URigVMPin* InModelPin,
 	return Pair.IsValid() && bResult;
 }
 
+bool URigVMEdGraphNode::CreateGraphPinFromCategory(const FString& InCategory, EEdGraphPinDirection InDirection)
+{
+	const FPinPair PairConst = CachedCategoryPins.FindOrAdd(InCategory);
+
+	auto CreatePinLambda = [this](const FString& InCategory, const FString& InCategoryName, EEdGraphPinDirection InDirection, UEdGraphPin* InParentPin) -> UEdGraphPin*
+	{
+		// check if we already have a pin like this in the last pins arrays
+		static const FEdGraphPinType PinType = GetPinTypeForCategoryPin();
+
+		UEdGraphPin* GraphPin = nullptr;
+		UEdGraphPin** ExistingEdGraphPinPtr = LastEdGraphPins.FindByPredicate([InCategory, InDirection](const UEdGraphPin* ExistingPin) -> bool
+		{
+			return !ExistingPin->bWasTrashed &&
+				ExistingPin->GetName() == InCategory &&
+				ExistingPin->Direction == InDirection &&
+				ExistingPin->PinType == PinType;
+		});
+		if(ExistingEdGraphPinPtr)
+		{
+			GraphPin = *ExistingEdGraphPinPtr;
+			GraphPin->ParentPin = nullptr;
+			GraphPin->SubPins.Reset();
+			ExistingEdGraphPinPtr = nullptr;
+			LastEdGraphPins.Remove(GraphPin);
+			Pins.Add(GraphPin);
+		}
+
+		if(GraphPin == nullptr)
+		{
+			GraphPin = CreatePin(InDirection, PinType, FName(*InCategory));
+		}
+		
+		if (GraphPin)
+		{
+			GraphPin->bHidden = false;
+			GraphPin->PinFriendlyName = FText::FromString(InCategoryName);
+			GraphPin->bNotConnectable = true;
+			GraphPin->bOrphanedPin = false; 
+			GraphPin->bDisplayAsMutableRef = false;
+
+			if (InParentPin)
+			{
+				InParentPin->SubPins.Add(GraphPin);
+				GraphPin->ParentPin = InParentPin;
+			}
+		}
+
+		return GraphPin;
+	};
+
+	UEdGraphPin* InputPin = nullptr;
+	UEdGraphPin* OutputPin = nullptr;
+	UEdGraphPin* ParentPin = nullptr;
+
+	FString CategoryName = InCategory;
+	TArray<FString> Categories;
+	if(RigVMStringUtils::SplitNodePath(InCategory, Categories))
+	{
+		CategoryName = Categories.Pop();
+		if(Categories.Num() > 0)
+		{
+			const FString ParentCategory = RigVMStringUtils::JoinNodePath(Categories);
+			(void)CreateGraphPinFromCategory(ParentCategory, InDirection);
+			ParentPin = FindGraphPinFromCategory(ParentCategory, InDirection == EGPD_Input);
+		}
+	}
+
+	bool bResult = false;
+	if (InDirection == EGPD_Input && PairConst.InputPin == nullptr)
+	{
+		InputPin = CreatePinLambda(InCategory, CategoryName, EGPD_Input, ParentPin);
+		bResult = true;
+	}
+	if (InDirection == EGPD_Output && PairConst.OutputPin == nullptr)
+	{
+		OutputPin = CreatePinLambda(InCategory, CategoryName, EGPD_Output, ParentPin);
+		bResult = true;
+	}
+
+	FPinPair& Pair = CachedCategoryPins.FindChecked(InCategory);
+	Pair.InputPin = InputPin != nullptr ? InputPin : Pair.InputPin;
+	Pair.OutputPin = OutputPin != nullptr ? OutputPin : Pair.OutputPin;
+	return Pair.IsValid() && bResult;
+}
+
 void URigVMEdGraphNode::RemoveGraphSubPins(UEdGraphPin* InParentPin, const TArray<UEdGraphPin*>& InPinsToKeep)
 {
 	TArray<UEdGraphPin*> SubPins = InParentPin->SubPins;
 
 	TArray<URigVMPin*> ModelSubPins;
+	TArray<FString> CategorySubPins;
 	for (const UEdGraphPin* SubPin: SubPins)
 	{
 		if(InPinsToKeep.Contains(SubPin))
@@ -816,6 +919,14 @@ void URigVMEdGraphNode::RemoveGraphSubPins(UEdGraphPin* InParentPin, const TArra
 				ModelSubPins.Add(Pair.Key);
 			}
 		}
+		for(const auto& Pair : CachedCategoryPins)
+		{
+			if(Pair.Value.InputPin == SubPin ||
+				Pair.Value.OutputPin == SubPin)
+			{
+				CategorySubPins.Add(Pair.Key);
+			}
+		}
 	}
 
 	for(const URigVMPin* ModelSubPin : ModelSubPins)
@@ -848,6 +959,37 @@ void URigVMEdGraphNode::RemoveGraphSubPins(UEdGraphPin* InParentPin, const TArra
 
 		CachedPins.Remove(ModelSubPin);
 	}
+
+	for(const FString& CategorySubPin : CategorySubPins)
+	{
+		PinPathToModelPin.Remove(CategorySubPin);
+
+		if(const FPinPair* PinPair = CachedCategoryPins.Find(CategorySubPin))
+		{
+			auto Traverse = [this](UEdGraphPin* SubPin)
+			{
+				if(SubPin == nullptr)
+				{
+					return;
+				}
+				
+				// Remove this pin from our owned pins
+				Pins.Remove(SubPin);
+
+				if (!SubPin->SubPins.IsEmpty())
+				{
+					RemoveGraphSubPins(SubPin);
+				}
+
+				SubPin->MarkAsGarbage();
+			};
+			
+			Traverse(PinPair->InputPin);
+			Traverse(PinPair->OutputPin);
+		}
+
+		CachedCategoryPins.Remove(CategorySubPin);
+	}
 	
 	InParentPin->SubPins.RemoveAll([InPinsToKeep](const UEdGraphPin* Pin) -> bool
 	{
@@ -874,6 +1016,37 @@ bool URigVMEdGraphNode::ModelPinsChanged(bool bForce)
 		return PinsAdded;
 	};
 
+	auto AddMissingCategoryPins = [this]()
+	{
+		int32 PinsAdded = 0;
+		if(const URigVMNode* ModelNode = GetModelNode())
+		{
+			for(const FString& Category : ModelNode->GetPinCategories())
+			{
+				if(FindGraphPinFromCategory(Category, true) == nullptr && 
+					FindGraphPinFromCategory(Category, false) == nullptr)
+				{
+					const TArray<URigVMPin*> PinsForCategory = ModelNode->GetPinsForCategory(Category);
+					if(!PinsForCategory.IsEmpty())
+					{
+						const ERigVMPinDirection PinDirection = PinsForCategory[0]->GetDirection();
+						if(PinDirection == ERigVMPinDirection::Input ||
+							PinDirection == ERigVMPinDirection::Visible ||
+							PinDirection == ERigVMPinDirection::IO)
+						{
+							PinsAdded += CategoryPinAdded_Internal(Category, EGPD_Input) ? 1 : 0;
+						}
+						if(PinDirection == ERigVMPinDirection::Output ||
+							PinDirection == ERigVMPinDirection::IO)
+						{
+							PinsAdded += CategoryPinAdded_Internal(Category, EGPD_Output) ? 1 : 0;
+						}
+					}
+				}
+			}
+		}
+		return PinsAdded;
+	};
 	auto RemoveObsoletePins = [this]()
 	{
 		TArray<URigVMPin*> PinsToRemove; 
@@ -901,6 +1074,22 @@ bool URigVMEdGraphNode::ModelPinsChanged(bool bForce)
 		{
 			PinsRemoved += ModelPinRemoved_Internal(ModelPin) ? 1 : 0;
 		}
+
+		if(const URigVMNode* ModelNode = GetModelNode())
+		{
+			for(const auto& Pair : CachedCategoryPins)
+			{
+				const FString Category = Pair.Key;
+				if(!ModelNode->GetPinCategories().Contains(Category))
+				{
+					PinsRemoved += CategoryPinRemoved_Internal(Category);
+				}
+				else if(ModelNode->GetPinsForCategory(Category).IsEmpty())
+				{
+					PinsRemoved += CategoryPinRemoved_Internal(Category);
+				}
+			}
+		}
 		return PinsRemoved;
 	};
 
@@ -955,8 +1144,21 @@ bool URigVMEdGraphNode::ModelPinsChanged(bool bForce)
 
 		return PinsToReorder;
 	};
-	
+
+	auto ReparentPins = [this](const TArray<URigVMPin*>& InModelPins) -> int32
+	{
+		// todo
+		return 0;
+	};
+
+	auto ReparentCategoryPins = [this](const TArray<FString>& InCategories) -> int32
+	{
+		// todo
+		return 0;
+	};
+
 	int32 PinsAdded = 0;
+	PinsAdded += AddMissingCategoryPins();
 	PinsAdded += AddMissingPins(ExecutePins);
 	PinsAdded += AddMissingPins(OutputPins);
 	PinsAdded += AddMissingPins(InputOutputPins);
@@ -972,7 +1174,17 @@ bool URigVMEdGraphNode::ModelPinsChanged(bool bForce)
 	PinsReordered += OrderPins(OutputPins);
 	PinsReordered += OrderPins(ExecutePins);
 
-	const bool bResult = (PinsAdded > 0) || (PinsRemoved > 0) || (PinsReordered > 0); 
+	int32 PinsReparented = 0;
+	if(const URigVMNode* ModelNode = GetModelNode())
+	{
+		PinsReparented += ReparentCategoryPins(ModelNode->GetPinCategories());
+		PinsReparented += ReparentPins(InputPins);
+		PinsReparented += ReparentPins(InputOutputPins);
+		PinsReparented += ReparentPins(OutputPins);
+		PinsReparented += ReparentPins(ExecutePins);
+	}
+
+	const bool bResult = (PinsAdded > 0) || (PinsRemoved > 0) || (PinsReordered > 0) || (PinsReparented > 0); 
 	if(bResult || bForce)
 	{
 		OnNodePinsChanged().Broadcast();
@@ -1012,11 +1224,22 @@ bool URigVMEdGraphNode::ModelPinAdded_Internal(const URigVMPin* InModelPin)
 
 	UEdGraphPin* InputParentPin = FindGraphPinFromModelPin(InModelPin->GetParentPin(), true);
 	UEdGraphPin* OutputParentPin = FindGraphPinFromModelPin(InModelPin->GetParentPin(), false);
+
+	const FString Category = InModelPin->GetCategory();
 		
 	if(InModelPin->GetDirection() == ERigVMPinDirection::Input ||
 		InModelPin->GetDirection() == ERigVMPinDirection::Visible ||
 		InModelPin->GetDirection() == ERigVMPinDirection::IO)
 	{
+		if(!Category.IsEmpty())
+		{
+			(void)CreateGraphPinFromCategory(Category, EGPD_Input);
+			if(UEdGraphPin* CategoryPin = FindGraphPinFromCategory(Category, true))
+			{
+				InputParentPin = CategoryPin;
+			}
+		}
+
 		if(CreateGraphPinFromModelPin(InModelPin, EGPD_Input, InputParentPin))
 		{
 			bResult = true;
@@ -1026,6 +1249,15 @@ bool URigVMEdGraphNode::ModelPinAdded_Internal(const URigVMPin* InModelPin)
 	if(InModelPin->GetDirection() == ERigVMPinDirection::Output ||
 		InModelPin->GetDirection() == ERigVMPinDirection::IO)
 	{
+		if(!Category.IsEmpty())
+		{
+			(void)CreateGraphPinFromCategory(Category, EGPD_Output);
+			if(UEdGraphPin* CategoryPin = FindGraphPinFromCategory(Category, false))
+			{
+				OutputParentPin = CategoryPin;
+			}
+		}
+
 		if(CreateGraphPinFromModelPin(InModelPin, EGPD_Output, OutputParentPin))
 		{
 			bResult = true;
@@ -1138,6 +1370,30 @@ bool URigVMEdGraphNode::ModelPinRemoved_Internal(const URigVMPin* InModelPin)
 	return bResult;
 }
 
+bool URigVMEdGraphNode::CategoryPinAdded_Internal(const FString& InCategory, EEdGraphPinDirection InDirection)
+{
+	return CreateGraphPinFromCategory(InCategory, InDirection);
+}
+
+bool URigVMEdGraphNode::CategoryPinRemoved_Internal(const FString& InCategory)
+{
+	auto RemoveGraphPin = [this](const FString& InCategory, bool bAsInput)
+	{
+		if(UEdGraphPin* GraphPin = FindGraphPinFromCategory(InCategory, bAsInput))
+		{
+			RemoveGraphSubPins(GraphPin);
+			Pins.Remove(GraphPin);
+			GraphPin->MarkAsGarbage();
+		}
+	};
+	RemoveGraphPin(InCategory, true);
+	RemoveGraphPin(InCategory, false);
+		
+	const bool bResult = PinPathToModelPin.Remove(InCategory) > 0;
+	CachedCategoryPins.Remove(InCategory);
+	return bResult;
+}
+
 FLinearColor URigVMEdGraphNode::GetNodeProfilingColor() const
 {
 #if WITH_EDITOR
@@ -1198,7 +1454,7 @@ void URigVMEdGraphNode::UpdatePinLists()
 	InputOutputPins.Reset();
 	InputPins.Reset();
 
-	if (URigVMNode* ModelNode = GetModelNode())
+	if (const URigVMNode* ModelNode = GetModelNode())
 	{
 		for(int32 PinListIndex=0; PinListIndex<2; PinListIndex++)
 		{
@@ -1241,6 +1497,7 @@ void URigVMEdGraphNode::AllocateDefaultPins()
 	UpdatePinLists();
 
 	CachedPins.Reset();
+	CachedCategoryPins.Reset();
 	PinPathToModelPin.Reset();
 
 	auto CreateGraphPins = [this](const TArray<URigVMPin*>& InModelPins)
@@ -1804,6 +2061,12 @@ FEdGraphPinType URigVMEdGraphNode::GetPinTypeForModelPin(const URigVMPin* InMode
 	return PinType;
 }
 
+FEdGraphPinType URigVMEdGraphNode::GetPinTypeForCategoryPin()
+{
+	static const FEdGraphPinType PinType = RigVMTypeUtils::PinTypeFromCPPType(NAME_None, FRigVMGraphFunctionCategory::StaticStruct());
+	return PinType;
+}
+
 void URigVMEdGraphNode::ConfigurePin(UEdGraphPin* EdGraphPin, const URigVMPin* ModelPin) const
 {
 	const bool bConnectable =
diff --git a/Engine/Plugins/Runtime/RigVM/Source/RigVMDeveloper/Private/RigVMModel/Nodes/RigVMFunctionInterfaceNode.cpp b/Engine/Plugins/Runtime/RigVM/Source/RigVMDeveloper/Private/RigVMModel/Nodes/RigVMFunctionInterfaceNode.cpp
--- a/Engine/Plugins/Runtime/RigVM/Source/RigVMDeveloper/Private/RigVMModel/Nodes/RigVMFunctionInterfaceNode.cpp
+++ b/Engine/Plugins/Runtime/RigVM/Source/RigVMDeveloper/Private/RigVMModel/Nodes/RigVMFunctionInterfaceNode.cpp
@@ -61,6 +61,15 @@ FString URigVMFunctionInterfaceNode::GetCategoryForPin(const FString& InPinPath)
 	return Super::GetCategoryForPin(InPinPath);
 }
 
+int32 URigVMFunctionInterfaceNode::GetIndexInCategoryForPin(const FString& InPinPath) const
+{
+	if(const URigVMPin* OuterPin = FindReferencedPin(InPinPath))
+	{
+		return OuterPin->GetIndexInCategory();
+	}
+	return Super::GetIndexInCategoryForPin(InPinPath);
+}
+
 TArray<FString> URigVMFunctionInterfaceNode::GetPinCategories() const
 {
 	if(const URigVMCollapseNode* OuterNode = Cast<URigVMCollapseNode>(GetGraph()->GetOuter()))
diff --git a/Engine/Plugins/Runtime/RigVM/Source/RigVMDeveloper/Private/RigVMModel/Nodes/RigVMFunctionReferenceNode.cpp b/Engine/Plugins/Runtime/RigVM/Source/RigVMDeveloper/Private/RigVMModel/Nodes/RigVMFunctionReferenceNode.cpp
--- a/Engine/Plugins/Runtime/RigVM/Source/RigVMDeveloper/Private/RigVMModel/Nodes/RigVMFunctionReferenceNode.cpp
+++ b/Engine/Plugins/Runtime/RigVM/Source/RigVMDeveloper/Private/RigVMModel/Nodes/RigVMFunctionReferenceNode.cpp
@@ -74,6 +74,21 @@ FString URigVMFunctionReferenceNode::GetCategoryForPin(const FString& InPinPath)
 	return Super::GetCategoryForPin(InPinPath);
 }
 
+int32 URigVMFunctionReferenceNode::GetIndexInCategoryForPin(const FString& InPinPath) const
+{
+	if(const int32* Index = ReferencedFunctionHeader.Layout.PinIndexInCategory.Find(InPinPath))
+	{
+		return *Index;
+	}
+
+	if(const URigVMPin* ReferencedPin = FindReferencedPin(InPinPath))
+	{
+		return ReferencedPin->GetIndexInCategory();
+	}
+
+	return INDEX_NONE;
+}
+
 FString URigVMFunctionReferenceNode::GetNodeCategory() const
 {
 	return ReferencedFunctionHeader.Category;
diff --git a/Engine/Plugins/Runtime/RigVM/Source/RigVMDeveloper/Private/RigVMModel/Nodes/RigVMLibraryNode.cpp b/Engine/Plugins/Runtime/RigVM/Source/RigVMDeveloper/Private/RigVMModel/Nodes/RigVMLibraryNode.cpp
--- a/Engine/Plugins/Runtime/RigVM/Source/RigVMDeveloper/Private/RigVMModel/Nodes/RigVMLibraryNode.cpp
+++ b/Engine/Plugins/Runtime/RigVM/Source/RigVMDeveloper/Private/RigVMModel/Nodes/RigVMLibraryNode.cpp
@@ -353,6 +353,10 @@ FRigVMGraphFunctionHeader URigVMLibraryNode::GetFunctionHeader(IRigVMGraphFuncti
 	// fill in all user provided display names
 	for(const URigVMPin* Pin : AllPins)
 	{
+		if(!Pin->GetCategory().IsEmpty() && Pin->GetIndexInCategory() != INDEX_NONE)
+		{
+			Header.Layout.PinIndexInCategory.Add(Pin->GetSegmentPath(true), Pin->GetIndexInCategory());
+		}
 		if(!Pin->DisplayName.IsNone())
 		{
 			Header.Layout.DisplayNames.Add(Pin->GetSegmentPath(true), Pin->DisplayName.ToString());
diff --git a/Engine/Plugins/Runtime/RigVM/Source/RigVMDeveloper/Private/RigVMModel/RigVMController.cpp b/Engine/Plugins/Runtime/RigVM/Source/RigVMDeveloper/Private/RigVMModel/RigVMController.cpp
--- a/Engine/Plugins/Runtime/RigVM/Source/RigVMDeveloper/Private/RigVMModel/RigVMController.cpp
+++ b/Engine/Plugins/Runtime/RigVM/Source/RigVMDeveloper/Private/RigVMModel/RigVMController.cpp
@@ -8092,6 +8092,51 @@ bool URigVMController::SetPinCategory(URigVMPin* InPin, const FString& InCategor
 		}
 	}
 
+	// you cannot set categories on IO, outputs or hidden pins
+	if(InPin->GetDirection() == ERigVMPinDirection::IO ||
+		InPin->GetDirection() == ERigVMPinDirection::Output ||
+		InPin->GetDirection() == ERigVMPinDirection::Hidden)
+	{
+		ReportErrorf(
+			TEXT("Cannot set category (%s) on pin '%s' - categories not allowed for pin direction."),
+			*NewCategory,
+			*InPin->GetPinPath()
+		);
+		return false;
+	}
+
+	TArray<FString> ParentAndThisCategories = InPin->GetNode()->GetParentPinCategories(NewCategory, false, true);
+	Algo::Reverse(ParentAndThisCategories);
+	for(const FString& CategoryOnPath : ParentAndThisCategories)
+	{
+		auto IsInput = [](const URigVMPin* InPin)
+		{
+			return InPin->GetDirection() == ERigVMPinDirection::Input || InPin->GetDirection() == ERigVMPinDirection::Visible;
+		};
+		
+		// you cannot have inputs and IO on the same category or on parent categories
+		const bool bIsInput = IsInput(InPin);
+		const TArray<URigVMPin*> AllPins = InPin->GetNode()->GetAllPinsRecursively();
+		for(const URigVMPin* OtherPin : AllPins)
+		{
+			if(OtherPin != InPin)
+			{
+				if(OtherPin->GetCategory().Equals(CategoryOnPath, ESearchCase::CaseSensitive))
+				{
+					if(bIsInput != IsInput(OtherPin))
+					{
+						ReportErrorf(
+							TEXT("Cannot set category (%s) on pin '%s', it's already used on pin '%s' with a different direction."),
+							*NewCategory,
+							*InPin->GetPinPath(),
+							*OtherPin->GetPinPath()
+						);
+					}
+				}
+			}
+		}
+	}
+
 	FRigVMControllerCompileBracketScope CompileScope(this);
 	FRigVMSetPinCategoryAction Action;
 	if (bSetupUndoRedo)
@@ -8107,12 +8152,14 @@ bool URigVMController::SetPinCategory(URigVMPin* InPin, const FString& InCategor
 	{
 		if(bSetupUndoRedo)
 		{
-			if(!InPin->GetNode()->PinCategories.Contains(NewCategory))
+			TArray<FString> CombinedCategories = InPin->GetNode()->PinCategories;
+			const TArray<FString> CategoriesToAdd = InPin->GetNode()->GetParentPinCategories(NewCategory, false, true);
+			Algo::Reverse(CategoriesToAdd);
+			for(const FString& CategoryToAdd : CategoriesToAdd)
 			{
-				TArray<FString> NewCategories = InPin->GetNode()->PinCategories;
-				NewCategories.Add(NewCategory);
-				(void)SetPinCategories(InPin->GetNode(), NewCategories, bSetupUndoRedo);
+				CombinedCategories.AddUnique(CategoryToAdd);
 			}
+			(void)SetPinCategories(InPin->GetNode(), CombinedCategories, bSetupUndoRedo);
 		}
 	}
 
@@ -8191,10 +8238,19 @@ bool URigVMController::RemovePinCategory(const URigVMNode* InNode, const FString
 		return false;
 	}
 
-	const TArray<URigVMPin*> PinsForCategory = InNode->GetPinsForCategory(InPinCategory);
-	if(!InNode->GetPinCategories().Contains(InPinCategory) && PinsForCategory.IsEmpty())
+	const TArray<FString> AllCategories = InNode->GetPinCategories();
+	TArray<FString> CategoriesToRemove = { InPinCategory };
+	CategoriesToRemove.Append(InNode->GetSubPinCategories(InPinCategory, false, true));
+
+	CategoriesToRemove.RemoveAll([&AllCategories](const FString InCategory) -> bool
 	{
-		return false;
+		return !AllCategories.Contains(InCategory);
+	});
+
+	TArray<URigVMPin*> PinsForCategory;
+	for(const FString& CategoryToRemove : CategoriesToRemove)
+	{
+		PinsForCategory.Append(InNode->GetPinsForCategory(CategoryToRemove));
 	}
 
 	FRigVMControllerCompileBracketScope CompileScope(this);
@@ -8211,7 +8267,10 @@ bool URigVMController::RemovePinCategory(const URigVMNode* InNode, const FString
 		(void)SetPinCategory(PinInCategory, FString(), bSetupUndoRedo);
 	}
 
-	const_cast<URigVMNode*>(InNode)->PinCategories.Remove(InPinCategory);
+	for(const FString& CategoryToRemove : CategoriesToRemove)
+	{
+		const_cast<URigVMNode*>(InNode)->PinCategories.Remove(CategoryToRemove);
+	}
 	Notify(ERigVMGraphNotifType::PinCategoriesChanged, const_cast<URigVMNode*>(InNode));
 
 	if (!bSuspendNotifications)
@@ -8455,6 +8514,203 @@ bool URigVMController::SetPinCategoryIndex(const URigVMNode* InNode, const FStri
 	return SetPinCategories(InNode, Categories, bSetupUndoRedo);
 }
 
+bool URigVMController::SetPinCategoryExpansion(const FName& InNodeName, const FString& InPinCategory, bool bIsExpanded, bool bSetupUndoRedo, bool bPrintPythonCommand)
+{
+	if (!IsValidGraph())
+	{
+		return false;
+	}
+
+	if (!bIsTransacting && !IsGraphEditable())
+	{
+		return false;
+	}
+
+	const URigVMGraph* Graph = GetGraph();
+	check(Graph);
+
+	const URigVMNode* Node = Graph->FindNodeByName(InNodeName);
+	if (Node == nullptr)
+	{
+		ReportErrorf(TEXT("Cannot find node '%s'."), *InNodeName.ToString());
+		return false;
+	}
+
+	const bool bSuccess = SetPinCategoryExpansion(Node, InPinCategory, bIsExpanded, bSetupUndoRedo);
+	if (bSuccess && bPrintPythonCommand)
+	{
+		const FString GraphName = GetSchema()->GetSanitizedGraphName(GetGraph()->GetGraphName());
+
+		RigVMPythonUtils::Print(GetSchema()->GetGraphOuterName(GetGraph()),
+			FString::Printf(TEXT("blueprint.get_controller_by_name('%s').set_pin_category_expansion('%s', '%s', %s)"),
+			*GraphName,
+			*GetSchema()->GetSanitizedNodeName(InNodeName.ToString()),
+			*InPinCategory,
+			(bIsExpanded) ? TEXT("True") : TEXT("False")));
+	}
+
+	return bSuccess;
+}
+
+bool URigVMController::SetPinCategoryExpansion(const URigVMNode* InNode, const FString& InPinCategory, bool bIsExpanded, bool bSetupUndoRedo)
+{
+	if(!IsValidNodeForGraph(InNode))
+	{
+		return false;
+	}
+
+	const TArray<FString> AllCategories = InNode->GetPinCategories();
+	if(!AllCategories.Contains(InPinCategory))
+	{
+		ReportErrorf(TEXT("Cannot find category '%s' on node '%s'."), *InPinCategory, *InNode->GetName());
+		return false;
+	}
+
+	if(InNode->IsPinCategoryExpanded(InPinCategory) == bIsExpanded)
+	{
+		return false;
+	}
+
+	if(InNode->GetPinsForCategory(InPinCategory).IsEmpty() &&
+		InNode->GetSubPinCategories(InPinCategory).IsEmpty())
+	{
+		return false;
+	}
+	
+	FRigVMControllerCompileBracketScope CompileScope(this);
+	FRigVMSetPinCategoryExpansionAction Action;
+	if (bSetupUndoRedo)
+	{
+		Action = FRigVMSetPinCategoryExpansionAction(this, InNode, InPinCategory);
+		if(bIsExpanded)
+		{
+			Action.SetTitle(TEXT("Collapse Pin Category"));
+		}
+		else
+		{
+			Action.SetTitle(TEXT("Expand Pin Category"));
+		}
+		GetActionStack()->BeginAction(Action);
+	}
+
+	if(bIsExpanded)
+	{
+		const_cast<URigVMNode*>(InNode)->PinCategoryExpansion.FindOrAdd(InPinCategory) = true;
+	}
+	else
+	{
+		const_cast<URigVMNode*>(InNode)->PinCategoryExpansion.Remove(InPinCategory);
+	}
+	const_cast<URigVMNode*>(InNode)->LastAffectedPinCategory = InPinCategory;
+	Notify(ERigVMGraphNotifType::PinCategoryExpansionChanged, const_cast<URigVMNode*>(InNode));
+
+	if (!bSuspendNotifications)
+	{
+		const URigVMGraph* Graph = GetGraph();
+		check(Graph);
+		(void)Graph->MarkPackageDirty();
+	}
+
+	if (bSetupUndoRedo)
+	{
+		GetActionStack()->EndAction(Action);
+	}
+
+	return true;
+}
+
+bool URigVMController::SetPinIndexInCategory(const FString& InPinPath, int32 InIndexInCategory, bool bSetupUndoRedo, bool bPrintPythonCommand)
+{
+	if (!IsValidGraph())
+	{
+		return false;
+	}
+
+	if (!bIsTransacting && !IsGraphEditable())
+	{
+		return false;
+	}
+
+	const URigVMGraph* Graph = GetGraph();
+	check(Graph);
+
+	URigVMPin* Pin = Graph->FindPin(InPinPath);
+	if (Pin == nullptr)
+	{
+		ReportErrorf(TEXT("Cannot find pin '%s'."), *InPinPath);
+		return false;
+	}
+
+	const bool bSuccess = SetPinIndexInCategory(Pin, InIndexInCategory, bSetupUndoRedo);
+	if (bSuccess && bPrintPythonCommand)
+	{
+		const FString GraphName = GetSchema()->GetSanitizedGraphName(GetGraph()->GetGraphName());
+
+		RigVMPythonUtils::Print(GetSchema()->GetGraphOuterName(GetGraph()),
+			FString::Printf(TEXT("blueprint.get_controller_by_name('%s').set_pin_index_in_category('%s', %d)"),
+			*GraphName,
+			*GetSchema()->GetSanitizedPinPath(InPinPath),
+			InIndexInCategory));
+	}
+
+	return bSuccess;
+}
+
+bool URigVMController::SetPinIndexInCategory(URigVMPin* InPin, int32 InIndexInCategory, bool bSetupUndoRedo)
+{
+	if(!IsValidPinForGraph(InPin))
+	{
+		return false;
+	}
+
+	// for entry or return nodes we relay to the outer pins
+	if(const URigVMFunctionInterfaceNode* InterfaceNode = Cast<URigVMFunctionInterfaceNode>(InPin->GetNode()))
+	{
+		if(const URigVMPin* OuterPin = InterfaceNode->FindReferencedPin(InPin))
+		{
+			if(URigVMController* OuterController = GetControllerForGraph(OuterPin->GetGraph()))
+			{
+				return OuterController->SetPinIndexInCategory(OuterPin->GetPinPath(), InIndexInCategory, bSetupUndoRedo);
+			}
+		}
+	}
+	else if(InPin->GetNode()->IsA<URigVMFunctionReferenceNode>())
+	{
+		return false;
+	}
+
+	if(InPin->GetIndexInCategory() == InIndexInCategory)
+	{
+		return false;
+	}
+
+	FRigVMControllerCompileBracketScope CompileScope(this);
+	FRigVMSetPinIndexInCategoryAction Action;
+	if (bSetupUndoRedo)
+	{
+		Action = FRigVMSetPinIndexInCategoryAction(this, InPin, InIndexInCategory);
+		Action.SetTitle(TEXT("Set Pin Index in Category"));
+		GetActionStack()->BeginAction(Action);
+	}
+
+	InPin->IndexInCategory = InIndexInCategory;
+	Notify(ERigVMGraphNotifType::PinCategoryChanged, InPin);
+
+	if (!bSuspendNotifications)
+	{
+		const URigVMGraph* Graph = GetGraph();
+		check(Graph);
+		(void)Graph->MarkPackageDirty();
+	}
+
+	if (bSetupUndoRedo)
+	{
+		GetActionStack()->EndAction(Action);
+	}
+
+	return true;
+}
+
 bool URigVMController::SetPinCategories(const FName& InNodeName, const TArray<FString>& InCategories, bool bSetupUndoRedo)
 {
 	if (!IsValidGraph())
diff --git a/Engine/Plugins/Runtime/RigVM/Source/RigVMDeveloper/Private/RigVMModel/RigVMControllerActions.cpp b/Engine/Plugins/Runtime/RigVM/Source/RigVMDeveloper/Private/RigVMModel/RigVMControllerActions.cpp
--- a/Engine/Plugins/Runtime/RigVM/Source/RigVMDeveloper/Private/RigVMModel/RigVMControllerActions.cpp
+++ b/Engine/Plugins/Runtime/RigVM/Source/RigVMDeveloper/Private/RigVMModel/RigVMControllerActions.cpp
@@ -1171,6 +1171,82 @@ void FRigVMChangeNodePinCategoriesAction::UpdateAfterModification(const URigVMNo
 	NewCategories = InNode->GetPinCategories();
 }
 
+FRigVMSetPinCategoryExpansionAction::FRigVMSetPinCategoryExpansionAction()
+: FRigVMBaseAction(nullptr)
+, NodeName()
+, PinCategory()
+, bOldExpansionState(false)
+{
+}
+
+FRigVMSetPinCategoryExpansionAction::FRigVMSetPinCategoryExpansionAction(URigVMController* InController, const URigVMNode* InNode, const FString& InPinCategory)
+: FRigVMBaseAction(InController)
+, NodeName(InNode->GetName())
+, PinCategory(InPinCategory)
+, bOldExpansionState(InNode->IsPinCategoryExpanded(InPinCategory))
+{
+}
+
+bool FRigVMSetPinCategoryExpansionAction::Undo()
+{
+	if(!FRigVMBaseAction::Undo())
+	{
+		return false;
+	}
+	return GetController()->SetPinCategoryExpansion(*NodeName, PinCategory, bOldExpansionState, false);
+}
+
+bool FRigVMSetPinCategoryExpansionAction::Redo()
+{
+	if(!CanUndoRedo())
+	{
+		return false;
+	}
+	if(!GetController()->SetPinCategoryExpansion(*NodeName, PinCategory, !bOldExpansionState, false))
+	{
+		return false;
+	}
+	return FRigVMBaseAction::Redo();
+}
+
+FRigVMSetPinIndexInCategoryAction::FRigVMSetPinIndexInCategoryAction()
+: FRigVMBaseAction(nullptr)
+, PinPath()
+, OldIndexInCategory(INDEX_NONE)
+, NewIndexInCategory(INDEX_NONE)
+{
+}
+
+FRigVMSetPinIndexInCategoryAction::FRigVMSetPinIndexInCategoryAction(URigVMController* InController, URigVMPin* InPin, int32 InNewIndexInCategory)
+: FRigVMBaseAction(InController)
+, PinPath(InPin->GetPinPath())
+, OldIndexInCategory(InPin->GetIndexInCategory())
+, NewIndexInCategory(InNewIndexInCategory)
+{
+}
+
+bool FRigVMSetPinIndexInCategoryAction::Undo()
+{
+	if(!FRigVMBaseAction::Undo())
+	{
+		return false;
+	}
+	return GetController()->SetPinIndexInCategory(PinPath, OldIndexInCategory, false);
+}
+
+bool FRigVMSetPinIndexInCategoryAction::Redo()
+{
+	if(!CanUndoRedo())
+	{
+		return false;
+	}
+	if(!GetController()->SetPinIndexInCategory(PinPath, NewIndexInCategory, false))
+	{
+		return false;
+	}
+	return FRigVMBaseAction::Redo();
+}
+
 FRigVMSetPinWatchAction::FRigVMSetPinWatchAction()
 : FRigVMBaseAction(nullptr)
 , OldIsWatched(false)
diff --git a/Engine/Plugins/Runtime/RigVM/Source/RigVMDeveloper/Private/RigVMModel/RigVMNode.cpp b/Engine/Plugins/Runtime/RigVM/Source/RigVMDeveloper/Private/RigVMModel/RigVMNode.cpp
--- a/Engine/Plugins/Runtime/RigVM/Source/RigVMDeveloper/Private/RigVMModel/RigVMNode.cpp
+++ b/Engine/Plugins/Runtime/RigVM/Source/RigVMDeveloper/Private/RigVMModel/RigVMNode.cpp
@@ -119,6 +119,138 @@ TArray<FString> URigVMNode::GetPinCategories() const
 	return PinCategories;
 }
 
+TArray<FString> URigVMNode::GetSubPinCategories(const FString InCategory, bool bOnlyExisting, bool bRecursive) const
+{
+	if(InCategory.IsEmpty())
+	{
+		return {};
+	}
+	
+	const TArray<FString> ExistingCategories = GetPinCategories();
+	const FString Prefix = InCategory + TEXT("|");
+
+	const TArray<FString> IncompleteSubCategories = ExistingCategories.FilterByPredicate([Prefix](const FString& ExistingCategory)
+	{
+		return ExistingCategory.StartsWith(Prefix, ESearchCase::CaseSensitive);
+	});
+
+	TArray<FString> SubCategories;
+	for(const FString& SubCategory : IncompleteSubCategories)
+	{
+		TArray<FString> Parts;
+		verify(RigVMStringUtils::SplitNodePath(SubCategory, Parts));
+
+		TArray<FString> ParentsOfSubCategory;
+		while(!Parts.IsEmpty())
+		{
+			const FString ParentCategory = RigVMStringUtils::JoinNodePath(Parts);
+			if(!ParentCategory.StartsWith(Prefix))
+			{
+				break;
+			}
+			ParentsOfSubCategory.Add(ParentCategory);
+			Parts.Pop();
+		}
+
+		for(int32 Index = ParentsOfSubCategory.Num() - 1; Index >= 0; Index--)
+		{
+			SubCategories.AddUnique(ParentsOfSubCategory[Index]);
+		}
+	}
+
+	if(!bRecursive)
+	{
+		// remove any category that is not a direct child of the input category
+		SubCategories.RemoveAll([Prefix](const FString& InCategory) -> bool
+		{
+			return InCategory.Mid(Prefix.Len()).Contains(TEXT("|"));
+		});
+	}
+
+	if(bOnlyExisting)
+	{
+		SubCategories.RemoveAll([&ExistingCategories](const FString& InCategory) -> bool
+		{
+			return !ExistingCategories.Contains(InCategory);
+		});
+	}
+
+	return SubCategories;
+}
+
+FString URigVMNode::GetPinCategoryName(const FString InCategory) const
+{
+	if(InCategory.IsEmpty())
+	{
+		return FString();
+	}
+	FString ParentCategory, CategoryName;
+	if(RigVMStringUtils::SplitNodePathAtEnd(InCategory, ParentCategory, CategoryName))
+	{
+		return CategoryName;
+	}
+	return FString();
+}
+
+FString URigVMNode::GetParentPinCategory(const FString InCategory, bool bOnlyExisting) const
+{
+	if(InCategory.IsEmpty())
+	{
+		return FString();
+	}
+	FString ParentCategory, CategoryName;
+	if(RigVMStringUtils::SplitNodePathAtEnd(InCategory, ParentCategory, CategoryName))
+	{
+		return ParentCategory;
+	}
+	return FString();
+}
+
+TArray<FString> URigVMNode::GetParentPinCategories(const FString InCategory, bool bOnlyExisting, bool bIncludeSelf) const
+{
+	if(InCategory.IsEmpty())
+	{
+		return {};
+	}
+	
+	const TArray<FString> ExistingCategories = GetPinCategories();
+
+	TArray<FString> Parts;
+	verify(RigVMStringUtils::SplitNodePath(InCategory, Parts));
+
+	TArray<FString> ParentCategories;
+	while(!Parts.IsEmpty())
+	{
+		ParentCategories.Add(RigVMStringUtils::JoinNodePath(Parts));
+		Parts.Pop();
+	}
+
+	if(!bIncludeSelf)
+	{
+		ParentCategories.Remove(InCategory);
+	}
+
+	if(bOnlyExisting)
+	{
+		ParentCategories.RemoveAll([&ExistingCategories](const FString& InCategory) -> bool
+		{
+			return !ExistingCategories.Contains(InCategory);
+		});
+	}
+
+	return ParentCategories;
+}
+
+int32 URigVMNode::GetPinCategoryDepth(const FString& InCategory)
+{
+	TArray<FString> Parts;
+	if(RigVMStringUtils::SplitNodePath(InCategory, Parts))
+	{
+		return Parts.Num() - 1;
+	}
+	return 0;
+}
+
 TArray<URigVMPin*> URigVMNode::GetPinsForCategory(FString InCategory) const
 {
 	InCategory.TrimStartAndEndInline();
@@ -140,6 +272,15 @@ TArray<URigVMPin*> URigVMNode::GetPinsForCategory(FString InCategory) const
 	return PinsInCategory;
 }
 
+bool URigVMNode::IsPinCategoryExpanded(FString InCategory) const
+{
+	if(const bool* ExpansionState = PinCategoryExpansion.Find(InCategory))
+	{
+		return *ExpansionState;
+	}
+	return false;
+}
+
 FString URigVMNode::GetOriginalPinDefaultValue(const URigVMPin* InPin) const
 {
 	const FString CompleteSegmentPath = InPin->GetSegmentPath(true);
@@ -684,6 +825,11 @@ FString URigVMNode::GetCategoryForPin(const FString& InPinPath) const
 	return FString();
 }
 
+int32 URigVMNode::GetIndexInCategoryForPin(const FString& InPinPath) const
+{
+	return INDEX_NONE;
+}
+
 URigVMLibraryNode* URigVMNode::FindFunctionForNode() const  
 {
 	const UObject* Subject = this;
diff --git a/Engine/Plugins/Runtime/RigVM/Source/RigVMDeveloper/Private/RigVMModel/RigVMPin.cpp b/Engine/Plugins/Runtime/RigVM/Source/RigVMDeveloper/Private/RigVMModel/RigVMPin.cpp
--- a/Engine/Plugins/Runtime/RigVM/Source/RigVMDeveloper/Private/RigVMModel/RigVMPin.cpp
+++ b/Engine/Plugins/Runtime/RigVM/Source/RigVMDeveloper/Private/RigVMModel/RigVMPin.cpp
@@ -115,6 +115,8 @@ URigVMPin::URigVMPin()
 	, CPPTypeObjectPath(NAME_None)
 	, DefaultValue(FString())
 	, DefaultValueType(ERigVMPinDefaultValueType::AutoDetect)
+	, CustomWidgetName(NAME_None)
+	, IndexInCategory(INDEX_NONE)
 	, BoundVariablePath_DEPRECATED()
 	, LastKnownTypeIndex(INDEX_NONE)
 {
@@ -240,6 +242,22 @@ FString URigVMPin::GetCategory() const
 	return UserDefinedCategory;
 }
 
+int32 URigVMPin::GetIndexInCategory() const
+{
+	if(IndexInCategory == INDEX_NONE)
+	{
+		if(const URigVMNode* Node = GetNode())
+		{
+			int32 IndexFromNode = Node->GetIndexInCategoryForPin(this->GetSegmentPath(true));
+			if(IndexFromNode != INDEX_NONE)
+			{
+				return IndexFromNode;
+			}
+		}
+	}
+	return IndexInCategory;
+}
+
 FString URigVMPin::GetSegmentPath(bool bIncludeRootPin) const
 {
 	URigVMPin* ParentPin = GetParentPin();
diff --git a/Engine/Plugins/Runtime/RigVM/Source/RigVMDeveloper/Public/EdGraph/RigVMEdGraphNode.h b/Engine/Plugins/Runtime/RigVM/Source/RigVMDeveloper/Public/EdGraph/RigVMEdGraphNode.h
--- a/Engine/Plugins/Runtime/RigVM/Source/RigVMDeveloper/Public/EdGraph/RigVMEdGraphNode.h
+++ b/Engine/Plugins/Runtime/RigVM/Source/RigVMDeveloper/Public/EdGraph/RigVMEdGraphNode.h
@@ -191,6 +191,7 @@ public:
 
 	URigVMPin* FindModelPinFromGraphPin(const UEdGraphPin* InGraphPin) const;
 	UEdGraphPin* FindGraphPinFromModelPin(const URigVMPin* InModelPin, bool bAsInput) const;
+	UEdGraphPin* FindGraphPinFromCategory(const FString& InCategory, bool bAsInput) const;
 
 	/// Synchronize the stored name/value/type on the graph pin with the value stored on the node. 
 	/// If the pin has sub-pins, the value update is done recursively.
@@ -229,10 +230,13 @@ protected:
 
 	/** Helper function for AllocateDefaultPins */
 	void UpdatePinLists();
+	bool CreateGraphPinFromCategory(const FString& InCategory, EEdGraphPinDirection InDirection);
 	bool CreateGraphPinFromModelPin(const URigVMPin* InModelPin, EEdGraphPinDirection InDirection,  UEdGraphPin* InParentPin = nullptr);
 	void RemoveGraphSubPins(UEdGraphPin *InParentPin, const TArray<UEdGraphPin*>& InPinsToKeep = TArray<UEdGraphPin*>());
 	bool ModelPinAdded_Internal(const URigVMPin* InModelPin);
 	bool ModelPinRemoved_Internal(const URigVMPin* InModelPin);
+	bool CategoryPinAdded_Internal(const FString& InCategory, EEdGraphPinDirection InDirection);
+	bool CategoryPinRemoved_Internal(const FString& InCategory);
 
 	/** Copies default values from underlying properties into pin defaults, for editing */
 	void SetupPinDefaultsFromModel(UEdGraphPin* Pin, const URigVMPin* InModelPin = nullptr);
@@ -259,6 +263,7 @@ protected:
 	UClass* GetRigVMSkeletonGeneratedClass() const;
 
 	static FEdGraphPinType GetPinTypeForModelPin(const URigVMPin* InModelPin);
+	static FEdGraphPinType GetPinTypeForCategoryPin();
 
 	virtual void ConfigurePin(UEdGraphPin* EdGraphPin, const URigVMPin* ModelPin) const;
 private:
@@ -287,6 +292,7 @@ private:
 	TArray<UEdGraphPin*> LastEdGraphPins;
 	
 	TMap<URigVMPin*, FPinPair> CachedPins;
+	TMap<FString, FPinPair> CachedCategoryPins;
 
 	FNodeTitleDirtied NodeTitleDirtied;
 	FNodePinsChanged NodePinsChanged;
diff --git a/Engine/Plugins/Runtime/RigVM/Source/RigVMDeveloper/Public/RigVMModel/Nodes/RigVMFunctionInterfaceNode.h b/Engine/Plugins/Runtime/RigVM/Source/RigVMDeveloper/Public/RigVMModel/Nodes/RigVMFunctionInterfaceNode.h
--- a/Engine/Plugins/Runtime/RigVM/Source/RigVMDeveloper/Public/RigVMModel/Nodes/RigVMFunctionInterfaceNode.h
+++ b/Engine/Plugins/Runtime/RigVM/Source/RigVMDeveloper/Public/RigVMModel/Nodes/RigVMFunctionInterfaceNode.h
@@ -28,6 +28,7 @@ public:
 	virtual FText GetToolTipTextForPin(const URigVMPin* InPin) const override;	
 	virtual FName GetDisplayNameForPin(const FString& InPinPath) const override;
 	virtual FString GetCategoryForPin(const FString& InPinPath) const override;
+	virtual int32 GetIndexInCategoryForPin(const FString& InPinPath) const override;
 	virtual TArray<FString> GetPinCategories() const override;
 
 protected:
diff --git a/Engine/Plugins/Runtime/RigVM/Source/RigVMDeveloper/Public/RigVMModel/Nodes/RigVMFunctionReferenceNode.h b/Engine/Plugins/Runtime/RigVM/Source/RigVMDeveloper/Public/RigVMModel/Nodes/RigVMFunctionReferenceNode.h
--- a/Engine/Plugins/Runtime/RigVM/Source/RigVMDeveloper/Public/RigVMModel/Nodes/RigVMFunctionReferenceNode.h
+++ b/Engine/Plugins/Runtime/RigVM/Source/RigVMDeveloper/Public/RigVMModel/Nodes/RigVMFunctionReferenceNode.h
@@ -26,6 +26,7 @@ public:
 	virtual FText GetToolTipText() const override;
 	virtual FName GetDisplayNameForPin(const FString& InPinPath) const override;
 	virtual FString GetCategoryForPin(const FString& InPinPath) const override;
+	virtual int32 GetIndexInCategoryForPin(const FString& InPinPath) const override;
 	virtual FText GetToolTipTextForPin(const URigVMPin* InPin) const override;
 	virtual TArray<FString> GetPinCategories() const override;
 	// end URigVMNode interface
diff --git a/Engine/Plugins/Runtime/RigVM/Source/RigVMDeveloper/Public/RigVMModel/RigVMController.h b/Engine/Plugins/Runtime/RigVM/Source/RigVMDeveloper/Public/RigVMModel/RigVMController.h
--- a/Engine/Plugins/Runtime/RigVM/Source/RigVMDeveloper/Public/RigVMModel/RigVMController.h
+++ b/Engine/Plugins/Runtime/RigVM/Source/RigVMDeveloper/Public/RigVMModel/RigVMController.h
@@ -825,6 +825,14 @@ public:
 	// to order pins in the user interface of the node as well as on the details panel.
 	bool SetPinCategoryIndex(const FName& InNodeName, const FString& InPinCategory, int32 InNewIndex, bool bSetupUndoRedo = true, bool bPrintPythonCommand = false);
 
+	// Changes a pin category's expansion state. The category is UI relevant only and used
+	// to order pins in the user interface of the node as well as on the details panel.
+	bool SetPinCategoryExpansion(const FName& InNodeName, const FString& InPinCategory, bool bIsExpanded, bool bSetupUndoRedo = true, bool bPrintPythonCommand = false);
+
+	// Changes a pin category's expansion state. The category is UI relevant only and used
+	// to order pins in the user interface of the node as well as on the details panel.
+	bool SetPinIndexInCategory(const FString& InPinPath, int32 InIndexInCategory, bool bSetupUndoRedo = true, bool bPrintPythonCommand = false);
+
 	// Returns the default value of a pin given its pinpath.
 	UFUNCTION(BlueprintCallable, Category = RigVMController)
 	FString GetPinDefaultValue(const FString& InPinPath);
@@ -1257,6 +1265,8 @@ private:
 	bool RemovePinCategory(const URigVMNode* InNode, const FString& InPinCategory, bool bSetupUndoRedo);
 	bool RenamePinCategory(const URigVMNode* InNode, const FString& InOldPinCategory, const FString& InNewPinCategory, bool bSetupUndoRedo);
 	bool SetPinCategoryIndex(const URigVMNode* InNode, const FString& InPinCategory, int32 InNewIndex, bool bSetupUndoRedo);
+	bool SetPinCategoryExpansion(const URigVMNode* InNode, const FString& InPinCategory, bool bIsExpanded, bool bSetupUndoRedo);
+	bool SetPinIndexInCategory(URigVMPin* InPin, int32 InIndexInCategory, bool bSetupUndoRedo);
 	bool SetPinCategories(const FName& InNodeName, const TArray<FString>& InCategories, bool bSetupUndoRedo);
 	bool SetPinCategories(const URigVMNode* InNode, const TArray<FString>& InCategories, bool bSetupUndoRedo);
 	bool SetVariableName(URigVMVariableNode* InVariableNode, const FName& InVariableName, bool bSetupUndoRedo);
diff --git a/Engine/Plugins/Runtime/RigVM/Source/RigVMDeveloper/Public/RigVMModel/RigVMControllerActions.h b/Engine/Plugins/Runtime/RigVM/Source/RigVMDeveloper/Public/RigVMModel/RigVMControllerActions.h
--- a/Engine/Plugins/Runtime/RigVM/Source/RigVMDeveloper/Public/RigVMModel/RigVMControllerActions.h
+++ b/Engine/Plugins/Runtime/RigVM/Source/RigVMDeveloper/Public/RigVMModel/RigVMControllerActions.h
@@ -919,7 +919,7 @@ public:
 	FRigVMChangeNodePinCategoriesAction();
 	FRigVMChangeNodePinCategoriesAction(URigVMController* InController, const URigVMNode* InNode);
 	virtual ~FRigVMChangeNodePinCategoriesAction() {};
-	virtual UScriptStruct* GetScriptStruct() const override { return FRigVMSetPinCategoryAction::StaticStruct(); }
+	virtual UScriptStruct* GetScriptStruct() const override { return FRigVMChangeNodePinCategoriesAction::StaticStruct(); }
 	virtual bool Undo() override;
 	virtual bool Redo() override;
 	
@@ -935,6 +935,60 @@ public:
 	TArray<FString> NewCategories;
 };
 
+/**
+ * An action to change the expansion on a pin's category in the graph.
+ */
+USTRUCT()
+struct FRigVMSetPinCategoryExpansionAction : public FRigVMBaseAction
+{
+	GENERATED_BODY()
+
+public:
+
+	FRigVMSetPinCategoryExpansionAction();
+	FRigVMSetPinCategoryExpansionAction(URigVMController* InController, const URigVMNode* InNode, const FString& InPinCategory);
+	virtual ~FRigVMSetPinCategoryExpansionAction() {};
+	virtual UScriptStruct* GetScriptStruct() const override { return FRigVMSetPinCategoryAction::StaticStruct(); }
+	virtual bool Undo() override;
+	virtual bool Redo() override;
+
+	UPROPERTY()
+	FString NodeName;
+
+	UPROPERTY()
+	FString PinCategory;
+
+	UPROPERTY()
+	bool bOldExpansionState;
+};
+
+/**
+ * An action setting a pin's index within a category.
+ */
+USTRUCT()
+struct FRigVMSetPinIndexInCategoryAction : public FRigVMBaseAction
+{
+	GENERATED_BODY()
+
+public:
+
+	FRigVMSetPinIndexInCategoryAction();
+	FRigVMSetPinIndexInCategoryAction(URigVMController* InController, URigVMPin* InPin, int32 InNewIndexInCategory);
+	virtual ~FRigVMSetPinIndexInCategoryAction() {};
+	virtual UScriptStruct* GetScriptStruct() const override { return FRigVMSetPinIndexInCategoryAction::StaticStruct(); }
+	virtual bool Undo() override;
+	virtual bool Redo() override;
+
+	UPROPERTY()
+	FString PinPath;
+
+	UPROPERTY()
+	int32 OldIndexInCategory;
+
+	UPROPERTY()
+	int32 NewIndexInCategory;
+};
+
 /**
  * An action setting a pin's default value in the graph.
  */
diff --git a/Engine/Plugins/Runtime/RigVM/Source/RigVMDeveloper/Public/RigVMModel/RigVMNode.h b/Engine/Plugins/Runtime/RigVM/Source/RigVMDeveloper/Public/RigVMModel/RigVMNode.h
--- a/Engine/Plugins/Runtime/RigVM/Source/RigVMDeveloper/Public/RigVMModel/RigVMNode.h
+++ b/Engine/Plugins/Runtime/RigVM/Source/RigVMDeveloper/Public/RigVMModel/RigVMNode.h
@@ -72,10 +72,35 @@ public:
 	UFUNCTION(BlueprintCallable, Category = RigVMNode)
 	virtual TArray<FString> GetPinCategories() const;
 
+	// Returns the name of pin category
+	UFUNCTION(BlueprintCallable, Category = RigVMNode)
+	FString GetPinCategoryName(const FString InCategory) const;
+
+	// Returns all sub user defined categories of a given parent category
+	UFUNCTION(BlueprintCallable, Category = RigVMNode)
+	TArray<FString> GetSubPinCategories(const FString InCategory, bool bOnlyExisting = false, bool bRecursive = false) const;
+
+	// Returns the parent pin category of the given category (or an empty string in case there's no parent)
+	UFUNCTION(BlueprintCallable, Category = RigVMNode)
+	FString GetParentPinCategory(const FString InCategory, bool bOnlyExisting = false) const;
+
+	// Returns all parent categories of a given 
+	UFUNCTION(BlueprintCallable, Category = RigVMNode)
+	TArray<FString> GetParentPinCategories(const FString InCategory, bool bOnlyExisting = false, bool bIncludeSelf = false) const;
+
+	// Returns the depth of the category (starting with 0 for "Foo" and 2 for "Foo|Bar|Waldo")
+	static int32 GetPinCategoryDepth(const FString& InCategory);
+
 	// Returns all pins for a given category
 	UFUNCTION(BlueprintCallable, Category = RigVMNode)
 	TArray<URigVMPin*> GetPinsForCategory(FString InCategory) const;
 
+	// Returns all pins for a given category
+	UFUNCTION(BlueprintCallable, Category = RigVMNode)
+	bool IsPinCategoryExpanded(FString InCategory) const;
+
+	const FString& GetLastAffectedPinCategory() const { return LastAffectedPinCategory; }
+
 	// Returns the default value for a given pin
 	FString GetOriginalPinDefaultValue(const URigVMPin* InPin) const;
 
@@ -350,6 +375,9 @@ public:
 	// returns the category for a pin
 	virtual FString GetCategoryForPin(const FString& InPinPath) const;
 
+	// returns the index of a pin within a category
+	virtual int32 GetIndexInCategoryForPin(const FString& InPinPath) const;
+
 private:
 
 	static const inline TCHAR* NodeColorName = TEXT("NodeColor");
@@ -401,9 +429,15 @@ private:
 	TArray<TObjectPtr<URigVMPin>> OrphanedPins;
 
 protected:
+	
 	UPROPERTY()
 	TArray<FString> PinCategories;
 
+	UPROPERTY()
+	TMap<FString, bool> PinCategoryExpansion;
+
+	FString LastAffectedPinCategory;
+	
 private:
 
 #if WITH_EDITOR
diff --git a/Engine/Plugins/Runtime/RigVM/Source/RigVMDeveloper/Public/RigVMModel/RigVMNotifications.h b/Engine/Plugins/Runtime/RigVM/Source/RigVMDeveloper/Public/RigVMModel/RigVMNotifications.h
--- a/Engine/Plugins/Runtime/RigVM/Source/RigVMDeveloper/Public/RigVMModel/RigVMNotifications.h
+++ b/Engine/Plugins/Runtime/RigVM/Source/RigVMDeveloper/Public/RigVMModel/RigVMNotifications.h
@@ -57,6 +57,7 @@ enum class ERigVMGraphNotifType : uint8
 	PinDisplayNameChanged, // The display name of a pin has changed - requiring a rebuild of the node user interface (Subject == URigVMPin) 
 	PinCategoryChanged, // The category of a pin has changed - requiring a rebuild of the node user interface (Subject == URigVMPin)
 	PinCategoriesChanged, // The category list of a node has changed - requiring a rebuild of the node user interface (Subject == URigVMNode)
+	PinCategoryExpansionChanged, // The category of a pin expanded / collapsed (Subject == URigVMNode)
 	Invalid // The max for this enum (used for guarding)
 };
 
diff --git a/Engine/Plugins/Runtime/RigVM/Source/RigVMDeveloper/Public/RigVMModel/RigVMPin.h b/Engine/Plugins/Runtime/RigVM/Source/RigVMDeveloper/Public/RigVMModel/RigVMPin.h
--- a/Engine/Plugins/Runtime/RigVM/Source/RigVMDeveloper/Public/RigVMModel/RigVMPin.h
+++ b/Engine/Plugins/Runtime/RigVM/Source/RigVMDeveloper/Public/RigVMModel/RigVMPin.h
@@ -173,6 +173,11 @@ public:
 	UFUNCTION(BlueprintCallable, Category = RigVMPin)
 	FString GetCategory() const;
 
+	// Returns index within a category on a pin. The category is UI relevant only and used
+	// to order pins in the user interface of the node as well as on the details panel.
+	UFUNCTION(BlueprintCallable, Category = RigVMPin)
+	int32 GetIndexInCategory() const;
+
 	// Returns a . separated path containing all names of the pin within its main
 	// memory owner / storage. This is typically used to create an offset pointer
 	// within memory (FRigVMRegisterOffset).
@@ -589,6 +594,9 @@ private:
 	UPROPERTY()
 	FString UserDefinedCategory;
 
+	UPROPERTY()
+	int32 IndexInCategory;
+
 #if WITH_EDITORONLY_DATA
 	UPROPERTY()
 	FString BoundVariablePath_DEPRECATED;
diff --git a/Engine/Plugins/Runtime/RigVM/Source/RigVMEditor/Private/EdGraph/RigVMEdGraphPanelPinFactory.cpp b/Engine/Plugins/Runtime/RigVM/Source/RigVMEditor/Private/EdGraph/RigVMEdGraphPanelPinFactory.cpp
--- a/Engine/Plugins/Runtime/RigVM/Source/RigVMEditor/Private/EdGraph/RigVMEdGraphPanelPinFactory.cpp
+++ b/Engine/Plugins/Runtime/RigVM/Source/RigVMEditor/Private/EdGraph/RigVMEdGraphPanelPinFactory.cpp
@@ -11,6 +11,7 @@
 #include "Widgets/SRigVMGraphPinUserDataNameSpace.h"
 #include "Widgets/SRigVMGraphPinUserDataPath.h"
 #include "Widgets/SRigVMGraphPinQuat.h"
+#include "Widgets/SRigVMGraphPinCategory.h"
 #include "KismetPins/SGraphPinExec.h"
 #include "SGraphPinComboBox.h"
 #include "RigVMHost.h"
@@ -124,6 +125,13 @@ TSharedPtr<SGraphPin> FRigVMEdGraphPanelPinFactory::CreatePin_Internal(UEdGraphP
 						.ModelPin(ModelPin);
 				}
 			}
+			else if(const URigVMNode* ModelNode = RigNode->GetModelNode())
+			{
+				if(ModelNode->GetPinCategories().Contains(InPin->GetName()))
+				{
+					return SNew(SRigVMGraphPinCategory, InPin);
+				}
+			}
 
 			if (InPin->PinType.PinCategory == UEdGraphSchema_K2::PC_Struct)
 			{
diff --git a/Engine/Plugins/Runtime/RigVM/Source/RigVMEditor/Private/Editor/RigVMMinimalEnvironment.cpp b/Engine/Plugins/Runtime/RigVM/Source/RigVMEditor/Private/Editor/RigVMMinimalEnvironment.cpp
--- a/Engine/Plugins/Runtime/RigVM/Source/RigVMEditor/Private/Editor/RigVMMinimalEnvironment.cpp
+++ b/Engine/Plugins/Runtime/RigVM/Source/RigVMEditor/Private/Editor/RigVMMinimalEnvironment.cpp
@@ -163,6 +163,9 @@ void FRigVMMinimalEnvironment::HandleModified(ERigVMGraphNotifType InNotificatio
 		case ERigVMGraphNotifType::InteractionBracketOpened:
 		case ERigVMGraphNotifType::InteractionBracketClosed:
 		case ERigVMGraphNotifType::InteractionBracketCanceled:
+		case ERigVMGraphNotifType::PinCategoryChanged:
+		case ERigVMGraphNotifType::PinCategoriesChanged:
+		case ERigVMGraphNotifType::PinCategoryExpansionChanged:
 		{
 			break;
 		}
diff --git a/Engine/Plugins/Runtime/RigVM/Source/RigVMEditor/Private/Widgets/SRigVMGraphNode.cpp b/Engine/Plugins/Runtime/RigVM/Source/RigVMEditor/Private/Widgets/SRigVMGraphNode.cpp
--- a/Engine/Plugins/Runtime/RigVM/Source/RigVMEditor/Private/Widgets/SRigVMGraphNode.cpp
+++ b/Engine/Plugins/Runtime/RigVM/Source/RigVMEditor/Private/Widgets/SRigVMGraphNode.cpp
@@ -39,6 +39,7 @@
 #include "RigVMFunctions/RigVMDispatch_Select.h"
 #include "Algo/Copy.h"
 #include "Brushes/SlateColorBrush.h"
+#include "Widgets/SRigVMGraphPinCategory.h"
 #include "Widgets/SRigVMVariantTagWidget.h"
 
 #if WITH_EDITOR
@@ -666,6 +667,10 @@ FText SRigVMGraphNode::GetPinLabel(TWeakPtr<SGraphPin> GraphPin) const
 
 TOptional<FSlateColor> SRigVMGraphNode::GetHighlightColor(const SGraphPin* InGraphPin) const
 {
+	if(IsCategoryPin(InGraphPin))
+	{
+		return FSlateColor(EStyleColor::Hover);
+	}
 	if(HasUserProvidedValue(InGraphPin))
 	{
 		return FSlateColor(FLinearColor::Red);
@@ -675,7 +680,7 @@ TOptional<FSlateColor> SRigVMGraphNode::GetHighlightColor(const SGraphPin* InGra
 
 TOptional<const FSlateBrush*> SRigVMGraphNode::GetPinBorder(const SGraphPin* InGraphPin) const
 {
-	if(HasUserProvidedValue(InGraphPin))
+	if(HasUserProvidedValue(InGraphPin) || IsCategoryPin(InGraphPin))
 	{
 		static const FName NAME_Pin_BackgroundHovered("Graph.Pin.BackgroundHovered");
 		static const FSlateBrush* CachedImg_Pin_BackgroundHovered = FAppStyle::GetBrush( NAME_Pin_BackgroundHovered );
@@ -1148,7 +1153,7 @@ void SRigVMGraphNode::HandleNodePinsChanged()
 		
 		const FPinInfo* PinInfoPtr = PinInfos.FindByPredicate([PinPath](const FPinInfo& PinInfo) -> bool
 		{
-			return PinInfo.ModelPinPath == PinPath;
+			return PinInfo.Identifier == PinPath;
 		});
 		
 		if (PinInfoPtr)
@@ -1173,7 +1178,7 @@ void SRigVMGraphNode::HandleNodePinsChanged()
 		
 		const FPinInfo* PinInfoPtr = PinInfos.FindByPredicate([PinPath](const FPinInfo& PinInfo) -> bool
 		{
-			return PinInfo.ModelPinPath == PinPath;
+			return PinInfo.Identifier == PinPath;
 		});
 		
 		if (PinInfoPtr)
@@ -1449,36 +1454,36 @@ FReply SRigVMGraphNode::OnExpanderArrowClicked(int32 InPinInfoIndex)
 			}
 			
 			const FPinInfo& PinInfo = PinInfos[InPinInfoIndex];
-			TArray<FString> PinPathsToModify;
-			PinPathsToModify.Add(PinInfo.ModelPinPath);
+			TArray<int32> PinInfoIndicesToModify;
+			PinInfoIndicesToModify.Add(PinInfo.Index);
 
 			// with shift clicked we expand recursively
 			if (FSlateApplication::Get().GetModifierKeys().IsShiftDown())
 			{
-				if(URigVMGraph* ModelGraph = EdGraphNode->GetModel())
+				for(const FPinInfo& OtherPinInfo : PinInfos)
 				{
-					if(URigVMPin* ModelPin = ModelGraph->FindPin(PinInfo.ModelPinPath))
+					if(PinInfoIndicesToModify.Contains(OtherPinInfo.ParentIndex))
 					{
-						TArray<URigVMPin*> SubPins = ModelPin->GetNode()->GetAllPinsRecursively();
-						for(URigVMPin* SubPin : SubPins)
-						{
-							if(SubPin->IsInOuter(ModelPin))
-							{
-								PinPathsToModify.Add(SubPin->GetPinPath());
-							}
-						}
-
-						Algo::Reverse(PinPathsToModify);
+						PinInfoIndicesToModify.Emplace(OtherPinInfo.Index);
+						PinInfoIndicesToModify.Add(OtherPinInfo.Index);
 					}
 				}
+				Algo::Reverse(PinInfoIndicesToModify);
 			}
 
 			TGuardValue<bool> GuardDirtyBlueprintStatus(Blueprint->bSkipDirtyBlueprintStatus, true);
 
 			Controller->OpenUndoBracket(PinInfo.bExpanded ? TEXT("Collapsing Pin") : TEXT("Expanding Pin"));
-			for(const FString& PinPathToModify : PinPathsToModify)
+			for(const int32& PinInfoIndexToModify : PinInfoIndicesToModify)
 			{
-				Controller->SetPinExpansion(PinPathToModify, !PinInfo.bExpanded, true, true);
+				if(PinInfos[PinInfoIndexToModify].bIsCategoryPin)
+				{
+					Controller->SetPinCategoryExpansion(EdGraphNode->GetFName(), PinInfos[PinInfoIndexToModify].Identifier, !PinInfo.bExpanded, true, true);
+				}
+				else
+				{
+					Controller->SetPinExpansion(PinInfos[PinInfoIndexToModify].Identifier, !PinInfo.bExpanded, true, true);
+				}
 			}
 			Controller->CloseUndoBracket();
 			return FReply::Handled();
@@ -1506,7 +1511,7 @@ void SRigVMGraphNode::HandleModifiedEvent(ERigVMGraphNotifType InNotifType, URig
 					const FString PinPath = Pin->GetPinPath();
 					for(FPinInfo& PinInfo : PinInfos)
 					{
-						if(PinInfo.ModelPinPath == PinPath)
+						if(PinInfo.Identifier == PinPath)
 						{
 							PinInfo.bExpanded = Pin->IsExpanded();
 							break;
@@ -1516,6 +1521,30 @@ void SRigVMGraphNode::HandleModifiedEvent(ERigVMGraphNotifType InNotifType, URig
 			}
 			break;
 		}
+		case ERigVMGraphNotifType::PinCategoryExpansionChanged:
+		{
+			if(!ModelNode.IsValid())
+			{
+				return;
+			}
+			
+			if(const URigVMNode* ModelNodeFromNotification = Cast<URigVMNode>(InSubject))
+			{
+				if(ModelNodeFromNotification == ModelNode.Get())
+				{
+					const FString PinCategory = ModelNodeFromNotification->GetLastAffectedPinCategory();
+					for(FPinInfo& PinInfo : PinInfos)
+					{
+						if(PinInfo.Identifier == PinCategory)
+						{
+							PinInfo.bExpanded = ModelNodeFromNotification->IsPinCategoryExpanded(PinCategory);
+							break;
+						}
+					}
+				}
+			}
+			break;
+		}
 		case ERigVMGraphNotifType::LibraryTemplateChanged:
 		{
 			if(!ModelNode.IsValid())
@@ -1575,7 +1604,8 @@ void SRigVMGraphNode::UpdatePinTreeView()
 	TArray<URigVMPin*> RootModelPins = ModelNode->GetPins();
 	// orphaned pins are appended to the end of pin list on each side of the node
 	RootModelPins.Append(ModelNode->GetOrphanedPins());
-	TArray<URigVMPin*> ModelPins;
+
+	TArray<TTuple<URigVMPin*, FString>> ModelPinsOrCategories;
 
 	const bool bSupportSubPins = !RigGraphNode->DrawAsCompactNode();
 
@@ -1583,9 +1613,41 @@ void SRigVMGraphNode::UpdatePinTreeView()
 	// a) execute IOs, b) IO pins, c) input / visible pins, d) output pins
 	struct Local
 	{
-		static void VisitPinRecursively(URigVMPin* InPin, TArray<URigVMPin*>& OutPins, bool bSupportSubPins)
+		static void VisitPinRecursively(URigVMPin* InPin, TArray<TTuple<URigVMPin*, FString>>& OutPinsOrCategories, bool bSupportSubPins, const TArray<ERigVMPinDirection>& DirectionsWithCategories)
 		{
-			OutPins.Add(InPin);
+			const TTuple<URigVMPin*, FString> PinKey(InPin, FString());
+			if(OutPinsOrCategories.Contains(PinKey))
+			{
+				return;
+			}
+			
+			bool bIncludePin = true;
+			const FString PinCategory = InPin->GetCategory();
+			if(!PinCategory.IsEmpty())
+			{
+				// add all categories from root to leaf
+				TArray<FString> ParentCategories = InPin->GetNode()->GetParentPinCategories(PinCategory, false, true);
+				Algo::Reverse(ParentCategories);
+				for(const FString& ParentCategory : ParentCategories)
+				{
+					OutPinsOrCategories.AddUnique({nullptr, ParentCategory});
+				}
+			}
+			// if the pin has no category - but there are pins with the same direction in a category - we need to skip it
+			else if(DirectionsWithCategories.Contains(InPin->GetDirection()))
+			{
+				// only add the pin if our parent is here - either the parent pin or the category
+				const TTuple<URigVMPin*, FString> ParentPinPair = {InPin->GetParentPin(), FString()}; 
+				const TTuple<URigVMPin*, FString> CategoryPair = {nullptr, PinCategory}; 
+				bIncludePin =
+					OutPinsOrCategories.Contains(ParentPinPair) ||
+					OutPinsOrCategories.Contains(CategoryPair);
+			}
+			
+			if(bIncludePin)
+			{
+				OutPinsOrCategories.Add(PinKey);
+			}
 			if(!bSupportSubPins)
 			{
 				return;
@@ -1596,25 +1658,55 @@ void SRigVMGraphNode::UpdatePinTreeView()
 				return;
 			}
 
+			// visit all pins within this category
+			if(!PinCategory.IsEmpty())
+			{
+				const TArray<URigVMPin*> PinsInCategory = InPin->GetNode()->GetPinsForCategory(PinCategory);
+				for (URigVMPin* PinInCategory : PinsInCategory)
+				{
+					VisitPinRecursively(PinInCategory, OutPinsOrCategories, bSupportSubPins, DirectionsWithCategories);
+				}
+			}
+
 			if (InPin->GetCPPType() == TEXT("FRotator"))
 			{
 				const TArray<URigVMPin*>& SubPins = InPin->GetSubPins();
 				if (SubPins.Num() == 3)
 				{
-					OutPins.Add(SubPins[2]);	
-					OutPins.Add(SubPins[0]);	
-					OutPins.Add(SubPins[1]);
+					VisitPinRecursively(SubPins[2], OutPinsOrCategories, bSupportSubPins, DirectionsWithCategories);
+					VisitPinRecursively(SubPins[0], OutPinsOrCategories, bSupportSubPins, DirectionsWithCategories);
+					VisitPinRecursively(SubPins[1], OutPinsOrCategories, bSupportSubPins, DirectionsWithCategories);
 				}	
 			}
 			else
 			{				
 				for (URigVMPin* SubPin : InPin->GetSubPins())
 				{
-					VisitPinRecursively(SubPin, OutPins, bSupportSubPins);
+					VisitPinRecursively(SubPin, OutPinsOrCategories, bSupportSubPins, DirectionsWithCategories);
 				}
 			}
 		}
 	};
+
+	// determine which directions use a category
+	TArray<ERigVMPinDirection> DirectionsWithCategories;
+	const TArray<URigVMPin*> AllPins = ModelNode->GetAllPinsRecursively();
+	for(const URigVMPin* Pin : AllPins)
+	{
+		if(!Pin->GetCategory().IsEmpty())
+		{
+			const ERigVMPinDirection Direction = Pin->GetDirection();
+			DirectionsWithCategories.AddUnique(Direction);
+			if(Direction == ERigVMPinDirection::Input)
+			{
+				DirectionsWithCategories.AddUnique(ERigVMPinDirection::Visible);
+			}
+			else if(Direction == ERigVMPinDirection::Visible)
+			{
+				DirectionsWithCategories.AddUnique(ERigVMPinDirection::Input);
+			}
+		}
+	}
 	
 	for(int32 SortPhase = 0; SortPhase < 4; SortPhase++)
 	{
@@ -1626,7 +1718,7 @@ void SRigVMGraphNode::UpdatePinTreeView()
 				{
 					if(RootPin->IsExecuteContext() && RootPin->GetDirection() == ERigVMPinDirection::IO)
 					{
-						Local::VisitPinRecursively(RootPin, ModelPins, bSupportSubPins);
+						Local::VisitPinRecursively(RootPin, ModelPinsOrCategories, bSupportSubPins, DirectionsWithCategories);
 					}
 					break;
 				}
@@ -1634,7 +1726,7 @@ void SRigVMGraphNode::UpdatePinTreeView()
 				{
 					if(RootPin->GetDirection() == ERigVMPinDirection::Output)
 					{
-						Local::VisitPinRecursively(RootPin, ModelPins, bSupportSubPins);
+						Local::VisitPinRecursively(RootPin, ModelPinsOrCategories, bSupportSubPins, DirectionsWithCategories);
 					}
 					break;
 				}
@@ -1642,7 +1734,7 @@ void SRigVMGraphNode::UpdatePinTreeView()
 				{
 					if(!RootPin->IsExecuteContext() && RootPin->GetDirection() == ERigVMPinDirection::IO)
 					{
-						Local::VisitPinRecursively(RootPin, ModelPins, bSupportSubPins);
+						Local::VisitPinRecursively(RootPin, ModelPinsOrCategories, bSupportSubPins, DirectionsWithCategories);
 					}
 					break;
 				}
@@ -1651,7 +1743,7 @@ void SRigVMGraphNode::UpdatePinTreeView()
 				{
 					if(RootPin->GetDirection() == ERigVMPinDirection::Input || RootPin->GetDirection() == ERigVMPinDirection::Visible)
 					{
-						Local::VisitPinRecursively(RootPin, ModelPins, bSupportSubPins);
+						Local::VisitPinRecursively(RootPin, ModelPinsOrCategories, bSupportSubPins, DirectionsWithCategories);
 					}
 					break;
 				}
@@ -1666,51 +1758,125 @@ void SRigVMGraphNode::UpdatePinTreeView()
 	{
 		DispatchContext = DispatchNode->GetDispatchContext();
 	}
-	
-	TMap<URigVMPin*, int32> ModelPinToInfoIndex;
-	for(URigVMPin* ModelPin : ModelPins)
+
+	// todo: we should implement traversers here which visit the pins in the right order.
+	// either by visiting the pins as we used to or using the categories. we need to figure out if for each group there's a category
+	// available and then use that instead. so in other words show all IO pins except if there's a category for an IO pin. etc.
+
+	TMap<FString, int32> IdentifierToInfoIndex;
+
+	auto AddModelPinOrCategoryPinInfo =
+		[this, &IdentifierToInfoIndex, RigSchema, Template, &DispatchContext, RigGraphNode, EdGraphPinToInputPin, EdGraphPinToOutputPin, bSupportSubPins]
+	(const URigVMPin* InModelPin, const FString& InCategory)
 	{
+		const URigVMNode* CurrentModelNode = RigGraphNode->GetModelNode();
+		check(CurrentModelNode);
+			
+		const bool bIsModelPin = InModelPin != nullptr;
+		const bool bIsCategoryPin = !InCategory.IsEmpty();
+		check(bIsModelPin != bIsCategoryPin);
+			
 		FPinInfo PinInfo;
 		PinInfo.Index = PinInfos.Num();
 		PinInfo.ParentIndex = INDEX_NONE;
-		PinInfo.bHasChildren = (ModelPin->GetSubPins().Num() > 0);
-		PinInfo.bIsContainer = ModelPin->IsArray();
+		PinInfo.bIsCategoryPin = bIsCategoryPin;
+		if(bIsModelPin)
+		{
+			PinInfo.bHasChildren = InModelPin->GetSubPins().Num() > 0;
+			PinInfo.bIsContainer = InModelPin->IsArray();
+			PinInfo.bExpanded = InModelPin->IsExpanded();
+			PinInfo.Identifier = InModelPin->GetPinPath();
+			PinInfo.bShowOnlySubPins = InModelPin->ShouldOnlyShowSubPins();
+		}
+		else
+		{
+			PinInfo.bHasChildren = false;
+			TArray<FString> CategoriesWithPins = {InCategory};
+			CategoriesWithPins.Append(CurrentModelNode->GetSubPinCategories(InCategory, false, true));
+			for(const FString& CategoryWithPins : CategoriesWithPins)
+			{
+				if(!CurrentModelNode->GetPinsForCategory(CategoryWithPins).IsEmpty())
+				{
+					PinInfo.bHasChildren = true;
+					break;
+				}
+			}
+			PinInfo.bIsContainer = false;
+			PinInfo.bExpanded = CurrentModelNode->IsPinCategoryExpanded(InCategory);
+			PinInfo.Identifier = InCategory;
+			PinInfo.bShowOnlySubPins = false;
+			PinInfo.bHideInputWidget = true;
+		}
 		PinInfo.Depth = 0;
-		PinInfo.bExpanded = ModelPin->IsExpanded();
-		PinInfo.ModelPinPath = ModelPin->GetPinPath();
 		PinInfo.bAutoHeight = false;
-		PinInfo.bShowOnlySubPins = ModelPin->ShouldOnlyShowSubPins();
 
-		if(!bSupportSubPins)
+		if(!bSupportSubPins && bIsModelPin)
 		{
 			PinInfo.bHasChildren = false;
 			PinInfo.bIsContainer = false;
 		}
-		
-		const bool bAskSchemaForEdition = RigSchema && ModelPin->IsStruct() && !ModelPin->IsBoundToVariable();
-		PinInfo.bHideInputWidget = (!ModelPin->IsBoundToVariable()) && PinInfo.bIsContainer;
-		if (!PinInfo.bHideInputWidget)
+
+		if(bIsModelPin)
 		{
-			if (bAskSchemaForEdition && !PinInfo.bHasChildren)
+			const bool bAskSchemaForEdition = RigSchema && InModelPin->IsStruct() && !InModelPin->IsBoundToVariable();
+			PinInfo.bHideInputWidget = (!InModelPin->IsBoundToVariable()) && PinInfo.bIsContainer;
+			if (!PinInfo.bHideInputWidget)
 			{
-				const bool bIsStructEditable = RigSchema->IsStructEditable(ModelPin->GetScriptStruct()); 
-				PinInfo.bHideInputWidget = !bIsStructEditable;
-				PinInfo.bAutoHeight = bIsStructEditable;
+				if (bAskSchemaForEdition && !PinInfo.bHasChildren)
+				{
+					const bool bIsStructEditable = RigSchema->IsStructEditable(InModelPin->GetScriptStruct()); 
+					PinInfo.bHideInputWidget = !bIsStructEditable;
+					PinInfo.bAutoHeight = bIsStructEditable;
+				}
+				else if(PinInfo.bHasChildren && !InModelPin->IsBoundToVariable())
+				{
+					PinInfo.bHideInputWidget = true;
+				}
 			}
-			else if(PinInfo.bHasChildren && !ModelPin->IsBoundToVariable())
+		}
+
+		if(bIsModelPin)
+		{
+			const FString PinCategory = InModelPin->GetCategory();
+			if(!PinCategory.IsEmpty())
+			{
+				if(const int32* ParentIndexPtr = IdentifierToInfoIndex.Find(PinCategory))
+				{
+					PinInfo.ParentIndex = *ParentIndexPtr;
+				}
+			}
+
+			// if the pin is not part of a category, add it under its default parent
+			if(PinInfo.ParentIndex == INDEX_NONE)
 			{
-				PinInfo.bHideInputWidget = true;
+				if(URigVMPin* ParentPin = InModelPin->GetParentPin())
+				{
+					const int32* ParentIndexPtr = IdentifierToInfoIndex.Find(ParentPin->GetPinPath());
+					if(ParentIndexPtr == nullptr)
+					{
+						// parent pins have to exist for the node to display correctly
+						return;
+					}
+					PinInfo.ParentIndex = *ParentIndexPtr;
+				}
 			}
 		}
-		
-		if(URigVMPin* ParentPin = ModelPin->GetParentPin())
+		else
 		{
-			const int32* ParentIndexPtr = ModelPinToInfoIndex.Find(ParentPin);
-			if(ParentIndexPtr == nullptr)
+			const FString ParentCategory = CurrentModelNode->GetParentPinCategory(InCategory);
+			if(!ParentCategory.IsEmpty())
 			{
-				continue;
+				const int32* ParentIndexPtr = IdentifierToInfoIndex.Find(ParentCategory);
+				if(ParentIndexPtr == nullptr)
+				{
+					return;
+				}
+				PinInfo.ParentIndex = *ParentIndexPtr;
 			}
-			PinInfo.ParentIndex = *ParentIndexPtr;
+		}
+
+		if(PinInfo.ParentIndex != INDEX_NONE)
+		{
 			PinInfo.Depth = PinInfos[PinInfo.ParentIndex].Depth + 1;
 			if(PinInfos[PinInfo.ParentIndex].bShowOnlySubPins)
 			{
@@ -1724,7 +1890,7 @@ void SRigVMGraphNode::UpdatePinTreeView()
 		TSharedPtr<SGraphPin> PinWidgetForExpander;
 
 		bool bPinInfoIsValid = false;
-		if(UEdGraphPin* OutputEdGraphPin = RigGraphNode->FindPin(ModelPin->GetPinPath(), EEdGraphPinDirection::EGPD_Output))
+		if(UEdGraphPin* OutputEdGraphPin = RigGraphNode->FindPin(PinInfo.Identifier, EEdGraphPinDirection::EGPD_Output))
 		{
 			if(const int32* PinIndexPtr = EdGraphPinToOutputPin.Find(OutputEdGraphPin))
 			{
@@ -1733,9 +1899,10 @@ void SRigVMGraphNode::UpdatePinTreeView()
 				PinWidgetForExpander = PinInfo.OutputPinWidget;
 				bPinWidgetForExpanderLeft = false;
 				bPinInfoIsValid = true;
-				if (Template && !ModelPin->IsExecuteContext())
+
+				if (bIsModelPin && Template && !InModelPin->IsExecuteContext())
 				{
-					if (URigVMPin* RootPin = ModelPin->GetRootPin())
+					if (URigVMPin* RootPin = InModelPin->GetRootPin())
 					{
 						FLinearColor PinColorAndOpacity = PinInfo.OutputPinWidget->GetColorAndOpacity();
 						if (Template->FindArgument(RootPin->GetFName()) == nullptr && Template->FindExecuteArgument(RootPin->GetFName(), DispatchContext) == nullptr)
@@ -1752,7 +1919,7 @@ void SRigVMGraphNode::UpdatePinTreeView()
 			}
 		}
 		
-		if(UEdGraphPin* InputEdGraphPin = RigGraphNode->FindPin(ModelPin->GetPinPath(), EEdGraphPinDirection::EGPD_Input))
+		if(UEdGraphPin* InputEdGraphPin = RigGraphNode->FindPin(PinInfo.Identifier, EEdGraphPinDirection::EGPD_Input))
 		{
 			if(const int32* PinIndexPtr = EdGraphPinToInputPin.Find(InputEdGraphPin))
 			{
@@ -1761,9 +1928,9 @@ void SRigVMGraphNode::UpdatePinTreeView()
 				PinWidgetForExpander = PinInfo.InputPinWidget;
 				bPinWidgetForExpanderLeft = true;
 				bPinInfoIsValid = true;
-				if (Template && !ModelPin->IsExecuteContext())
+				if (bIsModelPin && Template && !InModelPin->IsExecuteContext())
 				{
-					if (URigVMPin* RootPin = ModelPin->GetRootPin())
+					if (URigVMPin* RootPin = InModelPin->GetRootPin())
 					{
 						FLinearColor PinColorAndOpacity = PinInfo.InputPinWidget->GetColorAndOpacity();
 						if (Template->FindArgument(RootPin->GetFName()) == nullptr && Template->FindExecuteArgument(RootPin->GetFName(), DispatchContext) == nullptr)
@@ -1782,10 +1949,36 @@ void SRigVMGraphNode::UpdatePinTreeView()
 
 		if(!bPinInfoIsValid)
 		{
-			continue;
+			return;
 		}
 
-		ModelPinToInfoIndex.Add(ModelPin, PinInfos.Add(PinInfo));
+		ERigVMPinDirection PinDirection = ERigVMPinDirection::Invalid;
+		if(bIsModelPin)
+		{
+			PinDirection = InModelPin->GetDirection();
+		}
+		else
+		{
+			TArray<FString> CategoriesToCheck = {InCategory};
+			for(int32 CategoryIndex = 0; CategoryIndex < CategoriesToCheck.Num(); CategoryIndex++)
+			{
+				const FString& CategoryToCheck = CategoriesToCheck[CategoryIndex];
+				const TArray<URigVMPin*> PinsForCategory = CurrentModelNode->GetPinsForCategory(CategoryToCheck);
+				if(!PinsForCategory.IsEmpty())
+				{
+					PinDirection = PinsForCategory[0]->GetDirection();
+					break;
+				}
+				CategoriesToCheck.Append(CurrentModelNode->GetSubPinCategories(CategoryToCheck));
+			}
+		}
+
+		if(PinDirection == ERigVMPinDirection::Invalid)
+		{
+			return;
+		}
+
+		IdentifierToInfoIndex.Add(PinInfo.Identifier, PinInfos.Add(PinInfo));
 		
 		// check if this pin has sub pins
 		TSharedPtr<SHorizontalBox> FullPinHorizontalRowWidget = PinWidgetForExpander->GetFullPinHorizontalRowWidget().Pin();
@@ -1834,13 +2027,13 @@ void SRigVMGraphNode::UpdatePinTreeView()
 				if(
 					(
 						(
-							(ModelPin->GetDirection() == ERigVMPinDirection::Input) ||
-							(ModelPin->GetDirection() == ERigVMPinDirection::IO)
+							(PinDirection == ERigVMPinDirection::Input) ||
+							(PinDirection == ERigVMPinDirection::IO)
 						) &&
 						bPinWidgetForExpanderLeft
 					) ||
 					(
-						(ModelPin->GetDirection() == ERigVMPinDirection::Output) &&
+						(PinDirection == ERigVMPinDirection::Output) &&
 						(!bPinWidgetForExpanderLeft)
 					)
 				)
@@ -1875,6 +2068,11 @@ void SRigVMGraphNode::UpdatePinTreeView()
 				Slot.SetPadding(LineIndentation);
 			}
 		}
+	};
+	
+	for(const TTuple<URigVMPin*,FString>& ModelPinOrCategory : ModelPinsOrCategories)
+	{
+		AddModelPinOrCategoryPinInfo(ModelPinOrCategory.Get<0>(), ModelPinOrCategory.Get<1>());
 	}
 
 	// add spacer widget at the start
@@ -2066,7 +2264,7 @@ void SRigVMGraphNode::UpdatePinTreeView()
                     +SHorizontalBox::Slot()
                     .Expose(FirstSlot)
                     .FillWidth(1.f)
-                    .HAlign(HAlign_Left)
+                    .HAlign(HAlign_Fill)
                     .Padding(PinWidgetSidePadding, TopPadding, InputPinInfo.bIsContainer ? 0.f : MyEmptySidePadding, 0.f)
                     [
                         InputPinInfo.InputPinWidget.ToSharedRef()
@@ -2075,7 +2273,7 @@ void SRigVMGraphNode::UpdatePinTreeView()
 
 				if(InputPinInfo.bIsContainer)
 				{
-					URigVMPin* ModelPin = ModelNode->GetGraph()->FindPin(InputPinInfo.ModelPinPath);
+					URigVMPin* ModelPin = ModelNode->GetGraph()->FindPin(InputPinInfo.Identifier);
 					if(ModelPin)
 					{
 						// make sure to minimize the width of the label
@@ -2114,7 +2312,7 @@ void SRigVMGraphNode::UpdatePinTreeView()
 
 				if(InputPinInfo.bIsContainer)
 				{
-					URigVMPin* ModelPin = ModelNode->GetGraph()->FindPin(InputPinInfo.ModelPinPath);
+					URigVMPin* ModelPin = ModelNode->GetGraph()->FindPin(InputPinInfo.Identifier);
 					if(ModelPin)
 					{
 						// make sure to minimize the width of the label
@@ -2234,4 +2432,16 @@ bool SRigVMGraphNode::HasUserProvidedValue(const SGraphPin* InGraphPin) const
 	return false;
 }
 
+bool SRigVMGraphNode::IsCategoryPin(const SGraphPin* InGraphPin) const
+{
+	if (const URigVMEdGraphNode* RigVMEdGraphNode = Cast<URigVMEdGraphNode>(GraphNode))
+	{
+		if (const UEdGraphPin* Pin = InGraphPin->GetPinObj())
+		{
+			return RigVMEdGraphNode->CachedCategoryPins.Contains(Pin->GetName());
+		}
+	}
+	return false;
+}
+
 #undef LOCTEXT_NAMESPACE
diff --git a/Engine/Plugins/Runtime/RigVM/Source/RigVMEditor/Private/Widgets/SRigVMGraphPinCategory.cpp b/Engine/Plugins/Runtime/RigVM/Source/RigVMEditor/Private/Widgets/SRigVMGraphPinCategory.cpp
new file mode 100644
--- /dev/null
+++ b/Engine/Plugins/Runtime/RigVM/Source/RigVMEditor/Private/Widgets/SRigVMGraphPinCategory.cpp
@@ -0,0 +1,16 @@
+// Copyright Epic Games, Inc. All Rights Reserved.
+
+
+#include "Widgets/SRigVMGraphPinCategory.h"
+#include "Widgets/SNullWidget.h"
+
+void SRigVMGraphPinCategory::Construct(const FArguments& InArgs, UEdGraphPin* InGraphPinObj)
+{
+	SGraphPin::Construct(SGraphPin::FArguments(), InGraphPinObj);
+}
+
+TSharedRef<SWidget>	SRigVMGraphPinCategory::GetDefaultValueWidget()
+{
+	return SNullWidget::NullWidget;
+}
+
diff --git a/Engine/Plugins/Runtime/RigVM/Source/RigVMEditor/Public/Widgets/SRigVMGraphNode.h b/Engine/Plugins/Runtime/RigVM/Source/RigVMEditor/Public/Widgets/SRigVMGraphNode.h
--- a/Engine/Plugins/Runtime/RigVM/Source/RigVMEditor/Public/Widgets/SRigVMGraphNode.h
+++ b/Engine/Plugins/Runtime/RigVM/Source/RigVMEditor/Public/Widgets/SRigVMGraphNode.h
@@ -96,6 +96,7 @@ protected:
 	void HandleModifiedEvent(ERigVMGraphNotifType InNotifType, URigVMGraph* InGraph, UObject* InSubject);
 	virtual void UpdatePinTreeView();
 	bool HasUserProvidedValue(const SGraphPin* InGraphPin) const;
+	bool IsCategoryPin(const SGraphPin* InGraphPin) const;
 
 	/** Cached widget title area */
 	TSharedPtr<SOverlay> TitleAreaWidget;
@@ -120,11 +121,12 @@ protected:
 	{
 		int32 Index;
 		int32 ParentIndex;
+		bool bIsCategoryPin;
 		bool bHasChildren;
 		bool bHideInputWidget;
 		bool bIsContainer;
 		int32 Depth;
-		FString ModelPinPath;
+		FString Identifier;
 		TSharedPtr<SGraphPin> InputPinWidget;
 		TSharedPtr<SGraphPin> OutputPinWidget;
 		bool bExpanded;
diff --git a/Engine/Plugins/Runtime/RigVM/Source/RigVMEditor/Public/Widgets/SRigVMGraphPinCategory.h b/Engine/Plugins/Runtime/RigVM/Source/RigVMEditor/Public/Widgets/SRigVMGraphPinCategory.h
new file mode 100644
--- /dev/null
+++ b/Engine/Plugins/Runtime/RigVM/Source/RigVMEditor/Public/Widgets/SRigVMGraphPinCategory.h
@@ -0,0 +1,26 @@
+// Copyright Epic Games, Inc. All Rights Reserved.
+
+#pragma once
+
+#include "CoreMinimal.h"
+#include "Widgets/DeclarativeSyntaxSupport.h"
+#include "Widgets/SWidget.h"
+#include "SGraphPin.h"
+
+class RIGVMEDITOR_API SRigVMGraphPinCategory : public SGraphPin
+{
+public:
+	SLATE_BEGIN_ARGS(SRigVMGraphPinCategory)
+	{}
+	SLATE_END_ARGS()
+
+
+	void Construct(const FArguments& InArgs, UEdGraphPin* InGraphPinObj);
+
+protected:
+
+	//~ Begin SGraphPin Interface
+	virtual TSharedRef<SWidget>	GetDefaultValueWidget() override;
+	virtual const FSlateBrush* GetPinIcon() const override { return nullptr; }
+	//~ End SGraphPin Interface
+};
