Control Rig: Fix tooltip message on array init node
#rnx

diff --git a/Engine/Plugins/Runtime/RigVM/Source/RigVM/Private/RigVMFunctions/RigVMDispatch_Array.cpp b/Engine/Plugins/Runtime/RigVM/Source/RigVM/Private/RigVMFunctions/RigVMDispatch_Array.cpp
--- a/Engine/Plugins/Runtime/RigVM/Source/RigVM/Private/RigVMFunctions/RigVMDispatch_Array.cpp
+++ b/Engine/Plugins/Runtime/RigVM/Source/RigVM/Private/RigVMFunctions/RigVMDispatch_Array.cpp
@@ -542,7 +542,7 @@ FText FRigVMDispatch_ArrayInit::GetArgumentTooltip(const FName& InArgumentName,
 	}
 	if(InArgumentName == NumName)
 	{
-		return LOCTEXT("ArraySetNum_NumArgumentToolTip", "The new size of the array.");
+		return LOCTEXT("ArrayInit_NumArgumentToolTip", "The new size of the array.");
 	}
 	if(InArgumentName == ElementName)
 	{
