ControlRig: Fix crash when redoing change variable type. The action stack was not being informed in time.
#rnx

diff --git a/Engine/Plugins/Runtime/RigVM/Source/RigVMDeveloper/Private/RigVMBlueprint.cpp b/Engine/Plugins/Runtime/RigVM/Source/RigVMDeveloper/Private/RigVMBlueprint.cpp
--- a/Engine/Plugins/Runtime/RigVM/Source/RigVMDeveloper/Private/RigVMBlueprint.cpp
+++ b/Engine/Plugins/Runtime/RigVM/Source/RigVMDeveloper/Private/RigVMBlueprint.cpp
@@ -20,6 +20,7 @@
 #include "RigVMPythonUtils.h"
 #include "RigVMTypeUtils.h"
 #include "Algo/Count.h"
+#include "RigVMModel/RigVMControllerActions.h"
 #include "RigVMModel/Nodes/RigVMAggregateNode.h"
 #include "RigVMModel/Nodes/RigVMDispatchNode.h"
 #include "UObject/UE5MainStreamObjectVersion.h"
@@ -2580,6 +2581,14 @@ void URigVMBlueprint::SetObjectBeingDebugged(UObject* NewObject)
 void URigVMBlueprint::PostTransacted(const FTransactionObjectEvent& TransactionEvent)
 {
 	DECLARE_SCOPE_HIERARCHICAL_COUNTER_FUNC()
+
+	if (TransactionEvent.GetEventType() == ETransactionObjectEventType::UndoRedo)
+	{
+		// The action stack undo/redo transaction should always execute first
+		// It already knows whether or not it has already executed or not
+		RigVMClient.GetOrCreateActionStack()->PostTransacted(TransactionEvent);
+	}
+	
 	Super::PostTransacted(TransactionEvent);
 
 	if (TransactionEvent.GetEventType() == ETransactionObjectEventType::UndoRedo)
diff --git a/Engine/Plugins/Runtime/RigVM/Source/RigVMDeveloper/Public/RigVMModel/RigVMClient.h b/Engine/Plugins/Runtime/RigVM/Source/RigVMDeveloper/Public/RigVMModel/RigVMClient.h
--- a/Engine/Plugins/Runtime/RigVM/Source/RigVMDeveloper/Public/RigVMModel/RigVMClient.h
+++ b/Engine/Plugins/Runtime/RigVM/Source/RigVMDeveloper/Public/RigVMModel/RigVMClient.h
@@ -322,4 +322,5 @@ private:
 	FName OuterClientPropertyName;
 
 	friend class UEngineTestClientHost;
+	friend class URigVMBlueprint;
 };
