RigVM: Pin Categories - more work towards a pin layout

#rnx

diff --git a/Engine/Plugins/Runtime/RigVM/Config/BaseRigVM.ini b/Engine/Plugins/Runtime/RigVM/Config/BaseRigVM.ini
--- a/Engine/Plugins/Runtime/RigVM/Config/BaseRigVM.ini
+++ b/Engine/Plugins/Runtime/RigVM/Config/BaseRigVM.ini
@@ -2,4 +2,8 @@
 
 ; renaming decorator to trait
 +StructRedirects=(OldName="RigVMDecorator", NewName="/Script/RigVM.RigVMTrait")
-+PropertyRedirects=(OldName="RigVMNode.DecoratorRootPinNames",NewName="RigVMNode.TraitRootPinNames")
\ No newline at end of file
++PropertyRedirects=(OldName="RigVMNode.DecoratorRootPinNames",NewName="RigVMNode.TraitRootPinNames")
+
+; graph function layout renaming
++StructRedirects=(OldName="RigVMGraphFunctionLayout", NewName="/Script/RigVM.RigVMNodeLayout")
++StructRedirects=(OldName="RigVMGraphFunctionCategory", NewName="/Script/RigVM.RigVMPinCategory")
\ No newline at end of file
diff --git a/Engine/Plugins/Runtime/RigVM/Source/RigVM/Private/RigVMCore/RigVMGraphFunctionDefinition.cpp b/Engine/Plugins/Runtime/RigVM/Source/RigVM/Private/RigVMCore/RigVMGraphFunctionDefinition.cpp
--- a/Engine/Plugins/Runtime/RigVM/Source/RigVM/Private/RigVMCore/RigVMGraphFunctionDefinition.cpp
+++ b/Engine/Plugins/Runtime/RigVM/Source/RigVM/Private/RigVMCore/RigVMGraphFunctionDefinition.cpp
@@ -94,23 +94,6 @@ bool FRigVMGraphFunctionIdentifier::IsVariantOf(const FRigVMGraphFunctionIdentif
 	return ThisHeader.Variant.Guid == OtherHeader.Variant.Guid;
 }
 
-const FString* FRigVMGraphFunctionLayout::FindCategory(const FString& InElement) const
-{
-	for(const FRigVMGraphFunctionCategory& Category : Categories)
-	{
-		if(Category.Elements.Contains(InElement))
-		{
-			return &Category.Path;
-		}
-	}
-	return nullptr;
-}
-
-const FString* FRigVMGraphFunctionLayout::FindDisplayName(const FString& InElement) const
-{
-	return DisplayNames.Find(InElement);
-}
-
 TFunction<FRigVMGraphFunctionHeader(const FSoftObjectPath&, const FName&, bool*)> FRigVMGraphFunctionHeader::FindFunctionHeaderFromPathFunc;
 
 bool FRigVMGraphFunctionHeader::IsMutable() const
@@ -564,7 +547,7 @@ bool FRigVMGraphFunctionData::PatchSharedArgumentOperandsIfRequired()
 	return true;
 }
 
-FArchive& operator<<(FArchive& Ar, FRigVMGraphFunctionLayout& Layout)
+FArchive& operator<<(FArchive& Ar, FRigVMNodeLayout& Layout)
 {
 	Ar.UsingCustomVersion(FRigVMObjectVersion::GUID);
 
diff --git a/Engine/Plugins/Runtime/RigVM/Source/RigVM/Private/RigVMCore/RigVMNodeLayout.cpp b/Engine/Plugins/Runtime/RigVM/Source/RigVM/Private/RigVMCore/RigVMNodeLayout.cpp
new file mode 100644
--- /dev/null
+++ b/Engine/Plugins/Runtime/RigVM/Source/RigVM/Private/RigVMCore/RigVMNodeLayout.cpp
@@ -0,0 +1,23 @@
+// Copyright Epic Games, Inc. All Rights Reserved.
+
+#include "RigVMCore/RigVMNodeLayout.h"
+#include "RigVMStringUtils.h"
+
+#include UE_INLINE_GENERATED_CPP_BY_NAME(RigVMNodeLayout)
+
+const FString* FRigVMNodeLayout::FindCategory(const FString& InElement) const
+{
+	for(const FRigVMPinCategory& Category : Categories)
+	{
+		if(Category.Elements.Contains(InElement))
+		{
+			return &Category.Path;
+		}
+	}
+	return nullptr;
+}
+
+const FString* FRigVMNodeLayout::FindDisplayName(const FString& InElement) const
+{
+	return DisplayNames.Find(InElement);
+}
\ No newline at end of file
diff --git a/Engine/Plugins/Runtime/RigVM/Source/RigVM/Public/RigVMCore/RigVMGraphFunctionDefinition.h b/Engine/Plugins/Runtime/RigVM/Source/RigVM/Public/RigVMCore/RigVMGraphFunctionDefinition.h
--- a/Engine/Plugins/Runtime/RigVM/Source/RigVM/Public/RigVMCore/RigVMGraphFunctionDefinition.h
+++ b/Engine/Plugins/Runtime/RigVM/Source/RigVM/Public/RigVMCore/RigVMGraphFunctionDefinition.h
@@ -9,6 +9,7 @@
 #include "UObject/FortniteMainBranchObjectVersion.h"
 #include "UObject/UE5ReleaseStreamObjectVersion.h"
 #include "RigVMVariant.h"
+#include "RigVMNodeLayout.h"
 #include "RigVMGraphFunctionDefinition.generated.h"
 
 class IRigVMGraphFunctionHost;
@@ -486,96 +487,6 @@ protected:
 	friend class URigVMBuildData;
 };
 
-USTRUCT(BlueprintType)
-struct RIGVM_API FRigVMGraphFunctionCategory
-{
-	GENERATED_BODY()
-	
-	FRigVMGraphFunctionCategory()
-	: Path()
-	, Elements()
-	{}
-
-	UPROPERTY(BlueprintReadOnly, VisibleAnywhere, Category=FunctionArgument)
-	FString Path;
-
-	UPROPERTY(BlueprintReadOnly, VisibleAnywhere, Category=FunctionArgument)
-	TArray<FString> Elements;
-
-	friend uint32 GetTypeHash(const FRigVMGraphFunctionCategory& Category)
-	{
-		uint32 Hash = GetTypeHash(Category.Path);
-		for (const FString& Element : Category.Elements)
-		{
-			Hash = HashCombine(Hash, GetTypeHash(Element));
-		}
-		return Hash;
-	}
-
-	bool operator < (const FRigVMGraphFunctionCategory& Other) const
-	{
-		return FCString::Strcmp(*Path, *Other.Path) < 0;
-	}
-
-	friend FArchive& operator<<(FArchive& Ar, FRigVMGraphFunctionCategory& Category)
-	{
-		Ar << Category.Path;
-		Ar << Category.Elements;
-		return Ar;
-	}
-};
-
-USTRUCT(BlueprintType)
-struct RIGVM_API FRigVMGraphFunctionLayout
-{
-	GENERATED_BODY()
-	
-	FRigVMGraphFunctionLayout()
-	: Categories()
-	, DisplayNames()
-	{}
-
-	UPROPERTY(BlueprintReadOnly, VisibleAnywhere, Category=FunctionArgument)
-	TArray<FRigVMGraphFunctionCategory> Categories;
-
-	UPROPERTY(BlueprintReadOnly, VisibleAnywhere, Category=FunctionArgument)
-	TMap<FString, int32> PinIndexInCategory;
-
-	UPROPERTY(BlueprintReadOnly, VisibleAnywhere, Category=FunctionArgument)
-	TMap<FString, FString> DisplayNames;
-
-	void Reset()
-	{
-		Categories.Reset();
-		PinIndexInCategory.Reset();
-		DisplayNames.Reset();
-	}
-
-	friend uint32 GetTypeHash(const FRigVMGraphFunctionLayout& Layout)
-	{
-		uint32 Hash = 0;;
-		for (const FRigVMGraphFunctionCategory& Category : Layout.Categories)
-		{
-			Hash = HashCombine(Hash, GetTypeHash(Category));
-		}
-		for(const TPair<FString, int32>& Pair : Layout.PinIndexInCategory)
-		{
-			Hash = HashCombine(Hash, GetTypeHash(Pair));
-		}
-		for(const TPair<FString, FString>& Pair : Layout.DisplayNames)
-		{
-			Hash = HashCombine(Hash, GetTypeHash(Pair));
-		}
-		return Hash;
-	}
-
-	friend RIGVM_API FArchive& operator<<(FArchive& Ar, FRigVMGraphFunctionLayout& Layout);
-
-	const FString* FindCategory(const FString& InElement) const;
-	const FString* FindDisplayName(const FString& InElement) const;
-};
-
-
 USTRUCT(BlueprintType)
 struct RIGVM_API FRigVMGraphFunctionHeader
 {
@@ -623,7 +534,7 @@ struct RIGVM_API FRigVMGraphFunctionHeader
 	TArray<FRigVMExternalVariable> ExternalVariables;
 
 	UPROPERTY()
-	FRigVMGraphFunctionLayout Layout;
+	FRigVMNodeLayout Layout;
 
 	bool IsMutable() const;
 
diff --git a/Engine/Plugins/Runtime/RigVM/Source/RigVM/Public/RigVMCore/RigVMNodeLayout.h b/Engine/Plugins/Runtime/RigVM/Source/RigVM/Public/RigVMCore/RigVMNodeLayout.h
new file mode 100644
--- /dev/null
+++ b/Engine/Plugins/Runtime/RigVM/Source/RigVM/Public/RigVMCore/RigVMNodeLayout.h
@@ -0,0 +1,96 @@
+// Copyright Epic Games, Inc. All Rights Reserved.
+
+#pragma once
+
+#include "RigVMObjectVersion.h"
+#include "RigVMNodeLayout.generated.h"
+
+USTRUCT(BlueprintType)
+struct RIGVM_API FRigVMPinCategory
+{
+	GENERATED_BODY()
+	
+	FRigVMPinCategory()
+	: Path()
+	, Elements()
+	{}
+
+	UPROPERTY(BlueprintReadOnly, VisibleAnywhere, Category=FunctionArgument)
+	FString Path;
+
+	UPROPERTY(BlueprintReadOnly, VisibleAnywhere, Category=FunctionArgument)
+	TArray<FString> Elements;
+
+	friend uint32 GetTypeHash(const FRigVMPinCategory& Category)
+	{
+		uint32 Hash = GetTypeHash(Category.Path);
+		for (const FString& Element : Category.Elements)
+		{
+			Hash = HashCombine(Hash, GetTypeHash(Element));
+		}
+		return Hash;
+	}
+
+	bool operator < (const FRigVMPinCategory& Other) const
+	{
+		return FCString::Strcmp(*Path, *Other.Path) < 0;
+	}
+
+	friend FArchive& operator<<(FArchive& Ar, FRigVMPinCategory& Category)
+	{
+		Ar << Category.Path;
+		Ar << Category.Elements;
+		return Ar;
+	}
+};
+
+USTRUCT(BlueprintType)
+struct RIGVM_API FRigVMNodeLayout
+{
+	GENERATED_BODY()
+	
+	FRigVMNodeLayout()
+	: Categories()
+	, DisplayNames()
+	{}
+
+	UPROPERTY(BlueprintReadOnly, VisibleAnywhere, Category=FunctionArgument)
+	TArray<FRigVMPinCategory> Categories;
+
+	UPROPERTY(BlueprintReadOnly, VisibleAnywhere, Category=FunctionArgument)
+	TMap<FString, int32> PinIndexInCategory;
+
+	UPROPERTY(BlueprintReadOnly, VisibleAnywhere, Category=FunctionArgument)
+	TMap<FString, FString> DisplayNames;
+
+	void Reset()
+	{
+		Categories.Reset();
+		PinIndexInCategory.Reset();
+		DisplayNames.Reset();
+	}
+
+	friend uint32 GetTypeHash(const FRigVMNodeLayout& Layout)
+	{
+		uint32 Hash = 0;;
+		for (const FRigVMPinCategory& Category : Layout.Categories)
+		{
+			Hash = HashCombine(Hash, GetTypeHash(Category));
+		}
+		for(const TPair<FString, int32>& Pair : Layout.PinIndexInCategory)
+		{
+			Hash = HashCombine(Hash, GetTypeHash(Pair));
+		}
+		for(const TPair<FString, FString>& Pair : Layout.DisplayNames)
+		{
+			Hash = HashCombine(Hash, GetTypeHash(Pair));
+		}
+		return Hash;
+	}
+
+	friend RIGVM_API FArchive& operator<<(FArchive& Ar, FRigVMNodeLayout& Layout);
+
+	const FString* FindCategory(const FString& InElement) const;
+	const FString* FindDisplayName(const FString& InElement) const;
+};
+
diff --git a/Engine/Plugins/Runtime/RigVM/Source/RigVMDeveloper/Private/EdGraph/RigVMEdGraphNode.cpp b/Engine/Plugins/Runtime/RigVM/Source/RigVMDeveloper/Private/EdGraph/RigVMEdGraphNode.cpp
--- a/Engine/Plugins/Runtime/RigVM/Source/RigVMDeveloper/Private/EdGraph/RigVMEdGraphNode.cpp
+++ b/Engine/Plugins/Runtime/RigVM/Source/RigVMDeveloper/Private/EdGraph/RigVMEdGraphNode.cpp
@@ -2063,7 +2063,7 @@ FEdGraphPinType URigVMEdGraphNode::GetPinTypeForModelPin(const URigVMPin* InMode
 
 FEdGraphPinType URigVMEdGraphNode::GetPinTypeForCategoryPin()
 {
-	static const FEdGraphPinType PinType = RigVMTypeUtils::PinTypeFromCPPType(NAME_None, FRigVMGraphFunctionCategory::StaticStruct());
+	static const FEdGraphPinType PinType = RigVMTypeUtils::PinTypeFromCPPType(NAME_None, FRigVMPinCategory::StaticStruct());
 	return PinType;
 }
 
diff --git a/Engine/Plugins/Runtime/RigVM/Source/RigVMDeveloper/Private/RigVMModel/Nodes/RigVMFunctionReferenceNode.cpp b/Engine/Plugins/Runtime/RigVM/Source/RigVMDeveloper/Private/RigVMModel/Nodes/RigVMFunctionReferenceNode.cpp
--- a/Engine/Plugins/Runtime/RigVM/Source/RigVMDeveloper/Private/RigVMModel/Nodes/RigVMFunctionReferenceNode.cpp
+++ b/Engine/Plugins/Runtime/RigVM/Source/RigVMDeveloper/Private/RigVMModel/Nodes/RigVMFunctionReferenceNode.cpp
@@ -312,7 +312,7 @@ FText URigVMFunctionReferenceNode::GetToolTipTextForPin(const URigVMPin* InPin)
 TArray<FString> URigVMFunctionReferenceNode::GetPinCategories() const
 {
 	TArray<FString> TransientPinCategories;
-	for(const FRigVMGraphFunctionCategory& Category : ReferencedFunctionHeader.Layout.Categories)
+	for(const FRigVMPinCategory& Category : ReferencedFunctionHeader.Layout.Categories)
 	{
 		TransientPinCategories.Add(Category.Path);
 	}
diff --git a/Engine/Plugins/Runtime/RigVM/Source/RigVMDeveloper/Private/RigVMModel/Nodes/RigVMLibraryNode.cpp b/Engine/Plugins/Runtime/RigVM/Source/RigVMDeveloper/Private/RigVMModel/Nodes/RigVMLibraryNode.cpp
--- a/Engine/Plugins/Runtime/RigVM/Source/RigVMDeveloper/Private/RigVMModel/Nodes/RigVMLibraryNode.cpp
+++ b/Engine/Plugins/Runtime/RigVM/Source/RigVMDeveloper/Private/RigVMModel/Nodes/RigVMLibraryNode.cpp
@@ -328,40 +328,7 @@ FRigVMGraphFunctionHeader URigVMLibraryNode::GetFunctionHeader(IRigVMGraphFuncti
 		}
 	}
 
-	// fill in the pin categories based on the data stored on the pins themselves
-	const TArray<URigVMPin*> AllPins = GetAllPinsRecursively();
-	TMap<FString, FRigVMGraphFunctionCategory> CategoryMap;
-	for(const URigVMPin* Pin : AllPins)
-	{
-		if(!Pin->UserDefinedCategory.IsEmpty())
-		{
-			FRigVMGraphFunctionCategory& Category = CategoryMap.FindOrAdd(Pin->UserDefinedCategory);
-			Category.Path = Pin->UserDefinedCategory;
-			Category.Elements.Add(Pin->GetSegmentPath(true));
-		}
-	}
-
-	// add the categories in the order they have been added
-	for(const FString& PinCategory : PinCategories)
-	{
-		if(FRigVMGraphFunctionCategory* Category = CategoryMap.Find(PinCategory))
-		{
-			Header.Layout.Categories.Add(*Category);
-		}
-	}
-
-	// fill in all user provided display names
-	for(const URigVMPin* Pin : AllPins)
-	{
-		if(!Pin->GetCategory().IsEmpty() && Pin->GetIndexInCategory() != INDEX_NONE)
-		{
-			Header.Layout.PinIndexInCategory.Add(Pin->GetSegmentPath(true), Pin->GetIndexInCategory());
-		}
-		if(!Pin->DisplayName.IsNone())
-		{
-			Header.Layout.DisplayNames.Add(Pin->GetSegmentPath(true), Pin->DisplayName.ToString());
-		}
-	}
+	Header.Layout = GetPinLayout();
 	
 	return Header;
 }
diff --git a/Engine/Plugins/Runtime/RigVM/Source/RigVMDeveloper/Private/RigVMModel/Nodes/RigVMUnitNode.cpp b/Engine/Plugins/Runtime/RigVM/Source/RigVMDeveloper/Private/RigVMModel/Nodes/RigVMUnitNode.cpp
--- a/Engine/Plugins/Runtime/RigVM/Source/RigVMDeveloper/Private/RigVMModel/Nodes/RigVMUnitNode.cpp
+++ b/Engine/Plugins/Runtime/RigVM/Source/RigVMDeveloper/Private/RigVMModel/Nodes/RigVMUnitNode.cpp
@@ -322,6 +322,15 @@ FName URigVMUnitNode::GetNextAggregateName(const FName& InLastAggregatePinName)
 #endif
 }
 
+FName URigVMUnitNode::GetDisplayNameForPin(const FString& InPinPath) const
+{
+	if(const UScriptStruct* ScriptStruct = GetScriptStruct())
+	{
+		return GetDisplayNameForStructMember(ScriptStruct, InPinPath);
+	}
+	return Super::GetDisplayNameForPin(InPinPath);
+}
+
 UScriptStruct* URigVMUnitNode::GetScriptStruct() const
 {
 	if(UScriptStruct* ResolvedStruct = Super::GetScriptStruct())
diff --git a/Engine/Plugins/Runtime/RigVM/Source/RigVMDeveloper/Private/RigVMModel/RigVMController.cpp b/Engine/Plugins/Runtime/RigVM/Source/RigVMDeveloper/Private/RigVMModel/RigVMController.cpp
--- a/Engine/Plugins/Runtime/RigVM/Source/RigVMDeveloper/Private/RigVMModel/RigVMController.cpp
+++ b/Engine/Plugins/Runtime/RigVM/Source/RigVMDeveloper/Private/RigVMModel/RigVMController.cpp
@@ -8263,18 +8263,29 @@ bool URigVMController::RemovePinCategory(const URigVMNode* InNode, const FString
 		GetActionStack()->BeginAction(Action);
 	}
 
+	int32 NumPinsAffected = 0;
 	for(URigVMPin* PinInCategory : PinsForCategory)
 	{
-		(void)SetPinCategory(PinInCategory, FString(), bSetupUndoRedo);
+		if(SetPinCategory(PinInCategory, FString(), bSetupUndoRedo))
+		{
+			NumPinsAffected++;
+		}
 	}
 
+	int32 NumCategoriesRemoved = 0;
 	for(const FString& CategoryToRemove : CategoriesToRemove)
 	{
-		const_cast<URigVMNode*>(InNode)->PinCategories.Remove(CategoryToRemove);
+		NumCategoriesRemoved += const_cast<URigVMNode*>(InNode)->PinCategories.Remove(CategoryToRemove);
 	}
-	Notify(ERigVMGraphNotifType::PinCategoriesChanged, const_cast<URigVMNode*>(InNode));
 
-	if (!bSuspendNotifications)
+	if(NumCategoriesRemoved > 0)
+	{
+		Notify(ERigVMGraphNotifType::PinCategoriesChanged, const_cast<URigVMNode*>(InNode));
+	}
+
+	const bool bHasMadeChange = (NumCategoriesRemoved + NumPinsAffected) > 0;
+	
+	if (!bSuspendNotifications && bHasMadeChange)
 	{
 		const URigVMGraph* Graph = GetGraph();
 		check(Graph);
@@ -8283,8 +8294,15 @@ bool URigVMController::RemovePinCategory(const URigVMNode* InNode, const FString
 
 	if (bSetupUndoRedo)
 	{
-		Action.UpdateAfterModification(InNode);
-		GetActionStack()->EndAction(Action);
+		if(bHasMadeChange)
+		{
+			Action.UpdateAfterModification(InNode);
+			GetActionStack()->EndAction(Action);
+		}
+		else
+		{
+			GetActionStack()->CancelAction(Action);
+		}
 	}
 
 	return true;
@@ -8712,6 +8730,198 @@ bool URigVMController::SetPinIndexInCategory(URigVMPin* InPin, int32 InIndexInCa
 	return true;
 }
 
+bool URigVMController::SetPinLayout(const FName& InNodeName, FRigVMNodeLayout InLayout, bool bSetupUndoRedo, bool bPrintPythonCommand)
+{
+	if (!IsValidGraph())
+	{
+		return false;
+	}
+
+	if (!bIsTransacting && !IsGraphEditable())
+	{
+		return false;
+	}
+
+	const URigVMGraph* Graph = GetGraph();
+	check(Graph);
+
+	const URigVMNode* Node = Graph->FindNodeByName(InNodeName);
+	if (Node == nullptr)
+	{
+		ReportErrorf(TEXT("Cannot find node '%s'."), *InNodeName.ToString());
+		return false;
+	}
+
+	return SetPinLayout(Node, InLayout, bSetupUndoRedo, bPrintPythonCommand);
+}
+
+bool URigVMController::SetPinLayout(const URigVMNode* InNode, FRigVMNodeLayout InLayout, bool bSetupUndoRedo, bool bPrintPythonCommand)
+{
+	if(!IsValidNodeForGraph(InNode))
+	{
+		return false;
+	}
+
+	int32 NumChanges = 0;
+	if(bSetupUndoRedo)
+	{
+		OpenUndoBracket(TEXT("Setting Pin Layout"));
+	}
+
+	// first remove any pin layout from the node
+	if(ClearPinLayout(InNode->GetFName(), bSetupUndoRedo, bPrintPythonCommand))
+	{
+		NumChanges++;
+	}
+
+	for(const TPair<FString,FString>& DisplayNamePair : InLayout.DisplayNames)
+	{
+		if(const URigVMPin* Pin = InNode->FindPin(DisplayNamePair.Key))
+		{
+			if(SetPinDisplayName(Pin->GetPinPath(), DisplayNamePair.Value, bSetupUndoRedo, bPrintPythonCommand))
+			{
+				NumChanges++;
+			}
+		}
+		else
+		{
+			// just a warning - not an error. the layout can still be applied
+			ReportWarningf(TEXT("Cannot find pin '%s' on node '%s' - ignoring display name provided by layout."), *DisplayNamePair.Key, *InNode->GetName());
+		}
+	}
+
+	for(const FRigVMPinCategory& Category : InLayout.Categories)
+	{
+		for(const FString& SegmentPath : Category.Elements)
+		{
+			if(const URigVMPin* Pin = InNode->FindPin(SegmentPath))
+			{
+				if(SetPinCategory(Pin->GetPinPath(), Category.Path, bSetupUndoRedo, bPrintPythonCommand))
+				{
+					NumChanges++;
+				}
+			}
+			else
+			{
+				// just a warning - not an error. the layout can still be applied
+				ReportWarningf(TEXT("Cannot find pin '%s' on node '%s' - ignoring pin category provided by layout."), *SegmentPath, *InNode->GetName());
+			}
+		}
+	}
+
+	for(const TPair<FString,int32>& PinIndexPair : InLayout.PinIndexInCategory)
+	{
+		if(const URigVMPin* Pin = InNode->FindPin(PinIndexPair.Key))
+		{
+			if(SetPinIndexInCategory(Pin->GetPinPath(), PinIndexPair.Value, bSetupUndoRedo, bPrintPythonCommand))
+			{
+				NumChanges++;
+			}
+		}
+		else
+		{
+			// just a warning - not an error. the layout can still be applied
+			ReportWarningf(TEXT("Cannot find pin '%s' on node '%s' - ignoring pin index provided by layout."), *PinIndexPair.Key, *InNode->GetName());
+		}
+	}
+
+	if(bSetupUndoRedo)
+	{
+		if(NumChanges > 0)
+		{
+			CloseUndoBracket();
+		}
+		else
+		{
+			CancelUndoBracket();
+		}
+	}
+
+	return NumChanges > 0;
+}
+
+bool URigVMController::ClearPinLayout(const FName& InNodeName, bool bSetupUndoRedo, bool bPrintPythonCommand)
+{
+	if (!IsValidGraph())
+	{
+		return false;
+	}
+
+	if (!bIsTransacting && !IsGraphEditable())
+	{
+		return false;
+	}
+
+	const URigVMGraph* Graph = GetGraph();
+	check(Graph);
+
+	const URigVMNode* Node = Graph->FindNodeByName(InNodeName);
+	if (Node == nullptr)
+	{
+		ReportErrorf(TEXT("Cannot find node '%s'."), *InNodeName.ToString());
+		return false;
+	}
+
+	return ClearPinLayout(Node, bSetupUndoRedo, bPrintPythonCommand);
+}
+
+bool URigVMController::ClearPinLayout(const URigVMNode* InNode, bool bSetupUndoRedo, bool bPrintPythonCommand)
+{
+	if(!IsValidNodeForGraph(InNode))
+	{
+		return false;
+	}
+
+	if(bSetupUndoRedo)
+	{
+		OpenUndoBracket(TEXT("Clearing Pin Layout"));
+	}
+
+	int32 NumChanges = 0;
+
+	TArray<FString> PinCategories = InNode->GetPinCategories();
+	// sort by length and reverse
+	PinCategories.Sort();
+	Algo::Reverse(PinCategories);
+
+	// remove all pin categories
+	for(const FString& PinCategory : PinCategories)
+	{
+		if(RemovePinCategory(InNode->GetFName(), PinCategory, bSetupUndoRedo, bPrintPythonCommand))
+		{
+			NumChanges++;
+		}
+	}
+
+	// reset all pin display names
+	const TArray<URigVMPin*> AllPins = InNode->GetAllPinsRecursively();
+	for(const URigVMPin* Pin : AllPins)
+	{
+		const FName DefaultDisplayName = InNode->GetDisplayNameForPin(Pin->GetSegmentPath(true));
+		if(!Pin->DisplayName.IsEqual(DefaultDisplayName, ENameCase::CaseSensitive))
+		{
+			if(SetPinDisplayName(Pin->GetPinPath(), FString(), bSetupUndoRedo, bPrintPythonCommand))
+			{
+				NumChanges++;
+			}
+		}
+	}
+
+	if(bSetupUndoRedo)
+	{
+		if(NumChanges > 0)
+		{
+			CloseUndoBracket();
+		}
+		else
+		{
+			CancelUndoBracket();
+		}
+	}
+
+	return NumChanges > 0;
+}
+
 bool URigVMController::SetPinCategories(const FName& InNodeName, const TArray<FString>& InCategories, bool bSetupUndoRedo)
 {
 	if (!IsValidGraph())
@@ -16129,18 +16339,6 @@ void URigVMController::ConfigurePinFromProperty(FProperty* InProperty, URigVMPin
 
 #if WITH_EDITOR
 
-	if (!InOutPin->IsArrayElement())
-	{
-		FString DisplayNameText = InProperty->GetDisplayNameText().ToString();
-		if (!DisplayNameText.IsEmpty())
-		{
-			InOutPin->DisplayName = *DisplayNameText;
-		}
-		else
-		{
-			InOutPin->DisplayName = NAME_None;
-		}
-	}
 	InOutPin->bIsConstant = InProperty->HasMetaData(TEXT("Constant"));
 	FString CustomWidgetName = InProperty->GetMetaData(TEXT("CustomWidget"));
 	InOutPin->CustomWidgetName = CustomWidgetName.IsEmpty() ? FName(NAME_None) : FName(*CustomWidgetName);
diff --git a/Engine/Plugins/Runtime/RigVM/Source/RigVMDeveloper/Private/RigVMModel/RigVMNode.cpp b/Engine/Plugins/Runtime/RigVM/Source/RigVMDeveloper/Private/RigVMModel/RigVMNode.cpp
--- a/Engine/Plugins/Runtime/RigVM/Source/RigVMDeveloper/Private/RigVMModel/RigVMNode.cpp
+++ b/Engine/Plugins/Runtime/RigVM/Source/RigVMDeveloper/Private/RigVMModel/RigVMNode.cpp
@@ -281,6 +281,54 @@ bool URigVMNode::IsPinCategoryExpanded(FString InCategory) const
 	return false;
 }
 
+FRigVMNodeLayout URigVMNode::GetPinLayout() const
+{
+	FRigVMNodeLayout Layout;
+	
+	// fill in the pin categories based on the data stored on the pins themselves
+	const TArray<URigVMPin*> AllPins = GetAllPinsRecursively();
+	TMap<FString, FRigVMPinCategory> CategoryMap;
+	for(const URigVMPin* Pin : AllPins)
+	{
+		if(!Pin->UserDefinedCategory.IsEmpty())
+		{
+			FRigVMPinCategory& Category = CategoryMap.FindOrAdd(Pin->UserDefinedCategory);
+			Category.Path = Pin->UserDefinedCategory;
+			Category.Elements.Add(Pin->GetSegmentPath(true));
+		}
+	}
+
+	// add the categories in the order they have been added
+	for(const FString& PinCategory : PinCategories)
+	{
+		if(const FRigVMPinCategory* Category = CategoryMap.Find(PinCategory))
+		{
+			Layout.Categories.Add(*Category);
+		}
+	}
+
+	// fill in all user provided display names and pin category indices
+	for(const URigVMPin* Pin : AllPins)
+	{
+		const FString SegmentPath = Pin->GetSegmentPath(true);
+		
+		if(!Pin->GetCategory().IsEmpty() && Pin->GetIndexInCategory() != INDEX_NONE)
+		{
+			Layout.PinIndexInCategory.Add(SegmentPath, Pin->GetIndexInCategory());
+		}
+
+		if(!Pin->DisplayName.IsNone())
+		{
+			if(!Pin->DisplayName.IsEqual(GetDisplayNameForPin(SegmentPath), ENameCase::CaseSensitive))
+			{
+				Layout.DisplayNames.Add(SegmentPath, Pin->DisplayName.ToString());
+			}
+		}
+	}
+
+	return Layout;
+}
+
 FString URigVMNode::GetOriginalPinDefaultValue(const URigVMPin* InPin) const
 {
 	const FString CompleteSegmentPath = InPin->GetSegmentPath(true);
@@ -817,9 +865,60 @@ UScriptStruct* URigVMNode::GetTraitScriptStruct(const URigVMPin* InTraitPin) con
 
 FName URigVMNode::GetDisplayNameForPin(const FString& InPinPath) const
 {
+	if(const URigVMPin* Pin = FindPin(InPinPath))
+	{
+		if(Pin->IsArrayElement())
+		{
+			return *FString::FromInt(Pin->GetPinIndex());
+		}
+	}
+	return NAME_None;
+}
+
+FName URigVMNode::GetDisplayNameForStructMember(const UStruct* InStruct, const FString& InPath)
+{
+	check(InStruct);
+	if(!InPath.IsEmpty())
+	{
+		FString Left, Right;
+		if(!RigVMStringUtils::SplitPinPathAtStart(InPath, Left, Right))
+		{
+			Left = InPath;
+		}
+
+		if(const FProperty* Property = InStruct->FindPropertyByName(*Left))
+		{
+			return GetDisplayNameForProperty(Property, Right);
+		}
+	}
 	return NAME_None;
 }
 
+FName URigVMNode::GetDisplayNameForProperty(const FProperty* InProperty, const FString& InRemainingPath)
+{
+	check(InProperty);
+
+	FText DisplayNameText = InProperty->GetDisplayNameText();
+
+	if(!InRemainingPath.IsEmpty())
+	{
+		const FRigVMPropertyPath PropertyPath(InProperty, InRemainingPath);
+		if(PropertyPath.IsValid())
+		{
+			if(const FProperty* TailProperty = PropertyPath.GetTailProperty())
+			{
+				DisplayNameText = TailProperty->GetDisplayNameText();
+			}
+		}
+	}
+
+	if(DisplayNameText.IsEmpty())
+	{
+		return NAME_None;
+	}
+	return *DisplayNameText.ToString();
+}
+
 FString URigVMNode::GetCategoryForPin(const FString& InPinPath) const
 {
 	return FString();
diff --git a/Engine/Plugins/Runtime/RigVM/Source/RigVMDeveloper/Public/RigVMModel/Nodes/RigVMUnitNode.h b/Engine/Plugins/Runtime/RigVM/Source/RigVMDeveloper/Public/RigVMModel/Nodes/RigVMUnitNode.h
--- a/Engine/Plugins/Runtime/RigVM/Source/RigVMDeveloper/Public/RigVMModel/Nodes/RigVMUnitNode.h
+++ b/Engine/Plugins/Runtime/RigVM/Source/RigVMDeveloper/Public/RigVMModel/Nodes/RigVMUnitNode.h
@@ -37,6 +37,7 @@ public:
 	virtual TArray<URigVMPin*> GetAggregateInputs() const override;
 	virtual TArray<URigVMPin*> GetAggregateOutputs() const override;
 	virtual FName GetNextAggregateName(const FName& InLastAggregatePinName) const override;
+	virtual FName GetDisplayNameForPin(const FString& InPinPath) const override;
 
 	virtual bool IsOutDated() const override;
 	virtual FString GetDeprecatedMetadata() const override;
diff --git a/Engine/Plugins/Runtime/RigVM/Source/RigVMDeveloper/Public/RigVMModel/RigVMController.h b/Engine/Plugins/Runtime/RigVM/Source/RigVMDeveloper/Public/RigVMModel/RigVMController.h
--- a/Engine/Plugins/Runtime/RigVM/Source/RigVMDeveloper/Public/RigVMModel/RigVMController.h
+++ b/Engine/Plugins/Runtime/RigVM/Source/RigVMDeveloper/Public/RigVMModel/RigVMController.h
@@ -815,24 +815,37 @@ public:
 
 	// Removes a pin category. The category is UI relevant only and used
 	// to order pins in the user interface of the node as well as on the details panel.
+	UFUNCTION(BlueprintCallable, Category = RigVMController)
 	bool RemovePinCategory(const FName& InNodeName, const FString& InPinCategory, bool bSetupUndoRedo = true, bool bPrintPythonCommand = false);
 
 	// Renames a pin category. The category is UI relevant only and used
 	// to order pins in the user interface of the node as well as on the details panel.
+	UFUNCTION(BlueprintCallable, Category = RigVMController)
 	bool RenamePinCategory(const FName& InNodeName, const FString& InOldPinCategory, const FString& InNewPinCategory, bool bSetupUndoRedo = true, bool bPrintPythonCommand = false);
 
 	// Changes a pin category's index. The category is UI relevant only and used
 	// to order pins in the user interface of the node as well as on the details panel.
+	UFUNCTION(BlueprintCallable, Category = RigVMController)
 	bool SetPinCategoryIndex(const FName& InNodeName, const FString& InPinCategory, int32 InNewIndex, bool bSetupUndoRedo = true, bool bPrintPythonCommand = false);
 
 	// Changes a pin category's expansion state. The category is UI relevant only and used
 	// to order pins in the user interface of the node as well as on the details panel.
+	UFUNCTION(BlueprintCallable, Category = RigVMController)
 	bool SetPinCategoryExpansion(const FName& InNodeName, const FString& InPinCategory, bool bIsExpanded, bool bSetupUndoRedo = true, bool bPrintPythonCommand = false);
 
 	// Changes a pin category's expansion state. The category is UI relevant only and used
 	// to order pins in the user interface of the node as well as on the details panel.
+	UFUNCTION(BlueprintCallable, Category = RigVMController)
 	bool SetPinIndexInCategory(const FString& InPinPath, int32 InIndexInCategory, bool bSetupUndoRedo = true, bool bPrintPythonCommand = false);
 
+	// Applies a complete pin layout to a node
+	UFUNCTION(BlueprintCallable, Category = RigVMController)
+	bool SetPinLayout(const FName& InNodeName, FRigVMNodeLayout InLayout, bool bSetupUndoRedo = true, bool bPrintPythonCommand = false);
+
+	// Removes any layout information from a node
+	UFUNCTION(BlueprintCallable, Category = RigVMController)
+	bool ClearPinLayout(const FName& InNodeName, bool bSetupUndoRedo = true, bool bPrintPythonCommand = false);
+
 	// Returns the default value of a pin given its pinpath.
 	UFUNCTION(BlueprintCallable, Category = RigVMController)
 	FString GetPinDefaultValue(const FString& InPinPath);
@@ -1267,6 +1280,8 @@ private:
 	bool SetPinCategoryIndex(const URigVMNode* InNode, const FString& InPinCategory, int32 InNewIndex, bool bSetupUndoRedo);
 	bool SetPinCategoryExpansion(const URigVMNode* InNode, const FString& InPinCategory, bool bIsExpanded, bool bSetupUndoRedo);
 	bool SetPinIndexInCategory(URigVMPin* InPin, int32 InIndexInCategory, bool bSetupUndoRedo);
+	bool SetPinLayout(const URigVMNode* InNode, FRigVMNodeLayout InLayout, bool bSetupUndoRedo, bool bPrintPythonCommand);
+	bool ClearPinLayout(const URigVMNode* InNode, bool bSetupUndoRedo, bool bPrintPythonCommand);
 	bool SetPinCategories(const FName& InNodeName, const TArray<FString>& InCategories, bool bSetupUndoRedo);
 	bool SetPinCategories(const URigVMNode* InNode, const TArray<FString>& InCategories, bool bSetupUndoRedo);
 	bool SetVariableName(URigVMVariableNode* InVariableNode, const FName& InVariableName, bool bSetupUndoRedo);
diff --git a/Engine/Plugins/Runtime/RigVM/Source/RigVMDeveloper/Public/RigVMModel/RigVMNode.h b/Engine/Plugins/Runtime/RigVM/Source/RigVMDeveloper/Public/RigVMModel/RigVMNode.h
--- a/Engine/Plugins/Runtime/RigVM/Source/RigVMDeveloper/Public/RigVMModel/RigVMNode.h
+++ b/Engine/Plugins/Runtime/RigVM/Source/RigVMDeveloper/Public/RigVMModel/RigVMNode.h
@@ -5,9 +5,8 @@
 #include "RigVMPin.h"
 #include "RigVMCore/RigVM.h"
 #include "RigVMCore/RigVMStruct.h"
-#include "RigVMCore/RigVMTrait.h"
 #include "RigVMCore/RigVMUserWorkflow.h"
-#include "RigVMCore/RigVMExecuteContext.h"
+#include "RigVMCore/RigVMGraphFunctionDefinition.h"
 #include "UObject/StructOnScope.h"
 #include "RigVMNode.generated.h"
 
@@ -101,6 +100,10 @@ public:
 
 	const FString& GetLastAffectedPinCategory() const { return LastAffectedPinCategory; }
 
+	// Returns the pin UI layout for this node
+	UFUNCTION(BlueprintCallable, Category = RigVMNode)
+	FRigVMNodeLayout GetPinLayout() const;
+
 	// Returns the default value for a given pin
 	FString GetOriginalPinDefaultValue(const URigVMPin* InPin) const;
 
@@ -372,6 +375,12 @@ public:
 	// returns the display name for a pin
 	virtual FName GetDisplayNameForPin(const FString& InPinPath) const;
 
+	// returns the display name for a struct member
+	static FName GetDisplayNameForStructMember(const UStruct* InStruct, const FString& InPath);
+
+	// returns the display name within a property
+	static FName GetDisplayNameForProperty(const FProperty* InProperty, const FString& InRemainingPath);
+
 	// returns the category for a pin
 	virtual FString GetCategoryForPin(const FString& InPinPath) const;
 
diff --git a/Engine/Plugins/Runtime/RigVM/Source/RigVMEditor/Private/Widgets/SRigVMGraphPinCategory.cpp b/Engine/Plugins/Runtime/RigVM/Source/RigVMEditor/Private/Widgets/SRigVMGraphPinCategory.cpp
--- a/Engine/Plugins/Runtime/RigVM/Source/RigVMEditor/Private/Widgets/SRigVMGraphPinCategory.cpp
+++ b/Engine/Plugins/Runtime/RigVM/Source/RigVMEditor/Private/Widgets/SRigVMGraphPinCategory.cpp
@@ -2,6 +2,8 @@
 
 
 #include "Widgets/SRigVMGraphPinCategory.h"
+#include "EdGraph/RigVMEdGraphNode.h"
+#include "RigVMModel/RigVMController.h"
 #include "Widgets/SNullWidget.h"
 
 void SRigVMGraphPinCategory::Construct(const FArguments& InArgs, UEdGraphPin* InGraphPinObj)
@@ -9,6 +11,31 @@ void SRigVMGraphPinCategory::Construct(const FArguments& InArgs, UEdGraphPin* In
 	SGraphPin::Construct(SGraphPin::FArguments(), InGraphPinObj);
 }
 
+FReply SRigVMGraphPinCategory::OnMouseButtonDown(const FGeometry& MyGeometry, const FPointerEvent& MouseEvent)
+{
+	if(const UEdGraphPin* EdGraphPin = GraphPinObj)
+	{
+		if(const URigVMEdGraphNode* EdGraphNode = Cast<URigVMEdGraphNode>(EdGraphPin->GetOwningNode()))
+		{
+			if(URigVMController* Controller = EdGraphNode->GetController())
+			{
+				if(const URigVMGraph* Model = Controller->GetGraph())
+				{
+					if(const URigVMNode* Node = Model->FindNodeByName(EdGraphNode->GetFName()))
+					{
+						const bool bIsExpanded = Node->IsPinCategoryExpanded(EdGraphPin->GetName());
+						if(Controller->SetPinCategoryExpansion(Node->GetFName(), EdGraphPin->GetName(), !bIsExpanded))
+						{
+							return FReply::Handled();
+						}
+					}
+				}
+			}
+		}
+	}
+	return SGraphPin::OnMouseButtonDown(MyGeometry, MouseEvent);
+}
+
 TSharedRef<SWidget>	SRigVMGraphPinCategory::GetDefaultValueWidget()
 {
 	return SNullWidget::NullWidget;
diff --git a/Engine/Plugins/Runtime/RigVM/Source/RigVMEditor/Public/Widgets/SRigVMGraphPinCategory.h b/Engine/Plugins/Runtime/RigVM/Source/RigVMEditor/Public/Widgets/SRigVMGraphPinCategory.h
--- a/Engine/Plugins/Runtime/RigVM/Source/RigVMEditor/Public/Widgets/SRigVMGraphPinCategory.h
+++ b/Engine/Plugins/Runtime/RigVM/Source/RigVMEditor/Public/Widgets/SRigVMGraphPinCategory.h
@@ -16,6 +16,7 @@ public:
 
 
 	void Construct(const FArguments& InArgs, UEdGraphPin* InGraphPinObj);
+	virtual FReply OnMouseButtonDown(const FGeometry& MyGeometry, const FPointerEvent& MouseEvent) override;
 
 protected:
 
