RigVM: Correctly deal with arrays being emptied in details panel

#rnx

diff --git a/Engine/Plugins/Runtime/RigVM/Source/RigVMEditor/Private/Editor/RigVMEditor.cpp b/Engine/Plugins/Runtime/RigVM/Source/RigVMEditor/Private/Editor/RigVMEditor.cpp
--- a/Engine/Plugins/Runtime/RigVM/Source/RigVMEditor/Private/Editor/RigVMEditor.cpp
+++ b/Engine/Plugins/Runtime/RigVM/Source/RigVMEditor/Private/Editor/RigVMEditor.cpp
@@ -2232,7 +2232,27 @@ void FRigVMEditor::OnWrappedPropertyChangedChainEvent(URigVMDetailsViewWrapperOb
 
 		if (Property)
 		{
-			FString DefaultValue = FRigVMStruct::ExportToFullyQualifiedText(Property, PropertyStorage);
+			FString DefaultValue;
+			if(PropertyStorage == nullptr)
+			{
+				// this may happen when we remove the last element from an array.
+				// in that case just empty the array itself.
+				if(const FProperty* ParentProperty = Property->GetOwnerProperty())
+				{
+					if(ParentProperty->IsA<FArrayProperty>())
+					{
+						DefaultValue = TEXT("()");
+						FString Left, Right;
+						verify(URigVMPin::SplitPinPathAtEnd(PinPath, Left, Right));
+						PinPath = Left;
+					}
+				}
+			}
+			else
+			{
+				DefaultValue = FRigVMStruct::ExportToFullyQualifiedText(Property, PropertyStorage);
+			}
+			
 			if(Property->IsA<FStrProperty>() || Property->IsA<FNameProperty>())
 			{
 				DefaultValue.TrimCharInline(TEXT('\"'), nullptr);
