Fixed RigVM ccrash in AnimNext grapsh when expanding pins, due to missing Blueprint

diff --git a/Engine/Plugins/Runtime/RigVM/Source/RigVMDeveloper/Private/EdGraph/RigVMEdGraphSchema.cpp b/Engine/Plugins/Runtime/RigVM/Source/RigVMDeveloper/Private/EdGraph/RigVMEdGraphSchema.cpp
--- a/Engine/Plugins/Runtime/RigVM/Source/RigVMDeveloper/Private/EdGraph/RigVMEdGraphSchema.cpp
+++ b/Engine/Plugins/Runtime/RigVM/Source/RigVMDeveloper/Private/EdGraph/RigVMEdGraphSchema.cpp
@@ -2159,13 +2159,8 @@ void URigVMEdGraphSchema::EndGraphNodeInteraction(UEdGraphNode* InNode) const
 	
 	bool bMovedSomething = false;
 
-	URigVMBlueprint* Blueprint = Graph->GetBlueprint();
-	if (Blueprint == nullptr)
-	{
-		return;
-	}
-	TGuardValue<bool> GuardDirtyBlueprintStatus(Blueprint->bSkipDirtyBlueprintStatus, true);
-	
+	FGuardSkipDirtyBlueprintStatus GuardDirtyBlueprintStatus(Graph->GetBlueprint(), true);
+
 	Graph->GetController()->OpenUndoBracket(TEXT("Move Nodes"));
 
 	for (const UEdGraphNode* NodeToMove : NodesToMove)
diff --git a/Engine/Plugins/Runtime/RigVM/Source/RigVMDeveloper/Public/RigVMBlueprint.h b/Engine/Plugins/Runtime/RigVM/Source/RigVMDeveloper/Public/RigVMBlueprint.h
--- a/Engine/Plugins/Runtime/RigVM/Source/RigVMDeveloper/Public/RigVMBlueprint.h
+++ b/Engine/Plugins/Runtime/RigVM/Source/RigVMDeveloper/Public/RigVMBlueprint.h
@@ -27,6 +27,7 @@ class IRigVMEditorModule;
 #endif
 struct FEndLoadPackageContext;
 struct FRigVMMemoryStorageStruct;
+struct FGuardSkipDirtyBlueprintStatus;
 
 DECLARE_EVENT_ThreeParams(URigVMBlueprint, FOnRigVMCompiledEvent, UObject*, URigVM*, FRigVMExtendedExecuteContext&);
 DECLARE_EVENT_OneParam(URigVMBlueprint, FOnRigVMRefreshEditorEvent, URigVMBlueprint*);
@@ -836,6 +837,7 @@ private:
 	friend class FRigVMTreeAssetVariantFilter;
 	friend class FRigVMTreePackageNode;
 	friend class SRigVMGraphNode;
+	friend struct FGuardSkipDirtyBlueprintStatus;
 };
 
 class RIGVMDEVELOPER_API FRigVMBlueprintCompileScope
@@ -857,4 +859,28 @@ public:
 private:
 
 	URigVMBlueprint* Blueprint;
-};
\ No newline at end of file
+};
+
+struct FGuardSkipDirtyBlueprintStatus : private FNoncopyable
+{
+	[[nodiscard]] FGuardSkipDirtyBlueprintStatus(TWeakObjectPtr<URigVMBlueprint> InBlueprint, bool bNewValue)
+	{
+		if (InBlueprint.IsValid())
+		{
+			WeakBlueprint = InBlueprint;
+			bOldValue = InBlueprint->bSkipDirtyBlueprintStatus;
+			InBlueprint->bSkipDirtyBlueprintStatus = bNewValue;
+		}
+	}
+	~FGuardSkipDirtyBlueprintStatus()
+	{
+		if (URigVMBlueprint* Blueprint = WeakBlueprint.Get())
+		{
+			Blueprint->bSkipDirtyBlueprintStatus = bOldValue;
+		}
+	}
+
+private:
+	TWeakObjectPtr<URigVMBlueprint> WeakBlueprint;
+	bool bOldValue = false;
+};
diff --git a/Engine/Plugins/Runtime/RigVM/Source/RigVMEditor/Private/Widgets/SRigVMGraphNode.cpp b/Engine/Plugins/Runtime/RigVM/Source/RigVMEditor/Private/Widgets/SRigVMGraphNode.cpp
--- a/Engine/Plugins/Runtime/RigVM/Source/RigVMEditor/Private/Widgets/SRigVMGraphNode.cpp
+++ b/Engine/Plugins/Runtime/RigVM/Source/RigVMEditor/Private/Widgets/SRigVMGraphNode.cpp
@@ -1471,7 +1471,7 @@ FReply SRigVMGraphNode::OnExpanderArrowClicked(int32 InPinInfoIndex)
 				Algo::Reverse(PinInfoIndicesToModify);
 			}
 
-			TGuardValue<bool> GuardDirtyBlueprintStatus(Blueprint->bSkipDirtyBlueprintStatus, true);
+			FGuardSkipDirtyBlueprintStatus GuardDirtyBlueprintStatus(Blueprint, true);
 
 			Controller->OpenUndoBracket(PinInfo.bExpanded ? TEXT("Collapsing Pin") : TEXT("Expanding Pin"));
 			for(const int32& PinInfoIndexToModify : PinInfoIndicesToModify)
