RigVM: Use pretty names for dispatch based nodes

#rnx

diff --git a/Engine/Plugins/Runtime/RigVM/Source/RigVM/Private/RigVMCore/RigVMTemplate.cpp b/Engine/Plugins/Runtime/RigVM/Source/RigVM/Private/RigVMCore/RigVMTemplate.cpp
--- a/Engine/Plugins/Runtime/RigVM/Source/RigVM/Private/RigVMCore/RigVMTemplate.cpp
+++ b/Engine/Plugins/Runtime/RigVM/Source/RigVM/Private/RigVMCore/RigVMTemplate.cpp
@@ -5,6 +5,7 @@
 #include "RigVMCore/RigVMStruct.h"
 #include "RigVMCore/RigVMDispatchFactory.h"
 #include "RigVMModule.h"
+#include "RigVMStringUtils.h"
 #include "Algo/Accumulate.h"
 #include "Algo/ForEach.h"
 #include "Algo/Sort.h"
@@ -1124,6 +1125,28 @@ FName FRigVMTemplate::GetName() const
 	return NAME_None;
 }
 
+FName FRigVMTemplate::GetNodeName() const
+{
+#if WITH_EDITOR
+	if(UsesDispatch())
+	{
+		if(const FRigVMDispatchFactory* Factory = GetDispatchFactory())
+		{
+			if(const UScriptStruct* FactoryStruct = Factory->GetScriptStruct())
+			{
+				FString DisplayName = FactoryStruct->GetDisplayNameText().ToString();
+				RigVMStringUtils::SanitizeName(DisplayName, false, false, 100);
+				if(!DisplayName.IsEmpty())
+				{
+					return *DisplayName;
+				}
+			}
+		}
+	}
+#endif
+	return GetName();
+}
+
 #if WITH_EDITOR
 
 FLinearColor FRigVMTemplate::GetColor(const TArray<int32>& InPermutationIndices) const
diff --git a/Engine/Plugins/Runtime/RigVM/Source/RigVM/Private/RigVMStringUtils.cpp b/Engine/Plugins/Runtime/RigVM/Source/RigVM/Private/RigVMStringUtils.cpp
--- a/Engine/Plugins/Runtime/RigVM/Source/RigVM/Private/RigVMStringUtils.cpp
+++ b/Engine/Plugins/Runtime/RigVM/Source/RigVM/Private/RigVMStringUtils.cpp
@@ -187,3 +187,30 @@ TArray<FString> RigVMStringUtils::SplitDefaultValue(const FString& InDefaultValu
 	}
 	return Parts;
 }
+
+// Sanitizes a name as per ruleset
+void RigVMStringUtils::SanitizeName(FString& InOutName, bool bAllowPeriod, bool bAllowSpace, int32 InMaxNameLength)
+{
+	// Sanitize the name
+	for (int32 i = 0; i < InOutName.Len(); ++i)
+	{
+		TCHAR& C = InOutName[i];
+
+		const bool bGoodChar =
+			FChar::IsAlpha(C) ||											// Any letter (upper and lowercase) anytime
+			(C == '_') || (C == '-') || 									// _  and - anytime
+			(bAllowPeriod && (C == '.')) ||
+			(bAllowSpace && (C == ' ')) ||
+			((i > 0) && FChar::IsDigit(C));									// 0-9 after the first character
+
+		if (!bGoodChar)
+		{
+			C = '_';
+		}
+	}
+
+	if (InOutName.Len() > InMaxNameLength)
+	{
+		InOutName.LeftChopInline(InOutName.Len() - InMaxNameLength);
+	}
+}
diff --git a/Engine/Plugins/Runtime/RigVM/Source/RigVM/Public/RigVMCore/RigVMTemplate.h b/Engine/Plugins/Runtime/RigVM/Source/RigVM/Public/RigVMCore/RigVMTemplate.h
--- a/Engine/Plugins/Runtime/RigVM/Source/RigVM/Public/RigVMCore/RigVMTemplate.h
+++ b/Engine/Plugins/Runtime/RigVM/Source/RigVM/Public/RigVMCore/RigVMTemplate.h
@@ -406,6 +406,9 @@ public:
 	// Returns the name of the template
 	FName GetName() const;
 
+	// Returns the name to use for a node
+	FName GetNodeName() const;
+
 	// returns true if this template can merge another one
 	bool Merge(const FRigVMTemplate& InOther);
 
diff --git a/Engine/Plugins/Runtime/RigVM/Source/RigVM/Public/RigVMStringUtils.h b/Engine/Plugins/Runtime/RigVM/Source/RigVM/Public/RigVMStringUtils.h
--- a/Engine/Plugins/Runtime/RigVM/Source/RigVM/Public/RigVMStringUtils.h
+++ b/Engine/Plugins/Runtime/RigVM/Source/RigVM/Public/RigVMStringUtils.h
@@ -41,4 +41,7 @@ namespace RigVMStringUtils
 
 	// Splits the default value into name-value pairs
 	RIGVM_API TArray<FString> SplitDefaultValue(const FString& InDefaultValue);
+
+	// Sanitizes a name as per ruleset
+	RIGVM_API void SanitizeName(FString& InOutName, bool bAllowPeriod, bool bAllowSpace, int32 InMaxNameLength);
 }
\ No newline at end of file
diff --git a/Engine/Plugins/Runtime/RigVM/Source/RigVMDeveloper/Private/RigVMModel/Nodes/RigVMTemplateNode.cpp b/Engine/Plugins/Runtime/RigVM/Source/RigVMDeveloper/Private/RigVMModel/Nodes/RigVMTemplateNode.cpp
--- a/Engine/Plugins/Runtime/RigVM/Source/RigVMDeveloper/Private/RigVMModel/Nodes/RigVMTemplateNode.cpp
+++ b/Engine/Plugins/Runtime/RigVM/Source/RigVMDeveloper/Private/RigVMModel/Nodes/RigVMTemplateNode.cpp
@@ -150,7 +150,7 @@ FString URigVMTemplateNode::GetNodeTitle() const
 	{
 		if(const FRigVMTemplate* Template = GetTemplate())
 		{
-			return Template->GetName().ToString();
+			return Template->GetNodeName().ToString();
 		}
 	}
 	
diff --git a/Engine/Plugins/Runtime/RigVM/Source/RigVMDeveloper/Private/RigVMModel/RigVMController.cpp b/Engine/Plugins/Runtime/RigVM/Source/RigVMDeveloper/Private/RigVMModel/RigVMController.cpp
--- a/Engine/Plugins/Runtime/RigVM/Source/RigVMDeveloper/Private/RigVMModel/RigVMController.cpp
+++ b/Engine/Plugins/Runtime/RigVM/Source/RigVMDeveloper/Private/RigVMModel/RigVMController.cpp
@@ -1461,7 +1461,7 @@ URigVMUnitNode* URigVMController::AddUnitNode(UScriptStruct* InScriptStruct, TSu
 	{
 		if(bSetupUndoRedo)
 		{
-			OpenUndoBracket(FString::Printf(TEXT("Add %s Node"), *Template->GetName().ToString()));
+			OpenUndoBracket(FString::Printf(TEXT("Add %s Node"), *Template->GetNodeName().ToString()));
 		}
 
 		const FString Name = GetSchema()->GetValidNodeName(Graph, InNodeName.IsEmpty() ? InScriptStruct->GetName() : InNodeName);
@@ -15088,7 +15088,7 @@ URigVMTemplateNode* URigVMController::AddTemplateNode(const FName& InNotation, c
 		}
 	}
 
-	FString Name = GetSchema()->GetValidNodeName(Graph, InNodeName.IsEmpty() ? Template->GetName().ToString() : InNodeName);
+	FString Name = GetSchema()->GetValidNodeName(Graph, InNodeName.IsEmpty() ? Template->GetNodeName().ToString() : InNodeName);
 	URigVMTemplateNode* Node = nullptr;
 
 	// determine what kind of node we need to create
diff --git a/Engine/Plugins/Runtime/RigVM/Source/RigVMDeveloper/Private/RigVMModel/RigVMSchema.cpp b/Engine/Plugins/Runtime/RigVM/Source/RigVMDeveloper/Private/RigVMModel/RigVMSchema.cpp
--- a/Engine/Plugins/Runtime/RigVM/Source/RigVMDeveloper/Private/RigVMModel/RigVMSchema.cpp
+++ b/Engine/Plugins/Runtime/RigVM/Source/RigVMDeveloper/Private/RigVMModel/RigVMSchema.cpp
@@ -1,6 +1,7 @@
 // Copyright Epic Games, Inc. All Rights Reserved.
 
 #include "RigVMModel/RigVMSchema.h"
+#include "RigVMStringUtils.h"
 #include "RigVMModel/RigVMClient.h"
 #include "RigVMModel/RigVMController.h"
 
@@ -761,28 +762,7 @@ FString URigVMSchema::GetValidNodeName(const URigVMGraph* InGraph, const FString
 
 void URigVMSchema::SanitizeName(FString& InOutName, bool bAllowPeriod, bool bAllowSpace)
 {
-	// Sanitize the name
-	for (int32 i = 0; i < InOutName.Len(); ++i)
-	{
-		TCHAR& C = InOutName[i];
-
-		const bool bGoodChar =
-			FChar::IsAlpha(C) ||											// Any letter (upper and lowercase) anytime
-			(C == '_') || (C == '-') || 									// _  and - anytime
-			(bAllowPeriod && (C == '.')) ||
-			(bAllowSpace && (C == ' ')) ||
-			((i > 0) && FChar::IsDigit(C));									// 0-9 after the first character
-
-		if (!bGoodChar)
-		{
-			C = '_';
-		}
-	}
-
-	if (InOutName.Len() > GetMaxNameLength())
-	{
-		InOutName.LeftChopInline(InOutName.Len() - GetMaxNameLength());
-	}
+	RigVMStringUtils::SanitizeName(InOutName, bAllowPeriod, bAllowSpace, GetMaxNameLength());
 }
 
 FName URigVMSchema::GetUniqueName(const FName& InName, TFunction<bool(const FName&)> IsNameAvailablePredicate,
diff --git a/Engine/Plugins/Runtime/RigVM/Source/RigVMEditor/Private/EdGraph/NodeSpawners/RigVMEdGraphTemplateNodeSpawner.cpp b/Engine/Plugins/Runtime/RigVM/Source/RigVMEditor/Private/EdGraph/NodeSpawners/RigVMEdGraphTemplateNodeSpawner.cpp
--- a/Engine/Plugins/Runtime/RigVM/Source/RigVMEditor/Private/EdGraph/NodeSpawners/RigVMEdGraphTemplateNodeSpawner.cpp
+++ b/Engine/Plugins/Runtime/RigVM/Source/RigVMEditor/Private/EdGraph/NodeSpawners/RigVMEdGraphTemplateNodeSpawner.cpp
@@ -219,7 +219,7 @@ URigVMEdGraphNode* URigVMEdGraphTemplateNodeSpawner::SpawnNode(UEdGraph* ParentG
 
 	if (RigBlueprint != nullptr && RigGraph != nullptr)
 	{
-		const FName Name = FRigVMBlueprintUtils::ValidateName(RigBlueprint, Template->GetName().ToString());
+		const FName Name = FRigVMBlueprintUtils::ValidateName(RigBlueprint, Template->GetNodeName().ToString());
 		URigVMController* Controller = RigBlueprint->GetController(ParentGraph);
 
 		Controller->OpenUndoBracket(FString::Printf(TEXT("Add '%s' Node"), *Name.ToString()));
