RigVM: Ensure to only add an event once

#rnx

diff --git a/Engine/Plugins/Runtime/RigVM/Source/RigVMDeveloper/Private/RigVMModel/RigVMController.cpp b/Engine/Plugins/Runtime/RigVM/Source/RigVMDeveloper/Private/RigVMModel/RigVMController.cpp
--- a/Engine/Plugins/Runtime/RigVM/Source/RigVMDeveloper/Private/RigVMModel/RigVMController.cpp
+++ b/Engine/Plugins/Runtime/RigVM/Source/RigVMDeveloper/Private/RigVMModel/RigVMController.cpp
@@ -268,17 +268,7 @@ URigVMController* URigVMController::GetControllerForGraph(const URigVMGraph* InG
 		return const_cast<URigVMController*>(this);
 	}
 
-	// get the client host from the graph rather than this controller,
-	// since the graph may be in a different host / rigvmblueprint.
-	IRigVMClientHost* ClientHost = InGraph->GetImplementingOuter<IRigVMClientHost>();
-
-	// the graph may not be nested under a client in some unit tests
-	if(ClientHost == nullptr)
-	{
-		ClientHost = GetImplementingOuter<IRigVMClientHost>();
-	}
-	
-	if(ClientHost)
+	if(IRigVMClientHost* ClientHost = GetClientHost())
 	{
 		if(FRigVMClient* Client = ClientHost->GetRigVMClient())
 		{
@@ -301,6 +291,60 @@ URigVMController* URigVMController::GetControllerForGraph(const URigVMGraph* InG
 	return nullptr;
 }
 
+IRigVMClientHost* URigVMController::GetClientHost() const
+{
+	return GetClientHost_Internal(GetGraph());
+}
+
+IRigVMClientHost* URigVMController::GetClientHost_Internal(const URigVMGraph* InGraph) const
+{
+	// get the client host from the graph rather than this controller,
+	// since the graph may be in a different host / rigvmblueprint.
+	IRigVMClientHost* ClientHost = nullptr;
+
+	if(InGraph)
+	{
+		ClientHost = InGraph->GetImplementingOuter<IRigVMClientHost>();
+	}
+
+	// the graph may not be nested under a client in some unit tests
+	if(ClientHost == nullptr)
+	{
+		ClientHost = GetImplementingOuter<IRigVMClientHost>();
+	}
+
+	return ClientHost;
+}
+
+TArray<FName> URigVMController::GetAllEventNames() const
+{
+	TArray<FName> AllEventNames;
+	if(!IsValidGraph())
+	{
+		return AllEventNames;
+	}
+
+	TArray<const URigVMGraph*> GraphsToCheck = {GetGraph()};
+	if(const IRigVMClientHost* ClientHost = GetClientHost())
+	{
+		const TArray<URigVMGraph*> ClientModels = ClientHost->GetAllModels();
+		for(const URigVMGraph* ClientModel : ClientModels)
+		{
+			GraphsToCheck.AddUnique(ClientModel);
+		}
+	}
+	for(const URigVMGraph* GraphToCheck : GraphsToCheck)
+	{
+		const TArray<FName> EventNames = GraphToCheck->GetEventNames();
+		for(const FName& EventName : EventNames)
+		{
+			AllEventNames.AddUnique(EventName);
+		}
+	}
+
+	return AllEventNames;
+}
+
 FRigVMGraphModifiedEvent& URigVMController::OnModified()
 {
 	return ModifiedEventStatic;
@@ -1520,7 +1564,9 @@ URigVMUnitNode* URigVMController::AddUnitNode(UScriptStruct* InScriptStruct, TSu
 			return nullptr;
 		}
 
-		if (Graph->GetEventNames().Contains(StructMemory->GetEventName()))
+		// make sure to check all graphs when adding an event - we cannot
+		// have the same event twice under two graphs within the client.
+		if(GetAllEventNames().Contains(StructMemory->GetEventName()))
 		{
 			ReportAndNotifyErrorf(TEXT("Event %s already exists in the graph."), *StructMemory->GetEventName().ToString());
 			return nullptr;
@@ -4392,6 +4438,10 @@ TArray<FName> URigVMController::ImportNodesFromText(const FString& InText, bool
 						ControllerForGraph->AddLink(SourcePin, TargetPin, false);
 					}
 				}
+				else
+				{
+					Graph->DetachedLinks.AddUnique(CreatedLink);
+				}
 
 				// if the link is still part of the detached link array
 				// the restore was not successful.
diff --git a/Engine/Plugins/Runtime/RigVM/Source/RigVMDeveloper/Private/RigVMModel/RigVMSchema.cpp b/Engine/Plugins/Runtime/RigVM/Source/RigVMDeveloper/Private/RigVMModel/RigVMSchema.cpp
--- a/Engine/Plugins/Runtime/RigVM/Source/RigVMDeveloper/Private/RigVMModel/RigVMSchema.cpp
+++ b/Engine/Plugins/Runtime/RigVM/Source/RigVMDeveloper/Private/RigVMModel/RigVMSchema.cpp
@@ -385,7 +385,7 @@ bool URigVMSchema::CanAddNode(URigVMController* InController, const URigVMNode*
 					return false;
 				}
 
-				if (Graph->GetEventNames().Contains(InUnitNode->GetEventName()))
+				if (InController->GetAllEventNames().Contains(InUnitNode->GetEventName()))
 				{
 					static const FString Message = FString::Printf(TEXT("An event named %s already exists."), *InUnitNode->GetEventName().ToString()); 
 					InController->ReportError(Message);
diff --git a/Engine/Plugins/Runtime/RigVM/Source/RigVMDeveloper/Public/RigVMModel/RigVMController.h b/Engine/Plugins/Runtime/RigVM/Source/RigVMDeveloper/Public/RigVMModel/RigVMController.h
--- a/Engine/Plugins/Runtime/RigVM/Source/RigVMDeveloper/Public/RigVMModel/RigVMController.h
+++ b/Engine/Plugins/Runtime/RigVM/Source/RigVMDeveloper/Public/RigVMModel/RigVMController.h
@@ -24,6 +24,7 @@
 #include "RigVMController.generated.h"
 
 class URigVMActionStack;
+class IRigVMClientHost;
 struct FRigVMGraphFunctionArgument;
 struct FRigVMGraphFunctionHeader;
 
@@ -294,6 +295,12 @@ public:
 	UFUNCTION(BlueprintCallable, Category = RigVMController)
 	URigVMController* GetControllerForGraph(const URigVMGraph* InGraph) const;
 
+	// Returns the client host this controller belongs to
+	IRigVMClientHost* GetClientHost() const;
+
+	// Returns all events present on the client host
+	TArray<FName> GetAllEventNames() const; 
+
 	// The Modified event used to subscribe to changes
 	// happening within the Graph. This is broadcasted to 
 	// for any change happening - not only the changes 
@@ -1435,6 +1442,7 @@ public:
 
 protected:
 
+	IRigVMClientHost* GetClientHost_Internal(const URigVMGraph* InGraph) const;
 	URigVMPin* CreatePinFromPinInfo(const FRigVMRegistry& InRegistry, const FRigVMPinInfoArray& InPreviousPinInfos, const FRigVMPinInfo& InPinInfo, const FString& InPinPath, UObject* InOuter);
 
 	// backwards compatibility code
