RigVM: Add Box related math nodes

#rnx

diff --git a/Engine/Plugins/Runtime/RigVM/Source/RigVM/Private/RigVMFunctions/Debug/RigVMFunction_DebugPrimitives.cpp b/Engine/Plugins/Runtime/RigVM/Source/RigVM/Private/RigVMFunctions/Debug/RigVMFunction_DebugPrimitives.cpp
--- a/Engine/Plugins/Runtime/RigVM/Source/RigVM/Private/RigVMFunctions/Debug/RigVMFunction_DebugPrimitives.cpp
+++ b/Engine/Plugins/Runtime/RigVM/Source/RigVM/Private/RigVMFunctions/Debug/RigVMFunction_DebugPrimitives.cpp
@@ -81,3 +81,16 @@ FRigVMFunction_DebugArcNoSpace_Execute()
 
 	ExecuteContext.GetDrawInterface()->DrawArc(WorldOffset, Transform, Radius, FMath::DegreesToRadians(MinimumDegrees), FMath::DegreesToRadians(MaximumDegrees), Color, Thickness, Detail);
 }
+
+FRigVMFunction_DebugBoxNoSpace_Execute()
+{
+	if (ExecuteContext.GetDrawInterface() == nullptr || !bEnabled || !Box.IsValid)
+	{
+		return;
+	}
+
+	FTransform Transform = FTransform::Identity;
+	Transform.SetTranslation(Box.GetCenter());
+	Transform.SetScale3D(Box.GetSize());
+	ExecuteContext.GetDrawInterface()->DrawBox(WorldOffset, Transform, Color, Thickness);
+}
diff --git a/Engine/Plugins/Runtime/RigVM/Source/RigVM/Private/RigVMFunctions/Math/RigVMFunction_MathBox.cpp b/Engine/Plugins/Runtime/RigVM/Source/RigVM/Private/RigVMFunctions/Math/RigVMFunction_MathBox.cpp
new file mode 100644
--- /dev/null
+++ b/Engine/Plugins/Runtime/RigVM/Source/RigVM/Private/RigVMFunctions/Math/RigVMFunction_MathBox.cpp
@@ -0,0 +1,130 @@
+// Copyright Epic Games, Inc. All Rights Reserved.
+
+#include "RigVMFunctions/Math/RigVMFunction_MathBox.h"
+#include "Engine/Engine.h"
+
+#include UE_INLINE_GENERATED_CPP_BY_NAME(RigVMFunction_MathBox)
+
+FRigVMFunction_MathBoxFromArray_Execute()
+{
+	Box = FBox(EForceInit::ForceInit);
+	for(const FVector& Position : Array)
+	{
+		Box += Position;
+	}
+	if(Array.IsEmpty())
+	{
+		Box += FVector::ZeroVector;
+	}
+	Minimum = Box.Min;
+	Maximum = Box.Max;
+	Center = Box.GetCenter();
+	Size = Box.GetSize();
+}
+
+FRigVMFunction_MathBoxIsValid_Execute()
+{
+	Valid = Box.IsValid > 0; 
+}
+
+FRigVMFunction_MathBoxGetCenter_Execute()
+{
+	if(Box.IsValid)
+	{
+		Center = Box.GetCenter();
+	}
+	else
+	{
+		Center = FVector::ZeroVector;
+	}
+}
+
+FRigVMFunction_MathBoxGetSize_Execute()
+{
+	if(Box.IsValid)
+	{
+		Size = Box.GetSize();
+		Extent = Box.GetExtent();
+	}
+	else
+	{
+		Size = Extent = FVector::ZeroVector;
+	}
+}
+
+FRigVMFunction_MathBoxShift_Execute()
+{
+	if(!Box.IsValid)
+	{
+		Result = FBox(EForceInit::ForceInit);
+		return;
+	}
+	Result = Box.ShiftBy(Amount);
+}
+
+FRigVMFunction_MathBoxMoveTo_Execute()
+{
+	if(!Box.IsValid)
+	{
+		Result = FBox(EForceInit::ForceInit);
+		return;
+	}
+	Result = Box.MoveTo(Center);
+}
+
+FRigVMFunction_MathBoxExpand_Execute()
+{
+	if(!Box.IsValid)
+	{
+		Result = FBox(EForceInit::ForceInit);
+		return;
+	}
+	Result = Box.ExpandBy(Amount);
+}
+
+FRigVMFunction_MathBoxTransform_Execute()
+{
+	if(!Box.IsValid)
+	{
+		Result = FBox(EForceInit::ForceInit);
+		return;
+	}
+	Result = Box.TransformBy(Transform);
+}
+
+FRigVMFunction_MathBoxGetDistance_Execute()
+{
+	if(!Box.IsValid)
+	{
+		Valid = false;
+		Distance = 0.f;
+		return;
+	}
+
+	Distance = Box.ComputeSquaredDistanceToPoint(Position);
+	if(!Square && Distance >= 0)
+	{
+		Distance = sqrt(Distance);
+	}
+	Valid = true;
+}
+
+FRigVMFunction_MathBoxIsInside_Execute()
+{
+	if(!Box.IsValid)
+	{
+		Result = false;
+		return;
+	}
+	Result = Box.IsInsideOrOn(Position);
+}
+
+FRigVMFunction_MathBoxGetVolume_Execute()
+{
+	if(!Box.IsValid)
+	{
+		Volume = 0.f;
+		return;
+	}
+	Volume = Box.GetVolume();
+}
diff --git a/Engine/Plugins/Runtime/RigVM/Source/RigVM/Public/RigVMFunctions/Debug/RigVMFunction_DebugPrimitives.h b/Engine/Plugins/Runtime/RigVM/Source/RigVM/Public/RigVMFunctions/Debug/RigVMFunction_DebugPrimitives.h
--- a/Engine/Plugins/Runtime/RigVM/Source/RigVM/Public/RigVMFunctions/Debug/RigVMFunction_DebugPrimitives.h
+++ b/Engine/Plugins/Runtime/RigVM/Source/RigVM/Public/RigVMFunctions/Debug/RigVMFunction_DebugPrimitives.h
@@ -190,6 +190,43 @@ struct RIGVM_API FRigVMFunction_DebugArcNoSpace : public FRigVMFunction_DebugBas
 	UPROPERTY(meta = (Input))
 	FTransform WorldOffset;
 
+	UPROPERTY(meta = (Input))
+	bool bEnabled;
+};
+
+/**
+ * Draws a box in the viewport
+ */
+USTRUCT(meta=(DisplayName="Draw Box", Keywords="BoundingBox,Bbox"))
+struct RIGVM_API FRigVMFunction_DebugBoxNoSpace : public FRigVMFunction_DebugBaseMutable
+{
+	GENERATED_BODY()
+
+	FRigVMFunction_DebugBoxNoSpace()
+	{
+		Box = FBox(EForceInit::ForceInit);
+		WorldOffset = FTransform::Identity;
+		Color = FLinearColor::Red;
+		Thickness = 0.f;
+		WorldOffset = FTransform::Identity;
+		bEnabled = true;
+	}
+
+	RIGVM_METHOD()
+	virtual void Execute() override;
+
+	UPROPERTY(meta = (Input))
+	FBox Box;
+
+	UPROPERTY(meta = (Input))
+	FLinearColor Color;
+
+	UPROPERTY(meta = (Input))
+	float Thickness;
+
+	UPROPERTY(meta = (Input))
+	FTransform WorldOffset;
+
 	UPROPERTY(meta = (Input))
 	bool bEnabled;
 };
\ No newline at end of file
diff --git a/Engine/Plugins/Runtime/RigVM/Source/RigVM/Public/RigVMFunctions/Math/RigVMFunction_MathBox.h b/Engine/Plugins/Runtime/RigVM/Source/RigVM/Public/RigVMFunctions/Math/RigVMFunction_MathBox.h
new file mode 100644
--- /dev/null
+++ b/Engine/Plugins/Runtime/RigVM/Source/RigVM/Public/RigVMFunctions/Math/RigVMFunction_MathBox.h
@@ -0,0 +1,338 @@
+// Copyright Epic Games, Inc. All Rights Reserved.
+
+#pragma once
+
+#include "RigVMFunction_MathBase.h"
+#include "RigVMMathLibrary.h"
+#include "RigVMFunction_MathBox.generated.h"
+
+USTRUCT(meta=(Abstract, Category="Math|Box", MenuDescSuffix="(Box)"))
+struct RIGVM_API FRigVMFunction_MathBoxBase : public FRigVMFunction_MathBase
+{
+	GENERATED_BODY()
+};
+
+/**
+ * Returns bounding box of the given array of positions
+ */
+USTRUCT(meta = (DisplayName = "Box from Array", TemplateName = "ArrayBounds,CreateBox,CreateBoundingBox,NewBox,NewBoundingBoxMakeBox,MakeBoundingBox,BoundingBox,Bbox,Bounds"))
+struct RIGVM_API FRigVMFunction_MathBoxFromArray : public FRigVMFunction_MathBoxBase
+{
+	GENERATED_BODY()
+
+	FRigVMFunction_MathBoxFromArray()
+	{
+		Box = FBox(EForceInit::ForceInit);
+		Minimum = Maximum = Center = Size = FVector::ZeroVector;
+	}
+
+	/** Execute logic for this rig unit */
+	RIGVM_METHOD()
+	virtual void Execute() override;
+
+	UPROPERTY(meta = (Input))
+	TArray<FVector> Array;
+
+	UPROPERTY(meta = (Output))
+	FBox Box;
+
+	UPROPERTY(meta = (Output))
+	FVector Minimum;
+
+	UPROPERTY(meta = (Output))
+	FVector Maximum;
+
+	UPROPERTY(meta = (Output))
+	FVector Center;
+	
+	UPROPERTY(meta = (Output))
+	FVector Size;
+};
+
+/**
+ * Returns true if the box has any content / is valid
+ */
+USTRUCT(meta = (DisplayName = "Is Box Valid", TemplateName = "IsValid,HasVolume,ContainsPoints,Bounds,BoundingBox,Bbox"))
+struct RIGVM_API FRigVMFunction_MathBoxIsValid : public FRigVMFunction_MathBoxBase
+{
+	GENERATED_BODY()
+
+	FRigVMFunction_MathBoxIsValid()
+	{
+		Box = FBox(EForceInit::ForceInit);
+		Valid = false;
+	}
+
+	/** Execute logic for this rig unit */
+	RIGVM_METHOD()
+	virtual void Execute() override;
+
+	UPROPERTY(meta = (Input))
+	FBox Box;
+
+	UPROPERTY(meta = (Output))
+	bool Valid;
+};
+
+/**
+ * Returns the center of a bounding box
+ */
+USTRUCT(meta = (DisplayName = "Get Box Center", TemplateName = "Middle,Origin,Bounds,BoundingBox,Bbox"))
+struct RIGVM_API FRigVMFunction_MathBoxGetCenter : public FRigVMFunction_MathBoxBase
+{
+	GENERATED_BODY()
+
+	FRigVMFunction_MathBoxGetCenter()
+	{
+		Box = FBox(EForceInit::ForceInit);
+		Center = FVector::ZeroVector;
+	}
+
+	/** Execute logic for this rig unit */
+	RIGVM_METHOD()
+	virtual void Execute() override;
+
+	UPROPERTY(meta = (Input))
+	FBox Box;
+
+	UPROPERTY(meta = (Output))
+	FVector Center;
+};
+
+/**
+ * Returns the size of a bounding box
+ */
+USTRUCT(meta = (DisplayName = "Get Box Size", TemplateName = "Middle,Origin,Bounds,BoundingBox,Bbox"))
+struct RIGVM_API FRigVMFunction_MathBoxGetSize : public FRigVMFunction_MathBoxBase
+{
+	GENERATED_BODY()
+
+	FRigVMFunction_MathBoxGetSize()
+	{
+		Box = FBox(EForceInit::ForceInit);
+		Size = Extent = FVector::ZeroVector;
+	}
+
+	/** Execute logic for this rig unit */
+	RIGVM_METHOD()
+	virtual void Execute() override;
+
+	UPROPERTY(meta = (Input))
+	FBox Box;
+
+	// the overall size of the box
+	UPROPERTY(meta = (Output))
+	FVector Size;
+
+	// the half size of the box
+	UPROPERTY(meta = (Output))
+	FVector Extent;
+};
+
+/**
+ * Move the box by a certain amount 
+ */
+USTRUCT(meta = (DisplayName = "Shift Box", TemplateName = "Bbox,Translate,Move,BoundingBox"))
+struct RIGVM_API FRigVMFunction_MathBoxShift : public FRigVMFunction_MathBoxBase
+{
+	GENERATED_BODY()
+
+	FRigVMFunction_MathBoxShift()
+	{
+		Box = Result = FBox(EForceInit::ForceInit);
+		Amount = FVector::ZeroVector;
+	}
+
+	/** Execute logic for this rig unit */
+	RIGVM_METHOD()
+	virtual void Execute() override;
+
+	UPROPERTY(meta = (Input))
+	FBox Box;
+
+	// the amount / vector to shift the box by
+	UPROPERTY(meta = (Input))
+	FVector Amount;
+
+	UPROPERTY(meta = (Output))
+	FBox Result;
+};
+
+/**
+ * Moves the center of the box to a new location
+ */
+USTRUCT(meta = (DisplayName = "Move Box To", TemplateName = "Bbox,Translate,Move,BoundingBox"))
+struct RIGVM_API FRigVMFunction_MathBoxMoveTo : public FRigVMFunction_MathBoxBase
+{
+	GENERATED_BODY()
+
+	FRigVMFunction_MathBoxMoveTo()
+	{
+		Box = Result = FBox(EForceInit::ForceInit);
+		Center = FVector::ZeroVector;
+	}
+
+	/** Execute logic for this rig unit */
+	RIGVM_METHOD()
+	virtual void Execute() override;
+
+	UPROPERTY(meta = (Input))
+	FBox Box;
+
+	// the new center for the box
+	UPROPERTY(meta = (Input))
+	FVector Center;
+
+	UPROPERTY(meta = (Output))
+	FBox Result;
+};
+
+/**
+ * Expands the size of the box by a given amount
+ */
+USTRUCT(meta = (DisplayName = "Expand Box", TemplateName = "Bbox,Scale,Grow,Shrink,BoundingBox"))
+struct RIGVM_API FRigVMFunction_MathBoxExpand : public FRigVMFunction_MathBoxBase
+{
+	GENERATED_BODY()
+
+	FRigVMFunction_MathBoxExpand()
+	{
+		Box = Result = FBox(EForceInit::ForceInit);
+		Amount = FVector::ZeroVector;
+	}
+
+	/** Execute logic for this rig unit */
+	RIGVM_METHOD()
+	virtual void Execute() override;
+
+	UPROPERTY(meta = (Input))
+	FBox Box;
+
+	// the amount to grow / shrink the box by
+	UPROPERTY(meta = (Input))
+	FVector Amount;
+
+	UPROPERTY(meta = (Output))
+	FBox Result;
+};
+
+/**
+ * Transforms the box by a given transform
+ */
+USTRUCT(meta = (DisplayName = "Transform Box", TemplateName = "Bbox,BoundingBox"))
+struct RIGVM_API FRigVMFunction_MathBoxTransform : public FRigVMFunction_MathBoxBase
+{
+	GENERATED_BODY()
+
+	FRigVMFunction_MathBoxTransform()
+	{
+		Box = Result = FBox(EForceInit::ForceInit);
+		Transform = FTransform::Identity;
+	}
+
+	/** Execute logic for this rig unit */
+	RIGVM_METHOD()
+	virtual void Execute() override;
+
+	UPROPERTY(meta = (Input))
+	FBox Box;
+
+	UPROPERTY(meta = (Input))
+	FTransform Transform;
+
+	UPROPERTY(meta = (Output))
+	FBox Result;
+};
+
+/**
+ * Returns the distance to a given box
+ */
+USTRUCT(meta = (DisplayName = "Get Distance to Box", TemplateName = "Bbox,Closest,BoundingBox"))
+struct RIGVM_API FRigVMFunction_MathBoxGetDistance : public FRigVMFunction_MathBoxBase
+{
+	GENERATED_BODY()
+
+	FRigVMFunction_MathBoxGetDistance()
+	{
+		Box = FBox(EForceInit::ForceInit);
+		Position = FVector::ZeroVector;
+		Square = true;
+		Valid = false;
+		Distance = 0.f;
+	}
+
+	/** Execute logic for this rig unit */
+	RIGVM_METHOD()
+	virtual void Execute() override;
+
+	UPROPERTY(meta = (Input))
+	FBox Box;
+
+	UPROPERTY(meta = (Input))
+	FVector Position;
+
+	// if true the distance will be returned square
+	UPROPERTY(meta = (Input))
+	bool Square;
+
+	UPROPERTY(meta = (Output))
+	bool Valid;
+
+	UPROPERTY(meta = (Output))
+	float Distance;
+};
+
+/**
+ * Returns true if a point is inside a given box
+ */
+USTRUCT(meta = (DisplayName = "Is Inside Box", TemplateName = "Bbox,Contains,Encompasses,BoundingBox"))
+struct RIGVM_API FRigVMFunction_MathBoxIsInside : public FRigVMFunction_MathBoxBase
+{
+	GENERATED_BODY()
+
+	FRigVMFunction_MathBoxIsInside()
+	{
+		Box = FBox(EForceInit::ForceInit);
+		Position = FVector::ZeroVector;
+		Result = false;
+	}
+
+	/** Execute logic for this rig unit */
+	RIGVM_METHOD()
+	virtual void Execute() override;
+
+	UPROPERTY(meta = (Input))
+	FBox Box;
+
+	UPROPERTY(meta = (Output))
+	FVector Position;
+
+	UPROPERTY(meta = (Output))
+	bool Result;
+};
+
+
+/**
+ * Returns the volume of a given box
+ */
+USTRUCT(meta = (DisplayName = "Get Box Volume", TemplateName = "Bbox,BoundingBox"))
+struct RIGVM_API FRigVMFunction_MathBoxGetVolume : public FRigVMFunction_MathBoxBase
+{
+	GENERATED_BODY()
+
+	FRigVMFunction_MathBoxGetVolume()
+	{
+		Box = FBox(EForceInit::ForceInit);
+		Volume = 0.f;
+	}
+
+	/** Execute logic for this rig unit */
+	RIGVM_METHOD()
+	virtual void Execute() override;
+
+	UPROPERTY(meta = (Input))
+	FBox Box;
+
+	UPROPERTY(meta = (Output))
+	float Volume;
+};
\ No newline at end of file
