RigVM: fix display names for array element pins (fallout from pin category work)

#rnx

diff --git a/Engine/Plugins/Runtime/RigVM/Source/RigVMDeveloper/Private/RigVMModel/RigVMPin.cpp b/Engine/Plugins/Runtime/RigVM/Source/RigVMDeveloper/Private/RigVMModel/RigVMPin.cpp
--- a/Engine/Plugins/Runtime/RigVM/Source/RigVMDeveloper/Private/RigVMModel/RigVMPin.cpp
+++ b/Engine/Plugins/Runtime/RigVM/Source/RigVMDeveloper/Private/RigVMModel/RigVMPin.cpp
@@ -400,6 +400,11 @@ FName URigVMPin::GetDisplayName() const
 {
 	if (DisplayName == NAME_None)
 	{
+		if(IsArrayElement())
+		{
+			return *FString::FromInt(GetPinIndex());
+		}
+		
 		if(const URigVMNode* Node = GetNode())
 		{
 			const FName DisplayNameFromNode = Node->GetDisplayNameForPin(this->GetSegmentPath(true));
