RigVM: Fix shared lazy branch compilation

#rnx

diff --git a/Engine/Plugins/Runtime/RigVM/Source/RigVMDeveloper/Private/RigVMCompiler/RigVMAST.cpp b/Engine/Plugins/Runtime/RigVM/Source/RigVMDeveloper/Private/RigVMCompiler/RigVMAST.cpp
--- a/Engine/Plugins/Runtime/RigVM/Source/RigVMDeveloper/Private/RigVMCompiler/RigVMAST.cpp
+++ b/Engine/Plugins/Runtime/RigVM/Source/RigVMDeveloper/Private/RigVMCompiler/RigVMAST.cpp
@@ -725,6 +725,10 @@ bool FRigVMBlockExprAST::IsObsolete() const
 		}
 	};
 
+	if(NumParents() == 0)
+	{
+		return Local::IsObsoleteExpr(this);
+	}
 	return Local::HasOnlyObsoleteParents(this);
 }
 
