RigVM: Improve navigation for functions, allow navigation through details panel and auto select entry node

#rnx

diff --git a/Engine/Plugins/Runtime/RigVM/Source/RigVMEditor/Private/Editor/RigVMEditor.cpp b/Engine/Plugins/Runtime/RigVM/Source/RigVMEditor/Private/Editor/RigVMEditor.cpp
--- a/Engine/Plugins/Runtime/RigVM/Source/RigVMEditor/Private/Editor/RigVMEditor.cpp
+++ b/Engine/Plugins/Runtime/RigVM/Source/RigVMEditor/Private/Editor/RigVMEditor.cpp
@@ -2486,7 +2486,7 @@ void FRigVMEditor::HandleJumpToHyperlink(const UObject* InSubject)
 	{
 		if(URigVMEdGraph* EdGraph = Cast<URigVMEdGraph>(RigBlueprint->GetEdGraph(NodeToJumpTo->GetGraph())))
 		{
-			if(const URigVMEdGraphNode* EdGraphNode = Cast<URigVMEdGraphNode>(EdGraph->FindNodeForModelNodeName(NodeToJumpTo->GetFName())))
+			if(URigVMEdGraphNode* EdGraphNode = Cast<URigVMEdGraphNode>(EdGraph->FindNodeForModelNodeName(NodeToJumpTo->GetFName())))
 			{
 				if(PinToJumpTo)
 				{
@@ -2498,6 +2498,7 @@ void FRigVMEditor::HandleJumpToHyperlink(const UObject* InSubject)
 				}
 				
 				JumpToNode(EdGraphNode);
+				SetDetailObjects({EdGraphNode});
 				return;
 			}
 			
@@ -3393,15 +3394,15 @@ void FRigVMEditor::OnNodeDoubleClicked(URigVMBlueprint* InBlueprint, URigVMNode*
 
 	if (URigVMLibraryNode* LibraryNode = Cast<URigVMLibraryNode>(InNode))
 	{
-		URigVMGraph* ContainedGraph = LibraryNode->GetContainedGraph();
-		if (URigVMFunctionReferenceNode* FunctionReferenceNode = Cast<URigVMFunctionReferenceNode>(LibraryNode))
+		if (const URigVMFunctionReferenceNode* FunctionReferenceNode = Cast<URigVMFunctionReferenceNode>(LibraryNode))
 		{
-			if (URigVMLibraryNode* ReferencedNode = FunctionReferenceNode->LoadReferencedNode())
+			if (const URigVMLibraryNode* ReferencedNode = FunctionReferenceNode->LoadReferencedNode())
 			{
-				ContainedGraph = ReferencedNode->GetContainedGraph();
+				HandleJumpToHyperlink(ReferencedNode);
+				return;
 			}
 		}
-		if(ContainedGraph)
+		if(URigVMGraph* ContainedGraph = LibraryNode->GetContainedGraph())
 		{
 			if (UEdGraph* EdGraph = InBlueprint->GetEdGraph(ContainedGraph))
 			{
@@ -3409,13 +3410,13 @@ void FRigVMEditor::OnNodeDoubleClicked(URigVMBlueprint* InBlueprint, URigVMNode*
 			}
 			else
 			{
-				if(URigVMCollapseNode* FunctionLibraryNode = Cast<URigVMCollapseNode>(ContainedGraph->GetOuter()))
+				if(const URigVMCollapseNode* FunctionLibraryNode = Cast<URigVMCollapseNode>(ContainedGraph->GetOuter()))
 				{
-					if(URigVMFunctionLibrary* FunctionLibrary = FunctionLibraryNode->GetLibrary())
+					if(const URigVMFunctionLibrary* FunctionLibrary = FunctionLibraryNode->GetLibrary())
 					{
-						if(URigVMBlueprint* FunctionBlueprint = Cast<URigVMBlueprint>(FunctionLibrary->GetOuter()))
+						if(const URigVMBlueprint* FunctionBlueprint = Cast<URigVMBlueprint>(FunctionLibrary->GetOuter()))
 						{
-							if (UEdGraph* FunctionEdGraph = FunctionBlueprint->GetEdGraph(ContainedGraph))
+							if (const UEdGraph* FunctionEdGraph = FunctionBlueprint->GetEdGraph(ContainedGraph))
 							{
 								FKismetEditorUtilities::BringKismetToFocusAttentionOnObject(FunctionEdGraph);
 							}
diff --git a/Engine/Plugins/Runtime/RigVM/Source/RigVMEditor/Private/Editor/RigVMGraphDetailCustomization.cpp b/Engine/Plugins/Runtime/RigVM/Source/RigVMEditor/Private/Editor/RigVMGraphDetailCustomization.cpp
--- a/Engine/Plugins/Runtime/RigVM/Source/RigVMEditor/Private/Editor/RigVMGraphDetailCustomization.cpp
+++ b/Engine/Plugins/Runtime/RigVM/Source/RigVMEditor/Private/Editor/RigVMGraphDetailCustomization.cpp
@@ -1502,6 +1502,70 @@ void FRigVMWrappedNodeDetailCustomization::CustomizeDetails(IDetailLayoutBuilder
 
 	UClass* WrapperClass = ObjectsBeingCustomized[0]->GetClass();
 
+	if(NodesBeingCustomized.Num() == 1)
+	{
+		if(NodesBeingCustomized[0].IsValid())
+		{
+			if(const URigVMFunctionReferenceNode* FunctionReferenceNode = Cast<URigVMFunctionReferenceNode>(NodesBeingCustomized[0].Get()))
+			{
+				const FRigVMGraphFunctionHeader& Header = FunctionReferenceNode->GetReferencedFunctionHeader();
+				const FRigVMGraphFunctionIdentifier& Identifier = Header.LibraryPointer;
+				
+				IDetailCategoryBuilder& FunctionCategory = DetailLayout.EditCategory("Function", LOCTEXT("Function", "Function"), ECategoryPriority::Uncommon);
+				FunctionCategory.InitiallyCollapsed(false);
+
+				FunctionCategory.AddCustomRow(LOCTEXT("FunctionName", "FunctionName"))
+					.NameContent()
+					[
+						SNew(STextBlock)
+						.Text(FText::FromString(TEXT("Name")))
+						.Font(IDetailLayoutBuilder::GetDetailFont())
+					]
+					.ValueContent()
+					[
+						SNew(STextBlock)
+						.Text(FText::FromString(Identifier.GetFunctionName()))
+						.Font(IDetailLayoutBuilder::GetDetailFont())
+					];
+
+				FunctionCategory.AddCustomRow(LOCTEXT("FunctionPath", "FunctionPath"))
+				.NameContent()
+				[
+					SNew(STextBlock)
+					.Text(FText::FromString(TEXT("Path")))
+					.Font(IDetailLayoutBuilder::GetDetailFont())
+				]
+				.ValueContent()
+				[
+					SNew(SButton)
+					.ButtonStyle(FAppStyle::Get(), TEXT("SimpleButton"))
+					.ContentPadding(0)
+					.Text(FText::FromString(Identifier.GetLibraryNodePath()))
+					.OnClicked_Lambda([Header]() -> FReply
+					{
+						if(const URigVMLibraryNode* LibraryNode = Cast<URigVMLibraryNode>(Header.LibraryPointer.GetNodeSoftPath().TryLoad()))
+						{
+							if(UBlueprint* Blueprint = LibraryNode->GetTypedOuter<UBlueprint>())
+							{
+								GEditor->GetEditorSubsystem<UAssetEditorSubsystem>()->OpenEditorForAsset(Blueprint);
+			
+								if(IAssetEditorInstance* Editor = GEditor->GetEditorSubsystem<UAssetEditorSubsystem>()->FindEditorForAsset(Blueprint, true))
+								{
+									if(FRigVMEditor* RigVMEditor = static_cast<FRigVMEditor*>(Editor))
+									{
+										RigVMEditor->HandleJumpToHyperlink(LibraryNode);
+										return FReply::Handled();
+									}
+								}
+							}
+						}
+						return FReply::Unhandled();
+					})
+				];
+			}
+		}
+	}
+
 	// now loop over all of the properties and display them
 	TArray<TSharedPtr<IPropertyHandle>> PropertiesToVisit;
 	for (TFieldIterator<FProperty> PropertyIt(WrapperClass); PropertyIt; ++PropertyIt)
