RigVM: Adapting graph details customization to respect new node layout

#rnx

diff --git a/Engine/Plugins/Runtime/RigVM/Source/RigVM/Private/RigVMCore/RigVMNodeLayout.cpp b/Engine/Plugins/Runtime/RigVM/Source/RigVM/Private/RigVMCore/RigVMNodeLayout.cpp
--- a/Engine/Plugins/Runtime/RigVM/Source/RigVM/Private/RigVMCore/RigVMNodeLayout.cpp
+++ b/Engine/Plugins/Runtime/RigVM/Source/RigVM/Private/RigVMCore/RigVMNodeLayout.cpp
@@ -5,6 +5,16 @@
 
 #include UE_INLINE_GENERATED_CPP_BY_NAME(RigVMNodeLayout)
 
+FString FRigVMPinCategory::GetName() const
+{
+	FString Left, Right;
+	if(!RigVMStringUtils::SplitNodePathAtEnd(Path, Left, Right))
+	{
+		Right = Path;
+	}
+	return Right;
+}
+
 const FString* FRigVMNodeLayout::FindCategory(const FString& InElement) const
 {
 	for(const FRigVMPinCategory& Category : Categories)
diff --git a/Engine/Plugins/Runtime/RigVM/Source/RigVM/Public/RigVMCore/RigVMNodeLayout.h b/Engine/Plugins/Runtime/RigVM/Source/RigVM/Public/RigVMCore/RigVMNodeLayout.h
--- a/Engine/Plugins/Runtime/RigVM/Source/RigVM/Public/RigVMCore/RigVMNodeLayout.h
+++ b/Engine/Plugins/Runtime/RigVM/Source/RigVM/Public/RigVMCore/RigVMNodeLayout.h
@@ -21,6 +21,8 @@ struct RIGVM_API FRigVMPinCategory
 	UPROPERTY(BlueprintReadOnly, VisibleAnywhere, Category=FunctionArgument)
 	TArray<FString> Elements;
 
+	FString GetName() const;
+
 	friend uint32 GetTypeHash(const FRigVMPinCategory& Category)
 	{
 		uint32 Hash = GetTypeHash(Category.Path);
@@ -88,6 +90,16 @@ struct RIGVM_API FRigVMNodeLayout
 		return Hash;
 	}
 
+	bool operator ==(const FRigVMNodeLayout& OtherLayout) const
+	{
+		return GetTypeHash(*this) == GetTypeHash(OtherLayout);
+	}
+
+	bool operator !=(const FRigVMNodeLayout& OtherLayout) const
+	{
+		return !(*this == OtherLayout);
+	}
+	
 	friend RIGVM_API FArchive& operator<<(FArchive& Ar, FRigVMNodeLayout& Layout);
 
 	const FString* FindCategory(const FString& InElement) const;
diff --git a/Engine/Plugins/Runtime/RigVM/Source/RigVM/Public/RigVMCore/RigVMPropertyPath.h b/Engine/Plugins/Runtime/RigVM/Source/RigVM/Public/RigVMCore/RigVMPropertyPath.h
--- a/Engine/Plugins/Runtime/RigVM/Source/RigVM/Public/RigVMCore/RigVMPropertyPath.h
+++ b/Engine/Plugins/Runtime/RigVM/Source/RigVM/Public/RigVMCore/RigVMPropertyPath.h
@@ -128,6 +128,9 @@ public:
 	// Index based access operator for a segment
 	const FRigVMPropertyPathSegment& operator[](int32 InIndex) const { return Segments[InIndex]; }
 
+	// Returns the segments of this property path
+	const TArray<FRigVMPropertyPathSegment>& GetSegments() const { return Segments; }
+
 	bool operator ==(const FRigVMPropertyPath& Other) const
 	{
 		return GetTypeHash(this) == GetTypeHash(Other);
diff --git a/Engine/Plugins/Runtime/RigVM/Source/RigVMDeveloper/Private/RigVMModel/RigVMNode.cpp b/Engine/Plugins/Runtime/RigVM/Source/RigVMDeveloper/Private/RigVMModel/RigVMNode.cpp
--- a/Engine/Plugins/Runtime/RigVM/Source/RigVMDeveloper/Private/RigVMModel/RigVMNode.cpp
+++ b/Engine/Plugins/Runtime/RigVM/Source/RigVMDeveloper/Private/RigVMModel/RigVMNode.cpp
@@ -303,7 +303,36 @@ FRigVMNodeLayout URigVMNode::GetPinLayout() const
 	{
 		if(const FRigVMPinCategory* Category = CategoryMap.Find(PinCategory))
 		{
-			Layout.Categories.Add(*Category);
+			FRigVMPinCategory CategoryCopy = *Category;
+
+			// sort the elements based on pin index
+			// we start by assuming indices above the user defined range,
+			// so say for 4 pins we'll use (4,5,6,7) and then inline the
+			// user provided pin indices within the range of 0 to 3.
+			TMap<FString,int32> PinPathToIndex;
+			for(const FString& PinPath : CategoryCopy.Elements)
+			{
+				const int32 Index = CategoryCopy.Elements.Num() + PinPathToIndex.Num();
+				PinPathToIndex.Add(PinPath, Index);
+			}
+			for(const FString& PinPath : CategoryCopy.Elements)
+			{
+				if(const URigVMPin* Pin = FindPin(PinPath))
+				{
+					const int32 Index = Pin->GetIndexInCategory();
+					if(CategoryCopy.Elements.IsValidIndex(Index))
+					{
+						PinPathToIndex.FindChecked(PinPath) = Index;
+					}
+				}
+			}
+
+			Algo::SortBy(Category->Elements, [PinPathToIndex](const FString& PinPath) -> int32
+			{
+ 				return PinPathToIndex.FindChecked(PinPath);
+			});
+			
+			Layout.Categories.Add(CategoryCopy);
 		}
 	}
 
diff --git a/Engine/Plugins/Runtime/RigVM/Source/RigVMEditor/Private/Editor/RigVMGraphDetailCustomization.cpp b/Engine/Plugins/Runtime/RigVM/Source/RigVMEditor/Private/Editor/RigVMGraphDetailCustomization.cpp
--- a/Engine/Plugins/Runtime/RigVM/Source/RigVMEditor/Private/Editor/RigVMGraphDetailCustomization.cpp
+++ b/Engine/Plugins/Runtime/RigVM/Source/RigVMEditor/Private/Editor/RigVMGraphDetailCustomization.cpp
@@ -26,6 +26,7 @@
 #include "RigVMModel/Nodes/RigVMAggregateNode.h"
 #include "Widgets/SRigVMGraphPinVariableBinding.h"
 #include "InstancedPropertyBagStructureDataProvider.h"
+#include "RigVMStringUtils.h"
 #include "Widgets/SRigVMGraphPinEnumPicker.h"
 #include "Widgets/SRigVMVariantWidget.h"
 #include "ScopedTransaction.h"
@@ -1566,20 +1567,221 @@ void FRigVMWrappedNodeDetailCustomization::CustomizeDetails(IDetailLayoutBuilder
 		}
 	}
 
+	// determine the order of things
+	typedef TTuple<const FProperty*, FRigVMPropertyPath, FString> TPropertyToToShow;
+	TArray<TPropertyToToShow> PropertiesToShow;
+
+	bool bInspectingOnlyOneNodeType = NodesBeingCustomized.Num() == 1;
+	if(NodesBeingCustomized.Num() > 1)
+	{
+		const UClass* NodeClass = nullptr;
+		TArray<TTuple<FString,const UScriptStruct*>> Traits;
+		FName TemplateNotation(NAME_None);
+		FRigVMNodeLayout NodeLayout;
+		for(const TWeakObjectPtr<URigVMNode>& NodePtr : NodesBeingCustomized)
+		{
+			if(!NodePtr.IsValid())
+			{
+				continue;
+			}
+			
+			const URigVMNode* Node = NodePtr.Get();
+			if(NodeClass == nullptr)
+			{
+				// when looking at the first node - remember the relevant bits
+				NodeClass = Node->GetClass();
+
+				if(const URigVMTemplateNode* TemplateNode = Cast<URigVMTemplateNode>(Node))
+				{
+					TemplateNotation = TemplateNode->GetNotation();
+				}
+
+				NodeLayout = Node->GetPinLayout();
+
+				for(const FString& TraitName : Node->GetTraitNames())
+				{
+					Traits.Emplace(TraitName, Node->GetTraitScriptStruct(*TraitName));
+				}
+			}
+			else
+			{
+				if(Node->GetClass() != NodeClass)
+				{
+					bInspectingOnlyOneNodeType = false;
+					break;
+				}
+
+				if(const URigVMTemplateNode* TemplateNode = Cast<URigVMTemplateNode>(Node))
+				{
+					if(!TemplateNotation.IsEqual(TemplateNode->GetNotation()))
+					{
+						bInspectingOnlyOneNodeType = false;
+						break;
+					}
+				}
+
+				if(NodeLayout != Node->GetPinLayout())
+				{
+					bInspectingOnlyOneNodeType = false;
+					break;
+				}
+
+				if(Node->GetTraitNames().Num() != Traits.Num())
+				{
+					bInspectingOnlyOneNodeType = false;
+					break;
+				}
+				
+				for(int32 TraitIndex = 0; TraitIndex < Traits.Num(); TraitIndex++)
+				{
+					const FString TraitName = Node->GetTraitNames()[TraitIndex];
+					if(TraitName != Traits[TraitIndex].Get<0>())
+					{
+						bInspectingOnlyOneNodeType = false;
+						break;
+					}
+					if(Traits[TraitIndex].Get<1>() != Node->GetTraitScriptStruct(*TraitName))
+					{
+						bInspectingOnlyOneNodeType = false;
+						break;
+					}
+				}
+				if(!bInspectingOnlyOneNodeType)
+				{
+					break;
+				}
+			}
+		}
+	}
+
+	const URigVMNode* NodeWithCategories = nullptr;
+	if(bInspectingOnlyOneNodeType)
+	{
+		// determine if we should be using pin categories to display the node
+		for(const TWeakObjectPtr<URigVMNode>& NodePtr : NodesBeingCustomized)
+		{
+			if(NodePtr.IsValid())
+			{
+				NodeWithCategories = NodePtr.Get();
+				if(NodeWithCategories->GetPinCategories().IsEmpty())
+				{
+					NodeWithCategories = nullptr;
+				}
+				break;
+			}
+		}
+	}
+
+	if(NodeWithCategories)
+	{
+		const FRigVMNodeLayout NodeLayout = NodeWithCategories->GetPinLayout();
+		for(const FRigVMPinCategory& Category : NodeLayout.Categories)
+		{
+			for(const FString& PinPath : Category.Elements)
+			{
+				FString Left, Right;
+				if(!URigVMPin::SplitPinPathAtStart(PinPath, Left, Right))
+				{
+					Left = PinPath;
+				}
+				if(const FProperty* Property = WrapperClass->FindPropertyByName(*Left))
+				{
+					FRigVMPropertyPath PropertyPath;
+					if(!Right.IsEmpty())
+					{
+						PropertyPath = FRigVMPropertyPath(Property, Right);
+					}
+					PropertiesToShow.Emplace(Property, PropertyPath, Category.Path);
+				}
+			}
+		}
+	}
+	else
+	{
+		// if we don't have a pin category layout let's just use all root properties
+		for (TFieldIterator<FProperty> PropertyIt(WrapperClass); PropertyIt; ++PropertyIt)
+		{
+			FProperty* Property = *PropertyIt;
+			const FName PropertyName = Property->GetFName();
+			TSharedPtr<IPropertyHandle> PropertyHandle = DetailLayout.GetProperty(PropertyName, WrapperClass);
+			if (!PropertyHandle->IsValidHandle())
+			{
+				continue;
+			}
+			PropertiesToShow.Emplace(Property, FRigVMPropertyPath(), FString());
+		}
+	}
+
 	// now loop over all of the properties and display them
-	TArray<TSharedPtr<IPropertyHandle>> PropertiesToVisit;
-	for (TFieldIterator<FProperty> PropertyIt(WrapperClass); PropertyIt; ++PropertyIt)
+	TArray<TSharedPtr<IPropertyHandle>> PropertiesAddedToLayout;
+	FRigVMNodeLayout NodeLayout;
+	if(NodeWithCategories)
 	{
-		FProperty* Property = *PropertyIt;
-		const FName PropertyName = Property->GetFName();
-		TSharedPtr<IPropertyHandle> PropertyHandle = DetailLayout.GetProperty(PropertyName, WrapperClass);
+		NodeLayout = NodeWithCategories->GetPinLayout();
+	}
+	for (const TTuple<const FProperty*, FRigVMPropertyPath, FString>& Tuple : PropertiesToShow)
+	{
+		const FProperty* Property = Tuple.Get<0>();
+		const FRigVMPropertyPath& PropertyPath = Tuple.Get<1>();
+		FString PinPath = Property->GetName();
+		if(PropertyPath.IsValid())
+		{
+			PinPath = URigVMPin::JoinPinPath(PinPath, PropertyPath.ToString());
+		}
+		const FString Category = Tuple.Get<2>();
+		
+		TSharedPtr<IPropertyHandle> PropertyHandle = DetailLayout.GetProperty(Property->GetFName(), WrapperClass);
+		if (!PropertyHandle->IsValidHandle())
+		{
+			continue;
+		}
+		for(const FRigVMPropertyPathSegment& Segment : PropertyPath.GetSegments())
+		{
+			switch(Segment.Type)
+			{
+				case ERigVMPropertyPathSegmentType::StructMember:
+				{
+					PropertyHandle = PropertyHandle->GetChildHandle(Segment.Name);
+					break;
+				}
+				case ERigVMPropertyPathSegmentType::ArrayElement:
+				{
+					PropertyHandle = PropertyHandle->GetChildHandle(Segment.Index);
+					break;
+				}
+				case ERigVMPropertyPathSegmentType::MapValue:
+				{
+					// not supported just yet
+					checkNoEntry();
+					break;
+				}
+			}
+			if (!PropertyHandle->IsValidHandle())
+			{
+				break;
+			}
+		}
 		if (!PropertyHandle->IsValidHandle())
 		{
 			continue;
 		}
-		PropertiesToVisit.Add(PropertyHandle);
 
-		TAttribute<bool> HasUserProvidedDefaultValue = TAttribute<bool>::CreateLambda([this, PropertyName, PropertyHandle]() -> bool
+		const URigVMPin* Pin = nullptr;
+		for(const TWeakObjectPtr<URigVMNode>& Node : NodesBeingCustomized)
+		{
+			if(Node.IsValid())
+			{
+				Pin = Node->FindPin(PinPath);
+				if(Pin)
+				{
+					break;
+				}
+			}	
+		}
+
+		PropertiesAddedToLayout.Add(PropertyHandle);
+
+		TAttribute<bool> HasUserProvidedDefaultValue = TAttribute<bool>::CreateLambda([this, PinPath, PropertyHandle]() -> bool
 		{
 			if(CVarRigVMEnablePinDefaultTypes.GetValueOnAnyThread())
 			{
@@ -1587,7 +1789,7 @@ void FRigVMWrappedNodeDetailCustomization::CustomizeDetails(IDetailLayoutBuilder
 				{
 					if(Node.IsValid())
 					{
-						if(const URigVMPin* Pin = Node->FindRootPinByName(PropertyName))
+						if(const URigVMPin* Pin = Node->FindPin(PinPath))
 						{
 							if(Pin->HasUserProvidedDefaultValue())
 							{
@@ -1602,7 +1804,7 @@ void FRigVMWrappedNodeDetailCustomization::CustomizeDetails(IDetailLayoutBuilder
 
 		FResetToDefaultOverride ResetToDefault = FResetToDefaultOverride::Create(
 			HasUserProvidedDefaultValue,
-			FSimpleDelegate::CreateLambda([this, PropertyHandle, PropertyName]()
+			FSimpleDelegate::CreateLambda([this, PropertyHandle, PinPath]()
 			{
 				FScopedTransaction Transaction(LOCTEXT("ResetValueToDefault", "Reset Value To Default"));
 				const URigVMGraph* Graph = NodesBeingCustomized[0]->GetGraph();
@@ -1612,7 +1814,7 @@ void FRigVMWrappedNodeDetailCustomization::CustomizeDetails(IDetailLayoutBuilder
 				Controller->OpenUndoBracket(TEXT("Reset pin default value"));
 				for(const TWeakObjectPtr<URigVMNode> Node : NodesBeingCustomized)
 				{
-					if(const URigVMPin* Pin = Node->FindRootPinByName(PropertyName))
+					if(const URigVMPin* Pin = Node->FindPin(PinPath))
 					{
 						Controller->ResetPinDefaultValue(Pin->GetPinPath());
 					}
@@ -1622,8 +1824,13 @@ void FRigVMWrappedNodeDetailCustomization::CustomizeDetails(IDetailLayoutBuilder
 		);
 
 		static const FSlateFontInfo NameFont = FAppStyle::GetFontStyle( TEXT("PropertyWindow.NormalFont") );
-		
-		TSharedRef<SWidget> LabelWidget = PropertyHandle->CreatePropertyNameWidget();
+
+		FText LabelOverride;
+		if(const FString* DisplayName = NodeLayout.FindDisplayName(PinPath))
+		{
+			LabelOverride = FText::FromString(*DisplayName);
+		}
+		TSharedRef<SWidget> LabelWidget = PropertyHandle->CreatePropertyNameWidget(LabelOverride);
 
 		/*
 		// in the future we may want some visual alignment of the label widget on top of the 
@@ -1665,6 +1872,23 @@ void FRigVMWrappedNodeDetailCustomization::CustomizeDetails(IDetailLayoutBuilder
 		];
 		*/
 
+		IDetailPropertyRow* Row = nullptr;
+		if(NodeWithCategories)
+		{
+			DetailLayout.HideProperty(PropertyHandle);
+			FString Left, CategoryName;
+			if(!RigVMStringUtils::SplitNodePathAtEnd(Category, Left, CategoryName))
+			{
+				CategoryName = Category;
+			}
+			Row = &DetailLayout.EditCategory(*Category, FText::FromString(CategoryName))
+				.AddProperty(PropertyHandle);
+		}
+		else
+		{
+			Row = DetailLayout.EditDefaultProperty(PropertyHandle);
+		}
+		
 		// check if any / all pins are bound to a variable
 		int32 PinsBoundToVariable = 0;
 		TArray<URigVMPin*> ModelPins;
@@ -1681,37 +1905,31 @@ void FRigVMWrappedNodeDetailCustomization::CustomizeDetails(IDetailLayoutBuilder
 		{
 			if(PinsBoundToVariable == ModelPins.Num())
 			{
-				if(IDetailPropertyRow* Row = DetailLayout.EditDefaultProperty(PropertyHandle))
-				{
-					Row->CustomWidget()
-					.NameContent()
-					[
-						LabelWidget
-					]
-					.ValueContent()
-					[
-					SNew(SRigVMGraphVariableBinding)
-						.ModelPins(ModelPins)
-						.Blueprint(BlueprintBeingCustomized)
-					];
-					continue;
-				}
+				Row->CustomWidget()
+				.NameContent()
+				[
+					LabelWidget
+				]
+				.ValueContent()
+				[
+				SNew(SRigVMGraphVariableBinding)
+					.ModelPins(ModelPins)
+					.Blueprint(BlueprintBeingCustomized)
+				];
+				continue;
 			}
 			else // in this case some pins are bound, and some are not - we'll hide the input value widget
 			{
-				if(IDetailPropertyRow* Row = DetailLayout.EditDefaultProperty(PropertyHandle))
-				{
-					Row->CustomWidget()
-					.NameContent()
-					[
-						LabelWidget
-					];
-					continue;
-				}
+				Row->CustomWidget()
+				.NameContent()
+				[
+					LabelWidget
+				];
+				continue;
 			}
 		}
 		
-		if (FNameProperty* NameProperty = CastField<FNameProperty>(Property))
+		if (const FNameProperty* NameProperty = CastField<FNameProperty>(Property))
         {
         	FString CustomWidgetName = NameProperty->GetMetaData(TEXT("CustomWidget"));
         	if (!CustomWidgetName.IsEmpty())
@@ -1725,78 +1943,69 @@ void FRigVMWrappedNodeDetailCustomization::CustomizeDetails(IDetailLayoutBuilder
         		{
         			TSharedPtr<SRigVMGraphPinNameListValueWidget> NameListWidget;
 
-        			if(IDetailPropertyRow* Row = DetailLayout.EditDefaultProperty(PropertyHandle))
-        			{
-        				Row->CustomWidget()
-						.NameContent()
-						[
-							LabelWidget
-						]
-						.ValueContent()
-						[
-							SAssignNew(NameListWidget, SRigVMGraphPinNameListValueWidget)
-							.OptionsSource(NameList)
-							.OnGenerateWidget(this, &FRigVMWrappedNodeDetailCustomization::MakeNameListItemWidget)
-							.OnSelectionChanged(this, &FRigVMWrappedNodeDetailCustomization::OnNameListChanged, NameProperty, DetailLayout.GetPropertyUtilities())
-							.OnComboBoxOpening(this, &FRigVMWrappedNodeDetailCustomization::OnNameListComboBox, NameProperty, NameList)
-							.InitiallySelectedItem(GetCurrentlySelectedItem(NameProperty, NameList))
-							.Content()
-							[
-								SNew(STextBlock)
-								.Text(this, &FRigVMWrappedNodeDetailCustomization::GetNameListText, NameProperty)
-								.ColorAndOpacity_Lambda([this, NameProperty]() -> FSlateColor
-								{
-									static FText NoneText = LOCTEXT("None", "None"); 
-									if(GetNameListText(NameProperty).EqualToCaseIgnored(NoneText))
-									{
-										return FSlateColor(FLinearColor::Red);
-									}
-									return FSlateColor::UseForeground();
-								})
-							]
-						]
-        				.OverrideResetToDefault(ResetToDefault);
-        			}        			
-        			NameListWidgets.Add(Property->GetFName(), NameListWidget);
-       				continue;
-        		}
-        		
-        		if(IDetailPropertyRow* Row = DetailLayout.EditDefaultProperty(PropertyHandle))
-        		{
         			Row->CustomWidget()
 					.NameContent()
 					[
 						LabelWidget
 					]
+					.ValueContent()
+					[
+						SAssignNew(NameListWidget, SRigVMGraphPinNameListValueWidget)
+						.OptionsSource(NameList)
+						.OnGenerateWidget(this, &FRigVMWrappedNodeDetailCustomization::MakeNameListItemWidget)
+						.OnSelectionChanged(this, &FRigVMWrappedNodeDetailCustomization::OnNameListChanged, NameProperty, DetailLayout.GetPropertyUtilities())
+						.OnComboBoxOpening(this, &FRigVMWrappedNodeDetailCustomization::OnNameListComboBox, NameProperty, NameList)
+						.InitiallySelectedItem(GetCurrentlySelectedItem(NameProperty, NameList))
+						.Content()
+						[
+							SNew(STextBlock)
+							.Text(this, &FRigVMWrappedNodeDetailCustomization::GetNameListText, NameProperty)
+							.ColorAndOpacity_Lambda([this, NameProperty]() -> FSlateColor
+							{
+								static FText NoneText = LOCTEXT("None", "None"); 
+								if(GetNameListText(NameProperty).EqualToCaseIgnored(NoneText))
+								{
+									return FSlateColor(FLinearColor::Red);
+								}
+								return FSlateColor::UseForeground();
+							})
+						]
+					]
         			.OverrideResetToDefault(ResetToDefault);
-        			continue;
+        			
+        			NameListWidgets.Add(Property->GetFName(), NameListWidget);
+       				continue;
         		}
+        		
+        		Row->CustomWidget()
+				.NameContent()
+				[
+					LabelWidget
+				]
+        		.OverrideResetToDefault(ResetToDefault);
+        		continue;
         	}
         }
 
-		if(IDetailPropertyRow* Row = DetailLayout.EditDefaultProperty(PropertyHandle))
-		{
-			TSharedPtr<SWidget> ValueWidget = PropertyHandle->CreatePropertyValueWidgetWithCustomization(DetailLayout.GetDetailsView());
+		TSharedPtr<SWidget> ValueWidget = PropertyHandle->CreatePropertyValueWidgetWithCustomization(DetailLayout.GetDetailsView());
 
-			constexpr bool bShowChildren = true;
-			Row->CustomWidget(bShowChildren)
-			.NameContent()
-			[
-				LabelWidget
-			]
-			.ValueContent()
-			[
-				ValueWidget ? ValueWidget.ToSharedRef() : SNullWidget::NullWidget
-			]
-			.OverrideResetToDefault(ResetToDefault);
-			continue;
-		}
+		constexpr bool bShowChildren = true;
+		Row->CustomWidget(bShowChildren)
+		.NameContent()
+		[
+			LabelWidget
+		]
+		.ValueContent()
+		[
+			ValueWidget ? ValueWidget.ToSharedRef() : SNullWidget::NullWidget
+		]
+		.OverrideResetToDefault(ResetToDefault);
 	}
 
 	// now loop over all handles and determine expansion states of the corresponding pins
-	for (int32 Index = 0; Index < PropertiesToVisit.Num(); Index++)
+	for (int32 Index = 0; Index < PropertiesAddedToLayout.Num(); Index++)
 	{
-		TSharedPtr<IPropertyHandle> PropertyHandle = PropertiesToVisit[Index];
+		TSharedPtr<IPropertyHandle> PropertyHandle = PropertiesAddedToLayout[Index];
 		FProperty* Property = PropertyHandle->GetProperty();
 
 		// certain properties we don't look at for expansion states
@@ -1839,7 +2048,26 @@ void FRigVMWrappedNodeDetailCustomization::CustomizeDetails(IDetailLayoutBuilder
 		PropertyHandle->GetNumChildren(NumChildren);
 		for(uint32 ChildIndex = 0; ChildIndex < NumChildren; ChildIndex++)
 		{
-			PropertiesToVisit.Add(PropertyHandle->GetChildHandle(ChildIndex));
+			PropertiesAddedToLayout.Add(PropertyHandle->GetChildHandle(ChildIndex));
+		}
+	}
+
+	// hide all root properties not listed in the properties to show list
+	for (TFieldIterator<FProperty> PropertyIt(WrapperClass); PropertyIt; ++PropertyIt)
+	{
+		FProperty* Property = *PropertyIt;
+		if(!PropertiesToShow.ContainsByPredicate([Property](const TPropertyToToShow& Tuple) -> bool
+		{
+			return Tuple.Get<0>() == Property && !Tuple.Get<1>().IsValid();
+		}))
+		{
+			const FName PropertyName = Property->GetFName();
+			TSharedPtr<IPropertyHandle> PropertyHandle = DetailLayout.GetProperty(PropertyName, WrapperClass);
+			if (!PropertyHandle->IsValidHandle())
+			{
+				continue;
+			}
+			DetailLayout.HideProperty(PropertyHandle);
 		}
 	}
 
@@ -1852,7 +2080,7 @@ TSharedRef<SWidget> FRigVMWrappedNodeDetailCustomization::MakeNameListItemWidget
 	return 	SNew(STextBlock).Text(FText::FromString(InItem->GetStringWithTag()));// .Font(FAppStyle::GetFontStyle(TEXT("PropertyWindow.NormalFont")));
 }
 
-FText FRigVMWrappedNodeDetailCustomization::GetNameListText(FNameProperty* InProperty) const
+FText FRigVMWrappedNodeDetailCustomization::GetNameListText(const FNameProperty* InProperty) const
 {
 	FText FirstText;
 	for(TWeakObjectPtr<URigVMDetailsViewWrapperObject> ObjectBeingCustomized : ObjectsBeingCustomized)
@@ -1873,7 +2101,7 @@ FText FRigVMWrappedNodeDetailCustomization::GetNameListText(FNameProperty* InPro
 	return FirstText;
 }
 
-TSharedPtr<FRigVMStringWithTag> FRigVMWrappedNodeDetailCustomization::GetCurrentlySelectedItem(FNameProperty* InProperty, const TArray<TSharedPtr<FRigVMStringWithTag>>* InNameList) const
+TSharedPtr<FRigVMStringWithTag> FRigVMWrappedNodeDetailCustomization::GetCurrentlySelectedItem(const FNameProperty* InProperty, const TArray<TSharedPtr<FRigVMStringWithTag>>* InNameList) const
 {
 	FString CurrentItem = GetNameListText(InProperty).ToString();
 	for (const TSharedPtr<FRigVMStringWithTag>& Item : *InNameList)
@@ -1887,7 +2115,7 @@ TSharedPtr<FRigVMStringWithTag> FRigVMWrappedNodeDetailCustomization::GetCurrent
 }
 
 
-void FRigVMWrappedNodeDetailCustomization::SetNameListText(const FText& NewTypeInValue, ETextCommit::Type, FNameProperty* InProperty, TSharedRef<IPropertyUtilities> PropertyUtilities)
+void FRigVMWrappedNodeDetailCustomization::SetNameListText(const FText& NewTypeInValue, ETextCommit::Type, const FNameProperty* InProperty, TSharedRef<IPropertyUtilities> PropertyUtilities)
 {
 	URigVMGraph* Graph = NodesBeingCustomized[0]->GetGraph();
 	URigVMController* Controller = BlueprintBeingCustomized->GetController(Graph);
@@ -1906,7 +2134,7 @@ void FRigVMWrappedNodeDetailCustomization::SetNameListText(const FText& NewTypeI
 	Controller->CloseUndoBracket();
 }
 
-void FRigVMWrappedNodeDetailCustomization::OnNameListChanged(TSharedPtr<FRigVMStringWithTag> NewSelection, ESelectInfo::Type SelectInfo, FNameProperty* InProperty, TSharedRef<IPropertyUtilities> PropertyUtilities)
+void FRigVMWrappedNodeDetailCustomization::OnNameListChanged(TSharedPtr<FRigVMStringWithTag> NewSelection, ESelectInfo::Type SelectInfo, const FNameProperty* InProperty, TSharedRef<IPropertyUtilities> PropertyUtilities)
 {
 	if (SelectInfo != ESelectInfo::Direct)
 	{
@@ -1915,7 +2143,7 @@ void FRigVMWrappedNodeDetailCustomization::OnNameListChanged(TSharedPtr<FRigVMSt
 	}
 }
 
-void FRigVMWrappedNodeDetailCustomization::OnNameListComboBox(FNameProperty* InProperty, const TArray<TSharedPtr<FRigVMStringWithTag>>* InNameList)
+void FRigVMWrappedNodeDetailCustomization::OnNameListComboBox(const FNameProperty* InProperty, const TArray<TSharedPtr<FRigVMStringWithTag>>* InNameList)
 {
 	TSharedPtr<SRigVMGraphPinNameListValueWidget> Widget = NameListWidgets.FindChecked(InProperty->GetFName());
 	const TSharedPtr<FRigVMStringWithTag> CurrentlySelected = GetCurrentlySelectedItem(InProperty, InNameList);
diff --git a/Engine/Plugins/Runtime/RigVM/Source/RigVMEditor/Public/Editor/RigVMGraphDetailCustomization.h b/Engine/Plugins/Runtime/RigVM/Source/RigVMEditor/Public/Editor/RigVMGraphDetailCustomization.h
--- a/Engine/Plugins/Runtime/RigVM/Source/RigVMEditor/Public/Editor/RigVMGraphDetailCustomization.h
+++ b/Engine/Plugins/Runtime/RigVM/Source/RigVMEditor/Public/Editor/RigVMGraphDetailCustomization.h
@@ -237,11 +237,11 @@ public:
 	virtual void CustomizeDetails(IDetailLayoutBuilder& DetailLayout) override;
 
 	TSharedRef<SWidget> MakeNameListItemWidget(TSharedPtr<FRigVMStringWithTag> InItem);
-	FText GetNameListText(FNameProperty* InProperty) const;
-	TSharedPtr<FRigVMStringWithTag> GetCurrentlySelectedItem(FNameProperty* InProperty, const TArray<TSharedPtr<FRigVMStringWithTag>>* InNameList) const;
-	void SetNameListText(const FText& NewTypeInValue, ETextCommit::Type, FNameProperty* InProperty, TSharedRef<IPropertyUtilities> PropertyUtilities);
-	void OnNameListChanged(TSharedPtr<FRigVMStringWithTag> NewSelection, ESelectInfo::Type SelectInfo, FNameProperty* InProperty, TSharedRef<IPropertyUtilities> PropertyUtilities);
-	void OnNameListComboBox(FNameProperty* InProperty, const TArray<TSharedPtr<FRigVMStringWithTag>>* InNameList);
+	FText GetNameListText(const FNameProperty* InProperty) const;
+	TSharedPtr<FRigVMStringWithTag> GetCurrentlySelectedItem(const FNameProperty* InProperty, const TArray<TSharedPtr<FRigVMStringWithTag>>* InNameList) const;
+	void SetNameListText(const FText& NewTypeInValue, ETextCommit::Type, const FNameProperty* InProperty, TSharedRef<IPropertyUtilities> PropertyUtilities);
+	void OnNameListChanged(TSharedPtr<FRigVMStringWithTag> NewSelection, ESelectInfo::Type SelectInfo, const FNameProperty* InProperty, TSharedRef<IPropertyUtilities> PropertyUtilities);
+	void OnNameListComboBox(const FNameProperty* InProperty, const TArray<TSharedPtr<FRigVMStringWithTag>>* InNameList);
 	void CustomizeLiveValues(IDetailLayoutBuilder& DetailLayout);
 
 	URigVMBlueprint* BlueprintBeingCustomized;
