RigVM: Improving array pin handling through details panel

#rnx

diff --git a/Engine/Plugins/Runtime/RigVM/Source/RigVMDeveloper/Private/RigVMModel/RigVMController.cpp b/Engine/Plugins/Runtime/RigVM/Source/RigVMDeveloper/Private/RigVMModel/RigVMController.cpp
--- a/Engine/Plugins/Runtime/RigVM/Source/RigVMDeveloper/Private/RigVMModel/RigVMController.cpp
+++ b/Engine/Plugins/Runtime/RigVM/Source/RigVMDeveloper/Private/RigVMModel/RigVMController.cpp
@@ -10206,6 +10206,13 @@ URigVMPin* URigVMController::InsertArrayPin(URigVMPin* ArrayPin, int32 InIndex,
 	{
 		URigVMPin* ExistingPin = ArrayPin->GetSubPins()[ExistingIndex];
 		RenameObject(ExistingPin, *FString::FormatAsNumber(ExistingIndex + 1));
+
+		// temporarily force a the index based display name
+		// so that the ed graph node can pick that up during handling the notify
+		// (the default is that it's based the index in the sub pin array,
+		// which at this point hasn't changed)
+		const TGuardValue<FName> ForcePinIndexDisplayName(ExistingPin->DisplayName, *FString::FromInt(ExistingIndex+1));
+		Notify(ERigVMGraphNotifType::PinRenamed, ExistingPin);
 	}
 
 	URigVMPin* Pin = NewObject<URigVMPin>(ArrayPin, *FString::FormatAsNumber(InIndex));
diff --git a/Engine/Plugins/Runtime/RigVM/Source/RigVMEditor/Private/Editor/RigVMEditor.cpp b/Engine/Plugins/Runtime/RigVM/Source/RigVMEditor/Private/Editor/RigVMEditor.cpp
--- a/Engine/Plugins/Runtime/RigVM/Source/RigVMEditor/Private/Editor/RigVMEditor.cpp
+++ b/Engine/Plugins/Runtime/RigVM/Source/RigVMEditor/Private/Editor/RigVMEditor.cpp
@@ -1071,7 +1071,7 @@ void FRigVMEditor::Compile()
 		FString LastDebuggedObjectName = GetCustomDebugObjectLabel(RigVMBlueprint->GetObjectBeingDebugged());
 		RigVMBlueprint->SetObjectBeingDebugged(nullptr);
 
-		TArray< TWeakObjectPtr<UObject> > SelectedObjects = Inspector->GetSelectedObjects();
+		TArray< TWeakObjectPtr<UObject> > SelectedObjects = GetSelectedObjects();
 
 		if (URigVMHost* RigVMHost = GetRigVMHost())
 		{
@@ -1724,7 +1724,7 @@ void FRigVMEditor::HandleModifiedEvent(ERigVMGraphNotifType InNotifType, URigVMG
 				if(!DefaultValue.IsEmpty())
 				{
 					// sync the value change with the unit(s) displayed 
-					TArray< TWeakObjectPtr<UObject> > SelectedObjects = Inspector->GetSelectedObjects();
+					TArray< TWeakObjectPtr<UObject> > SelectedObjects = GetSelectedObjects();
 					for (TWeakObjectPtr<UObject> SelectedObject : SelectedObjects)
 					{
 						if (SelectedObject.IsValid())
@@ -2201,31 +2201,50 @@ void FRigVMEditor::OnWrappedPropertyChangedChainEvent(URigVMDetailsViewWrapperOb
 				check(InPropertyPath.StartsWith(RootPinNameString));
 				FString RemainingPropertyPath = InPropertyPath.Mid(RootPinNameString.Len());
 				RemainingPropertyPath.RemoveFromStart(TEXT("->"));
-				RemainingPropertyPath.ReplaceInline(TEXT("->"), TEXT("."));
-
-				// traverse each property one by one to make sure the expected pin exists.
-				// this may not be the case for an array element.
-				while(!RemainingPropertyPath.IsEmpty())
+                RemainingPropertyPath.ReplaceInline(TEXT("->"), TEXT("."));
+				RemainingPropertyPath.RemoveFromStart(TEXT("["));
+				RemainingPropertyPath.RemoveFromEnd(TEXT("]"));
+				RemainingPropertyPath.ReplaceInline(TEXT("["), TEXT("."));
+				RemainingPropertyPath.ReplaceInline(TEXT("]"), TEXT(""));
+				
+				if(InPropertyChangedChainEvent.ChangeType == EPropertyChangeType::ArrayAdd)
+				{
+					PinPath = URigVMPin::JoinPinPath(PinPath, RemainingPropertyPath);
+					
+					const FRigVMPropertyPath PropertyTraverser(Property, RemainingPropertyPath);
+					PropertyStorage = PropertyTraverser.GetData<uint8>(PropertyStorage, Property);
+					Property = PropertyTraverser.GetTailProperty();
+				}
+				else if((InPropertyChangedChainEvent.ChangeType == EPropertyChangeType::ArrayRemove) ||
+					(InPropertyChangedChainEvent.ChangeType == EPropertyChangeType::ArrayClear) ||
+					(InPropertyChangedChainEvent.ChangeType == EPropertyChangeType::Duplicate))
+				{
+					PinPath = URigVMPin::JoinPinPath(PinPath, RemainingPropertyPath);
+				}
+				else
 				{
-					FString Left = RemainingPropertyPath, Right;
-					(void)URigVMPin::SplitPinPathAtStart(RemainingPropertyPath, Left, Right);
+					// traverse each property one by one to make sure the expected pin exists.
+					// this may not be the case for an array element.
+					while(!RemainingPropertyPath.IsEmpty())
+					{
+						FString Left = RemainingPropertyPath, Right;
+						(void)URigVMPin::SplitPinPathAtStart(RemainingPropertyPath, Left, Right);
 
-					FString NewPinPath = URigVMPin::JoinPinPath(PinPath, Left);
-					NewPinPath.ReplaceInline(TEXT("["), TEXT(""));
-					NewPinPath.ReplaceInline(TEXT("]"), TEXT(""));
+						const FString NewPinPath = URigVMPin::JoinPinPath(PinPath, Left);
 
-					// this may be an array pin which doesn't exist yet
-					if(!Controller->GetGraph()->FindPin(NewPinPath))
-					{
-						break;
-					}
+						// this may be an array pin which doesn't exist yet
+						if(!Controller->GetGraph()->FindPin(NewPinPath))
+						{
+							break;
+						}
 
-					const FRigVMPropertyPath PropertyTraverser(Property, Left);
-					PropertyStorage = PropertyTraverser.GetData<uint8>(PropertyStorage, Property);
-					Property = PropertyTraverser.GetTailProperty();
-					PinPath = NewPinPath;
+						const FRigVMPropertyPath PropertyTraverser(Property, Left);
+						PropertyStorage = PropertyTraverser.GetData<uint8>(PropertyStorage, Property);
+						Property = PropertyTraverser.GetTailProperty();
+						PinPath = NewPinPath;
 
-					RemainingPropertyPath = Right;
+						RemainingPropertyPath = Right;
+					}
 				}
 			}
 		}
@@ -2233,31 +2252,56 @@ void FRigVMEditor::OnWrappedPropertyChangedChainEvent(URigVMDetailsViewWrapperOb
 		if (Property)
 		{
 			FString DefaultValue;
-			if(PropertyStorage == nullptr)
+			
+			if((InPropertyChangedChainEvent.ChangeType != EPropertyChangeType::ArrayRemove) &&
+				(InPropertyChangedChainEvent.ChangeType != EPropertyChangeType::ArrayClear) &&
+				(InPropertyChangedChainEvent.ChangeType != EPropertyChangeType::Duplicate))
 			{
-				// this may happen when we remove the last element from an array.
-				// in that case just empty the array itself.
-				if(const FProperty* ParentProperty = Property->GetOwnerProperty())
+				if(PropertyStorage == nullptr)
 				{
-					if(ParentProperty->IsA<FArrayProperty>())
+					// this may happen when we remove the last element from an array.
+					// in that case just empty the array itself.
+					if(const FProperty* ParentProperty = Property->GetOwnerProperty())
 					{
-						DefaultValue = TEXT("()");
-						FString Left, Right;
-						verify(URigVMPin::SplitPinPathAtEnd(PinPath, Left, Right));
-						PinPath = Left;
+						if(ParentProperty->IsA<FArrayProperty>())
+						{
+							DefaultValue = TEXT("()");
+							FString Left, Right;
+							verify(URigVMPin::SplitPinPathAtEnd(PinPath, Left, Right));
+							PinPath = Left;
+						}
 					}
 				}
-			}
-			else
-			{
-				DefaultValue = FRigVMStruct::ExportToFullyQualifiedText(Property, PropertyStorage);
+				else
+				{
+					DefaultValue = FRigVMStruct::ExportToFullyQualifiedText(Property, PropertyStorage);
+				}
 			}
 			
 			if(Property->IsA<FStrProperty>() || Property->IsA<FNameProperty>())
 			{
 				DefaultValue.TrimCharInline(TEXT('\"'), nullptr);
 			}
-			if (!DefaultValue.IsEmpty())
+			
+			if(InPropertyChangedChainEvent.ChangeType == EPropertyChangeType::ArrayAdd)
+			{
+				FString ArrayPinPath, ArrayElementIndex;
+				verify(URigVMPin::SplitPinPathAtEnd(PinPath, ArrayPinPath, ArrayElementIndex));
+				Controller->AddArrayPin(ArrayPinPath, DefaultValue, true, true);
+			}
+			else if(InPropertyChangedChainEvent.ChangeType == EPropertyChangeType::ArrayRemove)
+			{
+				Controller->RemoveArrayPin(PinPath, true, true);
+			}
+			else if(InPropertyChangedChainEvent.ChangeType == EPropertyChangeType::ArrayClear)
+			{
+				Controller->ClearArrayPin(PinPath, true, true);
+			}
+			else if(InPropertyChangedChainEvent.ChangeType == EPropertyChangeType::Duplicate)
+			{
+				Controller->DuplicateArrayPin(PinPath, true, true);
+			}
+			else if (!DefaultValue.IsEmpty())
 			{
 				const bool bInteractive = InPropertyChangedChainEvent.ChangeType == EPropertyChangeType::Interactive;
 				Controller->SetPinDefaultValue(PinPath, DefaultValue, true, !bInteractive, true, !bInteractive);
@@ -2642,6 +2686,27 @@ void FRigVMEditor::OnCreateComment()
 	}
 }
 
+TArray<TWeakObjectPtr<UObject>> FRigVMEditor::GetSelectedObjects() const
+{
+	// if the inspector shows wrapped objects - look in that array instead.
+	// with recent weak object pointer changes on the property detail view
+	// we cannot rely on the GetSelectedObjects being valid after blueprint compilation.
+	if(WrapperObjects.Num() == Inspector->GetSelectedObjects().Num())
+	{
+		TArray<TWeakObjectPtr<UObject>> WeakWrapperObjects;
+		for(const TStrongObjectPtr<URigVMDetailsViewWrapperObject>& WrapperObjectPtr : WrapperObjects)
+		{
+			URigVMDetailsViewWrapperObject* WrapperObject = WrapperObjectPtr.Get();
+			if(IsValid(WrapperObject->GetSubject()))
+			{
+				WeakWrapperObjects.Add(WrapperObject);
+			}
+		}
+		return WeakWrapperObjects;
+	}
+	return Inspector->GetSelectedObjects();
+}
+
 void FRigVMEditor::SetDetailObjects(const TArray<UObject*>& InObjects)
 {
 	SetDetailObjects(InObjects, true);
@@ -2884,7 +2949,7 @@ void FRigVMEditor::SetDetailViewForFocusedGraph()
 void FRigVMEditor::SetDetailViewForLocalVariable()
 {
 	FName VariableName;
-	TArray< TWeakObjectPtr<UObject> > SelectedObjects = Inspector->GetSelectedObjects();
+	TArray< TWeakObjectPtr<UObject> > SelectedObjects = GetSelectedObjects();
 	for (TWeakObjectPtr<UObject> SelectedObject : SelectedObjects)
 	{
 		if (SelectedObject.IsValid())
@@ -2931,7 +2996,7 @@ bool FRigVMEditor::DetailViewShowsAnyRigUnit() const
 		return true;
 	}
 
-	const TArray< TWeakObjectPtr<UObject> >& SelectedObjects = Inspector->GetSelectedObjects();
+	const TArray< TWeakObjectPtr<UObject> >& SelectedObjects = GetSelectedObjects();
 	for (TWeakObjectPtr<UObject> SelectedObject : SelectedObjects)
 	{
 		if (SelectedObject.IsValid())
@@ -2957,7 +3022,7 @@ bool FRigVMEditor::DetailViewShowsLocalVariable() const
 
 bool FRigVMEditor::DetailViewShowsStruct(UScriptStruct* InStruct) const
 {
-	TArray< TWeakObjectPtr<UObject> > SelectedObjects = Inspector->GetSelectedObjects();
+	TArray< TWeakObjectPtr<UObject> > SelectedObjects = GetSelectedObjects();
 	for (TWeakObjectPtr<UObject> SelectedObject : SelectedObjects)
 	{
 		if (SelectedObject.IsValid())
diff --git a/Engine/Plugins/Runtime/RigVM/Source/RigVMEditor/Public/Editor/RigVMEditor.h b/Engine/Plugins/Runtime/RigVM/Source/RigVMEditor/Public/Editor/RigVMEditor.h
--- a/Engine/Plugins/Runtime/RigVM/Source/RigVMEditor/Public/Editor/RigVMEditor.h
+++ b/Engine/Plugins/Runtime/RigVM/Source/RigVMEditor/Public/Editor/RigVMEditor.h
@@ -238,6 +238,7 @@ protected:
 
 	bool IsDetailsPanelRefreshSuspended() const { return bSuspendDetailsPanelRefresh; }
 	bool& GetSuspendDetailsPanelRefreshFlag() { return bSuspendDetailsPanelRefresh; }
+	TArray<TWeakObjectPtr<UObject>> GetSelectedObjects() const;
 	virtual void SetDetailObjects(const TArray<UObject*>& InObjects);
 	virtual void SetDetailObjects(const TArray<UObject*>& InObjects, bool bChangeUISelectionState);
 	virtual void SetMemoryStorageDetails(const TArray<FRigVMMemoryStorageStruct*>& InStructs);
