AnimNext RigVM collapse node and functions support
RigVM external model interface
One file per asset support for models

#rnx

diff --git a/Engine/Plugins/Experimental/AnimNext/Source/AnimNextEditor/Private/AnimNextEditorModule.cpp b/Engine/Plugins/Experimental/AnimNext/Source/AnimNextEditor/Private/AnimNextEditorModule.cpp
--- a/Engine/Plugins/Experimental/AnimNext/Source/AnimNextEditor/Private/AnimNextEditorModule.cpp
+++ b/Engine/Plugins/Experimental/AnimNext/Source/AnimNextEditor/Private/AnimNextEditorModule.cpp
@@ -227,6 +227,9 @@ void FModule::StartupModule()
 	Workspace::IWorkspaceEditorModule& WorkspaceModule = FModuleManager::Get().LoadModuleChecked<Workspace::IWorkspaceEditorModule>("WorkspaceEditor");
 	const TSharedPtr<FAnimNextGraphItemDetails> GraphItemDetails = MakeShareable<FAnimNextGraphItemDetails>(new FAnimNextGraphItemDetails());
 	WorkspaceModule.RegisterWorkspaceItemDetails(Workspace::FOutlinerItemDetailsId(FAnimNextGraphOutlinerData::StaticStruct()->GetFName()), StaticCastSharedPtr<UE::Workspace::IWorkspaceOutlinerItemDetails>(GraphItemDetails));
+	WorkspaceModule.RegisterWorkspaceItemDetails(Workspace::FOutlinerItemDetailsId(FAnimNextCollapseGraphOutlinerData::StaticStruct()->GetFName()), StaticCastSharedPtr<UE::Workspace::IWorkspaceOutlinerItemDetails>(GraphItemDetails));
+	WorkspaceModule.RegisterWorkspaceItemDetails(Workspace::FOutlinerItemDetailsId(FAnimNextGraphFunctionOutlinerData::StaticStruct()->GetFName()), StaticCastSharedPtr<UE::Workspace::IWorkspaceOutlinerItemDetails>(GraphItemDetails));
+
 	FAnimNextGraphItemDetails::RegisterToolMenuExtensions();
 }
 
@@ -268,6 +271,8 @@ void FModule::ShutdownModule()
 	{
 		Workspace::IWorkspaceEditorModule& WorkspaceModule = FModuleManager::Get().LoadModuleChecked<Workspace::IWorkspaceEditorModule>("WorkspaceEditor");
 		WorkspaceModule.UnregisterWorkspaceItemDetails(Workspace::FOutlinerItemDetailsId(FAnimNextGraphOutlinerData::StaticStruct()->GetFName()));
+		WorkspaceModule.UnregisterWorkspaceItemDetails(Workspace::FOutlinerItemDetailsId(FAnimNextCollapseGraphOutlinerData::StaticStruct()->GetFName()));
+		WorkspaceModule.UnregisterWorkspaceItemDetails(Workspace::FOutlinerItemDetailsId(FAnimNextGraphFunctionOutlinerData::StaticStruct()->GetFName()));
 		FAnimNextGraphItemDetails::UnregisterToolMenuExtensions();
 	}
 
@@ -786,6 +791,40 @@ void FModule::RegisterWorkspaceDocumentTypes(Workspace::IWorkspaceEditorModule&
 
 		InContext.WorkspaceEditor->SetDetailsObjects(NewSelection.Array());
 	});
+	GraphArgs.OnNodeDoubleClicked = Workspace::FOnNodeDoubleClicked::CreateLambda([](const Workspace::FWorkspaceEditorContext& InContext, const UEdGraphNode* InNode)
+	{
+		if (const URigVMEdGraphNode* RigVMEdGraphNode = Cast<URigVMEdGraphNode>(InNode))
+		{
+			const URigVMNode* ModelNode = RigVMEdGraphNode->GetModelNode();
+
+			if (const URigVMLibraryNode* LibraryNode = Cast<URigVMLibraryNode>(ModelNode))
+			{
+				URigVMGraph* ContainedGraph = LibraryNode->GetContainedGraph();
+
+				if (const URigVMFunctionReferenceNode* FunctionReferenceNode = Cast<URigVMFunctionReferenceNode>(LibraryNode))
+				{
+					if (URigVMLibraryNode* ReferencedNode = FunctionReferenceNode->LoadReferencedNode())
+					{
+						ContainedGraph = ReferencedNode->GetContainedGraph();
+					}
+				}
+
+				if (ContainedGraph)
+				{
+					if (TSharedPtr<Workspace::IWorkspaceEditor> WorkspaceEditor = InContext.WorkspaceEditor)
+					{
+						if (IRigVMClientHost* RigVMClientHost = ContainedGraph->GetImplementingOuter<IRigVMClientHost>())
+						{
+							if (UObject* EditorObject = RigVMClientHost->GetEditorObjectForRigVMGraph(ContainedGraph))
+							{
+								WorkspaceEditor->OpenObjects({ EditorObject });
+							}
+						}
+					}
+				}
+			}
+		}
+	});
 
 	Workspace::FObjectDocumentArgs GraphDocumentArgs = WorkspaceEditorModule.CreateGraphDocumentArgs(GraphArgs);
 	Workspace::FOnMakeDocumentWidget WorkspaceMakeDocumentWidgetDelegate = GraphDocumentArgs.OnMakeDocumentWidget;
@@ -831,66 +870,118 @@ void FModule::RegisterWorkspaceDocumentTypes(Workspace::IWorkspaceEditorModule&
 
 	GraphDocumentArgs.OnGetDocumentBreadcrumbTrail = Workspace::FOnGetDocumentBreadcrumbTrail::CreateLambda([](const Workspace::FWorkspaceEditorContext& InContext, TArray<TSharedPtr<Workspace::FWorkspaceBreadcrumb>>& OutBreadcrumbs)
 	{
-		if (const URigVMEdGraph* RigVMEdGraph = Cast<URigVMEdGraph>(InContext.Object))
+		if (URigVMEdGraph* EdGraph = Cast<URigVMEdGraph>(InContext.Object))
 		{
-			if (const UAnimNextRigVMAssetEntry* OuterAssetEntry = CastChecked<UAnimNextRigVMAssetEntry>(RigVMEdGraph->GetOuter()))
+			if (const UAnimNextModule_EditorData* GraphEditorData = EdGraph->GetTypedOuter<UAnimNextModule_EditorData>())
 			{
-				const IAnimNextRigVMGraphInterface* GraphInterface = Cast<IAnimNextRigVMGraphInterface>(OuterAssetEntry);
-				const TSharedPtr<Workspace::FWorkspaceBreadcrumb>& GraphCrumb = OutBreadcrumbs.Add_GetRef(MakeShared<Workspace::FWorkspaceBreadcrumb>());
+				// Iterate model tree, so we display all graph parents until we reach the Entry
+				URigVMGraph* ModelGraph = EdGraph->GetModel();
+				while (ModelGraph != nullptr)
+				{
+					URigVMEdGraph* RigVMEdGraph = Cast<URigVMEdGraph>(GraphEditorData->GetEditorObjectForRigVMGraph(ModelGraph));
 
-				TWeakObjectPtr<const URigVMEdGraph> WeakGraph = RigVMEdGraph;
-				GraphCrumb->OnGetLabel = Workspace::FWorkspaceBreadcrumb::FOnGetBreadcrumbLabel::CreateLambda([GraphName = OuterAssetEntry->GetEntryName()]{ return FText::FromName(GraphName); });
-				GraphCrumb->CanSave = Workspace::FWorkspaceBreadcrumb::FCanSaveBreadcrumb::CreateLambda(
-					[WeakGraph]
+					if (GraphEditorData && GraphEditorData->GetLocalFunctionLibrary() != RigVMEdGraph->GetModel())
 					{
-						if (const URigVMEdGraph* Graph = WeakGraph.Get())
-						{
-							return Graph->GetPackage()->IsDirty();
-						}
-						return false;
-					}
-				);
+						const TSharedPtr<Workspace::FWorkspaceBreadcrumb>& GraphCrumb = OutBreadcrumbs.Add_GetRef(MakeShared<Workspace::FWorkspaceBreadcrumb>());
 
-				GraphCrumb->OnSave = Workspace::FWorkspaceBreadcrumb::FOnSaveBreadcrumb::CreateLambda(
-				[WeakGraph]
-					{
-						if (const URigVMEdGraph* Graph = WeakGraph.Get())
+						TWeakObjectPtr<URigVMEdGraph> WeakEdGraph = RigVMEdGraph;
+						TWeakPtr<Workspace::IWorkspaceEditor> WeakWorkspaceEditor = InContext.WorkspaceEditor;
+						TWeakObjectPtr<const UAnimNextModule_EditorData> WeakGraphEditorData = GraphEditorData;
+
+						FText GraphName;
+						if (WeakEdGraph.IsValid())
 						{
-							FEditorFileUtils::PromptForCheckoutAndSave({Graph->GetPackage()}, false, /*bPromptToSave=*/ false);
+							if (URigVMCollapseNode* CollapseNode = Cast<URigVMCollapseNode>(WeakEdGraph->GetModel()->GetOuter()))
+							{
+								GraphName = FText::FromName(CollapseNode->GetFName());
+							}
+							else if (URigVMFunctionReferenceNode* FunctionReferenceNode = Cast<URigVMFunctionReferenceNode>(WeakEdGraph->GetModel()->GetOuter()))
+							{
+								if (URigVMLibraryNode* ReferencedNode = Cast<URigVMLibraryNode>(FunctionReferenceNode->GetReferencedFunctionHeader().LibraryPointer.GetNodeSoftPath().ResolveObject()))
+								{
+									GraphName = FText::FromName(ReferencedNode->GetFName());
+								}
+							}
+
+							if (GraphName.IsEmpty())
+							{
+								if (WeakGraphEditorData.IsValid() && WeakGraphEditorData->GetLocalFunctionLibrary() == WeakEdGraph->GetModel())
+								{
+									GraphName = UE::AnimNext::UncookedOnly::FUtils::GetFunctionLibraryDisplayName();
+								}
+								else
+								{
+									GraphName = FText::FromName(WeakEdGraph->GetFName());
+								}
+							}
 						}
-					}
-				);
-	
-				if (UAnimNextModule_EditorData* OuterEditorData = CastChecked<UAnimNextModule_EditorData>(OuterAssetEntry->GetOuter()))
-				{
-					if(UAnimNextModule* OuterGraph = UncookedOnly::FUtils::GetGraph(OuterEditorData))
-					{
-						const TSharedPtr<Workspace::FWorkspaceBreadcrumb>& OuterGraphCrumb = OutBreadcrumbs.Add_GetRef(MakeShared<Workspace::FWorkspaceBreadcrumb>());
-						TWeakObjectPtr<UAnimNextModule> WeakOuterGraph = OuterGraph;
-						TWeakPtr<Workspace::IWorkspaceEditor> WeakWorkspaceEditor = InContext.WorkspaceEditor;
-						OuterGraphCrumb->OnGetLabel = Workspace::FWorkspaceBreadcrumb::FOnGetBreadcrumbLabel::CreateLambda([GraphName = OuterGraph->GetFName()]{ return FText::FromName(GraphName); });
-						OuterGraphCrumb->OnClicked = Workspace::FWorkspaceBreadcrumb::FOnBreadcrumbClicked::CreateLambda(
-							[WeakOuterGraph, WeakWorkspaceEditor]
+
+						GraphCrumb->OnGetLabel = Workspace::FWorkspaceBreadcrumb::FOnGetBreadcrumbLabel::CreateLambda(
+							[GraphName]
+							{
+								return GraphName;
+							});
+						GraphCrumb->CanSave = Workspace::FWorkspaceBreadcrumb::FCanSaveBreadcrumb::CreateLambda(
+							[WeakEdGraph]
+							{
+								if (const URigVMEdGraph* Graph = WeakEdGraph.Get())
+								{
+									return Graph->GetPackage()->IsDirty();
+								}
+								return false;
+							}
+						);
+						GraphCrumb->OnClicked = Workspace::FWorkspaceBreadcrumb::FOnBreadcrumbClicked::CreateLambda(
+							[WeakGraphEditorData, WeakEdGraph, WeakWorkspaceEditor]
 							{
 								if (const TSharedPtr<Workspace::IWorkspaceEditor> SharedWorkspaceEditor = WeakWorkspaceEditor.Pin())
 								{
-									SharedWorkspaceEditor->OpenObjects({WeakOuterGraph.Get()});
+									SharedWorkspaceEditor->OpenObjects({ WeakEdGraph.Get() });
 								}
 							}
 						);
-						OuterGraphCrumb->CanSave = Workspace::FWorkspaceBreadcrumb::FCanSaveBreadcrumb::CreateLambda(
-							[WeakOuterGraph]
+						GraphCrumb->OnSave = Workspace::FWorkspaceBreadcrumb::FOnSaveBreadcrumb::CreateLambda(
+							[WeakEdGraph]
 							{
-								if (UAnimNextModule* Module = WeakOuterGraph.Get())
+								if (const URigVMEdGraph* Graph = WeakEdGraph.Get())
 								{
-									return Module->GetPackage()->IsDirty();
+									FEditorFileUtils::PromptForCheckoutAndSave({ Graph->GetPackage() }, false, /*bPromptToSave=*/ false);
 								}
-
-								return false;
 							}
 						);
+					}
+
+					ModelGraph = ModelGraph->GetTypedOuter<URigVMGraph>();
+				}
+
+				// Display the Module
+				if(UAnimNextModule* OuterGraph = UncookedOnly::FUtils::GetGraph(GraphEditorData))
+				{
+					const TSharedPtr<Workspace::FWorkspaceBreadcrumb>& OuterGraphCrumb = OutBreadcrumbs.Add_GetRef(MakeShared<Workspace::FWorkspaceBreadcrumb>());
+					TWeakObjectPtr<UAnimNextModule> WeakOuterGraph = OuterGraph;
+					TWeakPtr<Workspace::IWorkspaceEditor> WeakWorkspaceEditor = InContext.WorkspaceEditor;
+					OuterGraphCrumb->OnGetLabel = Workspace::FWorkspaceBreadcrumb::FOnGetBreadcrumbLabel::CreateLambda([GraphName = OuterGraph->GetFName()]{ return FText::FromName(GraphName); });
+					OuterGraphCrumb->OnClicked = Workspace::FWorkspaceBreadcrumb::FOnBreadcrumbClicked::CreateLambda(
+						[WeakOuterGraph, WeakWorkspaceEditor]
+						{
+							if (const TSharedPtr<Workspace::IWorkspaceEditor> SharedWorkspaceEditor = WeakWorkspaceEditor.Pin())
+							{
+								SharedWorkspaceEditor->OpenObjects({WeakOuterGraph.Get()});
+							}
+						}
+					);
+					OuterGraphCrumb->CanSave = Workspace::FWorkspaceBreadcrumb::FCanSaveBreadcrumb::CreateLambda(
+						[WeakOuterGraph]
+						{
+							if (UAnimNextModule* Module = WeakOuterGraph.Get())
+							{
+								return Module->GetPackage()->IsDirty();
+							}
 
-						OuterGraphCrumb->OnSave = Workspace::FWorkspaceBreadcrumb::FOnSaveBreadcrumb::CreateLambda(
+							return false;
+						}
+					);
+					OuterGraphCrumb->OnSave = Workspace::FWorkspaceBreadcrumb::FOnSaveBreadcrumb::CreateLambda(
 						[WeakOuterGraph]
 							{
 								if (UAnimNextModule* Module = WeakOuterGraph.Get())
@@ -899,8 +990,6 @@ void FModule::RegisterWorkspaceDocumentTypes(Workspace::IWorkspaceEditorModule&
 								}
 							}
 						);
-						
-					}
 				}
 			}
 		}
diff --git a/Engine/Plugins/Experimental/AnimNext/Source/AnimNextEditor/Private/Graph/AnimNextGraphItemDetails.cpp b/Engine/Plugins/Experimental/AnimNext/Source/AnimNextEditor/Private/Graph/AnimNextGraphItemDetails.cpp
--- a/Engine/Plugins/Experimental/AnimNext/Source/AnimNextEditor/Private/Graph/AnimNextGraphItemDetails.cpp
+++ b/Engine/Plugins/Experimental/AnimNext/Source/AnimNextEditor/Private/Graph/AnimNextGraphItemDetails.cpp
@@ -11,12 +11,14 @@
 #include "WorkspaceItemMenuContext.h"
 #include "IWorkspaceEditor.h"
 #include "RigVMModel/RigVMClient.h"
+#include "EdGraph/RigVMEdGraph.h"
 #include "ToolMenus.h"
 
 #define LOCTEXT_NAMESPACE "FAnimNextGraphItemDetails"
 
 namespace UE::AnimNext::Editor
 {
+
 void FAnimNextGraphItemDetails::HandleDoubleClick(const FToolMenuContext& ToolMenuContext) const
 {
 	const UWorkspaceItemMenuContext* WorkspaceItemContext = ToolMenuContext.FindContext<UWorkspaceItemMenuContext>();
@@ -26,22 +28,42 @@ void FAnimNextGraphItemDetails::HandleDoubleClick(const FToolMenuContext& ToolMe
 		if(const TSharedPtr<UE::Workspace::IWorkspaceEditor> WorkspaceEditor = StaticCastSharedPtr<UE::Workspace::IWorkspaceEditor>(AssetEditorContext->Toolkit.Pin()))
 		{
 			const TInstancedStruct<FWorkspaceOutlinerItemData>& Data = WorkspaceItemContext->SelectedExports[0].GetData();
-			if (Data.IsValid() && Data.GetScriptStruct() == FAnimNextGraphOutlinerData::StaticStruct())
+			if (Data.IsValid())
 			{
-				const FAnimNextGraphOutlinerData& GraphData = Data.Get<FAnimNextGraphOutlinerData>();							
-				if (GraphData.GraphInterface)
+				if (Data.GetScriptStruct() == FAnimNextGraphOutlinerData::StaticStruct())
 				{
-					if (URigVMGraph* RigVMGraph = GraphData.GraphInterface->GetRigVMGraph())
+					const FAnimNextGraphOutlinerData& GraphData = Data.Get<FAnimNextGraphOutlinerData>();							
+					if (GraphData.GraphInterface)
 					{
-						if(const IRigVMClientHost* RigVMClientHost = RigVMGraph->GetImplementingOuter<IRigVMClientHost>())
+						if (URigVMGraph* RigVMGraph = GraphData.GraphInterface->GetRigVMGraph())
 						{
-							if(UObject* EditorObject = RigVMClientHost->GetEditorObjectForRigVMGraph(RigVMGraph))
+							if(const IRigVMClientHost* RigVMClientHost = RigVMGraph->GetImplementingOuter<IRigVMClientHost>())
 							{
-								WorkspaceEditor->OpenObjects({EditorObject});
+								if(UObject* EditorObject = RigVMClientHost->GetEditorObjectForRigVMGraph(RigVMGraph))
+								{
+									WorkspaceEditor->OpenObjects({EditorObject});
+								}
 							}
 						}
 					}
 				}
+				else if (Data.GetScriptStruct() == FAnimNextGraphFunctionOutlinerData::StaticStruct())
+				{
+					const FAnimNextGraphFunctionOutlinerData& GraphFunctionData = Data.Get<FAnimNextGraphFunctionOutlinerData>();
+					if (GraphFunctionData.EditorObject.IsValid())
+					{
+						WorkspaceEditor->OpenObjects({ GraphFunctionData.EditorObject.Get()});
+					}
+				}
+				else if (Data.GetScriptStruct() == FAnimNextCollapseGraphOutlinerData::StaticStruct())
+				{
+					const FAnimNextCollapseGraphOutlinerData& CollapseGraphData = Data.Get<FAnimNextCollapseGraphOutlinerData>();
+					if (CollapseGraphData.EditorObject.IsValid())
+					{
+						UObject* EditorObject = CollapseGraphData.EditorObject.Get();
+						WorkspaceEditor->OpenObjects({ EditorObject });
+					}
+				}
 			}
 		}
 	}
@@ -50,12 +72,31 @@ void FAnimNextGraphItemDetails::HandleDoubleClick(const FToolMenuContext& ToolMe
 UPackage* FAnimNextGraphItemDetails::GetPackage(const FWorkspaceOutlinerItemExport& Export) const 
 {
 	const TInstancedStruct<FWorkspaceOutlinerItemData>& Data = Export.GetData();
-	if (Data.IsValid() && Data.GetScriptStruct() == FAnimNextGraphOutlinerData::StaticStruct())
+	if (Data.IsValid())
 	{
-		const FAnimNextGraphOutlinerData& GraphData = Data.Get<FAnimNextGraphOutlinerData>();
-		if (GraphData.GraphInterface)
+		if (Data.GetScriptStruct() == FAnimNextGraphOutlinerData::StaticStruct())
 		{
-			return GraphData.GraphInterface.GetObject()->GetExternalPackage();
+			const FAnimNextGraphOutlinerData& GraphData = Data.Get<FAnimNextGraphOutlinerData>();
+			if (GraphData.GraphInterface)
+			{
+				return GraphData.GraphInterface.GetObject()->GetExternalPackage();
+			}
+		}
+		else if(Data.GetScriptStruct() == FAnimNextGraphFunctionOutlinerData::StaticStruct())
+		{
+			const FAnimNextGraphFunctionOutlinerData& GraphFunctionData = Data.Get<FAnimNextGraphFunctionOutlinerData>();
+			if (GraphFunctionData.EditorObject.IsValid())
+			{
+				return GraphFunctionData.EditorObject->GetPackage();
+			}
+		}
+		else if (Data.GetScriptStruct() == FAnimNextCollapseGraphOutlinerData::StaticStruct())
+		{
+			const FAnimNextCollapseGraphOutlinerData& CollapseGraphData = Data.Get<FAnimNextCollapseGraphOutlinerData>();
+			if (CollapseGraphData.EditorObject.IsValid())
+			{
+				return CollapseGraphData.EditorObject->GetPackage();
+			}
 		}
 	}
 	return nullptr;
@@ -80,53 +121,107 @@ void FAnimNextGraphItemDetails::RegisterToolMenuExtensions()
 				FToolMenuSection& Section = InMenu->AddSection("WorkspaceOutliner.ItemContextMenu.RootAsset", FText::FromString(TEXT("Animation Next")));
 				if(const TSharedPtr<UE::Workspace::IWorkspaceEditor> WorkspaceEditor = StaticCastSharedPtr<UE::Workspace::IWorkspaceEditor>(AssetEditorContext->Toolkit.Pin()))
 				{
-
 					TArray<FWorkspaceOutlinerItemExport> GraphExports;
 					Algo::TransformIf(WorkspaceItemContext->SelectedExports, GraphExports, [](const FWorkspaceOutlinerItemExport& Export)
-					{
-						return Export.GetData().IsValid() && Export.GetData().GetScriptStruct() == FAnimNextGraphOutlinerData::StaticStruct();
-					},
-					[](const FWorkspaceOutlinerItemExport& Export)
-					{
-						return Export;
-					});
+						{
+							return Export.GetData().IsValid() 
+								&& (Export.GetData().GetScriptStruct() == FAnimNextGraphOutlinerData::StaticStruct()
+									|| Export.GetData().GetScriptStruct() == FAnimNextGraphFunctionOutlinerData::StaticStruct()
+									|| Export.GetData().GetScriptStruct() == FAnimNextCollapseGraphOutlinerData::StaticStruct());
+						},
+						[](const FWorkspaceOutlinerItemExport& Export)
+						{
+							return Export;
+						});
 
 					if (GraphExports.Num() == WorkspaceItemContext->SelectedExports.Num() && GraphExports.Num() > 0)
 					{
 						TWeakPtr<UE::Workspace::IWorkspaceEditor> WeakWorkspaceEditor = WorkspaceEditor;
-						Section.AddMenuEntry(
-						TEXT("OpenGraphMenuEntry"),
-						FText::FormatOrdered(LOCTEXT("OpenGraphMenuEntryLabel", "Open Animation {0}|plural(one=Graph,other=Graphs)"), GraphExports.Num()),
-						FText::FormatOrdered(LOCTEXT("OpenGraphMenuEntryTooltip", "Open the selected Animation {0}|plural(one=Graph,other=Graphs)"), GraphExports.Num()),
-						FSlateIcon(FAppStyle::GetAppStyleSetName(), "GraphEditor.EventGraph_16x"),
-						FUIAction(
-							FExecuteAction::CreateWeakLambda(WorkspaceItemContext, [GraphExports, WeakWorkspaceEditor]()
+
+						TInstancedStruct<FWorkspaceOutlinerItemData>& Data = WorkspaceItemContext->SelectedExports[0].GetData();
+						if (Data.IsValid())
+						{
+							FName EntryText;
+							FText EntryLabelText;
+							FText EntryTooltipText;
+							FSlateIcon EntryIcon;
+
+							if (Data.GetScriptStruct() == FAnimNextGraphOutlinerData::StaticStruct())
 							{
-								if (TSharedPtr<UE::Workspace::IWorkspaceEditor> SharedWorkspaceEditor = WeakWorkspaceEditor.Pin())
-								{
-									TArray<UObject*> ObjectsToOpen;
-									for (const FWorkspaceOutlinerItemExport& Export : GraphExports)
+								EntryText = TEXT("OpenGraphMenuEntry");
+								EntryLabelText = FText::FormatOrdered(LOCTEXT("OpenGraphMenuEntryLabel", "Open {0}|plural(one=Animation Graph,other=Graphs)"), GraphExports.Num());
+								EntryTooltipText = FText::FormatOrdered(LOCTEXT("OpenGraphMenuEntryTooltip", "Open the selected {0}|plural(one=Animation Graph,other=Graphs)"), GraphExports.Num());
+								EntryIcon = FSlateIcon(FAppStyle::GetAppStyleSetName(), "GraphEditor.EventGraph_16x");
+							}
+							else if (Data.GetScriptStruct() == FAnimNextGraphFunctionOutlinerData::StaticStruct())
+							{
+								EntryText = TEXT("OpenFunctionActionEntry");
+								EntryLabelText = FText::FormatOrdered(LOCTEXT("OpenFunctionActionLabel", "Open {0}|plural(one=Function Graph,other=Grapsh)"), GraphExports.Num());
+								EntryTooltipText = FText::FormatOrdered(LOCTEXT("OpenFunctionActionTooltip", "Open the selected {0}|plural(one=Function Graph,other=Graphs)"), GraphExports.Num());
+								EntryIcon = FSlateIcon(FAppStyle::GetAppStyleSetName(), "GraphEditor.Function_16x");
+							}
+							else if (Data.GetScriptStruct() == FAnimNextCollapseGraphOutlinerData::StaticStruct())
+							{
+								EntryText = TEXT("OpenCollapseNodeActionEntry");
+								EntryLabelText = FText::FormatOrdered(LOCTEXT("OpenCollapseGraphActionLabel", "Open {0}|plural(one=Collapse Graph,other=Graphs)"), GraphExports.Num());
+								EntryTooltipText = FText::FormatOrdered(LOCTEXT("OpenCollapseGraphActionTooltip", "Open the selected {0}|plural(one=Collapse Graph,other=Graphs)"), GraphExports.Num());
+								EntryIcon = FSlateIcon(FAppStyle::GetAppStyleSetName(), "GraphEditor.SubGraph_16x");
+							}
+
+							Section.AddMenuEntry(
+								EntryText,
+								EntryLabelText,
+								EntryTooltipText,
+								EntryIcon,
+								FUIAction(
+									FExecuteAction::CreateWeakLambda(WorkspaceItemContext, [GraphExports, WeakWorkspaceEditor]()
 									{
-										const FAnimNextGraphOutlinerData& GraphData = Export.GetData().Get<FAnimNextGraphOutlinerData>();
-										if (GraphData.GraphInterface)
+										if (TSharedPtr<UE::Workspace::IWorkspaceEditor> SharedWorkspaceEditor = WeakWorkspaceEditor.Pin())
 										{
-											if (URigVMGraph* RigVMGraph = GraphData.GraphInterface->GetRigVMGraph())
+											TArray<UObject*> ObjectsToOpen;
+											for (const FWorkspaceOutlinerItemExport& Export : GraphExports)
 											{
-												if(const IRigVMClientHost* RigVMClientHost = RigVMGraph->GetImplementingOuter<IRigVMClientHost>())
+												const TInstancedStruct<FWorkspaceOutlinerItemData>& Data = Export.GetData();
+												if (Data.GetScriptStruct() == FAnimNextGraphOutlinerData::StaticStruct())
 												{
-													if(UObject* EditorObject = RigVMClientHost->GetEditorObjectForRigVMGraph(RigVMGraph))
+													const FAnimNextGraphOutlinerData& GraphData = Export.GetData().Get<FAnimNextGraphOutlinerData>();
+													if (GraphData.GraphInterface)
 													{
-														ObjectsToOpen.Add(EditorObject);
+														if (URigVMGraph* RigVMGraph = GraphData.GraphInterface->GetRigVMGraph())
+														{
+															if (const IRigVMClientHost* RigVMClientHost = RigVMGraph->GetImplementingOuter<IRigVMClientHost>())
+															{
+																if (UObject* EditorObject = RigVMClientHost->GetEditorObjectForRigVMGraph(RigVMGraph))
+																{
+																	ObjectsToOpen.Add(EditorObject);
+																}
+															}
+														}
+													}
+												}
+												else if (Data.GetScriptStruct() == FAnimNextGraphFunctionOutlinerData::StaticStruct())
+												{
+													const FAnimNextGraphFunctionOutlinerData& GraphFunctionData = Data.Get<FAnimNextGraphFunctionOutlinerData>();
+													if (GraphFunctionData.EditorObject.IsValid())
+													{
+														ObjectsToOpen.Add(GraphFunctionData.EditorObject.Get());
+													}
+												}
+												else if (Data.GetScriptStruct() == FAnimNextCollapseGraphOutlinerData::StaticStruct())
+												{
+													const FAnimNextCollapseGraphOutlinerData& GraphFunctionData = Data.Get<FAnimNextCollapseGraphOutlinerData>();
+													if (GraphFunctionData.EditorObject.IsValid())
+													{
+														ObjectsToOpen.Add(GraphFunctionData.EditorObject.Get());
 													}
 												}
 											}
-										}
-									}
 
-									SharedWorkspaceEditor->OpenObjects({ObjectsToOpen});
-								}
-							})
-						));
+											SharedWorkspaceEditor->OpenObjects({ ObjectsToOpen });
+										}
+									})
+								));
+						}
 					}
 				}
 			}
@@ -141,6 +236,7 @@ void FAnimNextGraphItemDetails::UnregisterToolMenuExtensions()
 		ToolMenus->UnregisterOwnerByName("FAnimNextGraphItemDetails");
 	}
 }
-}
+
+} // UE::AnimNext::Editor
 
 #undef LOCTEXT_NAMESPACE // "FAnimNextGraphItemDetails"
\ No newline at end of file
diff --git a/Engine/Plugins/Experimental/AnimNext/Source/AnimNextUncookedOnly/Internal/AnimNextEdGraph.h b/Engine/Plugins/Experimental/AnimNext/Source/AnimNextUncookedOnly/Internal/AnimNextEdGraph.h
--- a/Engine/Plugins/Experimental/AnimNext/Source/AnimNextUncookedOnly/Internal/AnimNextEdGraph.h
+++ b/Engine/Plugins/Experimental/AnimNext/Source/AnimNextUncookedOnly/Internal/AnimNextEdGraph.h
@@ -22,6 +22,7 @@ class UAnimNextEdGraph : public URigVMEdGraph
 	GENERATED_BODY()
 	
 	friend class UAnimNextModule_EditorData;
+	friend class UAnimNextRigVMAssetEditorData;
 
 	// UObject interface
 	virtual void PostLoad() override;
diff --git a/Engine/Plugins/Experimental/AnimNext/Source/AnimNextUncookedOnly/Internal/AnimNextRigVMAssetEditorData.h b/Engine/Plugins/Experimental/AnimNext/Source/AnimNextUncookedOnly/Internal/AnimNextRigVMAssetEditorData.h
--- a/Engine/Plugins/Experimental/AnimNext/Source/AnimNextUncookedOnly/Internal/AnimNextRigVMAssetEditorData.h
+++ b/Engine/Plugins/Experimental/AnimNext/Source/AnimNextUncookedOnly/Internal/AnimNextRigVMAssetEditorData.h
@@ -65,7 +65,7 @@ class UAnimNextRigVMAssetLibrary : public UBlueprintFunctionLibrary
 
 /* Base class for all AnimNext editor data objects that use RigVM */
 UCLASS(MinimalAPI, Abstract)
-class UAnimNextRigVMAssetEditorData : public UObject, public IRigVMClientHost, public IRigVMGraphFunctionHost
+class UAnimNextRigVMAssetEditorData : public UObject, public IRigVMClientHost, public IRigVMGraphFunctionHost, public IRigVMClientExternalModelHost
 {
 	GENERATED_BODY()
 
@@ -82,6 +82,7 @@ protected:
 	friend class UE::AnimNext::Tests::FEditor_Graph;
 	friend class UE::AnimNext::Tests::FEditor_Parameters;
 	friend class UE::AnimNext::Editor::FParameterCustomization;
+	friend class UAnimNextModuleWorkspaceAssetUserData;
 
 	// UObject interface
 	virtual void Serialize(FArchive& Ar) override;
@@ -133,11 +134,16 @@ protected:
 	virtual URigVMController* GetController(const UEdGraph* InEdGraph) const override;
 	virtual URigVMController* GetOrCreateController(const UEdGraph* InGraph) override;
 	virtual TArray<FString> GeneratePythonCommands(const FString InNewBlueprintName) override;
+	virtual void SetupPinRedirectorsForBackwardsCompatibility() override;
 
 	// IRigVMGraphFunctionHost interface
 	virtual FRigVMGraphFunctionStore* GetRigVMGraphFunctionStore() override;
 	virtual const FRigVMGraphFunctionStore* GetRigVMGraphFunctionStore() const override;
 
+	// IRigVMClientExternalModelHost interface
+	virtual const TArray<TObjectPtr<URigVMGraph>>& GetExternalModels() const override { return GraphModels; }
+	virtual TObjectPtr<URigVMGraph> CreateContainedGraphModel(URigVMCollapseNode* CollapseNode, const FName& Name) override;
+
 	// Override called during initialization to determine what RigVM controller class is used
 	virtual TSubclassOf<URigVMController> GetControllerClass() const { return URigVMController::StaticClass(); }
 
@@ -148,7 +154,10 @@ protected:
 	virtual UEdGraph* CreateEdGraph(URigVMGraph* InRigVMGraph, bool bForce) PURE_VIRTUAL(UAnimNextRigVMAssetEditorData::CreateEdGraph, return nullptr;)
 
 	// Create and store a UEdGraph that corresponds to a URigVMCollapseNode
-	virtual void CreateEdGraphForCollapseNode(URigVMCollapseNode* InNode) PURE_VIRTUAL(UAnimNextRigVMAssetEditorData::CreateEdGraphForCollapseNode, )
+	virtual void CreateEdGraphForCollapseNode(URigVMCollapseNode* InNode, bool bForce) PURE_VIRTUAL(UAnimNextRigVMAssetEditorData::CreateEdGraphForCollapseNode, )
+
+	// Destroy a UEdGraph that corresponds to a URigVMCollapseNode
+	virtual void RemoveEdGraphForCollapseNode(URigVMCollapseNode* InNode, bool bNotify) PURE_VIRTUAL(UAnimNextRigVMAssetEditorData::RemoveEdGraphForCollapseNode, )
 
 	// Remove the UEdGraph that corresponds to a URigVMGraph
 	virtual bool RemoveEdGraph(URigVMGraph* InModel) PURE_VIRTUAL(UAnimNextRigVMAssetEditorData::RemoveEdGraph, return false;)
@@ -197,15 +206,48 @@ protected:
 
 	// Returns all nodes in all graphs of the specified class
 	template<class T>
-	void GetAllNodesOfClass(TArray<T*>& OutNodes)
+	void GetAllNodesOfClass(TArray<T*>& OutNodes) const
 	{
 		ForEachEntryOfType<IAnimNextRigVMGraphInterface>([&OutNodes](IAnimNextRigVMGraphInterface* InGraphInterface)
 		{
+			URigVMEdGraph* RigVMEdGraph = InGraphInterface->GetEdGraph();
+			check(RigVMEdGraph)
+
 			TArray<T*> GraphNodes;
-			InGraphInterface->GetEdGraph()->GetNodesOfClass<T>(GraphNodes);
+			RigVMEdGraph->GetNodesOfClass<T>(GraphNodes);
+
+			TArray<UEdGraph*> SubGraphs;
+			RigVMEdGraph->GetAllChildrenGraphs(SubGraphs);
+			for (const TObjectPtr<UEdGraph>& SubGraph : SubGraphs)
+			{
+				if (SubGraph)
+				{
+					SubGraph->GetNodesOfClass<T>(GraphNodes);
+				}
+			}
+
 			OutNodes.Append(GraphNodes);
+
 			return true;
 		});
+
+		for (URigVMEdGraph* RigVMEdGraph : FunctionEdGraphs)
+		{
+			if (RigVMEdGraph)
+			{
+				RigVMEdGraph->GetNodesOfClass<T>(OutNodes);
+
+				TArray<UEdGraph*> SubGraphs;
+				RigVMEdGraph->GetAllChildrenGraphs(SubGraphs);
+				for (const TObjectPtr<UEdGraph>& SubGraph : SubGraphs)
+				{
+					if (SubGraph)
+					{
+						SubGraph->GetNodesOfClass<T>(OutNodes);
+					}
+				}
+			}
+		}
 	}
 	
 	// Find an entry by name
@@ -229,7 +271,13 @@ protected:
 	void PostLoadExternalPackages();
 
 	// Find an entry that corresponds to the specified RigVMGraph. This uses the name of the graph to match the entry 
-	UAnimNextRigVMAssetEntry* FindEntryForRigVMGraph(URigVMGraph* InRigVMGraph) const;
+	UAnimNextRigVMAssetEntry* FindEntryForRigVMGraph(const URigVMGraph* InRigVMGraph) const;
+
+	// Find an entry that corresponds to the specified RigVMGraph. This uses the name of the graph to match the entry 
+	UAnimNextRigVMAssetEntry* FindEntryForRigVMEdGraph(const URigVMEdGraph* InRigVMEdGraph) const;
+
+	// Refresh the 'external' models for the RigVM client to reference
+	void RefreshExternalModels();
 	
 	/** All entries in this asset - not saved, either serialized or discovered at load time */
 	UPROPERTY(transient)
@@ -277,6 +325,19 @@ protected:
 
 	// Cached exports, generated lazily or on compilation
 	mutable TOptional<FAnimNextParameterProviderAssetRegistryExports> CachedExports;
+	
+	// Collection of models gleaned from graphs
+	TArray<TObjectPtr<URigVMGraph>> GraphModels;
+
+
+	// Set of functions implemented for this graph
+	UPROPERTY()
+	TArray<TObjectPtr<URigVMEdGraph>> FunctionEdGraphs;
+
+	// Default FunctionLibrary EdGraph
+	UPROPERTY()
+	TObjectPtr<UAnimNextEdGraph> FunctionLibraryEdGraph;
+
 
 	bool bAutoRecompileVM = true;
 	bool bErrorsDuringCompilation = false;
diff --git a/Engine/Plugins/Experimental/AnimNext/Source/AnimNextUncookedOnly/Internal/Graph/AnimNextModule_AnimationGraph.h b/Engine/Plugins/Experimental/AnimNext/Source/AnimNextUncookedOnly/Internal/Graph/AnimNextModule_AnimationGraph.h
--- a/Engine/Plugins/Experimental/AnimNext/Source/AnimNextUncookedOnly/Internal/Graph/AnimNextModule_AnimationGraph.h
+++ b/Engine/Plugins/Experimental/AnimNext/Source/AnimNextUncookedOnly/Internal/Graph/AnimNextModule_AnimationGraph.h
@@ -1,4 +1,4 @@
-﻿// Copyright Epic Games, Inc. All Rights Reserved.
+// Copyright Epic Games, Inc. All Rights Reserved.
 
 #pragma once
 
@@ -37,6 +37,7 @@ class UAnimNextModule_AnimationGraph : public UAnimNextRigVMAssetEntry, public I
 	virtual void SetEntryName(FName InName, bool bSetupUndoRedo = true) override;
 	
 	// IAnimNextRigVMGraphInterface interface
+	virtual const FName& GetGraphName() const override;
 	virtual URigVMGraph* GetRigVMGraph() const override;
 	virtual URigVMEdGraph* GetEdGraph() const override;
 	virtual void SetRigVMGraph(URigVMGraph* InGraph) override;
@@ -68,4 +69,4 @@ class UAnimNextGraphEntry : public UAnimNextRigVMAssetEntry
 
 	// UAnimNextRigVMAssetEntry interface
 	virtual FName GetEntryName() const override { return NAME_None; }
-};
\ No newline at end of file
+};
diff --git a/Engine/Plugins/Experimental/AnimNext/Source/AnimNextUncookedOnly/Internal/IAnimNextRigVMGraphInterface.h b/Engine/Plugins/Experimental/AnimNext/Source/AnimNextUncookedOnly/Internal/IAnimNextRigVMGraphInterface.h
--- a/Engine/Plugins/Experimental/AnimNext/Source/AnimNextUncookedOnly/Internal/IAnimNextRigVMGraphInterface.h
+++ b/Engine/Plugins/Experimental/AnimNext/Source/AnimNextUncookedOnly/Internal/IAnimNextRigVMGraphInterface.h
@@ -20,6 +20,9 @@ class ANIMNEXTUNCOOKEDONLY_API IAnimNextRigVMGraphInterface
 	GENERATED_BODY()
 
 public:
+	// Get the AnimNext graph name
+	virtual const FName& GetGraphName() const = 0;
+
 	// Get the RigVM graph
 	virtual URigVMGraph* GetRigVMGraph() const = 0;
 
diff --git a/Engine/Plugins/Experimental/AnimNext/Source/AnimNextUncookedOnly/Internal/Module/AnimNextModuleWorkspaceAssetUserData.h b/Engine/Plugins/Experimental/AnimNext/Source/AnimNextUncookedOnly/Internal/Module/AnimNextModuleWorkspaceAssetUserData.h
--- a/Engine/Plugins/Experimental/AnimNext/Source/AnimNextUncookedOnly/Internal/Module/AnimNextModuleWorkspaceAssetUserData.h
+++ b/Engine/Plugins/Experimental/AnimNext/Source/AnimNextUncookedOnly/Internal/Module/AnimNextModuleWorkspaceAssetUserData.h
@@ -5,6 +5,7 @@
 #include "WorkspaceAssetRegistryInfo.h"
 #include "Param/ParamType.h"
 #include "IAnimNextRigVMGraphInterface.h"
+#include "RigVMModel/Nodes/RigVMFunctionReferenceNode.h"
 
 #include "AnimNextModuleWorkspaceAssetUserData.generated.h"
 
@@ -38,6 +39,28 @@ struct FAnimNextGraphOutlinerData : public FWorkspaceOutlinerItemData
 	TScriptInterface<IAnimNextRigVMGraphInterface> GraphInterface;
 };
 
+USTRUCT()
+struct FAnimNextCollapseGraphOutlinerData : public FWorkspaceOutlinerItemData
+{
+	GENERATED_BODY()
+	
+	FAnimNextCollapseGraphOutlinerData() = default;
+
+	UPROPERTY(VisibleAnywhere, Category=AnimNext)
+	TWeakObjectPtr<URigVMEdGraph> EditorObject;
+};
+
+USTRUCT()
+struct FAnimNextGraphFunctionOutlinerData : public FWorkspaceOutlinerItemData
+{
+	GENERATED_BODY()
+	
+	FAnimNextGraphFunctionOutlinerData() = default;
+
+	UPROPERTY(VisibleAnywhere, Category=AnimNext)
+	TWeakObjectPtr<URigVMEdGraph> EditorObject;
+};
+
 USTRUCT()
 struct FAnimNextSchedulerData : public FAnimNextGraphAssetOutlinerData
 {
@@ -52,4 +75,4 @@ class UAnimNextModuleWorkspaceAssetUserData : public UAssetUserData
 	GENERATED_BODY()
 
 	virtual void GetAssetRegistryTags(FAssetRegistryTagsContext Context) const override;
-};
\ No newline at end of file
+};
diff --git a/Engine/Plugins/Experimental/AnimNext/Source/AnimNextUncookedOnly/Internal/Module/AnimNextModule_EditorData.h b/Engine/Plugins/Experimental/AnimNext/Source/AnimNextUncookedOnly/Internal/Module/AnimNextModule_EditorData.h
--- a/Engine/Plugins/Experimental/AnimNext/Source/AnimNextUncookedOnly/Internal/Module/AnimNextModule_EditorData.h
+++ b/Engine/Plugins/Experimental/AnimNext/Source/AnimNextUncookedOnly/Internal/Module/AnimNextModule_EditorData.h
@@ -12,6 +12,7 @@
 #include "AnimNextModule_EditorData.generated.h"
 
 class UAnimNextModule;
+class UAnimNextModule_FunctionGraph;
 enum class ERigVMGraphNotifType : uint8;
 class FAnimationAnimNextRuntimeTest_GraphAddTrait;
 class FAnimationAnimNextRuntimeTest_GraphExecute;
@@ -90,7 +91,8 @@ private:
 	virtual UScriptStruct* GetExecuteContextStruct() const override { return FAnimNextExecuteContext::StaticStruct(); }
 	virtual UEdGraph* CreateEdGraph(URigVMGraph* InRigVMGraph, bool bForce) override;
 	virtual bool RemoveEdGraph(URigVMGraph* InModel) override;
-	virtual void CreateEdGraphForCollapseNode(URigVMCollapseNode* InNode) override;
+	virtual void CreateEdGraphForCollapseNode(URigVMCollapseNode* InNode, bool bForce) override;
+	virtual void RemoveEdGraphForCollapseNode(URigVMCollapseNode* InNode, bool bNotify) override;
 	virtual void HandleModifiedEvent(ERigVMGraphNotifType InNotifType, URigVMGraph* InGraph, UObject* InSubject) override;
 	virtual TConstArrayView<TSubclassOf<UAnimNextRigVMAssetEntry>> GetEntryClasses() const override;
 
diff --git a/Engine/Plugins/Experimental/AnimNext/Source/AnimNextUncookedOnly/Internal/Module/AnimNextModule_EventGraph.h b/Engine/Plugins/Experimental/AnimNext/Source/AnimNextUncookedOnly/Internal/Module/AnimNextModule_EventGraph.h
--- a/Engine/Plugins/Experimental/AnimNext/Source/AnimNextUncookedOnly/Internal/Module/AnimNextModule_EventGraph.h
+++ b/Engine/Plugins/Experimental/AnimNext/Source/AnimNextUncookedOnly/Internal/Module/AnimNextModule_EventGraph.h
@@ -1,4 +1,4 @@
-﻿// Copyright Epic Games, Inc. All Rights Reserved.
+// Copyright Epic Games, Inc. All Rights Reserved.
 
 #pragma once
 
@@ -32,6 +32,7 @@ class UAnimNextModule_EventGraph : public UAnimNextRigVMAssetEntry, public IAnim
 	virtual FText GetDisplayNameTooltip() const override;
 
 	// IAnimNextRigVMGraphInterface interface
+	virtual const FName& GetGraphName() const override;
 	virtual URigVMGraph* GetRigVMGraph() const override;
 	virtual URigVMEdGraph* GetEdGraph() const override;
 	virtual void SetRigVMGraph(URigVMGraph* InGraph) override;
diff --git a/Engine/Plugins/Experimental/AnimNext/Source/AnimNextUncookedOnly/Internal/UncookedOnlyUtils.h b/Engine/Plugins/Experimental/AnimNext/Source/AnimNextUncookedOnly/Internal/UncookedOnlyUtils.h
--- a/Engine/Plugins/Experimental/AnimNext/Source/AnimNextUncookedOnly/Internal/UncookedOnlyUtils.h
+++ b/Engine/Plugins/Experimental/AnimNext/Source/AnimNextUncookedOnly/Internal/UncookedOnlyUtils.h
@@ -15,6 +15,8 @@ struct FAnimNextEditorParam;
 struct FAnimNextParam;
 struct FEdGraphPinType;
 struct FWorkspaceOutlinerItemExports;
+struct FWorkspaceOutlinerItemExport;
+struct FRigVMGraphFunctionData;
 class UAnimNextSchedule;
 class UAnimNextModule;
 class UAnimNextModule_EditorData;
@@ -191,6 +193,9 @@ struct ANIMNEXTUNCOOKEDONLY_API FUtils
 
 	// Gets the asset-registry information needed for representing the contained data into the Workspace Outliner
 	static void GetAssetOutlinerItems(const UAnimNextRigVMAssetEditorData* EditorData, FWorkspaceOutlinerItemExports& OutExports);
+	static void CreateSubGraphsOutlinerItemsRecursive(const UAnimNextRigVMAssetEditorData* EditorData, FWorkspaceOutlinerItemExports& OutExports, FWorkspaceOutlinerItemExport& ParentExport, URigVMEdGraph* RigVMEdGraph);
+	static void CreateFunctionLibraryOutlinerItemsRecursive(const UAnimNextRigVMAssetEditorData* EditorData, FWorkspaceOutlinerItemExports& OutExports, FWorkspaceOutlinerItemExport& ParentExport, const TArray<FRigVMGraphFunctionData>& PublicFunctions, const TArray<FRigVMGraphFunctionData>& PrivateFunctions);
+	static void CreateFunctionsOutlinerItemsRecursive(const UAnimNextRigVMAssetEditorData* EditorData, FWorkspaceOutlinerItemExports& OutExports, FWorkspaceOutlinerItemExport& ParentExport, const TArray<FRigVMGraphFunctionData>& Functions, bool bPublicFunctions);
 
 	// Attempts to determine the type from a parameter name
 	// If the name cannot be found, the returned type will be invalid
@@ -202,6 +207,9 @@ struct ANIMNEXTUNCOOKEDONLY_API FUtils
 
 	// Sorts the incoming array of parameters, then generates a hash and returns it.
 	static uint64 SortAndHashParameters(TArray<FAnimNextParam>& InParameters);
+
+	// Returns an user friendly name for the Function Library
+	static const FText& GetFunctionLibraryDisplayName();
 };
 
 }
\ No newline at end of file
diff --git a/Engine/Plugins/Experimental/AnimNext/Source/AnimNextUncookedOnly/Private/AnimNextRigVMAssetEditorData.cpp b/Engine/Plugins/Experimental/AnimNext/Source/AnimNextUncookedOnly/Private/AnimNextRigVMAssetEditorData.cpp
--- a/Engine/Plugins/Experimental/AnimNext/Source/AnimNextUncookedOnly/Private/AnimNextRigVMAssetEditorData.cpp
+++ b/Engine/Plugins/Experimental/AnimNext/Source/AnimNextUncookedOnly/Private/AnimNextRigVMAssetEditorData.cpp
@@ -11,6 +11,8 @@
 #include "IAnimNextRigVMGraphInterface.h"
 #include "IAnimNextRigVMParameterInterface.h"
 #include "UncookedOnlyUtils.h"
+#include "AnimNextEdGraph.h"
+#include "AnimNextEdGraphSchema.h"
 #include "Misc/TransactionObjectEvent.h"
 #include "Param/AnimNextTag.h"
 #include "RigVMModel/RigVMFunctionLibrary.h"
@@ -85,11 +87,16 @@ void UAnimNextRigVMAssetEditorData::Initialize(bool bRecompileVM)
 	RigVMClient.SetDefaultSchemaClass(UAnimNextRigVMAssetSchema::StaticClass());
 	RigVMClient.SetControllerClass(GetControllerClass());
 	RigVMClient.SetOuterClientHost(this, GET_MEMBER_NAME_CHECKED(UAnimNextRigVMAssetEditorData, RigVMClient));
+	RigVMClient.SetExternalModelHost(this);
+
+	URigVMFunctionLibrary* RigVMFunctionLibrary = nullptr;
 	{
 		TGuardValue<bool> DisableClientNotifs(RigVMClient.bSuspendNotifications, true);
-		RigVMClient.GetOrCreateFunctionLibrary(false);
+		RigVMFunctionLibrary = RigVMClient.GetOrCreateFunctionLibrary(false);
 	}
 
+	ensure(RigVMFunctionLibrary->GetFunctionHostObjectPathDelegate.IsBound());
+
 	if (RigVMClient.GetController(0) == nullptr)
 	{
 		if(RigVMClient.GetDefaultModel())
@@ -97,8 +104,21 @@ void UAnimNextRigVMAssetEditorData::Initialize(bool bRecompileVM)
 			RigVMClient.GetOrCreateController(RigVMClient.GetDefaultModel());
 		}
 
-		check(RigVMClient.GetFunctionLibrary());
-		RigVMClient.GetOrCreateController(RigVMClient.GetFunctionLibrary());
+		check(RigVMFunctionLibrary);
+		RigVMClient.GetOrCreateController(RigVMFunctionLibrary);
+
+		if (!FunctionLibraryEdGraph)
+		{
+			FunctionLibraryEdGraph = NewObject<UAnimNextEdGraph>(CastChecked<UObject>(this), NAME_None, RF_Transactional);
+
+			FunctionLibraryEdGraph->Schema = UAnimNextEdGraphSchema::StaticClass();
+			FunctionLibraryEdGraph->bAllowRenaming = 0;
+			FunctionLibraryEdGraph->bEditable = 0;
+			FunctionLibraryEdGraph->bAllowDeletion = 0;
+			FunctionLibraryEdGraph->bIsFunctionDefinition = false;
+			FunctionLibraryEdGraph->ModelNodePath = RigVMClient.GetFunctionLibrary()->GetNodePath();
+			FunctionLibraryEdGraph->Initialize(this);
+		}
 
 		// Init function library controllers
 		for(URigVMLibraryNode* LibraryNode : RigVMClient.GetFunctionLibrary()->GetFunctions())
@@ -130,11 +150,18 @@ void UAnimNextRigVMAssetEditorData::PostLoad()
 {
 	Super::PostLoad();
 
+	GraphModels.Reset();
+	
+	PostLoadExternalPackages();
+	RefreshExternalModels();
 	Initialize(/*bRecompileVM*/false);
+	
+	GetRigVMClient()->RefreshAllModels(ERigVMLoadType::PostLoad, false, bIsCompiling);
 
-	RefreshAllModels(ERigVMLoadType::PostLoad);
-
-	PostLoadExternalPackages();
+	GetRigVMClient()->PatchFunctionReferencesOnLoad();
+	TMap<URigVMLibraryNode*, FRigVMGraphFunctionHeader> OldHeaders;
+	TArray<FName> BackwardsCompatiblePublicFunctions;
+	GetRigVMClient()->PatchFunctionsOnLoad(this, BackwardsCompatiblePublicFunctions, OldHeaders);
 
 	// delay compilation until the package has been loaded
 	FCoreUObjectDelegates::OnEndLoadPackage.AddUObject(this, &UAnimNextRigVMAssetEditorData::HandlePackageDone);
@@ -208,128 +235,17 @@ void UAnimNextRigVMAssetEditorData::HandlePackageDone()
 
 	RecompileVM();
 
-	ReconstructAllNodes();
+	ReconstructAllNodes(); // If this is not executed on a node for whatever reason, it will appear transparent in the editor
 }
 
 void UAnimNextRigVMAssetEditorData::RefreshAllModels(ERigVMLoadType InLoadType)
 {
-	const bool bIsPostLoad = InLoadType == ERigVMLoadType::PostLoad;
-
-	TGuardValue<bool> IsCompilingGuard(bIsCompiling, true);
-	TGuardValue<bool> ClientIgnoreModificationsGuard(RigVMClient.bIgnoreModelNotifications, true);
-
-	TArray<URigVMGraph*> GraphsToDetach = RigVMClient.GetAllModels(true, false);
-	TMap<const URigVMGraph*, TArray<URigVMController::FLinkedPath>> LinkedPaths;
-	
-	if (ensure(IsInGameThread()))
-	{
-		for (const URigVMGraph* GraphToDetach : GraphsToDetach)
-		{
-			URigVMController* Controller = RigVMClient.GetOrCreateController(GraphToDetach);
-			// temporarily disable default value validation during load time, serialized values should always be accepted
-			TGuardValue<bool> PerGraphDisablePinDefaultValueValidation(Controller->bValidatePinDefaults, false);
-			FRigVMControllerNotifGuard NotifGuard(Controller, true);
-			LinkedPaths.Add(GraphToDetach, Controller->GetLinkedPaths());
-			Controller->FastBreakLinkedPaths(LinkedPaths.FindChecked(GraphToDetach));
-			TArray<URigVMNode*> Nodes = GraphToDetach->GetNodes();
-			for (URigVMNode* Node : Nodes)
-			{
-				Controller->RepopulatePinsOnNode(Node, true, true);
-			}
-		}
-		//SetupPinRedirectorsForBackwardsCompatibility();
-	}
-
-	for (const URigVMGraph* GraphToDetach : GraphsToDetach)
-	{
-		URigVMController* Controller = RigVMClient.GetOrCreateController(GraphToDetach);
-		// at this stage, allow all links to be reattached,
-		// RecomputeAllTemplateFilteredPermutations() later should break any invalid links
-		FRigVMControllerNotifGuard NotifGuard(Controller, true);
-		Controller->RestoreLinkedPaths(LinkedPaths.FindChecked(GraphToDetach));
-	}
-
-	if (bIsPostLoad)
-	{
-		//PatchTemplateNodesWithPreferredPermutation();
-	}
-
-	TArray<URigVMGraph*> GraphsToClean = RigVMClient.GetAllModels(true, true);
-
-	// Sort from leaf graphs to root
-	TArray<URigVMGraph*> SortedGraphsToClean;
-	SortedGraphsToClean.Reserve(GraphsToClean.Num());
-	while (SortedGraphsToClean.Num() < GraphsToClean.Num())
-	{
-		bool bGraphAdded = false;
-		for (URigVMGraph* Graph : GraphsToClean)
-		{
-			if (SortedGraphsToClean.Contains(Graph))
-			{
-				continue;
-			}
-
-			TArray<URigVMGraph*> ContainedGraphs;
-			for (URigVMNode* Node : Graph->GetNodes())
-			{
-				if (URigVMLibraryNode* LibraryNode = Cast<URigVMLibraryNode>(Node))
-				{
-					if (URigVMFunctionReferenceNode* FunctionReferenceNode = Cast<URigVMFunctionReferenceNode>(LibraryNode))
-					{
-						if (FunctionReferenceNode->GetReferencedFunctionHeader().LibraryPointer.GetNodeSoftPath().GetLongPackageName() != GetPackage()->GetPathName())
-						{
-							continue;
-						}
-						if (URigVMLibraryNode* ReferencedNode = FunctionReferenceNode->LoadReferencedNode())
-						{
-							ContainedGraphs.Add(ReferencedNode->GetContainedGraph());
-							continue;
-						}
-					}
-
-					if (URigVMGraph* ContainedGraph = LibraryNode->GetContainedGraph())
-					{
-						ContainedGraphs.Add(ContainedGraph);
-					}
-				}
-			}
-
-			bool bAllContained = true;
-			for (URigVMGraph* Contained : ContainedGraphs)
-			{
-				if (!SortedGraphsToClean.Contains(Contained))
-				{
-					bAllContained = false;
-					break;
-				}
-			}
-			if (bAllContained)
-			{
-				SortedGraphsToClean.Add(Graph);
-				bGraphAdded = true;
-			}
-		}
-		ensure(bGraphAdded);
-	}
-
-	for (int32 GraphIndex = 0; GraphIndex < SortedGraphsToClean.Num(); GraphIndex++)
-	{
-		URigVMGraph* GraphToClean = SortedGraphsToClean[GraphIndex];
-		URigVMController* Controller = RigVMClient.GetOrCreateController(GraphToClean);
-		//TGuardValue<bool> GuardEditGraph(GraphToClean->bEditable, true);
-		FRigVMControllerNotifGuard NotifGuard(Controller, true);
-
-		for (URigVMNode* ModelNode : GraphToClean->GetNodes())
-		{
-			Controller->RemoveUnusedOrphanedPins(ModelNode);
-		}
-	}
 }
 
 void UAnimNextRigVMAssetEditorData::OnRigVMRegistryChanged()
 {
-	RefreshAllModels(ERigVMLoadType::PostLoad);
-	//RebuildGraphFromModel(); // TODO zzz : How we do this on AnimNext ?
+	GetRigVMClient()->RefreshAllModels(ERigVMLoadType::PostLoad, false, bIsCompiling);
+	//RebuildGraphFromModel(); // TODO zzz : Move from blueprint to client
 }
 
 void UAnimNextRigVMAssetEditorData::RequestRigVMInit()
@@ -419,6 +335,10 @@ TArray<FString> UAnimNextRigVMAssetEditorData::GeneratePythonCommands(const FStr
 	return TArray<FString>();
 }
 
+void UAnimNextRigVMAssetEditorData::SetupPinRedirectorsForBackwardsCompatibility()
+{
+}
+
 FRigVMClient* UAnimNextRigVMAssetEditorData::GetRigVMClient()
 {
 	return &RigVMClient;
@@ -516,17 +436,62 @@ UObject* UAnimNextRigVMAssetEditorData::GetEditorObjectForRigVMGraph(URigVMGraph
 {
 	if(InVMGraph)
 	{
+		if (InVMGraph->IsA<URigVMFunctionLibrary>())
+		{
+			return Cast<UObject>(FunctionLibraryEdGraph.Get());
+		}
+
+		const auto FindSubgraph = ([](const FString SearchGraphNodePath, URigVMEdGraph* EdGraph) -> URigVMEdGraph*
+		{
+			TArray<UEdGraph*> SubGraphs;
+			EdGraph->GetAllChildrenGraphs(SubGraphs);
+			for (const TObjectPtr<UEdGraph>& SubGraph : SubGraphs)
+			{
+				if (URigVMEdGraph* RigVMEdGraph = Cast<URigVMEdGraph>(SubGraph))
+				{
+					if (RigVMEdGraph->ModelNodePath == SearchGraphNodePath)
+					{
+						return RigVMEdGraph;
+					}
+				}
+			}
+			return nullptr;
+		});
+
+		const FString GraphNodePath = InVMGraph->GetNodePath();
 		for(UAnimNextRigVMAssetEntry* Entry : Entries)
 		{
 			if(IAnimNextRigVMGraphInterface* GraphInterface = Cast<IAnimNextRigVMGraphInterface>(Entry))
 			{
 				URigVMEdGraph* EdGraph = GraphInterface->GetEdGraph();
-				if(EdGraph->ModelNodePath == InVMGraph->GetNodePath())
+
+				if (const URigVMGraph* RigVMGraph = GraphInterface->GetRigVMGraph())
+				{
+					if (RigVMGraph == InVMGraph)
+					{
+						return EdGraph;
+					}
+				}
+
+				if (URigVMEdGraph* RigVMEdGraph = FindSubgraph(GraphNodePath, EdGraph))
 				{
-					return EdGraph;
+					return RigVMEdGraph;
 				}
 			}
 		}
+
+		for (const TObjectPtr<URigVMEdGraph>& FunctionEdGraph : FunctionEdGraphs)
+		{
+			if (FunctionEdGraph->ModelNodePath == GraphNodePath)
+			{
+				return FunctionEdGraph;
+			}
+
+			if (URigVMEdGraph* RigVMEdGraph = FindSubgraph(GraphNodePath, FunctionEdGraph))
+			{
+				return RigVMEdGraph;
+			}
+		}
 	}
 	return nullptr;
 }
@@ -561,6 +526,25 @@ const FRigVMGraphFunctionStore* UAnimNextRigVMAssetEditorData::GetRigVMGraphFunc
 	return &GraphFunctionStore;
 }
 
+TObjectPtr<URigVMGraph> UAnimNextRigVMAssetEditorData::CreateContainedGraphModel(URigVMCollapseNode* CollapseNode, const FName& Name)
+{
+	check(CollapseNode);
+
+	TObjectPtr<URigVMGraph> Model = NewObject<URigVMGraph>(CollapseNode, Name);
+	Model->SetSchemaClass(RigVMClient.GetDefaultSchemaClass());
+
+	URigVMGraph* CollapseNodeModelRootGraph = CollapseNode->GetRootGraph();
+	check(CollapseNodeModelRootGraph);
+
+	// If we are a transient asset, dont use external packages
+	if (!CollapseNodeModelRootGraph->HasAnyFlags(RF_Transient))
+	{
+		Model->SetExternalPackage(CollapseNodeModelRootGraph->GetExternalPackage());
+	}
+
+	return Model;
+}
+
 void UAnimNextRigVMAssetEditorData::RecompileVMIfRequired()
 {
 	if (bVMRecompilationRequired)
@@ -635,11 +619,22 @@ void UAnimNextRigVMAssetEditorData::HandleModifiedEvent(ERigVMGraphNotifType InN
 		{
 			if (URigVMCollapseNode* CollapseNode = Cast<URigVMCollapseNode>(InSubject))
 			{
-				CreateEdGraphForCollapseNode(CollapseNode);
+				CreateEdGraphForCollapseNode(CollapseNode, false);
 				break;
 			}
+			RequestAutoVMRecompilation();
+			break;
+	}
+	case ERigVMGraphNotifType::NodeRemoved:
+	{
+		if (URigVMCollapseNode* CollapseNode = Cast<URigVMCollapseNode>(InSubject))
+		{
+			RemoveEdGraphForCollapseNode(CollapseNode, false);
+			break;
 		}
-		// Fall through to the next case
+		RequestAutoVMRecompilation();
+		break;
+	}
 	case ERigVMGraphNotifType::LinkAdded:
 	case ERigVMGraphNotifType::LinkRemoved:
 	case ERigVMGraphNotifType::PinArraySizeChanged:
@@ -691,6 +686,11 @@ TArray<UEdGraph*> UAnimNextRigVMAssetEditorData::GetAllEdGraphs() const
 			Graphs.Add(GraphInterface->GetEdGraph());
 		}
 	}
+	for (URigVMEdGraph* RigVMEdGraph : FunctionEdGraphs)
+	{
+		Graphs.Add(RigVMEdGraph);
+	}
+
 	return Graphs;
 }
 
@@ -743,6 +743,7 @@ bool UAnimNextRigVMAssetEditorData::RemoveEntry(UAnimNextRigVMAssetEntry* InEntr
 	// Remove from internal array
 	UAnimNextRigVMAssetEntry* EntryToRemove = *EntryToRemovePtr;
 	Entries.Remove(EntryToRemove);
+	RefreshExternalModels();
 
 	if (bSetupUndoRedo)
 	{
@@ -804,13 +805,32 @@ UObject* UAnimNextRigVMAssetEditorData::CreateNewSubEntry(UAnimNextRigVMAssetEdi
 	return NewEntry;
 }
 
-UAnimNextRigVMAssetEntry* UAnimNextRigVMAssetEditorData::FindEntryForRigVMGraph(URigVMGraph* InRigVMGraph) const
+UAnimNextRigVMAssetEntry* UAnimNextRigVMAssetEditorData::FindEntryForRigVMGraph(const URigVMGraph* InRigVMGraph) const
 {
 	for(UAnimNextRigVMAssetEntry* Entry : Entries)
 	{
 		if(IAnimNextRigVMGraphInterface* GraphInterface = Cast<IAnimNextRigVMGraphInterface>(Entry))
 		{
-			if(GraphInterface->GetRigVMGraph() == InRigVMGraph)
+			if (const URigVMGraph* RigVMGraph = GraphInterface->GetRigVMGraph())
+			{
+				if(RigVMGraph == InRigVMGraph)
+				{
+					return Entry;
+				}
+			}
+		}
+	}
+
+	return nullptr;
+}
+
+UAnimNextRigVMAssetEntry* UAnimNextRigVMAssetEditorData::FindEntryForRigVMEdGraph(const URigVMEdGraph* InRigVMEdGraph) const
+{
+	for (UAnimNextRigVMAssetEntry* Entry : Entries)
+	{
+		if (IAnimNextRigVMGraphInterface* GraphInterface = Cast<IAnimNextRigVMGraphInterface>(Entry))
+		{
+			if (GraphInterface->GetEdGraph() == InRigVMEdGraph)
 			{
 				return Entry;
 			}
@@ -818,4 +838,21 @@ UAnimNextRigVMAssetEntry* UAnimNextRigVMAssetEditorData::FindEntryForRigVMGraph(
 	}
 
 	return nullptr;
-}
\ No newline at end of file
+}
+
+void UAnimNextRigVMAssetEditorData::RefreshExternalModels()
+{
+	GraphModels.Reset();
+
+	for (UAnimNextRigVMAssetEntry* Entry : Entries)
+	{
+		if (IAnimNextRigVMGraphInterface* GraphInterface = Cast<IAnimNextRigVMGraphInterface>(Entry))
+		{
+			if(URigVMGraph* Model = GraphInterface->GetRigVMGraph())
+			{
+				GraphModels.Add(Model);
+			}
+		}
+	}
+
+}
diff --git a/Engine/Plugins/Experimental/AnimNext/Source/AnimNextUncookedOnly/Private/Graph/AnimNextModule_AnimationGraph.cpp b/Engine/Plugins/Experimental/AnimNext/Source/AnimNextUncookedOnly/Private/Graph/AnimNextModule_AnimationGraph.cpp
--- a/Engine/Plugins/Experimental/AnimNext/Source/AnimNextUncookedOnly/Private/Graph/AnimNextModule_AnimationGraph.cpp
+++ b/Engine/Plugins/Experimental/AnimNext/Source/AnimNextUncookedOnly/Private/Graph/AnimNextModule_AnimationGraph.cpp
@@ -1,4 +1,4 @@
-﻿// Copyright Epic Games, Inc. All Rights Reserved.
+// Copyright Epic Games, Inc. All Rights Reserved.
 
 #include "Graph/AnimNextModule_AnimationGraph.h"
 
@@ -74,6 +74,11 @@ void UAnimNextModule_AnimationGraph::SetEntryName(FName InName, bool bSetupUndoR
 	BroadcastModified();
 }
 
+const FName& UAnimNextModule_AnimationGraph::GetGraphName() const
+{
+	return GraphName;
+}
+
 URigVMGraph* UAnimNextModule_AnimationGraph::GetRigVMGraph() const
 {
 	return Graph;
diff --git a/Engine/Plugins/Experimental/AnimNext/Source/AnimNextUncookedOnly/Private/Module/AnimNextModule_EditorData.cpp b/Engine/Plugins/Experimental/AnimNext/Source/AnimNextUncookedOnly/Private/Module/AnimNextModule_EditorData.cpp
--- a/Engine/Plugins/Experimental/AnimNext/Source/AnimNextUncookedOnly/Private/Module/AnimNextModule_EditorData.cpp
+++ b/Engine/Plugins/Experimental/AnimNext/Source/AnimNextUncookedOnly/Private/Module/AnimNextModule_EditorData.cpp
@@ -2,6 +2,7 @@
 
 #include "Module/AnimNextModule_EditorData.h"
 
+#include "ExternalPackageHelper.h"
 #include "UncookedOnlyUtils.h"
 #include "AssetRegistry/IAssetRegistry.h"
 #include "Curves/CurveFloat.h"
@@ -15,6 +16,7 @@
 #include "RigVMModel/RigVMFunctionLibrary.h"
 #include "RigVMModel/RigVMNotifications.h"
 #include "RigVMModel/Nodes/RigVMCollapseNode.h"
+#include "RigVMModel/Nodes/RigVMAggregateNode.h"
 #include "UObject/AssetRegistryTagsContext.h"
 #include "UObject/FortniteMainBranchObjectVersion.h"
 #include "UObject/LinkerLoad.h"
@@ -195,24 +197,169 @@ TConstArrayView<TSubclassOf<UAnimNextRigVMAssetEntry>> UAnimNextModule_EditorDat
 	return Classes;
 }
 
-void UAnimNextModule_EditorData::CreateEdGraphForCollapseNode(URigVMCollapseNode* InNode)
+void UAnimNextModule_EditorData::CreateEdGraphForCollapseNode(URigVMCollapseNode* InNode, bool bForce)
 {
+	check(InNode);
+	URigVMGraph* CollapseNodeGraph = InNode->GetGraph();
+	check(CollapseNodeGraph);
+
+	if (bForce)
+	{
+		RemoveEdGraphForCollapseNode(InNode, false);
+	}
+
+	// For Function node
+	if (InNode->GetGraph()->IsA<URigVMFunctionLibrary>())
+	{
+		if (URigVMGraph* ContainedGraph = InNode->GetContainedGraph())
+		{
+			bool bFunctionGraphExists = false;
+			for (UEdGraph* FunctionGraph : FunctionEdGraphs)
+			{
+				if (URigVMEdGraph* RigFunctionGraph = Cast<URigVMEdGraph>(FunctionGraph))
+				{
+					if (RigFunctionGraph->ModelNodePath == ContainedGraph->GetNodePath())
+					{
+						bFunctionGraphExists = true;
+						break;
+					}
+				}
+			}
+
+			if (!bFunctionGraphExists)
+			{
+				const FName SubGraphName = RigVMClient.GetUniqueName(this, *InNode->GetName());
+				// create a sub graph
+				UAnimNextEdGraph* RigFunctionGraph = NewObject<UAnimNextEdGraph>(this, SubGraphName, RF_Transactional);
+				RigFunctionGraph->Schema = UAnimNextEdGraphSchema::StaticClass();
+				RigFunctionGraph->bAllowRenaming = true;
+				RigFunctionGraph->bEditable = true;
+				RigFunctionGraph->bAllowDeletion = true;
+				RigFunctionGraph->ModelNodePath = ContainedGraph->GetNodePath();
+				RigFunctionGraph->bIsFunctionDefinition = true;
+
+				RigFunctionGraph->Initialize(this);
+
+				FunctionEdGraphs.Add(RigFunctionGraph);
+
+				RigVMClient.GetOrCreateController(ContainedGraph)->ResendAllNotifications();
+			}
+		}
+	}
+	// --- For Collapse nodes ---
+	else if (URigVMEdGraph* RigEdGraph = Cast<URigVMEdGraph>(GetEditorObjectForRigVMGraph(InNode->GetGraph())))
+	{
+		if (URigVMGraph* ContainedGraph = InNode->GetContainedGraph())
+		{
+			bool bSubGraphExists = false;
+
+			const FString ContainedGraphNodePath = ContainedGraph->GetNodePath();
+			for (UEdGraph* SubGraph : RigEdGraph->SubGraphs)
+			{
+				if (UAnimNextEdGraph* SubRigGraph = Cast<UAnimNextEdGraph>(SubGraph))
+				{
+					if (SubRigGraph->ModelNodePath == ContainedGraphNodePath)
+					{
+						bSubGraphExists = true;
+						break;
+					}
+				}
+			}
+
+			if (!bSubGraphExists)
+			{
+				bool bEditable = true;
+				if (InNode->IsA<URigVMAggregateNode>())
+				{
+					bEditable = false;
+				}
+
+				UObject* Outer = FindEntryForRigVMGraph(CollapseNodeGraph->GetRootGraph());
+				if (Outer == nullptr)
+				{
+					Outer = this; // function library graph has no entry
+				}
+
+				const FName SubGraphName = RigVMClient.GetUniqueName(Outer, *InNode->GetEditorSubGraphName());
+				// create a sub graph, no need to set external package if outer is an Entry
+				UAnimNextEdGraph* SubRigGraph = NewObject<UAnimNextEdGraph>(Outer, SubGraphName, RF_Transactional);
+				SubRigGraph->Schema = UAnimNextEdGraphSchema::StaticClass();
+				SubRigGraph->bAllowRenaming = 1;
+				SubRigGraph->bEditable = bEditable;
+				SubRigGraph->bAllowDeletion = 1;
+				SubRigGraph->ModelNodePath = ContainedGraphNodePath;
+				SubRigGraph->bIsFunctionDefinition = false;
+
+				RigEdGraph->SubGraphs.Add(SubRigGraph);
+
+				SubRigGraph->Initialize(this);
+
+				GetOrCreateController(ContainedGraph)->ResendAllNotifications();
+			}
+		}
+	}
+}
+
+void UAnimNextModule_EditorData::RemoveEdGraphForCollapseNode(URigVMCollapseNode* InNode, bool bNotify)
+{
+	check(InNode);
+
 	if (InNode->GetGraph()->IsA<URigVMFunctionLibrary>())
 	{
 		if (URigVMGraph* ContainedGraph = InNode->GetContainedGraph())
 		{
-			// create a sub graph
-			UAnimNextEdGraph* RigFunctionGraph = NewObject<UAnimNextEdGraph>(this, *InNode->GetName(), RF_Transactional);
-			RigFunctionGraph->Schema = UAnimNextEdGraphSchema::StaticClass();
-			RigFunctionGraph->bAllowRenaming = true;
-			RigFunctionGraph->bEditable = true;
-			RigFunctionGraph->bAllowDeletion = true;
-			RigFunctionGraph->ModelNodePath = ContainedGraph->GetNodePath();
-			RigFunctionGraph->bIsFunctionDefinition = true;
-
-			RigFunctionGraph->Initialize(this);
-
-			RigVMClient.GetOrCreateController(ContainedGraph)->ResendAllNotifications();
+			for (UEdGraph* FunctionGraph : FunctionEdGraphs)
+			{
+				if (URigVMEdGraph* RigFunctionGraph = Cast<URigVMEdGraph>(FunctionGraph))
+				{
+					if (RigFunctionGraph->ModelNodePath == ContainedGraph->GetNodePath())
+					{
+						if (URigVMController* SubController = GetController(ContainedGraph))
+						{
+							SubController->OnModified().RemoveAll(RigFunctionGraph);
+						}
+
+						if (RigVMGraphModifiedEvent.IsBound() && bNotify)
+						{
+							RigVMGraphModifiedEvent.Broadcast(ERigVMGraphNotifType::NodeRemoved, InNode->GetGraph(), InNode);
+						}
+
+						FunctionEdGraphs.Remove(RigFunctionGraph);
+						RigFunctionGraph->Rename(nullptr, GetTransientPackage(), REN_ForceNoResetLoaders | REN_DontCreateRedirectors);
+						RigFunctionGraph->MarkAsGarbage();
+						break;
+					}
+				}
+			}
+		}
+	}
+	else if (URigVMEdGraph* RigGraph = Cast<URigVMEdGraph>(GetEditorObjectForRigVMGraph(InNode->GetGraph())))
+	{
+		if (URigVMGraph* ContainedGraph = InNode->GetContainedGraph())
+		{
+			for (UEdGraph* SubGraph : RigGraph->SubGraphs)
+			{
+				if (URigVMEdGraph* SubRigGraph = Cast<URigVMEdGraph>(SubGraph))
+				{
+					if (SubRigGraph->ModelNodePath == ContainedGraph->GetNodePath())
+					{
+						if (URigVMController* SubController = GetController(ContainedGraph))
+						{
+							SubController->OnModified().RemoveAll(SubRigGraph);
+						}
+
+						if (RigVMGraphModifiedEvent.IsBound() && bNotify)
+						{
+							RigVMGraphModifiedEvent.Broadcast(ERigVMGraphNotifType::NodeRemoved, InNode->GetGraph(), InNode);
+						}
+
+						RigGraph->SubGraphs.Remove(SubRigGraph);
+						SubRigGraph->Rename(nullptr, GetTransientPackage(), REN_ForceNoResetLoaders | REN_DontCreateRedirectors);
+						SubRigGraph->MarkAsGarbage();
+						break;
+					}
+				}
+			}
 		}
 	}
 }
@@ -246,10 +393,10 @@ UEdGraph* UAnimNextModule_EditorData::CreateEdGraph(URigVMGraph* InRigVMGraph, b
 		RemoveEdGraph(InRigVMGraph);
 	}
 
-	FString GraphName = InRigVMGraph->GetName();
-	check(!GraphName.IsEmpty());
+	const FName GraphName = RigVMClient.GetUniqueName(CastChecked<UObject>(Entry), Entry->GetGraphName());
+	UAnimNextEdGraph* RigFunctionGraph = NewObject<UAnimNextEdGraph>(CastChecked<UObject>(Entry), GraphName, RF_Transactional);
+	RigFunctionGraph->Schema = UAnimNextEdGraphSchema::StaticClass();
 
-	UAnimNextEdGraph* RigFunctionGraph = NewObject<UAnimNextEdGraph>(CastChecked<UObject>(Entry), NAME_None, RF_Transactional);
 	RigFunctionGraph->Schema = UAnimNextEdGraphSchema::StaticClass();
 	RigFunctionGraph->bAllowDeletion = true;
 	RigFunctionGraph->bIsFunctionDefinition = false;
@@ -426,15 +573,27 @@ UAnimNextModule_AnimationGraph* UAnimNextModule_EditorData::AddAnimationGraph(FN
 	{
 		TGuardValue<bool> EnablePythonPrint(bSuspendPythonMessagesForRigVMClient, !bPrintPythonCommand);
 		TGuardValue<bool> DisableAutoCompile(bAutoRecompileVM, false);
-		URigVMGraph* NewGraph = RigVMClient.AddModel(URigVMGraph::StaticClass()->GetFName(), UAnimNextAnimationGraphSchema::StaticClass(), bSetupUndoRedo);
-		ensure(NewGraph);
-		NewEntry->Graph = NewGraph;
 
-		URigVMController* Controller = RigVMClient.GetController(NewGraph);
+		// Editor data has to be the graph outer, or RigVM unique name generator will not work
+		URigVMGraph* NewRigVMGraphModel = RigVMClient.CreateModel(URigVMGraph::StaticClass()->GetFName(), UAnimNextAnimationGraphSchema::StaticClass(), bSetupUndoRedo, this);
+		// Then, to avoid the graph losing ref due to external package, set the same package as the Entry
+		if (!NewRigVMGraphModel->HasAnyFlags(RF_Transient))
+		{
+			NewRigVMGraphModel->SetExternalPackage(CastChecked<UObject>(NewEntry)->GetExternalPackage());
+		}
+
+		ensure(NewRigVMGraphModel);
+
+		NewEntry->Graph = NewRigVMGraphModel;
+		
+		RefreshExternalModels();
+		RigVMClient.AddModel(NewRigVMGraphModel, true);
+
+		URigVMController* Controller = RigVMClient.GetController(NewRigVMGraphModel);
 		UE::AnimNext::UncookedOnly::FUtils::SetupAnimGraph(NewEntry, Controller);
 	}
 
 	BroadcastModified();
 
 	return NewEntry;
-}
\ No newline at end of file
+}
diff --git a/Engine/Plugins/Experimental/AnimNext/Source/AnimNextUncookedOnly/Private/Module/AnimNextModule_EventGraph.cpp b/Engine/Plugins/Experimental/AnimNext/Source/AnimNextUncookedOnly/Private/Module/AnimNextModule_EventGraph.cpp
--- a/Engine/Plugins/Experimental/AnimNext/Source/AnimNextUncookedOnly/Private/Module/AnimNextModule_EventGraph.cpp
+++ b/Engine/Plugins/Experimental/AnimNext/Source/AnimNextUncookedOnly/Private/Module/AnimNextModule_EventGraph.cpp
@@ -1,4 +1,4 @@
-﻿// Copyright Epic Games, Inc. All Rights Reserved.
+// Copyright Epic Games, Inc. All Rights Reserved.
 
 #include "Module/AnimNextModule_EventGraph.h"
 
@@ -28,6 +28,11 @@ void UAnimNextModule_EventGraph::SetEntryName(FName InName, bool bSetupUndoRedo)
 	BroadcastModified();
 }
 
+const FName& UAnimNextModule_EventGraph::GetGraphName() const
+{
+	return GraphName;
+}
+
 URigVMGraph* UAnimNextModule_EventGraph::GetRigVMGraph() const
 {
 	return Graph;
diff --git a/Engine/Plugins/Experimental/AnimNext/Source/AnimNextUncookedOnly/Private/UncookedOnlyUtils.cpp b/Engine/Plugins/Experimental/AnimNext/Source/AnimNextUncookedOnly/Private/UncookedOnlyUtils.cpp
--- a/Engine/Plugins/Experimental/AnimNext/Source/AnimNextUncookedOnly/Private/UncookedOnlyUtils.cpp
+++ b/Engine/Plugins/Experimental/AnimNext/Source/AnimNextUncookedOnly/Private/UncookedOnlyUtils.cpp
@@ -50,6 +50,7 @@
 #include "Param/IParameterSourceType.h"
 #include "Param/RigVMDispatch_GetScopedParameter.h"
 #include "Module/AnimNextModuleWorkspaceAssetUserData.h"
+#include "WorkspaceAssetRegistryInfo.h"
 
 #define LOCTEXT_NAMESPACE "AnimNextUncookedOnlyUtils"
 
@@ -1246,6 +1247,7 @@ void FUtils::SetupAnimGraph(UAnimNextRigVMAssetEntry* InEntry, URigVMController*
 
 	// Add root node
 	URigVMUnitNode* MainEntryPointNode = InController->AddUnitNode(FRigUnit_AnimNextGraphRoot::StaticStruct(), FRigUnit_AnimNextGraphRoot::EventName, FVector2D(-400.0f, 0.0f), FString(), false);
+	check(MainEntryPointNode);
 	URigVMPin* BeginExecutePin = MainEntryPointNode->FindPin(GET_MEMBER_NAME_STRING_CHECKED(FRigUnit_AnimNextGraphRoot, Result));
 	check(BeginExecutePin);
 	check(BeginExecutePin->GetDirection() == ERigVMPinDirection::Input);
@@ -1688,6 +1690,89 @@ void FUtils::GetAssetOutlinerItems(const UAnimNextRigVMAssetEditorData* EditorDa
 			Export.GetData().InitializeAsScriptStruct(FAnimNextGraphOutlinerData::StaticStruct());
 			FAnimNextGraphOutlinerData& GraphData = Export.GetData().GetMutable<FAnimNextGraphOutlinerData>();
 			GraphData.GraphInterface = GraphInterface->_getUObject();
+
+			if (URigVMEdGraph* RigVMEdGraph = GraphInterface->GetEdGraph())
+			{
+				CreateSubGraphsOutlinerItemsRecursive(EditorData, OutExports, Export, RigVMEdGraph);
+			}
+		}
+	}
+
+	CreateFunctionLibraryOutlinerItemsRecursive(EditorData, OutExports, AssetIdentifier, EditorData->GetRigVMGraphFunctionStore()->PublicFunctions, EditorData->GetRigVMGraphFunctionStore()->PrivateFunctions);
+}
+
+void FUtils::CreateSubGraphsOutlinerItemsRecursive(const UAnimNextRigVMAssetEditorData* EditorData, FWorkspaceOutlinerItemExports& OutExports, FWorkspaceOutlinerItemExport& ParentExport, URigVMEdGraph* RigVMEdGraph)
+{
+	// ---- Collapsed graphs ---
+	for (const TObjectPtr<UEdGraph>& SubGraph : RigVMEdGraph->SubGraphs)
+	{
+		URigVMEdGraph* EditorObject = Cast<URigVMEdGraph>(SubGraph);
+		if (IsValid(EditorObject))
+		{
+			URigVMCollapseNode* CollapseNode = CastChecked<URigVMCollapseNode>(EditorObject->GetModel()->GetOuter());
+
+			FWorkspaceOutlinerItemExport& Export = OutExports.Exports.Add_GetRef(FWorkspaceOutlinerItemExport(CollapseNode->GetFName(), ParentExport));
+			Export.GetData().InitializeAsScriptStruct(FAnimNextCollapseGraphOutlinerData::StaticStruct());
+			
+			FAnimNextCollapseGraphOutlinerData& FnGraphData = Export.GetData().GetMutable<FAnimNextCollapseGraphOutlinerData>();
+			FnGraphData.EditorObject = EditorObject;
+
+			CreateSubGraphsOutlinerItemsRecursive(EditorData, OutExports, Export, EditorObject);
+		}
+	}
+
+	// ---- Function References ---
+	TArray<URigVMEdGraphNode*> EdNodes;
+	RigVMEdGraph->GetNodesOfClass(EdNodes);
+
+	for (const URigVMEdGraphNode* EdNode : EdNodes)
+	{
+		if (URigVMFunctionReferenceNode* FunctionReferenceNode = Cast<URigVMFunctionReferenceNode>(EdNode->GetModelNode()))
+		{
+			if (URigVMLibraryNode* ReferencedNode = Cast<URigVMLibraryNode>(FunctionReferenceNode->GetReferencedFunctionHeader().LibraryPointer.GetNodeSoftPath().ResolveObject()))
+			{
+				FWorkspaceOutlinerItemExport& Export = OutExports.Exports.Add_GetRef(FWorkspaceOutlinerItemExport(ReferencedNode->GetFName(), ParentExport));
+
+				Export.GetData().InitializeAsScriptStruct(FAnimNextGraphFunctionOutlinerData::StaticStruct());
+				FAnimNextGraphFunctionOutlinerData& FnGraphData = Export.GetData().GetMutable<FAnimNextGraphFunctionOutlinerData>();
+
+				URigVMEdGraph* ContainedGraph = Cast<URigVMEdGraph>(EditorData->GetEditorObjectForRigVMGraph(ReferencedNode->GetContainedGraph()));
+				FnGraphData.EditorObject = ContainedGraph;
+
+				CreateSubGraphsOutlinerItemsRecursive(EditorData, OutExports, Export, ContainedGraph);
+			}
+		}
+	}
+}
+
+void FUtils::CreateFunctionLibraryOutlinerItemsRecursive(const UAnimNextRigVMAssetEditorData* EditorData, FWorkspaceOutlinerItemExports& OutExports, FWorkspaceOutlinerItemExport& ParentExport, const TArray<FRigVMGraphFunctionData>& PublicFunctions, const TArray<FRigVMGraphFunctionData>& PrivateFunctions)
+{
+	if (PrivateFunctions.Num() > 0 || PublicFunctions.Num() > 0)
+	{
+		FWorkspaceOutlinerItemExport& Export = OutExports.Exports.Add_GetRef(FWorkspaceOutlinerItemExport(*GetFunctionLibraryDisplayName().ToString(), ParentExport));
+
+		CreateFunctionsOutlinerItemsRecursive(EditorData, OutExports, Export, PrivateFunctions, false);
+		CreateFunctionsOutlinerItemsRecursive(EditorData, OutExports, Export, PublicFunctions, true);
+	}
+}
+
+void FUtils::CreateFunctionsOutlinerItemsRecursive(const UAnimNextRigVMAssetEditorData* EditorData, FWorkspaceOutlinerItemExports& OutExports, FWorkspaceOutlinerItemExport& ParentExport, const TArray<FRigVMGraphFunctionData>& Functions, bool bPublicFunctions)
+{
+	for (const FRigVMGraphFunctionData& FunctionData : Functions)
+	{
+		if (URigVMLibraryNode* LibraryNode = Cast<URigVMLibraryNode>(FunctionData.Header.LibraryPointer.GetNodeSoftPath().ResolveObject()))
+		{
+			if (URigVMGraph* ContainedModelGraph = LibraryNode->GetContainedGraph())
+			{
+				if (URigVMEdGraph* EditorObject = Cast<URigVMEdGraph>(EditorData->GetEditorObjectForRigVMGraph(ContainedModelGraph)))
+				{
+					FWorkspaceOutlinerItemExport& Export = OutExports.Exports.Add_GetRef(FWorkspaceOutlinerItemExport(FunctionData.Header.Name, ParentExport));
+
+					Export.GetData().InitializeAsScriptStruct(FAnimNextGraphFunctionOutlinerData::StaticStruct());
+					FAnimNextGraphFunctionOutlinerData& FnGraphData = Export.GetData().GetMutable<FAnimNextGraphFunctionOutlinerData>();
+					FnGraphData.EditorObject = EditorObject;
+				}
+			}
 		}
 	}
 }
@@ -2294,6 +2379,12 @@ uint64 FUtils::SortAndHashParameters(TArray<FAnimNextParam>& InParameters)
 	return Hash;
 }
 
+const FText& FUtils::GetFunctionLibraryDisplayName()
+{
+	static const FText FunctionLibraryName = LOCTEXT("WorkspaceFunctionLibraryName", "Function Library");
+	return FunctionLibraryName;
+}
+
 }
 
 #undef LOCTEXT_NAMESPACE
\ No newline at end of file
diff --git a/Engine/Plugins/Experimental/Workspace/Source/WorkspaceEditor/Private/Outliner/WorkspaceOutlinerTreeItem.cpp b/Engine/Plugins/Experimental/Workspace/Source/WorkspaceEditor/Private/Outliner/WorkspaceOutlinerTreeItem.cpp
--- a/Engine/Plugins/Experimental/Workspace/Source/WorkspaceEditor/Private/Outliner/WorkspaceOutlinerTreeItem.cpp
+++ b/Engine/Plugins/Experimental/Workspace/Source/WorkspaceEditor/Private/Outliner/WorkspaceOutlinerTreeItem.cpp
@@ -119,7 +119,8 @@ namespace UE::Workspace
 	{
 		if (const TSharedPtr<IWorkspaceOutlinerItemDetails> SharedFactory = FWorkspaceEditorModule::GetOutlinerItemDetails(MakeOutlinerDetailsId(Export)))
 		{
-			return SharedFactory->GetPackage(Export)->GetName();
+			const UPackage* Package = SharedFactory->GetPackage(Export);
+			return Package != nullptr ? Package->GetName() : FString();
 		}
 		else if (Export.GetParentIdentifier() == NAME_None && Export.GetAssetPath().IsValid())
 		{
diff --git a/Engine/Plugins/Experimental/Workspace/Source/WorkspaceEditor/Private/SGraphDocument.cpp b/Engine/Plugins/Experimental/Workspace/Source/WorkspaceEditor/Private/SGraphDocument.cpp
--- a/Engine/Plugins/Experimental/Workspace/Source/WorkspaceEditor/Private/SGraphDocument.cpp
+++ b/Engine/Plugins/Experimental/Workspace/Source/WorkspaceEditor/Private/SGraphDocument.cpp
@@ -27,6 +27,7 @@ void SGraphDocument::Construct(const FArguments& InArgs, TSharedRef<FWorkspaceEd
 	OnPasteNodes = InArgs._OnPasteNodes;
 	OnCanDuplicateSelectedNodes = InArgs._OnCanDuplicateSelectedNodes;
 	OnDuplicateSelectedNodes = InArgs._OnDuplicateSelectedNodes;
+	OnNodeDoubleClicked = InArgs._OnNodeDoubleClicked;
 
 	SGraphEditor::FGraphEditorEvents Events;
 	Events.OnCreateActionMenu = SGraphEditor::FOnCreateActionMenu::CreateLambda([this, OnCreateActionMenu = InArgs._OnCreateActionMenu](UEdGraph* InGraph, const FVector2D& InNodePosition, const TArray<UEdGraphPin*>& InDraggedPins, bool bInAutoExpand, SGraphEditor::FActionMenuClosed InOnMenuClosed)
@@ -50,6 +51,10 @@ void SGraphDocument::Construct(const FArguments& InArgs, TSharedRef<FWorkspaceEd
 	{
 		OnNodeTextCommitted.ExecuteIfBound(FWorkspaceEditorContext(HostingAppPtr.Pin().ToSharedRef(), EdGraph), NewText, CommitInfo, NodeBeingChanged);
 	});
+	Events.OnNodeDoubleClicked = FSingleNodeEvent::CreateLambda([this, OnNodeDoubleClicked = InArgs._OnNodeDoubleClicked](UEdGraphNode* InNode)
+	{
+		OnNodeDoubleClicked.ExecuteIfBound(FWorkspaceEditorContext(HostingAppPtr.Pin().ToSharedRef(), EdGraph), InNode);
+	});
 
 	ChildSlot
 	[
diff --git a/Engine/Plugins/Experimental/Workspace/Source/WorkspaceEditor/Private/SGraphDocument.h b/Engine/Plugins/Experimental/Workspace/Source/WorkspaceEditor/Private/SGraphDocument.h
--- a/Engine/Plugins/Experimental/Workspace/Source/WorkspaceEditor/Private/SGraphDocument.h
+++ b/Engine/Plugins/Experimental/Workspace/Source/WorkspaceEditor/Private/SGraphDocument.h
@@ -44,6 +44,8 @@ class SGraphDocument : public SCompoundWidget
 	SLATE_EVENT(FSimpleDelegate, OnNavigateHistoryBack)
 	SLATE_EVENT(FSimpleDelegate, OnNavigateHistoryForward)
 
+	SLATE_EVENT(FOnNodeDoubleClicked, OnNodeDoubleClicked)
+
 	SLATE_END_ARGS()
 
 	void Construct(const FArguments& InArgs, TSharedRef<FWorkspaceEditor> InHostingApp, UEdGraph* InGraph);
@@ -105,6 +107,9 @@ class SGraphDocument : public SCompoundWidget
 	FOnCanPerformActionOnSelectedNodes OnCanDuplicateSelectedNodes;
 	FOnDuplicateSelectedNodes OnDuplicateSelectedNodes;
 
+	// Node Double click
+	FOnNodeDoubleClicked OnNodeDoubleClicked;
+
 	friend class FWorkspaceEditorModule;
 };
 
diff --git a/Engine/Plugins/Experimental/Workspace/Source/WorkspaceEditor/Private/WorkspaceEditorModule.cpp b/Engine/Plugins/Experimental/Workspace/Source/WorkspaceEditor/Private/WorkspaceEditorModule.cpp
--- a/Engine/Plugins/Experimental/Workspace/Source/WorkspaceEditor/Private/WorkspaceEditorModule.cpp
+++ b/Engine/Plugins/Experimental/Workspace/Source/WorkspaceEditor/Private/WorkspaceEditorModule.cpp
@@ -72,8 +72,7 @@ FObjectDocumentArgs FWorkspaceEditorModule::CreateGraphDocumentArgs(const FGraph
 			.OnDuplicateSelectedNodes(InArgs.OnDuplicateSelectedNodes)
 			.OnNavigateHistoryForward_Lambda([WeakWorkspaceEditor](){ if(const TSharedPtr<FWorkspaceEditor> SharedWorkspaceEditor = WeakWorkspaceEditor.Pin()) SharedWorkspaceEditor->NavigateForward(); })
 			.OnNavigateHistoryBack_Lambda([WeakWorkspaceEditor](){ if(const TSharedPtr<FWorkspaceEditor> SharedWorkspaceEditor = WeakWorkspaceEditor.Pin()) SharedWorkspaceEditor->NavigateBack(); })
-		
-		;
+			.OnNodeDoubleClicked(InArgs.OnNodeDoubleClicked);
 	});
 	Args.OnGetTabIcon = FOnGetTabIcon::CreateLambda([](const FWorkspaceEditorContext& InContext)
 	{
diff --git a/Engine/Plugins/Experimental/Workspace/Source/WorkspaceEditor/Public/IWorkspaceEditorModule.h b/Engine/Plugins/Experimental/Workspace/Source/WorkspaceEditor/Public/IWorkspaceEditorModule.h
--- a/Engine/Plugins/Experimental/Workspace/Source/WorkspaceEditor/Public/IWorkspaceEditorModule.h
+++ b/Engine/Plugins/Experimental/Workspace/Source/WorkspaceEditor/Public/IWorkspaceEditorModule.h
@@ -130,6 +130,7 @@ using FOnDuplicateSelectedNodes = TDelegate<void(const FWorkspaceEditorContext&,
 
 using FOnGetWorkspaceDetailCustomizationInstance = TDelegate<TSharedRef<IDetailCustomization>(TWeakPtr<IWorkspaceEditor>)>;
 
+using FOnNodeDoubleClicked = TDelegate<void(const FWorkspaceEditorContext&, const UEdGraphNode* InNode)>;
 
 // Arguments used to make document widgets for graphs
 struct FGraphDocumentWidgetArgs
@@ -160,6 +161,8 @@ struct FGraphDocumentWidgetArgs
 	FOnPerformActionOnSelectedNodes OnSelectAllNodes;
 
 	FOnGraphSelectionChanged OnGraphSelectionChanged;
+
+	FOnNodeDoubleClicked OnNodeDoubleClicked;
 };
 
 // Enum describing how to open a workspace
diff --git a/Engine/Plugins/Runtime/RigVM/Source/RigVMDeveloper/Private/RigVMBlueprint.cpp b/Engine/Plugins/Runtime/RigVM/Source/RigVMDeveloper/Private/RigVMBlueprint.cpp
--- a/Engine/Plugins/Runtime/RigVM/Source/RigVMDeveloper/Private/RigVMBlueprint.cpp
+++ b/Engine/Plugins/Runtime/RigVM/Source/RigVMDeveloper/Private/RigVMBlueprint.cpp
@@ -1009,14 +1009,22 @@ void URigVMBlueprint::PostLoad()
 				bDirtyDuringLoad = true;
 			}
 			
-			PatchFunctionReferencesOnLoad();
+			GetRigVMClient()->PatchFunctionReferencesOnLoad();
+			FunctionReferenceNodeData = GetReferenceNodeData();
+
 			PatchVariableNodesOnLoad();
 			PatchVariableNodesWithIncorrectType();
 			PathDomainSpecificContentOnLoad();
 			PatchBoundVariables();
 			PatchParameterNodesOnLoad();
 			PatchLinksWithCast();
-			PatchFunctionsOnLoad();
+			
+			TMap<URigVMLibraryNode*, FRigVMGraphFunctionHeader> OldHeaders;
+			// Backwards compatibility. Store public access in the model
+			TArray<FName> BackwardsCompatiblePublicFunctions;
+			GetBackwardsCompatibilityPublicFunctions(BackwardsCompatiblePublicFunctions, OldHeaders);
+
+			GetRigVMClient()->PatchFunctionsOnLoad(GetRigVMBlueprintGeneratedClass(), BackwardsCompatiblePublicFunctions, OldHeaders);
 
 			const FRigVMClientPatchResult PinDefaultValuePatchResult = GetRigVMClient()->PatchPinDefaultValues();
 			if(PinDefaultValuePatchResult.RequiresToMarkPackageDirty())
@@ -1402,161 +1410,9 @@ void URigVMBlueprint::DecrementVMRecompileBracket()
 
 void URigVMBlueprint::RefreshAllModels(ERigVMLoadType InLoadType)
 {
-	const bool bIsPostLoad = InLoadType == ERigVMLoadType::PostLoad;
-
-	// avoid any compute if the current structure hashes match with the serialized ones
-	if(CVarRigVMEnablePostLoadHashing->GetBool() && RigVMClient.GetStructureHash() == RigVMClient.GetSerializedStructureHash())
-	{
-		if(bIsPostLoad)
-		{
-			TArray<URigVMGraph*> ModelGraphs = RigVMClient.GetAllModels(true, true);
-			Algo::Reverse(ModelGraphs);
-			for (URigVMGraph* ModelGraph : ModelGraphs)
-			{
-				URigVMController* Controller = GetOrCreateController(ModelGraph);
-				URigVMController::FRestoreLinkedPathSettings Settings;
-				Settings.bFollowCoreRedirectors = true;
-				Settings.bRelayToOrphanPins = true;
-				Controller->ProcessDetachedLinks(Settings);
-			}
-		}
-		return;
-	}
-	
-	TGuardValue<bool> IsCompilingGuard(bIsCompiling, true);
-	TGuardValue<bool> ClientIgnoreModificationsGuard(RigVMClient.bIgnoreModelNotifications, true);
-	
-	TArray<URigVMGraph*> AllModelsLeavesFirst = RigVMClient.GetAllModelsLeavesFirst(true);
-	TMap<const URigVMGraph*, TArray<URigVMController::FLinkedPath>> LinkedPaths;
-
-	if (ensure(IsInGameThread()))
-	{
-		TArray<URigVMController::FRepopulatePinsNodeData> RepopulatePinsNodesData;
-		constexpr int32 REPOPULATE_NODES_NUM_RESERVED = 800;
-		RepopulatePinsNodesData.Reserve(REPOPULATE_NODES_NUM_RESERVED);
-
-		for (URigVMGraph* Graph : AllModelsLeavesFirst)
-		{
-			URigVMController* Controller = GetOrCreateController(Graph);
-			// temporarily disable default value validation during load time, serialized values should always be accepted
-			TGuardValue<bool> PerGraphDisablePinDefaultValueValidation(Controller->bValidatePinDefaults, false);
-			TGuardValue<bool> GuardEditGraph(Graph->bEditable, true);
-			FRigVMControllerNotifGuard NotifGuard(Controller, true);
-			LinkedPaths.Add(Graph, Controller->GetLinkedPaths());
-
-			const TArray<URigVMNode*> Nodes = Graph->GetNodes();
-			if (Nodes.Num() > 0)
-			{
-				RepopulatePinsNodesData.Reset();
-
-				for (URigVMNode* Node : Nodes)
-				{
-					Controller->GenerateRepopulatePinsNodeData(RepopulatePinsNodesData, Node, true, true);
-				}
-
-#if UE_RIGVMCONTROLLER_VERBOSE_REPOPULATE
-				UE_LOG(LogRigVMDeveloper, Display, TEXT("--- Graph: [%s/%s]  - NumNodes : [%d]"), *Graph->GetOuter()->GetName(), *Graph->GetName(), RepopulatePinsNodesData.Num());
-#endif
-
-				Controller->OrphanPins(RepopulatePinsNodesData);
-				Controller->FastBreakLinkedPaths(LinkedPaths.FindChecked(Graph));
-				Controller->RepopulatePins(RepopulatePinsNodesData);
-			}
-		}
-		SetupPinRedirectorsForBackwardsCompatibility();
-	}
-
-	for (URigVMGraph* Graph : AllModelsLeavesFirst)
-	{
-		URigVMController* Controller = GetOrCreateController(Graph);
-		TGuardValue<bool> GuardEditGraph(Graph->bEditable, true);
-		FRigVMControllerNotifGuard NotifGuard(Controller, true);
-		{
-			URigVMController::FRestoreLinkedPathSettings Settings;
-			Settings.bFollowCoreRedirectors = true;
-			Settings.bRelayToOrphanPins = true;
-			Controller->RestoreLinkedPaths(LinkedPaths.FindChecked(Graph), Settings);
-		}
-
-		for(URigVMNode* ModelNode : Graph->GetNodes())
-		{
-			Controller->RemoveUnusedOrphanedPins(ModelNode);
-		}
-
-		if(bIsPostLoad)
-		{
-			for(URigVMNode* ModelNode : Graph->GetNodes())
-			{
-				if (URigVMTemplateNode* TemplateNode = Cast<URigVMTemplateNode>(ModelNode))
-				{
-					TemplateNode->InvalidateCache();
-					TemplateNode->PostLoad();
-				}
-			}
-		}
-
-#if WITH_EDITOR
-
-		if(bIsPostLoad)
-		{
-			for(URigVMNode* ModelNode : Graph->GetNodes())
-			{
-				if(URigVMUnitNode* UnitNode = Cast<URigVMUnitNode>(ModelNode))
-				{
-					if (!UnitNode->HasWildCardPin())
-					{
-						UScriptStruct* ScriptStruct = UnitNode->GetScriptStruct(); 
-						if(ScriptStruct == nullptr)
-						{
-							Controller->FullyResolveTemplateNode(UnitNode, INDEX_NONE, false);
-						}
-
-						// Try to find a deprecated template
-						if (UnitNode->GetScriptStruct() == nullptr && !UnitNode->TemplateNotation.IsNone())
-						{
-							const FRigVMTemplate* Template = FRigVMRegistry::Get().FindTemplate(UnitNode->TemplateNotation, true);
-							FRigVMTemplate::FTypeMap TypeMap = UnitNode->GetTemplatePinTypeMap();
-
-							int32 Permutation;
-							if (Template->FullyResolve(TypeMap, Permutation))
-							{
-								const FRigVMFunction* Function = Template->GetPermutation(Permutation);
-								UnitNode->ResolvedFunctionName = Function->GetName();
-							}
-						}
+	const bool bEnablePostLoadHashing = CVarRigVMEnablePostLoadHashing->GetBool();
 
-						if (UnitNode->GetScriptStruct() == nullptr)
-						{
-							static const TCHAR UnresolvedUnitNodeMessage[] = TEXT("Node %s could not be resolved.");
-							Controller->ReportErrorf(UnresolvedUnitNodeMessage, *ModelNode->GetNodePath(true));
-						}
-					}
-				}
-				if (URigVMDispatchNode* DispatchNode = Cast<URigVMDispatchNode>(ModelNode))
-				{
-					if (DispatchNode->GetFactory() == nullptr)
-					{
-						static const TCHAR UnresolvedDispatchNodeMessage[] = TEXT("Dispatch node %s has no factory..");
-						Controller->ReportErrorf(UnresolvedDispatchNodeMessage, *ModelNode->GetNodePath(true));
-					}
-					else if (!DispatchNode->HasWildCardPin())
-					{
-						if (DispatchNode->GetResolvedFunction() == nullptr)
-						{
-							Controller->FullyResolveTemplateNode(DispatchNode, INDEX_NONE, false);
-						}
-						if (DispatchNode->GetResolvedFunction() == nullptr)
-						{
-							static const TCHAR UnresolvedDispatchNodeMessage[] = TEXT("Node %s could not be resolved.");
-							Controller->ReportErrorf(UnresolvedDispatchNodeMessage, *ModelNode->GetNodePath(true));
-						}
-					}
-				}
-			}
-		}
-#endif
-
-	}
+	RigVMClient.RefreshAllModels(InLoadType, bEnablePostLoadHashing, bIsCompiling);
 }
 
 void URigVMBlueprint::OnRigVMRegistryChanged()
@@ -3515,76 +3371,6 @@ FName URigVMBlueprint::AddHostMemberVariableFromExternal(FRigVMExternalVariable
 	return NAME_None;
 }
 
-void URigVMBlueprint::PatchFunctionReferencesOnLoad()
-{
-	// If the asset was copied from one project to another, the function referenced might have a different
-	// path, even if the function is internal to the contorl rig. In that case, let's try to find the function
-	// in the local function library.
-
-	for(URigVMGraph* Model : RigVMClient)
-	{
-		TArray<URigVMNode*> Nodes = Model->GetNodes();
-		for (URigVMLibraryNode* Library : RigVMClient.GetFunctionLibrary()->GetFunctions())
-		{
-			Nodes.Append(Library->GetContainedNodes());
-		}
-		
-		for (int32 i=0; i<Nodes.Num(); ++i)
-		{
-			URigVMNode* Node = Nodes[i];
-			if (URigVMFunctionReferenceNode* FunctionReferenceNode = Cast<URigVMFunctionReferenceNode>(Node))
-			{
-				if (!FunctionReferenceNode->ReferencedNodePtr_DEPRECATED.IsValid())
-				{
-					(void)FunctionReferenceNode->ReferencedNodePtr_DEPRECATED.LoadSynchronous();
-				}
-				if (!FunctionReferenceNode->ReferencedNodePtr_DEPRECATED)
-				{
-					if(URigVMFunctionLibrary* FunctionLibrary = RigVMClient.GetFunctionLibrary())
-					{
-						FString FunctionPath = FunctionReferenceNode->ReferencedNodePtr_DEPRECATED.ToSoftObjectPath().GetSubPathString();
-						
-						FString Left, Right;
-						if(FunctionPath.Split(TEXT("."), &Left, &Right))
-						{
-							FString LibraryNodePath = FunctionLibrary->GetNodePath();
-							if(Left == FunctionLibrary->GetName())
-							{
-								if (URigVMLibraryNode* LibraryNode = Cast<URigVMLibraryNode>(FunctionLibrary->FindNode(Right)))
-								{
-									FunctionReferenceNode->ReferencedNodePtr_DEPRECATED = LibraryNode;
-								}
-							}
-						}
-					}
-				}
-
-				if (FunctionReferenceNode->ReferencedNodePtr_DEPRECATED.IsValid())
-				{
-					FunctionReferenceNode->ReferencedFunctionHeader = FunctionReferenceNode->ReferencedNodePtr_DEPRECATED->GetFunctionHeader();
-				}
-				else if (!FunctionReferenceNode->ReferencedNodePtr_DEPRECATED.IsNull())
-				{
-					// At least lets make sure we store the path in the header
-					FunctionReferenceNode->ReferencedFunctionHeader.LibraryPointer.SetLibraryNodePath(FunctionReferenceNode->ReferencedNodePtr_DEPRECATED.ToSoftObjectPath().ToString());
-				}
-
-				if (FunctionReferenceNode->ReferencedFunctionHeader.LibraryPointer.LibraryNode_DEPRECATED.IsValid())
-				{
-					FunctionReferenceNode->ReferencedFunctionHeader.LibraryPointer.SetLibraryNodePath(FunctionReferenceNode->ReferencedFunctionHeader.LibraryPointer.LibraryNode_DEPRECATED.ToString());
-				}
-			}
-
-			if (URigVMCollapseNode* CollapseNode = Cast<URigVMCollapseNode>(Node))
-			{
-				Nodes.Append(CollapseNode->GetContainedNodes());
-			}
-			
-		}
-	}
-	FunctionReferenceNodeData = GetReferenceNodeData();
-}
-
 #endif
 
 void URigVMBlueprint::PatchVariableNodesOnLoad()
@@ -3736,71 +3522,6 @@ void URigVMBlueprint::PatchLinksWithCast()
 #endif
 }
 
-void URigVMBlueprint::PatchFunctionsOnLoad()
-{
-	URigVMBlueprintGeneratedClass* CRGeneratedClass = GetRigVMBlueprintGeneratedClass();
-	FRigVMGraphFunctionStore& Store = CRGeneratedClass->GraphFunctionStore;
-	URigVMFunctionLibrary* FunctionLibrary = GetLocalFunctionLibrary();
-
-	TMap<URigVMLibraryNode*, FRigVMGraphFunctionHeader> OldHeaders;
-
-	// Backwards compatibility. Store public access in the model
-	TArray<FName> BackwardsCompatiblePublicFunctions;
-	GetBackwardsCompatibilityPublicFunctions(BackwardsCompatiblePublicFunctions, OldHeaders);
-
-	// Lets rebuild the FunctionStore from the model
-	if (FunctionLibrary)
-	{
-		Store.PublicFunctions.Reset();
-		Store.PrivateFunctions.Reset();
-
-		for (URigVMLibraryNode* LibraryNode : FunctionLibrary->GetFunctions())
-		{
-			bool bIsPublic = FunctionLibrary->IsFunctionPublic(LibraryNode->GetFName());
-			if (!bIsPublic)
-			{
-				bIsPublic = BackwardsCompatiblePublicFunctions.Contains(LibraryNode->GetFName());
-				if (bIsPublic)
-				{
-					FunctionLibrary->PublicFunctionNames.Add(LibraryNode->GetFName());
-				}
-			}
-
-			FRigVMGraphFunctionHeader Header = LibraryNode->GetFunctionHeader(CRGeneratedClass);
-			if (FRigVMGraphFunctionHeader* OldHeader = OldHeaders.Find(LibraryNode))
-			{				
-				Header.ExternalVariables = OldHeader->ExternalVariables;
-				Header.Dependencies = OldHeader->Dependencies;
-				Header.Layout = OldHeader->Layout;
-			}
-
-			const FRigVMVariant* Variant = FunctionLibrary->GetFunctionVariant(LibraryNode->GetFName());
-			if (!Variant)
-			{
-				Header.Variant.Guid = FRigVMVariant::GenerateGUID(Header.LibraryPointer.GetLibraryNodePath());
-				FunctionLibrary->FunctionToVariant.FindOrAdd(Header.Name) = Header.Variant;
-			}
-			else
-			{
-				Header.Variant = *Variant;
-			}
-			
-			Store.AddFunction(Header, bIsPublic);
-			if (bIsPublic)
-			{
-				GetRigVMClient()->UpdateGraphFunctionSerializedGraph(LibraryNode);
-			}
-		}
-
-		// Update dependencies and external variables if needed
-		for (URigVMLibraryNode* LibraryNode : FunctionLibrary->GetFunctions())
-		{
-			GetRigVMClient()->UpdateExternalVariablesForFunction(LibraryNode);
-			GetRigVMClient()->UpdateDependenciesForFunction(LibraryNode);
-		}
-	}
-}
-
 void URigVMBlueprint::GetBackwardsCompatibilityPublicFunctions(TArray<FName>& BackwardsCompatiblePublicFunctions, TMap<URigVMLibraryNode*, FRigVMGraphFunctionHeader>& OldHeaders)
 {
 	URigVMBlueprintGeneratedClass* CRGeneratedClass = GetRigVMBlueprintGeneratedClass();
diff --git a/Engine/Plugins/Runtime/RigVM/Source/RigVMDeveloper/Private/RigVMCompiler/RigVMCompiler.cpp b/Engine/Plugins/Runtime/RigVM/Source/RigVMDeveloper/Private/RigVMCompiler/RigVMCompiler.cpp
--- a/Engine/Plugins/Runtime/RigVM/Source/RigVMDeveloper/Private/RigVMCompiler/RigVMCompiler.cpp
+++ b/Engine/Plugins/Runtime/RigVM/Source/RigVMDeveloper/Private/RigVMCompiler/RigVMCompiler.cpp
@@ -611,7 +611,16 @@ bool URigVMCompiler::Compile(const FRigVMCompileSettings& InSettings, TArray<URi
 			}
 			if (URigVMCollapseNode* CollapseNode = Cast<URigVMCollapseNode>(Nodes[i]))
 			{
-				Nodes.Append(CollapseNode->GetContainedGraph()->GetNodes());
+				if (CollapseNode->GetContainedGraph())
+				{
+					Nodes.Append(CollapseNode->GetContainedGraph()->GetNodes());
+				}
+				else
+				{
+					static const FString FunctionCompilationErrorMessage = TEXT("Could not find contained graph for collapse node @@.");
+					Settings.ASTSettings.Report(EMessageSeverity::Error, CollapseNode, FunctionCompilationErrorMessage);
+					bEncounteredGraphError = true;
+				}
 			}
 		}
 	}
diff --git a/Engine/Plugins/Runtime/RigVM/Source/RigVMDeveloper/Private/RigVMModel/RigVMClient.cpp b/Engine/Plugins/Runtime/RigVM/Source/RigVMDeveloper/Private/RigVMModel/RigVMClient.cpp
--- a/Engine/Plugins/Runtime/RigVM/Source/RigVMDeveloper/Private/RigVMModel/RigVMClient.cpp
+++ b/Engine/Plugins/Runtime/RigVM/Source/RigVMDeveloper/Private/RigVMModel/RigVMClient.cpp
@@ -7,6 +7,7 @@
 #include "UnrealExporter.h"
 #include "UObject/ObjectSaveContext.h"
 #include "RigVMModel/RigVMControllerActions.h"
+#include "RigVMModel/Nodes/RigVMDispatchNode.h"
 #include "EdGraph/RigVMEdGraph.h"
 
 #include UE_INLINE_GENERATED_CPP_BY_NAME(RigVMClient)
@@ -46,7 +47,7 @@ void FRigVMClient::SetControllerClass(TSubclassOf<URigVMController> InController
 		return;
 	}
 
-	for (URigVMGraph* Model : Models)
+	for (URigVMGraph* Model : GetModels())
 	{
 		RemoveController(Model);
 	}
@@ -99,9 +100,17 @@ void FRigVMClient::SetFromDeprecatedData(URigVMGraph* InDefaultGraph, URigVMFunc
 	}
 }
 
+void FRigVMClient::SetExternalModelHost(IRigVMClientExternalModelHost* InExternalModelHost)
+{
+	// We dont allow this to be set dynamically, assume it is set once at creation time
+	check(ExternalModelHost == nullptr);
+
+	ExternalModelHost = InExternalModelHost;
+}
+
 void FRigVMClient::Reset()
 {
-	for(URigVMGraph* Model : Models)
+	for(URigVMGraph* Model : GetModels())
 	{
 		DestroyObject(Model);
 	}
@@ -126,13 +135,24 @@ URigVMSchema* FRigVMClient::GetDefaultSchema() const
 
 URigVMGraph* FRigVMClient::GetDefaultModel() const
 {
-	if(Models.IsEmpty())
+	if(GetModels().IsEmpty())
 	{
 		return nullptr;
 	}
 	return GetModel(0);
 }
 
+URigVMGraph* FRigVMClient::GetModel(int32 InIndex) const
+{
+	const TArray<TObjectPtr<URigVMGraph>>& LocalModels = GetModels();
+	if(LocalModels.IsValidIndex(InIndex))
+	{
+		return LocalModels[InIndex];
+	}
+	return nullptr;
+}
+
+
 URigVMGraph* FRigVMClient::GetModel(const UEdGraph* InEdGraph) const
 {
 	if (InEdGraph == nullptr)
@@ -195,12 +215,180 @@ URigVMGraph* FRigVMClient::GetModel(const UObject* InEditorSideObject) const
 	return nullptr;
 }
 
+void FRigVMClient::RefreshAllModels(ERigVMLoadType InLoadType, bool bEnablePostLoadHashing, bool& bIsCompiling)
+{
+	const bool bIsPostLoad = InLoadType == ERigVMLoadType::PostLoad;
+
+	// avoid any compute if the current structure hashes match with the serialized ones
+	if (bEnablePostLoadHashing && GetStructureHash() == GetSerializedStructureHash())
+	{
+		if (bIsPostLoad)
+		{
+			TArray<URigVMGraph*> ModelGraphs = GetAllModels(true, true);
+			Algo::Reverse(ModelGraphs);
+			for (URigVMGraph* ModelGraph : ModelGraphs)
+			{
+				URigVMController* Controller = GetOrCreateController(ModelGraph);
+				URigVMController::FRestoreLinkedPathSettings Settings;
+				Settings.bFollowCoreRedirectors = true;
+				Settings.bRelayToOrphanPins = true;
+				Controller->ProcessDetachedLinks(Settings);
+			}
+		}
+		return;
+	}
+
+	TGuardValue<bool> IsCompilingGuard(bIsCompiling, true);
+	TGuardValue<bool> ClientIgnoreModificationsGuard(bIgnoreModelNotifications, true);
+
+	TArray<URigVMGraph*> AllModelsLeavesFirst = GetAllModelsLeavesFirst(true);
+	TMap<const URigVMGraph*, TArray<URigVMController::FLinkedPath>> LinkedPaths;
+
+	if (ensure(IsInGameThread()))
+	{
+		TArray<URigVMController::FRepopulatePinsNodeData> RepopulatePinsNodesData;
+		constexpr int32 REPOPULATE_NODES_NUM_RESERVED = 800;
+		RepopulatePinsNodesData.Reserve(REPOPULATE_NODES_NUM_RESERVED);
+
+		for (URigVMGraph* Graph : AllModelsLeavesFirst)
+		{
+			URigVMController* Controller = GetOrCreateController(Graph);
+			// temporarily disable default value validation during load time, serialized values should always be accepted
+			TGuardValue<bool> PerGraphDisablePinDefaultValueValidation(Controller->bValidatePinDefaults, false);
+			TGuardValue<bool> GuardEditGraph(Graph->bEditable, true);
+			FRigVMControllerNotifGuard NotifGuard(Controller, true);
+			LinkedPaths.Add(Graph, Controller->GetLinkedPaths());
+
+			const TArray<URigVMNode*> Nodes = Graph->GetNodes();
+			if (Nodes.Num() > 0)
+			{
+				RepopulatePinsNodesData.Reset();
+
+				for (URigVMNode* Node : Nodes)
+				{
+					Controller->GenerateRepopulatePinsNodeData(RepopulatePinsNodesData, Node, true, true);
+				}
+
+#if UE_RIGVMCONTROLLER_VERBOSE_REPOPULATE
+				UE_LOG(LogRigVMDeveloper, Display, TEXT("--- Graph: [%s/%s]  - NumNodes : [%d]"), *Graph->GetOuter()->GetName(), *Graph->GetName(), RepopulatePinsNodesData.Num());
+#endif
+
+				Controller->OrphanPins(RepopulatePinsNodesData);
+				Controller->FastBreakLinkedPaths(LinkedPaths.FindChecked(Graph));
+				Controller->RepopulatePins(RepopulatePinsNodesData);
+			}
+		}
+
+		if (IRigVMClientHost* ClientHost = Cast<IRigVMClientHost>(GetOuter()))
+		{
+			ClientHost->SetupPinRedirectorsForBackwardsCompatibility();
+		}
+	}
+
+	for (URigVMGraph* Graph : AllModelsLeavesFirst)
+	{
+		URigVMController* Controller = GetOrCreateController(Graph);
+		TGuardValue<bool> GuardEditGraph(Graph->bEditable, true);
+		FRigVMControllerNotifGuard NotifGuard(Controller, true);
+		{
+			URigVMController::FRestoreLinkedPathSettings Settings;
+			Settings.bFollowCoreRedirectors = true;
+			Settings.bRelayToOrphanPins = true;
+			Controller->RestoreLinkedPaths(LinkedPaths.FindChecked(Graph), Settings);
+		}
+
+		for (URigVMNode* ModelNode : Graph->GetNodes())
+		{
+			Controller->RemoveUnusedOrphanedPins(ModelNode);
+		}
+
+		if (bIsPostLoad)
+		{
+			for (URigVMNode* ModelNode : Graph->GetNodes())
+			{
+				if (URigVMTemplateNode* TemplateNode = Cast<URigVMTemplateNode>(ModelNode))
+				{
+					TemplateNode->InvalidateCache();
+					TemplateNode->PostLoad();
+				}
+			}
+		}
+
+#if WITH_EDITOR
+
+		if (bIsPostLoad)
+		{
+			for (URigVMNode* ModelNode : Graph->GetNodes())
+			{
+				if (URigVMUnitNode* UnitNode = Cast<URigVMUnitNode>(ModelNode))
+				{
+					if (!UnitNode->HasWildCardPin())
+					{
+						UScriptStruct* ScriptStruct = UnitNode->GetScriptStruct();
+						if (ScriptStruct == nullptr)
+						{
+							Controller->FullyResolveTemplateNode(UnitNode, INDEX_NONE, false);
+						}
+
+						// Try to find a deprecated template
+						if (UnitNode->GetScriptStruct() == nullptr && !UnitNode->TemplateNotation.IsNone())
+						{
+							const FRigVMTemplate* Template = FRigVMRegistry::Get().FindTemplate(UnitNode->TemplateNotation, true);
+							FRigVMTemplate::FTypeMap TypeMap = UnitNode->GetTemplatePinTypeMap();
+
+							int32 Permutation;
+							if (Template->FullyResolve(TypeMap, Permutation))
+							{
+								const FRigVMFunction* Function = Template->GetPermutation(Permutation);
+								UnitNode->ResolvedFunctionName = Function->GetName();
+							}
+						}
+
+						if (UnitNode->GetScriptStruct() == nullptr)
+						{
+							static const TCHAR UnresolvedUnitNodeMessage[] = TEXT("Node %s could not be resolved.");
+							Controller->ReportErrorf(UnresolvedUnitNodeMessage, *ModelNode->GetNodePath(true));
+						}
+					}
+				}
+				if (URigVMDispatchNode* DispatchNode = Cast<URigVMDispatchNode>(ModelNode))
+				{
+					if (DispatchNode->GetFactory() == nullptr)
+					{
+						static const TCHAR UnresolvedDispatchNodeMessage[] = TEXT("Dispatch node %s has no factory..");
+						Controller->ReportErrorf(UnresolvedDispatchNodeMessage, *ModelNode->GetNodePath(true));
+					}
+					else if (!DispatchNode->HasWildCardPin())
+					{
+						if (DispatchNode->GetResolvedFunction() == nullptr)
+						{
+							Controller->FullyResolveTemplateNode(DispatchNode, INDEX_NONE, false);
+						}
+						if (DispatchNode->GetResolvedFunction() == nullptr)
+						{
+							static const TCHAR UnresolvedDispatchNodeMessage[] = TEXT("Node %s could not be resolved.");
+							Controller->ReportErrorf(UnresolvedDispatchNodeMessage, *ModelNode->GetNodePath(true));
+						}
+					}
+				}
+			}
+		}
+#endif
+
+	}
+}
+
+const TArray<TObjectPtr<URigVMGraph>>& FRigVMClient::GetModels() const
+{
+	return ExternalModelHost != nullptr ? ExternalModelHost->GetExternalModels() : Models;
+}
+
 TArray<URigVMGraph*> FRigVMClient::GetAllModels(bool bIncludeFunctionLibrary, bool bRecursive) const
 {
-	TArray<URigVMGraph*> AllModels = Models;
+	TArray<URigVMGraph*> AllModels = GetModels();
 	if(bRecursive)
 	{
-		for(URigVMGraph* Model : Models)
+		for(URigVMGraph* Model : GetModels())
 		{
 			AllModels.Append(Model->GetContainedGraphs(true /* recursive */));
 		}
@@ -351,10 +539,20 @@ URigVMGraph* FRigVMClient::AddModel(const FString InName, bool bSetupUndoRedo, b
 URigVMGraph* FRigVMClient::AddModel(const FName& InName, bool bSetupUndoRedo, const FObjectInitializer* ObjectInitializer, bool bCreateController)
 {
 	check(DefaultSchemaClass);
-	return AddModel(InName, DefaultSchemaClass, bSetupUndoRedo, ObjectInitializer, bCreateController);
+
+	URigVMGraph* NewModel = CreateModel(InName, DefaultSchemaClass, bSetupUndoRedo, GetOuter(), ObjectInitializer, bCreateController);
+	AddModel(NewModel, bCreateController);
+	return NewModel;
 }
 
 URigVMGraph* FRigVMClient::AddModel(const FName& InName, TSubclassOf<URigVMSchema> InSchemaClass, bool bSetupUndoRedo, const FObjectInitializer* ObjectInitializer, bool bCreateController)
+{
+	URigVMGraph* NewModel = CreateModel(InName, InSchemaClass, bSetupUndoRedo, GetOuter(), ObjectInitializer, bCreateController);
+	AddModel(NewModel, bCreateController);
+	return NewModel;
+}
+
+URigVMGraph* FRigVMClient::CreateModel(const FName& InName, TSubclassOf<URigVMSchema> InSchemaClass, bool bSetupUndoRedo, UObject* InOuter, const FObjectInitializer* ObjectInitializer, bool bCreateController)
 {
 	check(InSchemaClass);
 
@@ -363,7 +561,7 @@ URigVMGraph* FRigVMClient::AddModel(const FName& InName, TSubclassOf<URigVMSchem
 	if(bSetupUndoRedo)
 	{
 		Transaction = MakeShared<FScopedTransaction>(NSLOCTEXT("RigVMClient", "AddModel", "Add new root graph"));
-		GetOuter()->Modify();
+		InOuter->Modify();
 	}
 #endif
 
@@ -371,20 +569,18 @@ URigVMGraph* FRigVMClient::AddModel(const FName& InName, TSubclassOf<URigVMSchem
 	URigVMGraph* Model = nullptr;
 	if(ObjectInitializer)
 	{
-		Model = ObjectInitializer->CreateDefaultSubobject<URigVMGraph>(GetOuter(), SafeGraphName);
+		Model = ObjectInitializer->CreateDefaultSubobject<URigVMGraph>(InOuter, SafeGraphName);
 	}
 	else
 	{
-		Model = NewObject<URigVMGraph>(GetOuter(), SafeGraphName);
+		Model = NewObject<URigVMGraph>(InOuter, SafeGraphName);
 	}
 
 	Model->SetSchemaClass(InSchemaClass);
 
-	AddModel(Model, bCreateController);
-
 	if(bSetupUndoRedo)
 	{
-		GetOuter()->Modify();
+		InOuter->Modify();
 		UndoRedoIndex++;
 
 		FRigVMClientAction Action;
@@ -396,6 +592,24 @@ URigVMGraph* FRigVMClient::AddModel(const FName& InName, TSubclassOf<URigVMSchem
 	return Model;
 }
 
+TObjectPtr<URigVMGraph> FRigVMClient::CreateContainedGraphModel(URigVMCollapseNode* CollapseNode, const FName& Name)
+{
+	TObjectPtr<URigVMGraph> Model = nullptr;
+
+	check(CollapseNode);
+
+	if (ExternalModelHost != nullptr)
+	{
+		Model = ExternalModelHost->CreateContainedGraphModel(CollapseNode, Name);
+	}
+	else
+	{
+		Model = NewObject<URigVMGraph>(CollapseNode, Name);
+	}
+
+	return Model;
+}
+
 void FRigVMClient::AddModel(URigVMGraph* InModel, bool bCreateController)
 {
 	check(InModel);
@@ -407,7 +621,10 @@ void FRigVMClient::AddModel(URigVMGraph* InModel, bool bCreateController)
 	}
 	else
 	{
-		Models.Add(InModel);
+		if(ExternalModelHost == nullptr)
+		{
+			Models.Add(InModel);
+		}
 	}
 
 	InModel->SetExecuteContextStruct(InModel->GetSchema()->GetExecuteContextStruct());
@@ -419,7 +636,7 @@ void FRigVMClient::AddModel(URigVMGraph* InModel, bool bCreateController)
 
 	if(InModel->IsA<URigVMFunctionLibrary>())
 	{
-		for(URigVMGraph* Model : Models)
+		for(URigVMGraph* Model : GetModels())
 		{
 			Model->SetDefaultFunctionLibrary(FunctionLibrary);
 		}
@@ -446,10 +663,36 @@ URigVMFunctionLibrary* FRigVMClient::GetOrCreateFunctionLibrary(bool bSetupUndoR
 	return GetOrCreateFunctionLibrary(DefaultSchemaClass, bSetupUndoRedo, ObjectInitializer, bCreateController);
 }
 
+static void SetGetFunctionHostObjectPathDelegate(FRigVMClient* RigVMClient, URigVMFunctionLibrary* FunctionLibrary)
+{
+	if (RigVMClient->GetOuter()->Implements<URigVMClientHost>())
+	{
+		if (IRigVMClientHost* ClientHost = Cast<IRigVMClientHost>(RigVMClient->GetOuter()))
+		{
+			TWeakObjectPtr<UObject> WeakClientHost = Cast<UObject>(ClientHost);
+			FunctionLibrary->GetFunctionHostObjectPathDelegate.BindLambda([WeakClientHost]() -> const FSoftObjectPath
+				{
+					if (WeakClientHost.IsValid())
+					{
+						if (IRigVMClientHost* ClientHost = Cast<IRigVMClientHost>(WeakClientHost.Get()))
+						{
+							return Cast<UObject>(ClientHost->GetRigVMGraphFunctionHost());
+						}
+					}
+					return nullptr;
+				});
+		}
+	}
+};
+
 URigVMFunctionLibrary* FRigVMClient::GetOrCreateFunctionLibrary(TSubclassOf<URigVMSchema> InSchemaClass, bool bSetupUndoRedo, const FObjectInitializer* ObjectInitializer, bool bCreateController)
 {
 	if(FunctionLibrary)
 	{
+		if (!FunctionLibrary->GetFunctionHostObjectPathDelegate.IsBound())
+		{
+			SetGetFunctionHostObjectPathDelegate(this, FunctionLibrary);
+		}
 		return FunctionLibrary;
 	}
 	
@@ -476,16 +719,11 @@ URigVMFunctionLibrary* FRigVMClient::GetOrCreateFunctionLibrary(TSubclassOf<URig
 
 	NewFunctionLibrary->SetSchemaClass(InSchemaClass);
 
-	NewFunctionLibrary->GetFunctionHostObjectPathDelegate.BindLambda([this]() -> const FSoftObjectPath 
-		{
-			if (GetOuter()->Implements<URigVMClientHost>())
-			{
-				IRigVMClientHost* ClientHost = Cast<IRigVMClientHost>(GetOuter());
-				return Cast<UObject>(ClientHost->GetRigVMGraphFunctionHost());				
-			}
-			return nullptr;
-		});
-	
+	if (!NewFunctionLibrary->GetFunctionHostObjectPathDelegate.IsBound())
+	{
+		SetGetFunctionHostObjectPathDelegate(this, NewFunctionLibrary);
+	}
+
 	AddModel(NewFunctionLibrary, bCreateController);
 	return NewFunctionLibrary;
 }
@@ -493,7 +731,7 @@ URigVMFunctionLibrary* FRigVMClient::GetOrCreateFunctionLibrary(TSubclassOf<URig
 TArray<FName> FRigVMClient::GetEntryNames() const
 {
 	TArray<FName> EntryNames;
-	for(const URigVMGraph* Model : Models)
+	for(const URigVMGraph* Model : GetModels())
 	{
 		for(const URigVMNode* Node : Model->GetNodes())
 		{
@@ -539,6 +777,9 @@ bool FRigVMClient::RemoveModel(const FString& InNodePathOrName, bool bSetupUndoR
 {
 	if(URigVMGraph* Model = GetModel(InNodePathOrName))
 	{
+		UObject* ModelOuter = Model->GetOuter();
+		check(ModelOuter);
+		
 		if(Model == GetDefaultModel() && !bDefaultModelCanBeRemoved)
 		{
 #if WITH_EDITOR
@@ -562,7 +803,7 @@ bool FRigVMClient::RemoveModel(const FString& InNodePathOrName, bool bSetupUndoR
 		if(bSetupUndoRedo)
 		{
 			Transaction = MakeShared<FScopedTransaction>(NSLOCTEXT("RigVMClient", "RemoveModel", "Remove a root graph"));
-			GetOuter()->Modify();
+			ModelOuter->Modify();
 		}
 #endif
 
@@ -574,7 +815,7 @@ bool FRigVMClient::RemoveModel(const FString& InNodePathOrName, bool bSetupUndoR
 
 		if(bSetupUndoRedo)
 		{
-			GetOuter()->Modify();
+			ModelOuter->Modify();
 			UndoRedoIndex++;
 
 			FRigVMClientAction Action;
@@ -585,10 +826,15 @@ bool FRigVMClient::RemoveModel(const FString& InNodePathOrName, bool bSetupUndoR
 		}
 
 		// clean up the model
-		Models.Remove(Model);
-		
-		// remove the controller
-		verify(RemoveController(Model));
+		if(ExternalModelHost == nullptr)
+		{
+			verify(Models.Remove(Model));
+		}
+		else
+		{
+			// Should have already been removed from external models
+			verify(!ExternalModelHost->GetExternalModels().Contains(Model));
+		}
 
 		NotifyOuterOfPropertyChange();
 		return true;
@@ -772,7 +1018,7 @@ void FRigVMClient::OnCollapseNodeRemoved(const URigVMCollapseNode* InCollapseNod
 
 URigVMNode* FRigVMClient::FindNode(const FString& InNodePathOrName) const
 {
-	for(const URigVMGraph* Model : Models)
+	for(const URigVMGraph* Model : GetModels())
 	{
 		if(URigVMNode* Node = Model->FindNode(InNodePathOrName))
 		{
@@ -788,7 +1034,7 @@ URigVMNode* FRigVMClient::FindNode(const FString& InNodePathOrName) const
 
 URigVMPin* FRigVMClient::FindPin(const FString& InPinPath) const
 {
-	for(const URigVMGraph* Model : Models)
+	for(const URigVMGraph* Model : GetModels())
 	{
 		if(URigVMPin* Pin = Model->FindPin(InPinPath))
 		{
@@ -974,6 +1220,135 @@ FRigVMClientPatchResult FRigVMClient::PatchModelsOnLoad()
 	return Result;
 }
 
+#if WITH_EDITOR
+void FRigVMClient::PatchFunctionReferencesOnLoad()
+{
+	// If the asset was copied from one project to another, the function referenced might have a different
+	// path, even if the function is internal to the contorl rig. In that case, let's try to find the function
+	// in the local function library.
+
+	for (URigVMGraph* Model : *this)
+	{
+		TArray<URigVMNode*> Nodes = Model->GetNodes();
+		for (URigVMLibraryNode* Library : GetFunctionLibrary()->GetFunctions())
+		{
+			Nodes.Append(Library->GetContainedNodes());
+		}
+
+		for (int32 i = 0; i < Nodes.Num(); ++i)
+		{
+			URigVMNode* Node = Nodes[i];
+			if (URigVMFunctionReferenceNode* FunctionReferenceNode = Cast<URigVMFunctionReferenceNode>(Node))
+			{
+				if (!FunctionReferenceNode->ReferencedNodePtr_DEPRECATED.IsValid())
+				{
+					(void)FunctionReferenceNode->ReferencedNodePtr_DEPRECATED.LoadSynchronous();
+				}
+				if (!FunctionReferenceNode->ReferencedNodePtr_DEPRECATED)
+				{
+					if (FunctionLibrary)
+					{
+						FString FunctionPath = FunctionReferenceNode->ReferencedNodePtr_DEPRECATED.ToSoftObjectPath().GetSubPathString();
+
+						FString Left, Right;
+						if (FunctionPath.Split(TEXT("."), &Left, &Right))
+						{
+							FString LibraryNodePath = FunctionLibrary->GetNodePath();
+							if (Left == FunctionLibrary->GetName())
+							{
+								if (URigVMLibraryNode* LibraryNode = Cast<URigVMLibraryNode>(FunctionLibrary->FindNode(Right)))
+								{
+									FunctionReferenceNode->ReferencedNodePtr_DEPRECATED = LibraryNode;
+								}
+							}
+						}
+					}
+				}
+
+				if (FunctionReferenceNode->ReferencedNodePtr_DEPRECATED.IsValid())
+				{
+					FunctionReferenceNode->ReferencedFunctionHeader = FunctionReferenceNode->ReferencedNodePtr_DEPRECATED->GetFunctionHeader();
+				}
+				else if (!FunctionReferenceNode->ReferencedNodePtr_DEPRECATED.IsNull())
+				{
+					// At least lets make sure we store the path in the header
+					FunctionReferenceNode->ReferencedFunctionHeader.LibraryPointer.SetLibraryNodePath(FunctionReferenceNode->ReferencedNodePtr_DEPRECATED.ToSoftObjectPath().ToString());
+				}
+
+				if (FunctionReferenceNode->ReferencedFunctionHeader.LibraryPointer.LibraryNode_DEPRECATED.IsValid())
+				{
+					FunctionReferenceNode->ReferencedFunctionHeader.LibraryPointer.SetLibraryNodePath(FunctionReferenceNode->ReferencedFunctionHeader.LibraryPointer.LibraryNode_DEPRECATED.ToString());
+				}
+			}
+
+			if (URigVMCollapseNode* CollapseNode = Cast<URigVMCollapseNode>(Node))
+			{
+				Nodes.Append(CollapseNode->GetContainedNodes());
+			}
+		}
+	}
+}
+
+#endif // WITH_EDITOR
+
+void FRigVMClient::PatchFunctionsOnLoad(IRigVMGraphFunctionHost* FunctionHost, TArray<FName>& BackwardsCompatiblePublicFunctions, TMap<URigVMLibraryNode*, FRigVMGraphFunctionHeader>& OldHeaders)
+{
+	if (FRigVMGraphFunctionStore* Store = FunctionHost->GetRigVMGraphFunctionStore())
+	{
+		// Lets rebuild the FunctionStore from the model
+		if (FunctionLibrary)
+		{
+			Store->PublicFunctions.Reset();
+			Store->PrivateFunctions.Reset();
+
+			for (URigVMLibraryNode* LibraryNode : FunctionLibrary->GetFunctions())
+			{
+				bool bIsPublic = FunctionLibrary->IsFunctionPublic(LibraryNode->GetFName());
+				if (!bIsPublic)
+				{
+					bIsPublic = BackwardsCompatiblePublicFunctions.Contains(LibraryNode->GetFName());
+					if (bIsPublic)
+					{
+						FunctionLibrary->PublicFunctionNames.Add(LibraryNode->GetFName());
+					}
+				}
+
+				FRigVMGraphFunctionHeader Header = LibraryNode->GetFunctionHeader(FunctionHost);
+				if (FRigVMGraphFunctionHeader* OldHeader = OldHeaders.Find(LibraryNode))
+				{
+					Header.ExternalVariables = OldHeader->ExternalVariables;
+					Header.Dependencies = OldHeader->Dependencies;
+					Header.Layout = OldHeader->Layout;
+				}
+
+				const FRigVMVariant* Variant = FunctionLibrary->GetFunctionVariant(LibraryNode->GetFName());
+				if (!Variant)
+				{
+					Header.Variant.Guid = FRigVMVariant::GenerateGUID(Header.LibraryPointer.GetLibraryNodePath());
+					FunctionLibrary->FunctionToVariant.FindOrAdd(Header.Name) = Header.Variant;
+				}
+				else
+				{
+					Header.Variant = *Variant;
+				}
+
+				Store->AddFunction(Header, bIsPublic);
+				if (bIsPublic)
+				{
+					UpdateGraphFunctionSerializedGraph(LibraryNode);
+				}
+			}
+
+			// Update dependencies and external variables if needed
+			for (URigVMLibraryNode* LibraryNode : FunctionLibrary->GetFunctions())
+			{
+				UpdateExternalVariablesForFunction(LibraryNode);
+				UpdateDependenciesForFunction(LibraryNode);
+			}
+		}
+	}
+}
+
 FRigVMClientPatchResult FRigVMClient::PatchPinDefaultValues()
 {
 	FRigVMClientPatchResult Result;
@@ -1322,14 +1697,17 @@ bool FRigVMClient::UpdateFunctionReferences(const FRigVMGraphFunctionHeader& Hea
 
 bool FRigVMClient::UpdateGraphFunctionData(const URigVMLibraryNode* InLibraryNode)
 {
-	if (GetOuter()->Implements<URigVMClientHost>())
+	if (IRigVMClientHost* ClientHost = InLibraryNode->GetImplementingOuter<IRigVMClientHost>())
 	{
-		if (FRigVMGraphFunctionStore* Store = FindFunctionStore(InLibraryNode))
+		if (IRigVMGraphFunctionHost* FunctionHost = ClientHost->GetRigVMGraphFunctionHost())
 		{
-			if (FRigVMGraphFunctionData* Data = Store->UpdateFunctionInterface(InLibraryNode->GetFunctionHeader()))
+			if (FRigVMGraphFunctionStore* Store = FindFunctionStore(InLibraryNode))
 			{
-				UpdateFunctionReferences(Data->Header, false, false);
-				return true;
+				if (FRigVMGraphFunctionData* Data = Store->UpdateFunctionInterface(InLibraryNode->GetFunctionHeader(FunctionHost)))
+				{
+					UpdateFunctionReferences(Data->Header, false, false);
+					return true;
+				}
 			}
 		}
 	}	
diff --git a/Engine/Plugins/Runtime/RigVM/Source/RigVMDeveloper/Private/RigVMModel/RigVMController.cpp b/Engine/Plugins/Runtime/RigVM/Source/RigVMDeveloper/Private/RigVMModel/RigVMController.cpp
--- a/Engine/Plugins/Runtime/RigVM/Source/RigVMDeveloper/Private/RigVMModel/RigVMController.cpp
+++ b/Engine/Plugins/Runtime/RigVM/Source/RigVMDeveloper/Private/RigVMModel/RigVMController.cpp
@@ -5025,7 +5025,19 @@ URigVMCollapseNode* URigVMController::CollapseNodes(const TArray<URigVMNode*>& I
 	URigVMCollapseNode* CollapseNode = NewObject<URigVMCollapseNode>(Graph, *CollapseNodeName);
 #endif
 	FString ContainedGraphName = CollapseNodeName + TEXT("_ContainedGraph");
-	CollapseNode->ContainedGraph = NewObject<URigVMGraph>(CollapseNode, *ContainedGraphName);
+
+	if (IRigVMClientHost* ClientHost = GetImplementingOuter<IRigVMClientHost>())
+	{
+		if (FRigVMClient* RigVMClient = ClientHost->GetRigVMClient())
+		{
+			CollapseNode->ContainedGraph = RigVMClient->CreateContainedGraphModel(CollapseNode, *ContainedGraphName);
+		}
+	}
+
+	if (CollapseNode->ContainedGraph == nullptr)
+	{
+		return nullptr;
+	}
 
 #if UE_RIGVM_AGGREGATE_NODES_ENABLED
 	if (bIsAggregate)
@@ -5431,7 +5443,12 @@ TArray<URigVMNode*> URigVMController::ExpandLibraryNode(URigVMLibraryNode* InNod
 			ReportError(TEXT("Cannot expand nodes from function reference because the source graph is not found."));
 			return TArray<URigVMNode*>();			
 		}
-	}	
+	}
+	else if(InnerGraph == nullptr)
+	{
+		ReportError(TEXT("Cannot expand nodes from collapse node because the source graph is not found."));
+		return TArray<URigVMNode*>();
+	}
 
 	TArray<URigVMNode*> ContainedNodes = InnerGraph->GetNodes();
 	TArray<URigVMLink*> ContainedLinks = InnerGraph->GetLinks();
@@ -6470,7 +6487,8 @@ void URigVMController::SetReferencedFunction(URigVMFunctionReferenceNode* InFunc
 	}
 	
 	FRigVMGraphFunctionHeader OldReferencedNode = InFunctionRefNode->GetReferencedFunctionHeader();
-	InFunctionRefNode->ReferencedFunctionHeader = InNewReferencedNode->GetFunctionHeader();
+	IRigVMClientHost* ClientHost = GetImplementingOuter<IRigVMClientHost>();
+	InFunctionRefNode->ReferencedFunctionHeader = InNewReferencedNode->GetFunctionHeader(ClientHost ? ClientHost->GetRigVMGraphFunctionHost() : nullptr);
 	
 	if(!(OldReferencedNode == InFunctionRefNode->ReferencedFunctionHeader))
 	{
@@ -13113,7 +13131,20 @@ URigVMLibraryNode* URigVMController::AddFunctionToLibrary(const FName& InFunctio
 	FString FunctionName = GetSchema()->GetValidNodeName(Graph, InFunctionName.IsNone() ? FString(TEXT("Function")) : InFunctionName.ToString());
 	URigVMCollapseNode* CollapseNode = NewObject<URigVMCollapseNode>(Graph, *FunctionName);
 	FString ContainedGraphName = FunctionName + TEXT("_ContainedGraph");
-	CollapseNode->ContainedGraph = NewObject<URigVMGraph>(CollapseNode, *ContainedGraphName);
+
+	if (IRigVMClientHost* ClientHost = GetImplementingOuter<IRigVMClientHost>())
+	{
+		if (FRigVMClient* RigVMClient = ClientHost->GetRigVMClient())
+		{
+			CollapseNode->ContainedGraph = RigVMClient->CreateContainedGraphModel(CollapseNode, *ContainedGraphName);
+		}
+	}
+
+	if (CollapseNode->ContainedGraph == nullptr)
+	{
+		return nullptr;
+	}
+
 	CollapseNode->Position = InNodePosition;
 
 	if(!AddGraphNode(CollapseNode, true))
diff --git a/Engine/Plugins/Runtime/RigVM/Source/RigVMDeveloper/Private/RigVMModel/RigVMGraph.cpp b/Engine/Plugins/Runtime/RigVM/Source/RigVMDeveloper/Private/RigVMModel/RigVMGraph.cpp
--- a/Engine/Plugins/Runtime/RigVM/Source/RigVMDeveloper/Private/RigVMModel/RigVMGraph.cpp
+++ b/Engine/Plugins/Runtime/RigVM/Source/RigVMDeveloper/Private/RigVMModel/RigVMGraph.cpp
@@ -70,10 +70,13 @@ TArray<URigVMGraph*> URigVMGraph::GetContainedGraphs(bool bRecursive) const
 	{
 		if (const URigVMCollapseNode* CollapseNode = Cast<URigVMCollapseNode>(Node))
 		{
-			Graphs.AddUnique(CollapseNode->GetContainedGraph());
-			if (bRecursive)
+			if (ensure(CollapseNode->GetContainedGraph()))
 			{
-				Graphs.Append(CollapseNode->GetContainedGraph()->GetContainedGraphs(true));
+				Graphs.AddUnique(CollapseNode->GetContainedGraph());
+				if (bRecursive)
+				{
+					Graphs.Append(CollapseNode->GetContainedGraph()->GetContainedGraphs(true));
+				}
 			}
 		}
 	}
diff --git a/Engine/Plugins/Runtime/RigVM/Source/RigVMDeveloper/Public/RigVMBlueprint.h b/Engine/Plugins/Runtime/RigVM/Source/RigVMDeveloper/Public/RigVMBlueprint.h
--- a/Engine/Plugins/Runtime/RigVM/Source/RigVMDeveloper/Public/RigVMBlueprint.h
+++ b/Engine/Plugins/Runtime/RigVM/Source/RigVMDeveloper/Public/RigVMBlueprint.h
@@ -401,6 +401,8 @@ public:
 	UFUNCTION(BlueprintCallable, Category = "RigVM Blueprint")
 	virtual TArray<FString> GeneratePythonCommands(const FString InNewBlueprintName) override;
 
+	virtual void SetupPinRedirectorsForBackwardsCompatibility() override {};
+
 	//  --- IRigVMClientHost interface End ---
 
 	//  --- IRigVMExternalDependencyManager interface Start ---
@@ -528,7 +530,6 @@ public:
 	bool bSuspendModelNotificationsForOthers;
 	bool bSuspendAllNotifications;
 
-	virtual void SetupPinRedirectorsForBackwardsCompatibility() {};
 	void RebuildGraphFromModel();
 
 	FRigVMGraphModifiedEvent& OnModified();
@@ -614,13 +615,11 @@ private:
 
 	FOnRigVMCompiledEvent VMCompiledEvent;
 
-	virtual void PatchFunctionReferencesOnLoad();
 	virtual void PathDomainSpecificContentOnLoad() {}
 	virtual void PatchBoundVariables();
 	virtual void PatchVariableNodesWithIncorrectType();
 	virtual void PatchParameterNodesOnLoad() {}
 	virtual void PatchLinksWithCast();
-	virtual void PatchFunctionsOnLoad();
 	virtual void GetBackwardsCompatibilityPublicFunctions(TArray<FName> &BackwardsCompatiblePublicFunctions, TMap<URigVMLibraryNode*, FRigVMGraphFunctionHeader>& OldHeaders);
 
 protected:
diff --git a/Engine/Plugins/Runtime/RigVM/Source/RigVMDeveloper/Public/RigVMModel/Nodes/RigVMTemplateNode.h b/Engine/Plugins/Runtime/RigVM/Source/RigVMDeveloper/Public/RigVMModel/Nodes/RigVMTemplateNode.h
--- a/Engine/Plugins/Runtime/RigVM/Source/RigVMDeveloper/Public/RigVMModel/Nodes/RigVMTemplateNode.h
+++ b/Engine/Plugins/Runtime/RigVM/Source/RigVMDeveloper/Public/RigVMModel/Nodes/RigVMTemplateNode.h
@@ -169,6 +169,7 @@ protected:
 
 	friend class URigVMController;
 	friend class URigVMBlueprint;
+	friend struct FRigVMClient;
 	friend struct FRigVMSetTemplateFilteredPermutationsAction;
 	friend struct FRigVMSetPreferredTemplatePermutationsAction;
 	friend struct FRigVMRemoveNodeAction;
diff --git a/Engine/Plugins/Runtime/RigVM/Source/RigVMDeveloper/Public/RigVMModel/RigVMClient.h b/Engine/Plugins/Runtime/RigVM/Source/RigVMDeveloper/Public/RigVMModel/RigVMClient.h
--- a/Engine/Plugins/Runtime/RigVM/Source/RigVMDeveloper/Public/RigVMModel/RigVMClient.h
+++ b/Engine/Plugins/Runtime/RigVM/Source/RigVMDeveloper/Public/RigVMModel/RigVMClient.h
@@ -6,6 +6,7 @@
 #include "RigVMSchema.h"
 #include "RigVMFunctionLibrary.h"
 #include "RigVMController.h"
+#include "UObject/WeakInterfacePtr.h"
 #include "RigVMClient.generated.h"
 
 struct FRigVMClient;
@@ -122,6 +123,8 @@ public:
 	virtual URigVMController* GetOrCreateController(const UEdGraph* InGraph) = 0;
 
 	virtual TArray<FString> GeneratePythonCommands(const FString InNewBlueprintName)  = 0;
+
+	virtual void SetupPinRedirectorsForBackwardsCompatibility() = 0;
 };
 
 UINTERFACE()
@@ -147,6 +150,27 @@ public:
 	virtual void HandleRigVMGraphRenamed(const FString& InOldNodePath, const FString& InNewNodePath) = 0;
 };
 
+UINTERFACE()
+class URigVMClientExternalModelHost : public UInterface
+{
+	GENERATED_BODY()
+};
+
+// Interface that allows access to externally-hosted models
+class IRigVMClientExternalModelHost
+{
+	GENERATED_BODY()
+
+public:
+
+	// Returns the externally-held models for a client
+	virtual const TArray<TObjectPtr<URigVMGraph>>& GetExternalModels() const = 0;
+
+	// Sets a new created model as 
+	virtual TObjectPtr<URigVMGraph> CreateContainedGraphModel(URigVMCollapseNode* CollapseNode, const FName& Name) = 0;
+};
+
+
 // A management struct containing graphs and controllers.
 USTRUCT()
 struct RIGVMDEVELOPER_API FRigVMClient
@@ -164,7 +188,8 @@ public:
 		, bIgnoreModelNotifications(false)
 		, bDefaultModelCanBeRemoved(false)
 		, OuterClientHost(nullptr)
-		, OuterClientPropertyName(NAME_None)
+		, OuterClientPropertyName(NAME_None)	
+		, ExternalModelHost(nullptr)
 	{
 	}
 
@@ -174,9 +199,10 @@ public:
 	void SetControllerClass(TSubclassOf<URigVMController> InControllerClass);
 	void SetOuterClientHost(UObject* InOuterClientHost, const FName& InOuterClientHostPropertyName);
 	void SetFromDeprecatedData(URigVMGraph* InDefaultGraph, URigVMFunctionLibrary* InFunctionLibrary);
+	void SetExternalModelHost(IRigVMClientExternalModelHost* InExternalModelHost);
 
 	void Reset();
-	int32 Num() const { return Models.Num(); }
+	int32 Num() const { return GetModels().Num(); }
 	UE_DEPRECATED(5.5, "Please use GetDefaultSchema or get a schema per controller/graph.")
 	URigVMSchema* GetSchema() { return GetDefaultSchema(); }
 	URigVMSchema* GetDefaultSchema() const;
@@ -184,10 +210,12 @@ public:
 	UE_DEPRECATED(5.5, "Please use GetSchema or get a schema per controller/graph.")
 	URigVMSchema* GetOrCreateSchema() { return GetDefaultSchema(); }
 	URigVMGraph* GetDefaultModel() const;
-	URigVMGraph* GetModel(int32 InIndex) const { return Models.IsValidIndex(InIndex) ? Models[InIndex] : nullptr; }
+	URigVMGraph* GetModel(int32 InIndex) const;
 	URigVMGraph* GetModel(const UEdGraph* InEdGraph = nullptr) const;
 	URigVMGraph* GetModel(const FString& InNodePathOrName) const;
 	URigVMGraph* GetModel(const UObject* InEditorSideObject) const;
+	void RefreshAllModels(ERigVMLoadType InLoadType, bool bEnablePostLoadHashing, bool& bIsCompiling);
+	const TArray<TObjectPtr<URigVMGraph>>& GetModels() const;
 	TArray<URigVMGraph*> GetAllModels(bool bIncludeFunctionLibrary, bool bRecursive) const;
 	TArray<URigVMGraph*> GetAllModelsLeavesFirst(bool bIncludeFunctionLibrary) const;
 	URigVMController* GetController(int32 InIndex) const;
@@ -214,10 +242,13 @@ public:
 	FRigVMGetFocusedGraph& OnGetFocusedGraph() { return OnGetFocusedGraphDelegate;}
 	const FRigVMGetFocusedGraph& OnGetFocusedGraph() const { return OnGetFocusedGraphDelegate; }
 	URigVMGraph* GetFocusedModel() const;
-
 	URigVMGraph* AddModel(const FString InName, bool bSetupUndoRedo, bool bPrintPythonCommand);
 	URigVMGraph* AddModel(const FName& InName, bool bSetupUndoRedo, const FObjectInitializer* ObjectInitializer = nullptr, bool bCreateController = true);
 	URigVMGraph* AddModel(const FName& InName, TSubclassOf<URigVMSchema> InSchemaClass, bool bSetupUndoRedo, const FObjectInitializer* ObjectInitializer = nullptr, bool bCreateController = true);
+
+	URigVMGraph* CreateModel(const FName& InName, TSubclassOf<URigVMSchema> InSchemaClass, bool bSetupUndoRedo, UObject* InOuter, const FObjectInitializer* ObjectInitializer = nullptr, bool bCreateController = true);
+	TObjectPtr<URigVMGraph> CreateContainedGraphModel(URigVMCollapseNode* CollapseNode, const FName& Name);
+	
 	void AddModel(URigVMGraph* InModel, bool bCreateController);
 	bool RemoveModel(FString InName, bool bSetupUndoRedo, bool bPrintPythonCommand);
 	bool RemoveModel(const FString& InNodePathOrName, bool bSetupUndoRedo);
@@ -229,10 +260,10 @@ public:
 	URigVMNode* FindNode(const FString& InNodePathOrName) const;
 	URigVMPin* FindPin(const FString& InPinPath) const;
 	
-	TArray<TObjectPtr<URigVMGraph>>::RangedForIteratorType      begin() { return Models.begin(); }
-	TArray<TObjectPtr<URigVMGraph>>::RangedForConstIteratorType begin() const { return Models.begin(); }
-	TArray<TObjectPtr<URigVMGraph>>::RangedForIteratorType      end() { return Models.end(); }
-	TArray<TObjectPtr<URigVMGraph>>::RangedForConstIteratorType end() const { return Models.end(); }
+	TArray<TObjectPtr<URigVMGraph>>::RangedForIteratorType      begin() { return const_cast<TArray<TObjectPtr<URigVMGraph>>&>(GetModels()).begin(); }
+	TArray<TObjectPtr<URigVMGraph>>::RangedForConstIteratorType begin() const { return GetModels().begin(); }
+	TArray<TObjectPtr<URigVMGraph>>::RangedForIteratorType      end() { return const_cast<TArray<TObjectPtr<URigVMGraph>>&>(GetModels()).end(); }
+	TArray<TObjectPtr<URigVMGraph>>::RangedForConstIteratorType end() const { return GetModels().end(); }
 
 	UObject* GetOuter() const;
 	FProperty* GetOuterClientProperty() const;
@@ -246,6 +277,8 @@ public:
 
 	// backwards compatibility
 	FRigVMClientPatchResult PatchModelsOnLoad();
+	void PatchFunctionReferencesOnLoad();
+	void PatchFunctionsOnLoad(IRigVMGraphFunctionHost* FunctionHost, TArray<FName>& BackwardsCompatiblePublicFunctions, TMap<URigVMLibraryNode*, FRigVMGraphFunctionHeader>& OldHeaders);
 	FRigVMClientPatchResult PatchPinDefaultValues();
 
 	// try to reattach detached links and delete remaining ones
@@ -321,6 +354,8 @@ private:
 	TWeakObjectPtr<UObject> OuterClientHost;
 	FName OuterClientPropertyName;
 
+	TWeakInterfacePtr<IRigVMClientExternalModelHost> ExternalModelHost;
+
 	friend class UEngineTestClientHost;
 	friend class URigVMBlueprint;
 };
diff --git a/Engine/Plugins/Runtime/RigVM/Source/RigVMDeveloper/Public/RigVMModel/RigVMFunctionLibrary.h b/Engine/Plugins/Runtime/RigVM/Source/RigVMDeveloper/Public/RigVMModel/RigVMFunctionLibrary.h
--- a/Engine/Plugins/Runtime/RigVM/Source/RigVMDeveloper/Public/RigVMModel/RigVMFunctionLibrary.h
+++ b/Engine/Plugins/Runtime/RigVM/Source/RigVMDeveloper/Public/RigVMModel/RigVMFunctionLibrary.h
@@ -100,6 +100,7 @@ private:
 	friend class URigVMController;
 	friend class URigVMCompiler;
 	friend class URigVMBlueprint;
+	friend struct FRigVMClient;
 	friend class UControlRigBlueprint;
 };
 
diff --git a/Engine/Plugins/Runtime/RigVM/Source/RigVMDeveloper/Public/RigVMModel/RigVMGraph.h b/Engine/Plugins/Runtime/RigVM/Source/RigVMDeveloper/Public/RigVMModel/RigVMGraph.h
--- a/Engine/Plugins/Runtime/RigVM/Source/RigVMDeveloper/Public/RigVMModel/RigVMGraph.h
+++ b/Engine/Plugins/Runtime/RigVM/Source/RigVMDeveloper/Public/RigVMModel/RigVMGraph.h
@@ -226,7 +226,8 @@ private:
 	friend class URigVMBlueprint;
 	friend class FRigVMControllerCompileBracketScope;
 	friend class URigVMCompiler;
-	friend struct FRigVMControllerObjectFactory;
 	friend class URigVMSchema;
+	friend struct FRigVMClient;
+	friend struct FRigVMControllerObjectFactory;
 };
 
