Mac build fix due to missing dll export tag

#rnx

diff --git a/Engine/Plugins/Runtime/RigVM/Source/RigVMDeveloper/Public/RigVMModel/RigVMClient.h b/Engine/Plugins/Runtime/RigVM/Source/RigVMDeveloper/Public/RigVMModel/RigVMClient.h
--- a/Engine/Plugins/Runtime/RigVM/Source/RigVMDeveloper/Public/RigVMModel/RigVMClient.h
+++ b/Engine/Plugins/Runtime/RigVM/Source/RigVMDeveloper/Public/RigVMModel/RigVMClient.h
@@ -151,13 +151,13 @@ public:
 };
 
 UINTERFACE()
-class URigVMClientExternalModelHost : public UInterface
+class RIGVMDEVELOPER_API URigVMClientExternalModelHost : public UInterface
 {
 	GENERATED_BODY()
 };
 
 // Interface that allows access to externally-hosted models
-class IRigVMClientExternalModelHost
+class RIGVMDEVELOPER_API IRigVMClientExternalModelHost
 {
 	GENERATED_BODY()
 
