RigVM: Improve Variant Ref widget and create variant workflows

#rnx

diff --git a/Engine/Plugins/Runtime/RigVM/Source/RigVMEditor/Private/Editor/RigVMGraphDetailCustomization.cpp b/Engine/Plugins/Runtime/RigVM/Source/RigVMEditor/Private/Editor/RigVMGraphDetailCustomization.cpp
--- a/Engine/Plugins/Runtime/RigVM/Source/RigVMEditor/Private/Editor/RigVMGraphDetailCustomization.cpp
+++ b/Engine/Plugins/Runtime/RigVM/Source/RigVMEditor/Private/Editor/RigVMGraphDetailCustomization.cpp
@@ -963,6 +963,12 @@ void FRigVMGraphDetailCustomization::CustomizeDetails(IDetailLayoutBuilder& Deta
         ];
 
 		// variant
+		FRigVMVariantWidgetContext VariantContext;
+		if(const URigVMFunctionLibrary* FunctionLbirary = Model->GetTypedOuter<URigVMFunctionLibrary>())
+		{
+			VariantContext.ParentPath = FunctionLbirary->GetPathName();
+		}
+			
 		SettingsCategory.AddCustomRow(FText::GetEmpty())
 		.OverrideResetToDefault(FResetToDefaultOverride::Hide())
 		.Visibility(TAttribute<EVisibility>::CreateLambda([this]()
@@ -978,6 +984,7 @@ void FRigVMGraphDetailCustomization::CustomizeDetails(IDetailLayoutBuilder& Deta
 		.ValueContent()
 		[
 			SNew(SRigVMVariantWidget)
+			.Context(VariantContext)
 			.Variant(this, &FRigVMGraphDetailCustomization::GetVariant)
 			.VariantRefs(this, &FRigVMGraphDetailCustomization::GetVariantRefs)
 			.OnVariantChanged(this, &FRigVMGraphDetailCustomization::OnVariantChanged)
diff --git a/Engine/Plugins/Runtime/RigVM/Source/RigVMEditor/Private/Editor/RigVMVariantDetailCustomization.cpp b/Engine/Plugins/Runtime/RigVM/Source/RigVMEditor/Private/Editor/RigVMVariantDetailCustomization.cpp
--- a/Engine/Plugins/Runtime/RigVM/Source/RigVMEditor/Private/Editor/RigVMVariantDetailCustomization.cpp
+++ b/Engine/Plugins/Runtime/RigVM/Source/RigVMEditor/Private/Editor/RigVMVariantDetailCustomization.cpp
@@ -33,6 +33,12 @@ void FRigVMVariantDetailCustomization::CustomizeHeader(TSharedRef<IPropertyHandl
 		}
 	}
 
+	FRigVMVariantWidgetContext VariantContext;
+	if(BlueprintBeingCustomized)
+	{
+		VariantContext.ParentPath = BlueprintBeingCustomized->GetPathName();
+	}
+
 	HeaderRow
 	.OverrideResetToDefault(FResetToDefaultOverride::Hide())
 	.NameContent()
@@ -42,6 +48,7 @@ void FRigVMVariantDetailCustomization::CustomizeHeader(TSharedRef<IPropertyHandl
 	.ValueContent()
 	[
 		SNew(SRigVMVariantWidget)
+		.Context(VariantContext)
 		.Variant(this, &FRigVMVariantDetailCustomization::GetVariant)
 		.VariantRefs(this, &FRigVMVariantDetailCustomization::GetVariantRefs)
 		.OnVariantChanged(this, &FRigVMVariantDetailCustomization::OnVariantChanged)
diff --git a/Engine/Plugins/Runtime/RigVM/Source/RigVMEditor/Private/RigVMEditorModule.cpp b/Engine/Plugins/Runtime/RigVM/Source/RigVMEditor/Private/RigVMEditorModule.cpp
--- a/Engine/Plugins/Runtime/RigVM/Source/RigVMEditor/Private/RigVMEditorModule.cpp
+++ b/Engine/Plugins/Runtime/RigVM/Source/RigVMEditor/Private/RigVMEditorModule.cpp
@@ -116,12 +116,14 @@ void FRigVMEditorModule::StartupModule()
 									return;
 								}
 
-								FAssetToolsModule& AssetToolsModule = FModuleManager::LoadModuleChecked<FAssetToolsModule>("AssetTools");
+								const FAssetToolsModule& AssetToolsModule = FModuleManager::LoadModuleChecked<FAssetToolsModule>("AssetTools");
 								FString PathName = SoftObjectPath.GetLongPackageName();
 								FString ObjectName = SelectedObject->GetName();
 								FString PackageName;
-								FString BasePackageName = PathName + "/" + ObjectName;
-								AssetToolsModule.Get().CreateUniqueAssetName(BasePackageName, TEXT(""), PackageName, ObjectName);
+								AssetToolsModule.Get().CreateUniqueAssetName(PathName, TEXT(""), PackageName, ObjectName);
+
+								FString Extension;
+								FPaths::Split(PackageName, PathName, ObjectName, Extension);
 
 								UObject* DuplicateAsset = AssetToolsModule.Get().DuplicateAsset(ObjectName, PathName, SelectedObject);
 								if (URigVMBlueprint* DuplicateBlueprint = Cast<URigVMBlueprint>(DuplicateAsset))
diff --git a/Engine/Plugins/Runtime/RigVM/Source/RigVMEditor/Private/Widgets/SRigVMVariantWidget.cpp b/Engine/Plugins/Runtime/RigVM/Source/RigVMEditor/Private/Widgets/SRigVMVariantWidget.cpp
--- a/Engine/Plugins/Runtime/RigVM/Source/RigVMEditor/Private/Widgets/SRigVMVariantWidget.cpp
+++ b/Engine/Plugins/Runtime/RigVM/Source/RigVMEditor/Private/Widgets/SRigVMVariantWidget.cpp
@@ -2,8 +2,13 @@
 
 #include "Widgets/SRigVMVariantWidget.h"
 #include "Widgets/Text/STextBlock.h"
+#include "Widgets/Input/SButton.h"
+#include "Widgets/Input/SComboButton.h"
 #include "Widgets/Layout/SScrollBox.h"
+#include "Widgets/Images/SImage.h"
 #include "DetailLayoutBuilder.h"
+#include "Editor/RigVMEditorTools.h"
+#include "AssetThumbnail.h"
 
 #define LOCTEXT_NAMESPACE "SRigVMVariantWidget"
 
@@ -32,6 +37,12 @@ void SRigVMVariantWidget::Construct(
 		OnCreateVariantRefRow.BindSP(this, &SRigVMVariantWidget::CreateDefaultVariantRefRow);
 	}
 
+	ContextAttribute = InArgs._Context;
+	if(!ContextAttribute.IsSet() && !ContextAttribute.IsBound())
+	{
+		ContextAttribute = FRigVMVariantWidgetContext();
+	}
+
 	ChildSlot
 	[
 		SNew(SVerticalBox)
@@ -110,34 +121,270 @@ EVisibility SRigVMVariantWidget::GetVariantRefListVisibility() const
 
 TSharedPtr<SWidget> SRigVMVariantWidget::CreateDefaultVariantRefRow(const FRigVMVariantRef& InVariantRef) const
 {
-	return SNew(STextBlock)
-	.Text(FText::FromString(InVariantRef.ObjectPath.ToString()));
+	const TAttribute<FText> ToolTipAttribute = FText::FromString(InVariantRef.ObjectPath.ToString());
+
+	const FRigVMVariantRef LocalVariantRef = InVariantRef;
+	
+	if(!InVariantRef.ObjectPath.IsSubobject())
+	{
+		const FAssetData AssetData = UE::RigVM::Editor::Tools::FindAssetFromAnyPath(InVariantRef.ObjectPath.ToString(), true);
+		const TSharedPtr<FAssetThumbnail> AssetThumbnail = MakeShareable(new FAssetThumbnail(AssetData, 32, 32, TSharedPtr<FAssetThumbnailPool>() ));
+		const FAssetThumbnailConfig ThumbnailConfig;
+
+		TSharedRef<SBorder> ThumbnailBorder = SNew(SBorder);
+		ThumbnailBorder->SetVisibility(EVisibility::SelfHitTestInvisible);
+		ThumbnailBorder->SetPadding(FMargin(0.0f, 0.0f, 4.0f, 4.0f));
+		ThumbnailBorder->SetBorderImage(FAppStyle::Get().GetBrush("PropertyEditor.AssetTileItem.DropShadow"));
+		ThumbnailBorder->SetContent(
+			SNew(SOverlay)
+			+SOverlay::Slot()
+			.Padding(1.0f)
+			[
+				SNew(SBorder)
+				.Padding(0)
+				.BorderImage(FStyleDefaults::GetNoBrush())
+				.OnMouseDoubleClick_Lambda(
+					[this, LocalVariantRef](
+						const FGeometry&,
+						const FPointerEvent&) -> FReply
+					{
+						(void)OnBrowseVariantRef.ExecuteIfBound(LocalVariantRef);
+						return FReply::Handled();
+					})
+				[
+					SNew(SBox)
+					.ToolTipText(ToolTipAttribute)
+					.WidthOverride(32)
+					.HeightOverride(32)
+					[
+						AssetThumbnail->MakeThumbnailWidget(ThumbnailConfig)
+					]
+				]
+			]
+			+ SOverlay::Slot()
+			[
+				SNew(SImage)
+				.Image(this, &SRigVMVariantWidget::GetThumbnailBorder, ThumbnailBorder)
+				.Visibility(EVisibility::SelfHitTestInvisible)
+			]
+		);
+
+		return SNew(SHorizontalBox)
+		+ SHorizontalBox::Slot()
+		.Padding(0.0f,3.0f,5.0f,0.0f)
+		.AutoWidth()
+		.VAlign(VAlign_Center)
+		[
+			ThumbnailBorder
+		]
+
+		+ SHorizontalBox::Slot()
+		.Padding(0.0f)
+		.VAlign(VAlign_Center)
+		[
+			SNew(SVerticalBox)
+			+ SVerticalBox::Slot()
+			.VAlign(VAlign_Center)
+			.AutoHeight()
+			[
+				SNew(SComboButton)
+				.ToolTipText(ToolTipAttribute)
+				.IsEnabled(false)
+				.ButtonContent()
+				[
+					SNew(SHorizontalBox)
+
+					+SHorizontalBox::Slot()
+					.FillWidth(1)
+					.VAlign(VAlign_Center)
+					[
+						// Show the name of the asset or actor
+						SNew(STextBlock)
+						.Font( FAppStyle::GetFontStyle( TEXT("PropertyWindow.NormalFont") ) )
+						.Text(FText::FromString(InVariantRef.ObjectPath.GetAssetName()))
+					]
+				]
+			]
+		];
+	}
+	
+	FString ParentPath = ContextAttribute.Get().ParentPath;
+
+	static const TArray<FString> Separators = { TEXT("\\"), TEXT(":"), TEXT("."), TEXT("/") };
+
+	auto TreatPath = [](FString& InOutPath)
+	{
+		// replace the first period with an underscore - since that may as well be the object
+		// for example: /Game/Animation/ControlRigVariants.ControlRigVariants:FunctionLibrary/MyFunction
+		// becomes /Game/Animation/ControlRigVariants_ControlRigVariants/FunctionLibrary/MyFunction
+		const int32 FirstPeriodIndex = InOutPath.Find(TEXT("."), ESearchCase::IgnoreCase, ESearchDir::FromStart);
+		if(FirstPeriodIndex != INDEX_NONE)
+		{
+			InOutPath[FirstPeriodIndex] = TEXT('_');
+		}
+
+		// replace all other separators with "/"
+		for(const FString& Separator : Separators)
+		{
+			if(Separator != TEXT("/"))
+			{
+				InOutPath.ReplaceInline(*Separator, TEXT("/"), ESearchCase::IgnoreCase);
+			}
+		}
+	};
+
+	FString ObjectPath = InVariantRef.ObjectPath.ToString(); 
+	FString SearchPath = ObjectPath;
+
+	TreatPath(ParentPath);
+	TreatPath(SearchPath);
+	SearchPath.ReplaceInline(TEXT("\\"), TEXT("/"), ESearchCase::IgnoreCase);
+	SearchPath.ReplaceInline(TEXT(":"), TEXT("/"), ESearchCase::IgnoreCase);
+	
+	while(!ParentPath.IsEmpty() && !SearchPath.StartsWith(ParentPath, ESearchCase::CaseSensitive))
+	{
+		if(!ParentPath.Split(TEXT("/"), &ParentPath, nullptr, ESearchCase::IgnoreCase, ESearchDir::FromEnd))
+		{
+			break;
+		}
+	}
+
+	if(SearchPath.StartsWith(ParentPath, ESearchCase::CaseSensitive))
+	{
+		ObjectPath = ObjectPath.Mid(ParentPath.Len());
+		for(const FString& Separator : Separators)
+		{
+			ObjectPath.RemoveFromStart(Separator);
+		}
+	}
+
+	// cut off double object name
+	FString PackageName, ObjectName;
+	if(ObjectPath.Split(TEXT("."), &PackageName, &ObjectName, ESearchCase::IgnoreCase, ESearchDir::FromEnd))
+	{
+		const FString CompleteObjectName = ObjectName + TEXT(".") + ObjectName;
+		const FString CompleteObjectNameWithSlash = TEXT("/") + CompleteObjectName;  
+		if(ObjectPath.Equals(CompleteObjectName, ESearchCase::CaseSensitive))
+		{
+			ObjectPath = ObjectName;
+		}
+		else if(ObjectPath.EndsWith(CompleteObjectNameWithSlash, ESearchCase::CaseSensitive))
+		{
+			ObjectPath = ObjectPath.LeftChop(CompleteObjectNameWithSlash.Len());
+		}
+	}
+
+	static constexpr int32 MaxPathLength = 50;
+	if(ObjectPath.Len() > MaxPathLength)
+	{
+		ObjectPath = TEXT("...") + ObjectPath.Right(MaxPathLength - 3);
+	}
+
+	TSharedRef<SHorizontalBox> HorizontalBox = SNew(SHorizontalBox);
+
+	const FSlateBrush* Icon = nullptr;
+	
+	static const FString RigVMFunctionLibraryToken = TEXT("/RigVMFunctionLibrary/");
+	if(SearchPath.Contains(RigVMFunctionLibraryToken, ESearchCase::CaseSensitive))
+	{
+		static FSlateIcon FunctionIcon(FAppStyle::GetAppStyleSetName(), "Kismet.AllClasses.FunctionIcon");
+		Icon = FunctionIcon.GetIcon(); 
+	}
+
+	if(Icon)
+	{
+		HorizontalBox->AddSlot()
+		.HAlign(HAlign_Left)
+		.VAlign(VAlign_Center)
+		.AutoWidth()
+		.Padding(0, 0, 3, 0)
+		[
+			SNew(SImage)
+			.Image(Icon)
+			.DesiredSizeOverride(FVector2D(16, 16))
+		];
+	}
+
+	HorizontalBox->AddSlot()
+	.HAlign(HAlign_Left)
+	.VAlign(VAlign_Center)
+	.AutoWidth()
+	.Padding(0, 0, 0, 0)
+	[
+		SNew(STextBlock)
+		.Text(FText::FromString(ObjectPath))
+	];
+
+	return SNew(SButton)
+		.ButtonStyle(FAppStyle::Get(), "SimpleButton")
+		.OnClicked_Lambda([this, LocalVariantRef]() -> FReply
+			{
+				(void)OnBrowseVariantRef.ExecuteIfBound(LocalVariantRef);
+				return FReply::Handled();
+			})
+		.ContentPadding(FMargin(1, 0))
+		.ToolTipText(ToolTipAttribute)
+		[
+			HorizontalBox
+		];
 }
 
 void SRigVMVariantWidget::RebuildVariantRefList()
 {
 	VariantRefListBox->ClearChildren();
 
-	for(FRigVMVariantRef VariantRef : VariantRefs)
+	if(VariantRefs.IsEmpty())
 	{
-		TSharedPtr<SWidget> Widget = OnCreateVariantRefRow.Execute(VariantRef);
-		Widget->SetOnMouseDoubleClick(
-			FPointerEventHandler::CreateLambda(
-				[VariantRef, this](const FGeometry& InGeometry, const FPointerEvent& InMouseEvent) -> FReply
-				{
-					(void)OnBrowseVariantRef.ExecuteIfBound(VariantRef);
-					return FReply::Handled();
-				}
-			)
-		);
 		VariantRefListBox->AddSlot()
 		.AutoHeight()
-		.HAlign(HAlign_Fill)
+		.HAlign(HAlign_Left)
+		.Padding(0, 4, 0, 0)
+		[
+			SNew(STextBlock)
+			.Font(IDetailLayoutBuilder::GetDetailFont())
+			.Text(LOCTEXT("NoOtherVariants", "No other variants found."))
+		];
+	}
+	else
+	{
+		VariantRefListBox->AddSlot()
+		.AutoHeight()
+		.HAlign(HAlign_Left)
 		.Padding(0, 4, 0, 0)
 		[
-			Widget.ToSharedRef()
+			SNew(STextBlock)
+			.Font(IDetailLayoutBuilder::GetDetailFont())
+			.Text(LOCTEXT("MatchingVariants", "Matching Variants:"))
 		];
+
+		for(FRigVMVariantRef VariantRef : VariantRefs)
+		{
+			TSharedPtr<SWidget> Widget = OnCreateVariantRefRow.Execute(VariantRef);
+			Widget->SetOnMouseDoubleClick(
+				FPointerEventHandler::CreateLambda(
+					[VariantRef, this](const FGeometry& InGeometry, const FPointerEvent& InMouseEvent) -> FReply
+					{
+						(void)OnBrowseVariantRef.ExecuteIfBound(VariantRef);
+						return FReply::Handled();
+					}
+				)
+			);
+			VariantRefListBox->AddSlot()
+			.AutoHeight()
+			.HAlign(HAlign_Fill)
+			.Padding(0, 4, 0, 0)
+			[
+				Widget.ToSharedRef()
+			];
+		}
 	}
 }
 
-#undef LOCTEXT_NAMESPACE // SRigVMVariantWidget
\ No newline at end of file
+const FSlateBrush* SRigVMVariantWidget::GetThumbnailBorder(TSharedRef<SBorder> InThumbnailBorder) const
+{
+	static const FName HoveredBorderName("PropertyEditor.AssetThumbnailBorderHovered");
+	static const FName RegularBorderName("PropertyEditor.AssetThumbnailBorder");
+	return InThumbnailBorder->IsHovered() ? FAppStyle::Get().GetBrush(HoveredBorderName) : FAppStyle::Get().GetBrush(RegularBorderName);
+}
+
+#undef LOCTEXT_NAMESPACE // SRigVMVariantWidget
diff --git a/Engine/Plugins/Runtime/RigVM/Source/RigVMEditor/Public/Widgets/SRigVMVariantWidget.h b/Engine/Plugins/Runtime/RigVM/Source/RigVMEditor/Public/Widgets/SRigVMVariantWidget.h
--- a/Engine/Plugins/Runtime/RigVM/Source/RigVMEditor/Public/Widgets/SRigVMVariantWidget.h
+++ b/Engine/Plugins/Runtime/RigVM/Source/RigVMEditor/Public/Widgets/SRigVMVariantWidget.h
@@ -3,6 +3,7 @@
 #pragma once
 
 #include "Widgets/Layout/SBox.h"
+#include "Widgets/Layout/SBorder.h"
 #include "Widgets/SBoxPanel.h"
 #include "RigVMCore/RigVMVariant.h"
 #include "Widgets/SRigVMVariantTagWidget.h"
@@ -11,18 +12,31 @@ DECLARE_DELEGATE_OneParam(FRigVMVariantWidget_OnVariantChanged, const FRigVMVari
 DECLARE_DELEGATE_RetVal_OneParam(TSharedPtr<SWidget>, FRigVMVariantWidget_OnCreateVariantRefRow, const FRigVMVariantRef&);
 DECLARE_DELEGATE_OneParam(FRigVMVariantWidget_OnBrowseVariantRef, const FRigVMVariantRef&);
 
-class SRigVMVariantWidget : public SBox
+struct RIGVMEDITOR_API FRigVMVariantWidgetContext
+{
+	FRigVMVariantWidgetContext()
+		: ParentPath()
+	{
+	}
+	
+	// the path the current context is in
+	FString ParentPath;
+};
+
+class RIGVMEDITOR_API SRigVMVariantWidget : public SBox
 {
 public:
 	
 	SLATE_BEGIN_ARGS(SRigVMVariantWidget)
-		: _MaxVariantRefListHeight(200.f)
+		: _Context(FRigVMVariantWidgetContext())
+		, _MaxVariantRefListHeight(200.f)
 		, _CanAddTags(false)
 		, _EnableTagContextMenu(false)
 	{
 	}
 	SLATE_ATTRIBUTE(FRigVMVariant, Variant)
 	SLATE_ATTRIBUTE(TArray<FRigVMVariantRef>, VariantRefs)
+	SLATE_ATTRIBUTE(FRigVMVariantWidgetContext, Context)
 	SLATE_EVENT(FRigVMVariantWidget_OnVariantChanged, OnVariantChanged)
 	SLATE_EVENT(FRigVMVariantWidget_OnCreateVariantRefRow, OnCreateVariantRefRow);
 	SLATE_EVENT(FRigVMVariantWidget_OnBrowseVariantRef, OnBrowseVariantRef)
@@ -46,6 +60,7 @@ private:
 	EVisibility GetVariantRefListVisibility() const;
 	TSharedPtr<SWidget> CreateDefaultVariantRefRow(const FRigVMVariantRef& InVariantRef) const;
 	void RebuildVariantRefList();
+	const FSlateBrush* GetThumbnailBorder(TSharedRef<SBorder> InThumbnailBorder) const;
 
 	TAttribute<FRigVMVariant> VariantAttribute;
 	FRigVMVariantWidget_OnVariantChanged OnVariantChanged;
@@ -58,4 +73,5 @@ private:
 	TArray<FRigVMVariantRef> VariantRefs;
 	uint32 VariantRefHash;
 	TSharedPtr<SVerticalBox> VariantRefListBox;
+	TAttribute<FRigVMVariantWidgetContext> ContextAttribute;
 };
\ No newline at end of file
