RigVM: Improve Variant widget tooltip and treeview

#rnx

diff --git a/Engine/Plugins/Runtime/RigVM/Source/RigVMEditor/Private/Widgets/SRigVMVariantWidget.cpp b/Engine/Plugins/Runtime/RigVM/Source/RigVMEditor/Private/Widgets/SRigVMVariantWidget.cpp
--- a/Engine/Plugins/Runtime/RigVM/Source/RigVMEditor/Private/Widgets/SRigVMVariantWidget.cpp
+++ b/Engine/Plugins/Runtime/RigVM/Source/RigVMEditor/Private/Widgets/SRigVMVariantWidget.cpp
@@ -9,9 +9,78 @@
 #include "DetailLayoutBuilder.h"
 #include "Editor/RigVMEditorTools.h"
 #include "AssetThumbnail.h"
+#include "RigVMModel/RigVMBuildData.h"
 
 #define LOCTEXT_NAMESPACE "SRigVMVariantWidget"
 
+void SRigVMVariantToolTipWithTags::Construct(const FArguments& InArgs)
+{
+	GetTagsDelegate = InArgs._OnGetTags;
+
+	SuperClassArgs._Text = InArgs._ToolTipText;
+	SToolTip::Construct(
+	SuperClassArgs
+		.TextMargin(11.0f)
+		.BorderImage( FCoreStyle::Get().GetBrush("ToolTip.BrightBackground") )
+	);
+}
+
+bool SRigVMVariantToolTipWithTags::IsEmpty() const
+{
+	if(!GetTextTooltip().IsEmpty())
+	{
+		return false;
+	}
+	if(GetTagsDelegate.IsBound())
+	{
+		return GetTagsDelegate.Execute().IsEmpty();
+	}
+	return true;
+}
+
+void SRigVMVariantToolTipWithTags::OnOpening()
+{
+	const TSharedPtr<SVerticalBox> ContentsWidget = SNew(SVerticalBox);
+
+	ContentsWidget->AddSlot()
+	.AutoHeight()
+	.HAlign(HAlign_Fill)
+	.Padding(0.f, 0.f, 0.f, 0.f)
+	[
+		SNew( STextBlock )
+		.Text( SuperClassArgs._Text )
+		.Font( SuperClassArgs._Font )
+		.ColorAndOpacity( FLinearColor::Black )
+		.WrapTextAt_Static( &SToolTip::GetToolTipWrapWidth )
+	];
+
+	if(GetTagsDelegate.IsBound())
+	{
+		ContentsWidget->AddSlot()
+		.AutoHeight()
+		.HAlign(HAlign_Left)
+		.Padding(0.f, 4.f, 0.f, 0.f)
+		[
+			SNew(SRigVMVariantTagWidget)
+			.Visibility_Lambda([this]() -> EVisibility
+			{
+				return GetTagsDelegate.Execute().IsEmpty() ? EVisibility::Collapsed : EVisibility::Visible;
+			})
+			.OnGetTags(GetTagsDelegate)
+			.Orientation(EOrientation::Orient_Horizontal)
+			.CanAddTags(false)
+			.EnableContextMenu(false)
+		];
+	}
+	SetContentWidget(ContentsWidget->AsShared());
+}
+	
+void SRigVMVariantToolTipWithTags::OnClosed()
+{
+	SToolTip::OnClosed();
+	ResetContentWidget();
+}
+
 SRigVMVariantWidget::SRigVMVariantWidget()
 	: VariantRefHash(UINT32_MAX)
 {
@@ -87,6 +156,41 @@ void SRigVMVariantWidget::Construct(
 		]
 	];
 
+	VariantRefListBox->AddSlot()
+	.AutoHeight()
+	.HAlign(HAlign_Left)
+	.Padding(0, 4, 0, 0)
+	[
+		SNew(STextBlock)
+		.Font(IDetailLayoutBuilder::GetDetailFont())
+		.Text_Lambda([this]()
+		{
+			return VariantTreeRowInfos.IsEmpty() ?
+				LOCTEXT("NoOtherVariants", "No other variants found.") :
+				LOCTEXT("MatchingVariants", "Matching Variants:");
+		})
+	];
+
+	VariantRefListBox->AddSlot()
+	.AutoHeight()
+	.HAlign(HAlign_Left)
+	.Padding(0, 0, 0, 0)
+	[
+		SAssignNew(VariantRefTreeView, STreeView<TSharedPtr<FVariantTreeRowInfo>>)
+		.SelectionMode(ESelectionMode::None)
+		.OnMouseButtonDoubleClick_Lambda([this](TSharedPtr<FVariantTreeRowInfo> InRowInfo)
+		{
+			(void)OnBrowseVariantRef.ExecuteIfBound(InRowInfo->VariantRef);
+		})
+		.Visibility_Lambda([this]()
+		{
+			return VariantTreeRowInfos.IsEmpty() ? EVisibility::Collapsed : EVisibility::Visible;
+		})
+		.TreeItemsSource(&VariantTreeRowInfos)
+		.OnGenerateRow(this, &SRigVMVariantWidget::GenerateVariantTreeRow)
+		.OnGetChildren(this, &SRigVMVariantWidget::GetChildrenForVariantInfo)
+	];
+
 	SetCanTick(true);
 }
 
@@ -105,26 +209,78 @@ void SRigVMVariantWidget::Tick(const FGeometry& AllottedGeometry, const double I
 	{
 		VariantRefHash = NewHash;
 		VariantRefs = NewVariantRefs;
-		VariantRefs.Sort([](const FRigVMVariantRef& A, const FRigVMVariantRef& B)
+
+		// sort the variants by path length - but make sure that
+		// variant refs within our own context come first
+		const FString ParentPath = GetVariantContext().ParentPath; 
+		VariantRefs.Sort([ParentPath](const FRigVMVariantRef& A, const FRigVMVariantRef& B)
 		{
-			return A.ObjectPath.ToString().Compare(B.ObjectPath.ToString()) < 0;
+			FString PathA = A.ObjectPath.ToString(); 
+			FString PathB = B.ObjectPath.ToString();
+			if(PathA.StartsWith(ParentPath, ESearchCase::CaseSensitive))
+			{
+				PathA = PathA.Mid(ParentPath.Len());
+			}
+			if(PathB.StartsWith(ParentPath, ESearchCase::CaseSensitive))
+			{
+				PathB = PathB.Mid(ParentPath.Len());
+			}
+			return PathA.Compare(PathB) < 0;
 		});
 		
 		RebuildVariantRefList();
 	}
 }
 
+SRigVMVariantWidget::SRigVMVariantRefTreeRow::~SRigVMVariantRefTreeRow()
+{
+}
+
+void SRigVMVariantWidget::SRigVMVariantRefTreeRow::Construct(const FArguments& InArgs, const TSharedRef<STableViewBase>& OwnerTableView)
+{
+	STableRow<TSharedPtr<FVariantTreeRowInfo>>::FArguments SuperArguments;
+	SuperArguments.Content()
+	[
+		InArgs._Content.ToSharedRef()
+	];
+	SuperArguments.Padding(0);
+	
+	STableRow< TSharedPtr<FVariantTreeRowInfo> >::Construct(SuperArguments, OwnerTableView);
+}
+
+const FRigVMVariantWidgetContext& SRigVMVariantWidget::GetVariantContext() const
+{
+	return ContextAttribute.Get();
+}
+
 EVisibility SRigVMVariantWidget::GetVariantRefListVisibility() const
 {
 	return VariantRefs.IsEmpty() ? EVisibility::Collapsed : EVisibility::Visible;
 }
 
-TSharedPtr<SWidget> SRigVMVariantWidget::CreateDefaultVariantRefRow(const FRigVMVariantRef& InVariantRef) const
+TSharedRef<ITableRow> SRigVMVariantWidget::GenerateVariantTreeRow(TSharedPtr<FVariantTreeRowInfo> InRowInfo, const TSharedRef<STableViewBase>& OwnerTable)
 {
-	const TAttribute<FText> ToolTipAttribute = FText::FromString(InVariantRef.ObjectPath.ToString());
+	return SNew(SRigVMVariantRefTreeRow, OwnerTable)
+		.Content(CreateDefaultVariantRefRow(InRowInfo->VariantRef));
+	
+}
+
+void SRigVMVariantWidget::GetChildrenForVariantInfo(TSharedPtr<FVariantTreeRowInfo> InInfo, TArray<TSharedPtr<FVariantTreeRowInfo>>& OutChildren)
+{
+	OutChildren = InInfo->NestedInfos;
+}
 
+TSharedPtr<SWidget> SRigVMVariantWidget::CreateDefaultVariantRefRow(const FRigVMVariantRef& InVariantRef) const
+{
 	const FRigVMVariantRef LocalVariantRef = InVariantRef;
 	
+	TSharedPtr<SToolTip> TooltipWithTags = SNew(SRigVMVariantToolTipWithTags)
+		.ToolTipText(FText::FromString(InVariantRef.ObjectPath.ToString()))
+		.OnGetTags_Lambda([LocalVariantRef]()
+		{
+			return LocalVariantRef.Variant.Tags;
+		});
+	
 	if(!InVariantRef.ObjectPath.IsSubobject())
 	{
 		const FAssetData AssetData = UE::RigVM::Editor::Tools::FindAssetFromAnyPath(InVariantRef.ObjectPath.ToString(), true);
@@ -153,7 +309,7 @@ TSharedPtr<SWidget> SRigVMVariantWidget::CreateDefaultVariantRefRow(const FRigVM
 					})
 				[
 					SNew(SBox)
-					.ToolTipText(ToolTipAttribute)
+					.ToolTip(TooltipWithTags)
 					.WidthOverride(32)
 					.HeightOverride(32)
 					[
@@ -188,7 +344,7 @@ TSharedPtr<SWidget> SRigVMVariantWidget::CreateDefaultVariantRefRow(const FRigVM
 			.AutoHeight()
 			[
 				SNew(SComboButton)
-				.ToolTipText(ToolTipAttribute)
+				.ToolTip(TooltipWithTags)
 				.IsEnabled(false)
 				.ButtonContent()
 				[
@@ -208,84 +364,13 @@ TSharedPtr<SWidget> SRigVMVariantWidget::CreateDefaultVariantRefRow(const FRigVM
 		];
 	}
 	
-	FString ParentPath = ContextAttribute.Get().ParentPath;
-
-	static const TArray<FString> Separators = { TEXT("\\"), TEXT(":"), TEXT("."), TEXT("/") };
-
-	auto TreatPath = [](FString& InOutPath)
-	{
-		// replace the first period with an underscore - since that may as well be the object
-		// for example: /Game/Animation/ControlRigVariants.ControlRigVariants:FunctionLibrary/MyFunction
-		// becomes /Game/Animation/ControlRigVariants_ControlRigVariants/FunctionLibrary/MyFunction
-		const int32 FirstPeriodIndex = InOutPath.Find(TEXT("."), ESearchCase::IgnoreCase, ESearchDir::FromStart);
-		if(FirstPeriodIndex != INDEX_NONE)
-		{
-			InOutPath[FirstPeriodIndex] = TEXT('_');
-		}
-
-		// replace all other separators with "/"
-		for(const FString& Separator : Separators)
-		{
-			if(Separator != TEXT("/"))
-			{
-				InOutPath.ReplaceInline(*Separator, TEXT("/"), ESearchCase::IgnoreCase);
-			}
-		}
-	};
-
-	FString ObjectPath = InVariantRef.ObjectPath.ToString(); 
-	FString SearchPath = ObjectPath;
-
-	TreatPath(ParentPath);
-	TreatPath(SearchPath);
-	SearchPath.ReplaceInline(TEXT("\\"), TEXT("/"), ESearchCase::IgnoreCase);
-	SearchPath.ReplaceInline(TEXT(":"), TEXT("/"), ESearchCase::IgnoreCase);
-	
-	while(!ParentPath.IsEmpty() && !SearchPath.StartsWith(ParentPath, ESearchCase::CaseSensitive))
-	{
-		if(!ParentPath.Split(TEXT("/"), &ParentPath, nullptr, ESearchCase::IgnoreCase, ESearchDir::FromEnd))
-		{
-			break;
-		}
-	}
-
-	if(SearchPath.StartsWith(ParentPath, ESearchCase::CaseSensitive))
-	{
-		ObjectPath = ObjectPath.Mid(ParentPath.Len());
-		for(const FString& Separator : Separators)
-		{
-			ObjectPath.RemoveFromStart(Separator);
-		}
-	}
-
-	// cut off double object name
-	FString PackageName, ObjectName;
-	if(ObjectPath.Split(TEXT("."), &PackageName, &ObjectName, ESearchCase::IgnoreCase, ESearchDir::FromEnd))
-	{
-		const FString CompleteObjectName = ObjectName + TEXT(".") + ObjectName;
-		const FString CompleteObjectNameWithSlash = TEXT("/") + CompleteObjectName;  
-		if(ObjectPath.Equals(CompleteObjectName, ESearchCase::CaseSensitive))
-		{
-			ObjectPath = ObjectName;
-		}
-		else if(ObjectPath.EndsWith(CompleteObjectNameWithSlash, ESearchCase::CaseSensitive))
-		{
-			ObjectPath = ObjectPath.LeftChop(CompleteObjectNameWithSlash.Len());
-		}
-	}
-
-	static constexpr int32 MaxPathLength = 50;
-	if(ObjectPath.Len() > MaxPathLength)
-	{
-		ObjectPath = TEXT("...") + ObjectPath.Right(MaxPathLength - 3);
-	}
 
 	TSharedRef<SHorizontalBox> HorizontalBox = SNew(SHorizontalBox);
 
 	const FSlateBrush* Icon = nullptr;
 	
-	static const FString RigVMFunctionLibraryToken = TEXT("/RigVMFunctionLibrary/");
-	if(SearchPath.Contains(RigVMFunctionLibraryToken, ESearchCase::CaseSensitive))
+	static const FString RigVMFunctionLibraryToken = TEXT("RigVMFunctionLibrary");
+	if(InVariantRef.ObjectPath.ToString().Contains(RigVMFunctionLibraryToken, ESearchCase::CaseSensitive))
 	{
 		static FSlateIcon FunctionIcon(FAppStyle::GetAppStyleSetName(), "Kismet.AllClasses.FunctionIcon");
 		Icon = FunctionIcon.GetIcon(); 
@@ -305,6 +390,17 @@ TSharedPtr<SWidget> SRigVMVariantWidget::CreateDefaultVariantRefRow(const FRigVM
 		];
 	}
 
+	FString DisplayLabel;
+	if(InVariantRef.ObjectPath.IsSubobject())
+	{
+		DisplayLabel = InVariantRef.ObjectPath.GetSubPathString();
+		(void)DisplayLabel.Split(TEXT("."), nullptr, &DisplayLabel, ESearchCase::CaseSensitive, ESearchDir::FromEnd);
+	}
+	else
+	{
+		DisplayLabel = InVariantRef.ObjectPath.GetAssetName();
+	}
+
 	HorizontalBox->AddSlot()
 	.HAlign(HAlign_Left)
 	.VAlign(VAlign_Center)
@@ -312,7 +408,7 @@ TSharedPtr<SWidget> SRigVMVariantWidget::CreateDefaultVariantRefRow(const FRigVM
 	.Padding(0, 0, 0, 0)
 	[
 		SNew(STextBlock)
-		.Text(FText::FromString(ObjectPath))
+		.Text(FText::FromString(DisplayLabel))
 	];
 
 	return SNew(SButton)
@@ -323,7 +419,7 @@ TSharedPtr<SWidget> SRigVMVariantWidget::CreateDefaultVariantRefRow(const FRigVM
 				return FReply::Handled();
 			})
 		.ContentPadding(FMargin(1, 0))
-		.ToolTipText(ToolTipAttribute)
+		.ToolTip(TooltipWithTags)
 		[
 			HorizontalBox
 		];
@@ -331,53 +427,57 @@ TSharedPtr<SWidget> SRigVMVariantWidget::CreateDefaultVariantRefRow(const FRigVM
 
 void SRigVMVariantWidget::RebuildVariantRefList()
 {
-	VariantRefListBox->ClearChildren();
+	VariantTreeRowInfos.Reset();
+	TMap<FString, TSharedPtr<FVariantTreeRowInfo>> PathToRowInfo;
 
-	if(VariantRefs.IsEmpty())
+	const TArray<FRigVMVariantRef> AllAssetVariantRefs = URigVMBuildData::Get()->GatherAllAssetVariantRefs();
+	TMap<FString, FRigVMVariantRef> AssetPathToVariantRef;
+	for(const FRigVMVariantRef& AssetVariantRef : AllAssetVariantRefs)
 	{
-		VariantRefListBox->AddSlot()
-		.AutoHeight()
-		.HAlign(HAlign_Left)
-		.Padding(0, 4, 0, 0)
-		[
-			SNew(STextBlock)
-			.Font(IDetailLayoutBuilder::GetDetailFont())
-			.Text(LOCTEXT("NoOtherVariants", "No other variants found."))
-		];
+		AssetPathToVariantRef.FindOrAdd(AssetVariantRef.ObjectPath.ToString()) = AssetVariantRef;
 	}
-	else
-	{
-		VariantRefListBox->AddSlot()
-		.AutoHeight()
-		.HAlign(HAlign_Left)
-		.Padding(0, 4, 0, 0)
-		[
-			SNew(STextBlock)
-			.Font(IDetailLayoutBuilder::GetDetailFont())
-			.Text(LOCTEXT("MatchingVariants", "Matching Variants:"))
-		];
 
-		for(FRigVMVariantRef VariantRef : VariantRefs)
+	const FSoftObjectPath ContextAssetObjectPath = FSoftObjectPath(GetVariantContext().ParentPath).GetWithoutSubPath();
+	const FString ContextAssetPath = ContextAssetObjectPath.ToString();
+	 
+	for(FRigVMVariantRef VariantRef : VariantRefs)
+	{
+		TSharedPtr<FVariantTreeRowInfo> ParentRowInfo;
+		if(VariantRef.ObjectPath.IsSubobject())
 		{
-			TSharedPtr<SWidget> Widget = OnCreateVariantRefRow.Execute(VariantRef);
-			Widget->SetOnMouseDoubleClick(
-				FPointerEventHandler::CreateLambda(
-					[VariantRef, this](const FGeometry& InGeometry, const FPointerEvent& InMouseEvent) -> FReply
+			const FString AssetPath = VariantRef.ObjectPath.GetWithoutSubPath().ToString();
+			if(AssetPath != ContextAssetPath)
+			{
+				if(const FRigVMVariantRef* AssetVariantRef = AssetPathToVariantRef.Find(AssetPath))
+				{
+					if(!PathToRowInfo.Contains(AssetPath))
 					{
-						(void)OnBrowseVariantRef.ExecuteIfBound(VariantRef);
-						return FReply::Handled();
+						ParentRowInfo = MakeShareable(new FVariantTreeRowInfo);
+						ParentRowInfo->VariantRef = *AssetVariantRef;
+						PathToRowInfo.Add(AssetPath, ParentRowInfo);
+						VariantTreeRowInfos.Add(ParentRowInfo);
 					}
-				)
-			);
-			VariantRefListBox->AddSlot()
-			.AutoHeight()
-			.HAlign(HAlign_Fill)
-			.Padding(0, 4, 0, 0)
-			[
-				Widget.ToSharedRef()
-			];
+					else
+					{
+						ParentRowInfo = PathToRowInfo.FindChecked(AssetPath);
+					}
+				}
+			}
+		}
+
+		TSharedPtr<FVariantTreeRowInfo> RowInfo = MakeShareable(new FVariantTreeRowInfo);
+		RowInfo->VariantRef = VariantRef;
+		if(ParentRowInfo)
+		{
+			ParentRowInfo->NestedInfos.Add(RowInfo);
+		}
+		else
+		{
+			VariantTreeRowInfos.Add(RowInfo);
 		}
 	}
+	
+	VariantRefTreeView->RequestTreeRefresh();
 }
 
 const FSlateBrush* SRigVMVariantWidget::GetThumbnailBorder(TSharedRef<SBorder> InThumbnailBorder) const
diff --git a/Engine/Plugins/Runtime/RigVM/Source/RigVMEditor/Public/Widgets/SRigVMVariantWidget.h b/Engine/Plugins/Runtime/RigVM/Source/RigVMEditor/Public/Widgets/SRigVMVariantWidget.h
--- a/Engine/Plugins/Runtime/RigVM/Source/RigVMEditor/Public/Widgets/SRigVMVariantWidget.h
+++ b/Engine/Plugins/Runtime/RigVM/Source/RigVMEditor/Public/Widgets/SRigVMVariantWidget.h
@@ -5,6 +5,8 @@
 #include "Widgets/Layout/SBox.h"
 #include "Widgets/Layout/SBorder.h"
 #include "Widgets/SBoxPanel.h"
+#include "Widgets/Views/STreeView.h"
+#include "Widgets/SToolTip.h"
 #include "RigVMCore/RigVMVariant.h"
 #include "Widgets/SRigVMVariantTagWidget.h"
 
@@ -23,6 +25,26 @@ struct RIGVMEDITOR_API FRigVMVariantWidgetContext
 	FString ParentPath;
 };
 
+class SRigVMVariantToolTipWithTags : public SToolTip
+{
+public:
+	SLATE_BEGIN_ARGS(SRigVMVariantToolTipWithTags) {}
+		SLATE_ATTRIBUTE(FText, ToolTipText)
+		SLATE_EVENT(FRigVMVariant_OnGetTags, OnGetTags)
+	SLATE_END_ARGS()
+
+	void Construct(const FArguments& InArgs);
+
+	// IToolTip interface
+	virtual bool IsEmpty() const override;
+	virtual void OnOpening() override;
+	virtual void OnClosed() override;
+
+private:
+	FRigVMVariant_OnGetTags GetTagsDelegate;
+	SToolTip::FArguments SuperClassArgs;
+};
+
 class RIGVMEDITOR_API SRigVMVariantWidget : public SBox
 {
 public:
@@ -57,7 +79,30 @@ public:
 
 private:
 
+	struct FVariantTreeRowInfo
+	{
+		FRigVMVariantRef VariantRef;
+		TArray<TSharedPtr<FVariantTreeRowInfo>> NestedInfos;
+	};
+
+	class RIGVMEDITOR_API SRigVMVariantRefTreeRow
+	: public STableRow<TSharedPtr<FVariantTreeRowInfo>>
+	{
+	public:
+	
+		SLATE_BEGIN_ARGS(SRigVMVariantRefTreeRow)
+		{}
+		SLATE_ARGUMENT(TSharedPtr<SWidget>, Content)
+		SLATE_END_ARGS()
+
+		virtual ~SRigVMVariantRefTreeRow() override;
+		void Construct(const FArguments& InArgs, const TSharedRef<STableViewBase>& OwnerTableView);
+	};
+
+	const FRigVMVariantWidgetContext& GetVariantContext() const;
 	EVisibility GetVariantRefListVisibility() const;
+	TSharedRef<ITableRow> GenerateVariantTreeRow(TSharedPtr<FVariantTreeRowInfo> InRowInfo, const TSharedRef<STableViewBase>& OwnerTable);
+	void GetChildrenForVariantInfo(TSharedPtr<FVariantTreeRowInfo> InInfo, TArray<TSharedPtr<FVariantTreeRowInfo>>& OutChildren);
 	TSharedPtr<SWidget> CreateDefaultVariantRefRow(const FRigVMVariantRef& InVariantRef) const;
 	void RebuildVariantRefList();
 	const FSlateBrush* GetThumbnailBorder(TSharedRef<SBorder> InThumbnailBorder) const;
@@ -71,7 +116,9 @@ private:
 	FRigVMVariantWidget_OnCreateVariantRefRow OnCreateVariantRefRow;
 	FRigVMVariantWidget_OnBrowseVariantRef OnBrowseVariantRef;
 	TArray<FRigVMVariantRef> VariantRefs;
+	TArray<TSharedPtr<FVariantTreeRowInfo>> VariantTreeRowInfos;
 	uint32 VariantRefHash;
 	TSharedPtr<SVerticalBox> VariantRefListBox;
+	TSharedPtr<STreeView<TSharedPtr<FVariantTreeRowInfo>>> VariantRefTreeView;
 	TAttribute<FRigVMVariantWidgetContext> ContextAttribute;
 };
\ No newline at end of file
