RigVM: Improve Variant Ref widget and create variant workflows

#rnx

diff --git a/Engine/Plugins/Runtime/RigVM/Source/RigVMEditor/Private/Editor/RigVMGraphDetailCustomization.cpp b/Engine/Plugins/Runtime/RigVM/Source/RigVMEditor/Private/Editor/RigVMGraphDetailCustomization.cpp
--- a/Engine/Plugins/Runtime/RigVM/Source/RigVMEditor/Private/Editor/RigVMGraphDetailCustomization.cpp
+++ b/Engine/Plugins/Runtime/RigVM/Source/RigVMEditor/Private/Editor/RigVMGraphDetailCustomization.cpp
@@ -30,7 +30,9 @@
 #include "Widgets/SRigVMGraphPinEnumPicker.h"
 #include "Widgets/SRigVMVariantWidget.h"
 #include "ScopedTransaction.h"
-#include "Misc/UObjectToken.h"
+#include "Editor/RigVMEditorTools.h"
+#include "ContentBrowserModule.h"
+#include "IContentBrowserSingleton.h"
 
 #define LOCTEXT_NAMESPACE "RigVMGraphDetailCustomization"
 
@@ -1425,6 +1427,15 @@ void FRigVMGraphDetailCustomization::OnBrowseVariantRef(const FRigVMVariantRef&
 			}
 		}
 	}
+	else
+	{
+		const FAssetData AssetData = UE::RigVM::Editor::Tools::FindAssetFromAnyPath(InVariantRef.ObjectPath.ToString(), true);
+		if(AssetData.IsValid())
+		{
+			const FContentBrowserModule& ContentBrowserModule = FModuleManager::Get().LoadModuleChecked<FContentBrowserModule>(TEXT("ContentBrowser"));
+			ContentBrowserModule.Get().SyncBrowserToAssets({AssetData});
+		}
+	}
 }
 
 TArray<FRigVMTag> FRigVMGraphDetailCustomization::OnGetAssignedTags() const
