RigVM: Support for Pin Categories - Front End

#rnx

diff --git a/Engine/Plugins/Runtime/RigVM/Source/RigVM/Private/RigVMCore/RigVMNodeLayout.cpp b/Engine/Plugins/Runtime/RigVM/Source/RigVM/Private/RigVMCore/RigVMNodeLayout.cpp
--- a/Engine/Plugins/Runtime/RigVM/Source/RigVM/Private/RigVMCore/RigVMNodeLayout.cpp
+++ b/Engine/Plugins/Runtime/RigVM/Source/RigVM/Private/RigVMCore/RigVMNodeLayout.cpp
@@ -15,6 +15,29 @@ FString FRigVMPinCategory::GetName() const
 	return Right;
 }
 
+bool FRigVMPinCategory::IsDefaultCategory() const
+{
+	return Path.Equals(GetDefaultCategoryName(), ESearchCase::IgnoreCase);
+}
+
+const FString& FRigVMPinCategory::GetDefaultCategoryName()
+{
+	static const FString DefaultCategoryName = TEXT("Default");
+	return DefaultCategoryName;
+}
+
+bool FRigVMNodeLayout::IsValid() const
+{
+	for(const FRigVMPinCategory& Category : Categories)
+	{
+		if(!Category.Elements.IsEmpty())
+		{
+			return true;
+		}
+	}
+	return false;
+}
+
 const FString* FRigVMNodeLayout::FindCategory(const FString& InElement) const
 {
 	for(const FRigVMPinCategory& Category : Categories)
diff --git a/Engine/Plugins/Runtime/RigVM/Source/RigVM/Public/RigVMCore/RigVMNodeLayout.h b/Engine/Plugins/Runtime/RigVM/Source/RigVM/Public/RigVMCore/RigVMNodeLayout.h
--- a/Engine/Plugins/Runtime/RigVM/Source/RigVM/Public/RigVMCore/RigVMNodeLayout.h
+++ b/Engine/Plugins/Runtime/RigVM/Source/RigVM/Public/RigVMCore/RigVMNodeLayout.h
@@ -13,6 +13,7 @@ struct RIGVM_API FRigVMPinCategory
 	FRigVMPinCategory()
 	: Path()
 	, Elements()
+	, bExpandedByDefault(true)
 	{}
 
 	UPROPERTY(BlueprintReadOnly, VisibleAnywhere, Category=FunctionArgument)
@@ -21,6 +22,9 @@ struct RIGVM_API FRigVMPinCategory
 	UPROPERTY(BlueprintReadOnly, VisibleAnywhere, Category=FunctionArgument)
 	TArray<FString> Elements;
 
+	UPROPERTY(BlueprintReadOnly, VisibleAnywhere, Category=FunctionArgument)
+	bool bExpandedByDefault;
+
 	FString GetName() const;
 
 	friend uint32 GetTypeHash(const FRigVMPinCategory& Category)
@@ -30,6 +34,7 @@ struct RIGVM_API FRigVMPinCategory
 		{
 			Hash = HashCombine(Hash, GetTypeHash(Element));
 		}
+		Hash = HashCombine(Hash, GetTypeHash(Category.bExpandedByDefault));
 		return Hash;
 	}
 
@@ -42,8 +47,27 @@ struct RIGVM_API FRigVMPinCategory
 	{
 		Ar << Category.Path;
 		Ar << Category.Elements;
+		if(Ar.IsLoading())
+		{
+			if(Ar.CustomVer(FRigVMObjectVersion::GUID) < FRigVMObjectVersion::FunctionHeaderLayoutStoresCategoryExpansion)
+			{
+				Category.bExpandedByDefault = true;
+			}
+			else
+			{
+				Ar << Category.bExpandedByDefault;
+			}
+		}
+		else
+		{
+			Ar << Category.bExpandedByDefault;
+		}
 		return Ar;
 	}
+
+	bool IsDefaultCategory() const;
+
+	static const FString& GetDefaultCategoryName();
 };
 
 USTRUCT(BlueprintType)
@@ -72,6 +96,8 @@ struct RIGVM_API FRigVMNodeLayout
 		DisplayNames.Reset();
 	}
 
+	bool IsValid() const;
+
 	friend uint32 GetTypeHash(const FRigVMNodeLayout& Layout)
 	{
 		uint32 Hash = 0;;
diff --git a/Engine/Plugins/Runtime/RigVM/Source/RigVM/Public/RigVMObjectVersion.h b/Engine/Plugins/Runtime/RigVM/Source/RigVM/Public/RigVMObjectVersion.h
--- a/Engine/Plugins/Runtime/RigVM/Source/RigVM/Public/RigVMObjectVersion.h
+++ b/Engine/Plugins/Runtime/RigVM/Source/RigVM/Public/RigVMObjectVersion.h
@@ -54,6 +54,9 @@ struct RIGVM_API FRigVMObjectVersion
 		// Storing user interface relevant pin index in category
 		FunctionHeaderLayoutStoresPinIndexInCategory,
 
+		// Storing user interface relevant category expansion
+		FunctionHeaderLayoutStoresCategoryExpansion,
+
 		// -----<new versions can be added above this line>-------------------------------------------------
 		VersionPlusOne,
 		LatestVersion = VersionPlusOne - 1
diff --git a/Engine/Plugins/Runtime/RigVM/Source/RigVMDeveloper/Private/EdGraph/RigVMEdGraphNode.cpp b/Engine/Plugins/Runtime/RigVM/Source/RigVMDeveloper/Private/EdGraph/RigVMEdGraphNode.cpp
--- a/Engine/Plugins/Runtime/RigVM/Source/RigVMDeveloper/Private/EdGraph/RigVMEdGraphNode.cpp
+++ b/Engine/Plugins/Runtime/RigVM/Source/RigVMDeveloper/Private/EdGraph/RigVMEdGraphNode.cpp
@@ -741,6 +741,16 @@ bool URigVMEdGraphNode::CreateGraphPinFromModelPin(const URigVMPin* InModelPin,
 		return false;
 	}
 
+	const FString Category = InModelPin->GetCategory();
+	if(!Category.IsEmpty())
+	{
+		(void)CreateGraphPinFromCategory(Category, EGPD_Input);
+		if(UEdGraphPin* CategoryPin = FindGraphPinFromCategory(Category, true))
+		{
+			InParentPin = CategoryPin;
+		}
+	}
+
 	const FPinPair PairConst = CachedPins.FindOrAdd((URigVMPin*)InModelPin);
 
 	auto CreatePinLambda = [this](const URigVMPin* InModelPin, EEdGraphPinDirection InDirection, UEdGraphPin* InParentPin) -> UEdGraphPin*
@@ -1077,9 +1087,11 @@ bool URigVMEdGraphNode::ModelPinsChanged(bool bForce)
 
 		if(const URigVMNode* ModelNode = GetModelNode())
 		{
-			for(const auto& Pair : CachedCategoryPins)
+			TArray<FString> Categories;
+			CachedCategoryPins.GetKeys(Categories);
+			
+			for(const FString& Category : Categories)
 			{
-				const FString Category = Pair.Key;
 				if(!ModelNode->GetPinCategories().Contains(Category))
 				{
 					PinsRemoved += CategoryPinRemoved_Internal(Category);
@@ -1147,14 +1159,83 @@ bool URigVMEdGraphNode::ModelPinsChanged(bool bForce)
 
 	auto ReparentPins = [this](const TArray<URigVMPin*>& InModelPins) -> int32
 	{
-		// todo
-		return 0;
+		int32 PinsReparented = 0;
+		for(const URigVMPin* ModelPin : InModelPins)
+		{
+			for(int32 Phase = 0; Phase < 2; Phase++)
+			{
+				if(UEdGraphPin* EdGraphPin = FindGraphPinFromModelPin(ModelPin, Phase == 0))
+				{
+					UEdGraphPin* ParentEdGraphPin = nullptr;
+
+					const FString Category = ModelPin->GetCategory();
+					if(Category.IsEmpty())
+					{
+						if(const URigVMPin* ParentModelPin = ModelPin->GetParentPin())
+						{
+							ParentEdGraphPin = FindGraphPinFromModelPin(ParentModelPin, Phase == 0);
+						}
+					}
+					else
+					{
+						ParentEdGraphPin = FindGraphPinFromCategory(Category, Phase == 0);
+					}
+
+					if(EdGraphPin->ParentPin != ParentEdGraphPin)
+					{
+						if(EdGraphPin->ParentPin)
+						{
+							EdGraphPin->ParentPin->SubPins.Remove(EdGraphPin);
+							EdGraphPin->ParentPin = nullptr;
+						}
+						if(ParentEdGraphPin)
+						{
+							ParentEdGraphPin->SubPins.Add(EdGraphPin);
+							EdGraphPin->ParentPin = ParentEdGraphPin;
+						}
+						PinsReparented++;
+					}
+				}
+			}
+		}
+		return PinsReparented;
 	};
 
 	auto ReparentCategoryPins = [this](const TArray<FString>& InCategories) -> int32
 	{
-		// todo
-		return 0;
+		int32 PinsReparented = 0;
+		for(const FString& Category : InCategories)
+		{
+			for(int32 Phase = 0; Phase < 2; Phase++)
+			{
+				if(UEdGraphPin* EdGraphPin = FindGraphPinFromCategory(Category, Phase == 0))
+				{
+					UEdGraphPin* ParentEdGraphPin = nullptr;
+					
+					FString ParentCategory, CategoryName;
+					if(RigVMStringUtils::SplitNodePathAtEnd(Category, ParentCategory, CategoryName))
+					{
+						ParentEdGraphPin = FindGraphPinFromCategory(ParentCategory, Phase == 0);
+					}
+
+					if(EdGraphPin->ParentPin != ParentEdGraphPin)
+					{
+						if(EdGraphPin->ParentPin)
+						{
+							EdGraphPin->ParentPin->SubPins.Remove(EdGraphPin);
+							EdGraphPin->ParentPin = nullptr;
+						}
+						if(ParentEdGraphPin)
+						{
+							ParentEdGraphPin->SubPins.Add(EdGraphPin);
+							EdGraphPin->ParentPin = ParentEdGraphPin;
+						}
+						PinsReparented++;
+					}
+				}
+			}
+		}
+		return PinsReparented;
 	};
 
 	int32 PinsAdded = 0;
@@ -1164,6 +1245,16 @@ bool URigVMEdGraphNode::ModelPinsChanged(bool bForce)
 	PinsAdded += AddMissingPins(InputOutputPins);
 	PinsAdded += AddMissingPins(InputPins);
 
+	int32 PinsReparented = 0;
+	if(const URigVMNode* ModelNode = GetModelNode())
+	{
+		PinsReparented += ReparentCategoryPins(ModelNode->GetPinCategories());
+	}
+	PinsReparented += ReparentPins(InputPins);
+	PinsReparented += ReparentPins(InputOutputPins);
+	PinsReparented += ReparentPins(OutputPins);
+	PinsReparented += ReparentPins(ExecutePins);
+
 	const int32 PinsRemoved = RemoveObsoletePins();
 
 	// working through it in the opposite order
@@ -1174,17 +1265,45 @@ bool URigVMEdGraphNode::ModelPinsChanged(bool bForce)
 	PinsReordered += OrderPins(OutputPins);
 	PinsReordered += OrderPins(ExecutePins);
 
-	int32 PinsReparented = 0;
-	if(const URigVMNode* ModelNode = GetModelNode())
+	int32 PinLabelsChanged = 0;
+	auto SyncPinLabel = [](const URigVMPin* InModelPin, const FPinPair& OutPinPair)
 	{
-		PinsReparented += ReparentCategoryPins(ModelNode->GetPinCategories());
-		PinsReparented += ReparentPins(InputPins);
-		PinsReparented += ReparentPins(InputOutputPins);
-		PinsReparented += ReparentPins(OutputPins);
-		PinsReparented += ReparentPins(ExecutePins);
+		int32 PinLabelsChanged = 0;
+		if(InModelPin)
+		{
+			const FName DesiredFriendlyName = InModelPin->GetDisplayName();
+			if(!DesiredFriendlyName.IsNone())
+			{
+				const FText DesiredFriendlyNameText = FText::FromName(DesiredFriendlyName);
+				for(int32 Phase = 0; Phase < 2; Phase++)
+				{
+					UEdGraphPin* Pin = Phase == 0 ? OutPinPair.InputPin : OutPinPair.OutputPin;
+					if(Pin)
+					{
+						if(!Pin->PinFriendlyName.EqualTo(DesiredFriendlyNameText))
+						{
+							Pin->PinFriendlyName = DesiredFriendlyNameText;
+							PinLabelsChanged++;
+						}
+					}
+				}
+			}
+		}
+		return PinLabelsChanged;
+	};
+	for(const TPair<FString,FPinPair>& Pair : CachedCategoryPins)
+	{
+		if(const URigVMPin* ModelPin = FindModelPinFromGraphPin(Pair.Value.InputPin != nullptr ? Pair.Value.InputPin : Pair.Value.OutputPin))
+		{
+			PinLabelsChanged += SyncPinLabel(ModelPin, Pair.Value);
+		}
+	}
+	for(const TPair<URigVMPin*,FPinPair>& Pair : CachedPins)
+	{
+		PinLabelsChanged += SyncPinLabel(Pair.Key, Pair.Value);
 	}
 
-	const bool bResult = (PinsAdded > 0) || (PinsRemoved > 0) || (PinsReordered > 0) || (PinsReparented > 0); 
+	const bool bResult = (PinsAdded > 0) || (PinLabelsChanged > 0) || (PinsRemoved > 0) || (PinsReordered > 0) || (PinsReparented > 0); 
 	if(bResult || bForce)
 	{
 		OnNodePinsChanged().Broadcast();
@@ -1225,21 +1344,10 @@ bool URigVMEdGraphNode::ModelPinAdded_Internal(const URigVMPin* InModelPin)
 	UEdGraphPin* InputParentPin = FindGraphPinFromModelPin(InModelPin->GetParentPin(), true);
 	UEdGraphPin* OutputParentPin = FindGraphPinFromModelPin(InModelPin->GetParentPin(), false);
 
-	const FString Category = InModelPin->GetCategory();
-		
 	if(InModelPin->GetDirection() == ERigVMPinDirection::Input ||
 		InModelPin->GetDirection() == ERigVMPinDirection::Visible ||
 		InModelPin->GetDirection() == ERigVMPinDirection::IO)
 	{
-		if(!Category.IsEmpty())
-		{
-			(void)CreateGraphPinFromCategory(Category, EGPD_Input);
-			if(UEdGraphPin* CategoryPin = FindGraphPinFromCategory(Category, true))
-			{
-				InputParentPin = CategoryPin;
-			}
-		}
-
 		if(CreateGraphPinFromModelPin(InModelPin, EGPD_Input, InputParentPin))
 		{
 			bResult = true;
@@ -1249,15 +1357,6 @@ bool URigVMEdGraphNode::ModelPinAdded_Internal(const URigVMPin* InModelPin)
 	if(InModelPin->GetDirection() == ERigVMPinDirection::Output ||
 		InModelPin->GetDirection() == ERigVMPinDirection::IO)
 	{
-		if(!Category.IsEmpty())
-		{
-			(void)CreateGraphPinFromCategory(Category, EGPD_Output);
-			if(UEdGraphPin* CategoryPin = FindGraphPinFromCategory(Category, false))
-			{
-				OutputParentPin = CategoryPin;
-			}
-		}
-
 		if(CreateGraphPinFromModelPin(InModelPin, EGPD_Output, OutputParentPin))
 		{
 			bResult = true;
@@ -1381,7 +1480,7 @@ bool URigVMEdGraphNode::CategoryPinRemoved_Internal(const FString& InCategory)
 	{
 		if(UEdGraphPin* GraphPin = FindGraphPinFromCategory(InCategory, bAsInput))
 		{
-			RemoveGraphSubPins(GraphPin);
+			check(GraphPin->SubPins.IsEmpty());
 			Pins.Remove(GraphPin);
 			GraphPin->MarkAsGarbage();
 		}
diff --git a/Engine/Plugins/Runtime/RigVM/Source/RigVMDeveloper/Private/RigVMModel/Nodes/RigVMFunctionInterfaceNode.cpp b/Engine/Plugins/Runtime/RigVM/Source/RigVMDeveloper/Private/RigVMModel/Nodes/RigVMFunctionInterfaceNode.cpp
--- a/Engine/Plugins/Runtime/RigVM/Source/RigVMDeveloper/Private/RigVMModel/Nodes/RigVMFunctionInterfaceNode.cpp
+++ b/Engine/Plugins/Runtime/RigVM/Source/RigVMDeveloper/Private/RigVMModel/Nodes/RigVMFunctionInterfaceNode.cpp
@@ -43,42 +43,6 @@ FText URigVMFunctionInterfaceNode::GetToolTipTextForPin(const URigVMPin* InPin)
 	return Super::GetToolTipTextForPin(InPin);
 }
 
-FName URigVMFunctionInterfaceNode::GetDisplayNameForPin(const FString& InPinPath) const
-{
-	if(const URigVMPin* OuterPin = FindReferencedPin(InPinPath))
-	{
-		return OuterPin->GetDisplayName();
-	}
-	return Super::GetDisplayNameForPin(InPinPath);
-}
-
-FString URigVMFunctionInterfaceNode::GetCategoryForPin(const FString& InPinPath) const
-{
-	if(const URigVMPin* OuterPin = FindReferencedPin(InPinPath))
-	{
-		return OuterPin->GetCategory();
-	}
-	return Super::GetCategoryForPin(InPinPath);
-}
-
-int32 URigVMFunctionInterfaceNode::GetIndexInCategoryForPin(const FString& InPinPath) const
-{
-	if(const URigVMPin* OuterPin = FindReferencedPin(InPinPath))
-	{
-		return OuterPin->GetIndexInCategory();
-	}
-	return Super::GetIndexInCategoryForPin(InPinPath);
-}
-
-TArray<FString> URigVMFunctionInterfaceNode::GetPinCategories() const
-{
-	if(const URigVMCollapseNode* OuterNode = Cast<URigVMCollapseNode>(GetGraph()->GetOuter()))
-	{
-		return OuterNode->GetPinCategories();
-	}
-	return PinCategories;
-}
-
 const URigVMPin* URigVMFunctionInterfaceNode::FindReferencedPin(const URigVMPin* InPin) const
 {
 	return FindReferencedPin(InPin->GetSegmentPath(true));
diff --git a/Engine/Plugins/Runtime/RigVM/Source/RigVMDeveloper/Private/RigVMModel/Nodes/RigVMFunctionReferenceNode.cpp b/Engine/Plugins/Runtime/RigVM/Source/RigVMDeveloper/Private/RigVMModel/Nodes/RigVMFunctionReferenceNode.cpp
--- a/Engine/Plugins/Runtime/RigVM/Source/RigVMDeveloper/Private/RigVMModel/Nodes/RigVMFunctionReferenceNode.cpp
+++ b/Engine/Plugins/Runtime/RigVM/Source/RigVMDeveloper/Private/RigVMModel/Nodes/RigVMFunctionReferenceNode.cpp
@@ -311,14 +311,30 @@ FText URigVMFunctionReferenceNode::GetToolTipTextForPin(const URigVMPin* InPin)
 
 TArray<FString> URigVMFunctionReferenceNode::GetPinCategories() const
 {
+	FRigVMNodeLayout ReferenceLayout;
+	const FRigVMNodeLayout* Layout = &ReferencedFunctionHeader.Layout;
+	if(IsReferencedNodeLoaded())
+	{
+		ReferenceLayout = LoadReferencedNode()->GetNodeLayout();
+		Layout = &ReferenceLayout;
+	}
 	TArray<FString> TransientPinCategories;
-	for(const FRigVMPinCategory& Category : ReferencedFunctionHeader.Layout.Categories)
+	for(const FRigVMPinCategory& Category : Layout->Categories)
 	{
 		TransientPinCategories.Add(Category.Path);
 	}
 	return TransientPinCategories;
 }
 
+FRigVMNodeLayout URigVMFunctionReferenceNode::GetNodeLayout(bool bIncludeEmptyCategories) const
+{
+	if(IsReferencedNodeLoaded())
+	{
+		return LoadReferencedNode()->GetNodeLayout(bIncludeEmptyCategories);
+	}
+	return ReferencedFunctionHeader.Layout;
+}
+
 FRigVMGraphFunctionIdentifier URigVMFunctionReferenceNode::GetFunctionIdentifier() const
 {
 	return GetReferencedFunctionHeader().LibraryPointer;
diff --git a/Engine/Plugins/Runtime/RigVM/Source/RigVMDeveloper/Private/RigVMModel/Nodes/RigVMLibraryNode.cpp b/Engine/Plugins/Runtime/RigVM/Source/RigVMDeveloper/Private/RigVMModel/Nodes/RigVMLibraryNode.cpp
--- a/Engine/Plugins/Runtime/RigVM/Source/RigVMDeveloper/Private/RigVMModel/Nodes/RigVMLibraryNode.cpp
+++ b/Engine/Plugins/Runtime/RigVM/Source/RigVMDeveloper/Private/RigVMModel/Nodes/RigVMLibraryNode.cpp
@@ -328,7 +328,7 @@ FRigVMGraphFunctionHeader URigVMLibraryNode::GetFunctionHeader(IRigVMGraphFuncti
 		}
 	}
 
-	Header.Layout = GetPinLayout();
+	Header.Layout = GetNodeLayout();
 	
 	return Header;
 }
diff --git a/Engine/Plugins/Runtime/RigVM/Source/RigVMDeveloper/Private/RigVMModel/RigVMClient.cpp b/Engine/Plugins/Runtime/RigVM/Source/RigVMDeveloper/Private/RigVMModel/RigVMClient.cpp
--- a/Engine/Plugins/Runtime/RigVM/Source/RigVMDeveloper/Private/RigVMModel/RigVMClient.cpp
+++ b/Engine/Plugins/Runtime/RigVM/Source/RigVMDeveloper/Private/RigVMModel/RigVMClient.cpp
@@ -1565,21 +1565,24 @@ void FRigVMClient::HandleGraphModifiedEvent(ERigVMGraphNotifType InNotifType, UR
 		case ERigVMGraphNotifType::PinWatchedChanged: // A pin's watch state has changed (Subject == URigVMPin)
 		case ERigVMGraphNotifType::PinDisplayNameChanged: // A pin's display name / UI label has changed (Subject == URigVMPin)
 		case ERigVMGraphNotifType::PinCategoryChanged: // A pin's category has changed (Subject == URigVMPin)
+		case ERigVMGraphNotifType::PinCategoryExpansionChanged: // A pin category has been collapsed / expanded on the function node (Subject == URigVMNode)
 		{
-			if (URigVMPin* Pin = Cast<URigVMPin>(InSubject))
+			URigVMNode* Node = Cast<URigVMNode>(InSubject);
+			if (const URigVMPin* Pin = Cast<URigVMPin>(InSubject))
 			{
-				if (URigVMNode* Node = Cast<URigVMNode>(Pin->GetNode()))
+				Node = Cast<URigVMNode>(Pin->GetNode());
+			}
+			if(Node)
+			{
+				if (URigVMLibraryNode* LibraryNode = Node->FindFunctionForNode())
 				{
-					if (URigVMLibraryNode* LibraryNode = Node->FindFunctionForNode())
-					{
-						DirtyGraphFunctionCompilationData(LibraryNode);
-					}
-					if(Node->GetOuter()->IsA<URigVMFunctionLibrary>())
+					DirtyGraphFunctionCompilationData(LibraryNode);
+				}
+				if(Node->GetOuter()->IsA<URigVMFunctionLibrary>())
+				{
+					if (const URigVMCollapseNode* CollapseNode = Cast<URigVMCollapseNode>(Node))
 					{
-						if (const URigVMCollapseNode* CollapseNode = Cast<URigVMCollapseNode>(Node))
-						{
-							UpdateGraphFunctionData(CollapseNode);
-						}
+						UpdateGraphFunctionData(CollapseNode);
 					}
 				}
 			}
diff --git a/Engine/Plugins/Runtime/RigVM/Source/RigVMDeveloper/Private/RigVMModel/RigVMController.cpp b/Engine/Plugins/Runtime/RigVM/Source/RigVMDeveloper/Private/RigVMModel/RigVMController.cpp
--- a/Engine/Plugins/Runtime/RigVM/Source/RigVMDeveloper/Private/RigVMModel/RigVMController.cpp
+++ b/Engine/Plugins/Runtime/RigVM/Source/RigVMDeveloper/Private/RigVMModel/RigVMController.cpp
@@ -87,6 +87,7 @@ void FRigVMClientPatchResult::Merge(const FRigVMClientPatchResult& InOther)
 URigVMController::URigVMController()
 	: bValidatePinDefaults(true)
 	, bSuspendNotifications(false)
+	, bSuspendRefreshingFunctionReferences(false)
 	, bReportWarningsAndErrors(true)
 	, bIgnoreRerouteCompactnessChanges(false)
 	, UserLinkDirection(ERigVMPinDirection::Invalid)
@@ -108,6 +109,7 @@ URigVMController::URigVMController(const FObjectInitializer& ObjectInitializer)
 	: Super(ObjectInitializer)
 	, bValidatePinDefaults(true)
 	, bSuspendNotifications(false)
+	, bSuspendRefreshingFunctionReferences(false)
 	, bReportWarningsAndErrors(true)
 	, bIgnoreRerouteCompactnessChanges(false)
 	, UserLinkDirection(ERigVMPinDirection::Invalid)
@@ -8029,6 +8031,130 @@ bool URigVMController::SetPinDisplayName(URigVMPin* InPin, const FString& InDisp
 	InPin->DisplayName = InDisplayName.IsEmpty() ? FName(NAME_None) : FName(*InDisplayName);
 	Notify(ERigVMGraphNotifType::PinDisplayNameChanged, InPin);
 
+	if(const URigVMLibraryNode* LibraryNode = Cast<URigVMLibraryNode>(InPin->GetNode()))
+	{
+		if (LibraryNode->GetGraph()->IsA<URigVMFunctionLibrary>())
+		{
+			PropagateNotificationToFunctionReferences(LibraryNode, ERigVMGraphNotifType::PinDisplayNameChanged, InPin, false);
+		}
+	}
+
+	if (!bSuspendNotifications)
+	{
+		const URigVMGraph* Graph = GetGraph();
+		check(Graph);
+		(void)Graph->MarkPackageDirty();
+	}
+
+	if (bSetupUndoRedo)
+	{
+		GetActionStack()->EndAction(Action);
+	}
+
+	return true;
+}
+
+bool URigVMController::AddEmptyPinCategory(const FName& InNodeName, const FString& InCategory, bool bSetupUndoRedo, bool bPrintPythonCommand)
+{
+	if (!IsValidGraph())
+	{
+		return false;
+	}
+
+	if (!bIsTransacting && !IsGraphEditable())
+	{
+		return false;
+	}
+
+	const URigVMGraph* Graph = GetGraph();
+	check(Graph);
+
+	const URigVMNode* Node = Graph->FindNodeByName(InNodeName);
+	if (Node == nullptr)
+	{
+		ReportErrorf(TEXT("Cannot find node '%s'."), *InNodeName.ToString());
+		return false;
+	}
+
+	const bool bSuccess = AddEmptyPinCategory(Node, InCategory, bSetupUndoRedo);
+	if (bSuccess && bPrintPythonCommand)
+	{
+		const FString GraphName = GetSchema()->GetSanitizedGraphName(GetGraph()->GetGraphName());
+
+		RigVMPythonUtils::Print(GetSchema()->GetGraphOuterName(GetGraph()),
+			FString::Printf(TEXT("blueprint.get_controller_by_name('%s').add_empty_pin_category('%s', '%s')"),
+			*GraphName,
+			*GetSchema()->GetSanitizedNodeName(InNodeName.ToString()),
+			*InCategory));
+	}
+
+	return bSuccess;
+}
+
+bool URigVMController::AddEmptyPinCategory(const URigVMNode* InNode, const FString& InPinCategory, bool bSetupUndoRedo)
+{
+	if(!IsValidNodeForGraph(InNode))
+	{
+		return false;
+	}
+
+	// for entry or return nodes we relay to the outer pins
+	if(InNode->IsA<URigVMFunctionInterfaceNode>())
+	{
+		if(const URigVMCollapseNode* OuterNode = Cast<URigVMCollapseNode>(InNode->GetGraph()->GetOuter()))
+		{
+			if(URigVMController* OuterController = GetControllerForGraph(OuterNode->GetGraph()))
+			{
+				return OuterController->AddEmptyPinCategory(OuterNode, InPinCategory, bSetupUndoRedo);
+			}
+		}
+	}
+	else if(InNode->IsA<URigVMFunctionReferenceNode>())
+	{
+		return false;
+	}
+
+	FRigVMControllerCompileBracketScope CompileScope(this);
+	FRigVMChangeNodePinCategoriesAction Action;
+	if (bSetupUndoRedo)
+	{
+		Action = FRigVMChangeNodePinCategoriesAction(this, InNode);
+		Action.SetTitle(TEXT("Add Empty Pin Category"));
+		GetActionStack()->BeginAction(Action);
+	}
+
+	const FName PinCategoryName = URigVMSchema::GetUniqueName(*InPinCategory, [InNode](const FName& InName) {
+		return !InNode->PinCategories.Contains(InName.ToString());
+	}, false, true);
+
+	const FString PinCategoryNameString = PinCategoryName.ToString();
+	
+	const_cast<URigVMNode*>(InNode)->PinCategories.Add(PinCategoryNameString);
+	Notify(ERigVMGraphNotifType::PinCategoriesChanged, const_cast<URigVMNode*>(InNode));
+
+	if(const URigVMLibraryNode* LibraryNode = Cast<URigVMLibraryNode>(InNode))
+	{
+		if (LibraryNode->GetGraph()->IsA<URigVMFunctionLibrary>())
+		{
+			RefreshFunctionReferences(LibraryNode, bSetupUndoRedo, false);
+			PropagateNotificationToFunctionReferences(LibraryNode, ERigVMGraphNotifType::PinCategoriesChanged, const_cast<URigVMNode*>(InNode), false);
+		}
+	}
+
+	// if this was the first category and it is the default category
+	// add all of the pins to it
+	if(InNode->PinCategories.Num() == 1)
+	{
+		if(PinCategoryNameString.Equals(FRigVMPinCategory::GetDefaultCategoryName(), ESearchCase::IgnoreCase))
+		{
+			// apply this to all pins - it's going to filter invalid pins anyway
+			for(URigVMPin* RootPin : InNode->Pins)
+			{
+				(void)SetPinCategory(RootPin, PinCategoryNameString, bSetupUndoRedo);
+			}
+		}
+	}
+
 	if (!bSuspendNotifications)
 	{
 		const URigVMGraph* Graph = GetGraph();
@@ -8038,6 +8164,7 @@ bool URigVMController::SetPinDisplayName(URigVMPin* InPin, const FString& InDisp
 
 	if (bSetupUndoRedo)
 	{
+		Action.UpdateAfterModification(InNode);
 		GetActionStack()->EndAction(Action);
 	}
 
@@ -8215,8 +8342,47 @@ bool URigVMController::SetPinCategory(URigVMPin* InPin, const FString& InCategor
 		GetActionStack()->BeginAction(Action);
 	}
 
+	// move all pins in the old category
+	if(!InPin->UserDefinedCategory.IsEmpty() && InPin->IndexInCategory != INDEX_NONE)
+	{
+		const TArray<URigVMPin*> PinsInOldCategory = InPin->GetNode()->GetPinsForCategory(InPin->UserDefinedCategory);
+		for(URigVMPin* OtherPin : PinsInOldCategory)
+		{
+			if(OtherPin != InPin && OtherPin->IndexInCategory != INDEX_NONE && OtherPin->IndexInCategory > InPin->IndexInCategory)
+			{
+				OtherPin->IndexInCategory--;
+				Notify(ERigVMGraphNotifType::PinCategoryChanged, OtherPin);
+			}
+		}
+	}
+
+	int32 IndexInCategory = 0;
+	if(NewCategory.IsEmpty())
+	{
+		IndexInCategory = INDEX_NONE;
+	}
+	else
+	{
+		IndexInCategory = InPin->GetNode()->GetPinsForCategory(NewCategory).Num();
+	}
+
 	InPin->UserDefinedCategory = NewCategory;
+	InPin->IndexInCategory = IndexInCategory;
 	Notify(ERigVMGraphNotifType::PinCategoryChanged, InPin);
+
+	if(const URigVMLibraryNode* LibraryNode = Cast<URigVMLibraryNode>(InPin->GetNode()))
+	{
+		if (LibraryNode->GetGraph()->IsA<URigVMFunctionLibrary>())
+		{
+			PropagateNotificationToFunctionReferences(LibraryNode, ERigVMGraphNotifType::PinCategoryChanged, InPin, false);
+		}
+	}
+
+	if(NewCategory.IsEmpty())
+	{
+		(void)SetPinDisplayName(InPin, FString(), bSetupUndoRedo);
+	}
+	
 	if(!NewCategory.IsEmpty())
 	{
 		if(bSetupUndoRedo)
@@ -8232,6 +8398,19 @@ bool URigVMController::SetPinCategory(URigVMPin* InPin, const FString& InCategor
 		}
 	}
 
+	if(!NewCategory.IsEmpty())
+	{
+		(void)SetPinCategoryExpansion(InPin->GetNode(), NewCategory, true, bSetupUndoRedo);
+	}
+
+	if(const URigVMLibraryNode* LibraryNode = Cast<URigVMLibraryNode>(InPin->GetNode()))
+	{
+		if (LibraryNode->GetGraph()->IsA<URigVMFunctionLibrary>())
+		{
+			RefreshFunctionReferences(LibraryNode, bSetupUndoRedo, false);
+		}
+	}
+
 	if (!bSuspendNotifications)
 	{
 		const URigVMGraph* Graph = GetGraph();
@@ -8292,7 +8471,7 @@ bool URigVMController::RemovePinCategory(const URigVMNode* InNode, const FString
 	}
 
 	// for entry or return nodes we relay to the outer pins
-	if(const URigVMFunctionInterfaceNode* InterfaceNode = Cast<URigVMFunctionInterfaceNode>(InNode))
+	if(InNode->IsA<URigVMFunctionInterfaceNode>())
 	{
 		if(const URigVMCollapseNode* OuterNode = Cast<URigVMCollapseNode>(InNode->GetGraph()->GetOuter()))
 		{
@@ -8334,6 +8513,7 @@ bool URigVMController::RemovePinCategory(const URigVMNode* InNode, const FString
 	int32 NumPinsAffected = 0;
 	for(URigVMPin* PinInCategory : PinsForCategory)
 	{
+		const TGuardValue<bool> GuardRefreshFunctions(bSuspendRefreshingFunctionReferences, true); 
 		if(SetPinCategory(PinInCategory, FString(), bSetupUndoRedo))
 		{
 			NumPinsAffected++;
@@ -8352,6 +8532,18 @@ bool URigVMController::RemovePinCategory(const URigVMNode* InNode, const FString
 	}
 
 	const bool bHasMadeChange = (NumCategoriesRemoved + NumPinsAffected) > 0;
+
+	if(bHasMadeChange)
+	{
+		if(const URigVMLibraryNode* LibraryNode = Cast<URigVMLibraryNode>(InNode))
+		{
+			if (LibraryNode->GetGraph()->IsA<URigVMFunctionLibrary>())
+			{
+				RefreshFunctionReferences(LibraryNode, bSetupUndoRedo, false);
+				PropagateNotificationToFunctionReferences(LibraryNode, ERigVMGraphNotifType::PinCategoriesChanged, const_cast<URigVMNode*>(InNode), false);
+			}
+		}
+	}
 	
 	if (!bSuspendNotifications && bHasMadeChange)
 	{
@@ -8421,8 +8613,14 @@ bool URigVMController::RenamePinCategory(const URigVMNode* InNode, const FString
 		return false;
 	}
 
+	if(InOldPinCategory.Equals(FRigVMPinCategory::GetDefaultCategoryName(), ESearchCase::IgnoreCase))
+	{
+		ReportErrorf(TEXT("Cannot rename '%s' category."), *InOldPinCategory);
+		return false;
+	}
+
 	// for entry or return nodes we relay to the outer pins
-	if(const URigVMFunctionInterfaceNode* InterfaceNode = Cast<URigVMFunctionInterfaceNode>(InNode))
+	if(InNode->IsA<URigVMFunctionInterfaceNode>())
 	{
 		if(const URigVMCollapseNode* OuterNode = Cast<URigVMCollapseNode>(InNode->GetGraph()->GetOuter()))
 		{
@@ -8469,6 +8667,7 @@ bool URigVMController::RenamePinCategory(const URigVMNode* InNode, const FString
 	// first update the category list on the node
 	for(const FString& CategoryToRename : CategoriesToRename)
 	{
+		const bool bWasExpanded = InNode->IsPinCategoryExpanded(CategoryToRename);
 		const int32 Index = InNode->PinCategories.Find(CategoryToRename);
 		if(CategoryToRename.Equals(InOldPinCategory, ESearchCase::CaseSensitive))
         {
@@ -8480,11 +8679,16 @@ bool URigVMController::RenamePinCategory(const URigVMNode* InNode, const FString
 			const_cast<URigVMNode*>(InNode)->PinCategories[Index] = NewCategory;
         }
 		PinsToUpdate.Append(InNode->GetPinsForCategory(CategoryToRename));
+
+		const_cast<URigVMNode*>(InNode)->PinCategoryExpansion.Remove(CategoryToRename);
+		const_cast<URigVMNode*>(InNode)->PinCategoryExpansion.FindOrAdd(InNode->PinCategories[Index]) = bWasExpanded;
 	}
 
 	// then update the pins and their categories
 	for(URigVMPin* PinToUpdate : PinsToUpdate)
 	{
+		const TGuardValue<bool> GuardRefreshFunctions(bSuspendRefreshingFunctionReferences, true); 
+
 		const FString CategoryToRename = PinToUpdate->GetCategory();
 		if(CategoryToRename.Equals(InOldPinCategory, ESearchCase::CaseSensitive))
 		{
@@ -8499,6 +8703,14 @@ bool URigVMController::RenamePinCategory(const URigVMNode* InNode, const FString
 	
 	Notify(ERigVMGraphNotifType::PinCategoriesChanged, const_cast<URigVMNode*>(InNode));
 
+	if(const URigVMLibraryNode* LibraryNode = Cast<URigVMLibraryNode>(InNode))
+	{
+		if (LibraryNode->GetGraph()->IsA<URigVMFunctionLibrary>())
+		{
+			PropagateNotificationToFunctionReferences(LibraryNode, ERigVMGraphNotifType::PinCategoriesChanged, const_cast<URigVMNode*>(InNode), false);
+		}
+	}
+	
 	if (!bSuspendNotifications)
 	{
 		const URigVMGraph* Graph = GetGraph();
@@ -8561,7 +8773,7 @@ bool URigVMController::SetPinCategoryIndex(const URigVMNode* InNode, const FStri
 	}
 
 	// for entry or return nodes we relay to the outer pins
-	if(const URigVMFunctionInterfaceNode* InterfaceNode = Cast<URigVMFunctionInterfaceNode>(InNode))
+	if(InNode->IsA<URigVMFunctionInterfaceNode>())
 	{
 		if(const URigVMCollapseNode* OuterNode = Cast<URigVMCollapseNode>(InNode->GetGraph()->GetOuter()))
 		{
@@ -8653,6 +8865,12 @@ bool URigVMController::SetPinCategoryExpansion(const URigVMNode* InNode, const F
 		return false;
 	}
 
+	if(InPinCategory.Equals(FRigVMPinCategory::GetDefaultCategoryName(), ESearchCase::IgnoreCase))
+	{
+		// nothing to do here - the default category is always expanded
+		return false;
+	}
+
 	if(InNode->IsPinCategoryExpanded(InPinCategory) == bIsExpanded)
 	{
 		return false;
@@ -8766,6 +8984,11 @@ bool URigVMController::SetPinIndexInCategory(URigVMPin* InPin, int32 InIndexInCa
 		return false;
 	}
 
+	if(InPin->GetCategory().IsEmpty())
+	{
+		InIndexInCategory = INDEX_NONE;
+	}
+
 	if(InPin->GetIndexInCategory() == InIndexInCategory)
 	{
 		return false;
@@ -8780,8 +9003,49 @@ bool URigVMController::SetPinIndexInCategory(URigVMPin* InPin, int32 InIndexInCa
 		GetActionStack()->BeginAction(Action);
 	}
 
-	InPin->IndexInCategory = InIndexInCategory;
-	Notify(ERigVMGraphNotifType::PinCategoryChanged, InPin);
+	if(!InPin->GetCategory().IsEmpty() && InIndexInCategory != INDEX_NONE)
+	{
+		TArray<URigVMPin*> PinsInCategory = InPin->GetNode()->GetPinsForCategory(InPin->GetCategory());
+		check(PinsInCategory.Contains(InPin));
+
+		PinsInCategory.Remove(InPin);
+		if(InIndexInCategory >= PinsInCategory.Num())
+		{
+			PinsInCategory.Add(InPin);
+		}
+		else
+		{
+			PinsInCategory.Insert(InPin, InIndexInCategory);
+		}
+
+		TArray<bool> HasChanged;
+		for(int32 Index = 0; Index < PinsInCategory.Num(); Index++)
+		{
+			HasChanged.Add(PinsInCategory[Index]->IndexInCategory != Index);
+			PinsInCategory[Index]->IndexInCategory = Index;
+		}
+		for(int32 Index = 0; Index < PinsInCategory.Num(); Index++)
+		{
+			if(HasChanged[Index])
+			{
+				Notify(ERigVMGraphNotifType::PinCategoryChanged, PinsInCategory[Index]);
+			}
+		}
+	}
+	else
+	{
+		InPin->IndexInCategory = InIndexInCategory;
+		Notify(ERigVMGraphNotifType::PinCategoryChanged, InPin);
+	}
+
+	if(const URigVMLibraryNode* LibraryNode = Cast<URigVMLibraryNode>(InPin->GetNode()))
+	{
+		if (LibraryNode->GetGraph()->IsA<URigVMFunctionLibrary>())
+		{
+			RefreshFunctionReferences(LibraryNode, bSetupUndoRedo, false);
+			PropagateNotificationToFunctionReferences(LibraryNode, ERigVMGraphNotifType::PinCategoryChanged, InPin, false);
+		}
+	}
 
 	if (!bSuspendNotifications)
 	{
@@ -8798,7 +9062,7 @@ bool URigVMController::SetPinIndexInCategory(URigVMPin* InPin, int32 InIndexInCa
 	return true;
 }
 
-bool URigVMController::SetPinLayout(const FName& InNodeName, FRigVMNodeLayout InLayout, bool bSetupUndoRedo, bool bPrintPythonCommand)
+bool URigVMController::SetNodeLayout(const FName& InNodeName, FRigVMNodeLayout InLayout, bool bSetupUndoRedo, bool bPrintPythonCommand)
 {
 	if (!IsValidGraph())
 	{
@@ -8820,10 +9084,10 @@ bool URigVMController::SetPinLayout(const FName& InNodeName, FRigVMNodeLayout In
 		return false;
 	}
 
-	return SetPinLayout(Node, InLayout, bSetupUndoRedo, bPrintPythonCommand);
+	return SetNodeLayout(Node, InLayout, bSetupUndoRedo, bPrintPythonCommand);
 }
 
-bool URigVMController::SetPinLayout(const URigVMNode* InNode, FRigVMNodeLayout InLayout, bool bSetupUndoRedo, bool bPrintPythonCommand)
+bool URigVMController::SetNodeLayout(const URigVMNode* InNode, FRigVMNodeLayout InLayout, bool bSetupUndoRedo, bool bPrintPythonCommand)
 {
 	if(!IsValidNodeForGraph(InNode))
 	{
@@ -8837,13 +9101,19 @@ bool URigVMController::SetPinLayout(const URigVMNode* InNode, FRigVMNodeLayout I
 	}
 
 	// first remove any pin layout from the node
-	if(ClearPinLayout(InNode->GetFName(), bSetupUndoRedo, bPrintPythonCommand))
 	{
-		NumChanges++;
+		const TGuardValue<bool> GuardRefreshFunctions(bSuspendRefreshingFunctionReferences, true); 
+
+		if(ClearNodeLayout(InNode->GetFName(), bSetupUndoRedo, bPrintPythonCommand))
+		{
+			NumChanges++;
+		}
 	}
 
 	for(const TPair<FString,FString>& DisplayNamePair : InLayout.DisplayNames)
 	{
+		const TGuardValue<bool> GuardRefreshFunctions(bSuspendRefreshingFunctionReferences, true); 
+
 		if(const URigVMPin* Pin = InNode->FindPin(DisplayNamePair.Key))
 		{
 			if(SetPinDisplayName(Pin->GetPinPath(), DisplayNamePair.Value, bSetupUndoRedo, bPrintPythonCommand))
@@ -8860,6 +9130,8 @@ bool URigVMController::SetPinLayout(const URigVMNode* InNode, FRigVMNodeLayout I
 
 	for(const FRigVMPinCategory& Category : InLayout.Categories)
 	{
+		const TGuardValue<bool> GuardRefreshFunctions(bSuspendRefreshingFunctionReferences, true); 
+
 		for(const FString& SegmentPath : Category.Elements)
 		{
 			if(const URigVMPin* Pin = InNode->FindPin(SegmentPath))
@@ -8879,6 +9151,8 @@ bool URigVMController::SetPinLayout(const URigVMNode* InNode, FRigVMNodeLayout I
 
 	for(const TPair<FString,int32>& PinIndexPair : InLayout.PinIndexInCategory)
 	{
+		const TGuardValue<bool> GuardRefreshFunctions(bSuspendRefreshingFunctionReferences, true); 
+
 		if(const URigVMPin* Pin = InNode->FindPin(PinIndexPair.Key))
 		{
 			if(SetPinIndexInCategory(Pin->GetPinPath(), PinIndexPair.Value, bSetupUndoRedo, bPrintPythonCommand))
@@ -8893,6 +9167,14 @@ bool URigVMController::SetPinLayout(const URigVMNode* InNode, FRigVMNodeLayout I
 		}
 	}
 
+	if(const URigVMLibraryNode* LibraryNode = Cast<URigVMLibraryNode>(InNode))
+	{
+		if (LibraryNode->GetGraph()->IsA<URigVMFunctionLibrary>())
+		{
+			RefreshFunctionReferences(LibraryNode, bSetupUndoRedo, false);
+		}
+	}
+	
 	if(bSetupUndoRedo)
 	{
 		if(NumChanges > 0)
@@ -8908,7 +9190,7 @@ bool URigVMController::SetPinLayout(const URigVMNode* InNode, FRigVMNodeLayout I
 	return NumChanges > 0;
 }
 
-bool URigVMController::ClearPinLayout(const FName& InNodeName, bool bSetupUndoRedo, bool bPrintPythonCommand)
+bool URigVMController::ClearNodeLayout(const FName& InNodeName, bool bSetupUndoRedo, bool bPrintPythonCommand)
 {
 	if (!IsValidGraph())
 	{
@@ -8930,10 +9212,10 @@ bool URigVMController::ClearPinLayout(const FName& InNodeName, bool bSetupUndoRe
 		return false;
 	}
 
-	return ClearPinLayout(Node, bSetupUndoRedo, bPrintPythonCommand);
+	return ClearNodeLayout(Node, bSetupUndoRedo, bPrintPythonCommand);
 }
 
-bool URigVMController::ClearPinLayout(const URigVMNode* InNode, bool bSetupUndoRedo, bool bPrintPythonCommand)
+bool URigVMController::ClearNodeLayout(const URigVMNode* InNode, bool bSetupUndoRedo, bool bPrintPythonCommand)
 {
 	if(!IsValidNodeForGraph(InNode))
 	{
@@ -8955,6 +9237,8 @@ bool URigVMController::ClearPinLayout(const URigVMNode* InNode, bool bSetupUndoR
 	// remove all pin categories
 	for(const FString& PinCategory : PinCategories)
 	{
+		const TGuardValue<bool> GuardRefreshFunctions(bSuspendRefreshingFunctionReferences, true); 
+
 		if(RemovePinCategory(InNode->GetFName(), PinCategory, bSetupUndoRedo, bPrintPythonCommand))
 		{
 			NumChanges++;
@@ -8965,6 +9249,8 @@ bool URigVMController::ClearPinLayout(const URigVMNode* InNode, bool bSetupUndoR
 	const TArray<URigVMPin*> AllPins = InNode->GetAllPinsRecursively();
 	for(const URigVMPin* Pin : AllPins)
 	{
+		const TGuardValue<bool> GuardRefreshFunctions(bSuspendRefreshingFunctionReferences, true); 
+
 		const FName DefaultDisplayName = InNode->GetDisplayNameForPin(Pin->GetSegmentPath(true));
 		if(!Pin->DisplayName.IsEqual(DefaultDisplayName, ENameCase::CaseSensitive))
 		{
@@ -8975,6 +9261,14 @@ bool URigVMController::ClearPinLayout(const URigVMNode* InNode, bool bSetupUndoR
 		}
 	}
 
+	if(const URigVMLibraryNode* LibraryNode = Cast<URigVMLibraryNode>(InNode))
+	{
+		if (LibraryNode->GetGraph()->IsA<URigVMFunctionLibrary>())
+		{
+			RefreshFunctionReferences(LibraryNode, bSetupUndoRedo, false);
+		}
+	}
+	
 	if(bSetupUndoRedo)
 	{
 		if(NumChanges > 0)
@@ -9023,7 +9317,7 @@ bool URigVMController::SetPinCategories(const URigVMNode* InNode, const TArray<F
 	}
 
 	// for entry or return nodes we relay to the outer pins
-	if(const URigVMFunctionInterfaceNode* InterfaceNode = Cast<URigVMFunctionInterfaceNode>(InNode))
+	if(InNode->IsA<URigVMFunctionInterfaceNode>())
 	{
 		if(const URigVMCollapseNode* OuterNode = Cast<URigVMCollapseNode>(InNode->GetGraph()->GetOuter()))
 		{
@@ -9068,6 +9362,15 @@ bool URigVMController::SetPinCategories(const URigVMNode* InNode, const TArray<F
 	const_cast<URigVMNode*>(InNode)->PinCategories = InCategories;
 	Notify(ERigVMGraphNotifType::PinCategoriesChanged, const_cast<URigVMNode*>(InNode));
 
+	if(const URigVMLibraryNode* LibraryNode = Cast<URigVMLibraryNode>(InNode))
+	{
+		if (LibraryNode->GetGraph()->IsA<URigVMFunctionLibrary>())
+		{
+			RefreshFunctionReferences(LibraryNode, bSetupUndoRedo, false);
+			PropagateNotificationToFunctionReferences(LibraryNode, ERigVMGraphNotifType::PinCategoriesChanged, const_cast<URigVMNode*>(InNode), false);
+		}
+	}
+
 	if (!bSuspendNotifications)
 	{
 		const URigVMGraph* Graph = GetGraph();
@@ -12939,6 +13242,12 @@ URigVMFunctionReferenceNode* URigVMController::AddFunctionReferenceNodeFromDescr
 		}
 	}
 
+	// auto expand categories as needed
+	for(const FRigVMPinCategory& Category : FunctionRefNode->ReferencedFunctionHeader.Layout.Categories)
+	{
+		(void)SetPinCategoryExpansion(FunctionRefNode, Category.Path, Category.bExpandedByDefault, bSetupUndoRedo);
+	}
+
 	if (bSetupUndoRedo)
 	{
 		FRigVMImportFromTextAction Action(this, FunctionRefNode);
@@ -20264,8 +20573,13 @@ const FRigVMByteCode* URigVMController::GetCurrentByteCode() const
 	return nullptr;
 }
 
-void URigVMController::RefreshFunctionReferences(URigVMLibraryNode* InFunctionDefinition, bool bSetupUndoRedo, bool bLoadIfNecessary)
+void URigVMController::RefreshFunctionReferences(const URigVMLibraryNode* InFunctionDefinition, bool bSetupUndoRedo, bool bLoadIfNecessary)
 {
+	if(bSuspendRefreshingFunctionReferences)
+	{
+		return;
+	}
+	
 	check(InFunctionDefinition);
 
 	if (const URigVMFunctionLibrary* FunctionLibrary = Cast<URigVMFunctionLibrary>(InFunctionDefinition->GetGraph()))
@@ -20289,6 +20603,39 @@ void URigVMController::RefreshFunctionReferences(URigVMLibraryNode* InFunctionDe
 	}
 }
 
+void URigVMController::PropagateNotificationToFunctionReferences(const URigVMLibraryNode* InFunctionDefinition, ERigVMGraphNotifType InNotifType, UObject* InSubject, bool bLoadIfNecessary)
+{
+	check(InFunctionDefinition);
+
+	if (const URigVMFunctionLibrary* FunctionLibrary = Cast<URigVMFunctionLibrary>(InFunctionDefinition->GetGraph()))
+	{
+		FunctionLibrary->ForEachReference(InFunctionDefinition->GetFName(), [this, InNotifType, InSubject](URigVMFunctionReferenceNode* ReferenceNode)
+		{
+			if(const URigVMController* ReferenceController = GetControllerForGraph(ReferenceNode->GetGraph()))
+			{
+				UObject* TargetSubject = nullptr;
+				if(InSubject)
+				{
+					if(const URigVMNode* NodeSubject = Cast<URigVMNode>(InSubject))
+					{
+						TargetSubject = ReferenceNode;
+					}
+					if(const URigVMPin* PinSubject = Cast<URigVMPin>(InSubject))
+					{
+						const FString SegmentPath = PinSubject->GetSegmentPath(true);
+						TargetSubject = ReferenceNode->FindPin(SegmentPath);
+					}
+				}
+				if(InSubject != nullptr && TargetSubject == nullptr)
+				{
+					return;
+				}
+				ReferenceController->Notify(InNotifType, TargetSubject);
+			}
+		}, bLoadIfNecessary);
+	}
+}
+
 URigVMController::FLinkedPath::FLinkedPath(URigVMLink* InLink)
 	: GraphPtr(InLink->GetGraph())
 	, SourcePinPath(InLink->GetSourcePinPath())
diff --git a/Engine/Plugins/Runtime/RigVM/Source/RigVMDeveloper/Private/RigVMModel/RigVMNode.cpp b/Engine/Plugins/Runtime/RigVM/Source/RigVMDeveloper/Private/RigVMModel/RigVMNode.cpp
--- a/Engine/Plugins/Runtime/RigVM/Source/RigVMDeveloper/Private/RigVMModel/RigVMNode.cpp
+++ b/Engine/Plugins/Runtime/RigVM/Source/RigVMDeveloper/Private/RigVMModel/RigVMNode.cpp
@@ -269,11 +269,21 @@ TArray<URigVMPin*> URigVMNode::GetPinsForCategory(FString InCategory) const
 			PinsInCategory.Add(Pin);
 		}
 	}
+
+	Algo::SortBy(PinsInCategory, [](const URigVMPin* Pin) -> int32
+	{
+		return Pin->GetIndexInCategory();
+	});
+	
 	return PinsInCategory;
 }
 
 bool URigVMNode::IsPinCategoryExpanded(FString InCategory) const
 {
+	if(InCategory.Equals(FRigVMPinCategory::GetDefaultCategoryName(), ESearchCase::IgnoreCase))
+	{
+		return true;
+	}
 	if(const bool* ExpansionState = PinCategoryExpansion.Find(InCategory))
 	{
 		return *ExpansionState;
@@ -281,7 +291,7 @@ bool URigVMNode::IsPinCategoryExpanded(FString InCategory) const
 	return false;
 }
 
-FRigVMNodeLayout URigVMNode::GetPinLayout() const
+FRigVMNodeLayout URigVMNode::GetNodeLayout(bool bIncludeEmptyCategories) const
 {
 	FRigVMNodeLayout Layout;
 	
@@ -327,13 +337,24 @@ FRigVMNodeLayout URigVMNode::GetPinLayout() const
 				}
 			}
 
-			Algo::SortBy(Category->Elements, [PinPathToIndex](const FString& PinPath) -> int32
+			Algo::SortBy(CategoryCopy.Elements, [PinPathToIndex](const FString& PinPath) -> int32
 			{
  				return PinPathToIndex.FindChecked(PinPath);
 			});
 			
 			Layout.Categories.Add(CategoryCopy);
 		}
+		else if(bIncludeEmptyCategories)
+		{
+			FRigVMPinCategory EmptyCategory;
+			EmptyCategory.Path = PinCategory;
+			Layout.Categories.Add(EmptyCategory);
+		}
+	}
+
+	for(FRigVMPinCategory& Category : Layout.Categories)
+	{
+		Category.bExpandedByDefault = IsPinCategoryExpanded(Category.Path);
 	}
 
 	// fill in all user provided display names and pin category indices
diff --git a/Engine/Plugins/Runtime/RigVM/Source/RigVMDeveloper/Public/RigVMModel/Nodes/RigVMFunctionInterfaceNode.h b/Engine/Plugins/Runtime/RigVM/Source/RigVMDeveloper/Public/RigVMModel/Nodes/RigVMFunctionInterfaceNode.h
--- a/Engine/Plugins/Runtime/RigVM/Source/RigVMDeveloper/Public/RigVMModel/Nodes/RigVMFunctionInterfaceNode.h
+++ b/Engine/Plugins/Runtime/RigVM/Source/RigVMDeveloper/Public/RigVMModel/Nodes/RigVMFunctionInterfaceNode.h
@@ -26,10 +26,6 @@ public:
 	// URigVMNode interface
 	virtual  FText GetToolTipText() const override;
 	virtual FText GetToolTipTextForPin(const URigVMPin* InPin) const override;	
-	virtual FName GetDisplayNameForPin(const FString& InPinPath) const override;
-	virtual FString GetCategoryForPin(const FString& InPinPath) const override;
-	virtual int32 GetIndexInCategoryForPin(const FString& InPinPath) const override;
-	virtual TArray<FString> GetPinCategories() const override;
 
 protected:
 
diff --git a/Engine/Plugins/Runtime/RigVM/Source/RigVMDeveloper/Public/RigVMModel/Nodes/RigVMFunctionReferenceNode.h b/Engine/Plugins/Runtime/RigVM/Source/RigVMDeveloper/Public/RigVMModel/Nodes/RigVMFunctionReferenceNode.h
--- a/Engine/Plugins/Runtime/RigVM/Source/RigVMDeveloper/Public/RigVMModel/Nodes/RigVMFunctionReferenceNode.h
+++ b/Engine/Plugins/Runtime/RigVM/Source/RigVMDeveloper/Public/RigVMModel/Nodes/RigVMFunctionReferenceNode.h
@@ -29,6 +29,7 @@ public:
 	virtual int32 GetIndexInCategoryForPin(const FString& InPinPath) const override;
 	virtual FText GetToolTipTextForPin(const URigVMPin* InPin) const override;
 	virtual TArray<FString> GetPinCategories() const override;
+	virtual FRigVMNodeLayout GetNodeLayout(bool bIncludeEmptyCategories) const override;
 	// end URigVMNode interface
 
 	// URigVMLibraryNode interface
diff --git a/Engine/Plugins/Runtime/RigVM/Source/RigVMDeveloper/Public/RigVMModel/RigVMController.h b/Engine/Plugins/Runtime/RigVM/Source/RigVMDeveloper/Public/RigVMModel/RigVMController.h
--- a/Engine/Plugins/Runtime/RigVM/Source/RigVMDeveloper/Public/RigVMModel/RigVMController.h
+++ b/Engine/Plugins/Runtime/RigVM/Source/RigVMDeveloper/Public/RigVMModel/RigVMController.h
@@ -810,6 +810,11 @@ public:
 	UFUNCTION(BlueprintCallable, Category = RigVMController)
 	bool SetPinDisplayName(const FString& InPinPath, const FString& InDisplayName, bool bSetupUndoRedo = true, bool bPrintPythonCommand = false);
 
+	// Adds a new pin category. The category is UI relevant only and used
+	// to order pins in the user interface of the node as well as on the details panel.
+	UFUNCTION(BlueprintCallable, Category = RigVMController)
+	bool AddEmptyPinCategory(const FName& InNodeName, const FString& InCategory, bool bSetupUndoRedo = true, bool bPrintPythonCommand = false);
+
 	// Sets the pin category. The category is UI relevant only and used
 	// to order pins in the user interface of the node as well as on the details panel.
 	UFUNCTION(BlueprintCallable, Category = RigVMController)
@@ -845,13 +850,13 @@ public:
 	UFUNCTION(BlueprintCallable, Category = RigVMController)
 	bool SetPinIndexInCategory(const FString& InPinPath, int32 InIndexInCategory, bool bSetupUndoRedo = true, bool bPrintPythonCommand = false);
 
-	// Applies a complete pin layout to a node
+	// Applies a complete node layout to a node
 	UFUNCTION(BlueprintCallable, Category = RigVMController)
-	bool SetPinLayout(const FName& InNodeName, FRigVMNodeLayout InLayout, bool bSetupUndoRedo = true, bool bPrintPythonCommand = false);
+	bool SetNodeLayout(const FName& InNodeName, FRigVMNodeLayout InLayout, bool bSetupUndoRedo = true, bool bPrintPythonCommand = false);
 
 	// Removes any layout information from a node
 	UFUNCTION(BlueprintCallable, Category = RigVMController)
-	bool ClearPinLayout(const FName& InNodeName, bool bSetupUndoRedo = true, bool bPrintPythonCommand = false);
+	bool ClearNodeLayout(const FName& InNodeName, bool bSetupUndoRedo = true, bool bPrintPythonCommand = false);
 
 	// Returns the default value of a pin given its pinpath.
 	UFUNCTION(BlueprintCallable, Category = RigVMController)
@@ -1281,14 +1286,15 @@ private:
 	void ExpandPinRecursively(URigVMPin* InPin, bool bSetupUndoRedo);
 	bool SetPinIsWatched(URigVMPin* InPin, bool bIsWatched, bool bSetupUndoRedo);
 	bool SetPinDisplayName(URigVMPin* InPin, const FString& InDisplayName, bool bSetupUndoRedo);
+	bool AddEmptyPinCategory(const URigVMNode* InNode, const FString& InPinCategory, bool bSetupUndoRedo);
 	bool SetPinCategory(URigVMPin* InPin, const FString& InCategory, bool bSetupUndoRedo);
 	bool RemovePinCategory(const URigVMNode* InNode, const FString& InPinCategory, bool bSetupUndoRedo);
 	bool RenamePinCategory(const URigVMNode* InNode, const FString& InOldPinCategory, const FString& InNewPinCategory, bool bSetupUndoRedo);
 	bool SetPinCategoryIndex(const URigVMNode* InNode, const FString& InPinCategory, int32 InNewIndex, bool bSetupUndoRedo);
 	bool SetPinCategoryExpansion(const URigVMNode* InNode, const FString& InPinCategory, bool bIsExpanded, bool bSetupUndoRedo);
 	bool SetPinIndexInCategory(URigVMPin* InPin, int32 InIndexInCategory, bool bSetupUndoRedo);
-	bool SetPinLayout(const URigVMNode* InNode, FRigVMNodeLayout InLayout, bool bSetupUndoRedo, bool bPrintPythonCommand);
-	bool ClearPinLayout(const URigVMNode* InNode, bool bSetupUndoRedo, bool bPrintPythonCommand);
+	bool SetNodeLayout(const URigVMNode* InNode, FRigVMNodeLayout InLayout, bool bSetupUndoRedo, bool bPrintPythonCommand);
+	bool ClearNodeLayout(const URigVMNode* InNode, bool bSetupUndoRedo, bool bPrintPythonCommand);
 	bool SetPinCategories(const FName& InNodeName, const TArray<FString>& InCategories, bool bSetupUndoRedo);
 	bool SetPinCategories(const URigVMNode* InNode, const TArray<FString>& InCategories, bool bSetupUndoRedo);
 	bool SetVariableName(URigVMVariableNode* InVariableNode, const FName& InVariableName, bool bSetupUndoRedo);
@@ -1360,7 +1366,8 @@ private:
 	FRigVMExternalVariable GetVariableByName(const FName& InExternalVariableName, const bool bIncludeInputArguments = false) const;
 	TArray<FRigVMExternalVariable> GetAllVariables(const bool bIncludeInputArguments = false) const;
 
-	void RefreshFunctionReferences(URigVMLibraryNode* InFunctionDefinition, bool bSetupUndoRedo, bool bLoadIfNecessary);
+	void RefreshFunctionReferences(const URigVMLibraryNode* InFunctionDefinition, bool bSetupUndoRedo, bool bLoadIfNecessary);
+	void PropagateNotificationToFunctionReferences(const URigVMLibraryNode* InFunctionDefinition, ERigVMGraphNotifType InNotifType, UObject* InSubject, bool bLoadIfNecessary);
 
 public:
 
@@ -1532,6 +1539,7 @@ private:
 	mutable FDelegateHandle ActionStackHandle;
 
 	bool bSuspendNotifications;
+	bool bSuspendRefreshingFunctionReferences;
 	bool bReportWarningsAndErrors;
 	bool bIgnoreRerouteCompactnessChanges;
 	ERigVMPinDirection UserLinkDirection;
diff --git a/Engine/Plugins/Runtime/RigVM/Source/RigVMDeveloper/Public/RigVMModel/RigVMNode.h b/Engine/Plugins/Runtime/RigVM/Source/RigVMDeveloper/Public/RigVMModel/RigVMNode.h
--- a/Engine/Plugins/Runtime/RigVM/Source/RigVMDeveloper/Public/RigVMModel/RigVMNode.h
+++ b/Engine/Plugins/Runtime/RigVM/Source/RigVMDeveloper/Public/RigVMModel/RigVMNode.h
@@ -102,7 +102,7 @@ public:
 
 	// Returns the pin UI layout for this node
 	UFUNCTION(BlueprintCallable, Category = RigVMNode)
-	FRigVMNodeLayout GetPinLayout() const;
+	virtual FRigVMNodeLayout GetNodeLayout(bool bIncludeEmptyCategories = false) const;
 
 	// Returns the default value for a given pin
 	FString GetOriginalPinDefaultValue(const URigVMPin* InPin) const;
diff --git a/Engine/Plugins/Runtime/RigVM/Source/RigVMEditor/Private/EdGraph/NodeSpawners/RigVMEdGraphFunctionRefNodeSpawner.cpp b/Engine/Plugins/Runtime/RigVM/Source/RigVMEditor/Private/EdGraph/NodeSpawners/RigVMEdGraphFunctionRefNodeSpawner.cpp
--- a/Engine/Plugins/Runtime/RigVM/Source/RigVMEditor/Private/EdGraph/NodeSpawners/RigVMEdGraphFunctionRefNodeSpawner.cpp
+++ b/Engine/Plugins/Runtime/RigVM/Source/RigVMEditor/Private/EdGraph/NodeSpawners/RigVMEdGraphFunctionRefNodeSpawner.cpp
@@ -15,6 +15,7 @@
 #include "BlueprintNodeTemplateCache.h"
 #include "RigVMBlueprintUtils.h"
 #include "ScopedTransaction.h"
+#include "Editor/RigVMEditorTools.h"
 
 #include UE_INLINE_GENERATED_CPP_BY_NAME(RigVMEdGraphFunctionRefNodeSpawner)
 
@@ -39,6 +40,8 @@ URigVMEdGraphFunctionRefNodeSpawner* URigVMEdGraphFunctionRefNodeSpawner::Create
 	NodeSpawner->ReferencedPublicFunctionHeader = InFunction->GetFunctionHeader();
 	NodeSpawner->NodeClass = URigVMEdGraphNode::StaticClass();
 	NodeSpawner->bIsLocalFunction = true;
+	const FAssetData Asset = UE::RigVM::Editor::Tools::FindAssetFromAnyPath(InFunction->GetPathName(), true);
+	NodeSpawner->AssetPath = Asset.ToSoftObjectPath(); 	
 
 	FBlueprintActionUiSpec& MenuSignature = NodeSpawner->DefaultMenuSignature;
 
@@ -91,6 +94,7 @@ URigVMEdGraphFunctionRefNodeSpawner* URigVMEdGraphFunctionRefNodeSpawner::Create
 	NodeSpawner->ReferencedPublicFunctionHeader = InPublicFunction;
 	NodeSpawner->NodeClass = URigVMEdGraphNode::StaticClass();
 	NodeSpawner->bIsLocalFunction = false;
+	NodeSpawner->AssetPath = InAssetData.ToSoftObjectPath(); 	
 
 	FBlueprintActionUiSpec& MenuSignature = NodeSpawner->DefaultMenuSignature;
 
@@ -258,7 +262,7 @@ UEdGraphNode* URigVMEdGraphFunctionRefNodeSpawner::Invoke(UEdGraph* ParentGraph,
 	// if we are trying to build the real function ref - but we haven't loaded the asset yet...
 	if(!FBlueprintNodeTemplateCache::IsTemplateOuter(ParentGraph))
 	{
-		if (!ReferencedPublicFunctionHeader.IsValid() && AssetPath.IsValid())
+		if (AssetPath.IsValid())
 		{
 			if (URigVMBlueprint* Blueprint = Cast<URigVMBlueprint>(AssetPath.TryLoad()))
 			{
diff --git a/Engine/Plugins/Runtime/RigVM/Source/RigVMEditor/Private/Editor/RigVMGraphDetailCustomization.cpp b/Engine/Plugins/Runtime/RigVM/Source/RigVMEditor/Private/Editor/RigVMGraphDetailCustomization.cpp
--- a/Engine/Plugins/Runtime/RigVM/Source/RigVMEditor/Private/Editor/RigVMGraphDetailCustomization.cpp
+++ b/Engine/Plugins/Runtime/RigVM/Source/RigVMEditor/Private/Editor/RigVMGraphDetailCustomization.cpp
@@ -29,10 +29,12 @@
 #include "RigVMStringUtils.h"
 #include "Widgets/SRigVMGraphPinEnumPicker.h"
 #include "Widgets/SRigVMVariantWidget.h"
+#include "Widgets/SRigVMNodeLayoutWidget.h"
 #include "ScopedTransaction.h"
 #include "Editor/RigVMEditorTools.h"
 #include "ContentBrowserModule.h"
 #include "IContentBrowserSingleton.h"
+#include "UObject/GarbageCollectionSchema.h"
 
 #define LOCTEXT_NAMESPACE "RigVMGraphDetailCustomization"
 
@@ -1022,6 +1024,40 @@ void FRigVMGraphDetailCustomization::CustomizeDetails(IDetailLayoutBuilder& Deta
 		];
 	}
 
+	SettingsCategory.AddCustomRow(FText::GetEmpty())
+	.OverrideResetToDefault(FResetToDefaultOverride::Hide())
+	.Visibility(TAttribute<EVisibility>::CreateLambda([this]()
+	{
+		return IsValidFunction() ? EVisibility::Visible : EVisibility::Collapsed;
+	}))
+	.NameContent()
+	[
+		SNew(STextBlock)
+		.Text(FText::FromString(TEXT("Layout")))
+		.Font(IDetailLayoutBuilder::GetDetailFont())
+	]
+	.ValueContent()
+	.HAlign(HAlign_Fill)
+	[
+		SNew(SRigVMNodeLayoutWidget)
+		.OnGetUncategorizedPins(this, &FRigVMGraphDetailCustomization::GetUncategorizedPins)
+		.OnGetCategories(this, &FRigVMGraphDetailCustomization::GetPinCategories)
+		.OnGetElementCategory(this, &FRigVMGraphDetailCustomization::GetPinCategory)
+		.OnGetElementIndexInCategory(this, &FRigVMGraphDetailCustomization::GetPinIndexInCategory)
+		.OnGetElementLabel(this, &FRigVMGraphDetailCustomization::GetPinLabel)
+		.OnGetElementColor(this, &FRigVMGraphDetailCustomization::GetPinColor)
+		.OnGetElementIcon(this, &FRigVMGraphDetailCustomization::GetPinIcon)
+		.OnCategoryAdded(this, &FRigVMGraphDetailCustomization::HandleCategoryAdded)
+		.OnCategoryRemoved(this, &FRigVMGraphDetailCustomization::HandleCategoryRemoved)
+		.OnCategoryRenamed(this, &FRigVMGraphDetailCustomization::HandleCategoryRenamed)
+		.OnElementCategoryChanged(this, &FRigVMGraphDetailCustomization::HandlePinCategoryChanged)
+		.OnElementLabelChanged(this, &FRigVMGraphDetailCustomization::HandlePinLabelChanged)
+		.OnElementIndexInCategoryChanged(this, &FRigVMGraphDetailCustomization::HandlePinIndexInCategoryChanged)
+		.OnValidateCategoryName(this, &FRigVMGraphDetailCustomization::HandleValidateCategoryName)
+		.OnValidateElementName(this, &FRigVMGraphDetailCustomization::HandleValidatePinDisplayName)
+		.OnGetStructuralHash(this, &FRigVMGraphDetailCustomization::GetNodeLayoutHash)
+	];
+
 	IDetailCategoryBuilder& DefaultsCategory = DetailLayout.EditCategory("NodeDefaults", LOCTEXT("FunctionDetailsNodeDefaults", "Node Defaults"));
 	TSharedRef<FRigVMFunctionArgumentDefaultNode> DefaultsArgumentNode = MakeShareable(new FRigVMFunctionArgumentDefaultNode(
 		Model,
@@ -1483,6 +1519,378 @@ void FRigVMGraphDetailCustomization::OnRemoveAssignedTag(const FName& InTagName)
 	}
 }
 
+URigVMLibraryNode* FRigVMGraphDetailCustomization::GetLibraryNode() const
+{
+	if (GraphPtr.IsValid() && RigVMBlueprintPtr.IsValid())
+	{
+		const URigVMBlueprint* Blueprint = RigVMBlueprintPtr.Get();
+		if (const URigVMGraph* Model = Blueprint->GetModel(GraphPtr.Get()))
+		{
+			if (const URigVMGraph* FunctionLibrary = Blueprint->GetLocalFunctionLibrary())
+			{
+				if (URigVMLibraryNode* LibraryNode = Cast<URigVMLibraryNode>(Model->GetOuter()))
+				{
+					if(LibraryNode->GetGraph() == FunctionLibrary)
+					{
+						return LibraryNode;
+					}
+				}
+			}
+		}
+	}
+	return nullptr;
+}
+
+URigVMNode* FRigVMGraphDetailCustomization::GetNodeForLayout() const
+{
+	return GetLibraryNode();
+}
+
+const FRigVMNodeLayout* FRigVMGraphDetailCustomization::GetNodeLayout() const
+{
+	if(const URigVMNode* Node = GetNodeForLayout())
+	{
+		CachedNodeLayout = Node->GetNodeLayout(true);
+		return &CachedNodeLayout.GetValue();
+	}
+	return nullptr;
+}
+
+TArray<FString> FRigVMGraphDetailCustomization::GetUncategorizedPins() const
+{
+	if(const URigVMNode* Node = GetNodeForLayout())
+	{
+		TArray<FString> PinPaths;
+		const TArray<URigVMPin*> AllPins = Node->GetAllPinsRecursively();
+		PinPaths.Reserve(AllPins.Num());
+		for(const URigVMPin* Pin : AllPins)
+		{
+			if(Pin->IsExecuteContext())
+			{
+				continue;
+			}
+			if(Pin->GetDirection() != ERigVMPinDirection::Input &&
+				Pin->GetDirection() != ERigVMPinDirection::Visible)
+			{
+				continue;
+			}
+			if(!Pin->GetCategory().IsEmpty())
+			{
+				continue;
+			}
+			PinPaths.Add(Pin->GetSegmentPath(true));
+		}
+		return PinPaths;
+	}
+	return TArray<FString>();
+}
+
+TArray<FRigVMPinCategory> FRigVMGraphDetailCustomization::GetPinCategories() const
+{
+	if (const FRigVMNodeLayout* NodeLayout = GetNodeLayout())
+	{
+		return NodeLayout->Categories;
+	}
+	return TArray<FRigVMPinCategory>();
+}
+
+FString FRigVMGraphDetailCustomization::GetPinCategory(FString InPinPath) const
+{
+	if(RigVMBlueprintPtr.IsValid())
+	{
+		if(const URigVMNode* Node = GetNodeForLayout())
+		{
+			if(const URigVMPin* Pin = Node->FindPin(InPinPath))
+			{
+				return Pin->GetCategory();
+			}
+		}
+	}
+	return FString();
+}
+
+int32 FRigVMGraphDetailCustomization::GetPinIndexInCategory(FString InPinPath) const
+{
+	if(RigVMBlueprintPtr.IsValid())
+	{
+		if(const URigVMNode* Node = GetNodeForLayout())
+		{
+			if(const URigVMPin* Pin = Node->FindPin(InPinPath))
+			{
+				return Pin->GetIndexInCategory();
+			}
+		}
+	}
+	return INDEX_NONE;
+}
+
+FString FRigVMGraphDetailCustomization::GetPinLabel(FString InPinPath) const
+{
+	if (const FRigVMNodeLayout* NodeLayout = GetNodeLayout())
+	{
+		if(const FString* DisplayName = NodeLayout->FindDisplayName(InPinPath))
+		{
+			return *DisplayName;
+		}
+	}
+	return FString();
+}
+
+FLinearColor FRigVMGraphDetailCustomization::GetPinColor(FString InPinPath) const
+{
+	if(RigVMBlueprintPtr.IsValid())
+	{
+		if(const URigVMNode* Node = GetNodeForLayout())
+		{
+			if(const URigVMPin* Pin = Node->FindPin(InPinPath))
+			{
+				if(const URigVMEdGraphSchema* Schema = Cast<URigVMEdGraphSchema>(RigVMBlueprintPtr->GetRigVMEdGraphSchemaClass()->GetDefaultObject()))
+				{
+					const FEdGraphPinType PinType = RigVMTypeUtils::PinTypeFromCPPType(*Pin->GetCPPType(), Pin->GetCPPTypeObject());
+					return Schema->GetPinTypeColor(PinType);
+				}
+			}
+		}
+	}
+	return FLinearColor::White;
+}
+
+const FSlateBrush* FRigVMGraphDetailCustomization::GetPinIcon(FString InPinPath) const
+{
+	if(RigVMBlueprintPtr.IsValid())
+	{
+		if(const URigVMNode* Node = GetNodeForLayout())
+		{
+			if(const URigVMPin* Pin = Node->FindPin(InPinPath))
+			{
+				const FEdGraphPinType PinType = RigVMTypeUtils::PinTypeFromCPPType(*Pin->GetCPPType(), Pin->GetCPPTypeObject());
+				return FBlueprintEditorUtils::GetIconFromPin(PinType, /* bIsLarge = */ false);
+			}
+		}
+	}
+	return nullptr;
+}
+
+void FRigVMGraphDetailCustomization::HandleCategoryAdded(FString InCategory)
+{
+	if(RigVMBlueprintPtr.IsValid())
+	{
+		if(const URigVMNode* Node = GetNodeForLayout())
+		{
+			if(URigVMController* Controller = RigVMBlueprintPtr->GetController(Node->GetGraph()))
+			{
+				Controller->AddEmptyPinCategory(Node->GetFName(), InCategory);
+				CachedNodeLayout.Reset();
+			}
+		}
+	}
+}
+
+void FRigVMGraphDetailCustomization::HandleCategoryRemoved(FString InCategory)
+{
+	if(RigVMBlueprintPtr.IsValid())
+	{
+		if(const URigVMNode* Node = GetNodeForLayout())
+		{
+			if(URigVMController* Controller = RigVMBlueprintPtr->GetController(Node->GetGraph()))
+			{
+				Controller->RemovePinCategory(Node->GetFName(), InCategory);
+				CachedNodeLayout.Reset();
+			}
+		}
+	}
+}
+
+void FRigVMGraphDetailCustomization::HandleCategoryRenamed(FString InOldCategory, FString InNewCategory)
+{
+	if(RigVMBlueprintPtr.IsValid())
+	{
+		if(const URigVMNode* Node = GetNodeForLayout())
+		{
+			if(URigVMController* Controller = RigVMBlueprintPtr->GetController(Node->GetGraph()))
+			{
+				Controller->RenamePinCategory(Node->GetFName(), InOldCategory, InNewCategory);
+				CachedNodeLayout.Reset();
+			}
+		}
+	}
+}
+
+void FRigVMGraphDetailCustomization::HandlePinCategoryChanged(FString InPinPath, FString InCategory)
+{
+	if(RigVMBlueprintPtr.IsValid())
+	{
+		if (const URigVMLibraryNode* LibraryNode = GetLibraryNode())
+		{
+			if(const URigVMPin* Pin = LibraryNode->FindPin(InPinPath))
+			{
+				if(URigVMController* Controller = RigVMBlueprintPtr->GetController(LibraryNode->GetGraph()))
+				{
+					Controller->SetPinCategory(Pin->GetPinPath(), InCategory);
+					CachedNodeLayout.Reset();
+				}
+			}
+		}
+	}
+}
+
+void FRigVMGraphDetailCustomization::HandlePinLabelChanged(FString InPinPath, FString InNewLabel)
+{
+	if(RigVMBlueprintPtr.IsValid())
+	{
+		if (const URigVMLibraryNode* LibraryNode = GetLibraryNode())
+		{
+			if(const URigVMPin* Pin = LibraryNode->FindPin(InPinPath))
+			{
+				if(URigVMController* Controller = RigVMBlueprintPtr->GetController(LibraryNode->GetGraph()))
+				{
+					Controller->SetPinDisplayName(Pin->GetPinPath(), InNewLabel);
+					CachedNodeLayout.Reset();
+				}
+			}
+		}
+	}
+}
+
+void FRigVMGraphDetailCustomization::HandlePinIndexInCategoryChanged(FString InPinPath, int32 InIndexInCategory)
+{
+	if(RigVMBlueprintPtr.IsValid())
+	{
+		if (const URigVMLibraryNode* LibraryNode = GetLibraryNode())
+		{
+			if(const URigVMPin* Pin = LibraryNode->FindPin(InPinPath))
+			{
+				if(URigVMController* Controller = RigVMBlueprintPtr->GetController(LibraryNode->GetGraph()))
+				{
+					Controller->SetPinIndexInCategory(Pin->GetPinPath(), InIndexInCategory);
+					CachedNodeLayout.Reset();
+				}
+			}
+		}
+	}
+}
+
+bool FRigVMGraphDetailCustomization::ValidateName(FString InNewName, FText& OutErrorMessage)
+{
+	if(InNewName.IsEmpty())
+	{
+		OutErrorMessage = LOCTEXT("EmptyNamesAreNotAllowed", "Empty names are not allowed.");
+		return false;
+	}
+
+	if(FChar::IsDigit(InNewName[0]))
+	{
+		OutErrorMessage = LOCTEXT("NamesCannotStartWithADigit", "Names cannot start with a digit.");
+		return false;
+	}
+
+	for (int32 i = 0; i < InNewName.Len(); ++i)
+	{
+		TCHAR& C = InNewName[i];
+
+		const bool bGoodChar = FChar::IsAlpha(C) ||					 // Any letter
+			(C == '_') || (C == '-') || (C == ' ') ||				 // _  - space anytime
+			FChar::IsDigit(C);										 // 0-9 anytime
+
+		if (!bGoodChar)
+		{
+			const FText Character = FText::FromString(InNewName.Mid(i, 1));
+			OutErrorMessage = FText::Format(LOCTEXT("CharacterNotAllowedFormat", "'{0}' not allowed."), Character);
+			return false;
+		}
+	}
+
+	if (InNewName.Len() > 100)
+	{
+		OutErrorMessage = LOCTEXT("NameIsTooLong", "Name is too long.");
+		return false;
+	}
+
+	return true;
+}
+
+bool FRigVMGraphDetailCustomization::HandleValidateCategoryName(FString InCategoryPath, FString InNewName, FText& OutErrorMessage)
+{
+	if(!ValidateName(InNewName, OutErrorMessage))
+	{
+		return false;
+	}
+	if(const URigVMNode* Node = GetNodeForLayout())
+	{
+		const FString ParentCategory = Node->GetParentPinCategory(InCategoryPath);
+		if(!ParentCategory.IsEmpty())
+		{
+			const TArray<FString> SiblingCategories = Node->GetSubPinCategories(ParentCategory);
+			const FString NewNameSuffix = TEXT("|") + InNewName;
+			if(SiblingCategories.ContainsByPredicate([InNewName, NewNameSuffix](const FString& Category)
+			{
+				return Category.Equals(InNewName, ESearchCase::IgnoreCase) || Category.EndsWith(NewNameSuffix, ESearchCase::IgnoreCase); 
+			}))
+			{
+				OutErrorMessage = LOCTEXT("NameIsAlreadyUsed", "Duplicate name.");
+				return false;
+			}
+		} 
+	}
+	return true;
+}
+
+bool FRigVMGraphDetailCustomization::HandleValidatePinDisplayName(FString InPinPath, FString InNewName, FText& OutErrorMessage)
+{
+	if(!ValidateName(InNewName, OutErrorMessage))
+	{
+		return false;
+	}
+	if(const URigVMNode* Node = GetNodeForLayout())
+	{
+		if(const URigVMPin* Pin = Node->FindPin(InPinPath))
+		{
+			const FString Category = Pin->GetCategory();
+			if(!Category.IsEmpty())
+			{
+				const TArray<URigVMPin*> PinsInCategory = Node->GetPinsForCategory(Category);
+				if(PinsInCategory.ContainsByPredicate([InNewName](const URigVMPin* PinInCategory)
+				{
+					return PinInCategory->GetDisplayName().ToString().Equals(InNewName, ESearchCase::IgnoreCase);
+				}))
+				{
+					OutErrorMessage = LOCTEXT("NameIsAlreadyUsedInCategory", "Duplicate name (category).");
+					return false;
+				}
+			}
+
+			if(const URigVMPin* ParentPin = Pin->GetParentPin())
+			{
+				const TArray<URigVMPin*> SubPins = ParentPin->GetSubPins();
+				if(SubPins.ContainsByPredicate([InNewName](const URigVMPin* SubPin)
+				{
+					return SubPin->GetDisplayName().ToString().Equals(InNewName, ESearchCase::IgnoreCase);
+				}))
+				{
+					OutErrorMessage = LOCTEXT("NameIsAlreadyUsedWithinPin", "Duplicate name (parent pin).");
+					return false;
+				}
+			}
+		}
+	}
+	return true;
+}
+
+uint32 FRigVMGraphDetailCustomization::GetNodeLayoutHash() const
+{
+	uint32 Hash = 0;
+	if(const FRigVMNodeLayout* Layout = GetNodeLayout())
+	{
+		Hash = HashCombine(Hash, GetTypeHash(*Layout));
+	}
+	const TArray<FString> UncategorizedPins = GetUncategorizedPins();
+	for(const FString& UncategorizedPin : UncategorizedPins)
+	{
+		Hash = HashCombine(Hash, GetTypeHash(UncategorizedPin));
+	}
+	return Hash;
+}
+
 FRigVMWrappedNodeDetailCustomization::FRigVMWrappedNodeDetailCustomization()
 : BlueprintBeingCustomized(nullptr)
 {
@@ -1614,7 +2022,7 @@ void FRigVMWrappedNodeDetailCustomization::CustomizeDetails(IDetailLayoutBuilder
 					TemplateNotation = TemplateNode->GetNotation();
 				}
 
-				NodeLayout = Node->GetPinLayout();
+				NodeLayout = Node->GetNodeLayout();
 
 				for(const FString& TraitName : Node->GetTraitNames())
 				{
@@ -1638,7 +2046,7 @@ void FRigVMWrappedNodeDetailCustomization::CustomizeDetails(IDetailLayoutBuilder
 					}
 				}
 
-				if(NodeLayout != Node->GetPinLayout())
+				if(NodeLayout != Node->GetNodeLayout())
 				{
 					bInspectingOnlyOneNodeType = false;
 					break;
@@ -1692,7 +2100,7 @@ void FRigVMWrappedNodeDetailCustomization::CustomizeDetails(IDetailLayoutBuilder
 
 	if(NodeWithCategories)
 	{
-		const FRigVMNodeLayout NodeLayout = NodeWithCategories->GetPinLayout();
+		const FRigVMNodeLayout NodeLayout = NodeWithCategories->GetNodeLayout();
 		for(const FRigVMPinCategory& Category : NodeLayout.Categories)
 		{
 			for(const FString& PinPath : Category.Elements)
@@ -1735,7 +2143,7 @@ void FRigVMWrappedNodeDetailCustomization::CustomizeDetails(IDetailLayoutBuilder
 	FRigVMNodeLayout NodeLayout;
 	if(NodeWithCategories)
 	{
-		NodeLayout = NodeWithCategories->GetPinLayout();
+		NodeLayout = NodeWithCategories->GetNodeLayout();
 	}
 	for (const TTuple<const FProperty*, FRigVMPropertyPath, FString>& Tuple : PropertiesToShow)
 	{
diff --git a/Engine/Plugins/Runtime/RigVM/Source/RigVMEditor/Private/Widgets/SRigVMGraphNode.cpp b/Engine/Plugins/Runtime/RigVM/Source/RigVMEditor/Private/Widgets/SRigVMGraphNode.cpp
--- a/Engine/Plugins/Runtime/RigVM/Source/RigVMEditor/Private/Widgets/SRigVMGraphNode.cpp
+++ b/Engine/Plugins/Runtime/RigVM/Source/RigVMEditor/Private/Widgets/SRigVMGraphNode.cpp
@@ -1605,7 +1605,7 @@ void SRigVMGraphNode::UpdatePinTreeView()
 	// orphaned pins are appended to the end of pin list on each side of the node
 	RootModelPins.Append(ModelNode->GetOrphanedPins());
 
-	TArray<TTuple<URigVMPin*, FString>> ModelPinsOrCategories;
+	TArray<TTuple<const URigVMPin*, FString>> ModelPinsOrCategories;
 
 	const bool bSupportSubPins = !RigGraphNode->DrawAsCompactNode();
 
@@ -1613,41 +1613,16 @@ void SRigVMGraphNode::UpdatePinTreeView()
 	// a) execute IOs, b) IO pins, c) input / visible pins, d) output pins
 	struct Local
 	{
-		static void VisitPinRecursively(URigVMPin* InPin, TArray<TTuple<URigVMPin*, FString>>& OutPinsOrCategories, bool bSupportSubPins, const TArray<ERigVMPinDirection>& DirectionsWithCategories)
+		static void VisitPinRecursively(const URigVMPin* InPin, TArray<TTuple<const URigVMPin*, FString>>& OutPinsOrCategories, bool bSupportSubPins)
 		{
-			const TTuple<URigVMPin*, FString> PinKey(InPin, FString());
+			const TTuple<const URigVMPin*, FString> PinKey(InPin, FString());
 			if(OutPinsOrCategories.Contains(PinKey))
 			{
 				return;
 			}
 			
-			bool bIncludePin = true;
-			const FString PinCategory = InPin->GetCategory();
-			if(!PinCategory.IsEmpty())
-			{
-				// add all categories from root to leaf
-				TArray<FString> ParentCategories = InPin->GetNode()->GetParentPinCategories(PinCategory, false, true);
-				Algo::Reverse(ParentCategories);
-				for(const FString& ParentCategory : ParentCategories)
-				{
-					OutPinsOrCategories.AddUnique({nullptr, ParentCategory});
-				}
-			}
-			// if the pin has no category - but there are pins with the same direction in a category - we need to skip it
-			else if(DirectionsWithCategories.Contains(InPin->GetDirection()))
-			{
-				// only add the pin if our parent is here - either the parent pin or the category
-				const TTuple<URigVMPin*, FString> ParentPinPair = {InPin->GetParentPin(), FString()}; 
-				const TTuple<URigVMPin*, FString> CategoryPair = {nullptr, PinCategory}; 
-				bIncludePin =
-					OutPinsOrCategories.Contains(ParentPinPair) ||
-					OutPinsOrCategories.Contains(CategoryPair);
-			}
+			OutPinsOrCategories.AddUnique(PinKey);
 			
-			if(bIncludePin)
-			{
-				OutPinsOrCategories.Add(PinKey);
-			}
 			if(!bSupportSubPins)
 			{
 				return;
@@ -1658,55 +1633,27 @@ void SRigVMGraphNode::UpdatePinTreeView()
 				return;
 			}
 
-			// visit all pins within this category
-			if(!PinCategory.IsEmpty())
-			{
-				const TArray<URigVMPin*> PinsInCategory = InPin->GetNode()->GetPinsForCategory(PinCategory);
-				for (URigVMPin* PinInCategory : PinsInCategory)
-				{
-					VisitPinRecursively(PinInCategory, OutPinsOrCategories, bSupportSubPins, DirectionsWithCategories);
-				}
-			}
-
 			if (InPin->GetCPPType() == TEXT("FRotator"))
 			{
 				const TArray<URigVMPin*>& SubPins = InPin->GetSubPins();
 				if (SubPins.Num() == 3)
 				{
-					VisitPinRecursively(SubPins[2], OutPinsOrCategories, bSupportSubPins, DirectionsWithCategories);
-					VisitPinRecursively(SubPins[0], OutPinsOrCategories, bSupportSubPins, DirectionsWithCategories);
-					VisitPinRecursively(SubPins[1], OutPinsOrCategories, bSupportSubPins, DirectionsWithCategories);
+					VisitPinRecursively(SubPins[2], OutPinsOrCategories, bSupportSubPins);
+					VisitPinRecursively(SubPins[0], OutPinsOrCategories, bSupportSubPins);
+					VisitPinRecursively(SubPins[1], OutPinsOrCategories, bSupportSubPins);
 				}	
 			}
 			else
 			{				
 				for (URigVMPin* SubPin : InPin->GetSubPins())
 				{
-					VisitPinRecursively(SubPin, OutPinsOrCategories, bSupportSubPins, DirectionsWithCategories);
+					VisitPinRecursively(SubPin, OutPinsOrCategories, bSupportSubPins);
 				}
 			}
 		}
 	};
 
-	// determine which directions use a category
-	TArray<ERigVMPinDirection> DirectionsWithCategories;
-	const TArray<URigVMPin*> AllPins = ModelNode->GetAllPinsRecursively();
-	for(const URigVMPin* Pin : AllPins)
-	{
-		if(!Pin->GetCategory().IsEmpty())
-		{
-			const ERigVMPinDirection Direction = Pin->GetDirection();
-			DirectionsWithCategories.AddUnique(Direction);
-			if(Direction == ERigVMPinDirection::Input)
-			{
-				DirectionsWithCategories.AddUnique(ERigVMPinDirection::Visible);
-			}
-			else if(Direction == ERigVMPinDirection::Visible)
-			{
-				DirectionsWithCategories.AddUnique(ERigVMPinDirection::Input);
-			}
-		}
-	}
+	const FRigVMNodeLayout NodeLayout = ModelNode->GetNodeLayout();
 	
 	for(int32 SortPhase = 0; SortPhase < 4; SortPhase++)
 	{
@@ -1718,7 +1665,7 @@ void SRigVMGraphNode::UpdatePinTreeView()
 				{
 					if(RootPin->IsExecuteContext() && RootPin->GetDirection() == ERigVMPinDirection::IO)
 					{
-						Local::VisitPinRecursively(RootPin, ModelPinsOrCategories, bSupportSubPins, DirectionsWithCategories);
+						Local::VisitPinRecursively(RootPin, ModelPinsOrCategories, bSupportSubPins);
 					}
 					break;
 				}
@@ -1726,7 +1673,7 @@ void SRigVMGraphNode::UpdatePinTreeView()
 				{
 					if(RootPin->GetDirection() == ERigVMPinDirection::Output)
 					{
-						Local::VisitPinRecursively(RootPin, ModelPinsOrCategories, bSupportSubPins, DirectionsWithCategories);
+						Local::VisitPinRecursively(RootPin, ModelPinsOrCategories, bSupportSubPins);
 					}
 					break;
 				}
@@ -1734,16 +1681,20 @@ void SRigVMGraphNode::UpdatePinTreeView()
 				{
 					if(!RootPin->IsExecuteContext() && RootPin->GetDirection() == ERigVMPinDirection::IO)
 					{
-						Local::VisitPinRecursively(RootPin, ModelPinsOrCategories, bSupportSubPins, DirectionsWithCategories);
+						Local::VisitPinRecursively(RootPin, ModelPinsOrCategories, bSupportSubPins);
 					}
 					break;
 				}
 				case 3: // input / visible pins
 				default:
 				{
-					if(RootPin->GetDirection() == ERigVMPinDirection::Input || RootPin->GetDirection() == ERigVMPinDirection::Visible)
+					// for now we only allow node layouts for input pins
+					if(!NodeLayout.IsValid())
 					{
-						Local::VisitPinRecursively(RootPin, ModelPinsOrCategories, bSupportSubPins, DirectionsWithCategories);
+						if(RootPin->GetDirection() == ERigVMPinDirection::Input || RootPin->GetDirection() == ERigVMPinDirection::Visible)
+						{
+							Local::VisitPinRecursively(RootPin, ModelPinsOrCategories, bSupportSubPins);
+						}
 					}
 					break;
 				}
@@ -1751,6 +1702,45 @@ void SRigVMGraphNode::UpdatePinTreeView()
 		}
 	}
 
+	// visit the categories / input / visible pins in the right order
+	if(NodeLayout.IsValid())
+	{
+		for(const FRigVMPinCategory& Category : NodeLayout.Categories)
+		{
+			if(!Category.IsDefaultCategory())
+			{
+				ModelPinsOrCategories.AddUnique({nullptr, Category.Path});
+			}
+			for(const FString& PinPath : Category.Elements)
+			{
+				if(const URigVMPin* Pin = ModelNode->FindPin(PinPath))
+				{
+					Local::VisitPinRecursively(Pin, ModelPinsOrCategories, bSupportSubPins);
+				}
+			}
+		}
+
+		// when adding pins based on the node layout
+		// we may be skipping over pins which have a connection but are not part of the layout
+		const TArray<URigVMLink*> Links = ModelNode->GetLinks();
+		for(const URigVMLink* Link : Links)
+		{
+			const URigVMPin* SourcePin = Link->GetSourcePin();
+			const URigVMPin* TargetPin = Link->GetTargetPin();
+			if(SourcePin == nullptr || TargetPin == nullptr)
+			{
+				continue;
+			}
+			const URigVMPin* PinToShow = SourcePin->GetNode() == ModelNode.Get() ? SourcePin : TargetPin;
+			const TTuple<const URigVMPin*, FString> Key = {PinToShow, FString()};
+			if(ModelPinsOrCategories.Contains(Key))
+			{
+				continue;
+			}
+			Local::VisitPinRecursively(PinToShow, ModelPinsOrCategories, bSupportSubPins);
+		}
+	}
+
 	const URigVMEdGraphSchema* RigSchema = Cast<URigVMEdGraphSchema>(RigGraphNode->GetSchema());
 
 	FRigVMDispatchContext DispatchContext;
@@ -1759,10 +1749,6 @@ void SRigVMGraphNode::UpdatePinTreeView()
 		DispatchContext = DispatchNode->GetDispatchContext();
 	}
 
-	// todo: we should implement traversers here which visit the pins in the right order.
-	// either by visiting the pins as we used to or using the categories. we need to figure out if for each group there's a category
-	// available and then use that instead. so in other words show all IO pins except if there's a category for an IO pin. etc.
-
 	TMap<FString, int32> IdentifierToInfoIndex;
 
 	auto AddModelPinOrCategoryPinInfo =
@@ -1837,27 +1823,41 @@ void SRigVMGraphNode::UpdatePinTreeView()
 
 		if(bIsModelPin)
 		{
+			bool bNeedsParentPin = true;
 			const FString PinCategory = InModelPin->GetCategory();
 			if(!PinCategory.IsEmpty())
 			{
-				if(const int32* ParentIndexPtr = IdentifierToInfoIndex.Find(PinCategory))
+				if(PinCategory.Equals(FRigVMPinCategory::GetDefaultCategoryName(), ESearchCase::IgnoreCase))
+				{
+					bNeedsParentPin = false;
+				}
+				else if(const int32* ParentIndexPtr = IdentifierToInfoIndex.Find(PinCategory))
 				{
 					PinInfo.ParentIndex = *ParentIndexPtr;
+					bNeedsParentPin = false;
 				}
 			}
 
 			// if the pin is not part of a category, add it under its default parent
-			if(PinInfo.ParentIndex == INDEX_NONE)
+			if(bNeedsParentPin)
 			{
 				if(URigVMPin* ParentPin = InModelPin->GetParentPin())
 				{
 					const int32* ParentIndexPtr = IdentifierToInfoIndex.Find(ParentPin->GetPinPath());
 					if(ParentIndexPtr == nullptr)
 					{
-						// parent pins have to exist for the node to display correctly
-						return;
+						if(InModelPin->GetLinks().IsEmpty())
+						{
+							// parent pins have to exist for the node to display correctly
+							// but allow pins with links to show up
+							return;
+						}
+					}
+					else
+					{
+						PinInfo.ParentIndex = *ParentIndexPtr;
+						bNeedsParentPin = false;
 					}
-					PinInfo.ParentIndex = *ParentIndexPtr;
 				}
 			}
 		}
@@ -2070,7 +2070,7 @@ void SRigVMGraphNode::UpdatePinTreeView()
 		}
 	};
 	
-	for(const TTuple<URigVMPin*,FString>& ModelPinOrCategory : ModelPinsOrCategories)
+	for(const TTuple<const URigVMPin*,FString>& ModelPinOrCategory : ModelPinsOrCategories)
 	{
 		AddModelPinOrCategoryPinInfo(ModelPinOrCategory.Get<0>(), ModelPinOrCategory.Get<1>());
 	}
diff --git a/Engine/Plugins/Runtime/RigVM/Source/RigVMEditor/Private/Widgets/SRigVMNodeLayoutWidget.cpp b/Engine/Plugins/Runtime/RigVM/Source/RigVMEditor/Private/Widgets/SRigVMNodeLayoutWidget.cpp
new file mode 100644
--- /dev/null
+++ b/Engine/Plugins/Runtime/RigVM/Source/RigVMEditor/Private/Widgets/SRigVMNodeLayoutWidget.cpp
@@ -0,0 +1,787 @@
+// Copyright Epic Games, Inc. All Rights Reserved.
+
+#include "Widgets/SRigVMNodeLayoutWidget.h"
+#include "Widgets/Layout/SBorder.h"
+#include "Widgets/SBoxPanel.h"
+#include "Widgets/Text/STextBlock.h"
+#include "Widgets/Input/SButton.h"
+#include "Widgets/Layout/SScrollBox.h"
+#include "Widgets/Layout/SSpacer.h"
+#include "Widgets/Images/SImage.h"
+#include "DetailLayoutBuilder.h"
+#include "RigVMStringUtils.h"
+#include "Framework/Application/SlateApplication.h"
+
+#define LOCTEXT_NAMESPACE "SRigVMNodeLayoutWidget"
+
+SRigVMNodeLayoutWidget::SRigVMNodeLayoutWidget()
+{
+}
+
+SRigVMNodeLayoutWidget::~SRigVMNodeLayoutWidget()
+{
+}
+
+void SRigVMNodeLayoutWidget::Construct(
+	const FArguments& InArgs)
+{
+	OnGetUncategorizedPins = InArgs._OnGetUncategorizedPins;
+	OnGetCategories = InArgs._OnGetCategories;
+	OnCategoryAdded = InArgs._OnCategoryAdded;
+	OnCategoryRemoved = InArgs._OnCategoryRemoved;
+	OnCategoryRenamed = InArgs._OnCategoryRenamed;
+	OnGetElementLabel = InArgs._OnGetElementLabel;
+	OnElementLabelChanged = InArgs._OnElementLabelChanged;
+	OnGetElementCategory = InArgs._OnGetElementCategory;
+	OnGetElementIndexInCategory = InArgs._OnGetElementIndexInCategory;
+	OnGetElementColor = InArgs._OnGetElementColor;
+	OnGetElementIcon = InArgs._OnGetElementIcon;
+	OnElementIndexInCategoryChanged = InArgs._OnElementIndexInCategoryChanged;
+	OnElementCategoryChanged = InArgs._OnElementCategoryChanged;
+	OnGetStructuralHash = InArgs._OnGetStructuralHash;
+	OnValidateCategoryName = InArgs._OnValidateCategoryName;
+	OnValidateElementName = InArgs._OnValidateElementName;
+
+	ChildSlot
+	[
+		SNew(SVerticalBox)
+		+ SVerticalBox::Slot()
+		.AutoHeight()
+		.HAlign(HAlign_Right)
+		.VAlign(VAlign_Center)
+		.Padding(0, 4, 4, 0)
+		[
+			SNew(SButton)
+			.ContentPadding(0.0f)
+			.ButtonStyle(FAppStyle::Get(), "NoBorder")
+			.OnClicked(this, &SRigVMNodeLayoutWidget::HandleAddCategory)
+			.Cursor(EMouseCursor::Default)
+			.Visibility_Lambda([this]()
+			{
+				return IsNodeLayoutEditable() ? EVisibility::Visible : EVisibility::Collapsed;
+			})
+			.ToolTipText(LOCTEXT("AddCategory", "Add Category"))
+			[
+				SNew(SHorizontalBox)
+				+SHorizontalBox::Slot()
+				.AutoWidth()
+				.VAlign(VAlign_Center)
+				[
+					SNew(SImage)
+					.Image(FAppStyle::GetBrush(TEXT("Icons.PlusCircle")))
+				]
+			]
+		]
+
+		+ SVerticalBox::Slot()
+		.AutoHeight()
+		.HAlign(HAlign_Fill)
+		.Padding(4, 4, 4, 4)
+		[
+			SNew(SScrollBox)
+			.Visibility_Lambda([this]()
+			{
+				if(!OnGetCategories.IsBound())
+				{
+					return EVisibility::Collapsed;
+				}
+				return OnGetCategories.Execute().IsEmpty() ? EVisibility::Collapsed : EVisibility::Visible;
+			})
+			+ SScrollBox::Slot()
+			.MaxSize(InArgs._MaxScrollBoxSize)
+			[
+				SAssignNew(TreeView, STreeView<TSharedPtr<FNodeLayoutRow>>)
+				.SelectionMode(ESelectionMode::Single)
+				.TreeItemsSource(&NodeLayoutRows)
+				.OnGenerateRow(this, &SRigVMNodeLayoutWidget::GenerateRow)
+				.OnGetChildren(this, &SRigVMNodeLayoutWidget::GetChildrenForRow)
+				.OnExpansionChanged(this, &SRigVMNodeLayoutWidget::OnItemExpansionChanged)
+				.OnSelectionChanged(this, &SRigVMNodeLayoutWidget::OnItemSelectionChanged)
+			]
+		]
+	];
+
+	Refresh();
+
+	if(OnGetStructuralHash.IsBound())
+	{
+		SetCanTick(true);
+	}
+}
+
+void SRigVMNodeLayoutWidget::Refresh()
+{
+	struct FItemState : FNodeLayoutRow::FState
+	{
+		static void RecordItemState(const TSharedPtr<FNodeLayoutRow>& InRow, TMap<FString, FItemState>& OutItemStates)
+		{
+			FItemState State;
+			*static_cast<FNodeLayoutRow::FState*>(&State) = InRow->State;
+			OutItemStates.FindOrAdd(InRow->Path, State);
+			RecordItemStates(InRow->ChildRows, OutItemStates);
+		}
+
+		static void RecordItemStates(const TArray<TSharedPtr<FNodeLayoutRow>>& InRows, TMap<FString, FItemState>& OutItemStates)
+		{
+			for(const TSharedPtr<FNodeLayoutRow>& Row : InRows)
+			{
+				RecordItemState(Row, OutItemStates);
+			}
+		}
+
+		static void ApplyItemState(const TSharedPtr<FNodeLayoutRow>& InOutRow, const TMap<FString, FItemState>& InItemStates)
+		{
+			if(const FItemState* State = InItemStates.Find(InOutRow->Path))
+			{
+				InOutRow->State = *State;
+			}
+			ApplyItemStates(InOutRow->ChildRows, InItemStates);
+		}
+
+		static void ApplyItemStates(const TArray<TSharedPtr<FNodeLayoutRow>>& InOutRows, const TMap<FString, FItemState>& InItemStates)
+		{
+			for(const TSharedPtr<FNodeLayoutRow>& Row : InOutRows)
+			{
+				ApplyItemState(Row, InItemStates);
+			}
+		}
+	};
+	
+	TMap<FString, FItemState> ItemStates;
+	FItemState::RecordItemStates(NodeLayoutRows, ItemStates);
+	NodeLayoutRows.Reset();
+
+	TMap<FString, TSharedPtr<FNodeLayoutRow>> PathToRow;
+
+	auto AddElement = [&PathToRow, this](const FString& InElementPath, bool bAddToList) -> TSharedPtr<FNodeLayoutRow>
+	{
+		TSharedPtr<FNodeLayoutRow> Row = MakeShareable(new FNodeLayoutRow);
+		Row->bIsCategory = false;
+		Row->bIsUncategorized = true;
+		Row->Path = Row->Label = InElementPath;
+
+		if(bAddToList)
+		{
+			FString ParentPath;
+			if(RigVMStringUtils::SplitPinPathAtEnd(Row->Path, ParentPath, Row->Label))
+			{
+				if(const TSharedPtr<FNodeLayoutRow>* ParentRow = PathToRow.Find(ParentPath))
+				{
+					(*ParentRow)->ChildRows.Add(Row);
+				}
+				else
+				{
+					NodeLayoutRows.Add(Row);
+				}
+			}
+			else
+			{
+				NodeLayoutRows.Add(Row);
+			}
+		}
+		PathToRow.Add(Row->Path, Row);
+
+		//  use the user defined label if need be
+		if(OnGetElementLabel.IsBound())
+		{
+			const FString Label = OnGetElementLabel.Execute(Row->Path);
+			if(!Label.IsEmpty())
+			{
+				Row->Label = Label;
+			}
+		}
+		if(OnGetElementColor.IsBound())
+		{
+			Row->Color = OnGetElementColor.Execute(Row->Path);
+		}
+		if(OnGetElementIcon.IsBound())
+		{
+			Row->Icon = OnGetElementIcon.Execute(Row->Path);
+		}
+
+		return Row;
+	};
+
+	if(OnGetCategories.IsBound())
+	{
+		const TArray<FRigVMPinCategory> Categories = OnGetCategories.Execute();
+		for(const FRigVMPinCategory& Category : Categories)
+		{
+			TSharedPtr<FNodeLayoutRow> Row = MakeShareable(new FNodeLayoutRow);
+			Row->bIsCategory = true;
+			Row->bIsUncategorized = true;
+			Row->Path = Category.Path;
+			Row->Label = Category.GetName();
+			// this may get overriden later when the states get reapplied
+			Row->State.bExpanded = true;
+			
+			FString ParentPath, LastBit;;
+			if(RigVMStringUtils::SplitNodePathAtEnd(Row->Path, ParentPath, LastBit))
+			{
+				if(const TSharedPtr<FNodeLayoutRow>* ParentRow = PathToRow.Find(ParentPath))
+				{
+					(*ParentRow)->ChildRows.Add(Row);
+				}
+				else
+				{
+					NodeLayoutRows.Add(Row);
+				}
+			}
+			else
+			{
+				NodeLayoutRows.Add(Row);
+			}
+			PathToRow.Add(Row->Path, Row);
+
+			for(const FString& Element : Category.Elements)
+			{
+				TSharedPtr<FNodeLayoutRow> ElementRow = AddElement(Element, false);
+				ElementRow->bIsUncategorized = false;
+				Row->ChildRows.Add(ElementRow);
+			}
+
+			if(!ItemStates.Contains(Row->Path))
+			{
+				FItemState State;
+				*static_cast<FNodeLayoutRow::FState*>(&State) = Row->State;
+				ItemStates.Add(Row->Path, State);
+			}
+		}
+	}
+
+	if(OnGetUncategorizedPins.IsBound())
+	{
+		const TArray<FString> PinPaths = OnGetUncategorizedPins.Execute();
+		for(const FString& PinPath : PinPaths)
+		{
+			(void)AddElement(PinPath, true);
+		}
+	}
+	
+	TreeView->RequestTreeRefresh();
+	
+	FItemState::ApplyItemStates(NodeLayoutRows, ItemStates);
+	TreeView->ClearSelection();
+
+	for(const TPair<FString, FItemState>& Pair : ItemStates)
+	{
+		if(const TSharedPtr<FNodeLayoutRow>* RowPtr = PathToRow.Find(Pair.Key))
+		{
+			const TSharedPtr<FNodeLayoutRow>& Row = *RowPtr;
+			if(Row)
+			{
+				TreeView->SetItemExpansion(Row, Pair.Value.bExpanded);
+				TreeView->SetItemSelection(Row, Pair.Value.bSelected);
+			}
+		}
+	}
+}
+
+void SRigVMNodeLayoutWidget::Tick(const FGeometry& AllottedGeometry, const double InCurrentTime, const float InDeltaTime)
+{
+	SBox::Tick(AllottedGeometry, InCurrentTime, InDeltaTime);
+
+	if(OnGetStructuralHash.IsBound())
+	{
+		const uint32 CurrentStructuralHash = OnGetStructuralHash.Execute();
+		if(!LastStructuralHash.IsSet() || LastStructuralHash.GetValue() != CurrentStructuralHash)
+		{
+			LastStructuralHash = CurrentStructuralHash;
+			Refresh();
+		}
+	}
+}
+
+FReply SRigVMNodeLayoutWidget::OnKeyDown(const FGeometry& MyGeometry, const FKeyEvent& InKeyEvent)
+{
+	if(InKeyEvent.GetKey() == EKeys::F2)
+	{
+		const TArray<TSharedPtr<FNodeLayoutRow>> SelectedRows = TreeView->GetSelectedItems();
+		if(!SelectedRows.IsEmpty())
+		{
+			SelectedRows[0]->RequestRename();
+			return FReply::Handled();
+		}
+	}
+	return SBox::OnKeyDown(MyGeometry, InKeyEvent);
+}
+
+SRigVMNodeLayoutWidget::SRigVMNowLayoutRow::~SRigVMNowLayoutRow()
+{
+}
+
+void SRigVMNodeLayoutWidget::SRigVMNowLayoutRow::Construct(const FArguments& InArgs, const TSharedRef<STableViewBase>& OwnerTableView)
+{
+	NodeLayoutRow = InArgs._NodeLayoutRow;
+	check(NodeLayoutRow.IsValid());
+	OnGetCategories = InArgs._OnGetCategories;
+	OnCategoryRenamed = InArgs._OnCategoryRenamed;
+	OnElementLabelChanged = InArgs._OnElementLabelChanged;
+	OnElementCategoryChanged = InArgs._OnElementCategoryChanged;
+	OnCategoryRemoved = InArgs._OnCategoryRemoved;
+	OnValidateCategoryName = InArgs._OnValidateCategoryName;
+	OnValidateElementName = InArgs._OnValidateElementName;
+
+	const TSharedRef<SHorizontalBox> HorizontalBox = SNew(SHorizontalBox);
+
+	LabelEditWidget = ConstructLabel(NodeLayoutRow, HorizontalBox, this);
+
+	if(NodeLayoutRow->IsCategory() || NodeLayoutRow->IsCategorizedPin())
+	{
+		HorizontalBox->AddSlot()
+		.HAlign(HAlign_Fill)
+		[
+			SNew(SSpacer)
+		];
+
+		HorizontalBox->AddSlot()
+		.HAlign(HAlign_Right)
+		.VAlign(VAlign_Center)
+		.AutoWidth()
+		[
+			SNew(SButton)
+			.ButtonStyle( FAppStyle::Get(), "SimpleButton" )
+			.Visibility_Lambda([this]()
+			{
+				if(NodeLayoutRow->IsCategory())
+				{
+					if(OnGetCategories.IsBound())
+					{
+						if(OnGetCategories.Execute().Num() > 1)
+						{
+							return EVisibility::Collapsed;
+						}
+					}
+				}
+				return EVisibility::Visible;
+			})
+			.OnClicked_Lambda([this]() -> FReply
+			{
+				if(NodeLayoutRow->IsCategory())
+				{
+					if(OnCategoryRemoved.IsBound())
+					{
+						OnCategoryRemoved.Execute(NodeLayoutRow->Path);
+						return FReply::Handled();
+					}
+				}
+				else
+				{
+					if(OnElementCategoryChanged.IsBound())
+					{
+						OnElementCategoryChanged.Execute(NodeLayoutRow->Path, FString());
+						return FReply::Handled();
+					}
+				}
+				return FReply::Unhandled();
+			})
+			.ContentPadding(0.0f)
+			[ 
+				SNew(SImage)
+				.Image( FAppStyle::GetBrush("Icons.X") )
+				.DesiredSizeOverride(FVector2D(16, 16))
+				.ColorAndOpacity(FSlateColor::UseForeground())
+			]
+		];
+	}
+	
+	STableRow<TSharedPtr<FNodeLayoutRow>>::FArguments SuperArguments;
+	SuperArguments.Content()
+	[
+		HorizontalBox
+	];
+	SuperArguments.Padding(0);
+	SuperArguments.OnCanAcceptDrop(InArgs._OnCanAcceptDrop);
+	SuperArguments.OnAcceptDrop(InArgs._OnAcceptDrop);
+	SuperArguments.OnPaintDropIndicator(InArgs._OnPaintDropIndicator);
+	SuperArguments.OnDragDetected(InArgs._OnDragDetected);
+	SuperArguments.OnDragEnter(InArgs._OnDragEnter);
+	SuperArguments.OnDragLeave(InArgs._OnDragLeave);
+	SuperArguments.OnDrop(InArgs._OnDrop);
+	STableRow< TSharedPtr<FNodeLayoutRow> >::Construct(SuperArguments, OwnerTableView);
+}
+
+TSharedPtr<SInlineEditableTextBlock> SRigVMNodeLayoutWidget::SRigVMNowLayoutRow::ConstructLabel(TSharedPtr<FNodeLayoutRow> InNodeLayoutRow, TSharedRef<SHorizontalBox> OutHorizontalBox, SRigVMNowLayoutRow* InRow)
+{
+	if(InNodeLayoutRow->Icon)
+	{
+		OutHorizontalBox->AddSlot()
+		.AutoWidth()
+		.HAlign(HAlign_Left)
+		.VAlign(VAlign_Center)
+		.Padding(0, 0, 4, 0)
+		[
+			SNew(SImage)
+			.Image(InNodeLayoutRow->Icon)
+			.DesiredSizeOverride(FVector2D(16, 16))
+			.ColorAndOpacity(InNodeLayoutRow->Color)
+		];
+	}
+
+	TSharedPtr<SInlineEditableTextBlock> InlineEditableTextBlock;
+
+	if(InRow && (
+		(InNodeLayoutRow->IsCategory() &&
+			!InNodeLayoutRow->Path.Equals(FRigVMPinCategory::GetDefaultCategoryName(), ESearchCase::IgnoreCase)
+		) || InNodeLayoutRow->IsCategorizedPin()))
+	{
+		OutHorizontalBox->AddSlot()
+		.AutoWidth()
+		.HAlign(HAlign_Left)
+		.VAlign(VAlign_Center)
+		[
+			SAssignNew(InlineEditableTextBlock, SInlineEditableTextBlock)
+			.Text(FText::FromString(InNodeLayoutRow->Label))
+			.OnVerifyTextChanged(InRow, &SRigVMNowLayoutRow::OnVerifyLabelChanged)
+			.OnTextCommitted(InRow, &SRigVMNowLayoutRow::OnLabelCommitted)
+			.MultiLine(false)
+			.Font(IDetailLayoutBuilder::GetDetailFont())
+		];
+
+		InNodeLayoutRow->OnRequestRename.BindSP(InlineEditableTextBlock.Get(), &SInlineEditableTextBlock::EnterEditingMode);
+	}
+	else
+	{
+		OutHorizontalBox->AddSlot()
+		.AutoWidth()
+		.HAlign(HAlign_Left)
+		.VAlign(VAlign_Center)
+		[
+			SNew(STextBlock)
+			.Text(FText::FromString(InNodeLayoutRow->Label))
+			.Font(IDetailLayoutBuilder::GetDetailFont())
+		];
+	}
+	
+	return InlineEditableTextBlock;
+}
+
+void SRigVMNodeLayoutWidget::SRigVMNowLayoutRow::OnLabelCommitted(const FText& InText, ETextCommit::Type InCommitType) const
+{
+	if(NodeLayoutRow->IsCategory())
+	{
+		if(OnCategoryRenamed.IsBound())
+		{
+			FString NewPath = InText.ToString();
+			FString Left, Right = NodeLayoutRow->Path;
+			if(RigVMStringUtils::SplitNodePathAtEnd(NodeLayoutRow->Path, Left, Right))
+			{
+				NewPath = RigVMStringUtils::JoinNodePath(Left, NewPath);
+			}
+			OnCategoryRenamed.Execute(NodeLayoutRow->Path, NewPath);
+		}
+	}
+	else if(OnElementLabelChanged.IsBound())
+	{
+		OnElementLabelChanged.Execute(NodeLayoutRow->Path, InText.ToString());
+	}
+}
+
+bool SRigVMNodeLayoutWidget::SRigVMNowLayoutRow::OnVerifyLabelChanged(const FText& InText, FText& OutErrorMessage)
+{
+	if(NodeLayoutRow->IsCategory())
+	{
+		if(OnValidateCategoryName.IsBound())
+		{
+			return OnValidateCategoryName.Execute(NodeLayoutRow->Path, InText.ToString(), OutErrorMessage);
+		}
+	}
+	if(OnValidateElementName.IsBound())
+	{
+		return OnValidateElementName.Execute(NodeLayoutRow->Path, InText.ToString(), OutErrorMessage);
+	}
+	return true;
+}
+
+TSharedRef<ITableRow> SRigVMNodeLayoutWidget::GenerateRow(TSharedPtr<FNodeLayoutRow> InNodeLayoutRow, const TSharedRef<STableViewBase>& OwnerTable)
+{
+	return SNew(SRigVMNowLayoutRow, OwnerTable)
+		.NodeLayoutRow(InNodeLayoutRow)
+		.OnElementLabelChanged(OnElementLabelChanged)
+		.OnElementCategoryChanged(OnElementCategoryChanged)
+		.OnCategoryRemoved(OnCategoryRemoved)
+		.OnCategoryRenamed(OnCategoryRenamed)
+		.OnValidateCategoryName(OnValidateCategoryName)
+		.OnValidateElementName(OnValidateElementName)
+		.OnDragDetected_Lambda([this, InNodeLayoutRow](const FGeometry&, const FPointerEvent&) -> FReply
+			{
+				return OnDragDetectedForRow(InNodeLayoutRow);
+			})
+		.OnCanAcceptDrop_Lambda(
+			[this](const FDragDropEvent& InDragDropEvent, EItemDropZone InDropZone, TSharedPtr<FNodeLayoutRow> InTargetRow) -> TOptional<EItemDropZone>
+			{
+				return OnCanAcceptDrop(InDragDropEvent, InTargetRow, InDropZone);
+			})
+		.OnAcceptDrop_Lambda(
+			[this](const FDragDropEvent& InDragDropEvent, EItemDropZone InDropZone, TSharedPtr<FNodeLayoutRow> InTargetRow) -> FReply
+			{
+				return OnAcceptDrop(InDragDropEvent, InTargetRow, InDropZone);
+			})
+		.ToolTipText_Lambda([InNodeLayoutRow]()
+		{
+			return FText::FromString(InNodeLayoutRow->Path);
+		})
+	;
+}
+
+void SRigVMNodeLayoutWidget::GetChildrenForRow(TSharedPtr<FNodeLayoutRow> InNodeLayoutRow, TArray<TSharedPtr<FNodeLayoutRow>>& OutChildren)
+{
+	OutChildren = InNodeLayoutRow->ChildRows;
+}
+
+void SRigVMNodeLayoutWidget::OnItemExpansionChanged(TSharedPtr<FNodeLayoutRow> InRow, bool bExpanded)
+{
+	if(InRow)
+	{
+		InRow->State.bExpanded = bExpanded;
+	}
+}
+
+void SRigVMNodeLayoutWidget::OnItemSelectionChanged(TSharedPtr<FNodeLayoutRow> InRow, ESelectInfo::Type InSelectInfo)
+{
+	struct Local
+	{
+		static void SyncSelectionState(TSharedPtr<STreeView<TSharedPtr<FNodeLayoutRow>>>& InTreeView, const TSharedPtr<FNodeLayoutRow>& InRow)
+		{
+			InRow->State.bSelected = InTreeView->IsItemSelected(InRow);
+			SyncSelectionStates(InTreeView, InRow->ChildRows);
+		}
+
+		static void SyncSelectionStates(TSharedPtr<STreeView<TSharedPtr<FNodeLayoutRow>>>& InTreeView, const TArray<TSharedPtr<FNodeLayoutRow>>& InRows)
+		{
+			for(const TSharedPtr<FNodeLayoutRow>& Row : InRows)
+			{
+				SyncSelectionState(InTreeView, Row);
+			}
+		}
+	};
+	Local::SyncSelectionStates(TreeView, NodeLayoutRows);
+}
+
+bool SRigVMNodeLayoutWidget::IsNodeLayoutEditable() const
+{
+	return IsEnabled();
+}
+
+FReply SRigVMNodeLayoutWidget::HandleAddCategory()
+{
+	if(OnCategoryAdded.IsBound())
+	{
+		FString NewCategoryName = TEXT("Category");
+		if(OnGetCategories.IsBound())
+		{
+			if(OnGetCategories.Execute().IsEmpty())
+			{
+				NewCategoryName = FRigVMPinCategory::GetDefaultCategoryName();
+			}
+		}
+		OnCategoryAdded.Execute(NewCategoryName);
+		return FReply::Handled();
+	}
+	return FReply::Unhandled();
+}
+
+TSharedRef<SRigVMNodeLayoutWidget::FRigVMNodeLayoutDragDropOp> SRigVMNodeLayoutWidget::FRigVMNodeLayoutDragDropOp::New(const TArray<TSharedPtr<FNodeLayoutRow>>& InNodeLayoutRows)
+{
+	TSharedRef<FRigVMNodeLayoutDragDropOp> Op = MakeShared<FRigVMNodeLayoutDragDropOp>();
+	Op->NodeLayoutRows = InNodeLayoutRows;
+	return Op;
+}
+
+TSharedPtr<SWidget> SRigVMNodeLayoutWidget::FRigVMNodeLayoutDragDropOp::GetDefaultDecorator() const
+{
+	TSharedPtr<SHorizontalBox> HorizontalBox;
+	TSharedPtr<SWidget> Result = SNew(SBorder)
+	.Padding(2)
+	.Visibility(EVisibility::Visible)
+	.BorderImage(FAppStyle::GetBrush("Menu.Background"))
+	[
+		SAssignNew(HorizontalBox, SHorizontalBox)
+	];
+	SRigVMNowLayoutRow::ConstructLabel(NodeLayoutRows[0], HorizontalBox.ToSharedRef(), nullptr);
+	return Result;
+}
+
+FVector2D SRigVMNodeLayoutWidget::FRigVMNodeLayoutDragDropOp::GetDecoratorPosition() const
+{
+	return FSlateApplication::Get().GetCursorPos();
+}
+
+FReply SRigVMNodeLayoutWidget::OnDragDetectedForRow(TSharedPtr<FNodeLayoutRow> InSourceRow)
+{
+	if(IsNodeLayoutEditable())
+	{
+		if(InSourceRow.IsValid())
+		{
+			const TSharedRef<FRigVMNodeLayoutDragDropOp> DragDropOp = FRigVMNodeLayoutDragDropOp::New({InSourceRow});
+			return FReply::Handled().BeginDragDrop(DragDropOp);
+		}
+	}
+	return FReply::Unhandled();
+}
+
+TOptional<EItemDropZone> SRigVMNodeLayoutWidget::OnCanAcceptDrop(const FDragDropEvent& InDragDropEvent, TSharedPtr<FNodeLayoutRow> InTargetRow, EItemDropZone InDropZone)
+{
+	static const TOptional<EItemDropZone> InvalidDropZone;
+
+	const TSharedPtr<FRigVMNodeLayoutDragDropOp> NodeLayoutDragDropOp = InDragDropEvent.GetOperationAs<FRigVMNodeLayoutDragDropOp>();
+	if(!NodeLayoutDragDropOp.IsValid())
+	{
+		return InvalidDropZone;
+	}
+	
+	if(NodeLayoutDragDropOp->GetNodeLayoutRows().IsEmpty())
+	{
+		return InvalidDropZone;
+	}
+	
+	const TSharedPtr<FNodeLayoutRow> SourceRow = NodeLayoutDragDropOp->GetNodeLayoutRows()[0];
+	if(SourceRow.Get() == InTargetRow.Get())
+	{
+		return InvalidDropZone;
+	}
+
+	if(SourceRow->IsCategory() && InTargetRow->IsPin())
+	{
+		return InvalidDropZone;
+	}
+
+	if(InTargetRow->IsCategory())
+	{
+		if(SourceRow->IsPin())
+		{
+			return EItemDropZone::OntoItem;
+		}
+		if(SourceRow->IsCategory())
+		{
+			// for now we don't allow nesting of categories
+			// even though the API in the controller allows for it.
+			if(InDropZone == EItemDropZone::AboveItem ||
+				InDropZone == EItemDropZone::BelowItem)
+			{
+				return InDropZone;
+			}
+		}
+	}
+	else if(InTargetRow->IsPin())
+	{
+		if(SourceRow->IsPin())
+		{
+			if(InDropZone == EItemDropZone::AboveItem ||
+				InDropZone == EItemDropZone::BelowItem)
+			{
+				return InDropZone;
+			}
+		}
+	}
+			
+	return InvalidDropZone;
+}
+
+FReply SRigVMNodeLayoutWidget::OnAcceptDrop(const FDragDropEvent& InDragDropEvent, TSharedPtr<FNodeLayoutRow> InTargetRow, EItemDropZone InDropZone)
+{
+	if(!OnCanAcceptDrop(InDragDropEvent, InTargetRow, InDropZone).IsSet())
+	{
+		return FReply::Unhandled();
+	}
+
+	const TSharedPtr<FRigVMNodeLayoutDragDropOp> NodeLayoutDragDropOp = InDragDropEvent.GetOperationAs<FRigVMNodeLayoutDragDropOp>();
+	if(!NodeLayoutDragDropOp.IsValid())
+	{
+		return FReply::Unhandled();
+	}
+	
+	if(NodeLayoutDragDropOp->GetNodeLayoutRows().IsEmpty())
+	{
+		return FReply::Unhandled();
+	}
+	
+	const TSharedPtr<FNodeLayoutRow> SourceRow = NodeLayoutDragDropOp->GetNodeLayoutRows()[0];
+
+	if(InTargetRow->IsCategory())
+	{
+		if(SourceRow->IsPin())
+		{
+			if(OnElementCategoryChanged.IsBound())
+			{
+				int32 IndexInCategory = INDEX_NONE;
+				const TArray<FRigVMPinCategory> Categories = OnGetCategories.Execute();
+				for(const FRigVMPinCategory& Category : Categories)
+				{
+					if(Category.Path == InTargetRow->Path)
+					{
+						IndexInCategory = Category.Elements.Num();
+						break;
+					}
+				}
+
+				OnElementCategoryChanged.Execute(SourceRow->Path, InTargetRow->Path);
+				
+				if(OnElementIndexInCategoryChanged.IsBound())
+				{
+					OnElementIndexInCategoryChanged.Execute(SourceRow->Path, IndexInCategory);
+				}
+				return FReply::Handled();
+			}
+		}
+		if(SourceRow->IsCategory())
+		{
+			if(InDropZone == EItemDropZone::AboveItem ||
+				InDropZone == EItemDropZone::BelowItem)
+			{
+				// reorder categories
+				// todo
+				return FReply::Handled();
+			}
+		}
+	}
+	else if(InTargetRow->IsPin())
+	{
+		if(SourceRow->IsPin())
+		{
+			if(InDropZone == EItemDropZone::AboveItem ||
+				InDropZone == EItemDropZone::BelowItem)
+			{
+				if(OnElementCategoryChanged.IsBound())
+				{
+					if(InTargetRow->IsUncategorizedPin())
+					{
+						// remove the pin category / set the pin to uncategorized
+						OnElementCategoryChanged.Execute(SourceRow->Path, FString());
+						return FReply::Handled();
+					}
+					if(OnGetElementCategory.IsBound())
+					{
+						// copy the category from the target onto me
+						const FString TargetCategory = OnGetElementCategory.Execute(InTargetRow->Path);
+						if(!TargetCategory.IsEmpty())
+						{
+							OnElementCategoryChanged.Execute(SourceRow->Path, TargetCategory);
+
+							if(OnGetElementIndexInCategory.IsBound() && OnElementIndexInCategoryChanged.IsBound())
+							{
+								const int32 TargetIndex = OnGetElementIndexInCategory.Execute(InTargetRow->Path);
+								if(TargetIndex != INDEX_NONE)
+								{
+									OnElementIndexInCategoryChanged.Execute(
+										SourceRow->Path,
+										InDropZone == EItemDropZone::AboveItem ? TargetIndex : TargetIndex + 1);
+								}
+							}
+							return FReply::Handled();
+						}
+					}
+				}
+			}
+		}
+	}
+	return FReply::Unhandled();
+}
+
+FReply SRigVMNodeLayoutWidget::OnDrop(const FGeometry& InGeometry, const FDragDropEvent& InDragDropEvent)
+{
+	// don't allow to drop anything onto the widget itself
+	return FReply::Unhandled();
+}
+
+#undef LOCTEXT_NAMESPACE // SRigVMNodeLayoutWidget
diff --git a/Engine/Plugins/Runtime/RigVM/Source/RigVMEditor/Public/Editor/RigVMGraphDetailCustomization.h b/Engine/Plugins/Runtime/RigVM/Source/RigVMEditor/Public/Editor/RigVMGraphDetailCustomization.h
--- a/Engine/Plugins/Runtime/RigVM/Source/RigVMEditor/Public/Editor/RigVMGraphDetailCustomization.h
+++ b/Engine/Plugins/Runtime/RigVM/Source/RigVMEditor/Public/Editor/RigVMGraphDetailCustomization.h
@@ -205,6 +205,28 @@ private:
 	void OnAddAssignedTag(const FName& InTagName);
 	void OnRemoveAssignedTag(const FName& InTagName);
 
+	URigVMLibraryNode* GetLibraryNode() const;
+	URigVMNode* GetNodeForLayout() const;
+	const FRigVMNodeLayout* GetNodeLayout() const;
+	TArray<FString> GetUncategorizedPins() const;
+	TArray<FRigVMPinCategory> GetPinCategories() const;
+	FString GetPinCategory(FString InPinPath) const;
+	int32 GetPinIndexInCategory(FString InPinPath) const;
+	FString GetPinLabel(FString InPinPath) const;
+	FLinearColor GetPinColor(FString InPinPath) const;
+	const FSlateBrush* GetPinIcon(FString InPinPath) const;
+	void HandleCategoryAdded(FString InCategory);
+	void HandleCategoryRemoved(FString InCategory);
+	void HandleCategoryRenamed(FString InOldCategory, FString InNewCategory);
+	void HandlePinCategoryChanged(FString InPinPath, FString InCategory);
+	void HandlePinLabelChanged(FString InPinPath, FString InNewLabel);
+	void HandlePinIndexInCategoryChanged(FString InPinPath, int32 InIndexInCategory);
+	static bool ValidateName(FString InNewName, FText& OutErrorMessage); 
+	bool HandleValidateCategoryName(FString InCategoryPath, FString InNewName, FText& OutErrorMessage);
+	bool HandleValidatePinDisplayName(FString InPinPath, FString InNewName, FText& OutErrorMessage);
+
+	uint32 GetNodeLayoutHash() const;
+
 	/** The Blueprint editor we are embedded in */
 	TWeakPtr<FRigVMEditor> RigVMEditorPtr;
 
@@ -221,6 +243,7 @@ private:
 	bool bIsPickingColor;
 
 	static TArray<TSharedPtr<FRigVMStringWithTag>> AccessSpecifierStrings;
+	mutable TOptional<FRigVMNodeLayout> CachedNodeLayout;
 };
 
 /** Customization for editing a rig vm node */
diff --git a/Engine/Plugins/Runtime/RigVM/Source/RigVMEditor/Public/Widgets/SRigVMNodeLayoutWidget.h b/Engine/Plugins/Runtime/RigVM/Source/RigVMEditor/Public/Widgets/SRigVMNodeLayoutWidget.h
new file mode 100644
--- /dev/null
+++ b/Engine/Plugins/Runtime/RigVM/Source/RigVMEditor/Public/Widgets/SRigVMNodeLayoutWidget.h
@@ -0,0 +1,226 @@
+// Copyright Epic Games, Inc. All Rights Reserved.
+
+#pragma once
+
+#include "Widgets/Layout/SBox.h"
+#include "Widgets/Views/STreeView.h"
+#include "Widgets/Text/SInlineEditableTextBlock.h"
+#include "RigVMModel/RigVMNode.h"
+#include "Input/DragAndDrop.h"
+
+DECLARE_DELEGATE_RetVal(TArray<FString> /* pin paths  */, FRigVMNodeLayoutWidget_OnGetUncategorizedPins);
+DECLARE_DELEGATE_RetVal(TArray<FRigVMPinCategory> /* categories */, FRigVMNodeLayoutWidget_OnGetCategories);
+DECLARE_DELEGATE_OneParam(FRigVMNodeLayoutWidget_OnCategoryAdded, FString /* category */);
+DECLARE_DELEGATE_OneParam(FRigVMNodeLayoutWidget_OnCategoryRemoved, FString /* category */);
+DECLARE_DELEGATE_TwoParams(FRigVMNodeLayoutWidget_OnCategoryRenamed, FString /* old */, FString /* new */);
+DECLARE_DELEGATE_RetVal_OneParam(FString, FRigVMNodeLayoutWidget_OnGetElementLabel, FString /* element path */);
+DECLARE_DELEGATE_TwoParams(FRigVMNodeLayoutWidget_OnElementLabelChanged, FString /* element path */, FString /* new label */);
+DECLARE_DELEGATE_RetVal_OneParam(FString, FRigVMNodeLayoutWidget_OnGetElementCategory, FString /* element path */);
+DECLARE_DELEGATE_RetVal_OneParam(int32, FRigVMNodeLayoutWidget_OnGetElementIndexInCategory, FString /* element path */);
+DECLARE_DELEGATE_RetVal_OneParam(FLinearColor, FRigVMNodeLayoutWidget_OnGetElementColor, FString /* element path */);
+DECLARE_DELEGATE_RetVal_OneParam(const FSlateBrush*, FRigVMNodeLayoutWidget_OnGetElementIcon, FString /* element path */);
+DECLARE_DELEGATE_TwoParams(FRigVMNodeLayoutWidget_OnElementIndexInCategoryChanged, FString /* element path */, int32 /* new element index */);
+DECLARE_DELEGATE_TwoParams(FRigVMNodeLayoutWidget_OnElementCategoryChanged, FString /* element */, FString /* new category */);
+DECLARE_DELEGATE_RetVal(uint32, FRigVMNodeLayoutWidget_OnGetStructuralHash);
+DECLARE_DELEGATE_RetVal_ThreeParams(bool, FRigVMNodeLayoutWidget_ValidateName, FString /* in path */, FString /* new name */, FText& /* out error */);
+
+class RIGVMEDITOR_API SRigVMNodeLayoutWidget : public SBox
+{
+public:
+	
+	SLATE_BEGIN_ARGS(SRigVMNodeLayoutWidget)
+		: _MaxScrollBoxSize(300.f)
+	{
+	}
+	SLATE_EVENT(FRigVMNodeLayoutWidget_OnGetUncategorizedPins, OnGetUncategorizedPins)
+	SLATE_EVENT(FRigVMNodeLayoutWidget_OnGetCategories, OnGetCategories)
+	SLATE_EVENT(FRigVMNodeLayoutWidget_OnCategoryAdded, OnCategoryAdded)
+	SLATE_EVENT(FRigVMNodeLayoutWidget_OnCategoryRemoved, OnCategoryRemoved)
+	SLATE_EVENT(FRigVMNodeLayoutWidget_OnCategoryRenamed, OnCategoryRenamed)
+	SLATE_EVENT(FRigVMNodeLayoutWidget_OnGetElementLabel, OnGetElementLabel)
+	SLATE_EVENT(FRigVMNodeLayoutWidget_OnElementLabelChanged, OnElementLabelChanged)
+	SLATE_EVENT(FRigVMNodeLayoutWidget_OnGetElementCategory, OnGetElementCategory)
+	SLATE_EVENT(FRigVMNodeLayoutWidget_OnGetElementIndexInCategory, OnGetElementIndexInCategory)
+	SLATE_EVENT(FRigVMNodeLayoutWidget_OnGetElementColor, OnGetElementColor)
+	SLATE_EVENT(FRigVMNodeLayoutWidget_OnGetElementIcon, OnGetElementIcon)
+	SLATE_EVENT(FRigVMNodeLayoutWidget_OnElementIndexInCategoryChanged, OnElementIndexInCategoryChanged)
+	SLATE_EVENT(FRigVMNodeLayoutWidget_OnElementCategoryChanged, OnElementCategoryChanged)
+	SLATE_EVENT(FRigVMNodeLayoutWidget_OnGetStructuralHash, OnGetStructuralHash)
+	SLATE_EVENT(FRigVMNodeLayoutWidget_ValidateName, OnValidateCategoryName)
+	SLATE_EVENT(FRigVMNodeLayoutWidget_ValidateName, OnValidateElementName)
+	SLATE_ATTRIBUTE(float, MaxScrollBoxSize)
+	SLATE_END_ARGS()
+
+	SRigVMNodeLayoutWidget();
+	virtual ~SRigVMNodeLayoutWidget() override;
+	
+	void Construct(const FArguments& InArgs);
+	void Refresh();
+	virtual void Tick(const FGeometry& AllottedGeometry, const double InCurrentTime, const float InDeltaTime) override;
+	virtual FReply OnKeyDown(const FGeometry& MyGeometry, const FKeyEvent& InKeyEvent) override;
+
+private:
+
+	FRigVMNodeLayoutWidget_OnGetUncategorizedPins OnGetUncategorizedPins;
+	FRigVMNodeLayoutWidget_OnGetCategories OnGetCategories;
+	FRigVMNodeLayoutWidget_OnCategoryAdded OnCategoryAdded;
+	FRigVMNodeLayoutWidget_OnCategoryRemoved OnCategoryRemoved;
+	FRigVMNodeLayoutWidget_OnCategoryRenamed OnCategoryRenamed;
+	FRigVMNodeLayoutWidget_OnGetElementLabel OnGetElementLabel;
+	FRigVMNodeLayoutWidget_OnElementLabelChanged OnElementLabelChanged;
+	FRigVMNodeLayoutWidget_OnGetElementCategory OnGetElementCategory;
+	FRigVMNodeLayoutWidget_OnGetElementIndexInCategory OnGetElementIndexInCategory;
+	FRigVMNodeLayoutWidget_OnGetElementColor OnGetElementColor;
+	FRigVMNodeLayoutWidget_OnGetElementIcon OnGetElementIcon;
+	FRigVMNodeLayoutWidget_OnElementIndexInCategoryChanged OnElementIndexInCategoryChanged;
+	FRigVMNodeLayoutWidget_OnElementCategoryChanged OnElementCategoryChanged;
+	FRigVMNodeLayoutWidget_OnGetStructuralHash OnGetStructuralHash;
+	FRigVMNodeLayoutWidget_ValidateName OnValidateCategoryName;
+	FRigVMNodeLayoutWidget_ValidateName OnValidateElementName;
+
+	struct FNodeLayoutRow
+	{
+		FNodeLayoutRow()
+			: bIsCategory(true)
+			, bIsUncategorized(false)
+			, Color(FLinearColor::White)
+			, Icon(nullptr)
+		{
+		}
+
+		struct FState
+		{
+			FState()
+				: bExpanded(false)
+				, bSelected(false)
+			{
+			}
+
+			bool bExpanded;
+			bool bSelected;
+		};
+		
+		bool bIsCategory;
+		bool bIsUncategorized;
+		FString Path;
+		FString Label;
+		FLinearColor Color;
+		const FSlateBrush* Icon;
+		FState State;
+		TArray<TSharedPtr<FNodeLayoutRow>> ChildRows;
+		FSimpleDelegate OnRequestRename;
+
+		bool IsCategory() const
+		{
+			return bIsCategory;
+		}
+
+		bool IsPin() const
+		{
+			return !bIsCategory;
+		}
+
+		bool IsUncategorizedPin() const
+		{
+			return IsPin() && bIsUncategorized;
+		}
+
+		bool IsCategorizedPin() const
+		{
+			return IsPin() && !bIsUncategorized;
+		}
+
+		void RequestRename()
+		{
+			(void)OnRequestRename.ExecuteIfBound();
+		}
+	};
+
+	DECLARE_DELEGATE_RetVal_ThreeParams(TOptional<EItemDropZone>, FRigVMNodeLayoutWidget_OnElementCanDrop, TSharedPtr<FNodeLayoutRow>, TSharedPtr<FNodeLayoutRow>, EItemDropZone);
+	DECLARE_DELEGATE_RetVal_ThreeParams(FReply, FRigVMNodeLayoutWidget_OnElementAcceptDrop, TSharedPtr<FNodeLayoutRow>, TSharedPtr<FNodeLayoutRow>, EItemDropZone);
+
+	class RIGVMEDITOR_API SRigVMNowLayoutRow
+	: public STableRow<TSharedPtr<FNodeLayoutRow>>
+	{
+	public:
+	
+		SLATE_BEGIN_ARGS(SRigVMNowLayoutRow)
+		{}
+		SLATE_ARGUMENT(TSharedPtr<FNodeLayoutRow>, NodeLayoutRow)
+		SLATE_EVENT(FRigVMNodeLayoutWidget_OnGetCategories, OnGetCategories)
+		SLATE_EVENT(FRigVMNodeLayoutWidget_OnCategoryRenamed, OnCategoryRenamed)
+		SLATE_EVENT(FRigVMNodeLayoutWidget_OnElementLabelChanged, OnElementLabelChanged)
+		SLATE_EVENT(FRigVMNodeLayoutWidget_OnElementCategoryChanged, OnElementCategoryChanged)
+		SLATE_EVENT(FRigVMNodeLayoutWidget_OnCategoryRemoved, OnCategoryRemoved)
+		SLATE_EVENT(FRigVMNodeLayoutWidget_ValidateName, OnValidateCategoryName)
+		SLATE_EVENT(FRigVMNodeLayoutWidget_ValidateName, OnValidateElementName)
+
+		SLATE_EVENT( FOnCanAcceptDrop, OnCanAcceptDrop )
+		SLATE_EVENT( FOnAcceptDrop,    OnAcceptDrop )
+		SLATE_EVENT( FOnPaintDropIndicator, OnPaintDropIndicator )
+		SLATE_EVENT( FOnDragDetected,      OnDragDetected )
+		SLATE_EVENT( FOnTableRowDragEnter, OnDragEnter )
+		SLATE_EVENT( FOnTableRowDragLeave, OnDragLeave )
+		SLATE_EVENT( FOnTableRowDrop,      OnDrop )
+
+		SLATE_END_ARGS()
+
+		virtual ~SRigVMNowLayoutRow() override;
+		void Construct(const FArguments& InArgs, const TSharedRef<STableViewBase>& OwnerTableView);
+		static TSharedPtr<SInlineEditableTextBlock> ConstructLabel(TSharedPtr<FNodeLayoutRow> InNodeLayoutRow, TSharedRef<SHorizontalBox> OutHorizontalBox, SRigVMNowLayoutRow* InRow);
+
+		TSharedPtr<FNodeLayoutRow> GetNodeLayoutRow() const { return NodeLayoutRow; }
+		void OnLabelCommitted(const FText& InText, ETextCommit::Type InCommitType) const;
+		bool OnVerifyLabelChanged(const FText& InText, FText& OutErrorMessage);
+
+	private:
+		TSharedPtr<FNodeLayoutRow> NodeLayoutRow;
+		FRigVMNodeLayoutWidget_OnGetCategories OnGetCategories;
+		FRigVMNodeLayoutWidget_OnCategoryRenamed OnCategoryRenamed;
+		FRigVMNodeLayoutWidget_OnElementLabelChanged OnElementLabelChanged;
+		FRigVMNodeLayoutWidget_OnElementCategoryChanged OnElementCategoryChanged;
+		FRigVMNodeLayoutWidget_OnCategoryRemoved OnCategoryRemoved;
+		FRigVMNodeLayoutWidget_ValidateName OnValidateCategoryName;
+		FRigVMNodeLayoutWidget_ValidateName OnValidateElementName;
+		
+		TSharedPtr<SInlineEditableTextBlock> LabelEditWidget;
+	};
+
+	TSharedRef<ITableRow> GenerateRow(TSharedPtr<FNodeLayoutRow> InRow, const TSharedRef<STableViewBase>& OwnerTable);
+	void GetChildrenForRow(TSharedPtr<FNodeLayoutRow> InRow, TArray<TSharedPtr<FNodeLayoutRow>>& OutChildren);
+	void OnItemExpansionChanged(TSharedPtr<FNodeLayoutRow> InRow, bool bExpanded);
+	void OnItemSelectionChanged(TSharedPtr<FNodeLayoutRow> InRow, ESelectInfo::Type InSelectInfo);
+	bool IsNodeLayoutEditable() const;
+	FReply HandleAddCategory();
+
+	class FRigVMNodeLayoutDragDropOp : public FDragDropOperation
+	{
+	public:
+		DRAG_DROP_OPERATOR_TYPE(FRigVMNodeLayoutDragDropOp, FDragDropOperation)
+
+		static TSharedRef<FRigVMNodeLayoutDragDropOp> New(const TArray<TSharedPtr<FNodeLayoutRow>>& InNodeLayoutRows);
+
+		virtual TSharedPtr<SWidget> GetDefaultDecorator() const override;
+		virtual FVector2D GetDecoratorPosition() const override;
+		virtual bool IsWindowlessOperation() const override { return true; }
+
+		const TArray<TSharedPtr<FNodeLayoutRow>>& GetNodeLayoutRows() const
+		{
+			return NodeLayoutRows;
+		}
+		
+	private:
+
+		/** Data for the property paths this item represents */
+		TArray<TSharedPtr<FNodeLayoutRow>> NodeLayoutRows;
+	};
+
+	FReply OnDragDetectedForRow(TSharedPtr<FNodeLayoutRow> InSourceRow);
+	TOptional<EItemDropZone> OnCanAcceptDrop(const FDragDropEvent& InDragDropEvent, TSharedPtr<FNodeLayoutRow> InTargetRow, EItemDropZone InDropZone);
+	FReply OnAcceptDrop(const FDragDropEvent& InDragDropEvent, TSharedPtr<FNodeLayoutRow> InTargetRow, EItemDropZone InDropZone);
+	virtual FReply OnDrop(const FGeometry& InGeometry, const FDragDropEvent& InDragDropEvent);
+	
+	TArray<TSharedPtr<FNodeLayoutRow>> NodeLayoutRows;
+	TSharedPtr<STreeView<TSharedPtr<FNodeLayoutRow>>> TreeView;
+	TOptional<uint32> LastStructuralHash;
+};
\ No newline at end of file
