RigVM: Support for Pin Categories - hide feature behind console variable

#rnx

diff --git a/Engine/Plugins/Experimental/AnimNext/Source/AnimNextUncookedOnly/Private/AnimNextRigVMAssetSchema.h b/Engine/Plugins/Experimental/AnimNext/Source/AnimNextUncookedOnly/Private/AnimNextRigVMAssetSchema.h
--- a/Engine/Plugins/Experimental/AnimNext/Source/AnimNextUncookedOnly/Private/AnimNextRigVMAssetSchema.h
+++ b/Engine/Plugins/Experimental/AnimNext/Source/AnimNextUncookedOnly/Private/AnimNextRigVMAssetSchema.h
@@ -11,6 +11,11 @@ class UAnimNextRigVMAssetSchema : public URigVMSchema
 {
 	GENERATED_BODY()
 
+	virtual bool SupportsNodeLayouts(const URigVMGraph* InGraph) const override
+	{
+		return true;
+	}
+
 protected:
 	UAnimNextRigVMAssetSchema();
 };
diff --git a/Engine/Plugins/Runtime/RigVM/Source/RigVMDeveloper/Private/RigVMModel/RigVMClient.cpp b/Engine/Plugins/Runtime/RigVM/Source/RigVMDeveloper/Private/RigVMModel/RigVMClient.cpp
--- a/Engine/Plugins/Runtime/RigVM/Source/RigVMDeveloper/Private/RigVMModel/RigVMClient.cpp
+++ b/Engine/Plugins/Runtime/RigVM/Source/RigVMDeveloper/Private/RigVMModel/RigVMClient.cpp
@@ -627,6 +627,7 @@ void FRigVMClient::AddModel(URigVMGraph* InModel, bool bCreateController)
 		}
 	}
 
+	InModel->SetSchemaClass(GetDefaultSchemaClass());
 	InModel->SetExecuteContextStruct(InModel->GetSchema()->GetExecuteContextStruct());
 
 	if(bCreateController)
@@ -1192,6 +1193,7 @@ FRigVMClientPatchResult FRigVMClient::PatchModelsOnLoad()
 	for(URigVMGraph* Model : AllModels)
 	{
 		Model->PostLoad();
+		Model->SetSchemaClass(GetDefaultSchemaClass());
 		
 		URigVMController* Controller = GetOrCreateController(Model);
 		TGuardValue<bool> GuardSuspendTemplateComputation(Controller->bSuspendTemplateComputation, true);
diff --git a/Engine/Plugins/Runtime/RigVM/Source/RigVMDeveloper/Private/RigVMModel/RigVMController.cpp b/Engine/Plugins/Runtime/RigVM/Source/RigVMDeveloper/Private/RigVMModel/RigVMController.cpp
--- a/Engine/Plugins/Runtime/RigVM/Source/RigVMDeveloper/Private/RigVMModel/RigVMController.cpp
+++ b/Engine/Plugins/Runtime/RigVM/Source/RigVMDeveloper/Private/RigVMModel/RigVMController.cpp
@@ -8114,6 +8114,11 @@ bool URigVMController::AddEmptyPinCategory(const URigVMNode* InNode, const FStri
 		return false;
 	}
 
+	if(!GetGraph()->GetSchema()->SupportsNodeLayouts(GetGraph()))
+	{
+		return false;
+	}
+
 	FRigVMControllerCompileBracketScope CompileScope(this);
 	FRigVMChangeNodePinCategoriesAction Action;
 	if (bSetupUndoRedo)
@@ -8236,6 +8241,11 @@ bool URigVMController::SetPinCategory(URigVMPin* InPin, const FString& InCategor
 		return false;
 	}
 
+	if(!GetGraph()->GetSchema()->SupportsNodeLayouts(GetGraph()))
+	{
+		return false;
+	}
+
 	TArray<FString> Categories;
 	RigVMStringUtils::SplitNodePath(InCategory, Categories);
 
@@ -8486,6 +8496,11 @@ bool URigVMController::RemovePinCategory(const URigVMNode* InNode, const FString
 		return false;
 	}
 
+	if(!GetGraph()->GetSchema()->SupportsNodeLayouts(GetGraph()))
+	{
+		return false;
+	}
+
 	const TArray<FString> AllCategories = InNode->GetPinCategories();
 	TArray<FString> CategoriesToRemove = { InPinCategory };
 	CategoriesToRemove.Append(InNode->GetSubPinCategories(InPinCategory, false, true));
@@ -8613,6 +8628,11 @@ bool URigVMController::RenamePinCategory(const URigVMNode* InNode, const FString
 		return false;
 	}
 
+	if(!GetGraph()->GetSchema()->SupportsNodeLayouts(GetGraph()))
+	{
+		return false;
+	}
+
 	if(InOldPinCategory.Equals(FRigVMPinCategory::GetDefaultCategoryName(), ESearchCase::IgnoreCase))
 	{
 		ReportErrorf(TEXT("Cannot rename '%s' category."), *InOldPinCategory);
@@ -8772,6 +8792,11 @@ bool URigVMController::SetPinCategoryIndex(const URigVMNode* InNode, const FStri
 		return false;
 	}
 
+	if(!GetGraph()->GetSchema()->SupportsNodeLayouts(GetGraph()))
+	{
+		return false;
+	}
+
 	// for entry or return nodes we relay to the outer pins
 	if(InNode->IsA<URigVMFunctionInterfaceNode>())
 	{
@@ -8858,6 +8883,11 @@ bool URigVMController::SetPinCategoryExpansion(const URigVMNode* InNode, const F
 		return false;
 	}
 
+	if(!GetGraph()->GetSchema()->SupportsNodeLayouts(GetGraph()))
+	{
+		return false;
+	}
+
 	const TArray<FString> AllCategories = InNode->GetPinCategories();
 	if(!AllCategories.Contains(InPinCategory))
 	{
@@ -9077,6 +9107,11 @@ bool URigVMController::SetNodeLayout(const FName& InNodeName, FRigVMNodeLayout I
 	const URigVMGraph* Graph = GetGraph();
 	check(Graph);
 
+	if(!Graph->GetSchema()->SupportsNodeLayouts(Graph))
+	{
+		return false;
+	}
+
 	const URigVMNode* Node = Graph->FindNodeByName(InNodeName);
 	if (Node == nullptr)
 	{
diff --git a/Engine/Plugins/Runtime/RigVM/Source/RigVMDeveloper/Private/RigVMModel/RigVMSchema.cpp b/Engine/Plugins/Runtime/RigVM/Source/RigVMDeveloper/Private/RigVMModel/RigVMSchema.cpp
--- a/Engine/Plugins/Runtime/RigVM/Source/RigVMDeveloper/Private/RigVMModel/RigVMSchema.cpp
+++ b/Engine/Plugins/Runtime/RigVM/Source/RigVMDeveloper/Private/RigVMModel/RigVMSchema.cpp
@@ -5,6 +5,8 @@
 #include "RigVMModel/RigVMClient.h"
 #include "RigVMModel/RigVMController.h"
 
+TAutoConsoleVariable<bool> CVarRigVMEnableNodeLayouts(TEXT("RigVM.EnableNodeLayouts"), false, TEXT("Set to true to turn on support for node layouts"));
+
 URigVMSchema::URigVMSchema()
 	: ExecuteContextStruct(nullptr)
 	, Registry(&FRigVMRegistry::Get())
@@ -218,6 +220,11 @@ bool URigVMSchema::SupportsExternalVariable(URigVMController* InController, cons
 	return false;
 }
 
+bool URigVMSchema::SupportsNodeLayouts(const URigVMGraph* InGraph) const
+{
+	return CVarRigVMEnableNodeLayouts.GetValueOnAnyThread();
+}
+
 bool URigVMSchema::ShouldUnfoldStruct(URigVMController* InController, const UStruct* InStruct) const
 {
 	RIGVMSCHEMA_DEFAULT_FUNCTION_BODY
diff --git a/Engine/Plugins/Runtime/RigVM/Source/RigVMDeveloper/Public/RigVMModel/RigVMSchema.h b/Engine/Plugins/Runtime/RigVM/Source/RigVMDeveloper/Public/RigVMModel/RigVMSchema.h
--- a/Engine/Plugins/Runtime/RigVM/Source/RigVMDeveloper/Public/RigVMModel/RigVMSchema.h
+++ b/Engine/Plugins/Runtime/RigVM/Source/RigVMDeveloper/Public/RigVMModel/RigVMSchema.h
@@ -55,6 +55,9 @@ public:
 	// Returns true if a graph supports a given external variable
 	virtual bool SupportsExternalVariable(URigVMController* InController, const FRigVMExternalVariable* InExternalVariable) const;
 
+	// Returns true if a graph supports pin categories / node layouts
+	virtual bool SupportsNodeLayouts(const URigVMGraph* InGraph) const;
+
 	// Returns true if the pin for a given struct should be unfolded into subpins for a graph
 	virtual bool ShouldUnfoldStruct(URigVMController* InController, const UStruct* InStruct) const;
 
diff --git a/Engine/Plugins/Runtime/RigVM/Source/RigVMEditor/Private/Editor/RigVMGraphDetailCustomization.cpp b/Engine/Plugins/Runtime/RigVM/Source/RigVMEditor/Private/Editor/RigVMGraphDetailCustomization.cpp
--- a/Engine/Plugins/Runtime/RigVM/Source/RigVMEditor/Private/Editor/RigVMGraphDetailCustomization.cpp
+++ b/Engine/Plugins/Runtime/RigVM/Source/RigVMEditor/Private/Editor/RigVMGraphDetailCustomization.cpp
@@ -1024,39 +1024,48 @@ void FRigVMGraphDetailCustomization::CustomizeDetails(IDetailLayoutBuilder& Deta
 		];
 	}
 
-	SettingsCategory.AddCustomRow(FText::GetEmpty())
-	.OverrideResetToDefault(FResetToDefaultOverride::Hide())
-	.Visibility(TAttribute<EVisibility>::CreateLambda([this]()
+	if(Model)
 	{
-		return IsValidFunction() ? EVisibility::Visible : EVisibility::Collapsed;
-	}))
-	.NameContent()
-	[
-		SNew(STextBlock)
-		.Text(FText::FromString(TEXT("Layout")))
-		.Font(IDetailLayoutBuilder::GetDetailFont())
-	]
-	.ValueContent()
-	.HAlign(HAlign_Fill)
-	[
-		SNew(SRigVMNodeLayoutWidget)
-		.OnGetUncategorizedPins(this, &FRigVMGraphDetailCustomization::GetUncategorizedPins)
-		.OnGetCategories(this, &FRigVMGraphDetailCustomization::GetPinCategories)
-		.OnGetElementCategory(this, &FRigVMGraphDetailCustomization::GetPinCategory)
-		.OnGetElementIndexInCategory(this, &FRigVMGraphDetailCustomization::GetPinIndexInCategory)
-		.OnGetElementLabel(this, &FRigVMGraphDetailCustomization::GetPinLabel)
-		.OnGetElementColor(this, &FRigVMGraphDetailCustomization::GetPinColor)
-		.OnGetElementIcon(this, &FRigVMGraphDetailCustomization::GetPinIcon)
-		.OnCategoryAdded(this, &FRigVMGraphDetailCustomization::HandleCategoryAdded)
-		.OnCategoryRemoved(this, &FRigVMGraphDetailCustomization::HandleCategoryRemoved)
-		.OnCategoryRenamed(this, &FRigVMGraphDetailCustomization::HandleCategoryRenamed)
-		.OnElementCategoryChanged(this, &FRigVMGraphDetailCustomization::HandlePinCategoryChanged)
-		.OnElementLabelChanged(this, &FRigVMGraphDetailCustomization::HandlePinLabelChanged)
-		.OnElementIndexInCategoryChanged(this, &FRigVMGraphDetailCustomization::HandlePinIndexInCategoryChanged)
-		.OnValidateCategoryName(this, &FRigVMGraphDetailCustomization::HandleValidateCategoryName)
-		.OnValidateElementName(this, &FRigVMGraphDetailCustomization::HandleValidatePinDisplayName)
-		.OnGetStructuralHash(this, &FRigVMGraphDetailCustomization::GetNodeLayoutHash)
-	];
+		if(const URigVMSchema* Schema = Model->GetSchema())
+		{
+			if(Schema->SupportsNodeLayouts(Model))
+			{
+				SettingsCategory.AddCustomRow(FText::GetEmpty())
+				.OverrideResetToDefault(FResetToDefaultOverride::Hide())
+				.Visibility(TAttribute<EVisibility>::CreateLambda([this]()
+				{
+					return IsValidFunction() ? EVisibility::Visible : EVisibility::Collapsed;
+				}))
+				.NameContent()
+				[
+					SNew(STextBlock)
+					.Text(FText::FromString(TEXT("Layout")))
+					.Font(IDetailLayoutBuilder::GetDetailFont())
+				]
+				.ValueContent()
+				.HAlign(HAlign_Fill)
+				[
+					SNew(SRigVMNodeLayoutWidget)
+					.OnGetUncategorizedPins(this, &FRigVMGraphDetailCustomization::GetUncategorizedPins)
+					.OnGetCategories(this, &FRigVMGraphDetailCustomization::GetPinCategories)
+					.OnGetElementCategory(this, &FRigVMGraphDetailCustomization::GetPinCategory)
+					.OnGetElementIndexInCategory(this, &FRigVMGraphDetailCustomization::GetPinIndexInCategory)
+					.OnGetElementLabel(this, &FRigVMGraphDetailCustomization::GetPinLabel)
+					.OnGetElementColor(this, &FRigVMGraphDetailCustomization::GetPinColor)
+					.OnGetElementIcon(this, &FRigVMGraphDetailCustomization::GetPinIcon)
+					.OnCategoryAdded(this, &FRigVMGraphDetailCustomization::HandleCategoryAdded)
+					.OnCategoryRemoved(this, &FRigVMGraphDetailCustomization::HandleCategoryRemoved)
+					.OnCategoryRenamed(this, &FRigVMGraphDetailCustomization::HandleCategoryRenamed)
+					.OnElementCategoryChanged(this, &FRigVMGraphDetailCustomization::HandlePinCategoryChanged)
+					.OnElementLabelChanged(this, &FRigVMGraphDetailCustomization::HandlePinLabelChanged)
+					.OnElementIndexInCategoryChanged(this, &FRigVMGraphDetailCustomization::HandlePinIndexInCategoryChanged)
+					.OnValidateCategoryName(this, &FRigVMGraphDetailCustomization::HandleValidateCategoryName)
+					.OnValidateElementName(this, &FRigVMGraphDetailCustomization::HandleValidatePinDisplayName)
+					.OnGetStructuralHash(this, &FRigVMGraphDetailCustomization::GetNodeLayoutHash)
+				];
+			}
+		}
+	}
 
 	IDetailCategoryBuilder& DefaultsCategory = DetailLayout.EditCategory("NodeDefaults", LOCTEXT("FunctionDetailsNodeDefaults", "Node Defaults"));
 	TSharedRef<FRigVMFunctionArgumentDefaultNode> DefaultsArgumentNode = MakeShareable(new FRigVMFunctionArgumentDefaultNode(
