[ControlRig] Added support for soft object ptrs stored in traits to react to asset renaming

diff --git a/Engine/Plugins/Runtime/RigVM/Source/RigVMDeveloper/Private/RigVMModel/RigVMController.cpp b/Engine/Plugins/Runtime/RigVM/Source/RigVMDeveloper/Private/RigVMModel/RigVMController.cpp
--- a/Engine/Plugins/Runtime/RigVM/Source/RigVMDeveloper/Private/RigVMModel/RigVMController.cpp
+++ b/Engine/Plugins/Runtime/RigVM/Source/RigVMDeveloper/Private/RigVMModel/RigVMController.cpp
@@ -9641,6 +9641,18 @@ bool URigVMController::SetPinDefaultValue(URigVMPin* InPin, const FString& InDef
 		if(InPin->DefaultValue != ClampedDefaultValue)
 		{
 			InPin->DefaultValue = ClampedDefaultValue;
+			
+			// root trait pin store their default value in a separate property bag so that
+			// things like soft object ptr can be used and tracked in a uproperty
+			if (InPin->IsTraitPin() && InPin->IsRootPin())
+			{
+				FRigVMTraitDefaultValueStruct* DefaultValueStructPtr = InPin->GetNode()->TraitDefaultValues.Find(InPin->GetName());
+				if (ensure(DefaultValueStructPtr))
+				{
+					DefaultValueStructPtr->SetValue(InPin->DefaultValue);
+				}
+			}
+			
 			Notify(ERigVMGraphNotifType::PinDefaultValueChanged, InPin);
 			if (!bSuspendNotifications)
 			{
@@ -16077,6 +16089,12 @@ FName URigVMController::AddTrait(URigVMNode* InNode, UScriptStruct* InTraitScrip
 	}
 
 	InNode->TraitRootPinNames.Add(ValidTraitName.ToString());
+	
+	// root trait pin store their default value in a separate property bag so that
+	// things like soft object ptr can be used and tracked in a uproperty
+	FRigVMTraitDefaultValueStruct& TraitDefaultValueStruct = InNode->TraitDefaultValues.Add(ValidTraitName.ToString());
+	TraitDefaultValueStruct.Init(InTraitScriptStruct);
+	TraitDefaultValueStruct.SetValue(DefaultValue);
 
 	URigVMPin* TraitPin = NewObject<URigVMPin>(InNode, ValidTraitName);
 	const FString DisplayName = Trait->GetDisplayName();
@@ -16202,6 +16220,8 @@ bool URigVMController::RemoveTrait(URigVMNode* InNode, const FName& InTraitName,
 		return TraitNameString.Equals(TraitRootPinName, ESearchCase::CaseSensitive);
 	});
 
+	(void)InNode->TraitDefaultValues.Remove(TraitNameString);
+
 	RemovePin(TraitPin, bSetupUndoRedo, true);
 
 	if(bSetupUndoRedo)
@@ -18085,14 +18105,6 @@ void URigVMController::RepopulatePinsOnNode(const FRigVMRegistry& Registry, cons
 		ApplyPinStates(InNode, PinStates, RedirectedPinPaths);
 	}
 
-	InNode->TraitRootPinNames.Reset();
-	for (int32 Index = 0; Index < NodeData.NewPinInfos.Num(); Index++)
-	{
-		if (NodeData.NewPinInfos[Index].bIsTrait)
-		{
-			InNode->TraitRootPinNames.Add(NodeData.NewPinInfos[Index].Name.ToString());
-		}
-	}
 	InNode->UpdateTraitRootPinNames();
 
 	if (!LinkedPaths.IsEmpty())
diff --git a/Engine/Plugins/Runtime/RigVM/Source/RigVMDeveloper/Private/RigVMModel/RigVMNode.cpp b/Engine/Plugins/Runtime/RigVM/Source/RigVMDeveloper/Private/RigVMModel/RigVMNode.cpp
--- a/Engine/Plugins/Runtime/RigVM/Source/RigVMDeveloper/Private/RigVMModel/RigVMNode.cpp
+++ b/Engine/Plugins/Runtime/RigVM/Source/RigVMDeveloper/Private/RigVMModel/RigVMNode.cpp
@@ -6,6 +6,7 @@
 #include "RigVMModel/RigVMGraph.h"
 #include "RigVMModel/Nodes/RigVMLibraryNode.h"
 #include "RigVMModel/RigVMFunctionLibrary.h"
+#include "RigVMModel/RigVMTraitDefaultValueStruct.h"
 #include "RigVMCore/RigVMExecuteContext.h"
 #include "RigVMCore/RigVMStruct.h"
 #include "RigVMUserWorkflowRegistry.h"
@@ -35,6 +36,25 @@ URigVMNode::~URigVMNode()
 {
 }
 
+void URigVMNode::PostLoad()
+{
+	Super::PostLoad();
+
+	for (const FString& TraitRootPinName : TraitRootPinNames)
+	{
+		if (!TraitDefaultValues.Contains(TraitRootPinName))
+		{
+			if (URigVMPin* TraitPin = FindPin(TraitRootPinName))
+			{
+				UScriptStruct* TraitScriptStruct = TraitPin->GetScriptStruct();
+				FRigVMTraitDefaultValueStruct& TraitDefaultValueStruct = TraitDefaultValues.Add(TraitRootPinName);
+				TraitDefaultValueStruct.Init(TraitScriptStruct);
+				TraitDefaultValueStruct.SetValue(TraitPin->DefaultValue);
+			}
+		}
+	}
+}
+
 FString URigVMNode::GetNodePath(bool bRecursive) const
 {
 	if (bRecursive)
diff --git a/Engine/Plugins/Runtime/RigVM/Source/RigVMDeveloper/Private/RigVMModel/RigVMPin.cpp b/Engine/Plugins/Runtime/RigVM/Source/RigVMDeveloper/Private/RigVMModel/RigVMPin.cpp
--- a/Engine/Plugins/Runtime/RigVM/Source/RigVMDeveloper/Private/RigVMModel/RigVMPin.cpp
+++ b/Engine/Plugins/Runtime/RigVM/Source/RigVMDeveloper/Private/RigVMModel/RigVMPin.cpp
@@ -857,6 +857,17 @@ FString URigVMPin::GetDefaultValue(const URigVMPin::FPinOverride& InOverride, bo
 		if (SubPins.Num() > 0 || IsTraitPin())
 		{
 			FString FinalDefaultValue = DefaultValue;
+			
+			// root trait pin store their default value in a separate property bag so that
+			// things like soft object ptr can be used and tracked in a uproperty
+			if (IsTraitPin() && IsRootPin())
+			{
+				FRigVMTraitDefaultValueStruct* DefaultValueStructPtr = GetNode()->TraitDefaultValues.Find(GetName());
+				if (ensure(DefaultValueStructPtr))
+				{
+					FinalDefaultValue = DefaultValueStructPtr->GetValue();
+				}
+			}
 
 			for (const URigVMPin* SubPin : SubPins)
 			{
diff --git a/Engine/Plugins/Runtime/RigVM/Source/RigVMDeveloper/Private/RigVMModel/RigVMTraitDefaultValueStruct.cpp b/Engine/Plugins/Runtime/RigVM/Source/RigVMDeveloper/Private/RigVMModel/RigVMTraitDefaultValueStruct.cpp
new file mode 100644
--- /dev/null
+++ b/Engine/Plugins/Runtime/RigVM/Source/RigVMDeveloper/Private/RigVMModel/RigVMTraitDefaultValueStruct.cpp
@@ -0,0 +1,23 @@
+// Copyright Epic Games, Inc. All Rights Reserved.
+
+#include "RigVMModel/RigVMTraitDefaultValueStruct.h"
+
+const TCHAR* FRigVMTraitDefaultValueStruct::DefaultValuePropertyName = TEXT("DefaultValue");
+
+void FRigVMTraitDefaultValueStruct::Init(UScriptStruct* InTraitScriptStruct)
+{
+	PropertyBag.AddProperty(DefaultValuePropertyName, EPropertyBagPropertyType::Struct, InTraitScriptStruct);	
+}
+
+void FRigVMTraitDefaultValueStruct::SetValue(const FString& InDefaultValue)
+{
+	PropertyBag.SetValueSerializedString(DefaultValuePropertyName, InDefaultValue);	
+}
+
+FString FRigVMTraitDefaultValueStruct::GetValue() const
+{
+	TValueOrError<FString, EPropertyBagResult> TraitDefaultValue = PropertyBag.GetValueSerializedString(DefaultValuePropertyName);
+	
+	check(TraitDefaultValue.HasValue())
+	return TraitDefaultValue.GetValue();
+}
diff --git a/Engine/Plugins/Runtime/RigVM/Source/RigVMDeveloper/Public/RigVMModel/RigVMNode.h b/Engine/Plugins/Runtime/RigVM/Source/RigVMDeveloper/Public/RigVMModel/RigVMNode.h
--- a/Engine/Plugins/Runtime/RigVM/Source/RigVMDeveloper/Public/RigVMModel/RigVMNode.h
+++ b/Engine/Plugins/Runtime/RigVM/Source/RigVMDeveloper/Public/RigVMModel/RigVMNode.h
@@ -3,6 +3,7 @@
 #pragma once
 
 #include "RigVMPin.h"
+#include "RigVMTraitDefaultValueStruct.h"
 #include "RigVMCore/RigVM.h"
 #include "RigVMCore/RigVMStruct.h"
 #include "RigVMCore/RigVMUserWorkflow.h"
@@ -34,6 +35,8 @@ public:
 	// Default destructor
 	virtual ~URigVMNode();
 
+	void PostLoad() override;
+
 	// Returns the a . separated string containing all of the
 	// names used to reach this Node within the Graph.
 	// (for now this is the same as the Node's name)
@@ -429,6 +432,11 @@ protected:
 	UPROPERTY()
 	TArray<FString> TraitRootPinNames;
 
+	// root trait pin store their default value in a separate property bag so that
+	// things like soft object ptr can be used and tracked in a uproperty 
+	UPROPERTY()
+	TMap<FString, FRigVMTraitDefaultValueStruct> TraitDefaultValues;
+	
 private:
 
 	UPROPERTY()
diff --git a/Engine/Plugins/Runtime/RigVM/Source/RigVMDeveloper/Public/RigVMModel/RigVMTraitDefaultValueStruct.h b/Engine/Plugins/Runtime/RigVM/Source/RigVMDeveloper/Public/RigVMModel/RigVMTraitDefaultValueStruct.h
new file mode 100644
--- /dev/null
+++ b/Engine/Plugins/Runtime/RigVM/Source/RigVMDeveloper/Public/RigVMModel/RigVMTraitDefaultValueStruct.h
@@ -0,0 +1,19 @@
+// Copyright Epic Games, Inc. All Rights Reserved.
+#pragma once
+#include "StructUtils/PropertyBag.h"
+#include "RigVMTraitDefaultValueStruct.generated.h"
+
+USTRUCT()
+struct FRigVMTraitDefaultValueStruct
+{
+	GENERATED_BODY()
+
+	static const TCHAR* DefaultValuePropertyName;
+
+	UPROPERTY()
+	FInstancedPropertyBag PropertyBag;
+	
+	void Init(UScriptStruct* InTraitScriptStruct);
+	void SetValue(const FString& InDefaultValue);
+	FString GetValue() const;
+};
\ No newline at end of file
