RigVM: Avoid clobbering graph schemas that have already been set up in CreateModel()

And ensure that contained graph models have a valid schema class

diff --git a/Engine/Plugins/Runtime/RigVM/Source/RigVMDeveloper/Private/RigVMModel/RigVMClient.cpp b/Engine/Plugins/Runtime/RigVM/Source/RigVMDeveloper/Private/RigVMModel/RigVMClient.cpp
--- a/Engine/Plugins/Runtime/RigVM/Source/RigVMDeveloper/Private/RigVMModel/RigVMClient.cpp
+++ b/Engine/Plugins/Runtime/RigVM/Source/RigVMDeveloper/Private/RigVMModel/RigVMClient.cpp
@@ -605,6 +605,7 @@ TObjectPtr<URigVMGraph> FRigVMClient::CreateContainedGraphModel(URigVMCollapseNo
 	else
 	{
 		Model = NewObject<URigVMGraph>(CollapseNode, Name);
+		Model->SetSchemaClass(GetDefaultSchemaClass());
 	}
 
 	return Model;
@@ -627,7 +628,11 @@ void FRigVMClient::AddModel(URigVMGraph* InModel, bool bCreateController)
 		}
 	}
 
-	InModel->SetSchemaClass(GetDefaultSchemaClass());
+	if(InModel->GetSchemaClass() == nullptr)
+	{
+		InModel->SetSchemaClass(GetDefaultSchemaClass());
+	}
+
 	InModel->SetExecuteContextStruct(InModel->GetSchema()->GetExecuteContextStruct());
 
 	if(bCreateController)
