AnimNext: Module refactor

diff --git a/Engine/Plugins/Experimental/Workspace/Source/WorkspaceEditor/Private/Outliner/WorkspaceOutlinerMode.h b/Engine/Plugins/Experimental/Workspace/Source/WorkspaceEditor/Private/Outliner/WorkspaceOutlinerMode.h
--- a/Engine/Plugins/Experimental/Workspace/Source/WorkspaceEditor/Private/Outliner/WorkspaceOutlinerMode.h
+++ b/Engine/Plugins/Experimental/Workspace/Source/WorkspaceEditor/Private/Outliner/WorkspaceOutlinerMode.h
@@ -26,6 +26,10 @@ namespace UE::Workspace
 		virtual void OnItemSelectionChanged(FSceneOutlinerTreeItemPtr Item, ESelectInfo::Type SelectionType, const FSceneOutlinerItemSelection& Selection) override;
 		virtual bool CanCustomizeToolbar() const { return true; }
 		virtual ESelectionMode::Type GetSelectionMode() const { return ESelectionMode::Multi; }
+		virtual bool CanDelete() const override;
+		virtual bool CanRename() const override;
+		virtual bool CanRenameItem(const ISceneOutlinerTreeItem& Item) const override;
+		virtual void BindCommands(const TSharedRef<FUICommandList>& OutCommandList) override;
 	protected:
 		virtual TUniquePtr<ISceneOutlinerHierarchy> CreateHierarchy() override;		
 		// End ISceneOutlinerMode overrides
@@ -33,6 +37,11 @@ namespace UE::Workspace
 		void OnWorkspaceModified(UWorkspace* InWorkspace);
 		void ResetOutlinerSelection();
 
+		void Open();
+		void Delete();
+		bool CanDeleteItem(const ISceneOutlinerTreeItem& Item) const;
+		void Rename();
+
 		void OpenItems(TArrayView<const FSceneOutlinerTreeItemPtr> Items) const;
 		void DeleteItems(TArrayView<const FSceneOutlinerTreeItemPtr> Items) const;
 
@@ -40,5 +49,6 @@ namespace UE::Workspace
 	private:
 		TWeakObjectPtr<UWorkspace> WeakWorkspace;
 		TWeakPtr<IWorkspaceEditor> WeakWorkspaceEditor;
+		TSharedPtr<FUICommandList> CommandList;
 	};
 }
\ No newline at end of file
diff --git a/Engine/Plugins/Experimental/Workspace/Source/WorkspaceEditor/Private/Outliner/WorkspaceOutlinerTreeItem.cpp b/Engine/Plugins/Experimental/Workspace/Source/WorkspaceEditor/Private/Outliner/WorkspaceOutlinerTreeItem.cpp
--- a/Engine/Plugins/Experimental/Workspace/Source/WorkspaceEditor/Private/Outliner/WorkspaceOutlinerTreeItem.cpp
+++ b/Engine/Plugins/Experimental/Workspace/Source/WorkspaceEditor/Private/Outliner/WorkspaceOutlinerTreeItem.cpp
@@ -11,6 +11,7 @@
 #include "Widgets/SCompoundWidget.h"
 #include "Widgets/DeclarativeSyntaxSupport.h"
 #include "Widgets/Images/SImage.h"
+#include "Widgets/Text/SInlineEditableTextBlock.h"
 #include "Widgets/Text/STextBlock.h"
 #include "Widgets/Views/STableRow.h"
 #include "Widgets/Views/SListView.h"
@@ -49,7 +50,7 @@ namespace UE::Workspace
 							{
 								if (SharedTreeItem->ItemDetails.IsValid())
 								{
-									return SharedTreeItem->ItemDetails->GetItemIcon();
+									return SharedTreeItem->ItemDetails->GetItemIcon(SharedTreeItem->Export);
 								}
 							}
 							
@@ -59,14 +60,16 @@ namespace UE::Workspace
 					]
 				]
 				+ SHorizontalBox::Slot()
-				.FillWidth(1.0f)
+				.AutoWidth()
 				.VAlign(VAlign_Center)
 				.Padding(0.0f, 2.0f)
 				[
-					SNew(STextBlock)
+					SAssignNew(TextBlock, SInlineEditableTextBlock)
 					.Text(this, &SWorkspaceOutlinerTreelabel::GetDisplayText)
 					.HighlightText(SceneOutliner.GetFilterHighlightText())
 					.ColorAndOpacity(this, &SWorkspaceOutlinerTreelabel::GetForegroundColor)
+					.OnTextCommitted(this, &SWorkspaceOutlinerTreelabel::OnTextCommited)
+					.OnVerifyTextChanged(this, &SWorkspaceOutlinerTreelabel::OnVerifyTextChanged)
 				]
 			];
 		}
@@ -74,22 +77,48 @@ namespace UE::Workspace
 		FText GetDisplayText() const
 		{
 			if (const TSharedPtr<FWorkspaceOutlinerTreeItem> Item = TreeItem.Pin())
-			{				
+			{
 				return FText::FromString(Item->GetDisplayString());
 			}
 			return FText();
 		}
-		
+
+		void OnTextCommited(const FText& InLabel, ETextCommit::Type InCommitInfo)
+		{
+			if(InCommitInfo == ETextCommit::OnEnter)
+			{
+				if (const TSharedPtr<FWorkspaceOutlinerTreeItem> Item = TreeItem.Pin())
+				{
+					if(Item->ItemDetails.IsValid())
+					{
+						Item->ItemDetails->Rename(Item->Export, InLabel); 
+					}
+				}
+			}
+		}
+
+		bool OnVerifyTextChanged(const FText& InLabel, FText& OutErrorMessage)
+		{
+			if (const TSharedPtr<FWorkspaceOutlinerTreeItem> Item = TreeItem.Pin())
+			{
+				if(Item->ItemDetails.IsValid())
+				{
+					return Item->ItemDetails->ValidateName(Item->Export, InLabel, OutErrorMessage); 
+				}
+			}
+			return false;
+		}
+
 		virtual FSlateColor GetForegroundColor() const override
 		{
 			const TOptional<FLinearColor> BaseColor = FSceneOutlinerCommonLabelData::GetForegroundColor(*TreeItem.Pin());
 			return BaseColor.IsSet() ? BaseColor.GetValue() : FSlateColor::UseForeground();
 		}
 
-	private:
 		TWeakPtr<FWorkspaceOutlinerTreeItem> TreeItem;
-	};	
-	
+		TSharedPtr<SInlineEditableTextBlock> TextBlock;
+	};
+
 	FWorkspaceOutlinerTreeItem::FWorkspaceOutlinerTreeItem(const FItemData& InItemData) : ISceneOutlinerTreeItem(FWorkspaceOutlinerTreeItem::Type), Export(InItemData.Export)
 	{
 		ItemDetails = FWorkspaceEditorModule::GetOutlinerItemDetails(MakeOutlinerDetailsId(Export));
@@ -112,7 +141,9 @@ namespace UE::Workspace
 
 	TSharedRef<SWidget> FWorkspaceOutlinerTreeItem::GenerateLabelWidget(ISceneOutliner& Outliner, const STableRow<FSceneOutlinerTreeItemPtr>& InRow)
 	{
-		return SNew(SWorkspaceOutlinerTreelabel, *this, Outliner, InRow);
+		TSharedRef<SWorkspaceOutlinerTreelabel> LabelWidget = SNew(SWorkspaceOutlinerTreelabel, *this, Outliner, InRow);
+		RenameRequestEvent.BindSP(LabelWidget->TextBlock.Get(), &SInlineEditableTextBlock::EnterEditingMode);
+		return LabelWidget;
 	}
 
 	FString FWorkspaceOutlinerTreeItem::GetPackageName() const
@@ -125,7 +156,7 @@ namespace UE::Workspace
 		else if (Export.GetParentIdentifier() == NAME_None && Export.GetAssetPath().IsValid())
 		{
 			return Export.GetAssetPath().GetLongPackageName();
-		}		
+		}	
 		
 		return ISceneOutlinerTreeItem::GetPackageName();
 	}
diff --git a/Engine/Plugins/Experimental/Workspace/Source/WorkspaceEditor/Private/Outliner/WorkspaceOutlinerTreeItem.h b/Engine/Plugins/Experimental/Workspace/Source/WorkspaceEditor/Private/Outliner/WorkspaceOutlinerTreeItem.h
--- a/Engine/Plugins/Experimental/Workspace/Source/WorkspaceEditor/Private/Outliner/WorkspaceOutlinerTreeItem.h
+++ b/Engine/Plugins/Experimental/Workspace/Source/WorkspaceEditor/Private/Outliner/WorkspaceOutlinerTreeItem.h
@@ -17,7 +17,7 @@ namespace UE::Workspace
 		};
 				
 		FWorkspaceOutlinerTreeItem(const FItemData& InItemData);
-		virtual ~FWorkspaceOutlinerTreeItem() override {}
+		virtual ~FWorkspaceOutlinerTreeItem() override = default;
 
 		// Begin ISceneOutlinerTreeItem overrides
 		virtual bool IsValid() const override;
diff --git a/Engine/Plugins/Experimental/Workspace/Source/WorkspaceEditor/Private/SWorkspaceTabWrapper.cpp b/Engine/Plugins/Experimental/Workspace/Source/WorkspaceEditor/Private/SWorkspaceTabWrapper.cpp
--- a/Engine/Plugins/Experimental/Workspace/Source/WorkspaceEditor/Private/SWorkspaceTabWrapper.cpp
+++ b/Engine/Plugins/Experimental/Workspace/Source/WorkspaceEditor/Private/SWorkspaceTabWrapper.cpp
@@ -229,14 +229,14 @@ void SWorkspaceTabWrapper::RebuildBreadcrumbTrail() const
 		if (const TSharedPtr<UE::Workspace::FWorkspaceEditor> SharedWorkspaceEditor = WeakWorkspaceEditor.Pin())
 		{
 			const UE::Workspace::FWorkspaceEditorModule& WorkspaceEditorModule = FModuleManager::LoadModuleChecked<UE::Workspace::FWorkspaceEditorModule>("WorkspaceEditor");
-			if(const UE::Workspace::FObjectDocumentArgs* DocumentArgs = WorkspaceEditorModule.FindObjectDocumentType(DocumentID->GetClass()->GetClassPathName()))
+			if(const UE::Workspace::FObjectDocumentArgs* DocumentArgs = WorkspaceEditorModule.FindObjectDocumentType(DocumentID))
 			{
 				if(DocumentArgs->OnGetDocumentBreadcrumbTrail.IsBound())
 				{
 					DocumentArgs->OnGetDocumentBreadcrumbTrail.Execute(UE::Workspace::FWorkspaceEditorContext(SharedWorkspaceEditor.ToSharedRef(), DocumentID), Breadcrumbs);
 				}
 			}
-		}			
+		}
 	}
 
 	// Widgets have to be added in reverse order
@@ -254,7 +254,7 @@ const FSlateBrush* SWorkspaceTabWrapper::GetTabIcon() const
 		if (const TSharedPtr<UE::Workspace::FWorkspaceEditor> SharedWorkspaceEditor = WeakWorkspaceEditor.Pin())
 		{
 			const UE::Workspace::FWorkspaceEditorModule& WorkspaceEditorModule = FModuleManager::LoadModuleChecked<UE::Workspace::FWorkspaceEditorModule>("WorkspaceEditor");
-			if(const UE::Workspace::FObjectDocumentArgs* DocumentArgs = WorkspaceEditorModule.FindObjectDocumentType(DocumentID->GetClass()->GetClassPathName()))
+			if(const UE::Workspace::FObjectDocumentArgs* DocumentArgs = WorkspaceEditorModule.FindObjectDocumentType(DocumentID))
 			{
 				if(DocumentArgs->OnGetTabIcon.IsBound())
 				{
diff --git a/Engine/Plugins/Experimental/Workspace/Source/WorkspaceEditor/Private/SWorkspaceView.cpp b/Engine/Plugins/Experimental/Workspace/Source/WorkspaceEditor/Private/SWorkspaceView.cpp
--- a/Engine/Plugins/Experimental/Workspace/Source/WorkspaceEditor/Private/SWorkspaceView.cpp
+++ b/Engine/Plugins/Experimental/Workspace/Source/WorkspaceEditor/Private/SWorkspaceView.cpp
@@ -1,4 +1,4 @@
-﻿// Copyright Epic Games, Inc. All Rights Reserved.
+// Copyright Epic Games, Inc. All Rights Reserved.
 
 #include "SWorkspaceView.h"
 #include "Workspace.h"
@@ -10,6 +10,7 @@
 #include "ScopedTransaction.h"
 #include "SPositiveActionButton.h"
 #include "SSceneOutliner.h"
+#include "WorkspaceAssetRegistryInfo.h"
 #include "WorkspaceSchema.h"
 #include "Framework/Commands/GenericCommands.h"
 #include "Outliner/WorkspaceOutlinerColumns.h"
@@ -104,7 +105,7 @@ void SWorkspaceView::Construct(const FArguments& InArgs, UWorkspace* InWorkspace
 		InitOptions.ModeFactory = FCreateSceneOutlinerMode::CreateLambda([this, WeakWorkspaceEditor=InWorkspaceEditor.ToWeakPtr()](SSceneOutliner* InOutliner) { return new UE::Workspace::FWorkspaceOutlinerMode(UE::Workspace::FWorkspaceOutlinerMode(InOutliner, Workspace, WeakWorkspaceEditor)); });
 	}
 	SceneWorkspaceOutliner = SNew(SWorkspaceOutliner, InitOptions, Workspace);
-	
+
 	TWeakObjectPtr<UWorkspace> WeakWorkspace = Workspace;
 	ChildSlot
 	[
@@ -167,6 +168,17 @@ void SWorkspaceView::Construct(const FArguments& InArgs, UWorkspace* InWorkspace
 	];
 }
 
+void SWorkspaceView::SelectObject(UObject* InObject)
+{
+	FWorkspaceOutlinerItemExport Export(InObject->GetFName(), InObject);
+
+	FSceneOutlinerTreeItemPtr FoundItem = SceneWorkspaceOutliner->GetTreeItem(FSceneOutlinerTreeItemID(GetTypeHash(Export)));
+	if(FoundItem.IsValid())
+	{
+		SceneWorkspaceOutliner->SetItemSelection(FoundItem, true, ESelectInfo::OnMouseClick);	// Not direct, so we get callbacks
+	}
+}
+
 }
 
 #undef LOCTEXT_NAMESPACE
\ No newline at end of file
diff --git a/Engine/Plugins/Experimental/Workspace/Source/WorkspaceEditor/Private/SWorkspaceView.h b/Engine/Plugins/Experimental/Workspace/Source/WorkspaceEditor/Private/SWorkspaceView.h
--- a/Engine/Plugins/Experimental/Workspace/Source/WorkspaceEditor/Private/SWorkspaceView.h
+++ b/Engine/Plugins/Experimental/Workspace/Source/WorkspaceEditor/Private/SWorkspaceView.h
@@ -19,6 +19,9 @@ class SWorkspaceView : public SCompoundWidget
 	SLATE_END_ARGS()
 
 	void Construct(const FArguments& InArgs, UWorkspace* InWorkspace, TSharedRef<UE::Workspace::IWorkspaceEditor> InWorkspaceEditor);
+
+	void SelectObject(UObject* InObject);
+	
 private:
 	UWorkspace* Workspace = nullptr;
 	TSharedPtr<SWorkspaceOutliner> SceneWorkspaceOutliner;
diff --git a/Engine/Plugins/Experimental/Workspace/Source/WorkspaceEditor/Private/WorkspaceEditor.cpp b/Engine/Plugins/Experimental/Workspace/Source/WorkspaceEditor/Private/WorkspaceEditor.cpp
--- a/Engine/Plugins/Experimental/Workspace/Source/WorkspaceEditor/Private/WorkspaceEditor.cpp
+++ b/Engine/Plugins/Experimental/Workspace/Source/WorkspaceEditor/Private/WorkspaceEditor.cpp
@@ -169,6 +169,9 @@ void FWorkspaceEditor::CreateWidgets()
 	StandaloneDefaultLayout->ProcessExtensions(*LayoutExtender.Get());
 
 	WorkspaceView = SNew(SWorkspaceView, Workspace, StaticCastSharedRef<UE::Workspace::IWorkspaceEditor>(AsShared()));
+
+	OnOutlinerSelectionChangedDelegate.AddSP(this, &FWorkspaceEditor::HandleOutlinerSelectionChanged);
+	
 	BindCommands();
 }
 
@@ -193,7 +196,7 @@ void FWorkspaceEditor::RestoreEditedObjectState()
 			{
 				const TSharedRef<SWorkspaceTabWrapper> TabWrapper = StaticCastSharedRef<SWorkspaceTabWrapper>(DockTab->GetContent());	
 				FWorkspaceEditorModule& WorkspaceEditorModule = FModuleManager::LoadModuleChecked<FWorkspaceEditorModule>("WorkspaceEditor");
-				const FObjectDocumentArgs* DocumentArgs = WorkspaceEditorModule.FindObjectDocumentType(Object->GetClass()->GetClassPathName());
+				const FObjectDocumentArgs* DocumentArgs = WorkspaceEditorModule.FindObjectDocumentType(Object);
 				if(DocumentArgs != nullptr && DocumentArgs->OnSetDocumentState.IsBound())
 				{
 					DocumentArgs->OnSetDocumentState.Execute(FWorkspaceEditorContext(SharedThis(this), Object), TabWrapper->GetContent(), DocumentState);
@@ -230,7 +233,7 @@ TSharedPtr<SDockTab> FWorkspaceEditor::OpenDocument(const UObject* InForObject,
 	if(InCause != FDocumentTracker::RestorePreviousDocument)
 	{
 		const FWorkspaceEditorModule& WorkspaceEditorModule = FModuleManager::LoadModuleChecked<FWorkspaceEditorModule>("WorkspaceEditor");
-		const FObjectDocumentArgs* DocumentArgs = WorkspaceEditorModule.FindObjectDocumentType(InForObject->GetClass()->GetClassPathName());
+		const FObjectDocumentArgs* DocumentArgs = WorkspaceEditorModule.FindObjectDocumentType(InForObject);
 		if(DocumentArgs != nullptr && DocumentArgs->OnGetDocumentState.IsBound())
 		{
 			const TSharedRef<SWorkspaceTabWrapper> TabWrapper = StaticCastSharedRef<SWorkspaceTabWrapper>(NewTab->GetContent());
@@ -254,6 +257,14 @@ void FWorkspaceEditor::OpenAssets(TConstArrayView<FAssetData> InAssets)
 			OpenDocument(LoadedAsset, FDocumentTracker::EOpenDocumentCause::NavigatingCurrentDocument);
 		}
 	}
+
+	if(InAssets.Num() > 0)
+	{
+		if(UObject* LoadedAsset = InAssets.Last().GetAsset())
+		{
+			WorkspaceView->SelectObject(LoadedAsset);
+		}
+	}
 }
 
 void FWorkspaceEditor::OpenObjects(TConstArrayView<UObject*> InObjects)
@@ -296,6 +307,17 @@ UWorkspaceSchema* FWorkspaceEditor::GetSchema() const
 	return Workspace.Get() != nullptr ? Workspace->GetSchema() : nullptr;
 }
 
+bool FWorkspaceEditor::GetOutlinerSelection(TArray<FWorkspaceOutlinerItemExport>& OutExports) const
+{
+	OutExports = LastSelectedExports;
+	return LastSelectedExports.Num() > 0;
+}
+
+IWorkspaceEditor::FOnOutlinerSelectionChanged& FWorkspaceEditor::OnOutlinerSelectionChanged()
+{
+	return OnOutlinerSelectionChangedDelegate;
+}
+
 void FWorkspaceEditor::SetGlobalSelection(FGlobalSelectionId SelectionId, FOnClearGlobalSelection OnClearSelectionDelegate)
 {
 	// Only execute if widget is still valid, and it is not the same as the previous call 
@@ -330,9 +352,13 @@ const TObjectPtr<UObject> FWorkspaceEditor::GetFocussedAssetOfClass(const TObjec
 	return nullptr;
 }
 
+void FWorkspaceEditor::HandleOutlinerSelectionChanged(TConstArrayView<FWorkspaceOutlinerItemExport> InExports)
+{
+	LastSelectedExports = InExports;
+}
+
 void FWorkspaceEditor::BindCommands()
-{	
-	FWorkspaceAssetEditorCommands::Register();
+{
 	const FWorkspaceAssetEditorCommands& Commands = FWorkspaceAssetEditorCommands::Get();
 
 	ToolkitCommands->MapAction
diff --git a/Engine/Plugins/Experimental/Workspace/Source/WorkspaceEditor/Private/WorkspaceEditor.h b/Engine/Plugins/Experimental/Workspace/Source/WorkspaceEditor/Private/WorkspaceEditor.h
--- a/Engine/Plugins/Experimental/Workspace/Source/WorkspaceEditor/Private/WorkspaceEditor.h
+++ b/Engine/Plugins/Experimental/Workspace/Source/WorkspaceEditor/Private/WorkspaceEditor.h
@@ -5,6 +5,7 @@
 #include "CoreMinimal.h"
 #include "StructUtils/InstancedStruct.h"
 #include "IWorkspaceEditor.h"
+#include "WorkspaceAssetRegistryInfo.h"
 #include "WorkflowOrientedApp/WorkflowTabManager.h"
 
 struct FWorkspaceDocumentState;
@@ -97,9 +98,12 @@ private:
 	virtual void SetDetailsObjects(const TArray<UObject*>& InObjects) override;
 	virtual void RefreshDetails() override;
 	virtual UWorkspaceSchema* GetSchema() const override;
+	virtual bool GetOutlinerSelection(TArray<FWorkspaceOutlinerItemExport>& OutExports) const override;
+	virtual FOnOutlinerSelectionChanged& OnOutlinerSelectionChanged() override;
 	virtual void SetGlobalSelection(FGlobalSelectionId SelectionId, FOnClearGlobalSelection OnClearSelectionDelegate) override;
 	virtual void SetFocussedAsset(const TObjectPtr<UObject> InAsset) override;	
 	virtual const TObjectPtr<UObject> GetFocussedAssetOfClass(const TObjectPtr<UClass> AssetClass) const override;	
+	void HandleOutlinerSelectionChanged(TConstArrayView<FWorkspaceOutlinerItemExport> InExports);
 
 	void BindCommands();
 
@@ -143,6 +147,8 @@ private:
 
 	FGlobalSelectionId LastGlobalSelectionId = nullptr;
 	FOnClearGlobalSelection LastOnClearSelectionDelegate = nullptr;
+	FOnOutlinerSelectionChanged OnOutlinerSelectionChangedDelegate;
+	TArray<FWorkspaceOutlinerItemExport> LastSelectedExports;
 };
 
 }
\ No newline at end of file
diff --git a/Engine/Plugins/Experimental/Workspace/Source/WorkspaceEditor/Private/WorkspaceEditorCommands.cpp b/Engine/Plugins/Experimental/Workspace/Source/WorkspaceEditor/Private/WorkspaceEditorCommands.cpp
--- a/Engine/Plugins/Experimental/Workspace/Source/WorkspaceEditor/Private/WorkspaceEditorCommands.cpp
+++ b/Engine/Plugins/Experimental/Workspace/Source/WorkspaceEditor/Private/WorkspaceEditorCommands.cpp
@@ -28,6 +28,7 @@ namespace UE::Workspace
 		UI_COMMAND(NavigateForward, "Navigate Forward", "Moves forwards to previous location",
 			EUserInterfaceActionType::Button, FInputChord(EKeys::ThumbMouseButton2), FInputChord(EKeys::Right, EModifierKey::Alt));
 		UI_COMMAND(SaveAssetEntries, "Save Assets", "Saves all loaded and modified assets in the workspace", EUserInterfaceActionType::Button, FInputChord());
+		UI_COMMAND(Open, "Open", "Open items", EUserInterfaceActionType::Button, FInputChord(EKeys::Enter))
 	}
 
 }  // namespace UE::Workspace
diff --git a/Engine/Plugins/Experimental/Workspace/Source/WorkspaceEditor/Private/WorkspaceEditorCommands.h b/Engine/Plugins/Experimental/Workspace/Source/WorkspaceEditor/Private/WorkspaceEditorCommands.h
--- a/Engine/Plugins/Experimental/Workspace/Source/WorkspaceEditor/Private/WorkspaceEditorCommands.h
+++ b/Engine/Plugins/Experimental/Workspace/Source/WorkspaceEditor/Private/WorkspaceEditorCommands.h
@@ -12,11 +12,12 @@ namespace UE::Workspace
 	{
 	public:
 		FWorkspaceAssetEditorCommands();
-		virtual void RegisterCommands() override;	
+		virtual void RegisterCommands() override;
 
-		TSharedPtr<FUICommandInfo> NavigateBackward;		
+		TSharedPtr<FUICommandInfo> NavigateBackward;
 		TSharedPtr<FUICommandInfo> NavigateForward;
 		TSharedPtr<FUICommandInfo> SaveAssetEntries;
+		TSharedPtr<FUICommandInfo> Open;
 	};
 
 }  // namespace UE::Workspace
diff --git a/Engine/Plugins/Experimental/Workspace/Source/WorkspaceEditor/Private/WorkspaceEditorModule.cpp b/Engine/Plugins/Experimental/Workspace/Source/WorkspaceEditor/Private/WorkspaceEditorModule.cpp
--- a/Engine/Plugins/Experimental/Workspace/Source/WorkspaceEditor/Private/WorkspaceEditorModule.cpp
+++ b/Engine/Plugins/Experimental/Workspace/Source/WorkspaceEditor/Private/WorkspaceEditorModule.cpp
@@ -26,6 +26,26 @@ namespace UE::Workspace
 
 TMap<FOutlinerItemDetailsId, TSharedPtr<IWorkspaceOutlinerItemDetails>> FWorkspaceEditorModule::OutlinerItemDetails;
 
+FWorkspaceEditorContext::FWorkspaceEditorContext(const TSharedRef<IWorkspaceEditor>& InWorkspaceEditor, UObject* InObject)
+	: WorkspaceEditor(InWorkspaceEditor)
+	, Object(InObject)
+{
+	FWorkspaceEditorModule& WorkspaceEditorModule = FModuleManager::Get().LoadModuleChecked<FWorkspaceEditorModule>("WorkspaceEditor");
+	const FObjectDocumentArgs* FoundArgs = WorkspaceEditorModule.ObjectDocumentArgs.Find(InObject->GetClass()->GetClassPathName());
+	while(FoundArgs && FoundArgs->OnRedirectWorkspaceContext.IsBound())
+	{
+		InObject = FoundArgs->OnRedirectWorkspaceContext.Execute(InObject);
+		FoundArgs = WorkspaceEditorModule.ObjectDocumentArgs.Find(InObject->GetClass()->GetClassPathName());
+	}
+
+	Object = InObject;
+}
+
+void FWorkspaceEditorModule::StartupModule()
+{
+	FWorkspaceAssetEditorCommands::Register();
+}
+
 void FWorkspaceEditorModule::RegisterObjectDocumentType(const FTopLevelAssetPath& InClassPath, const FObjectDocumentArgs& InParams)
 {
 	ensure(InParams.SpawnLocation != NAME_None);
@@ -36,7 +56,7 @@ void FWorkspaceEditorModule::RegisterObjectDocumentType(const FTopLevelAssetPath
 
 void FWorkspaceEditorModule::UnregisterObjectDocumentType(const FTopLevelAssetPath& InClassPath)
 {
-	if(const FObjectDocumentArgs* ExistingType = FindObjectDocumentType(InClassPath))
+	if(const FObjectDocumentArgs* ExistingType = ObjectDocumentArgs.Find(InClassPath))
 	{
 		DocumentAreaMap.FindChecked(ExistingType->SpawnLocation).Remove(InClassPath);
 	}
@@ -148,7 +168,7 @@ bool FWorkspaceEditorModule::GetExportedAssetsForWorkspace(const FAssetData& InW
 	return OutExports.Assets.Num() > 0;
 }
 
-void FWorkspaceEditorModule::OpenWorkspaceForObject(UObject* InObject, EOpenWorkspaceMethod InOpenMethod, const TSubclassOf<UWorkspaceFactory> WorkSpaceFactoryClass/*= UWorkspaceFactory::StaticClass()*/)
+ IWorkspaceEditor* FWorkspaceEditorModule::OpenWorkspaceForObject(UObject* InObject, EOpenWorkspaceMethod InOpenMethod, const TSubclassOf<UWorkspaceFactory> WorkSpaceFactoryClass/*= UWorkspaceFactory::StaticClass()*/)
 {
 	FAssetRegistryModule& AssetRegistryModule = FModuleManager::LoadModuleChecked<FAssetRegistryModule>(TEXT("AssetRegistry"));
 
@@ -238,11 +258,20 @@ void FWorkspaceEditorModule::OpenWorkspaceForObject(UObject* InObject, EOpenWork
 	{
 		WorkspaceEditor->OpenAssets({InObject});
 	}
+
+	return WorkspaceEditor;
 }
 
-const FObjectDocumentArgs* FWorkspaceEditorModule::FindObjectDocumentType(const FTopLevelAssetPath& InClassPath) const
+const FObjectDocumentArgs* FWorkspaceEditorModule::FindObjectDocumentType(const UObject* InObject) const
 {
-	return ObjectDocumentArgs.Find(InClassPath);
+	const FObjectDocumentArgs* FoundArgs = ObjectDocumentArgs.Find(InObject->GetClass()->GetClassPathName());
+	while(FoundArgs && FoundArgs->OnRedirectWorkspaceContext.IsBound())
+	{
+		InObject = FoundArgs->OnRedirectWorkspaceContext.Execute(const_cast<UObject*>(InObject));
+		FoundArgs = ObjectDocumentArgs.Find(InObject->GetClass()->GetClassPathName());
+	}
+
+	return FoundArgs;
 }
 
 TArray<FTopLevelAssetPath> FWorkspaceEditorModule::GetAllowedObjectTypesForArea(FName InSpawnLocation) const
diff --git a/Engine/Plugins/Experimental/Workspace/Source/WorkspaceEditor/Private/WorkspaceEditorModule.h b/Engine/Plugins/Experimental/Workspace/Source/WorkspaceEditor/Private/WorkspaceEditorModule.h
--- a/Engine/Plugins/Experimental/Workspace/Source/WorkspaceEditor/Private/WorkspaceEditorModule.h
+++ b/Engine/Plugins/Experimental/Workspace/Source/WorkspaceEditor/Private/WorkspaceEditorModule.h
@@ -23,19 +23,22 @@ namespace UE::Workspace
 class FWorkspaceEditorModule : public IWorkspaceEditorModule
 {
 private:
+	// IModuleInterface interface
+	virtual void StartupModule() override;
+	
 	// IWorkspaceModule interface
 	virtual void RegisterObjectDocumentType(const FTopLevelAssetPath& InClassPath, const FObjectDocumentArgs& InParams) override;
 	virtual void UnregisterObjectDocumentType(const FTopLevelAssetPath& InClassPath) override;
 	virtual FObjectDocumentArgs CreateGraphDocumentArgs(const FGraphDocumentWidgetArgs& InArgs) override;
-	virtual void OpenWorkspaceForObject(UObject* InObject, EOpenWorkspaceMethod InOpenMethod, const TSubclassOf<UWorkspaceFactory> WorkSpaceFactoryClass) override;	
+	virtual IWorkspaceEditor* OpenWorkspaceForObject(UObject* InObject, EOpenWorkspaceMethod InOpenMethod, const TSubclassOf<UWorkspaceFactory> WorkSpaceFactoryClass) override;	
 	virtual FOnRegisterDetailCustomizations& OnRegisterWorkspaceDetailsCustomization() override;
 	virtual void RegisterWorkspaceItemDetails(const FOutlinerItemDetailsId& InItemDetailsId, TSharedPtr<IWorkspaceOutlinerItemDetails> InItemDetails) override;	
 	virtual void UnregisterWorkspaceItemDetails(const FOutlinerItemDetailsId& InItemDetails) override;
 	virtual FOnRegisterTabs& OnRegisterTabsForEditor() override { return RegisterTabsForEditor; }
 	virtual FOnExtendTabs& OnExtendTabs() override { return ExtendTabsForEditor; }
 
-	// Find an existing registered object document type
-	const FObjectDocumentArgs* FindObjectDocumentType(const FTopLevelAssetPath& InClassPath) const;
+	// Find an existing registered object document type. Note this redirects based on FObjectDocumentArgs::OnRedirectWorkspaceContext
+	const FObjectDocumentArgs* FindObjectDocumentType(const UObject* InObject) const;
 
 	// Find the set of allowed object types for the specified spawn location
 	TArray<FTopLevelAssetPath> GetAllowedObjectTypesForArea(FName InSpawnLocation) const;
@@ -70,6 +73,7 @@ private:
 	friend class FWorkspaceOutlinerSourceControlColumn;
 	friend class FWorkspaceOutlinerFileStateColumn;
 	friend SWorkspaceTabWrapper;
+	friend FWorkspaceEditorContext;
 };
 
 }
diff --git a/Engine/Plugins/Experimental/Workspace/Source/WorkspaceEditor/Public/IWorkspaceEditor.h b/Engine/Plugins/Experimental/Workspace/Source/WorkspaceEditor/Public/IWorkspaceEditor.h
--- a/Engine/Plugins/Experimental/Workspace/Source/WorkspaceEditor/Public/IWorkspaceEditor.h
+++ b/Engine/Plugins/Experimental/Workspace/Source/WorkspaceEditor/Public/IWorkspaceEditor.h
@@ -6,6 +6,7 @@
 
 class UWorkspaceAssetEditor;
 class UWorkspaceSchema;
+struct FWorkspaceOutlinerItemExport;
 
 namespace UE::Workspace
 {
@@ -45,6 +46,14 @@ public:
 	template<typename AssetClass>
 	TObjectPtr<AssetClass> GetFocussedAsset() const { return GetFocussedAssetOfClass(AssetClass::StaticClass()); }
 	TObjectPtr<UObject> GetFocussedAsset() const { return GetFocussedAssetOfClass(UObject::StaticClass()); }
+
+	// Get the current single selection of the outliner.
+	// @return true if a single selection is active
+	virtual bool GetOutlinerSelection(TArray<FWorkspaceOutlinerItemExport>& OutExports) const = 0;
+
+	// Delegate fired when selection changes in the workspace outliner
+	using FOnOutlinerSelectionChanged = TMulticastDelegate<void(TConstArrayView<FWorkspaceOutlinerItemExport> InExports)>;
+	virtual FOnOutlinerSelectionChanged& OnOutlinerSelectionChanged() = 0;
 };
 
 }
\ No newline at end of file
diff --git a/Engine/Plugins/Experimental/Workspace/Source/WorkspaceEditor/Public/IWorkspaceEditorModule.h b/Engine/Plugins/Experimental/Workspace/Source/WorkspaceEditor/Public/IWorkspaceEditorModule.h
--- a/Engine/Plugins/Experimental/Workspace/Source/WorkspaceEditor/Public/IWorkspaceEditorModule.h
+++ b/Engine/Plugins/Experimental/Workspace/Source/WorkspaceEditor/Public/IWorkspaceEditorModule.h
@@ -44,10 +44,7 @@ namespace WorkspaceTabs
 // Context passed to workspace editor delegates
 struct FWorkspaceEditorContext
 {
-	FWorkspaceEditorContext(const TSharedRef<IWorkspaceEditor>& InWorkspaceEditor, UObject* InObject)
-		: WorkspaceEditor(InWorkspaceEditor)
-		, Object(InObject)
-	{}
+	FWorkspaceEditorContext(const TSharedRef<IWorkspaceEditor>& InWorkspaceEditor, UObject* InObject);
 
 	// The current workspace editor
 	TSharedRef<IWorkspaceEditor> WorkspaceEditor;
@@ -69,6 +66,8 @@ struct FWorkspaceBreadcrumb : TSharedFromThis<FWorkspaceBreadcrumb>
 	FOnSaveBreadcrumb OnSave;
 };
 
+using FOnRedirectWorkspaceContext = TDelegate<UObject*(UObject*)>;
+
 using FOnMakeDocumentWidget = TDelegate<TSharedRef<SWidget>(const FWorkspaceEditorContext&)>;
 
 using FOnGetTabIcon = TDelegate<const FSlateBrush*(const FWorkspaceEditorContext&)>;
@@ -91,6 +90,13 @@ struct FObjectDocumentArgs
 		, SpawnLocation(InSpawnLocation)
 	{}
 
+	FObjectDocumentArgs(FOnRedirectWorkspaceContext InOnRedirectWorkspaceContext)
+		: OnRedirectWorkspaceContext(InOnRedirectWorkspaceContext)
+	{}
+	
+	// Delegate called to redirect the context to another document object (e.g. a subobject)
+	FOnRedirectWorkspaceContext OnRedirectWorkspaceContext;
+
 	// Delegate called to generate a widget for the supplied object
 	FOnMakeDocumentWidget OnMakeDocumentWidget;
 
@@ -181,7 +187,7 @@ class IWorkspaceEditorModule : public IModuleInterface
 {
 public:
 	// Open an object inside a workspace editor.
-	virtual void OpenWorkspaceForObject(UObject* InObject, EOpenWorkspaceMethod InOpenMethod, const TSubclassOf<UWorkspaceFactory> WorkSpaceFactoryClass = UWorkspaceFactory::StaticClass()) = 0;
+	virtual IWorkspaceEditor* OpenWorkspaceForObject(UObject* InObject, EOpenWorkspaceMethod InOpenMethod, const TSubclassOf<UWorkspaceFactory> WorkSpaceFactoryClass = UWorkspaceFactory::StaticClass()) = 0;
 
 	// Register a widget factory method to spawn for a particular class
 	virtual void RegisterObjectDocumentType(const FTopLevelAssetPath& InClassPath, const FObjectDocumentArgs& InArgs) = 0;
diff --git a/Engine/Plugins/Experimental/Workspace/Source/WorkspaceEditor/Public/IWorkspaceOutlinerItemDetails.h b/Engine/Plugins/Experimental/Workspace/Source/WorkspaceEditor/Public/IWorkspaceOutlinerItemDetails.h
--- a/Engine/Plugins/Experimental/Workspace/Source/WorkspaceEditor/Public/IWorkspaceOutlinerItemDetails.h
+++ b/Engine/Plugins/Experimental/Workspace/Source/WorkspaceEditor/Public/IWorkspaceOutlinerItemDetails.h
@@ -21,8 +21,13 @@ class IWorkspaceOutlinerItemDetails : public TSharedFromThis<IWorkspaceOutlinerI
 {
 public:
     virtual ~IWorkspaceOutlinerItemDetails() = default;
-    virtual const FSlateBrush* GetItemIcon() const { return nullptr; }
+    virtual const FSlateBrush* GetItemIcon(const FWorkspaceOutlinerItemExport& Export) const { return nullptr; }
     virtual void HandleDoubleClick(const FToolMenuContext& ToolMenuContext) const {}
+	virtual bool CanDelete(const FWorkspaceOutlinerItemExport& Export) const { return true; }
+	virtual void Delete(TConstArrayView<FWorkspaceOutlinerItemExport> Exports) const {}
+    virtual bool CanRename(const FWorkspaceOutlinerItemExport& Export) const { return false; }
+    virtual void Rename(const FWorkspaceOutlinerItemExport& Export, const FText& InName) const {}
+    virtual bool ValidateName(const FWorkspaceOutlinerItemExport& Export, const FText& InName, FText& OutErrorMessage) const { return false; }
     virtual UPackage* GetPackage(const FWorkspaceOutlinerItemExport& Export) const { return nullptr; }
 };
 
diff --git a/Engine/Plugins/Experimental/Workspace/Source/WorkspaceEditor/Public/WorkspaceItemMenuContext.h b/Engine/Plugins/Experimental/Workspace/Source/WorkspaceEditor/Public/WorkspaceItemMenuContext.h
--- a/Engine/Plugins/Experimental/Workspace/Source/WorkspaceEditor/Public/WorkspaceItemMenuContext.h
+++ b/Engine/Plugins/Experimental/Workspace/Source/WorkspaceEditor/Public/WorkspaceItemMenuContext.h
@@ -5,10 +5,13 @@
 #include "WorkspaceAssetRegistryInfo.h"
 #include "WorkspaceItemMenuContext.generated.h"
 
+class FUICommandList;
+
 UCLASS()
 class WORKSPACEEDITOR_API UWorkspaceItemMenuContext : public UObject
 {
 	GENERATED_BODY()
 public:
 	TArray<FWorkspaceOutlinerItemExport> SelectedExports;
+	TWeakPtr<FUICommandList> WeakCommandList;
 };
diff --git a/Engine/Plugins/Runtime/RigVM/Source/RigVM/Private/RigVMCore/RigVM.cpp b/Engine/Plugins/Runtime/RigVM/Source/RigVM/Private/RigVMCore/RigVM.cpp
--- a/Engine/Plugins/Runtime/RigVM/Source/RigVM/Private/RigVMCore/RigVM.cpp
+++ b/Engine/Plugins/Runtime/RigVM/Source/RigVM/Private/RigVMCore/RigVM.cpp
@@ -2142,21 +2142,45 @@ ERigVMExecuteResult URigVM::ExecuteInstructions(FRigVMExtendedExecuteContext& Co
 			case ERigVMOpCode::SetupTraits:
 			{
 				ContextPublicData.Traits.Reset();
-				
+				ContextPublicData.AdditionalTraitMemoryHandles.Reset();
+
 				const TArray<int32>& TraitList = *(TArray<int32>*)Context.CachedMemoryHandles[FirstHandleForInstruction[ContextPublicData.InstructionIndex]].GetData();
+				ContextPublicData.Traits.Reserve(TraitList.Num());
+				ContextPublicData.AdditionalTraitMemoryHandles.Reserve(TraitList.Num());
+				int32 AdditionalStartIndex = INDEX_NONE;
+				int32 AdditionalNum = 0;
 				for(const int32 TraitPropertyIndex : TraitList)
 				{
+					// Properties from programmatic pins that are added to the trait index list are added before their respective trait.
+					// AdditionalMemoryHandles is accumulated prior to building the FRigVMTraitScope, and each FRigVMTraitScope has a view into
+					// the overall ContextPublicData.AdditionalTraitMemoryHandles buffer.
 					if(Context.WorkMemoryStorage.GetProperties().IsValidIndex(TraitPropertyIndex))
 					{
 						const FProperty* Property = Context.WorkMemoryStorage.GetProperties()[TraitPropertyIndex];
-						if(const FStructProperty* StructProperty = CastField<FStructProperty>(Property))
+						const FStructProperty* StructProperty = CastField<FStructProperty>(Property);
+						if(StructProperty && StructProperty->Struct && StructProperty->Struct->IsChildOf(FRigVMTrait::StaticStruct()))
+						{
+							TConstArrayView<FRigVMMemoryHandle> AdditionalMemoryHandles;
+							if(AdditionalStartIndex != INDEX_NONE)
+							{
+								AdditionalMemoryHandles = TConstArrayView<FRigVMMemoryHandle>(&ContextPublicData.AdditionalTraitMemoryHandles[AdditionalStartIndex], AdditionalNum);
+								AdditionalStartIndex = INDEX_NONE;
+							}
+
+							ContextPublicData.Traits.Emplace(
+								Context.WorkMemoryStorage.GetData<FRigVMTrait>(TraitPropertyIndex),
+								Cast<UScriptStruct>(StructProperty->Struct),
+								AdditionalMemoryHandles);
+						}
+						else
 						{
-							if(StructProperty->Struct && StructProperty->Struct->IsChildOf(FRigVMTrait::StaticStruct()))
+							if(AdditionalStartIndex == INDEX_NONE)
 							{
-								ContextPublicData.Traits.Emplace(
-									Context.WorkMemoryStorage.GetData<FRigVMTrait>(StructProperty),
-									Cast<UScriptStruct>(StructProperty->Struct));
+								AdditionalStartIndex = ContextPublicData.AdditionalTraitMemoryHandles.Num();
+								AdditionalNum = 0;
 							}
+							ContextPublicData.AdditionalTraitMemoryHandles.Emplace(Context.WorkMemoryStorage.GetHandle(TraitPropertyIndex));
+							AdditionalNum++;
 						}
 					}
 				}
diff --git a/Engine/Plugins/Runtime/RigVM/Source/RigVM/Private/RigVMCore/RigVMByteCode.cpp b/Engine/Plugins/Runtime/RigVM/Source/RigVM/Private/RigVMCore/RigVMByteCode.cpp
--- a/Engine/Plugins/Runtime/RigVM/Source/RigVM/Private/RigVMCore/RigVMByteCode.cpp
+++ b/Engine/Plugins/Runtime/RigVM/Source/RigVM/Private/RigVMCore/RigVMByteCode.cpp
@@ -2149,6 +2149,48 @@ TMap<int32, TArray<FRigVMTraitScope>> FRigVMByteCode::GetTraits(
 	return Result;
 }
 
+TMap<int32, TArray<FRigVMTraitScope>> FRigVMByteCode::GetTraits(
+	FRigVMMemoryStorageStruct& InLiteralMemory,
+	FRigVMMemoryStorageStruct& InWorkMemory,
+	TArray<FRigVMMemoryHandle>& OutAdditionalMemoryHandles,
+	const UScriptStruct* InScriptStruct) const
+{
+	// Count memory handles we will add and reserve to ensure we don't reallocate and invalidate
+	// handle views in returned FRigVMTraitScopes
+	int32 NumAdditionalMemoryHandles = 0;
+	const FRigVMInstructionArray Instructions = GetInstructions();
+	for(int32 InstructionIndex = 0; InstructionIndex < Instructions.Num(); InstructionIndex++)
+	{
+		const FRigVMInstruction& Instruction = Instructions[InstructionIndex];
+		if(Instruction.OpCode == ERigVMOpCode::SetupTraits)
+		{
+			TArray<FRigVMMemoryHandle> AdditionalMemoryHandles;
+			const TArray<FRigVMTraitScope> Traits = GetTraitsForInstruction(Instruction, InLiteralMemory, InWorkMemory, AdditionalMemoryHandles, InScriptStruct);
+			for(const FRigVMTraitScope& Trait : Traits)
+			{
+				NumAdditionalMemoryHandles += Trait.GetAdditionalMemoryHandles().Num();
+			}
+		}
+	}
+
+	OutAdditionalMemoryHandles.Reserve(NumAdditionalMemoryHandles);
+
+	TMap<int32, TArray<FRigVMTraitScope>> Result;
+	for(int32 InstructionIndex = 0; InstructionIndex < Instructions.Num(); InstructionIndex++)
+	{
+		const FRigVMInstruction& Instruction = Instructions[InstructionIndex];
+		if(Instruction.OpCode == ERigVMOpCode::SetupTraits)
+		{
+			const TArray<FRigVMTraitScope> Traits = GetTraitsForInstruction(Instruction, InLiteralMemory, InWorkMemory, OutAdditionalMemoryHandles, InScriptStruct);
+			if(!Traits.IsEmpty())
+			{
+				Result.Add(InstructionIndex, Traits);
+			}
+		}
+	}
+	return Result;
+}
+
 TArray<FRigVMTraitScope> FRigVMByteCode::GetTraitsForInstruction(
 	const FRigVMInstruction& InInstruction,
 	FRigVMMemoryStorageStruct& InLiteralMemory,
@@ -2195,6 +2237,88 @@ TArray<FRigVMTraitScope> FRigVMByteCode::GetTraitsForInstruction(
 	return Traits;
 }
 
+TArray<FRigVMTraitScope> FRigVMByteCode::GetTraitsForInstruction(
+	const FRigVMInstruction& InInstruction,
+	FRigVMMemoryStorageStruct& InLiteralMemory,
+	FRigVMMemoryStorageStruct& InWorkMemory,
+	TArray<FRigVMMemoryHandle>& OutAdditionalMemoryHandles,
+	const UScriptStruct* InScriptStruct) const
+{
+	TArray<FRigVMTraitScope> Traits;
+
+	if(InInstruction.OpCode == ERigVMOpCode::SetupTraits)
+	{
+		if(InScriptStruct == nullptr)
+		{
+			InScriptStruct = FRigVMTrait::StaticStruct();
+		}
+
+		const FRigVMSetupTraitsOp& Op = GetOpAt<FRigVMSetupTraitsOp>(InInstruction);
+		check(Op.Arg.GetMemoryType() == ERigVMMemoryType::Literal);
+
+		if(InLiteralMemory.GetProperties().IsValidIndex(Op.Arg.GetRegisterIndex()))
+		{
+			if(const FArrayProperty* TraitIndicesProperty = CastField<FArrayProperty>(InLiteralMemory.GetProperty(Op.Arg.GetRegisterIndex())))
+			{
+				if(TraitIndicesProperty->Inner->IsA<FIntProperty>())
+				{
+					const TArray<int32>& TraitIndices = *InLiteralMemory.GetData<TArray<int32>>(TraitIndicesProperty);
+
+					// Count memory handles we will add and reserve to ensure we don't reallocate and invalidate
+					// handle views in returned FRigVMTraitScopes
+					int32 NumAdditionalMemoryHandles = 0;
+					for(const int32 TraitIndex : TraitIndices)
+					{
+						const FProperty* Property = InWorkMemory.GetProperties()[TraitIndex];
+						const FStructProperty* StructProperty = CastField<FStructProperty>(Property);
+						if(Property && !(StructProperty && StructProperty->Struct && StructProperty->Struct->IsChildOf(FRigVMTrait::StaticStruct())))
+						{
+							NumAdditionalMemoryHandles++;
+						}
+					}
+
+					OutAdditionalMemoryHandles.Reserve(NumAdditionalMemoryHandles);
+
+					int32 AdditionalStartIndex = OutAdditionalMemoryHandles.Num();
+					int32 AdditionalNum = 0;
+
+					for(const int32 TraitIndex : TraitIndices)
+					{
+						const FProperty* Property = InWorkMemory.GetProperties()[TraitIndex];
+						const FStructProperty* StructProperty = CastField<FStructProperty>(Property);
+						if(StructProperty && StructProperty->Struct && StructProperty->Struct->IsChildOf(InScriptStruct))
+						{
+							TConstArrayView<FRigVMMemoryHandle> AdditionalMemoryHandles;
+							if(AdditionalStartIndex != INDEX_NONE)
+							{
+								AdditionalMemoryHandles = TConstArrayView<FRigVMMemoryHandle>(&OutAdditionalMemoryHandles[AdditionalStartIndex], AdditionalNum);
+								AdditionalStartIndex = INDEX_NONE;
+							}
+
+							Traits.Emplace(
+								InWorkMemory.GetData<FRigVMTrait>(StructProperty),
+								Cast<UScriptStruct>(StructProperty->Struct),
+								AdditionalMemoryHandles);
+						}
+						else if(Property && !(StructProperty && StructProperty->Struct && StructProperty->Struct->IsChildOf(FRigVMTrait::StaticStruct())))
+						{
+							if(AdditionalStartIndex == INDEX_NONE)
+							{
+								AdditionalStartIndex = OutAdditionalMemoryHandles.Num();
+								AdditionalNum = 0;
+							}
+							OutAdditionalMemoryHandles.Emplace(InWorkMemory.GetHandle(TraitIndex));
+							AdditionalNum++;
+						}
+					}
+				}
+			}
+		}
+	}
+
+	return Traits;
+}
+
 const FRigVMBranchInfo* FRigVMByteCode::GetBranchInfo(const FRigVMBranchInfoKey& InBranchInfoKey) const
 {
 	if(BranchInfos.IsEmpty())
diff --git a/Engine/Plugins/Runtime/RigVM/Source/RigVM/Private/RigVMPropertyUtils.cpp b/Engine/Plugins/Runtime/RigVM/Source/RigVM/Private/RigVMPropertyUtils.cpp
--- a/Engine/Plugins/Runtime/RigVM/Source/RigVM/Private/RigVMPropertyUtils.cpp
+++ b/Engine/Plugins/Runtime/RigVM/Source/RigVM/Private/RigVMPropertyUtils.cpp
@@ -20,6 +20,11 @@ void RigVMPropertyUtils::GetTypeFromProperty(const FProperty* InProperty, FName&
 		OutTypeName = TEXT("int32");
 		OutTypeObject = nullptr;
 	}
+	else if (CastField<FInt64Property>(InProperty))
+	{
+		OutTypeName = TEXT("int64");
+		OutTypeObject = nullptr;
+	}
 	else if (CastField<FFloatProperty>(InProperty))
 	{
 		OutTypeName = TEXT("float");
@@ -57,6 +62,11 @@ void RigVMPropertyUtils::GetTypeFromProperty(const FProperty* InProperty, FName&
 			OutTypeName = BytePropertyEnum->GetFName();
 			OutTypeObject = BytePropertyEnum;
 		}
+		else
+		{
+			OutTypeName = TEXT("uint8");
+			OutTypeObject = nullptr;
+		}
 	}
 	else if (const FStructProperty* StructProperty = CastField<FStructProperty>(InProperty))
 	{
diff --git a/Engine/Plugins/Runtime/RigVM/Source/RigVM/Public/RigVMCore/RigVM.h b/Engine/Plugins/Runtime/RigVM/Source/RigVM/Public/RigVMCore/RigVM.h
--- a/Engine/Plugins/Runtime/RigVM/Source/RigVM/Public/RigVMCore/RigVM.h
+++ b/Engine/Plugins/Runtime/RigVM/Source/RigVM/Public/RigVMCore/RigVM.h
@@ -441,6 +441,19 @@ public:
 		return GetByteCode().GetTraits<T>(*GetLiteralMemory(), InContext.WorkMemoryStorage);
 	}
 
+	// returns the traits for this VM's bytecode and Context, as well as any additional memory handles
+	TMap<int32, TArray<FRigVMTraitScope>> GetTraits(FRigVMExtendedExecuteContext& InContext, TArray<FRigVMMemoryHandle>& OutAdditionalMemoryHandles, const UScriptStruct* InScriptStruct = nullptr)
+	{
+		return GetByteCode().GetTraits(*GetLiteralMemory(), InContext.WorkMemoryStorage, OutAdditionalMemoryHandles, InScriptStruct);
+	}
+
+	// returns the traits of a given type for this VM's bytecode and Context, as well as any additional memory handles
+	template<typename T>
+	TMap<int32, TArray<FRigVMTraitScope>> GetTraits(FRigVMExtendedExecuteContext& InContext, TArray<FRigVMMemoryHandle>& OutAdditionalMemoryHandles)
+	{
+		return GetByteCode().GetTraits<T>(*GetLiteralMemory(), InContext.WorkMemoryStorage, OutAdditionalMemoryHandles);
+	}
+	
 	// returns the traits for the provided memory for a single instruction
 	TArray<FRigVMTraitScope> GetTraitsForInstruction(const FRigVMInstruction& InInstruction, FRigVMExtendedExecuteContext& InContext, const UScriptStruct* InScriptStruct = nullptr)
 	{
diff --git a/Engine/Plugins/Runtime/RigVM/Source/RigVM/Public/RigVMCore/RigVMByteCode.h b/Engine/Plugins/Runtime/RigVM/Source/RigVM/Public/RigVMCore/RigVMByteCode.h
--- a/Engine/Plugins/Runtime/RigVM/Source/RigVM/Public/RigVMCore/RigVMByteCode.h
+++ b/Engine/Plugins/Runtime/RigVM/Source/RigVM/Public/RigVMCore/RigVMByteCode.h
@@ -1433,6 +1433,9 @@ public:
 	// returns the traits for the provided memory
 	TMap<int32, TArray<FRigVMTraitScope>> GetTraits(FRigVMMemoryStorageStruct& InLiteralMemory, FRigVMMemoryStorageStruct& InWorkMemory, const UScriptStruct* InScriptStruct = nullptr) const;
 
+	// returns the traits for the provided memory and any additional memory handles
+	TMap<int32, TArray<FRigVMTraitScope>> GetTraits(FRigVMMemoryStorageStruct& InLiteralMemory, FRigVMMemoryStorageStruct& InWorkMemory, TArray<FRigVMMemoryHandle>& OutAdditionalMemoryHandles, const UScriptStruct* InScriptStruct = nullptr) const;
+	
 	// returns the traits of a given type for the provided memory
 	template<typename T>
 	TMap<int32, TArray<FRigVMTraitScope>> GetTraits(FRigVMMemoryStorageStruct& InLiteralMemory, FRigVMMemoryStorageStruct& InWorkMemory) const
@@ -1440,9 +1443,19 @@ public:
 		return GetTraits(InLiteralMemory, InWorkMemory, T::StaticStruct());
 	}
 
+	// returns the traits of a given type for the provided memory and any additional memorty handles
+	template<typename T>
+	TMap<int32, TArray<FRigVMTraitScope>> GetTraits(FRigVMMemoryStorageStruct& InLiteralMemory, FRigVMMemoryStorageStruct& InWorkMemory, TArray<FRigVMMemoryHandle>& OutAdditionalMemoryHandles) const
+	{
+		return GetTraits(InLiteralMemory, InWorkMemory, OutAdditionalMemoryHandles, T::StaticStruct());
+	}
+
 	// returns the traits for the provided memory for a single instruction
 	TArray<FRigVMTraitScope> GetTraitsForInstruction(const FRigVMInstruction& InInstruction, FRigVMMemoryStorageStruct& InLiteralMemory, FRigVMMemoryStorageStruct& InWorkMemory, const UScriptStruct* InScriptStruct = nullptr) const;
 
+	// returns the traits for the provided memory for a single instruction and any additional memory handles
+	TArray<FRigVMTraitScope> GetTraitsForInstruction(const FRigVMInstruction& InInstruction, FRigVMMemoryStorageStruct& InLiteralMemory, FRigVMMemoryStorageStruct& InWorkMemory, TArray<FRigVMMemoryHandle>& OutAdditionalMemoryHandles, const UScriptStruct* InScriptStruct = nullptr) const;
+
 	// returns the traits of a given type for the provided memory for a single instruction
 	template<typename T>
 	TArray<FRigVMTraitScope> GetTraitsForInstruction(const FRigVMInstruction& InInstruction, FRigVMMemoryStorageStruct& InLiteralMemory, FRigVMMemoryStorageStruct& InWorkMemory) const
@@ -1450,6 +1463,13 @@ public:
 		return GetTraitsForInstruction(InInstruction, InLiteralMemory, InWorkMemory, T::StaticStruct());
 	}
 
+	// returns the traits of a given type for the provided memory for a single instruction and any additional memory handles
+	template<typename T>
+	TArray<FRigVMTraitScope> GetTraitsForInstruction(const FRigVMInstruction& InInstruction, FRigVMMemoryStorageStruct& InLiteralMemory, FRigVMMemoryStorageStruct& InWorkMemory, TArray<FRigVMMemoryHandle>& OutAdditionalMemoryHandles) const
+	{
+		return GetTraitsForInstruction(InInstruction, InLiteralMemory, InWorkMemory, T::StaticStruct(), OutAdditionalMemoryHandles);
+	}
+
 private:
 
 	template<class OpType>
diff --git a/Engine/Plugins/Runtime/RigVM/Source/RigVM/Public/RigVMCore/RigVMExecuteContext.h b/Engine/Plugins/Runtime/RigVM/Source/RigVM/Public/RigVMCore/RigVMExecuteContext.h
--- a/Engine/Plugins/Runtime/RigVM/Source/RigVM/Public/RigVMCore/RigVMExecuteContext.h
+++ b/Engine/Plugins/Runtime/RigVM/Source/RigVM/Public/RigVMCore/RigVMExecuteContext.h
@@ -514,6 +514,9 @@ protected:
 	/** The traits accessible to the current instruction */
 	TArray<FRigVMTraitScope> Traits;
 
+	/** Additional memory handles for each trait */
+	TArray<FRigVMMemoryHandle> AdditionalTraitMemoryHandles;
+
 #if UE_RIGVM_DEBUG_EXECUTION
 public:
 	bool bDebugExecution = false;
diff --git a/Engine/Plugins/Runtime/RigVM/Source/RigVM/Public/RigVMCore/RigVMTraitScope.h b/Engine/Plugins/Runtime/RigVM/Source/RigVM/Public/RigVMCore/RigVMTraitScope.h
--- a/Engine/Plugins/Runtime/RigVM/Source/RigVM/Public/RigVMCore/RigVMTraitScope.h
+++ b/Engine/Plugins/Runtime/RigVM/Source/RigVM/Public/RigVMCore/RigVMTraitScope.h
@@ -7,6 +7,7 @@
 
 struct FRigVMTrait;
 class UScriptStruct;
+struct FRigVMMemoryHandle;
 
 class RIGVM_API FRigVMTraitScope
 {
@@ -16,12 +17,19 @@ public:
 		, ScriptStruct(nullptr)
 	{
 	}
-	
+
 	FRigVMTraitScope(FRigVMTrait* InTrait, const UScriptStruct* InScriptStruct)
 		: Trait(InTrait)
 		, ScriptStruct(InScriptStruct)
 	{
 	}
+	
+	FRigVMTraitScope(FRigVMTrait* InTrait, const UScriptStruct* InScriptStruct, TConstArrayView<FRigVMMemoryHandle> InAdditionalMemoryHandles)
+		: Trait(InTrait)
+		, ScriptStruct(InScriptStruct)
+		, AdditionalMemoryHandles(InAdditionalMemoryHandles)
+	{
+	}
 
 	bool IsValid() const
 	{
@@ -73,8 +81,14 @@ public:
 		return ScriptStruct;
 	}
 
+	TConstArrayView<FRigVMMemoryHandle> GetAdditionalMemoryHandles() const
+	{
+		return AdditionalMemoryHandles;
+	}
+
 private:
 
 	FRigVMTrait* Trait;
 	const UScriptStruct* ScriptStruct;
+	TConstArrayView<FRigVMMemoryHandle> AdditionalMemoryHandles;
 };
\ No newline at end of file
diff --git a/Engine/Plugins/Runtime/RigVM/Source/RigVM/Public/RigVMHost.h b/Engine/Plugins/Runtime/RigVM/Source/RigVM/Public/RigVMHost.h
--- a/Engine/Plugins/Runtime/RigVM/Source/RigVM/Public/RigVMHost.h
+++ b/Engine/Plugins/Runtime/RigVM/Source/RigVM/Public/RigVMHost.h
@@ -365,7 +365,7 @@ protected:
 	
 	void HandleExecutionReachedExit(const FName& InEventName);
 	
-	TArray<FRigVMExternalVariable> GetExternalVariablesImpl(bool bFallbackToBlueprint) const;
+	virtual TArray<FRigVMExternalVariable> GetExternalVariablesImpl(bool bFallbackToBlueprint) const;
 
 	FProperty* GetPublicVariableProperty(const FName& InVariableName) const
 	{
diff --git a/Engine/Plugins/Runtime/RigVM/Source/RigVMDeveloper/Private/EdGraph/RigVMEdGraph.cpp b/Engine/Plugins/Runtime/RigVM/Source/RigVMDeveloper/Private/EdGraph/RigVMEdGraph.cpp
--- a/Engine/Plugins/Runtime/RigVM/Source/RigVMDeveloper/Private/EdGraph/RigVMEdGraph.cpp
+++ b/Engine/Plugins/Runtime/RigVM/Source/RigVMDeveloper/Private/EdGraph/RigVMEdGraph.cpp
@@ -1027,7 +1027,7 @@ void URigVMEdGraph::AddNode(UEdGraphNode* NodeToAdd, bool bUserAction, bool bSel
 	// Comments are added outside of the ControlRigEditor, so we add here the node to the model
 	if (const UEdGraphNode_Comment* CommentNode = Cast<const UEdGraphNode_Comment>(NodeToAdd))
 	{
-		if (URigVMController* Controller = GetBlueprint()->GetOrCreateController(GetModel()))
+		if (URigVMController* Controller = GetController())
 		{
 			if (GetModel()->FindNodeByName(NodeToAdd->GetFName()) == nullptr) // When recreating nodes at RebuildGraphFromModel, the model node already exists
 			{
@@ -1045,7 +1045,7 @@ void URigVMEdGraph::AddNode(UEdGraphNode* NodeToAdd, bool bUserAction, bool bSel
 					GEditor->CancelTransaction(0);
 				}
 #endif // WITH_EDITOR
-				TGuardValue<bool> BlueprintNotifGuard(GetBlueprint()->bSuspendModelNotificationsForOthers, true);
+				TGuardValue<bool> BlueprintNotifGuard(GetRigVMClient()->bSuspendModelNotificationsForOthers, true);
 				FVector2D NodePos(CommentNode->NodePosX, CommentNode->NodePosY);
 				FVector2D NodeSize(CommentNode->NodeWidth, CommentNode->NodeHeight);
 				FLinearColor NodeColor = CommentNode->CommentColor;
diff --git a/Engine/Plugins/Runtime/RigVM/Source/RigVMDeveloper/Private/RigVMBlueprint.cpp b/Engine/Plugins/Runtime/RigVM/Source/RigVMDeveloper/Private/RigVMBlueprint.cpp
--- a/Engine/Plugins/Runtime/RigVM/Source/RigVMDeveloper/Private/RigVMBlueprint.cpp
+++ b/Engine/Plugins/Runtime/RigVM/Source/RigVMDeveloper/Private/RigVMBlueprint.cpp
@@ -161,7 +161,6 @@ URigVMBlueprint::URigVMBlueprint()
 URigVMBlueprint::URigVMBlueprint(const FObjectInitializer& ObjectInitializer)
 {
 	bSuspendModelNotificationsForSelf = false;
-	bSuspendModelNotificationsForOthers = false;
 	bSuspendAllNotifications = false;
 	bSuspendPythonMessagesForRigVMClient = true;
 	bMarkBlueprintAsStructurallyModifiedPending = false;
@@ -1300,7 +1299,7 @@ void URigVMBlueprint::RecompileVM()
 	if (CDO && CDO->VM != nullptr)
 	{
 		TGuardValue<bool> ReentrantGuardSelf(bSuspendModelNotificationsForSelf, true);
-		TGuardValue<bool> ReentrantGuardOthers(bSuspendModelNotificationsForOthers, true);
+		TGuardValue<bool> ReentrantGuardOthers(RigVMClient.bSuspendModelNotificationsForOthers, true);
 
 		SetupDefaultObjectDuringCompilation(CDO);
 
@@ -3234,7 +3233,7 @@ void URigVMBlueprint::HandleModifiedEvent(ERigVMGraphNotifType InNotifType, URig
 	}
 
 	// if the notification still has to be sent...
-	if (bNotifForOthersPending && !bSuspendModelNotificationsForOthers)
+	if (bNotifForOthersPending && !RigVMClient.bSuspendModelNotificationsForOthers)
 	{
 		if (ModifiedEvent.IsBound())
 		{
diff --git a/Engine/Plugins/Runtime/RigVM/Source/RigVMDeveloper/Private/RigVMCompiler/RigVMAST.cpp b/Engine/Plugins/Runtime/RigVM/Source/RigVMDeveloper/Private/RigVMCompiler/RigVMAST.cpp
--- a/Engine/Plugins/Runtime/RigVM/Source/RigVMDeveloper/Private/RigVMCompiler/RigVMAST.cpp
+++ b/Engine/Plugins/Runtime/RigVM/Source/RigVMDeveloper/Private/RigVMCompiler/RigVMAST.cpp
@@ -1420,6 +1420,23 @@ FRigVMExprAST* FRigVMParserAST::TraversePin(const FRigVMASTProxy& InPinProxy, FR
 	if(Pin->IsTraitPin())
 	{
 		PinExpr = MakeExpr<FRigVMVarExprAST>(FRigVMExprAST::EType::Var, InPinProxy);
+
+		if(Settings.bSetupTraits)
+		{
+			// Traits can generate their own programmatic pins via FRigVMTrait::GetProgrammaticPins. We account for these as additional expressions if the
+			// pin is not part of the set of sub-pins exposed on the struct
+			for(URigVMPin* SubPin : Pin->SubPins)
+			{
+				if(SubPin->IsProgrammaticPin())
+				{
+					// Not a pin from the struct - add a synthetic var for this parent too
+					FRigVMASTProxy SubPinProxy = InPinProxy.GetSibling(SubPin);
+					FRigVMExprAST* SubPinExpr = TraversePin(SubPinProxy, InParentExpr);
+					const int32 ChildIndex = InParentExpr->Children.Find(SubPinExpr);
+					InParentExpr->PinNameToChildIndex.FindOrAdd(SubPinExpr->GetName()) = ChildIndex;
+				}
+			}
+		}
 	}
 	else if ((Pin->GetDirection() == ERigVMPinDirection::Input ||
 		Pin->GetDirection() == ERigVMPinDirection::Visible) &&
@@ -1617,6 +1634,15 @@ FRigVMExprAST* FRigVMParserAST::TraverseLink(int32 InLinkIndex, FRigVMExprAST* I
 		}
 	}
 
+	if (!bRequiresCopy)
+	{
+		// Programmatic pins always require a copy
+		if(TargetPin->IsProgrammaticPin())
+		{
+			bRequiresCopy = true;
+		}
+	}
+
 	FRigVMAssignExprAST* AssignExpr = nullptr;
 	if (bRequiresCopy)
 	{
@@ -1922,6 +1948,12 @@ void FRigVMParserAST::FoldAssignments()
 			}
 		}
 
+		// Skip folding for programmatic pins as we expect them to always exist in work memory
+		if(TargetPin->IsProgrammaticPin())
+		{
+			continue;
+		}
+		
 		FRigVMExprAST* Parent = AssignExpr->Parents[0];
 		if (!Parent->IsA(FRigVMExprAST::EType::Var))
 		{
@@ -2812,8 +2844,13 @@ TArray<int32> FRigVMParserAST::GetLinkIndices(const FRigVMASTProxy& InPinProxy,
 		URigVMPin* Pin = InPinProxy.GetSubjectChecked<URigVMPin>();
 		for (URigVMPin* SubPin : Pin->GetSubPins())
 		{
-			FRigVMASTProxy SubPinProxy = InPinProxy.GetSibling(SubPin);
-			LinkIndices.Append(GetLinkIndices(SubPinProxy, bGetSource, true));
+			// We dont get links to programmatic pins as links are traversed manually in TraversePin rather than
+			// using links from parent traits. This ensures we dont double-up on programmatic pin traversal.
+			if(!SubPin->IsProgrammaticPin())
+			{
+				FRigVMASTProxy SubPinProxy = InPinProxy.GetSibling(SubPin);
+				LinkIndices.Append(GetLinkIndices(SubPinProxy, bGetSource, true));
+			}
 		}
 	}
 
diff --git a/Engine/Plugins/Runtime/RigVM/Source/RigVMDeveloper/Private/RigVMCompiler/RigVMCompiler.cpp b/Engine/Plugins/Runtime/RigVM/Source/RigVMDeveloper/Private/RigVMCompiler/RigVMCompiler.cpp
--- a/Engine/Plugins/Runtime/RigVM/Source/RigVMDeveloper/Private/RigVMCompiler/RigVMCompiler.cpp
+++ b/Engine/Plugins/Runtime/RigVM/Source/RigVMDeveloper/Private/RigVMCompiler/RigVMCompiler.cpp
@@ -2108,12 +2108,27 @@ bool URigVMCompiler::TraverseCallExtern(const FRigVMCallExternExprAST* InExpr, F
 					TArray<FString> DefaultValues;
 					for(const URigVMPin* TraitPin : TraitPins)
 					{
-						if(const FRigVMExprAST* TraitExpr = InExpr->FindVarWithPinName(TraitPin->GetFName()))
+						auto AddPinRegister = [this, &DefaultValues, &WorkData](const FRigVMVarExprAST* InPinExpr)
 						{
-							check(TraitExpr->IsVar());
-							const FRigVMOperand TraitOperand = FindOrAddRegister(TraitExpr->To<FRigVMVarExprAST>(), WorkData);
+							const FRigVMOperand TraitOperand = FindOrAddRegister(InPinExpr, WorkData);
 							check(TraitOperand.GetMemoryType() == ERigVMMemoryType::Work);
 							DefaultValues.Add(FString::FromInt(TraitOperand.GetRegisterIndex()));
+						};
+
+						if(const FRigVMVarExprAST* TraitPinExpr = InExpr->FindVarWithPinName(TraitPin->GetFName()))
+						{
+							// Add programmatic pin expressions first. This is done to allow memory handles to be built at runtime and passed to their
+							// respective 'owning' traits
+							const TArray<URigVMPin*> ProgrammaticPins = TraitPin->GetProgrammaticSubPins(); 
+							for(URigVMPin* ProgrammaticPin : ProgrammaticPins)
+							{
+								if(const FRigVMVarExprAST* ProgrammaticPinExpr = InExpr->FindVarWithPinName(*ProgrammaticPin->GetPinPath()))
+								{
+									AddPinRegister(ProgrammaticPinExpr);
+								}
+							}
+
+							AddPinRegister(TraitPinExpr);
 						}
 					}
 
@@ -3118,6 +3133,11 @@ bool URigVMCompiler::TraverseAssign(const FRigVMAssignExprAST* InExpr, FRigVMCom
 						return;
 					}
 
+					if(Pin->IsProgrammaticPin())
+					{
+						return;
+					}
+
 					FString SegmentPath = Pin->GetSegmentPath(false);
 					if(bHasTargetSegmentPath)
 					{
@@ -4169,6 +4189,11 @@ const FRigVMCompilerWorkData::FRigVMASTProxyArray& URigVMCompiler::FindProxiesWi
 				{
 					continue;
 				}
+
+				if(Pin->IsProgrammaticPin())
+				{
+					continue;
+				}
 			}
 			PinProxies.Add(PinProxiesToProcess[ProxyIndex]);
 		}
diff --git a/Engine/Plugins/Runtime/RigVM/Source/RigVMDeveloper/Private/RigVMModel/RigVMPin.cpp b/Engine/Plugins/Runtime/RigVM/Source/RigVMDeveloper/Private/RigVMModel/RigVMPin.cpp
--- a/Engine/Plugins/Runtime/RigVM/Source/RigVMDeveloper/Private/RigVMModel/RigVMPin.cpp
+++ b/Engine/Plugins/Runtime/RigVM/Source/RigVMDeveloper/Private/RigVMModel/RigVMPin.cpp
@@ -726,6 +726,16 @@ bool URigVMPin::ShouldOnlyShowSubPins() const
 	{
 		if(const URigVMNode* Node = GetNode())
 		{
+			if(IsStruct())
+			{
+				// Never show sub-pins for custom import/export text as we cant do memberwise manipulations in the graph via text
+				// This change is to allow things like FUniversalObjectLocator (with its native text serialization) to be used on RigVM pins.
+				if((GetScriptStruct()->StructFlags & (STRUCT_ExportTextItemNative | STRUCT_ImportTextItemNative)) != 0)
+				{
+					return false;
+				}
+			}
+
 			if(const URigVMUnitNode* UnitNode = Cast<URigVMUnitNode>(Node))
 			{
 				if(const UScriptStruct* Struct = UnitNode->GetScriptStruct())
@@ -762,6 +772,15 @@ bool URigVMPin::ShouldHideSubPins() const
 	{
 		if(const URigVMNode* Node = GetNode())
 		{
+			if(IsStruct())
+			{
+				// Hide sub-pins for custom import/export text as we cant do memberwise manipulations in the graph via text
+				if((GetScriptStruct()->StructFlags & (STRUCT_ExportTextItemNative | STRUCT_ImportTextItemNative)) != 0)
+				{
+					return true;
+				}
+			}
+
 			if(const URigVMUnitNode* UnitNode = Cast<URigVMUnitNode>(Node))
 			{
 				if(const UScriptStruct* Struct = UnitNode->GetScriptStruct())
@@ -1588,6 +1607,39 @@ bool URigVMPin::IsTraitPin() const
 	return false;
 }
 
+bool URigVMPin::IsProgrammaticPin() const
+{
+	// Traits can generate their own programmatic pins via FRigVMTrait::GetProgrammaticPins. We account for these as additional expressions if the
+	// pin is not part of the set of sub-pins exposed on the struct
+	if(const URigVMPin* ParentPin = GetParentPin())
+	{
+		UScriptStruct* ScriptStruct = ParentPin->GetScriptStruct();
+		if(ScriptStruct && ScriptStruct->IsChildOf(FRigVMTrait::StaticStruct()))
+		{
+			if(ScriptStruct->FindPropertyByName(GetFName()) == nullptr)
+			{
+				return true;
+			}
+		}
+	}
+
+	return false;
+}
+
+TArray<URigVMPin*> URigVMPin::GetProgrammaticSubPins() const
+{
+	TArray<URigVMPin*> ProgrammaticPins;
+	for(URigVMPin* SubPin : SubPins)
+	{
+		if(SubPin->IsProgrammaticPin())
+		{
+			ProgrammaticPins.Add(SubPin);
+		}
+	}
+
+	return ProgrammaticPins;
+}
+
 TSharedPtr<FStructOnScope> URigVMPin::GetTraitInstance(bool bUseDefaultValueFromPin) const
 {
 	if(const URigVMNode* Node = GetNode())
diff --git a/Engine/Plugins/Runtime/RigVM/Source/RigVMDeveloper/Public/RigVMBlueprint.h b/Engine/Plugins/Runtime/RigVM/Source/RigVMDeveloper/Public/RigVMBlueprint.h
--- a/Engine/Plugins/Runtime/RigVM/Source/RigVMDeveloper/Public/RigVMBlueprint.h
+++ b/Engine/Plugins/Runtime/RigVM/Source/RigVMDeveloper/Public/RigVMBlueprint.h
@@ -527,7 +527,6 @@ public:
 	TMap<FString, FRigVMOperand> PinToOperandMap;
 
 	bool bSuspendModelNotificationsForSelf;
-	bool bSuspendModelNotificationsForOthers;
 	bool bSuspendAllNotifications;
 
 	void RebuildGraphFromModel();
diff --git a/Engine/Plugins/Runtime/RigVM/Source/RigVMDeveloper/Public/RigVMModel/Nodes/RigVMVariableNode.h b/Engine/Plugins/Runtime/RigVM/Source/RigVMDeveloper/Public/RigVMModel/Nodes/RigVMVariableNode.h
--- a/Engine/Plugins/Runtime/RigVM/Source/RigVMDeveloper/Public/RigVMModel/Nodes/RigVMVariableNode.h
+++ b/Engine/Plugins/Runtime/RigVM/Source/RigVMDeveloper/Public/RigVMModel/Nodes/RigVMVariableNode.h
@@ -66,15 +66,17 @@ public:
 	virtual FLinearColor GetNodeColor() const override { return FLinearColor::Blue; }
 	virtual bool IsDefinedAsVarying() const override { return true; }
 
+	// Get the variable name pin
+	URigVMPin* GetVariableNamePin() const;
+
+	// Get the value pin
+	URigVMPin* GetValuePin() const;
+
 private:
 
 	static const inline TCHAR* VariableName = TEXT("Variable");
 	static const inline TCHAR* ValueName = TEXT("Value");
 
-	URigVMPin* GetVariableNamePin() const;
-	URigVMPin* GetValuePin() const;
-	
-
 	friend class URigVMController;
 	friend class URigVMBlueprint;
 	friend struct FRigVMRemoveNodeAction;
diff --git a/Engine/Plugins/Runtime/RigVM/Source/RigVMDeveloper/Public/RigVMModel/RigVMClient.h b/Engine/Plugins/Runtime/RigVM/Source/RigVMDeveloper/Public/RigVMModel/RigVMClient.h
--- a/Engine/Plugins/Runtime/RigVM/Source/RigVMDeveloper/Public/RigVMModel/RigVMClient.h
+++ b/Engine/Plugins/Runtime/RigVM/Source/RigVMDeveloper/Public/RigVMModel/RigVMClient.h
@@ -187,6 +187,7 @@ public:
 		, bSuspendNotifications(false)
 		, bIgnoreModelNotifications(false)
 		, bDefaultModelCanBeRemoved(false)
+		, bSuspendModelNotificationsForOthers(false)
 		, OuterClientHost(nullptr)
 		, OuterClientPropertyName(NAME_None)	
 		, ExternalModelHost(nullptr)
@@ -350,6 +351,7 @@ public:
 	bool bSuspendNotifications;
 	bool bIgnoreModelNotifications;
 	bool bDefaultModelCanBeRemoved;
+	bool bSuspendModelNotificationsForOthers;
 private:
 	TWeakObjectPtr<UObject> OuterClientHost;
 	FName OuterClientPropertyName;
diff --git a/Engine/Plugins/Runtime/RigVM/Source/RigVMDeveloper/Public/RigVMModel/RigVMController.h b/Engine/Plugins/Runtime/RigVM/Source/RigVMDeveloper/Public/RigVMModel/RigVMController.h
--- a/Engine/Plugins/Runtime/RigVM/Source/RigVMDeveloper/Public/RigVMModel/RigVMController.h
+++ b/Engine/Plugins/Runtime/RigVM/Source/RigVMDeveloper/Public/RigVMModel/RigVMController.h
@@ -575,10 +575,12 @@ public:
 	// Adds a trait to a node
 	UFUNCTION(BlueprintCallable, Category = RigVMController)
 	FName AddTrait(const FName& InNodeName, const FName& InTraitTypeObjectPath, const FName& InTraitName = NAME_None, const FString& InDefaultValue = TEXT(""), int32 InPinIndex = -1, bool bSetupUndoRedo = true, bool bPrintPythonCommand = false);
+	FName AddTrait(URigVMNode* InNode, UScriptStruct* InTraitScriptStruct, const FName& InTraitName, const FString& InDefaultValue, int32 InPinIndex = -1, bool bSetupUndoRedo = true);
 
 	// Removes a trait from a node
 	UFUNCTION(BlueprintCallable, Category = RigVMController)
 	virtual bool RemoveTrait(const FName& InNodeName, const FName& InTraitName, bool bSetupUndoRedo = true, bool bPrintPythonCommand = false);
+	bool RemoveTrait(URigVMNode* InNode, const FName& InTraitName, bool bSetupUndoRedo = true);
 
 	// Un-does the last action on the stack.
 	// Note: This should really only be used for unit tests,
@@ -866,6 +868,7 @@ public:
 	// This causes a PinDefaultValueChanged modified event.
 	UFUNCTION(BlueprintCallable, Category = RigVMController)
 	bool SetPinDefaultValue(const FString& InPinPath, const FString& InDefaultValue, bool bResizeArrays = true, bool bSetupUndoRedo = true, bool bMergeUndoAction = false, bool bPrintPythonCommand = false, bool bSetValueOnLinkedPins = true);
+	bool SetPinDefaultValue(URigVMPin* InPin, const FString& InDefaultValue, bool bResizeArrays, bool bSetupUndoRedo, bool bMergeUndoAction, bool bSetValueOnLinkedPins = true);
 
 	// Resets the default value of a pin given its pinpath.
 	// This causes a PinDefaultValueChanged modified event.
@@ -1242,7 +1245,6 @@ private:
 	void ConfigurePinFromProperty(FProperty* InProperty, URigVMPin* InOutPin, ERigVMPinDirection InPinDirection = ERigVMPinDirection::Invalid) const;
 	void ConfigurePinFromPin(URigVMPin* InOutPin, URigVMPin* InPin, bool bCopyDisplayName = false);
 	void ConfigurePinFromArgument(URigVMPin* InOutPin, const FRigVMGraphFunctionArgument& InArgument, bool bCopyDisplayName = false);
-	bool SetPinDefaultValue(URigVMPin* InPin, const FString& InDefaultValue, bool bResizeArrays, bool bSetupUndoRedo, bool bMergeUndoAction, bool bSetValueOnLinkedPins = true);
 	bool ResetPinDefaultValue(URigVMPin* InPin, bool bSetupUndoRedo);
 	static FString GetPinInitialDefaultValue(const URigVMPin* InPin);
 	static FString GetPinInitialDefaultValueFromStruct(UScriptStruct* ScriptStruct, const URigVMPin* InPin, uint32 InOffset);
@@ -1257,8 +1259,6 @@ private:
 	URigVMInjectionInfo* InjectNodeIntoPin(URigVMPin* InPin, bool bAsInput, const FName& InInputPinName, const FName& InOutputPinName, bool bSetupUndoRedo = true);
 	URigVMNode* EjectNodeFromPin(URigVMPin* InPin, bool bSetupUndoRedo = true, bool bPrintPythonCommands = false);
 	bool EjectAllInjectedNodes(URigVMNode* InNode, bool bSetupUndoRedo = true, bool bPrintPythonCommands = false);
-	FName AddTrait(URigVMNode* InNode, UScriptStruct* InTraitScriptStruct, const FName& InTraitName, const FString& InDefaultValue, int32 InPinIndex = -1, bool bSetupUndoRedo = true);
-	bool RemoveTrait(URigVMNode* InNode, const FName& InTraitName, bool bSetupUndoRedo = true);
 
 protected:
 
diff --git a/Engine/Plugins/Runtime/RigVM/Source/RigVMDeveloper/Public/RigVMModel/RigVMPin.h b/Engine/Plugins/Runtime/RigVM/Source/RigVMDeveloper/Public/RigVMModel/RigVMPin.h
--- a/Engine/Plugins/Runtime/RigVM/Source/RigVMDeveloper/Public/RigVMModel/RigVMPin.h
+++ b/Engine/Plugins/Runtime/RigVM/Source/RigVMDeveloper/Public/RigVMModel/RigVMPin.h
@@ -511,6 +511,12 @@ public:
 	UFUNCTION(BlueprintPure, Category = RigVMPin)
 	bool IsTraitPin() const;
 
+	// Returns true if this pin represents a trait's programmatic pin
+	bool IsProgrammaticPin() const;
+
+	// Get all the sub-pins that are programmatic
+	TArray<URigVMPin*> GetProgrammaticSubPins() const;
+
 	// Returns the trait backing up this pin
 	TSharedPtr<FStructOnScope> GetTraitInstance(bool bUseDefaultValueFromPin = true) const;
 
diff --git a/Engine/Plugins/Runtime/RigVM/Source/RigVMEditor/Private/Editor/RigVMEditor.cpp b/Engine/Plugins/Runtime/RigVM/Source/RigVMEditor/Private/Editor/RigVMEditor.cpp
--- a/Engine/Plugins/Runtime/RigVM/Source/RigVMEditor/Private/Editor/RigVMEditor.cpp
+++ b/Engine/Plugins/Runtime/RigVM/Source/RigVMEditor/Private/Editor/RigVMEditor.cpp
@@ -1063,7 +1063,7 @@ void FRigVMEditor::Compile()
 		}
 
 		// force to disable the supended notif brackets
-		RigVMBlueprint->bSuspendModelNotificationsForOthers = false;
+		RigVMBlueprint->GetRigVMClient()->bSuspendModelNotificationsForOthers = false;
 		RigVMBlueprint->bSuspendModelNotificationsForSelf = false;
 
 		RigVMBlueprint->GetCompileLog().Messages.Reset();
diff --git a/Engine/Source/Runtime/Core/Public/UObject/FortniteMainBranchObjectVersion.h b/Engine/Source/Runtime/Core/Public/UObject/FortniteMainBranchObjectVersion.h
--- a/Engine/Source/Runtime/Core/Public/UObject/FortniteMainBranchObjectVersion.h
+++ b/Engine/Source/Runtime/Core/Public/UObject/FortniteMainBranchObjectVersion.h
@@ -510,12 +510,16 @@ struct FFortniteMainBranchObjectVersion
 		// Deprecation of the bPlaying flag on FTimeline struct types in favor of a better
 		// PlayingStateTracker type to improve replication reliability 
 		TimelinePlayingStateTrackerDeprecation,
+		
 
 		// Enable SkipOnlyEditorOnly style cooking of UStaticMeshComponent::MeshPaintTexture
 		MeshPaintTextureUsesEditorOnly,
 
 		// Fixup and synchronize some landscape properties that have moved to the property sharing/overriding system :
 		LandscapeBodyInstanceAsSharedProperty,
+		
+		// Multiple changes to AnimNext modules, variables etc.
+		AnimNextModuleRefactor,
 
 		// -----<new versions can be added above this line>-------------------------------------------------
 		VersionPlusOne,
