RigVM: Python API frontend for variants

#rnx

diff --git a/Engine/Plugins/Experimental/AnimNext/Source/AnimNextUncookedOnly/Internal/AnimNextRigVMAssetEditorData.h b/Engine/Plugins/Experimental/AnimNext/Source/AnimNextUncookedOnly/Internal/AnimNextRigVMAssetEditorData.h
--- a/Engine/Plugins/Experimental/AnimNext/Source/AnimNextUncookedOnly/Internal/AnimNextRigVMAssetEditorData.h
+++ b/Engine/Plugins/Experimental/AnimNext/Source/AnimNextUncookedOnly/Internal/AnimNextRigVMAssetEditorData.h
@@ -177,6 +177,7 @@ protected:
 	virtual URigVMGraph* GetDefaultModel() const override;
 	virtual TArray<URigVMGraph*> GetAllModels() const override;
 	virtual URigVMFunctionLibrary* GetLocalFunctionLibrary() const override;
+	virtual URigVMFunctionLibrary* GetOrCreateLocalFunctionLibrary(bool bSetupUndoRedo =  true) override;
 	virtual URigVMGraph* AddModel(FString InName = TEXT("Rig Graph"), bool bSetupUndoRedo = true, bool bPrintPythonCommand = true) override;
 	virtual bool RemoveModel(FString InName = TEXT("Rig Graph"), bool bSetupUndoRedo = true, bool bPrintPythonCommand = true) override;
 	virtual FRigVMGetFocusedGraph& OnGetFocusedGraph() override;
diff --git a/Engine/Plugins/Experimental/AnimNext/Source/AnimNextUncookedOnly/Private/AnimNextRigVMAssetEditorData.cpp b/Engine/Plugins/Experimental/AnimNext/Source/AnimNextUncookedOnly/Private/AnimNextRigVMAssetEditorData.cpp
--- a/Engine/Plugins/Experimental/AnimNext/Source/AnimNextUncookedOnly/Private/AnimNextRigVMAssetEditorData.cpp
+++ b/Engine/Plugins/Experimental/AnimNext/Source/AnimNextUncookedOnly/Private/AnimNextRigVMAssetEditorData.cpp
@@ -295,6 +295,11 @@ URigVMFunctionLibrary* UAnimNextRigVMAssetEditorData::GetLocalFunctionLibrary()
 	return RigVMClient.GetFunctionLibrary();
 }
 
+URigVMFunctionLibrary* UAnimNextRigVMAssetEditorData::GetOrCreateLocalFunctionLibrary(bool bSetupUndoRedo)
+{
+	return RigVMClient.GetOrCreateFunctionLibrary(bSetupUndoRedo);
+}
+
 URigVMGraph* UAnimNextRigVMAssetEditorData::AddModel(FString InName, bool bSetupUndoRedo, bool bPrintPythonCommand)
 {
 	TGuardValue<bool> EnablePythonPrint(bSuspendPythonMessagesForRigVMClient, !bPrintPythonCommand);
diff --git a/Engine/Plugins/Runtime/RigVM/Source/RigVM/Public/RigVMCore/RigVMVariant.h b/Engine/Plugins/Runtime/RigVM/Source/RigVM/Public/RigVMCore/RigVMVariant.h
--- a/Engine/Plugins/Runtime/RigVM/Source/RigVM/Public/RigVMCore/RigVMVariant.h
+++ b/Engine/Plugins/Runtime/RigVM/Source/RigVM/Public/RigVMCore/RigVMVariant.h
@@ -161,6 +161,9 @@ struct RIGVM_API FRigVMVariantRef
 		return HashCombine(GetTypeHash(InVariantRef.ObjectPath), GetTypeHash(InVariantRef.Variant.Guid));
 	}
 
+	UPROPERTY(BlueprintReadOnly, Category=RigVMVariantRef)
 	FSoftObjectPath ObjectPath;
+
+	UPROPERTY(BlueprintReadOnly, Category=RigVMVariantRef)
 	FRigVMVariant Variant;
 };
diff --git a/Engine/Plugins/Runtime/RigVM/Source/RigVMDeveloper/Private/RigVMBlueprint.cpp b/Engine/Plugins/Runtime/RigVM/Source/RigVMDeveloper/Private/RigVMBlueprint.cpp
--- a/Engine/Plugins/Runtime/RigVM/Source/RigVMDeveloper/Private/RigVMBlueprint.cpp
+++ b/Engine/Plugins/Runtime/RigVM/Source/RigVMDeveloper/Private/RigVMBlueprint.cpp
@@ -37,6 +37,7 @@
 #include "Editor/Transactor.h"
 #include "CookOnTheSide/CookOnTheFlyServer.h"
 #include "RigVMEditorModule.h"
+#include "ScopedTransaction.h"
 #endif//WITH_EDITOR
 
 #define LOCTEXT_NAMESPACE "RigVMBlueprint"
@@ -1869,6 +1870,11 @@ URigVMFunctionLibrary* URigVMBlueprint::GetLocalFunctionLibrary() const
 	return RigVMClient.GetFunctionLibrary();
 }
 
+URigVMFunctionLibrary* URigVMBlueprint::GetOrCreateLocalFunctionLibrary(bool bSetupUndoRedo)
+{
+	return RigVMClient.GetOrCreateFunctionLibrary(bSetupUndoRedo);
+}
+
 URigVMGraph* URigVMBlueprint::AddModel(FString InName, bool bSetupUndoRedo, bool bPrintPythonCommand)
 {
 	TGuardValue<bool> EnablePythonPrint(bSuspendPythonMessagesForRigVMClient, !bPrintPythonCommand);
@@ -2813,6 +2819,64 @@ bool URigVMBlueprint::ChangeMemberVariableType(const FName& InName, const FStrin
 	return true;
 }
 
+FRigVMVariant URigVMBlueprint::GetAssetVariantBP() const
+{
+	return GetAssetVariant();
+}
+
+bool URigVMBlueprint::SplitAssetVariant()
+{
+	if(GetMatchingVariants().IsEmpty())
+	{
+		return false;
+	}
+
+	FScopedTransaction Transaction(LOCTEXT("SplitAssetVariant", "Split Asset Variant"));
+	Modify();
+
+	// prefer the path based (deterministic) guid - and fall back on random.
+	const FGuid PathBasedGuid = FRigVMVariant::GenerateGUID(GetPathName());
+	if(PathBasedGuid != AssetVariant.Guid)
+	{
+		AssetVariant.Guid = PathBasedGuid;
+	}
+	else
+	{
+		AssetVariant.Guid = FRigVMVariant::GenerateGUID();
+	}
+	
+	return true;
+}
+
+bool URigVMBlueprint::JoinAssetVariant(const FGuid& InGuid)
+{
+	if(AssetVariant.Guid != InGuid)
+	{
+		FScopedTransaction Transaction(LOCTEXT("JoinAssetVariant", "Join Asset Variant"));
+		Modify();
+		
+		AssetVariant.Guid = InGuid;
+		return true;
+	}
+
+	return false;
+}
+
+TArray<FRigVMVariantRef> URigVMBlueprint::GetMatchingVariants() const
+{
+	if(URigVMBuildData* BuildData = URigVMBuildData::Get())
+	{
+		TArray<FRigVMVariantRef> Variants = BuildData->FindAssetVariantRefs(AssetVariant.Guid);
+		const FRigVMVariantRef MyVariantRef = FRigVMVariantRef(GetPathName(), AssetVariant);
+		Variants.RemoveAll([MyVariantRef](const FRigVMVariantRef& VariantRef) -> bool
+		{
+			return VariantRef == MyVariantRef;
+		});
+		return Variants;
+	}
+	return TArray<FRigVMVariantRef>();
+}
+
 #endif
 
 void URigVMBlueprint::RebuildGraphFromModel()
diff --git a/Engine/Plugins/Runtime/RigVM/Source/RigVMDeveloper/Private/RigVMModel/Nodes/RigVMLibraryNode.cpp b/Engine/Plugins/Runtime/RigVM/Source/RigVMDeveloper/Private/RigVMModel/Nodes/RigVMLibraryNode.cpp
--- a/Engine/Plugins/Runtime/RigVM/Source/RigVMDeveloper/Private/RigVMModel/Nodes/RigVMLibraryNode.cpp
+++ b/Engine/Plugins/Runtime/RigVM/Source/RigVMDeveloper/Private/RigVMModel/Nodes/RigVMLibraryNode.cpp
@@ -333,4 +333,19 @@ FRigVMGraphFunctionHeader URigVMLibraryNode::GetFunctionHeader(IRigVMGraphFuncti
 	return Header;
 }
 
+FRigVMVariant URigVMLibraryNode::GetFunctionVariant() const
+{
+	return GetFunctionHeader().Variant;
+}
+
+TArray<FRigVMVariantRef> URigVMLibraryNode::GetMatchingVariants() const
+{
+	const FRigVMGraphFunctionHeader Header = GetFunctionHeader();
+	if(Header.IsValid())
+	{
+		return Header.LibraryPointer.GetVariants(false);
+	}
+	return TArray<FRigVMVariantRef>();
+}
+
 
diff --git a/Engine/Plugins/Runtime/RigVM/Source/RigVMDeveloper/Private/RigVMModel/RigVMBuildData.cpp b/Engine/Plugins/Runtime/RigVM/Source/RigVMDeveloper/Private/RigVMModel/RigVMBuildData.cpp
--- a/Engine/Plugins/Runtime/RigVM/Source/RigVMDeveloper/Private/RigVMModel/RigVMBuildData.cpp
+++ b/Engine/Plugins/Runtime/RigVM/Source/RigVMDeveloper/Private/RigVMModel/RigVMBuildData.cpp
@@ -515,6 +515,26 @@ TArray<FRigVMVariantRef> URigVMBuildData::GatherAllFunctionVariantRefs()
 
 	for (const FAssetData& Asset : Assets)
 	{
+		if(Asset.IsAssetLoaded())
+		{
+			if(const IRigVMClientHost* ClientHost = Cast<IRigVMClientHost>(Asset.GetAsset()))
+			{
+				if(const URigVMFunctionLibrary* FunctionLibrary = ClientHost->GetLocalFunctionLibrary())
+				{
+					const TArray<URigVMLibraryNode*> Functions = FunctionLibrary->GetFunctions();
+					for(const URigVMLibraryNode* Function : Functions)
+					{
+						FRigVMVariant FunctionVariant = Function->GetFunctionVariant();
+						if(!FunctionVariant.Guid.IsValid())
+						{
+							FunctionVariant.Guid = FRigVMVariant::GenerateGUID(Function->GetPathName());
+						}
+						Result.Add(FRigVMVariantRef(Function->GetPathName(), FunctionVariant));
+					}
+					continue;
+				}
+			}
+		}
 		TArray<FRigVMVariantRef> VariantRefs = URigVMBuildData::GatherFunctionVariantRefsForAsset(Asset);
 		Result.Append(VariantRefs);
 	}
@@ -563,11 +583,25 @@ TArray<FRigVMVariantRef> URigVMBuildData::GatherAllAssetVariantRefs()
 
 	for (const FAssetData& Asset : Assets)
 	{
-		static const FName AssetVariantPropertyName = GET_MEMBER_NAME_CHECKED(URigVMBlueprint, AssetVariant);
-		const FProperty* AssetVariantProperty = CastField<FProperty>(URigVMBlueprint::StaticClass()->FindPropertyByName(AssetVariantPropertyName));
-		const FString VariantStr = Asset.GetTagValueRef<FString>(AssetVariantPropertyName);
 		FRigVMVariant AssetVariant;
-		AssetVariantProperty->ImportText_Direct(*VariantStr, &AssetVariant, nullptr, EPropertyPortFlags::PPF_None);
+		if(Asset.IsAssetLoaded())
+		{
+			if(const URigVMBlueprint* Blueprint = Cast<URigVMBlueprint>(Asset.GetAsset()))
+			{
+				AssetVariant = Blueprint->GetAssetVariant();
+			}
+		}
+		
+		if (!AssetVariant.Guid.IsValid())
+		{
+			static const FName AssetVariantPropertyName = GET_MEMBER_NAME_CHECKED(URigVMBlueprint, AssetVariant);
+			const FProperty* AssetVariantProperty = CastField<FProperty>(URigVMBlueprint::StaticClass()->FindPropertyByName(AssetVariantPropertyName));
+			const FString VariantStr = Asset.GetTagValueRef<FString>(AssetVariantPropertyName);
+			if(!VariantStr.IsEmpty())
+			{
+				AssetVariantProperty->ImportText_Direct(*VariantStr, &AssetVariant, nullptr, EPropertyPortFlags::PPF_None);
+			}
+		}
 
 		if (!AssetVariant.Guid.IsValid())
 		{
diff --git a/Engine/Plugins/Runtime/RigVM/Source/RigVMDeveloper/Private/RigVMModel/RigVMClient.cpp b/Engine/Plugins/Runtime/RigVM/Source/RigVMDeveloper/Private/RigVMModel/RigVMClient.cpp
--- a/Engine/Plugins/Runtime/RigVM/Source/RigVMDeveloper/Private/RigVMModel/RigVMClient.cpp
+++ b/Engine/Plugins/Runtime/RigVM/Source/RigVMDeveloper/Private/RigVMModel/RigVMClient.cpp
@@ -1573,6 +1573,7 @@ void FRigVMClient::HandleGraphModifiedEvent(ERigVMGraphNotifType InNotifType, UR
 		case ERigVMGraphNotifType::PinDisplayNameChanged: // A pin's display name / UI label has changed (Subject == URigVMPin)
 		case ERigVMGraphNotifType::PinCategoryChanged: // A pin's category has changed (Subject == URigVMPin)
 		case ERigVMGraphNotifType::PinCategoryExpansionChanged: // A pin category has been collapsed / expanded on the function node (Subject == URigVMNode)
+		case ERigVMGraphNotifType::FunctionVariantGuidChanged: // A function has changed its guid (Subject == URigVMLibraryNode)
 		{
 			URigVMNode* Node = Cast<URigVMNode>(InSubject);
 			if (const URigVMPin* Pin = Cast<URigVMPin>(InSubject))
@@ -1642,7 +1643,6 @@ void FRigVMClient::HandleGraphModifiedEvent(ERigVMGraphNotifType InNotifType, UR
 			break;
 		}
 		
-		
 		default:
 		{
 			break;
diff --git a/Engine/Plugins/Runtime/RigVM/Source/RigVMDeveloper/Private/RigVMModel/RigVMController.cpp b/Engine/Plugins/Runtime/RigVM/Source/RigVMDeveloper/Private/RigVMModel/RigVMController.cpp
--- a/Engine/Plugins/Runtime/RigVM/Source/RigVMDeveloper/Private/RigVMModel/RigVMController.cpp
+++ b/Engine/Plugins/Runtime/RigVM/Source/RigVMDeveloper/Private/RigVMModel/RigVMController.cpp
@@ -34,6 +34,7 @@
 #include "RigVMModel/Nodes/RigVMBranchNode.h"
 #include "RigVMModel/Nodes/RigVMArrayNode.h"
 #include "Logging/LogScopedVerbosityOverride.h"
+#include "ScopedTransaction.h"
 
 #include UE_INLINE_GENERATED_CPP_BY_NAME(RigVMController)
 
@@ -8114,9 +8115,12 @@ bool URigVMController::AddEmptyPinCategory(const URigVMNode* InNode, const FStri
 		return false;
 	}
 
-	if(!GetGraph()->GetSchema()->SupportsNodeLayouts(GetGraph()))
+	if(const URigVMSchema* Schema = GetGraph()->GetSchema())
 	{
-		return false;
+		if(!Schema->SupportsNodeLayouts(GetGraph()))
+		{
+			return false;
+		}
 	}
 
 	FRigVMControllerCompileBracketScope CompileScope(this);
@@ -8241,9 +8245,12 @@ bool URigVMController::SetPinCategory(URigVMPin* InPin, const FString& InCategor
 		return false;
 	}
 
-	if(!GetGraph()->GetSchema()->SupportsNodeLayouts(GetGraph()))
+	if(const URigVMSchema* Schema = GetGraph()->GetSchema())
 	{
-		return false;
+		if(!Schema->SupportsNodeLayouts(GetGraph()))
+		{
+			return false;
+		}
 	}
 
 	TArray<FString> Categories;
@@ -8496,9 +8503,12 @@ bool URigVMController::RemovePinCategory(const URigVMNode* InNode, const FString
 		return false;
 	}
 
-	if(!GetGraph()->GetSchema()->SupportsNodeLayouts(GetGraph()))
+	if(const URigVMSchema* Schema = GetGraph()->GetSchema())
 	{
-		return false;
+		if(!Schema->SupportsNodeLayouts(GetGraph()))
+		{
+			return false;
+		}
 	}
 
 	const TArray<FString> AllCategories = InNode->GetPinCategories();
@@ -8628,9 +8638,12 @@ bool URigVMController::RenamePinCategory(const URigVMNode* InNode, const FString
 		return false;
 	}
 
-	if(!GetGraph()->GetSchema()->SupportsNodeLayouts(GetGraph()))
+	if(const URigVMSchema* Schema = GetGraph()->GetSchema())
 	{
-		return false;
+		if(!Schema->SupportsNodeLayouts(GetGraph()))
+		{
+			return false;
+		}
 	}
 
 	if(InOldPinCategory.Equals(FRigVMPinCategory::GetDefaultCategoryName(), ESearchCase::IgnoreCase))
@@ -8792,9 +8805,12 @@ bool URigVMController::SetPinCategoryIndex(const URigVMNode* InNode, const FStri
 		return false;
 	}
 
-	if(!GetGraph()->GetSchema()->SupportsNodeLayouts(GetGraph()))
+	if(const URigVMSchema* Schema = GetGraph()->GetSchema())
 	{
-		return false;
+		if(!Schema->SupportsNodeLayouts(GetGraph()))
+		{
+			return false;
+		}
 	}
 
 	// for entry or return nodes we relay to the outer pins
@@ -8883,9 +8899,12 @@ bool URigVMController::SetPinCategoryExpansion(const URigVMNode* InNode, const F
 		return false;
 	}
 
-	if(!GetGraph()->GetSchema()->SupportsNodeLayouts(GetGraph()))
+	if(const URigVMSchema* Schema = GetGraph()->GetSchema())
 	{
-		return false;
+		if(!Schema->SupportsNodeLayouts(GetGraph()))
+		{
+			return false;
+		}
 	}
 
 	const TArray<FString> AllCategories = InNode->GetPinCategories();
@@ -9107,9 +9126,12 @@ bool URigVMController::SetNodeLayout(const FName& InNodeName, FRigVMNodeLayout I
 	const URigVMGraph* Graph = GetGraph();
 	check(Graph);
 
-	if(!Graph->GetSchema()->SupportsNodeLayouts(Graph))
+	if(const URigVMSchema* Schema = GetGraph()->GetSchema())
 	{
-		return false;
+		if(!Schema->SupportsNodeLayouts(GetGraph()))
+		{
+			return false;
+		}
 	}
 
 	const URigVMNode* Node = Graph->FindNodeByName(InNodeName);
@@ -14110,6 +14132,132 @@ TArray<FRigVMVariantRef> URigVMController::FindVariantsOfFunction(const FName& I
 	return FunctionData->Header.LibraryPointer.GetVariants(true);
 }
 
+bool URigVMController::SplitFunctionVariant(const FName& InFunctionName, bool bSetupUndoRedo, bool bPrintPythonCommand)
+{
+	if(!IsValidGraph())
+	{
+		return false;
+	}
+
+	URigVMGraph* Graph = GetGraph();
+	check(Graph);
+
+	if (!Graph->IsA<URigVMFunctionLibrary>())
+	{
+		ReportError(TEXT("Cannot split function variants on normal (non-function-library) graphs."));
+		return false;
+	}
+
+	URigVMFunctionLibrary* Library = CastChecked<URigVMFunctionLibrary>(Graph);
+	const URigVMLibraryNode* LibraryNode = Library->FindFunction(InFunctionName);
+	if(LibraryNode == nullptr)
+	{
+		ReportError(TEXT("Function '%s' cannot be found"));
+		return false;
+	}
+
+	FRigVMVariant Variant;
+	if (FRigVMGraphFunctionStore* FunctionStore = GetGraphFunctionStore())
+	{
+		if (FRigVMGraphFunctionData* FunctionData = FunctionStore->FindFunctionByName(InFunctionName))
+		{
+			Variant = FunctionData->Header.Variant;
+		}
+	}
+
+	// prefer the deterministic path based guid - and only fall back on
+	// random guids if necessary
+	const FGuid PreviousGuid = Variant.Guid; 
+	Variant.Guid = FRigVMVariant::GenerateGUID(LibraryNode->GetPathName());
+	if(Variant.Guid == PreviousGuid)
+	{
+		Variant.Guid = FRigVMVariant::GenerateGUID();
+	}
+
+	FScopedTransaction Transaction(NSLOCTEXT("RigVMController", "SplitAssetVariant", "Split Asset Variant"));
+	if(bSetupUndoRedo)
+	{
+		Library->Modify();
+	}
+	Library->FunctionToVariant.FindOrAdd(InFunctionName) = Variant;
+
+	Notify(ERigVMGraphNotifType::FunctionVariantGuidChanged, Library->FindFunction(InFunctionName));
+
+	if (bPrintPythonCommand)
+	{
+		const FString GraphName = GetSchema()->GetSanitizedGraphName(GetGraph()->GetGraphName());
+		RigVMPythonUtils::Print(GetSchema()->GetGraphOuterName(GetGraph()), 
+			FString::Printf(TEXT("blueprint.get_controller_by_name('%s').split_function_variant('%s')"),
+				*GraphName,
+				*InFunctionName.ToString()));
+	}
+
+	return true;
+}
+
+bool URigVMController::JoinFunctionVariant(const FName& InFunctionName, const FGuid& InGuid, bool bSetupUndoRedo, bool bPrintPythonCommand)
+{
+	if(!IsValidGraph())
+	{
+		return false;
+	}
+
+	URigVMGraph* Graph = GetGraph();
+	check(Graph);
+
+	if (!Graph->IsA<URigVMFunctionLibrary>())
+	{
+		ReportError(TEXT("Cannot join function variants on normal (non-function-library) graphs."));
+		return false;
+	}
+
+	URigVMFunctionLibrary* Library = CastChecked<URigVMFunctionLibrary>(Graph);
+	if(Library->FindFunction(InFunctionName) == nullptr)
+	{
+		ReportError(TEXT("Function '%s' cannot be found"));
+		return false;
+	}
+
+	FRigVMVariant Variant;
+	if (FRigVMGraphFunctionStore* FunctionStore = GetGraphFunctionStore())
+	{
+		if (FRigVMGraphFunctionData* FunctionData = FunctionStore->FindFunctionByName(InFunctionName))
+		{
+			Variant = FunctionData->Header.Variant;
+		}
+	}
+
+	if(Variant.Guid == InGuid)
+	{
+		return false;
+	}
+	Variant.Guid = InGuid;
+
+	FScopedTransaction Transaction(NSLOCTEXT("RigVMController", "SplitAssetVariant", "Split Asset Variant"));
+	if(bSetupUndoRedo)
+	{
+		Library->Modify();
+	}
+	Library->FunctionToVariant.FindOrAdd(InFunctionName) = Variant;
+
+	Notify(ERigVMGraphNotifType::FunctionVariantGuidChanged, Library->FindFunction(InFunctionName));
+
+	if (bPrintPythonCommand)
+	{
+		const FString GraphName = GetSchema()->GetSanitizedGraphName(GetGraph()->GetGraphName());
+		RigVMPythonUtils::Print(GetSchema()->GetGraphOuterName(GetGraph()), 
+			FString::Printf(TEXT("guid = unreal.Guid()")));
+		RigVMPythonUtils::Print(GetSchema()->GetGraphOuterName(GetGraph()), 
+			FString::Printf(TEXT("guid.import_text(\"%s\")"), *Variant.Guid.ToString()));
+		RigVMPythonUtils::Print(GetSchema()->GetGraphOuterName(GetGraph()), 
+			FString::Printf(TEXT("blueprint.get_controller_by_name('%s').join_function_variant('%s', guid)"),
+				*GraphName,
+				*InFunctionName.ToString()));
+	}
+
+	return true;
+}
+
 FRigVMGraphVariableDescription URigVMController::AddLocalVariable(const FName& InVariableName, const FString& InCPPType, UObject* InCPPTypeObject, const FString& InDefaultValue, bool bSetupUndoRedo, bool bPrintPythonCommand)
 {
 	FRigVMGraphVariableDescription NewVariable;
diff --git a/Engine/Plugins/Runtime/RigVM/Source/RigVMDeveloper/Public/RigVMBlueprint.h b/Engine/Plugins/Runtime/RigVM/Source/RigVMDeveloper/Public/RigVMBlueprint.h
--- a/Engine/Plugins/Runtime/RigVM/Source/RigVMDeveloper/Public/RigVMBlueprint.h
+++ b/Engine/Plugins/Runtime/RigVM/Source/RigVMDeveloper/Public/RigVMBlueprint.h
@@ -374,6 +374,9 @@ public:
 	UFUNCTION(BlueprintCallable, Category = "RigVM Blueprint")
 	virtual URigVMFunctionLibrary* GetLocalFunctionLibrary() const override;
 
+	UFUNCTION(BlueprintCallable, Category = "RigVM Blueprint")
+	virtual URigVMFunctionLibrary* GetOrCreateLocalFunctionLibrary(bool bSetupUndoRedo = true) override;
+
 	UFUNCTION(BlueprintCallable, Category = "RigVM Blueprint")
 	virtual URigVMGraph* AddModel(FString InName = TEXT("Rig Graph"), bool bSetupUndoRedo = true, bool bPrintPythonCommand = true) override;
 
@@ -563,6 +566,21 @@ public:
 	bool ChangeMemberVariableType(const FName& InName, const FString& InCPPType, bool bIsPublic = false, bool bIsReadOnly = false, FString InDefaultValue = TEXT(""));
 
 	const FRigVMVariant& GetAssetVariant() const { return AssetVariant; }
+
+	UFUNCTION(BlueprintPure, Category = "Variables", meta = (DisplayName = "GetAssetVariant", ScriptName = "GetAssetVariant"))
+	FRigVMVariant GetAssetVariantBP() const;
+
+	/** Resets the asset's guid to a new one and splits it from the former variant set */
+	UFUNCTION(BlueprintCallable, Category = "Variables")
+	bool SplitAssetVariant();
+
+	/** Merges the asset's guid with a provided one to join the variant set */
+	UFUNCTION(BlueprintCallable, Category = "Variables")
+	bool JoinAssetVariant(const FGuid& InGuid);
+
+	UFUNCTION(BlueprintPure, Category = "Variables")
+	TArray<FRigVMVariantRef> GetMatchingVariants() const;
+
 #endif
 
 private:
diff --git a/Engine/Plugins/Runtime/RigVM/Source/RigVMDeveloper/Public/RigVMModel/Nodes/RigVMLibraryNode.h b/Engine/Plugins/Runtime/RigVM/Source/RigVMDeveloper/Public/RigVMModel/Nodes/RigVMLibraryNode.h
--- a/Engine/Plugins/Runtime/RigVM/Source/RigVMDeveloper/Public/RigVMModel/Nodes/RigVMLibraryNode.h
+++ b/Engine/Plugins/Runtime/RigVM/Source/RigVMDeveloper/Public/RigVMModel/Nodes/RigVMLibraryNode.h
@@ -60,6 +60,12 @@ public:
 
 	virtual FRigVMGraphFunctionIdentifier GetFunctionIdentifier() const;
 	FRigVMGraphFunctionHeader GetFunctionHeader(IRigVMGraphFunctionHost* InHostObject = nullptr) const;
+	
+	UFUNCTION(BlueprintPure, Category = RigVMLibraryNode)
+	FRigVMVariant GetFunctionVariant() const;
+
+	UFUNCTION(BlueprintPure, Category = RigVMLibraryNode)
+	TArray<FRigVMVariantRef> GetMatchingVariants() const;
 
 protected:
 
diff --git a/Engine/Plugins/Runtime/RigVM/Source/RigVMDeveloper/Public/RigVMModel/RigVMClient.h b/Engine/Plugins/Runtime/RigVM/Source/RigVMDeveloper/Public/RigVMModel/RigVMClient.h
--- a/Engine/Plugins/Runtime/RigVM/Source/RigVMDeveloper/Public/RigVMModel/RigVMClient.h
+++ b/Engine/Plugins/Runtime/RigVM/Source/RigVMDeveloper/Public/RigVMModel/RigVMClient.h
@@ -103,6 +103,8 @@ public:
 
 	virtual URigVMFunctionLibrary* GetLocalFunctionLibrary() const = 0;
 
+	virtual URigVMFunctionLibrary* GetOrCreateLocalFunctionLibrary(bool bSetupUndoRedo = true) = 0;
+
 	virtual URigVMGraph* AddModel(FString InName = TEXT("Rig Graph"), bool bSetupUndoRedo = true, bool bPrintPythonCommand = true) = 0;
 
 	virtual bool RemoveModel(FString InName = TEXT("Rig Graph"), bool bSetupUndoRedo = true, bool bPrintPythonCommand = true) = 0;
diff --git a/Engine/Plugins/Runtime/RigVM/Source/RigVMDeveloper/Public/RigVMModel/RigVMController.h b/Engine/Plugins/Runtime/RigVM/Source/RigVMDeveloper/Public/RigVMModel/RigVMController.h
--- a/Engine/Plugins/Runtime/RigVM/Source/RigVMDeveloper/Public/RigVMModel/RigVMController.h
+++ b/Engine/Plugins/Runtime/RigVM/Source/RigVMDeveloper/Public/RigVMModel/RigVMController.h
@@ -1037,7 +1037,6 @@ public:
 	UFUNCTION(BlueprintCallable, Category = RigVMController)
 	bool AddTagToFunctionVariant(const FName& InFunctionName, const FRigVMTag& InTag, bool bSetupUndoRedo = true, bool bPrintPythonCommand = false);
 
-
 	// Adds a tag to a function variant
 	UFUNCTION(BlueprintCallable, Category = RigVMController)
 	bool RemoveTagFromFunctionVariant(const FName& InFunctionName, const FName& InTagName, bool bSetupUndoRedo = true, bool bPrintPythonCommand = false);
@@ -1046,6 +1045,14 @@ public:
 	UFUNCTION(BlueprintCallable, Category = RigVMController)
 	TArray<FRigVMVariantRef> FindVariantsOfFunction(const FName& InFunctionName);
 
+	/** Resets the function's guid to a new one and splits it from the former variant set */
+	UFUNCTION(BlueprintCallable, Category = RigVMController)
+	bool SplitFunctionVariant(const FName& InFunctionName, bool bSetupUndoRedo = true, bool bPrintPythonCommand = false);
+
+	/** Merges the function's guid with a provided one to join the variant set */
+	UFUNCTION(BlueprintCallable, Category = RigVMController)
+	bool JoinFunctionVariant(const FName& InFunctionName, const FGuid& InGuid, bool bSetupUndoRedo = true, bool bPrintPythonCommand = false);
+
 	// Add a local variable to the graph
 	UFUNCTION(BlueprintCallable, Category = RigVMController)
 	FRigVMGraphVariableDescription AddLocalVariable(const FName& InVariableName, const FString& InCPPType, UObject* InCPPTypeObject, const FString& InDefaultValue, bool bSetupUndoRedo = true, bool bPrintPythonCommand = false);
diff --git a/Engine/Plugins/Runtime/RigVM/Source/RigVMDeveloper/Public/RigVMModel/RigVMNotifications.h b/Engine/Plugins/Runtime/RigVM/Source/RigVMDeveloper/Public/RigVMModel/RigVMNotifications.h
--- a/Engine/Plugins/Runtime/RigVM/Source/RigVMDeveloper/Public/RigVMModel/RigVMNotifications.h
+++ b/Engine/Plugins/Runtime/RigVM/Source/RigVMDeveloper/Public/RigVMModel/RigVMNotifications.h
@@ -58,6 +58,7 @@ enum class ERigVMGraphNotifType : uint8
 	PinCategoryChanged, // The category of a pin has changed - requiring a rebuild of the node user interface (Subject == URigVMPin)
 	PinCategoriesChanged, // The category list of a node has changed - requiring a rebuild of the node user interface (Subject == URigVMNode)
 	PinCategoryExpansionChanged, // The category of a pin expanded / collapsed (Subject == URigVMNode)
+	FunctionVariantGuidChanged, // The guid for a function has changed (Subject == URigVMLibraryNode)
 	Invalid // The max for this enum (used for guarding)
 };
 
