RigVM: Move trait default value initialization to Serialize

Avoids issues with PostLoad ordering and pin default access

diff --git a/Engine/Plugins/Runtime/RigVM/Source/RigVMDeveloper/Private/RigVMModel/RigVMNode.cpp b/Engine/Plugins/Runtime/RigVM/Source/RigVMDeveloper/Private/RigVMModel/RigVMNode.cpp
--- a/Engine/Plugins/Runtime/RigVM/Source/RigVMDeveloper/Private/RigVMModel/RigVMNode.cpp
+++ b/Engine/Plugins/Runtime/RigVM/Source/RigVMDeveloper/Private/RigVMModel/RigVMNode.cpp
@@ -36,9 +36,9 @@ URigVMNode::~URigVMNode()
 {
 }
 
-void URigVMNode::PostLoad()
+void URigVMNode::Serialize(FArchive& Ar)
 {
-	Super::PostLoad();
+	Super::Serialize(Ar);
 
 	for (const FString& TraitRootPinName : TraitRootPinNames)
 	{
diff --git a/Engine/Plugins/Runtime/RigVM/Source/RigVMDeveloper/Public/RigVMModel/RigVMNode.h b/Engine/Plugins/Runtime/RigVM/Source/RigVMDeveloper/Public/RigVMModel/RigVMNode.h
--- a/Engine/Plugins/Runtime/RigVM/Source/RigVMDeveloper/Public/RigVMModel/RigVMNode.h
+++ b/Engine/Plugins/Runtime/RigVM/Source/RigVMDeveloper/Public/RigVMModel/RigVMNode.h
@@ -35,7 +35,7 @@ public:
 	// Default destructor
 	virtual ~URigVMNode();
 
-	void PostLoad() override;
+	void Serialize(FArchive& Ar) override;
 
 	// Returns the a . separated string containing all of the
 	// names used to reach this Node within the Graph.
