RigVM: Guard against schema nullptr

#rnx

diff --git a/Engine/Plugins/Runtime/RigVM/Source/RigVMEditor/Private/EdGraph/RigVMEdGraphPanelNodeFactory.cpp b/Engine/Plugins/Runtime/RigVM/Source/RigVMEditor/Private/EdGraph/RigVMEdGraphPanelNodeFactory.cpp
--- a/Engine/Plugins/Runtime/RigVM/Source/RigVMEditor/Private/EdGraph/RigVMEdGraphPanelNodeFactory.cpp
+++ b/Engine/Plugins/Runtime/RigVM/Source/RigVMEditor/Private/EdGraph/RigVMEdGraphPanelNodeFactory.cpp
@@ -50,13 +50,16 @@ TSharedPtr<SGraphNode> FRigVMEdGraphPanelNodeFactory::CreateNode(UEdGraphNode* N
 	
 	if (UEdGraphNode_Comment* CommentNode = Cast<UEdGraphNode_Comment>(Node))
 	{
-		if (CommentNode->GetSchema()->IsA(URigVMEdGraphSchema::StaticClass()))
+		if (const UEdGraphSchema* Schema = CommentNode->GetSchema())
 		{
-			TSharedRef<SGraphNode> GraphNode =
-				SNew(SRigVMGraphNodeComment, CommentNode);
+			if (Schema->IsA(URigVMEdGraphSchema::StaticClass()))
+			{
+				TSharedRef<SGraphNode> GraphNode =
+					SNew(SRigVMGraphNodeComment, CommentNode);
 
-			GraphNode->SlatePrepass();
-			return GraphNode;
+				GraphNode->SlatePrepass();
+				return GraphNode;
+			}
 		}
 	}
 
