Fixed RigVM crash when renaming pin categories in the node customization layout

#rnx

diff --git a/Engine/Plugins/Runtime/RigVM/Source/RigVMDeveloper/Private/RigVMModel/RigVMController.cpp b/Engine/Plugins/Runtime/RigVM/Source/RigVMDeveloper/Private/RigVMModel/RigVMController.cpp
--- a/Engine/Plugins/Runtime/RigVM/Source/RigVMDeveloper/Private/RigVMModel/RigVMController.cpp
+++ b/Engine/Plugins/Runtime/RigVM/Source/RigVMDeveloper/Private/RigVMModel/RigVMController.cpp
@@ -8696,8 +8696,30 @@ bool URigVMController::RenamePinCategory(const URigVMNode* InNode, const FString
 	}
 
 	TArray<URigVMPin*> PinsToUpdate;
+	TArray<FString> PinsCategoriesToRename;
 	
-	// first update the category list on the node
+	// first obtain the pins to update
+	for (const FString& CategoryToRename : CategoriesToRename)
+	{
+		PinsToUpdate.Append(InNode->GetPinsForCategory(CategoryToRename));
+	}
+
+	{
+		const TGuardValue<bool> GuardRefreshFunctions(bSuspendRefreshingFunctionReferences, true);
+
+		// Set all pins temporarily to the default category
+		for (URigVMPin* PinToUpdate : PinsToUpdate)
+		{
+			const FString CategoryToRename = PinToUpdate->GetCategory();
+
+			// store the original pin category name as it will be needed later at the final set
+			PinsCategoriesToRename.Add(CategoryToRename);
+
+			(void)SetPinCategory(PinToUpdate, FString(), bSetupUndoRedo);
+		}
+	}
+
+	// update the category list on the node
 	for(const FString& CategoryToRename : CategoriesToRename)
 	{
 		const bool bWasExpanded = InNode->IsPinCategoryExpanded(CategoryToRename);
@@ -8711,26 +8733,29 @@ bool URigVMController::RenamePinCategory(const URigVMNode* InNode, const FString
         	const FString NewCategory = NewPrefix + CategoryToRename.RightChop(OldPrefix.Len());
 			const_cast<URigVMNode*>(InNode)->PinCategories[Index] = NewCategory;
         }
-		PinsToUpdate.Append(InNode->GetPinsForCategory(CategoryToRename));
 
 		const_cast<URigVMNode*>(InNode)->PinCategoryExpansion.Remove(CategoryToRename);
 		const_cast<URigVMNode*>(InNode)->PinCategoryExpansion.FindOrAdd(InNode->PinCategories[Index]) = bWasExpanded;
 	}
 
-	// then update the pins and their categories
-	for(URigVMPin* PinToUpdate : PinsToUpdate)
 	{
-		const TGuardValue<bool> GuardRefreshFunctions(bSuspendRefreshingFunctionReferences, true); 
+		const TGuardValue<bool> GuardRefreshFunctions(bSuspendRefreshingFunctionReferences, true);
 
-		const FString CategoryToRename = PinToUpdate->GetCategory();
-		if(CategoryToRename.Equals(InOldPinCategory, ESearchCase::CaseSensitive))
-		{
-			(void)SetPinCategory(PinToUpdate, InNewPinCategory, bSetupUndoRedo);
-		}
-		else if(CategoryToRename.StartsWith(OldPrefix, ESearchCase::CaseSensitive))
+		// then update the pins to their new categories
+		for (int PinIndex = 0; PinIndex < PinsToUpdate.Num(); ++PinIndex)
 		{
-			const FString NewCategory = NewPrefix + CategoryToRename.RightChop(OldPrefix.Len());
-			(void)SetPinCategory(PinToUpdate, NewCategory, bSetupUndoRedo);
+			URigVMPin* PinToUpdate = PinsToUpdate[PinIndex];
+
+			const FString& CategoryToRename = PinsCategoriesToRename[PinIndex];
+			if(CategoryToRename.Equals(InOldPinCategory, ESearchCase::CaseSensitive))
+			{
+				(void)SetPinCategory(PinToUpdate, InNewPinCategory, bSetupUndoRedo);
+			}
+			else if(CategoryToRename.StartsWith(OldPrefix, ESearchCase::CaseSensitive))
+			{
+				const FString NewCategory = NewPrefix + CategoryToRename.RightChop(OldPrefix.Len());
+				(void)SetPinCategory(PinToUpdate, NewCategory, bSetupUndoRedo);
+			}
 		}
 	}
 	
