RigVM: Unit test for traits in functions

#rnx

diff --git a/Engine/Plugins/Runtime/RigVM/Source/RigVM/Private/RigVMCore/RigVMGraphFunctionDefinition.cpp b/Engine/Plugins/Runtime/RigVM/Source/RigVM/Private/RigVMCore/RigVMGraphFunctionDefinition.cpp
--- a/Engine/Plugins/Runtime/RigVM/Source/RigVM/Private/RigVMCore/RigVMGraphFunctionDefinition.cpp
+++ b/Engine/Plugins/Runtime/RigVM/Source/RigVM/Private/RigVMCore/RigVMGraphFunctionDefinition.cpp
@@ -6,6 +6,23 @@
 
 #include UE_INLINE_GENERATED_CPP_BY_NAME(RigVMGraphFunctionDefinition)
 
+FRigVMPropertyDescription FRigVMFunctionCompilationPropertyDescription::ToPropertyDescription() const
+{
+	return FRigVMPropertyDescription(Name, CPPType, CPPTypeObject.LoadSynchronous(), DefaultValue);
+}
+
+TArray<FRigVMPropertyDescription> FRigVMFunctionCompilationPropertyDescription::ToPropertyDescription(
+	const TArray<FRigVMFunctionCompilationPropertyDescription>& InDescriptions)
+{
+	TArray<FRigVMPropertyDescription> Result;
+	Result.Reserve(InDescriptions.Num());
+	for(const FRigVMFunctionCompilationPropertyDescription& Description : InDescriptions)
+	{
+		Result.Add(Description.ToPropertyDescription());
+	}
+	return Result;
+}
+
 FRigVMExternalVariable FRigVMGraphFunctionArgument::GetExternalVariable() const
 {
 	FRigVMExternalVariable Variable;
diff --git a/Engine/Plugins/Runtime/RigVM/Source/RigVM/Public/RigVMCore/RigVMGraphFunctionDefinition.h b/Engine/Plugins/Runtime/RigVM/Source/RigVM/Public/RigVMCore/RigVMGraphFunctionDefinition.h
--- a/Engine/Plugins/Runtime/RigVM/Source/RigVM/Public/RigVMCore/RigVMGraphFunctionDefinition.h
+++ b/Engine/Plugins/Runtime/RigVM/Source/RigVM/Public/RigVMCore/RigVMGraphFunctionDefinition.h
@@ -16,7 +16,7 @@ class IRigVMGraphFunctionHost;
 struct FRigVMGraphFunctionData;
 
 USTRUCT()
-struct FRigVMFunctionCompilationPropertyDescription
+struct RIGVM_API FRigVMFunctionCompilationPropertyDescription
 {
 	GENERATED_BODY()
 
@@ -54,6 +54,9 @@ struct FRigVMFunctionCompilationPropertyDescription
 		Ar << Data.DefaultValue;
 		return Ar;
 	}
+
+	FRigVMPropertyDescription ToPropertyDescription() const;
+	static TArray<FRigVMPropertyDescription> ToPropertyDescription(const TArray<FRigVMFunctionCompilationPropertyDescription>& InDescriptions);
 };
 
 USTRUCT()
diff --git a/Engine/Plugins/Runtime/RigVM/Source/RigVMDeveloper/Private/RigVMModel/RigVMBuildData.cpp b/Engine/Plugins/Runtime/RigVM/Source/RigVMDeveloper/Private/RigVMModel/RigVMBuildData.cpp
--- a/Engine/Plugins/Runtime/RigVM/Source/RigVMDeveloper/Private/RigVMModel/RigVMBuildData.cpp
+++ b/Engine/Plugins/Runtime/RigVM/Source/RigVMDeveloper/Private/RigVMModel/RigVMBuildData.cpp
@@ -130,9 +130,27 @@ void URigVMBuildData::SetupRigVMGraphFunctionPointers()
 	};
 	FRigVMGraphFunctionData::GetFunctionHostFromObjectFunc = [this](UObject* InObject) -> IRigVMGraphFunctionHost*
 	{
-		if(IRigVMGraphFunctionHost* Host = InObject->GetImplementingOuter<IRigVMGraphFunctionHost>())
+		if(IRigVMGraphFunctionHost* FunctionHost = Cast<IRigVMGraphFunctionHost>(InObject))
 		{
-			return Host;
+			return FunctionHost;
+		}
+		if(IRigVMGraphFunctionHost* OuterFunctionHost = InObject->GetImplementingOuter<IRigVMGraphFunctionHost>())
+		{
+			return OuterFunctionHost;
+		}
+		if(IRigVMClientHost* ClientHost = Cast<IRigVMClientHost>(InObject))
+		{
+			if(IRigVMGraphFunctionHost* FunctionHost = ClientHost->GetRigVMGraphFunctionHost())
+			{
+				return FunctionHost;
+			}
+		}
+		if(IRigVMClientHost* OuterClientHost = InObject->GetImplementingOuter<IRigVMClientHost>())
+		{
+			if(IRigVMGraphFunctionHost* FunctionHost = OuterClientHost->GetRigVMGraphFunctionHost())
+			{
+				return FunctionHost;
+			}
 		}
 		if(URigVMBlueprint* Blueprint = Cast<URigVMBlueprint>(InObject))
 		{
