Determinism pass on RigVM hashing

diff --git a/Engine/Plugins/Runtime/RigVM/Source/RigVM/Public/RigVMCore/RigVMByteCode.h b/Engine/Plugins/Runtime/RigVM/Source/RigVM/Public/RigVMCore/RigVMByteCode.h
--- a/Engine/Plugins/Runtime/RigVM/Source/RigVM/Public/RigVMCore/RigVMByteCode.h
+++ b/Engine/Plugins/Runtime/RigVM/Source/RigVM/Public/RigVMCore/RigVMByteCode.h
@@ -65,7 +65,7 @@ struct FRigVMBranchInfoKey
 				GetTypeHash(InKey.InstructionIndex),
 				GetTypeHash(InKey.ArgumentIndex)
 			),
-			GetTypeHash(InKey.Label)
+			GetTypeHash(InKey.Label.ToString())
 		);
 	}
 
@@ -925,7 +925,7 @@ struct RIGVM_API FRigVMInvokeEntryOp : public FRigVMBaseOp
 
 	friend uint32 GetTypeHash(const FRigVMInvokeEntryOp& Op)
 	{
-		return HashCombine(GetTypeHash((const FRigVMBaseOp&)Op), GetTypeHash(Op.EntryName));
+		return HashCombine(GetTypeHash((const FRigVMBaseOp&)Op), GetTypeHash(Op.EntryName.ToString()));
 	}
 
 	void Serialize(FArchive& Ar);
diff --git a/Engine/Plugins/Runtime/RigVM/Source/RigVM/Public/RigVMCore/RigVMGraphFunctionDefinition.h b/Engine/Plugins/Runtime/RigVM/Source/RigVM/Public/RigVMCore/RigVMGraphFunctionDefinition.h
--- a/Engine/Plugins/Runtime/RigVM/Source/RigVM/Public/RigVMCore/RigVMGraphFunctionDefinition.h
+++ b/Engine/Plugins/Runtime/RigVM/Source/RigVM/Public/RigVMCore/RigVMGraphFunctionDefinition.h
@@ -335,8 +335,8 @@ struct RIGVM_API FRigVMGraphFunctionArgument
 
 	friend uint32 GetTypeHash(const FRigVMGraphFunctionArgument& Argument)
 	{
-		uint32 Hash = HashCombine(GetTypeHash(Argument.Name), GetTypeHash(Argument.DisplayName));
-		Hash = HashCombine(Hash, GetTypeHash(Argument.CPPType));
+		uint32 Hash = HashCombine(GetTypeHash(Argument.Name.ToString()), GetTypeHash(Argument.DisplayName.ToString()));
+		Hash = HashCombine(Hash, GetTypeHash(Argument.CPPType.ToString()));
 		Hash = HashCombine(Hash, GetTypeHash(Argument.CPPTypeObject));
 		Hash = HashCombine(Hash, GetTypeHash(Argument.bIsArray));
 		Hash = HashCombine(Hash, GetTypeHash(Argument.Direction));
@@ -396,7 +396,7 @@ public:
 
 	friend uint32 GetTypeHash(const FRigVMGraphFunctionIdentifier& Pointer)
 	{
-		return HashCombine(GetTypeHash(Pointer.GetLibraryNodePath()), GetTypeHash(Pointer.HostObject));
+		return HashCombine(GetTypeHash(Pointer.GetLibraryNodePath()), GetTypeHash(Pointer.HostObject.ToString()));
 	}
 
 	bool operator==(const FRigVMGraphFunctionIdentifier& Other) const
diff --git a/Engine/Plugins/Runtime/RigVM/Source/RigVM/Public/RigVMCore/RigVMTemplate.h b/Engine/Plugins/Runtime/RigVM/Source/RigVM/Public/RigVMCore/RigVMTemplate.h
--- a/Engine/Plugins/Runtime/RigVM/Source/RigVM/Public/RigVMCore/RigVMTemplate.h
+++ b/Engine/Plugins/Runtime/RigVM/Source/RigVM/Public/RigVMCore/RigVMTemplate.h
@@ -103,7 +103,7 @@ struct RIGVM_API FRigVMTemplateArgumentType
 
 	RIGVM_API friend FORCEINLINE uint32 GetTypeHash(const FRigVMTemplateArgumentType& InType)
 	{
-		return GetTypeHash(InType.CPPType);
+		return GetTypeHash(InType.CPPType.ToString());
 	}
 
 	FName GetCPPTypeObjectPath() const
diff --git a/Engine/Plugins/Runtime/RigVM/Source/RigVM/Public/RigVMCore/RigVMVariant.h b/Engine/Plugins/Runtime/RigVM/Source/RigVM/Public/RigVMCore/RigVMVariant.h
--- a/Engine/Plugins/Runtime/RigVM/Source/RigVM/Public/RigVMCore/RigVMVariant.h
+++ b/Engine/Plugins/Runtime/RigVM/Source/RigVM/Public/RigVMCore/RigVMVariant.h
@@ -65,7 +65,7 @@ struct RIGVM_API FRigVMTag
 
 	friend uint32 GetTypeHash(const FRigVMTag& InTag)
 	{
-		uint32 Hash = GetTypeHash(InTag.Name);
+		uint32 Hash = GetTypeHash(InTag.Name.ToString());
 		Hash = HashCombine(Hash, GetTypeHash(InTag.Label));
 		Hash = HashCombine(Hash, GetTypeHash(InTag.ToolTip.ToString()));
 		Hash = HashCombine(Hash, GetTypeHash(InTag.Color));
@@ -158,7 +158,7 @@ struct RIGVM_API FRigVMVariantRef
 
 	friend uint32 GetTypeHash(const FRigVMVariantRef& InVariantRef)
 	{
-		return HashCombine(GetTypeHash(InVariantRef.ObjectPath), GetTypeHash(InVariantRef.Variant.Guid));
+		return HashCombine(GetTypeHash(InVariantRef.ObjectPath.ToString()), GetTypeHash(InVariantRef.Variant.Guid));
 	}
 
 	UPROPERTY(BlueprintReadOnly, Category=RigVMVariantRef)
diff --git a/Engine/Plugins/Runtime/RigVM/Source/RigVMDeveloper/Private/RigVMModel/RigVMController.cpp b/Engine/Plugins/Runtime/RigVM/Source/RigVMDeveloper/Private/RigVMModel/RigVMController.cpp
--- a/Engine/Plugins/Runtime/RigVM/Source/RigVMDeveloper/Private/RigVMModel/RigVMController.cpp
+++ b/Engine/Plugins/Runtime/RigVM/Source/RigVMDeveloper/Private/RigVMModel/RigVMController.cpp
@@ -583,7 +583,7 @@ void FRigVMPinInfo::CorrectExecuteTypeIndex()
 uint32 GetTypeHash(const FRigVMPinInfo& InPin)
 {
 	uint32 Hash = 0; //GetTypeHash(InPin.ParentIndex);
-	Hash = HashCombine(Hash, GetTypeHash(InPin.Name));
+	Hash = HashCombine(Hash, GetTypeHash(InPin.Name.ToString()));
 	Hash = HashCombine(Hash, GetTypeHash((int32)InPin.Direction));
 	Hash = HashCombine(Hash, GetTypeHash((int32)InPin.TypeIndex));
 	Hash = HashCombine(Hash, GetTypeHash(InPin.bIsArray));
diff --git a/Engine/Plugins/Runtime/RigVM/Source/RigVMDeveloper/Public/RigVMCompiler/RigVMCompiler.h b/Engine/Plugins/Runtime/RigVM/Source/RigVMDeveloper/Public/RigVMCompiler/RigVMCompiler.h
--- a/Engine/Plugins/Runtime/RigVM/Source/RigVMDeveloper/Public/RigVMCompiler/RigVMCompiler.h
+++ b/Engine/Plugins/Runtime/RigVM/Source/RigVMDeveloper/Public/RigVMCompiler/RigVMCompiler.h
@@ -161,7 +161,7 @@ public:
 
 		friend inline uint32 GetTypeHash(const FFunctionRegisterData& Data)
 		{
-			uint32 Result = HashCombine(GetTypeHash(Data.ReferenceNode), GetTypeHash(Data.MemoryType));
+			uint32 Result = HashCombine(GetTypeHash(Data.ReferenceNode.ToSoftObjectPath().ToString()), GetTypeHash(Data.MemoryType));
 			return HashCombine(Result, GetTypeHash(Data.RegisterIndex));
 		}
 
