AnimNext: Fix memory storage type conversions to support uint8/byte and FText

diff --git a/Engine/Plugins/Runtime/RigVM/Source/RigVM/Private/RigVMCore/RigVMMemoryStorageStruct.cpp b/Engine/Plugins/Runtime/RigVM/Source/RigVM/Private/RigVMCore/RigVMMemoryStorageStruct.cpp
--- a/Engine/Plugins/Runtime/RigVM/Source/RigVM/Private/RigVMCore/RigVMMemoryStorageStruct.cpp
+++ b/Engine/Plugins/Runtime/RigVM/Source/RigVM/Private/RigVMCore/RigVMMemoryStorageStruct.cpp
@@ -619,6 +619,10 @@ void FRigVMMemoryStorageStruct::SetDefaultValues(const TArray<FRigVMPropertyDesc
 	{
 		OutBagPropertyType = EPropertyBagPropertyType::Bool;
 	}
+	else if (VMType == RigVMTypeUtils::UInt8TypeName)
+	{
+		OutBagPropertyType = EPropertyBagPropertyType::Byte;
+	}
 	else if (VMType == RigVMTypeUtils::Int32TypeName || VMType == RigVMTypeUtils::IntTypeName)
 	{
 		OutBagPropertyType = EPropertyBagPropertyType::Int32;
@@ -651,6 +655,10 @@ void FRigVMMemoryStorageStruct::SetDefaultValues(const TArray<FRigVMPropertyDesc
 	{
 		OutBagPropertyType = EPropertyBagPropertyType::String;
 	}
+	else if (VMType == RigVMTypeUtils::FTextTypeName)
+	{
+		OutBagPropertyType = EPropertyBagPropertyType::Text;
+	}
 	else if (Cast<UScriptStruct>(RigVMDescriptor.CPPTypeObject))
 	{
 		OutBagPropertyType = EPropertyBagPropertyType::Struct;
