[RigVM] fixed calling locked version of the GetArrayType API in a NoLock path

diff --git a/Engine/Plugins/Runtime/RigVM/Source/RigVM/Private/RigVMCore/RigVMTemplate.cpp b/Engine/Plugins/Runtime/RigVM/Source/RigVM/Private/RigVMCore/RigVMTemplate.cpp
--- a/Engine/Plugins/Runtime/RigVM/Source/RigVM/Private/RigVMCore/RigVMTemplate.cpp
+++ b/Engine/Plugins/Runtime/RigVM/Source/RigVM/Private/RigVMCore/RigVMTemplate.cpp
@@ -1791,7 +1791,7 @@ bool FRigVMTemplate::Resolve(FTypeMap& InOutTypes, TArray<int32>& OutPermutation
 			}
 		}
 
-		const FRigVMTemplateArgument::EArrayType ArrayType = Argument.GetArrayType();
+		const FRigVMTemplateArgument::EArrayType ArrayType = Argument.GetArrayType_NoLock();
 		if(ArrayType == FRigVMTemplateArgument::EArrayType_Mixed)
 		{
 			InOutTypes.Add(Argument.Name, RigVMTypeUtils::TypeIndex::WildCard);
