Fixed RigVMEdGraph generating invalid model node paths for EdGraph subgraphs, when a function that contains a collapsed graph is renamed.
Also created a test to validate the fix or future updates to the function

diff --git a/Engine/Plugins/Runtime/RigVM/Source/RigVMDeveloper/Private/EdGraph/RigVMEdGraph.cpp b/Engine/Plugins/Runtime/RigVM/Source/RigVMDeveloper/Private/EdGraph/RigVMEdGraph.cpp
--- a/Engine/Plugins/Runtime/RigVM/Source/RigVMDeveloper/Private/EdGraph/RigVMEdGraph.cpp
+++ b/Engine/Plugins/Runtime/RigVM/Source/RigVMDeveloper/Private/EdGraph/RigVMEdGraph.cpp
@@ -83,7 +83,7 @@ void URigVMEdGraph::HandleRigVMGraphRenamed(const FString& InOldNodePath, const
 	else if(ModelNodePath.StartsWith(OldPrefix))
 	{
 		Modify();
-		ModelNodePath = NewPrefix + ModelNodePath.RightChop(OldPrefix.Len() - 1);
+		ModelNodePath = NewPrefix + ModelNodePath.RightChop(OldPrefix.Len());
 	}
 	else
 	{
