Control Rig: Fix module swapping dialog

#rnx

diff --git a/Engine/Plugins/Runtime/RigVM/Source/RigVMEditor/Private/Widgets/SRigVMSwapAssetReferencesWidget.cpp b/Engine/Plugins/Runtime/RigVM/Source/RigVMEditor/Private/Widgets/SRigVMSwapAssetReferencesWidget.cpp
--- a/Engine/Plugins/Runtime/RigVM/Source/RigVMEditor/Private/Widgets/SRigVMSwapAssetReferencesWidget.cpp
+++ b/Engine/Plugins/Runtime/RigVM/Source/RigVMEditor/Private/Widgets/SRigVMSwapAssetReferencesWidget.cpp
@@ -32,6 +32,10 @@ TArray<TSharedRef<FRigVMTreeNode>> FRigVMTreeAssetRefAssetNode::GetChildrenImpl(
 				if (SoftObjectPath == Reference.GetWithoutSubPath())
 				{
 					TSharedRef<FRigVMTreeReferenceNode> RefNode = FRigVMTreeReferenceNode::Create(Reference);
+					if(GetCheckState() == ECheckBoxState::Checked)
+					{
+						RefNode->SetCheckState(ECheckBoxState::Checked);
+					}
 					AssetRefNodes.Add(RefNode);
 				}
 			}
@@ -215,12 +219,20 @@ void SRigVMSwapAssetReferencesWidget::Construct(const FArguments& InArgs)
 		{
 			GetBulkEditWidget()->GetTreeView()->GetTreeView()->ClearSelection();
 			
-			const TArray<TSharedRef<FRigVMTreeNode>> AllCheckedNodes = GetBulkEditWidget()->GetCheckedNodes();
+			TArray<TSharedRef<FRigVMTreeNode>> AllCheckedNodes = GetBulkEditWidget()->GetCheckedNodes();
 			if(AllCheckedNodes.IsEmpty())
 			{
 				GetBulkEditWidget()->CloseDialog();
 				return FReply::Handled();
 			}
+
+			const TSharedPtr<FRigVMSwapAssetReferencesContext> Context = Cast<FRigVMSwapAssetReferencesContext>(Phase->GetContext());
+			check(Context);
+
+			for(int32 Index = 0; Index < AllCheckedNodes.Num(); Index++)
+			{
+				AllCheckedNodes.Append(AllCheckedNodes[Index]->GetChildren(Context.ToSharedRef()));
+			}
 			
 			const TArray<TSharedRef<FRigVMTreeNode>> ModuleRefs = AllCheckedNodes.FilterByPredicate([](const TSharedRef<FRigVMTreeNode>& Node) {
 				return Node->IsA<FRigVMTreeReferenceNode>();
@@ -230,10 +242,7 @@ void SRigVMSwapAssetReferencesWidget::Construct(const FArguments& InArgs)
 			{
 				return FReply::Handled();
 			}
-			
-			const TSharedPtr<FRigVMSwapAssetReferencesContext> Context = Cast<FRigVMSwapAssetReferencesContext>(Phase->GetContext());
-			check(Context);
-			
+		
 			TSet<FSoftObjectPath> VisitedPackages;
 			TArray<TSharedRef<FRigVMTreeTask>> Tasks;
 			
