Added RigVM functions support to AnimNext :
- Added functions details customization support. This enables adding and removing parameters to the function interface, category, public setting and new layout editor.
- Added functions to AnimNext context menu (local and public from other graphs)
- Added asset registry tags export / import for public AnimNext graph functions, so these can be shared from different graphs
- Added function and collapse graphs rename and delete Outliner support
RigVM required changes :
- Refactored RigVM function details to be able to use it at both plugins
- Moved some functions from RigVMBlueprint to RigVMClientHostInterface
- Updated EngineRigVMTools to compile with the new client host interface
- Made collapse graphs inherit schema from collapse node graph, if available, so custom schemas from other plugins are not removed in the process

#rnx

diff --git a/Engine/Plugins/Animation/ControlRig/Source/ControlRigDeveloper/Public/ControlRigBlueprint.h b/Engine/Plugins/Animation/ControlRig/Source/ControlRigDeveloper/Public/ControlRigBlueprint.h
--- a/Engine/Plugins/Animation/ControlRig/Source/ControlRigDeveloper/Public/ControlRigBlueprint.h
+++ b/Engine/Plugins/Animation/ControlRig/Source/ControlRigDeveloper/Public/ControlRigBlueprint.h
@@ -76,15 +76,17 @@ class CONTROLRIGDEVELOPER_API UControlRigBlueprint : public URigVMBlueprint, pub
 public:
 	UControlRigBlueprint();
 
-	// URigVMBlueprint interface
-	virtual UClass* GetRigVMBlueprintGeneratedClassPrototype() const override { return UControlRigBlueprintGeneratedClass::StaticClass(); }
+	//  --- IRigVMClientHost interface ---
 	virtual UClass* GetRigVMSchemaClass() const override { return UControlRigSchema::StaticClass(); }
 	virtual UScriptStruct* GetRigVMExecuteContextStruct() const override { return FControlRigExecuteContext::StaticStruct(); }
 	virtual UClass* GetRigVMEdGraphClass() const override;
 	virtual UClass* GetRigVMEdGraphNodeClass() const override;
 	virtual UClass* GetRigVMEdGraphSchemaClass() const override;
-	virtual TArray<FString> GeneratePythonCommands(const FString InNewBlueprintName) override;
 	virtual UClass* GetRigVMEditorSettingsClass() const override;
+
+	// URigVMBlueprint interface
+	virtual UClass* GetRigVMBlueprintGeneratedClassPrototype() const override { return UControlRigBlueprintGeneratedClass::StaticClass(); }
+	virtual TArray<FString> GeneratePythonCommands(const FString InNewBlueprintName) override;
 	virtual void GetPreloadDependencies(TArray<UObject*>& OutDeps) override;
 #if WITH_EDITOR
 	virtual const FLazyName& GetPanelPinFactoryName() const override;
diff --git a/Engine/Plugins/Experimental/AnimNext/Source/AnimNextEditor/Private/AnimNextEditorModule.cpp b/Engine/Plugins/Experimental/AnimNext/Source/AnimNextEditor/Private/AnimNextEditorModule.cpp
--- a/Engine/Plugins/Experimental/AnimNext/Source/AnimNextEditor/Private/AnimNextEditorModule.cpp
+++ b/Engine/Plugins/Experimental/AnimNext/Source/AnimNextEditor/Private/AnimNextEditorModule.cpp
@@ -38,6 +38,8 @@
 #include "Variables/SAddVariablesDialog.h"
 #include "AnimNextAssetWorkspaceAssetUserData.h"
 #include "Graph/AnimNextGraphItemDetails.h"
+#include "Graph/AnimNextCollapseNodeItemDetails.h"
+#include "Graph/AnimNextFunctionItemDetails.h"
 #include "Param/AnimNextActorLocatorEditor.h"
 #include "IWorkspaceEditor.h"
 #include "Entries/AnimNextAnimationGraphEntry.h"
@@ -225,11 +227,13 @@ void FAnimNextEditorModule::StartupModule()
 	RegisterLocatorFragmentEditorType("AnimNextActor");
 
 	Workspace::IWorkspaceEditorModule& WorkspaceModule = FModuleManager::Get().LoadModuleChecked<Workspace::IWorkspaceEditorModule>("WorkspaceEditor");
-	const TSharedPtr<FAnimNextGraphItemDetails> GraphItemDetails = MakeShared<FAnimNextGraphItemDetails>();
-	WorkspaceModule.RegisterWorkspaceItemDetails(Workspace::FOutlinerItemDetailsId(FAnimNextGraphOutlinerData::StaticStruct()->GetFName()), GraphItemDetails);
-	WorkspaceModule.RegisterWorkspaceItemDetails(Workspace::FOutlinerItemDetailsId(FAnimNextCollapseGraphOutlinerData::StaticStruct()->GetFName()), GraphItemDetails);
-	WorkspaceModule.RegisterWorkspaceItemDetails(Workspace::FOutlinerItemDetailsId(FAnimNextGraphFunctionOutlinerData::StaticStruct()->GetFName()), GraphItemDetails);
+	WorkspaceModule.RegisterWorkspaceItemDetails(Workspace::FOutlinerItemDetailsId(FAnimNextGraphOutlinerData::StaticStruct()->GetFName()), MakeShared<FAnimNextGraphItemDetails>());
+	WorkspaceModule.RegisterWorkspaceItemDetails(Workspace::FOutlinerItemDetailsId(FAnimNextCollapseGraphOutlinerData::StaticStruct()->GetFName()), MakeShared<FAnimNextCollapseNodeItemDetails>());
+	WorkspaceModule.RegisterWorkspaceItemDetails(Workspace::FOutlinerItemDetailsId(FAnimNextGraphFunctionOutlinerData::StaticStruct()->GetFName()), MakeShared<FAnimNextFunctionItemDetails>());
+
 	FAnimNextGraphItemDetails::RegisterToolMenuExtensions();
+	FAnimNextCollapseNodeItemDetails::RegisterToolMenuExtensions();
+	FAnimNextFunctionItemDetails::RegisterToolMenuExtensions();
 
 	const TSharedPtr<FAnimNextAssetItemDetails> AssetItemDetails = MakeShared<FAnimNextAssetItemDetails>();
 	WorkspaceModule.RegisterWorkspaceItemDetails(Workspace::FOutlinerItemDetailsId(FAnimNextModuleOutlinerData::StaticStruct()->GetFName()), AssetItemDetails);
diff --git a/Engine/Plugins/Experimental/AnimNext/Source/AnimNextEditor/Private/Common/GraphEditorSchemaActions.cpp b/Engine/Plugins/Experimental/AnimNext/Source/AnimNextEditor/Private/Common/GraphEditorSchemaActions.cpp
--- a/Engine/Plugins/Experimental/AnimNext/Source/AnimNextEditor/Private/Common/GraphEditorSchemaActions.cpp
+++ b/Engine/Plugins/Experimental/AnimNext/Source/AnimNextEditor/Private/Common/GraphEditorSchemaActions.cpp
@@ -250,4 +250,67 @@ UEdGraphNode* FAnimNextSchemaAction_AddComment::PerformAction(UEdGraph* ParentGr
 	return FEdGraphSchemaAction_NewNode::SpawnNodeFromTemplate<UEdGraphNode_Comment>(ParentGraph, CommentTemplate, SpawnLocation, bSelectNewNode);
 }
 
-#undef LOCTEXT_NAMESPACE
\ No newline at end of file
+// *** Graph Function ***
+
+FAnimNextSchemaAction_Function::FAnimNextSchemaAction_Function(const FRigVMGraphFunctionHeader& InReferencedPublicFunctionHeader, const FText& InNodeCategory, const FText& InMenuDesc, const FText& InToolTip, const FText& InKeywords)
+	: FAnimNextSchemaAction(InNodeCategory, InMenuDesc, InToolTip, InKeywords)
+{
+	ReferencedPublicFunctionHeader = InReferencedPublicFunctionHeader;
+	NodeClass = UAnimNextEdGraphNode::StaticClass();
+	bIsLocalFunction = true;
+}
+
+FAnimNextSchemaAction_Function::FAnimNextSchemaAction_Function(const URigVMLibraryNode* InFunctionLibraryNode, const FText& InNodeCategory, const FText& InMenuDesc, const FText& InToolTip, const FText& InKeywords)
+	: FAnimNextSchemaAction(InNodeCategory, InMenuDesc, InToolTip, InKeywords)
+{
+	ReferencedPublicFunctionHeader = InFunctionLibraryNode->GetFunctionHeader();
+	NodeClass = UAnimNextEdGraphNode::StaticClass();
+	bIsLocalFunction = true;
+}
+
+const FSlateBrush* FAnimNextSchemaAction_Function::GetIconBrush() const
+{
+	return FAppStyle::GetBrush("GraphEditor.Function_16x");
+}
+
+UEdGraphNode* FAnimNextSchemaAction_Function::PerformAction(UEdGraph* ParentGraph, TArray<UEdGraphPin*>& FromPins, const FVector2D Location, bool bSelectNewNode)
+{
+	IRigVMClientHost* Host = ParentGraph->GetImplementingOuter<IRigVMClientHost>();
+	URigVMEdGraphNode* NewNode = nullptr;
+	URigVMEdGraph* EdGraph = Cast<URigVMEdGraph>(ParentGraph);
+
+	UEdGraphPin* FromPin = nullptr;
+	if (FromPins.Num() > 0)
+	{
+		FromPin = FromPins[0];
+	}
+
+	if (Host != nullptr && EdGraph != nullptr)
+	{
+		FName Name = UE::AnimNext::Editor::FUtils::ValidateName(Cast<UObject>(Host), ReferencedPublicFunctionHeader.Name.ToString());
+		URigVMController* Controller = EdGraph->GetController();
+
+		Controller->OpenUndoBracket(FString::Printf(TEXT("Add '%s' Node"), *Name.ToString()));
+
+		if (URigVMFunctionReferenceNode* ModelNode = Controller->AddFunctionReferenceNodeFromDescription(ReferencedPublicFunctionHeader, Location, Name.ToString(), true, true))
+		{
+			NewNode = Cast<URigVMEdGraphNode>(EdGraph->FindNodeForModelNodeName(ModelNode->GetFName()));
+			check(NewNode);
+
+			if (NewNode)
+			{
+				Controller->ClearNodeSelection(true);
+				Controller->SelectNode(ModelNode, true, true);
+			}
+			Controller->CloseUndoBracket();
+		}
+		else
+		{
+			Controller->CancelUndoBracket();
+		}
+	}
+
+	return NewNode;
+}
+
+#undef LOCTEXT_NAMESPACE
diff --git a/Engine/Plugins/Experimental/AnimNext/Source/AnimNextEditor/Private/Common/GraphEditorSchemaActions.h b/Engine/Plugins/Experimental/AnimNext/Source/AnimNextEditor/Private/Common/GraphEditorSchemaActions.h
--- a/Engine/Plugins/Experimental/AnimNext/Source/AnimNextEditor/Private/Common/GraphEditorSchemaActions.h
+++ b/Engine/Plugins/Experimental/AnimNext/Source/AnimNextEditor/Private/Common/GraphEditorSchemaActions.h
@@ -6,10 +6,12 @@
 #include "EdGraph/EdGraphSchema.h"
 #include "Editor/RigVMEditorStyle.h"
 #include "RigVMCore/RigVMExternalVariable.h"
+#include "RigVMCore/RigVMGraphFunctionDefinition.h"
 #include "Styling/AppStyle.h"
 #include "GraphEditorSchemaActions.generated.h"
 
 struct FSlateBrush;
+class URigVMLibraryNode;
 
 USTRUCT()
 struct FAnimNextSchemaAction : public FEdGraphSchemaAction
@@ -127,4 +129,34 @@ struct FAnimNextSchemaAction_AddComment : public FAnimNextSchemaAction
 	{
 		return FAppStyle::Get().GetBrush("Icons.Comment");
 	}
+};
+
+USTRUCT()
+struct FAnimNextSchemaAction_Function : public FAnimNextSchemaAction
+{
+	GENERATED_BODY()
+
+	FAnimNextSchemaAction_Function() = default;
+
+	FAnimNextSchemaAction_Function(const FRigVMGraphFunctionHeader& InReferencedPublicFunctionHeader, const FText& InNodeCategory, const FText& InMenuDesc, const FText& InToolTip, const FText& InKeywords = FText::GetEmpty());
+	FAnimNextSchemaAction_Function(const URigVMLibraryNode* InFunctionLibraryNode, const FText& InNodeCategory, const FText& InMenuDesc, const FText& InToolTip, const FText& InKeywords = FText::GetEmpty());
+
+	virtual const FSlateBrush* GetIconBrush() const;
+	
+	// FEdGraphSchemaAction Interface
+	virtual UEdGraphNode* PerformAction(UEdGraph* ParentGraph, TArray<UEdGraphPin*>& FromPins, const FVector2D Location, bool bSelectNewNode = true) override;
+	virtual UEdGraphNode* PerformAction(UEdGraph* ParentGraph, UEdGraphPin* FromPin, const FVector2D Location, bool bSelectNewNode = true) { return nullptr; }
+
+private:
+	/** The public function definition we will spawn from [optional] */
+	UPROPERTY(Transient)
+	FRigVMGraphFunctionHeader ReferencedPublicFunctionHeader;
+
+	/** Marked as true for local function definitions */
+	UPROPERTY(Transient)
+	bool bIsLocalFunction;
+
+	/** Holds the node type that this spawner will instantiate. */
+	UPROPERTY(Transient)
+	TSubclassOf<UEdGraphNode> NodeClass;
 };
\ No newline at end of file
diff --git a/Engine/Plugins/Experimental/AnimNext/Source/AnimNextEditor/Private/Common/SActionMenu.cpp b/Engine/Plugins/Experimental/AnimNext/Source/AnimNextEditor/Private/Common/SActionMenu.cpp
--- a/Engine/Plugins/Experimental/AnimNext/Source/AnimNextEditor/Private/Common/SActionMenu.cpp
+++ b/Engine/Plugins/Experimental/AnimNext/Source/AnimNextEditor/Private/Common/SActionMenu.cpp
@@ -16,6 +16,8 @@
 #include "RigVMModel/RigVMClient.h"
 #include "RigVMModel/RigVMSchema.h"
 #include "Units/RigUnit.h"
+#include "UncookedOnlyUtils.h"
+#include "RigVMBlueprintGeneratedClass.h"
 
 #define LOCTEXT_NAMESPACE "AnimNextEditor"
 
@@ -108,14 +110,43 @@ void SActionMenu::CollectAllAnimNextGraphActions(FGraphContextMenuBuilder& MenuB
 		MenuBuilder.AddAction(MakeShared<FAnimNextSchemaAction_DispatchFactory>(Template->GetNotation(), NodeCategory, MenuDesc, ToolTip));
 	};
 
-	TArray<FRigVMExternalVariable> ExternalVariables = RigVMHost->GetVM()->GetExternalVariables(RigVMHost->GetRigVMExtendedExecuteContext());
-	for (const FRigVMExternalVariable& ExternalVariable : ExternalVariables)
+	if (URigVMFunctionLibrary* LocalFunctionLibrary = RigVMClientHost->GetLocalFunctionLibrary())
 	{
-		MenuBuilder.AddAction(MakeShared<FAnimNextSchemaAction_Variable>(ExternalVariable, true));
-		MenuBuilder.AddAction(MakeShared<FAnimNextSchemaAction_Variable>(ExternalVariable, false));
+		const FSoftObjectPath LocalLibrarySoftPath = LocalFunctionLibrary->GetFunctionHostObjectPath();
+
+		TArray<URigVMLibraryNode*> Functions = LocalFunctionLibrary->GetFunctions();
+		for (URigVMLibraryNode* FunctionLibraryNode : Functions)
+		{
+			if (LocalFunctionLibrary->IsFunctionPublic(FunctionLibraryNode->GetFName()))	// Public functions will be added when processing asset registry exports
+			{
+				continue;
+			}
+			const FText NodeCategory = FText::FromString(FunctionLibraryNode->GetNodeCategory());
+			const FText MenuDesc = FText::FromString(FunctionLibraryNode->GetName());
+			const FText ToolTip = FunctionLibraryNode->GetToolTipText();
+
+			MenuBuilder.AddAction(MakeShared<FAnimNextSchemaAction_Function>(FunctionLibraryNode, NodeCategory, MenuDesc, ToolTip));
+		}
 	}
 
-	MenuBuilder.AddAction(MakeShared<FAnimNextSchemaAction_AddComment>());
+	TMap<FAssetData, FRigVMGraphFunctionHeaderArray> FunctionExports;
+	AnimNext::UncookedOnly::FUtils::GetExportedFunctionsFromAssetRegistry(UE::AnimNext::AnimNextPublicGraphFunctionsExportsRegistryTag, FunctionExports);
+	AnimNext::UncookedOnly::FUtils::GetExportedFunctionsFromAssetRegistry(UE::AnimNext::ControlRigAssetPublicGraphFunctionsExportsRegistryTag, FunctionExports);
+
+	for (const auto& Export : FunctionExports.Array())
+	{
+		for (const FRigVMGraphFunctionHeader& FunctionHeader : Export.Value.Headers)
+		{
+			if (FunctionHeader.LibraryPointer.IsValid())
+			{
+				const FText NodeCategory = FText::FromString(FunctionHeader.Category);
+				const FText MenuDesc = FText::FromString(FunctionHeader.NodeTitle);
+				const FText ToolTip = FunctionHeader.GetTooltip();
+
+				MenuBuilder.AddAction(MakeShared<FAnimNextSchemaAction_Function>(FunctionHeader, NodeCategory, MenuDesc, ToolTip));
+			}
+		}
+	}
 
 	RigVMController->EnableReporting(true);
 }
diff --git a/Engine/Plugins/Experimental/AnimNext/Source/AnimNextEditor/Private/Graph/AnimNextCollapseNodeItemDetails.cpp b/Engine/Plugins/Experimental/AnimNext/Source/AnimNextEditor/Private/Graph/AnimNextCollapseNodeItemDetails.cpp
new file mode 100644
--- /dev/null
+++ b/Engine/Plugins/Experimental/AnimNext/Source/AnimNextEditor/Private/Graph/AnimNextCollapseNodeItemDetails.cpp
@@ -0,0 +1,194 @@
+// Copyright Epic Games, Inc. All Rights Reserved.
+
+#include "AnimNextCollapseNodeItemDetails.h"
+
+#include "Toolkits/AssetEditorToolkitMenuContext.h"
+#include "Entries/AnimNextAnimationGraphEntry.h"
+#include "Entries/AnimNextEventGraphEntry.h"
+#include "StructUtils/InstancedStruct.h"
+#include "AnimNextAssetWorkspaceAssetUserData.h"
+#include "RigVMModel/RigVMGraph.h"
+#include "WorkspaceItemMenuContext.h"
+#include "IWorkspaceEditor.h"
+#include "ScopedTransaction.h"
+#include "RigVMModel/RigVMClient.h"
+#include "EdGraph/RigVMEdGraph.h"
+#include "ToolMenus.h"
+#include "Framework/Commands/GenericCommands.h"
+#include "Module/AnimNextModule_EditorData.h"
+
+#define LOCTEXT_NAMESPACE "FAnimNextCollapseNodeItemDetails"
+
+namespace UE::AnimNext::Editor
+{
+
+void FAnimNextCollapseNodeItemDetails::HandleDoubleClick(const FToolMenuContext& ToolMenuContext) const
+{
+	const UWorkspaceItemMenuContext* WorkspaceItemContext = ToolMenuContext.FindContext<UWorkspaceItemMenuContext>();
+	const UAssetEditorToolkitMenuContext* AssetEditorContext = ToolMenuContext.FindContext<UAssetEditorToolkitMenuContext>();
+	if (WorkspaceItemContext && AssetEditorContext)
+	{
+		if (const TSharedPtr<UE::Workspace::IWorkspaceEditor> WorkspaceEditor = StaticCastSharedPtr<UE::Workspace::IWorkspaceEditor>(AssetEditorContext->Toolkit.Pin()))
+		{
+			const TInstancedStruct<FWorkspaceOutlinerItemData>& Data = WorkspaceItemContext->SelectedExports[0].GetData();
+			if (Data.IsValid() && Data.GetScriptStruct() == FAnimNextCollapseGraphOutlinerData::StaticStruct())
+			{
+				const FAnimNextCollapseGraphOutlinerData& CollapseGraphData = Data.Get<FAnimNextCollapseGraphOutlinerData>();
+				if (CollapseGraphData.EditorObject.IsValid())
+				{
+					UObject* EditorObject = CollapseGraphData.EditorObject.Get();
+					WorkspaceEditor->OpenObjects({ EditorObject });
+				}
+			}
+		}
+	}
+}
+
+bool FAnimNextCollapseNodeItemDetails::CanDelete(const FWorkspaceOutlinerItemExport& Export) const
+{
+	const TInstancedStruct<FWorkspaceOutlinerItemData>& Data = Export.GetData();
+	if (Data.IsValid() && Data.GetScriptStruct() == FAnimNextCollapseGraphOutlinerData::StaticStruct())
+	{
+		const FAnimNextCollapseGraphOutlinerData& CollapseGraphData = Export.GetData().Get<FAnimNextCollapseGraphOutlinerData>();
+		if (CollapseGraphData.EditorObject.IsValid())
+		{
+			URigVMEdGraph* EdGraph = CollapseGraphData.EditorObject.Get();
+			if (EdGraph->bAllowDeletion)
+			{
+				return true;
+			}
+		}
+	}
+	return false;
+}
+
+void FAnimNextCollapseNodeItemDetails::Delete(TConstArrayView<FWorkspaceOutlinerItemExport> Exports) const
+{
+	TMap<UAnimNextRigVMAssetEditorData*, TArray<UAnimNextRigVMAssetEntry*>> EntriesToDelete;
+	for (const FWorkspaceOutlinerItemExport& Export : Exports)
+	{
+		const TInstancedStruct<FWorkspaceOutlinerItemData>& Data = Export.GetData();
+		if (Data.IsValid() && Data.GetScriptStruct() == FAnimNextCollapseGraphOutlinerData::StaticStruct())
+		{
+			const FAnimNextCollapseGraphOutlinerData& CollapseGraphData = Export.GetData().Get<FAnimNextCollapseGraphOutlinerData>();
+			if (CollapseGraphData.EditorObject.IsValid())
+			{
+				URigVMEdGraph* EdGraph = CollapseGraphData.EditorObject.Get();
+				if (EdGraph->bAllowDeletion)
+				{
+					if (URigVMGraph* Model = EdGraph->GetModel())
+					{
+						if (URigVMCollapseNode* CollapseNode = CastChecked<URigVMCollapseNode>(EdGraph->GetModel()->GetOuter()))
+						{
+							URigVMGraph* ContainerGraph = CollapseNode->GetGraph();
+							if (IRigVMClientHost* ClientHost = ContainerGraph->GetImplementingOuter<IRigVMClientHost>())
+							{
+								if (URigVMEdGraph* ContainerEdGraph = Cast<URigVMEdGraph>(ClientHost->GetEditorObjectForRigVMGraph(ContainerGraph)))
+								{
+									FScopedTransaction Transaction(LOCTEXT("DeleteCollapseNodeInOutliner", "Delete Collapse Node"));
+									ContainerEdGraph->GetController()->RemoveNode(CollapseNode);
+								}
+							}
+						}
+					}
+				}
+			}
+		}
+	}
+}
+
+bool FAnimNextCollapseNodeItemDetails::CanRename(const FWorkspaceOutlinerItemExport& Export) const
+{
+	const TInstancedStruct<FWorkspaceOutlinerItemData>& Data = Export.GetData();
+	if (Data.IsValid() && Data.GetScriptStruct() == FAnimNextCollapseGraphOutlinerData::StaticStruct())
+	{
+		const FAnimNextCollapseGraphOutlinerData& CollapseGraphData = Export.GetData().Get<FAnimNextCollapseGraphOutlinerData>();
+		if (CollapseGraphData.EditorObject.IsValid())
+		{
+			URigVMEdGraph* EdGraph = CollapseGraphData.EditorObject.Get();
+			if (EdGraph->bAllowRenaming)
+			{
+				return true;
+			}
+		}
+	}
+	return false;
+}
+
+void FAnimNextCollapseNodeItemDetails::Rename(const FWorkspaceOutlinerItemExport& Export, const FText& InName) const
+{
+	const TInstancedStruct<FWorkspaceOutlinerItemData>& Data = Export.GetData();
+	if (Data.IsValid() && Data.GetScriptStruct() == FAnimNextCollapseGraphOutlinerData::StaticStruct())
+	{
+		const FAnimNextCollapseGraphOutlinerData& CollapseGraphData = Export.GetData().Get<FAnimNextCollapseGraphOutlinerData>();
+		if (CollapseGraphData.EditorObject.IsValid())
+		{
+			URigVMEdGraph* EdGraph = CollapseGraphData.EditorObject.Get();
+			if (EdGraph->bAllowRenaming)
+			{
+				if (const UEdGraphSchema* GraphSchema = EdGraph->GetSchema())
+				{
+					FGraphDisplayInfo DisplayInfo;
+					GraphSchema->GetGraphDisplayInformation(*EdGraph, DisplayInfo);
+
+					// Check if the name is unchanged
+					if (InName.EqualTo(DisplayInfo.PlainName))
+					{
+						return;
+					}
+
+					FScopedTransaction Transaction(LOCTEXT("RenameCollapseNodeInOutliner", "Rename Collapse Node"));
+					FRigVMControllerCompileBracketScope CompileScope(EdGraph->GetController());
+					if (GraphSchema->TryRenameGraph(EdGraph, *InName.ToString()))
+					{
+						return;
+					}
+				}
+			}
+		}
+	}
+}
+
+bool FAnimNextCollapseNodeItemDetails::ValidateName(const FWorkspaceOutlinerItemExport& Export, const FText& InName, FText& OutErrorMessage) const
+{
+	const TInstancedStruct<FWorkspaceOutlinerItemData>& Data = Export.GetData();
+	if (Data.IsValid() && Data.GetScriptStruct() == FAnimNextCollapseGraphOutlinerData::StaticStruct())
+	{
+		return true;
+	}
+
+	OutErrorMessage = LOCTEXT("UnsupportedTypeCollapseNodeRenameError", "Element type is not supported for rename");
+	return false;
+}
+
+UPackage* FAnimNextCollapseNodeItemDetails::GetPackage(const FWorkspaceOutlinerItemExport& Export) const
+{
+	const TInstancedStruct<FWorkspaceOutlinerItemData>& Data = Export.GetData();
+	if (Data.IsValid() && Data.GetScriptStruct() == FAnimNextCollapseGraphOutlinerData::StaticStruct())
+	{
+		const FAnimNextCollapseGraphOutlinerData& CollapseGraphData = Data.Get<FAnimNextCollapseGraphOutlinerData>();
+		if (CollapseGraphData.EditorObject.IsValid())
+		{
+			return CollapseGraphData.EditorObject->GetPackage();
+		}
+	}
+	return nullptr;
+}
+
+const FSlateBrush* FAnimNextCollapseNodeItemDetails::GetItemIcon(const FWorkspaceOutlinerItemExport& Export) const
+{
+	return FAppStyle::GetBrush(TEXT("GraphEditor.EventGraph_24x"));
+}
+
+void FAnimNextCollapseNodeItemDetails::RegisterToolMenuExtensions()
+{
+}
+
+void FAnimNextCollapseNodeItemDetails::UnregisterToolMenuExtensions()
+{
+}
+
+
+} // UE::AnimNext::Editor
+
+#undef LOCTEXT_NAMESPACE // "FAnimNextCollapseNodeItemDetails"
diff --git a/Engine/Plugins/Experimental/AnimNext/Source/AnimNextEditor/Private/Graph/AnimNextCollapseNodeItemDetails.h b/Engine/Plugins/Experimental/AnimNext/Source/AnimNextEditor/Private/Graph/AnimNextCollapseNodeItemDetails.h
new file mode 100644
--- /dev/null
+++ b/Engine/Plugins/Experimental/AnimNext/Source/AnimNextEditor/Private/Graph/AnimNextCollapseNodeItemDetails.h
@@ -0,0 +1,29 @@
+// Copyright Epic Games, Inc. All Rights Reserved.
+
+#pragma once
+
+#include "IWorkspaceOutlinerItemDetails.h"
+
+namespace UE::AnimNext::Editor
+{
+
+class FAnimNextCollapseNodeItemDetails : public UE::Workspace::IWorkspaceOutlinerItemDetails
+{
+public:
+	FAnimNextCollapseNodeItemDetails() = default;
+	virtual ~FAnimNextCollapseNodeItemDetails() override = default;
+
+	virtual void HandleDoubleClick(const FToolMenuContext& ToolMenuContext) const override;
+	virtual bool CanDelete(const FWorkspaceOutlinerItemExport& Export) const override;
+	virtual void Delete(TConstArrayView<FWorkspaceOutlinerItemExport> Exports) const override;
+	virtual bool CanRename(const FWorkspaceOutlinerItemExport& Export) const override;
+	virtual void Rename(const FWorkspaceOutlinerItemExport& Export, const FText& InName) const override;
+	virtual bool ValidateName(const FWorkspaceOutlinerItemExport& Export, const FText& InName, FText& OutErrorMessage) const override;
+	virtual UPackage* GetPackage(const FWorkspaceOutlinerItemExport& Export) const override;
+	virtual const FSlateBrush* GetItemIcon(const FWorkspaceOutlinerItemExport& Export) const override;
+
+	static void RegisterToolMenuExtensions();
+	static void UnregisterToolMenuExtensions();
+};
+
+}
diff --git a/Engine/Plugins/Experimental/AnimNext/Source/AnimNextEditor/Private/Graph/AnimNextEdGraphNodeCustomization.cpp b/Engine/Plugins/Experimental/AnimNext/Source/AnimNextEditor/Private/Graph/AnimNextEdGraphNodeCustomization.cpp
--- a/Engine/Plugins/Experimental/AnimNext/Source/AnimNextEditor/Private/Graph/AnimNextEdGraphNodeCustomization.cpp
+++ b/Engine/Plugins/Experimental/AnimNext/Source/AnimNextEditor/Private/Graph/AnimNextEdGraphNodeCustomization.cpp
@@ -5,11 +5,15 @@
 #include "DetailLayoutBuilder.h"
 #include "DetailCategoryBuilder.h"
 #include "AnimNextEdGraphNode.h"
+#include "Editor/RigVMGraphDetailCustomization.h"
 #include "Graph/RigDecorator_AnimNextCppTrait.h"
 #include "Graph/TraitEditorTabSummoner.h"
 #include "Logging/LogScopedVerbosityOverride.h"
 #include "InstancedPropertyBagStructureDataProvider.h"
+#include "RigVMModel/Nodes/RigVMLibraryNode.h"
+#include "RigVMModel/Nodes/RigVMAggregateNode.h"
 #include "RigVMModel/RigVMController.h"
+#include "RigVMModel/RigVMFunctionLibrary.h"
 #include "IWorkspaceEditor.h"
 #include "STraitEditorView.h"
 #include "Widgets/Docking/SDockTab.h"
@@ -27,6 +31,7 @@ FAnimNextEdGraphNodeCustomization::FAnimNextEdGraphNodeCustomization(const TWeak
 void FAnimNextEdGraphNodeCustomization::PendingDelete()
 {
 	CategoryDetailsData.Reset();
+	RigVMGraphDetailCustomizationImpl.Reset();
 
 	if (TSharedPtr<UE::Workspace::IWorkspaceEditor> WorkspaceEditor = WorkspaceEditorWeak.Pin())
 	{
@@ -45,6 +50,8 @@ void FAnimNextEdGraphNodeCustomization::CustomizeDetails(IDetailLayoutBuilder& D
 	TArray<TWeakObjectPtr<UObject>> Objects;
 	DetailBuilder.GetObjectsBeingCustomized(Objects);
 
+	RigVMGraphDetailCustomizationImpl.Reset();
+
 	if (Objects.IsEmpty())
 	{
 		return;
@@ -62,10 +69,55 @@ void FAnimNextEdGraphNodeCustomization::CustomizeObjects(IDetailLayoutBuilder& D
 		{
 			if (EdGraphNode->IsTraitStack())
 			{
+				RigVMGraphDetailCustomizationImpl.Reset();
 				GenerateTraitData(EdGraphNode, CategoryDetailsData);
 			}
-			else
+			else 
 			{
+				if (URigVMNode* ModelNode = EdGraphNode->GetModelNode())
+				{
+					bool bIsFunction = false;
+					bool bIsCollapseNode = false;
+					bool bIsAggregate = false;
+
+					if (URigVMLibraryNode* LibraryNode = Cast<URigVMLibraryNode>(ModelNode))
+					{
+						if (URigVMFunctionReferenceNode* FunctionReferenceNode = Cast<URigVMFunctionReferenceNode>(ModelNode))
+						{
+							if (URigVMLibraryNode* ReferencedNode = Cast<URigVMLibraryNode>(FunctionReferenceNode->GetReferencedFunctionHeader().LibraryPointer.GetNodeSoftPath().ResolveObject()))
+							{
+								LibraryNode = ReferencedNode;
+								bIsFunction = true;
+							}
+						}
+						else if (URigVMCollapseNode* CollapseNode = Cast<URigVMCollapseNode>(ModelNode))
+						{
+							bIsCollapseNode = true;
+						}
+						bIsAggregate = LibraryNode->IsA<URigVMAggregateNode>();
+
+						if (bIsFunction || bIsCollapseNode || bIsAggregate)
+						{
+							if (NumNodes > 1) // no multiselection of function / collapse nodes
+							{
+								continue;
+							}
+
+							TWeakPtr<FRigVMEditor> WeakEditor; // passing an empty RigVMEditor to FRigVMCollapseGraphLayout
+							IRigVMClientHost* RigVMClientHost = EdGraphNode->GetController()->GetClientHost();
+							check(RigVMClientHost);
+
+
+							RigVMGraphDetailCustomizationImpl = MakeShared<FRigVMGraphDetailCustomizationImpl>();
+							RigVMGraphDetailCustomizationImpl->CustomizeDetails(DetailBuilder
+								, LibraryNode->GetContainedGraph()
+								, RigVMClientHost->GetController(LibraryNode->GetContainedGraph())
+								, RigVMClientHost
+								, WeakEditor);
+						}
+					}
+				}
+
 				GenerateRigVMData(EdGraphNode, CategoryDetailsData);
 			}
 		}
diff --git a/Engine/Plugins/Experimental/AnimNext/Source/AnimNextEditor/Private/Graph/AnimNextEdGraphNodeCustomization.h b/Engine/Plugins/Experimental/AnimNext/Source/AnimNextEditor/Private/Graph/AnimNextEdGraphNodeCustomization.h
--- a/Engine/Plugins/Experimental/AnimNext/Source/AnimNextEditor/Private/Graph/AnimNextEdGraphNodeCustomization.h
+++ b/Engine/Plugins/Experimental/AnimNext/Source/AnimNextEditor/Private/Graph/AnimNextEdGraphNodeCustomization.h
@@ -10,6 +10,7 @@
 class URigVMPin;
 class IDetailCategoryBuilder;
 class UAnimNextEdGraphNode;
+class FRigVMGraphDetailCustomizationImpl;
 
 namespace UE::Workspace
 {
@@ -107,6 +108,8 @@ protected:
 
 	TWeakPtr<UE::Workspace::IWorkspaceEditor> WorkspaceEditorWeak;
 	TArray<TSharedPtr<FCategoryDetailsData>> CategoryDetailsData;
+
+	TSharedPtr<FRigVMGraphDetailCustomizationImpl> RigVMGraphDetailCustomizationImpl;
 };
 
 }
diff --git a/Engine/Plugins/Experimental/AnimNext/Source/AnimNextEditor/Private/Graph/AnimNextFunctionItemDetails.cpp b/Engine/Plugins/Experimental/AnimNext/Source/AnimNextEditor/Private/Graph/AnimNextFunctionItemDetails.cpp
new file mode 100644
--- /dev/null
+++ b/Engine/Plugins/Experimental/AnimNext/Source/AnimNextEditor/Private/Graph/AnimNextFunctionItemDetails.cpp
@@ -0,0 +1,186 @@
+// Copyright Epic Games, Inc. All Rights Reserved.
+
+#include "AnimNextFunctionItemDetails.h"
+
+#include "Toolkits/AssetEditorToolkitMenuContext.h"
+#include "Entries/AnimNextAnimationGraphEntry.h"
+#include "Entries/AnimNextEventGraphEntry.h"
+#include "StructUtils/InstancedStruct.h"
+#include "AnimNextAssetWorkspaceAssetUserData.h"
+#include "RigVMModel/RigVMGraph.h"
+#include "WorkspaceItemMenuContext.h"
+#include "IWorkspaceEditor.h"
+#include "ScopedTransaction.h"
+#include "RigVMModel/RigVMClient.h"
+#include "EdGraph/RigVMEdGraph.h"
+#include "ToolMenus.h"
+#include "Framework/Commands/GenericCommands.h"
+#include "Module/AnimNextModule_EditorData.h"
+
+#define LOCTEXT_NAMESPACE "FAnimNextFunctionItemDetails"
+
+namespace UE::AnimNext::Editor
+{
+
+void FAnimNextFunctionItemDetails::HandleDoubleClick(const FToolMenuContext& ToolMenuContext) const
+{
+	const UWorkspaceItemMenuContext* WorkspaceItemContext = ToolMenuContext.FindContext<UWorkspaceItemMenuContext>();
+	const UAssetEditorToolkitMenuContext* AssetEditorContext = ToolMenuContext.FindContext<UAssetEditorToolkitMenuContext>();
+	if (WorkspaceItemContext && AssetEditorContext)
+	{
+		if (const TSharedPtr<UE::Workspace::IWorkspaceEditor> WorkspaceEditor = StaticCastSharedPtr<UE::Workspace::IWorkspaceEditor>(AssetEditorContext->Toolkit.Pin()))
+		{
+			const TInstancedStruct<FWorkspaceOutlinerItemData>& Data = WorkspaceItemContext->SelectedExports[0].GetData();
+			if (Data.IsValid() && Data.GetScriptStruct() == FAnimNextGraphFunctionOutlinerData::StaticStruct())
+			{
+				const FAnimNextGraphFunctionOutlinerData& GraphFunctionData = Data.Get<FAnimNextGraphFunctionOutlinerData>();
+				if (GraphFunctionData.EditorObject.IsValid())
+				{
+					WorkspaceEditor->OpenObjects({ GraphFunctionData.EditorObject.Get() });
+				}
+			}
+		}
+	}
+}
+
+bool FAnimNextFunctionItemDetails::CanDelete(const FWorkspaceOutlinerItemExport& Export) const
+{
+	const TInstancedStruct<FWorkspaceOutlinerItemData>& Data = Export.GetData();
+	if (Data.IsValid() && Data.GetScriptStruct() == FAnimNextGraphFunctionOutlinerData::StaticStruct())
+	{
+		const FAnimNextGraphFunctionOutlinerData& CollapseGraphData = Export.GetData().Get<FAnimNextGraphFunctionOutlinerData>();
+		if (CollapseGraphData.EditorObject.IsValid())
+		{
+			URigVMEdGraph* EdGraph = CollapseGraphData.EditorObject.Get();
+			if (EdGraph->bAllowDeletion)
+			{
+				return true;
+			}
+		}
+	}
+	return false;
+}
+
+void FAnimNextFunctionItemDetails::Delete(TConstArrayView<FWorkspaceOutlinerItemExport> Exports) const
+{
+	TMap<UAnimNextRigVMAssetEditorData*, TArray<UAnimNextRigVMAssetEntry*>> EntriesToDelete;
+	for (const FWorkspaceOutlinerItemExport& Export : Exports)
+	{
+		const TInstancedStruct<FWorkspaceOutlinerItemData>& Data = Export.GetData();
+		if (Data.IsValid() && Data.GetScriptStruct() == FAnimNextGraphFunctionOutlinerData::StaticStruct())
+		{
+			const FAnimNextGraphFunctionOutlinerData& GraphData = Export.GetData().Get<FAnimNextGraphFunctionOutlinerData>();
+			if (GraphData.EditorObject.IsValid() && GraphData.EdGraphNode.IsValid())
+			{
+				const URigVMEdGraphNode* EdGraphNode = GraphData.EdGraphNode.Get();
+				if (EdGraphNode->CanUserDeleteNode())
+				{
+					if (URigVMGraph* Model = EdGraphNode->GetModel())
+					{
+						if (URigVMNode* ModelNode = Model->FindNodeByName(*EdGraphNode->GetModelNodePath()))
+						{
+							FScopedTransaction Transaction(LOCTEXT("DeleteFunctionInOutliner", "Delete Function"));
+							EdGraphNode->GetController()->RemoveNode(ModelNode);
+						}
+					}
+				}
+			}
+		}
+	}
+}
+
+bool FAnimNextFunctionItemDetails::CanRename(const FWorkspaceOutlinerItemExport& Export) const
+{
+	const TInstancedStruct<FWorkspaceOutlinerItemData>& Data = Export.GetData();
+	if (Data.IsValid() && Data.GetScriptStruct() == FAnimNextGraphFunctionOutlinerData::StaticStruct())
+	{
+		const FAnimNextGraphFunctionOutlinerData& CollapseGraphData = Export.GetData().Get<FAnimNextGraphFunctionOutlinerData>();
+		if (CollapseGraphData.EditorObject.IsValid())
+		{
+			URigVMEdGraph* EdGraph = CollapseGraphData.EditorObject.Get();
+			if (EdGraph->bAllowRenaming)
+			{
+				return true;
+			}
+		}
+	}
+
+	return false;
+}
+
+void FAnimNextFunctionItemDetails::Rename(const FWorkspaceOutlinerItemExport& Export, const FText& InName) const
+{
+	const TInstancedStruct<FWorkspaceOutlinerItemData>& Data = Export.GetData();
+	if (Data.IsValid() && Data.GetScriptStruct() == FAnimNextGraphFunctionOutlinerData::StaticStruct())
+	{
+		const FAnimNextGraphFunctionOutlinerData& CollapseGraphData = Export.GetData().Get<FAnimNextGraphFunctionOutlinerData>();
+		if (CollapseGraphData.EditorObject.IsValid())
+		{
+			URigVMEdGraph* EdGraph = CollapseGraphData.EditorObject.Get();
+			if (EdGraph->bAllowRenaming)
+			{
+				if (const UEdGraphSchema* GraphSchema = EdGraph->GetSchema())
+				{
+					FGraphDisplayInfo DisplayInfo;
+					GraphSchema->GetGraphDisplayInformation(*EdGraph, DisplayInfo);
+
+					// Check if the name is unchanged
+					if (InName.EqualTo(DisplayInfo.PlainName))
+					{
+						return;
+					}
+
+					FScopedTransaction Transaction(LOCTEXT("RenameFunctionInOutliner", "Rename Function"));
+					FRigVMControllerCompileBracketScope CompileScope(EdGraph->GetController());
+					if (GraphSchema->TryRenameGraph(EdGraph, *InName.ToString()))
+					{
+						return;
+					}
+				}
+			}
+		}
+	}
+}
+
+bool FAnimNextFunctionItemDetails::ValidateName(const FWorkspaceOutlinerItemExport& Export, const FText& InName, FText& OutErrorMessage) const
+{
+	const TInstancedStruct<FWorkspaceOutlinerItemData>& Data = Export.GetData();
+	if (Data.IsValid() && Data.GetScriptStruct() == FAnimNextGraphFunctionOutlinerData::StaticStruct())
+	{
+		return true;
+	}
+
+	OutErrorMessage = LOCTEXT("UnsupportedTypeRenameError", "Element type is not supported for rename");
+	return false;
+}
+
+UPackage* FAnimNextFunctionItemDetails::GetPackage(const FWorkspaceOutlinerItemExport& Export) const
+{
+	const TInstancedStruct<FWorkspaceOutlinerItemData>& Data = Export.GetData();
+	if (Data.IsValid() && Data.GetScriptStruct() == FAnimNextGraphFunctionOutlinerData::StaticStruct())
+	{
+		const FAnimNextGraphFunctionOutlinerData& GraphFunctionData = Data.Get<FAnimNextGraphFunctionOutlinerData>();
+		if (GraphFunctionData.EditorObject.IsValid())
+		{
+			return GraphFunctionData.EditorObject->GetPackage();
+		}
+	}
+	return nullptr;
+}
+
+const FSlateBrush* FAnimNextFunctionItemDetails::GetItemIcon(const FWorkspaceOutlinerItemExport& Export) const
+{
+	return FAppStyle::GetBrush(TEXT("GraphEditor.EventGraph_24x"));
+}
+
+void FAnimNextFunctionItemDetails::RegisterToolMenuExtensions()
+{
+}
+
+void FAnimNextFunctionItemDetails::UnregisterToolMenuExtensions()
+{
+}
+
+} // UE::AnimNext::Editor
+
+#undef LOCTEXT_NAMESPACE // "FAnimNextFunctionItemDetails"
diff --git a/Engine/Plugins/Experimental/AnimNext/Source/AnimNextEditor/Private/Graph/AnimNextFunctionItemDetails.h b/Engine/Plugins/Experimental/AnimNext/Source/AnimNextEditor/Private/Graph/AnimNextFunctionItemDetails.h
new file mode 100644
--- /dev/null
+++ b/Engine/Plugins/Experimental/AnimNext/Source/AnimNextEditor/Private/Graph/AnimNextFunctionItemDetails.h
@@ -0,0 +1,29 @@
+// Copyright Epic Games, Inc. All Rights Reserved.
+
+#pragma once
+
+#include "IWorkspaceOutlinerItemDetails.h"
+
+namespace UE::AnimNext::Editor
+{
+
+class FAnimNextFunctionItemDetails : public UE::Workspace::IWorkspaceOutlinerItemDetails
+{
+public:
+	FAnimNextFunctionItemDetails() = default;
+	virtual ~FAnimNextFunctionItemDetails() override = default;
+
+	virtual void HandleDoubleClick(const FToolMenuContext& ToolMenuContext) const override;
+	virtual bool CanDelete(const FWorkspaceOutlinerItemExport& Export) const override;
+	virtual void Delete(TConstArrayView<FWorkspaceOutlinerItemExport> Exports) const override;
+	virtual bool CanRename(const FWorkspaceOutlinerItemExport& Export) const override;
+	virtual void Rename(const FWorkspaceOutlinerItemExport& Export, const FText& InName) const override;
+	virtual bool ValidateName(const FWorkspaceOutlinerItemExport& Export, const FText& InName, FText& OutErrorMessage) const override;
+	virtual UPackage* GetPackage(const FWorkspaceOutlinerItemExport& Export) const override;
+	virtual const FSlateBrush* GetItemIcon(const FWorkspaceOutlinerItemExport& Export) const override;
+
+	static void RegisterToolMenuExtensions();
+	static void UnregisterToolMenuExtensions();
+};
+
+}
diff --git a/Engine/Plugins/Experimental/AnimNext/Source/AnimNextEditor/Private/Graph/AnimNextGraphItemDetails.cpp b/Engine/Plugins/Experimental/AnimNext/Source/AnimNextEditor/Private/Graph/AnimNextGraphItemDetails.cpp
--- a/Engine/Plugins/Experimental/AnimNext/Source/AnimNextEditor/Private/Graph/AnimNextGraphItemDetails.cpp
+++ b/Engine/Plugins/Experimental/AnimNext/Source/AnimNextEditor/Private/Graph/AnimNextGraphItemDetails.cpp
@@ -31,47 +31,37 @@ void FAnimNextGraphItemDetails::HandleDoubleClick(const FToolMenuContext& ToolMe
 		if(const TSharedPtr<UE::Workspace::IWorkspaceEditor> WorkspaceEditor = StaticCastSharedPtr<UE::Workspace::IWorkspaceEditor>(AssetEditorContext->Toolkit.Pin()))
 		{
 			const TInstancedStruct<FWorkspaceOutlinerItemData>& Data = WorkspaceItemContext->SelectedExports[0].GetData();
-			if (Data.IsValid())
+			if (Data.IsValid() && Data.GetScriptStruct() == FAnimNextGraphOutlinerData::StaticStruct())
 			{
-				if (Data.GetScriptStruct() == FAnimNextGraphOutlinerData::StaticStruct())
+				const FAnimNextGraphOutlinerData& GraphData = Data.Get<FAnimNextGraphOutlinerData>();							
+				if (GraphData.GraphInterface)
 				{
-					const FAnimNextGraphOutlinerData& GraphData = Data.Get<FAnimNextGraphOutlinerData>();							
-					if (GraphData.GraphInterface)
+					if (URigVMGraph* RigVMGraph = GraphData.GraphInterface->GetRigVMGraph())
 					{
-						if (URigVMGraph* RigVMGraph = GraphData.GraphInterface->GetRigVMGraph())
+						if(const IRigVMClientHost* RigVMClientHost = RigVMGraph->GetImplementingOuter<IRigVMClientHost>())
 						{
-							if(const IRigVMClientHost* RigVMClientHost = RigVMGraph->GetImplementingOuter<IRigVMClientHost>())
+							if(UObject* EditorObject = RigVMClientHost->GetEditorObjectForRigVMGraph(RigVMGraph))
 							{
-								if(UObject* EditorObject = RigVMClientHost->GetEditorObjectForRigVMGraph(RigVMGraph))
-								{
-									WorkspaceEditor->OpenObjects({EditorObject});
-								}
+								WorkspaceEditor->OpenObjects({EditorObject});
 							}
 						}
 					}
 				}
-				else if (Data.GetScriptStruct() == FAnimNextGraphFunctionOutlinerData::StaticStruct())
-				{
-					const FAnimNextGraphFunctionOutlinerData& GraphFunctionData = Data.Get<FAnimNextGraphFunctionOutlinerData>();
-					if (GraphFunctionData.EditorObject.IsValid())
-					{
-						WorkspaceEditor->OpenObjects({ GraphFunctionData.EditorObject.Get()});
-					}
-				}
-				else if (Data.GetScriptStruct() == FAnimNextCollapseGraphOutlinerData::StaticStruct())
-				{
-					const FAnimNextCollapseGraphOutlinerData& CollapseGraphData = Data.Get<FAnimNextCollapseGraphOutlinerData>();
-					if (CollapseGraphData.EditorObject.IsValid())
-					{
-						UObject* EditorObject = CollapseGraphData.EditorObject.Get();
-						WorkspaceEditor->OpenObjects({ EditorObject });
-					}
-				}
 			}
 		}
 	}
 }
 
+bool FAnimNextGraphItemDetails::CanDelete(const FWorkspaceOutlinerItemExport& Export) const
+{
+	const TInstancedStruct<FWorkspaceOutlinerItemData>& Data = Export.GetData();
+	if (Data.IsValid() && Data.GetScriptStruct() == FAnimNextGraphOutlinerData::StaticStruct())
+	{
+		return true;
+	}
+	return false;
+}
+
 void FAnimNextGraphItemDetails::Delete(TConstArrayView<FWorkspaceOutlinerItemExport> Exports) const
 {
 	TMap<UAnimNextRigVMAssetEditorData*, TArray<UAnimNextRigVMAssetEntry*>> EntriesToDelete;
@@ -81,6 +71,7 @@ void FAnimNextGraphItemDetails::Delete(TConstArrayView<FWorkspaceOutlinerItemExp
 		if (Data.IsValid() && Data.GetScriptStruct() == FAnimNextGraphOutlinerData::StaticStruct())
 		{
 			const FAnimNextGraphOutlinerData& GraphData = Export.GetData().Get<FAnimNextGraphOutlinerData>();
+
 			UAnimNextRigVMAssetEditorData* EditorData = GraphData.Entry->GetTypedOuter<UAnimNextRigVMAssetEditorData>();
 			if(EditorData == nullptr)
 			{
@@ -125,7 +116,7 @@ void FAnimNextGraphItemDetails::Rename(const FWorkspaceOutlinerItemExport& Expor
 		UAnimNextRigVMAssetEditorData* EditorData = GraphData.Entry->GetTypedOuter<UAnimNextRigVMAssetEditorData>();
 		FName NewName = FName(*InName.ToString());
 		UAnimNextRigVMAssetEntry* ExistingEntry = EditorData->FindEntry(NewName);
-		if(ExistingEntry == nullptr && GraphData.Entry->GetEntryName() != NewName)
+		if (ExistingEntry == nullptr && GraphData.Entry->GetEntryName() != NewName)
 		{
 			FScopedTransaction Transaction(LOCTEXT("SetName", "Set Name"));
 			GraphData.Entry->SetEntryName(NewName);
@@ -146,40 +137,21 @@ bool FAnimNextGraphItemDetails::ValidateName(const FWorkspaceOutlinerItemExport&
 			OutErrorMessage = LOCTEXT("NameAlreadyExistsError", "Name already exists in this module");
 			return false;
 		}
-
-		return true;
 	}
+
+	OutErrorMessage = LOCTEXT("UnsupportedTypeRenameError", "Element type is not supported for rename");
 	return false;
 }
 
-UPackage* FAnimNextGraphItemDetails::GetPackage(const FWorkspaceOutlinerItemExport& Export) const 
+UPackage* FAnimNextGraphItemDetails::GetPackage(const FWorkspaceOutlinerItemExport& Export) const
 {
 	const TInstancedStruct<FWorkspaceOutlinerItemData>& Data = Export.GetData();
-	if (Data.IsValid())
+	if (Data.IsValid() && Data.GetScriptStruct() == FAnimNextGraphOutlinerData::StaticStruct())
 	{
-		if (Data.GetScriptStruct() == FAnimNextGraphOutlinerData::StaticStruct())
+		const FAnimNextGraphOutlinerData& GraphData = Data.Get<FAnimNextGraphOutlinerData>();
+		if (GraphData.GraphInterface)
 		{
-			const FAnimNextGraphOutlinerData& GraphData = Data.Get<FAnimNextGraphOutlinerData>();
-			if (GraphData.GraphInterface)
-			{
-				return GraphData.GraphInterface.GetObject()->GetExternalPackage();
-			}
-		}
-		else if(Data.GetScriptStruct() == FAnimNextGraphFunctionOutlinerData::StaticStruct())
-		{
-			const FAnimNextGraphFunctionOutlinerData& GraphFunctionData = Data.Get<FAnimNextGraphFunctionOutlinerData>();
-			if (GraphFunctionData.EditorObject.IsValid())
-			{
-				return GraphFunctionData.EditorObject->GetPackage();
-			}
-		}
-		else if (Data.GetScriptStruct() == FAnimNextCollapseGraphOutlinerData::StaticStruct())
-		{
-			const FAnimNextCollapseGraphOutlinerData& CollapseGraphData = Data.Get<FAnimNextCollapseGraphOutlinerData>();
-			if (CollapseGraphData.EditorObject.IsValid())
-			{
-				return CollapseGraphData.EditorObject->GetPackage();
-			}
+			return GraphData.GraphInterface.GetObject()->GetExternalPackage();
 		}
 	}
 	return nullptr;
@@ -192,14 +164,12 @@ const FSlateBrush* FAnimNextGraphItemDetails::GetItemIcon(const FWorkspaceOutlin
 
 void FAnimNextGraphItemDetails::RegisterToolMenuExtensions()
 {
-
 }
 
 void FAnimNextGraphItemDetails::UnregisterToolMenuExtensions()
 {
-
 }
 
 } // UE::AnimNext::Editor
 
-#undef LOCTEXT_NAMESPACE // "FAnimNextGraphItemDetails"
\ No newline at end of file
+#undef LOCTEXT_NAMESPACE // "FAnimNextGraphItemDetails"
diff --git a/Engine/Plugins/Experimental/AnimNext/Source/AnimNextEditor/Private/Graph/AnimNextGraphItemDetails.h b/Engine/Plugins/Experimental/AnimNext/Source/AnimNextEditor/Private/Graph/AnimNextGraphItemDetails.h
--- a/Engine/Plugins/Experimental/AnimNext/Source/AnimNextEditor/Private/Graph/AnimNextGraphItemDetails.h
+++ b/Engine/Plugins/Experimental/AnimNext/Source/AnimNextEditor/Private/Graph/AnimNextGraphItemDetails.h
@@ -6,12 +6,14 @@
 
 namespace UE::AnimNext::Editor
 {
+
 class FAnimNextGraphItemDetails : public UE::Workspace::IWorkspaceOutlinerItemDetails
 {
 public:
 	FAnimNextGraphItemDetails() = default;
 	virtual ~FAnimNextGraphItemDetails() override = default;
 	virtual void HandleDoubleClick(const FToolMenuContext& ToolMenuContext) const override;
+	virtual bool CanDelete(const FWorkspaceOutlinerItemExport& Export) const override;
 	virtual void Delete(TConstArrayView<FWorkspaceOutlinerItemExport> Exports) const override; 
 	virtual bool CanRename(const FWorkspaceOutlinerItemExport& Export) const override;
 	virtual void Rename(const FWorkspaceOutlinerItemExport& Export, const FText& InName) const override;
@@ -22,4 +24,5 @@ public:
 	static void RegisterToolMenuExtensions();
 	static void UnregisterToolMenuExtensions();
 };
+
 }
diff --git a/Engine/Plugins/Experimental/AnimNext/Source/AnimNextUncookedOnly/Internal/AnimNextAssetWorkspaceAssetUserData.h b/Engine/Plugins/Experimental/AnimNext/Source/AnimNextUncookedOnly/Internal/AnimNextAssetWorkspaceAssetUserData.h
--- a/Engine/Plugins/Experimental/AnimNext/Source/AnimNextUncookedOnly/Internal/AnimNextAssetWorkspaceAssetUserData.h
+++ b/Engine/Plugins/Experimental/AnimNext/Source/AnimNextUncookedOnly/Internal/AnimNextAssetWorkspaceAssetUserData.h
@@ -14,6 +14,8 @@
 #include "AnimNextAssetWorkspaceAssetUserData.generated.h"
 
 class UAnimNextRigVMAssetEntry;
+class URigVMEdGraphNode;
+
 // Base struct used to identify asset entries
 USTRUCT()
 struct FAnimNextRigVMAssetOutlinerData : public FWorkspaceOutlinerItemData
@@ -89,25 +91,33 @@ struct FAnimNextVariableOutlinerData : public FAnimNextAssetEntryOutlinerData
 };
 
 USTRUCT()
-struct FAnimNextCollapseGraphOutlinerData : public FWorkspaceOutlinerItemData
+struct FAnimNextCollapseGraphsOutlinerDataBase : public FWorkspaceOutlinerItemData
 {
 	GENERATED_BODY()
 	
-	FAnimNextCollapseGraphOutlinerData() = default;
+	FAnimNextCollapseGraphsOutlinerDataBase() = default;
 
 	UPROPERTY(VisibleAnywhere, Category=AnimNext)
 	TWeakObjectPtr<URigVMEdGraph> EditorObject;
 };
 
 USTRUCT()
-struct FAnimNextGraphFunctionOutlinerData : public FWorkspaceOutlinerItemData
+struct FAnimNextCollapseGraphOutlinerData : public FAnimNextCollapseGraphsOutlinerDataBase
+{
+	GENERATED_BODY()
+	
+	FAnimNextCollapseGraphOutlinerData() = default;
+};
+
+USTRUCT()
+struct FAnimNextGraphFunctionOutlinerData : public FAnimNextCollapseGraphsOutlinerDataBase
 {
 	GENERATED_BODY()
 	
 	FAnimNextGraphFunctionOutlinerData() = default;
 
 	UPROPERTY(VisibleAnywhere, Category=AnimNext)
-	TWeakObjectPtr<URigVMEdGraph> EditorObject;
+	TWeakObjectPtr<URigVMEdGraphNode> EdGraphNode;
 };
 
 USTRUCT()
diff --git a/Engine/Plugins/Experimental/AnimNext/Source/AnimNextUncookedOnly/Internal/AnimNextRigVMAssetEditorData.h b/Engine/Plugins/Experimental/AnimNext/Source/AnimNextUncookedOnly/Internal/AnimNextRigVMAssetEditorData.h
--- a/Engine/Plugins/Experimental/AnimNext/Source/AnimNextUncookedOnly/Internal/AnimNextRigVMAssetEditorData.h
+++ b/Engine/Plugins/Experimental/AnimNext/Source/AnimNextUncookedOnly/Internal/AnimNextRigVMAssetEditorData.h
@@ -33,6 +33,7 @@ namespace UE::AnimNext::Editor
 	class FWorkspaceEditor;
 	class FAnimNextAssetItemDetails;
 	class FAnimNextGraphItemDetails;
+	class FAnimNextFunctionItemDetails;
 	struct FVariablesOutlinerEntryItem;
 	class FVariablesOutlinerMode;
 	class FVariablesOutlinerHierarchy;
@@ -151,6 +152,7 @@ protected:
 	friend class UE::AnimNext::Editor::FVariableCustomization;
 	friend class UE::AnimNext::Editor::FAnimNextAssetItemDetails;
 	friend class UE::AnimNext::Editor::FAnimNextGraphItemDetails;
+	friend class UE::AnimNext::Editor::FAnimNextFunctionItemDetails;
 	friend struct UE::AnimNext::Editor::FVariablesOutlinerEntryItem;
 	friend class UE::AnimNext::Editor::FVariablesOutlinerMode;
 	friend class UE::AnimNext::Editor::FVariablesOutlinerHierarchy;
@@ -175,13 +177,20 @@ protected:
 	void HandlePackageDone();
 
 	// IRigVMClientHost interface
+	virtual FString GetAssetName() const override { return GetName(); }
+	virtual UClass* GetRigVMSchemaClass() const override;
+	virtual UScriptStruct* GetRigVMExecuteContextStruct() const override;
+	virtual UClass* GetRigVMEdGraphClass() const override;
+	virtual UClass* GetRigVMEdGraphNodeClass() const override;
+	virtual UClass* GetRigVMEdGraphSchemaClass() const override;
+	virtual UClass* GetRigVMEditorSettingsClass() const override;
 	virtual FRigVMClient* GetRigVMClient() override;
 	virtual const FRigVMClient* GetRigVMClient() const override;
 	virtual IRigVMGraphFunctionHost* GetRigVMGraphFunctionHost() override;
 	virtual const IRigVMGraphFunctionHost* GetRigVMGraphFunctionHost() const override;
 	virtual void HandleRigVMGraphAdded(const FRigVMClient* InClient, const FString& InNodePath) override;
 	virtual void HandleRigVMGraphRemoved(const FRigVMClient* InClient, const FString& InNodePath) override;
-	virtual void HandleRigVMGraphRenamed(const FRigVMClient* InClient, const FString& InOldNodePath, const FString& InNewNodePath) override {}
+	virtual void HandleRigVMGraphRenamed(const FRigVMClient* InClient, const FString& InOldNodePath, const FString& InNewNodePath) override;
 	virtual void HandleConfigureRigVMController(const FRigVMClient* InClient, URigVMController* InControllerToConfigure) override;
 	virtual UObject* GetEditorObjectForRigVMGraph(URigVMGraph* InVMGraph) const override;
 	virtual URigVMGraph* GetRigVMGraphForEditorObject(UObject* InObject) const override;
@@ -213,6 +222,11 @@ protected:
 	virtual URigVMController* GetOrCreateController(const UEdGraph* InGraph) override;
 	virtual TArray<FString> GeneratePythonCommands(const FString InNewBlueprintName) override;
 	virtual void SetupPinRedirectorsForBackwardsCompatibility() override;
+	virtual FRigVMGraphModifiedEvent& OnModified() override;
+	virtual bool IsFunctionPublic(const FName& InFunctionName) const override;
+	virtual void MarkFunctionPublic(const FName& InFunctionName, bool bIsPublic = true) override;
+	virtual void RenameGraph(const FString& InNodePath, const FName& InNewName) override;
+
 
 	// IRigVMGraphFunctionHost interface
 	virtual FRigVMGraphFunctionStore* GetRigVMGraphFunctionStore() override;
diff --git a/Engine/Plugins/Experimental/AnimNext/Source/AnimNextUncookedOnly/Internal/UncookedOnlyUtils.h b/Engine/Plugins/Experimental/AnimNext/Source/AnimNextUncookedOnly/Internal/UncookedOnlyUtils.h
--- a/Engine/Plugins/Experimental/AnimNext/Source/AnimNextUncookedOnly/Internal/UncookedOnlyUtils.h
+++ b/Engine/Plugins/Experimental/AnimNext/Source/AnimNextUncookedOnly/Internal/UncookedOnlyUtils.h
@@ -9,6 +9,7 @@
 #include "Param/ParamType.h"
 #include "RigVMCore/RigVMTemplate.h"
 #include "StructUtils/InstancedStruct.h"
+#include "RigVMCore/RigVMGraphFunctionDefinition.h"
 #include "UncookedOnlyUtils.generated.h"
 
 struct FAnimNextVariableBindingData;
@@ -17,6 +18,7 @@ struct FWorkspaceOutlinerItemExports;
 struct FWorkspaceOutlinerItemExport;
 struct FRigVMGraphFunctionData;
 struct FRigVMCompileSettings;
+struct FRigVMGraphFunctionHeaderArray;
 class UAnimNextModule;
 class UAnimNextModule_EditorData;
 class UAnimNextEdGraph;
@@ -27,11 +29,14 @@ class UAnimNextRigVMAsset;
 class UAnimNextRigVMAssetEditorData;
 class UAnimNextRigVMAssetEntry;
 
+
 namespace UE
 {
 	namespace AnimNext
 	{
 		static const FLazyName ExportsAnimNextAssetRegistryTag = TEXT("AnimNextExports");
+		static const FLazyName AnimNextPublicGraphFunctionsExportsRegistryTag = TEXT("AnimNextPublicGraphFunctions");
+		static const FLazyName ControlRigAssetPublicGraphFunctionsExportsRegistryTag = TEXT("PublicGraphFunctions");
 	}
 }
 
@@ -93,6 +98,9 @@ struct FAnimNextAssetRegistryExports
 
 	UPROPERTY()
 	TArray<FAnimNextAssetRegistryExportedVariable> Variables;
+
+	UPROPERTY()
+	TArray<FRigVMGraphFunctionHeader> PublicHeaders;
 };
 
 namespace UE::AnimNext::UncookedOnly
@@ -158,6 +166,12 @@ struct ANIMNEXTUNCOOKEDONLY_API FUtils
 	// Gets all the variables that are exported to the asset registry
 	static bool GetExportedVariablesFromAssetRegistry(TMap<FAssetData, FAnimNextAssetRegistryExports>& OutExports);
 
+	// Gets the functions that are exported to the asset registry for an asset
+	static bool GetExportedFunctionsForAsset(const FAssetData& InAsset, FAnimNextAssetRegistryExports& OutExports);
+
+	// Gets all the functions that are exported to the asset registry for the specified Tag
+	static bool GetExportedFunctionsFromAssetRegistry(FName Tag, TMap<FAssetData, FRigVMGraphFunctionHeaderArray>& OutExports);
+
 	// Gets the exported variables that are used by a RigVM asset
 	static void GetAssetVariables(const UAnimNextRigVMAssetEditorData* EditorData, FAnimNextAssetRegistryExports& OutExports);
 	static void GetAssetVariables(const UAnimNextRigVMAssetEditorData* EditorData, TSet<FAnimNextAssetRegistryExportedVariable>& OutExports);
diff --git a/Engine/Plugins/Experimental/AnimNext/Source/AnimNextUncookedOnly/Private/AnimNextRigVMAssetEditorData.cpp b/Engine/Plugins/Experimental/AnimNext/Source/AnimNextUncookedOnly/Private/AnimNextRigVMAssetEditorData.cpp
--- a/Engine/Plugins/Experimental/AnimNext/Source/AnimNextUncookedOnly/Private/AnimNextRigVMAssetEditorData.cpp
+++ b/Engine/Plugins/Experimental/AnimNext/Source/AnimNextUncookedOnly/Private/AnimNextRigVMAssetEditorData.cpp
@@ -167,6 +167,7 @@ void UAnimNextRigVMAssetEditorData::PostLoad()
 	
 	PostLoadExternalPackages();
 	RefreshExternalModels();
+
 	Initialize(/*bRecompileVM*/false);
 	
 	GetRigVMClient()->RefreshAllModels(ERigVMLoadType::PostLoad, false, bIsCompiling);
@@ -176,6 +177,31 @@ void UAnimNextRigVMAssetEditorData::PostLoad()
 	TArray<FName> BackwardsCompatiblePublicFunctions;
 	GetRigVMClient()->PatchFunctionsOnLoad(this, BackwardsCompatiblePublicFunctions, OldHeaders);
 
+	// Register function references at RigVMBuildData
+	if (URigVMBuildData* BuildData = URigVMBuildData::Get())
+	{
+		TArray<FRigVMReferenceNodeData> ReferenceNodeDatas;
+		const TArray<URigVMGraph*> AllModels = GetAllModels();
+		for (URigVMGraph* ModelToVisit : AllModels)
+		{
+			for (URigVMNode* Node : ModelToVisit->GetNodes())
+			{
+				if (URigVMFunctionReferenceNode* ReferenceNode = Cast<URigVMFunctionReferenceNode>(Node))
+				{
+					ReferenceNodeDatas.Add(FRigVMReferenceNodeData(ReferenceNode));
+				}
+			}
+		}
+
+		// update the build data from the current function references
+		for (const FRigVMReferenceNodeData& ReferenceNodeData : ReferenceNodeDatas)
+		{
+			BuildData->RegisterFunctionReference(ReferenceNodeData);
+		}
+
+		BuildData->ClearInvalidReferences();
+	}
+
 	// delay compilation until the package has been loaded
 	FCoreUObjectDelegates::OnEndLoadPackage.AddUObject(this, &UAnimNextRigVMAssetEditorData::HandlePackageDone);
 }
@@ -209,16 +235,36 @@ void UAnimNextRigVMAssetEditorData::GetAssetRegistryTags(FAssetRegistryTagsConte
 {
 	Super::GetAssetRegistryTags(Context);
 
-	// We may not have compiled yet, so cache exports if we havent already
-	if(!CachedExports.IsSet())
 	{
-		CachedExports = FAnimNextAssetRegistryExports();
-		UE::AnimNext::UncookedOnly::FUtils::GetAssetVariables(this, CachedExports.GetValue());
+		// We may not have compiled yet, so cache exports if we havent already
+		if(!CachedExports.IsSet())
+		{
+			CachedExports = FAnimNextAssetRegistryExports();
+			FAnimNextAssetRegistryExports& OutExports = CachedExports.GetValue();
+
+			UE::AnimNext::UncookedOnly::FUtils::GetAssetVariables(this, OutExports);
+		}
+
+		FString TagValue;
+		FAnimNextAssetRegistryExports::StaticStruct()->ExportText(TagValue, &CachedExports.GetValue(), nullptr, nullptr, PPF_None, nullptr);
+		Context.AddTag(FAssetRegistryTag(UE::AnimNext::ExportsAnimNextAssetRegistryTag, TagValue, FAssetRegistryTag::TT_Hidden));
 	}
 
-	FString TagValue;
-	FAnimNextAssetRegistryExports::StaticStruct()->ExportText(TagValue, &CachedExports.GetValue(), nullptr, nullptr, PPF_None, nullptr);
-	Context.AddTag(FAssetRegistryTag(UE::AnimNext::ExportsAnimNextAssetRegistryTag, TagValue, FAssetRegistryTag::TT_Hidden));
+	{
+		FRigVMGraphFunctionHeaderArray FunctionExports;
+		for (const FRigVMGraphFunctionData& FunctionData : GraphFunctionStore.PublicFunctions)
+		{
+			if (FunctionData.CompilationData.IsValid())
+			{
+				FunctionExports.Headers.Add(FunctionData.Header);
+			}
+		}
+
+		FString TagValue;
+		const FArrayProperty* HeadersProperty = CastField<FArrayProperty>(FRigVMGraphFunctionHeaderArray::StaticStruct()->FindPropertyByName(TEXT("Headers")));
+		HeadersProperty->ExportText_Direct(TagValue, &(FunctionExports.Headers), &(FunctionExports.Headers), nullptr, PPF_None, nullptr);
+		Context.AddTag(FAssetRegistryTag(UE::AnimNext::AnimNextPublicGraphFunctionsExportsRegistryTag, TagValue, FAssetRegistryTag::TT_Hidden));
+	}
 }
 
 bool UAnimNextRigVMAssetEditorData::Rename(const TCHAR* NewName, UObject* NewOuter, ERenameFlags Flags)
@@ -246,9 +292,9 @@ void UAnimNextRigVMAssetEditorData::HandlePackageDone()
 {
 	FCoreUObjectDelegates::OnEndLoadPackage.RemoveAll(this);
 
-	RecompileVM();
-
 	ReconstructAllNodes(); // If this is not executed on a node for whatever reason, it will appear transparent in the editor
+
+	RecompileVM();
 }
 
 void UAnimNextRigVMAssetEditorData::RefreshAllModels(ERigVMLoadType InLoadType)
@@ -357,6 +403,71 @@ void UAnimNextRigVMAssetEditorData::SetupPinRedirectorsForBackwardsCompatibility
 {
 }
 
+FRigVMGraphModifiedEvent& UAnimNextRigVMAssetEditorData::OnModified()
+{
+	return RigVMGraphModifiedEvent;
+}
+
+bool UAnimNextRigVMAssetEditorData::IsFunctionPublic(const FName& InFunctionName) const
+{
+	return GetLocalFunctionLibrary()->IsFunctionPublic(InFunctionName);
+}
+
+void UAnimNextRigVMAssetEditorData::MarkFunctionPublic(const FName& InFunctionName, bool bIsPublic)
+{
+	if (IsFunctionPublic(InFunctionName) == bIsPublic)
+	{
+		return;
+	}
+
+	URigVMController* Controller = RigVMClient.GetOrCreateController(GetLocalFunctionLibrary());
+	Controller->MarkFunctionAsPublic(InFunctionName, bIsPublic);
+}
+
+void UAnimNextRigVMAssetEditorData::RenameGraph(const FString& InNodePath, const FName& InNewName)
+{
+	if (URigVMGraph* ModelForNodePath = GetModel(InNodePath))
+	{
+		if (UEdGraph* EdGraph = Cast<UEdGraph>(GetEditorObjectForRigVMGraph(ModelForNodePath)))
+		{
+			FName OldName = NAME_None;
+			OldName = EdGraph->GetFName();
+
+			RigVMClient.RenameModel(InNodePath, InNewName, true);
+		}
+	}
+}
+
+UClass* UAnimNextRigVMAssetEditorData::GetRigVMSchemaClass() const
+{
+	return UAnimNextRigVMAssetSchema::StaticClass();
+}
+
+UScriptStruct* UAnimNextRigVMAssetEditorData::GetRigVMExecuteContextStruct() const 
+{
+	return FAnimNextExecuteContext::StaticStruct();
+}
+
+UClass* UAnimNextRigVMAssetEditorData::GetRigVMEdGraphClass() const 
+{
+	return UAnimNextEdGraph::StaticClass();
+}
+
+UClass* UAnimNextRigVMAssetEditorData::GetRigVMEdGraphNodeClass() const
+{
+	return UAnimNextEdGraphNode::StaticClass();
+}
+
+UClass* UAnimNextRigVMAssetEditorData::GetRigVMEdGraphSchemaClass() const
+{
+	return UAnimNextEdGraphSchema::StaticClass();
+}
+
+UClass* UAnimNextRigVMAssetEditorData::GetRigVMEditorSettingsClass() const
+{
+	return URigVMEditorSettings::StaticClass();
+}
+
 FRigVMClient* UAnimNextRigVMAssetEditorData::GetRigVMClient()
 {
 	return &RigVMClient;
@@ -426,6 +537,22 @@ void UAnimNextRigVMAssetEditorData::HandleRigVMGraphRemoved(const FRigVMClient*
 	}
 }
 
+void UAnimNextRigVMAssetEditorData::HandleRigVMGraphRenamed(const FRigVMClient* InClient, const FString& InOldNodePath, const FString& InNewNodePath)
+{
+	if (InClient->GetModel(InNewNodePath))
+	{
+		TArray<UEdGraph*> EdGraphs = GetAllEdGraphs();
+		for (UEdGraph* EdGraph : EdGraphs)
+		{
+			if (URigVMEdGraph* RigGraph = Cast<URigVMEdGraph>(EdGraph))
+			{
+				RigGraph->HandleRigVMGraphRenamed(InOldNodePath, InNewNodePath);
+			}
+		}
+	}
+}
+
+
 void UAnimNextRigVMAssetEditorData::HandleConfigureRigVMController(const FRigVMClient* InClient, URigVMController* InControllerToConfigure)
 {
 	InControllerToConfigure->OnModified().AddUObject(this, &UAnimNextRigVMAssetEditorData::HandleModifiedEvent);
@@ -488,7 +615,7 @@ UObject* UAnimNextRigVMAssetEditorData::GetEditorObjectForRigVMGraph(URigVMGraph
 			{
 				if (URigVMEdGraph* RigVMEdGraph = Cast<URigVMEdGraph>(SubGraph))
 				{
-					if (RigVMEdGraph->ModelNodePath == SearchGraphNodePath)
+					if (RigVMEdGraph->GetRigVMNodePath() == SearchGraphNodePath)
 					{
 						return RigVMEdGraph;
 					}
@@ -570,7 +697,16 @@ TObjectPtr<URigVMGraph> UAnimNextRigVMAssetEditorData::CreateContainedGraphModel
 	check(CollapseNode);
 
 	TObjectPtr<URigVMGraph> Model = NewObject<URigVMGraph>(CollapseNode, Name);
-	Model->SetSchemaClass(RigVMClient.GetDefaultSchemaClass());
+
+	check(CollapseNode->GetGraph());
+	if (CollapseNode->GetGraph()->GetSchema() != nullptr)
+	{
+		Model->SetSchemaClass(CollapseNode->GetGraph()->GetSchema()->GetClass());
+	}
+	else
+	{
+		Model->SetSchemaClass(RigVMClient.GetDefaultSchemaClass());
+	}
 
 	URigVMGraph* CollapseNodeModelRootGraph = CollapseNode->GetRootGraph();
 	check(CollapseNodeModelRootGraph);
@@ -597,8 +733,7 @@ void UAnimNextRigVMAssetEditorData::RecompileVM()
 
 	UAnimNextRigVMAsset* Asset = FUtils::GetAsset<UAnimNextRigVMAsset>(this);
 
-	CachedExports = FAnimNextAssetRegistryExports();
-	FUtils::GetAssetVariables(this, CachedExports.GetValue());
+	CachedExports.Reset();  // asset variables and other tags will be updated at the end by AssetRegistry->AssetUpdateTags
 
 	bErrorsDuringCompilation = false;
 
@@ -683,6 +818,7 @@ void UAnimNextRigVMAssetEditorData::RecompileVM()
 //	RefreshBreakpoints(EditorData);
 #endif
 
+	// Refresh CachedExports
 	if (IAssetRegistry* AssetRegistry = IAssetRegistry::Get())
 	{
 		AssetRegistry->AssetUpdateTags(Asset, EAssetRegistryTagsCaller::Fast);
@@ -779,6 +915,28 @@ void UAnimNextRigVMAssetEditorData::HandleModifiedEvent(ERigVMGraphNotifType InN
 		RequestAutoVMRecompilation();
 		break;
 	}
+	case ERigVMGraphNotifType::NodeRenamed:
+	{
+		if (URigVMCollapseNode* CollapseNode = Cast<URigVMCollapseNode>(InSubject))
+		{
+			FString NewNodePath = CollapseNode->GetNodePath(true /* recursive */);
+			FString Left, Right = NewNodePath;
+			URigVMNode::SplitNodePathAtEnd(NewNodePath, Left, Right);
+			FString OldNodePath = CollapseNode->GetPreviousFName().ToString();
+			if (!Left.IsEmpty())
+			{
+				OldNodePath = URigVMNode::JoinNodePath(Left, OldNodePath);
+			}
+
+			HandleRigVMGraphRenamed(GetRigVMClient(), OldNodePath, NewNodePath);
+
+			if (UEdGraph* ContainedEdGraph = Cast<UEdGraph>(GetEditorObjectForRigVMGraph(CollapseNode->GetContainedGraph())))
+			{
+				ContainedEdGraph->Rename(*CollapseNode->GetEditorSubGraphName(), nullptr);
+			}
+		}
+		break;
+	}
 	case ERigVMGraphNotifType::LinkAdded:
 	case ERigVMGraphNotifType::LinkRemoved:
 	case ERigVMGraphNotifType::PinArraySizeChanged:
@@ -843,12 +1001,15 @@ TArray<UEdGraph*> UAnimNextRigVMAssetEditorData::GetAllEdGraphs() const
 	{
 		if(IAnimNextRigVMGraphInterface* GraphInterface = Cast<IAnimNextRigVMGraphInterface>(Entry))
 		{
-			Graphs.Add(GraphInterface->GetEdGraph());
+			UEdGraph* EdGraph = GraphInterface->GetEdGraph();
+			Graphs.Add(EdGraph);
+			EdGraph->GetAllChildrenGraphs(Graphs);
 		}
 	}
 	for (URigVMEdGraph* RigVMEdGraph : FunctionEdGraphs)
 	{
 		Graphs.Add(RigVMEdGraph);
+		RigVMEdGraph->GetAllChildrenGraphs(Graphs);
 	}
 
 	return Graphs;
diff --git a/Engine/Plugins/Experimental/AnimNext/Source/AnimNextUncookedOnly/Private/Graph/AnimNextAnimationGraph_EditorData.cpp b/Engine/Plugins/Experimental/AnimNext/Source/AnimNextUncookedOnly/Private/Graph/AnimNextAnimationGraph_EditorData.cpp
--- a/Engine/Plugins/Experimental/AnimNext/Source/AnimNextUncookedOnly/Private/Graph/AnimNextAnimationGraph_EditorData.cpp
+++ b/Engine/Plugins/Experimental/AnimNext/Source/AnimNextUncookedOnly/Private/Graph/AnimNextAnimationGraph_EditorData.cpp
@@ -1,4 +1,4 @@
-﻿// Copyright Epic Games, Inc. All Rights Reserved.
+// Copyright Epic Games, Inc. All Rights Reserved.
 
 #include "Graph/AnimNextAnimationGraph_EditorData.h"
 
@@ -466,8 +466,7 @@ void UAnimNextAnimationGraph_EditorData::RecompileVM()
 	// to call trait destructors etc
 	AnimationGraph->FreezeGraphInstances();
 
-	CachedExports = FAnimNextAssetRegistryExports();
-	FUtils::GetAssetVariables(this, CachedExports.GetValue());
+	CachedExports.Reset(); // asset variables and other tags will be updated at the end by AssetRegistry->AssetUpdateTags
 
 	bErrorsDuringCompilation = false;
 
@@ -673,6 +672,7 @@ void UAnimNextAnimationGraph_EditorData::RecompileVM()
 //	RefreshBreakpoints(EditorData);
 #endif
 
+	// Refresh CachedExports, also updates variables at UAnimNextRigVMAssetEditorData::GetAssetRegistryTags
 	if (IAssetRegistry* AssetRegistry = IAssetRegistry::Get())
 	{
 		AssetRegistry->AssetUpdateTags(GetTypedOuter<UAnimNextAnimationGraph>(), EAssetRegistryTagsCaller::Fast);
diff --git a/Engine/Plugins/Experimental/AnimNext/Source/AnimNextUncookedOnly/Private/UncookedOnlyUtils.cpp b/Engine/Plugins/Experimental/AnimNext/Source/AnimNextUncookedOnly/Private/UncookedOnlyUtils.cpp
--- a/Engine/Plugins/Experimental/AnimNext/Source/AnimNextUncookedOnly/Private/UncookedOnlyUtils.cpp
+++ b/Engine/Plugins/Experimental/AnimNext/Source/AnimNextUncookedOnly/Private/UncookedOnlyUtils.cpp
@@ -742,6 +742,48 @@ bool FUtils::GetExportedVariablesFromAssetRegistry(TMap<FAssetData, FAnimNextAss
 	return OutExports.Num() > 0;
 }
 
+bool FUtils::GetExportedFunctionsForAsset(const FAssetData& InAsset, FAnimNextAssetRegistryExports& OutExports)
+{
+	const FString TagValue = InAsset.GetTagValueRef<FString>(UE::AnimNext::ExportsAnimNextAssetRegistryTag);
+	
+	bool bContainsFunctions = false;
+	FAnimNextAssetRegistryExports AssetExports;
+	if (FAnimNextAssetRegistryExports::StaticStruct()->ImportText(*TagValue, &AssetExports, nullptr, PPF_None, nullptr, FAnimNextAssetRegistryExports::StaticStruct()->GetName()) != nullptr)
+	{
+		if (AssetExports.PublicHeaders.Num() > 0)
+		{
+			OutExports = MoveTemp(AssetExports);
+			bContainsFunctions = true;
+		}
+	}
+
+	return bContainsFunctions;
+}
+
+bool FUtils::GetExportedFunctionsFromAssetRegistry(FName Tag, TMap<FAssetData, FRigVMGraphFunctionHeaderArray>& OutExports)
+{
+	TArray<FAssetData> AssetData;
+	IAssetRegistry::GetChecked().GetAssetsByTags({ Tag }, AssetData);
+
+	const FArrayProperty* HeadersProperty = CastField<FArrayProperty>(FRigVMGraphFunctionHeaderArray::StaticStruct()->FindPropertyByName(TEXT("Headers")));
+
+	for (const FAssetData& Asset : AssetData)
+	{
+		const FString TagValue = Asset.GetTagValueRef<FString>(Tag);
+		FRigVMGraphFunctionHeaderArray AssetExports;
+
+		if (HeadersProperty->ImportText_Direct(*TagValue, &AssetExports, nullptr, EPropertyPortFlags::PPF_None) != nullptr)
+		{
+			if (AssetExports.Headers.Num() > 0)
+			{
+				OutExports.Add(Asset, MoveTemp(AssetExports));
+			}
+		}
+	}
+
+	return OutExports.Num() > 0;
+}
+
 static void AddParamToSet(const FAnimNextAssetRegistryExportedVariable& InNewParam, TSet<FAnimNextAssetRegistryExportedVariable>& OutExports)
 {
 	if(FAnimNextAssetRegistryExportedVariable* ExistingEntry = OutExports.Find(InNewParam))
@@ -824,26 +866,34 @@ void FUtils::GetAssetOutlinerItems(const UAnimNextRigVMAssetEditorData* EditorDa
 		}
 	}
 
-	CreateFunctionLibraryOutlinerItemsRecursive(EditorData, OutExports, AssetIdentifier, EditorData->GetRigVMGraphFunctionStore()->PublicFunctions, EditorData->GetRigVMGraphFunctionStore()->PrivateFunctions);
+	//CreateFunctionLibraryOutlinerItemsRecursive(EditorData, OutExports, AssetIdentifier, EditorData->GetRigVMGraphFunctionStore()->PublicFunctions, EditorData->GetRigVMGraphFunctionStore()->PrivateFunctions);
 }
 
 void FUtils::CreateSubGraphsOutlinerItemsRecursive(const UAnimNextRigVMAssetEditorData* EditorData, FWorkspaceOutlinerItemExports& OutExports, FWorkspaceOutlinerItemExport& ParentExport, URigVMEdGraph* RigVMEdGraph)
 {
+	if (RigVMEdGraph == nullptr)
+	{
+		return;
+	}
+
 	// ---- Collapsed graphs ---
 	for (const TObjectPtr<UEdGraph>& SubGraph : RigVMEdGraph->SubGraphs)
 	{
 		URigVMEdGraph* EditorObject = Cast<URigVMEdGraph>(SubGraph);
 		if (IsValid(EditorObject))
 		{
-			URigVMCollapseNode* CollapseNode = CastChecked<URigVMCollapseNode>(EditorObject->GetModel()->GetOuter());
+			if(ensure(EditorObject->GetModel()))
+			{
+				URigVMCollapseNode* CollapseNode = CastChecked<URigVMCollapseNode>(EditorObject->GetModel()->GetOuter());
 
-			FWorkspaceOutlinerItemExport& Export = OutExports.Exports.Add_GetRef(FWorkspaceOutlinerItemExport(CollapseNode->GetFName(), ParentExport));
-			Export.GetData().InitializeAsScriptStruct(FAnimNextCollapseGraphOutlinerData::StaticStruct());
+				FWorkspaceOutlinerItemExport& Export = OutExports.Exports.Add_GetRef(FWorkspaceOutlinerItemExport(CollapseNode->GetFName(), ParentExport));
+				Export.GetData().InitializeAsScriptStruct(FAnimNextCollapseGraphOutlinerData::StaticStruct());
 			
-			FAnimNextCollapseGraphOutlinerData& FnGraphData = Export.GetData().GetMutable<FAnimNextCollapseGraphOutlinerData>();
-			FnGraphData.EditorObject = EditorObject;
+				FAnimNextCollapseGraphOutlinerData& FnGraphData = Export.GetData().GetMutable<FAnimNextCollapseGraphOutlinerData>();
+				FnGraphData.EditorObject = EditorObject;
 
-			CreateSubGraphsOutlinerItemsRecursive(EditorData, OutExports, Export, EditorObject);
+				CreateSubGraphsOutlinerItemsRecursive(EditorData, OutExports, Export, EditorObject);
+			}
 		}
 	}
 
@@ -851,21 +901,32 @@ void FUtils::CreateSubGraphsOutlinerItemsRecursive(const UAnimNextRigVMAssetEdit
 	TArray<URigVMEdGraphNode*> EdNodes;
 	RigVMEdGraph->GetNodesOfClass(EdNodes);
 
-	for (const URigVMEdGraphNode* EdNode : EdNodes)
+	for (URigVMEdGraphNode* EdNode : EdNodes)
 	{
 		if (URigVMFunctionReferenceNode* FunctionReferenceNode = Cast<URigVMFunctionReferenceNode>(EdNode->GetModelNode()))
 		{
 			if (URigVMLibraryNode* ReferencedNode = Cast<URigVMLibraryNode>(FunctionReferenceNode->GetReferencedFunctionHeader().LibraryPointer.GetNodeSoftPath().ResolveObject()))
 			{
-				FWorkspaceOutlinerItemExport& Export = OutExports.Exports.Add_GetRef(FWorkspaceOutlinerItemExport(ReferencedNode->GetFName(), ParentExport));
+				if (URigVMGraph* ContainedGraph = ReferencedNode->GetContainedGraph())
+				{
+					if (EditorData->GetEditorObjectForRigVMGraph(ContainedGraph) == nullptr)
+					{
+						continue; // Do not show references to other assets functions in the outliner
+					}
 
-				Export.GetData().InitializeAsScriptStruct(FAnimNextGraphFunctionOutlinerData::StaticStruct());
-				FAnimNextGraphFunctionOutlinerData& FnGraphData = Export.GetData().GetMutable<FAnimNextGraphFunctionOutlinerData>();
+					FWorkspaceOutlinerItemExport& Export = OutExports.Exports.Add_GetRef(FWorkspaceOutlinerItemExport(ReferencedNode->GetFName(), ParentExport));
 
-				URigVMEdGraph* ContainedGraph = Cast<URigVMEdGraph>(EditorData->GetEditorObjectForRigVMGraph(ReferencedNode->GetContainedGraph()));
-				FnGraphData.EditorObject = ContainedGraph;
+					Export.GetData().InitializeAsScriptStruct(FAnimNextGraphFunctionOutlinerData::StaticStruct());
+					FAnimNextGraphFunctionOutlinerData& FnGraphData = Export.GetData().GetMutable<FAnimNextGraphFunctionOutlinerData>();
 
-				CreateSubGraphsOutlinerItemsRecursive(EditorData, OutExports, Export, ContainedGraph);
+					if (URigVMEdGraph* ContainedEdGraph = Cast<URigVMEdGraph>(EditorData->GetEditorObjectForRigVMGraph(ContainedGraph)))
+					{
+						FnGraphData.EditorObject = ContainedEdGraph;
+						FnGraphData.EdGraphNode = EdNode;
+						
+						CreateSubGraphsOutlinerItemsRecursive(EditorData, OutExports, Export, ContainedEdGraph);
+					}
+				}
 			}
 		}
 	}
diff --git a/Engine/Plugins/Runtime/RigVM/Source/RigVMDeveloper/Private/EdGraph/RigVMEdGraphSchema.cpp b/Engine/Plugins/Runtime/RigVM/Source/RigVMDeveloper/Private/EdGraph/RigVMEdGraphSchema.cpp
--- a/Engine/Plugins/Runtime/RigVM/Source/RigVMDeveloper/Private/EdGraph/RigVMEdGraphSchema.cpp
+++ b/Engine/Plugins/Runtime/RigVM/Source/RigVMDeveloper/Private/EdGraph/RigVMEdGraphSchema.cpp
@@ -111,7 +111,7 @@ FRigVMNameValidator::FRigVMNameValidator(const UBlueprint* Blueprint, const UStr
 		FBlueprintEditorUtils::GetSCSVariableNameList(Blueprint, NamesTemp);
 		FBlueprintEditorUtils::GetImplementingBlueprintsFunctionNameList(Blueprint, NamesTemp);
 
-		for (FName & Name : NamesTemp)
+		for (FName& Name : NamesTemp)
 		{
 			Names.Add(Name.ToString());
 		}
@@ -1629,21 +1629,21 @@ bool URigVMEdGraphSchema::TryRenameGraph(UEdGraph* GraphToRename, const FName& I
 {
 	if (const URigVMEdGraph* RigGraph = Cast<URigVMEdGraph>(GraphToRename))
 	{
-		if (URigVMBlueprint* RigBlueprint = Cast<URigVMBlueprint>(FBlueprintEditorUtils::FindBlueprintForGraph(RigGraph)))
+		if (IRigVMClientHost* ClientHost = RigGraph->GetImplementingOuter<IRigVMClientHost>())
 		{
 			if (const URigVMGraph* Model = RigGraph->GetModel())
 			{
 				if(Model->IsRootGraph())
 				{
 					const FString NewName = FString::Printf(TEXT("%s %s"), FRigVMClient::RigVMModelPrefix, *InNewName.ToString()); 
-					RigBlueprint->RenameGraph(Model->GetNodePath(), *NewName);
+					ClientHost->RenameGraph(Model->GetNodePath(), *NewName);
 				}
 				else if (const URigVMGraph* RootModel = Model->GetRootGraph())
 				{
 					URigVMLibraryNode* LibraryNode = Cast<URigVMLibraryNode>(RootModel->FindNode(RigGraph->ModelNodePath));
 					if (LibraryNode)
 					{
-						if (URigVMController* Controller = RigBlueprint->GetOrCreateController(LibraryNode->GetGraph()))
+						if (URigVMController* Controller = ClientHost->GetOrCreateController(LibraryNode->GetGraph()))
 						{
 							Controller->RenameNode(LibraryNode, InNewName, true, true);
 							return true;
diff --git a/Engine/Plugins/Runtime/RigVM/Source/RigVMDeveloper/Private/RigVMModel/RigVMClient.cpp b/Engine/Plugins/Runtime/RigVM/Source/RigVMDeveloper/Private/RigVMModel/RigVMClient.cpp
--- a/Engine/Plugins/Runtime/RigVM/Source/RigVMDeveloper/Private/RigVMModel/RigVMClient.cpp
+++ b/Engine/Plugins/Runtime/RigVM/Source/RigVMDeveloper/Private/RigVMModel/RigVMClient.cpp
@@ -605,7 +605,16 @@ TObjectPtr<URigVMGraph> FRigVMClient::CreateContainedGraphModel(URigVMCollapseNo
 	else
 	{
 		Model = NewObject<URigVMGraph>(CollapseNode, Name);
-		Model->SetSchemaClass(GetDefaultSchemaClass());
+
+		// keep schema from collapse node graph, if exists
+		if (CollapseNode->GetGraph() != nullptr && CollapseNode->GetGraph()->GetSchema() != nullptr)
+		{
+			Model->SetSchemaClass(CollapseNode->GetGraph()->GetSchema()->GetClass());
+		}
+		else
+		{
+			Model->SetSchemaClass(GetDefaultSchemaClass());
+		}
 	}
 
 	return Model;
diff --git a/Engine/Plugins/Runtime/RigVM/Source/RigVMDeveloper/Public/RigVMBlueprint.h b/Engine/Plugins/Runtime/RigVM/Source/RigVMDeveloper/Public/RigVMBlueprint.h
--- a/Engine/Plugins/Runtime/RigVM/Source/RigVMDeveloper/Public/RigVMBlueprint.h
+++ b/Engine/Plugins/Runtime/RigVM/Source/RigVMDeveloper/Public/RigVMBlueprint.h
@@ -236,24 +236,6 @@ public:
 	/** Returns the class used as the super class for all generated classes */
 	virtual UClass* GetRigVMBlueprintGeneratedClassPrototype() const { return URigVMBlueprintGeneratedClass::StaticClass(); }
 
-	/** Returns the expected schema class to use for this blueprint */
-	virtual UClass* GetRigVMSchemaClass() const { return URigVMSchema::StaticClass(); }
-
-	/** Returns the expected execute context struct to use for this blueprint */
-	virtual UScriptStruct* GetRigVMExecuteContextStruct() const { return FRigVMExecuteContext::StaticStruct(); }
-
-	/** Returns the expected ed graph class to use for this blueprint */
-	virtual UClass* GetRigVMEdGraphClass() const { return URigVMEdGraph::StaticClass(); }
-
-	/** Returns the expected ed graph node class to use for this blueprint */
-	virtual UClass* GetRigVMEdGraphNodeClass() const { return URigVMEdGraphNode::StaticClass(); }
-
-	/** Returns the expected ed graph schema class to use for this blueprint */
-	virtual UClass* GetRigVMEdGraphSchemaClass() const { return URigVMEdGraphSchema::StaticClass(); }
-
-	/** Returns the class of the settings to use */
-	virtual UClass* GetRigVMEditorSettingsClass() const { return URigVMEditorSettings::StaticClass(); }
-
 	/** Returns the settings defaults for this blueprint */
 	URigVMEditorSettings* GetRigVMEditorSettings() const;
 
@@ -322,6 +304,13 @@ public:
 	virtual void GetPreloadDependencies(TArray<UObject*>& OutDeps) override;
 
 	//  --- IRigVMClientHost interface Start---
+	virtual FString GetAssetName() const override { return GetName(); }
+	virtual UClass* GetRigVMSchemaClass() const override { return URigVMSchema::StaticClass(); }
+	virtual UScriptStruct* GetRigVMExecuteContextStruct() const { return FRigVMExecuteContext::StaticStruct(); }
+	virtual UClass* GetRigVMEdGraphClass() const override { return URigVMEdGraph::StaticClass(); }
+	virtual UClass* GetRigVMEdGraphNodeClass() const override { return URigVMEdGraphNode::StaticClass(); }
+	virtual UClass* GetRigVMEdGraphSchemaClass() const override { return URigVMEdGraphSchema::StaticClass(); }
+	virtual UClass* GetRigVMEditorSettingsClass() const override { return URigVMEditorSettings::StaticClass(); }
 	virtual FRigVMClient* GetRigVMClient() override;
 	virtual const FRigVMClient* GetRigVMClient() const override;
 	virtual IRigVMGraphFunctionHost* GetRigVMGraphFunctionHost() override;
@@ -334,6 +323,7 @@ public:
 	virtual void HandleConfigureRigVMController(const FRigVMClient* InClient, URigVMController* InControllerToConfigure) override;
 	virtual UObject* ResolveUserDefinedTypeById(const FString& InTypeName) const override;
 
+
 	UFUNCTION(BlueprintCallable, Category = "RigVM Blueprint")
 	virtual void RecompileVM() override;
 	UFUNCTION(BlueprintCallable, Category = "RigVM Blueprint")
@@ -406,6 +396,13 @@ public:
 
 	virtual void SetupPinRedirectorsForBackwardsCompatibility() override {};
 
+	virtual FRigVMGraphModifiedEvent& OnModified() override;
+
+	virtual bool IsFunctionPublic(const FName& InFunctionName) const override;
+	virtual void MarkFunctionPublic(const FName& InFunctionName, bool bIsPublic = true) override;
+
+	virtual void RenameGraph(const FString& InNodePath, const FName& InNewName) override;
+
 	//  --- IRigVMClientHost interface End ---
 
 	//  --- IRigVMExternalDependencyManager interface Start ---
@@ -440,9 +437,6 @@ public:
 	TObjectPtr<URigVMEdGraph> FunctionLibraryEdGraph;
 #endif
 
-	bool IsFunctionPublic(const FName& InFunctionName) const;
-	void MarkFunctionPublic(const FName& InFunctionName, bool bIsPublic = true);
-
 	// Returns a list of dependencies of this blueprint.
 	// Dependencies are blueprints that contain functions used in this blueprint
 	TArray<URigVMBlueprint*> GetDependencies(bool bRecursive = false) const;
@@ -534,7 +528,6 @@ public:
 
 	void RebuildGraphFromModel();
 
-	FRigVMGraphModifiedEvent& OnModified();
 	FOnRigVMCompiledEvent& OnVMCompiled();
 
 	UFUNCTION(BlueprintCallable, Category = "VM")
@@ -753,7 +746,6 @@ private:
 	UEdGraph* CreateEdGraph(URigVMGraph* InModel, bool bForce = false);
 	bool RemoveEdGraph(URigVMGraph* InModel);
 	void DestroyObject(UObject* InObject);
-	void RenameGraph(const FString& InNodePath, const FName& InNewName);
 	void CreateEdGraphForCollapseNodeIfNeeded(URigVMCollapseNode* InNode, bool bForce = false);
 	bool RemoveEdGraphForCollapseNode(URigVMCollapseNode* InNode, bool bNotify = false);
 	void HandleReportFromCompiler(EMessageSeverity::Type InSeverity, UObject* InSubject, const FString& InMessage);
diff --git a/Engine/Plugins/Runtime/RigVM/Source/RigVMDeveloper/Public/RigVMModel/RigVMClient.h b/Engine/Plugins/Runtime/RigVM/Source/RigVMDeveloper/Public/RigVMModel/RigVMClient.h
--- a/Engine/Plugins/Runtime/RigVM/Source/RigVMDeveloper/Public/RigVMModel/RigVMClient.h
+++ b/Engine/Plugins/Runtime/RigVM/Source/RigVMDeveloper/Public/RigVMModel/RigVMClient.h
@@ -31,7 +31,28 @@ class RIGVMDEVELOPER_API IRigVMClientHost
 	GENERATED_BODY()
 
 public:
-	
+
+	/** Returns the expected schema class to use for this blueprint */
+	virtual FString GetAssetName() const = 0;
+
+	/** Returns the expected schema class to use for this blueprint */
+	virtual UClass* GetRigVMSchemaClass() const = 0;
+
+	/** Returns the expected execute context struct to use for this blueprint */
+	virtual UScriptStruct* GetRigVMExecuteContextStruct() const = 0;
+
+	/** Returns the expected ed graph class to use for this blueprint */
+	virtual UClass* GetRigVMEdGraphClass() const = 0;
+
+	/** Returns the expected ed graph node class to use for this blueprint */
+	virtual UClass* GetRigVMEdGraphNodeClass() const = 0;
+
+	/** Returns the expected ed graph schema class to use for this blueprint */
+	virtual UClass* GetRigVMEdGraphSchemaClass() const = 0;
+
+	/** Returns the class of the settings to use */
+	virtual UClass* GetRigVMEditorSettingsClass() const = 0;
+
 	// Returns the rigvm client for this host
 	virtual FRigVMClient* GetRigVMClient() = 0;
 
@@ -109,7 +130,6 @@ public:
 
 	virtual bool RemoveModel(FString InName = TEXT("Rig Graph"), bool bSetupUndoRedo = true, bool bPrintPythonCommand = true) = 0;
 
-
 	virtual FRigVMGetFocusedGraph& OnGetFocusedGraph() = 0;
 	virtual const FRigVMGetFocusedGraph& OnGetFocusedGraph() const = 0;
 
@@ -127,6 +147,13 @@ public:
 	virtual TArray<FString> GeneratePythonCommands(const FString InNewBlueprintName)  = 0;
 
 	virtual void SetupPinRedirectorsForBackwardsCompatibility() = 0;
+
+	virtual FRigVMGraphModifiedEvent& OnModified() = 0;
+
+	virtual bool IsFunctionPublic(const FName& InFunctionName) const = 0;
+	virtual void MarkFunctionPublic(const FName& InFunctionName, bool bIsPublic = true) = 0;
+
+	virtual void RenameGraph(const FString& InNodePath, const FName& InNewName) = 0;
 };
 
 UINTERFACE()
diff --git a/Engine/Plugins/Runtime/RigVM/Source/RigVMEditor/Private/Editor/RigVMGraphDetailCustomization.cpp b/Engine/Plugins/Runtime/RigVM/Source/RigVMEditor/Private/Editor/RigVMGraphDetailCustomization.cpp
--- a/Engine/Plugins/Runtime/RigVM/Source/RigVMEditor/Private/Editor/RigVMGraphDetailCustomization.cpp
+++ b/Engine/Plugins/Runtime/RigVM/Source/RigVMEditor/Private/Editor/RigVMGraphDetailCustomization.cpp
@@ -41,26 +41,26 @@
 static const FText RigVMGraphDetailCustomizationMultipleValues = LOCTEXT("MultipleValues", "Multiple Values");
 
 FRigVMFunctionArgumentGroupLayout::FRigVMFunctionArgumentGroupLayout(
-	URigVMGraph* InGraph, 
-	URigVMBlueprint* InBlueprint, 
-	TWeakPtr<FRigVMEditor> InEditor,
+	const TWeakObjectPtr<URigVMGraph>& InGraph,
+	const TWeakInterfacePtr<IRigVMClientHost>& InRigVMClientHost,
+	const TWeakPtr<FRigVMEditor>& InEditor,
 	bool bInputs)
 	: GraphPtr(InGraph)
-	, RigVMBlueprintPtr(InBlueprint)
+	, WeakRigVMClientHost(InRigVMClientHost)
 	, RigVMEditorPtr(InEditor)
 	, bIsInputGroup(bInputs)
 {
-	if (RigVMBlueprintPtr.IsValid())
+	if (WeakRigVMClientHost.IsValid())
 	{
-		RigVMBlueprintPtr.Get()->OnModified().AddRaw(this, &FRigVMFunctionArgumentGroupLayout::HandleModifiedEvent);
+		WeakRigVMClientHost->OnModified().AddRaw(this, &FRigVMFunctionArgumentGroupLayout::HandleModifiedEvent);
 	}
 }
 
 FRigVMFunctionArgumentGroupLayout::~FRigVMFunctionArgumentGroupLayout()
 {
-	if (RigVMBlueprintPtr.IsValid())
+	if (WeakRigVMClientHost.IsValid())
 	{
-		RigVMBlueprintPtr.Get()->OnModified().RemoveAll(this);
+		WeakRigVMClientHost->OnModified().RemoveAll(this);
 	}
 }
 
@@ -79,8 +79,8 @@ void FRigVMFunctionArgumentGroupLayout::GenerateChildContent(IDetailChildrenBuil
 				{
 					TSharedRef<class FRigVMFunctionArgumentLayout> ArgumentLayout = MakeShareable(new FRigVMFunctionArgumentLayout(
 						Pin,
-						Graph,
-						RigVMBlueprintPtr.Get(),
+						GraphPtr,
+						WeakRigVMClientHost,
 						RigVMEditorPtr
 					));
 					ChildrenBuilder.AddCustomBuilder(ArgumentLayout);
@@ -150,7 +150,7 @@ void FRigVMFunctionArgumentGroupLayout::HandleModifiedEvent(ERigVMGraphNotifType
 class FRigVMFunctionArgumentPinTypeSelectorFilter : public IPinTypeSelectorFilter
 {
 public:
-	FRigVMFunctionArgumentPinTypeSelectorFilter(TWeakPtr<FRigVMEditor> InRigVMEditor, TWeakObjectPtr<URigVMGraph> InGraph)
+	FRigVMFunctionArgumentPinTypeSelectorFilter(const TWeakPtr<FRigVMEditor>& InRigVMEditor, const TWeakObjectPtr<URigVMGraph>& InGraph)
 		: RigVMEditorPtr(InRigVMEditor), GraphPtr(InGraph)
 	{
 	}
@@ -326,13 +326,13 @@ void FRigVMFunctionArgumentLayout::GenerateChildContent(IDetailChildrenBuilder&
 
 void FRigVMFunctionArgumentLayout::OnRemoveClicked()
 {
-	if (PinPtr.IsValid() && RigVMBlueprintPtr.IsValid())
+	if (PinPtr.IsValid() && WeakRigVMClientHost.IsValid())
 	{
 		URigVMPin* Pin = PinPtr.Get();
-		URigVMBlueprint* Blueprint = RigVMBlueprintPtr.Get();
+		IRigVMClientHost* RigVMClientHost = WeakRigVMClientHost.Get();
 		if (URigVMLibraryNode* LibraryNode = Cast<URigVMLibraryNode>(Pin->GetNode()))
 		{
-			if (URigVMController* Controller = Blueprint->GetController(LibraryNode->GetContainedGraph()))
+			if (URigVMController* Controller = RigVMClientHost->GetController(LibraryNode->GetContainedGraph()))
 			{
 				Controller->RemoveExposedPin(Pin->GetFName(), true, true);
 			}
@@ -342,13 +342,13 @@ void FRigVMFunctionArgumentLayout::OnRemoveClicked()
 
 FReply FRigVMFunctionArgumentLayout::OnArgMoveUp()
 {
-	if (PinPtr.IsValid() && RigVMBlueprintPtr.IsValid())
+	if (PinPtr.IsValid() && WeakRigVMClientHost.IsValid())
 	{
 		URigVMPin* Pin = PinPtr.Get();
-		URigVMBlueprint* Blueprint = RigVMBlueprintPtr.Get();
+		IRigVMClientHost* RigVMClientHost = WeakRigVMClientHost.Get();
 		if (URigVMLibraryNode* LibraryNode = Cast<URigVMLibraryNode>(Pin->GetNode()))
 		{
-			if (URigVMController* Controller = Blueprint->GetController(LibraryNode->GetContainedGraph()))
+			if (URigVMController* Controller = RigVMClientHost->GetController(LibraryNode->GetContainedGraph()))
 			{
 				bool bIsInput = Pin->GetDirection() == ERigVMPinDirection::Input || Pin->GetDirection() == ERigVMPinDirection::IO;
 				
@@ -385,13 +385,13 @@ FReply FRigVMFunctionArgumentLayout::OnArgMoveUp()
 
 FReply FRigVMFunctionArgumentLayout::OnArgMoveDown()
 {
-	if (PinPtr.IsValid() && RigVMBlueprintPtr.IsValid())
+	if (PinPtr.IsValid() && WeakRigVMClientHost.IsValid())
 	{
 		URigVMPin* Pin = PinPtr.Get();
-		URigVMBlueprint* Blueprint = RigVMBlueprintPtr.Get();
+		IRigVMClientHost* RigVMClientHost = WeakRigVMClientHost.Get();
 		if (URigVMLibraryNode* LibraryNode = Cast<URigVMLibraryNode>(Pin->GetNode()))
 		{
-			if (URigVMController* Controller = Blueprint->GetController(LibraryNode->GetContainedGraph()))
+			if (URigVMController* Controller = RigVMClientHost->GetController(LibraryNode->GetContainedGraph()))
 			{
 				bool bIsInput = Pin->GetDirection() == ERigVMPinDirection::Input || Pin->GetDirection() == ERigVMPinDirection::IO;
 				
@@ -527,13 +527,13 @@ void FRigVMFunctionArgumentLayout::OnArgNameTextCommitted(const FText& NewText,
 {
 	if (InTextCommit == ETextCommit::OnEnter)
 	{
-		if (!NewText.IsEmpty() && PinPtr.IsValid() && RigVMBlueprintPtr.IsValid() && !ShouldPinBeReadOnly())
+		if (!NewText.IsEmpty() && PinPtr.IsValid() && WeakRigVMClientHost.IsValid() && !ShouldPinBeReadOnly())
 		{
 			URigVMPin* Pin = PinPtr.Get();
-			URigVMBlueprint* Blueprint = RigVMBlueprintPtr.Get();
+			IRigVMClientHost* RigVMClientHost = WeakRigVMClientHost.Get();
 			if (URigVMLibraryNode* LibraryNode = Cast<URigVMLibraryNode>(Pin->GetNode()))
 			{
-				if (URigVMController* Controller = Blueprint->GetController(LibraryNode->GetContainedGraph()))
+				if (URigVMController* Controller = RigVMClientHost->GetController(LibraryNode->GetContainedGraph()))
 				{
 					const FString& NewName = NewText.ToString();
 					Controller->RenameExposedPin(Pin->GetFName(), *NewName, true, true);
@@ -554,13 +554,13 @@ FEdGraphPinType FRigVMFunctionArgumentLayout::OnGetPinInfo() const
 
 void FRigVMFunctionArgumentLayout::PinInfoChanged(const FEdGraphPinType& PinType)
 {
-	if (PinPtr.IsValid() && RigVMBlueprintPtr.IsValid() && FBlueprintEditorUtils::IsPinTypeValid(PinType))
+	if (PinPtr.IsValid() && WeakRigVMClientHost.IsValid() && FBlueprintEditorUtils::IsPinTypeValid(PinType))
 	{
 		URigVMPin* Pin = PinPtr.Get();
-		URigVMBlueprint* Blueprint = RigVMBlueprintPtr.Get();
+		IRigVMClientHost* RigVMClientHost = WeakRigVMClientHost.Get();
 		if (URigVMLibraryNode* LibraryNode = Cast<URigVMLibraryNode>(Pin->GetNode()))
 		{
-			if (URigVMController* Controller = Blueprint->GetController(LibraryNode->GetContainedGraph()))
+			if (URigVMController* Controller = RigVMClientHost->GetController(LibraryNode->GetContainedGraph()))
 			{
 				FString CPPType;
 				FName CPPTypeObjectName = NAME_None;
@@ -586,19 +586,20 @@ void FRigVMFunctionArgumentLayout::OnPrePinInfoChange(const FEdGraphPinType& Pin
 }
 
 FRigVMFunctionArgumentDefaultNode::FRigVMFunctionArgumentDefaultNode(
-	URigVMGraph* InGraph,
-	URigVMBlueprint* InBlueprint
+	const TWeakObjectPtr<URigVMGraph>& InGraph,
+	const TWeakInterfacePtr<IRigVMClientHost>& InClientHost
 )
 	: GraphPtr(InGraph)
-	, RigVMBlueprintPtr(InBlueprint)
+	, WeakRigVMClientHost(InClientHost)
 {
-	if (GraphPtr.IsValid() && RigVMBlueprintPtr.IsValid())
+	if (GraphPtr.IsValid() && WeakRigVMClientHost.IsValid())
 	{
-		RigVMBlueprintPtr.Get()->OnModified().AddRaw(this, &FRigVMFunctionArgumentDefaultNode::HandleModifiedEvent);
+		IRigVMClientHost* RigVMClientHost = WeakRigVMClientHost.Get();
+		RigVMClientHost->OnModified().AddRaw(this, &FRigVMFunctionArgumentDefaultNode::HandleModifiedEvent);
 
 		if (URigVMLibraryNode* LibraryNode = Cast<URigVMLibraryNode>(GraphPtr->GetOuter()))
 		{
-			if (URigVMEdGraph* RigGraph = Cast<URigVMEdGraph>(RigVMBlueprintPtr->GetEdGraph(LibraryNode->GetGraph())))
+			if (URigVMEdGraph* RigGraph = Cast<URigVMEdGraph>(RigVMClientHost->GetEditorObjectForRigVMGraph(LibraryNode->GetGraph())))
 			{
 				EdGraphOuterPtr = RigGraph;
 				GraphChangedDelegateHandle = RigGraph->AddOnGraphChangedHandler(
@@ -612,9 +613,9 @@ FRigVMFunctionArgumentDefaultNode::FRigVMFunctionArgumentDefaultNode(
 
 FRigVMFunctionArgumentDefaultNode::~FRigVMFunctionArgumentDefaultNode()
 {
-	if (RigVMBlueprintPtr.IsValid())
+	if (WeakRigVMClientHost.IsValid())
 	{
-		RigVMBlueprintPtr.Get()->OnModified().RemoveAll(this);
+		WeakRigVMClientHost->OnModified().RemoveAll(this);
 	}
 	
 	if (EdGraphOuterPtr.IsValid())
@@ -628,17 +629,17 @@ FRigVMFunctionArgumentDefaultNode::~FRigVMFunctionArgumentDefaultNode()
 
 void FRigVMFunctionArgumentDefaultNode::GenerateChildContent(IDetailChildrenBuilder& ChildrenBuilder)
 {
-	if (!GraphPtr.IsValid() || !RigVMBlueprintPtr.IsValid())
+	if (!GraphPtr.IsValid() || !WeakRigVMClientHost.IsValid())
 	{
 		return;
 	}
 
-	URigVMBlueprint* Blueprint = RigVMBlueprintPtr.Get();
+	IRigVMClientHost* RigVMClientHost = WeakRigVMClientHost.Get();
 	URigVMGraph* Graph = GraphPtr.Get();
 	URigVMEdGraphNode* RigVMEdGraphNode = nullptr;
 	if (URigVMLibraryNode* LibraryNode = Cast<URigVMLibraryNode>(Graph->GetOuter()))
 	{
-		if (URigVMEdGraph* RigGraph = Cast<URigVMEdGraph>(Blueprint->GetEdGraph(LibraryNode->GetGraph())))
+		if (URigVMEdGraph* RigGraph = Cast<URigVMEdGraph>(RigVMClientHost->GetEditorObjectForRigVMGraph(LibraryNode->GetGraph())))
 		{
 			RigVMEdGraphNode = Cast<URigVMEdGraphNode>(RigGraph->FindNodeForModelNodeName(LibraryNode->GetFName()));
 		}
@@ -668,7 +669,7 @@ void FRigVMFunctionArgumentDefaultNode::GenerateChildContent(IDetailChildrenBuil
 
 void FRigVMFunctionArgumentDefaultNode::OnGraphChanged(const FEdGraphEditAction& InAction)
 {
-	if (GraphPtr.IsValid() && RigVMBlueprintPtr.IsValid())
+	if (GraphPtr.IsValid() && WeakRigVMClientHost.IsValid())
 	{
 		OnRebuildChildren.ExecuteIfBound();
 	}
@@ -741,10 +742,15 @@ TSharedPtr<IDetailCustomization> FRigVMGraphDetailCustomization::MakeInstance(TS
 	return nullptr;
 }
 
-void FRigVMGraphDetailCustomization::CustomizeDetails(IDetailLayoutBuilder& DetailLayout)
+FRigVMGraphDetailCustomization::FRigVMGraphDetailCustomization(TSharedPtr<FRigVMEditor> RigVMigEditor, URigVMBlueprint* RigVMBlueprint)
+	: RigVMEditorPtr(RigVMigEditor)
+	, RigVMBlueprintPtr(RigVMBlueprint)
+	, RigVMGraphDetailCustomizationImpl(MakeShared<FRigVMGraphDetailCustomizationImpl>())
 {
-	bIsPickingColor = false;
+}
 
+void FRigVMGraphDetailCustomization::CustomizeDetails(IDetailLayoutBuilder& DetailLayout)
+{
 	TArray<TWeakObjectPtr<UObject>> Objects;
 	DetailLayout.GetObjectsBeingCustomized(Objects);
 
@@ -772,6 +778,24 @@ void FRigVMGraphDetailCustomization::CustomizeDetails(IDetailLayoutBuilder& Deta
 		return;
 	}
 
+	RigVMGraphDetailCustomizationImpl->CustomizeDetails(DetailLayout, Model, Controller, Blueprint, RigVMEditorPtr);
+}
+
+void FRigVMGraphDetailCustomizationImpl::CustomizeDetails(IDetailLayoutBuilder& DetailLayout,
+	URigVMGraph* InModel,
+	URigVMController* InController,
+	IRigVMClientHost* InRigVMClientHost,
+	TWeakPtr<FRigVMEditor> InEditor)
+{
+	WeakModel = InModel;
+	WeakController = InController;
+	RigVMClientHost = InRigVMClientHost;
+	RigVMEditorPtr = InEditor;
+
+	bIsPickingColor = false;
+
+	URigVMGraph* Model = WeakModel.Get();
+
 	if (Model->IsTopLevelGraph())
 	{
 		IDetailCategoryBuilder& Category = DetailLayout.EditCategory("Graph", LOCTEXT("FunctionDetailsGraph", "Graph"));
@@ -795,9 +819,9 @@ void FRigVMGraphDetailCustomization::CustomizeDetails(IDetailLayoutBuilder& Deta
 	}
 
 	IDetailCategoryBuilder& InputsCategory = DetailLayout.EditCategory("Inputs", LOCTEXT("FunctionDetailsInputs", "Inputs"));
-	TSharedRef<FRigVMFunctionArgumentGroupLayout> InputArgumentGroup = MakeShareable(new FRigVMFunctionArgumentGroupLayout(
-		Model, 
-		Blueprint,
+	TSharedRef<FRigVMFunctionArgumentGroupLayout> InputArgumentGroup = MakeShareable<FRigVMFunctionArgumentGroupLayout>(new FRigVMFunctionArgumentGroupLayout(
+		Model,
+		RigVMClientHost.Get(),
 		RigVMEditorPtr,
 		true));
 	InputsCategory.AddCustomBuilder(InputArgumentGroup);
@@ -812,13 +836,13 @@ void FRigVMGraphDetailCustomization::CustomizeDetails(IDetailLayoutBuilder& Deta
 			SNew(SButton)
 			.ButtonStyle(FAppStyle::Get(), "SimpleButton")
 			.ContentPadding(FMargin(10.f, 0))
-			.OnClicked(this, &FRigVMGraphDetailCustomization::OnAddNewInputClicked)
-			.Visibility(this, &FRigVMGraphDetailCustomization::GetAddNewInputOutputVisibility)
+			.OnClicked(this, &FRigVMGraphDetailCustomizationImpl::OnAddNewInputClicked)
+			.Visibility(this, &FRigVMGraphDetailCustomizationImpl::GetAddNewInputOutputVisibility)
 			.HAlign(HAlign_Right)
 			.ToolTipText(LOCTEXT("FunctionNewInputArgTooltip", "Create a new input argument"))
 			.VAlign(VAlign_Center)
 			.AddMetaData<FTagMetaData>(FTagMetaData(TEXT("FunctionNewInputArg")))
-			.IsEnabled(this, &FRigVMGraphDetailCustomization::IsAddNewInputOutputEnabled)
+			.IsEnabled(this, &FRigVMGraphDetailCustomizationImpl::IsAddNewInputOutputEnabled)
 			[
 				SNew(SImage)
 				.Image(FAppStyle::Get().GetBrush("Icons.PlusCircle"))
@@ -830,8 +854,8 @@ void FRigVMGraphDetailCustomization::CustomizeDetails(IDetailLayoutBuilder& Deta
 	
 	IDetailCategoryBuilder& OutputsCategory = DetailLayout.EditCategory("Outputs", LOCTEXT("FunctionDetailsOutputs", "Outputs"));
 	TSharedRef<FRigVMFunctionArgumentGroupLayout> OutputArgumentGroup = MakeShareable(new FRigVMFunctionArgumentGroupLayout(
-		Model, 
-		Blueprint,
+		Model,
+		RigVMClientHost,
 		RigVMEditorPtr,
 		false));
 	OutputsCategory.AddCustomBuilder(OutputArgumentGroup);
@@ -846,13 +870,13 @@ void FRigVMGraphDetailCustomization::CustomizeDetails(IDetailLayoutBuilder& Deta
 			SNew(SButton)
 			.ButtonStyle(FAppStyle::Get(), "SimpleButton")
 			.ContentPadding(FMargin(10.f, 0))
-			.OnClicked(this, &FRigVMGraphDetailCustomization::OnAddNewOutputClicked)
-			.Visibility(this, &FRigVMGraphDetailCustomization::GetAddNewInputOutputVisibility)
+			.OnClicked(this, &FRigVMGraphDetailCustomizationImpl::OnAddNewOutputClicked)
+			.Visibility(this, &FRigVMGraphDetailCustomizationImpl::GetAddNewInputOutputVisibility)
 			.HAlign(HAlign_Right)
 			.ToolTipText(LOCTEXT("FunctionNewOutputArgTooltip", "Create a new output argument"))
 			.VAlign(VAlign_Center)
 			.AddMetaData<FTagMetaData>(FTagMetaData(TEXT("FunctionNewOutputArg")))
-			.IsEnabled(this, &FRigVMGraphDetailCustomization::IsAddNewInputOutputEnabled)
+			.IsEnabled(this, &FRigVMGraphDetailCustomizationImpl::IsAddNewInputOutputEnabled)
 			[
 				SNew(SImage)
 				.Image(FAppStyle::Get().GetBrush("Icons.PlusCircle"))
@@ -878,8 +902,8 @@ void FRigVMGraphDetailCustomization::CustomizeDetails(IDetailLayoutBuilder& Deta
 		[
 			SNew(SEditableTextBox)
 			.Font(IDetailLayoutBuilder::GetDetailFont())
-			.Text(this, &FRigVMGraphDetailCustomization::GetNodeCategory)
-			.OnTextCommitted(this, &FRigVMGraphDetailCustomization::SetNodeCategory)
+			.Text(this, &FRigVMGraphDetailCustomizationImpl::GetNodeCategory)
+			.OnTextCommitted(this, &FRigVMGraphDetailCustomizationImpl::SetNodeCategory)
 			.OnVerifyTextChanged_Lambda([&](const FText& InNewText, FText& OutErrorMessage) -> bool
 			{
 				const FText NewText = FEditorCategoryUtils::GetCategoryDisplayString(InNewText);
@@ -889,11 +913,11 @@ void FRigVMGraphDetailCustomization::CustomizeDetails(IDetailLayoutBuilder& Deta
 					return false;
 				}
 				
-				if (RigVMBlueprintPtr.IsValid())
+				if (RigVMClientHost.IsValid())
 				{
-					if (NewText.EqualTo(FText::FromString(RigVMBlueprintPtr.Get()->GetName())))
+					if (NewText.EqualTo(FText::FromString(RigVMClientHost->GetAssetName())))
 					{
-						OutErrorMessage = LOCTEXT("CategoryEqualsBlueprintName", "Cannot add a category with the same name as the blueprint.");
+						OutErrorMessage = LOCTEXT("CategoryEqualsBlueprintName", "Cannot add a category with the same name as the owner asset.");
 						return false;
 					}
 				}
@@ -913,8 +937,8 @@ void FRigVMGraphDetailCustomization::CustomizeDetails(IDetailLayoutBuilder& Deta
 		[
 			SNew(SEditableTextBox)
 			.Font(IDetailLayoutBuilder::GetDetailFont())
-			.Text(this, &FRigVMGraphDetailCustomization::GetNodeKeywords)
-			.OnTextCommitted(this, &FRigVMGraphDetailCustomization::SetNodeKeywords)
+			.Text(this, &FRigVMGraphDetailCustomizationImpl::GetNodeKeywords)
+			.OnTextCommitted(this, &FRigVMGraphDetailCustomizationImpl::SetNodeKeywords)
 		];
 
 		// description
@@ -929,11 +953,11 @@ void FRigVMGraphDetailCustomization::CustomizeDetails(IDetailLayoutBuilder& Deta
 		[
 			SNew(SMultiLineEditableText)
 			.Font(IDetailLayoutBuilder::GetDetailFont())
-			.Text(this, &FRigVMGraphDetailCustomization::GetNodeDescription)
-			.OnTextCommitted(this, &FRigVMGraphDetailCustomization::SetNodeDescription)
+			.Text(this, &FRigVMGraphDetailCustomizationImpl::GetNodeDescription)
+			.OnTextCommitted(this, &FRigVMGraphDetailCustomizationImpl::SetNodeDescription)
 		];
 
-		if(AccessSpecifierStrings.IsEmpty())
+		if (AccessSpecifierStrings.IsEmpty())
 		{
 			AccessSpecifierStrings.Add(TSharedPtr<FRigVMStringWithTag>(new FRigVMStringWithTag(TEXT("Public"))));
 			AccessSpecifierStrings.Add(TSharedPtr<FRigVMStringWithTag>(new FRigVMStringWithTag(TEXT("Private"))));
@@ -954,15 +978,15 @@ void FRigVMGraphDetailCustomization::CustomizeDetails(IDetailLayoutBuilder& Deta
             .ButtonContent()
             [
                 SNew(STextBlock)
-                    .Text(this, &FRigVMGraphDetailCustomization::GetCurrentAccessSpecifierName)
+                    .Text(this, &FRigVMGraphDetailCustomizationImpl::GetCurrentAccessSpecifierName)
                     .Font( IDetailLayoutBuilder::GetDetailFont() )
             ]
             .MenuContent()
             [
                 SNew(SListView<TSharedPtr<FRigVMStringWithTag> >)
                     .ListItemsSource( &AccessSpecifierStrings )
-                    .OnGenerateRow(this, &FRigVMGraphDetailCustomization::HandleGenerateRowAccessSpecifier)
-                    .OnSelectionChanged(this, &FRigVMGraphDetailCustomization::OnAccessSpecifierSelected)
+                    .OnGenerateRow(this, &FRigVMGraphDetailCustomizationImpl::HandleGenerateRowAccessSpecifier)
+                    .OnSelectionChanged(this, &FRigVMGraphDetailCustomizationImpl::OnAccessSpecifierSelected)
             ]
         ];
 
@@ -989,13 +1013,13 @@ void FRigVMGraphDetailCustomization::CustomizeDetails(IDetailLayoutBuilder& Deta
 		[
 			SNew(SRigVMVariantWidget)
 			.Context(VariantContext)
-			.Variant(this, &FRigVMGraphDetailCustomization::GetVariant)
-			.VariantRefs(this, &FRigVMGraphDetailCustomization::GetVariantRefs)
-			.OnVariantChanged(this, &FRigVMGraphDetailCustomization::OnVariantChanged)
-			.OnBrowseVariantRef(this, &FRigVMGraphDetailCustomization::OnBrowseVariantRef)
-			.OnGetTags(this, &FRigVMGraphDetailCustomization::OnGetAssignedTags)
-			.OnAddTag(this, &FRigVMGraphDetailCustomization::OnAddAssignedTag)
-			.OnRemoveTag(this, &FRigVMGraphDetailCustomization::OnRemoveAssignedTag)
+			.Variant(this, &FRigVMGraphDetailCustomizationImpl::GetVariant)
+			.VariantRefs(this, &FRigVMGraphDetailCustomizationImpl::GetVariantRefs)
+			.OnVariantChanged(this, &FRigVMGraphDetailCustomizationImpl::OnVariantChanged)
+			.OnBrowseVariantRef(this, &FRigVMGraphDetailCustomizationImpl::OnBrowseVariantRef)
+			.OnGetTags(this, &FRigVMGraphDetailCustomizationImpl::OnGetAssignedTags)
+			.OnAddTag(this, &FRigVMGraphDetailCustomizationImpl::OnAddAssignedTag)
+			.OnRemoveTag(this, &FRigVMGraphDetailCustomizationImpl::OnRemoveAssignedTag)
 			.CanAddTags(true)
 			.EnableTagContextMenu(true)
 		];
@@ -1015,10 +1039,10 @@ void FRigVMGraphDetailCustomization::CustomizeDetails(IDetailLayoutBuilder& Deta
 		[
 			SNew(SButton)
 			.ButtonStyle(FAppStyle::Get(), "Menu.Button")
-			.OnClicked(this, &FRigVMGraphDetailCustomization::OnNodeColorClicked)
+			.OnClicked(this, &FRigVMGraphDetailCustomizationImpl::OnNodeColorClicked)
 			[
 				SAssignNew(ColorBlock, SColorBlock)
-				.Color(this, &FRigVMGraphDetailCustomization::GetNodeColor)
+				.Color(this, &FRigVMGraphDetailCustomizationImpl::GetNodeColor)
 				.Size(FVector2D(77, 16))
 			]
 		];
@@ -1046,22 +1070,22 @@ void FRigVMGraphDetailCustomization::CustomizeDetails(IDetailLayoutBuilder& Deta
 				.HAlign(HAlign_Fill)
 				[
 					SNew(SRigVMNodeLayoutWidget)
-					.OnGetUncategorizedPins(this, &FRigVMGraphDetailCustomization::GetUncategorizedPins)
-					.OnGetCategories(this, &FRigVMGraphDetailCustomization::GetPinCategories)
-					.OnGetElementCategory(this, &FRigVMGraphDetailCustomization::GetPinCategory)
-					.OnGetElementIndexInCategory(this, &FRigVMGraphDetailCustomization::GetPinIndexInCategory)
-					.OnGetElementLabel(this, &FRigVMGraphDetailCustomization::GetPinLabel)
-					.OnGetElementColor(this, &FRigVMGraphDetailCustomization::GetPinColor)
-					.OnGetElementIcon(this, &FRigVMGraphDetailCustomization::GetPinIcon)
-					.OnCategoryAdded(this, &FRigVMGraphDetailCustomization::HandleCategoryAdded)
-					.OnCategoryRemoved(this, &FRigVMGraphDetailCustomization::HandleCategoryRemoved)
-					.OnCategoryRenamed(this, &FRigVMGraphDetailCustomization::HandleCategoryRenamed)
-					.OnElementCategoryChanged(this, &FRigVMGraphDetailCustomization::HandlePinCategoryChanged)
-					.OnElementLabelChanged(this, &FRigVMGraphDetailCustomization::HandlePinLabelChanged)
-					.OnElementIndexInCategoryChanged(this, &FRigVMGraphDetailCustomization::HandlePinIndexInCategoryChanged)
-					.OnValidateCategoryName(this, &FRigVMGraphDetailCustomization::HandleValidateCategoryName)
-					.OnValidateElementName(this, &FRigVMGraphDetailCustomization::HandleValidatePinDisplayName)
-					.OnGetStructuralHash(this, &FRigVMGraphDetailCustomization::GetNodeLayoutHash)
+					.OnGetUncategorizedPins(this, &FRigVMGraphDetailCustomizationImpl::GetUncategorizedPins)
+					.OnGetCategories(this, &FRigVMGraphDetailCustomizationImpl::GetPinCategories)
+					.OnGetElementCategory(this, &FRigVMGraphDetailCustomizationImpl::GetPinCategory)
+					.OnGetElementIndexInCategory(this, &FRigVMGraphDetailCustomizationImpl::GetPinIndexInCategory)
+					.OnGetElementLabel(this, &FRigVMGraphDetailCustomizationImpl::GetPinLabel)
+					.OnGetElementColor(this, &FRigVMGraphDetailCustomizationImpl::GetPinColor)
+					.OnGetElementIcon(this, &FRigVMGraphDetailCustomizationImpl::GetPinIcon)
+					.OnCategoryAdded(this, &FRigVMGraphDetailCustomizationImpl::HandleCategoryAdded)
+					.OnCategoryRemoved(this, &FRigVMGraphDetailCustomizationImpl::HandleCategoryRemoved)
+					.OnCategoryRenamed(this, &FRigVMGraphDetailCustomizationImpl::HandleCategoryRenamed)
+					.OnElementCategoryChanged(this, &FRigVMGraphDetailCustomizationImpl::HandlePinCategoryChanged)
+					.OnElementLabelChanged(this, &FRigVMGraphDetailCustomizationImpl::HandlePinLabelChanged)
+					.OnElementIndexInCategoryChanged(this, &FRigVMGraphDetailCustomizationImpl::HandlePinIndexInCategoryChanged)
+					.OnValidateCategoryName(this, &FRigVMGraphDetailCustomizationImpl::HandleValidateCategoryName)
+					.OnValidateElementName(this, &FRigVMGraphDetailCustomizationImpl::HandleValidatePinDisplayName)
+					.OnGetStructuralHash(this, &FRigVMGraphDetailCustomizationImpl::GetNodeLayoutHash)
 				];
 			}
 		}
@@ -1070,34 +1094,32 @@ void FRigVMGraphDetailCustomization::CustomizeDetails(IDetailLayoutBuilder& Deta
 	IDetailCategoryBuilder& DefaultsCategory = DetailLayout.EditCategory("NodeDefaults", LOCTEXT("FunctionDetailsNodeDefaults", "Node Defaults"));
 	TSharedRef<FRigVMFunctionArgumentDefaultNode> DefaultsArgumentNode = MakeShareable(new FRigVMFunctionArgumentDefaultNode(
 		Model,
-		Blueprint));
+		RigVMClientHost));
 	DefaultsCategory.AddCustomBuilder(DefaultsArgumentNode);
 }
 
-bool FRigVMGraphDetailCustomization::IsAddNewInputOutputEnabled() const
+bool FRigVMGraphDetailCustomizationImpl::IsAddNewInputOutputEnabled() const
 {
 	return true;
 }
 
-EVisibility FRigVMGraphDetailCustomization::GetAddNewInputOutputVisibility() const
+EVisibility FRigVMGraphDetailCustomizationImpl::GetAddNewInputOutputVisibility() const
 {
 	return EVisibility::Visible;
 }
 
-FReply FRigVMGraphDetailCustomization::OnAddNewInputClicked()
+FReply FRigVMGraphDetailCustomizationImpl::OnAddNewInputClicked()
 {
-	if (GraphPtr.IsValid() && RigVMBlueprintPtr.IsValid())
+	if (WeakModel.IsValid() && WeakController.IsValid())
 	{
-		URigVMBlueprint* Blueprint = RigVMBlueprintPtr.Get();
-		URigVMGraph* Model = Blueprint->GetModel(GraphPtr.Get());
-		if (URigVMController* Controller = Blueprint->GetController(Model))
+		if (URigVMController* Controller = WeakController.Get())
 		{
 			FName ArgumentName = TEXT("Argument");
 			FString CPPType = TEXT("bool");
 			FName CPPTypeObjectPath = NAME_None;
 			FString DefaultValue = TEXT("False");
 
-			if (URigVMLibraryNode* LibraryNode = Cast<URigVMLibraryNode>(Model->GetOuter()))
+			if (URigVMLibraryNode* LibraryNode = Cast<URigVMLibraryNode>(WeakModel->GetOuter()))
 			{
 				if (LibraryNode->GetPins().Num() > 0)
 				{
@@ -1138,13 +1160,11 @@ FReply FRigVMGraphDetailCustomization::OnAddNewInputClicked()
 	return FReply::Unhandled();
 }
 
-FReply FRigVMGraphDetailCustomization::OnAddNewOutputClicked()
+FReply FRigVMGraphDetailCustomizationImpl::OnAddNewOutputClicked()
 {
-	if (GraphPtr.IsValid() && RigVMBlueprintPtr.IsValid())
+	if (WeakModel.IsValid() && WeakController.IsValid())
 	{
-		URigVMBlueprint* Blueprint = RigVMBlueprintPtr.Get();
-		URigVMGraph* Model = Blueprint->GetModel(GraphPtr.Get());
-		if (URigVMController* Controller = Blueprint->GetController(Model))
+		if (URigVMController* Controller = WeakController.Get())
 		{
 			FName ArgumentName = TEXT("Argument");
 			FString CPPType = TEXT("bool");
@@ -1158,12 +1178,11 @@ FReply FRigVMGraphDetailCustomization::OnAddNewOutputClicked()
 	return FReply::Unhandled();
 }
 
-FText FRigVMGraphDetailCustomization::GetNodeCategory() const
+FText FRigVMGraphDetailCustomizationImpl::GetNodeCategory() const
 {
-	if (GraphPtr.IsValid() && RigVMBlueprintPtr.IsValid())
+	if (WeakModel.IsValid() && WeakController.IsValid())
 	{
-		URigVMBlueprint* Blueprint = RigVMBlueprintPtr.Get();
-		if (URigVMGraph* Model = Blueprint->GetModel(GraphPtr.Get()))
+		if (URigVMGraph* Model = WeakModel.Get())
 		{
 			if (URigVMCollapseNode* OuterNode = Cast<URigVMCollapseNode>(Model->GetOuter()))
 			{
@@ -1175,21 +1194,20 @@ FText FRigVMGraphDetailCustomization::GetNodeCategory() const
 	return FText();
 }
 
-void FRigVMGraphDetailCustomization::SetNodeCategory(const FText& InNewText, ETextCommit::Type InCommitType)
+void FRigVMGraphDetailCustomizationImpl::SetNodeCategory(const FText& InNewText, ETextCommit::Type InCommitType)
 {
 	if(InCommitType == ETextCommit::OnCleared)
 	{
 		return;
 	}
 
-	if (GraphPtr.IsValid() && RigVMBlueprintPtr.IsValid())
+	if (WeakModel.IsValid() && WeakController.IsValid())
 	{
-		URigVMBlueprint* Blueprint = RigVMBlueprintPtr.Get();
-		if (URigVMGraph* Model = Blueprint->GetModel(GraphPtr.Get()))
+		if (URigVMGraph* Model = WeakModel.Get())
 		{
 			if (URigVMCollapseNode* OuterNode = Cast<URigVMCollapseNode>(Model->GetOuter()))
 			{
-				if (URigVMController* Controller = Blueprint->GetOrCreateController(OuterNode->GetGraph()))
+				if (URigVMController* Controller = RigVMClientHost->GetOrCreateController(OuterNode->GetGraph()))
 				{
 					Controller->SetNodeCategory(OuterNode, InNewText.ToString(), true, false, true);
 				}
@@ -1198,12 +1216,11 @@ void FRigVMGraphDetailCustomization::SetNodeCategory(const FText& InNewText, ETe
 	}
 }
 
-FText FRigVMGraphDetailCustomization::GetNodeKeywords() const
+FText FRigVMGraphDetailCustomizationImpl::GetNodeKeywords() const
 {
-	if (GraphPtr.IsValid() && RigVMBlueprintPtr.IsValid())
+	if (WeakModel.IsValid() && WeakController.IsValid())
 	{
-		URigVMBlueprint* Blueprint = RigVMBlueprintPtr.Get();
-		if (URigVMGraph* Model = Blueprint->GetModel(GraphPtr.Get()))
+		if (URigVMGraph* Model = WeakModel.Get())
 		{
 			if (URigVMCollapseNode* OuterNode = Cast<URigVMCollapseNode>(Model->GetOuter()))
 			{
@@ -1215,21 +1232,20 @@ FText FRigVMGraphDetailCustomization::GetNodeKeywords() const
 	return FText();
 }
 
-void FRigVMGraphDetailCustomization::SetNodeKeywords(const FText& InNewText, ETextCommit::Type InCommitType)
+void FRigVMGraphDetailCustomizationImpl::SetNodeKeywords(const FText& InNewText, ETextCommit::Type InCommitType)
 {
 	if(InCommitType == ETextCommit::OnCleared)
 	{
 		return;
 	}
 
-	if (GraphPtr.IsValid() && RigVMBlueprintPtr.IsValid())
+	if (WeakModel.IsValid() && WeakController.IsValid())
 	{
-		URigVMBlueprint* Blueprint = RigVMBlueprintPtr.Get();
-		if (URigVMGraph* Model = Blueprint->GetModel(GraphPtr.Get()))
+		if (URigVMGraph* Model = WeakModel.Get())
 		{
 			if (URigVMCollapseNode* OuterNode = Cast<URigVMCollapseNode>(Model->GetOuter()))
 			{
-				if (URigVMController* Controller = Blueprint->GetOrCreateController(OuterNode->GetGraph()))
+				if (URigVMController* Controller = RigVMClientHost->GetOrCreateController(OuterNode->GetGraph()))
 				{
 					Controller->SetNodeKeywords(OuterNode, InNewText.ToString(), true, false, true);
 				}
@@ -1238,12 +1254,11 @@ void FRigVMGraphDetailCustomization::SetNodeKeywords(const FText& InNewText, ETe
 	}
 }
 
-FText FRigVMGraphDetailCustomization::GetNodeDescription() const
+FText FRigVMGraphDetailCustomizationImpl::GetNodeDescription() const
 {
-	if (GraphPtr.IsValid() && RigVMBlueprintPtr.IsValid())
+	if (WeakModel.IsValid() && WeakController.IsValid())
 	{
-		URigVMBlueprint* Blueprint = RigVMBlueprintPtr.Get();
-		if (URigVMGraph* Model = Blueprint->GetModel(GraphPtr.Get()))
+		if (URigVMGraph* Model = WeakModel.Get())
 		{
 			if (URigVMCollapseNode* OuterNode = Cast<URigVMCollapseNode>(Model->GetOuter()))
 			{
@@ -1255,21 +1270,20 @@ FText FRigVMGraphDetailCustomization::GetNodeDescription() const
 	return FText();
 }
 
-void FRigVMGraphDetailCustomization::SetNodeDescription(const FText& InNewText, ETextCommit::Type InCommitType)
+void FRigVMGraphDetailCustomizationImpl::SetNodeDescription(const FText& InNewText, ETextCommit::Type InCommitType)
 {
-	if(InCommitType == ETextCommit::OnCleared)
+	if (InCommitType == ETextCommit::OnCleared)
 	{
 		return;
 	}
 
-	if (GraphPtr.IsValid() && RigVMBlueprintPtr.IsValid())
+	if (WeakModel.IsValid() && WeakController.IsValid())
 	{
-		URigVMBlueprint* Blueprint = RigVMBlueprintPtr.Get();
-		if (URigVMGraph* Model = Blueprint->GetModel(GraphPtr.Get()))
+		if (URigVMGraph* Model = WeakModel.Get())
 		{
 			if (URigVMCollapseNode* OuterNode = Cast<URigVMCollapseNode>(Model->GetOuter()))
 			{
-				if (URigVMController* Controller = Blueprint->GetOrCreateController(OuterNode->GetGraph()))
+				if (URigVMController* Controller = RigVMClientHost->GetOrCreateController(OuterNode->GetGraph()))
 				{
 					Controller->SetNodeDescription(OuterNode, InNewText.ToString(), true, false, true);
 				}
@@ -1278,12 +1292,11 @@ void FRigVMGraphDetailCustomization::SetNodeDescription(const FText& InNewText,
 	}
 }
 
-FLinearColor FRigVMGraphDetailCustomization::GetNodeColor() const
+FLinearColor FRigVMGraphDetailCustomizationImpl::GetNodeColor() const
 {
-	if (GraphPtr.IsValid() && RigVMBlueprintPtr.IsValid())
+	if (WeakModel.IsValid() && WeakController.IsValid())
 	{
-		URigVMBlueprint* Blueprint = RigVMBlueprintPtr.Get();
-		if (URigVMGraph* Model = Blueprint->GetModel(GraphPtr.Get()))
+		if (URigVMGraph* Model = WeakModel.Get())
 		{
 			if (URigVMCollapseNode* OuterNode = Cast<URigVMCollapseNode>(Model->GetOuter()))
 			{
@@ -1294,16 +1307,15 @@ FLinearColor FRigVMGraphDetailCustomization::GetNodeColor() const
 	return FLinearColor::White;
 }
 
-void FRigVMGraphDetailCustomization::SetNodeColor(FLinearColor InColor, bool bSetupUndoRedo)
+void FRigVMGraphDetailCustomizationImpl::SetNodeColor(FLinearColor InColor, bool bSetupUndoRedo)
 {
-	if (GraphPtr.IsValid() && RigVMBlueprintPtr.IsValid())
+	if (WeakModel.IsValid() && WeakController.IsValid())
 	{
-		URigVMBlueprint* Blueprint = RigVMBlueprintPtr.Get();
-		if (URigVMGraph* Model = Blueprint->GetModel(GraphPtr.Get()))
+		if (URigVMGraph* Model = WeakModel.Get())
 		{
 			if (URigVMCollapseNode* OuterNode = Cast<URigVMCollapseNode>(Model->GetOuter()))
 			{
-				if (URigVMController* Controller = Blueprint->GetOrCreateController(OuterNode->GetGraph()))
+				if (URigVMController* Controller = RigVMClientHost->GetOrCreateController(OuterNode->GetGraph()))
 				{
 					Controller->SetNodeColor(OuterNode, InColor, bSetupUndoRedo, bIsPickingColor, true);
 				}
@@ -1312,47 +1324,44 @@ void FRigVMGraphDetailCustomization::SetNodeColor(FLinearColor InColor, bool bSe
 	}
 }
 
-void FRigVMGraphDetailCustomization::OnNodeColorBegin()
+void FRigVMGraphDetailCustomizationImpl::OnNodeColorBegin()
 {
 	bIsPickingColor = true;
 }
-void FRigVMGraphDetailCustomization::OnNodeColorEnd()
-{ 
-	bIsPickingColor = false; 
+void FRigVMGraphDetailCustomizationImpl::OnNodeColorEnd()
+{
+	bIsPickingColor = false;
 }
 
-void FRigVMGraphDetailCustomization::OnNodeColorCancelled(FLinearColor OriginalColor)
+void FRigVMGraphDetailCustomizationImpl::OnNodeColorCancelled(FLinearColor OriginalColor)
 {
 	SetNodeColor(OriginalColor, true);
 }
 
-FReply FRigVMGraphDetailCustomization::OnNodeColorClicked()
+FReply FRigVMGraphDetailCustomizationImpl::OnNodeColorClicked()
 {
-	FColorPickerArgs PickerArgs = FColorPickerArgs(GetNodeColor(), FOnLinearColorValueChanged::CreateSP(this, &FRigVMGraphDetailCustomization::SetNodeColor, true));
+	FColorPickerArgs PickerArgs = FColorPickerArgs(GetNodeColor(), FOnLinearColorValueChanged::CreateSP(this, &FRigVMGraphDetailCustomizationImpl::SetNodeColor, true));
 	PickerArgs.ParentWidget = ColorBlock;
 	PickerArgs.bUseAlpha = false;
 	PickerArgs.DisplayGamma = false;
-	PickerArgs.OnInteractivePickBegin = FSimpleDelegate::CreateSP(this, &FRigVMGraphDetailCustomization::OnNodeColorBegin);
-	PickerArgs.OnInteractivePickEnd = FSimpleDelegate::CreateSP(this, &FRigVMGraphDetailCustomization::OnNodeColorEnd);
-	PickerArgs.OnColorPickerCancelled = FOnColorPickerCancelled::CreateSP(this, &FRigVMGraphDetailCustomization::OnNodeColorCancelled);
+	PickerArgs.OnInteractivePickBegin = FSimpleDelegate::CreateSP(this, &FRigVMGraphDetailCustomizationImpl::OnNodeColorBegin);
+	PickerArgs.OnInteractivePickEnd = FSimpleDelegate::CreateSP(this, &FRigVMGraphDetailCustomizationImpl::OnNodeColorEnd);
+	PickerArgs.OnColorPickerCancelled = FOnColorPickerCancelled::CreateSP(this, &FRigVMGraphDetailCustomizationImpl::OnNodeColorCancelled);
 	OpenColorPicker(PickerArgs);
 	return FReply::Handled();
 }
 
-TArray<TSharedPtr<FRigVMStringWithTag>> FRigVMGraphDetailCustomization::AccessSpecifierStrings;
+TArray<TSharedPtr<FRigVMStringWithTag>> FRigVMGraphDetailCustomizationImpl::AccessSpecifierStrings;
 
-FText FRigVMGraphDetailCustomization::GetCurrentAccessSpecifierName() const
+FText FRigVMGraphDetailCustomizationImpl::GetCurrentAccessSpecifierName() const
 {
-	if(RigVMBlueprintPtr.IsValid() && GraphPtr.IsValid())
+	if (WeakModel.IsValid() && WeakController.IsValid())
 	{
-		URigVMEdGraph* Graph = GraphPtr.Get();
-		URigVMBlueprint* RigVMBlueprint = RigVMBlueprintPtr.Get();
-
-		if (URigVMGraph* RigVMGraph = Graph->GetModel())
+		if (URigVMGraph* RigVMGraph = WeakModel.Get())
 		{
 			if (URigVMLibraryNode* LibraryNode = RigVMGraph->GetTypedOuter<URigVMLibraryNode>())
 			{
-				if(RigVMBlueprint->IsFunctionPublic(LibraryNode->GetFName()))
+				if(RigVMClientHost->GetLocalFunctionLibrary()->IsFunctionPublic(LibraryNode->GetFName()))
 				{
 					return FText::FromString(AccessSpecifierStrings[0]->GetString()); // public
 				}
@@ -1363,31 +1372,28 @@ FText FRigVMGraphDetailCustomization::GetCurrentAccessSpecifierName() const
 	return FText::FromString(AccessSpecifierStrings[1]->GetString()); // private
 }
 
-void FRigVMGraphDetailCustomization::OnAccessSpecifierSelected( TSharedPtr<FRigVMStringWithTag> SpecifierName, ESelectInfo::Type SelectInfo )
+void FRigVMGraphDetailCustomizationImpl::OnAccessSpecifierSelected(TSharedPtr<FRigVMStringWithTag> SpecifierName, ESelectInfo::Type SelectInfo)
 {
-	if(RigVMBlueprintPtr.IsValid() && GraphPtr.IsValid())
+	if(WeakModel.IsValid() && WeakController.IsValid())
 	{
-		URigVMEdGraph* Graph = GraphPtr.Get();
-		URigVMBlueprint* RigVMBlueprint = RigVMBlueprintPtr.Get();
-
-		if (URigVMGraph* RigVMGraph = Graph->GetModel())
+		if (URigVMGraph* RigVMGraph = WeakModel.Get())
 		{
 			if (URigVMLibraryNode* LibraryNode = RigVMGraph->GetTypedOuter<URigVMLibraryNode>())
 			{
 				if(SpecifierName->Equals(TEXT("Private")))
 				{
-					RigVMBlueprint->MarkFunctionPublic(LibraryNode->GetFName(), false);
+					RigVMClientHost->MarkFunctionPublic(LibraryNode->GetFName(), false);
 				}
 				else
 				{
-					RigVMBlueprint->MarkFunctionPublic(LibraryNode->GetFName(), true);
+					RigVMClientHost->MarkFunctionPublic(LibraryNode->GetFName(), true);
 				}
 			}
 		}
 	}
 }
 
-TSharedRef<ITableRow> FRigVMGraphDetailCustomization::HandleGenerateRowAccessSpecifier( TSharedPtr<FRigVMStringWithTag> SpecifierName, const TSharedRef<STableViewBase>& OwnerTable )
+TSharedRef<ITableRow> FRigVMGraphDetailCustomizationImpl::HandleGenerateRowAccessSpecifier(TSharedPtr<FRigVMStringWithTag> SpecifierName, const TSharedRef<STableViewBase>& OwnerTable)
 {
 	return SNew(STableRow< TSharedPtr<FRigVMStringWithTag> >, OwnerTable)
         .Content()
@@ -1397,48 +1403,45 @@ TSharedRef<ITableRow> FRigVMGraphDetailCustomization::HandleGenerateRowAccessSpe
         ];
 }
 
-bool FRigVMGraphDetailCustomization::IsValidFunction() const
+bool FRigVMGraphDetailCustomizationImpl::IsValidFunction() const
 {
-	if (GraphPtr.IsValid() && RigVMBlueprintPtr.IsValid())
+	if (WeakModel.IsValid() && RigVMClientHost.IsValid())
 	{
-		URigVMBlueprint* Blueprint = RigVMBlueprintPtr.Get();
-		if (const URigVMGraph* Model = Blueprint->GetModel(GraphPtr.Get()))
+		if (const URigVMGraph* Model = WeakModel.Get())
 		{
 			if (const URigVMLibraryNode* LibraryNode = Cast<URigVMLibraryNode>(Model->GetOuter()))
 			{
-				return LibraryNode->GetFunctionHeader(Blueprint->GetRigVMGraphFunctionHost()).IsValid(); 
+				return LibraryNode->GetFunctionHeader(RigVMClientHost->GetRigVMGraphFunctionHost()).IsValid();
 			}
 		}
 	}
 	return false;
 }
 
-FRigVMVariant FRigVMGraphDetailCustomization::GetVariant() const
+FRigVMVariant FRigVMGraphDetailCustomizationImpl::GetVariant() const
 {
-	if (GraphPtr.IsValid() && RigVMBlueprintPtr.IsValid())
+	if (WeakModel.IsValid() && RigVMClientHost.IsValid())
 	{
-		URigVMBlueprint* Blueprint = RigVMBlueprintPtr.Get();
-		if (const URigVMGraph* Model = Blueprint->GetModel(GraphPtr.Get()))
+		if (const URigVMGraph* Model = WeakModel.Get())
 		{
 			if (const URigVMLibraryNode* LibraryNode = Cast<URigVMLibraryNode>(Model->GetOuter()))
 			{
-				return LibraryNode->GetFunctionHeader(Blueprint->GetRigVMGraphFunctionHost()).Variant; 
+				return LibraryNode->GetFunctionHeader(RigVMClientHost->GetRigVMGraphFunctionHost()).Variant;
 			}
 		}
 	}
 	return FRigVMVariant();
 }
 
-TArray<FRigVMVariantRef> FRigVMGraphDetailCustomization::GetVariantRefs() const
+TArray<FRigVMVariantRef> FRigVMGraphDetailCustomizationImpl::GetVariantRefs() const
 {
-	if (GraphPtr.IsValid() && RigVMBlueprintPtr.IsValid())
+	if (WeakModel.IsValid() && RigVMClientHost.IsValid())
 	{
-		URigVMBlueprint* Blueprint = RigVMBlueprintPtr.Get();
-		if (const URigVMGraph* Model = Blueprint->GetModel(GraphPtr.Get()))
+		if (const URigVMGraph* Model = WeakModel.Get())
 		{
 			if (const URigVMLibraryNode* LibraryNode = Cast<URigVMLibraryNode>(Model->GetOuter()))
 			{
-				const FRigVMGraphFunctionHeader Header = LibraryNode->GetFunctionHeader(Blueprint->GetRigVMGraphFunctionHost());
+				const FRigVMGraphFunctionHeader Header = LibraryNode->GetFunctionHeader(RigVMClientHost->GetRigVMGraphFunctionHost());
 				return Header.LibraryPointer.GetVariants(false);
 			}
 		}
@@ -1446,12 +1449,12 @@ TArray<FRigVMVariantRef> FRigVMGraphDetailCustomization::GetVariantRefs() const
 	return TArray<FRigVMVariantRef>();
 }
 
-void FRigVMGraphDetailCustomization::OnVariantChanged(const FRigVMVariant& InVariant)
+void FRigVMGraphDetailCustomizationImpl::OnVariantChanged(const FRigVMVariant& InVariant)
 {
 	// todo: update the function's variant info
 }
 
-void FRigVMGraphDetailCustomization::OnBrowseVariantRef(const FRigVMVariantRef& InVariantRef)
+void FRigVMGraphDetailCustomizationImpl::OnBrowseVariantRef(const FRigVMVariantRef& InVariantRef)
 {
 	const FRigVMGraphFunctionHeader Header = FRigVMGraphFunctionHeader::FindGraphFunctionHeader(InVariantRef.ObjectPath);
 	if(Header.IsValid())
@@ -1483,24 +1486,23 @@ void FRigVMGraphDetailCustomization::OnBrowseVariantRef(const FRigVMVariantRef&
 	}
 }
 
-TArray<FRigVMTag> FRigVMGraphDetailCustomization::OnGetAssignedTags() const
+TArray<FRigVMTag> FRigVMGraphDetailCustomizationImpl::OnGetAssignedTags() const
 {
 	return GetVariant().Tags;
 }
 
-void FRigVMGraphDetailCustomization::OnAddAssignedTag(const FName& InTagName)
+void FRigVMGraphDetailCustomizationImpl::OnAddAssignedTag(const FName& InTagName)
 {
-	if (GraphPtr.IsValid() && RigVMBlueprintPtr.IsValid())
+	if (WeakModel.IsValid() && RigVMClientHost.IsValid())
 	{
-		URigVMBlueprint* Blueprint = RigVMBlueprintPtr.Get();
-		if (const URigVMGraph* Model = Blueprint->GetModel(GraphPtr.Get()))
+		if (const URigVMGraph* Model = WeakModel.Get())
 		{
-			if (URigVMGraph* FunctionLibrary = Blueprint->GetLocalFunctionLibrary())
+			if (URigVMGraph* FunctionLibrary = RigVMClientHost->GetLocalFunctionLibrary())
 			{
 				if (const URigVMLibraryNode* LibraryNode = Cast<URigVMLibraryNode>(Model->GetOuter()))
 				{
 					const FString& FunctionName = LibraryNode->GetFunctionHeader().LibraryPointer.GetFunctionName();
-					URigVMController* FunctionLibraryController = Blueprint->GetOrCreateController(FunctionLibrary);
+					URigVMController* FunctionLibraryController = RigVMClientHost->GetOrCreateController(FunctionLibrary);
 					FunctionLibraryController->AddDefaultTagToFunctionVariant(*FunctionName, InTagName);
 				}
 			}
@@ -1508,19 +1510,18 @@ void FRigVMGraphDetailCustomization::OnAddAssignedTag(const FName& InTagName)
 	}
 }
 
-void FRigVMGraphDetailCustomization::OnRemoveAssignedTag(const FName& InTagName)
+void FRigVMGraphDetailCustomizationImpl::OnRemoveAssignedTag(const FName& InTagName)
 {
-	if (GraphPtr.IsValid() && RigVMBlueprintPtr.IsValid())
+	if (WeakModel.IsValid() && RigVMClientHost.IsValid())
 	{
-		URigVMBlueprint* Blueprint = RigVMBlueprintPtr.Get();
-		if (const URigVMGraph* Model = Blueprint->GetModel(GraphPtr.Get()))
+		if (const URigVMGraph* Model = WeakModel.Get())
 		{
-			if (URigVMGraph* FunctionLibrary = Blueprint->GetLocalFunctionLibrary())
+			if (URigVMGraph* FunctionLibrary = RigVMClientHost->GetLocalFunctionLibrary())
 			{
 				if (const URigVMLibraryNode* LibraryNode = Cast<URigVMLibraryNode>(Model->GetOuter()))
 				{
 					const FString& FunctionName = LibraryNode->GetFunctionHeader().LibraryPointer.GetFunctionName();
-					URigVMController* FunctionLibraryController = Blueprint->GetOrCreateController(FunctionLibrary);
+					URigVMController* FunctionLibraryController = RigVMClientHost->GetOrCreateController(FunctionLibrary);
 					FunctionLibraryController->RemoveTagFromFunctionVariant(*FunctionName, InTagName);
 				}
 			}
@@ -1528,14 +1529,13 @@ void FRigVMGraphDetailCustomization::OnRemoveAssignedTag(const FName& InTagName)
 	}
 }
 
-URigVMLibraryNode* FRigVMGraphDetailCustomization::GetLibraryNode() const
+URigVMLibraryNode* FRigVMGraphDetailCustomizationImpl::GetLibraryNode() const
 {
-	if (GraphPtr.IsValid() && RigVMBlueprintPtr.IsValid())
+	if (WeakModel.IsValid() && RigVMClientHost.IsValid())
 	{
-		const URigVMBlueprint* Blueprint = RigVMBlueprintPtr.Get();
-		if (const URigVMGraph* Model = Blueprint->GetModel(GraphPtr.Get()))
+		if (const URigVMGraph* Model = WeakModel.Get())
 		{
-			if (const URigVMGraph* FunctionLibrary = Blueprint->GetLocalFunctionLibrary())
+			if (const URigVMGraph* FunctionLibrary = RigVMClientHost->GetLocalFunctionLibrary())
 			{
 				if (URigVMLibraryNode* LibraryNode = Cast<URigVMLibraryNode>(Model->GetOuter()))
 				{
@@ -1550,12 +1550,12 @@ URigVMLibraryNode* FRigVMGraphDetailCustomization::GetLibraryNode() const
 	return nullptr;
 }
 
-URigVMNode* FRigVMGraphDetailCustomization::GetNodeForLayout() const
+URigVMNode* FRigVMGraphDetailCustomizationImpl::GetNodeForLayout() const
 {
 	return GetLibraryNode();
 }
 
-const FRigVMNodeLayout* FRigVMGraphDetailCustomization::GetNodeLayout() const
+const FRigVMNodeLayout* FRigVMGraphDetailCustomizationImpl::GetNodeLayout() const
 {
 	if(const URigVMNode* Node = GetNodeForLayout())
 	{
@@ -1565,7 +1565,7 @@ const FRigVMNodeLayout* FRigVMGraphDetailCustomization::GetNodeLayout() const
 	return nullptr;
 }
 
-TArray<FString> FRigVMGraphDetailCustomization::GetUncategorizedPins() const
+TArray<FString> FRigVMGraphDetailCustomizationImpl::GetUncategorizedPins() const
 {
 	if(const URigVMNode* Node = GetNodeForLayout())
 	{
@@ -1594,7 +1594,7 @@ TArray<FString> FRigVMGraphDetailCustomization::GetUncategorizedPins() const
 	return TArray<FString>();
 }
 
-TArray<FRigVMPinCategory> FRigVMGraphDetailCustomization::GetPinCategories() const
+TArray<FRigVMPinCategory> FRigVMGraphDetailCustomizationImpl::GetPinCategories() const
 {
 	if (const FRigVMNodeLayout* NodeLayout = GetNodeLayout())
 	{
@@ -1603,9 +1603,9 @@ TArray<FRigVMPinCategory> FRigVMGraphDetailCustomization::GetPinCategories() con
 	return TArray<FRigVMPinCategory>();
 }
 
-FString FRigVMGraphDetailCustomization::GetPinCategory(FString InPinPath) const
+FString FRigVMGraphDetailCustomizationImpl::GetPinCategory(FString InPinPath) const
 {
-	if(RigVMBlueprintPtr.IsValid())
+	if(RigVMClientHost.IsValid())
 	{
 		if(const URigVMNode* Node = GetNodeForLayout())
 		{
@@ -1618,9 +1618,9 @@ FString FRigVMGraphDetailCustomization::GetPinCategory(FString InPinPath) const
 	return FString();
 }
 
-int32 FRigVMGraphDetailCustomization::GetPinIndexInCategory(FString InPinPath) const
+int32 FRigVMGraphDetailCustomizationImpl::GetPinIndexInCategory(FString InPinPath) const
 {
-	if(RigVMBlueprintPtr.IsValid())
+	if(RigVMClientHost.IsValid())
 	{
 		if(const URigVMNode* Node = GetNodeForLayout())
 		{
@@ -1633,7 +1633,7 @@ int32 FRigVMGraphDetailCustomization::GetPinIndexInCategory(FString InPinPath) c
 	return INDEX_NONE;
 }
 
-FString FRigVMGraphDetailCustomization::GetPinLabel(FString InPinPath) const
+FString FRigVMGraphDetailCustomizationImpl::GetPinLabel(FString InPinPath) const
 {
 	if (const FRigVMNodeLayout* NodeLayout = GetNodeLayout())
 	{
@@ -1645,15 +1645,15 @@ FString FRigVMGraphDetailCustomization::GetPinLabel(FString InPinPath) const
 	return FString();
 }
 
-FLinearColor FRigVMGraphDetailCustomization::GetPinColor(FString InPinPath) const
+FLinearColor FRigVMGraphDetailCustomizationImpl::GetPinColor(FString InPinPath) const
 {
-	if(RigVMBlueprintPtr.IsValid())
+	if(RigVMClientHost.IsValid())
 	{
 		if(const URigVMNode* Node = GetNodeForLayout())
 		{
 			if(const URigVMPin* Pin = Node->FindPin(InPinPath))
 			{
-				if(const URigVMEdGraphSchema* Schema = Cast<URigVMEdGraphSchema>(RigVMBlueprintPtr->GetRigVMEdGraphSchemaClass()->GetDefaultObject()))
+				if(const URigVMEdGraphSchema* Schema = Cast<URigVMEdGraphSchema>(RigVMClientHost->GetRigVMEdGraphSchemaClass()->GetDefaultObject()))
 				{
 					const FEdGraphPinType PinType = RigVMTypeUtils::PinTypeFromCPPType(*Pin->GetCPPType(), Pin->GetCPPTypeObject());
 					return Schema->GetPinTypeColor(PinType);
@@ -1664,9 +1664,9 @@ FLinearColor FRigVMGraphDetailCustomization::GetPinColor(FString InPinPath) cons
 	return FLinearColor::White;
 }
 
-const FSlateBrush* FRigVMGraphDetailCustomization::GetPinIcon(FString InPinPath) const
+const FSlateBrush* FRigVMGraphDetailCustomizationImpl::GetPinIcon(FString InPinPath) const
 {
-	if(RigVMBlueprintPtr.IsValid())
+	if(RigVMClientHost.IsValid())
 	{
 		if(const URigVMNode* Node = GetNodeForLayout())
 		{
@@ -1680,13 +1680,13 @@ const FSlateBrush* FRigVMGraphDetailCustomization::GetPinIcon(FString InPinPath)
 	return nullptr;
 }
 
-void FRigVMGraphDetailCustomization::HandleCategoryAdded(FString InCategory)
+void FRigVMGraphDetailCustomizationImpl::HandleCategoryAdded(FString InCategory)
 {
-	if(RigVMBlueprintPtr.IsValid())
+	if(RigVMClientHost.IsValid())
 	{
 		if(const URigVMNode* Node = GetNodeForLayout())
 		{
-			if(URigVMController* Controller = RigVMBlueprintPtr->GetController(Node->GetGraph()))
+			if(URigVMController* Controller = RigVMClientHost->GetController(Node->GetGraph()))
 			{
 				Controller->AddEmptyPinCategory(Node->GetFName(), InCategory);
 				CachedNodeLayout.Reset();
@@ -1695,13 +1695,13 @@ void FRigVMGraphDetailCustomization::HandleCategoryAdded(FString InCategory)
 	}
 }
 
-void FRigVMGraphDetailCustomization::HandleCategoryRemoved(FString InCategory)
+void FRigVMGraphDetailCustomizationImpl::HandleCategoryRemoved(FString InCategory)
 {
-	if(RigVMBlueprintPtr.IsValid())
+	if(RigVMClientHost.IsValid())
 	{
 		if(const URigVMNode* Node = GetNodeForLayout())
 		{
-			if(URigVMController* Controller = RigVMBlueprintPtr->GetController(Node->GetGraph()))
+			if(URigVMController* Controller = RigVMClientHost->GetController(Node->GetGraph()))
 			{
 				Controller->RemovePinCategory(Node->GetFName(), InCategory);
 				CachedNodeLayout.Reset();
@@ -1710,13 +1710,13 @@ void FRigVMGraphDetailCustomization::HandleCategoryRemoved(FString InCategory)
 	}
 }
 
-void FRigVMGraphDetailCustomization::HandleCategoryRenamed(FString InOldCategory, FString InNewCategory)
+void FRigVMGraphDetailCustomizationImpl::HandleCategoryRenamed(FString InOldCategory, FString InNewCategory)
 {
-	if(RigVMBlueprintPtr.IsValid())
+	if(RigVMClientHost.IsValid())
 	{
 		if(const URigVMNode* Node = GetNodeForLayout())
 		{
-			if(URigVMController* Controller = RigVMBlueprintPtr->GetController(Node->GetGraph()))
+			if(URigVMController* Controller = RigVMClientHost->GetController(Node->GetGraph()))
 			{
 				Controller->RenamePinCategory(Node->GetFName(), InOldCategory, InNewCategory);
 				CachedNodeLayout.Reset();
@@ -1725,15 +1725,15 @@ void FRigVMGraphDetailCustomization::HandleCategoryRenamed(FString InOldCategory
 	}
 }
 
-void FRigVMGraphDetailCustomization::HandlePinCategoryChanged(FString InPinPath, FString InCategory)
+void FRigVMGraphDetailCustomizationImpl::HandlePinCategoryChanged(FString InPinPath, FString InCategory)
 {
-	if(RigVMBlueprintPtr.IsValid())
+	if(RigVMClientHost.IsValid())
 	{
 		if (const URigVMLibraryNode* LibraryNode = GetLibraryNode())
 		{
 			if(const URigVMPin* Pin = LibraryNode->FindPin(InPinPath))
 			{
-				if(URigVMController* Controller = RigVMBlueprintPtr->GetController(LibraryNode->GetGraph()))
+				if(URigVMController* Controller = RigVMClientHost->GetController(LibraryNode->GetGraph()))
 				{
 					Controller->SetPinCategory(Pin->GetPinPath(), InCategory);
 					CachedNodeLayout.Reset();
@@ -1743,15 +1743,15 @@ void FRigVMGraphDetailCustomization::HandlePinCategoryChanged(FString InPinPath,
 	}
 }
 
-void FRigVMGraphDetailCustomization::HandlePinLabelChanged(FString InPinPath, FString InNewLabel)
+void FRigVMGraphDetailCustomizationImpl::HandlePinLabelChanged(FString InPinPath, FString InNewLabel)
 {
-	if(RigVMBlueprintPtr.IsValid())
+	if(RigVMClientHost.IsValid())
 	{
 		if (const URigVMLibraryNode* LibraryNode = GetLibraryNode())
 		{
 			if(const URigVMPin* Pin = LibraryNode->FindPin(InPinPath))
 			{
-				if(URigVMController* Controller = RigVMBlueprintPtr->GetController(LibraryNode->GetGraph()))
+				if(URigVMController* Controller = RigVMClientHost->GetController(LibraryNode->GetGraph()))
 				{
 					Controller->SetPinDisplayName(Pin->GetPinPath(), InNewLabel);
 					CachedNodeLayout.Reset();
@@ -1761,15 +1761,15 @@ void FRigVMGraphDetailCustomization::HandlePinLabelChanged(FString InPinPath, FS
 	}
 }
 
-void FRigVMGraphDetailCustomization::HandlePinIndexInCategoryChanged(FString InPinPath, int32 InIndexInCategory)
+void FRigVMGraphDetailCustomizationImpl::HandlePinIndexInCategoryChanged(FString InPinPath, int32 InIndexInCategory)
 {
-	if(RigVMBlueprintPtr.IsValid())
+	if(RigVMClientHost.IsValid())
 	{
 		if (const URigVMLibraryNode* LibraryNode = GetLibraryNode())
 		{
 			if(const URigVMPin* Pin = LibraryNode->FindPin(InPinPath))
 			{
-				if(URigVMController* Controller = RigVMBlueprintPtr->GetController(LibraryNode->GetGraph()))
+				if(URigVMController* Controller = RigVMClientHost->GetController(LibraryNode->GetGraph()))
 				{
 					Controller->SetPinIndexInCategory(Pin->GetPinPath(), InIndexInCategory);
 					CachedNodeLayout.Reset();
@@ -1779,7 +1779,7 @@ void FRigVMGraphDetailCustomization::HandlePinIndexInCategoryChanged(FString InP
 	}
 }
 
-bool FRigVMGraphDetailCustomization::ValidateName(FString InNewName, FText& OutErrorMessage)
+bool FRigVMGraphDetailCustomizationImpl::ValidateName(FString InNewName, FText& OutErrorMessage)
 {
 	if(InNewName.IsEmpty())
 	{
@@ -1818,7 +1818,7 @@ bool FRigVMGraphDetailCustomization::ValidateName(FString InNewName, FText& OutE
 	return true;
 }
 
-bool FRigVMGraphDetailCustomization::HandleValidateCategoryName(FString InCategoryPath, FString InNewName, FText& OutErrorMessage)
+bool FRigVMGraphDetailCustomizationImpl::HandleValidateCategoryName(FString InCategoryPath, FString InNewName, FText& OutErrorMessage)
 {
 	if(!ValidateName(InNewName, OutErrorMessage))
 	{
@@ -1844,7 +1844,7 @@ bool FRigVMGraphDetailCustomization::HandleValidateCategoryName(FString InCatego
 	return true;
 }
 
-bool FRigVMGraphDetailCustomization::HandleValidatePinDisplayName(FString InPinPath, FString InNewName, FText& OutErrorMessage)
+bool FRigVMGraphDetailCustomizationImpl::HandleValidatePinDisplayName(FString InPinPath, FString InNewName, FText& OutErrorMessage)
 {
 	if(!ValidateName(InNewName, OutErrorMessage))
 	{
@@ -1885,7 +1885,7 @@ bool FRigVMGraphDetailCustomization::HandleValidatePinDisplayName(FString InPinP
 	return true;
 }
 
-uint32 FRigVMGraphDetailCustomization::GetNodeLayoutHash() const
+uint32 FRigVMGraphDetailCustomizationImpl::GetNodeLayoutHash() const
 {
 	uint32 Hash = 0;
 	if(const FRigVMNodeLayout* Layout = GetNodeLayout())
diff --git a/Engine/Plugins/Runtime/RigVM/Source/RigVMEditor/Public/Editor/RigVMGraphDetailCustomization.h b/Engine/Plugins/Runtime/RigVM/Source/RigVMEditor/Public/Editor/RigVMGraphDetailCustomization.h
--- a/Engine/Plugins/Runtime/RigVM/Source/RigVMEditor/Public/Editor/RigVMGraphDetailCustomization.h
+++ b/Engine/Plugins/Runtime/RigVM/Source/RigVMEditor/Public/Editor/RigVMGraphDetailCustomization.h
@@ -21,14 +21,15 @@
 #include "HAL/PlatformApplicationMisc.h"
 
 class IDetailLayoutBuilder;
+class FRigVMGraphDetailCustomizationImpl;
 
 class RIGVMEDITOR_API FRigVMFunctionArgumentGroupLayout : public IDetailCustomNodeBuilder, public TSharedFromThis<FRigVMFunctionArgumentGroupLayout>
 {
 public:
 	FRigVMFunctionArgumentGroupLayout(
-		URigVMGraph* InGraph, 
-		URigVMBlueprint* InBlueprint,
-		TWeakPtr<FRigVMEditor> InEditor,
+		const TWeakObjectPtr<URigVMGraph>& InGraph,
+		const TWeakInterfacePtr<IRigVMClientHost>& InRigVMClientHost,
+		const TWeakPtr<FRigVMEditor>& InEditor,
 		bool bInputs);
 	virtual ~FRigVMFunctionArgumentGroupLayout();
 
@@ -47,7 +48,7 @@ private:
 	void HandleModifiedEvent(ERigVMGraphNotifType InNotifType, URigVMGraph* InGraph, UObject* InSubject);
 
 	TWeakObjectPtr<URigVMGraph> GraphPtr;
-	TWeakObjectPtr<URigVMBlueprint> RigVMBlueprintPtr;
+	TWeakInterfacePtr<IRigVMClientHost> WeakRigVMClientHost;
 	TWeakPtr<FRigVMEditor> RigVMEditorPtr;
 	bool bIsInputGroup;
 	FSimpleDelegate OnRebuildChildren;
@@ -58,15 +59,15 @@ class RIGVMEDITOR_API FRigVMFunctionArgumentLayout : public IDetailCustomNodeBui
 public:
 
 	FRigVMFunctionArgumentLayout(
-		URigVMPin* InPin, 
-		URigVMGraph* InGraph, 
-		URigVMBlueprint* InBlueprint,
-		TWeakPtr<FRigVMEditor> InEditor)
+		const TWeakObjectPtr<URigVMPin>& InPin,
+		const TWeakObjectPtr<URigVMGraph>& InGraph, 
+		const TWeakInterfacePtr<IRigVMClientHost>& InRigVMClientHost,
+		const TWeakPtr<FRigVMEditor>& InEditor)
 		: PinPtr(InPin)
 		, GraphPtr(InGraph)
-		, RigVMBlueprintPtr(InBlueprint)
+		, WeakRigVMClientHost(InRigVMClientHost)
 		, RigVMEditorPtr(InEditor)
-		, NameValidator(InBlueprint, InGraph, InPin->GetFName())
+		, NameValidator(nullptr, InGraph.Get(), InPin->GetFName()) // TODO check if missing Blueprint affects name validation
 	{}
 
 private:
@@ -111,8 +112,8 @@ private:
 	/** The target graph that this argument is on */
 	TWeakObjectPtr<URigVMGraph> GraphPtr;
 
-	/** The blueprint we are editing */
-	TWeakObjectPtr<URigVMBlueprint> RigVMBlueprintPtr;
+	/** The asset host we are editing */
+	TWeakInterfacePtr<IRigVMClientHost> WeakRigVMClientHost;
 
 	/** The editor we are editing */
 	TWeakPtr<FRigVMEditor> RigVMEditorPtr;
@@ -128,8 +129,8 @@ class RIGVMEDITOR_API FRigVMFunctionArgumentDefaultNode : public IDetailCustomNo
 {
 public:
 	FRigVMFunctionArgumentDefaultNode(
-		URigVMGraph* InGraph,
-		URigVMBlueprint* InBlueprint
+		const TWeakObjectPtr<URigVMGraph>& InGraph,
+		const TWeakInterfacePtr<IRigVMClientHost>& InClientHost
 	);
 	virtual ~FRigVMFunctionArgumentDefaultNode();
 
@@ -150,30 +151,48 @@ private:
 
 	TWeakObjectPtr<URigVMGraph> GraphPtr;
 	TWeakObjectPtr<URigVMEdGraph> EdGraphOuterPtr;
-	TWeakObjectPtr<URigVMBlueprint> RigVMBlueprintPtr;
+	TWeakInterfacePtr<IRigVMClientHost> WeakRigVMClientHost;
 	FSimpleDelegate OnRebuildChildren;
 	TSharedPtr<SRigVMGraphNode> OwnedNodeWidget;
 	FDelegateHandle GraphChangedDelegateHandle;
 };
 
-
 /** Customization for editing rig vm graphs */
 class RIGVMEDITOR_API FRigVMGraphDetailCustomization : public IDetailCustomization
 {
 public:
-
 	/** Makes a new instance of this detail layout class for a specific detail view requesting it */
 	static TSharedPtr<IDetailCustomization> MakeInstance(TSharedPtr<IBlueprintEditor> InBlueprintEditor, const UClass* InExpectedBlueprintClass);
 
-	FRigVMGraphDetailCustomization(TSharedPtr<FRigVMEditor> RigVMigEditor, URigVMBlueprint* RigVMBlueprint)
-		: RigVMEditorPtr(RigVMigEditor)
-		, RigVMBlueprintPtr(RigVMBlueprint)
-		, bIsPickingColor(false)
-	{}
+	FRigVMGraphDetailCustomization(TSharedPtr<FRigVMEditor> RigVMigEditor, URigVMBlueprint* RigVMBlueprint);
 
 	// IDetailCustomization interface
 	virtual void CustomizeDetails(IDetailLayoutBuilder& DetailLayout) override;
 
+private:
+	/** The Blueprint editor we are embedded in */
+	TWeakPtr<FRigVMEditor> RigVMEditorPtr;
+
+	/** The graph we are editing */
+	TWeakObjectPtr<URigVMEdGraph> GraphPtr;
+
+	/** The blueprint we are editing */
+	TWeakObjectPtr<URigVMBlueprint> RigVMBlueprintPtr;
+
+	TSharedPtr<FRigVMGraphDetailCustomizationImpl> RigVMGraphDetailCustomizationImpl;
+};
+
+class RIGVMEDITOR_API FRigVMGraphDetailCustomizationImpl : public TSharedFromThis<FRigVMGraphDetailCustomizationImpl>
+{
+public:
+	void CustomizeDetails(IDetailLayoutBuilder& DetailLayout,
+		URigVMGraph* Model,
+		URigVMController* Controller,
+		IRigVMClientHost* InRigVMClientHost,
+		TWeakPtr<FRigVMEditor> InEditor);
+
+private:
+
 	bool IsAddNewInputOutputEnabled() const;
 	EVisibility GetAddNewInputOutputVisibility() const;
 	FReply OnAddNewInputClicked();
@@ -221,20 +240,23 @@ private:
 	void HandlePinCategoryChanged(FString InPinPath, FString InCategory);
 	void HandlePinLabelChanged(FString InPinPath, FString InNewLabel);
 	void HandlePinIndexInCategoryChanged(FString InPinPath, int32 InIndexInCategory);
-	static bool ValidateName(FString InNewName, FText& OutErrorMessage); 
+	static bool ValidateName(FString InNewName, FText& OutErrorMessage);
 	bool HandleValidateCategoryName(FString InCategoryPath, FString InNewName, FText& OutErrorMessage);
 	bool HandleValidatePinDisplayName(FString InPinPath, FString InNewName, FText& OutErrorMessage);
 
 	uint32 GetNodeLayoutHash() const;
 
-	/** The Blueprint editor we are embedded in */
-	TWeakPtr<FRigVMEditor> RigVMEditorPtr;
-
 	/** The graph we are editing */
-	TWeakObjectPtr<URigVMEdGraph> GraphPtr;
+	TWeakObjectPtr<URigVMGraph> WeakModel;
 
-	/** The blueprint we are editing */
-	TWeakObjectPtr<URigVMBlueprint> RigVMBlueprintPtr;
+	/** The graph controller we are editing */
+	TWeakObjectPtr<URigVMController> WeakController;
+
+	/** The editor we are embedded in */
+	TWeakPtr<FRigVMEditor> RigVMEditorPtr;
+
+	/** The asset host we are editing */
+	TWeakInterfacePtr<IRigVMClientHost> RigVMClientHost;
 
 	/** The color block widget */
 	TSharedPtr<SColorBlock> ColorBlock;
