RigVM Editor: Catch nullptr crash

#rnx

diff --git a/Engine/Plugins/Runtime/RigVM/Source/RigVMEditor/Private/Editor/RigVMEditor.cpp b/Engine/Plugins/Runtime/RigVM/Source/RigVMEditor/Private/Editor/RigVMEditor.cpp
--- a/Engine/Plugins/Runtime/RigVM/Source/RigVMEditor/Private/Editor/RigVMEditor.cpp
+++ b/Engine/Plugins/Runtime/RigVM/Source/RigVMEditor/Private/Editor/RigVMEditor.cpp
@@ -1747,8 +1747,11 @@ void FRigVMEditor::HandleModifiedEvent(ERigVMGraphNotifType InNotifType, URigVMG
 										}
 
 										// we are ok with not reacting to errors here
-										FRigVMPinDefaultValueImportErrorContext ErrorPipe;										
-										Property->ImportText_Direct(*DefaultValue, PropertyStorage, nullptr, PPF_None, &ErrorPipe);
+										if(Property && PropertyStorage)
+										{
+											FRigVMPinDefaultValueImportErrorContext ErrorPipe;										
+											Property->ImportText_Direct(*DefaultValue, PropertyStorage, nullptr, PPF_None, &ErrorPipe);
+										}
 									}
 								}
 							}
