RigVM: Fix nullptr crash

#rnx

diff --git a/Engine/Plugins/Runtime/RigVM/Source/RigVMDeveloper/Private/RigVMModel/RigVMController.cpp b/Engine/Plugins/Runtime/RigVM/Source/RigVMDeveloper/Private/RigVMModel/RigVMController.cpp
--- a/Engine/Plugins/Runtime/RigVM/Source/RigVMDeveloper/Private/RigVMModel/RigVMController.cpp
+++ b/Engine/Plugins/Runtime/RigVM/Source/RigVMDeveloper/Private/RigVMModel/RigVMController.cpp
@@ -5674,7 +5674,10 @@ TArray<URigVMNode*> URigVMController::ExpandLibraryNode(URigVMLibraryNode* InNod
 
 			if (!Link->GetSourcePin()->IsRootPin())
 			{
-				LibraryPinsToTurnIntoConstant.AddUnique(InNode->FindPin(Link->GetSourcePin()->GetRootPin()->GetName()));
+				if(URigVMPin* RootPinOnOtherNode = InNode->FindPin(Link->GetSourcePin()->GetRootPin()->GetName()))
+				{
+					LibraryPinsToTurnIntoConstant.AddUnique(RootPinOnOtherNode);
+				}
 			}
 
 			FString NodeName, PinPath;
@@ -5715,7 +5718,10 @@ TArray<URigVMNode*> URigVMController::ExpandLibraryNode(URigVMLibraryNode* InNod
 
 			if (!Link->GetTargetPin()->IsRootPin())
 			{
-				LibraryPinsToTurnIntoConstant.AddUnique(InNode->FindPin(Link->GetTargetPin()->GetRootPin()->GetName()));
+				if(URigVMPin* RootPinOnOtherNode = InNode->FindPin(Link->GetTargetPin()->GetRootPin()->GetName()))
+				{
+					LibraryPinsToTurnIntoConstant.AddUnique(RootPinOnOtherNode);
+				}
 			}
 
 			FString NodeName, PinPath;
