Fix static analysis at RigVM graph details customization

#rnx

diff --git a/Engine/Plugins/Runtime/RigVM/Source/RigVMEditor/Private/Editor/RigVMGraphDetailCustomization.cpp b/Engine/Plugins/Runtime/RigVM/Source/RigVMEditor/Private/Editor/RigVMGraphDetailCustomization.cpp
--- a/Engine/Plugins/Runtime/RigVM/Source/RigVMEditor/Private/Editor/RigVMGraphDetailCustomization.cpp
+++ b/Engine/Plugins/Runtime/RigVM/Source/RigVMEditor/Private/Editor/RigVMGraphDetailCustomization.cpp
@@ -796,7 +796,7 @@ void FRigVMGraphDetailCustomizationImpl::CustomizeDetails(IDetailLayoutBuilder&
 
 	URigVMGraph* Model = WeakModel.Get();
 
-	if (Model->IsTopLevelGraph())
+	if (Model != nullptr && Model->IsTopLevelGraph())
 	{
 		IDetailCategoryBuilder& Category = DetailLayout.EditCategory("Graph", LOCTEXT("FunctionDetailsGraph", "Graph"));
 		Category.AddCustomRow(FText::GetEmpty())
