RigVMEditor: Fix jumping to hyperlinks between assets

#rnx

diff --git a/Engine/Plugins/Runtime/RigVM/Source/RigVMEditor/Private/Editor/RigVMEditor.cpp b/Engine/Plugins/Runtime/RigVM/Source/RigVMEditor/Private/Editor/RigVMEditor.cpp
--- a/Engine/Plugins/Runtime/RigVM/Source/RigVMEditor/Private/Editor/RigVMEditor.cpp
+++ b/Engine/Plugins/Runtime/RigVM/Source/RigVMEditor/Private/Editor/RigVMEditor.cpp
@@ -2531,25 +2531,41 @@ void FRigVMEditor::HandleJumpToHyperlink(const UObject* InSubject)
 
 	if (GraphToJumpTo && NodeToJumpTo)
 	{
-		if(URigVMEdGraph* EdGraph = Cast<URigVMEdGraph>(RigBlueprint->GetEdGraph(NodeToJumpTo->GetGraph())))
+		if(URigVMBlueprint* OtherBlueprint = NodeToJumpTo->GetTypedOuter<URigVMBlueprint>())
 		{
-			if(URigVMEdGraphNode* EdGraphNode = Cast<URigVMEdGraphNode>(EdGraph->FindNodeForModelNodeName(NodeToJumpTo->GetFName())))
+			if(OtherBlueprint != RigBlueprint)
 			{
-				if(PinToJumpTo)
+				if (GEditor->GetEditorSubsystem<UAssetEditorSubsystem>()->OpenEditorForAsset(OtherBlueprint))
 				{
-					if(const UEdGraphPin* EdGraphPin = EdGraphNode->FindPin(PinToJumpTo->GetSegmentPath(true)))
+					FRigVMEditor* OtherEditor = static_cast<FRigVMEditor*>(GEditor->GetEditorSubsystem<UAssetEditorSubsystem>()->FindEditorForAsset(OtherBlueprint, /*bFocusIfOpen =*/true));
+					if(OtherEditor)
 					{
-						JumpToPin(EdGraphPin);
+						OtherEditor->HandleJumpToHyperlink(NodeToJumpTo);
 						return;
 					}
 				}
+			}
+					
+			if(URigVMEdGraph* EdGraph = Cast<URigVMEdGraph>(OtherBlueprint->GetEdGraph(NodeToJumpTo->GetGraph())))
+			{
+				if(URigVMEdGraphNode* EdGraphNode = Cast<URigVMEdGraphNode>(EdGraph->FindNodeForModelNodeName(NodeToJumpTo->GetFName())))
+				{
+					if(PinToJumpTo)
+					{
+						if(const UEdGraphPin* EdGraphPin = EdGraphNode->FindPin(PinToJumpTo->GetSegmentPath(true)))
+						{
+							JumpToPin(EdGraphPin);
+							return;
+						}
+					}
+					
+					JumpToNode(EdGraphNode);
+					SetDetailObjects({EdGraphNode});
+					return;
+				}
 				
-				JumpToNode(EdGraphNode);
-				SetDetailObjects({EdGraphNode});
-				return;
+				JumpToHyperlink(EdGraph);
 			}
-			
-			JumpToHyperlink(EdGraph);
 		}
 	}
 }
