RigVM: Use weak pointers for pin cache

#rnx

diff --git a/Engine/Plugins/Runtime/RigVM/Source/RigVMDeveloper/Private/EdGraph/RigVMEdGraphNode.cpp b/Engine/Plugins/Runtime/RigVM/Source/RigVMDeveloper/Private/EdGraph/RigVMEdGraphNode.cpp
--- a/Engine/Plugins/Runtime/RigVM/Source/RigVMDeveloper/Private/EdGraph/RigVMEdGraphNode.cpp
+++ b/Engine/Plugins/Runtime/RigVM/Source/RigVMDeveloper/Private/EdGraph/RigVMEdGraphNode.cpp
@@ -534,10 +534,13 @@ URigVMPin* URigVMEdGraphNode::FindModelPinFromGraphPin(const UEdGraphPin* InGrap
 	
 	for(const auto& Pair : CachedPins)
 	{
-		if(Pair.Value.InputPin == InGraphPin ||
-			Pair.Value.OutputPin == InGraphPin)
+		if(Pair.Key.IsValid())
 		{
-			return Pair.Key;
+			if(Pair.Value.InputPin == InGraphPin ||
+				Pair.Value.OutputPin == InGraphPin)
+			{
+				return Pair.Key.Get();
+			}
 		}
 	}
 
@@ -923,10 +926,13 @@ void URigVMEdGraphNode::RemoveGraphSubPins(UEdGraphPin* InParentPin, const TArra
 		}
 		for(const auto& Pair : CachedPins)
 		{
-			if(Pair.Value.InputPin == SubPin ||
-				Pair.Value.OutputPin == SubPin)
+			if(Pair.Key.IsValid())
 			{
-				ModelSubPins.Add(Pair.Key);
+				if(Pair.Value.InputPin == SubPin ||
+					Pair.Value.OutputPin == SubPin)
+				{
+					ModelSubPins.Add(Pair.Key.Get());
+				}
 			}
 		}
 		for(const auto& Pair : CachedCategoryPins)
@@ -1059,10 +1065,15 @@ bool URigVMEdGraphNode::ModelPinsChanged(bool bForce)
 	};
 	auto RemoveObsoletePins = [this]()
 	{
+		CachedPins = CachedPins.FilterByPredicate([](const TPair<TWeakObjectPtr<URigVMPin>, FPinPair>& InPair) -> bool
+		{
+			return InPair.Key.IsValid();
+		});
+		
 		TArray<URigVMPin*> PinsToRemove; 
 		for(const auto& Pair : CachedPins)
 		{
-			URigVMPin* ModelPin = Pair.Key;
+			URigVMPin* ModelPin = Pair.Key.Get();
 			TArray<URigVMPin*>& List = PinListForPin(ModelPin);
 			bool bRemove = true;
 			if (URigVMPin** FoundPin = List.FindByKey(ModelPin->GetRootPin()))
@@ -1298,9 +1309,12 @@ bool URigVMEdGraphNode::ModelPinsChanged(bool bForce)
 			PinLabelsChanged += SyncPinLabel(ModelPin, Pair.Value);
 		}
 	}
-	for(const TPair<URigVMPin*,FPinPair>& Pair : CachedPins)
+	for(const TPair<TWeakObjectPtr<URigVMPin>,FPinPair>& Pair : CachedPins)
 	{
-		PinLabelsChanged += SyncPinLabel(Pair.Key, Pair.Value);
+		if(Pair.Key.IsValid())
+		{
+			PinLabelsChanged += SyncPinLabel(Pair.Key.Get(), Pair.Value);
+		}
 	}
 
 	const bool bResult = (PinsAdded > 0) || (PinLabelsChanged > 0) || (PinsRemoved > 0) || (PinsReordered > 0) || (PinsReparented > 0); 
diff --git a/Engine/Plugins/Runtime/RigVM/Source/RigVMDeveloper/Public/EdGraph/RigVMEdGraphNode.h b/Engine/Plugins/Runtime/RigVM/Source/RigVMDeveloper/Public/EdGraph/RigVMEdGraphNode.h
--- a/Engine/Plugins/Runtime/RigVM/Source/RigVMDeveloper/Public/EdGraph/RigVMEdGraphNode.h
+++ b/Engine/Plugins/Runtime/RigVM/Source/RigVMDeveloper/Public/EdGraph/RigVMEdGraphNode.h
@@ -291,7 +291,7 @@ private:
 	TArray<TSharedPtr<FRigVMExternalVariable>> ExternalVariables;
 	TArray<UEdGraphPin*> LastEdGraphPins;
 	
-	TMap<URigVMPin*, FPinPair> CachedPins;
+	mutable TMap<TWeakObjectPtr<URigVMPin>, FPinPair> CachedPins;
 	TMap<FString, FPinPair> CachedCategoryPins;
 
 	FNodeTitleDirtied NodeTitleDirtied;
