RigVM: Add option to toggle the tag display on nodes

#rnx

diff --git a/Engine/Plugins/Runtime/RigVM/Source/RigVMDeveloper/Public/RigVMBlueprint.h b/Engine/Plugins/Runtime/RigVM/Source/RigVMDeveloper/Public/RigVMBlueprint.h
--- a/Engine/Plugins/Runtime/RigVM/Source/RigVMDeveloper/Public/RigVMBlueprint.h
+++ b/Engine/Plugins/Runtime/RigVM/Source/RigVMDeveloper/Public/RigVMBlueprint.h
@@ -55,6 +55,15 @@ struct RIGVMDEVELOPER_API FRigVMPythonSettings
 	}
 };
 
+UENUM(BlueprintType)
+enum class ERigVMTagDisplayMode : uint8
+{
+	None = 0,
+	All = 0x001,
+	DeprecationOnly = 0x002,
+	Last = DeprecationOnly UMETA(Hidden), 
+};
+
 USTRUCT()
 struct RIGVMDEVELOPER_API FRigVMEdGraphDisplaySettings
 {
@@ -74,6 +83,7 @@ struct RIGVMDEVELOPER_API FRigVMEdGraphDisplaySettings
 		, LastMaxMicroSeconds(1.0)
 		, MinDurationColor(FLinearColor::Green)
 		, MaxDurationColor(FLinearColor::Red)
+		, TagDisplayMode(ERigVMTagDisplayMode::All)
 	{
 	}
 
@@ -141,6 +151,10 @@ struct RIGVMDEVELOPER_API FRigVMEdGraphDisplaySettings
 	UPROPERTY(EditAnywhere, Category = "Graph Display Settings")
 	FLinearColor MaxDurationColor;
 
+	// The color of the slowest instruction / node
+	UPROPERTY(EditAnywhere, Category = "Graph Display Settings")
+	ERigVMTagDisplayMode TagDisplayMode;
+
 	void SetTotalMicroSeconds(double InTotalMicroSeconds);
 	void SetLastMinMicroSeconds(double InMinMicroSeconds);
 	void SetLastMaxMicroSeconds(double InMaxMicroSeconds);
diff --git a/Engine/Plugins/Runtime/RigVM/Source/RigVMEditor/Private/RigVMEditorModule.cpp b/Engine/Plugins/Runtime/RigVM/Source/RigVMEditor/Private/RigVMEditorModule.cpp
--- a/Engine/Plugins/Runtime/RigVM/Source/RigVMEditor/Private/RigVMEditorModule.cpp
+++ b/Engine/Plugins/Runtime/RigVM/Source/RigVMEditor/Private/RigVMEditorModule.cpp
@@ -640,6 +640,7 @@ void FRigVMEditorModule::GetNodeContextMenuActions(IRigVMClientHost* RigVMClient
 	GetNodeOrganizationContextMenuActions(RigVMClientHost, EdGraphNode, ModelNode, Menu);
 	GetNodeVariantContextMenuActions(RigVMClientHost, EdGraphNode, ModelNode, Menu);
 	GetNodeTestContextMenuActions(RigVMClientHost, EdGraphNode, ModelNode, Menu);
+	GetNodeDisplayContextMenuActions(RigVMClientHost, EdGraphNode, ModelNode, Menu);
 }
 
 void FRigVMEditorModule::GetNodeWorkflowContextMenuActions(IRigVMClientHost* RigVMClientHost, const URigVMEdGraphNode* EdGraphNode, URigVMNode* ModelNode, UToolMenu* Menu) const
@@ -1296,6 +1297,68 @@ void FRigVMEditorModule::GetNodeTestContextMenuActions(IRigVMClientHost* RigVMCl
 	);
 }
 
+void FRigVMEditorModule::GetNodeDisplayContextMenuActions(IRigVMClientHost* RigVMClientHost, const URigVMEdGraphNode* EdGraphNode, URigVMNode* ModelNode, UToolMenu* Menu) const
+{
+	URigVMBlueprint* Blueprint = Cast<URigVMBlueprint>(RigVMClientHost);
+	if(Blueprint == nullptr)
+	{
+		return;
+	}
+	
+	FToolMenuSection& DisplaySection = Menu->AddSection("RigVMEditorContextMenuDisplay", LOCTEXT("DisplayHeader", "Display"));
+
+	DisplaySection.AddMenuEntry(
+		TEXT("EnableProfiler"),
+		LOCTEXT("EnableProfiler", "Enable Profiler"),
+		LOCTEXT("EnableProfiler_Tooltip", "Enables the heat map profiler"),
+		FSlateIcon(),
+		FUIAction(FExecuteAction::CreateLambda([Blueprint]()
+		{
+			FScopedTransaction Transaction(LOCTEXT("ToggleProfiler", "Toggle Profiler"));
+			Blueprint->Modify();
+			Blueprint->VMRuntimeSettings.bEnableProfiling = !Blueprint->VMRuntimeSettings.bEnableProfiling;
+			Blueprint->PropagateRuntimeSettingsFromBPToInstances();
+			Blueprint->RequestAutoVMRecompilation();
+		}),
+		FCanExecuteAction(),
+		FGetActionCheckState::CreateLambda([Blueprint]() -> ECheckBoxState
+		{
+			return Blueprint->VMRuntimeSettings.bEnableProfiling ? ECheckBoxState::Checked : ECheckBoxState::Unchecked;
+		})),
+		EUserInterfaceActionType::ToggleButton
+	);
+
+	DisplaySection.AddMenuEntry(
+		TEXT("ShowAllTags"),
+		LOCTEXT("ShowAllTags", "Show all tags"),
+		LOCTEXT("ShowAllTags_Tooltip", "Shows all of the tags on nodes. If turned off this will show the deprecation tags only."),
+		FSlateIcon(),
+		FUIAction(FExecuteAction::CreateLambda([Blueprint]()
+		{
+			FScopedTransaction Transaction(LOCTEXT("ToggleTagDisplayMode", "Toggle Tag Display Mode"));
+			Blueprint->Modify();
+
+			if(Blueprint->RigGraphDisplaySettings.TagDisplayMode == ERigVMTagDisplayMode::All)
+			{
+				Blueprint->RigGraphDisplaySettings.TagDisplayMode = ERigVMTagDisplayMode::DeprecationOnly;
+			}
+			else
+			{
+				Blueprint->RigGraphDisplaySettings.TagDisplayMode = ERigVMTagDisplayMode::All;
+			}
+
+			// this causes all nodes to refresh
+			Blueprint->PropagateRuntimeSettingsFromBPToInstances();
+		}),
+		FCanExecuteAction(),
+		FGetActionCheckState::CreateLambda([Blueprint]() -> ECheckBoxState
+		{
+			return Blueprint->RigGraphDisplaySettings.TagDisplayMode == ERigVMTagDisplayMode::All ? ECheckBoxState::Checked : ECheckBoxState::Unchecked;
+		})),
+		EUserInterfaceActionType::ToggleButton
+	);
+}
+
 void FRigVMEditorModule::GetPinContextMenuActions(IRigVMClientHost* RigVMClientHost, const UEdGraphPin* EdGraphPin, URigVMPin* ModelPin, UToolMenu* Menu) const
 {
 	GetPinWorkflowContextMenuActions(RigVMClientHost, EdGraphPin, ModelPin, Menu);
diff --git a/Engine/Plugins/Runtime/RigVM/Source/RigVMEditor/Private/Widgets/SRigVMGraphNode.cpp b/Engine/Plugins/Runtime/RigVM/Source/RigVMEditor/Private/Widgets/SRigVMGraphNode.cpp
--- a/Engine/Plugins/Runtime/RigVM/Source/RigVMEditor/Private/Widgets/SRigVMGraphNode.cpp
+++ b/Engine/Plugins/Runtime/RigVM/Source/RigVMEditor/Private/Widgets/SRigVMGraphNode.cpp
@@ -2086,12 +2086,27 @@ void SRigVMGraphNode::UpdatePinTreeView()
 	];
 
 	// add the tags this node potentially has
-	if(ModelNode.IsValid())
+	ERigVMTagDisplayMode TagDisplayMode = ERigVMTagDisplayMode::All;
+	if(Blueprint.IsValid())
+	{
+		TagDisplayMode = Blueprint->RigGraphDisplaySettings.TagDisplayMode;
+	}
+	if(ModelNode.IsValid() && TagDisplayMode != ERigVMTagDisplayMode::None)
 	{
 		if(URigVMFunctionReferenceNode* FunctionRefNode = Cast<URigVMFunctionReferenceNode>(ModelNode.Get()))
 		{
 			const FRigVMVariant Variant = FunctionRefNode->GetReferencedFunctionHeader().Variant;
-			if(!Variant.Tags.IsEmpty())
+
+			const bool bContainsAnyTags = Variant.Tags.ContainsByPredicate([TagDisplayMode](const FRigVMTag& InTag) -> bool
+			{
+				if(TagDisplayMode == ERigVMTagDisplayMode::DeprecationOnly)
+				{
+					return InTag.bMarksSubjectAsInvalid;
+				}
+				return true;
+			});
+			
+			if(bContainsAnyTags)
 			{
 				TWeakObjectPtr<URigVMFunctionReferenceNode> WeakFunctionRefNode = FunctionRefNode;
 				
@@ -2117,11 +2132,18 @@ void SRigVMGraphNode::UpdatePinTreeView()
 					.Orientation(EOrientation::Orient_Horizontal)
 					.CanAddTags(false)
 					.EnableContextMenu(false)
-					.OnGetTags_Lambda([WeakFunctionRefNode]() -> TArray<FRigVMTag>
+					.OnGetTags_Lambda([WeakFunctionRefNode, TagDisplayMode]() -> TArray<FRigVMTag>
 					{
 						if(WeakFunctionRefNode.IsValid())
 						{
 							const FRigVMVariant Variant = WeakFunctionRefNode->GetReferencedFunctionHeader().Variant;
+							if(TagDisplayMode == ERigVMTagDisplayMode::DeprecationOnly)
+							{
+								return Variant.Tags.FilterByPredicate([TagDisplayMode](const FRigVMTag& InTag) -> bool
+								{
+									return InTag.bMarksSubjectAsInvalid;
+								});
+							}
 							return Variant.Tags;
 						}
 						return {};
diff --git a/Engine/Plugins/Runtime/RigVM/Source/RigVMEditor/Public/RigVMEditorModule.h b/Engine/Plugins/Runtime/RigVM/Source/RigVMEditor/Public/RigVMEditorModule.h
--- a/Engine/Plugins/Runtime/RigVM/Source/RigVMEditor/Public/RigVMEditorModule.h
+++ b/Engine/Plugins/Runtime/RigVM/Source/RigVMEditor/Public/RigVMEditorModule.h
@@ -94,6 +94,7 @@ protected:
 	virtual void GetNodeOrganizationContextMenuActions(IRigVMClientHost* RigVMClientHost, const URigVMEdGraphNode* EdGraphNode, URigVMNode* ModelNode, UToolMenu* Menu) const;
 	virtual void GetNodeVariantContextMenuActions(IRigVMClientHost* RigVMClientHost, const URigVMEdGraphNode* EdGraphNode, URigVMNode* ModelNode, UToolMenu* Menu) const;
 	virtual void GetNodeTestContextMenuActions(IRigVMClientHost* RigVMClientHost, const URigVMEdGraphNode* EdGraphNode, URigVMNode* ModelNode, UToolMenu* Menu) const;
+	virtual void GetNodeDisplayContextMenuActions(IRigVMClientHost* RigVMClientHost, const URigVMEdGraphNode* EdGraphNode, URigVMNode* ModelNode, UToolMenu* Menu) const;
 	virtual void GetPinWorkflowContextMenuActions(IRigVMClientHost* RigVMClientHost, const UEdGraphPin* EdGraphPin, URigVMPin* ModelPin, UToolMenu* Menu) const;
 	virtual void GetPinDebugContextMenuActions(IRigVMClientHost* RigVMClientHost, const UEdGraphPin* EdGraphPin, URigVMPin* ModelPin, UToolMenu* Menu) const;
 	virtual void GetPinArrayContextMenuActions(IRigVMClientHost* RigVMClientHost, const UEdGraphPin* EdGraphPin, URigVMPin* ModelPin, UToolMenu* Menu) const;
