Moved dll storage from type to individual functions/staticvars in order to resolve linker warnings on some modular builds using clang platforms

diff --git a/Engine/Plugins/Experimental/AnimNext/Source/AnimNext/Public/TraitInterfaces/IHierarchy.h b/Engine/Plugins/Experimental/AnimNext/Source/AnimNext/Public/TraitInterfaces/IHierarchy.h
--- a/Engine/Plugins/Experimental/AnimNext/Source/AnimNext/Public/TraitInterfaces/IHierarchy.h
+++ b/Engine/Plugins/Experimental/AnimNext/Source/AnimNext/Public/TraitInterfaces/IHierarchy.h
@@ -18,34 +18,34 @@ namespace UE::AnimNext
 	 * 
 	 * This interface exposes hierarchy traversal information to navigate the graph.
 	 */
-	struct ANIMNEXT_API IHierarchy : ITraitInterface
+	struct IHierarchy : ITraitInterface
 	{
 		DECLARE_ANIM_TRAIT_INTERFACE(IHierarchy, 0x846d8a37)
 
 		// Returns the number of children of the trait implementation (not the whole stack)
 		// Includes inactive children
-		virtual uint32 GetNumChildren(const FExecutionContext& Context, const TTraitBinding<IHierarchy>& Binding) const;
+		ANIMNEXT_API virtual uint32 GetNumChildren(const FExecutionContext& Context, const TTraitBinding<IHierarchy>& Binding) const;
 
 		// Appends weak handles to any children we wish to traverse on the trait implementation (not the whole stack).
 		// Traits are responsible for allocating and releasing child instance data.
 		// Empty handles can be appended.
-		virtual void GetChildren(const FExecutionContext& Context, const TTraitBinding<IHierarchy>& Binding, FChildrenArray& Children) const;
+		ANIMNEXT_API virtual void GetChildren(const FExecutionContext& Context, const TTraitBinding<IHierarchy>& Binding, FChildrenArray& Children) const;
 
 		// Queries the trait stack and calls GetChildren for each trait, appending the result.
-		static void GetStackChildren(const FExecutionContext& Context, const FTraitStackBinding& Binding, FChildrenArray& Children);
+		ANIMNEXT_API static void GetStackChildren(const FExecutionContext& Context, const FTraitStackBinding& Binding, FChildrenArray& Children);
 
 		// Queries the trait stack of the specified binding and calls GetChildren for each trait, appending the result.
-		static void GetStackChildren(const FExecutionContext& Context, const FTraitBinding& Binding, FChildrenArray& Children);
+		ANIMNEXT_API static void GetStackChildren(const FExecutionContext& Context, const FTraitBinding& Binding, FChildrenArray& Children);
 
 		// Queries the trait stack and calls GetNumChildren for each trait, accumulating the result.
-		static uint32 GetNumStackChildren(const FExecutionContext& Context, const FTraitStackBinding& Binding);
+		ANIMNEXT_API static uint32 GetNumStackChildren(const FExecutionContext& Context, const FTraitStackBinding& Binding);
 
 		// Queries the trait stack of the specified binding and calls GetNumChildren for each trait, accumulating the result.
-		static uint32 GetNumStackChildren(const FExecutionContext& Context, const FTraitBinding& Binding);
+		ANIMNEXT_API static uint32 GetNumStackChildren(const FExecutionContext& Context, const FTraitBinding& Binding);
 
 #if WITH_EDITOR
-		virtual const FText& GetDisplayName() const override;
-		virtual const FText& GetDisplayShortName() const override;
+		ANIMNEXT_API virtual const FText& GetDisplayName() const override;
+		ANIMNEXT_API virtual const FText& GetDisplayShortName() const override;
 #endif // WITH_EDITOR
 	};
 
diff --git a/Engine/Plugins/Runtime/Harmonix/Source/HarmonixMetasound/Public/HarmonixMetasound/DataTypes/MidiStream.h b/Engine/Plugins/Runtime/Harmonix/Source/HarmonixMetasound/Public/HarmonixMetasound/DataTypes/MidiStream.h
--- a/Engine/Plugins/Runtime/Harmonix/Source/HarmonixMetasound/Public/HarmonixMetasound/DataTypes/MidiStream.h
+++ b/Engine/Plugins/Runtime/Harmonix/Source/HarmonixMetasound/Public/HarmonixMetasound/DataTypes/MidiStream.h
@@ -35,22 +35,22 @@ namespace HarmonixMetasound
 		FMidiVoiceId VoiceId;
 	};
 
-	class HARMONIXMETASOUND_API FMidiStream
+	class FMidiStream
 	{
 	public:
-		void SetClock(const FMidiClock& InClock);
-		void ResetClock();
-		TSharedPtr<const FMidiClock, ESPMode::NotThreadSafe> GetClock() const;
+		HARMONIXMETASOUND_API void SetClock(const FMidiClock& InClock);
+		HARMONIXMETASOUND_API void ResetClock();
+		HARMONIXMETASOUND_API TSharedPtr<const FMidiClock, ESPMode::NotThreadSafe> GetClock() const;
 
-		void PrepareBlock();
+		HARMONIXMETASOUND_API void PrepareBlock();
 
-		void AddMidiEvent(const FMidiStreamEvent& Event);
-		void InsertMidiEvent(const FMidiStreamEvent& Event);
-		void AddNoteOffEventOrCancelPendingNoteOn(const FMidiStreamEvent& Event);
-		void InsertNoteOffEventOrCancelPendingNoteOn(const FMidiStreamEvent& Event);
+		HARMONIXMETASOUND_API void AddMidiEvent(const FMidiStreamEvent& Event);
+		HARMONIXMETASOUND_API void InsertMidiEvent(const FMidiStreamEvent& Event);
+		HARMONIXMETASOUND_API void AddNoteOffEventOrCancelPendingNoteOn(const FMidiStreamEvent& Event);
+		HARMONIXMETASOUND_API void InsertNoteOffEventOrCancelPendingNoteOn(const FMidiStreamEvent& Event);
 
 		void SetMidiFile(const FMidiFileProxyPtr& MidiFile) { MidiFileSourceOfEvents = MidiFile; }
-		const FString* GetMidiTrackText(int32 TrackNumber, int32 TextIndex) const;
+		HARMONIXMETASOUND_API const FString* GetMidiTrackText(int32 TrackNumber, int32 TextIndex) const;
 
 
 		const TArray<FMidiStreamEvent>& GetEventsInBlock() const
@@ -67,26 +67,26 @@ namespace HarmonixMetasound
 		using FEventTransformer = TFunction<FMidiStreamEvent(const FMidiStreamEvent&)>;
 		inline static const FEventTransformer NoOpTransformer = [](const FMidiStreamEvent& Event) { return Event; };
 		
-		static void Copy(
+		static HARMONIXMETASOUND_API void Copy(
 			const FMidiStream& From,
 			FMidiStream& To,
 			const FEventFilter& Filter = NoOpFilter,
 			const FEventTransformer& Transformer = NoOpTransformer);
 
-		static void Merge(
+		static HARMONIXMETASOUND_API void Merge(
 			const FMidiStream& From,
 			FMidiStream& To,
 			const FEventFilter& Filter = NoOpFilter,
 			const FEventTransformer& Transformer = NoOpTransformer);
 		
-		static void Merge(
+		static HARMONIXMETASOUND_API void Merge(
 			const FMidiStream& FromA,
 			const FMidiStream& FromB,
 			FMidiStream& To,
 			const FEventFilter& Filter = NoOpFilter,
 			const FEventTransformer& Transformer = NoOpTransformer);
 
-		bool NoteIsActive(const FMidiStreamEvent& Event) const;
+		HARMONIXMETASOUND_API bool NoteIsActive(const FMidiStreamEvent& Event) const;
 
 	private:
 		FMidiFileProxyPtr MidiFileSourceOfEvents;
diff --git a/Engine/Plugins/Runtime/MassGameplay/Source/MassSimulation/Public/MassSimulationSubsystem.h b/Engine/Plugins/Runtime/MassGameplay/Source/MassSimulation/Public/MassSimulationSubsystem.h
--- a/Engine/Plugins/Runtime/MassGameplay/Source/MassSimulation/Public/MassSimulationSubsystem.h
+++ b/Engine/Plugins/Runtime/MassGameplay/Source/MassSimulation/Public/MassSimulationSubsystem.h
@@ -14,52 +14,52 @@ class IConsoleVariable;
 
 DECLARE_LOG_CATEGORY_EXTERN(LogMassSim, Log, All);
 
-UCLASS(config = Game, defaultconfig)
-class MASSSIMULATION_API UMassSimulationSubsystem : public UMassSubsystemBase
+UCLASS(config = Game, defaultconfig, MinimalAPI)
+class UMassSimulationSubsystem : public UMassSubsystemBase
 {
 	GENERATED_BODY()
 public:
 	DECLARE_MULTICAST_DELEGATE_OneParam(FOnSimulationStarted, UWorld* /*World*/);
 	
-	UMassSimulationSubsystem(const FObjectInitializer& ObjectInitializer = FObjectInitializer::Get());
+	MASSSIMULATION_API UMassSimulationSubsystem(const FObjectInitializer& ObjectInitializer = FObjectInitializer::Get());
 
 	const FMassProcessingPhaseManager& GetPhaseManager() const { return PhaseManager; }
 
-	FMassProcessingPhase::FOnPhaseEvent& GetOnProcessingPhaseStarted(const EMassProcessingPhase Phase);
-	FMassProcessingPhase::FOnPhaseEvent& GetOnProcessingPhaseFinished(const EMassProcessingPhase Phase);
+	MASSSIMULATION_API FMassProcessingPhase::FOnPhaseEvent& GetOnProcessingPhaseStarted(const EMassProcessingPhase Phase);
+	MASSSIMULATION_API FMassProcessingPhase::FOnPhaseEvent& GetOnProcessingPhaseFinished(const EMassProcessingPhase Phase);
 	static FOnSimulationStarted& GetOnSimulationStarted() { return OnSimulationStarted; }
 
-	void RegisterDynamicProcessor(UMassProcessor& Processor);
-	void UnregisterDynamicProcessor(UMassProcessor& Processor);
+	MASSSIMULATION_API void RegisterDynamicProcessor(UMassProcessor& Processor);
+	MASSSIMULATION_API void UnregisterDynamicProcessor(UMassProcessor& Processor);
 
 	bool IsSimulationStarted() const { return bSimulationStarted; }
 
 	/** @return whether hosted EntityManager is currently, actively being used for processing purposes. Equivalent to calling FMassEntityManager.IsProcessing() */
-	bool IsDuringMassProcessing() const;
+	MASSSIMULATION_API bool IsDuringMassProcessing() const;
 
 	/** Starts/stops simulation ticking for all worlds, based on new `mass.SimulationTickingEnabled` cvar value */
-	static void HandleSimulationTickingEnabledCVarChange(IConsoleVariable*);
+	MASSSIMULATION_API static void HandleSimulationTickingEnabledCVarChange(IConsoleVariable*);
 
 protected:
 	// UWorldSubsystem BEGIN
-	virtual void Initialize(FSubsystemCollectionBase& Collection) override;
-	virtual void PostInitialize() override;
-	virtual void OnWorldBeginPlay(UWorld& InWorld) override;
-	virtual void Deinitialize() override;
+	MASSSIMULATION_API virtual void Initialize(FSubsystemCollectionBase& Collection) override;
+	MASSSIMULATION_API virtual void PostInitialize() override;
+	MASSSIMULATION_API virtual void OnWorldBeginPlay(UWorld& InWorld) override;
+	MASSSIMULATION_API virtual void Deinitialize() override;
 	// UWorldSubsystem END
-	virtual void BeginDestroy() override;
+	MASSSIMULATION_API virtual void BeginDestroy() override;
 	
-	void RebuildTickPipeline();
+	MASSSIMULATION_API void RebuildTickPipeline();
 
-	void StartSimulation(UWorld& InWorld);
-	void StopSimulation();
+	MASSSIMULATION_API void StartSimulation(UWorld& InWorld);
+	MASSSIMULATION_API void StopSimulation();
 
-	void OnProcessingPhaseStarted(const float DeltaSeconds, const EMassProcessingPhase Phase) const;
+	MASSSIMULATION_API void OnProcessingPhaseStarted(const float DeltaSeconds, const EMassProcessingPhase Phase) const;
 
 #if WITH_EDITOR
-	void OnPieBegin(const bool bIsSimulation);
-	void OnPieEnded(const bool bIsSimulation);
-	void OnMassEntitySettingsChange(const FPropertyChangedEvent& PropertyChangedEvent);
+	MASSSIMULATION_API void OnPieBegin(const bool bIsSimulation);
+	MASSSIMULATION_API void OnPieEnded(const bool bIsSimulation);
+	MASSSIMULATION_API void OnMassEntitySettingsChange(const FPropertyChangedEvent& PropertyChangedEvent);
 #endif // WITH_EDITOR
 
 protected:
diff --git a/Engine/Plugins/Runtime/Metasound/Source/MetasoundGraphCore/Private/MetasoundNodeInterface.cpp b/Engine/Plugins/Runtime/Metasound/Source/MetasoundGraphCore/Private/MetasoundNodeInterface.cpp
--- a/Engine/Plugins/Runtime/Metasound/Source/MetasoundGraphCore/Private/MetasoundNodeInterface.cpp
+++ b/Engine/Plugins/Runtime/Metasound/Source/MetasoundGraphCore/Private/MetasoundNodeInterface.cpp
@@ -99,6 +99,12 @@ namespace Metasound
 		return *this != InvalidNodeClassName;
 	}
 
+	const FNodeClassMetadata& FNodeClassMetadata::GetEmpty()
+	{
+		static const FNodeClassMetadata EmptyInfo;
+		return EmptyInfo;
+	}
+
 	bool operator==(const FOutputDataSource& InLeft, const FOutputDataSource& InRight)
 	{
 		return (InLeft.Node == InRight.Node) && (InLeft.Vertex == InRight.Vertex);
diff --git a/Engine/Plugins/Runtime/Metasound/Source/MetasoundGraphCore/Public/MetasoundNodeInterface.h b/Engine/Plugins/Runtime/Metasound/Source/MetasoundGraphCore/Public/MetasoundNodeInterface.h
--- a/Engine/Plugins/Runtime/Metasound/Source/MetasoundGraphCore/Public/MetasoundNodeInterface.h
+++ b/Engine/Plugins/Runtime/Metasound/Source/MetasoundGraphCore/Public/MetasoundNodeInterface.h
@@ -99,7 +99,7 @@ namespace Metasound
 	};
 
 	/** Provides metadata for a given node. */
-	struct METASOUNDGRAPHCORE_API FNodeClassMetadata
+	struct FNodeClassMetadata
 	{
 		/** Name of class. Used for registration and lookup. */
 		FNodeClassName ClassName;
@@ -138,11 +138,7 @@ namespace Metasound
 		bool bDeprecated = false;
 
 		/** Returns an empty FNodeClassMetadata object. */
-		static const FNodeClassMetadata& GetEmpty()
-		{
-			static const FNodeClassMetadata EmptyInfo;
-			return EmptyInfo;
-		}
+		static METASOUNDGRAPHCORE_API const FNodeClassMetadata& GetEmpty();
 	};
 
 	/** INodeBase
diff --git a/Engine/Plugins/Runtime/RigVM/Source/RigVM/Public/RigVMCore/RigVMAssetUserData.h b/Engine/Plugins/Runtime/RigVM/Source/RigVM/Public/RigVMCore/RigVMAssetUserData.h
--- a/Engine/Plugins/Runtime/RigVM/Source/RigVM/Public/RigVMCore/RigVMAssetUserData.h
+++ b/Engine/Plugins/Runtime/RigVM/Source/RigVM/Public/RigVMCore/RigVMAssetUserData.h
@@ -10,15 +10,15 @@
 /**
 * User data that can be attached to assets to provide namespaced data access.
 */
-UCLASS(Abstract)
-class RIGVM_API UNameSpacedUserData : public UAssetUserData
+UCLASS(Abstract, MinimalAPI)
+class UNameSpacedUserData : public UAssetUserData
 {
 	GENERATED_BODY()
 
 public:
 
 	/** A helper struct to represent a single user data */
-	struct RIGVM_API FUserData
+	struct FUserData
 	{
 	public:
 
@@ -60,12 +60,12 @@ public:
 		const FString& GetPath() const { return Path; }
 
 		// returns the name of the last segment of the path
-		FString GetName() const;
+		RIGVM_API FString GetName() const;
 
 #if WITH_EDITOR
 
 		// returns the display name of the last segment of the path
-		FString GetDisplayName() const;
+		RIGVM_API FString GetDisplayName() const;
 
 #endif
 
@@ -82,11 +82,11 @@ public:
 			return CPPType;
 		}
 
-		bool IsArray() const;
+		RIGVM_API bool IsArray() const;
 
-		bool IsArrayElement() const;
+		RIGVM_API bool IsArrayElement() const;
 
-		bool IsUObject() const;
+		RIGVM_API bool IsUObject() const;
 
 		// returns the property of this user data or nullptr
 		const FProperty* GetProperty() const { return Property; }
@@ -129,7 +129,7 @@ public:
 		InvalidateCache();
 	}
 
-	virtual void Serialize(FArchive& Ar) override;
+	RIGVM_API virtual void Serialize(FArchive& Ar) override;
 
 	/** The namespace to use when looking up values inside of the user data. */
 	UPROPERTY(EditAnywhere, Category = General)
@@ -139,10 +139,10 @@ public:
 	bool ContainsUserData(const FString& InPath) const { return GetUserData(InPath) != nullptr; }
 
 	/** Returns a user data path given its string path */
-	virtual const FUserData* GetUserData(const FString& InPath, FString* OutErrorMessage = nullptr) const;
+	RIGVM_API virtual const FUserData* GetUserData(const FString& InPath, FString* OutErrorMessage = nullptr) const;
 
 	/** Retrieves the user data paths given a (optional) parent path */
-	virtual const TArray<const FUserData*>& GetUserDataArray(const FString& InParentPath = FString(), FString* OutErrorMessage = nullptr) const;
+	RIGVM_API virtual const TArray<const FUserData*>& GetUserDataArray(const FString& InParentPath = FString(), FString* OutErrorMessage = nullptr) const;
 
 	/** Retrieves the user data path given the path string */
 	const FUserData* GetParentUserData(const FString& InPath, FString* OutErrorMessage = nullptr) const
@@ -157,17 +157,17 @@ public:
 
 protected:
 
-	const FUserData* GetUserDataWithinStruct(const UStruct* InStruct, const uint8* InMemory, const FString& InPath, const FString& InPropertyName, FString* OutErrorMessage = nullptr) const;
-	const TArray<const FUserData*>& GetUserDataArrayWithinStruct(UStruct* InStruct, const uint8* InMemory, const FString& InPath, FString* OutErrorMessage = nullptr) const;
-	const TArray<const FUserData*>& GetUserDataArrayWithinArray(const FArrayProperty* InArrayProperty, const uint8* InMemory, const FString& InPath, FString* OutErrorMessage = nullptr) const;
-	static const FProperty* FindPropertyByName(const UStruct* InStruct, const FName& InName);
+	RIGVM_API const FUserData* GetUserDataWithinStruct(const UStruct* InStruct, const uint8* InMemory, const FString& InPath, const FString& InPropertyName, FString* OutErrorMessage = nullptr) const;
+	RIGVM_API const TArray<const FUserData*>& GetUserDataArrayWithinStruct(UStruct* InStruct, const uint8* InMemory, const FString& InPath, FString* OutErrorMessage = nullptr) const;
+	RIGVM_API const TArray<const FUserData*>& GetUserDataArrayWithinArray(const FArrayProperty* InArrayProperty, const uint8* InMemory, const FString& InPath, FString* OutErrorMessage = nullptr) const;
+	RIGVM_API static const FProperty* FindPropertyByName(const UStruct* InStruct, const FName& InName);
 
-	void InvalidateCache() const;
-	const FUserData* StoreCacheForUserData(const FUserData& InUserData) const; 
-	const TArray<const FUserData*>& StoreCacheForUserDataArray(const FString& InPath, const TArray<const FUserData*>& InUserDataArray) const;
-	virtual bool IsPropertySupported(const FProperty* InProperty, const FString& InPath, bool bCheckPropertyFlags, FString* OutErrorMessage = nullptr) const;
+	RIGVM_API void InvalidateCache() const;
+	RIGVM_API const FUserData* StoreCacheForUserData(const FUserData& InUserData) const; 
+	RIGVM_API const TArray<const FUserData*>& StoreCacheForUserDataArray(const FString& InPath, const TArray<const FUserData*>& InUserDataArray) const;
+	RIGVM_API virtual bool IsPropertySupported(const FProperty* InProperty, const FString& InPath, bool bCheckPropertyFlags, FString* OutErrorMessage = nullptr) const;
 	
-	static const TArray<const FUserData*> EmptyUserDatas;
+	RIGVM_API static const TArray<const FUserData*> EmptyUserDatas;
 
 	mutable TMap<FString, FUserData*> CachedUserData;
 	mutable TMap<FString, TArray<const FUserData*>> CachedUserDataArray;
@@ -183,8 +183,8 @@ protected:
 /**
 * Namespaced user data which provides access to a linked data asset
 */
-UCLASS(BlueprintType)
-class RIGVM_API UDataAssetLink : public UNameSpacedUserData
+UCLASS(BlueprintType, MinimalAPI)
+class UDataAssetLink : public UNameSpacedUserData
 {
 	GENERATED_BODY()
 
@@ -198,16 +198,16 @@ public:
 	TSoftObjectPtr<UDataAsset> GetDataAsset() const { return DataAsset; }
 
 	UFUNCTION(BlueprintSetter)
-	void SetDataAsset(TSoftObjectPtr<UDataAsset> InDataAsset);
+	RIGVM_API void SetDataAsset(TSoftObjectPtr<UDataAsset> InDataAsset);
 
-	virtual const FUserData* GetUserData(const FString& InPath, FString* OutErrorMessage = nullptr) const override;
-	virtual const TArray<const FUserData*>& GetUserDataArray(const FString& InParentPath = FString(), FString* OutErrorMessage = nullptr) const override;
+	RIGVM_API virtual const FUserData* GetUserData(const FString& InPath, FString* OutErrorMessage = nullptr) const override;
+	RIGVM_API virtual const TArray<const FUserData*>& GetUserDataArray(const FString& InParentPath = FString(), FString* OutErrorMessage = nullptr) const override;
 
-	virtual void Serialize(FArchive& Ar) override;
+	RIGVM_API virtual void Serialize(FArchive& Ar) override;
 #if WITH_EDITOR
 	virtual bool IsPostLoadThreadSafe() const override { return false; }
-	virtual void PostLoad() override;
-	virtual void PostEditChangeProperty(FPropertyChangedEvent& PropertyChangedEvent) override;
+	RIGVM_API virtual void PostLoad() override;
+	RIGVM_API virtual void PostEditChangeProperty(FPropertyChangedEvent& PropertyChangedEvent) override;
 #endif
 
 protected:
diff --git a/Engine/Plugins/Runtime/RigVM/Source/RigVM/Public/RigVMCore/RigVMStruct.h b/Engine/Plugins/Runtime/RigVM/Source/RigVM/Public/RigVMCore/RigVMStruct.h
--- a/Engine/Plugins/Runtime/RigVM/Source/RigVM/Public/RigVMCore/RigVMStruct.h
+++ b/Engine/Plugins/Runtime/RigVM/Source/RigVM/Public/RigVMCore/RigVMStruct.h
@@ -50,7 +50,7 @@ enum class ERigVMNodeCreatedReason : uint8
 /**
  * A context struct passed to FRigVMStruct::OnUnitNodeCreated
  */
-struct RIGVM_API FRigVMUnitNodeCreatedContext
+struct FRigVMUnitNodeCreatedContext
 {
 public:
 
@@ -75,7 +75,7 @@ public:
 	};
 
 	/** Returns true if this context is valid to use */
-	bool IsValid() const;
+	RIGVM_API bool IsValid() const;
 
 	/** Get the reason why this node was created */
 	ERigVMNodeCreatedReason GetReason() const { return Reason; }
@@ -88,19 +88,19 @@ public:
 	FName GetNodeName() const { return NodeName; }
 
 	/** Returns all currently existing external variables */
-	TArray<FRigVMExternalVariable> GetExternalVariables() const;
+	RIGVM_API TArray<FRigVMExternalVariable> GetExternalVariables() const;
 
 	/** Creates a new variable within the host of this VM */
-	FName AddExternalVariable(const FRigVMExternalVariable& InVariableToCreate, FString InDefaultValue = FString());
+	RIGVM_API FName AddExternalVariable(const FRigVMExternalVariable& InVariableToCreate, FString InDefaultValue = FString());
 
 	/** Binds a pin to an external variable on the created node */
-	bool BindPinToExternalVariable(FString InPinPath, FString InVariablePath);
+	RIGVM_API bool BindPinToExternalVariable(FString InPinPath, FString InVariablePath);
 
 	/** Returns a variable given a name (or a non-valid variable if not found) */
-	FRigVMExternalVariable FindVariable(FName InVariableName) const;
+	RIGVM_API FRigVMExternalVariable FindVariable(FName InVariableName) const;
 
 	/** Returns the name of the first variable given a(or NAME_None if not found) */
-	FName FindFirstVariableOfType(FName InCPPTypeName) const;
+	RIGVM_API FName FindFirstVariableOfType(FName InCPPTypeName) const;
 
 	/** Returns the name of the first variable given a type (or NAME_None if not found) */
 	template <
@@ -167,7 +167,7 @@ private:
  * The base class for all RigVM enabled structs.
  */
 USTRUCT()
-struct RIGVM_API FRigVMStruct
+struct FRigVMStruct
 {
 	GENERATED_BODY()
 
@@ -187,30 +187,30 @@ public:
 	virtual void Execute() {}
 
 	// control flow related
-	bool IsForLoop() const;
-	bool IsControlFlowNode() const; 
+	RIGVM_API bool IsForLoop() const;
+	RIGVM_API bool IsControlFlowNode() const; 
 	virtual int32 GetNumSlices() const { return 1; }
-	const TArray<FName>& GetControlFlowBlocks() const;
+	RIGVM_API const TArray<FName>& GetControlFlowBlocks() const;
 	virtual const bool IsControlFlowBlockSliced(const FName& InBlockName) const { return false; }
 
 	// node creation
 	virtual void OnUnitNodeCreated(FRigVMUnitNodeCreatedContext& InContext) const {}
 
 	// user workflow
-	TArray<FRigVMUserWorkflow> GetWorkflows(ERigVMUserWorkflowType InType, const UObject* InSubject) const; 
+	RIGVM_API TArray<FRigVMUserWorkflow> GetWorkflows(ERigVMUserWorkflowType InType, const UObject* InSubject) const; 
 
 #if WITH_EDITOR
-	static bool ValidateStruct(UScriptStruct* InStruct, FString* OutErrorMessage);
-	static bool CheckPinType(UScriptStruct* InStruct, const FName& PinName, const FString& ExpectedType, FString* OutErrorMessage = nullptr);
-	static bool CheckPinDirection(UScriptStruct* InStruct, const FName& PinName, const FName& InDirectionMetaName);
-	static ERigVMPinDirection GetPinDirectionFromProperty(FProperty* InProperty);
-	static bool CheckPinExists(UScriptStruct* InStruct, const FName& PinName, const FString& ExpectedType = FString(), FString* OutErrorMessage = nullptr);
-	static bool CheckMetadata(UScriptStruct* InStruct, const FName& PinName, const FName& InMetadataKey, FString* OutErrorMessage = nullptr);
-	static bool CheckFunctionExists(UScriptStruct* InStruct, const FName& FunctionName, FString* OutErrorMessage = nullptr);
-	virtual bool ShouldCreatePinForProperty(const FProperty* InProperty) const; 
+	RIGVM_API static bool ValidateStruct(UScriptStruct* InStruct, FString* OutErrorMessage);
+	RIGVM_API static bool CheckPinType(UScriptStruct* InStruct, const FName& PinName, const FString& ExpectedType, FString* OutErrorMessage = nullptr);
+	RIGVM_API static bool CheckPinDirection(UScriptStruct* InStruct, const FName& PinName, const FName& InDirectionMetaName);
+	RIGVM_API static ERigVMPinDirection GetPinDirectionFromProperty(FProperty* InProperty);
+	RIGVM_API static bool CheckPinExists(UScriptStruct* InStruct, const FName& PinName, const FString& ExpectedType = FString(), FString* OutErrorMessage = nullptr);
+	RIGVM_API static bool CheckMetadata(UScriptStruct* InStruct, const FName& PinName, const FName& InMetadataKey, FString* OutErrorMessage = nullptr);
+	RIGVM_API static bool CheckFunctionExists(UScriptStruct* InStruct, const FName& FunctionName, FString* OutErrorMessage = nullptr);
+	RIGVM_API virtual bool ShouldCreatePinForProperty(const FProperty* InProperty) const; 
 #endif
-	static FString ExportToFullyQualifiedText(const FProperty* InMemberProperty, const uint8* InMemberMemoryPtr, bool bUseQuotes = true);
-	static FString ExportToFullyQualifiedText(const UScriptStruct* InStruct, const uint8* InStructMemoryPtr, bool bUseQuotes = true);
+	RIGVM_API static FString ExportToFullyQualifiedText(const FProperty* InMemberProperty, const uint8* InMemberMemoryPtr, bool bUseQuotes = true);
+	RIGVM_API static FString ExportToFullyQualifiedText(const UScriptStruct* InStruct, const uint8* InStructMemoryPtr, bool bUseQuotes = true);
 
 	template <
 		typename T,
@@ -230,9 +230,9 @@ public:
 		return ExportToFullyQualifiedText(T::StaticStruct(), (const uint8*)&InStructValue);
 	}
 
-	FString ExportToFullyQualifiedText(const UScriptStruct* InScriptStruct, const FName& InPropertyName, const uint8* InStructMemoryPointer = nullptr, bool bUseQuotes = true) const;
+	RIGVM_API FString ExportToFullyQualifiedText(const UScriptStruct* InScriptStruct, const FName& InPropertyName, const uint8* InStructMemoryPointer = nullptr, bool bUseQuotes = true) const;
 	
-	virtual FName GetNextAggregateName(const FName& InLastAggregatePinName) const;
+	RIGVM_API virtual FName GetNextAggregateName(const FName& InLastAggregatePinName) const;
 	virtual FRigVMStructUpgradeInfo GetUpgradeInfo() const { return FRigVMStructUpgradeInfo(); }
 
 	static inline const FLazyName DeprecatedMetaName = FLazyName(TEXT("Deprecated"));
@@ -280,14 +280,14 @@ public:
 
 protected:
 
-	static float GetRatioFromIndex(int32 InIndex, int32 InCount);
-	TMap<FName, FString> GetDefaultValues(UScriptStruct* InScriptStruct) const;
-	bool ApplyUpgradeInfo(const FRigVMStructUpgradeInfo& InUpgradeInfo);
+	RIGVM_API static float GetRatioFromIndex(int32 InIndex, int32 InCount);
+	RIGVM_API TMap<FName, FString> GetDefaultValues(UScriptStruct* InScriptStruct) const;
+	RIGVM_API bool ApplyUpgradeInfo(const FRigVMStructUpgradeInfo& InUpgradeInfo);
 	virtual TArray<FRigVMUserWorkflow> GetSupportedWorkflows(const UObject* InSubject) const { return TArray<FRigVMUserWorkflow>(); }
-	virtual const TArray<FName>& GetControlFlowBlocks_Impl() const;
+	RIGVM_API virtual const TArray<FName>& GetControlFlowBlocks_Impl() const;
 
 #if WITH_EDITOR
-	static void ValidateControlFlowBlocks(const TArray<FName>& InBlocks);
+	RIGVM_API static void ValidateControlFlowBlocks(const TArray<FName>& InBlocks);
 #endif
 
 	friend struct FRigVMStructUpgradeInfo;
@@ -300,7 +300,7 @@ protected:
  * The base mutable class for all RigVM enabled structs.
  */
 USTRUCT()
-struct RIGVM_API FRigVMStructMutable : public FRigVMStruct
+struct FRigVMStructMutable : public FRigVMStruct
 {
 	GENERATED_BODY()
 
diff --git a/Engine/Source/Runtime/CoreUObject/Public/StructUtils/StructTypeBitSet.h b/Engine/Source/Runtime/CoreUObject/Public/StructUtils/StructTypeBitSet.h
--- a/Engine/Source/Runtime/CoreUObject/Public/StructUtils/StructTypeBitSet.h
+++ b/Engine/Source/Runtime/CoreUObject/Public/StructUtils/StructTypeBitSet.h
@@ -686,7 +686,6 @@ private:
 		using FBaseStructType = BaseType; \
 		EXPORTED_API static FStructTracker StructTracker; \
 	}; \
-	template struct EXPORTED_API TStructTypeBitSet<BaseType, ContainerTypeName##StructTrackerWrapper, BaseUStructType>; \
 	using ContainerTypeName = TStructTypeBitSet<BaseType, ContainerTypeName##StructTrackerWrapper, BaseUStructType>
 
 #define DECLARE_STRUCTTYPEBITSET_EXPORTED(EXPORTED_API, ContainerTypeName, BaseStructType) _DECLARE_TYPEBITSET_IMPL(EXPORTED_API, ContainerTypeName, BaseStructType, UScriptStruct)
diff --git a/Engine/Source/Runtime/Experimental/ChaosVisualDebugger/Public/DataWrappers/ChaosVDCollisionDataWrappers.h b/Engine/Source/Runtime/Experimental/ChaosVisualDebugger/Public/DataWrappers/ChaosVDCollisionDataWrappers.h
--- a/Engine/Source/Runtime/Experimental/ChaosVisualDebugger/Public/DataWrappers/ChaosVDCollisionDataWrappers.h
+++ b/Engine/Source/Runtime/Experimental/ChaosVisualDebugger/Public/DataWrappers/ChaosVDCollisionDataWrappers.h
@@ -53,7 +53,7 @@ enum class EChaosVDContactPointType : int8
 };
 
 USTRUCT()
-struct CHAOSVDRUNTIME_API FChaosVDContactPoint
+struct FChaosVDContactPoint
 {
 	GENERATED_BODY()
 
@@ -77,7 +77,7 @@ struct CHAOSVDRUNTIME_API FChaosVDContactPoint
 	UPROPERTY(VisibleAnywhere, Category=Contact)
 	EChaosVDContactPointType ContactType = EChaosVDContactPointType::Unknown;
 
-	bool Serialize(FArchive& Ar);
+	CHAOSVDRUNTIME_API bool Serialize(FArchive& Ar);
 };
 
 CVD_IMPLEMENT_SERIALIZER(FChaosVDContactPoint)
@@ -96,7 +96,7 @@ enum class EChaosVDManifoldPointFlags : uint8
 ENUM_CLASS_FLAGS(EChaosVDManifoldPointFlags)
 
 USTRUCT()
-struct CHAOSVDRUNTIME_API FChaosVDManifoldPoint
+struct FChaosVDManifoldPoint
 {
 	GENERATED_BODY()
 
@@ -135,7 +135,7 @@ struct CHAOSVDRUNTIME_API FChaosVDManifoldPoint
 
 	bool bIsSelectedInEditor = false;
 
-	bool Serialize(FArchive& Ar);
+	CHAOSVDRUNTIME_API bool Serialize(FArchive& Ar);
 };
 
 CVD_IMPLEMENT_SERIALIZER(FChaosVDManifoldPoint)
@@ -206,7 +206,7 @@ enum class EChaosVDConstraintFlags : uint16
 };
 
 USTRUCT()
-struct CHAOSVDRUNTIME_API FChaosVDConstraint
+struct FChaosVDConstraint
 {
 	GENERATED_BODY()
 
@@ -304,7 +304,7 @@ struct CHAOSVDRUNTIME_API FChaosVDConstraint
 	UPROPERTY()
 	int32 SolverID = INDEX_NONE;
 
-	bool Serialize(FArchive& Ar);
+	CHAOSVDRUNTIME_API bool Serialize(FArchive& Ar);
 };
 
 CVD_IMPLEMENT_SERIALIZER(FChaosVDConstraint)
@@ -339,7 +339,7 @@ enum class EChaosVDMidPhaseType : int8
 };
 
 USTRUCT()
-struct CHAOSVDRUNTIME_API FChaosVDParticlePairMidPhase
+struct FChaosVDParticlePairMidPhase
 {
 	GENERATED_BODY()
 
@@ -373,7 +373,7 @@ struct CHAOSVDRUNTIME_API FChaosVDParticlePairMidPhase
 	UPROPERTY(VisibleAnywhere, Category=Constraints)
 	TArray<FChaosVDConstraint> Constraints;
 
-	bool Serialize(FArchive& Ar);
+	CHAOSVDRUNTIME_API bool Serialize(FArchive& Ar);
 };
 
 CVD_IMPLEMENT_SERIALIZER(FChaosVDParticlePairMidPhase)
@@ -394,7 +394,7 @@ enum class EChaosVDCollisionTraceFlag
 };
 
 USTRUCT()
-struct CHAOSVDRUNTIME_API FChaosVDCollisionFilterData
+struct FChaosVDCollisionFilterData
 {
 	GENERATED_BODY()
 
@@ -407,7 +407,7 @@ struct CHAOSVDRUNTIME_API FChaosVDCollisionFilterData
 	UPROPERTY(VisibleAnywhere, Category=CollisionData)
 	uint32 Word3 = 0;
 
-	bool Serialize(FArchive& Ar);
+	CHAOSVDRUNTIME_API bool Serialize(FArchive& Ar);
 
 	bool operator==(const FChaosVDCollisionFilterData& Other) const = default;
 };
@@ -430,7 +430,7 @@ enum class EChaosVDCollisionShapeDataFlags : uint8
 };
 
 USTRUCT()
-struct CHAOSVDRUNTIME_API FChaosVDShapeCollisionData
+struct FChaosVDShapeCollisionData
 {
 	GENERATED_BODY()
 
@@ -456,9 +456,9 @@ struct CHAOSVDRUNTIME_API FChaosVDShapeCollisionData
 	UPROPERTY(VisibleAnywhere, Category="CVD Data")
 	bool bIsValid = false;
 
-	bool Serialize(FArchive& Ar);
+	CHAOSVDRUNTIME_API bool Serialize(FArchive& Ar);
 
-	bool operator==(const FChaosVDShapeCollisionData& Other) const;
+	CHAOSVDRUNTIME_API bool operator==(const FChaosVDShapeCollisionData& Other) const;
 };
 
 template <>
@@ -473,7 +473,7 @@ CVD_IMPLEMENT_SERIALIZER(FChaosVDShapeCollisionData)
  * based on already serialized flags
  */
 USTRUCT()
-struct CHAOSVDRUNTIME_API FChaosVDCollisionChannelInfo
+struct FChaosVDCollisionChannelInfo
 {
 	GENERATED_BODY()
 
@@ -486,14 +486,14 @@ struct CHAOSVDRUNTIME_API FChaosVDCollisionChannelInfo
 	UPROPERTY()
 	bool bIsTraceType = false;
 	
-	bool Serialize(FArchive& Ar);
+	CHAOSVDRUNTIME_API bool Serialize(FArchive& Ar);
 };
 
 CVD_IMPLEMENT_SERIALIZER(FChaosVDCollisionChannelInfo)
 
 /** Container for recorded custom collision profile data */
 USTRUCT()
-struct CHAOSVDRUNTIME_API FChaosVDCollisionChannelsInfoContainer
+struct FChaosVDCollisionChannelsInfoContainer
 {
 	GENERATED_BODY()
 
@@ -502,7 +502,7 @@ struct CHAOSVDRUNTIME_API FChaosVDCollisionChannelsInfoContainer
 	UPROPERTY(VisibleAnywhere, Category=CollisionData)
 	FChaosVDCollisionChannelInfo CustomChannelsNames[32] = {};
 
-	bool Serialize(FArchive& Ar);
+	CHAOSVDRUNTIME_API bool Serialize(FArchive& Ar);
 };
 
 CVD_IMPLEMENT_SERIALIZER(FChaosVDCollisionChannelsInfoContainer)
diff --git a/Engine/Source/Runtime/Experimental/ChaosVisualDebugger/Public/DataWrappers/ChaosVDDebugShapeDataWrapper.h b/Engine/Source/Runtime/Experimental/ChaosVisualDebugger/Public/DataWrappers/ChaosVDDebugShapeDataWrapper.h
--- a/Engine/Source/Runtime/Experimental/ChaosVisualDebugger/Public/DataWrappers/ChaosVDDebugShapeDataWrapper.h
+++ b/Engine/Source/Runtime/Experimental/ChaosVisualDebugger/Public/DataWrappers/ChaosVDDebugShapeDataWrapper.h
@@ -7,7 +7,7 @@
 #include "ChaosVDDebugShapeDataWrapper.generated.h"
 
 USTRUCT()
-struct CHAOSVDRUNTIME_API FChaosVDDebugDrawShapeBase : public FChaosVDWrapperDataBase
+struct FChaosVDDebugDrawShapeBase : public FChaosVDWrapperDataBase
 {
 	GENERATED_BODY()
 
@@ -21,11 +21,11 @@ struct CHAOSVDRUNTIME_API FChaosVDDebugDrawShapeBase : public FChaosVDWrapperDat
 	FColor Color = FColor::Blue;
 
 protected:
-	void SerializeBase_Internal(FArchive& Ar);
+	CHAOSVDRUNTIME_API void SerializeBase_Internal(FArchive& Ar);
 };
 
 USTRUCT()
-struct CHAOSVDRUNTIME_API FChaosVDDebugDrawBoxDataWrapper : public FChaosVDDebugDrawShapeBase
+struct FChaosVDDebugDrawBoxDataWrapper : public FChaosVDDebugDrawShapeBase
 {
 	GENERATED_BODY()
 
@@ -34,13 +34,13 @@ struct CHAOSVDRUNTIME_API FChaosVDDebugDrawBoxDataWrapper : public FChaosVDDebug
 	UPROPERTY()
 	FBox Box = FBox(ForceInitToZero);
 
-	bool Serialize(FArchive& Ar);
+	CHAOSVDRUNTIME_API bool Serialize(FArchive& Ar);
 };
 
 CVD_IMPLEMENT_SERIALIZER(FChaosVDDebugDrawBoxDataWrapper)
 
 USTRUCT()
-struct CHAOSVDRUNTIME_API FChaosVDDebugDrawSphereDataWrapper : public FChaosVDDebugDrawShapeBase
+struct FChaosVDDebugDrawSphereDataWrapper : public FChaosVDDebugDrawShapeBase
 {
 	GENERATED_BODY()
 	
@@ -52,13 +52,13 @@ struct CHAOSVDRUNTIME_API FChaosVDDebugDrawSphereDataWrapper : public FChaosVDDe
 	UPROPERTY()
 	float Radius = 0.0f;
 
-	bool Serialize(FArchive& Ar);
+	CHAOSVDRUNTIME_API bool Serialize(FArchive& Ar);
 };
 
 CVD_IMPLEMENT_SERIALIZER(FChaosVDDebugDrawSphereDataWrapper)
 
 USTRUCT()
-struct CHAOSVDRUNTIME_API FChaosVDDebugDrawLineDataWrapper : public FChaosVDDebugDrawShapeBase
+struct FChaosVDDebugDrawLineDataWrapper : public FChaosVDDebugDrawShapeBase
 {
 	GENERATED_BODY()
 	
@@ -73,13 +73,13 @@ struct CHAOSVDRUNTIME_API FChaosVDDebugDrawLineDataWrapper : public FChaosVDDebu
 	UPROPERTY()
 	bool bIsArrow = false;
 
-	bool Serialize(FArchive& Ar);
+	CHAOSVDRUNTIME_API bool Serialize(FArchive& Ar);
 };
 
 CVD_IMPLEMENT_SERIALIZER(FChaosVDDebugDrawLineDataWrapper)
 
 USTRUCT()
-struct CHAOSVDRUNTIME_API FChaosVDDebugDrawImplicitObjectDataWrapper : public FChaosVDDebugDrawShapeBase
+struct FChaosVDDebugDrawImplicitObjectDataWrapper : public FChaosVDDebugDrawShapeBase
 {
 	GENERATED_BODY()
 	
@@ -89,5 +89,5 @@ struct CHAOSVDRUNTIME_API FChaosVDDebugDrawImplicitObjectDataWrapper : public FC
 
 	FTransform ParentTransform = FTransform();
 
-	bool Serialize(FArchive& Ar);
+	CHAOSVDRUNTIME_API bool Serialize(FArchive& Ar);
 };
diff --git a/Engine/Source/Runtime/UMG/Private/Binding/States/WidgetStateRegistration.cpp b/Engine/Source/Runtime/UMG/Private/Binding/States/WidgetStateRegistration.cpp
--- a/Engine/Source/Runtime/UMG/Private/Binding/States/WidgetStateRegistration.cpp
+++ b/Engine/Source/Runtime/UMG/Private/Binding/States/WidgetStateRegistration.cpp
@@ -20,6 +20,8 @@ void UWidgetBinaryStateRegistration::InitializeStaticBitfields() const
 {
 }
 
+FWidgetStateBitfield UWidgetHoveredStateRegistration::Bit;
+
 FName UWidgetHoveredStateRegistration::GetStateName() const
 {
 	return StateName;
@@ -35,6 +37,8 @@ void UWidgetHoveredStateRegistration::InitializeStaticBitfields() const
 	Bit = FWidgetStateBitfield(GetStateName());
 }
 
+FWidgetStateBitfield UWidgetPressedStateRegistration::Bit;
+
 FName UWidgetPressedStateRegistration::GetStateName() const
 {
 	return StateName;
@@ -50,6 +54,8 @@ void UWidgetPressedStateRegistration::InitializeStaticBitfields() const
 	Bit = FWidgetStateBitfield(GetStateName());
 }
 
+FWidgetStateBitfield UWidgetDisabledStateRegistration::Bit;
+
 FName UWidgetDisabledStateRegistration::GetStateName() const
 {
 	return StateName;
@@ -65,6 +71,8 @@ void UWidgetDisabledStateRegistration::InitializeStaticBitfields() const
 	Bit = FWidgetStateBitfield(GetStateName());
 }
 
+FWidgetStateBitfield UWidgetSelectedStateRegistration::Bit;
+
 FName UWidgetSelectedStateRegistration::GetStateName() const
 {
 	return StateName;
diff --git a/Engine/Source/Runtime/UMG/Private/Components/CheckBox.cpp b/Engine/Source/Runtime/UMG/Private/Components/CheckBox.cpp
--- a/Engine/Source/Runtime/UMG/Private/Components/CheckBox.cpp
+++ b/Engine/Source/Runtime/UMG/Private/Components/CheckBox.cpp
@@ -294,6 +294,10 @@ const FText UCheckBox::GetPaletteCategory()
 
 #endif
 
+FWidgetStateBitfield UWidgetCheckedStateRegistration::Unchecked;
+FWidgetStateBitfield UWidgetCheckedStateRegistration::Checked;
+FWidgetStateBitfield UWidgetCheckedStateRegistration::Undetermined;
+
 FName UWidgetCheckedStateRegistration::GetStateName() const
 {
 	return StateName;
diff --git a/Engine/Source/Runtime/UMG/Public/Binding/States/WidgetStateRegistration.h b/Engine/Source/Runtime/UMG/Public/Binding/States/WidgetStateRegistration.h
--- a/Engine/Source/Runtime/UMG/Public/Binding/States/WidgetStateRegistration.h
+++ b/Engine/Source/Runtime/UMG/Public/Binding/States/WidgetStateRegistration.h
@@ -48,7 +48,7 @@ class UWidgetHoveredStateRegistration : public UWidgetBinaryStateRegistration
 public:
 
 	/** Post-load initialized bit corresponding to this binary state */
-	static UMG_API inline FWidgetStateBitfield Bit;
+	static UMG_API FWidgetStateBitfield Bit;
 
 	static const inline FName StateName = FName("Hovered");
 
@@ -73,7 +73,7 @@ class UWidgetPressedStateRegistration : public UWidgetBinaryStateRegistration
 public:
 
 	/** Post-load initialized bit corresponding to this binary state */
-	static UMG_API inline FWidgetStateBitfield Bit;
+	static UMG_API FWidgetStateBitfield Bit;
 
 	static const inline FName StateName = FName("Pressed");
 
@@ -98,7 +98,7 @@ class UWidgetDisabledStateRegistration : public UWidgetBinaryStateRegistration
 public:
 
 	/** Post-load initialized bit corresponding to this binary state */
-	static UMG_API inline FWidgetStateBitfield Bit;
+	static UMG_API FWidgetStateBitfield Bit;
 
 	static const inline FName StateName = FName("Disabled");
 
@@ -123,7 +123,7 @@ class UWidgetSelectedStateRegistration : public UWidgetBinaryStateRegistration
 public:
 
 	/** Post-load initialized bit corresponding to this binary state */
-	static UMG_API inline FWidgetStateBitfield Bit;
+	static UMG_API FWidgetStateBitfield Bit;
 
 	static const inline FName StateName = FName("Selected");
 
diff --git a/Engine/Source/Runtime/UMG/Public/Components/CheckBox.h b/Engine/Source/Runtime/UMG/Public/Components/CheckBox.h
--- a/Engine/Source/Runtime/UMG/Public/Components/CheckBox.h
+++ b/Engine/Source/Runtime/UMG/Public/Components/CheckBox.h
@@ -188,9 +188,9 @@ class UWidgetCheckedStateRegistration : public UWidgetEnumStateRegistration
 public:
 
 	/** Post-load initialized values corresponding to this enum state */
-	static UMG_API inline FWidgetStateBitfield Unchecked;
-	static UMG_API inline FWidgetStateBitfield Checked;
-	static UMG_API inline FWidgetStateBitfield Undetermined;
+	static UMG_API FWidgetStateBitfield Unchecked;
+	static UMG_API FWidgetStateBitfield Checked;
+	static UMG_API FWidgetStateBitfield Undetermined;
 
 	static const inline FName StateName = FName("CheckedState");
 
