Fixed broken FindFirstVariableOfType missing RIGVM_API preventing EngineTest from compiling and breaking engine preflights that include engine test.

#rnx

diff --git a/Engine/Plugins/Runtime/RigVM/Source/RigVM/Public/RigVMCore/RigVMStruct.h b/Engine/Plugins/Runtime/RigVM/Source/RigVM/Public/RigVMCore/RigVMStruct.h
--- a/Engine/Plugins/Runtime/RigVM/Source/RigVM/Public/RigVMCore/RigVMStruct.h
+++ b/Engine/Plugins/Runtime/RigVM/Source/RigVM/Public/RigVMCore/RigVMStruct.h
@@ -157,7 +157,7 @@ private:
 	FRigVMCreateExternalVariableDelegate CreateExternalVariableDelegate;
 	FRigVMBindPinToExternalVariableDelegate BindPinToExternalVariableDelegate;
 
-	FName FindFirstVariableOfType(UObject* InCPPTypeObject) const;
+	RIGVM_API FName FindFirstVariableOfType(UObject* InCPPTypeObject) const;
 
 	friend class URigVMController;
 	friend struct FScope;
