Variants: Improve tool tips / display names

#rnx

diff --git a/Engine/Plugins/Runtime/RigVM/Source/RigVM/Public/RigVMCore/RigVMVariant.h b/Engine/Plugins/Runtime/RigVM/Source/RigVM/Public/RigVMCore/RigVMVariant.h
--- a/Engine/Plugins/Runtime/RigVM/Source/RigVM/Public/RigVMCore/RigVMVariant.h
+++ b/Engine/Plugins/Runtime/RigVM/Source/RigVM/Public/RigVMCore/RigVMVariant.h
@@ -43,10 +43,16 @@ struct RIGVM_API FRigVMTag
 	UPROPERTY(BlueprintReadWrite, EditAnywhere, Category=Variant)
 	FLinearColor Color;
 
+	// Enabling this will show tags in the user interface within
+	// the variant widgets and other places.
+	// Disabling this will interpret this tag as automation / scripting only.
 	UPROPERTY(BlueprintReadWrite, EditAnywhere, Category=Variant)
 	bool bShowInUserInterface;
 
-	UPROPERTY(BlueprintReadWrite, EditAnywhere, Category=Variant)
+	// Enabling this will mark the subject (asset, function etc) as (soft) deprecated.
+	// The content referring to variants with this tag will continue to work, but the
+	// user interface will suggest the user to upgrade to a later variant.
+	UPROPERTY(BlueprintReadWrite, EditAnywhere, Category=Variant, meta=(DisplayName="Mark Subject Deprecated"))
 	bool bMarksSubjectAsInvalid;
 
 	bool IsValid() const
