RigVM: Don't resolve singleton templates against any type - but matching types

#rnx

diff --git a/Engine/Plugins/Runtime/RigVM/Source/RigVM/Private/RigVMCore/RigVMTemplate.cpp b/Engine/Plugins/Runtime/RigVM/Source/RigVM/Private/RigVMCore/RigVMTemplate.cpp
--- a/Engine/Plugins/Runtime/RigVM/Source/RigVM/Private/RigVMCore/RigVMTemplate.cpp
+++ b/Engine/Plugins/Runtime/RigVM/Source/RigVM/Private/RigVMCore/RigVMTemplate.cpp
@@ -1734,12 +1734,28 @@ bool FRigVMTemplate::Resolve(FTypeMap& InOutTypes, TArray<int32>& OutPermutation
 	
 	for (const FRigVMTemplateArgument& Argument : Arguments)
 	{
+		const TRigVMTypeIndex* InputType = InputTypes.Find(Argument.Name);
+		
 		if (Argument.IsSingleton_NoLock())
 		{
-			InOutTypes.Add(Argument.Name, Argument.GetTypeIndex_NoLock(0));
+			// if we are singleton - we still need to check if the potentially provided
+			// type is compatible with the singleton type.
+			const TRigVMTypeIndex SingleTypeIndex = Argument.GetTypeIndex_NoLock(0);
+			if (InputType)
+			{
+				if (!Registry.IsWildCardType_NoLock(*InputType))
+				{
+					if (!Registry.CanMatchTypes_NoLock(*InputType, SingleTypeIndex, bAllowFloatingPointCasts))
+					{
+						OutPermutationIndices.Reset();
+						return false;
+					}
+				}
+			}
+			InOutTypes.Add(Argument.Name, SingleTypeIndex);
 			continue;
 		}
-		else if (const TRigVMTypeIndex* InputType = InputTypes.Find(Argument.Name))
+		else if (InputType)
 		{
 			TArray<TRigVMTypeIndex> AllTypes; Argument.GetAllTypes_NoLock(AllTypes);
 
@@ -1796,7 +1812,7 @@ bool FRigVMTemplate::Resolve(FTypeMap& InOutTypes, TArray<int32>& OutPermutation
 		{
 			InOutTypes.Add(Argument.Name, RigVMTypeUtils::TypeIndex::WildCard);
 
-			if(const TRigVMTypeIndex* InputType = InputTypes.Find(Argument.Name))
+			if(InputType)
 			{
 				if(Registry.IsArrayType_NoLock(*InputType))
 				{
