RigVM: Avoid race condition when accessing URigVM::FunctionStorage

#rnx

diff --git a/Engine/Plugins/Runtime/RigVM/Source/RigVM/Private/RigVMCore/RigVM.cpp b/Engine/Plugins/Runtime/RigVM/Source/RigVM/Private/RigVMCore/RigVM.cpp
--- a/Engine/Plugins/Runtime/RigVM/Source/RigVM/Private/RigVMCore/RigVM.cpp
+++ b/Engine/Plugins/Runtime/RigVM/Source/RigVM/Private/RigVMCore/RigVM.cpp
@@ -325,8 +325,12 @@ void URigVM::PostLoad()
 
 void URigVM::RefreshArgumentNameCaches()
 {
+	// make sure that the functions cannot change at the same time as
+	// the argument name caches are being updated.
+	FScopeLock ResolveFunctionsScopeLock(&ResolveFunctionsMutex);
+
 	// make sure to update all argument name caches
-	TArray<const FRigVMFunction*>& Functions = GetFunctions();
+	const TArray<const FRigVMFunction*>& Functions = GetFunctions();
 	for (const FRigVMInstruction& Instruction : Instructions)
 	{
 		if (Instruction.OpCode == ERigVMOpCode::Execute)
