RigVM: hide variants behind a console variable

#rnx

diff --git a/Engine/Plugins/Runtime/RigVM/Source/RigVM/Private/RigVMCore/RigVMVariant.cpp b/Engine/Plugins/Runtime/RigVM/Source/RigVM/Private/RigVMCore/RigVMVariant.cpp
new file mode 100644
--- /dev/null
+++ b/Engine/Plugins/Runtime/RigVM/Source/RigVM/Private/RigVMCore/RigVMVariant.cpp
@@ -0,0 +1,5 @@
+// Copyright Epic Games, Inc. All Rights Reserved.
+
+#include "RigVMCore/RigVMVariant.h"
+
+TAutoConsoleVariable<bool> CVarRigVMEnableVariants(TEXT("RigVM.EnableVariants"), false, TEXT("enables the UI bits to create a variant of assets / functions"));
diff --git a/Engine/Plugins/Runtime/RigVM/Source/RigVM/Public/RigVMCore/RigVMVariant.h b/Engine/Plugins/Runtime/RigVM/Source/RigVM/Public/RigVMCore/RigVMVariant.h
--- a/Engine/Plugins/Runtime/RigVM/Source/RigVM/Public/RigVMCore/RigVMVariant.h
+++ b/Engine/Plugins/Runtime/RigVM/Source/RigVM/Public/RigVMCore/RigVMVariant.h
@@ -3,10 +3,13 @@
 #pragma once
 
 #include "RigVMObjectVersion.h"
+#include "HAL/IConsoleManager.h"
 #include "RigVMVariant.generated.h"
 
 struct FRigVMVariantRef;
 
+extern RIGVM_API TAutoConsoleVariable<bool> CVarRigVMEnableVariants;
+
 // User applied tag
 USTRUCT(BlueprintType)
 struct RIGVM_API FRigVMTag
diff --git a/Engine/Plugins/Runtime/RigVM/Source/RigVMDeveloper/Public/EdGraph/RigVMEdGraphSchema.h b/Engine/Plugins/Runtime/RigVM/Source/RigVMDeveloper/Public/EdGraph/RigVMEdGraphSchema.h
--- a/Engine/Plugins/Runtime/RigVM/Source/RigVMDeveloper/Public/EdGraph/RigVMEdGraphSchema.h
+++ b/Engine/Plugins/Runtime/RigVM/Source/RigVMDeveloper/Public/EdGraph/RigVMEdGraphSchema.h
@@ -6,6 +6,7 @@
 #include "EdGraph/EdGraphSchema.h"
 #include "GraphEditorDragDropAction.h"
 #include "RigVMModel/RigVMGraph.h"
+#include "RigVMCore/RigVMVariant.h"
 #include "EdGraphSchema_K2_Actions.h"
 #include "Kismet2/Kismet2NameValidators.h"
 
@@ -312,7 +313,7 @@ public:
 	virtual bool TryRenameGraph(UEdGraph* GraphToRename, const FName& InNewName) const override;
 	virtual bool TryToGetChildEvents(const UEdGraph* Graph, const int32 SectionId, TArray<TSharedPtr<FEdGraphSchemaAction>>& Actions, const FText& ParentCategory) const override;
 	virtual bool CanDuplicateGraph(UEdGraph* InSourceGraph) const { return false; }
-	virtual bool AllowsFunctionVariants() const override { return true; }
+	virtual bool AllowsFunctionVariants() const override { return CVarRigVMEnableVariants.GetValueOnAnyThread(); }
 	virtual UEdGraphPin* DropPinOnNode(UEdGraphNode* InTargetNode, const FName& InSourcePinName, const FEdGraphPinType& InSourcePinType, EEdGraphPinDirection InSourcePinDirection) const override;
 	virtual bool SupportsDropPinOnNode(UEdGraphNode* InTargetNode, const FEdGraphPinType& InSourcePinType, EEdGraphPinDirection InSourcePinDirection, FText& OutErrorMessage) const override;
 	virtual void SetPinBeingDroppedOnNode(UEdGraphPin* InSourcePin) const override { PinBeingDropped = InSourcePin; }
diff --git a/Engine/Plugins/Runtime/RigVM/Source/RigVMEditor/Private/Editor/RigVMGraphDetailCustomization.cpp b/Engine/Plugins/Runtime/RigVM/Source/RigVMEditor/Private/Editor/RigVMGraphDetailCustomization.cpp
--- a/Engine/Plugins/Runtime/RigVM/Source/RigVMEditor/Private/Editor/RigVMGraphDetailCustomization.cpp
+++ b/Engine/Plugins/Runtime/RigVM/Source/RigVMEditor/Private/Editor/RigVMGraphDetailCustomization.cpp
@@ -991,38 +991,41 @@ void FRigVMGraphDetailCustomizationImpl::CustomizeDetails(IDetailLayoutBuilder&
         ];
 
 		// variant
-		FRigVMVariantWidgetContext VariantContext;
-		if(const URigVMFunctionLibrary* FunctionLbirary = Model->GetTypedOuter<URigVMFunctionLibrary>())
+		if(CVarRigVMEnableVariants.GetValueOnAnyThread())
 		{
-			VariantContext.ParentPath = FunctionLbirary->GetPathName();
+			FRigVMVariantWidgetContext VariantContext;
+			if(const URigVMFunctionLibrary* FunctionLbirary = Model->GetTypedOuter<URigVMFunctionLibrary>())
+			{
+				VariantContext.ParentPath = FunctionLbirary->GetPathName();
+			}
+				
+			SettingsCategory.AddCustomRow(FText::GetEmpty())
+			.OverrideResetToDefault(FResetToDefaultOverride::Hide())
+			.Visibility(TAttribute<EVisibility>::CreateLambda([this]()
+			{
+				return IsValidFunction() ? EVisibility::Visible : EVisibility::Collapsed;
+			}))
+			.NameContent()
+			[
+				SNew(STextBlock)
+				.Text(FText::FromString(TEXT("Variant")))
+				.Font(IDetailLayoutBuilder::GetDetailFont())
+			]
+			.ValueContent()
+			[
+				SNew(SRigVMVariantWidget)
+				.Context(VariantContext)
+				.Variant(this, &FRigVMGraphDetailCustomizationImpl::GetVariant)
+				.VariantRefs(this, &FRigVMGraphDetailCustomizationImpl::GetVariantRefs)
+				.OnVariantChanged(this, &FRigVMGraphDetailCustomizationImpl::OnVariantChanged)
+				.OnBrowseVariantRef(this, &FRigVMGraphDetailCustomizationImpl::OnBrowseVariantRef)
+				.OnGetTags(this, &FRigVMGraphDetailCustomizationImpl::OnGetAssignedTags)
+				.OnAddTag(this, &FRigVMGraphDetailCustomizationImpl::OnAddAssignedTag)
+				.OnRemoveTag(this, &FRigVMGraphDetailCustomizationImpl::OnRemoveAssignedTag)
+				.CanAddTags(true)
+				.EnableTagContextMenu(true)
+			];
 		}
-			
-		SettingsCategory.AddCustomRow(FText::GetEmpty())
-		.OverrideResetToDefault(FResetToDefaultOverride::Hide())
-		.Visibility(TAttribute<EVisibility>::CreateLambda([this]()
-		{
-			return IsValidFunction() ? EVisibility::Visible : EVisibility::Collapsed;
-		}))
-		.NameContent()
-		[
-			SNew(STextBlock)
-			.Text(FText::FromString(TEXT("Variant")))
-			.Font(IDetailLayoutBuilder::GetDetailFont())
-		]
-		.ValueContent()
-		[
-			SNew(SRigVMVariantWidget)
-			.Context(VariantContext)
-			.Variant(this, &FRigVMGraphDetailCustomizationImpl::GetVariant)
-			.VariantRefs(this, &FRigVMGraphDetailCustomizationImpl::GetVariantRefs)
-			.OnVariantChanged(this, &FRigVMGraphDetailCustomizationImpl::OnVariantChanged)
-			.OnBrowseVariantRef(this, &FRigVMGraphDetailCustomizationImpl::OnBrowseVariantRef)
-			.OnGetTags(this, &FRigVMGraphDetailCustomizationImpl::OnGetAssignedTags)
-			.OnAddTag(this, &FRigVMGraphDetailCustomizationImpl::OnAddAssignedTag)
-			.OnRemoveTag(this, &FRigVMGraphDetailCustomizationImpl::OnRemoveAssignedTag)
-			.CanAddTags(true)
-			.EnableTagContextMenu(true)
-		];
 	}
 
 	// node color
diff --git a/Engine/Plugins/Runtime/RigVM/Source/RigVMEditor/Private/Editor/RigVMVariantDetailCustomization.cpp b/Engine/Plugins/Runtime/RigVM/Source/RigVMEditor/Private/Editor/RigVMVariantDetailCustomization.cpp
--- a/Engine/Plugins/Runtime/RigVM/Source/RigVMEditor/Private/Editor/RigVMVariantDetailCustomization.cpp
+++ b/Engine/Plugins/Runtime/RigVM/Source/RigVMEditor/Private/Editor/RigVMVariantDetailCustomization.cpp
@@ -20,6 +20,11 @@ class FUObjectToken;
 
 void FRigVMVariantDetailCustomization::CustomizeHeader(TSharedRef<IPropertyHandle> InStructPropertyHandle, FDetailWidgetRow& HeaderRow, IPropertyTypeCustomizationUtils& StructCustomizationUtils)
 {
+	if(!CVarRigVMEnableVariants.GetValueOnAnyThread())
+	{
+		return;
+	}
+		
 	TArray<UObject*> Objects;
 	InStructPropertyHandle->GetOuterObjects(Objects);
 	ensure(Objects.Num() == 1); // This is in here to ensure we are only showing the modifier details in the blueprint editor
diff --git a/Engine/Plugins/Runtime/RigVM/Source/RigVMEditor/Private/RigVMEditorModule.cpp b/Engine/Plugins/Runtime/RigVM/Source/RigVMEditor/Private/RigVMEditorModule.cpp
--- a/Engine/Plugins/Runtime/RigVM/Source/RigVMEditor/Private/RigVMEditorModule.cpp
+++ b/Engine/Plugins/Runtime/RigVM/Source/RigVMEditor/Private/RigVMEditorModule.cpp
@@ -106,31 +106,34 @@ void FRigVMEditorModule::StartupModule()
 							return;
 						}
 
-						FSoftObjectPath SoftObjectPath = SelectedAssetData.GetSoftObjectPath();
-						InSection.AddMenuEntry("CreateVariant", LOCTEXT("CreateVariant", "Create variant"), LOCTEXT("CreateVariant_ToolTip", "Create a variant for this asset"), FSlateIcon(FRigVMEditorStyle::Get().GetStyleSetName(), "RigVM", "RigVM.Unit"), FExecuteAction::CreateLambda([SoftObjectPath]()
-							{
-								// Perform the load from within our lambda since this can be expensive, and should not be done speculatively
-								UObject* SelectedObject = SoftObjectPath.TryLoad();
-								if (!SelectedObject)
+						if(CVarRigVMEnableVariants.GetValueOnAnyThread())
+						{
+							FSoftObjectPath SoftObjectPath = SelectedAssetData.GetSoftObjectPath();
+							InSection.AddMenuEntry("CreateVariant", LOCTEXT("CreateVariant", "Create variant"), LOCTEXT("CreateVariant_ToolTip", "Create a variant for this asset"), FSlateIcon(FRigVMEditorStyle::Get().GetStyleSetName(), "RigVM", "RigVM.Unit"), FExecuteAction::CreateLambda([SoftObjectPath]()
 								{
-									return;
-								}
+									// Perform the load from within our lambda since this can be expensive, and should not be done speculatively
+									UObject* SelectedObject = SoftObjectPath.TryLoad();
+									if (!SelectedObject)
+									{
+										return;
+									}
 
-								const FAssetToolsModule& AssetToolsModule = FModuleManager::LoadModuleChecked<FAssetToolsModule>("AssetTools");
-								FString PathName = SoftObjectPath.GetLongPackageName();
-								FString ObjectName = SelectedObject->GetName();
-								FString PackageName;
-								AssetToolsModule.Get().CreateUniqueAssetName(PathName, TEXT(""), PackageName, ObjectName);
+									const FAssetToolsModule& AssetToolsModule = FModuleManager::LoadModuleChecked<FAssetToolsModule>("AssetTools");
+									FString PathName = SoftObjectPath.GetLongPackageName();
+									FString ObjectName = SelectedObject->GetName();
+									FString PackageName;
+									AssetToolsModule.Get().CreateUniqueAssetName(PathName, TEXT(""), PackageName, ObjectName);
 
-								FString Extension;
-								FPaths::Split(PackageName, PathName, ObjectName, Extension);
+									FString Extension;
+									FPaths::Split(PackageName, PathName, ObjectName, Extension);
 
-								UObject* DuplicateAsset = AssetToolsModule.Get().DuplicateAsset(ObjectName, PathName, SelectedObject);
-								if (URigVMBlueprint* DuplicateBlueprint = Cast<URigVMBlueprint>(DuplicateAsset))
-								{
-									DuplicateBlueprint->AssetVariant = Cast<URigVMBlueprint>(SelectedObject)->AssetVariant;
-								}
-							}));
+									UObject* DuplicateAsset = AssetToolsModule.Get().DuplicateAsset(ObjectName, PathName, SelectedObject);
+									if (URigVMBlueprint* DuplicateBlueprint = Cast<URigVMBlueprint>(DuplicateAsset))
+									{
+										DuplicateBlueprint->AssetVariant = Cast<URigVMBlueprint>(SelectedObject)->AssetVariant;
+									}
+								}));
+						}
 					}
 				}));
 			}
