RigVM: New node to determine if a graph is being debugged in an asset editor

#rnx

diff --git a/Engine/Plugins/Animation/ControlRig/Source/ControlRig/Private/ControlRig.cpp b/Engine/Plugins/Animation/ControlRig/Source/ControlRig/Private/ControlRig.cpp
--- a/Engine/Plugins/Animation/ControlRig/Source/ControlRig/Private/ControlRig.cpp
+++ b/Engine/Plugins/Animation/ControlRig/Source/ControlRig/Private/ControlRig.cpp
@@ -774,6 +774,9 @@ bool UControlRig::Execute(const FName& InEventName)
 	PublicContext.SetDeltaTime(DeltaTime);
 	PublicContext.SetAbsoluteTime(AbsoluteTime);
 	PublicContext.SetFramesPerSecond(GetCurrentFramesPerSecond());
+#if WITH_EDITOR
+	PublicContext.SetHostBeingDebugged(bIsBeingDebugged);
+#endif
 
 #if UE_RIGVM_DEBUG_EXECUTION
 	PublicContext.bDebugExecution = bDebugExecutionEnabled;
diff --git a/Engine/Plugins/Animation/ControlRig/Source/ControlRig/Private/ModularRig.cpp b/Engine/Plugins/Animation/ControlRig/Source/ControlRig/Private/ModularRig.cpp
--- a/Engine/Plugins/Animation/ControlRig/Source/ControlRig/Private/ModularRig.cpp
+++ b/Engine/Plugins/Animation/ControlRig/Source/ControlRig/Private/ModularRig.cpp
@@ -491,6 +491,9 @@ void UModularRig::ExecuteQueue()
 				RigPublicContext.SetLog(PublicContext.GetLog());
 #endif
 				RigPublicContext.SetFramesPerSecond(PublicContext.GetFramesPerSecond());
+#if WITH_EDITOR
+				RigPublicContext.SetHostBeingDebugged(bIsBeingDebugged);
+#endif
 				RigPublicContext.SetToWorldSpaceTransform(PublicContext.GetToWorldSpaceTransform());
 				RigPublicContext.OnAddShapeLibraryDelegate = PublicContext.OnAddShapeLibraryDelegate;
 				RigPublicContext.OnShapeExistsDelegate = PublicContext.OnShapeExistsDelegate;
diff --git a/Engine/Plugins/Runtime/RigVM/Source/RigVM/Private/RigVMFunctions/Execution/RigVMFunction_Context.cpp b/Engine/Plugins/Runtime/RigVM/Source/RigVM/Private/RigVMFunctions/Execution/RigVMFunction_Context.cpp
new file mode 100644
--- /dev/null
+++ b/Engine/Plugins/Runtime/RigVM/Source/RigVM/Private/RigVMFunctions/Execution/RigVMFunction_Context.cpp
@@ -0,0 +1,14 @@
+// Copyright Epic Games, Inc. All Rights Reserved.
+
+#include "RigVMFunctions/Execution/RigVMFunction_Context.h"
+
+FRigVMFunction_IsHostBeingDebugged_Execute()
+{
+#if WITH_EDITOR
+	Result = ExecuteContext.IsHostBeingDebugged();
+#else
+	Result = false;
+#endif
+}
+
+
diff --git a/Engine/Plugins/Runtime/RigVM/Source/RigVM/Private/RigVMHost.cpp b/Engine/Plugins/Runtime/RigVM/Source/RigVM/Private/RigVMHost.cpp
--- a/Engine/Plugins/Runtime/RigVM/Source/RigVM/Private/RigVMHost.cpp
+++ b/Engine/Plugins/Runtime/RigVM/Source/RigVM/Private/RigVMHost.cpp
@@ -32,6 +32,7 @@ URigVMHost::URigVMHost(const FObjectInitializer& ObjectInitializer)
 	, FramesPerSecond(0.0f)
 	, bAccumulateTime(true)
 #if WITH_EDITOR
+	, bIsBeingDebugged(false)
 	, RigVMLog(nullptr)
 	, bEnableLogging(true)
 #endif
@@ -526,6 +527,9 @@ bool URigVMHost::Execute(const FName& InEventName)
 	PublicContext.SetDeltaTime(DeltaTime);
 	PublicContext.SetAbsoluteTime(AbsoluteTime);
 	PublicContext.SetFramesPerSecond(GetCurrentFramesPerSecond());
+#if WITH_EDITOR
+	PublicContext.SetHostBeingDebugged(bIsBeingDebugged);
+#endif
 	PublicContext.SetOwningComponent(GetOwningSceneComponent());
 #if UE_RIGVM_DEBUG_EXECUTION
 	PublicContext.bDebugExecution = bDebugExecutionEnabled;
diff --git a/Engine/Plugins/Runtime/RigVM/Source/RigVM/Public/RigVMCore/RigVMExecuteContext.h b/Engine/Plugins/Runtime/RigVM/Source/RigVM/Public/RigVMCore/RigVMExecuteContext.h
--- a/Engine/Plugins/Runtime/RigVM/Source/RigVM/Public/RigVMCore/RigVMExecuteContext.h
+++ b/Engine/Plugins/Runtime/RigVM/Source/RigVM/Public/RigVMCore/RigVMExecuteContext.h
@@ -270,6 +270,9 @@ struct FRigVMExecuteContext
 		, DeltaTime(0.0)
 		, AbsoluteTime(0.0)
 		, FramesPerSecond(1.0 / 60.0)
+#if WITH_EDITOR
+		, bHostBeingDebugged(false)
+#endif
 		, RuntimeSettings()
 		, NameCache(nullptr)
 #if WITH_EDITOR
@@ -334,6 +337,11 @@ struct FRigVMExecuteContext
 	double GetFramesPerSecond() const { return FramesPerSecond; } 
 	void SetFramesPerSecond(double InFramesPerSecond) { FramesPerSecond = InFramesPerSecond; }
 
+#if WITH_EDITOR
+	bool IsHostBeingDebugged() const { return bHostBeingDebugged; }
+	void SetHostBeingDebugged(bool InIsHostBeingDebugged) { bHostBeingDebugged = InIsHostBeingDebugged; }
+#endif
+
 	/** The current transform going from rig (global) space to world space */
 	const FTransform& GetToWorldSpaceTransform() const { return ToWorldSpaceTransform; };
 
@@ -489,6 +497,10 @@ protected:
 
 	double FramesPerSecond;
 
+#if WITH_EDITOR
+	bool bHostBeingDebugged;
+#endif
+
 	FRigVMRuntimeSettings RuntimeSettings;
 
 	mutable FRigVMNameCache* NameCache;
diff --git a/Engine/Plugins/Runtime/RigVM/Source/RigVM/Public/RigVMFunctions/Execution/RigVMFunction_Context.h b/Engine/Plugins/Runtime/RigVM/Source/RigVM/Public/RigVMFunctions/Execution/RigVMFunction_Context.h
new file mode 100644
--- /dev/null
+++ b/Engine/Plugins/Runtime/RigVM/Source/RigVM/Public/RigVMFunctions/Execution/RigVMFunction_Context.h
@@ -0,0 +1,26 @@
+// Copyright Epic Games, Inc. All Rights Reserved.
+
+#pragma once
+
+#include "RigVMCore/RigVMStruct.h"
+#include "RigVMFunction_Context.generated.h"
+
+/**
+ * Returns true if a graph is run with its asset editor open. This is editor only - in shipping it always returns false.
+ */
+USTRUCT(meta=(DisplayName="Is Asset Editor Open", Category="Execution", TitleColor="1 0 0", NodeColor="1 1 1", Keywords="Debug,Open,Inspect"))
+struct RIGVM_API FRigVMFunction_IsHostBeingDebugged : public FRigVMStruct
+{
+	GENERATED_BODY()
+
+	FRigVMFunction_IsHostBeingDebugged()
+	: Result(false)
+	{}
+
+	RIGVM_METHOD()
+	virtual void Execute() override;
+
+	// True if the graph is currently open in the asset editor
+	UPROPERTY(EditAnywhere, Transient, Category = "Execution", meta = (Output))
+	bool Result;
+};
\ No newline at end of file
diff --git a/Engine/Plugins/Runtime/RigVM/Source/RigVM/Public/RigVMHost.h b/Engine/Plugins/Runtime/RigVM/Source/RigVM/Public/RigVMHost.h
--- a/Engine/Plugins/Runtime/RigVM/Source/RigVM/Public/RigVMHost.h
+++ b/Engine/Plugins/Runtime/RigVM/Source/RigVM/Public/RigVMHost.h
@@ -334,6 +334,11 @@ protected:
 	/** true if we should increase the AbsoluteTime */
 	bool bAccumulateTime;
 
+#if WITH_EDITOR
+	/** true if the instance is being debugged in an asset editor */
+	bool bIsBeingDebugged;
+#endif
+
 	UPROPERTY(Transient)
 	TObjectPtr<URigVM> VM;
 
diff --git a/Engine/Plugins/Runtime/RigVM/Source/RigVMDeveloper/Private/RigVMBlueprint.cpp b/Engine/Plugins/Runtime/RigVM/Source/RigVMDeveloper/Private/RigVMBlueprint.cpp
--- a/Engine/Plugins/Runtime/RigVM/Source/RigVMDeveloper/Private/RigVMBlueprint.cpp
+++ b/Engine/Plugins/Runtime/RigVM/Source/RigVMDeveloper/Private/RigVMBlueprint.cpp
@@ -2434,9 +2434,19 @@ void URigVMBlueprint::SetObjectBeingDebugged(UObject* NewObject)
 	{
 		PreviousRigBeingDebugged->DrawInterface.Reset();
 		PreviousRigBeingDebugged->RigVMLog = nullptr;
+#if WITH_EDITOR
+		PreviousRigBeingDebugged->bIsBeingDebugged = false;
+#endif
 	}
 
 	Super::SetObjectBeingDebugged(NewObject);
+
+#if WITH_EDITOR
+	if(URigVMHost* NewRigBeingDebugged = Cast<URigVMHost>(NewObject))
+	{
+		NewRigBeingDebugged->bIsBeingDebugged = true;
+	}
+#endif
 }
 
 void URigVMBlueprint::PostTransacted(const FTransactionObjectEvent& TransactionEvent)
