RigVM: Re-enable variant console var, but mark as experimental

#rnx

diff --git a/Engine/Plugins/Runtime/RigVM/Source/RigVM/Private/RigVMCore/RigVMVariant.cpp b/Engine/Plugins/Runtime/RigVM/Source/RigVM/Private/RigVMCore/RigVMVariant.cpp
--- a/Engine/Plugins/Runtime/RigVM/Source/RigVM/Private/RigVMCore/RigVMVariant.cpp
+++ b/Engine/Plugins/Runtime/RigVM/Source/RigVM/Private/RigVMCore/RigVMVariant.cpp
@@ -2,4 +2,4 @@
 
 #include "RigVMCore/RigVMVariant.h"
 
-TAutoConsoleVariable<bool> CVarRigVMEnableVariants(TEXT("RigVM.EnableVariants"), false, TEXT("enables the UI bits to create a variant of assets / functions"));
+TAutoConsoleVariable<bool> CVarRigVMEnableVariants(TEXT("RigVM.EnableVariants"), true, TEXT("enables the UI bits to create a variant of assets / functions"));
diff --git a/Engine/Plugins/Runtime/RigVM/Source/RigVM/Public/RigVMSettings.h b/Engine/Plugins/Runtime/RigVM/Source/RigVM/Public/RigVMSettings.h
--- a/Engine/Plugins/Runtime/RigVM/Source/RigVM/Public/RigVMSettings.h
+++ b/Engine/Plugins/Runtime/RigVM/Source/RigVM/Public/RigVMSettings.h
@@ -44,7 +44,7 @@ class RIGVM_API URigVMProjectSettings : public UDeveloperSettings
 
 public:
 
-	UPROPERTY(EditAnywhere, config, Category = Variants)
+	UPROPERTY(EditAnywhere, config, Category = Experimental)
 	TArray<FRigVMTag> VariantTags;
 
 	UFUNCTION(BlueprintPure, Category= Variants)
diff --git a/Engine/Plugins/Runtime/RigVM/Source/RigVMEditor/Private/RigVMEditorModule.cpp b/Engine/Plugins/Runtime/RigVM/Source/RigVMEditor/Private/RigVMEditorModule.cpp
--- a/Engine/Plugins/Runtime/RigVM/Source/RigVMEditor/Private/RigVMEditorModule.cpp
+++ b/Engine/Plugins/Runtime/RigVM/Source/RigVMEditor/Private/RigVMEditorModule.cpp
@@ -109,7 +109,7 @@ void FRigVMEditorModule::StartupModule()
 						if(CVarRigVMEnableVariants.GetValueOnAnyThread())
 						{
 							FSoftObjectPath SoftObjectPath = SelectedAssetData.GetSoftObjectPath();
-							InSection.AddMenuEntry("CreateVariant", LOCTEXT("CreateVariant", "Create variant"), LOCTEXT("CreateVariant_ToolTip", "Create a variant for this asset"), FSlateIcon(FRigVMEditorStyle::Get().GetStyleSetName(), "RigVM", "RigVM.Unit"), FExecuteAction::CreateLambda([SoftObjectPath]()
+							InSection.AddMenuEntry("CreateVariant", LOCTEXT("CreateVariant", "Create variant (Experimental)"), LOCTEXT("CreateVariant_ToolTip", "Create a variant for this asset"), FSlateIcon(FRigVMEditorStyle::Get().GetStyleSetName(), "RigVM", "RigVM.Unit"), FExecuteAction::CreateLambda([SoftObjectPath]()
 								{
 									// Perform the load from within our lambda since this can be expensive, and should not be done speculatively
 									UObject* SelectedObject = SoftObjectPath.TryLoad();
diff --git a/Engine/Plugins/Runtime/RigVM/Source/RigVMEditor/Private/Widgets/SRigVMVariantWidget.cpp b/Engine/Plugins/Runtime/RigVM/Source/RigVMEditor/Private/Widgets/SRigVMVariantWidget.cpp
--- a/Engine/Plugins/Runtime/RigVM/Source/RigVMEditor/Private/Widgets/SRigVMVariantWidget.cpp
+++ b/Engine/Plugins/Runtime/RigVM/Source/RigVMEditor/Private/Widgets/SRigVMVariantWidget.cpp
@@ -117,6 +117,15 @@ void SRigVMVariantWidget::Construct(
 		SNew(SVerticalBox)
 		+ SVerticalBox::Slot()
 		.AutoHeight()
+		.Padding(0, 8, 0, 8)
+		.HAlign(HAlign_Left)
+		[
+			SNew(STextBlock)
+			.Font(IDetailLayoutBuilder::GetDetailFont())
+			.Text(LOCTEXT("Experimental", "Note: This is experimental"))
+		]
+		+ SVerticalBox::Slot()
+		.AutoHeight()
 		.Padding(0, 8, 0, 0)
 		.HAlign(HAlign_Left)
 		[
diff --git a/Engine/Source/Editor/Kismet/Private/SMyBlueprint.cpp b/Engine/Source/Editor/Kismet/Private/SMyBlueprint.cpp
--- a/Engine/Source/Editor/Kismet/Private/SMyBlueprint.cpp
+++ b/Engine/Source/Editor/Kismet/Private/SMyBlueprint.cpp
@@ -2547,7 +2547,7 @@ TSharedPtr<SWidget> SMyBlueprint::OnContextMenuOpening()
 				if (Schema->AllowsFunctionVariants())
 				{
 					MenuBuilder.AddMenuEntry(
-						LOCTEXT("MyBlueprint_Add_Func_Variant", "Add Variant"), FText(), FSlateIcon(),
+						LOCTEXT("MyBlueprint_Add_Func_Variant", "Add Variant (Experimental)"), FText(), FSlateIcon(),
 						FExecuteAction::CreateLambda([BlueprintEditor, Graph]()
 						{
 							BlueprintEditor->AddNewFunctionVariant(Graph->EdGraph);
