RigVM: Preload external dependencies prior compilation to avoid async lock

#rnx

diff --git a/Engine/Plugins/Runtime/RigVM/Source/RigVMDeveloper/Private/RigVMBlueprint.cpp b/Engine/Plugins/Runtime/RigVM/Source/RigVMDeveloper/Private/RigVMBlueprint.cpp
--- a/Engine/Plugins/Runtime/RigVM/Source/RigVMDeveloper/Private/RigVMBlueprint.cpp
+++ b/Engine/Plugins/Runtime/RigVM/Source/RigVMDeveloper/Private/RigVMBlueprint.cpp
@@ -403,6 +403,41 @@ void URigVMBlueprint::GetPreloadDependencies(TArray<UObject*>& OutDeps)
 	}
 }
 
+void URigVMBlueprint::PreloadDependenciesBeforeCompilation()
+{
+	Super::PreloadDependenciesBeforeCompilation();
+
+	TArray<URigVMGraph*> AllGraphs = GetAllModels();
+
+	for(int32 GraphIndex = 0; GraphIndex < AllGraphs.Num(); GraphIndex++)
+	{
+		const URigVMGraph* Graph = AllGraphs[GraphIndex];
+		
+		for(const URigVMNode* Node : Graph->GetNodes())
+		{
+			const TArray<URigVMPin*> AllPins = Node->GetAllPinsRecursively();
+			for(const URigVMPin* Pin : AllPins)
+			{
+				// force load a potential type dependency
+				(void)Pin->GetCPPTypeObject();
+			}
+
+			if(const URigVMFunctionReferenceNode* FunctionReferenceNode = Cast<URigVMFunctionReferenceNode>(Node))
+			{
+				// force load the referenced function
+				if(const URigVMLibraryNode* ReferencedNode = FunctionReferenceNode->LoadReferencedNode(true))
+				{
+					if(URigVMGraph* ReferencedGraph = ReferencedNode->GetContainedGraph())
+					{
+						// keep recursing into the function to load additional dependencies
+						AllGraphs.AddUnique(ReferencedGraph);
+					}
+				}
+			}
+		}
+	}
+}
+
 FRigVMClient* URigVMBlueprint::GetRigVMClient()
 {
 	return &RigVMClient;
diff --git a/Engine/Plugins/Runtime/RigVM/Source/RigVMDeveloper/Private/RigVMModel/Nodes/RigVMFunctionReferenceNode.cpp b/Engine/Plugins/Runtime/RigVM/Source/RigVMDeveloper/Private/RigVMModel/Nodes/RigVMFunctionReferenceNode.cpp
--- a/Engine/Plugins/Runtime/RigVM/Source/RigVMDeveloper/Private/RigVMModel/Nodes/RigVMFunctionReferenceNode.cpp
+++ b/Engine/Plugins/Runtime/RigVM/Source/RigVMDeveloper/Private/RigVMModel/Nodes/RigVMFunctionReferenceNode.cpp
@@ -1,6 +1,8 @@
 // Copyright Epic Games, Inc. All Rights Reserved.
 
 #include "RigVMModel/Nodes/RigVMFunctionReferenceNode.h"
+
+#include "RigVMBlueprint.h"
 #include "RigVMModel/RigVMFunctionLibrary.h"
 #include "RigVMCore/RigVMGraphFunctionDefinition.h"
 #include "RigVMCore/RigVMGraphFunctionHost.h"
@@ -350,12 +352,19 @@ bool URigVMFunctionReferenceNode::IsReferencedNodeLoaded() const
 	return ReferencedFunctionHeader.LibraryPointer.GetNodeSoftPath().ResolveObject() != nullptr;
 }
 
-URigVMLibraryNode* URigVMFunctionReferenceNode::LoadReferencedNode() const
+URigVMLibraryNode* URigVMFunctionReferenceNode::LoadReferencedNode(bool bForce) const
 {
-	FSoftObjectPath SoftObjectPath = ReferencedFunctionHeader.LibraryPointer.GetNodeSoftPath();
+	const FSoftObjectPath SoftObjectPath = ReferencedFunctionHeader.LibraryPointer.GetNodeSoftPath();
 	UObject* LibraryNode = SoftObjectPath.ResolveObject();
 	if (!LibraryNode)
 	{
+		if(const URigVMBlueprint* Blueprint = GetTypedOuter<URigVMBlueprint>())
+		{
+			if(Blueprint->bBeingCompiled && !bForce)
+			{
+				return nullptr;
+			}
+		}
 		LibraryNode = SoftObjectPath.TryLoad();
 	}
 	return Cast<URigVMLibraryNode>(LibraryNode);
diff --git a/Engine/Plugins/Runtime/RigVM/Source/RigVMDeveloper/Public/RigVMBlueprint.h b/Engine/Plugins/Runtime/RigVM/Source/RigVMDeveloper/Public/RigVMBlueprint.h
--- a/Engine/Plugins/Runtime/RigVM/Source/RigVMDeveloper/Public/RigVMBlueprint.h
+++ b/Engine/Plugins/Runtime/RigVM/Source/RigVMDeveloper/Public/RigVMBlueprint.h
@@ -287,6 +287,7 @@ public:
 	virtual void PreSave(FObjectPreSaveContext ObjectSaveContext) override;
 	virtual void PostSaveRoot(FObjectPostSaveRootContext ObjectSaveContext) override;
 	virtual void PostLoad() override;
+	virtual void PreloadDependenciesBeforeCompilation() override;
 #if WITH_EDITORONLY_DATA
 	static void DeclareConstructClasses(TArray<FTopLevelAssetPath>& OutConstructClasses, const UClass* SpecificSubclass);
 #endif
diff --git a/Engine/Plugins/Runtime/RigVM/Source/RigVMDeveloper/Public/RigVMModel/Nodes/RigVMFunctionReferenceNode.h b/Engine/Plugins/Runtime/RigVM/Source/RigVMDeveloper/Public/RigVMModel/Nodes/RigVMFunctionReferenceNode.h
--- a/Engine/Plugins/Runtime/RigVM/Source/RigVMDeveloper/Public/RigVMModel/Nodes/RigVMFunctionReferenceNode.h
+++ b/Engine/Plugins/Runtime/RigVM/Source/RigVMDeveloper/Public/RigVMModel/Nodes/RigVMFunctionReferenceNode.h
@@ -42,7 +42,7 @@ public:
 
 	bool IsReferencedFunctionHostLoaded() const;
 	bool IsReferencedNodeLoaded() const;
-	URigVMLibraryNode* LoadReferencedNode() const;
+	URigVMLibraryNode* LoadReferencedNode(bool bForce = false) const;
 
 	// Variable remapping
 	bool RequiresVariableRemapping() const;
diff --git a/Engine/Source/Editor/UnrealEd/Private/Kismet2/BlueprintEditorUtils.cpp b/Engine/Source/Editor/UnrealEd/Private/Kismet2/BlueprintEditorUtils.cpp
--- a/Engine/Source/Editor/UnrealEd/Private/Kismet2/BlueprintEditorUtils.cpp
+++ b/Engine/Source/Editor/UnrealEd/Private/Kismet2/BlueprintEditorUtils.cpp
@@ -1280,6 +1280,8 @@ void FBlueprintEditorUtils::PreloadBlueprintSpecificData(UBlueprint* Blueprint)
 	{
 		K2Node->PreloadRequiredAssets();
 	}
+
+	Blueprint->PreloadDependenciesBeforeCompilation();
 }
 
 void FBlueprintEditorUtils::LinkExternalDependencies(UBlueprint* Blueprint)
diff --git a/Engine/Source/Runtime/Engine/Classes/Engine/Blueprint.h b/Engine/Source/Runtime/Engine/Classes/Engine/Blueprint.h
--- a/Engine/Source/Runtime/Engine/Classes/Engine/Blueprint.h
+++ b/Engine/Source/Runtime/Engine/Classes/Engine/Blueprint.h
@@ -1147,6 +1147,11 @@ public:
 	 */
 	virtual bool ShouldBeMarkedDirtyUponTransaction() const { return true; }
 
+	/**
+	 * Gives the blueprint a chance to load dependencies if may have to get before compilation
+	 */
+	virtual void PreloadDependenciesBeforeCompilation() {}
+
 #if WITH_EDITOR
 private:
 
