RigVM: Add missing double to/from float casts

#rnx

diff --git a/Engine/Plugins/Runtime/RigVM/Source/RigVM/Private/RigVMFunctions/Math/RigVMFunction_MathDouble.cpp b/Engine/Plugins/Runtime/RigVM/Source/RigVM/Private/RigVMFunctions/Math/RigVMFunction_MathDouble.cpp
--- a/Engine/Plugins/Runtime/RigVM/Source/RigVM/Private/RigVMFunctions/Math/RigVMFunction_MathDouble.cpp
+++ b/Engine/Plugins/Runtime/RigVM/Source/RigVM/Private/RigVMFunctions/Math/RigVMFunction_MathDouble.cpp
@@ -125,6 +125,11 @@ FRigVMFunction_MathDoubleToInt_Execute()
 	Result = FMath::FloorToInt(Value);
 }
 
+FRigVMFunction_MathDoubleToFloat_Execute()
+{
+	Result = (float)Value;
+}
+
 FRigVMFunction_MathDoubleSign_Execute()
 {
 	Result = Value < 0.0 ? -1.0 : 1.0;
diff --git a/Engine/Plugins/Runtime/RigVM/Source/RigVM/Private/RigVMFunctions/Math/RigVMFunction_MathFloat.cpp b/Engine/Plugins/Runtime/RigVM/Source/RigVM/Private/RigVMFunctions/Math/RigVMFunction_MathFloat.cpp
--- a/Engine/Plugins/Runtime/RigVM/Source/RigVM/Private/RigVMFunctions/Math/RigVMFunction_MathFloat.cpp
+++ b/Engine/Plugins/Runtime/RigVM/Source/RigVM/Private/RigVMFunctions/Math/RigVMFunction_MathFloat.cpp
@@ -125,6 +125,11 @@ FRigVMFunction_MathFloatToInt_Execute()
 	Result = FMath::FloorToInt(Value);
 }
 
+FRigVMFunction_MathFloatToDouble_Execute()
+{
+	Result = Value;
+}
+
 FRigVMFunction_MathFloatSign_Execute()
 {
 	Result = Value < 0.f ? -1.f : 1.f;
diff --git a/Engine/Plugins/Runtime/RigVM/Source/RigVM/Public/RigVMFunctions/Math/RigVMFunction_MathDouble.h b/Engine/Plugins/Runtime/RigVM/Source/RigVM/Public/RigVMFunctions/Math/RigVMFunction_MathDouble.h
--- a/Engine/Plugins/Runtime/RigVM/Source/RigVM/Public/RigVMFunctions/Math/RigVMFunction_MathDouble.h
+++ b/Engine/Plugins/Runtime/RigVM/Source/RigVM/Public/RigVMFunctions/Math/RigVMFunction_MathDouble.h
@@ -411,6 +411,30 @@ struct RIGVM_API FRigVMFunction_MathDoubleToInt : public FRigVMFunction_MathDoub
 	int32 Result;
 };
 
+/**
+ * Returns the double cast to a float
+ */
+USTRUCT(meta=(DisplayName="To Float", TemplateName="Cast", ExecuteContext="FRigVMExecuteContext"))
+struct RIGVM_API FRigVMFunction_MathDoubleToFloat : public FRigVMFunction_MathDoubleBase
+{
+	GENERATED_BODY()
+
+	FRigVMFunction_MathDoubleToFloat()
+	{
+		Value = 0.0;
+		Result = 0.f;
+	}
+
+	RIGVM_METHOD()
+	virtual void Execute() override;
+
+	UPROPERTY(meta=(Input))
+	double Value;
+
+	UPROPERTY(meta=(Output))
+	float Result;
+};
+
 /**
  * Returns the sign of the value (+1 for >= 0.0, -1 for < 0.0)
  */
diff --git a/Engine/Plugins/Runtime/RigVM/Source/RigVM/Public/RigVMFunctions/Math/RigVMFunction_MathFloat.h b/Engine/Plugins/Runtime/RigVM/Source/RigVM/Public/RigVMFunctions/Math/RigVMFunction_MathFloat.h
--- a/Engine/Plugins/Runtime/RigVM/Source/RigVM/Public/RigVMFunctions/Math/RigVMFunction_MathFloat.h
+++ b/Engine/Plugins/Runtime/RigVM/Source/RigVM/Public/RigVMFunctions/Math/RigVMFunction_MathFloat.h
@@ -410,6 +410,30 @@ struct RIGVM_API FRigVMFunction_MathFloatToInt : public FRigVMFunction_MathFloat
 	int32 Result;
 };
 
+/**
+ * Returns the float cast to a double
+ */
+USTRUCT(meta=(DisplayName="To Double", TemplateName="Cast", ExecuteContext="FRigVMExecuteContext"))
+struct RIGVM_API FRigVMFunction_MathFloatToDouble : public FRigVMFunction_MathFloatBase
+{
+	GENERATED_BODY()
+
+	FRigVMFunction_MathFloatToDouble()
+	{
+		Value = 0.f;
+		Result = 0.0;
+	}
+
+	RIGVM_METHOD()
+	virtual void Execute() override;
+
+	UPROPERTY(meta=(Input))
+	float Value;
+
+	UPROPERTY(meta=(Output))
+	double Result;
+};
+
 /**
  * Returns the sign of the value (+1 for >= 0.f, -1 for < 0.f)
  */
