[RigVM Settings] moved default tags into a config file such that external plugins adding additional tags via config files don't override/remove default tags

diff --git a/Engine/Plugins/Runtime/RigVM/Config/Editor.ini b/Engine/Plugins/Runtime/RigVM/Config/Editor.ini
new file mode 100644
--- /dev/null
+++ b/Engine/Plugins/Runtime/RigVM/Config/Editor.ini
@@ -0,0 +1,3 @@
+﻿[/Script/RigVM.RigVMProjectSettings]
++VariantTags=(Name="Outdated",Label="Outdated",ToolTip=NSLOCTEXT("RigVMSettings", "OutdatedToolTip", "This item is outdated and should no longer be used."),Color=(R=1.0,G=0.0,B=0.0,A=1.0),bShowInUserInterface=True,bMarksSubjectAsInvalid=True)
++VariantTags=(Name="Stable",Label="Stable",ToolTip=NSLOCTEXT("RigVMSettings", "StableToolTip", "This item is stable and ready to use."),Color=(R=0.0,G=1.0,B=0.0,A=1.0),bShowInUserInterface=True,bMarksSubjectAsInvalid=False)
diff --git a/Engine/Plugins/Runtime/RigVM/Source/RigVM/Private/RigVMSettings.cpp b/Engine/Plugins/Runtime/RigVM/Source/RigVM/Private/RigVMSettings.cpp
--- a/Engine/Plugins/Runtime/RigVM/Source/RigVM/Private/RigVMSettings.cpp
+++ b/Engine/Plugins/Runtime/RigVM/Source/RigVM/Private/RigVMSettings.cpp
@@ -17,8 +17,9 @@ URigVMEditorSettings::URigVMEditorSettings(const FObjectInitializer& ObjectIniti
 URigVMProjectSettings::URigVMProjectSettings(const FObjectInitializer& ObjectInitializer)
 	: Super(ObjectInitializer)
 {
-	VariantTags.Emplace(TEXT("Outdated"), TEXT("Outdated"), LOCTEXT("OutdatedToolTip", "This item is outdated and should no longer be used."), FLinearColor::Red, true, true);
-	VariantTags.Emplace(TEXT("Stable"), TEXT("Stable"), LOCTEXT("StableToolTip", "This item is stable and ready to use."), FLinearColor::Green, true);
+	// Default tags are now added via Engine/Plugins/Runtime/RigVM/Config/Editor.ini
+	// to avoid external plugins overriding default tags when they only want to append new tags
+	// See UObject::LoadConfig for details
 }
 
 #undef LOCTEXT_NAMESPACE // RigVMSettings
