RigVM: Added visual logging support

[RN] Added visual logging support to Control Rig and RigVM assets

diff --git a/Engine/Plugins/Experimental/AnimNext/Source/AnimNext/Private/Module/ModuleTickFunction.cpp b/Engine/Plugins/Experimental/AnimNext/Source/AnimNext/Private/Module/ModuleTickFunction.cpp
--- a/Engine/Plugins/Experimental/AnimNext/Source/AnimNext/Private/Module/ModuleTickFunction.cpp
+++ b/Engine/Plugins/Experimental/AnimNext/Source/AnimNext/Private/Module/ModuleTickFunction.cpp
@@ -158,6 +158,7 @@ void FModuleEventTickFunction::Run(float InDeltaTime)
 
 		// RigVM setup
 		AnimNextContext.SetDeltaTime(InDeltaTime);
+		AnimNextContext.SetOwningObject(ModuleInstance->Object);
 
 #if UE_ENABLE_DEBUG_DRAWING
 		AnimNextContext.SetDrawInterface(&ModuleInstance->GetDebugDrawInterface());
diff --git a/Engine/Plugins/Runtime/RigVM/Source/RigVM/Private/RigVMCore/RigVMExecuteContext.cpp b/Engine/Plugins/Runtime/RigVM/Source/RigVM/Private/RigVMCore/RigVMExecuteContext.cpp
--- a/Engine/Plugins/Runtime/RigVM/Source/RigVM/Private/RigVMCore/RigVMExecuteContext.cpp
+++ b/Engine/Plugins/Runtime/RigVM/Source/RigVM/Private/RigVMCore/RigVMExecuteContext.cpp
@@ -10,17 +10,37 @@ TAutoConsoleVariable<bool> CVarRigVMReportAllMessages(
 	false,
 	TEXT("Report all log messages, even when no logging procedure has been setup."));
 
+void FRigVMExecuteContext::SetOwningObject(const UObject* InOwningObject)
+{
+	OwningObject = InOwningObject;
+	OwningActor = nullptr;
+	World = nullptr;
+	ToWorldSpaceTransform = FTransform::Identity;
+
+	if(InOwningObject)
+	{
+		if(const UActorComponent* ActorComponent = Cast<UActorComponent>(InOwningObject))
+		{
+			SetOwningActor(ActorComponent->GetOwner());
+		}
+		else
+		{
+			SetWorld(OwningObject->GetWorld());
+		}
+	}
+}
+
 void FRigVMExecuteContext::SetOwningComponent(const USceneComponent* InOwningComponent)
 {
-	OwningComponent = InOwningComponent;
+	OwningObject = InOwningComponent;
 	OwningActor = nullptr;
 	World = nullptr;
 	ToWorldSpaceTransform = FTransform::Identity;
 	
-	if(OwningComponent)
+	if(InOwningComponent)
 	{
-		ToWorldSpaceTransform = OwningComponent->GetComponentToWorld();
-		SetOwningActor(OwningComponent->GetOwner());
+		ToWorldSpaceTransform = InOwningComponent->GetComponentToWorld();
+		SetOwningActor(InOwningComponent->GetOwner());
 	}
 }
 
diff --git a/Engine/Plugins/Runtime/RigVM/Source/RigVM/Private/RigVMFunctions/Debug/RigVMFunction_VisualLog.cpp b/Engine/Plugins/Runtime/RigVM/Source/RigVM/Private/RigVMFunctions/Debug/RigVMFunction_VisualLog.cpp
new file mode 100644
--- /dev/null
+++ b/Engine/Plugins/Runtime/RigVM/Source/RigVM/Private/RigVMFunctions/Debug/RigVMFunction_VisualLog.cpp
@@ -0,0 +1,87 @@
+﻿// Copyright Epic Games, Inc. All Rights Reserved.
+
+#include "RigVMFunctions/Debug/RigVMFunction_VisualLog.h"
+#include "VisualLogger/VisualLogger.h"
+
+namespace UE::RigVM::Private
+{
+	constexpr ELogVerbosity::Type DefaultVerbosity = ELogVerbosity::Display;
+}
+
+FRigVMFunction_VisualLogText_Execute()
+{
+#if ENABLE_VISUAL_LOG
+	FVisualLogger::CategorizedLogf(ExecuteContext.GetOwningObject(), Category, UE::RigVM::Private::DefaultVerbosity, TEXT("%s"), *Text);
+#endif
+}
+
+FRigVMFunction_VisualLogLocation_Execute()
+{
+#if ENABLE_VISUAL_LOG
+	FVisualLogger::LocationLogf(ExecuteContext.GetOwningObject(), Category, UE::RigVM::Private::DefaultVerbosity, Location, static_cast<uint16>(Radius), ObjectColor.ToFColor(true), TEXT("%s"), *Text);
+#endif
+}
+
+FRigVMFunction_VisualLogSphere_Execute()
+{
+#if ENABLE_VISUAL_LOG
+	FVisualLogger::SphereLogf(ExecuteContext.GetOwningObject(), Category, UE::RigVM::Private::DefaultVerbosity, Center, Radius, ObjectColor.ToFColor(true), bWireframe, TEXT("%s"), *Text);
+#endif
+}
+
+FRigVMFunction_VisualLogCone_Execute()
+{
+#if ENABLE_VISUAL_LOG
+	FVisualLogger::ConeLogf(ExecuteContext.GetOwningObject(), Category, UE::RigVM::Private::DefaultVerbosity, Origin, Direction, Length, Angle, ObjectColor.ToFColor(true), bWireframe, TEXT("%s"), *Text);
+#endif
+}
+
+FRigVMFunction_VisualLogCylinder_Execute()
+{
+#if ENABLE_VISUAL_LOG
+	FVisualLogger::CylinderLogf(ExecuteContext.GetOwningObject(), Category, UE::RigVM::Private::DefaultVerbosity, Start, End, Radius, ObjectColor.ToFColor(true), bWireframe, TEXT("%s"), *Text);
+#endif
+}
+
+FRigVMFunction_VisualLogCapsule_Execute()
+{
+#if ENABLE_VISUAL_LOG
+	FVisualLogger::CapsuleLogf(ExecuteContext.GetOwningObject(), Category, UE::RigVM::Private::DefaultVerbosity, Base, HalfHeight, Radius, Rotation, ObjectColor.ToFColor(true), bWireframe, TEXT("%s"), *Text);
+#endif
+}
+
+FRigVMFunction_VisualLogBox_Execute()
+{
+#if ENABLE_VISUAL_LOG
+	FVisualLogger::BoxLogf(ExecuteContext.GetOwningObject(), Category, UE::RigVM::Private::DefaultVerbosity, Box, FMatrix::Identity, ObjectColor.ToFColor(true), bWireframe, TEXT("%s"), *Text);
+#endif
+}
+
+FRigVMFunction_VisualLogOrientedBox_Execute()
+{
+#if ENABLE_VISUAL_LOG
+	FVisualLogger::BoxLogf(ExecuteContext.GetOwningObject(), Category, UE::RigVM::Private::DefaultVerbosity, Box, Transform.ToMatrixWithScale(), ObjectColor.ToFColor(true), bWireframe, TEXT("%s"), *Text);
+#endif
+}
+
+FRigVMFunction_VisualLogArrow_Execute()
+{
+#if ENABLE_VISUAL_LOG
+	FVisualLogger::ArrowLineLogf(ExecuteContext.GetOwningObject(), Category, UE::RigVM::Private::DefaultVerbosity, SegmentStart, SegmentEnd, ObjectColor.ToFColor(true), ArrowHeadSize, TEXT("%s"), *Text);
+#endif
+}
+
+FRigVMFunction_VisualLogCircle_Execute()
+{
+#if ENABLE_VISUAL_LOG
+	FVisualLogger::DiscLogf(ExecuteContext.GetOwningObject(), Category, UE::RigVM::Private::DefaultVerbosity, Center, UpAxis, Radius, ObjectColor.ToFColor(true), (uint16)Thickness, bWireframe, TEXT("%s"), *Text);
+#endif
+}
+
+FRigVMFunction_VisualLogSegment_Execute()
+{
+#if ENABLE_VISUAL_LOG
+	FVisualLogger::SegmentLogf(ExecuteContext.GetOwningObject(), Category, UE::RigVM::Private::DefaultVerbosity, SegmentStart, SegmentEnd, ObjectColor.ToFColor(true), (uint16)Thickness, TEXT("%s"), *Text);
+#endif
+}
+
diff --git a/Engine/Plugins/Runtime/RigVM/Source/RigVM/Public/RigVMCore/RigVMExecuteContext.h b/Engine/Plugins/Runtime/RigVM/Source/RigVM/Public/RigVMCore/RigVMExecuteContext.h
--- a/Engine/Plugins/Runtime/RigVM/Source/RigVM/Public/RigVMCore/RigVMExecuteContext.h
+++ b/Engine/Plugins/Runtime/RigVM/Source/RigVM/Public/RigVMCore/RigVMExecuteContext.h
@@ -281,7 +281,7 @@ struct FRigVMExecuteContext
 		, DrawInterfacePtr(nullptr)
 		, DrawContainerPtr(nullptr)
 		, ToWorldSpaceTransform(FTransform::Identity)
-		, OwningComponent(nullptr)
+		, OwningObject(nullptr)
 		, OwningActor(nullptr)
 		, World(nullptr)
 		, Traits()
@@ -345,11 +345,14 @@ struct FRigVMExecuteContext
 	/** The current transform going from rig (global) space to world space */
 	const FTransform& GetToWorldSpaceTransform() const { return ToWorldSpaceTransform; };
 
+	/** The current object this VM is owned by */
+	const UObject* GetOwningObject() const { return OwningObject; }
+
 	/** The current component this VM is owned by */
-	const USceneComponent* GetOwningComponent() const { return OwningComponent; }
-	
+	const USceneComponent* GetOwningComponent() const { return Cast<USceneComponent>(OwningObject); }
+
 	/** The current component this VM is owned by */
-	USceneComponent* GetMutableOwningComponent() const { return const_cast<USceneComponent*>(OwningComponent); }
+	USceneComponent* GetMutableOwningComponent() const { return const_cast<USceneComponent*>(Cast<USceneComponent>(OwningObject)); }
 
 	/** The current actor this VM is owned by */
 	const AActor* GetOwningActor() const { return OwningActor; }
@@ -357,6 +360,8 @@ struct FRigVMExecuteContext
 	/** The world this VM is running in */
 	const UWorld* GetWorld() const { return World; }
 
+	RIGVM_API void SetOwningObject(const UObject* InOwningObject);
+
 	RIGVM_API void SetOwningComponent(const USceneComponent* InOwningComponent);
 
 	RIGVM_API void SetOwningActor(const AActor* InActor);
@@ -514,8 +519,8 @@ protected:
 	/** The current transform going from rig (global) space to world space */
 	FTransform ToWorldSpaceTransform;
 
-	/** The current component this VM is owned by */
-	const USceneComponent* OwningComponent;
+	/** The current object (usually a component) this VM is owned by */
+	const UObject* OwningObject;
 
 	/** The current actor this VM is owned by */
 	const AActor* OwningActor;
diff --git a/Engine/Plugins/Runtime/RigVM/Source/RigVM/Public/RigVMFunctions/Debug/RigVMFunction_VisualLog.h b/Engine/Plugins/Runtime/RigVM/Source/RigVM/Public/RigVMFunctions/Debug/RigVMFunction_VisualLog.h
new file mode 100644
--- /dev/null
+++ b/Engine/Plugins/Runtime/RigVM/Source/RigVM/Public/RigVMFunctions/Debug/RigVMFunction_VisualLog.h
@@ -0,0 +1,265 @@
+﻿// Copyright Epic Games, Inc. All Rights Reserved.
+
+#pragma once
+
+#include "RigVMFunction_DebugBase.h"
+#include "RigVMFunction_VisualLog.generated.h"
+
+/** Base RigVMFunction used for Visual Logger output */
+USTRUCT(meta=(Abstract, Keywords = "Draw,Point"))
+struct FRigVMFunction_VisualLogBase : public FRigVMFunction_DebugBaseMutable
+{
+	GENERATED_BODY()
+
+	/** The text to log */ 
+	UPROPERTY(meta=(Input))
+	FString Text;
+
+	/** The category of the logged text */
+	UPROPERTY(meta=(Input))
+	FName Category = TEXT("VisLogRigVM");
+};
+
+/** Logs simple text string with Visual Logger - recording for Visual Logs has to be enabled to record this data */
+USTRUCT(meta=(DisplayName = "Visual Log Text", Keywords = "Draw,String"))
+struct FRigVMFunction_VisualLogText : public FRigVMFunction_VisualLogBase
+{
+	GENERATED_BODY()
+
+	RIGVM_METHOD()
+	virtual void Execute() override;
+};
+
+/** Base RigVMFunction for visual logging of objects */ 
+USTRUCT(meta=(Abstract, Keywords = "Draw,String"))
+struct FRigVMFunction_VisualLogObject : public FRigVMFunction_VisualLogBase
+{
+	GENERATED_BODY()
+
+	/** The color of the logged object */
+	UPROPERTY(meta=(Input))
+	FLinearColor ObjectColor = FLinearColor::Blue;
+};
+
+/** Logs location as sphere with given radius - recording for Visual Logs has to be enabled to record this data */
+USTRUCT(meta=(DisplayName = "Visual Log Location", Keywords = "Draw,String,Sphere"))
+struct FRigVMFunction_VisualLogLocation : public FRigVMFunction_VisualLogObject
+{
+	GENERATED_BODY()
+
+	RIGVM_METHOD()
+	virtual void Execute() override;
+
+	/** The location to log */
+	UPROPERTY(meta=(Input))
+	FVector Location;
+
+	/** The radius of the sphere to log the location with */
+	UPROPERTY(meta=(Input))
+	float Radius = 10.0f;
+};
+
+/** Base RigVMFunction for visual logging of objects that can be wirefame */ 
+USTRUCT(meta=(Abstract))
+struct FRigVMFunction_VisualLogWireframeOptional : public FRigVMFunction_VisualLogObject
+{
+	GENERATED_BODY()
+
+	/** Whether to display as wireframe */
+	UPROPERTY(meta=(Input))
+	bool bWireframe = false;
+};
+
+/** Logs sphere shape - recording for Visual Logs has to be enabled to record this data */
+USTRUCT(meta=(DisplayName = "Visual Log Sphere", Keywords = "Draw,String"))
+struct FRigVMFunction_VisualLogSphere : public FRigVMFunction_VisualLogWireframeOptional
+{
+	GENERATED_BODY()
+
+	RIGVM_METHOD()
+	virtual void Execute() override;
+
+	/** The centre of the sphere */
+	UPROPERTY(meta=(Input))
+	FVector Center;
+
+	/** The radius of the sphere */
+	UPROPERTY(meta=(Input))
+	float Radius = 10.0f;
+};
+
+/** Logs cone shape - recording for Visual Logs has to be enabled to record this data */
+USTRUCT(meta=(DisplayName = "Visual Log Cone", Keywords = "Draw,String"))
+struct FRigVMFunction_VisualLogCone : public FRigVMFunction_VisualLogWireframeOptional
+{
+	GENERATED_BODY()
+
+	RIGVM_METHOD()
+	virtual void Execute() override;
+
+	/** The origin of the cone */
+	UPROPERTY(meta=(Input))
+	FVector Origin;
+
+	/** The direction of the cone */
+	UPROPERTY(meta=(Input))
+	FVector Direction = FVector(0.0f, 0.0f, 1.0f);
+
+	/** The length of the cone */
+	UPROPERTY(meta=(Input))
+	float Length = 10.0f;
+
+	/** The angle of the cone */
+	UPROPERTY(meta=(Input))
+	float Angle = 10.0f;
+};
+
+/** Logs cylinder shape - recording for Visual Logs has to be enabled to record this data */
+USTRUCT(meta=(DisplayName = "Visual Log Cylinder", Keywords = "Draw,String"))
+struct FRigVMFunction_VisualLogCylinder : public FRigVMFunction_VisualLogWireframeOptional
+{
+	GENERATED_BODY()
+
+	RIGVM_METHOD()
+	virtual void Execute() override;
+
+	/** The start of the line segment forming the cylinder */
+	UPROPERTY(meta=(Input))
+	FVector Start;
+
+	/** The end of the line segment forming the cylinder */
+	UPROPERTY(meta=(Input))
+	FVector End;
+
+	/** The radius of the cylinder */
+	UPROPERTY(meta=(Input))
+	float Radius = 10.0f;
+};
+
+/** Logs capsule shape - recording for Visual Logs has to be enabled to record this data */
+USTRUCT(meta=(DisplayName = "Visual Log Capsule", Keywords = "Draw,String"))
+struct FRigVMFunction_VisualLogCapsule : public FRigVMFunction_VisualLogWireframeOptional
+{
+	GENERATED_BODY()
+
+	RIGVM_METHOD()
+	virtual void Execute() override;
+
+	/** The base or origin of the capsule */
+	UPROPERTY(meta=(Input))
+	FVector Base;
+
+	/** Half the height of the capsule */
+	UPROPERTY(meta=(Input))
+	float HalfHeight = 10.0f;
+
+	/** The radius of the capsule */
+	UPROPERTY(meta=(Input))
+	float Radius = 10.0f;
+
+	/** The orientation of the capsule */
+	UPROPERTY(meta=(Input))
+	FQuat Rotation;
+};
+
+/** Logs box shape - recording for Visual Logs has to be enabled to record this data */
+USTRUCT(meta=(DisplayName = "Visual Log Box", Keywords = "Draw,String"))
+struct FRigVMFunction_VisualLogBox : public FRigVMFunction_VisualLogWireframeOptional
+{
+	GENERATED_BODY()
+
+	RIGVM_METHOD()
+	virtual void Execute() override;
+
+	/** The box to draw */
+	UPROPERTY(meta=(Input))
+	FBox Box;
+};
+
+/** Logs oriented box shape - recording for Visual Logs has to be enabled to record this data */
+USTRUCT(meta=(DisplayName = "Visual Log Oriented Box", Keywords = "Draw,String"))
+struct FRigVMFunction_VisualLogOrientedBox : public FRigVMFunction_VisualLogWireframeOptional
+{
+	GENERATED_BODY()
+
+	RIGVM_METHOD()
+	virtual void Execute() override;
+
+	/** The box to draw */
+	UPROPERTY(meta=(Input))
+	FBox Box;
+
+	/** The transform of the box */
+	UPROPERTY(meta=(Input))
+	FTransform Transform;
+};
+
+/** Logs arrow - recording for Visual Logs has to be enabled to record this data */
+USTRUCT(meta=(DisplayName = "Visual Log Arrow", Keywords = "Draw,String,Direction"))
+struct FRigVMFunction_VisualLogArrow : public FRigVMFunction_VisualLogObject
+{
+	GENERATED_BODY()
+
+	RIGVM_METHOD()
+	virtual void Execute() override;
+
+	/** The start of the arrow */
+	UPROPERTY(meta=(Input))
+	FVector SegmentStart;
+
+	/** The end of the arrow */
+	UPROPERTY(meta=(Input))
+	FVector SegmentEnd;
+
+	/** The size of the arrow head */
+	UPROPERTY(meta=(Input))
+	float ArrowHeadSize = 8.0f;
+};
+
+/** Logs circle - recording for Visual Logs has to be enabled to record this data */
+USTRUCT(meta=(DisplayName = "Visual Log Circle", Keywords = "Draw,String,Disc"))
+struct FRigVMFunction_VisualLogCircle : public FRigVMFunction_VisualLogWireframeOptional
+{
+	GENERATED_BODY()
+
+	RIGVM_METHOD()
+	virtual void Execute() override;
+
+	/** The center of the circle */
+	UPROPERTY(meta=(Input))
+	FVector Center;
+
+	/** The up axis/normal of the circle's plane */
+	UPROPERTY(meta=(Input))
+	FVector UpAxis = FVector(0.0f, 0.0f, 1.0f);
+
+	/** The radius of the circle */
+	UPROPERTY(meta=(Input))
+	float Radius = 10.0f;
+
+	/** The thickness the circle */
+	UPROPERTY(meta=(Input))
+	float Thickness = 0.0f;
+};
+
+/** Logs segment - recording for Visual Logs has to be enabled to record this data */
+USTRUCT(meta=(DisplayName = "Visual Log Segment", Keywords = "Draw,String,Line"))
+struct FRigVMFunction_VisualLogSegment : public FRigVMFunction_VisualLogObject
+{
+	GENERATED_BODY()
+
+	RIGVM_METHOD()
+	virtual void Execute() override;
+
+	/** The start of the line segment */
+	UPROPERTY(meta=(Input))
+	FVector SegmentStart;
+
+	/** The end of the line segment */
+	UPROPERTY(meta=(Input))
+	FVector SegmentEnd;
+
+	/** The thickness the circle */
+	UPROPERTY(meta=(Input))
+	float Thickness = 0.0f;
+};
