Add type information to RigVMPin tooltip text

diff --git a/Engine/Plugins/Runtime/RigVM/Source/RigVMDeveloper/Private/RigVMModel/Nodes/RigVMUnitNode.cpp b/Engine/Plugins/Runtime/RigVM/Source/RigVMDeveloper/Private/RigVMModel/Nodes/RigVMUnitNode.cpp
--- a/Engine/Plugins/Runtime/RigVM/Source/RigVMDeveloper/Private/RigVMModel/Nodes/RigVMUnitNode.cpp
+++ b/Engine/Plugins/Runtime/RigVM/Source/RigVMDeveloper/Private/RigVMModel/Nodes/RigVMUnitNode.cpp
@@ -146,7 +146,8 @@ FText URigVMUnitNode::GetToolTipTextForPin(const URigVMPin* InPin) const
 
 			if (PartIndex == Parts.Num() - 1)
 			{
-				return Property->GetToolTipText();
+				const FText ToolTipBody = Property->GetToolTipText();
+				return GetTypedToolTipText(InPin, ToolTipBody);
 			}
 
 			if (FStructProperty* StructProperty = CastField<FStructProperty>(Property))
diff --git a/Engine/Plugins/Runtime/RigVM/Source/RigVMDeveloper/Private/RigVMModel/RigVMNode.cpp b/Engine/Plugins/Runtime/RigVM/Source/RigVMDeveloper/Private/RigVMModel/RigVMNode.cpp
--- a/Engine/Plugins/Runtime/RigVM/Source/RigVMDeveloper/Private/RigVMModel/RigVMNode.cpp
+++ b/Engine/Plugins/Runtime/RigVM/Source/RigVMDeveloper/Private/RigVMModel/RigVMNode.cpp
@@ -613,7 +613,8 @@ FText URigVMNode::GetToolTipText() const
 
 FText URigVMNode::GetToolTipTextForPin(const URigVMPin* InPin) const
 {
-	return FText::FromName(InPin->GetFName());
+	const FText NodePinToolTipText = FText::FromName(InPin->GetFName());
+	return GetTypedToolTipText(InPin, NodePinToolTipText);
 }
 
 FString URigVMNode::GetOriginalDefaultValueForRootPin(const URigVMPin* InRootPin) const
@@ -999,6 +1000,36 @@ int32 URigVMNode::GetIndexInCategoryForPin(const FString& InPinPath) const
 	return INDEX_NONE;
 }
 
+FText URigVMNode::GetTypedToolTipText(const URigVMPin* InPin, const FText& ToolTipBody) const
+{
+	FText PinTypeText = FText::GetEmpty();
+
+	if (UObject* PinTypeObject = InPin->GetCPPTypeObject())
+	{
+		if (const UField* Field = Cast<const UField>(PinTypeObject))
+		{
+			PinTypeText = Field->GetDisplayNameText();
+		}
+		else
+		{
+			PinTypeText = FText::FromName(PinTypeObject->GetFName());
+		}
+	}
+	else
+	{
+		PinTypeText = FText::FromString(InPin->GetCPPType());
+	}
+
+	if (!ToolTipBody.IsEmptyOrWhitespace())
+	{
+		return FText::Format(INVTEXT("{0}\n{1}"), ToolTipBody, PinTypeText);
+	}
+	else
+	{
+		return PinTypeText;
+	}
+}
+
 URigVMLibraryNode* URigVMNode::FindFunctionForNode() const  
 {
 	const UObject* Subject = this;
diff --git a/Engine/Plugins/Runtime/RigVM/Source/RigVMDeveloper/Public/RigVMModel/RigVMNode.h b/Engine/Plugins/Runtime/RigVM/Source/RigVMDeveloper/Public/RigVMModel/RigVMNode.h
--- a/Engine/Plugins/Runtime/RigVM/Source/RigVMDeveloper/Public/RigVMModel/RigVMNode.h
+++ b/Engine/Plugins/Runtime/RigVM/Source/RigVMDeveloper/Public/RigVMModel/RigVMNode.h
@@ -390,6 +390,8 @@ public:
 	// returns the index of a pin within a category
 	virtual int32 GetIndexInCategoryForPin(const FString& InPinPath) const;
 
+	virtual FText GetTypedToolTipText(const URigVMPin* InPin, const FText& ToolTipBody) const;
+
 private:
 
 	static const inline TCHAR* NodeColorName = TEXT("NodeColor");
