Correctly initialize struct members in RigVM visual log units

#rnx

diff --git a/Engine/Plugins/Runtime/RigVM/Source/RigVM/Public/RigVMFunctions/Debug/RigVMFunction_VisualLog.h b/Engine/Plugins/Runtime/RigVM/Source/RigVM/Public/RigVMFunctions/Debug/RigVMFunction_VisualLog.h
--- a/Engine/Plugins/Runtime/RigVM/Source/RigVM/Public/RigVMFunctions/Debug/RigVMFunction_VisualLog.h
+++ b/Engine/Plugins/Runtime/RigVM/Source/RigVM/Public/RigVMFunctions/Debug/RigVMFunction_VisualLog.h
@@ -52,7 +52,7 @@ struct FRigVMFunction_VisualLogLocation : public FRigVMFunction_VisualLogObject
 
 	/** The location to log */
 	UPROPERTY(meta=(Input))
-	FVector Location;
+	FVector Location = FVector::ZeroVector;
 
 	/** The radius of the sphere to log the location with */
 	UPROPERTY(meta=(Input))
@@ -81,7 +81,7 @@ struct FRigVMFunction_VisualLogSphere : public FRigVMFunction_VisualLogWireframe
 
 	/** The centre of the sphere */
 	UPROPERTY(meta=(Input))
-	FVector Center;
+	FVector Center = FVector::ZeroVector;
 
 	/** The radius of the sphere */
 	UPROPERTY(meta=(Input))
@@ -99,7 +99,7 @@ struct FRigVMFunction_VisualLogCone : public FRigVMFunction_VisualLogWireframeOp
 
 	/** The origin of the cone */
 	UPROPERTY(meta=(Input))
-	FVector Origin;
+	FVector Origin = FVector::ZeroVector;
 
 	/** The direction of the cone */
 	UPROPERTY(meta=(Input))
@@ -125,11 +125,11 @@ struct FRigVMFunction_VisualLogCylinder : public FRigVMFunction_VisualLogWirefra
 
 	/** The start of the line segment forming the cylinder */
 	UPROPERTY(meta=(Input))
-	FVector Start;
+	FVector Start = FVector::ZeroVector;
 
 	/** The end of the line segment forming the cylinder */
 	UPROPERTY(meta=(Input))
-	FVector End;
+	FVector End = FVector::ZeroVector;
 
 	/** The radius of the cylinder */
 	UPROPERTY(meta=(Input))
@@ -147,7 +147,7 @@ struct FRigVMFunction_VisualLogCapsule : public FRigVMFunction_VisualLogWirefram
 
 	/** The base or origin of the capsule */
 	UPROPERTY(meta=(Input))
-	FVector Base;
+	FVector Base = FVector::ZeroVector;
 
 	/** Half the height of the capsule */
 	UPROPERTY(meta=(Input))
@@ -159,7 +159,7 @@ struct FRigVMFunction_VisualLogCapsule : public FRigVMFunction_VisualLogWirefram
 
 	/** The orientation of the capsule */
 	UPROPERTY(meta=(Input))
-	FQuat Rotation;
+	FQuat Rotation = FQuat::Identity;
 };
 
 /** Logs box shape - recording for Visual Logs has to be enabled to record this data */
@@ -173,7 +173,7 @@ struct FRigVMFunction_VisualLogBox : public FRigVMFunction_VisualLogWireframeOpt
 
 	/** The box to draw */
 	UPROPERTY(meta=(Input))
-	FBox Box;
+	FBox Box = FBox(ForceInit);
 };
 
 /** Logs oriented box shape - recording for Visual Logs has to be enabled to record this data */
@@ -187,7 +187,7 @@ struct FRigVMFunction_VisualLogOrientedBox : public FRigVMFunction_VisualLogWire
 
 	/** The box to draw */
 	UPROPERTY(meta=(Input))
-	FBox Box;
+	FBox Box = FBox(ForceInit);
 
 	/** The transform of the box */
 	UPROPERTY(meta=(Input))
@@ -205,11 +205,11 @@ struct FRigVMFunction_VisualLogArrow : public FRigVMFunction_VisualLogObject
 
 	/** The start of the arrow */
 	UPROPERTY(meta=(Input))
-	FVector SegmentStart;
+	FVector SegmentStart = FVector::ZeroVector;
 
 	/** The end of the arrow */
 	UPROPERTY(meta=(Input))
-	FVector SegmentEnd;
+	FVector SegmentEnd = FVector::ZeroVector;
 
 	/** The size of the arrow head */
 	UPROPERTY(meta=(Input))
@@ -227,7 +227,7 @@ struct FRigVMFunction_VisualLogCircle : public FRigVMFunction_VisualLogWireframe
 
 	/** The center of the circle */
 	UPROPERTY(meta=(Input))
-	FVector Center;
+	FVector Center = FVector::ZeroVector;
 
 	/** The up axis/normal of the circle's plane */
 	UPROPERTY(meta=(Input))
@@ -253,11 +253,11 @@ struct FRigVMFunction_VisualLogSegment : public FRigVMFunction_VisualLogObject
 
 	/** The start of the line segment */
 	UPROPERTY(meta=(Input))
-	FVector SegmentStart;
+	FVector SegmentStart = FVector::ZeroVector;
 
 	/** The end of the line segment */
 	UPROPERTY(meta=(Input))
-	FVector SegmentEnd;
+	FVector SegmentEnd = FVector::ZeroVector;
 
 	/** The thickness the circle */
 	UPROPERTY(meta=(Input))
