RigVM avoid crash when checking if a pin can link and one of the target pins is an unresolved wildcard pin

diff --git a/Engine/Plugins/Runtime/RigVM/Source/RigVMDeveloper/Private/RigVMModel/RigVMPin.cpp b/Engine/Plugins/Runtime/RigVM/Source/RigVMDeveloper/Private/RigVMModel/RigVMPin.cpp
--- a/Engine/Plugins/Runtime/RigVM/Source/RigVMDeveloper/Private/RigVMModel/RigVMPin.cpp
+++ b/Engine/Plugins/Runtime/RigVM/Source/RigVMDeveloper/Private/RigVMModel/RigVMPin.cpp
@@ -2165,6 +2165,19 @@ bool URigVMPin::CanLink(const URigVMPin* InSourcePin, const URigVMPin* InTargetP
 					TArray<int32> TargetPermutations = TargetTemplateNode->GetResolvedPermutationIndices(true);
 					const FRigVMTemplate* SourceTemplate = SourceTemplateNode->GetTemplate();
 					const FRigVMTemplate* TargetTemplate = TargetTemplateNode->GetTemplate();
+
+					ensureMsgf(SourceTemplate != nullptr, TEXT("Source Template can not be resolved. Might have a pin with a type not registered in the RigVM Registry."));
+					ensureMsgf(TargetTemplate != nullptr, TEXT("Target Template can not be resolved. Might have a pin with a type not registered in the RigVM Registry."));
+					if (SourceTemplate == nullptr || TargetTemplate == nullptr)
+					{
+						if (OutFailureReason)
+						{
+							*OutFailureReason = TEXT("One of the templates can not be resolved. Might have a pin with a type not registered in the RigVM Registry.");
+						}
+
+						return false;
+					}
+
 					const FRigVMTemplateArgument* SourceRootArgument = SourceTemplate->FindArgument(RootSourcePin->GetFName());
 					const FRigVMTemplateArgument* TargetRootArgument = TargetTemplate->FindArgument(RootTargetPin->GetFName());
 
