RigVM subpin orphaning for Traits

diff --git a/Engine/Plugins/Experimental/Workspace/Source/WorkspaceEditor/Private/Workspace.cpp b/Engine/Plugins/Experimental/Workspace/Source/WorkspaceEditor/Private/Workspace.cpp
--- a/Engine/Plugins/Experimental/Workspace/Source/WorkspaceEditor/Private/Workspace.cpp
+++ b/Engine/Plugins/Experimental/Workspace/Source/WorkspaceEditor/Private/Workspace.cpp
@@ -265,7 +265,7 @@ void UWorkspace::GetAssets(TArray<TObjectPtr<UObject>>& OutAssets) const
 void UWorkspace::GetAssetDataEntries(TArray<FAssetData>& OutAssetDataEntries) const
 {
 	FARFilter Filter;
-	Algo::Transform(AssetEntries, Filter.SoftObjectPaths, [](const UWorkspaceAssetEntry* AssetEntry) { return AssetEntry->Asset.ToSoftObjectPath(); });	
+	Algo::Transform(AssetEntries, Filter.SoftObjectPaths, [](const UWorkspaceAssetEntry* AssetEntry) { return AssetEntry != nullptr ? AssetEntry->Asset.ToSoftObjectPath() : nullptr; });
 	const IAssetRegistry& AssetRegistry = FAssetRegistryModule::GetRegistry();
 	AssetRegistry.GetAssets(Filter, OutAssetDataEntries);
 }
diff --git a/Engine/Plugins/Experimental/Workspace/Source/WorkspaceEditor/Private/WorkspaceEditor.cpp b/Engine/Plugins/Experimental/Workspace/Source/WorkspaceEditor/Private/WorkspaceEditor.cpp
--- a/Engine/Plugins/Experimental/Workspace/Source/WorkspaceEditor/Private/WorkspaceEditor.cpp
+++ b/Engine/Plugins/Experimental/Workspace/Source/WorkspaceEditor/Private/WorkspaceEditor.cpp
@@ -677,12 +677,15 @@ bool FWorkspaceEditor::AreAssetEntriesModified() const
 {
 	for(const UWorkspaceAssetEntry* Entry : Workspace->AssetEntries)
 	{
-		if(UObject* Asset = Entry->Asset.Get())
+		if (Entry != nullptr)
 		{
-			const UPackage* Package = Asset->GetOutermost();
-			if (Package->IsDirty() || Package->GetExternalPackages().ContainsByPredicate([](const UPackage* ExternalPackage)  { return ExternalPackage && ExternalPackage->IsDirty(); }))
+			if(UObject* Asset = Entry->Asset.Get())
 			{
-				return true;
+				const UPackage* Package = Asset->GetOutermost();
+				if (Package->IsDirty() || Package->GetExternalPackages().ContainsByPredicate([](const UPackage* ExternalPackage)  { return ExternalPackage && ExternalPackage->IsDirty(); }))
+				{
+					return true;
+				}
 			}
 		}
 	}
diff --git a/Engine/Plugins/Runtime/RigVM/Source/RigVMDeveloper/Private/EdGraph/RigVMEdGraphNode.cpp b/Engine/Plugins/Runtime/RigVM/Source/RigVMDeveloper/Private/EdGraph/RigVMEdGraphNode.cpp
--- a/Engine/Plugins/Runtime/RigVM/Source/RigVMDeveloper/Private/EdGraph/RigVMEdGraphNode.cpp
+++ b/Engine/Plugins/Runtime/RigVM/Source/RigVMDeveloper/Private/EdGraph/RigVMEdGraphNode.cpp
@@ -327,9 +327,18 @@ void URigVMEdGraphNode::RewireOldPinsToNewPins(TArray<UEdGraphPin*>& InOldPins,
 
 	for(UEdGraphPin* OldPin : InOldPins)
 	{
+		const FString Name = OldPin->PinName.ToString();
+		FString Left, Right;
+		if (!URigVMPin::SplitPinPathAtStart(Name, Left, Right))
+		{
+			Left = Name;
+		}
+		const FString OrphanName = FString::Printf(TEXT("%s.%s%s"), *Left, URigVMPin::OrphanPinPrefix, *Right);
+
 		for(UEdGraphPin* NewPin : InNewPins)
 		{
-			if(OldPin->PinName == NewPin->PinName && OldPin->Direction == NewPin->Direction)
+			if((OldPin->PinName == NewPin->PinName || NewPin->PinName == OrphanName)
+				&& OldPin->Direction == NewPin->Direction)
 			{
 				if (OldPin->PinType == NewPin->PinType ||
 					OldPin->PinType.PinSubCategoryObject == RigVMTypeUtils::GetWildCardCPPTypeObject() ||
@@ -342,6 +351,10 @@ void URigVMEdGraphNode::RewireOldPinsToNewPins(TArray<UEdGraphPin*>& InOldPins,
 					break;
 				}
 			}
+			else
+			{
+				continue;
+			}
 		}
 	}
 
diff --git a/Engine/Plugins/Runtime/RigVM/Source/RigVMDeveloper/Private/RigVMModel/RigVMController.cpp b/Engine/Plugins/Runtime/RigVM/Source/RigVMDeveloper/Private/RigVMModel/RigVMController.cpp
--- a/Engine/Plugins/Runtime/RigVM/Source/RigVMDeveloper/Private/RigVMModel/RigVMController.cpp
+++ b/Engine/Plugins/Runtime/RigVM/Source/RigVMDeveloper/Private/RigVMModel/RigVMController.cpp
@@ -17472,9 +17472,14 @@ void URigVMController::OrphanPins(const TArray<FRepopulatePinsNodeData>& NodesPi
 #if UE_RIGVMCONTROLLER_VERBOSE_REPOPULATE
 				UE_LOG(LogRigVMDeveloper, Display, TEXT("Orphaning pin '%s'."), *PinPath);
 #endif
-				check(Pin->IsRootPin());
+				const bool bIsTraitPin = Pin->GetRootPin()->IsTraitPin();
+				check(Pin->IsRootPin() || bIsTraitPin);
+
+				//const FString OrphanedName = FString::Printf(TEXT("%s%s"), URigVMPin::OrphanPinPrefix, *Pin->GetName());
+				const FString OrphanedName = bIsTraitPin
+					? FString::Printf(TEXT("%s%s.%s"), URigVMPin::OrphanPinPrefix, *Pin->GetParentPin()->GetName(), *Pin->GetName())
+					: FString::Printf(TEXT("%s%s"), URigVMPin::OrphanPinPrefix, *Pin->GetName());
 
-				const FString OrphanedName = FString::Printf(TEXT("%s%s"), URigVMPin::OrphanPinPrefix, *Pin->GetName());
 				if (!NodeData.Node->FindPin(OrphanedName))
 				{
 					URigVMPin* OrphanPin = NewObject<URigVMPin>(NodeData.Node, *OrphanedName);
@@ -17878,7 +17883,7 @@ void URigVMController::GenerateRepopulatePinLists(const FRigVMRegistry& Registry
 	{
 		const FString PinPath = NodeData.PreviousPinInfos.GetPinPath(Index);
 		const int32 NewIndex = NodeData.NewPinInfos.GetIndexFromPinPath(PinPath);
-		const int32 RootIndex = NodeData.PreviousPinInfos.GetRootIndex(Index);
+		int32 RootIndex = NodeData.PreviousPinInfos.GetRootIndex(Index);
 
 		if (NewIndex == INDEX_NONE)
 		{
@@ -17886,6 +17891,11 @@ void URigVMController::GenerateRepopulatePinLists(const FRigVMRegistry& Registry
 			{
 				if (URigVMPin* Pin = InNode->FindPin(PinPath))
 				{
+					const bool bIsTraitPin = Pin->GetRootPin()->IsTraitPin();
+					if (bIsTraitPin)
+					{
+						RootIndex = Index;
+					}
 					if (!Pin->GetLinks().IsEmpty())
 					{
 						NodeData.bRequireRecreateLinks = true;
@@ -17896,7 +17906,7 @@ void URigVMController::GenerateRepopulatePinLists(const FRigVMRegistry& Registry
 					{
 						if (!NodeData.PreviousPinsToOrphan.Contains(RootIndex))
 						{
-							URigVMPin* RootPin = Pin->GetRootPin();
+							URigVMPin* RootPin = bIsTraitPin ? Pin : Pin->GetRootPin();
 
 							if (RootPin->GetSourceLinks(true).Num() > 0 ||
 								RootPin->GetTargetLinks(true).Num() > 0)
@@ -18096,9 +18106,13 @@ void URigVMController::RepopulatePinsOnNode(const FRigVMRegistry& Registry, cons
 #if UE_RIGVMCONTROLLER_VERBOSE_REPOPULATE
 			UE_LOG(LogRigVMDeveloper, Display, TEXT("Orphaning pin '%s'."), *PinPath);
 #endif
-			check(Pin->IsRootPin());
+			const bool bIsTraitPin = Pin->GetRootPin()->IsTraitPin();
+			check(Pin->IsRootPin() || bIsTraitPin);
 
-			const FString OrphanedName = FString::Printf(TEXT("%s%s"), URigVMPin::OrphanPinPrefix, *Pin->GetName());
+			//const FString OrphanedName = FString::Printf(TEXT("%s%s"), URigVMPin::OrphanPinPrefix, *Pin->GetName());
+			const FString OrphanedName = bIsTraitPin
+				? FString::Printf(TEXT("%s%s.%s"), URigVMPin::OrphanPinPrefix, *Pin->GetParentPin()->GetName(), *Pin->GetName())
+				: FString::Printf(TEXT("%s%s"), URigVMPin::OrphanPinPrefix, *Pin->GetName());
 			if (InNode->FindPin(OrphanedName) == nullptr)
 			{
 				Pin->DisplayName = Pin->GetFName();
@@ -18316,7 +18330,8 @@ void URigVMController::RemovePinsDuringRepopulate(URigVMNode* InNode, TArray<URi
 	{
 		if(bSetupOrphanedPins && !Pin->IsExecuteContext())
 		{
-			URigVMPin* RootPin = Pin->GetRootPin();
+			const bool bIsTraitPin = Pin->GetRootPin()->IsTraitPin();
+			URigVMPin* RootPin = bIsTraitPin ? Pin : Pin->GetRootPin();
 			const FString OrphanedName = FString::Printf(TEXT("%s%s"), URigVMPin::OrphanPinPrefix, *RootPin->GetName());
 
 			URigVMPin* OrphanedRootPin = nullptr;
@@ -18332,7 +18347,7 @@ void URigVMController::RemovePinsDuringRepopulate(URigVMNode* InNode, TArray<URi
 
 			if(OrphanedRootPin == nullptr)
 			{
-				if(Pin->IsRootPin()) // if we are passing root pins we can reparent them directly
+				if(Pin->IsRootPin() || bIsTraitPin) // if we are passing root pins we can reparent them directly
 				{
 					check(RootPin->GetLinkedSourcePins(true).IsEmpty());
 					check(RootPin->GetLinkedTargetPins(true).IsEmpty());
@@ -21146,6 +21161,7 @@ TArray<URigVMController::FLinkedPath> URigVMController::RemapLinkedPaths(
 	{
 		RemappedLinkedPaths.Add(LinkedPath);
 		FLinkedPath& RemappedLinkedPath = RemappedLinkedPaths.Last();
+		RemappedLinkedPath.OriginalPinPathRepresentation = LinkedPath.GetPinPathRepresentation();
 		RemappedLinkedPath.SourcePinPath = RemapPinPath(LinkedPath.SourcePinPath, false);
 		RemappedLinkedPath.TargetPinPath = RemapPinPath(LinkedPath.TargetPinPath, true);
 
@@ -21233,7 +21249,7 @@ bool URigVMController::RestoreLinkedPaths(
 	
 	for(const FLinkedPath& LinkedPath : RemappedLinkedPaths)
 	{
-		const FString OriginalPinPathRepresentation = LinkedPath.GetPinPathRepresentation();
+		const FString& OriginalPinPathRepresentation = LinkedPath.OriginalPinPathRepresentation;
 
 		const uint32 Hash = GetTypeHash(OriginalPinPathRepresentation);
 		if(ProcessedLinks.Contains(Hash))
diff --git a/Engine/Plugins/Runtime/RigVM/Source/RigVMDeveloper/Private/RigVMModel/RigVMGraph.cpp b/Engine/Plugins/Runtime/RigVM/Source/RigVMDeveloper/Private/RigVMModel/RigVMGraph.cpp
--- a/Engine/Plugins/Runtime/RigVM/Source/RigVMDeveloper/Private/RigVMModel/RigVMGraph.cpp
+++ b/Engine/Plugins/Runtime/RigVM/Source/RigVMDeveloper/Private/RigVMModel/RigVMGraph.cpp
@@ -267,7 +267,10 @@ URigVMPin* URigVMGraph::FindPin(const FString& InPinPath) const
 	URigVMNode* Node = FindNode(Left);
 	if (Node)
 	{
-		return Node->FindPin(Right);
+		if (URigVMPin* Pin = Node->FindPin(Right))
+		{
+			return Pin;
+		}
 	}
 
 	return nullptr;
diff --git a/Engine/Plugins/Runtime/RigVM/Source/RigVMDeveloper/Private/RigVMModel/RigVMNode.cpp b/Engine/Plugins/Runtime/RigVM/Source/RigVMDeveloper/Private/RigVMModel/RigVMNode.cpp
--- a/Engine/Plugins/Runtime/RigVM/Source/RigVMDeveloper/Private/RigVMModel/RigVMNode.cpp
+++ b/Engine/Plugins/Runtime/RigVM/Source/RigVMDeveloper/Private/RigVMModel/RigVMNode.cpp
@@ -512,6 +512,10 @@ URigVMPin* URigVMNode::FindPin(const FString& InPinPath) const
 	{
 		for (URigVMPin* Pin : OrphanedPins)
 		{
+			if (Pin->GetName() == InPinPath)
+			{
+				return Pin;
+			}
 			if (Pin->GetName() == Left)
 			{
 				if (Right.IsEmpty())
diff --git a/Engine/Plugins/Runtime/RigVM/Source/RigVMDeveloper/Public/RigVMModel/RigVMController.h b/Engine/Plugins/Runtime/RigVM/Source/RigVMDeveloper/Public/RigVMModel/RigVMController.h
--- a/Engine/Plugins/Runtime/RigVM/Source/RigVMDeveloper/Public/RigVMModel/RigVMController.h
+++ b/Engine/Plugins/Runtime/RigVM/Source/RigVMDeveloper/Public/RigVMModel/RigVMController.h
@@ -1400,6 +1400,7 @@ public:
 		TSoftObjectPtr<URigVMGraph> GraphPtr;
 		FString SourcePinPath;
 		FString TargetPinPath;
+		FString OriginalPinPathRepresentation;
 		bool bSourceNodeIsInjected;
 		bool bTargetNodeIsInjected;
 	};
diff --git a/Engine/Plugins/Runtime/RigVM/Source/RigVMDeveloper/Public/RigVMModel/RigVMLink.h b/Engine/Plugins/Runtime/RigVM/Source/RigVMDeveloper/Public/RigVMModel/RigVMLink.h
--- a/Engine/Plugins/Runtime/RigVM/Source/RigVMDeveloper/Public/RigVMModel/RigVMLink.h
+++ b/Engine/Plugins/Runtime/RigVM/Source/RigVMDeveloper/Public/RigVMModel/RigVMLink.h
@@ -23,10 +23,7 @@ class RIGVMDEVELOPER_API URigVMLink : public UObject
 public:
 
 	// Default constructor
-	URigVMLink()
-	{
-		SourcePin = TargetPin = nullptr;
-	}
+	URigVMLink() = default;
 
 	// Serialization override
 	virtual void Serialize(FArchive& Ar) override;
@@ -112,7 +109,7 @@ private:
 	UPROPERTY()
 	FString TargetPinPath;
 
-	mutable URigVMPin* SourcePin;
-	mutable URigVMPin* TargetPin;
+	mutable URigVMPin* SourcePin = nullptr;
+	mutable URigVMPin* TargetPin = nullptr;
 };
 
diff --git a/Engine/Plugins/Runtime/RigVM/Source/RigVMDeveloper/Public/RigVMModel/RigVMPin.h b/Engine/Plugins/Runtime/RigVM/Source/RigVMDeveloper/Public/RigVMModel/RigVMPin.h
--- a/Engine/Plugins/Runtime/RigVM/Source/RigVMDeveloper/Public/RigVMModel/RigVMPin.h
+++ b/Engine/Plugins/Runtime/RigVM/Source/RigVMDeveloper/Public/RigVMModel/RigVMPin.h
@@ -621,6 +621,7 @@ private:
 	friend struct FRigVMSetPinDisplayNameAction;
 	friend struct FRigVMSetPinCategoryAction;
 	friend class URigVMLibraryNode;
+	friend class URigVMEdGraphNode;
 };
 
 class RIGVMDEVELOPER_API FRigVMPinDefaultValueImportErrorContext : public FOutputDevice
