AnimNext: Alter non-debugged RigVM connections to default to solid for readability

https://jira.it.epicgames.com/browse/UE-227393

diff --git a/Engine/Plugins/Runtime/RigVM/Source/RigVMEditor/Private/EdGraph/RigVMEdGraphConnectionDrawingPolicy.cpp b/Engine/Plugins/Runtime/RigVM/Source/RigVMEditor/Private/EdGraph/RigVMEdGraphConnectionDrawingPolicy.cpp
--- a/Engine/Plugins/Runtime/RigVM/Source/RigVMEditor/Private/EdGraph/RigVMEdGraphConnectionDrawingPolicy.cpp
+++ b/Engine/Plugins/Runtime/RigVM/Source/RigVMEditor/Private/EdGraph/RigVMEdGraphConnectionDrawingPolicy.cpp
@@ -309,7 +309,8 @@ void FRigVMEdGraphConnectionDrawingPolicy::DetermineWiringStyle(UEdGraphPin* Out
 
 		if (OutputModelPin && InputModelPin)
 		{
-			bool bVisited = false;
+			// Assume visited unless debugging
+			bool bVisited = true;
 			int32 OutputInstructionIndex = OutputNode->GetInstructionIndex(false);
 			int32 InputInstructionIndex = InputNode->GetInstructionIndex(true);
 
@@ -319,6 +320,8 @@ void FRigVMEdGraphConnectionDrawingPolicy::DetermineWiringStyle(UEdGraphPin* Out
 				{
 					if (URigVMHost* Host = Cast<URigVMHost>(RigBlueprint->GetObjectBeingDebugged()))
 					{
+						// Assume not visited unless visited when debugging last run
+						bVisited = false;
 						if (const URigVM* VM = Host->GetVM())
 						{
 							if (VM->WasInstructionVisitedDuringLastRun(Host->GetRigVMExtendedExecuteContext(), OutputInstructionIndex) &&
