RigVM: Rely on binary compressed serialization for function graphs in cooked data

#rnx

diff --git a/Engine/Plugins/Runtime/RigVM/Source/RigVM/Private/RigVMCore/RigVMGraphFunctionDefinition.cpp b/Engine/Plugins/Runtime/RigVM/Source/RigVM/Private/RigVMCore/RigVMGraphFunctionDefinition.cpp
--- a/Engine/Plugins/Runtime/RigVM/Source/RigVM/Private/RigVMCore/RigVMGraphFunctionDefinition.cpp
+++ b/Engine/Plugins/Runtime/RigVM/Source/RigVM/Private/RigVMCore/RigVMGraphFunctionDefinition.cpp
@@ -64,6 +64,102 @@ bool FRigVMGraphFunctionArgument::IsExecuteContext() const
 	return false;
 }
 
+FArchive& operator<<(FArchive& Ar, FRigVMFunctionCompilationData& Data)
+{
+	Ar.UsingCustomVersion(FUE5ReleaseStreamObjectVersion::GUID);
+	Ar.UsingCustomVersion(FFortniteMainBranchObjectVersion::GUID);
+		
+	Ar << Data.ByteCode;
+	Ar << Data.FunctionNames;
+	Ar << Data.WorkPropertyDescriptions;
+	Ar << Data.WorkPropertyPathDescriptions;
+	Ar << Data.LiteralPropertyDescriptions;
+	Ar << Data.LiteralPropertyPathDescriptions;
+	Ar << Data.DebugPropertyDescriptions;
+	Ar << Data.DebugPropertyPathDescriptions;
+	Ar << Data.ExternalPropertyDescriptions;
+	Ar << Data.ExternalPropertyPathDescriptions;
+	Ar << Data.ExternalRegisterIndexToVariable;
+	Ar << Data.Operands;
+	Ar << Data.Hash;
+
+	if(Ar.IsLoading())
+	{
+		Data.bEncounteredSurpressedErrors = false;
+	}
+
+	if (Ar.CustomVer(FUE5ReleaseStreamObjectVersion::GUID) < FUE5ReleaseStreamObjectVersion::RigVMSaveDebugMapInGraphFunctionData &&
+		Ar.CustomVer(FFortniteMainBranchObjectVersion::GUID) < FFortniteMainBranchObjectVersion::RigVMSaveDebugMapInGraphFunctionData)
+	{
+		return Ar;
+	}
+
+	// Serialize OperandToDebugRegisters
+	{
+		uint8 NumKeys = Data.OperandToDebugRegisters.Num();
+		Ar << NumKeys;
+
+		if (Ar.IsLoading())
+		{
+			for (int32 KeyIndex=0; KeyIndex<NumKeys; ++KeyIndex)
+			{
+				FRigVMOperand Key;
+				Ar << Key;
+				uint8 NumValues;
+				Ar << NumValues;
+				TArray<FRigVMOperand> Values;
+				Values.SetNumUninitialized(NumValues);
+				for (int32 ValueIndex=0; ValueIndex<NumValues; ++ValueIndex)
+				{
+					Ar << Values[ValueIndex];
+				}
+				Data.OperandToDebugRegisters.Add(Key, Values);
+			}
+		}
+		else
+		{
+			for (TPair<FRigVMOperand, TArray<FRigVMOperand>>& Pair : Data.OperandToDebugRegisters)
+			{
+				Ar << Pair.Key;
+				uint8 NumValues = Pair.Value.Num();
+				Ar << NumValues;
+				for (FRigVMOperand& Operand : Pair.Value)
+				{
+					Ar << Operand;
+				}
+			}
+		}
+	}
+	return Ar;
+}
+
+FArchive& operator<<(FArchive& Ar, FRigVMGraphFunctionIdentifier& Data)
+{
+	Ar.UsingCustomVersion(FRigVMObjectVersion::GUID);
+
+	if(Ar.IsSaving())
+	{
+		if(Data.LibraryNodePath.IsEmpty() && Data.LibraryNode_DEPRECATED.IsValid())
+		{
+			Data.LibraryNodePath = Data.GetLibraryNodePath();
+		}
+	}
+		
+	if (Ar.IsLoading() && Ar.CustomVer(FRigVMObjectVersion::GUID) < FRigVMObjectVersion::RemoveLibraryNodeReferenceFromFunctionIdentifier)
+	{
+		FSoftObjectPath SoftPath;
+		Ar << SoftPath;
+		Data.LibraryNodePath = SoftPath.ToString();
+	}
+	else
+	{
+		Ar << Data.LibraryNodePath;
+	}
+		
+	Ar << Data.HostObject;
+	return Ar;
+}
+
 TFunction<TArray<FRigVMVariantRef>(const FGuid& InGuid)> FRigVMGraphFunctionIdentifier::GetVariantRefsByGuidFunc;
 
 bool FRigVMGraphFunctionIdentifier::IsVariant() const
@@ -154,6 +250,104 @@ FRigVMGraphFunctionData* FRigVMGraphFunctionHeader::GetFunctionData(bool bLoadIf
 	return nullptr;
 }
 
+FArchive& operator<<(FArchive& Ar, FRigVMGraphFunctionHeader& Data)
+{
+	Ar.UsingCustomVersion(FRigVMObjectVersion::GUID);
+	
+	Ar << Data.LibraryPointer;
+
+	if (!Ar.IsLoading() || Ar.CustomVer(FRigVMObjectVersion::GUID) >= FRigVMObjectVersion::AddVariantToFunctionIdentifier)
+	{
+		Ar << Data.Variant;
+	}
+	
+	Ar << Data.Name;
+	Ar << Data.NodeTitle;
+	Ar << Data.NodeColor;
+
+	if (Ar.IsLoading())
+	{
+		if (Ar.CustomVer(FRigVMObjectVersion::GUID) < FRigVMObjectVersion::VMRemoveTooltipFromFunctionHeader)
+		{
+			Ar << Data.Tooltip_DEPRECATED;
+		}
+		else
+		{
+			Ar << Data.Description;
+		}
+	}
+	else
+	{
+		Ar << Data.Description;
+	}
+	
+	Ar << Data.Category;
+	Ar << Data.Keywords;
+	Ar << Data.Arguments;
+	Ar << Data.Dependencies;
+	Ar << Data.ExternalVariables;
+
+	if (Ar.IsLoading())
+	{
+		if (Ar.CustomVer(FRigVMObjectVersion::GUID) >= FRigVMObjectVersion::FunctionHeaderStoresLayout)
+		{
+			Ar << Data.Layout;
+		}
+		else
+		{
+			Data.Layout.Reset();
+		}
+	}
+	else
+	{
+		Ar << Data.Layout;
+	}
+	
+	return Ar;
+}
+
+FArchive& operator<<(FArchive& Ar, FRigVMGraphFunctionData& Data)
+{
+	Ar << Data.Header;
+	Ar << Data.CompilationData;
+
+	if (Ar.CustomVer(FUE5MainStreamObjectVersion::GUID) < FUE5MainStreamObjectVersion::RigVMSaveSerializedGraphInGraphFunctionData)
+	{
+		if(Ar.IsLoading())
+		{
+			Data.SerializedCollapsedNode_DEPRECATED.Empty();
+		}
+		return Ar;
+	}
+
+	Ar << Data.SerializedCollapsedNode_DEPRECATED;
+
+	if (Ar.CustomVer(FRigVMObjectVersion::GUID) < FRigVMObjectVersion::RigVMSaveSerializedGraphInGraphFunctionDataAsByteArray)
+	{
+		if(Ar.IsLoading())
+		{
+			Data.CollapseNodeArchive.Empty();
+		}
+		return Ar;
+	}
+
+	Ar << Data.CollapseNodeArchive;
+
+#if UE_BUILD_SHIPPING
+	// these two members can store substantial data
+	// which is not needed for a shipping game.
+	Data.CollapseNodeArchive.Empty();
+	Data.SerializedCollapsedNode_DEPRECATED.Empty();
+#else
+	if(!Data.CollapseNodeArchive.IsEmpty())
+	{
+		Data.SerializedCollapsedNode_DEPRECATED.Empty();
+	}
+#endif
+
+	return Ar;
+}
+
 FRigVMGraphFunctionHeader FRigVMGraphFunctionHeader::FindGraphFunctionHeader(const FSoftObjectPath& InFunctionObjectPath, bool* bOutIsPublic, FString* OutErrorMessage)
 {
 	return FindGraphFunctionHeader(InFunctionObjectPath, NAME_None, bOutIsPublic, OutErrorMessage);
diff --git a/Engine/Plugins/Runtime/RigVM/Source/RigVM/Private/RigVMCore/RigVMGraphFunctionHost.cpp b/Engine/Plugins/Runtime/RigVM/Source/RigVM/Private/RigVMCore/RigVMGraphFunctionHost.cpp
--- a/Engine/Plugins/Runtime/RigVM/Source/RigVM/Private/RigVMCore/RigVMGraphFunctionHost.cpp
+++ b/Engine/Plugins/Runtime/RigVM/Source/RigVM/Private/RigVMCore/RigVMGraphFunctionHost.cpp
@@ -163,8 +163,14 @@ bool FRigVMGraphFunctionStore::MarkFunctionAsPublic(const FRigVMGraphFunctionIde
 	{
 		return false;
 	}
-	NewContainer->Add(OldContainer->operator[](Index));
+	const int32 NewIndex = NewContainer->Add(OldContainer->operator[](Index));
 	OldContainer->RemoveAt(Index);
+	FRigVMGraphFunctionData& FunctionData = NewContainer->operator[](NewIndex);
+	if(!bIsPublic)
+	{
+		FunctionData.SerializedCollapsedNode_DEPRECATED.Empty();
+		FunctionData.CollapseNodeArchive.Empty();
+	}
 	return true;
 }
 
diff --git a/Engine/Plugins/Runtime/RigVM/Source/RigVM/Private/RigVMCore/RigVMObjectArchive.cpp b/Engine/Plugins/Runtime/RigVM/Source/RigVM/Private/RigVMCore/RigVMObjectArchive.cpp
new file mode 100644
--- /dev/null
+++ b/Engine/Plugins/Runtime/RigVM/Source/RigVM/Private/RigVMCore/RigVMObjectArchive.cpp
@@ -0,0 +1,473 @@
+// Copyright Epic Games, Inc. All Rights Reserved.
+
+#include "RigVMCore/RigVMObjectArchive.h"
+#include "Misc/Compression.h"
+#include "UObject/UE5MainStreamObjectVersion.h"
+#include "RigVMObjectVersion.h"
+
+#include UE_INLINE_GENERATED_CPP_BY_NAME(RigVMObjectArchive)
+
+void FRigVMObjectArchive::Reset()
+{
+	Buffer.Reset();
+	bIsCompressed = false;
+	CompressedSize = UncompressedSize = INDEX_NONE;
+}
+
+void FRigVMObjectArchive::Empty()
+{
+	Buffer.Empty();
+	bIsCompressed = false;
+	CompressedSize = UncompressedSize = INDEX_NONE;
+}
+
+bool FRigVMObjectArchive::IsEmpty() const
+{
+	return Buffer.IsEmpty();
+}
+
+void FRigVMObjectArchive::Compress()
+{
+	if(IsCompressed())
+	{
+		return;
+	}
+
+	UncompressedSize = Buffer.Num();
+
+	if(Buffer.IsEmpty())
+	{
+		return;
+	}
+	
+	// It is possible for compression to actually increase the size of the data, so we over allocate here to handle that.
+	CompressedSize = UncompressedSize * 4 / 3;
+
+	TArray<uint8> CompressedBuffer;
+	CompressedBuffer.SetNumUninitialized(CompressedSize);
+
+	if (FCompression::CompressMemory(NAME_Zlib, CompressedBuffer.GetData(), CompressedSize, Buffer.GetData(), Buffer.Num(), COMPRESS_BiasMemory))
+	{
+		// In the case that compressing it actually increases the size, we leave it uncompressed
+		if (CompressedSize < UncompressedSize)
+		{
+			CompressedBuffer.SetNum(CompressedSize);
+			Buffer = MoveTemp(CompressedBuffer);
+			Buffer.Shrink();
+			bIsCompressed = true;
+		}
+		else
+		{
+			CompressedSize = INDEX_NONE;
+		}
+	}
+
+}
+
+void FRigVMObjectArchive::Decompress()
+{
+	if(!IsCompressed())
+	{
+		return;
+	}
+
+	CompressedSize = Buffer.Num();
+	
+	TArray<uint8> UncompressedBuffer;
+	UncompressedBuffer.SetNumUninitialized(UncompressedSize);
+
+	if (FCompression::UncompressMemory(NAME_Zlib, UncompressedBuffer.GetData(), UncompressedSize, Buffer.GetData(), Buffer.Num()))
+	{
+		Buffer = MoveTemp(UncompressedBuffer);
+		bIsCompressed = false;
+	}
+}
+
+FRigVMObjectArchiveWriter::FRigVMObjectArchiveWriter(FRigVMObjectArchive& InArchive, const UObject* InRoot)
+: Archive( InArchive )
+, Offset(0)
+, Root( InRoot )
+, RootPathName( InRoot ? InRoot->GetPathName() : FString() )
+{
+	SetIsSaving(true);
+	SetIsLoading(false);
+	UsingCustomVersion(FUE5MainStreamObjectVersion::GUID);
+	UsingCustomVersion(FRigVMObjectVersion::GUID);
+	ArIgnoreOuterRef = 1;
+}
+
+void FRigVMObjectArchiveWriter::Serialize(void* V, int64 Length)
+{
+	if(Length == 0)
+	{
+		check(V == nullptr);
+	}
+	else
+	{
+		check(V);
+		int32 Start = IntCastChecked<int32, int64>(Offset);
+
+		if(Archive.Buffer.IsValidIndex(Start))
+		{
+			const int32 MissingBytes = IntCastChecked<int32, int64>(Length) - (Archive.Buffer.Num() - Start);
+			if(MissingBytes > 0)
+			{
+				(void)Archive.Buffer.AddUninitialized(MissingBytes);
+			}
+		}
+		else if(Start == Archive.Buffer.Num())
+		{
+			Start = Archive.Buffer.AddUninitialized(IntCastChecked<int32, int64>(Length));
+		}
+		else
+		{
+			checkNoEntry();
+		}
+		
+		FMemory::Memcpy( &(Archive.Buffer[Start]), V, Length );
+		Offset += Length;
+	}
+}
+
+int64 FRigVMObjectArchiveWriter::Tell()
+{
+	return Offset;
+}
+
+int64 FRigVMObjectArchiveWriter::TotalSize()
+{
+	return Archive.Buffer.Num();
+}
+
+void FRigVMObjectArchiveWriter::Seek(int64 InPos)
+{
+	check((InPos >= 0) && (InPos <= Archive.Buffer.Num()));
+	Offset = InPos;
+}
+
+FArchive& FRigVMObjectArchiveWriter::operator<<(UObject*& Obj)
+{
+	if (Obj)
+	{
+		bool bStorePath = true;
+		if(!VisitedObjects.Contains(Obj))
+		{
+			if((Obj == Root) || Obj->GetPathName().StartsWith(RootPathName, ESearchCase::CaseSensitive))
+			{
+				VisitedObjects.Add(Obj);
+				bStorePath = false;
+
+				uint8 State = StoringFullObject;
+				*this << State;
+				FName ClassName = *Obj->GetClass()->GetPathName();
+				*this << ClassName;
+				FName ObjectName = *Obj->GetName();
+				*this << ObjectName;
+				int32 ObjectFlags = (int32)Obj->GetFlags();
+				*this << ObjectFlags;
+
+				FName OuterPathName = FName(NAME_None);
+				if(UObject* Outer = Obj->GetOuter())
+				{
+					if(VisitedObjects.Contains(Outer))
+					{
+						FString OuterPathNameString = Outer->GetPathName();
+						if(OuterPathNameString.StartsWith(RootPathName, ESearchCase::CaseSensitive))
+						{
+							OuterPathNameString = OuterPathNameString.Mid(RootPathName.Len());
+						}
+						OuterPathName = OuterPathNameString.IsEmpty() ? FName(NAME_None) : *OuterPathNameString;
+					}
+				}
+				
+				*this << OuterPathName;
+
+				Obj->Serialize(*this);
+			}
+		}
+
+		if(bStorePath)
+		{
+			FString PathNameString = Obj->GetPathName();
+			if(PathNameString.StartsWith(RootPathName, ESearchCase::CaseSensitive))
+			{
+				PathNameString = PathNameString.Mid(RootPathName.Len());
+			}
+			
+			uint8 State = StoringArchiveLocalPath;
+			*this << State;
+			FName PathName = *PathNameString;
+			*this << PathName;
+		}
+	}
+	else
+	{
+		uint8 State = StoringNullptr;
+		*this << State;
+	}
+
+	return *this;
+}
+
+FArchive& FRigVMObjectArchiveWriter::operator<<(FName& Value)
+{
+	if(const int64* NameOffsetPtr = NameToOffset.Find(Value))
+	{
+		uint8 State = StoringNameAsOffset;
+		*this << State;
+		
+		int64 NameOffset = *NameOffsetPtr;
+		*this << NameOffset;
+	}
+	else
+	{
+		uint8 State = StoringNameAsString;
+		*this << State;
+		NameToOffset.Add(Value, Tell());
+		FString NameAsString = Value.IsNone() ? FString() : Value.ToString();
+		*this << NameAsString;
+	}
+
+	return *this;
+}
+
+FArchive& FRigVMObjectArchiveWriter::operator<<(FText& Value)
+{
+	FString ValueString = Value.ToString();
+	*this << ValueString;
+	return *this;
+}
+
+FRigVMObjectArchiveReader::FRigVMObjectArchiveReader(FRigVMObjectArchive& InArchive, const UObject* InRoot)
+: FRigVMObjectArchiveWriter(InArchive, InRoot)
+{
+	SetIsSaving(false);
+	SetIsLoading(true);
+	Archive.Decompress();
+}
+
+void FRigVMObjectArchiveReader::Serialize(void* V, int64 Length)
+{
+	if(Length == 0)
+	{
+		check(V == nullptr);
+	}
+	else
+	{
+		check(V);
+		check(Archive.Buffer.Num() >= Offset + Length);
+		FMemory::Memcpy( V, &(Archive.Buffer[Offset]), Length );
+		Offset += Length;
+	}
+}
+
+FArchive& FRigVMObjectArchiveReader::operator<<(UObject*& Obj)
+{
+	uint8 State = UINT8_MAX;
+	*this << State;
+
+	if (State == StoringNullptr)
+	{
+		Obj = nullptr;
+	}
+	else
+	{
+		if(State == StoringFullObject)
+		{
+			const FObjectHeader Header = ReadObjectHeader();
+			check(Header.IsValid());
+
+			const EObjectFlags Flags = Header.Flags | RF_NeedPostLoad;
+
+			if(Obj != Root)
+			{
+				UObject* Outer = const_cast<UObject*>(Root);
+				if(!Header.OuterPathName.IsNone())
+				{
+					Outer = ReadObjects.FindChecked(Header.OuterPathName);
+				}
+
+				Obj = StaticFindObjectFast(Header.Class, Outer, Header.Name);
+				if(Obj == nullptr || !IsValid(Obj))
+				{
+					Obj = NewObject<UObject>(Outer, Header.Class, Header.Name, Flags);
+				}
+				else
+				{
+					Obj->SetFlags(Flags);
+				}
+			}
+			else
+			{
+				check(IsValid(Obj));
+			}
+
+			FString PathNameString = Obj->GetPathName();
+			if(PathNameString.StartsWith(RootPathName, ESearchCase::CaseSensitive))
+			{
+				PathNameString = PathNameString.Mid(RootPathName.Len());
+			}
+			const FName PathName = *PathNameString;
+
+			check(!ReadObjects.Contains(PathName));
+			ReadObjects.Add(PathName, Obj);
+			
+			Obj->Serialize(*this);
+			DeserializedObjects.Add(Obj);
+		}
+		else if(State == StoringArchiveLocalPath)
+		{
+			FName PathName;
+			*this << PathName;
+
+			if(UObject** PreviouslyReadObj = ReadObjects.Find(PathName))
+			{
+				Obj = *PreviouslyReadObj;
+				check(Obj);
+			}
+			else
+			{
+				Obj = FindObject<UObject>(nullptr, *PathName.ToString(), false);
+				check(Obj);
+				ReadObjects.Add(PathName, Obj);
+			}
+		}
+		else
+		{
+			checkNoEntry();
+		}
+	}
+
+	if(Obj == Root)
+	{
+		// perform post load
+		for(UObject* DeserializedObject : DeserializedObjects)
+		{
+			DeserializedObject->PostLoad();
+			DeserializedObject->ClearFlags(RF_NeedPostLoad);
+		}
+	}
+	
+	return *this;
+}
+
+FArchive& FRigVMObjectArchiveReader::operator<<(FName& Value)
+{
+	uint8 State = UINT8_MAX;
+	*this << State;
+
+	auto ReadNameString = [this, &Value]()
+	{
+		const int64 OffsetOfString = Tell();
+		FString NameAsString;
+		*this << NameAsString;
+		if(NameAsString.IsEmpty())
+		{
+			Value = FName(NAME_None);
+		}
+		else
+		{
+			Value = *NameAsString;
+		}
+		OffsetToName.FindOrAdd(OffsetOfString) = Value;
+	};
+
+	if(State == StoringNameAsString)
+	{
+		// read the string at the current position
+		ReadNameString();
+	}
+	else if(State == StoringNameAsOffset)
+	{
+		int64 OffsetOfString = -1;
+		*this << OffsetOfString;
+		check(Archive.Buffer.IsValidIndex(OffsetOfString));
+		
+		if(const FName* NamePtr = OffsetToName.Find(OffsetOfString))
+		{
+			Value = *NamePtr;
+		}
+		else
+		{
+			TGuardValue<int64> OffsetGuard(Offset, Offset);
+			Seek(OffsetOfString);
+			ReadNameString();
+		}
+	}
+	else
+	{
+		checkNoEntry();
+	}
+
+	return *this;
+}
+
+FArchive& FRigVMObjectArchiveReader::operator<<(FText& Value)
+{
+	FString ValueString;
+	*this << ValueString;
+	Value = FText::FromString(ValueString);
+	return *this;
+}
+
+FRigVMObjectArchiveReader::FObjectHeader FRigVMObjectArchiveReader::ReadObjectHeader(int64 InPosition)
+{
+	const int64 CurrentPosition = Tell();
+	if(InPosition != INDEX_NONE)
+	{
+		Seek(InPosition);
+	}
+	
+	FObjectHeader Header;
+	FName ClassName = FName(NAME_None);
+	*this << ClassName;
+
+	if(UClass** ClassPtr = ReadClasses.Find(ClassName))
+	{
+		Header.Class = *ClassPtr;
+	}
+	else
+	{
+		Header.Class = FindObject<UClass>(nullptr, *ClassName.ToString(), false);
+		ReadClasses.Add(ClassName, Header.Class);
+	}
+
+	*this << Header.Name;
+	int32 ObjectFlags = (int32)Header.Flags;
+	*this << ObjectFlags;
+	Header.Flags = (EObjectFlags)ObjectFlags;
+	*this << Header.OuterPathName;
+
+	if(InPosition != INDEX_NONE)
+	{
+		Seek(CurrentPosition);
+	}
+
+	return Header;
+}
+
+FRigVMObjectArchiveReader::FObjectHeader FRigVMObjectArchiveReader::GetRootObjectHeader()
+{
+	TGuardValue<int64> OffsetGuard(Offset, Offset);
+
+	Seek(0);
+	
+	uint8 State = UINT8_MAX;
+	*this << State;
+
+	if(State == StoringFullObject)
+	{
+		return ReadObjectHeader();
+	}
+
+	return FObjectHeader();
+}
+
+FArchive& operator<<(FArchive& Ar, FRigVMObjectArchive& Data)
+{
+	Ar << Data.Buffer;
+	Ar << Data.UncompressedSize;
+	Ar << Data.CompressedSize;
+	Ar << Data.bIsCompressed;
+	return Ar;
+}
diff --git a/Engine/Plugins/Runtime/RigVM/Source/RigVM/Public/RigVMCore/RigVMGraphFunctionDefinition.h b/Engine/Plugins/Runtime/RigVM/Source/RigVM/Public/RigVMCore/RigVMGraphFunctionDefinition.h
--- a/Engine/Plugins/Runtime/RigVM/Source/RigVM/Public/RigVMCore/RigVMGraphFunctionDefinition.h
+++ b/Engine/Plugins/Runtime/RigVM/Source/RigVM/Public/RigVMCore/RigVMGraphFunctionDefinition.h
@@ -5,6 +5,7 @@
 #include "RigVMObjectVersion.h"
 #include "RigVMCore/RigVMExternalVariable.h"
 #include "RigVMCore/RigVMByteCode.h"
+#include "RigVMCore/RigVMObjectArchive.h"
 #include "UObject/UE5MainStreamObjectVersion.h"
 #include "UObject/FortniteMainBranchObjectVersion.h"
 #include "UObject/UE5ReleaseStreamObjectVersion.h"
@@ -46,7 +47,7 @@ struct RIGVM_API FRigVMFunctionCompilationPropertyDescription
 		return Hash;
 	}
 
-	friend FArchive& operator<<(FArchive& Ar, FRigVMFunctionCompilationPropertyDescription& Data)
+	RIGVM_API friend FArchive& operator<<(FArchive& Ar, FRigVMFunctionCompilationPropertyDescription& Data)
 	{
 		Ar << Data.Name;
 		Ar << Data.CPPType;
@@ -81,7 +82,7 @@ struct FRigVMFunctionCompilationPropertyPath
 		return Hash;
 	}
 
-	friend FArchive& operator<<(FArchive& Ar, FRigVMFunctionCompilationPropertyPath& Data)
+	RIGVM_API friend FArchive& operator<<(FArchive& Ar, FRigVMFunctionCompilationPropertyPath& Data)
 	{
 		Ar << Data.PropertyIndex;
 		Ar << Data.HeadCPPType;
@@ -213,74 +214,7 @@ struct RIGVM_API FRigVMFunctionCompilationData
 		return DataHash;
 	}
 
-	friend FArchive& operator<<(FArchive& Ar, FRigVMFunctionCompilationData& Data)
-	{
-		Ar.UsingCustomVersion(FUE5ReleaseStreamObjectVersion::GUID);
-		Ar.UsingCustomVersion(FFortniteMainBranchObjectVersion::GUID);
-		
-		Ar << Data.ByteCode;
-		Ar << Data.FunctionNames;
-		Ar << Data.WorkPropertyDescriptions;
-		Ar << Data.WorkPropertyPathDescriptions;
-		Ar << Data.LiteralPropertyDescriptions;
-		Ar << Data.LiteralPropertyPathDescriptions;
-		Ar << Data.DebugPropertyDescriptions;
-		Ar << Data.DebugPropertyPathDescriptions;
-		Ar << Data.ExternalPropertyDescriptions;
-		Ar << Data.ExternalPropertyPathDescriptions;
-		Ar << Data.ExternalRegisterIndexToVariable;
-		Ar << Data.Operands;
-		Ar << Data.Hash;
-
-		if(Ar.IsLoading())
-		{
-			Data.bEncounteredSurpressedErrors = false;
-		}
-
-		if (Ar.CustomVer(FUE5ReleaseStreamObjectVersion::GUID) < FUE5ReleaseStreamObjectVersion::RigVMSaveDebugMapInGraphFunctionData &&
-		    Ar.CustomVer(FFortniteMainBranchObjectVersion::GUID) < FFortniteMainBranchObjectVersion::RigVMSaveDebugMapInGraphFunctionData)
-		{
-			return Ar;
-		}
-
-		// Serialize OperandToDebugRegisters
-		{
-			uint8 NumKeys = Data.OperandToDebugRegisters.Num();
-			Ar << NumKeys;
-
-			if (Ar.IsLoading())
-			{
-				for (int32 KeyIndex=0; KeyIndex<NumKeys; ++KeyIndex)
-				{
-					FRigVMOperand Key;
-					Ar << Key;
-					uint8 NumValues;
-					Ar << NumValues;
-					TArray<FRigVMOperand> Values;
-					Values.SetNumUninitialized(NumValues);
-					for (int32 ValueIndex=0; ValueIndex<NumValues; ++ValueIndex)
-					{
-						Ar << Values[ValueIndex];
-					}
-					Data.OperandToDebugRegisters.Add(Key, Values);
-				}
-			}
-			else
-			{
-				for (TPair<FRigVMOperand, TArray<FRigVMOperand>>& Pair : Data.OperandToDebugRegisters)
-				{
-					Ar << Pair.Key;
-					uint8 NumValues = Pair.Value.Num();
-					Ar << NumValues;
-					for (FRigVMOperand& Operand : Pair.Value)
-					{
-						Ar << Operand;
-					}
-				}
-			}
-		}
-		return Ar;
-	}
+	RIGVM_API friend FArchive& operator<<(FArchive& Ar, FRigVMFunctionCompilationData& Data);
 };
 
 USTRUCT(BlueprintType)
@@ -355,7 +289,7 @@ struct RIGVM_API FRigVMGraphFunctionArgument
 		return true;
 	}
 
-	friend FArchive& operator<<(FArchive& Ar, FRigVMGraphFunctionArgument& Data)
+	RIGVM_API friend FArchive& operator<<(FArchive& Ar, FRigVMGraphFunctionArgument& Data)
 	{
 		Ar << Data.Name;
 		Ar << Data.DisplayName;
@@ -451,32 +385,7 @@ public:
 		return GetLibraryNodePath();
 	}
 
-	friend FArchive& operator<<(FArchive& Ar, FRigVMGraphFunctionIdentifier& Data)
-	{
-		Ar.UsingCustomVersion(FRigVMObjectVersion::GUID);
-
-		if(Ar.IsSaving())
-		{
-			if(Data.LibraryNodePath.IsEmpty() && Data.LibraryNode_DEPRECATED.IsValid())
-			{
-				Data.LibraryNodePath = Data.GetLibraryNodePath();
-			}
-		}
-		
-		if (Ar.IsLoading() && Ar.CustomVer(FRigVMObjectVersion::GUID) < FRigVMObjectVersion::RemoveLibraryNodeReferenceFromFunctionIdentifier)
-		{
-			FSoftObjectPath SoftPath;
-			Ar << SoftPath;
-			Data.LibraryNodePath = SoftPath.ToString();
-		}
-		else
-		{
-			Ar << Data.LibraryNodePath;
-		}
-		
-		Ar << Data.HostObject;
-		return Ar;
-	}
+	RIGVM_API friend FArchive& operator<<(FArchive& Ar, FRigVMGraphFunctionIdentifier& Data);
 
 	bool IsVariant() const;
 	TArray<FRigVMVariantRef> GetVariants(bool bIncludeSelf = false) const;
@@ -571,61 +480,7 @@ struct RIGVM_API FRigVMGraphFunctionHeader
 		return FText::FromString(TooltipStr);
 	}
 
-	friend FArchive& operator<<(FArchive& Ar, FRigVMGraphFunctionHeader& Data)
-	{
-		Ar.UsingCustomVersion(FRigVMObjectVersion::GUID);
-		
-		Ar << Data.LibraryPointer;
-
-		if (!Ar.IsLoading() || Ar.CustomVer(FRigVMObjectVersion::GUID) >= FRigVMObjectVersion::AddVariantToFunctionIdentifier)
-		{
-			Ar << Data.Variant;
-		}
-		
-		Ar << Data.Name;
-		Ar << Data.NodeTitle;
-		Ar << Data.NodeColor;
-
-		if (Ar.IsLoading())
-		{
-			if (Ar.CustomVer(FRigVMObjectVersion::GUID) < FRigVMObjectVersion::VMRemoveTooltipFromFunctionHeader)
-			{
-				Ar << Data.Tooltip_DEPRECATED;
-			}
-			else
-			{
-				Ar << Data.Description;
-			}
-		}
-		else
-		{
-			Ar << Data.Description;
-		}
-		
-		Ar << Data.Category;
-		Ar << Data.Keywords;
-		Ar << Data.Arguments;
-		Ar << Data.Dependencies;
-		Ar << Data.ExternalVariables;
-
-		if (Ar.IsLoading())
-		{
-			if (Ar.CustomVer(FRigVMObjectVersion::GUID) >= FRigVMObjectVersion::FunctionHeaderStoresLayout)
-			{
-				Ar << Data.Layout;
-			}
-			else
-			{
-				Data.Layout.Reset();
-			}
-		}
-		else
-		{
-			Ar << Data.Layout;
-		}
-		
-		return Ar;
-	}
+	RIGVM_API friend FArchive& operator<<(FArchive& Ar, FRigVMGraphFunctionHeader& Data);
 
 	static FRigVMGraphFunctionHeader FindGraphFunctionHeader(const FSoftObjectPath& InFunctionObjectPath, bool* bOutIsPublic = nullptr, FString* OutErrorMessage = nullptr);
 
@@ -660,7 +515,10 @@ struct RIGVM_API FRigVMGraphFunctionData
 	FRigVMFunctionCompilationData CompilationData;
 
 	UPROPERTY()
-	FString SerializedCollapsedNode;
+	FString SerializedCollapsedNode_DEPRECATED;
+
+	UPROPERTY()
+	FRigVMObjectArchive CollapseNodeArchive;
 
 	bool IsMutable() const;
 
@@ -671,19 +529,7 @@ struct RIGVM_API FRigVMGraphFunctionData
 
 	void ClearCompilationData() { CompilationData = FRigVMFunctionCompilationData(); }
 
-	friend FArchive& operator<<(FArchive& Ar, FRigVMGraphFunctionData& Data)
-	{
-		Ar << Data.Header;
-		Ar << Data.CompilationData;
-
-		if (Ar.CustomVer(FUE5MainStreamObjectVersion::GUID) < FUE5MainStreamObjectVersion::RigVMSaveSerializedGraphInGraphFunctionData)
-		{
-			return Ar;
-		}
-
-		Ar << Data.SerializedCollapsedNode;
-		return Ar;
-	}
+	RIGVM_API friend FArchive& operator<<(FArchive& Ar, FRigVMGraphFunctionData& Data);
 
 	static FRigVMGraphFunctionData* FindFunctionData(const FSoftObjectPath& InHostObjectPath, const FName& InFunctionName, bool* bOutIsPublic = nullptr, FString* OutErrorMessage = nullptr);	
 
diff --git a/Engine/Plugins/Runtime/RigVM/Source/RigVM/Public/RigVMCore/RigVMObjectArchive.h b/Engine/Plugins/Runtime/RigVM/Source/RigVM/Public/RigVMCore/RigVMObjectArchive.h
new file mode 100644
--- /dev/null
+++ b/Engine/Plugins/Runtime/RigVM/Source/RigVM/Public/RigVMCore/RigVMObjectArchive.h
@@ -0,0 +1,120 @@
+// Copyright Epic Games, Inc. All Rights Reserved.
+
+#pragma once
+
+#include "Serialization/ArchiveUObject.h"
+#include "RigVMObjectArchive.generated.h"
+
+USTRUCT()
+struct RIGVM_API FRigVMObjectArchive
+{
+public:
+	GENERATED_BODY()
+
+	FRigVMObjectArchive()
+		: UncompressedSize(INDEX_NONE)
+		, CompressedSize(INDEX_NONE)
+		, bIsCompressed(false)
+	{
+	}
+	
+	void Reset();
+	void Empty();
+	bool IsEmpty() const;
+	void Compress();
+	void Decompress();
+	bool IsCompressed() const
+	{
+		return bIsCompressed;
+	}
+
+	RIGVM_API friend FArchive& operator<<(FArchive& Ar, FRigVMObjectArchive& Data);
+
+private:
+
+	UPROPERTY()
+	TArray<uint8> Buffer;
+
+	UPROPERTY()
+	int32 UncompressedSize;
+	
+	UPROPERTY()
+	int32 CompressedSize;
+
+	UPROPERTY()
+	bool bIsCompressed;
+
+	friend class FRigVMObjectArchiveWriter;
+	friend class FRigVMObjectArchiveReader;
+};
+
+class RIGVM_API FRigVMObjectArchiveWriter : public FArchiveUObject
+{
+public:
+	FRigVMObjectArchiveWriter( FRigVMObjectArchive& InOutArchive, const UObject* InRoot );
+	virtual void Serialize( void* V, int64 Length ) override;
+	virtual int64 Tell() override;
+	virtual int64 TotalSize() override;
+	virtual void Seek(int64 InPos) override;
+	using FArchiveUObject::operator<<; // For visibility of the overloads we don't override
+	virtual FArchive& operator<<(UObject*& Obj) override;
+	virtual FArchive& operator<<(FName& Value) override;
+	virtual FArchive& operator<<(FText& Value) override;
+
+protected:
+	
+	FRigVMObjectArchive& Archive;
+	int64 Offset;
+	const UObject* Root;
+	const FString RootPathName;
+	TSet<UObject*> VisitedObjects;
+	TMap<FName,int64> NameToOffset;
+	
+	static inline constexpr uint8 StoringNullptr = 0;
+	static inline constexpr uint8 StoringFullObject = 1;
+	static inline constexpr uint8 StoringArchiveLocalPath = 2;
+	static inline constexpr uint8 StoringNameAsString = 0;
+	static inline constexpr uint8 StoringNameAsOffset = 1;
+};
+
+class RIGVM_API FRigVMObjectArchiveReader : public FRigVMObjectArchiveWriter
+{
+public:
+	
+	FRigVMObjectArchiveReader( FRigVMObjectArchive& InOutArchive, const UObject* InRoot );
+	virtual void Serialize( void* V, int64 Length ) override;
+	using FArchiveUObject::operator<<; // For visibility of the overloads we don't override
+	virtual FArchive& operator<<(UObject*& Obj) override;
+	virtual FArchive& operator<<(FName& Value) override;
+	virtual FArchive& operator<<(FText& Value) override;
+
+	struct FObjectHeader
+	{
+		FObjectHeader()
+		{
+			Class = nullptr;
+			Name = OuterPathName = NAME_None;
+			Flags = RF_NoFlags;
+		}
+
+		bool IsValid() const
+		{
+			return Class != nullptr && !Name.IsNone();   
+		}
+		
+		UClass* Class;
+		FName Name;
+		EObjectFlags Flags;
+		FName OuterPathName;
+	};
+
+	FObjectHeader ReadObjectHeader(int64 InPosition = INDEX_NONE);
+	FObjectHeader GetRootObjectHeader();
+
+protected:
+
+	TMap<FName, UClass*> ReadClasses;
+	TMap<FName, UObject*> ReadObjects;
+	TArray<UObject*> DeserializedObjects;
+	TMap<int64,FName> OffsetToName;
+};
\ No newline at end of file
diff --git a/Engine/Plugins/Runtime/RigVM/Source/RigVM/Public/RigVMObjectVersion.h b/Engine/Plugins/Runtime/RigVM/Source/RigVM/Public/RigVMObjectVersion.h
--- a/Engine/Plugins/Runtime/RigVM/Source/RigVM/Public/RigVMObjectVersion.h
+++ b/Engine/Plugins/Runtime/RigVM/Source/RigVM/Public/RigVMObjectVersion.h
@@ -57,6 +57,9 @@ struct RIGVM_API FRigVMObjectVersion
 		// Storing user interface relevant category expansion
 		FunctionHeaderLayoutStoresCategoryExpansion,
 
+		// Storing function graph collapse node content as part of the header
+		RigVMSaveSerializedGraphInGraphFunctionDataAsByteArray,
+
 		// -----<new versions can be added above this line>-------------------------------------------------
 		VersionPlusOne,
 		LatestVersion = VersionPlusOne - 1
diff --git a/Engine/Plugins/Runtime/RigVM/Source/RigVMDeveloper/Private/RigVMModel/RigVMClient.cpp b/Engine/Plugins/Runtime/RigVM/Source/RigVMDeveloper/Private/RigVMModel/RigVMClient.cpp
--- a/Engine/Plugins/Runtime/RigVM/Source/RigVMDeveloper/Private/RigVMModel/RigVMClient.cpp
+++ b/Engine/Plugins/Runtime/RigVM/Source/RigVMDeveloper/Private/RigVMModel/RigVMClient.cpp
@@ -3,11 +3,10 @@
 #include "RigVMModel/RigVMClient.h"
 #include "Misc/TransactionObjectEvent.h"
 #include "Exporters/Exporter.h"
-#include "Algo/Sort.h"
-#include "UnrealExporter.h"
 #include "UObject/ObjectSaveContext.h"
 #include "RigVMModel/RigVMControllerActions.h"
 #include "RigVMModel/Nodes/RigVMDispatchNode.h"
+#include "RigVMCore/RigVMObjectArchive.h"
 #include "EdGraph/RigVMEdGraph.h"
 
 #include UE_INLINE_GENERATED_CPP_BY_NAME(RigVMClient)
@@ -1348,11 +1347,16 @@ void FRigVMClient::PatchFunctionsOnLoad(IRigVMGraphFunctionHost* FunctionHost, T
 					Header.Variant = *Variant;
 				}
 
-				Store->AddFunction(Header, bIsPublic);
+				FRigVMGraphFunctionData* FunctionData = Store->AddFunction(Header, bIsPublic);
 				if (bIsPublic)
 				{
 					UpdateGraphFunctionSerializedGraph(LibraryNode);
 				}
+				else
+				{
+					FunctionData->SerializedCollapsedNode_DEPRECATED.Empty();
+					FunctionData->CollapseNodeArchive.Empty();
+				}
 			}
 
 			// Update dependencies and external variables if needed
@@ -1646,6 +1650,10 @@ void FRigVMClient::HandleGraphModifiedEvent(ERigVMGraphNotifType InNotifType, UR
 					if (FRigVMGraphFunctionStore* Store = FindFunctionStore(LibraryNode))
 					{
 						Store->MarkFunctionAsPublic(LibraryNode->GetFunctionIdentifier(), bIsPublic);
+						if(bIsPublic)
+						{
+							UpdateGraphFunctionSerializedGraph(LibraryNode);
+						}
 					}
 				}
 			}
@@ -1790,13 +1798,16 @@ bool FRigVMClient::UpdateGraphFunctionSerializedGraph(URigVMLibraryNode* InLibra
 {
 	if (FRigVMGraphFunctionStore* Store = FindFunctionStore(InLibraryNode))
 	{
-		FRigVMGraphFunctionIdentifier Identifier = InLibraryNode->GetFunctionIdentifier();
+		const FRigVMGraphFunctionIdentifier Identifier = InLibraryNode->GetFunctionIdentifier();
 		if (FRigVMGraphFunctionData* Data = Store->FindFunction(Identifier))
 		{
-			FStringOutputDevice Archive;
-			const FExportObjectInnerContext Context;
-			UExporter::ExportToOutputDevice(&Context, InLibraryNode, NULL, Archive, TEXT("copy"), 0, PPF_ExportsNotFullyQualified | PPF_Copy | PPF_Delimited, false, nullptr);
-			Data->SerializedCollapsedNode = Archive;
+			Data->SerializedCollapsedNode_DEPRECATED.Empty();
+
+			URigVMController* Controller = GetOrCreateController(InLibraryNode->GetGraph());
+			check(Controller);
+			Data->CollapseNodeArchive.Reset();
+			(void)Controller->ExportFunctionToArchive(Identifier.GetFunctionFName(), Data->CollapseNodeArchive);
+
 			return true;
 		}
 	}
diff --git a/Engine/Plugins/Runtime/RigVM/Source/RigVMDeveloper/Private/RigVMModel/RigVMController.cpp b/Engine/Plugins/Runtime/RigVM/Source/RigVMDeveloper/Private/RigVMModel/RigVMController.cpp
--- a/Engine/Plugins/Runtime/RigVM/Source/RigVMDeveloper/Private/RigVMModel/RigVMController.cpp
+++ b/Engine/Plugins/Runtime/RigVM/Source/RigVMDeveloper/Private/RigVMModel/RigVMController.cpp
@@ -33,6 +33,7 @@
 #include "RigVMModel/RigVMClient.h"
 #include "RigVMModel/Nodes/RigVMBranchNode.h"
 #include "RigVMModel/Nodes/RigVMArrayNode.h"
+#include "RigVMCore/RigVMObjectArchive.h"
 #include "Logging/LogScopedVerbosityOverride.h"
 #include "ScopedTransaction.h"
 
@@ -4504,6 +4505,148 @@ TArray<FName> URigVMController::ImportNodesFromText(const FString& InText, bool
 	return NodeNames;
 }
 
+bool URigVMController::ExportFunctionToArchive(const FName& InFunctionName, FRigVMObjectArchive& OutArchive)
+{
+	if (!IsValidGraph())
+	{
+		return false;
+	}
+
+	if (!bIsTransacting && !IsGraphEditable())
+	{
+		return false;
+	}
+
+	URigVMGraph* Graph = GetGraph();
+	check(Graph);
+
+	URigVMFunctionLibrary* FunctionLibrary = Cast<URigVMFunctionLibrary>(Graph);
+	if (FunctionLibrary == nullptr)
+	{
+		ReportError(TEXT("Cannot export functions from non-function-libraries"));
+		return false;
+	}
+
+	URigVMLibraryNode* Function = FunctionLibrary->FindFunction(InFunctionName);
+	if(Function == nullptr)
+	{
+		ReportErrorf(TEXT("Cannot find function '%s'"), *InFunctionName.ToString());
+		return false;
+	}
+
+	OutArchive.Reset();
+	{
+		FRigVMObjectArchiveWriter Writer(OutArchive, Function);
+		Writer << Function;
+	}
+	OutArchive.Compress();
+	
+	return true;
+}
+
+URigVMLibraryNode* URigVMController::ImportFunctionFromArchive(const FRigVMObjectArchive& InArchive, const FName& InFunctionName)
+{
+	if (!IsValidGraph())
+	{
+		return {};
+	}
+
+	if (!bIsTransacting && !IsGraphEditable())
+	{
+		return {};
+	}
+
+	URigVMGraph* Graph = GetGraph();
+	check(Graph);
+
+	URigVMFunctionLibrary* FunctionLibrary = Cast<URigVMFunctionLibrary>(Graph);
+	if (FunctionLibrary == nullptr)
+	{
+		ReportError(TEXT("Cannot import functions into non-function-libraries"));
+		return {};
+	}
+
+	if(InArchive.IsEmpty())
+	{
+		ReportError(TEXT("Provided archive is empty."));
+		return {};
+	}
+
+	// test if the archive contains the right thin.
+	// since the archive is passed as const - we need to copy it.
+	// the process here may also decompress the archive internally.
+	FRigVMObjectArchive CopyOfArchive = InArchive;
+	FRigVMObjectArchiveReader::FObjectHeader ObjectHeader;
+	{
+		FRigVMObjectArchiveReader ReaderForHeader(CopyOfArchive, nullptr);
+		ObjectHeader = ReaderForHeader.GetRootObjectHeader();
+	}
+	if(!ObjectHeader.IsValid())
+	{
+		ReportError(TEXT("Cannot read object header from archive. Corrupt data?"));
+		return {};
+	}
+	
+	if(ObjectHeader.Class != URigVMCollapseNode::StaticClass())
+	{
+		ReportError(TEXT("Archive doesn't contain a collapse node."));
+		return {};
+	}
+
+	const FName NameForFunction = URigVMSchema::GetUniqueName(InFunctionName.IsNone() ? ObjectHeader.Name : InFunctionName, [Graph](const FName& InName)
+	{
+		return Graph->IsNameAvailable(InName.ToString());
+	}, false, true);
+
+	FRigVMVariant Variant;
+	if (FRigVMGraphFunctionStore* FunctionStore = GetGraphFunctionStore())
+	{
+		if (const FRigVMGraphFunctionData* FunctionData = FunctionStore->FindFunctionByName(NameForFunction))
+		{
+			Variant = FunctionData->Header.Variant;
+		}
+	}
+	
+	if (!Variant.Guid.IsValid())
+	{
+		Variant.Guid = FRigVMVariant::GenerateGUID();
+	}
+
+	URigVMCollapseNode * CollapseNode = NewObject<URigVMCollapseNode>(FunctionLibrary, NameForFunction);
+
+	// reinitialize the reader
+	FRigVMObjectArchiveReader Reader(CopyOfArchive, CollapseNode);
+
+	// perform the import
+	Reader << CollapseNode;
+
+	if(CollapseNode->GetContainedGraph() == nullptr)
+	{
+		ReportError(TEXT("Unable to deserialize archive."));
+		RemoveNode(CollapseNode, false, false);
+		return nullptr;
+	}
+
+	if(!AddGraphNode(CollapseNode, false))
+	{
+		ReportError(TEXT("Unable to add imported function to function library."));
+		DestroyObject(CollapseNode);
+		return nullptr;
+	}
+
+	// add the variant information back
+	FunctionLibrary->FunctionToVariant.Add(NameForFunction, Variant);
+
+	// notify everyone
+	Notify(ERigVMGraphNotifType::NodeAdded, CollapseNode);
+	if(URigVMController* ContainedGraphController = GetControllerForGraph(CollapseNode->GetContainedGraph()))
+	{
+		ContainedGraphController->ResendAllNotifications();
+	}
+
+	return CollapseNode;
+}
+
 URigVMLibraryNode* URigVMController::LocalizeFunctionFromPath(const FString& InHostPath, const FName& InFunctionName, bool bLocalizeDependentPrivateFunctions, bool bSetupUndoRedo, bool bPrintPythonCommand)
 {
 	if (!IsValidGraph())
@@ -4683,11 +4826,23 @@ TMap<FRigVMGraphFunctionIdentifier, URigVMLibraryNode*> URigVMController::Locali
 					ClientHost->GetRigVMClient()->UpdateGraphFunctionSerializedGraph(ReferencedFunction);
 				}
 			}
-			
-			TArray<FName> NodeNames = LibraryController->ImportNodesFromText(FunctionToLocalize->SerializedCollapsedNode, false);
-			if (NodeNames.Num() > 0)
+
+			URigVMLibraryNode* LocalizedFunction = nullptr;
+			if(FunctionToLocalize->CollapseNodeArchive.IsEmpty())
+			{
+				TArray<FName> NodeNames = LibraryController->ImportNodesFromText(FunctionToLocalize->SerializedCollapsedNode_DEPRECATED, false);
+				if (NodeNames.Num() > 0)
+				{
+					LocalizedFunction = ThisLibrary->FindFunction(NodeNames[0]);
+				}
+			}
+			else
+			{
+				LocalizedFunction = LibraryController->ImportFunctionFromArchive(FunctionToLocalize->CollapseNodeArchive);
+			}
+
+			if(LocalizedFunction)
 			{
-				URigVMLibraryNode* LocalizedFunction = ThisLibrary->FindFunction(NodeNames[0]);
 				LocalizedFunctions.Add(FunctionToLocalize->Header.LibraryPointer, LocalizedFunction);
 				ThisLibrary->LocalizedFunctions.FindOrAdd(FunctionToLocalize->Header.LibraryPointer.GetLibraryNodePath(), LocalizedFunction);
 			}
@@ -13897,7 +14052,19 @@ URigVMLibraryNode* URigVMController::CreateFunctionVariant(const FName& InFuncti
 	}
 
 	ClientHost->GetRigVMClient()->UpdateGraphFunctionSerializedGraph(FunctionToClone);
-	TArray<FName> NodeNames = ImportNodesFromText(OriginalFunction->SerializedCollapsedNode, false);
+
+	TArray<FName> NodeNames;
+	if(OriginalFunction->CollapseNodeArchive.IsEmpty())
+	{
+		NodeNames = ImportNodesFromText(OriginalFunction->SerializedCollapsedNode_DEPRECATED, false);
+	}
+	else
+	{
+		if(const URigVMLibraryNode* NewCollapseNode = ImportFunctionFromArchive(OriginalFunction->CollapseNodeArchive, VariantName))
+		{
+			NodeNames.Add(NewCollapseNode->GetFName());
+		}
+	}
 
 	if (!NodeNames.IsEmpty())
 	{
diff --git a/Engine/Plugins/Runtime/RigVM/Source/RigVMDeveloper/Public/RigVMModel/RigVMController.h b/Engine/Plugins/Runtime/RigVM/Source/RigVMDeveloper/Public/RigVMModel/RigVMController.h
--- a/Engine/Plugins/Runtime/RigVM/Source/RigVMDeveloper/Public/RigVMModel/RigVMController.h
+++ b/Engine/Plugins/Runtime/RigVM/Source/RigVMDeveloper/Public/RigVMModel/RigVMController.h
@@ -19,6 +19,7 @@
 #include "RigVMModel/Nodes/RigVMInvokeEntryNode.h"
 #include "RigVMModel/RigVMBuildData.h"
 #include "RigVMCore/RigVMUserWorkflow.h"
+#include "RigVMCore/RigVMObjectArchive.h"
 #include "UObject/Interface.h"
 #include "Algo/Sort.h"
 #include "RigVMController.generated.h"
@@ -629,6 +630,12 @@ public:
 	UFUNCTION(BlueprintCallable, Category = RigVMController)
 	TArray<FName> ImportNodesFromText(const FString& InText, bool bSetupUndoRedo = true, bool bPrintPythonCommands = false);
 
+	// Exports the given function to a binary archive
+	bool ExportFunctionToArchive(const FName& InFunctionName, FRigVMObjectArchive& OutArchive);
+
+	// Imports a function from a binary archive
+	URigVMLibraryNode* ImportFunctionFromArchive(const FRigVMObjectArchive& InArchive, const FName& InFunctionName = FName(NAME_None));
+
 	// Copies a function declaration into this graph's local function library
 	UFUNCTION(BlueprintCallable, Category = RigVMController)
 	URigVMLibraryNode* LocalizeFunctionFromPath(const FString& InHostPath, const FName& InFunctionName, bool bLocalizeDependentPrivateFunctions = true, bool bSetupUndoRedo = true, bool bPrintPythonCommand = false);
