RigVM: Archive tracing to track what contributes to the size of a package

#rnx

diff --git a/Engine/Plugins/Animation/ControlRig/Source/ControlRig/Private/ControlRig.cpp b/Engine/Plugins/Animation/ControlRig/Source/ControlRig/Private/ControlRig.cpp
--- a/Engine/Plugins/Animation/ControlRig/Source/ControlRig/Private/ControlRig.cpp
+++ b/Engine/Plugins/Animation/ControlRig/Source/ControlRig/Private/ControlRig.cpp
@@ -1794,14 +1794,17 @@ TArray<UControlRig*> UControlRig::FindControlRigs(UObject* Outer, TSubclassOf<UC
 void UControlRig::Serialize(FArchive& Ar)
 {
 	DECLARE_SCOPE_HIERARCHICAL_COUNTER_FUNC()
+	UE_RIGVM_ARCHIVETRACE_SCOPE(Ar, FString::Printf(TEXT("UControlRig(%s)"), *GetName()));
 
 	Super::Serialize(Ar);
+	UE_RIGVM_ARCHIVETRACE_ENTRY(Ar, TEXT("Super::Serialize"));
 
 	Ar.UsingCustomVersion(FControlRigObjectVersion::GUID);
 
 	if (Ar.CustomVer(FControlRigObjectVersion::GUID) >= FControlRigObjectVersion::ControlRigStoresPhysicsSolvers)
 	{
 		Ar << PhysicsSolvers;
+		UE_RIGVM_ARCHIVETRACE_ENTRY(Ar, TEXT("PhysicsSolvers"));
 	}
 	else
 	{
diff --git a/Engine/Plugins/Animation/ControlRig/Source/ControlRig/Private/ControlRigBlueprintGeneratedClass.cpp b/Engine/Plugins/Animation/ControlRig/Source/ControlRig/Private/ControlRigBlueprintGeneratedClass.cpp
--- a/Engine/Plugins/Animation/ControlRig/Source/ControlRig/Private/ControlRigBlueprintGeneratedClass.cpp
+++ b/Engine/Plugins/Animation/ControlRig/Source/ControlRig/Private/ControlRigBlueprintGeneratedClass.cpp
@@ -15,7 +15,10 @@ UControlRigBlueprintGeneratedClass::UControlRigBlueprintGeneratedClass(const FOb
 
 void UControlRigBlueprintGeneratedClass::Serialize(FArchive& Ar)
 {
+	UE_RIGVM_ARCHIVETRACE_SCOPE(Ar, FString::Printf(TEXT("UControlRigBlueprintGeneratedClass(%s)"), *GetName()));
+
 	DECLARE_SCOPE_HIERARCHICAL_COUNTER_FUNC()
+	UE_RIGVM_ARCHIVETRACE_ENTRY(Ar, TEXT("Super::Serialize"));
 
 	// don't use URigVMBlueprintGeneratedClass
 	// to avoid backwards compat issues.
@@ -43,6 +46,7 @@ void UControlRigBlueprintGeneratedClass::Serialize(FArchive& Ar)
 	}
 	
 	VM->Serialize(Ar);
+	UE_RIGVM_ARCHIVETRACE_ENTRY(Ar, TEXT("VM"));
 
 	if (UControlRig* CDO = Cast<UControlRig>(GetDefaultObject(false)))
 	{
@@ -58,5 +62,6 @@ void UControlRigBlueprintGeneratedClass::Serialize(FArchive& Ar)
 	}
 	
 	Ar << GraphFunctionStore;
+	UE_RIGVM_ARCHIVETRACE_ENTRY(Ar, TEXT("GraphFunctionStore"));
 }
 
diff --git a/Engine/Plugins/Animation/ControlRig/Source/ControlRig/Private/Rigs/RigHierarchy.cpp b/Engine/Plugins/Animation/ControlRig/Source/ControlRig/Private/Rigs/RigHierarchy.cpp
--- a/Engine/Plugins/Animation/ControlRig/Source/ControlRig/Private/Rigs/RigHierarchy.cpp
+++ b/Engine/Plugins/Animation/ControlRig/Source/ControlRig/Private/Rigs/RigHierarchy.cpp
@@ -184,11 +184,14 @@ void URigHierarchy::AddReferencedObjects(UObject* InpThis, FReferenceCollector&
 void URigHierarchy::Save(FArchive& Ar)
 {
 	FScopeLock Lock(&ElementsLock);
-	
+	UE_RIGVM_ARCHIVETRACE_SCOPE(Ar, FString::Printf(TEXT("URigHierarchy(%s)"), *GetName()));
+
 	if(Ar.IsTransacting())
 	{
 		Ar << TransformStackIndex;
+		UE_RIGVM_ARCHIVETRACE_ENTRY(Ar, TEXT("TransformStackIndex"));
 		Ar << bTransactingForTransformChange;
+		UE_RIGVM_ARCHIVETRACE_ENTRY(Ar, TEXT("bTransactingForTransformChange"));
 		
 		if(bTransactingForTransformChange)
 		{
@@ -197,6 +200,7 @@ void URigHierarchy::Save(FArchive& Ar)
 
 		TArray<FRigElementKey> SelectedKeys = GetSelectedKeys();
 		Ar << SelectedKeys;
+		UE_RIGVM_ARCHIVETRACE_ENTRY(Ar, TEXT("SelectedKeys"));
 	}
 
 	// make sure all parts of pose are valid.
@@ -207,6 +211,7 @@ void URigHierarchy::Save(FArchive& Ar)
 
 	int32 ElementCount = Elements.Num();
 	Ar << ElementCount;
+	UE_RIGVM_ARCHIVETRACE_ENTRY(Ar, TEXT("ElementCount"));
 
 	for(int32 ElementIndex = 0; ElementIndex < ElementCount; ElementIndex++)
 	{
@@ -228,9 +233,12 @@ void URigHierarchy::Save(FArchive& Ar)
 		FRigBaseElement* Element = Elements[ElementIndex];
 		Element->Serialize(Ar, FRigBaseElement::InterElementData);
 	}
+	UE_RIGVM_ARCHIVETRACE_ENTRY(Ar, TEXT("Elements"));
 
 	Ar << PreviousNameMap;
+	UE_RIGVM_ARCHIVETRACE_ENTRY(Ar, TEXT("PreviousNameMap"));
 	Ar << PreviousParentMap;
+	UE_RIGVM_ARCHIVETRACE_ENTRY(Ar, TEXT("PreviousParentMap"));
 
 	{
 		TMap<FRigElementKey, FMetadataStorage> ElementMetadataToSave;
@@ -244,6 +252,7 @@ void URigHierarchy::Save(FArchive& Ar)
 		}
 		
 		Ar << ElementMetadataToSave;
+		UE_RIGVM_ARCHIVETRACE_ENTRY(Ar, TEXT("ElementMetadataToSave"));
 	}
 }
 
diff --git a/Engine/Plugins/Animation/ControlRig/Source/ControlRigDeveloper/Private/ControlRigBlueprint.cpp b/Engine/Plugins/Animation/ControlRig/Source/ControlRigDeveloper/Private/ControlRigBlueprint.cpp
--- a/Engine/Plugins/Animation/ControlRig/Source/ControlRigDeveloper/Private/ControlRigBlueprint.cpp
+++ b/Engine/Plugins/Animation/ControlRig/Source/ControlRigDeveloper/Private/ControlRigBlueprint.cpp
@@ -847,6 +847,8 @@ void UControlRigBlueprint::SetPreviewMesh(USkeletalMesh* PreviewMesh, bool bMark
 
 void UControlRigBlueprint::Serialize(FArchive& Ar)
 {
+	UE_RIGVM_ARCHIVETRACE_SCOPE(Ar, FString::Printf(TEXT("UControlRigBlueprint(%s)"), *GetName()));
+
 	if(IsValid(this))
 	{
 		RigVMClient.SetOuterClientHost(this, GET_MEMBER_NAME_CHECKED(UControlRigBlueprint, RigVMClient));
@@ -854,6 +856,7 @@ void UControlRigBlueprint::Serialize(FArchive& Ar)
 	}
 	
 	Super::Serialize(Ar);
+	UE_RIGVM_ARCHIVETRACE_ENTRY(Ar, TEXT("Super::Serialize"));
 
 	if(Ar.IsObjectReferenceCollector())
 	{
diff --git a/Engine/Plugins/Runtime/RigVM/Source/RigVM/Private/RigVMBlueprintGeneratedClass.cpp b/Engine/Plugins/Runtime/RigVM/Source/RigVM/Private/RigVMBlueprintGeneratedClass.cpp
--- a/Engine/Plugins/Runtime/RigVM/Source/RigVM/Private/RigVMBlueprintGeneratedClass.cpp
+++ b/Engine/Plugins/Runtime/RigVM/Source/RigVM/Private/RigVMBlueprintGeneratedClass.cpp
@@ -43,7 +43,10 @@ void URigVMBlueprintGeneratedClass::PostInitInstance(UObject* InObj, FObjectInst
 
 void URigVMBlueprintGeneratedClass::Serialize(FArchive& Ar)
 {
+	UE_RIGVM_ARCHIVETRACE_SCOPE(Ar, FString::Printf(TEXT("URigVMBlueprintGeneratedClass(%s)"), *GetName()));
+
 	Super::Serialize(Ar);
+	UE_RIGVM_ARCHIVETRACE_ENTRY(Ar, TEXT("Super::Serialize"));
 
 	Ar.UsingCustomVersion(FUE5MainStreamObjectVersion::GUID);
 
@@ -79,6 +82,7 @@ void URigVMBlueprintGeneratedClass::Serialize(FArchive& Ar)
 	}
 	
 	VM->Serialize(Ar);
+	UE_RIGVM_ARCHIVETRACE_ENTRY(Ar, TEXT("VM"));
 
 	if (!IsSkeletonClass(this))
 	{
@@ -92,6 +96,7 @@ void URigVMBlueprintGeneratedClass::Serialize(FArchive& Ar)
 	}
 
 	Ar << GraphFunctionStore;
+	UE_RIGVM_ARCHIVETRACE_ENTRY(Ar, TEXT("GraphFunctionStore"));
 }
 
 void URigVMBlueprintGeneratedClass::PostLoad()
diff --git a/Engine/Plugins/Runtime/RigVM/Source/RigVM/Private/RigVMCore/RigVM.cpp b/Engine/Plugins/Runtime/RigVM/Source/RigVM/Private/RigVMCore/RigVM.cpp
--- a/Engine/Plugins/Runtime/RigVM/Source/RigVM/Private/RigVMCore/RigVM.cpp
+++ b/Engine/Plugins/Runtime/RigVM/Source/RigVM/Private/RigVMCore/RigVM.cpp
@@ -139,19 +139,30 @@ void URigVM::Serialize(FArchive& Ar)
 
 void URigVM::Save(FArchive& Ar)
 {
-	// The Save function has to be in sync with CopyDataForSerialization
+	UE_RIGVM_ARCHIVETRACE_SCOPE(Ar, FString::Printf(TEXT("URigVM(%s)"), *GetName()));
 
+	// The Save function has to be in sync with CopyDataForSerialization
 	Ar << CachedVMHash;
+	UE_RIGVM_ARCHIVETRACE_ENTRY(Ar, TEXT("CachedVMHash"));
+
 	Ar << ExternalPropertyPathDescriptions;
+	UE_RIGVM_ARCHIVETRACE_ENTRY(Ar, TEXT("ExternalPropertyPathDescriptions"));
 	Ar << FunctionNamesStorage;
+	UE_RIGVM_ARCHIVETRACE_ENTRY(Ar, TEXT("FunctionNamesStorage"));
 	Ar << ByteCodeStorage;
+	UE_RIGVM_ARCHIVETRACE_ENTRY(Ar, TEXT("ByteCodeStorage"));
 	Ar << Parameters;
+	UE_RIGVM_ARCHIVETRACE_ENTRY(Ar, TEXT("Parameters"));
 
 	Ar << OperandToDebugRegisters;
+	UE_RIGVM_ARCHIVETRACE_ENTRY(Ar, TEXT("OperandToDebugRegisters"));
 	
 	Ar << LiteralMemoryStorage;
+	UE_RIGVM_ARCHIVETRACE_ENTRY(Ar, TEXT("LiteralMemoryStorage"));
 	Ar << DefaultWorkMemoryStorage;
+	UE_RIGVM_ARCHIVETRACE_ENTRY(Ar, TEXT("DefaultWorkMemoryStorage"));
 	Ar << DefaultDebugMemoryStorage;
+	UE_RIGVM_ARCHIVETRACE_ENTRY(Ar, TEXT("DefaultDebugMemoryStorage"));
 }
 
 void URigVM::Load(FArchive& Ar)
diff --git a/Engine/Plugins/Runtime/RigVM/Source/RigVM/Private/RigVMCore/RigVMByteCode.cpp b/Engine/Plugins/Runtime/RigVM/Source/RigVM/Private/RigVMCore/RigVMByteCode.cpp
--- a/Engine/Plugins/Runtime/RigVM/Source/RigVM/Private/RigVMCore/RigVMByteCode.cpp
+++ b/Engine/Plugins/Runtime/RigVM/Source/RigVM/Private/RigVMCore/RigVMByteCode.cpp
@@ -292,6 +292,8 @@ void FRigVMByteCode::Serialize(FArchive& Ar)
 
 void FRigVMByteCode::Save(FArchive& Ar)
 {
+	UE_RIGVM_ARCHIVETRACE_SCOPE(Ar, TEXT("FRigVMByteCode"));
+
 	FRigVMInstructionArray Instructions;
 
 	int32 InstructionCount = 0;
@@ -447,6 +449,8 @@ void FRigVMByteCode::Save(FArchive& Ar)
 		}
 	}
 
+	UE_RIGVM_ARCHIVETRACE_ENTRY(Ar, TEXT("Operators"));
+
 	UScriptStruct* ScriptStruct = FRigVMByteCodeEntry::StaticStruct();
 	TArray<uint8, TAlignedHeapAllocator<16>> DefaultStructData;
 	DefaultStructData.AddZeroed(ScriptStruct->GetStructureSize());
@@ -461,11 +465,14 @@ void FRigVMByteCode::Save(FArchive& Ar)
 	}
 
 	Ar << View;
+	UE_RIGVM_ARCHIVETRACE_ENTRY(Ar, TEXT("Entries"));
 
 	TArray<FRigVMBranchInfo> TempBranchInfos = BranchInfos;
 	Ar << TempBranchInfos;
+	UE_RIGVM_ARCHIVETRACE_ENTRY(Ar, TEXT("BranchInfos"));
 
 	Ar << PublicContextPathName;
+	UE_RIGVM_ARCHIVETRACE_ENTRY(Ar, TEXT("PublicContextPathName"));
 }
 
 void FRigVMByteCode::Load(FArchive& Ar)
diff --git a/Engine/Plugins/Runtime/RigVM/Source/RigVM/Private/RigVMCore/RigVMDefines.cpp b/Engine/Plugins/Runtime/RigVM/Source/RigVM/Private/RigVMCore/RigVMDefines.cpp
new file mode 100644
--- /dev/null
+++ b/Engine/Plugins/Runtime/RigVM/Source/RigVM/Private/RigVMCore/RigVMDefines.cpp
@@ -0,0 +1,114 @@
+// Copyright Epic Games, Inc. All Rights Reserved.
+
+#include "RigVMDefines.h"
+
+#if UE_RIGVM_ARCHIVETRACE_ENABLE
+
+TMap<FArchive*, TSharedPtr<FRigVMArchiveTrace>> FRigVMArchiveTrace::ActiveTraces;
+FCriticalSection FRigVMArchiveTrace::AddRemoveTracesMutex;
+
+FRigVMArchiveTrace* FRigVMArchiveTrace::AddRefTrace(FArchive* InArchive)
+{
+	const FScopeLock Lock(&AddRemoveTracesMutex);
+	if(const TSharedPtr<FRigVMArchiveTrace>* ExistingTrace = ActiveTraces.Find(InArchive))
+	{
+		ExistingTrace->Get()->Counter++;
+		return ExistingTrace->Get();
+	}
+
+	const TSharedPtr<FRigVMArchiveTrace> NewTrace = MakeShared<FRigVMArchiveTrace>();
+	ActiveTraces.Add(InArchive, NewTrace);
+	NewTrace->Archive = InArchive;
+	return NewTrace.Get();
+}
+
+void FRigVMArchiveTrace::DecRefTrace(FRigVMArchiveTrace* InTrace)
+{
+	InTrace->Counter--;
+	if(InTrace->Counter == 0)
+	{
+		const FScopeLock Lock(&AddRemoveTracesMutex);
+		verify(ActiveTraces.Remove(InTrace->Archive) == 1);
+	}
+}
+
+FRigVMArchiveTraceBracket::FRigVMArchiveTraceBracket(FArchive& InArchive, const FString& InScope)
+: Trace(nullptr)
+, Indentation(0)
+, ArchivePos(0)
+, LastArchivePos(0)
+, bEnabled(true)
+{
+	Trace = FRigVMArchiveTrace::AddRefTrace(&InArchive);
+	Indentation = Trace->Counter - 1;
+
+	bEnabled = InArchive.IsSaving() && !InArchive.IsObjectReferenceCollector() && !InArchive.IsTransacting() && InArchive.IsPersistent();
+
+	if(bEnabled)
+	{
+		ArchivePos = InArchive.GetArchiveState().Tell();
+		LastArchivePos = ArchivePos;
+		
+		ArchiveName = InArchive.GetArchiveName();
+		ArchiveWhiteSpace = WhiteSpace(Indentation * 2);
+		ArchivePrefix = FString::Printf(TEXT("  %s%s"), *ArchiveWhiteSpace, *InScope);
+
+		const FString ArchiveOffset = ArchiveOffsetToString(ArchivePos);
+		static const FString EntrySize = WhiteSpace(14);
+		UE_LOG(LogRigVM, Display, TEXT("%s %s%s %s%s"), *ArchiveName, *ArchiveOffset, *EntrySize, *ArchiveWhiteSpace, *InScope);
+	}
+}
+
+FRigVMArchiveTraceBracket::~FRigVMArchiveTraceBracket()
+{
+	FRigVMArchiveTrace::DecRefTrace(Trace);
+}
+
+void FRigVMArchiveTraceBracket::AddEntry(FArchive& InArchive, const FString& InScope)
+{
+	if(bEnabled)
+	{
+		ArchivePos = InArchive.GetArchiveState().Tell();
+		const FString ArchiveOffset = FRigVMArchiveTraceBracket::ArchiveOffsetToString(LastArchivePos);
+		const FString EntrySize = FRigVMArchiveTraceBracket::ArchiveOffsetToString(ArchivePos - LastArchivePos);
+		UE_LOG(LogRigVM, Display, TEXT("%s %s,%s %s %s"), *ArchiveName, *ArchiveOffset, *EntrySize, *ArchivePrefix, *InScope);
+		Swap(ArchivePos, LastArchivePos);
+	}
+}
+
+FString FRigVMArchiveTraceBracket::WhiteSpace(int32 InCount)
+{
+	if(InCount <= 0)
+	{
+		return FString();
+	}
+	return FString(TEXT("\t")).ConvertTabsToSpaces(InCount);
+}
+
+FString FRigVMArchiveTraceBracket::ArchiveOffsetToString(int64 InOffset)
+{
+	const FString OffsetStr = FString::FromInt(InOffset);
+	FString FormattedOffsetStr = WhiteSpace(FMath::Max(11, int32(OffsetStr.Len() / 3) * 4 + 3)); 
+	int32 Input = OffsetStr.Len() - 1;
+	int32 Output = FormattedOffsetStr.Len() - 1;
+	int32 CharsLeft = 3;
+	while(Input >= 0)
+	{
+		FormattedOffsetStr[Output] = OffsetStr[Input];
+		Input--;
+		Output--;
+		CharsLeft--;
+
+		if(CharsLeft == 0 && Input >= 0)
+		{
+			FormattedOffsetStr[Output] = TCHAR('.');
+			Output--;
+			CharsLeft = 3;
+		}
+	}
+	static const FString Prefix = TEXT("[");
+	static const FString Suffix = TEXT("]");
+	return Prefix + FormattedOffsetStr + Suffix;
+}
+
+#endif
diff --git a/Engine/Plugins/Runtime/RigVM/Source/RigVM/Private/RigVMCore/RigVMGraphFunctionDefinition.cpp b/Engine/Plugins/Runtime/RigVM/Source/RigVM/Private/RigVMCore/RigVMGraphFunctionDefinition.cpp
--- a/Engine/Plugins/Runtime/RigVM/Source/RigVM/Private/RigVMCore/RigVMGraphFunctionDefinition.cpp
+++ b/Engine/Plugins/Runtime/RigVM/Source/RigVM/Private/RigVMCore/RigVMGraphFunctionDefinition.cpp
@@ -252,18 +252,25 @@ FRigVMGraphFunctionData* FRigVMGraphFunctionHeader::GetFunctionData(bool bLoadIf
 
 FArchive& operator<<(FArchive& Ar, FRigVMGraphFunctionHeader& Data)
 {
+	UE_RIGVM_ARCHIVETRACE_SCOPE(Ar, FString::Printf(TEXT("FRigVMGraphFunctionHeader(%s)"), *Data.Name.ToString()));
+
 	Ar.UsingCustomVersion(FRigVMObjectVersion::GUID);
 	
 	Ar << Data.LibraryPointer;
+	UE_RIGVM_ARCHIVETRACE_ENTRY(Ar, TEXT("LibraryPointer"));
 
 	if (!Ar.IsLoading() || Ar.CustomVer(FRigVMObjectVersion::GUID) >= FRigVMObjectVersion::AddVariantToFunctionIdentifier)
 	{
 		Ar << Data.Variant;
+		UE_RIGVM_ARCHIVETRACE_ENTRY(Ar, TEXT("Variant"));
 	}
 	
 	Ar << Data.Name;
+	UE_RIGVM_ARCHIVETRACE_ENTRY(Ar, TEXT("Name"));
 	Ar << Data.NodeTitle;
+	UE_RIGVM_ARCHIVETRACE_ENTRY(Ar, TEXT("NodeTitle"));
 	Ar << Data.NodeColor;
+	UE_RIGVM_ARCHIVETRACE_ENTRY(Ar, TEXT("NodeColor"));
 
 	if (Ar.IsLoading())
 	{
@@ -279,13 +286,19 @@ FArchive& operator<<(FArchive& Ar, FRigVMGraphFunctionHeader& Data)
 	else
 	{
 		Ar << Data.Description;
+		UE_RIGVM_ARCHIVETRACE_ENTRY(Ar, TEXT("Description"));
 	}
 	
 	Ar << Data.Category;
+	UE_RIGVM_ARCHIVETRACE_ENTRY(Ar, TEXT("Category"));
 	Ar << Data.Keywords;
+	UE_RIGVM_ARCHIVETRACE_ENTRY(Ar, TEXT("Keywords"));
 	Ar << Data.Arguments;
+	UE_RIGVM_ARCHIVETRACE_ENTRY(Ar, TEXT("Arguments"));
 	Ar << Data.Dependencies;
+	UE_RIGVM_ARCHIVETRACE_ENTRY(Ar, TEXT("Dependencies"));
 	Ar << Data.ExternalVariables;
+	UE_RIGVM_ARCHIVETRACE_ENTRY(Ar, TEXT("ExternalVariables"));
 
 	if (Ar.IsLoading())
 	{
@@ -301,6 +314,7 @@ FArchive& operator<<(FArchive& Ar, FRigVMGraphFunctionHeader& Data)
 	else
 	{
 		Ar << Data.Layout;
+		UE_RIGVM_ARCHIVETRACE_ENTRY(Ar, TEXT("Layout"));
 	}
 	
 	return Ar;
@@ -308,8 +322,12 @@ FArchive& operator<<(FArchive& Ar, FRigVMGraphFunctionHeader& Data)
 
 FArchive& operator<<(FArchive& Ar, FRigVMGraphFunctionData& Data)
 {
+	UE_RIGVM_ARCHIVETRACE_SCOPE(Ar, FString::Printf(TEXT("FRigVMGraphFunctionStore(%s)"), *Data.Header.Name.ToString()));
+
 	Ar << Data.Header;
+	UE_RIGVM_ARCHIVETRACE_ENTRY(Ar, TEXT("Header"));
 	Ar << Data.CompilationData;
+	UE_RIGVM_ARCHIVETRACE_ENTRY(Ar, TEXT("CompilationData"));
 
 	if (Ar.CustomVer(FUE5MainStreamObjectVersion::GUID) < FUE5MainStreamObjectVersion::RigVMSaveSerializedGraphInGraphFunctionData)
 	{
@@ -321,6 +339,7 @@ FArchive& operator<<(FArchive& Ar, FRigVMGraphFunctionData& Data)
 	}
 
 	Ar << Data.SerializedCollapsedNode_DEPRECATED;
+	UE_RIGVM_ARCHIVETRACE_ENTRY(Ar, TEXT("SerializedCollapsedNode"));
 
 	if (Ar.CustomVer(FRigVMObjectVersion::GUID) < FRigVMObjectVersion::RigVMSaveSerializedGraphInGraphFunctionDataAsByteArray)
 	{
@@ -332,6 +351,7 @@ FArchive& operator<<(FArchive& Ar, FRigVMGraphFunctionData& Data)
 	}
 
 	Ar << Data.CollapseNodeArchive;
+	UE_RIGVM_ARCHIVETRACE_ENTRY(Ar, TEXT("CollapseNodeArchive"));
 
 #if UE_BUILD_SHIPPING
 	// these two members can store substantial data
diff --git a/Engine/Plugins/Runtime/RigVM/Source/RigVM/Private/RigVMCore/RigVMMemoryStorageStruct.cpp b/Engine/Plugins/Runtime/RigVM/Source/RigVM/Private/RigVMCore/RigVMMemoryStorageStruct.cpp
--- a/Engine/Plugins/Runtime/RigVM/Source/RigVM/Private/RigVMCore/RigVMMemoryStorageStruct.cpp
+++ b/Engine/Plugins/Runtime/RigVM/Source/RigVM/Private/RigVMCore/RigVMMemoryStorageStruct.cpp
@@ -38,10 +38,14 @@ FRigVMMemoryStorageStruct::~FRigVMMemoryStorageStruct()
 
 bool FRigVMMemoryStorageStruct::Serialize(FArchive& Ar)
 {
+	UE_RIGVM_ARCHIVETRACE_SCOPE(Ar, TEXT("FRigVMMemoryStorageStruct"));
 	Super::Serialize(Ar);
+	UE_RIGVM_ARCHIVETRACE_ENTRY(Ar, TEXT("Super::Serialize"));
 
 	Ar << MemoryType;
+	UE_RIGVM_ARCHIVETRACE_ENTRY(Ar, TEXT("MemoryType"));
 	Ar << PropertyPathDescriptions;
+	UE_RIGVM_ARCHIVETRACE_ENTRY(Ar, TEXT("PropertyPathDescriptions"));
 
 	if (Ar.IsLoading())
 	{
diff --git a/Engine/Plugins/Runtime/RigVM/Source/RigVM/Private/RigVMHost.cpp b/Engine/Plugins/Runtime/RigVM/Source/RigVM/Private/RigVMHost.cpp
--- a/Engine/Plugins/Runtime/RigVM/Source/RigVM/Private/RigVMHost.cpp
+++ b/Engine/Plugins/Runtime/RigVM/Source/RigVM/Private/RigVMHost.cpp
@@ -119,7 +119,10 @@ UWorld* URigVMHost::GetWorld() const
 
 void URigVMHost::Serialize(FArchive& Ar)
 {
+	UE_RIGVM_ARCHIVETRACE_SCOPE(Ar, FString::Printf(TEXT("URigVMHost(%s)"), *GetName()));
+
 	Super::Serialize(Ar);
+	UE_RIGVM_ARCHIVETRACE_ENTRY(Ar, TEXT("Super::Serialize"));
 
 	Ar.UsingCustomVersion(FRigVMObjectVersion::GUID);
 
diff --git a/Engine/Plugins/Runtime/RigVM/Source/RigVM/Public/RigVMCore/RigVMGraphFunctionHost.h b/Engine/Plugins/Runtime/RigVM/Source/RigVM/Public/RigVMCore/RigVMGraphFunctionHost.h
--- a/Engine/Plugins/Runtime/RigVM/Source/RigVM/Public/RigVMCore/RigVMGraphFunctionHost.h
+++ b/Engine/Plugins/Runtime/RigVM/Source/RigVM/Public/RigVMCore/RigVMGraphFunctionHost.h
@@ -70,13 +70,17 @@ public:
 
 	friend FArchive& operator<<(FArchive& Ar, FRigVMGraphFunctionStore& Host)
 	{
+		UE_RIGVM_ARCHIVETRACE_SCOPE(Ar, TEXT("FRigVMGraphFunctionStore"));
+
 		Ar << Host.PublicFunctions;
+		UE_RIGVM_ARCHIVETRACE_ENTRY(Ar, TEXT("PublicFunctions"));
 
 		// This is only added to make sure SoftObjectPaths can be gathered and fixed up in the case of asset rename
 		// It should not affect data on disk
 		if (Ar.IsObjectReferenceCollector())
 		{
 			Ar << Host.PrivateFunctions;
+			UE_RIGVM_ARCHIVETRACE_ENTRY(Ar, TEXT("PrivateFunctions"));
 		}
 
 		return Ar;
diff --git a/Engine/Plugins/Runtime/RigVM/Source/RigVM/Public/RigVMDefines.h b/Engine/Plugins/Runtime/RigVM/Source/RigVM/Public/RigVMDefines.h
--- a/Engine/Plugins/Runtime/RigVM/Source/RigVM/Public/RigVMDefines.h
+++ b/Engine/Plugins/Runtime/RigVM/Source/RigVM/Public/RigVMDefines.h
@@ -41,3 +41,86 @@
 #ifndef UE_RIGVMCONTROLLER_VERBOSE_REPOPULATE
 	#define UE_RIGVMCONTROLLER_VERBOSE_REPOPULATE 0
 #endif
+
+#ifndef UE_RIGVM_ARCHIVETRACE_ENABLE
+	#define UE_RIGVM_ARCHIVETRACE_ENABLE 0
+#endif
+
+#if WITH_EDITOR
+#else
+	#undef UE_RIGVM_ARCHIVETRACE_ENABLE
+	#define UE_RIGVM_ARCHIVETRACE_ENABLE 0
+#endif
+
+#if UE_RIGVM_ARCHIVETRACE_ENABLE
+
+#include "Containers/Map.h"
+#include "Templates/SharedPointer.h"
+
+class RIGVM_API FRigVMArchiveTrace
+{
+public:
+
+	FRigVMArchiveTrace()
+		: Counter(1)
+		, Archive(nullptr)
+	{
+	}
+
+private:
+	
+	int32 Counter;
+	FArchive* Archive;
+
+	static FRigVMArchiveTrace* AddRefTrace(FArchive* InArchive);
+	static void DecRefTrace(FRigVMArchiveTrace* InTrace);
+
+	static TMap<FArchive*, TSharedPtr<FRigVMArchiveTrace>> ActiveTraces;
+	static FCriticalSection AddRemoveTracesMutex;
+	
+	friend class FRigVMArchiveTraceBracket;
+};
+
+class RIGVM_API FRigVMArchiveTraceBracket
+{
+public:
+
+	FRigVMArchiveTraceBracket(FArchive& InArchive, const FString& InScope);
+	~FRigVMArchiveTraceBracket();
+
+	void AddEntry(FArchive& InArchive, const FString& InScope);
+
+	static FString WhiteSpace(int32 InCount);
+	static FString ArchiveOffsetToString(int64 InOffset);
+	
+private:
+
+	FRigVMArchiveTrace* Trace;
+	int32 Indentation;
+	int64 ArchivePos;
+	int64 LastArchivePos;
+	bool bEnabled;
+	FString ArchiveName;
+	FString ArchiveWhiteSpace;
+	FString ArchivePrefix;
+};
+#endif
+
+#ifndef UE_RIGVM_ARCHIVETRACE_SCOPE
+	#if UE_RIGVM_ARCHIVETRACE_ENABLE
+		#define UE_RIGVM_ARCHIVETRACE_SCOPE(Ar, Scope) \
+			FRigVMArchiveTraceBracket _ArchiveTraceBracket(Ar, FString(Scope));
+	#else
+		#define UE_RIGVM_ARCHIVETRACE_SCOPE(Ar, Scope)
+	#endif
+#endif
+
+// if you are hitting an error with this - you likely have to place UE_RIGVM_ARCHIVETRACE_SCOPE in the top of the scope
+#ifndef UE_RIGVM_ARCHIVETRACE_ENTRY
+	#if UE_RIGVM_ARCHIVETRACE_ENABLE
+		#define UE_RIGVM_ARCHIVETRACE_ENTRY(Ar, Scope) \
+			_ArchiveTraceBracket.AddEntry(Ar, Scope);
+	#else
+		#define UE_RIGVM_ARCHIVETRACE_ENTRY(Ar, Scope)
+	#endif
+#endif
diff --git a/Engine/Plugins/Runtime/RigVM/Source/RigVMDeveloper/Private/RigVMBlueprint.cpp b/Engine/Plugins/Runtime/RigVM/Source/RigVMDeveloper/Private/RigVMBlueprint.cpp
--- a/Engine/Plugins/Runtime/RigVM/Source/RigVMDeveloper/Private/RigVMBlueprint.cpp
+++ b/Engine/Plugins/Runtime/RigVM/Source/RigVMDeveloper/Private/RigVMBlueprint.cpp
@@ -783,6 +783,8 @@ IRigVMEditorModule* URigVMBlueprint::GetEditorModule() const
 
 void URigVMBlueprint::Serialize(FArchive& Ar)
 {
+	UE_RIGVM_ARCHIVETRACE_SCOPE(Ar, FString::Printf(TEXT("URigVMBlueprint(%s)"), *GetName()));
+	
 	Ar.UsingCustomVersion(FRigVMObjectVersion::GUID);
 	
 	if(IsValid(this))
@@ -791,6 +793,7 @@ void URigVMBlueprint::Serialize(FArchive& Ar)
 	}
 	
 	Super::Serialize(Ar);
+	UE_RIGVM_ARCHIVETRACE_ENTRY(Ar, TEXT("Super::Serialize"));
 
 	if(Ar.IsObjectReferenceCollector())
 	{
