RigVM: Reuse literal properties across functions

#rnx

diff --git a/Engine/Plugins/Runtime/RigVM/Source/RigVMDeveloper/Private/RigVMCompiler/RigVMCompiler.cpp b/Engine/Plugins/Runtime/RigVM/Source/RigVMDeveloper/Private/RigVMCompiler/RigVMCompiler.cpp
--- a/Engine/Plugins/Runtime/RigVM/Source/RigVMDeveloper/Private/RigVMCompiler/RigVMCompiler.cpp
+++ b/Engine/Plugins/Runtime/RigVM/Source/RigVMDeveloper/Private/RigVMCompiler/RigVMCompiler.cpp
@@ -2623,24 +2623,45 @@ bool URigVMCompiler::TraverseInlineFunction(const FRigVMInlineFunctionExprAST* I
 			}
 
 			auto FindOrAddProperty = [&WorkData, MemoryType, LiteralValueToSetupTraitArg]
-			(const FRigVMFunctionCompilationPropertyDescription& InProperty, const FString& InNewName, const bool bIsExecuteState) -> FRigVMOperand
+			(const FRigVMFunctionCompilationPropertyDescription& InProperty, const FString& InNewName, const bool bMakeUniqueProperty) -> FRigVMOperand
 			{
 				// Sharing / reusing memory / operands happens as per following contract:
 				// 1. properties are only shared if their CPP type matches
 				// 2. Literal / constant memory is only shared if the constant values match (and it is not a trait list)
 				// 3. Work state is only shared if it is not internal work state private to the instruction referring to it
-				// 4. Work state of type FRigVMInstructionSetExecuteState (bIsExecuteState) is never shared either since it is work state private to a lazy branch.
+				// 4. Work state of type FRigVMInstructionSetExecuteState (bMakeUniqueProperty) is never shared either since it is work state private to a lazy branch.
 
 				FString NewName = InNewName;
 				int32 Suffix = 0;
 
+				// always share literals
+				if(MemoryType == ERigVMMemoryType::Literal && !bMakeUniqueProperty && WorkData.PropertyDescriptions.Contains(MemoryType))
+				{
+					const TArray<FRigVMPropertyDescription>& LiteralDescriptions = WorkData.PropertyDescriptions[MemoryType]; 
+					for(int32 Index = 0; Index < LiteralDescriptions.Num(); Index++)
+					{
+						if(!LiteralValueToSetupTraitArg.Contains(Index))
+						{
+							const FRigVMPropertyDescription& ExistingProperty = LiteralDescriptions[Index];
+							if(ExistingProperty.CPPType.Equals(InProperty.CPPType, ESearchCase::CaseSensitive))
+							{
+								// if the value is the same and this is not a trait setup list
+								if(ExistingProperty.DefaultValue.Equals(InProperty.DefaultValue, ESearchCase::CaseSensitive))
+								{
+									return FRigVMOperand(MemoryType, Index);
+								}
+							}
+						}
+					}
+				}
+
 				FRigVMOperand Operand = WorkData.FindProperty(MemoryType, *NewName);
 
 				// look for a free / unused operand
 				while(Operand.IsValid())
 				{
 					const FRigVMPropertyDescription& ExistingProperty = WorkData.PropertyDescriptions[MemoryType][Operand.GetRegisterIndex()];
-					if(!bIsExecuteState &&
+					if(!bMakeUniqueProperty &&
 						ExistingProperty.CPPType.Equals(InProperty.CPPType, ESearchCase::CaseSensitive))
 					{
 						if(MemoryType == ERigVMMemoryType::Literal)
@@ -2698,7 +2719,7 @@ bool URigVMCompiler::TraverseInlineFunction(const FRigVMInlineFunctionExprAST* I
 					FRigVMPropertyDescription::SanitizeName(NewName);
 				}
 
-				const FRigVMOperand Operand = FindOrAddProperty(Description, NewName, bIsExecuteState);
+				const FRigVMOperand Operand = FindOrAddProperty(Description, NewName, bIsExecuteState || bIsTrait);
 				FRigVMCompilerWorkData::FFunctionRegisterData Data = {FunctionReferenceNode, MemoryType, PropertyIndex};
 				WorkData.FunctionRegisterToOperand.Add(Data, Operand);
 
