[RigVM] Name related function should use ignorecase instead of case sensitive because names in packaged build may have unpredictable capitalization

diff --git a/Engine/Plugins/Runtime/RigVM/Source/RigVM/Private/RigVMFunctions/RigVMFunction_Name.cpp b/Engine/Plugins/Runtime/RigVM/Source/RigVM/Private/RigVMFunctions/RigVMFunction_Name.cpp
--- a/Engine/Plugins/Runtime/RigVM/Source/RigVM/Private/RigVMFunctions/RigVMFunction_Name.cpp
+++ b/Engine/Plugins/Runtime/RigVM/Source/RigVM/Private/RigVMFunctions/RigVMFunction_Name.cpp
@@ -117,7 +117,7 @@ FRigVMFunction_NameReplace_Execute()
 {
 	check(ExecuteContext.GetNameCache());
 
-	Result = ExecuteContext.GetNameCache()->Replace(Name, Old, New, ESearchCase::CaseSensitive);
+	Result = ExecuteContext.GetNameCache()->Replace(Name, Old, New, ESearchCase::IgnoreCase);
 }
 
 #if WITH_DEV_AUTOMATION_TESTS
@@ -163,7 +163,7 @@ FRigVMFunction_EndsWith_Execute()
 {
 	check(ExecuteContext.GetNameCache());
 
-	Result = ExecuteContext.GetNameCache()->EndsWith(Name, Ending, ESearchCase::CaseSensitive);
+	Result = ExecuteContext.GetNameCache()->EndsWith(Name, Ending, ESearchCase::IgnoreCase);
 }
 
 #if WITH_DEV_AUTOMATION_TESTS
@@ -193,7 +193,7 @@ FRigVMFunction_StartsWith_Execute()
 {
 	check(ExecuteContext.GetNameCache());
 
-	Result = ExecuteContext.GetNameCache()->StartsWith(Name, Start, ESearchCase::CaseSensitive);
+	Result = ExecuteContext.GetNameCache()->StartsWith(Name, Start, ESearchCase::IgnoreCase);
 }
 
 #if WITH_DEV_AUTOMATION_TESTS
@@ -222,7 +222,7 @@ FRigVMFunction_Contains_Execute()
 {
 	check(ExecuteContext.GetNameCache());
 
-	Result = ExecuteContext.GetNameCache()->Contains(Name, Search, ESearchCase::CaseSensitive);
+	Result = ExecuteContext.GetNameCache()->Contains(Name, Search, ESearchCase::IgnoreCase);
 }
 
 #if WITH_DEV_AUTOMATION_TESTS
