RigVM: Fix struct redirector handling in the RigVM registry

A redirector that only changes the package/module/plugin rather than the name itself can cause a stack overflow. Added shortcut to skip recursion if the name does not change when redirected.

diff --git a/Engine/Plugins/Runtime/RigVM/Source/RigVM/Private/RigVMCore/RigVMRegistry.cpp b/Engine/Plugins/Runtime/RigVM/Source/RigVM/Private/RigVMCore/RigVMRegistry.cpp
--- a/Engine/Plugins/Runtime/RigVM/Source/RigVM/Private/RigVMCore/RigVMRegistry.cpp
+++ b/Engine/Plugins/Runtime/RigVM/Source/RigVM/Private/RigVMCore/RigVMRegistry.cpp
@@ -1672,20 +1672,25 @@ const FRigVMFunction* FRigVMRegistry_NoLock::FindFunction_NoLock(const TCHAR* In
 			for(const FCoreRedirect* Redirect : Redirects)
 			{
 				FString NewStructOrFactoryName = Redirect->NewName.ObjectName.ToString();
-				if(bIsDispatchFactory)
-				{
-					NewStructOrFactoryName = FRigVMDispatchFactory::DispatchPrefix + NewStructOrFactoryName;
-				}
-				else
-				{
-					NewStructOrFactoryName = StructPrefix + NewStructOrFactoryName;
-				}
-				const FRigVMFunction* RedirectedFunction = FindFunction_NoLock(*(NewStructOrFactoryName + TEXT("::") + SuffixString), InTypeResolver);
-				if(RedirectedFunction)
+
+				// Check name differs - this could just be a struct that moved package
+				if(NewStructOrFactoryName != StructOrFactoryName)
 				{
-					FRigVMRegistry_NoLock& MutableRegistry = Get(LockType_Write);
-					MutableRegistry.FunctionNameToIndex.Add(InName, RedirectedFunction->Index);
-					return RedirectedFunction;
+					if(bIsDispatchFactory)
+					{
+						NewStructOrFactoryName = FRigVMDispatchFactory::DispatchPrefix + NewStructOrFactoryName;
+					}
+					else
+					{
+						NewStructOrFactoryName = StructPrefix + NewStructOrFactoryName;
+					}
+					const FRigVMFunction* RedirectedFunction = FindFunction_NoLock(*(NewStructOrFactoryName + TEXT("::") + SuffixString), InTypeResolver);
+					if(RedirectedFunction)
+					{
+						FRigVMRegistry_NoLock& MutableRegistry = Get(LockType_Write);
+						MutableRegistry.FunctionNameToIndex.Add(InName, RedirectedFunction->Index);
+						return RedirectedFunction;
+					}
 				}
 			}
 		}
