RigVM: Make sure to unroot ed graphs for old assets

#rnx

diff --git a/Engine/Plugins/Runtime/RigVM/Source/RigVMDeveloper/Private/RigVMBlueprint.cpp b/Engine/Plugins/Runtime/RigVM/Source/RigVMDeveloper/Private/RigVMBlueprint.cpp
--- a/Engine/Plugins/Runtime/RigVM/Source/RigVMDeveloper/Private/RigVMBlueprint.cpp
+++ b/Engine/Plugins/Runtime/RigVM/Source/RigVMDeveloper/Private/RigVMBlueprint.cpp
@@ -4143,6 +4143,10 @@ bool URigVMBlueprint::RemoveEdGraphForCollapseNode(URigVMCollapseNode* InNode, b
 
 						FunctionGraphs.Remove(RigFunctionGraph);
 						RigFunctionGraph->Rename(nullptr, GetTransientPackage(), REN_DontCreateRedirectors);
+						if(RigFunctionGraph->IsRooted())
+						{
+							RigFunctionGraph->RemoveFromRoot();
+						}
 						RigFunctionGraph->MarkAsGarbage();
 						return bNotify;
 					}
@@ -4172,6 +4176,10 @@ bool URigVMBlueprint::RemoveEdGraphForCollapseNode(URigVMCollapseNode* InNode, b
 
 						RigGraph->SubGraphs.Remove(SubRigGraph);
 						SubRigGraph->Rename(nullptr, GetTransientPackage(), REN_DontCreateRedirectors);
+						if(SubRigGraph->IsRooted())
+						{
+							SubRigGraph->RemoveFromRoot();
+						}
 						SubRigGraph->MarkAsGarbage();
 						return bNotify;
 					}
