RigVMCompiler: Cache results for GetFirstChildOfType, GetFirstParentOfType and GetBlocksImpl

#rnx

diff --git a/Engine/Plugins/Runtime/RigVM/Source/RigVMDeveloper/Private/RigVMCompiler/RigVMAST.cpp b/Engine/Plugins/Runtime/RigVM/Source/RigVMDeveloper/Private/RigVMCompiler/RigVMAST.cpp
--- a/Engine/Plugins/Runtime/RigVM/Source/RigVMDeveloper/Private/RigVMCompiler/RigVMAST.cpp
+++ b/Engine/Plugins/Runtime/RigVM/Source/RigVMDeveloper/Private/RigVMCompiler/RigVMAST.cpp
@@ -27,6 +27,7 @@ FRigVMExprAST::FRigVMExprAST(EType InType, const FRigVMASTProxy& InProxy)
 	: Name(NAME_None)
 	, Type(InType)
 	, Index(INDEX_NONE)
+	, ParserPtr(nullptr)
 {
 }
 
@@ -105,21 +106,53 @@ const FRigVMExprAST* FRigVMExprAST::GetParent() const
 
 const FRigVMExprAST* FRigVMExprAST::GetFirstParentOfType(EType InExprType) const
 {
+	if(const FRigVMParserAST* Parser = GetParser())
+	{
+		if(FirstParentOfTypeCacheVersion.Get(INDEX_NONE) == Parser->CacheVersion)
+		{
+			if(const FRigVMExprAST* const* ExistingFirstParent = CachedFirstParentOfType.Find(InExprType))
+			{
+				return *ExistingFirstParent;
+			}
+		}
+		else
+		{
+			// if the cache version no longer matches,
+			// clear the map to make sure we re-retrieve the first parent accordingly.
+			CachedFirstParentOfType.Reset();
+		}
+	}
+
+	const FRigVMExprAST* FirstParent = nullptr;
+	
 	for(const FRigVMExprAST* Parent : Parents)
 	{
 		if (Parent->IsA(InExprType))
 		{
-			return Parent;
+			FirstParent = Parent;
+			break;
 		}
 	}
-	for (const FRigVMExprAST* Parent : Parents)
+
+	if(FirstParent == nullptr)
 	{
-		if (const FRigVMExprAST* GrandParent = Parent->GetFirstParentOfType(InExprType))
+		for (const FRigVMExprAST* Parent : Parents)
 		{
-			return GrandParent;
+			if (const FRigVMExprAST* GrandParent = Parent->GetFirstParentOfType(InExprType))
+			{
+				FirstParent = GrandParent;
+				break;
+			}
 		}
 	}
-	return nullptr;
+
+	if(const FRigVMParserAST* Parser = GetParser())
+	{
+		FirstParentOfTypeCacheVersion = Parser->CacheVersion;
+		CachedFirstParentOfType.FindOrAdd(InExprType, nullptr) = FirstParent;
+	}
+
+	return FirstParent;
 }
 
 bool FRigVMExprAST::IsParentedTo(const FRigVMExprAST* InParentExpr) const
@@ -148,21 +181,53 @@ bool FRigVMExprAST::IsParentOf(const FRigVMExprAST* InChildExpr) const
 
 const FRigVMExprAST* FRigVMExprAST::GetFirstChildOfType(EType InExprType) const
 {
+	if(const FRigVMParserAST* Parser = GetParser())
+	{
+		if(FirstChildOfTypeCacheVersion.Get(INDEX_NONE) == Parser->CacheVersion)
+		{
+			if(const FRigVMExprAST* const* ExistingFirstChild = CachedFirstChildOfType.Find(InExprType))
+			{
+				return *ExistingFirstChild;
+			}
+		}
+		else
+		{
+			// if the cache version no longer matches,
+			// clear the map to make sure we re-retrieve the first child accordingly.
+			CachedFirstChildOfType.Reset();
+		}
+	}
+
+	const FRigVMExprAST* FirstChild = nullptr;
+	
 	for (const FRigVMExprAST* Child : Children)
 	{
 		if (Child->IsA(InExprType))
 		{
-			return Child;
+			FirstChild = Child;
+			break;
 		}
 	}
-	for (const FRigVMExprAST* Child : Children)
+
+	if(FirstChild == nullptr)
 	{
-		if (const FRigVMExprAST* GrandChild = Child->GetFirstChildOfType(InExprType))
+		for (const FRigVMExprAST* Child : Children)
 		{
-			return GrandChild;
+			if (const FRigVMExprAST* GrandChild = Child->GetFirstChildOfType(InExprType))
+			{
+				FirstChild = GrandChild;
+				break;
+			}
 		}
 	}
-	return nullptr;
+
+	if(const FRigVMParserAST* Parser = GetParser())
+	{
+		FirstChildOfTypeCacheVersion = Parser->CacheVersion;
+		CachedFirstChildOfType.FindOrAdd(InExprType, nullptr) = FirstChild;
+	}
+
+	return FirstChild;
 }
 
 const FRigVMBlockExprAST* FRigVMExprAST::GetBlock() const
@@ -394,10 +459,30 @@ void FRigVMExprAST::GetBlocksImpl(FRigVMBlockArray& InOutBlocks) const
 		return;
 	}
 	
-	for(int32 ParentIndex = 0; ParentIndex < NumParents(); ParentIndex++)
+	const FRigVMParserAST* Parser = GetParser();
+	check(Parser);
+	
+	if(BlocksCacheVersion.Get(INDEX_NONE) != Parser->CacheVersion)
 	{
-		const FRigVMExprAST* ParentExpression = ParentAt(ParentIndex);
-		ParentExpression->GetBlocksImpl(InOutBlocks);
+		CachedBlocks.Reset();
+		for(int32 ParentIndex = 0; ParentIndex < NumParents(); ParentIndex++)
+		{
+			const FRigVMExprAST* ParentExpression = ParentAt(ParentIndex);
+			ParentExpression->GetBlocksImpl(CachedBlocks);
+		}
+		BlocksCacheVersion = Parser->CacheVersion;
+	}
+
+	if(InOutBlocks.IsEmpty())
+	{
+		InOutBlocks = CachedBlocks;
+	}
+	else
+	{
+		for(const FRigVMBlockExprAST* Block : CachedBlocks)
+		{
+			InOutBlocks.AddUnique(Block);
+		}
 	}
 }
 
@@ -532,8 +617,12 @@ void FRigVMExprAST::ReplaceChild(FRigVMExprAST* InCurrentChild, FRigVMExprAST* I
 			Children[ChildIndex] = InNewChild;
 			InCurrentChild->Parents.Remove(this);
 			InNewChild->Parents.Add(this);
-			InCurrentChild->InvalidateCaches();
-			InNewChild->InvalidateCaches();
+
+			TArray<bool> InvalidateCachesProcessed;
+			InvalidateCachesProcessed.AddZeroed(GetParser()->Expressions.Num());
+			InCurrentChild->InvalidateCachesImpl(InvalidateCachesProcessed);
+			InNewChild->InvalidateCachesImpl(InvalidateCachesProcessed);
+			InvalidateCachesImpl(InvalidateCachesProcessed);
 		}
 	}
 }
@@ -623,9 +712,17 @@ void FRigVMExprAST::InvalidateCachesImpl(TArray<bool>& OutProcessed)
 	}
 	
 	BlockCombinationHash.Reset();
+	BlocksCacheVersion.Reset();
+	FirstChildOfTypeCacheVersion.Reset();
+	FirstParentOfTypeCacheVersion.Reset();
 	MaximumDepth.Reset();
 	OutProcessed[Index] = true;
 
+	if(const FRigVMParserAST* Parser = GetParser())
+	{
+		Parser->IncrementCacheVersion();
+	}
+
 	for(FRigVMExprAST* ChildExpression : Children)
 	{
 		ChildExpression->InvalidateCachesImpl(OutProcessed);
@@ -927,7 +1024,8 @@ void FRigVMParserASTSettings::Report(EMessageSeverity::Type InSeverity, UObject*
 const TArray<FRigVMASTProxy> FRigVMParserAST::EmptyProxyArray;
 
 FRigVMParserAST::FRigVMParserAST(TArray<URigVMGraph*> InGraphs, URigVMController* InController, const FRigVMParserASTSettings& InSettings, const TArray<FRigVMExternalVariable>& InExternalVariables)
-	: LibraryNodeBeingCompiled(nullptr)
+	: CacheVersion(0)
+	, LibraryNodeBeingCompiled(nullptr)
 {
 	DECLARE_SCOPE_HIERARCHICAL_COUNTER_FUNC()
 
@@ -2793,6 +2891,7 @@ void FRigVMParserAST::RemoveExpressions(TArray<FRigVMExprAST*> InExprs)
 	}
 
 	RefreshExprIndices();
+	IncrementCacheVersion();
 }
 
 void FRigVMParserAST::TraverseParents(const FRigVMExprAST* InExpr, TFunctionRef<bool(const FRigVMExprAST*)> InContinuePredicate)
@@ -3470,3 +3569,8 @@ FString FRigVMParserAST::GetLinkAsString(const FRigVMASTLinkDescription& InLink)
 		 FString::Printf(TEXT("%s%s%s"), *TargetPin->GetPinPath(),
 			*(InLink.SegmentPath.IsEmpty() ? EmptyString : PeriodString), *InLink.SegmentPath));
 }
+
+void FRigVMParserAST::IncrementCacheVersion() const
+{
+	CacheVersion++;
+}
diff --git a/Engine/Plugins/Runtime/RigVM/Source/RigVMDeveloper/Public/RigVMCompiler/RigVMAST.h b/Engine/Plugins/Runtime/RigVM/Source/RigVMDeveloper/Public/RigVMCompiler/RigVMAST.h
--- a/Engine/Plugins/Runtime/RigVM/Source/RigVMDeveloper/Public/RigVMCompiler/RigVMAST.h
+++ b/Engine/Plugins/Runtime/RigVM/Source/RigVMDeveloper/Public/RigVMCompiler/RigVMAST.h
@@ -342,6 +342,12 @@ protected:
 	TMap<FName, int32> PinNameToChildIndex;
 	mutable TOptional<TOptional<uint32>> BlockCombinationHash;
 	mutable TOptional<int32> MaximumDepth;
+	mutable TOptional<int32> BlocksCacheVersion;
+	mutable FRigVMBlockArray CachedBlocks;
+	mutable TOptional<int32> FirstChildOfTypeCacheVersion;
+	mutable TMap<int32, const FRigVMExprAST*> CachedFirstChildOfType;
+	mutable TOptional<int32> FirstParentOfTypeCacheVersion;
+	mutable TMap<int32, const FRigVMExprAST*> CachedFirstParentOfType;
 
 	friend class FRigVMParserAST;
 	friend class URigVMCompiler;
@@ -1416,6 +1422,8 @@ private:
 
 	static FString GetLinkAsString(const FRigVMASTLinkDescription& InLink);
 
+	void IncrementCacheVersion() const;
+
 	TMap<FRigVMASTProxy, FRigVMExprAST*> SubjectToExpression;
 	TMap<FRigVMASTProxy, int32> NodeExpressionIndex;
 	TArray<FRigVMExprAST*> Expressions;
@@ -1423,6 +1431,7 @@ private:
 	TArray<FRigVMExprAST*> DeletedExpressions;
 	FRigVMBlockExprAST* ObsoleteBlock;
 	mutable TMap<uint32, FString> BlockCombinationHashToName;
+	mutable int32 CacheVersion;
 
 	TArray<FRigVMASTProxy> NodeProxies;
 	TMap<FRigVMASTProxy, FRigVMASTProxy> SharedOperandPins;
