RigVM: Rely on compression format specific bounds rather than magic numbers

#rnx

diff --git a/Engine/Plugins/Runtime/RigVM/Source/RigVM/Private/RigVMCore/RigVMObjectArchive.cpp b/Engine/Plugins/Runtime/RigVM/Source/RigVM/Private/RigVMCore/RigVMObjectArchive.cpp
--- a/Engine/Plugins/Runtime/RigVM/Source/RigVM/Private/RigVMCore/RigVMObjectArchive.cpp
+++ b/Engine/Plugins/Runtime/RigVM/Source/RigVM/Private/RigVMCore/RigVMObjectArchive.cpp
@@ -41,7 +41,7 @@ void FRigVMObjectArchive::Compress()
 	}
 	
 	// It is possible for compression to actually increase the size of the data, so we over allocate here to handle that.
-	CompressedSize = UncompressedSize * 4 / 3;
+	CompressedSize = FCompression::CompressMemoryBound(NAME_Zlib, UncompressedSize);
 
 	TArray<uint8> CompressedBuffer;
 	CompressedBuffer.SetNumUninitialized(CompressedSize);
