Added RigVM support to set node titles and import nodes from text at a specific locatioon

diff --git a/Engine/Plugins/Runtime/RigVM/Source/RigVMDeveloper/Private/EdGraph/RigVMEdGraph.cpp b/Engine/Plugins/Runtime/RigVM/Source/RigVMDeveloper/Private/EdGraph/RigVMEdGraph.cpp
--- a/Engine/Plugins/Runtime/RigVM/Source/RigVMDeveloper/Private/EdGraph/RigVMEdGraph.cpp
+++ b/Engine/Plugins/Runtime/RigVM/Source/RigVMDeveloper/Private/EdGraph/RigVMEdGraph.cpp
@@ -405,6 +405,7 @@ bool URigVMEdGraph::HandleModifiedEvent_Internal(ERigVMGraphNotifType InNotifTyp
 		}
 		case ERigVMGraphNotifType::NodeDescriptionChanged:
 		case ERigVMGraphNotifType::NodeCategoryChanged:
+		case ERigVMGraphNotifType::NodeTitleChanged:
 		{
 			if (URigVMNode* ModelNode = Cast<URigVMNode>(InSubject))
 			{
diff --git a/Engine/Plugins/Runtime/RigVM/Source/RigVMDeveloper/Private/RigVMModel/RigVMClient.cpp b/Engine/Plugins/Runtime/RigVM/Source/RigVMDeveloper/Private/RigVMModel/RigVMClient.cpp
--- a/Engine/Plugins/Runtime/RigVM/Source/RigVMDeveloper/Private/RigVMModel/RigVMClient.cpp
+++ b/Engine/Plugins/Runtime/RigVM/Source/RigVMDeveloper/Private/RigVMModel/RigVMClient.cpp
@@ -1562,6 +1562,7 @@ void FRigVMClient::HandleGraphModifiedEvent(ERigVMGraphNotifType InNotifType, UR
 		case ERigVMGraphNotifType::NodeCategoryChanged: // A node's category has changed (Subject == URigVMNode)
 		case ERigVMGraphNotifType::NodeKeywordsChanged: // A node's keywords have changed (Subject == URigVMNode)
 		case ERigVMGraphNotifType::NodeDescriptionChanged: // A node's description has changed (Subject == URigVMNode)
+		case ERigVMGraphNotifType::NodeTitleChanged: // A node's title has changed (Subject == URigVMNode)
 		case ERigVMGraphNotifType::VariantTagsChanged: // The tags in the header of this function variant have changed (Subject == URigVMLibraryNode)
 		{
 			if(InSubject->GetOuter()->IsA<URigVMFunctionLibrary>())
diff --git a/Engine/Plugins/Runtime/RigVM/Source/RigVMDeveloper/Private/RigVMModel/RigVMController.cpp b/Engine/Plugins/Runtime/RigVM/Source/RigVMDeveloper/Private/RigVMModel/RigVMController.cpp
--- a/Engine/Plugins/Runtime/RigVM/Source/RigVMDeveloper/Private/RigVMModel/RigVMController.cpp
+++ b/Engine/Plugins/Runtime/RigVM/Source/RigVMDeveloper/Private/RigVMModel/RigVMController.cpp
@@ -4050,6 +4050,17 @@ FString URigVMController::ExportNodesToText(const TArray<FName>& InNodeNames, bo
 	return MoveTemp(Archive);
 }
 
+FString URigVMController::ExportNodeToText(const URigVMNode* InNode, bool bIncludeExteriorLinks)
+{
+	TArray<FName> AllNodeNames;
+	if (InNode != nullptr)
+	{
+		AllNodeNames.Add(InNode->GetFName());
+	}
+
+	return ExportNodesToText(AllNodeNames, bIncludeExteriorLinks);
+}
+
 FString URigVMController::ExportSelectedNodesToText(bool bIncludeExteriorLinks)
 {
 	if (!IsValidGraph())
@@ -7514,6 +7525,73 @@ bool URigVMController::SetNodeSizeByName(const FName& InNodeName, const FVector2
 	return SetNodeSize(Node, InSize, bSetupUndoRedo, bMergeUndoAction, bPrintPythonCommand);
 }
 
+bool URigVMController::SetNodeTitle(URigVMNode* InNode, const FString InNodeTitle, bool bSetupUndoRedo, bool bMergeUndoAction, bool bPrintPythonCommand)
+{
+	if (!bIsTransacting && !IsGraphEditable())
+	{
+		return false;
+	}
+
+	if (!IsValidNodeForGraph(InNode))
+	{
+		return false;
+	}
+
+	if (!GetSchema()->CanSetNodeTitle(this, InNode))
+	{
+		return false;
+	}
+
+	FRigVMSetNodeTitleAction Action;
+	if (bSetupUndoRedo)
+	{
+		Action = FRigVMSetNodeTitleAction(this, InNode, InNodeTitle);
+		Action.SetTitle(FString::Printf(TEXT("Set Node Title")));
+		GetActionStack()->BeginAction(Action);
+	}
+
+	InNode->NodeTitle = InNodeTitle;
+	Notify(ERigVMGraphNotifType::NodeTitleChanged, InNode);
+
+	if (bSetupUndoRedo)
+	{
+		GetActionStack()->EndAction(Action, bMergeUndoAction);
+	}
+
+	if (bPrintPythonCommand)
+	{
+		const FString GraphName = GetSchema()->GetSanitizedGraphName(GetGraph()->GetGraphName());
+		const FString NodePath = GetSchema()->GetSanitizedPinPath(InNode->GetNodePath());
+
+		RigVMPythonUtils::Print(GetSchema()->GetGraphOuterName(GetGraph()),
+			FString::Printf(TEXT("blueprint.get_controller_by_name('%s').set_node_title_by_name('%s', %s)"),
+				*GraphName,
+				*NodePath,
+				*InNodeTitle));
+	}
+
+	return true;
+}
+
+bool URigVMController::SetNodeTitleByName(const FName& InNodeName, const FString InNodeTitle, bool bSetupUndoRedo, bool bMergeUndoAction)
+{
+	if (!IsValidGraph())
+	{
+		return false;
+	}
+
+	if (!bIsTransacting && !IsGraphEditable())
+	{
+		return false;
+	}
+
+	URigVMGraph* Graph = GetGraph();
+	check(Graph);
+
+	URigVMNode* Node = Graph->FindNodeByName(InNodeName);
+	return SetNodeTitle(Node, InNodeTitle, bSetupUndoRedo, bMergeUndoAction);
+}
+
 bool URigVMController::SetNodeColor(URigVMNode* InNode, const FLinearColor& InColor, bool bSetupUndoRedo, bool bMergeUndoAction, bool bPrintPythonCommand)
 {
 	if (!bIsTransacting && !IsGraphEditable())
diff --git a/Engine/Plugins/Runtime/RigVM/Source/RigVMDeveloper/Private/RigVMModel/RigVMControllerActions.cpp b/Engine/Plugins/Runtime/RigVM/Source/RigVMDeveloper/Private/RigVMModel/RigVMControllerActions.cpp
--- a/Engine/Plugins/Runtime/RigVM/Source/RigVMDeveloper/Private/RigVMModel/RigVMControllerActions.cpp
+++ b/Engine/Plugins/Runtime/RigVM/Source/RigVMDeveloper/Private/RigVMModel/RigVMControllerActions.cpp
@@ -712,6 +712,60 @@ bool FRigVMSetNodeSizeAction::Redo()
 	return FRigVMBaseAction::Redo();
 }
 
+FRigVMSetNodeTitleAction::FRigVMSetNodeTitleAction()
+	: FRigVMBaseAction(nullptr)
+	, OldTitle()
+	, NewTitle()
+{
+}
+
+FRigVMSetNodeTitleAction::FRigVMSetNodeTitleAction(URigVMController* InController, URigVMNode* InNode, const FString& InNewTitle)
+	: FRigVMBaseAction(InController)
+	, NodePath(InNode->GetNodePath())
+	, OldTitle(InNode->GetNodeTitleRaw())
+	, NewTitle(InNewTitle)
+{
+}
+
+bool FRigVMSetNodeTitleAction::Merge(const FRigVMBaseAction* Other)
+{
+	if (!FRigVMBaseAction::Merge(Other))
+	{
+		return false;
+	}
+
+	const FRigVMSetNodeTitleAction* Action = (const FRigVMSetNodeTitleAction*)Other;
+	if (NodePath != Action->NodePath)
+	{
+		return false;
+	}
+
+	NewTitle = Action->NewTitle;
+	return true;
+}
+
+bool FRigVMSetNodeTitleAction::Undo()
+{
+	if (!FRigVMBaseAction::Undo())
+	{
+		return false;
+	}
+	return GetController()->SetNodeTitleByName(*NodePath, OldTitle, false);
+}
+
+bool FRigVMSetNodeTitleAction::Redo()
+{
+	if (!CanUndoRedo())
+	{
+		return false;
+	}
+	if (!GetController()->SetNodeTitleByName(*NodePath, NewTitle, false))
+	{
+		return false;
+	}
+	return FRigVMBaseAction::Redo();
+}
+
 FRigVMSetNodeColorAction::FRigVMSetNodeColorAction()
 : FRigVMBaseAction(nullptr)
 , OldColor(FLinearColor::Black)
diff --git a/Engine/Plugins/Runtime/RigVM/Source/RigVMDeveloper/Private/RigVMModel/RigVMNode.cpp b/Engine/Plugins/Runtime/RigVM/Source/RigVMDeveloper/Private/RigVMModel/RigVMNode.cpp
--- a/Engine/Plugins/Runtime/RigVM/Source/RigVMDeveloper/Private/RigVMModel/RigVMNode.cpp
+++ b/Engine/Plugins/Runtime/RigVM/Source/RigVMDeveloper/Private/RigVMModel/RigVMNode.cpp
@@ -595,6 +595,11 @@ FString URigVMNode::GetNodeTitle() const
 	return GetName();
 }
 
+const FString& URigVMNode::GetNodeTitleRaw() const
+{
+	return NodeTitle;
+}
+
 FVector2D URigVMNode::GetPosition() const
 {
 	return Position;
diff --git a/Engine/Plugins/Runtime/RigVM/Source/RigVMDeveloper/Private/RigVMModel/RigVMSchema.cpp b/Engine/Plugins/Runtime/RigVM/Source/RigVMDeveloper/Private/RigVMModel/RigVMSchema.cpp
--- a/Engine/Plugins/Runtime/RigVM/Source/RigVMDeveloper/Private/RigVMModel/RigVMSchema.cpp
+++ b/Engine/Plugins/Runtime/RigVM/Source/RigVMDeveloper/Private/RigVMModel/RigVMSchema.cpp
@@ -469,6 +469,13 @@ bool URigVMSchema::CanResizeNode(URigVMController* InController, const URigVMNod
 	return InNode->IsA<URigVMCommentNode>();
 }
 
+bool URigVMSchema::CanSetNodeTitle(URigVMController* InController, const URigVMNode* InNode) const
+{
+	RIGVMSCHEMA_DEFAULT_FUNCTION_BODY
+
+	return false;
+}
+
 bool URigVMSchema::CanRecolorNode(URigVMController* InController, const URigVMNode* InNode, const FLinearColor& InNewColor) const
 {
 	RIGVMSCHEMA_DEFAULT_FUNCTION_BODY
diff --git a/Engine/Plugins/Runtime/RigVM/Source/RigVMDeveloper/Public/RigVMModel/RigVMController.h b/Engine/Plugins/Runtime/RigVM/Source/RigVMDeveloper/Public/RigVMModel/RigVMController.h
--- a/Engine/Plugins/Runtime/RigVM/Source/RigVMDeveloper/Public/RigVMModel/RigVMController.h
+++ b/Engine/Plugins/Runtime/RigVM/Source/RigVMDeveloper/Public/RigVMModel/RigVMController.h
@@ -618,6 +618,10 @@ public:
 	UFUNCTION(BlueprintCallable, Category = RigVMController)
 	FString ExportNodesToText(const TArray<FName>& InNodeNames, bool bIncludeExteriorLinks = false);
 
+	// Exports the given node as text
+	UFUNCTION(BlueprintCallable, Category = RigVMController)
+	FString ExportNodeToText(const URigVMNode* InNode, bool bIncludeExteriorLinks = false);
+
 	// Exports the selected nodes as text
 	UFUNCTION(BlueprintCallable, Category = RigVMController)
 	FString ExportSelectedNodesToText(bool bIncludeExteriorLinks = false);
@@ -739,6 +743,16 @@ public:
 	UFUNCTION(BlueprintCallable, Category = RigVMController)
 	bool SetNodeSizeByName(const FName& InNodeName, const FVector2D& InSize, bool bSetupUndoRedo = true, bool bMergeUndoAction = false, bool bPrintPythonCommand = false);
 
+	// Sets the raw node title of a node in the graph.
+	// Some nodes generate customs node titles that override this setting.
+	UFUNCTION(BlueprintCallable, Category = RigVMController)
+	bool SetNodeTitle(URigVMNode* InNode, const FString InNodeTitle, bool bSetupUndoRedo = true, bool bMergeUndoAction = false, bool bPrintPythonCommand = false);
+
+	// Sets the raw node title of a node in the graph.
+	// Some nodes generate customs node titles that override this setting.
+	UFUNCTION(BlueprintCallable, Category = RigVMController)
+	bool SetNodeTitleByName(const FName& InNodeName, const FString InNodeTitle, bool bSetupUndoRedo = true, bool bMergeUndoAction = false);
+
 	// Sets the color of a node in the graph.
 	// This causes a NodeColorChanged modified event.
 	UFUNCTION(BlueprintCallable, Category = RigVMController)
diff --git a/Engine/Plugins/Runtime/RigVM/Source/RigVMDeveloper/Public/RigVMModel/RigVMControllerActions.h b/Engine/Plugins/Runtime/RigVM/Source/RigVMDeveloper/Public/RigVMModel/RigVMControllerActions.h
--- a/Engine/Plugins/Runtime/RigVM/Source/RigVMDeveloper/Public/RigVMModel/RigVMControllerActions.h
+++ b/Engine/Plugins/Runtime/RigVM/Source/RigVMDeveloper/Public/RigVMModel/RigVMControllerActions.h
@@ -616,6 +616,34 @@ public:
 	FVector2D NewSize;
 };
 
+/**
+ * An action setting a node's title in the graph.
+ */
+USTRUCT()
+struct FRigVMSetNodeTitleAction : public FRigVMBaseAction
+{
+	GENERATED_BODY()
+
+public:
+
+	FRigVMSetNodeTitleAction();
+	FRigVMSetNodeTitleAction(URigVMController* InController, URigVMNode* InNode, const FString& InNewTitle);
+	virtual ~FRigVMSetNodeTitleAction() = default;
+	virtual UScriptStruct* GetScriptStruct() const override { return FRigVMSetNodeTitleAction::StaticStruct(); }
+	virtual bool Merge(const FRigVMBaseAction* Other);
+	virtual bool Undo() override;
+	virtual bool Redo() override;
+
+	UPROPERTY()
+	FString NodePath;
+
+	UPROPERTY()
+	FString OldTitle;
+
+	UPROPERTY()
+	FString NewTitle;
+};
+
 /**
  * An action setting a node's color in the graph.
  */
diff --git a/Engine/Plugins/Runtime/RigVM/Source/RigVMDeveloper/Public/RigVMModel/RigVMNode.h b/Engine/Plugins/Runtime/RigVM/Source/RigVMDeveloper/Public/RigVMModel/RigVMNode.h
--- a/Engine/Plugins/Runtime/RigVM/Source/RigVMDeveloper/Public/RigVMModel/RigVMNode.h
+++ b/Engine/Plugins/Runtime/RigVM/Source/RigVMDeveloper/Public/RigVMModel/RigVMNode.h
@@ -147,6 +147,10 @@ public:
 	UFUNCTION(BlueprintCallable, Category = RigVMNode)
 	virtual FString GetNodeTitle() const;
 
+	// Returns currently node title set in the node, without any transformations
+	// Warning : The raw title might be ignored by some nodes overrides to GetNodeTitle
+	const FString& GetNodeTitleRaw() const;
+
 	// Returns the 2d position of this node - used for UI.
 	UFUNCTION(BlueprintCallable, Category = RigVMNode)
 	FVector2D GetPosition() const;
diff --git a/Engine/Plugins/Runtime/RigVM/Source/RigVMDeveloper/Public/RigVMModel/RigVMNotifications.h b/Engine/Plugins/Runtime/RigVM/Source/RigVMDeveloper/Public/RigVMModel/RigVMNotifications.h
--- a/Engine/Plugins/Runtime/RigVM/Source/RigVMDeveloper/Public/RigVMModel/RigVMNotifications.h
+++ b/Engine/Plugins/Runtime/RigVM/Source/RigVMDeveloper/Public/RigVMModel/RigVMNotifications.h
@@ -23,6 +23,7 @@ enum class ERigVMGraphNotifType : uint8
 	NodeSelectionChanged, // The set of selected nodes has changed (Subject == nullptr)
 	NodePositionChanged, // A node's position has changed (Subject == URigVMNode)
 	NodeSizeChanged, // A node's size has changed (Subject == URigVMNode)
+	NodeTitleChanged, // A node's title has changed (Subject == URigVMNode)
 	NodeColorChanged, // A node's color has changed (Subject == URigVMNode)
 	PinAdded, // A pin has been added to a given node (Subject == URigVMPin)
 	PinRemoved, // A pin has been removed from a given node (Subject == URigVMPin)
diff --git a/Engine/Plugins/Runtime/RigVM/Source/RigVMDeveloper/Public/RigVMModel/RigVMSchema.h b/Engine/Plugins/Runtime/RigVM/Source/RigVMDeveloper/Public/RigVMModel/RigVMSchema.h
--- a/Engine/Plugins/Runtime/RigVM/Source/RigVMDeveloper/Public/RigVMModel/RigVMSchema.h
+++ b/Engine/Plugins/Runtime/RigVM/Source/RigVMDeveloper/Public/RigVMModel/RigVMSchema.h
@@ -79,6 +79,9 @@ public:
 	// Returns true if a node can resized to a new size
 	virtual bool CanResizeNode(URigVMController* InController, const URigVMNode* InNode, const FVector2D& InNewSize = {100, 100}) const;
 
+	// Returns true if a node supports setting a node title
+	virtual bool CanSetNodeTitle(URigVMController* InController, const URigVMNode* InNode) const;
+
 	// Returns true if a node can recolored to a new color
 	virtual bool CanRecolorNode(URigVMController* InController, const URigVMNode* InNode, const FLinearColor& InNewColor = FLinearColor::White) const;
 
diff --git a/Engine/Plugins/Runtime/RigVM/Source/RigVMEditor/Private/Editor/RigVMEditorTools.cpp b/Engine/Plugins/Runtime/RigVM/Source/RigVMEditor/Private/Editor/RigVMEditorTools.cpp
--- a/Engine/Plugins/Runtime/RigVM/Source/RigVMEditor/Private/Editor/RigVMEditorTools.cpp
+++ b/Engine/Plugins/Runtime/RigVM/Source/RigVMEditor/Private/Editor/RigVMEditorTools.cpp
@@ -18,12 +18,40 @@ RIGVMEDITOR_API bool PasteNodes(const FVector2D& PasteLocation
 	, URigVMController* InFocusedController
 	, URigVMGraph* InFocusedModel
 	, URigVMFunctionLibrary* InLocalFunctionLibrary
-	, IRigVMGraphFunctionHost* InGraphFunctionHost)
+	, IRigVMGraphFunctionHost* InGraphFunctionHost
+	, bool bSetupUndoRedo
+	, bool bPrintPythonCommands)
+
 {
-	bool bPastePerformed = false;
+	if (bSetupUndoRedo)
+	{
+		InFocusedController->OpenUndoBracket(TEXT("Paste Nodes."));
+	}
+
+	TArray<FName> NodeNamesCreated = ImportNodesFromText(PasteLocation, TextToImport, InFocusedController, InFocusedModel, InLocalFunctionLibrary, InGraphFunctionHost);
+
+	const bool bPastePerformed = NodeNamesCreated.Num() > 0;
+	if (bPastePerformed)
+	{
+		InFocusedController->CloseUndoBracket();
+	}
+	else
+	{
+		InFocusedController->CancelUndoBracket();
+	}
 
-	InFocusedController->OpenUndoBracket(TEXT("Pasted Nodes."));
+	return bPastePerformed;
+}
 
+RIGVMEDITOR_API TArray<FName> ImportNodesFromText(const FVector2D& PasteLocation
+	, const FString& TextToImport
+	, URigVMController* InFocusedController
+	, URigVMGraph* InFocusedModel
+	, URigVMFunctionLibrary* InLocalFunctionLibrary
+	, IRigVMGraphFunctionHost* InGraphFunctionHost
+	, bool bSetupUndoRedo
+	, bool bPrintPythonCommands)
+{
 	TGuardValue<FRigVMController_RequestLocalizeFunctionDelegate> RequestLocalizeDelegateGuard(
 		InFocusedController->RequestLocalizeFunctionDelegate,
 		FRigVMController_RequestLocalizeFunctionDelegate::CreateLambda([InFocusedController, InLocalFunctionLibrary, InGraphFunctionHost](FRigVMGraphFunctionIdentifier& InFunctionToLocalize)
@@ -35,7 +63,12 @@ RIGVMEDITOR_API bool PasteNodes(const FVector2D& PasteLocation
 			})
 	);
 
-	TArray<FName> NodeNames = InFocusedController->ImportNodesFromText(TextToImport, true, true);
+	if (bSetupUndoRedo)
+	{
+		InFocusedController->OpenUndoBracket(TEXT("Import Nodes."));
+	}
+
+	TArray<FName> NodeNames = InFocusedController->ImportNodesFromText(TextToImport, bSetupUndoRedo, bPrintPythonCommands);
 
 	if (NodeNames.Num() > 0)
 	{
@@ -72,16 +105,22 @@ RIGVMEDITOR_API bool PasteNodes(const FVector2D& PasteLocation
 			check(Node);
 
 			FVector2D Position = Node->GetPosition();
-			InFocusedController->SetNodePositionByName(NodeName, PasteLocation + Position - Bounds.GetCenter(), true, false, true);
+			InFocusedController->SetNodePositionByName(NodeName, PasteLocation + Position - Bounds.GetCenter(), bSetupUndoRedo, false, bPrintPythonCommands);
 		}
 
-		InFocusedController->SetNodeSelection(NodesToSelect);
-		InFocusedController->CloseUndoBracket();
-
-		bPastePerformed = true;
+		InFocusedController->SetNodeSelection(NodesToSelect, bSetupUndoRedo, bPrintPythonCommands);
+		
+		if (NodeNames.Num() > 0)
+		{
+			InFocusedController->CloseUndoBracket();
+		}
+		else
+		{
+			InFocusedController->CancelUndoBracket();
+		}
 	}
 
-	return bPastePerformed;
+	return NodeNames;
 }
 
 RIGVMEDITOR_API void OnRequestLocalizeFunctionDialog(FRigVMGraphFunctionIdentifier& InFunction
diff --git a/Engine/Plugins/Runtime/RigVM/Source/RigVMEditor/Public/Editor/RigVMEditorTools.h b/Engine/Plugins/Runtime/RigVM/Source/RigVMEditor/Public/Editor/RigVMEditorTools.h
--- a/Engine/Plugins/Runtime/RigVM/Source/RigVMEditor/Public/Editor/RigVMEditorTools.h
+++ b/Engine/Plugins/Runtime/RigVM/Source/RigVMEditor/Public/Editor/RigVMEditorTools.h
@@ -19,7 +19,18 @@ RIGVMEDITOR_API bool PasteNodes(const FVector2D& PasteLocation
 	, URigVMController* InFocusedController
 	, URigVMGraph* InFocusedModel
 	, URigVMFunctionLibrary* InLocalFunctionLibrary
-	, IRigVMGraphFunctionHost* InGraphFunctionHost);
+	, IRigVMGraphFunctionHost* InGraphFunctionHost
+	, bool bSetupUndoRedo = true
+	, bool bPrintPythonCommands = false);
+
+RIGVMEDITOR_API TArray<FName> ImportNodesFromText(const FVector2D& PasteLocation
+	, const FString& TextToImport
+	, URigVMController* InFocusedController
+	, URigVMGraph* InFocusedModel
+	, URigVMFunctionLibrary* InLocalFunctionLibrary
+	, IRigVMGraphFunctionHost* InGraphFunctionHost
+	, bool bSetupUndoRedo = true
+	, bool bPrintPythonCommands = false);
 
 RIGVMEDITOR_API void OnRequestLocalizeFunctionDialog(FRigVMGraphFunctionIdentifier& InFunction
 	, URigVMController* InTargetController
