Make sure graphs are bound to their controller modify events.
- Fixes issues where AnimNext StateTree condtions would not run after being modified in editor

#tests EngineTest UnitTests, Lyra PF, Editor, PIE

diff --git a/Engine/Plugins/Runtime/RigVM/Source/RigVMDeveloper/Private/RigVMModel/RigVMClient.cpp b/Engine/Plugins/Runtime/RigVM/Source/RigVMDeveloper/Private/RigVMModel/RigVMClient.cpp
--- a/Engine/Plugins/Runtime/RigVM/Source/RigVMDeveloper/Private/RigVMModel/RigVMClient.cpp
+++ b/Engine/Plugins/Runtime/RigVM/Source/RigVMDeveloper/Private/RigVMModel/RigVMClient.cpp
@@ -473,6 +473,12 @@ URigVMController* FRigVMClient::GetOrCreateController(const URigVMGraph* InModel
 	{
 		if(URigVMController* Controller = GetController(InModel))
 		{
+			// We associate controllers to graphs via soft path, so they can match newly created graphs
+			// If this happens make sure the graph is correctly bound to the controller
+			if (!InModel->ModifiedEvent.IsBound())
+			{
+				const_cast<URigVMGraph*>(InModel)->OnModified().AddUObject(Controller, &URigVMController::HandleModifiedEvent);
+			}
 			return Controller;
 		}
 		return CreateController(InModel);
