[RigVM & ControlRig] Avoiding loading a soft path that can point to the asset itself when loading the asset. Delay the access to be after postload so that edgraphs don't get created using data at an intermedate state and that loader don't complain/log spam about cyclic references during load

diff --git a/Engine/Plugins/Runtime/RigVM/Source/RigVMDeveloper/Private/RigVMBlueprintUtils.cpp b/Engine/Plugins/Runtime/RigVM/Source/RigVMDeveloper/Private/RigVMBlueprintUtils.cpp
--- a/Engine/Plugins/Runtime/RigVM/Source/RigVMDeveloper/Private/RigVMBlueprintUtils.cpp
+++ b/Engine/Plugins/Runtime/RigVM/Source/RigVMDeveloper/Private/RigVMBlueprintUtils.cpp
@@ -110,6 +110,15 @@ void FRigVMBlueprintUtils::HandleRefreshAllNodes(UBlueprint* InBlueprint)
 	
 	if(const URigVMBlueprint* RigVMBlueprint = Cast<URigVMBlueprint>(InBlueprint))
 	{
+		// Avoid refreshing EdGraph if PostLoad() hasn't been called, since reconstruct node later
+		// can access model data that hasn't been fully loaded. And it is ok to skip here because
+		// the EdGraph will be reconstructed later when the CR editor
+		// initializes, as that is when the EdGraph is actually used.
+		if (RigVMBlueprint->HasAnyFlags(RF_NeedPostLoad))
+		{
+			return;
+		}
+		
 		if (RigVMBlueprint->GetDefaultModel() == nullptr)
 		{
 			return;
