RigVM: Fix CIS failure on Mac / Linux

#rnx

diff --git a/Engine/Plugins/Runtime/RigVM/Source/RigVM/Public/RigVMCore/RigVMTemplate.h b/Engine/Plugins/Runtime/RigVM/Source/RigVM/Public/RigVMCore/RigVMTemplate.h
--- a/Engine/Plugins/Runtime/RigVM/Source/RigVM/Public/RigVMCore/RigVMTemplate.h
+++ b/Engine/Plugins/Runtime/RigVM/Source/RigVM/Public/RigVMCore/RigVMTemplate.h
@@ -29,6 +29,7 @@ struct FRigVMTemplate;
 struct FRigVMDispatchContext;
 struct FRigVMUserDefinedTypeResolver;
 struct FRigVMRegistry_NoLock;
+class FRigVMTypeCacheScope_NoLock;
 
 typedef TMap<FName, TRigVMTypeIndex> FRigVMTemplateTypeMap;
 
