RigVM: Performance optimization: Avoid TSet<TArray> and preallocate as much as possible

#rnx

diff --git a/Engine/Plugins/Runtime/RigVM/Source/RigVM/Private/RigVMCore/RigVMRegistry.cpp b/Engine/Plugins/Runtime/RigVM/Source/RigVM/Private/RigVMCore/RigVMRegistry.cpp
--- a/Engine/Plugins/Runtime/RigVM/Source/RigVM/Private/RigVMCore/RigVMRegistry.cpp
+++ b/Engine/Plugins/Runtime/RigVM/Source/RigVM/Private/RigVMCore/RigVMRegistry.cpp
@@ -1889,64 +1889,66 @@ const FRigVMTemplate* FRigVMRegistry_NoLock::AddTemplateFromArguments_NoLock(con
 	}
 
 	// Remove duplicate permutations
+	TArray<FRigVMTypeCacheScope_NoLock> TypeCaches;
+	if(!Template.Arguments.IsEmpty())
 	{
-		TArray<int32> ToRemove;
-		
+		TArray<bool> ToRemove;
+		int32 NumToRemove = 0;
+		ToRemove.AddZeroed(NumPermutations);
+
+		TSet<uint32> EncounteredPermutations;
+		EncounteredPermutations.Reserve(NumPermutations);
+
 		const int32 NumArguments = Template.Arguments.Num();
-		if (NumArguments == 1)
+		TypeCaches.SetNum(NumArguments);
+		
+		bool bAnyArgumentWithZeroTypes = false;
+		for(int32 ArgIndex = 0; ArgIndex < NumArguments; ArgIndex++)
 		{
-			TSet< TRigVMTypeIndex > PermutationTypes; PermutationTypes.Reserve(NumPermutations);
-			const FRigVMTypeCacheScope_NoLock TypeCache(Template.Arguments[0]);
-			for(int32 Index = 0; Index < NumPermutations; Index++)
-			{
-				const TRigVMTypeIndex ArgType = TypeCache.GetTypeIndex_NoLock(Index);
-				if (PermutationTypes.Contains(ArgType))
-				{
-					ToRemove.Add(Index);
-				}
-				else
-				{
-					PermutationTypes.Add(ArgType);
-				}
-			}
+			(void)TypeCaches[ArgIndex].UpdateIfRequired(Template.Arguments[ArgIndex]);
+			bAnyArgumentWithZeroTypes = bAnyArgumentWithZeroTypes || TypeCaches[ArgIndex].GetNumTypes_NoLock() == 0;
 		}
-		else
+
+		if(!bAnyArgumentWithZeroTypes)
 		{
-			TSet< TArray<TRigVMTypeIndex> > PermutationTypes;
-			PermutationTypes.Reserve(NumPermutations);
-			TArray<TRigVMTypeIndex> ArgTypes; ArgTypes.SetNum(NumArguments);
-			TArray<FRigVMTypeCacheScope_NoLock> TypeCaches;
-			TypeCaches.SetNum(NumArguments);
 			for(int32 Index = 0; Index < NumPermutations; Index++)
 			{
-				for(int32 ArgIndex = 0; ArgIndex < NumArguments; ArgIndex++)
+				uint32 Hash = GetTypeHash(TypeCaches[0].GetTypeIndex_NoLock(Index));
+				for(int32 ArgIndex = 1; ArgIndex < NumArguments; ArgIndex++)
 				{
-					(void)TypeCaches[ArgIndex].UpdateIfRequired(Template.Arguments[ArgIndex]);
-					ArgTypes[ArgIndex] = TypeCaches[ArgIndex].GetTypeIndex_NoLock(Index);
+					Hash = HashCombine(Hash, GetTypeHash(TypeCaches[ArgIndex].GetTypeIndex_NoLock(Index)));
 				}
-				
-				if (PermutationTypes.Contains(ArgTypes))
+
+				if (EncounteredPermutations.Contains(Hash))
 				{
-					ToRemove.Add(Index);
+					ToRemove[Index] = true;
+					NumToRemove++;
 				}
 				else
 				{
-					PermutationTypes.Add(ArgTypes);
+					EncounteredPermutations.Add(Hash);
 				}
 			}
 		}
 		
-		for (int32 i=ToRemove.Num()-1; i>=0; --i)
+		if(NumToRemove > 0)
 		{
 			for(FRigVMTemplateArgument& Argument : Template.Arguments)
 			{
-				if (Argument.TypeIndices.IsValidIndex(ToRemove[i]))
+				TArray<TRigVMTypeIndex> NewTypeIndices;
+				NewTypeIndices.Reserve(Argument.TypeIndices.Num() - NumToRemove);
+
+				for(int32 Index = 0; Index < Argument.TypeIndices.Num(); Index++)
 				{
-					Argument.TypeIndices.RemoveAt(ToRemove[i]);
+					if(!ToRemove[Index])
+					{
+						NewTypeIndices.Add(Argument.TypeIndices[Index]);
+					}
 				}
+				Argument.TypeIndices = MoveTemp(NewTypeIndices);
 			}
+			NumPermutations -= NumToRemove;
 		}
-		NumPermutations -= ToRemove.Num();
 	}
 
 	for(FRigVMTemplateArgument& Argument : Template.Arguments)
@@ -1955,7 +1957,7 @@ const FRigVMTemplate* FRigVMRegistry_NoLock::AddTemplateFromArguments_NoLock(con
 	}
 
 	Template.Permutations.Init(INDEX_NONE, NumPermutations);
-	Template.RecomputeTypesHashToPermutations();
+	Template.RecomputeTypesHashToPermutations(TypeCaches);
 
 	const int32 Index = Templates.AddElement(Template);
 	Templates[Index].Index = Index;
diff --git a/Engine/Plugins/Runtime/RigVM/Source/RigVM/Private/RigVMCore/RigVMTemplate.cpp b/Engine/Plugins/Runtime/RigVM/Source/RigVM/Private/RigVMCore/RigVMTemplate.cpp
--- a/Engine/Plugins/Runtime/RigVM/Source/RigVM/Private/RigVMCore/RigVMTemplate.cpp
+++ b/Engine/Plugins/Runtime/RigVM/Source/RigVM/Private/RigVMCore/RigVMTemplate.cpp
@@ -138,10 +138,19 @@ FRigVMTemplateArgument::FRigVMTemplateArgument(const FName& InName, ERigVMPinDir
 		bUseCategories = NumCategories == 1;
 		if (!bUseCategories)
 		{
+			// preallocate AllTypes since it can be a large array
+			{
+				int32 NumTypes = 0;
+				for (const ETypeCategory TypeCategory : TypeCategories)
+				{
+					NumTypes += FRigVMRegistry_NoLock::GetForRead().GetTypesForCategory_NoLock(TypeCategory).Num();
+				}
+				AllTypes.Reserve(NumTypes);
+			}
+
 			for (const ETypeCategory TypeCategory : TypeCategories)
 			{
 				const TArray<TRigVMTypeIndex>& Types = FRigVMRegistry_NoLock::GetForRead().GetTypesForCategory_NoLock(TypeCategory);
-				AllTypes.Reserve(AllTypes.Num() + Types.Num());
 				for (const TRigVMTypeIndex Type: Types)
 				{
 					AllTypes.Add(Type);
@@ -819,6 +828,17 @@ TArray<TRigVMTypeIndex> FRigVMTemplateArgumentInfo::GetTypesFromCategories(
 	const FRigVMTemplateArgument::FTypeFilter& InTypeFilter)
 {
 	TSet<TRigVMTypeIndex> AllTypes;
+
+	// preallocate AllTypes since it can be a large array
+	{
+		int32 NumTypes = 0;
+		for (const FRigVMTemplateArgument::ETypeCategory TypeCategory : InTypeCategories)
+		{
+			NumTypes += FRigVMRegistry_NoLock::GetForRead().GetTypesForCategory_NoLock(TypeCategory).Num();
+		}
+		AllTypes.Reserve(NumTypes);
+	}
+
 	for (const FRigVMTemplateArgument::ETypeCategory TypeCategory : InTypeCategories)
 	{
 		AllTypes.Append(FRigVMRegistry_NoLock::GetForRead().GetTypesForCategory_NoLock(TypeCategory));
@@ -2296,16 +2316,39 @@ const FRigVMDispatchFactory* FRigVMTemplate::GetDispatchFactory(const bool bLock
 }
 
 void FRigVMTemplate::RecomputeTypesHashToPermutations()
+{
+	TArray<FRigVMTypeCacheScope_NoLock> TypeCaches;
+	TypeCaches.SetNum(NumArguments());
+	
+	for(int32 ArgIndex = 0; ArgIndex < NumArguments(); ArgIndex++)
+	{
+		(void)TypeCaches[ArgIndex].UpdateIfRequired(Arguments[ArgIndex]);
+	}
+	
+	RecomputeTypesHashToPermutations(TypeCaches);
+}
+
+void FRigVMTemplate::RecomputeTypesHashToPermutations(const TArray<FRigVMTypeCacheScope_NoLock>& InTypeCaches)
 {
 	TypesHashToPermutation.Reset();
-	for (int32 PermutationIndex=0; PermutationIndex<NumPermutations(); ++PermutationIndex)
+
+	bool bAnyArgumentWithZeroTypes = false;
+	for(int32 ArgIndex = 0; ArgIndex < InTypeCaches.Num(); ArgIndex++)
 	{
-		uint32 TypesHash=0;
-		for (int32 ArgIndex=0; ArgIndex<NumArguments(); ++ArgIndex)
+		bAnyArgumentWithZeroTypes = bAnyArgumentWithZeroTypes || InTypeCaches[ArgIndex].GetNumTypes_NoLock() == 0;
+	}
+
+	if(!bAnyArgumentWithZeroTypes)
+	{
+		for (int32 PermutationIndex=0; PermutationIndex<NumPermutations(); ++PermutationIndex)
 		{
-			TypesHash = HashCombine(Hash, GetTypeHash(Arguments[ArgIndex].GetTypeIndex_NoLock(PermutationIndex)));
+			uint32 TypesHash=0;
+			for (int32 ArgIndex=0; ArgIndex < InTypeCaches.Num(); ++ArgIndex)
+			{
+				TypesHash = HashCombine(TypesHash, GetTypeHash(InTypeCaches[ArgIndex].GetTypeIndex_NoLock(PermutationIndex)));
+			}
+			TypesHashToPermutation.Add(TypesHash, PermutationIndex);
 		}
-		TypesHashToPermutation.Add(TypesHash, PermutationIndex);
 	}
 }
 
@@ -2430,7 +2473,10 @@ TRigVMTypeIndex FRigVMTypeCacheScope_NoLock::GetTypeIndex_NoLock(int32 InIndex)
 			return Argument->GetTypeIndex_NoLock(InIndex);
 		}
 	}
-	return Types.GetValue()[InIndex];
+
+	// singleton arguments may have only one type
+	const TArray<TRigVMTypeIndex>& TypeArray = Types.GetValue();
+	return TypeArray.IsValidIndex(InIndex) ? TypeArray[InIndex] : TypeArray[0];
 }
 
 void FRigVMTypeCacheScope_NoLock::UpdateTypesIfRequired() const
diff --git a/Engine/Plugins/Runtime/RigVM/Source/RigVM/Public/RigVMCore/RigVMTemplate.h b/Engine/Plugins/Runtime/RigVM/Source/RigVM/Public/RigVMCore/RigVMTemplate.h
--- a/Engine/Plugins/Runtime/RigVM/Source/RigVM/Public/RigVMCore/RigVMTemplate.h
+++ b/Engine/Plugins/Runtime/RigVM/Source/RigVM/Public/RigVMCore/RigVMTemplate.h
@@ -550,6 +550,7 @@ public:
 	}
 
 	void RecomputeTypesHashToPermutations();
+	void RecomputeTypesHashToPermutations(const TArray<FRigVMTypeCacheScope_NoLock>& InTypeCaches);
 
 	void UpdateTypesHashToPermutation(const int32 InPermutation);
 
