RigVM: Performance optimization: Don't load factories lazily

#rnx

diff --git a/Engine/Plugins/Runtime/RigVM/Source/RigVM/Private/RigVMCore/RigVMRegistry.cpp b/Engine/Plugins/Runtime/RigVM/Source/RigVM/Private/RigVMCore/RigVMRegistry.cpp
--- a/Engine/Plugins/Runtime/RigVM/Source/RigVM/Private/RigVMCore/RigVMRegistry.cpp
+++ b/Engine/Plugins/Runtime/RigVM/Source/RigVM/Private/RigVMCore/RigVMRegistry.cpp
@@ -353,6 +353,7 @@ void FRigVMRegistry_NoLock::Initialize_NoLock()
 			}
 		}
 
+		IPluginManager::Get().OnNewPluginMounted().RemoveAll(this);
 		IPluginManager::Get().OnPluginUnmounted().RemoveAll(this);
 
 		UE::Anim::AttributeTypes::GetOnAttributeTypesChanged().RemoveAll(this);
@@ -449,10 +450,43 @@ void FRigVMRegistry_NoLock::RefreshEngineTypes_NoLock()
 			}
 		}
 	}
+
+	// also refresh the functions and dispatches
+	(void)RefreshFunctionsAndDispatches_NoLock();
 	
 	bEverRefreshedEngineTypes = true;
 }
 
+bool FRigVMRegistry_NoLock::RefreshFunctionsAndDispatches_NoLock()
+{
+	bool bRegistryChanged = false;
+	for (TObjectIterator<UScriptStruct> ScriptStructIt; ScriptStructIt; ++ScriptStructIt)
+	{
+		UScriptStruct* ScriptStruct = *ScriptStructIt;
+		if(ScriptStruct->IsChildOf(FRigVMStruct::StaticStruct()) && ScriptStruct != FRigVMStruct::StaticStruct())
+		{
+			// nothing to do here - functions are registered by their static initialize
+		}
+		else if(ScriptStruct->IsChildOf(FRigVMDispatchFactory::StaticStruct()) && ScriptStruct != FRigVMDispatchFactory::StaticStruct())
+		{
+			// factories are also registered by their static initialize,
+			// but particularly in shipping it's expensive to update factories lazily.
+			if(const FRigVMDispatchFactory* Factory = FindOrAddDispatchFactory_NoLock(ScriptStruct))
+			{
+				// pulling on the template will cause the template to be initialized.
+				// that may introduce a certain cost - which we don't want to experience during
+				// the game.
+				if(Factory->CachedTemplate == nullptr)
+				{
+					(void)Factory->GetTemplate_NoLock();
+					bRegistryChanged = true;
+				}
+			}
+		}
+	}
+	return bRegistryChanged; 
+}
+
 void FRigVMRegistry_NoLock::OnAssetRenamed_NoLock(const FAssetData& InAssetData, const FString& InOldObjectPath)
 {
 	const FSoftObjectPath OldPath(InOldObjectPath);
@@ -470,6 +504,11 @@ bool FRigVMRegistry_NoLock::OnAssetRemoved_NoLock(const FAssetData& InAssetData)
 	return RemoveType_NoLock(InAssetData.ToSoftObjectPath(), InAssetData.GetClass());
 }
 
+bool FRigVMRegistry_NoLock::OnPluginLoaded_NoLock(IPlugin& InPlugin)
+{
+	return RefreshFunctionsAndDispatches_NoLock();
+}
+
 bool FRigVMRegistry_NoLock::OnPluginUnloaded_NoLock(IPlugin& InPlugin)
 {
 	const FString PluginContentPath = InPlugin.GetMountedAssetPath();
@@ -977,6 +1016,20 @@ void FRigVMRegistry_RWLock::OnAssetRemoved(const FAssetData& InAssetData)
 	}
 }
 
+void FRigVMRegistry_RWLock::OnPluginLoaded(IPlugin& InPlugin)
+{
+	bool bRegistryChanged = false;
+	{
+		FConditionalWriteScopeLock _(*this);
+		bRegistryChanged = Super::OnPluginLoaded_NoLock(InPlugin);
+	}
+		
+	if (bRegistryChanged)
+	{
+		OnRigVMRegistryChangedDelegate.Broadcast();
+	}
+}
+
 void FRigVMRegistry_RWLock::OnPluginUnloaded(IPlugin& InPlugin)
 {
 	bool bRegistryChanged = false;
@@ -2057,6 +2110,7 @@ void FRigVMRegistry_RWLock::Initialize(bool bLockRegistry)
 	AssetRegistryModule.Get().OnAssetRemoved().AddRaw(this, &FRigVMRegistry_RWLock::OnAssetRemoved);
 	AssetRegistryModule.Get().OnAssetRenamed().AddRaw(this, &FRigVMRegistry_RWLock::OnAssetRenamed);
 
+	IPluginManager::Get().OnNewPluginMounted().AddRaw(this, &FRigVMRegistry_RWLock::OnPluginLoaded);
 	IPluginManager::Get().OnPluginUnmounted().AddRaw(this, &FRigVMRegistry_RWLock::OnPluginUnloaded);
 	
 	UE::Anim::AttributeTypes::GetOnAttributeTypesChanged().AddRaw(this, &FRigVMRegistry_RWLock::OnAnimationAttributeTypesChanged);
diff --git a/Engine/Plugins/Runtime/RigVM/Source/RigVM/Public/RigVMCore/RigVMRegistry.h b/Engine/Plugins/Runtime/RigVM/Source/RigVM/Public/RigVMCore/RigVMRegistry.h
--- a/Engine/Plugins/Runtime/RigVM/Source/RigVM/Public/RigVMCore/RigVMRegistry.h
+++ b/Engine/Plugins/Runtime/RigVM/Source/RigVM/Public/RigVMCore/RigVMRegistry.h
@@ -92,6 +92,9 @@ public:
 	// based on the names.
 	virtual void RefreshEngineTypes_NoLock();
 
+	// Refreshes the registered functions and dispatches.
+	virtual bool RefreshFunctionsAndDispatches_NoLock();
+
 	// Refreshes the list and finds the function pointers
 	// based on the names.
 	virtual void RefreshEngineTypesIfRequired_NoLock();
@@ -102,6 +105,9 @@ public:
 	// Update the registry when old types are removed
 	virtual bool OnAssetRemoved_NoLock(const FAssetData& InAssetData);
 
+	// May add factories and unit functions declared in the plugin 
+	virtual bool OnPluginLoaded_NoLock(IPlugin& InPlugin);
+
 	// Removes all types associated with a plugin that's being unloaded. 
 	virtual bool OnPluginUnloaded_NoLock(IPlugin& InPlugin);
 	
@@ -604,6 +610,13 @@ public:
 		FConditionalWriteScopeLock _(*this, bLockRegistry);
 		Super::RefreshEngineTypesIfRequired_NoLock();
 	}
+	
+	// Refreshes the registered functions and dispatches.
+	bool RefreshFunctionsAndDispatches()
+	{
+		FConditionalWriteScopeLock _(*this);
+		return Super::RefreshFunctionsAndDispatches_NoLock();
+	}
 
 	// Update the registry when types are renamed
 	void OnAssetRenamed(const FAssetData& InAssetData, const FString& InOldObjectPath)
@@ -615,6 +628,9 @@ public:
 	// Update the registry when old types are removed
     void OnAssetRemoved(const FAssetData& InAssetData);
 
+	// May add factories and unit functions declared in the plugin 
+	void OnPluginLoaded(IPlugin& InPlugin);
+
 	// Removes all types associated with a plugin that's being unloaded. 
 	void OnPluginUnloaded(IPlugin& InPlugin);
 
@@ -959,6 +975,12 @@ protected:
 		Super::RefreshEngineTypes_NoLock();
 	}
 
+	// Refreshes the registered functions and dispatches.
+	virtual bool RefreshFunctionsAndDispatches_NoLock() override
+	{
+		return Super::RefreshFunctionsAndDispatches_NoLock();
+	}
+
 	// Refreshes the list and finds the function pointers
 	// based on the names.
 	// Note: Only call NoLock methods on the FRigVMRegistry_NoLock
@@ -981,6 +1003,13 @@ protected:
 		return Super::OnAssetRemoved_NoLock(InAssetData);
 	}
 
+	// May add factories and unit functions declared in the plugin 
+	// Note: Only call NoLock methods on the FRigVMRegistry_NoLock
+	virtual bool OnPluginLoaded_NoLock(IPlugin& InPlugin) override
+	{
+		return Super::OnPluginLoaded_NoLock(InPlugin);
+	}
+
 	// Removes all types associated with a plugin that's being unloaded. 
 	// Note: Only call NoLock methods on the FRigVMRegistry_NoLock
 	virtual bool OnPluginUnloaded_NoLock(IPlugin& InPlugin) override
