Control Rig: Fix stall due to poor performance in InWithinLoop
#rnx

diff --git a/Engine/Plugins/Runtime/RigVM/Source/RigVMDeveloper/Private/RigVMModel/RigVMNode.cpp b/Engine/Plugins/Runtime/RigVM/Source/RigVMDeveloper/Private/RigVMModel/RigVMNode.cpp
--- a/Engine/Plugins/Runtime/RigVM/Source/RigVMDeveloper/Private/RigVMModel/RigVMNode.cpp
+++ b/Engine/Plugins/Runtime/RigVM/Source/RigVMDeveloper/Private/RigVMModel/RigVMNode.cpp
@@ -1237,6 +1237,7 @@ const bool URigVMNode::IsControlFlowBlockSliced(const FName& InBlockName) const
 
 bool URigVMNode::IsWithinLoop() const
 {
+	TArray<URigVMNode*> SourceNodes;
 	for(const URigVMPin* Pin : Pins)
 	{
 		const TArray<URigVMPin*> SourcePins = Pin->GetLinkedSourcePins(true);
@@ -1251,14 +1252,20 @@ bool URigVMNode::IsWithinLoop() const
 			}
 		}
 
-		for(const URigVMPin* SourcePin : SourcePins)
+		for (const URigVMPin* SourcePin : SourcePins)
 		{
-			if(SourcePin->GetNode()->IsWithinLoop())
-			{
-				return true;
-			}
+			SourceNodes.AddUnique(SourcePin->GetNode());
 		}
 	}
+
+	for(const URigVMNode* SourceNode : SourceNodes)
+	{
+		if(SourceNode->IsWithinLoop())
+		{
+			return true;
+		}
+	}
+	
 	return false;
 }
 
