RigVM: Make loading dispatches greedily optional, move TypesToPermutations to editor only

#rnx

diff --git a/Engine/Plugins/Runtime/RigVM/Source/RigVM/Private/RigVMCore/RigVMRegistry.cpp b/Engine/Plugins/Runtime/RigVM/Source/RigVM/Private/RigVMCore/RigVMRegistry.cpp
--- a/Engine/Plugins/Runtime/RigVM/Source/RigVM/Private/RigVMCore/RigVMRegistry.cpp
+++ b/Engine/Plugins/Runtime/RigVM/Source/RigVM/Private/RigVMCore/RigVMRegistry.cpp
@@ -17,15 +17,14 @@
 #include "Interfaces/IPluginManager.h"
 
 // When the object system has been completely loaded, load in all the engine types that we haven't registered already in InitializeIfNeeded 
-static FDelayedAutoRegisterHelper GRigVMRegistrySingletonHelper(EDelayedRegisterRunPhase::EndOfEngineInit, []() -> void
-{
-	FRigVMRegistry_RWLock::Get().RefreshEngineTypes();
-});
+static FDelayedAutoRegisterHelper GRigVMRegistrySingletonHelper(EDelayedRegisterRunPhase::EndOfEngineInit, &FRigVMRegistry_NoLock::OnEngineInit);
 
+static TAutoConsoleVariable<bool> CVarRigVMUpdateDispatchFactoriesGreedily(TEXT("RigVM.UpdateDispatchFactoriesGreedily"), true, TEXT("Set this to false to avoid loading dispatch factories during engine init / plugin mount."));
 
 FRigVMRegistry_NoLock::FRigVMRegistry_NoLock() :
 	bIsRefreshingEngineTypes(false),
-	bEverRefreshedEngineTypes(false)
+	bEverRefreshedEngineTypes(false),
+	bEverRefreshedDispatchFactoriesAfterEngineInit(false)
 {
 }
 
@@ -459,29 +458,26 @@ void FRigVMRegistry_NoLock::RefreshEngineTypes_NoLock()
 
 bool FRigVMRegistry_NoLock::RefreshFunctionsAndDispatches_NoLock()
 {
+	if(!CVarRigVMUpdateDispatchFactoriesGreedily->GetBool())
+	{
+		return false;
+	}
+
+	// nothing to do for functions for now - they are registered by their static initialize
+
 	bool bRegistryChanged = false;
-	for (TObjectIterator<UScriptStruct> ScriptStructIt; ScriptStructIt; ++ScriptStructIt)
+
+	// factories are also registered by RegisterFactory_NoLock, so we don't need to visit all
+	// currently known UScriptStructs. By the time we get here the factories are registered.
+	for(const FRigVMDispatchFactory* Factory : Factories)
 	{
-		UScriptStruct* ScriptStruct = *ScriptStructIt;
-		if(ScriptStruct->IsChildOf(FRigVMStruct::StaticStruct()) && ScriptStruct != FRigVMStruct::StaticStruct())
-		{
-			// nothing to do here - functions are registered by their static initialize
-		}
-		else if(ScriptStruct->IsChildOf(FRigVMDispatchFactory::StaticStruct()) && ScriptStruct != FRigVMDispatchFactory::StaticStruct())
+		// pulling on the template will cause the template to be initialized.
+		// that may introduce a certain cost - which we don't want to experience during
+		// the game.
+		if(Factory->CachedTemplate == nullptr)
 		{
-			// factories are also registered by their static initialize,
-			// but particularly in shipping it's expensive to update factories lazily.
-			if(const FRigVMDispatchFactory* Factory = FindOrAddDispatchFactory_NoLock(ScriptStruct))
-			{
-				// pulling on the template will cause the template to be initialized.
-				// that may introduce a certain cost - which we don't want to experience during
-				// the game.
-				if(Factory->CachedTemplate == nullptr)
-				{
-					(void)Factory->GetTemplate_NoLock();
-					bRegistryChanged = true;
-				}
-			}
+			(void)Factory->GetTemplate_NoLock();
+			bRegistryChanged = true;
 		}
 	}
 	return bRegistryChanged; 
@@ -506,6 +502,11 @@ bool FRigVMRegistry_NoLock::OnAssetRemoved_NoLock(const FAssetData& InAssetData)
 
 bool FRigVMRegistry_NoLock::OnPluginLoaded_NoLock(IPlugin& InPlugin)
 {
+	// only update the functions / dispatches once the engine has initialized.
+	if(!bEverRefreshedDispatchFactoriesAfterEngineInit)
+	{
+		return false;
+	}
 	return RefreshFunctionsAndDispatches_NoLock();
 }
 
@@ -992,6 +993,13 @@ void FRigVMRegistry_NoLock::RemoveTypeInCategory_NoLock(FRigVMTemplateArgument::
 	}
 }
 
+void FRigVMRegistry_NoLock::OnEngineInit()
+{
+	FRigVMRegistry_RWLock& Registry = FRigVMRegistry_RWLock::Get();
+	Registry.RefreshEngineTypes();
+	Registry.bEverRefreshedDispatchFactoriesAfterEngineInit = true;
+}
+
 // This function needs to be in cpp file instead of header file
 // to avoid confusing certain compilers into creating multiple copies of the registry
 FRigVMRegistry_RWLock& FRigVMRegistry_RWLock::Get()
@@ -1924,7 +1932,9 @@ const FRigVMTemplate* FRigVMRegistry_NoLock::AddTemplateFromArguments_NoLock(con
 			const TRigVMTypeIndex FirstTypeIndex = Argument.GetTypeIndex_NoLock(0);
 			if(IsWildCardType_NoLock(FirstTypeIndex))
 			{
+#if WITH_EDITOR
 				Argument.InvalidatePermutations(FirstTypeIndex);
+#endif
 				if(IsArrayType_NoLock(FirstTypeIndex))
 				{
 					Argument.TypeCategories.Add(FRigVMTemplateArgument::ETypeCategory_ArrayAnyValue);
@@ -2004,10 +2014,12 @@ const FRigVMTemplate* FRigVMRegistry_NoLock::AddTemplateFromArguments_NoLock(con
 		}
 	}
 
+#if WITH_EDITOR
 	for(FRigVMTemplateArgument& Argument : Template.Arguments)
 	{
 		Argument.UpdateTypeToPermutations();
 	}
+#endif
 
 	Template.Permutations.Init(INDEX_NONE, NumPermutations);
 	Template.RecomputeTypesHashToPermutations(TypeCaches);
diff --git a/Engine/Plugins/Runtime/RigVM/Source/RigVM/Private/RigVMCore/RigVMTemplate.cpp b/Engine/Plugins/Runtime/RigVM/Source/RigVM/Private/RigVMCore/RigVMTemplate.cpp
--- a/Engine/Plugins/Runtime/RigVM/Source/RigVM/Private/RigVMCore/RigVMTemplate.cpp
+++ b/Engine/Plugins/Runtime/RigVM/Source/RigVM/Private/RigVMCore/RigVMTemplate.cpp
@@ -82,7 +82,9 @@ FRigVMTemplateArgument::FRigVMTemplateArgument(FProperty* InProperty, FRigVMRegi
 
 	TypeIndices.Add(TypeIndex);
 	EnsureValidExecuteType_NoLock(InRegistry);
+#if WITH_EDITOR
 	UpdateTypeToPermutations();
+#endif
 }
 
 FRigVMTemplateArgument FRigVMTemplateArgument::Make(FProperty* InProperty)
@@ -107,7 +109,9 @@ FRigVMTemplateArgument::FRigVMTemplateArgument(const FName& InName, ERigVMPinDir
 	, TypeIndices({InTypeIndex})
 {
 	EnsureValidExecuteType_NoLock(FRigVMRegistry_RWLock::Get());
+#if WITH_EDITOR
 	UpdateTypeToPermutations();
+#endif
 }
 
 FRigVMTemplateArgument::FRigVMTemplateArgument(const FName& InName, ERigVMPinDirection InDirection, const TArray<TRigVMTypeIndex>& InTypeIndices)
@@ -117,7 +121,9 @@ FRigVMTemplateArgument::FRigVMTemplateArgument(const FName& InName, ERigVMPinDir
 {
 	check(TypeIndices.Num() > 0);
 	EnsureValidExecuteType_NoLock(FRigVMRegistry_RWLock::Get());
+#if WITH_EDITOR
 	UpdateTypeToPermutations();
+#endif
 }
 
 FRigVMTemplateArgument::FRigVMTemplateArgument(const FName& InName, ERigVMPinDirection InDirection, const TArray<ETypeCategory>& InTypeCategories, TFunction<bool(const TRigVMTypeIndex&)> InFilterType)
@@ -180,7 +186,9 @@ FRigVMTemplateArgument::FRigVMTemplateArgument(const FName& InName, ERigVMPinDir
 			EnsureValidExecuteType_NoLock(FRigVMRegistry_RWLock::Get());
 		}
 
+#if WITH_EDITOR
 		UpdateTypeToPermutations();
+#endif
 	}
 }
 
@@ -192,6 +200,8 @@ void FRigVMTemplateArgument::EnsureValidExecuteType_NoLock(FRigVMRegistry_NoLock
 	}
 }
 
+#if WITH_EDITOR
+
 void FRigVMTemplateArgument::UpdateTypeToPermutations()
 {
 	TypeToPermutations.Reset();
@@ -258,6 +268,8 @@ bool FRigVMTemplateArgument::SupportsTypeIndex_NoLock(TRigVMTypeIndex InTypeInde
 	return false;
 }
 
+#endif
+
 bool FRigVMTemplateArgument::IsSingleton(const TArray<int32>& InPermutationIndices, const bool bLockRegistry) const
 {
 	const FRigVMRegistryReadLock _(bLockRegistry);
@@ -266,10 +278,12 @@ bool FRigVMTemplateArgument::IsSingleton(const TArray<int32>& InPermutationIndic
 
 bool FRigVMTemplateArgument::IsSingleton_NoLock(const TArray<int32>& InPermutationIndices) const
 {
+#if WITH_EDITOR
 	if (TypeToPermutations.Num() == 1)
 	{
 		return true;
 	}
+#endif
 
 	// if a type is using categories it can't be singleton
 	// since categories provide more than one type.
@@ -358,6 +372,8 @@ FRigVMTemplateArgument::EArrayType FRigVMTemplateArgument::GetArrayType_NoLock()
 	return EArrayType_Invalid;
 }
 
+#if WITH_EDITOR
+
 const TArray<int32>& FRigVMTemplateArgument::GetPermutations(const TRigVMTypeIndex InType, bool bLockRegistry) const
 {
 	FRigVMRegistryReadLock _(bLockRegistry);
@@ -397,6 +413,8 @@ void FRigVMTemplateArgument::InvalidatePermutations(const TRigVMTypeIndex InType
 	TypeToPermutations.Remove(InType);
 }
 
+#endif
+
 void FRigVMTemplateArgument::GetAllTypes(TArray<TRigVMTypeIndex>& OutTypes, const bool bLockRegistry) const
 {
 	const FRigVMRegistryReadLock _(bLockRegistry);
@@ -1493,6 +1511,7 @@ bool FRigVMTemplate::Merge(const FRigVMTemplate& InOther)
 		{
 			FRigVMTemplateArgument& NewArg = NewArgs.Add_GetRef(Arguments[ArgumentIndex]);
 			const TRigVMTypeIndex OtherTypeIndex = OtherArg.GetTypeIndex_NoLock(0);
+#if WITH_EDITOR
 			const int32 NewPermutationIndex = NewArg.GetNumTypes_NoLock();
 			if (TArray<int32>* ArgTypePermutations = NewArg.TypeToPermutations.Find(OtherTypeIndex))
 			{
@@ -1502,6 +1521,7 @@ bool FRigVMTemplate::Merge(const FRigVMTemplate& InOther)
 			{
 				NewArg.TypeToPermutations.Add(OtherTypeIndex, {NewPermutationIndex});
 			}
+#endif
 			NewArg.TypeIndices.Add(OtherTypeIndex);
 		}
 	}
@@ -1599,6 +1619,8 @@ bool FRigVMTemplate::SupportsExecuteContextStruct(const UScriptStruct* InExecute
 	return InExecuteContextStruct->IsChildOf(GetExecuteContextStruct(true));
 }
 
+#if WITH_EDITOR
+
 bool FRigVMTemplate::ArgumentSupportsTypeIndex(const FName& InArgumentName, TRigVMTypeIndex InTypeIndex, TRigVMTypeIndex* OutTypeIndex) const
 {
 	if (const FRigVMTemplateArgument* Argument = FindArgument(InArgumentName))
@@ -1608,6 +1630,8 @@ bool FRigVMTemplate::ArgumentSupportsTypeIndex(const FName& InArgumentName, TRig
 	return false;
 }
 
+#endif
+
 const FRigVMFunction* FRigVMTemplate::GetPrimaryPermutation(bool bLockRegistry) const
 {
 	if (NumPermutations() > 0)
@@ -1919,7 +1943,8 @@ bool FRigVMTemplate::ContainsPermutation_NoLock(const FTypeMap& InTypes) const
 	{
 		return TypesHashToPermutation.Contains(TypesHash);
 	}
-	
+
+#if WITH_EDITOR
 	TArray<int32> PossiblePermutations;
 	for (const TPair<FName, TRigVMTypeIndex>& Pair : InTypes)
 	{
@@ -1960,6 +1985,9 @@ bool FRigVMTemplate::ContainsPermutation_NoLock(const FTypeMap& InTypes) const
 	}
 	
 	return true;
+#else
+	return false;
+#endif
 }
 
 bool FRigVMTemplate::ResolveArgument(const FName& InArgumentName, const TRigVMTypeIndex InTypeIndex,
@@ -2172,7 +2200,9 @@ bool FRigVMTemplate::UpdateArgumentTypes()
 		}
 
 		Argument.TypeIndices.Reset();
+#if WITH_EDITOR
 		Argument.TypeToPermutations.Reset();
+#endif
 	}
 
 	const FRigVMRegistry_NoLock& Registry = FRigVMRegistry_NoLock::GetForRead();
@@ -2244,7 +2274,9 @@ bool FRigVMTemplate::UpdateArgumentTypes()
 						}
 					
 						Argument.TypeIndices.Add(TypeIndex);
+#if WITH_EDITOR
 						Argument.TypeToPermutations.FindOrAdd(TypeIndex).Add(Permutations.Num());
+#endif
 					}
 
 					Permutations.Add(INDEX_NONE);
@@ -2263,11 +2295,13 @@ bool FRigVMTemplate::UpdateArgumentTypes()
 		return true;
 	});
 
+#if WITH_EDITOR
 	for(int32 ArgumentIndex = 0; ArgumentIndex < Arguments.Num(); ArgumentIndex++)
 	{
 		FRigVMTemplateArgument& Argument = Arguments[ArgumentIndex];
 		Argument.UpdateTypeToPermutations();
 	}
+#endif
 
 	return bResult;
 }
@@ -2302,11 +2336,13 @@ void FRigVMTemplate::HandleTypeRemoval(TRigVMTypeIndex InTypeIndex)
 			}
 		}
 	}
-	
+
+#if WITH_EDITOR
 	for (FRigVMTemplateArgument& Argument : Arguments)
 	{
 		Argument.TypeToPermutations.Remove(InTypeIndex);
 	}
+#endif
 }
 
 const FRigVMDispatchFactory* FRigVMTemplate::GetDispatchFactory(const bool bLockRegistry) const
diff --git a/Engine/Plugins/Runtime/RigVM/Source/RigVM/Public/RigVMCore/RigVMRegistry.h b/Engine/Plugins/Runtime/RigVM/Source/RigVM/Public/RigVMCore/RigVMRegistry.h
--- a/Engine/Plugins/Runtime/RigVM/Source/RigVM/Public/RigVMCore/RigVMRegistry.h
+++ b/Engine/Plugins/Runtime/RigVM/Source/RigVM/Public/RigVMCore/RigVMRegistry.h
@@ -313,6 +313,8 @@ public:
 
 	static inline const FLazyName TemplateNameMetaName = FLazyName(TEXT("TemplateName"));
 
+	static void OnEngineInit();
+
 protected:
 
 
@@ -419,6 +421,9 @@ protected:
 	// This is true if the engine has ever refreshed the engine types
 	bool bEverRefreshedEngineTypes;
 
+	// This is true if the dispatch factories and functions have been greedily loaded once during engine init
+	bool bEverRefreshedDispatchFactoriesAfterEngineInit;
+
 	friend struct FRigVMStruct;
 	friend struct FRigVMTemplate;
 	friend struct FRigVMTemplateArgument;
diff --git a/Engine/Plugins/Runtime/RigVM/Source/RigVM/Public/RigVMCore/RigVMTemplate.h b/Engine/Plugins/Runtime/RigVM/Source/RigVM/Public/RigVMCore/RigVMTemplate.h
--- a/Engine/Plugins/Runtime/RigVM/Source/RigVM/Public/RigVMCore/RigVMTemplate.h
+++ b/Engine/Plugins/Runtime/RigVM/Source/RigVM/Public/RigVMCore/RigVMTemplate.h
@@ -205,9 +205,11 @@ struct RIGVM_API FRigVMTemplateArgument
 	// returns the direction of the argument
 	ERigVMPinDirection GetDirection() const { return Direction; }
 
+#if WITH_EDITOR
 	// returns true if this argument supports a given type across a set of permutations
 	bool SupportsTypeIndex(TRigVMTypeIndex InTypeIndex, TRigVMTypeIndex* OutTypeIndex = nullptr, const bool bLockRegistry = true) const;
 	bool SupportsTypeIndex_NoLock(TRigVMTypeIndex InTypeIndex, TRigVMTypeIndex* OutTypeIndex = nullptr) const;
+#endif
 
 	// returns the flat list of types (including duplicates) of this argument
 	void GetAllTypes(TArray<TRigVMTypeIndex>& OutTypes, const bool bLockRegistry = true) const;
@@ -260,10 +262,8 @@ struct RIGVM_API FRigVMTemplateArgument
 	// returns an array of all of the supported types
 	TArray<TRigVMTypeIndex> GetSupportedTypeIndices(const TArray<int32>& InPermutationIndices = TArray<int32>()) const;
 
-#if WITH_EDITOR
 	// returns an array of all supported types as strings. this is used for automated testing only.
 	TArray<FString> GetSupportedTypeStrings(const TArray<int32>& InPermutationIndices = TArray<int32>()) const;
-#endif
 	
 	// returns true if an argument is singleton (same type for all variants)
 	bool IsSingleton(const TArray<int32>& InPermutationIndices = TArray<int32>(), const bool bLockRegistry = true) const;
@@ -279,10 +279,12 @@ struct RIGVM_API FRigVMTemplateArgument
 
 	RIGVM_API friend uint32 GetTypeHash_NoLock(const FRigVMTemplateArgument& InArgument);
 
+#if WITH_EDITOR
 	// Get the map of types to permutation indices
 	const TArray<int32>& GetPermutations(const TRigVMTypeIndex InType, bool bLockRegistry = true) const;
 	const TArray<int32>& GetPermutations_NoLock(const TRigVMTypeIndex InType) const;
 	void InvalidatePermutations(const TRigVMTypeIndex InType);
+#endif
 
 protected:
 
@@ -291,7 +293,9 @@ protected:
 	ERigVMPinDirection Direction = ERigVMPinDirection::IO;
 
 	TArray<TRigVMTypeIndex> TypeIndices;
+#if WITH_EDITOR
 	mutable TMap<TRigVMTypeIndex, TArray<int32>> TypeToPermutations;
+#endif
 
 	bool bUseCategories = false;
 	TArray<ETypeCategory> TypeCategories;
@@ -309,7 +313,10 @@ protected:
 	static FRigVMTemplateArgument Make_NoLock(FProperty* InProperty, FRigVMRegistry_NoLock& InRegistry);
 
 	void EnsureValidExecuteType_NoLock(FRigVMRegistry_NoLock& InRegistry);
+
+#if WITH_EDITOR
 	void UpdateTypeToPermutations();
+#endif
 
 	friend struct FRigVMTemplate;
 	friend struct FRigVMDispatchFactory;
@@ -437,8 +444,10 @@ public:
 	// returns true if this template supports a given execute context struct
 	bool SupportsExecuteContextStruct(const UScriptStruct* InExecuteContextStruct) const;
 
+#if WITH_EDITOR
 	// returns true if a given arg supports a type
 	bool ArgumentSupportsTypeIndex(const FName& InArgumentName, TRigVMTypeIndex InTypeIndex, TRigVMTypeIndex* OutTypeIndex = nullptr) const;
+#endif
 
 	// returns the number of permutations supported by this template
 	int32 NumPermutations() const { return Permutations.Num(); }
