RigVM: Various fixes for AnimNext use cases

Fixed issues using TSubClassof pins - in many cases the class being used was incorrect, so many controller operations would fail. This mainly involves replacing logic dealing with FObjectProperty::PropertyClass to use FClassProperty::MetaClass.
Fixed tooltips on cast nodes (and other dispatch node pins) not displaying.
Ensured that pin metadata gets propagated from the model to the UEdGraphNode (this allows things like GetOptions metadata to be correctly exposed, allowing dropdown selectors for FName/FString types).
Added a new kind of FRigVMUserWorkflow that rig units can opt into - NodeContextButton. This will add a button to the bottom-left of the node widget for each workflow added (looks like the 'Add Pin' button on Sequence nodes).

diff --git a/Engine/Plugins/Runtime/RigVM/Source/RigVM/Private/RigVMCore/RigVMMemoryStorage.cpp b/Engine/Plugins/Runtime/RigVM/Source/RigVM/Private/RigVMCore/RigVMMemoryStorage.cpp
--- a/Engine/Plugins/Runtime/RigVM/Source/RigVM/Private/RigVMCore/RigVMMemoryStorage.cpp
+++ b/Engine/Plugins/Runtime/RigVM/Source/RigVM/Private/RigVMCore/RigVMMemoryStorage.cpp
@@ -175,6 +175,10 @@ FRigVMPropertyDescription::FRigVMPropertyDescription(const FProperty* InProperty
 			CPPTypeObject = StructProperty->Struct;
 		}
 #if UE_RIGVM_UOBJECT_PROPERTIES_ENABLED
+		else if(const FClassProperty* ClassProperty = CastField<FClassProperty>(ValueProperty))
+		{
+			CPPTypeObject = ClassProperty->MetaClass;
+		}
 		else if(const FObjectProperty* ObjectProperty = CastField<FObjectProperty>(ValueProperty))
 		{
 			CPPTypeObject = ObjectProperty->PropertyClass;
@@ -753,9 +757,19 @@ FProperty* URigVMMemoryStorageGeneratorClass::AddProperty(URigVMMemoryStorageGen
 			{
 				check(RigVMCore::SupportsUObjects());
 
-				FObjectProperty* ObjectProperty = new FObjectProperty(PropertyOwner, InProperty.Name, RF_Public);
-				ObjectProperty->SetPropertyClass(PropertyClass);
-				(*ValuePropertyPtr) = ObjectProperty;
+				const bool bIsClass = RigVMTypeUtils::IsUClassType(InProperty.CPPType);
+				if(bIsClass)
+				{
+					FClassProperty* ClassProperty = new FClassProperty(PropertyOwner, InProperty.Name, RF_Public);
+					ClassProperty->SetMetaClass(PropertyClass);
+					(*ValuePropertyPtr) = ClassProperty;
+				}
+				else
+				{
+					FObjectProperty* ObjectProperty = new FObjectProperty(PropertyOwner, InProperty.Name, RF_Public);
+					ObjectProperty->SetPropertyClass(PropertyClass);
+					(*ValuePropertyPtr) = ObjectProperty;
+				}
 			}
 			else
 			{
diff --git a/Engine/Plugins/Runtime/RigVM/Source/RigVM/Private/RigVMCore/RigVMMemoryStorageStruct.cpp b/Engine/Plugins/Runtime/RigVM/Source/RigVM/Private/RigVMCore/RigVMMemoryStorageStruct.cpp
--- a/Engine/Plugins/Runtime/RigVM/Source/RigVM/Private/RigVMCore/RigVMMemoryStorageStruct.cpp
+++ b/Engine/Plugins/Runtime/RigVM/Source/RigVM/Private/RigVMCore/RigVMMemoryStorageStruct.cpp
@@ -673,7 +673,15 @@ void FRigVMMemoryStorageStruct::SetDefaultValues(const TArray<FRigVMPropertyDesc
 	}
 	else if (Cast<UObject>(RigVMDescriptor.CPPTypeObject))
 	{
-		OutBagPropertyType = EPropertyBagPropertyType::Object;
+		const bool bIsClass = RigVMTypeUtils::IsUClassType(RigVMDescriptor.CPPType);
+		if(bIsClass)
+		{
+			OutBagPropertyType = EPropertyBagPropertyType::Class;
+		}
+		else
+		{
+			OutBagPropertyType = EPropertyBagPropertyType::Object;
+		}
 	}
 	else
 	{
diff --git a/Engine/Plugins/Runtime/RigVM/Source/RigVM/Private/RigVMCore/RigVMRegistry.cpp b/Engine/Plugins/Runtime/RigVM/Source/RigVM/Private/RigVMCore/RigVMRegistry.cpp
--- a/Engine/Plugins/Runtime/RigVM/Source/RigVM/Private/RigVMCore/RigVMRegistry.cpp
+++ b/Engine/Plugins/Runtime/RigVM/Source/RigVM/Private/RigVMCore/RigVMRegistry.cpp
@@ -1345,6 +1345,10 @@ bool FRigVMRegistry_NoLock::IsAllowedType_NoLock(const FProperty* InProperty) co
 	{
 		return IsAllowedType_NoLock(StructProperty->Struct);
 	}
+	if(const FClassProperty* ClassProperty = CastField<FClassProperty>(InProperty))
+	{
+		return IsAllowedType_NoLock(ClassProperty->MetaClass);
+	}
 	if(const FObjectProperty* ObjectProperty = CastField<FObjectProperty>(InProperty))
 	{
 		return IsAllowedType_NoLock(ObjectProperty->PropertyClass);
diff --git a/Engine/Plugins/Runtime/RigVM/Source/RigVM/Private/RigVMCore/RigVMTemplate.cpp b/Engine/Plugins/Runtime/RigVM/Source/RigVM/Private/RigVMCore/RigVMTemplate.cpp
--- a/Engine/Plugins/Runtime/RigVM/Source/RigVM/Private/RigVMCore/RigVMTemplate.cpp
+++ b/Engine/Plugins/Runtime/RigVM/Source/RigVM/Private/RigVMCore/RigVMTemplate.cpp
@@ -69,6 +69,13 @@ FRigVMTemplateArgument::FRigVMTemplateArgument(FProperty* InProperty, FRigVMRegi
 	{
 		CPPTypeObject = ByteProperty->Enum;
 	}
+	else if (FClassProperty* ClassProperty = CastField<FClassProperty>(InProperty))
+	{
+		if(RigVMCore::SupportsUObjects())
+		{
+			CPPTypeObject = ClassProperty->MetaClass;
+		}
+	}
 	else if (FObjectProperty* ObjectProperty = CastField<FObjectProperty>(InProperty))
 	{
 		if(RigVMCore::SupportsUObjects())
diff --git a/Engine/Plugins/Runtime/RigVM/Source/RigVM/Private/RigVMFunctions/RigVMDispatch_CastObject.cpp b/Engine/Plugins/Runtime/RigVM/Source/RigVM/Private/RigVMFunctions/RigVMDispatch_CastObject.cpp
--- a/Engine/Plugins/Runtime/RigVM/Source/RigVM/Private/RigVMFunctions/RigVMDispatch_CastObject.cpp
+++ b/Engine/Plugins/Runtime/RigVM/Source/RigVM/Private/RigVMFunctions/RigVMDispatch_CastObject.cpp
@@ -85,6 +85,19 @@ FString FRigVMDispatch_CastObject::GetNodeTitle(const FRigVMTemplateTypeMap& InT
 
 FText FRigVMDispatch_CastObject::GetNodeTooltip(const FRigVMTemplateTypeMap& InTypes) const
 {
+	const TRigVMTypeIndex* ValueTypeIndexPtr = InTypes.Find(ValueName);
+	const TRigVMTypeIndex* ResultTypeIndexPtr = InTypes.Find(ResultName);
+	if(ValueTypeIndexPtr && ResultTypeIndexPtr)
+	{
+		const FRigVMRegistry& Registry = FRigVMRegistry::Get();
+		const FRigVMTemplateArgumentType& ValueType = Registry.GetType(*ValueTypeIndexPtr);
+		const FRigVMTemplateArgumentType& ResultType = Registry.GetType(*ResultTypeIndexPtr);
+		if(ValueType.CPPTypeObject && ValueType.CPPTypeObject->IsA<UClass>() && ResultType.CPPTypeObject && ResultType.CPPTypeObject->IsA<UClass>())
+		{
+			return FText::Format(LOCTEXT("CastToolTipFormat", "Cast from {0} to {1}"), CastChecked<UClass>(ValueType.CPPTypeObject)->GetDisplayNameText(), CastChecked<UClass>(ResultType.CPPTypeObject)->GetDisplayNameText());
+		}
+	}
+
 	return LOCTEXT("CastToolTip", "Casts between object types");
 }
 
diff --git a/Engine/Plugins/Runtime/RigVM/Source/RigVM/Private/RigVMPropertyUtils.cpp b/Engine/Plugins/Runtime/RigVM/Source/RigVM/Private/RigVMPropertyUtils.cpp
--- a/Engine/Plugins/Runtime/RigVM/Source/RigVM/Private/RigVMPropertyUtils.cpp
+++ b/Engine/Plugins/Runtime/RigVM/Source/RigVM/Private/RigVMPropertyUtils.cpp
@@ -89,6 +89,21 @@ void RigVMPropertyUtils::GetTypeFromProperty(const FProperty* InProperty, FName&
 		GetTypeFromProperty(ArrayProperty->Inner, InnerTypeName, OutTypeObject);
 		OutTypeName = *RigVMTypeUtils::ArrayTypeFromBaseType(InnerTypeName.ToString());
 	}
+	else if (const FClassProperty* ClassProperty = CastField<FClassProperty>(InProperty))
+	{
+		if (RigVMCore::SupportsUObjects())
+		{
+			OutTypeName = *FString::Printf(TEXT("TSubclassOf<%s%s>"),
+				ClassProperty->MetaClass->GetPrefixCPP(),
+				*ClassProperty->MetaClass->GetName());
+			OutTypeObject = ClassProperty->MetaClass;
+		}
+		else
+		{
+			OutTypeName = NAME_None;
+			OutTypeObject = nullptr;
+		}
+	}
 	else if (const FObjectProperty* ObjectProperty = CastField<FObjectProperty>(InProperty))
 	{
 		if (RigVMCore::SupportsUObjects())
diff --git a/Engine/Plugins/Runtime/RigVM/Source/RigVM/Public/RigVMCore/RigVMUserWorkflow.h b/Engine/Plugins/Runtime/RigVM/Source/RigVM/Public/RigVMCore/RigVMUserWorkflow.h
--- a/Engine/Plugins/Runtime/RigVM/Source/RigVM/Public/RigVMCore/RigVMUserWorkflow.h
+++ b/Engine/Plugins/Runtime/RigVM/Source/RigVM/Public/RigVMCore/RigVMUserWorkflow.h
@@ -36,7 +36,8 @@ enum class ERigVMUserWorkflowType : uint8
 	NodeContext = 0x001,
 	PinContext = 0x002,
 	OnPinDefaultChanged  = 0x004,
-	All = NodeContext | PinContext | OnPinDefaultChanged
+	NodeContextButton  = 0x008,
+	All = NodeContext | PinContext | OnPinDefaultChanged | NodeContextButton
 };
 
 USTRUCT(BlueprintType)
diff --git a/Engine/Plugins/Runtime/RigVM/Source/RigVMDeveloper/Private/EdGraph/RigVMEdGraphNode.cpp b/Engine/Plugins/Runtime/RigVM/Source/RigVMDeveloper/Private/EdGraph/RigVMEdGraphNode.cpp
--- a/Engine/Plugins/Runtime/RigVM/Source/RigVMDeveloper/Private/EdGraph/RigVMEdGraphNode.cpp
+++ b/Engine/Plugins/Runtime/RigVM/Source/RigVMDeveloper/Private/EdGraph/RigVMEdGraphNode.cpp
@@ -1913,6 +1913,7 @@ void URigVMEdGraphNode::CopyPinDefaultsToModel(UEdGraphPin* Pin, bool bUndo, boo
 
 		if(DefaultValue.IsEmpty() && (
 			Pin->PinType.PinCategory == UEdGraphSchema_K2::PC_Object ||
+			Pin->PinType.PinCategory == UEdGraphSchema_K2::PC_Class ||
 			Pin->PinType.PinCategory == UEdGraphSchema_K2::PC_SoftObject ||
 			Pin->PinType.PinCategory == UEdGraphSchema_K2::AllObjectTypes ||
 			Pin->PinType.PinCategory == UEdGraphSchema_K2::PC_Interface
@@ -2257,5 +2258,24 @@ void URigVMEdGraphNode::UpdateProfilingSettings()
 }
 #endif
 
+FString URigVMEdGraphNode::GetPinMetaData(FName InPinName, FName InKey)
+{
+	URigVMNode* ModelNode = GetModelNode();
+	if(ModelNode == nullptr)
+	{
+		return FString();
+	}
+
+	FString Left, Right = InPinName.ToString();
+	URigVMPin::SplitPinPathAtStart(InPinName.ToString(), Left, Right);
+	URigVMPin* Pin = ModelNode->FindPin(Right);
+	if(Pin == nullptr)
+	{
+		return FString();
+	}
+
+	return Pin->GetMetaData(InKey);
+}
+
 #undef LOCTEXT_NAMESPACE
 
diff --git a/Engine/Plugins/Runtime/RigVM/Source/RigVMDeveloper/Private/RigVMModel/Nodes/RigVMDispatchNode.cpp b/Engine/Plugins/Runtime/RigVM/Source/RigVMDeveloper/Private/RigVMModel/Nodes/RigVMDispatchNode.cpp
--- a/Engine/Plugins/Runtime/RigVM/Source/RigVMDeveloper/Private/RigVMModel/Nodes/RigVMDispatchNode.cpp
+++ b/Engine/Plugins/Runtime/RigVM/Source/RigVMDeveloper/Private/RigVMModel/Nodes/RigVMDispatchNode.cpp
@@ -18,7 +18,11 @@ FText URigVMDispatchNode::GetToolTipText() const
 {
 	if (const FRigVMDispatchFactory* Factory = GetFactory())
 	{
-		return Factory->GetNodeTooltip(GetTemplatePinTypeMap());
+		FText ToolTipText = Factory->GetNodeTooltip(GetTemplatePinTypeMap());
+		if(!ToolTipText.IsEmpty())
+		{
+			return ToolTipText;
+		}
 	}
 	return URigVMNode::GetToolTipText();
 }
@@ -73,7 +77,11 @@ FText URigVMDispatchNode::GetToolTipTextForPin(const URigVMPin* InPin) const
 	if (const FRigVMDispatchFactory* Factory = GetFactory())
 	{
 		const URigVMPin* RootPin = InPin->GetRootPin();
-		return Factory->GetArgumentTooltip(RootPin->GetFName(), RootPin->GetTypeIndex());
+		FText ToolTipText = Factory->GetArgumentTooltip(RootPin->GetFName(), RootPin->GetTypeIndex());
+		if(!ToolTipText.IsEmpty())
+		{
+			return ToolTipText;
+		}
 	}
 	return URigVMNode::GetToolTipTextForPin(InPin);
 }
diff --git a/Engine/Plugins/Runtime/RigVM/Source/RigVMDeveloper/Private/RigVMModel/RigVMController.cpp b/Engine/Plugins/Runtime/RigVM/Source/RigVMDeveloper/Private/RigVMModel/RigVMController.cpp
--- a/Engine/Plugins/Runtime/RigVM/Source/RigVMDeveloper/Private/RigVMModel/RigVMController.cpp
+++ b/Engine/Plugins/Runtime/RigVM/Source/RigVMDeveloper/Private/RigVMModel/RigVMController.cpp
@@ -528,6 +528,13 @@ FRigVMPinInfo::FRigVMPinInfo(FProperty* InProperty, ERigVMPinDirection InDirecti
 	{
 		CPPTypeObject = StructProperty->Struct;
 	}
+	else if (const FClassProperty* ClassProperty = CastField<FClassProperty>(PropertyForType))
+	{
+		if(RigVMCore::SupportsUObjects())
+		{
+			CPPTypeObject = ClassProperty->MetaClass;
+		}
+	}
 	else if (const FObjectProperty* ObjectProperty = CastField<FObjectProperty>(PropertyForType))
 	{
 		if(RigVMCore::SupportsUObjects())
@@ -17285,6 +17292,19 @@ void URigVMController::ConfigurePinFromProperty(FProperty* InProperty, URigVMPin
 	{
 		InOutPin->CPPTypeObject = StructProperty->Struct;
 	}
+	else if (FClassProperty* ClassProperty = CastField<FClassProperty>(PropertyForType))
+	{
+		if(RigVMCore::SupportsUObjects())
+		{
+			InOutPin->CPPTypeObject = ClassProperty->MetaClass;
+		}
+		else
+		{
+			ReportErrorf(TEXT("Unsupported type '%s' for pin."), *ClassProperty->MetaClass->GetName(), *InOutPin->GetName());
+			InOutPin->CPPType = FString();
+			InOutPin->CPPTypeObject = nullptr;
+		}
+	}
 	else if (FObjectProperty* ObjectProperty = CastField<FObjectProperty>(PropertyForType))
 	{
 		if(RigVMCore::SupportsUObjects())
diff --git a/Engine/Plugins/Runtime/RigVM/Source/RigVMDeveloper/Private/RigVMModel/RigVMPin.cpp b/Engine/Plugins/Runtime/RigVM/Source/RigVMDeveloper/Private/RigVMModel/RigVMPin.cpp
--- a/Engine/Plugins/Runtime/RigVM/Source/RigVMDeveloper/Private/RigVMModel/RigVMPin.cpp
+++ b/Engine/Plugins/Runtime/RigVM/Source/RigVMDeveloper/Private/RigVMModel/RigVMPin.cpp
@@ -1013,9 +1013,20 @@ bool URigVMPin::IsValidDefaultValue(const FString& InDefaultValue) const
 				return false;
 			}
 
-			if(!Object->GetClass()->IsChildOf(Class))
+			const bool bIsClass = RigVMTypeUtils::IsUClassType(GetCPPType());
+			if(bIsClass)
 			{
-				return false;
+				if(!CastChecked<UClass>(Object)->IsChildOf(Class))
+				{
+					return false;
+				}
+			}
+			else
+			{
+				if(!Object->GetClass()->IsChildOf(Class))
+				{
+					return false;
+				}
 			}
 		} 
 		else if (UScriptStruct* ScriptStruct = Cast<UScriptStruct>(GetCPPTypeObject()))
diff --git a/Engine/Plugins/Runtime/RigVM/Source/RigVMDeveloper/Public/EdGraph/RigVMEdGraphNode.h b/Engine/Plugins/Runtime/RigVM/Source/RigVMDeveloper/Public/EdGraph/RigVMEdGraphNode.h
--- a/Engine/Plugins/Runtime/RigVM/Source/RigVMDeveloper/Public/EdGraph/RigVMEdGraphNode.h
+++ b/Engine/Plugins/Runtime/RigVM/Source/RigVMDeveloper/Public/EdGraph/RigVMEdGraphNode.h
@@ -127,6 +127,7 @@ public:
 #if WITH_EDITOR
 	virtual void AddPinSearchMetaDataInfo(const UEdGraphPin* Pin, TArray<FSearchTagDataPair>& OutTaggedMetaData) const override;
 #endif
+	virtual FString GetPinMetaData(FName InPinName, FName InKey) override;
 
 	virtual bool IsDeprecated() const override;
 	bool IsOutDated() const;
diff --git a/Engine/Plugins/Runtime/RigVM/Source/RigVMEditor/Private/Editor/RigVMDetailsViewWrapperObject.cpp b/Engine/Plugins/Runtime/RigVM/Source/RigVMEditor/Private/Editor/RigVMDetailsViewWrapperObject.cpp
--- a/Engine/Plugins/Runtime/RigVM/Source/RigVMEditor/Private/Editor/RigVMDetailsViewWrapperObject.cpp
+++ b/Engine/Plugins/Runtime/RigVM/Source/RigVMEditor/Private/Editor/RigVMDetailsViewWrapperObject.cpp
@@ -495,9 +495,19 @@ UClass* URigVMDetailsViewWrapperObject::GetClassForNodes(TArray<URigVMNode*> InN
 			}
 			else if(UClass* Class = Cast<UClass>(Pin->GetCPPTypeObject()))
 			{
-				FObjectProperty* ObjectProperty = new FObjectProperty(PropertyOwner, Pin->GetFName(), RF_Public);
-				ObjectProperty->PropertyClass = Class; 
-				(*ElementProperty) = ObjectProperty;
+				const bool bIsClass = RigVMTypeUtils::IsUClassType(Pin->GetCPPType());
+				if(bIsClass)
+				{
+					FClassProperty* ClassProperty = new FClassProperty(PropertyOwner, Pin->GetFName(), RF_Public);
+					ClassProperty->MetaClass = Class; 
+					(*ElementProperty) = ClassProperty;
+				}
+				else
+				{
+					FObjectProperty* ObjectProperty = new FObjectProperty(PropertyOwner, Pin->GetFName(), RF_Public);
+					ObjectProperty->PropertyClass = Class; 
+					(*ElementProperty) = ObjectProperty;
+				}
 			}
 			else if(UScriptStruct* ScriptStruct = Cast<UScriptStruct>(Pin->GetCPPTypeObject()))
 			{
diff --git a/Engine/Plugins/Runtime/RigVM/Source/RigVMEditor/Private/Widgets/SRigVMGraphNode.cpp b/Engine/Plugins/Runtime/RigVM/Source/RigVMEditor/Private/Widgets/SRigVMGraphNode.cpp
--- a/Engine/Plugins/Runtime/RigVM/Source/RigVMEditor/Private/Widgets/SRigVMGraphNode.cpp
+++ b/Engine/Plugins/Runtime/RigVM/Source/RigVMEditor/Private/Widgets/SRigVMGraphNode.cpp
@@ -41,6 +41,7 @@
 #include "Brushes/SlateColorBrush.h"
 #include "Widgets/SRigVMGraphPinCategory.h"
 #include "Widgets/SRigVMVariantTagWidget.h"
+#include "Widgets/SToolTip.h"
 
 #if WITH_EDITOR
 #include "Editor.h"
@@ -243,6 +244,12 @@ void SRigVMGraphNode::UpdateCompactNode()
 		InnerVerticalBox
 	];
 
+	if (!SWidget::GetToolTip().IsValid())
+	{
+		TSharedRef<SToolTip> DefaultToolTip = IDocumentation::Get()->CreateToolTip( TAttribute< FText >( this, &SGraphNode::GetNodeTooltip ), NULL, GraphNode->GetDocumentationLink(), GraphNode->GetDocumentationExcerptName() );
+		SetToolTip(DefaultToolTip);
+	}
+
 	CreatePinWidgets();
 }
 
@@ -323,6 +330,62 @@ FReply SRigVMGraphNode::OnAddPin()
 	return FReply::Handled();
 }
 
+void SRigVMGraphNode::CreateWorkflowWidgets()
+{
+	if (!LeftNodeBox.IsValid() || !ModelNode.IsValid())
+	{
+		return;
+	}
+
+	URigVMUnitNode* UnitNode = Cast<URigVMUnitNode>(ModelNode.Get());
+	if(UnitNode == nullptr)
+	{
+		return;
+	}
+
+	URigVMEdGraphNode* EdGraphNode = Cast<URigVMEdGraphNode>(GraphNode);
+	if(EdGraphNode == nullptr)
+	{
+		return;
+	}
+
+	URigVMController* Controller = EdGraphNode->GetController();
+	if(Controller == nullptr)
+	{
+		return;
+	}
+
+	TArray<FRigVMUserWorkflow> Workflows = UnitNode->GetSupportedWorkflows(ERigVMUserWorkflowType::NodeContextButton, UnitNode);
+	for(const FRigVMUserWorkflow& Workflow : Workflows)
+	{
+		TSharedRef<SButton> WorkflowButton = StaticCastSharedRef<SButton>(AddPinButtonContent(
+			FText::FromString(Workflow.GetTitle()),
+			FText::FromString(Workflow.GetTooltip())));
+
+		WorkflowButton->SetOnClicked(FOnClicked::CreateLambda([UnitNode, Controller, Workflow]()
+		{
+			URigVMUserWorkflowOptions* Options = Controller->MakeOptionsForWorkflow(UnitNode, Workflow);
+			Workflow.Perform(Options, Controller);
+			return FReply::Handled();
+		}));
+
+		WorkflowButton->SetToolTipText(FText::FromString(Workflow.GetTooltip()));
+
+		FMargin Padding = Settings->GetInputPinPadding();
+		Padding.Top += 2.0f;
+		Padding.Left -= 2.0f;
+
+		LeftNodeBox->AddSlot()
+			.AutoHeight()
+			.VAlign(VAlign_Center)
+			.HAlign(HAlign_Right)
+			.Padding(Padding)
+			[
+				WorkflowButton
+			];
+	}
+}
+
 bool SRigVMGraphNode::UseLowDetailNodeContent() const
 {
 	if(LastHighDetailSize.IsNearlyZero())
@@ -2422,7 +2485,9 @@ void SRigVMGraphNode::UpdatePinTreeView()
 	}
 
 	CreateAddPinButton();
-	
+
+	CreateWorkflowWidgets();
+
 	// add spacer widget at the end
 	LeftNodeBox->AddSlot()
 	.HAlign(HAlign_Fill)
diff --git a/Engine/Plugins/Runtime/RigVM/Source/RigVMEditor/Public/Widgets/SRigVMGraphNode.h b/Engine/Plugins/Runtime/RigVM/Source/RigVMEditor/Public/Widgets/SRigVMGraphNode.h
--- a/Engine/Plugins/Runtime/RigVM/Source/RigVMEditor/Public/Widgets/SRigVMGraphNode.h
+++ b/Engine/Plugins/Runtime/RigVM/Source/RigVMEditor/Public/Widgets/SRigVMGraphNode.h
@@ -59,6 +59,8 @@ public:
 
 	void CreateAddPinButton();
 
+	void CreateWorkflowWidgets();
+
 	/** Callback function executed when Add pin button is clicked */
 	virtual FReply OnAddPin() override;
 
