RigVM: Allow customization of nodes to use for templates

#rnx

diff --git a/Engine/Plugins/Runtime/RigVM/Source/RigVMDeveloper/Private/RigVMModel/RigVMController.cpp b/Engine/Plugins/Runtime/RigVM/Source/RigVMDeveloper/Private/RigVMModel/RigVMController.cpp
--- a/Engine/Plugins/Runtime/RigVM/Source/RigVMDeveloper/Private/RigVMModel/RigVMController.cpp
+++ b/Engine/Plugins/Runtime/RigVM/Source/RigVMDeveloper/Private/RigVMModel/RigVMController.cpp
@@ -1668,6 +1668,24 @@ URigVMUnitNode* URigVMController::AddUnitNode(UScriptStruct* InScriptStruct, TSu
 	return Node;
 }
 
+const UClass* URigVMController::GetNodeClassForTemplate(FRigVMTemplate* InTemplate) const
+{
+	// determine what kind of node we need to create
+	if(InTemplate->UsesDispatch())
+	{
+		return URigVMDispatchNode::StaticClass();
+	}
+	if(const FRigVMFunction* FirstFunction = InTemplate->GetOrCreatePermutation(0))
+	{
+		const UScriptStruct* PotentialUnitStruct = FirstFunction->Struct;
+		if(PotentialUnitStruct && PotentialUnitStruct->IsChildOf(FRigVMStruct::StaticStruct()))
+		{
+			return URigVMUnitNode::StaticClass();
+		}
+	}
+	return nullptr;
+}
+
 URigVMUnitNode* URigVMController::AddUnitNodeFromStructPath(const FString& InScriptStructPath, const FName& InMethodName, const FVector2D& InPosition, const FString& InNodeName, bool bSetupUndoRedo, bool bPrintPythonCommand)
 {
 	if(!IsValidGraph())
@@ -15970,18 +15988,11 @@ URigVMTemplateNode* URigVMController::AddTemplateNode(const FName& InNotation, c
 	FString Name = GetSchema()->GetValidNodeName(Graph, InNodeName.IsEmpty() ? Template->GetNodeName().ToString() : InNodeName);
 	URigVMTemplateNode* Node = nullptr;
 
-	// determine what kind of node we need to create
-	if(Template->UsesDispatch())
-	{
-		Node = NewObject<URigVMDispatchNode>(Graph, *Name);
-	}
-	else if(const FRigVMFunction* FirstFunction = Template->GetOrCreatePermutation(0))
+	if(const UClass* NodeClassToUse = GetNodeClassForTemplate(Template))
 	{
-		const UScriptStruct* PotentialUnitStruct = FirstFunction->Struct;
-		if(PotentialUnitStruct && PotentialUnitStruct->IsChildOf(FRigVMStruct::StaticStruct()))
-		{
-			Node = NewObject<URigVMUnitNode>(Graph, *Name); 
-		}
+		check(NodeClassToUse == URigVMTemplateNode::StaticClass() ||
+			NodeClassToUse->IsChildOf(URigVMTemplateNode::StaticClass()));
+		Node = NewObject<URigVMTemplateNode>(Graph, NodeClassToUse, *Name);
 	}
 
 	if(Node == nullptr)
diff --git a/Engine/Plugins/Runtime/RigVM/Source/RigVMDeveloper/Public/RigVMModel/RigVMController.h b/Engine/Plugins/Runtime/RigVM/Source/RigVMDeveloper/Public/RigVMModel/RigVMController.h
--- a/Engine/Plugins/Runtime/RigVM/Source/RigVMDeveloper/Public/RigVMModel/RigVMController.h
+++ b/Engine/Plugins/Runtime/RigVM/Source/RigVMDeveloper/Public/RigVMModel/RigVMController.h
@@ -1298,6 +1298,8 @@ protected:
 	URigVMUnitNode* AddUnitNode(UScriptStruct* InScriptStruct, TSubclassOf<URigVMUnitNode> InUnitNodeClass, const FName& InMethodName, const FVector2D& InPosition, const FString& InNodeName, bool bSetupUndoRedo, bool bPrintPythonCommand);
 #endif
 
+	virtual const UClass* GetNodeClassForTemplate(FRigVMTemplate* InTemplate) const;
+
 public:
 
 	// try to reconnect source and target pins after a node deletion
