Control Rig: Fix names of subpins on UserDefinedStructs
#rnx

diff --git a/Engine/Plugins/Runtime/RigVM/Source/RigVMDeveloper/Private/RigVMModel/RigVMController.cpp b/Engine/Plugins/Runtime/RigVM/Source/RigVMDeveloper/Private/RigVMModel/RigVMController.cpp
--- a/Engine/Plugins/Runtime/RigVM/Source/RigVMDeveloper/Private/RigVMModel/RigVMController.cpp
+++ b/Engine/Plugins/Runtime/RigVM/Source/RigVMDeveloper/Private/RigVMModel/RigVMController.cpp
@@ -16915,6 +16915,10 @@ void URigVMController::AddPinsForStruct(UStruct* InStruct, URigVMNode* InNode, U
 			}
 			
 			FName PropertyName = It->GetFName();
+			if (UUserDefinedStruct* UDStruct = Cast<UUserDefinedStruct>(StructToVisit))
+			{
+				PropertyName = *UDStruct->GetAuthoredNameForField(*It);
+			}
 
 			URigVMPin* Pin = NewObject<URigVMPin>(InParentPin == nullptr ? Cast<UObject>(InNode) : Cast<UObject>(InParentPin), PropertyName);
 			ConfigurePinFromProperty(*It, Pin, InPinDirection);
@@ -17240,6 +17244,12 @@ void URigVMController::ConfigurePinFromProperty(FProperty* InProperty, URigVMPin
 
 #if WITH_EDITOR
 
+	if (UUserDefinedStruct* Struct = Cast<UUserDefinedStruct>(InProperty->GetOwnerStruct()))
+	{
+		const FString DisplayName = Struct->GetAuthoredNameForField(InProperty);
+		InOutPin->DisplayName =  DisplayName.IsEmpty() ? NAME_None : FName(*DisplayName);
+	}
+	
 	InOutPin->bIsConstant = InProperty->HasMetaData(TEXT("Constant"));
 	FString CustomWidgetName = InProperty->GetMetaData(TEXT("CustomWidget"));
 	InOutPin->CustomWidgetName = CustomWidgetName.IsEmpty() ? FName(NAME_None) : FName(*CustomWidgetName);
