RigVM: Fix truncation warnings

#rnx

diff --git a/Engine/Plugins/Runtime/RigVM/Source/RigVM/Private/RigVMCore/RigVM.cpp b/Engine/Plugins/Runtime/RigVM/Source/RigVM/Private/RigVMCore/RigVM.cpp
--- a/Engine/Plugins/Runtime/RigVM/Source/RigVM/Private/RigVMCore/RigVM.cpp
+++ b/Engine/Plugins/Runtime/RigVM/Source/RigVM/Private/RigVMCore/RigVM.cpp
@@ -1291,7 +1291,7 @@ void URigVM::RebuildByteCodeOnLoad()
 }
 
 #if WITH_EDITOR
-bool URigVM::ShouldHaltAtInstruction(FRigVMExtendedExecuteContext& Context, const FName& InEventName, const uint16 InstructionIndex)
+bool URigVM::ShouldHaltAtInstruction(FRigVMExtendedExecuteContext& Context, const FName& InEventName, const int32 InstructionIndex)
 {
 	FRigVMDebugInfo* RigVMDebugInfo = Context.GetRigVMDebugInfo();
 
@@ -1465,8 +1465,8 @@ bool URigVM::ShouldHaltAtInstruction(FRigVMExtendedExecuteContext& Context, cons
 
 				// Create new temporary breakpoint
 				const FRigVMBreakpoint& NewBreakpoint = RigVMDebugInfo->AddBreakpoint(ContextPublicData.InstructionIndex, NewBreakpointNode, 0, true);
-				RigVMDebugInfo->SetBreakpointHits(NewBreakpoint, GetInstructionVisitedCount(Context, ContextPublicData.InstructionIndex));
-				RigVMDebugInfo->SetBreakpointActivationOnHit(NewBreakpoint, GetInstructionVisitedCount(Context, ContextPublicData.InstructionIndex));
+				RigVMDebugInfo->SetBreakpointHits(NewBreakpoint, static_cast<uint16>(GetInstructionVisitedCount(Context, ContextPublicData.InstructionIndex)));
+				RigVMDebugInfo->SetBreakpointActivationOnHit(NewBreakpoint, static_cast<uint16>(GetInstructionVisitedCount(Context, ContextPublicData.InstructionIndex)));
 				RigVMDebugInfo->SetCurrentBreakpointAction(ERigVMBreakpointAction::None);
 
 				RigVMDebugInfo->SetHaltedAtBreakpoint(NewBreakpoint);
@@ -1718,7 +1718,7 @@ ERigVMExecuteResult URigVM::ExecuteInstructions(FRigVMExtendedExecuteContext& Co
 	check(!Context.CurrentVMMemory.IsEmpty());
 
 	FRigVMExecuteContext& ContextPublicData = Context.GetPublicData<>();
-	TGuardValue<uint16> InstructionIndexGuard(ContextPublicData.InstructionIndex, (uint16)InFirstInstruction);
+	TGuardValue<int32> InstructionIndexGuard(ContextPublicData.InstructionIndex, InFirstInstruction);
 #if WITH_EDITOR
 	FInstructionBracketGuard InstructionBracket(Context, InFirstInstruction, InLastInstruction);
 	if(!InstructionBracket.ErrorMessage.IsEmpty())
@@ -2308,7 +2308,7 @@ ERigVMExecuteResult URigVM::ExecuteBranch(FRigVMExtendedExecuteContext& Context,
 	TGuardValue<bool> bCurrentlyRunningRootEntryGuard(Context.bCurrentlyRunningRootEntry, Context.bCurrentlyRunningRootEntry);
 	TGuardValue<FName> EventNameGuard(PublicContext.EventName, PublicContext.EventName);
 	TGuardValue<FName> FunctionNameGuard(PublicContext.FunctionName, PublicContext.FunctionName);
-	TGuardValue<uint16> InstructionIndexGuard(PublicContext.InstructionIndex, PublicContext.InstructionIndex);
+	TGuardValue<int32> InstructionIndexGuard(PublicContext.InstructionIndex, PublicContext.InstructionIndex);
 	TGuardValue<double> DeltaTimeGuard(PublicContext.DeltaTime, PublicContext.DeltaTime);
 	TGuardValue<double> AbsoluteTimeGuard(PublicContext.AbsoluteTime, PublicContext.AbsoluteTime);
 	TGuardValue<double> FramesPerSecondGuard(PublicContext.FramesPerSecond, PublicContext.FramesPerSecond);
diff --git a/Engine/Plugins/Runtime/RigVM/Source/RigVM/Private/RigVMCore/RigVMByteCode.cpp b/Engine/Plugins/Runtime/RigVM/Source/RigVM/Private/RigVMCore/RigVMByteCode.cpp
--- a/Engine/Plugins/Runtime/RigVM/Source/RigVM/Private/RigVMCore/RigVMByteCode.cpp
+++ b/Engine/Plugins/Runtime/RigVM/Source/RigVM/Private/RigVMCore/RigVMByteCode.cpp
@@ -194,11 +194,11 @@ FRigVMInstructionArray::FRigVMInstructionArray()
 
 FRigVMInstructionArray::FRigVMInstructionArray(const FRigVMByteCode& InByteCode, bool bByteCodeIsAligned)
 {
-	uint64 ByteIndex = 0;
+	int32 ByteIndex = 0;
 	while (ByteIndex < InByteCode.Num())
 	{
-		ERigVMOpCode OpCode = InByteCode.GetOpCodeAt(ByteIndex);
-		if ((int32)OpCode >= (int32)ERigVMOpCode::Invalid)
+		const ERigVMOpCode OpCode = InByteCode.GetOpCodeAt(ByteIndex);
+		if (OpCode >= ERigVMOpCode::Invalid)
 		{
 			checkNoEntry();
 			Instructions.Reset();
@@ -209,7 +209,7 @@ FRigVMInstructionArray::FRigVMInstructionArray(const FRigVMByteCode& InByteCode,
 
 		if (bByteCodeIsAligned)
 		{
-			uint64 Alignment = InByteCode.GetOpAlignment(OpCode);
+			int32 Alignment = InByteCode.GetOpAlignment(OpCode);
 			if (Alignment > 0)
 			{
 				while (!IsAligned(&InByteCode[ByteIndex], Alignment))
@@ -220,7 +220,7 @@ FRigVMInstructionArray::FRigVMInstructionArray(const FRigVMByteCode& InByteCode,
 
 			if (OpCode == ERigVMOpCode::Execute)
 			{
-				uint64 OperandByteIndex = ByteIndex + (uint64)sizeof(FRigVMExecuteOp);
+				const int32 OperandByteIndex = ByteIndex + sizeof(FRigVMExecuteOp);
 
 				Alignment = InByteCode.GetOperandAlignment();
 				if (Alignment > 0)
@@ -483,8 +483,6 @@ void FRigVMByteCode::Load(FArchive& Ar)
 		return;
 	}
 
-	FRigVMInstructionArray Instructions;
-
 	int32 InstructionCount = 0;
 	
 	ByteCode.Reset();
@@ -493,7 +491,6 @@ void FRigVMByteCode::Load(FArchive& Ar)
 
 	for (int32 InstructionIndex = 0; InstructionIndex < InstructionCount; InstructionIndex++)
 	{
-		FRigVMInstruction Instruction;
 		ERigVMOpCode OpCode = ERigVMOpCode::Invalid;
 
 		Ar << OpCode;
@@ -744,8 +741,7 @@ uint32 FRigVMByteCode::GetByteCodeHash() const
 		Hash = HashCombine(Hash, GetTypeHash(GetEntry(EntryIndex).InstructionIndex));
 	}
 
-	const FRigVMInstructionArray Instructions = GetInstructions();
-	for(const FRigVMInstruction& Instruction : Instructions)
+	for(const FRigVMInstruction& Instruction : GetInstructions())
 	{
 		Hash = HashCombine(Hash, GetOperatorHash(Instruction));
 	}
@@ -760,15 +756,13 @@ uint32 FRigVMByteCode::GetByteCodeHash() const
 
 uint32 FRigVMByteCode::GetOperatorHash(const FRigVMInstruction& InInstruction) const
 {
-	const ERigVMOpCode OpCode = GetOpCodeAt(InInstruction.ByteCodeIndex);
-	switch (OpCode)
+	switch (GetOpCodeAt(InInstruction.ByteCodeIndex))
 	{
 		case ERigVMOpCode::Execute:
 		{
 			const FRigVMExecuteOp& Op = GetOpAt<FRigVMExecuteOp>(InInstruction);
 			uint32 Hash = GetTypeHash(Op);
-			const FRigVMOperandArray Operands = GetOperandsForExecuteOp(InInstruction);
-			for(const FRigVMOperand& Operand : Operands)
+			for(const FRigVMOperand& Operand : GetOperandsForExecuteOp(InInstruction))
 			{
 				Hash = HashCombine(Hash, GetTypeHash(Operand));
 			}
@@ -866,16 +860,20 @@ uint32 FRigVMByteCode::GetOperatorHash(const FRigVMInstruction& InInstruction) c
 			ensure(false);
 			return 0;
 		}
+		default:
+		{
+			checkNoEntry();
+		}
 	}
 	return 0;
 }
 
-uint64 FRigVMByteCode::Num() const
+int32 FRigVMByteCode::Num() const
 {
-	return (uint64)ByteCode.Num();
+	return ByteCode.Num();
 }
 
-uint64 FRigVMByteCode::NumEntries() const
+int32 FRigVMByteCode::NumEntries() const
 {
 	return Entries.Num();
 }
@@ -897,10 +895,9 @@ int32 FRigVMByteCode::FindEntryIndex(const FName& InEntryName) const
 	return INDEX_NONE;
 }
 
-uint64 FRigVMByteCode::GetOpNumBytesAt(uint64 InByteCodeIndex, bool bIncludeOperands) const
+int32 FRigVMByteCode::GetOpNumBytesAt(int32 InByteCodeIndex, bool bIncludeOperands) const
 {
-	ERigVMOpCode OpCode = GetOpCodeAt(InByteCodeIndex);
-	switch (OpCode)
+	switch (GetOpCodeAt(InByteCodeIndex))
 	{
 		case ERigVMOpCode::Execute_0_Operands:
 		case ERigVMOpCode::Execute_1_Operands:
@@ -969,11 +966,11 @@ uint64 FRigVMByteCode::GetOpNumBytesAt(uint64 InByteCodeIndex, bool bIncludeOper
 		case ERigVMOpCode::Execute_64_Operands:
 		case ERigVMOpCode::Execute:
 		{
-			uint64 NumBytes = (uint64)sizeof(FRigVMExecuteOp);
+			int32 NumBytes = sizeof(FRigVMExecuteOp);
 			if(bIncludeOperands)
 			{
 				FRigVMExecuteOp ExecuteOp;
-				uint8* ExecuteOpPtr = (uint8*)&ExecuteOp;
+				uint8* ExecuteOpPtr = reinterpret_cast<uint8*>(&ExecuteOp);
 				for (int32 ByteIndex = 0; ByteIndex < sizeof(FRigVMExecuteOp); ByteIndex++)
 				{
 					ExecuteOpPtr[ByteIndex] = ByteCode[InByteCodeIndex + ByteIndex];
@@ -990,13 +987,13 @@ uint64 FRigVMByteCode::GetOpNumBytesAt(uint64 InByteCodeIndex, bool bIncludeOper
 						}
 					}
 				}
-				NumBytes += (uint64)ExecuteOp.GetOperandCount() * (uint64)sizeof(FRigVMOperand);
+				NumBytes += ExecuteOp.GetOperandCount() * sizeof(FRigVMOperand);
 			}
 			return NumBytes;
 		}
 		case ERigVMOpCode::Copy:
 		{
-			return (uint64)sizeof(FRigVMCopyOp);
+			return sizeof(FRigVMCopyOp);
 		}
 		case ERigVMOpCode::Zero:
 		case ERigVMOpCode::BoolFalse:
@@ -1006,24 +1003,24 @@ uint64 FRigVMByteCode::GetOpNumBytesAt(uint64 InByteCodeIndex, bool bIncludeOper
 		case ERigVMOpCode::ArrayReset:
 		case ERigVMOpCode::ArrayReverse:
 		{
-			return (uint64)sizeof(FRigVMUnaryOp);
+			return sizeof(FRigVMUnaryOp);
 		}
 		case ERigVMOpCode::Equals:
 		case ERigVMOpCode::NotEquals:
 		{
-			return (uint64)sizeof(FRigVMComparisonOp);
+			return sizeof(FRigVMComparisonOp);
 		}
 		case ERigVMOpCode::JumpAbsolute:
 		case ERigVMOpCode::JumpForward:
 		case ERigVMOpCode::JumpBackward:
 		{
-			return (uint64)sizeof(FRigVMJumpOp);
+			return sizeof(FRigVMJumpOp);
 		}
 		case ERigVMOpCode::JumpAbsoluteIf:
 		case ERigVMOpCode::JumpForwardIf:
 		case ERigVMOpCode::JumpBackwardIf:
 		{
-			return (uint64)sizeof(FRigVMJumpIfOp);
+			return sizeof(FRigVMJumpIfOp);
 		}
 		case ERigVMOpCode::ChangeType:
 		{
@@ -1032,7 +1029,7 @@ uint64 FRigVMByteCode::GetOpNumBytesAt(uint64 InByteCodeIndex, bool bIncludeOper
 		}
 		case ERigVMOpCode::Exit:
 		{
-			return (uint64)sizeof(FRigVMBaseOp);
+			return sizeof(FRigVMBaseOp);
 		}
 		case ERigVMOpCode::BeginBlock:
 		case ERigVMOpCode::ArrayGetNum:
@@ -1042,7 +1039,7 @@ uint64 FRigVMByteCode::GetOpNumBytesAt(uint64 InByteCodeIndex, bool bIncludeOper
 		case ERigVMOpCode::ArrayRemove:
 		case ERigVMOpCode::ArrayUnion:
 		{
-			return (uint64)sizeof(FRigVMBinaryOp);
+			return sizeof(FRigVMBinaryOp);
 		}
 		case ERigVMOpCode::ArrayAdd:
 		case ERigVMOpCode::ArrayGetAtIndex:
@@ -1051,35 +1048,35 @@ uint64 FRigVMByteCode::GetOpNumBytesAt(uint64 InByteCodeIndex, bool bIncludeOper
 		case ERigVMOpCode::ArrayDifference:
 		case ERigVMOpCode::ArrayIntersection:
 		{				
-			return (uint64)sizeof(FRigVMTernaryOp);
+			return sizeof(FRigVMTernaryOp);
 		}
 		case ERigVMOpCode::ArrayFind:
 		{				
-			return (uint64)sizeof(FRigVMQuaternaryOp);
+			return sizeof(FRigVMQuaternaryOp);
 		}
 		case ERigVMOpCode::ArrayIterator:
 		{				
-			return (uint64)sizeof(FRigVMSenaryOp);
+			return sizeof(FRigVMSenaryOp);
 		}
 		case ERigVMOpCode::EndBlock:
 		{
-			return (uint64)sizeof(FRigVMBaseOp);
+			return sizeof(FRigVMBaseOp);
 		}
 		case ERigVMOpCode::InvokeEntry:
 		{
-			return (uint64)sizeof(FRigVMInvokeEntryOp);
+			return sizeof(FRigVMInvokeEntryOp);
 		}
 		case ERigVMOpCode::JumpToBranch:
 		{
-			return (uint64)sizeof(FRigVMJumpToBranchOp);
+			return sizeof(FRigVMJumpToBranchOp);
 		}
 		case ERigVMOpCode::RunInstructions:
 		{
-			return (uint64)sizeof(FRigVMRunInstructionsOp);
+			return sizeof(FRigVMRunInstructionsOp);
 		}
 		case ERigVMOpCode::SetupTraits:
 		{
-			return (uint64)sizeof(FRigVMSetupTraitsOp);
+			return sizeof(FRigVMSetupTraitsOp);
 		}
 		case ERigVMOpCode::Invalid:
 		{
@@ -1090,31 +1087,31 @@ uint64 FRigVMByteCode::GetOpNumBytesAt(uint64 InByteCodeIndex, bool bIncludeOper
 	return 0;
 }
 
-uint64 FRigVMByteCode::AddZeroOp(const FRigVMOperand& InArg)
+int32 FRigVMByteCode::AddZeroOp(const FRigVMOperand& InArg)
 {
 	FRigVMUnaryOp Op(ERigVMOpCode::Zero, InArg);
 	return AddOp(Op);
 }
 
-uint64 FRigVMByteCode::AddFalseOp(const FRigVMOperand& InArg)
+int32 FRigVMByteCode::AddFalseOp(const FRigVMOperand& InArg)
 {
 	FRigVMUnaryOp Op(ERigVMOpCode::BoolFalse, InArg);
 	return AddOp(Op);
 }
 
-uint64 FRigVMByteCode::AddTrueOp(const FRigVMOperand& InArg)
+int32 FRigVMByteCode::AddTrueOp(const FRigVMOperand& InArg)
 {
 	FRigVMUnaryOp Op(ERigVMOpCode::BoolTrue, InArg);
 	return AddOp(Op);
 }
 
-uint64 FRigVMByteCode::AddCopyOp(const FRigVMOperand& InSource, const FRigVMOperand& InTarget)
+int32 FRigVMByteCode::AddCopyOp(const FRigVMOperand& InSource, const FRigVMOperand& InTarget)
 {
 	check(InTarget.GetMemoryType() != ERigVMMemoryType::Literal);
 	check(InSource != InTarget);
 
 	const FRigVMCopyOp Op(InSource, InTarget);
-	const uint64 ByteIndex = AddOp(Op);
+	const int32 ByteIndex = AddOp(Op);
 
 #if WITH_EDITORONLY_DATA
 	const FRigVMOperandArray InputOperands(&InSource, 1);
@@ -1125,29 +1122,29 @@ uint64 FRigVMByteCode::AddCopyOp(const FRigVMOperand& InSource, const FRigVMOper
 	return ByteIndex;
 }
 
-uint64 FRigVMByteCode::AddCopyOp(const FRigVMCopyOp& InCopyOp)
+int32 FRigVMByteCode::AddCopyOp(const FRigVMCopyOp& InCopyOp)
 {
 	return AddCopyOp(InCopyOp.Source, InCopyOp.Target);
 }
 
-uint64 FRigVMByteCode::AddIncrementOp(const FRigVMOperand& InArg)
+int32 FRigVMByteCode::AddIncrementOp(const FRigVMOperand& InArg)
 {
 	ensure(InArg.GetMemoryType() != ERigVMMemoryType::Literal);
 	FRigVMUnaryOp Op(ERigVMOpCode::Increment, InArg);
 	return AddOp(Op);
 }
 
-uint64 FRigVMByteCode::AddDecrementOp(const FRigVMOperand& InArg)
+int32 FRigVMByteCode::AddDecrementOp(const FRigVMOperand& InArg)
 {
 	ensure(InArg.GetMemoryType() != ERigVMMemoryType::Literal);
 	FRigVMUnaryOp Op(ERigVMOpCode::Decrement, InArg);
 	return AddOp(Op);
 }
 
-uint64 FRigVMByteCode::AddEqualsOp(const FRigVMOperand& InA, const FRigVMOperand& InB, const FRigVMOperand& InResult)
+int32 FRigVMByteCode::AddEqualsOp(const FRigVMOperand& InA, const FRigVMOperand& InB, const FRigVMOperand& InResult)
 {
 	const FRigVMComparisonOp Op(ERigVMOpCode::Equals, InA, InB, InResult);
-	const uint64 ByteIndex = AddOp(Op);
+	const int32 ByteIndex = AddOp(Op);
 	
 #if WITH_EDITORONLY_DATA
 	TArray<FRigVMOperand> Inputs;
@@ -1161,10 +1158,10 @@ uint64 FRigVMByteCode::AddEqualsOp(const FRigVMOperand& InA, const FRigVMOperand
 	return ByteIndex;
 }
 
-uint64 FRigVMByteCode::AddNotEqualsOp(const FRigVMOperand& InA, const FRigVMOperand& InB, const FRigVMOperand& InResult)
+int32 FRigVMByteCode::AddNotEqualsOp(const FRigVMOperand& InA, const FRigVMOperand& InB, const FRigVMOperand& InResult)
 {
 	const FRigVMComparisonOp Op(ERigVMOpCode::NotEquals, InA, InB, InResult);
-	const uint64 ByteIndex = AddOp(Op);
+	const int32 ByteIndex = AddOp(Op);
 
 #if WITH_EDITORONLY_DATA
 	TArray<FRigVMOperand> Inputs;
@@ -1178,26 +1175,26 @@ uint64 FRigVMByteCode::AddNotEqualsOp(const FRigVMOperand& InA, const FRigVMOper
 	return ByteIndex;
 }
 
-uint64 FRigVMByteCode::AddJumpOp(ERigVMOpCode InOpCode, uint16 InInstructionIndex)
+int32 FRigVMByteCode::AddJumpOp(ERigVMOpCode InOpCode, uint16 InInstructionIndex)
 {
 	FRigVMJumpOp Op(InOpCode, InInstructionIndex);
 	return AddOp(Op);
 }
 
-uint64 FRigVMByteCode::AddJumpIfOp(ERigVMOpCode InOpCode, uint16 InInstructionIndex, const FRigVMOperand& InConditionArg, bool bJumpWhenConditionIs)
+int32 FRigVMByteCode::AddJumpIfOp(ERigVMOpCode InOpCode, uint16 InInstructionIndex, const FRigVMOperand& InConditionArg, bool bJumpWhenConditionIs)
 {
 	FRigVMJumpIfOp Op(InOpCode, InConditionArg, InInstructionIndex, bJumpWhenConditionIs);
 	return AddOp(Op);
 }
 
-uint64 FRigVMByteCode::AddExecuteOp(uint16 InFunctionIndex, const FRigVMOperandArray& InOperands, const int32& StartPredicateIndex, const int32 PredicateCount)
+int32 FRigVMByteCode::AddExecuteOp(uint16 InFunctionIndex, const FRigVMOperandArray& InOperands, const int32& StartPredicateIndex, const int32 PredicateCount)
 {
 	FRigVMExecuteOp Op(InFunctionIndex, IntCastChecked<uint16>(InOperands.Num()));
-	Op.FirstPredicateIndex = StartPredicateIndex;
-	Op.PredicateCount = PredicateCount;
-	uint64 OpByteIndex = AddOp(Op);
+	Op.FirstPredicateIndex = static_cast<uint16>(StartPredicateIndex);
+	Op.PredicateCount = static_cast<uint16>(PredicateCount);
+	int32 OpByteIndex = AddOp(Op);
 
-	const uint64 OperandsByteIndex = (uint64)ByteCode.AddZeroed(sizeof(FRigVMOperand) * InOperands.Num());
+	const int32 OperandsByteIndex = ByteCode.AddZeroed(sizeof(FRigVMOperand) * InOperands.Num());
 	FMemory::Memcpy(ByteCode.GetData() + OperandsByteIndex, InOperands.GetData(), sizeof(FRigVMOperand) * InOperands.Num());
 
 	for(int32 Index = 0; Index < InOperands.Num(); Index++)
@@ -1209,18 +1206,18 @@ uint64 FRigVMByteCode::AddExecuteOp(uint16 InFunctionIndex, const FRigVMOperandA
 	return OpByteIndex;
 }
 
-uint64 FRigVMByteCode::InlineFunction(const FRigVMByteCode* FunctionByteCode, const FRigVMOperandArray& InOperands)
+int32 FRigVMByteCode::InlineFunction(const FRigVMByteCode* FunctionByteCode, const FRigVMOperandArray& InOperands)
 {
 	check(FunctionByteCode);
 	check(!FunctionByteCode->bByteCodeIsAligned);
-	uint64 OpByteIndex = ByteCode.Num();
+	const int32 OpByteIndex = ByteCode.Num();
 	ByteCode.Append(FunctionByteCode->ByteCode);
 	NumInstructions += FunctionByteCode->NumInstructions;
 	
 	return OpByteIndex;
 }
 
-uint64 FRigVMByteCode::AddExitOp()
+int32 FRigVMByteCode::AddExitOp()
 {
 	FRigVMBaseOp Op(ERigVMOpCode::Exit);
 	return AddOp(Op);
@@ -1236,14 +1233,14 @@ FString FRigVMByteCode::DumpToText() const
 		const FRigVMInstruction& Instruction = Instructions[InstructionIndex];
 		ERigVMOpCode OpCode = Instruction.OpCode;
 
-		FString Line = StaticEnum<ERigVMOpCode>()->GetNameByValue((int64)OpCode).ToString();
+		FString Line = StaticEnum<ERigVMOpCode>()->GetNameByValue(static_cast<int64>(OpCode)).ToString();
 
 		switch (OpCode)
 		{
 			case ERigVMOpCode::Execute:
 			{
 				const FRigVMExecuteOp& Op = GetOpAt<FRigVMExecuteOp>(Instruction.ByteCodeIndex);
-				Line += FString::Printf(TEXT(", FunctionIndex %d"), (int32)Op.FunctionIndex);
+				Line += FString::Printf(TEXT(", FunctionIndex %u"), Op.FunctionIndex);
 
 				FRigVMOperandArray Operands = GetOperandsForExecuteOp(Instruction);
 				if (Operands.Num() > 0)
@@ -1355,7 +1352,7 @@ FString FRigVMByteCode::DumpToText() const
 					{
 						break;
 					}
-					BranchInfoTexts.Add(FString::Printf(TEXT("%s (%d)"), *BranchInfo.Label.ToString(), (int32)BranchInfo.FirstInstruction));
+					BranchInfoTexts.Add(FString::Printf(TEXT("%s (%u)"), *BranchInfo.Label.ToString(), BranchInfo.FirstInstruction));
 				}
 				Line += FString::Join(BranchInfoTexts, TEXT(", "));
 				break;
@@ -1400,34 +1397,34 @@ FString FRigVMByteCode::DumpToText() const
 	return FString::Join(Lines, TEXT("\n"));
 }
 
-uint64 FRigVMByteCode::AddBeginBlockOp(FRigVMOperand InCountArg, FRigVMOperand InIndexArg)
+int32 FRigVMByteCode::AddBeginBlockOp(FRigVMOperand InCountArg, FRigVMOperand InIndexArg)
 {
 	FRigVMBinaryOp Op(ERigVMOpCode::BeginBlock, InCountArg, InIndexArg);
 	return AddOp(Op);
 }
 
-uint64 FRigVMByteCode::AddEndBlockOp()
+int32 FRigVMByteCode::AddEndBlockOp()
 {
 	FRigVMBaseOp Op(ERigVMOpCode::EndBlock);
 	return AddOp(Op);
 }
 
-uint64 FRigVMByteCode::AddInvokeEntryOp(const FName& InEntryName)
+int32 FRigVMByteCode::AddInvokeEntryOp(const FName& InEntryName)
 {
 	return AddOp(FRigVMInvokeEntryOp(InEntryName));
 }
 
-uint64 FRigVMByteCode::AddJumpToBranchOp(FRigVMOperand InBranchNameArg, int32 InFirstBranchInfoIndex)
+int32 FRigVMByteCode::AddJumpToBranchOp(FRigVMOperand InBranchNameArg, int32 InFirstBranchInfoIndex)
 {
 	return AddOp(FRigVMJumpToBranchOp(InBranchNameArg, InFirstBranchInfoIndex));
 }
 
-uint64 FRigVMByteCode::AddRunInstructionsOp(FRigVMOperand InExecuteStateArg, int32 InStartInstruction, int32 InEndInstruction)
+int32 FRigVMByteCode::AddRunInstructionsOp(FRigVMOperand InExecuteStateArg, int32 InStartInstruction, int32 InEndInstruction)
 {
 	return AddOp(FRigVMRunInstructionsOp(InExecuteStateArg, InStartInstruction, InEndInstruction));
 }
 
-uint64 FRigVMByteCode::AddSetupTraitsOp(FRigVMOperand InTraitListArg)
+int32 FRigVMByteCode::AddSetupTraitsOp(FRigVMOperand InTraitListArg)
 {
 	return AddOp(FRigVMSetupTraitsOp(InTraitListArg));
 }
@@ -1448,8 +1445,8 @@ int32 FRigVMByteCode::AddBranchInfo(const FName& InBranchLabel, int32 InInstruct
 	BranchInfo.Label = InBranchLabel;
 	BranchInfo.InstructionIndex = InInstructionIndex;
 	BranchInfo.ArgumentIndex = InArgumentIndex;
-	BranchInfo.FirstInstruction = InFirstBranchInstruction;
-	BranchInfo.LastInstruction = InLastBranchInstruction;
+	BranchInfo.FirstInstruction = static_cast<uint16>(InFirstBranchInstruction);
+	BranchInfo.LastInstruction = static_cast<uint16>(InLastBranchInstruction);
 	return AddBranchInfo(BranchInfo);
 }
 
@@ -1557,22 +1554,21 @@ FRigVMOperandArray FRigVMByteCode::GetOperandsForOp(const FRigVMInstruction& InI
 	return FRigVMOperandArray();
 }
 
-uint64 FRigVMByteCode::GetFirstOperandByteIndex(const FRigVMInstruction& InInstruction) const
+int32 FRigVMByteCode::GetFirstOperandByteIndex(const FRigVMInstruction& InInstruction) const
 {
 	if (InInstruction.OpCode == ERigVMOpCode::Execute)
 	{
-		const uint64 ByteCodeIndex = InInstruction.ByteCodeIndex;
+		const int32 ByteCodeIndex = InInstruction.ByteCodeIndex;
 		// if the bytecode is not aligned the OperandAlignment needs to be 0
 		ensure(bByteCodeIsAligned || InInstruction.OperandAlignment == 0);
-		return ByteCodeIndex + sizeof(FRigVMExecuteOp) + (uint64)InInstruction.OperandAlignment;
+		return ByteCodeIndex + sizeof(FRigVMExecuteOp) + InInstruction.OperandAlignment;
 	}
 
 	switch(InInstruction.OpCode)
 	{
 		case ERigVMOpCode::Copy:
 		{
-			const FRigVMCopyOp& Op = GetOpAt<FRigVMCopyOp>(InInstruction.ByteCodeIndex);
-			return InInstruction.ByteCodeIndex + ((uint64)&Op.Source - (uint64)&Op);
+			return InInstruction.ByteCodeIndex + offsetof(FRigVMCopyOp, Source);
 		}
 		case ERigVMOpCode::Zero:
 		case ERigVMOpCode::BoolFalse:
@@ -1582,21 +1578,18 @@ uint64 FRigVMByteCode::GetFirstOperandByteIndex(const FRigVMInstruction& InInstr
 		case ERigVMOpCode::ArrayReset:
 		case ERigVMOpCode::ArrayReverse:
 		{
-			const FRigVMUnaryOp& Op = GetOpAt<FRigVMUnaryOp>(InInstruction.ByteCodeIndex);
-			return InInstruction.ByteCodeIndex + ((uint64)&Op.Arg - (uint64)&Op);
+			return InInstruction.ByteCodeIndex + offsetof(FRigVMUnaryOp, Arg);
 		}
 		case ERigVMOpCode::Equals:
 		case ERigVMOpCode::NotEquals:
 		{
-			const FRigVMComparisonOp& Op = GetOpAt<FRigVMComparisonOp>(InInstruction.ByteCodeIndex);
-			return InInstruction.ByteCodeIndex + ((uint64)&Op.A - (uint64)&Op);
+			return InInstruction.ByteCodeIndex + offsetof(FRigVMComparisonOp, A);
 		}
 		case ERigVMOpCode::JumpAbsoluteIf:
 		case ERigVMOpCode::JumpForwardIf:
 		case ERigVMOpCode::JumpBackwardIf:
 		{
-			const FRigVMJumpIfOp& Op = GetOpAt<FRigVMJumpIfOp>(InInstruction.ByteCodeIndex);
-			return InInstruction.ByteCodeIndex + ((uint64)&Op.Arg - (uint64)&Op);
+			return InInstruction.ByteCodeIndex + offsetof(FRigVMJumpIfOp, Arg);
 		}
 		case ERigVMOpCode::BeginBlock:
 		case ERigVMOpCode::ArrayGetNum:
@@ -1606,8 +1599,7 @@ uint64 FRigVMByteCode::GetFirstOperandByteIndex(const FRigVMInstruction& InInstr
 		case ERigVMOpCode::ArrayRemove:
 		case ERigVMOpCode::ArrayUnion:
 		{
-			const FRigVMBinaryOp& Op = GetOpAt<FRigVMBinaryOp>(InInstruction.ByteCodeIndex);
-			return InInstruction.ByteCodeIndex + ((uint64)&Op.ArgA - (uint64)&Op);
+			return InInstruction.ByteCodeIndex + offsetof(FRigVMBinaryOp, ArgA);
 		}
 		case ERigVMOpCode::ArrayAdd:
 		case ERigVMOpCode::ArrayGetAtIndex:
@@ -1616,33 +1608,27 @@ uint64 FRigVMByteCode::GetFirstOperandByteIndex(const FRigVMInstruction& InInstr
 		case ERigVMOpCode::ArrayDifference:
 		case ERigVMOpCode::ArrayIntersection:
 		{
-			const FRigVMTernaryOp& Op = GetOpAt<FRigVMTernaryOp>(InInstruction.ByteCodeIndex);
-			return InInstruction.ByteCodeIndex + ((uint64)&Op.ArgA - (uint64)&Op);
+			return InInstruction.ByteCodeIndex + offsetof(FRigVMTernaryOp, ArgA);
 		}
 		case ERigVMOpCode::ArrayFind:
 		{
-			const FRigVMQuaternaryOp& Op = GetOpAt<FRigVMQuaternaryOp>(InInstruction.ByteCodeIndex);
-			return InInstruction.ByteCodeIndex + ((uint64)&Op.ArgA - (uint64)&Op);
+			return InInstruction.ByteCodeIndex + offsetof(FRigVMQuaternaryOp, ArgA);
 		}
 		case ERigVMOpCode::ArrayIterator:
 		{
-			const FRigVMSenaryOp& Op = GetOpAt<FRigVMSenaryOp>(InInstruction.ByteCodeIndex);
-			return InInstruction.ByteCodeIndex + ((uint64)&Op.ArgA - (uint64)&Op);
+			return InInstruction.ByteCodeIndex + offsetof(FRigVMSenaryOp, ArgA);
 		}
 		case ERigVMOpCode::JumpToBranch:
 		{
-			const FRigVMJumpToBranchOp& Op = GetOpAt<FRigVMJumpToBranchOp>(InInstruction.ByteCodeIndex);
-			return InInstruction.ByteCodeIndex + ((uint64)&Op.Arg - (uint64)&Op);
+			return InInstruction.ByteCodeIndex + offsetof(FRigVMJumpToBranchOp, Arg);
 		}
 		case ERigVMOpCode::RunInstructions:
 		{
-			const FRigVMRunInstructionsOp& Op = GetOpAt<FRigVMRunInstructionsOp>(InInstruction.ByteCodeIndex);
-			return InInstruction.ByteCodeIndex + ((uint64)&Op.Arg - (uint64)&Op);
+			return InInstruction.ByteCodeIndex + offsetof(FRigVMRunInstructionsOp, Arg);
 		}
 		case ERigVMOpCode::SetupTraits:
 		{
-			const FRigVMSetupTraitsOp& Op = GetOpAt<FRigVMSetupTraitsOp>(InInstruction.ByteCodeIndex);
-			return InInstruction.ByteCodeIndex + ((uint64)&Op.Arg - (uint64)&Op);
+			return InInstruction.ByteCodeIndex + offsetof(FRigVMSetupTraitsOp, Arg);
 		}
 		case ERigVMOpCode::JumpAbsolute:
 		case ERigVMOpCode::JumpForward:
@@ -1676,18 +1662,18 @@ TArray<int32> FRigVMByteCode::GetInstructionsForOperand(const FRigVMOperand& InO
 	return InstructionIndices;
 }
 
-uint64 FRigVMByteCode::GetOpAlignment(ERigVMOpCode InOpCode) const
+int32 FRigVMByteCode::GetOpAlignment(ERigVMOpCode InOpCode)
 {
 	switch (InOpCode)
 	{
 		case ERigVMOpCode::Execute:
 		{
-			static const uint64 Alignment = FRigVMExecuteOp::StaticStruct()->GetCppStructOps()->GetAlignment();
+			static const int32 Alignment = FRigVMExecuteOp::StaticStruct()->GetCppStructOps()->GetAlignment();
 			return Alignment;
 		}
 		case ERigVMOpCode::Copy:
 		{
-			static const uint64 Alignment = FRigVMCopyOp::StaticStruct()->GetCppStructOps()->GetAlignment();
+			static const int32 Alignment = FRigVMCopyOp::StaticStruct()->GetCppStructOps()->GetAlignment();
 			return Alignment;
 		}
 		case ERigVMOpCode::Zero:
@@ -1698,27 +1684,27 @@ uint64 FRigVMByteCode::GetOpAlignment(ERigVMOpCode InOpCode) const
 		case ERigVMOpCode::ArrayReset:
 		case ERigVMOpCode::ArrayReverse:
 		{
-			static const uint64 Alignment = FRigVMUnaryOp::StaticStruct()->GetCppStructOps()->GetAlignment();
+			static const int32 Alignment = FRigVMUnaryOp::StaticStruct()->GetCppStructOps()->GetAlignment();
 			return Alignment;
 		}
 		case ERigVMOpCode::Equals:
 		case ERigVMOpCode::NotEquals:
 		{
-			static const uint64 Alignment = FRigVMComparisonOp::StaticStruct()->GetCppStructOps()->GetAlignment();
+			static const int32 Alignment = FRigVMComparisonOp::StaticStruct()->GetCppStructOps()->GetAlignment();
 			return Alignment;
 		}
 		case ERigVMOpCode::JumpAbsolute:
 		case ERigVMOpCode::JumpForward:
 		case ERigVMOpCode::JumpBackward:
 		{
-			static const uint64 Alignment = FRigVMJumpOp::StaticStruct()->GetCppStructOps()->GetAlignment();
+			static const int32 Alignment = FRigVMJumpOp::StaticStruct()->GetCppStructOps()->GetAlignment();
 			return Alignment;
 		}
 		case ERigVMOpCode::JumpAbsoluteIf:
 		case ERigVMOpCode::JumpForwardIf:
 		case ERigVMOpCode::JumpBackwardIf:
 		{
-			static const uint64 Alignment = FRigVMJumpIfOp::StaticStruct()->GetCppStructOps()->GetAlignment();
+			static const int32 Alignment = FRigVMJumpIfOp::StaticStruct()->GetCppStructOps()->GetAlignment();
 			return Alignment;
 		}
 		case ERigVMOpCode::ChangeType:
@@ -1728,7 +1714,7 @@ uint64 FRigVMByteCode::GetOpAlignment(ERigVMOpCode InOpCode) const
 		}
 		case ERigVMOpCode::Exit:
 		{
-			static const uint64 Alignment = FRigVMBaseOp::StaticStruct()->GetCppStructOps()->GetAlignment();
+			static const int32 Alignment = FRigVMBaseOp::StaticStruct()->GetCppStructOps()->GetAlignment();
 			return Alignment;
 		}
 		case ERigVMOpCode::BeginBlock:
@@ -1739,7 +1725,7 @@ uint64 FRigVMByteCode::GetOpAlignment(ERigVMOpCode InOpCode) const
 		case ERigVMOpCode::ArrayRemove:
 		case ERigVMOpCode::ArrayUnion:
 		{
-			static const uint64 Alignment = FRigVMBinaryOp::StaticStruct()->GetCppStructOps()->GetAlignment();
+			static const int32 Alignment = FRigVMBinaryOp::StaticStruct()->GetCppStructOps()->GetAlignment();
 			return Alignment;
 		}
 		case ERigVMOpCode::ArrayAdd:
@@ -1749,42 +1735,42 @@ uint64 FRigVMByteCode::GetOpAlignment(ERigVMOpCode InOpCode) const
 		case ERigVMOpCode::ArrayDifference:
 		case ERigVMOpCode::ArrayIntersection:
 		{				
-			static const uint64 Alignment = FRigVMTernaryOp::StaticStruct()->GetCppStructOps()->GetAlignment();
+			static const int32 Alignment = FRigVMTernaryOp::StaticStruct()->GetCppStructOps()->GetAlignment();
 			return Alignment;
 		}
 		case ERigVMOpCode::ArrayFind:
 		{				
-			static const uint64 Alignment = FRigVMQuaternaryOp::StaticStruct()->GetCppStructOps()->GetAlignment();
+			static const int32 Alignment = FRigVMQuaternaryOp::StaticStruct()->GetCppStructOps()->GetAlignment();
 			return Alignment;
 		}
 		case ERigVMOpCode::ArrayIterator:
 		{				
-			static const uint64 Alignment = FRigVMSenaryOp::StaticStruct()->GetCppStructOps()->GetAlignment();
+			static const int32 Alignment = FRigVMSenaryOp::StaticStruct()->GetCppStructOps()->GetAlignment();
 			return Alignment;
 		}
 		case ERigVMOpCode::EndBlock:
 		{
-			static const uint64 Alignment = FRigVMBaseOp::StaticStruct()->GetCppStructOps()->GetAlignment();
+			static const int32 Alignment = FRigVMBaseOp::StaticStruct()->GetCppStructOps()->GetAlignment();
 			return Alignment;
 		}
 		case ERigVMOpCode::InvokeEntry:
 		{
-			static const uint64 Alignment = FRigVMInvokeEntryOp::StaticStruct()->GetCppStructOps()->GetAlignment();
+			static const int32 Alignment = FRigVMInvokeEntryOp::StaticStruct()->GetCppStructOps()->GetAlignment();
 			return Alignment;
 		}
 		case ERigVMOpCode::JumpToBranch:
 		{
-			static const uint64 Alignment = FRigVMJumpToBranchOp::StaticStruct()->GetCppStructOps()->GetAlignment();
+			static const int32 Alignment = FRigVMJumpToBranchOp::StaticStruct()->GetCppStructOps()->GetAlignment();
 			return Alignment;
 		}
 		case ERigVMOpCode::RunInstructions:
 		{
-			static const uint64 Alignment = FRigVMRunInstructionsOp::StaticStruct()->GetCppStructOps()->GetAlignment();
+			static const int32 Alignment = FRigVMRunInstructionsOp::StaticStruct()->GetCppStructOps()->GetAlignment();
 			return Alignment;
 		}
 		case ERigVMOpCode::SetupTraits:
 		{
-			static const uint64 Alignment = FRigVMSetupTraitsOp::StaticStruct()->GetCppStructOps()->GetAlignment();
+			static const int32 Alignment = FRigVMSetupTraitsOp::StaticStruct()->GetCppStructOps()->GetAlignment();
 			return Alignment;
 		}
 		case ERigVMOpCode::Invalid:
@@ -1792,13 +1778,16 @@ uint64 FRigVMByteCode::GetOpAlignment(ERigVMOpCode InOpCode) const
 			ensure(false);
 			return 0;
 		}
+		default:
+			checkNoEntry();
+			break;
 	}
 	return 0;
 }
 
-uint64 FRigVMByteCode::GetOperandAlignment() const
+int32 FRigVMByteCode::GetOperandAlignment()
 {
-	static const uint64 OperandAlignment = !FRigVMOperand::StaticStruct()->GetCppStructOps()->GetAlignment();
+	static const int32 OperandAlignment = FRigVMOperand::StaticStruct()->GetCppStructOps()->GetAlignment();
 	return OperandAlignment;
 }
 
@@ -1814,8 +1803,8 @@ void FRigVMByteCode::AlignByteCode()
 		return;
 	}
 
-	FRigVMInstructionArray Instructions(*this, false);
-	uint64 BytesToReserve = ByteCode.Num();
+	const FRigVMInstructionArray Instructions(*this, false);
+	int32 BytesToReserve = ByteCode.Num();
 
 	for (int32 InstructionIndex = 0; InstructionIndex < Instructions.Num(); InstructionIndex++)
 	{
@@ -1832,26 +1821,26 @@ void FRigVMByteCode::AlignByteCode()
 	AlignedByteCode.Reserve(BytesToReserve);
 	AlignedByteCode.AddZeroed(ByteCode.Num());
 
-	uint64 ShiftedBytes = 0;
+	int32 ShiftedBytes = 0;
 	for (int32 InstructionIndex = 0; InstructionIndex < Instructions.Num(); InstructionIndex++)
 	{
 		const FRigVMInstruction& Instruction = Instructions[InstructionIndex];
-		uint64 OriginalByteCodeIndex = Instruction.ByteCodeIndex;
-		uint64 AlignedByteCodeIndex = OriginalByteCodeIndex + ShiftedBytes;
-		uint64 OpAlignment = GetOpAlignment(Instruction.OpCode);
+		int32 OriginalByteCodeIndex = Instruction.ByteCodeIndex;
+		int32 AlignedByteCodeIndex = OriginalByteCodeIndex + ShiftedBytes;
+		int32 OpAlignment = GetOpAlignment(Instruction.OpCode);
 
 		if (OpAlignment > 0)
 		{
 			while (!IsAligned(&AlignedByteCode[AlignedByteCodeIndex], OpAlignment))
 			{
-				AlignedByteCode[AlignedByteCodeIndex] = (uint8)Instruction.OpCode;
+				AlignedByteCode[AlignedByteCodeIndex] = static_cast<uint8>(Instruction.OpCode);
 				AlignedByteCodeIndex++;
 				ShiftedBytes++;
 				AlignedByteCode.AddZeroed(1);
 			}
 		}
 
-		uint64 NumBytesToCopy = GetOpNumBytesAt(OriginalByteCodeIndex, false);
+		int32 NumBytesToCopy = GetOpNumBytesAt(OriginalByteCodeIndex, false);
 		for (int32 ByteIndex = 0; ByteIndex < NumBytesToCopy; ByteIndex++)
 		{
 			AlignedByteCode[AlignedByteCodeIndex + ByteIndex] = ByteCode[OriginalByteCodeIndex + ByteIndex];
@@ -1861,7 +1850,7 @@ void FRigVMByteCode::AlignByteCode()
 		{
 			AlignedByteCodeIndex += NumBytesToCopy;
 
-			static const uint64 OperandAlignment = GetOperandAlignment();
+			static const int32 OperandAlignment = GetOperandAlignment();
 			if (OperandAlignment > 0)
 			{
 				while (!IsAligned(&AlignedByteCode[AlignedByteCodeIndex], OperandAlignment))
@@ -1873,7 +1862,7 @@ void FRigVMByteCode::AlignByteCode()
 			}
 
 			FRigVMExecuteOp ExecuteOp;
-			uint8* ExecuteOpPtr = (uint8*)&ExecuteOp;
+			uint8* ExecuteOpPtr = reinterpret_cast<uint8*>(&ExecuteOp);
 			for (int32 ByteIndex = 0; ByteIndex < sizeof(FRigVMExecuteOp); ByteIndex++)
 			{
 				ExecuteOpPtr[ByteIndex] = ByteCode[OriginalByteCodeIndex + ByteIndex];
diff --git a/Engine/Plugins/Runtime/RigVM/Source/RigVM/Private/RigVMCore/RigVMDebugInfo.cpp b/Engine/Plugins/Runtime/RigVM/Source/RigVM/Private/RigVMCore/RigVMDebugInfo.cpp
--- a/Engine/Plugins/Runtime/RigVM/Source/RigVM/Private/RigVMCore/RigVMDebugInfo.cpp
+++ b/Engine/Plugins/Runtime/RigVM/Source/RigVM/Private/RigVMCore/RigVMDebugInfo.cpp
@@ -20,7 +20,7 @@ void FRigVMDebugInfo::StartExecution()
 	CurrentActiveBreakpointCallstack.Empty();
 }
 
-const FRigVMBreakpoint& FRigVMDebugInfo::FindBreakpoint(const uint16 InInstructionIndex, const UObject* InSubject) const
+const FRigVMBreakpoint& FRigVMDebugInfo::FindBreakpoint(const int32 InInstructionIndex, const UObject* InSubject) const
 {
 	TArray<FRigVMBreakpoint> BreakpointAtInstruction = FindBreakpointsAtInstruction(InInstructionIndex);
 	for (const FRigVMBreakpoint& BP : BreakpointAtInstruction)
@@ -35,7 +35,7 @@ const FRigVMBreakpoint& FRigVMDebugInfo::FindBreakpoint(const uint16 InInstructi
 	return EmptyBreakpoint;
 }
 
-TArray<FRigVMBreakpoint> FRigVMDebugInfo::FindBreakpointsAtInstruction(const uint16 InInstructionIndex) const
+TArray<FRigVMBreakpoint> FRigVMDebugInfo::FindBreakpointsAtInstruction(const int32 InInstructionIndex) const
 {
 	TArray<FRigVMBreakpoint> Result;
 	for (const FRigVMBreakpoint& BP : Breakpoints)
@@ -79,7 +79,7 @@ const FRigVMBreakpoint& FRigVMDebugInfo::FindBreakpoint(const FGuid& InGuid) con
 	return EmptyBreakpoint;
 }
 
-const FRigVMBreakpoint& FRigVMDebugInfo::AddBreakpoint(const uint16 InstructionIndex, TWeakObjectPtr<UObject> InNode, const uint16 InDepth,
+const FRigVMBreakpoint& FRigVMDebugInfo::AddBreakpoint(const int32 InstructionIndex, TWeakObjectPtr<UObject> InNode, const uint16 InDepth,
                                                             const bool bIsTemporary)
 {
 	static const FRigVMBreakpoint EmptyBreakpoint;
diff --git a/Engine/Plugins/Runtime/RigVM/Source/RigVM/Private/RigVMCore/RigVMExternalVariable.cpp b/Engine/Plugins/Runtime/RigVM/Source/RigVM/Private/RigVMCore/RigVMExternalVariable.cpp
--- a/Engine/Plugins/Runtime/RigVM/Source/RigVM/Private/RigVMCore/RigVMExternalVariable.cpp
+++ b/Engine/Plugins/Runtime/RigVM/Source/RigVM/Private/RigVMCore/RigVMExternalVariable.cpp
@@ -193,7 +193,7 @@ FRigVMExternalVariable RigVMTypeUtils::ExternalVariableFromCPPTypePath(const FNa
 	{
 		Variable.TypeName = *CPPTypeFromEnum(Enum);
 		Variable.TypeObject = Enum;
-		Variable.Size = Enum->GetResourceSizeBytes(EResourceSizeMode::EstimatedTotal);
+		Variable.Size = static_cast<int32>(Enum->GetResourceSizeBytes(EResourceSizeMode::EstimatedTotal));
 	}
 	else
 	{
@@ -267,7 +267,7 @@ FRigVMExternalVariable RigVMTypeUtils::ExternalVariableFromCPPType(const FName&
 	{
 		Variable.TypeName = *CPPTypeFromEnum(Enum);
 		Variable.TypeObject = Enum;
-		Variable.Size = Enum->GetResourceSizeBytes(EResourceSizeMode::EstimatedTotal);
+		Variable.Size = static_cast<int32>(Enum->GetResourceSizeBytes(EResourceSizeMode::EstimatedTotal));
 	}
 	else
 	{
diff --git a/Engine/Plugins/Runtime/RigVM/Source/RigVM/Private/RigVMCore/RigVMGraphFunctionDefinition.cpp b/Engine/Plugins/Runtime/RigVM/Source/RigVM/Private/RigVMCore/RigVMGraphFunctionDefinition.cpp
--- a/Engine/Plugins/Runtime/RigVM/Source/RigVM/Private/RigVMCore/RigVMGraphFunctionDefinition.cpp
+++ b/Engine/Plugins/Runtime/RigVM/Source/RigVM/Private/RigVMCore/RigVMGraphFunctionDefinition.cpp
@@ -96,7 +96,7 @@ FArchive& operator<<(FArchive& Ar, FRigVMFunctionCompilationData& Data)
 
 	// Serialize OperandToDebugRegisters
 	{
-		uint8 NumKeys = Data.OperandToDebugRegisters.Num();
+		uint8 NumKeys = static_cast<uint8>(Data.OperandToDebugRegisters.Num());
 		Ar << NumKeys;
 
 		if (Ar.IsLoading())
@@ -121,7 +121,7 @@ FArchive& operator<<(FArchive& Ar, FRigVMFunctionCompilationData& Data)
 			for (TPair<FRigVMOperand, TArray<FRigVMOperand>>& Pair : Data.OperandToDebugRegisters)
 			{
 				Ar << Pair.Key;
-				uint8 NumValues = Pair.Value.Num();
+				uint8 NumValues = static_cast<uint8>(Pair.Value.Num());
 				Ar << NumValues;
 				for (FRigVMOperand& Operand : Pair.Value)
 				{
diff --git a/Engine/Plugins/Runtime/RigVM/Source/RigVM/Private/RigVMCore/RigVMMemoryDeprecated.cpp b/Engine/Plugins/Runtime/RigVM/Source/RigVM/Private/RigVMCore/RigVMMemoryDeprecated.cpp
--- a/Engine/Plugins/Runtime/RigVM/Source/RigVM/Private/RigVMCore/RigVMMemoryDeprecated.cpp
+++ b/Engine/Plugins/Runtime/RigVM/Source/RigVM/Private/RigVMCore/RigVMMemoryDeprecated.cpp
@@ -248,7 +248,7 @@ FRigVMRegisterOffset::FRigVMRegisterOffset(UScriptStruct* InScriptStruct, const
 			else
 			{
 				Offset.CPPType = *Property->GetCPPType();
-				Offset.ElementSize = Property->GetElementSize();
+				Offset.ElementSize = static_cast<uint16>(Property->GetElementSize());
 
 				if (FArrayProperty* ArrayProperty = CastField<FArrayProperty>(Property))
 				{
@@ -256,7 +256,7 @@ FRigVMRegisterOffset::FRigVMRegisterOffset(UScriptStruct* InScriptStruct, const
 					Property = ArrayProperty->Inner;
 
 					Offset.CPPType = *ArrayProperty->Inner->GetCPPType();
-					Offset.ElementSize = ArrayProperty->Inner->GetElementSize();
+					Offset.ElementSize = static_cast<uint16>(ArrayProperty->Inner->GetElementSize());
 				}
 
 				if (FStructProperty* StructProperty = CastField<FStructProperty>(Property))
@@ -314,7 +314,7 @@ FRigVMRegisterOffset::FRigVMRegisterOffset(UScriptStruct* InScriptStruct, const
 			else
 			{
 				Offset.CPPType = *InArrayProperty->Inner->GetCPPType();
-				Offset.ElementSize = InArrayProperty->Inner->GetElementSize();
+				Offset.ElementSize = static_cast<uint16>(InArrayProperty->Inner->GetElementSize());
 
 				if (FArrayProperty* ArrayProperty = CastField<FArrayProperty>(InArrayProperty->Inner))
 				{
@@ -521,13 +521,13 @@ FRigVMMemoryContainer& FRigVMMemoryContainer::operator= (const FRigVMMemoryConta
 		const FRigVMRegister& Register = Registers[Index];
 		if (Register.IsNestedDynamic())
 		{
-			FMemory::Memzero(&Data[Register.GetWorkByteIndex()], sizeof(FRigVMNestedByteArray));
-			*(FRigVMNestedByteArray*)&Data[Register.GetWorkByteIndex()] = FRigVMNestedByteArray();
+			FMemory::Memzero(&Data[static_cast<int32>(Register.GetWorkByteIndex())], sizeof(FRigVMNestedByteArray));
+			*(FRigVMNestedByteArray*)&Data[static_cast<int32>(Register.GetWorkByteIndex())] = FRigVMNestedByteArray();
 		}
 		else if(Register.IsDynamic())
 		{
-			FMemory::Memzero(&Data[Register.GetWorkByteIndex()], sizeof(FRigVMByteArray));
-			*(FRigVMByteArray*)&Data[Register.GetWorkByteIndex()] = FRigVMByteArray();
+			FMemory::Memzero(&Data[static_cast<int32>(Register.GetWorkByteIndex())], sizeof(FRigVMByteArray));
+			*(FRigVMByteArray*)&Data[static_cast<int32>(Register.GetWorkByteIndex())] = FRigVMByteArray();
 		}
 
 		Construct(Index);
@@ -636,10 +636,9 @@ void FRigVMMemoryContainer::Load(FArchive& Ar)
 		{
 			FRigVMRegister& Register = Registers[RegisterIndex];
 
-			UScriptStruct* ScriptStruct = GetScriptStruct(Register);
-			if (ScriptStruct)
+			if (UScriptStruct* ScriptStruct = GetScriptStruct(Register))
 			{
-				Register.ElementSize = ScriptStruct->GetStructureSize();
+				Register.ElementSize = static_cast<uint16>(ScriptStruct->GetStructureSize());
 			}
 			else if (Register.Type == ERigVMRegisterType::Name)
 			{
@@ -669,10 +668,9 @@ void FRigVMMemoryContainer::Load(FArchive& Ar)
 		{
 			FRigVMRegisterOffset& RegisterOffset = RegisterOffsets[RegisterOffsetIndex];
 
-			UScriptStruct* ScriptStruct = RegisterOffset.GetScriptStruct();
-			if (ScriptStruct)
+			if (UScriptStruct* ScriptStruct = RegisterOffset.GetScriptStruct())
 			{
-				RegisterOffset.SetElementSize(ScriptStruct->GetStructureSize());
+				RegisterOffset.SetElementSize(static_cast<uint16>(ScriptStruct->GetStructureSize()));
 			}
 			if (RegisterOffset.GetType() == ERigVMRegisterType::Name)
 			{
@@ -724,7 +722,7 @@ void FRigVMMemoryContainer::Load(FArchive& Ar)
 					if (!bEncounteredErrorDuringLoad)
 					{
 						ensure(View.Num() <= Register.GetAllocatedBytes());
-						FMemory::Memcpy(&Data[Register.GetWorkByteIndex()], View.GetData(), View.Num());
+						FMemory::Memcpy(&Data[static_cast<int32>(Register.GetWorkByteIndex())], View.GetData(), View.Num());
 					}
 					break;
 				}
@@ -736,7 +734,7 @@ void FRigVMMemoryContainer::Load(FArchive& Ar)
 					if (!bEncounteredErrorDuringLoad)
 					{
 						ensure(View.Num() == Register.GetTotalElementCount());
-						RigVMCopy<FName>(&Data[Register.GetWorkByteIndex()], View.GetData(), View.Num());
+						RigVMCopy<FName>(&Data[static_cast<int32>(Register.GetWorkByteIndex())], View.GetData(), View.Num());
 					}
 					break;
 				}
@@ -748,7 +746,7 @@ void FRigVMMemoryContainer::Load(FArchive& Ar)
 					if (!bEncounteredErrorDuringLoad)
 					{
 						ensure(View.Num() == Register.GetTotalElementCount());
-						RigVMCopy<FString>(&Data[Register.GetWorkByteIndex()], View.GetData(), View.Num());
+						RigVMCopy<FString>(&Data[static_cast<int32>(Register.GetWorkByteIndex())], View.GetData(), View.Num());
 					}
 					break;
 				}
@@ -761,7 +759,7 @@ void FRigVMMemoryContainer::Load(FArchive& Ar)
 					{
 						ensure(View.Num() == Register.GetTotalElementCount());
 
-						uint8* DataPtr = &Data[Register.GetWorkByteIndex()];
+						uint8* DataPtr = &Data[static_cast<int32>(Register.GetWorkByteIndex())];
 						UScriptStruct* ScriptStruct = GetScriptStruct(Register);
 						if (ScriptStruct && !bEncounteredErrorDuringLoad)
 						{
@@ -786,14 +784,6 @@ void FRigVMMemoryContainer::Load(FArchive& Ar)
 		}
 		else if (!Register.IsNestedDynamic())
 		{
-
-			FRigVMByteArray* ArrayStorage = nullptr;
-
-			if (!bEncounteredErrorDuringLoad)
-			{
-				ArrayStorage = (FRigVMByteArray*)&Data[Register.GetWorkByteIndex()];
-			}
-			
 			// notice that all dynamic registers' memory is not really serialized
 			// serialize empty arrays here just to avoid having to increment serialization version
 			// the loaded data is not getting used at all, it will be recomputed at execution time
@@ -831,19 +821,13 @@ void FRigVMMemoryContainer::Load(FArchive& Ar)
 			FRigVMNestedByteArray* NestedArrayStorage = nullptr;
 			if (!bEncounteredErrorDuringLoad)
 			{
-				NestedArrayStorage = (FRigVMNestedByteArray*)&Data[Register.GetWorkByteIndex()];
+				NestedArrayStorage = (FRigVMNestedByteArray*)&Data[static_cast<int32>(Register.GetWorkByteIndex())];
 				NestedArrayStorage->Reset();
 				NestedArrayStorage->SetNumZeroed(Register.SliceCount);
 			}
 			
 			for (int32 SliceIndex = 0; SliceIndex < Register.SliceCount; SliceIndex++)
 			{
-				FRigVMByteArray* ArrayStorage = nullptr;
-				if (NestedArrayStorage)
-				{
-					ArrayStorage = &(*NestedArrayStorage)[SliceIndex];
-				}
-
 				// notice that all dynamic registers' memory is not really serialized
 				// serialize empty arrays here just to avoid having to increment serialization version
 				// the loaded data is not getting used at all, it will be recomputed at execution time
@@ -952,7 +936,7 @@ int32 FRigVMMemoryContainer::Allocate(const FName& InNewName, int32 InElementSiz
 	{
 		for (uint16 SliceIndex = 0; SliceIndex < NewRegister.SliceCount; SliceIndex++)
 		{
-			FMemory::Memcpy(&Data[NewRegister.GetWorkByteIndex(SliceIndex)], InDataPtr, NewRegister.GetNumBytesPerSlice());
+			FMemory::Memcpy(&Data[static_cast<int32>(NewRegister.GetWorkByteIndex(SliceIndex))], InDataPtr, NewRegister.GetNumBytesPerSlice());
 		}
 	}
 
@@ -986,7 +970,7 @@ bool FRigVMMemoryContainer::Construct(int32 InRegisterIndex, int32 InElementInde
 	int32 ElementCount = InElementIndex == INDEX_NONE ? Register.GetTotalElementCount() : 1;
 	int32 NumSlices = 1;
 
-	uint8* TopDataPtr = &Data[Register.GetWorkByteIndex()];
+	uint8* TopDataPtr = &Data[static_cast<int32>(Register.GetWorkByteIndex())];
 	if (Register.IsDynamic())
 	{
 		if (Register.IsNestedDynamic())
@@ -1100,7 +1084,7 @@ bool FRigVMMemoryContainer::Destroy(int32 InRegisterIndex, int32 InElementIndex,
 	int32 ElementCount = InElementIndex == INDEX_NONE ? Register.GetTotalElementCount() : 1;
 	int32 NumSlices = 1;
 
-	uint8* TopDataPtr = &Data[Register.GetWorkByteIndex()];
+	uint8* TopDataPtr = &Data[static_cast<int32>(Register.GetWorkByteIndex())];
 	if (Register.IsDynamic())
 	{
 		if (Register.IsNestedDynamic())
@@ -1213,7 +1197,7 @@ bool FRigVMMemoryContainer::Destroy(int32 InRegisterIndex, int32 InElementIndex,
 
 	if (Register.IsDynamic())
 	{
-		TopDataPtr = &Data[Register.GetWorkByteIndex()];
+		TopDataPtr = &Data[static_cast<int32>(Register.GetWorkByteIndex())];
 
 		if (Register.IsNestedDynamic())
 		{
@@ -1241,7 +1225,7 @@ int32 FRigVMMemoryContainer::GetOrAddRegisterOffset(int32 InRegisterIndex, const
 	ensure(Registers.IsValidIndex(InRegisterIndex));
 
 	UScriptStruct* ScriptStruct = GetScriptStruct(InRegisterIndex);
-	int32 ElementSize = (int32)Registers[InRegisterIndex].ElementSize;
+	int32 ElementSize = Registers[InRegisterIndex].ElementSize;
 	int32 InitialOffset = InArrayElement * ElementSize;
 	return GetOrAddRegisterOffset(InRegisterIndex, ScriptStruct, InSegmentPath, InitialOffset);
 }
@@ -1259,7 +1243,7 @@ int32 FRigVMMemoryContainer::GetOrAddRegisterOffset(int32 InRegisterIndex, UScri
 	if (InElementSize == 0)
 	{
 		ensure(Registers.IsValidIndex(InRegisterIndex)); 
-		InElementSize = (int32)Registers[InRegisterIndex].ElementSize;
+		InElementSize = Registers[InRegisterIndex].ElementSize;
 	}
 	
 	FName BaseCPPType = NAME_None;
@@ -1267,14 +1251,14 @@ int32 FRigVMMemoryContainer::GetOrAddRegisterOffset(int32 InRegisterIndex, UScri
 	{
 		if(InElementSize == 0)
 		{
-			InElementSize = (int32)Registers[InRegisterIndex].ElementSize;
+			InElementSize = Registers[InRegisterIndex].ElementSize;
 		}
 #if WITH_EDITORONLY_DATA
 		BaseCPPType = Registers[InRegisterIndex].BaseCPPType;
 #endif
 	}
 
-	FRigVMRegisterOffset Offset(InScriptStruct, InSegmentPath, InInitialOffset, InElementSize, BaseCPPType);
+	FRigVMRegisterOffset Offset(InScriptStruct, InSegmentPath, InInitialOffset, static_cast<uint16>(InElementSize), BaseCPPType);
 	int32 ExistingIndex = RegisterOffsets.Find(Offset);
 	if (ExistingIndex == INDEX_NONE)
 	{
@@ -1315,32 +1299,29 @@ void FRigVMMemoryContainer::UpdateRegisters()
 				continue;
 			}
 
-			ensure(Data.IsValidIndex(Register.GetWorkByteIndex()));
-
-
-			if (ensure(Data.IsValidIndex(Register.GetWorkByteIndex())))
+			if (ensure(Data.IsValidIndex(static_cast<int32>(Register.GetWorkByteIndex()))))
 			{
-				uint8* Pointer = (uint8*)(&(Data[Register.GetWorkByteIndex()]));
+				uint8* Pointer = (&(Data[static_cast<int32>(Register.GetWorkByteIndex())]));
 			
 				if (Register.AlignmentBytes > 0)
 				{
 					if (!IsAligned(Pointer, Alignment))
 					{
-						Data.RemoveAt(Register.GetFirstAllocatedByte(), Register.AlignmentBytes);
+						Data.RemoveAt(static_cast<int32>(Register.GetFirstAllocatedByte()), Register.AlignmentBytes);
 						AlignmentShift -= Register.AlignmentBytes;
 						Register.ByteIndex -= Register.AlignmentBytes;
 						Register.AlignmentBytes = 0;
-						Pointer = (uint8*)(&(Data[Register.GetWorkByteIndex()]));
+						Pointer = (&(Data[static_cast<int32>(Register.GetWorkByteIndex())]));
 					}
 				}
 
 				while (!IsAligned(Pointer, Alignment))
 				{
-					Data.InsertZeroed(Register.GetFirstAllocatedByte(), 1);
+					Data.InsertZeroed(static_cast<int32>(Register.GetFirstAllocatedByte()), 1);
 					Register.AlignmentBytes++;
 					Register.ByteIndex++;
 					AlignmentShift++;
-					Pointer = (uint8*)(&(Data[Register.GetWorkByteIndex()]));
+					Pointer = (&(Data[static_cast<int32>(Register.GetWorkByteIndex())]));
 				}
 			}
 		}
diff --git a/Engine/Plugins/Runtime/RigVM/Source/RigVM/Private/RigVMCore/RigVMObjectArchive.cpp b/Engine/Plugins/Runtime/RigVM/Source/RigVM/Private/RigVMCore/RigVMObjectArchive.cpp
--- a/Engine/Plugins/Runtime/RigVM/Source/RigVM/Private/RigVMCore/RigVMObjectArchive.cpp
+++ b/Engine/Plugins/Runtime/RigVM/Source/RigVM/Private/RigVMCore/RigVMObjectArchive.cpp
@@ -256,7 +256,7 @@ void FRigVMObjectArchiveReader::Serialize(void* V, int64 Length)
 	{
 		check(V);
 		check(Archive.Buffer.Num() >= Offset + Length);
-		FMemory::Memcpy( V, &(Archive.Buffer[Offset]), Length );
+		FMemory::Memcpy( V, &(Archive.Buffer[IntCastChecked<int32, int64>(Offset)]), Length );
 		Offset += Length;
 	}
 }
@@ -381,7 +381,7 @@ FArchive& FRigVMObjectArchiveReader::operator<<(FName& Value)
 	{
 		int64 OffsetOfString = -1;
 		*this << OffsetOfString;
-		check(Archive.Buffer.IsValidIndex(OffsetOfString));
+		check(Archive.Buffer.IsValidIndex(IntCastChecked<int32, int64>(OffsetOfString)));
 		
 		if(const FName* NamePtr = OffsetToName.Find(OffsetOfString))
 		{
diff --git a/Engine/Plugins/Runtime/RigVM/Source/RigVM/Private/RigVMCore/RigVMProfilingInfo.cpp b/Engine/Plugins/Runtime/RigVM/Source/RigVM/Private/RigVMCore/RigVMProfilingInfo.cpp
--- a/Engine/Plugins/Runtime/RigVM/Source/RigVM/Private/RigVMCore/RigVMProfilingInfo.cpp
+++ b/Engine/Plugins/Runtime/RigVM/Source/RigVM/Private/RigVMCore/RigVMProfilingInfo.cpp
@@ -40,6 +40,6 @@ void FRigVMProfilingInfo::StopProfiling()
 {
 #if WITH_EDITOR
 	const uint64 Cycles = GetOverallCycles() > 0 ? GetOverallCycles() : (FPlatformTime::Cycles64() - GetStartCycles());
-	LastExecutionMicroSeconds = Cycles * FPlatformTime::GetSecondsPerCycle() * 1000.0 * 1000.0;
+	LastExecutionMicroSeconds = static_cast<double>(Cycles) * FPlatformTime::GetSecondsPerCycle() * 1000.0 * 1000.0;
 #endif
 }
diff --git a/Engine/Plugins/Runtime/RigVM/Source/RigVM/Private/RigVMFunctions/Animation/RigVMFunction_GetDeltaTime.cpp b/Engine/Plugins/Runtime/RigVM/Source/RigVM/Private/RigVMFunctions/Animation/RigVMFunction_GetDeltaTime.cpp
--- a/Engine/Plugins/Runtime/RigVM/Source/RigVM/Private/RigVMFunctions/Animation/RigVMFunction_GetDeltaTime.cpp
+++ b/Engine/Plugins/Runtime/RigVM/Source/RigVM/Private/RigVMFunctions/Animation/RigVMFunction_GetDeltaTime.cpp
@@ -6,7 +6,7 @@
 
 FRigVMFunction_GetDeltaTime_Execute()
 {
-	Result = ExecuteContext.GetDeltaTime();
+	Result = ExecuteContext.GetDeltaTime<float>();
 }
 
 #if WITH_DEV_AUTOMATION_TESTS
@@ -14,7 +14,7 @@ FRigVMFunction_GetDeltaTime_Execute()
 
 IMPLEMENT_RIGVMSTRUCT_AUTOMATION_TEST(FRigVMFunction_GetDeltaTime)
 {
-	ExecuteContext.SetDeltaTime(0.2f);
+	ExecuteContext.SetDeltaTime(0.2);
 	Execute();
 	AddErrorIfFalse(FMath::IsNearlyEqual(Unit.Result, 0.2f), TEXT("unexpected delta time"));
 	return true;
diff --git a/Engine/Plugins/Runtime/RigVM/Source/RigVM/Private/RigVMFunctions/Animation/RigVMFunction_TimeConversion.cpp b/Engine/Plugins/Runtime/RigVM/Source/RigVM/Private/RigVMFunctions/Animation/RigVMFunction_TimeConversion.cpp
--- a/Engine/Plugins/Runtime/RigVM/Source/RigVM/Private/RigVMFunctions/Animation/RigVMFunction_TimeConversion.cpp
+++ b/Engine/Plugins/Runtime/RigVM/Source/RigVM/Private/RigVMFunctions/Animation/RigVMFunction_TimeConversion.cpp
@@ -6,9 +6,9 @@
 
 FRigVMFunction_FramesToSeconds_Execute()
 {
-    if(ExecuteContext.GetFramesPerSecond() > SMALL_NUMBER)
+    if(ExecuteContext.GetFramesPerSecond() > UE_SMALL_NUMBER)
     {
-		Seconds = Frames / ExecuteContext.GetFramesPerSecond();
+		Seconds = Frames / ExecuteContext.GetFramesPerSecond<float>();
 	}
 	else
 	{
@@ -18,9 +18,9 @@ FRigVMFunction_FramesToSeconds_Execute()
 
 FRigVMFunction_SecondsToFrames_Execute()
 {
-    if(ExecuteContext.GetFramesPerSecond() > SMALL_NUMBER)
+    if(ExecuteContext.GetFramesPerSecond() > UE_SMALL_NUMBER)
     {
-		Frames = Seconds * ExecuteContext.GetFramesPerSecond();
+		Frames = Seconds * ExecuteContext.GetFramesPerSecond<float>();
 	}
 	else
 	{
diff --git a/Engine/Plugins/Runtime/RigVM/Source/RigVM/Private/RigVMFunctions/Debug/RigVMFunction_VisualLog.cpp b/Engine/Plugins/Runtime/RigVM/Source/RigVM/Private/RigVMFunctions/Debug/RigVMFunction_VisualLog.cpp
--- a/Engine/Plugins/Runtime/RigVM/Source/RigVM/Private/RigVMFunctions/Debug/RigVMFunction_VisualLog.cpp
+++ b/Engine/Plugins/Runtime/RigVM/Source/RigVM/Private/RigVMFunctions/Debug/RigVMFunction_VisualLog.cpp
@@ -67,21 +67,21 @@ FRigVMFunction_VisualLogOrientedBox_Execute()
 FRigVMFunction_VisualLogArrow_Execute()
 {
 #if ENABLE_VISUAL_LOG
-	FVisualLogger::ArrowLineLogf(ExecuteContext.GetOwningObject(), Category, UE::RigVM::Private::DefaultVerbosity, SegmentStart, SegmentEnd, ObjectColor.ToFColor(true), ArrowHeadSize, TEXT("%s"), *Text);
+	FVisualLogger::ArrowLineLogf(ExecuteContext.GetOwningObject(), Category, UE::RigVM::Private::DefaultVerbosity, SegmentStart, SegmentEnd, ObjectColor.ToFColor(true), static_cast<uint16>(ArrowHeadSize), TEXT("%s"), *Text);
 #endif
 }
 
 FRigVMFunction_VisualLogCircle_Execute()
 {
 #if ENABLE_VISUAL_LOG
-	FVisualLogger::DiscLogf(ExecuteContext.GetOwningObject(), Category, UE::RigVM::Private::DefaultVerbosity, Center, UpAxis, Radius, ObjectColor.ToFColor(true), (uint16)Thickness, bWireframe, TEXT("%s"), *Text);
+	FVisualLogger::DiscLogf(ExecuteContext.GetOwningObject(), Category, UE::RigVM::Private::DefaultVerbosity, Center, UpAxis, Radius, ObjectColor.ToFColor(true), static_cast<uint16>(Thickness), bWireframe, TEXT("%s"), *Text);
 #endif
 }
 
 FRigVMFunction_VisualLogSegment_Execute()
 {
 #if ENABLE_VISUAL_LOG
-	FVisualLogger::SegmentLogf(ExecuteContext.GetOwningObject(), Category, UE::RigVM::Private::DefaultVerbosity, SegmentStart, SegmentEnd, ObjectColor.ToFColor(true), (uint16)Thickness, TEXT("%s"), *Text);
+	FVisualLogger::SegmentLogf(ExecuteContext.GetOwningObject(), Category, UE::RigVM::Private::DefaultVerbosity, SegmentStart, SegmentEnd, ObjectColor.ToFColor(true), static_cast<uint16>(Thickness), TEXT("%s"), *Text);
 #endif
 }
 
diff --git a/Engine/Plugins/Runtime/RigVM/Source/RigVM/Private/RigVMFunctions/Math/RigVMFunction_MathBool.cpp b/Engine/Plugins/Runtime/RigVM/Source/RigVM/Private/RigVMFunctions/Math/RigVMFunction_MathBool.cpp
--- a/Engine/Plugins/Runtime/RigVM/Source/RigVM/Private/RigVMFunctions/Math/RigVMFunction_MathBool.cpp
+++ b/Engine/Plugins/Runtime/RigVM/Source/RigVM/Private/RigVMFunctions/Math/RigVMFunction_MathBool.cpp
@@ -91,7 +91,7 @@ FRigVMFunction_MathBoolToggled_Execute()
 
 FRigVMFunction_MathBoolFlipFlop_Execute()
 {
-	TimeLeft = TimeLeft - ExecuteContext.GetDeltaTime();
+	TimeLeft = TimeLeft - ExecuteContext.GetDeltaTime<float>();
 	if (TimeLeft <= 0.f)
 	{
 		LastValue = !LastValue;
@@ -110,7 +110,7 @@ FRigVMFunction_MathBoolOnce_Execute()
 
 	Result = LastValue;
 
-	TimeLeft -= ExecuteContext.GetDeltaTime();
+	TimeLeft -= ExecuteContext.GetDeltaTime<float>();
 	if(TimeLeft <= 0.f)
 	{
 		LastValue = false;
diff --git a/Engine/Plugins/Runtime/RigVM/Source/RigVM/Private/RigVMFunctions/Math/RigVMFunction_MathBox.cpp b/Engine/Plugins/Runtime/RigVM/Source/RigVM/Private/RigVMFunctions/Math/RigVMFunction_MathBox.cpp
--- a/Engine/Plugins/Runtime/RigVM/Source/RigVM/Private/RigVMFunctions/Math/RigVMFunction_MathBox.cpp
+++ b/Engine/Plugins/Runtime/RigVM/Source/RigVM/Private/RigVMFunctions/Math/RigVMFunction_MathBox.cpp
@@ -101,10 +101,10 @@ FRigVMFunction_MathBoxGetDistance_Execute()
 		return;
 	}
 
-	Distance = Box.ComputeSquaredDistanceToPoint(Position);
+	Distance = static_cast<float>(Box.ComputeSquaredDistanceToPoint(Position));
 	if(!Square && Distance >= 0)
 	{
-		Distance = sqrt(Distance);
+		Distance = FMath::Sqrt(Distance);
 	}
 	Valid = true;
 }
@@ -126,5 +126,5 @@ FRigVMFunction_MathBoxGetVolume_Execute()
 		Volume = 0.f;
 		return;
 	}
-	Volume = Box.GetVolume();
+	Volume = static_cast<float>(Box.GetVolume());
 }
diff --git a/Engine/Plugins/Runtime/RigVM/Source/RigVM/Private/RigVMFunctions/Math/RigVMFunction_MathDouble.cpp b/Engine/Plugins/Runtime/RigVM/Source/RigVM/Private/RigVMFunctions/Math/RigVMFunction_MathDouble.cpp
--- a/Engine/Plugins/Runtime/RigVM/Source/RigVM/Private/RigVMFunctions/Math/RigVMFunction_MathDouble.cpp
+++ b/Engine/Plugins/Runtime/RigVM/Source/RigVM/Private/RigVMFunctions/Math/RigVMFunction_MathDouble.cpp
@@ -105,29 +105,29 @@ FRigVMFunction_MathDoubleAbs_Execute()
 FRigVMFunction_MathDoubleFloor_Execute()
 {
 	Result = FMath::FloorToDouble(Value);
-	Int = (int32)Result;
+	Int = static_cast<int32>(Result);
 }
 
 FRigVMFunction_MathDoubleCeil_Execute()
 {
 	Result = FMath::CeilToDouble(Value);
-	Int = (int32)Result;
+	Int = static_cast<int32>(Result);
 }
 
 FRigVMFunction_MathDoubleRound_Execute()
 {
 	Result = FMath::RoundToDouble(Value);
-	Int = (int32)Result;
+	Int = static_cast<int32>(Result);
 }
 
 FRigVMFunction_MathDoubleToInt_Execute()
 {
-	Result = FMath::FloorToInt(Value);
+	Result = static_cast<int32>(FMath::FloorToInt(Value));
 }
 
 FRigVMFunction_MathDoubleToFloat_Execute()
 {
-	Result = (float)Value;
+	Result = static_cast<float>(Value);
 }
 
 FRigVMFunction_MathDoubleSign_Execute()
diff --git a/Engine/Plugins/Runtime/RigVM/Source/RigVM/Private/RigVMFunctions/Math/RigVMFunction_MathFloat.cpp b/Engine/Plugins/Runtime/RigVM/Source/RigVM/Private/RigVMFunctions/Math/RigVMFunction_MathFloat.cpp
--- a/Engine/Plugins/Runtime/RigVM/Source/RigVM/Private/RigVMFunctions/Math/RigVMFunction_MathFloat.cpp
+++ b/Engine/Plugins/Runtime/RigVM/Source/RigVM/Private/RigVMFunctions/Math/RigVMFunction_MathFloat.cpp
@@ -327,7 +327,7 @@ FRigVMFunction_MathFloatArrayAverage_Execute()
 		{
 			Average += Value;
 		}
-		Average = Average / Array.Num();
+		Average = Average / static_cast<float>(Array.Num());
 	}
 	else
 	{
diff --git a/Engine/Plugins/Runtime/RigVM/Source/RigVM/Private/RigVMFunctions/Math/RigVMFunction_MathInt.cpp b/Engine/Plugins/Runtime/RigVM/Source/RigVM/Private/RigVMFunctions/Math/RigVMFunction_MathInt.cpp
--- a/Engine/Plugins/Runtime/RigVM/Source/RigVM/Private/RigVMFunctions/Math/RigVMFunction_MathInt.cpp
+++ b/Engine/Plugins/Runtime/RigVM/Source/RigVM/Private/RigVMFunctions/Math/RigVMFunction_MathInt.cpp
@@ -4,6 +4,8 @@
 #include "RigVMFunctions/RigVMDispatch_Core.h"
 #include "RigVMCore/RigVMStructTest.h"
 
+#include "limits"
+
 #include UE_INLINE_GENERATED_CPP_BY_NAME(RigVMFunction_MathInt)
 
 FRigVMFunction_MathIntMake_Execute()
@@ -59,8 +61,39 @@ FRigVMFunction_MathIntMax_Execute()
 
 FRigVMFunction_MathIntPow_Execute()
 {
-	Result = FMath::Pow(static_cast<float>(A), static_cast<float>(B));
-}
+	if (B >= 0)
+	{
+		// A quick way to compute exponentiate an integer that doesn't require powf()
+		Result = 1;
+		int32 P = A;
+		uint32 N = static_cast<uint32>(B);
+		for (;;)
+		{
+			if (N & 1)
+			{
+				// If the lowest bit of n is set, multiply the current exponentiated value into the result.
+				Result = P * Result;
+			}
+			N >>= 1;
+			if (N == 0)
+			{
+				break;
+			}
+			P = P * P;
+		}
+	}
+	else if (FMath::Abs(A) > 1)
+	{
+		// Negative integer powers of values greater than 1 in magnitude will always be a fraction, round to zero.
+		Result = 0;
+	}
+	else
+	{
+		// Negative powers of 1 are always 1. Negative powers of 0 are inf.
+		Result = A == 0 ? std::numeric_limits<int32>::max() : 1;
+	}
+}	
+
 
 FRigVMFunction_MathIntNegate_Execute()
 {
@@ -74,12 +107,12 @@ FRigVMFunction_MathIntAbs_Execute()
 
 FRigVMFunction_MathIntToFloat_Execute()
 {
-	Result = (float)Value;
+	Result = static_cast<float>(Value);
 }
 
 FRigVMFunction_MathIntToDouble_Execute()
 {
-	Result = (double)Value;
+	Result = static_cast<double>(Value);
 }
 
 FRigVMFunction_MathIntSign_Execute()
diff --git a/Engine/Plugins/Runtime/RigVM/Source/RigVM/Private/RigVMFunctions/Math/RigVMFunction_MathQuaternion.cpp b/Engine/Plugins/Runtime/RigVM/Source/RigVM/Private/RigVMFunctions/Math/RigVMFunction_MathQuaternion.cpp
--- a/Engine/Plugins/Runtime/RigVM/Source/RigVM/Private/RigVMFunctions/Math/RigVMFunction_MathQuaternion.cpp
+++ b/Engine/Plugins/Runtime/RigVM/Source/RigVM/Private/RigVMFunctions/Math/RigVMFunction_MathQuaternion.cpp
@@ -179,7 +179,7 @@ FRigVMStructUpgradeInfo FRigVMFunction_MathQuaternionSelectBool::GetUpgradeInfo(
 
 FRigVMFunction_MathQuaternionDot_Execute()
 {
-	Result = A | B;
+	Result = static_cast<float>(A | B);
 }
 
 FRigVMFunction_MathQuaternionUnit_Execute()
diff --git a/Engine/Plugins/Runtime/RigVM/Source/RigVM/Private/RigVMFunctions/Math/RigVMFunction_MathQuaternionTests.cpp b/Engine/Plugins/Runtime/RigVM/Source/RigVM/Private/RigVMFunctions/Math/RigVMFunction_MathQuaternionTests.cpp
--- a/Engine/Plugins/Runtime/RigVM/Source/RigVM/Private/RigVMFunctions/Math/RigVMFunction_MathQuaternionTests.cpp
+++ b/Engine/Plugins/Runtime/RigVM/Source/RigVM/Private/RigVMFunctions/Math/RigVMFunction_MathQuaternionTests.cpp
@@ -13,31 +13,31 @@ namespace FRigVMFunction_MathQuatTest_Utils
 		return FMath::IsNearlyEqual(A, B, Tolerance);
 	}
 
-	bool IsNearlyEqual(const FVector& A, const FVector& B, float Tolerance = 0.001f)
+	bool IsNearlyEqual(const FVector& A, const FVector& B, double Tolerance = 0.001)
 	{
-		return IsNearlyEqual(A.X, B.X, Tolerance) &&
-			IsNearlyEqual(A.Y, B.Y, Tolerance) &&
-			IsNearlyEqual(A.Z, B.Z, Tolerance);
+		return FMath::IsNearlyEqual(A.X, B.X, Tolerance) &&
+			   FMath::IsNearlyEqual(A.Y, B.Y, Tolerance) &&
+			   FMath::IsNearlyEqual(A.Z, B.Z, Tolerance);
 	}
 
-	bool IsNearlyEqual(const FQuat& A, const FQuat& B, float Tolerance = 0.001f)
+	bool IsNearlyEqual(const FQuat& A, const FQuat& B, double Tolerance = 0.001)
 	{
-		return IsNearlyEqual(A.X, B.X, Tolerance) &&
-			IsNearlyEqual(A.Y, B.Y, Tolerance) &&
-			IsNearlyEqual(A.Z, B.Z, Tolerance) &&
-			IsNearlyEqual(A.W, B.W, Tolerance);
+		return FMath::IsNearlyEqual(A.X, B.X, Tolerance) &&
+			   FMath::IsNearlyEqual(A.Y, B.Y, Tolerance) &&
+			   FMath::IsNearlyEqual(A.Z, B.Z, Tolerance) &&
+			   FMath::IsNearlyEqual(A.W, B.W, Tolerance);
 	}
 
-	bool IsNearlyEqual(const FRotator& A, const FRotator& B, float Tolerance = 0.001f)
+	bool IsNearlyEqual(const FRotator& A, const FRotator& B, double Tolerance = 0.001)
 	{
-		return IsNearlyEqual(A.Yaw, B.Yaw, Tolerance) &&
-			IsNearlyEqual(A.Roll, B.Roll, Tolerance) &&
-			IsNearlyEqual(A.Pitch, B.Pitch, Tolerance);
+		return FMath::IsNearlyEqual(A.Yaw, B.Yaw, Tolerance) &&
+			   FMath::IsNearlyEqual(A.Roll, B.Roll, Tolerance) &&
+			   FMath::IsNearlyEqual(A.Pitch, B.Pitch, Tolerance);
 	}
 
-	bool IsNearlyEqual(const FQuat& A, const FVector& B, float Tolerance = 0.001f)
+	bool IsNearlyEqual(const FQuat& A, const FVector& B, double Tolerance = 0.001)
 	{
-		FVector Euler = A.Euler();
+		const FVector Euler = A.Euler();
 		return IsNearlyEqual(Euler, B, Tolerance);
 	}
 };
@@ -179,7 +179,7 @@ IMPLEMENT_RIGVMSTRUCT_AUTOMATION_TEST(FRigVMFunction_MathQuaternionDot)
 	Unit.A = FQuat(1.f, 2.f, 3.f, 4.f);
 	Unit.B = FQuat(5.f, 6.f, 7.f, 8.f);
 	Execute();
-	AddErrorIfFalse(FRigVMFunction_MathQuatTest_Utils::IsNearlyEqual(Unit.Result, Unit.A | Unit.B), TEXT("unexpected result"));
+	AddErrorIfFalse(FRigVMFunction_MathQuatTest_Utils::IsNearlyEqual(Unit.Result, static_cast<float>(Unit.A | Unit.B)), TEXT("unexpected result"));
 	return true;
 }
 
diff --git a/Engine/Plugins/Runtime/RigVM/Source/RigVM/Private/RigVMFunctions/Math/RigVMFunction_MathRBFInterpolate.cpp b/Engine/Plugins/Runtime/RigVM/Source/RigVM/Private/RigVMFunctions/Math/RigVMFunction_MathRBFInterpolate.cpp
--- a/Engine/Plugins/Runtime/RigVM/Source/RigVM/Private/RigVMFunctions/Math/RigVMFunction_MathRBFInterpolate.cpp
+++ b/Engine/Plugins/Runtime/RigVM/Source/RigVM/Private/RigVMFunctions/Math/RigVMFunction_MathRBFInterpolate.cpp
@@ -74,7 +74,7 @@ struct FRigVMFunction_MathRBFQuatWeightFunctor
 
 	float operator()(const FQuat &A, const FQuat &B) const
 	{
-		float Distance = 0.0f;
+		double Distance = 0.0;
 		switch (DistanceType)
 		{
 		case ERBFQuatDistanceType::Euclidean:
@@ -91,7 +91,7 @@ struct FRigVMFunction_MathRBFQuatWeightFunctor
 			break;
 		}
 
-		return SmoothingKernel(Distance);
+		return SmoothingKernel(static_cast<float>(Distance));
 	}
 
 	ERBFQuatDistanceType DistanceType;
@@ -112,7 +112,7 @@ struct FRigVMFunction_MathRBFVectorWeightFunctor
 
 	float operator()(const FVector &A, const FVector &B) const
 	{
-		float Distance = 0.0f;
+		double Distance = 0.0;
 		switch (DistanceType)
 		{
 		case ERBFVectorDistanceType::Euclidean:
@@ -126,7 +126,7 @@ struct FRigVMFunction_MathRBFVectorWeightFunctor
 			break;
 		}
 
-		return SmoothingKernel(Distance);
+		return SmoothingKernel(static_cast<float>(Distance));
 	}
 
 	ERBFVectorDistanceType DistanceType;
@@ -138,7 +138,16 @@ static uint64 HashValue(float v)
 	// To ensure both -0 and +0 hash to the same value. Yes, that's a thing.
 	if (v == 0.0f)
 		return 0ULL;
-	return *(uint32*)&v;
+	const uint64 x = *reinterpret_cast<uint32*>(&v);
+	return (x << 32ull) | x;
+}
+
+static uint64 HashValue(double v)
+{
+	// To ensure both -0 and +0 hash to the same value. Yes, that's a thing.
+	if (v == 0.0f)
+		return 0ULL;
+	return *reinterpret_cast<uint64*>(&v);
 }
 
 static uint64 HashValues(std::initializer_list<uint64> Values, uint64 Seed)
diff --git a/Engine/Plugins/Runtime/RigVM/Source/RigVM/Private/RigVMFunctions/Math/RigVMFunction_MathRay.cpp b/Engine/Plugins/Runtime/RigVM/Source/RigVM/Private/RigVMFunctions/Math/RigVMFunction_MathRay.cpp
--- a/Engine/Plugins/Runtime/RigVM/Source/RigVM/Private/RigVMFunctions/Math/RigVMFunction_MathRay.cpp
+++ b/Engine/Plugins/Runtime/RigVM/Source/RigVM/Private/RigVMFunctions/Math/RigVMFunction_MathRay.cpp
@@ -23,10 +23,10 @@ FRigVMFunction_MathRayIntersectRay_Execute()
 	Distance = 0;
 	if(!ClosestA.Equals(ClosestB))
 	{
-		Distance = (ClosestA - ClosestB).Size();
+		Distance = static_cast<float>((ClosestA - ClosestB).Size());
 	}
-	RatioA = A.GetParameter(ClosestA);	
-	RatioB = B.GetParameter(ClosestB);	
+	RatioA = static_cast<float>(A.GetParameter(ClosestA));	
+	RatioB = static_cast<float>(B.GetParameter(ClosestB));	
 }
 
 FRigVMFunction_MathRayIntersectPlane_Execute()
@@ -39,12 +39,12 @@ FRigVMFunction_MathRayIntersectPlane_Execute()
 		return;
 	}
 	
-	Ratio = FMath::RayPlaneIntersectionParam(Ray.Origin, Ray.Direction, FPlane(PlanePoint, NormalizedPlaneNormal));
+	Ratio = static_cast<float>(FMath::RayPlaneIntersectionParam(Ray.Origin, Ray.Direction, FPlane(PlanePoint, NormalizedPlaneNormal)));
 	Result = Ray.PointAt(Ratio);
 	Distance = 0;
 	if(!Ray.Origin.Equals(Result))
 	{
-		Distance = (Ray.Origin - Result).Size();
+		Distance = static_cast<float>((Ray.Origin - Result).Size());
 	}
 }
 
diff --git a/Engine/Plugins/Runtime/RigVM/Source/RigVM/Private/RigVMFunctions/Math/RigVMFunction_MathTransformTests.cpp b/Engine/Plugins/Runtime/RigVM/Source/RigVM/Private/RigVMFunctions/Math/RigVMFunction_MathTransformTests.cpp
--- a/Engine/Plugins/Runtime/RigVM/Source/RigVM/Private/RigVMFunctions/Math/RigVMFunction_MathTransformTests.cpp
+++ b/Engine/Plugins/Runtime/RigVM/Source/RigVM/Private/RigVMFunctions/Math/RigVMFunction_MathTransformTests.cpp
@@ -11,29 +11,29 @@ namespace FRigVMFunction_MathTransformTest_Utils
 		return FMath::IsNearlyEqual(A, B, Tolerance);
 	}
 
-	bool IsNearlyEqual(const FVector& A, const FVector& B, float Tolerance = 0.001f)
+	bool IsNearlyEqual(const FVector& A, const FVector& B, double Tolerance = 0.001)
 	{
-		return IsNearlyEqual(A.X, B.X, Tolerance) &&
-			IsNearlyEqual(A.Y, B.Y, Tolerance) &&
-			IsNearlyEqual(A.Z, B.Z, Tolerance);
+		return FMath::IsNearlyEqual(A.X, B.X, Tolerance) &&
+			FMath::IsNearlyEqual(A.Y, B.Y, Tolerance) &&
+			FMath::IsNearlyEqual(A.Z, B.Z, Tolerance);
 	}
 
-	bool IsNearlyEqual(const FQuat& A, const FQuat& B, float Tolerance = 0.001f)
+	bool IsNearlyEqual(const FQuat& A, const FQuat& B, double Tolerance = 0.001)
 	{
-		return IsNearlyEqual(A.X, B.X, Tolerance) &&
-			IsNearlyEqual(A.Y, B.Y, Tolerance) &&
-			IsNearlyEqual(A.Z, B.Z, Tolerance) &&
-			IsNearlyEqual(A.W, B.W, Tolerance);
+		return FMath::IsNearlyEqual(A.X, B.X, Tolerance) &&
+			FMath::IsNearlyEqual(A.Y, B.Y, Tolerance) &&
+			FMath::IsNearlyEqual(A.Z, B.Z, Tolerance) &&
+			FMath::IsNearlyEqual(A.W, B.W, Tolerance);
 	}
 
-	bool IsNearlyEqual(const FRotator& A, const FRotator& B, float Tolerance = 0.001f)
+	bool IsNearlyEqual(const FRotator& A, const FRotator& B, double Tolerance = 0.001)
 	{
-		return IsNearlyEqual(A.Yaw, B.Yaw, Tolerance) &&
-			IsNearlyEqual(A.Roll, B.Roll, Tolerance) &&
-			IsNearlyEqual(A.Pitch, B.Pitch, Tolerance);
+		return FMath::IsNearlyEqual(A.Yaw, B.Yaw, Tolerance) &&
+			FMath::IsNearlyEqual(A.Roll, B.Roll, Tolerance) &&
+			FMath::IsNearlyEqual(A.Pitch, B.Pitch, Tolerance);
 	}
 
-	bool IsNearlyEqual(const FQuat& A, const FVector& B, float Tolerance = 0.001f)
+	bool IsNearlyEqual(const FQuat& A, const FVector& B, double Tolerance = 0.001)
 	{
 		FVector Euler = A.Euler();
 		return IsNearlyEqual(Euler, B, Tolerance);
diff --git a/Engine/Plugins/Runtime/RigVM/Source/RigVM/Private/RigVMFunctions/Math/RigVMFunction_MathVector.cpp b/Engine/Plugins/Runtime/RigVM/Source/RigVM/Private/RigVMFunctions/Math/RigVMFunction_MathVector.cpp
--- a/Engine/Plugins/Runtime/RigVM/Source/RigVM/Private/RigVMFunctions/Math/RigVMFunction_MathVector.cpp
+++ b/Engine/Plugins/Runtime/RigVM/Source/RigVM/Private/RigVMFunctions/Math/RigVMFunction_MathVector.cpp
@@ -20,7 +20,7 @@ FRigVMFunction_MathVectorFromFloat_Execute()
 
 FRigVMFunction_MathVectorFromDouble_Execute()
 {
-	Result = FVector((float)Value, (float)Value, (float)Value);
+	Result = FVector(Value, Value, Value);
 }
 
 FRigVMFunction_MathVectorAdd_Execute()
@@ -92,16 +92,16 @@ FRigVMFunction_MathVectorMod_Execute()
 
 FRigVMFunction_MathVectorMin_Execute()
 {
-	Result.X = FMath::Min<float>(A.X, B.X);
-	Result.Y = FMath::Min<float>(A.Y, B.Y);
-	Result.Z = FMath::Min<float>(A.Z, B.Z);
+	Result.X = FMath::Min<double>(A.X, B.X);
+	Result.Y = FMath::Min<double>(A.Y, B.Y);
+	Result.Z = FMath::Min<double>(A.Z, B.Z);
 }
 
 FRigVMFunction_MathVectorMax_Execute()
 {
-	Result.X = FMath::Max<float>(A.X, B.X);
-	Result.Y = FMath::Max<float>(A.Y, B.Y);
-	Result.Z = FMath::Max<float>(A.Z, B.Z);
+	Result.X = FMath::Max<double>(A.X, B.X);
+	Result.Y = FMath::Max<double>(A.Y, B.Y);
+	Result.Z = FMath::Max<double>(A.Z, B.Z);
 }
 
 FRigVMFunction_MathVectorNegate_Execute()
@@ -144,9 +144,9 @@ FRigVMFunction_MathVectorSign_Execute()
 
 FRigVMFunction_MathVectorClamp_Execute()
 {
-	Result.X = FMath::Clamp<float>(Value.X, Minimum.X, Maximum.X);
-	Result.Y = FMath::Clamp<float>(Value.Y, Minimum.Y, Maximum.Y);
-	Result.Z = FMath::Clamp<float>(Value.Z, Minimum.Z, Maximum.Z);
+	Result.X = FMath::Clamp<double>(Value.X, Minimum.X, Maximum.X);
+	Result.Y = FMath::Clamp<double>(Value.Y, Minimum.Y, Maximum.Y);
+	Result.Z = FMath::Clamp<double>(Value.Z, Minimum.Z, Maximum.Z);
 }
 
 FRigVMFunction_MathVectorLerp_Execute()
@@ -156,7 +156,7 @@ FRigVMFunction_MathVectorLerp_Execute()
 
 FRigVMFunction_MathVectorRemap_Execute()
 {
-	FVector Ratio(0.f, 0.f, 0.f);
+	FVector Ratio = FVector::ZeroVector;
 	if (FMath::IsNearlyEqual(SourceMinimum.X, SourceMaximum.X) || FMath::IsNearlyEqual(SourceMinimum.Y, SourceMaximum.Y) || FMath::IsNearlyEqual(SourceMinimum.Z, SourceMaximum.Z))
 	{
 		if (FMath::IsNearlyEqual(SourceMinimum.X, SourceMaximum.X))
@@ -178,9 +178,9 @@ FRigVMFunction_MathVectorRemap_Execute()
 	}
 	if (bClamp)
 	{
-		Ratio.X = FMath::Clamp<float>(Ratio.X, 0.f, 1.f);
-		Ratio.Y = FMath::Clamp<float>(Ratio.Y, 0.f, 1.f);
-		Ratio.Z = FMath::Clamp<float>(Ratio.Z, 0.f, 1.f);
+		Ratio.X = FMath::Clamp<double>(Ratio.X, 0, 1);
+		Ratio.Y = FMath::Clamp<double>(Ratio.Y, 0, 1);
+		Ratio.Z = FMath::Clamp<double>(Ratio.Z, 0, 1);
 	}
 	Result = FMath::Lerp<FVector>(TargetMinimum, TargetMaximum, Ratio);
 }
@@ -246,17 +246,17 @@ FRigVMFunction_MathVectorRad_Execute()
 
 FRigVMFunction_MathVectorLengthSquared_Execute()
 {
-	Result = Value.SizeSquared();
+	Result = static_cast<float>(Value.SizeSquared());
 }
 
 FRigVMFunction_MathVectorLength_Execute()
 {
-	Result = Value.Size();
+	Result = static_cast<float>(Value.Size());
 }
 
 FRigVMFunction_MathVectorDistance_Execute()
 {
-	Result = FVector::Distance(A, B);
+	Result = static_cast<float>(FVector::Distance(A, B));
 }
 
 FRigVMFunction_MathVectorCross_Execute()
@@ -266,7 +266,7 @@ FRigVMFunction_MathVectorCross_Execute()
 
 FRigVMFunction_MathVectorDot_Execute()
 {
-	Result = A | B;
+	Result = static_cast<float>(A | B);
 }
 
 FRigVMFunction_MathVectorUnit_Execute()
@@ -299,7 +299,7 @@ FRigVMFunction_MathVectorClampLength_Execute()
 		Result = FVector::ZeroVector;
 		return;
 	}
-	float Length = Value.Size();
+	float Length = static_cast<float>(Value.Size());
 	Result = Value * FMath::Clamp<float>(Length, MinimumLength, MaximumLength) / Length;
 }
 
@@ -329,7 +329,7 @@ FRigVMFunction_MathVectorAngle_Execute()
 		Result = 0.f;
 		return;
 	}
-	Result = FRigVMMathLibrary::FindQuatBetweenVectors(A, B).GetAngle();
+	Result = static_cast<float>(FRigVMMathLibrary::FindQuatBetweenVectors(A, B).GetAngle());
 }
 
 FRigVMFunction_MathVectorParallel_Execute()
@@ -424,7 +424,7 @@ void FRigVMFunction_DrawPlane(FRigVMDrawInterface* DrawInterface, EAxis::Type Ax
 	}
 
 	const FVector PointOnPlane = FVector::PointPlaneProject(Result, FPlane(CurrentSpace.GetTranslation(), Normal));
-	const float PlaneScale = FMath::Max<float>(10.f, FVector::Distance(CurrentSpace.GetTranslation(), PointOnPlane)) * 0.525f;
+	const double PlaneScale = FMath::Max<double>(10., FVector::Distance(CurrentSpace.GetTranslation(), PointOnPlane)) * 0.525;
 
 	FTransform PlaneRotation = FTransform::Identity;
 	PlaneRotation.SetRotation(FMatrix(A, Normal.Cross(A), Normal, FVector::ZeroVector).ToQuat());
@@ -704,7 +704,7 @@ FRigVMFunction_MathIntersectPlane_Execute()
 	FPlane Plane(PlanePoint, NormalizedPlaneNormal);
 
 	Result = FMath::RayPlaneIntersection(Start, Direction, Plane);
-	Distance = (Start - Result).Size();
+	Distance = static_cast<float>((Start - Result).Size());
 }
 
 FRigVMFunction_MathDistanceToPlane_Execute()
@@ -716,7 +716,7 @@ FRigVMFunction_MathDistanceToPlane_Execute()
 		FPlane Plane(PlanePoint, UnitVectorNormal);
 
 		ClosestPointOnPlane = FVector::PointPlaneProject(Point, Plane);
-		SignedDistance = FVector::PointPlaneDist(Point, PlanePoint, UnitVectorNormal);
+		SignedDistance = static_cast<float>(FVector::PointPlaneDist(Point, PlanePoint, UnitVectorNormal));
 	}
 	else
 	{
diff --git a/Engine/Plugins/Runtime/RigVM/Source/RigVM/Private/RigVMFunctions/Math/RigVMFunction_Noise.cpp b/Engine/Plugins/Runtime/RigVM/Source/RigVM/Private/RigVMFunctions/Math/RigVMFunction_Noise.cpp
--- a/Engine/Plugins/Runtime/RigVM/Source/RigVM/Private/RigVMFunctions/Math/RigVMFunction_Noise.cpp
+++ b/Engine/Plugins/Runtime/RigVM/Source/RigVM/Private/RigVMFunctions/Math/RigVMFunction_Noise.cpp
@@ -9,14 +9,14 @@ FRigVMFunction_NoiseFloat_Execute()
 
 	const float Noise = FMath::PerlinNoise1D(Value * Frequency + Time) + 0.5f;
 	Result = FMath::Lerp<float>(Minimum, Maximum, Noise);
-	Time = Time + Speed * ExecuteContext.GetDeltaTime();
+	Time = Time + Speed * ExecuteContext.GetDeltaTime<float>();
 }
 
 FRigVMFunction_NoiseDouble_Execute()
 {
-	const double Noise = double(FMath::PerlinNoise1D(Value * Frequency + Time)) + 0.5;
+	const double Noise = FMath::PerlinNoise1D(static_cast<float>(Value * Frequency + Time)) + 0.5;
 	Result = FMath::Lerp<double>(Minimum, Maximum, Noise);
-	Time = Time + Speed * ExecuteContext.GetDeltaTime();
+	Time = Time + Speed * ExecuteContext.GetDeltaTime<double>();
 }
 
 FRigVMFunction_NoiseVector_Execute()
@@ -40,12 +40,12 @@ FRigVMStructUpgradeInfo FRigVMFunction_NoiseVector::GetUpgradeInfo() const
 
 FRigVMFunction_NoiseVector2_Execute()
 {
-	const double NoiseX = double(FMath::PerlinNoise1D(Value.X * Frequency.X + Time.X)) + 0.5;
-	const double NoiseY = double(FMath::PerlinNoise1D(Value.Y * Frequency.Y + Time.Y)) + 0.5;
-	const double NoiseZ = double(FMath::PerlinNoise1D(Value.Z * Frequency.Z + Time.Z)) + 0.5;
+	const double NoiseX = FMath::PerlinNoise1D(static_cast<float>(Value.X * Frequency.X + Time.X)) + 0.5;
+	const double NoiseY = FMath::PerlinNoise1D(static_cast<float>(Value.Y * Frequency.Y + Time.Y)) + 0.5;
+	const double NoiseZ = FMath::PerlinNoise1D(static_cast<float>(Value.Z * Frequency.Z + Time.Z)) + 0.5;
 	Result.X = FMath::Lerp<double>(Minimum, Maximum, NoiseX);
 	Result.Y = FMath::Lerp<double>(Minimum, Maximum, NoiseY);
 	Result.Z = FMath::Lerp<double>(Minimum, Maximum, NoiseZ);
-	Time = Time + Speed * ExecuteContext.GetDeltaTime();
+	Time = Time + Speed * ExecuteContext.GetDeltaTime<float>();
 }
 
diff --git a/Engine/Plugins/Runtime/RigVM/Source/RigVM/Private/RigVMFunctions/Math/RigVMFunction_Random.cpp b/Engine/Plugins/Runtime/RigVM/Source/RigVM/Private/RigVMFunctions/Math/RigVMFunction_Random.cpp
--- a/Engine/Plugins/Runtime/RigVM/Source/RigVM/Private/RigVMFunctions/Math/RigVMFunction_Random.cpp
+++ b/Engine/Plugins/Runtime/RigVM/Source/RigVM/Private/RigVMFunctions/Math/RigVMFunction_Random.cpp
@@ -24,7 +24,7 @@ FRigVMFunction_RandomFloat_Execute()
 		LastSeed = BaseSeed = Seed;
 	}
 
-	TimeLeft = TimeLeft - ExecuteContext.GetDeltaTime();
+	TimeLeft = TimeLeft - ExecuteContext.GetDeltaTime<float>();
 	if (TimeLeft > 0.f)
 	{
 		Result = LastResult;
@@ -45,7 +45,7 @@ FRigVMFunction_RandomVector_Execute()
 		LastSeed = BaseSeed = Seed;
 	}
 
-	TimeLeft = TimeLeft - ExecuteContext.GetDeltaTime();
+	TimeLeft = TimeLeft - ExecuteContext.GetDeltaTime<float>();
 	if (TimeLeft > 0.f)
 	{
 		Result = LastResult;
diff --git a/Engine/Plugins/Runtime/RigVM/Source/RigVM/Private/RigVMFunctions/Math/RigVMMathLibrary.cpp b/Engine/Plugins/Runtime/RigVM/Source/RigVM/Private/RigVMFunctions/Math/RigVMMathLibrary.cpp
--- a/Engine/Plugins/Runtime/RigVM/Source/RigVM/Private/RigVMFunctions/Math/RigVMMathLibrary.cpp
+++ b/Engine/Plugins/Runtime/RigVM/Source/RigVM/Private/RigVMFunctions/Math/RigVMMathLibrary.cpp
@@ -84,7 +84,7 @@ float FRigVMMathLibrary::AngleBetween(const FVector& A, const FVector& B)
 		return 0.f;
 	}
 
-	return FMath::Acos(FVector::DotProduct(A, B) / (A.Size() * B.Size()));
+	return static_cast<float>(FMath::Acos(FVector::DotProduct(A, B) / (A.Size() * B.Size())));
 }
 
 float FRigVMMathLibrary::EaseFloat(float Value, ERigVMAnimEasingType Type)
@@ -386,7 +386,7 @@ FVector FRigVMMathLibrary::ClampSpatially(const FVector& Value, EAxis::Type Axis
 				}
 			}
 
-			const float DotOnNormal = Normal.Dot(Local);
+			const double DotOnNormal = Normal.Dot(Local);
 
 			// if we are on the cylinder part of the collision, fall through 
 			if(FMath::Abs(DotOnNormal) < HalfCylinderLength)
@@ -507,8 +507,8 @@ FVector FRigVMMathLibrary::GetEquivalentEulerAngle(const FVector& InEulerAngle,
 FVector& FRigVMMathLibrary::ChooseBetterEulerAngleForAxisFilter(const FVector& Base, FVector& A, FVector& B)
 {
 	// simply compare the sum of difference should be enough
-	float Diff1 = (A - Base).GetAbs().Dot(FVector::OneVector);
-	float Diff2 = (B - Base).GetAbs().Dot(FVector::OneVector);
+	const double Diff1 = (A - Base).GetAbs().Dot(FVector::OneVector);
+	const double Diff2 = (B - Base).GetAbs().Dot(FVector::OneVector);
 
 	return Diff1 < Diff2 ? A : B;
 }
diff --git a/Engine/Plugins/Runtime/RigVM/Source/RigVM/Private/RigVMFunctions/Simulation/RigVMFunction_Accumulate.cpp b/Engine/Plugins/Runtime/RigVM/Source/RigVM/Private/RigVMFunctions/Simulation/RigVMFunction_Accumulate.cpp
--- a/Engine/Plugins/Runtime/RigVM/Source/RigVM/Private/RigVMFunctions/Simulation/RigVMFunction_Accumulate.cpp
+++ b/Engine/Plugins/Runtime/RigVM/Source/RigVM/Private/RigVMFunctions/Simulation/RigVMFunction_Accumulate.cpp
@@ -14,7 +14,7 @@ FRigVMFunction_AccumulateFloatAdd_Execute()
 	
 	if(bIntegrateDeltaTime)
 	{
-		Result = AccumulatedValue = AccumulatedValue + Increment * ExecuteContext.GetDeltaTime();
+		Result = AccumulatedValue = AccumulatedValue + Increment * ExecuteContext.GetDeltaTime<float>();
 	}
 	else
 	{
@@ -128,7 +128,7 @@ FRigVMFunction_AccumulateFloatLerp_Execute()
 		bIsInitialized = true;
 	}
 
-	Result = AccumulatedValue = FMath::Lerp<float>(AccumulatedValue, TargetValue, FMath::Clamp<float>(bIntegrateDeltaTime ? ExecuteContext.GetDeltaTime() * Blend : Blend, 0.f, 1.f));
+	Result = AccumulatedValue = FMath::Lerp<float>(AccumulatedValue, TargetValue, FMath::Clamp<float>(bIntegrateDeltaTime ? ExecuteContext.GetDeltaTime<float>() * Blend : Blend, 0.f, 1.f));
 }
 
 FRigVMFunction_AccumulateVectorLerp_Execute()
@@ -139,7 +139,7 @@ FRigVMFunction_AccumulateVectorLerp_Execute()
 		bIsInitialized = true;
 	}
 
-	Result = AccumulatedValue = FMath::Lerp<FVector>(AccumulatedValue, TargetValue, FMath::Clamp<float>(bIntegrateDeltaTime ? ExecuteContext.GetDeltaTime() * Blend : Blend, 0.f, 1.f));
+	Result = AccumulatedValue = FMath::Lerp<FVector>(AccumulatedValue, TargetValue, FMath::Clamp<double>(bIntegrateDeltaTime ? ExecuteContext.GetDeltaTime() * Blend : Blend, 0., 1.));
 }
 
 FRigVMFunction_AccumulateQuatLerp_Execute()
@@ -150,7 +150,7 @@ FRigVMFunction_AccumulateQuatLerp_Execute()
 		bIsInitialized = true;
 	}
 
-	Result = AccumulatedValue = FQuat::Slerp(AccumulatedValue, TargetValue, FMath::Clamp<float>(bIntegrateDeltaTime ? ExecuteContext.GetDeltaTime() * Blend : Blend, 0.f, 1.f));
+	Result = AccumulatedValue = FQuat::Slerp(AccumulatedValue, TargetValue, FMath::Clamp<double>(bIntegrateDeltaTime ? ExecuteContext.GetDeltaTime() * Blend : Blend, 0., 1.));
 }
 
 FRigVMFunction_AccumulateTransformLerp_Execute()
@@ -161,7 +161,7 @@ FRigVMFunction_AccumulateTransformLerp_Execute()
 		bIsInitialized = true;
 	}
 
-	float B = FMath::Clamp<float>(bIntegrateDeltaTime ? ExecuteContext.GetDeltaTime() * Blend : Blend, 0.f, 1.f);
+	double B = FMath::Clamp<double>(bIntegrateDeltaTime ? ExecuteContext.GetDeltaTime() * Blend : Blend, 0., 1.);
 	AccumulatedValue.SetTranslation(FMath::Lerp<FVector>(AccumulatedValue.GetTranslation(), TargetValue.GetTranslation(), B));
 	AccumulatedValue.SetRotation(FQuat::Slerp(AccumulatedValue.GetRotation(), TargetValue.GetRotation(), B));
 	AccumulatedValue.SetScale3D(FMath::Lerp<FVector>(AccumulatedValue.GetScale3D(), TargetValue.GetScale3D(), B));
diff --git a/Engine/Plugins/Runtime/RigVM/Source/RigVM/Private/RigVMFunctions/Simulation/RigVMFunction_AlphaInterp.cpp b/Engine/Plugins/Runtime/RigVM/Source/RigVM/Private/RigVMFunctions/Simulation/RigVMFunction_AlphaInterp.cpp
--- a/Engine/Plugins/Runtime/RigVM/Source/RigVM/Private/RigVMFunctions/Simulation/RigVMFunction_AlphaInterp.cpp
+++ b/Engine/Plugins/Runtime/RigVM/Source/RigVM/Private/RigVMFunctions/Simulation/RigVMFunction_AlphaInterp.cpp
@@ -19,7 +19,7 @@ FRigVMFunction_AlphaInterp_Execute()
 	ScaleBiasClamp.InterpSpeedIncreasing = InterpSpeedIncreasing;
 	ScaleBiasClamp.InterpSpeedDecreasing = InterpSpeedDecreasing;
 
-	Result = ScaleBiasClamp.ApplyTo(Value, ExecuteContext.GetDeltaTime());
+	Result = ScaleBiasClamp.ApplyTo(Value, ExecuteContext.GetDeltaTime<float>());
 }
 
 FString FRigVMFunction_AlphaInterp::ProcessPinLabelForInjection(const FString& InLabel) const
@@ -60,9 +60,9 @@ FRigVMFunction_AlphaInterpVector_Execute()
 	ScaleBiasClamp.InterpSpeedIncreasing = InterpSpeedIncreasing;
 	ScaleBiasClamp.InterpSpeedDecreasing = InterpSpeedDecreasing;
 
-	Result.X = ScaleBiasClamp.ApplyTo(Value.X, ExecuteContext.GetDeltaTime());
-	Result.Y = ScaleBiasClamp.ApplyTo(Value.Y, ExecuteContext.GetDeltaTime());
-	Result.Z = ScaleBiasClamp.ApplyTo(Value.Z, ExecuteContext.GetDeltaTime());
+	Result.X = ScaleBiasClamp.ApplyTo(static_cast<float>(Value.X), ExecuteContext.GetDeltaTime<float>());
+	Result.Y = ScaleBiasClamp.ApplyTo(static_cast<float>(Value.Y), ExecuteContext.GetDeltaTime<float>());
+	Result.Z = ScaleBiasClamp.ApplyTo(static_cast<float>(Value.Z), ExecuteContext.GetDeltaTime<float>());
 }
 
 FString FRigVMFunction_AlphaInterpVector::ProcessPinLabelForInjection(const FString& InLabel) const
@@ -103,7 +103,7 @@ FRigVMFunction_AlphaInterpQuat_Execute()
 	ScaleBiasClamp.InterpSpeedIncreasing = InterpSpeedIncreasing;
 	ScaleBiasClamp.InterpSpeedDecreasing = InterpSpeedDecreasing;
 
-	const float T = ScaleBiasClamp.ApplyTo(1.f, ExecuteContext.GetDeltaTime());
+	const double T = ScaleBiasClamp.ApplyTo(1.f, ExecuteContext.GetDeltaTime<float>());
 	Result = FQuat::Slerp(FQuat::Identity, Value, T);
 }
 
diff --git a/Engine/Plugins/Runtime/RigVM/Source/RigVM/Private/RigVMFunctions/Simulation/RigVMFunction_TimeOffset.cpp b/Engine/Plugins/Runtime/RigVM/Source/RigVM/Private/RigVMFunctions/Simulation/RigVMFunction_TimeOffset.cpp
--- a/Engine/Plugins/Runtime/RigVM/Source/RigVM/Private/RigVMFunctions/Simulation/RigVMFunction_TimeOffset.cpp
+++ b/Engine/Plugins/Runtime/RigVM/Source/RigVM/Private/RigVMFunctions/Simulation/RigVMFunction_TimeOffset.cpp
@@ -79,35 +79,35 @@ FRigVMFunction_TimeOffsetFloat_Execute()
 		Result = FMath::Lerp<float>(A, B, T);
 	}
 
-	if (ExecuteContext.GetDeltaTime() > SMALL_NUMBER)
+	if (ExecuteContext.GetDeltaTime() > UE_SMALL_NUMBER)
 	{
 		if (UpperBound == 0)
 		{
 			LastInsertIndex = UpperBound;
 			Buffer[UpperBound] = Value;
-			DeltaTimes[UpperBound++] = ExecuteContext.GetDeltaTime();
+			DeltaTimes[UpperBound++] = ExecuteContext.GetDeltaTime<float>();
 		}
 		else
 		{
-			float SecondsPerEntry = TimeRange / float(MaxSize - 1);
-			if (DeltaTimes[LastInsertIndex] > SecondsPerEntry - ExecuteContext.GetDeltaTime() * 0.5f)
+			const float SecondsPerEntry = TimeRange / static_cast<float>(MaxSize - 1);
+			if (DeltaTimes[LastInsertIndex] > (SecondsPerEntry - ExecuteContext.GetDeltaTime<float>() * 0.5f))
 			{
 				if (UpperBound < MaxSize)
 				{
 					LastInsertIndex = UpperBound;
 					Buffer[UpperBound] = Value;
-					DeltaTimes[UpperBound++] = ExecuteContext.GetDeltaTime();
+					DeltaTimes[UpperBound++] = ExecuteContext.GetDeltaTime<float>();
 				}
 				else
 				{
 					LastInsertIndex = (LastInsertIndex + 1) % UpperBound;
 					Buffer[LastInsertIndex] = Value;
-					DeltaTimes[LastInsertIndex] = ExecuteContext.GetDeltaTime();
+					DeltaTimes[LastInsertIndex] = ExecuteContext.GetDeltaTime<float>();
 				}
 			}
 			else
 			{
-				DeltaTimes[LastInsertIndex] += ExecuteContext.GetDeltaTime();
+				DeltaTimes[LastInsertIndex] += ExecuteContext.GetDeltaTime<float>();
 			}
 		}
 	}
@@ -187,35 +187,35 @@ FRigVMFunction_TimeOffsetVector_Execute()
 		Result = FMath::Lerp<FVector>(A, B, T);
 	}
 
-	if (ExecuteContext.GetDeltaTime() > SMALL_NUMBER)
+	if (ExecuteContext.GetDeltaTime() > UE_SMALL_NUMBER)
 	{
 		if (UpperBound == 0)
 		{
 			LastInsertIndex = UpperBound;
 			Buffer[UpperBound] = Value;
-			DeltaTimes[UpperBound++] = ExecuteContext.GetDeltaTime();
+			DeltaTimes[UpperBound++] = ExecuteContext.GetDeltaTime<float>();
 		}
 		else
 		{
-			float SecondsPerEntry = TimeRange / float(MaxSize - 1);
-			if (DeltaTimes[LastInsertIndex] > SecondsPerEntry - ExecuteContext.GetDeltaTime() * 0.5f)
+			const float SecondsPerEntry = TimeRange / static_cast<float>(MaxSize - 1);
+			if (DeltaTimes[LastInsertIndex] > (SecondsPerEntry - ExecuteContext.GetDeltaTime<float>() * 0.5f))
 			{
 				if (UpperBound < MaxSize)
 				{
 					LastInsertIndex = UpperBound;
 					Buffer[UpperBound] = Value;
-					DeltaTimes[UpperBound++] = ExecuteContext.GetDeltaTime();
+					DeltaTimes[UpperBound++] = ExecuteContext.GetDeltaTime<float>();
 				}
 				else
 				{
 					LastInsertIndex = (LastInsertIndex + 1) % UpperBound;
 					Buffer[LastInsertIndex] = Value;
-					DeltaTimes[LastInsertIndex] = ExecuteContext.GetDeltaTime();
+					DeltaTimes[LastInsertIndex] = ExecuteContext.GetDeltaTime<float>();
 				}
 			}
 			else
 			{
-				DeltaTimes[LastInsertIndex] += ExecuteContext.GetDeltaTime();
+				DeltaTimes[LastInsertIndex] += ExecuteContext.GetDeltaTime<float>();
 			}
 		}
 	}
@@ -297,35 +297,35 @@ FRigVMFunction_TimeOffsetTransform_Execute()
 		Result.SetScale3D(FMath::Lerp<FVector>(A.GetScale3D(), B.GetScale3D(), T));
 	}
 
-	if (ExecuteContext.GetDeltaTime() > SMALL_NUMBER)
+	if (ExecuteContext.GetDeltaTime() > UE_SMALL_NUMBER)
 	{
 		if (UpperBound == 0)
 		{
 			LastInsertIndex = UpperBound;
 			Buffer[UpperBound] = Value;
-			DeltaTimes[UpperBound++] = ExecuteContext.GetDeltaTime();
+			DeltaTimes[UpperBound++] = ExecuteContext.GetDeltaTime<float>();
 		}
 		else
 		{
-			float SecondsPerEntry = TimeRange / float(MaxSize - 1);
-			if (DeltaTimes[LastInsertIndex] > SecondsPerEntry - ExecuteContext.GetDeltaTime() * 0.5f)
+			const float SecondsPerEntry = TimeRange / static_cast<float>(MaxSize - 1);
+			if (DeltaTimes[LastInsertIndex] > SecondsPerEntry - ExecuteContext.GetDeltaTime<float>() * 0.5f)
 			{
 				if (UpperBound < MaxSize)
 				{
 					LastInsertIndex = UpperBound;
 					Buffer[UpperBound] = Value;
-					DeltaTimes[UpperBound++] = ExecuteContext.GetDeltaTime();
+					DeltaTimes[UpperBound++] = ExecuteContext.GetDeltaTime<float>();
 				}
 				else
 				{
 					LastInsertIndex = (LastInsertIndex + 1) % UpperBound;
 					Buffer[LastInsertIndex] = Value;
-					DeltaTimes[LastInsertIndex] = ExecuteContext.GetDeltaTime();
+					DeltaTimes[LastInsertIndex] = ExecuteContext.GetDeltaTime<float>();
 				}
 			}
 			else
 			{
-				DeltaTimes[LastInsertIndex] += ExecuteContext.GetDeltaTime();
+				DeltaTimes[LastInsertIndex] += ExecuteContext.GetDeltaTime<float>();
 			}
 		}
 	}
diff --git a/Engine/Plugins/Runtime/RigVM/Source/RigVM/Private/RigVMFunctions/Simulation/RigVMFunction_Timeline.cpp b/Engine/Plugins/Runtime/RigVM/Source/RigVM/Private/RigVMFunctions/Simulation/RigVMFunction_Timeline.cpp
--- a/Engine/Plugins/Runtime/RigVM/Source/RigVM/Private/RigVMFunctions/Simulation/RigVMFunction_Timeline.cpp
+++ b/Engine/Plugins/Runtime/RigVM/Source/RigVM/Private/RigVMFunctions/Simulation/RigVMFunction_Timeline.cpp
@@ -12,7 +12,7 @@ FRigVMFunction_Timeline_Execute()
 		bIsInitialized = true;
 	}
 
-	Time = AccumulatedValue = AccumulatedValue + ExecuteContext.GetDeltaTime() * Speed;
+	Time = AccumulatedValue = AccumulatedValue + ExecuteContext.GetDeltaTime<float>() * Speed;
 }
 
 #if WITH_DEV_AUTOMATION_TESTS
@@ -48,7 +48,7 @@ FRigVMFunction_TimeLoop_Execute()
 	}
 
 	const float DurationClamped = FMath::Max(Duration, 0.0001f);
-	const float Increment = ExecuteContext.GetDeltaTime() * Speed;
+	const float Increment = ExecuteContext.GetDeltaTime<float>() * Speed;
 	Absolute = AccumulatedAbsolute = AccumulatedAbsolute + Increment;
 
 	AccumulatedRelative = AccumulatedRelative + Increment;
diff --git a/Engine/Plugins/Runtime/RigVM/Source/RigVM/Private/RigVMFunctions/Simulation/RigVMFunction_Verlet.cpp b/Engine/Plugins/Runtime/RigVM/Source/RigVM/Private/RigVMFunctions/Simulation/RigVMFunction_Verlet.cpp
--- a/Engine/Plugins/Runtime/RigVM/Source/RigVM/Private/RigVMFunctions/Simulation/RigVMFunction_Verlet.cpp
+++ b/Engine/Plugins/Runtime/RigVM/Source/RigVM/Private/RigVMFunctions/Simulation/RigVMFunction_Verlet.cpp
@@ -15,13 +15,13 @@ FRigVMFunction_VerletIntegrateVector_Execute()
 	}
 
 	Point.LinearDamping = Damp;
-	if (ExecuteContext.GetDeltaTime() > SMALL_NUMBER)
+	if (ExecuteContext.GetDeltaTime() > UE_SMALL_NUMBER)
 	{
-		float U = FMath::Clamp<float>(Blend * ExecuteContext.GetDeltaTime(), 0.f, 1.f);
+		float U = FMath::Clamp<float>(Blend * ExecuteContext.GetDeltaTime<float>(), 0.f, 1.f);
 		const FVector CombinedForce = (Target - Point.Position) * FMath::Max(Strength, 0.0001f) + Force * ExecuteContext.GetDeltaTime() * 60.f;
 		const FVector PreviousVelocity = Point.LinearVelocity;
-		Point = Point.IntegrateVerlet(CombinedForce, Blend, ExecuteContext.GetDeltaTime());
-		Acceleration = (Point.LinearVelocity - PreviousVelocity) / ExecuteContext.GetDeltaTime();
+		Point = Point.IntegrateVerlet(CombinedForce, Blend, ExecuteContext.GetDeltaTime<float>());
+		Acceleration = (Point.LinearVelocity - PreviousVelocity) / ExecuteContext.GetDeltaTime<float>();
 		Position = Point.Position;
 		Velocity = Point.LinearVelocity;
 	}
diff --git a/Engine/Plugins/Runtime/RigVM/Source/RigVM/Public/RigVMCore/RigVM.h b/Engine/Plugins/Runtime/RigVM/Source/RigVM/Public/RigVMCore/RigVM.h
--- a/Engine/Plugins/Runtime/RigVM/Source/RigVM/Public/RigVMCore/RigVM.h
+++ b/Engine/Plugins/Runtime/RigVM/Source/RigVM/Public/RigVMCore/RigVM.h
@@ -884,7 +884,7 @@ public:
 		//}
 
 		Statistics.ByteCode = ByteCodePtr->GetStatistics();
-		Statistics.BytesForCaching = FirstHandleForInstruction.GetAllocatedSize(); // +Context.CachedMemoryHandles.GetAllocatedSize(); // Requires context, but fn deprecated already
+		Statistics.BytesForCaching = static_cast<int32>(FirstHandleForInstruction.GetAllocatedSize()); // +Context.CachedMemoryHandles.GetAllocatedSize(); // Requires context, but fn deprecated already
 		Statistics.BytesForCDO =
 			Statistics.LiteralMemory.TotalBytes +
 			Statistics.WorkMemory.TotalBytes +
@@ -1026,7 +1026,7 @@ public:
 private:
 
 #if WITH_EDITOR
-	bool ShouldHaltAtInstruction(FRigVMExtendedExecuteContext& Context, const FName& InEventName, const uint16 InstructionIndex);
+	bool ShouldHaltAtInstruction(FRigVMExtendedExecuteContext& Context, const FName& InEventName, const int32 InstructionIndex);
 #endif
 
 	UPROPERTY()
diff --git a/Engine/Plugins/Runtime/RigVM/Source/RigVM/Public/RigVMCore/RigVMByteCode.h b/Engine/Plugins/Runtime/RigVM/Source/RigVM/Public/RigVMCore/RigVMByteCode.h
--- a/Engine/Plugins/Runtime/RigVM/Source/RigVM/Public/RigVMCore/RigVMByteCode.h
+++ b/Engine/Plugins/Runtime/RigVM/Source/RigVM/Public/RigVMCore/RigVMByteCode.h
@@ -897,7 +897,7 @@ struct RIGVM_API FRigVMChangeTypeOp : public FRigVMUnaryOp
 {
 	GENERATED_USTRUCT_BODY()
 
-	void Serialize(FArchive& Ar);
+	static void Serialize(FArchive& Ar);
 	friend FArchive& operator<<(FArchive& Ar, FRigVMChangeTypeOp& P)
 	{
 		P.Serialize(Ar);
@@ -1042,7 +1042,7 @@ struct RIGVM_API FRigVMInstruction
 {
 	GENERATED_USTRUCT_BODY()
 
-	FRigVMInstruction(ERigVMOpCode InOpCode = ERigVMOpCode::Invalid, uint64 InByteCodeIndex = UINT64_MAX, uint8 InOperandAlignment = 0)
+	FRigVMInstruction(ERigVMOpCode InOpCode = ERigVMOpCode::Invalid, int32 InByteCodeIndex = INDEX_NONE, uint8 InOperandAlignment = 0)
 		: ByteCodeIndex(InByteCodeIndex)
 		, OpCode(InOpCode)
 		, OperandAlignment(InOperandAlignment)
@@ -1050,7 +1050,7 @@ struct RIGVM_API FRigVMInstruction
 	}
 
 	UPROPERTY()
-	uint64 ByteCodeIndex;
+	int32 ByteCodeIndex;
 
 	UPROPERTY()
 	ERigVMOpCode OpCode;
@@ -1164,13 +1164,13 @@ public:
 	}
 	
 	// returns the number of instructions in this container
-	uint64 Num() const;
+	int32 Num() const;
 
 	// const accessor for a byte given its index
 	const uint8& operator[](int32 InIndex) const { return ByteCode[InIndex]; }
 
 	// returns the number of entries
-	uint64 NumEntries() const;
+	int32 NumEntries() const;
 
 	// returns the entry with a given index
 	const FRigVMByteCodeEntry& GetEntry(int32 InEntryIndex) const;
@@ -1179,64 +1179,64 @@ public:
 	int32 FindEntryIndex(const FName& InEntryName) const;
 
 	// adds an execute operator given its function index operands
-	uint64 AddExecuteOp(uint16 InFunctionIndex, const FRigVMOperandArray& InOperands, const int32& StartPredicateIndex, const int32 PredicateCount);
+	int32 AddExecuteOp(uint16 InFunctionIndex, const FRigVMOperandArray& InOperands, const int32& StartPredicateIndex, const int32 PredicateCount);
 	
 	// adds an execute operator given its function index operands
-	uint64 InlineFunction(const FRigVMByteCode* FunctionByteCode, const FRigVMOperandArray& InOperands);
+	int32 InlineFunction(const FRigVMByteCode* FunctionByteCode, const FRigVMOperandArray& InOperands);
 
 	// adds a zero operator to zero the memory of a given argument
-	uint64 AddZeroOp(const FRigVMOperand& InArg);
+	int32 AddZeroOp(const FRigVMOperand& InArg);
 
 	// adds a false operator to set a given argument to false
-	uint64 AddFalseOp(const FRigVMOperand& InArg);
+	int32 AddFalseOp(const FRigVMOperand& InArg);
 
 	// adds a true operator to set a given argument to true
-	uint64 AddTrueOp(const FRigVMOperand& InArg);
+	int32 AddTrueOp(const FRigVMOperand& InArg);
 
 	// adds a copy operator to copy the content of a source argument to a target argument
-	uint64 AddCopyOp(const FRigVMOperand& InSource, const FRigVMOperand& InTarget);
+	int32 AddCopyOp(const FRigVMOperand& InSource, const FRigVMOperand& InTarget);
 
 	// adds a copy operator to copy the content of a source argument to a target argument
-	uint64 AddCopyOp(const FRigVMCopyOp& InCopyOp);
+	int32 AddCopyOp(const FRigVMCopyOp& InCopyOp);
 
 	// adds an increment operator to increment a int32 argument
-	uint64 AddIncrementOp(const FRigVMOperand& InArg);
+	int32 AddIncrementOp(const FRigVMOperand& InArg);
 
 	// adds an decrement operator to decrement a int32 argument
-	uint64 AddDecrementOp(const FRigVMOperand& InArg);
+	int32 AddDecrementOp(const FRigVMOperand& InArg);
 
 	// adds an equals operator to store the comparison result of A and B into a Result argument
-	uint64 AddEqualsOp(const FRigVMOperand& InA, const FRigVMOperand& InB, const FRigVMOperand& InResult);
+	int32 AddEqualsOp(const FRigVMOperand& InA, const FRigVMOperand& InB, const FRigVMOperand& InResult);
 
 	// adds an not-equals operator to store the comparison result of A and B into a Result argument
-	uint64 AddNotEqualsOp(const FRigVMOperand& InA, const FRigVMOperand& InB, const FRigVMOperand& InResult);
+	int32 AddNotEqualsOp(const FRigVMOperand& InA, const FRigVMOperand& InB, const FRigVMOperand& InResult);
 
 	// adds an absolute, forward or backward jump operator
-	uint64 AddJumpOp(ERigVMOpCode InOpCode, uint16 InInstructionIndex);
+	int32 AddJumpOp(ERigVMOpCode InOpCode, uint16 InInstructionIndex);
 
 	// adds an absolute, forward or backward jump operator based on a condition argument
-	uint64 AddJumpIfOp(ERigVMOpCode InOpCode, uint16 InInstructionIndex, const FRigVMOperand& InConditionArg, bool bJumpWhenConditionIs = false);
+	int32 AddJumpIfOp(ERigVMOpCode InOpCode, uint16 InInstructionIndex, const FRigVMOperand& InConditionArg, bool bJumpWhenConditionIs = false);
 
 	// adds an exit operator to exit the execution loop
-	uint64 AddExitOp();
+	int32 AddExitOp();
 
 	// adds an operator to end the last memory slice
-	uint64 AddBeginBlockOp(FRigVMOperand InCountArg, FRigVMOperand InIndexArg);
+	int32 AddBeginBlockOp(FRigVMOperand InCountArg, FRigVMOperand InIndexArg);
 
 	// adds an operator to end the last memory slice
-	uint64 AddEndBlockOp();
+	int32 AddEndBlockOp();
 
 	// adds an invoke entry operator
-	uint64 AddInvokeEntryOp(const FName& InEntryName);
+	int32 AddInvokeEntryOp(const FName& InEntryName);
 
 	// adds a jump to branch operator
-	uint64 AddJumpToBranchOp(FRigVMOperand InBranchNameArg, int32 InFirstBranchInfoIndex);
+	int32 AddJumpToBranchOp(FRigVMOperand InBranchNameArg, int32 InFirstBranchInfoIndex);
 
 	// adds a run instructions op
-	uint64 AddRunInstructionsOp(FRigVMOperand InExecuteStateArg, int32 InStartInstruction, int32 InEndInstruction);
+	int32 AddRunInstructionsOp(FRigVMOperand InExecuteStateArg, int32 InStartInstruction, int32 InEndInstruction);
 
 	// adds a setup traits op
-	uint64 AddSetupTraitsOp(FRigVMOperand InTraitListArg);
+	int32 AddSetupTraitsOp(FRigVMOperand InTraitListArg);
 
 	// adds information about a branch for an instruction's argument
 	int32 AddBranchInfo(const FRigVMBranchInfo& InBranchInfo);
@@ -1252,21 +1252,21 @@ public:
 	}
 
 	// returns the opcode at a given byte index
-	ERigVMOpCode GetOpCodeAt(uint64 InByteCodeIndex) const
+	ERigVMOpCode GetOpCodeAt(int32 InByteCodeIndex) const
 	{
-		ensure(InByteCodeIndex >= 0 && InByteCodeIndex < ByteCode.Num());
-		return (ERigVMOpCode)ByteCode[InByteCodeIndex];
+		checkSlow(InByteCodeIndex >= 0 && InByteCodeIndex < ByteCode.Num());
+		return static_cast<ERigVMOpCode>(ByteCode[InByteCodeIndex]);
 	}
 
 	// returns the size of the operator in bytes at a given byte index
-	uint64 GetOpNumBytesAt(uint64 InByteCodeIndex, bool bIncludeOperands) const;
+	int32 GetOpNumBytesAt(int32 InByteCodeIndex, bool bIncludeOperands) const;
 
 	// returns an operator at a given byte code index
 	template<class OpType>
-	const OpType& GetOpAt(uint64 InByteCodeIndex) const
+	const OpType& GetOpAt(int32 InByteCodeIndex) const
 	{
-		ensure(InByteCodeIndex >= 0 && InByteCodeIndex <= ByteCode.Num() - sizeof(OpType));
-		return *(const OpType*)(ByteCode.GetData() + InByteCodeIndex);
+		checkSlow(InByteCodeIndex >= 0 && InByteCodeIndex <= ByteCode.Num() - sizeof(OpType));
+		return *reinterpret_cast<const OpType*>(ByteCode.GetData() + InByteCodeIndex);
 	}
 
 	// returns an operator for a given instruction
@@ -1278,10 +1278,10 @@ public:
 
 	// returns an operator at a given byte code index
 	template<class OpType>
-	OpType& GetOpAt(uint64 InByteCodeIndex)
+	OpType& GetOpAt(int32 InByteCodeIndex)
 	{
-		ensure(InByteCodeIndex >= 0 && InByteCodeIndex <= ByteCode.Num() - sizeof(OpType));
-		return *(OpType*)(ByteCode.GetData() + InByteCodeIndex);
+		checkSlow(InByteCodeIndex >= 0 && InByteCodeIndex <= ByteCode.Num() - sizeof(OpType));
+		return *reinterpret_cast<OpType*>(ByteCode.GetData() + InByteCodeIndex);
 	}
 
 	// returns an operator for a given instruction
@@ -1292,20 +1292,20 @@ public:
 	}
 
 	// returns a list of operands at a given byte code index
-	FRigVMOperandArray GetOperandsAt(uint64 InByteCodeIndex, uint16 InArgumentCount) const
+	FRigVMOperandArray GetOperandsAt(int32 InByteCodeIndex, int32 InArgumentCount) const
 	{
-		ensure(InByteCodeIndex >= 0 && InByteCodeIndex <= ByteCode.Num() - sizeof(FRigVMOperand) * InArgumentCount);
-		return FRigVMOperandArray((FRigVMOperand*)(ByteCode.GetData() + InByteCodeIndex), InArgumentCount);
+		checkSlow(InByteCodeIndex >= 0 && InByteCodeIndex <= ByteCode.Num() - sizeof(FRigVMOperand) * InArgumentCount);
+		return FRigVMOperandArray(&GetOpAt<FRigVMOperand>(InByteCodeIndex), InArgumentCount);
 	}
 
 	// returns the operands for a given execute instruction
 	FRigVMOperandArray GetOperandsForExecuteOp(const FRigVMInstruction& InInstruction) const
 	{
-		uint64 ByteCodeIndex = InInstruction.ByteCodeIndex;
+		int32 ByteCodeIndex = InInstruction.ByteCodeIndex;
 		const FRigVMExecuteOp& ExecuteOp = GetOpAt<FRigVMExecuteOp>(ByteCodeIndex);
 		// if the bytecode is not aligned the OperandAlignment needs to be 0
-		ensure(bByteCodeIsAligned || InInstruction.OperandAlignment == 0);
-		ByteCodeIndex += sizeof(FRigVMExecuteOp) + (uint64)InInstruction.OperandAlignment;
+		check(bByteCodeIsAligned || InInstruction.OperandAlignment == 0);
+		ByteCodeIndex += sizeof(FRigVMExecuteOp) + InInstruction.OperandAlignment;
 		return GetOperandsAt(ByteCodeIndex, ExecuteOp.GetOperandCount());
 	}
 
@@ -1313,7 +1313,7 @@ public:
 	FRigVMOperandArray GetOperandsForOp(const FRigVMInstruction& InInstruction) const;
 
 	// returns the byte index of the first operand for this instructions
-	uint64 GetFirstOperandByteIndex(const FRigVMInstruction& InInstruction) const;
+	int32 GetFirstOperandByteIndex(const FRigVMInstruction& InInstruction) const;
 
 	// returns all of the operands for a given instruction
 	TArray<int32> GetInstructionsForOperand(const FRigVMOperand& InOperand) const;
@@ -1333,7 +1333,7 @@ public:
 	{
 		FRigVMByteCodeStatistics Statistics;
 		Statistics.InstructionCount = GetInstructions().Num();
-		Statistics.DataBytes = ByteCode.GetAllocatedSize();
+		Statistics.DataBytes = static_cast<int32>(ByteCode.GetAllocatedSize());
 		return Statistics;
 	}
 
@@ -1341,10 +1341,10 @@ public:
 	int32 GetNumInstructions() const { return NumInstructions; }
 
 	// returns the alignment for an operator given its opcode
-	uint64 GetOpAlignment(ERigVMOpCode InOpCode) const;
+	static int32 GetOpAlignment(ERigVMOpCode InOpCode);
 
 	// returns the alignment for an operand
-	uint64 GetOperandAlignment() const;
+	static int32 GetOperandAlignment();
 
 	FString DumpToText() const;
 
@@ -1407,7 +1407,7 @@ public:
 		{
 			if(InputOperandsPerInstruction[InInstructionIndex].Num() > 0)
 			{
-				return FRigVMOperandArray((FRigVMOperand*)(InputOperandsPerInstruction[InInstructionIndex].GetData()), InputOperandsPerInstruction[InInstructionIndex].Num());
+				return FRigVMOperandArray(const_cast<FRigVMOperand*>(InputOperandsPerInstruction[InInstructionIndex].GetData()), InputOperandsPerInstruction[InInstructionIndex].Num());
 			}
 		}
 		return FRigVMOperandArray();
@@ -1420,7 +1420,7 @@ public:
 		{
 			if(OutputOperandsPerInstruction[InInstructionIndex].Num() > 0)
 			{
-				return FRigVMOperandArray((FRigVMOperand*)(OutputOperandsPerInstruction[InInstructionIndex].GetData()), OutputOperandsPerInstruction[InInstructionIndex].Num());
+				return FRigVMOperandArray(const_cast<FRigVMOperand*>(OutputOperandsPerInstruction[InInstructionIndex].GetData()), OutputOperandsPerInstruction[InInstructionIndex].Num());
 			}
 		}
 		return FRigVMOperandArray();
@@ -1473,7 +1473,7 @@ public:
 private:
 
 	template<class OpType>
-	uint64 AddOp(const OpType& InOp)
+	int32 AddOp(const OpType& InOp)
 	{
 		check(
 			(InOp.OpCode > ERigVMOpCode::Execute_64_Operands) && 
@@ -1481,7 +1481,7 @@ private:
 			(InOp.OpCode < ERigVMOpCode::Invalid) 
 		);
 		
-		const uint64 ByteIndex = (uint64)ByteCode.AddZeroed(sizeof(OpType));
+		const int32 ByteIndex = ByteCode.AddZeroed(sizeof(OpType));
 		uint8* Pointer = &ByteCode[ByteIndex];
 		FMemory::Memcpy(Pointer, &InOp, sizeof(OpType));
 		InOp.ZeroPaddedMemoryIfNeeded(reinterpret_cast<OpType*>(Pointer));
diff --git a/Engine/Plugins/Runtime/RigVM/Source/RigVM/Public/RigVMCore/RigVMDebugInfo.h b/Engine/Plugins/Runtime/RigVM/Source/RigVM/Public/RigVMCore/RigVMDebugInfo.h
--- a/Engine/Plugins/Runtime/RigVM/Source/RigVM/Public/RigVMCore/RigVMDebugInfo.h
+++ b/Engine/Plugins/Runtime/RigVM/Source/RigVM/Public/RigVMCore/RigVMDebugInfo.h
@@ -40,7 +40,7 @@ struct RIGVM_API FRigVMBreakpoint
 		Guid.Invalidate();
 	}
 	
-	FRigVMBreakpoint(const uint16 InInstructionIndex, TWeakObjectPtr<UObject> InNode, const uint16 InDepth)
+	FRigVMBreakpoint(const int32 InInstructionIndex, TWeakObjectPtr<UObject> InNode, const uint16 InDepth)
 	: bIsActive(true)
 	, Guid(FGuid::NewGuid())
 	, InstructionIndex(InInstructionIndex)
@@ -76,7 +76,7 @@ struct RIGVM_API FRigVMBreakpoint
 	FGuid Guid;
 	
 	// Instruction where this breakpoint is set
-	uint16 InstructionIndex;
+	int32 InstructionIndex;
 
 	// Node related to the breakpoint
 	TWeakObjectPtr<UObject> Subject; 
@@ -125,11 +125,11 @@ struct RIGVM_API FRigVMDebugInfo
 		return GetBreakpoints().IsEmpty() && !TemporaryBreakpoint.IsValid();
 	}
 
-	const FRigVMBreakpoint& FindBreakpoint(const uint16 InInstructionIndex, const UObject* InSubject) const;
-	TArray<FRigVMBreakpoint> FindBreakpointsAtInstruction(const uint16 InInstructionIndex) const;
+	const FRigVMBreakpoint& FindBreakpoint(const int32 InInstructionIndex, const UObject* InSubject) const;
+	TArray<FRigVMBreakpoint> FindBreakpointsAtInstruction(const int32 InInstructionIndex) const;
 	const FRigVMBreakpoint& FindBreakpoint(const FGuid& InGuid) const;
 	
-	const FRigVMBreakpoint& AddBreakpoint(const uint16 InstructionIndex, TWeakObjectPtr<UObject> InNode, const uint16 InDepth, const bool bIsTemporary = false);
+	const FRigVMBreakpoint& AddBreakpoint(const int32 InstructionIndex, TWeakObjectPtr<UObject> InNode, const uint16 InDepth, const bool bIsTemporary = false);
 
 	bool RemoveBreakpoint(const FRigVMBreakpoint& Breakpoint);
 
diff --git a/Engine/Plugins/Runtime/RigVM/Source/RigVM/Public/RigVMCore/RigVMDrawContainer.h b/Engine/Plugins/Runtime/RigVM/Source/RigVM/Public/RigVMCore/RigVMDrawContainer.h
--- a/Engine/Plugins/Runtime/RigVM/Source/RigVM/Public/RigVMCore/RigVMDrawContainer.h
+++ b/Engine/Plugins/Runtime/RigVM/Source/RigVM/Public/RigVMCore/RigVMDrawContainer.h
@@ -24,7 +24,7 @@ struct RIGVM_API FRigVMDrawContainer
 	TArray<FRigVMDrawInstruction>::RangedForIteratorType      end() { return Instructions.end(); }
 	TArray<FRigVMDrawInstruction>::RangedForConstIteratorType end() const { return Instructions.end(); }
 
-	uint32 GetAllocatedSize(void) const { return Instructions.GetAllocatedSize(); }
+	uint32 GetAllocatedSize() const { return static_cast<uint32>(Instructions.GetAllocatedSize()); }
 
 	void Reset();
 
diff --git a/Engine/Plugins/Runtime/RigVM/Source/RigVM/Public/RigVMCore/RigVMExecuteContext.h b/Engine/Plugins/Runtime/RigVM/Source/RigVM/Public/RigVMCore/RigVMExecuteContext.h
--- a/Engine/Plugins/Runtime/RigVM/Source/RigVM/Public/RigVMCore/RigVMExecuteContext.h
+++ b/Engine/Plugins/Runtime/RigVM/Source/RigVM/Public/RigVMCore/RigVMExecuteContext.h
@@ -319,7 +319,7 @@ struct FRigVMExecuteContext
 		Log(InLogSettings, FString::Printf(Fmt, Args...));
 	}
 
-	uint16 GetInstructionIndex() const { return InstructionIndex; }
+	int32 GetInstructionIndex() const { return InstructionIndex; }
 
 	uint32 GetNumExecutions() const { return NumExecutions; }
 
@@ -328,13 +328,16 @@ struct FRigVMExecuteContext
 	FName GetEventName() const { return EventName; }
 	void SetEventName(const FName& InName) { EventName = InName; }
 
-	double GetDeltaTime() const { return DeltaTime; }
+	template<typename T = double> requires std::is_floating_point_v<T>
+	T GetDeltaTime() const { return static_cast<T>(DeltaTime); }
 	void SetDeltaTime(double InDeltaTime) { DeltaTime = InDeltaTime; }
 
-	double GetAbsoluteTime() const { return AbsoluteTime; } 
+	template<typename T = double> requires std::is_floating_point_v<T>
+	T GetAbsoluteTime() const { return static_cast<T>(AbsoluteTime); } 
 	void SetAbsoluteTime(double InAbsoluteTime) { AbsoluteTime = InAbsoluteTime; }
 
-	double GetFramesPerSecond() const { return FramesPerSecond; } 
+	template<typename T = double> requires std::is_floating_point_v<T>
+	T GetFramesPerSecond() const { return static_cast<T>(FramesPerSecond); } 
 	void SetFramesPerSecond(double InFramesPerSecond) { FramesPerSecond = InFramesPerSecond; }
 
 #if WITH_EDITOR
@@ -492,7 +495,7 @@ protected:
 	
 	FName FunctionName;
 	
-	uint16 InstructionIndex;
+	int32 InstructionIndex;
 
 	uint32 NumExecutions;
 
@@ -675,7 +678,8 @@ struct RIGVM_API FRigVMExtendedExecuteContext
 		Slices.Last().SetRelativeIndex(InRelativeIndex);
 		Slices.Last().SetInstructionIndex(GetPublicData<>().InstructionIndex);
 		Slices.Last().SetOffset(SliceOffsets[GetPublicData<>().InstructionIndex]);
-	}
+	}
+
 
 	void EndSlice()
 	{
@@ -691,7 +695,7 @@ struct RIGVM_API FRigVMExtendedExecuteContext
 		{
 			if(SliceOffsets.IsValidIndex(PoppedSlice.GetInstructionIndex()))
 			{
-				SliceOffsets[PoppedSlice.GetInstructionIndex()] += PoppedSlice.Num();
+				SliceOffsets[PoppedSlice.GetInstructionIndex()] += static_cast<uint16>(PoppedSlice.Num());
 			}
 		}
 	}
diff --git a/Engine/Plugins/Runtime/RigVM/Source/RigVM/Public/RigVMCore/RigVMMemoryDeprecated.h b/Engine/Plugins/Runtime/RigVM/Source/RigVM/Public/RigVMCore/RigVMMemoryDeprecated.h
--- a/Engine/Plugins/Runtime/RigVM/Source/RigVM/Public/RigVMCore/RigVMMemoryDeprecated.h
+++ b/Engine/Plugins/Runtime/RigVM/Source/RigVM/Public/RigVMCore/RigVMMemoryDeprecated.h
@@ -434,7 +434,7 @@ private:
 		uint8* Ptr = nullptr;
 		if (Register.IsDynamic())
 		{
-			Ptr = (uint8*)&Data[Register.GetWorkByteIndex()];
+			Ptr = (uint8*)&Data[static_cast<int32>(Register.GetWorkByteIndex())];
 
 			if (Register.IsNestedDynamic())
 			{
@@ -463,7 +463,7 @@ private:
 		}
 		else
 		{
-			Ptr = (uint8*)&Data[Register.GetWorkByteIndex(InSliceIndex)];
+			Ptr = (uint8*)&Data[static_cast<int32>(Register.GetWorkByteIndex(InSliceIndex))];
 		}
 
 		if (InRegisterOffset != INDEX_NONE && Ptr != nullptr)
diff --git a/Engine/Plugins/Runtime/RigVM/Source/RigVM/Public/RigVMFunctions/Simulation/RigVMFunction_Verlet.h b/Engine/Plugins/Runtime/RigVM/Source/RigVM/Public/RigVMFunctions/Simulation/RigVMFunction_Verlet.h
--- a/Engine/Plugins/Runtime/RigVM/Source/RigVM/Public/RigVMFunctions/Simulation/RigVMFunction_Verlet.h
+++ b/Engine/Plugins/Runtime/RigVM/Source/RigVM/Public/RigVMFunctions/Simulation/RigVMFunction_Verlet.h
@@ -19,7 +19,7 @@ struct RIGVM_API FRigVMFunction_VerletIntegrateVector : public FRigVMFunction_Si
 	{
 		Target = Force = Position = Velocity = Acceleration = FVector::ZeroVector;
 		Strength = 64.f;
-		Damp = 0.01;
+		Damp = 0.01f;
 		Blend = 5.f;
 		Point = FRigVMSimPoint();
 		bInitialized = false;
diff --git a/Engine/Plugins/Runtime/RigVM/Source/RigVM/RigVM.Build.cs b/Engine/Plugins/Runtime/RigVM/Source/RigVM/RigVM.Build.cs
--- a/Engine/Plugins/Runtime/RigVM/Source/RigVM/RigVM.Build.cs
+++ b/Engine/Plugins/Runtime/RigVM/Source/RigVM/RigVM.Build.cs
@@ -7,6 +7,8 @@ public class RigVM : ModuleRules
 {
     public RigVM(ReadOnlyTargetRules Target) : base(Target)
     {
+	    CppCompileWarningSettings.UnsafeTypeCastWarningLevel = WarningLevel.Error;
+
         PrivateIncludePaths.Add(Path.Combine(EngineDirectory,"Source/ThirdParty/AHEasing/AHEasing-1.3.2"));
 
         PublicDependencyModuleNames.AddRange(
