Control Rig: FIx async loading preventing deprecated if and select nodes from being patched
#rnx

diff --git a/Engine/Plugins/Runtime/RigVM/Source/RigVMDeveloper/Private/RigVMModel/RigVMController.cpp b/Engine/Plugins/Runtime/RigVM/Source/RigVMDeveloper/Private/RigVMModel/RigVMController.cpp
--- a/Engine/Plugins/Runtime/RigVM/Source/RigVMDeveloper/Private/RigVMModel/RigVMController.cpp
+++ b/Engine/Plugins/Runtime/RigVM/Source/RigVMDeveloper/Private/RigVMModel/RigVMController.cpp
@@ -36,6 +36,7 @@
 #include "RigVMCore/RigVMObjectArchive.h"
 #include "Logging/LogScopedVerbosityOverride.h"
 #include "ScopedTransaction.h"
+#include "UObject/LinkerLoad.h"
 
 #include UE_INLINE_GENERATED_CPP_BY_NAME(RigVMController)
 
@@ -20863,6 +20864,20 @@ void URigVMController::DestroyObject(UObject* InObjectToDestroy) const
 			}
 		}
 	}
+
+	if (URigVMNode* Node = Cast<URigVMNode>(InObjectToDestroy))
+	{
+		TArray<URigVMPin*> PinsToDestroy = Node->GetAllPinsRecursively();
+		for (int32 i=PinsToDestroy.Num()-1; i>=0; --i)
+		{
+			const FString DeletedName = GetSchema()->GetValidNodeName(GetGraph(), FString::Printf(TEXT("%s_Deleted"), *PinsToDestroy[i]->GetPinPath()));
+			RenameObject(PinsToDestroy[i], *DeletedName, nullptr, false);
+			DestroyObject(PinsToDestroy[i]);
+		}
+	}
+
+	// Invalidate the future loading of the destroyed object
+	FLinkerLoad::InvalidateExport(InObjectToDestroy);
 	
 	// We are renaming an object that may be loading, but ultimately want to destroy.
 	// Pass REN_AllowPackageLinkerMismatch to allow the linker to remain on the object
