RigVM: Avoid cleanup of duplicate permutations

#rnx

diff --git a/Engine/Plugins/Runtime/RigVM/Source/RigVM/Private/RigVMCore/RigVMRegistry.cpp b/Engine/Plugins/Runtime/RigVM/Source/RigVM/Private/RigVMCore/RigVMRegistry.cpp
--- a/Engine/Plugins/Runtime/RigVM/Source/RigVM/Private/RigVMCore/RigVMRegistry.cpp
+++ b/Engine/Plugins/Runtime/RigVM/Source/RigVM/Private/RigVMCore/RigVMRegistry.cpp
@@ -1956,15 +1956,21 @@ const FRigVMTemplate* FRigVMRegistry_NoLock::AddTemplateFromArguments_NoLock(con
 	}
 
 	// Remove duplicate permutations
+	/*
+	 * we'll disable this for now since it's not a valid approach.
+	 * most arguments use type indices by categories, so we can't just remove
+	 * single type indices.
+	 */
+
 	TArray<FRigVMTypeCacheScope_NoLock> TypeCaches;
 	if(!Template.Arguments.IsEmpty())
 	{
-		TArray<bool> ToRemove;
-		int32 NumToRemove = 0;
-		ToRemove.AddZeroed(NumPermutations);
+		//TArray<bool> ToRemove;
+		//int32 NumToRemove = 0;
+		//ToRemove.AddZeroed(NumPermutations);
 
-		TSet<uint32> EncounteredPermutations;
-		EncounteredPermutations.Reserve(NumPermutations);
+		//TSet<uint32> EncounteredPermutations;
+		//EncounteredPermutations.Reserve(NumPermutations);
 
 		const int32 NumArguments = Template.Arguments.Num();
 		TypeCaches.SetNum(NumArguments);
@@ -1976,6 +1982,7 @@ const FRigVMTemplate* FRigVMRegistry_NoLock::AddTemplateFromArguments_NoLock(con
 			bAnyArgumentWithZeroTypes = bAnyArgumentWithZeroTypes || TypeCaches[ArgIndex].GetNumTypes_NoLock() == 0;
 		}
 
+		/*
 		if(!bAnyArgumentWithZeroTypes)
 		{
 			for(int32 Index = 0; Index < NumPermutations; Index++)
@@ -2002,6 +2009,13 @@ const FRigVMTemplate* FRigVMRegistry_NoLock::AddTemplateFromArguments_NoLock(con
 		{
 			for(FRigVMTemplateArgument& Argument : Template.Arguments)
 			{
+				// this is not enough - we may have arguments where the type indices
+				// array is empty...
+				// if(Argument.IsSingleton_NoLock())
+				//{
+				//	continue;
+				//}
+				
 				TArray<TRigVMTypeIndex> NewTypeIndices;
 				NewTypeIndices.Reserve(Argument.TypeIndices.Num() - NumToRemove);
 
@@ -2016,6 +2030,7 @@ const FRigVMTemplate* FRigVMRegistry_NoLock::AddTemplateFromArguments_NoLock(con
 			}
 			NumPermutations -= NumToRemove;
 		}
+		*/
 	}
 
 #if WITH_EDITOR
