RigVM: Fix unnecessary array read for empty set of operands.

#rnx

diff --git a/Engine/Plugins/Runtime/RigVM/Source/RigVM/Public/RigVMCore/RigVMByteCode.h b/Engine/Plugins/Runtime/RigVM/Source/RigVM/Public/RigVMCore/RigVMByteCode.h
--- a/Engine/Plugins/Runtime/RigVM/Source/RigVM/Public/RigVMCore/RigVMByteCode.h
+++ b/Engine/Plugins/Runtime/RigVM/Source/RigVM/Public/RigVMCore/RigVMByteCode.h
@@ -1294,6 +1294,11 @@ public:
 	// returns a list of operands at a given byte code index
 	FRigVMOperandArray GetOperandsAt(int32 InByteCodeIndex, int32 InArgumentCount) const
 	{
+		if (InArgumentCount == 0)
+		{
+			return {};
+		}
+		
 		checkSlow(InByteCodeIndex >= 0 && InByteCodeIndex <= ByteCode.Num() - sizeof(FRigVMOperand) * InArgumentCount);
 		return FRigVMOperandArray(&GetOpAt<FRigVMOperand>(InByteCodeIndex), InArgumentCount);
 	}
