Avoid freeing slack on pop to avoid memory churn in RigVM entry scope

diff --git a/Engine/Plugins/Runtime/RigVM/Source/RigVM/Public/RigVMCore/RigVM.h b/Engine/Plugins/Runtime/RigVM/Source/RigVM/Public/RigVMCore/RigVM.h
--- a/Engine/Plugins/Runtime/RigVM/Source/RigVM/Public/RigVMCore/RigVM.h
+++ b/Engine/Plugins/Runtime/RigVM/Source/RigVM/Public/RigVMCore/RigVM.h
@@ -1114,7 +1114,8 @@ protected:
 
 		~FEntryExecuteGuard()
 		{
-			Stack.Pop();
+			// Don't release memory to avoid churn
+			Stack.Pop(EAllowShrinking::No);
 		}
 	private:
 
