Control Rig: Fix debugger not working on external RigVMHosts
#rnx

diff --git a/Engine/Plugins/Runtime/RigVM/Source/RigVMEditor/Private/Editor/RigVMEditor.cpp b/Engine/Plugins/Runtime/RigVM/Source/RigVMEditor/Private/Editor/RigVMEditor.cpp
--- a/Engine/Plugins/Runtime/RigVM/Source/RigVMEditor/Private/Editor/RigVMEditor.cpp
+++ b/Engine/Plugins/Runtime/RigVM/Source/RigVMEditor/Private/Editor/RigVMEditor.cpp
@@ -3548,9 +3548,9 @@ void FRigVMEditor::HandleShowCurrentStatement()
 
 void FRigVMEditor::HandleBreakpointActionRequested(const ERigVMBreakpointAction BreakpointAction)
 {
-	if (URigVMHost* RigVMHost = GetRigVMHost())
+	if(URigVMHost* DebuggedHost = Cast<URigVMHost>(GetBlueprintObj()->GetObjectBeingDebugged()))
 	{
-		RigVMHost->ExecuteBreakpointAction(BreakpointAction);
+		DebuggedHost->ExecuteBreakpointAction(BreakpointAction);
 	}
 }
 
@@ -3946,6 +3946,8 @@ void FRigVMEditor::HandleSetObjectBeingDebugged(UObject* InObject)
 		if(!URigVMHost::IsGarbageOrDestroyed(PreviouslyDebuggedHost))
 		{
 			PreviouslyDebuggedHost->OnExecuted_AnyThread().RemoveAll(this);
+			PreviouslyDebuggedHost->GetDebugInfo().ExecutionHalted().RemoveAll(this);
+			PreviouslyDebuggedHost->SetIsInDebugMode(false);
 		}
 	}
 	
@@ -3979,6 +3981,10 @@ void FRigVMEditor::HandleSetObjectBeingDebugged(UObject* InObject)
 	{
 		DebuggedHost->SetLog(&RigVMLog);
 		DebuggedHost->OnExecuted_AnyThread().AddSP(this, &FRigVMEditor::HandleVMExecutedEvent);
+#if WITH_EDITOR
+		DebuggedHost->GetDebugInfo().ExecutionHalted().AddSP(this, &FRigVMEditor::HandleVMExecutionHalted);
+		DebuggedHost->SetIsInDebugMode(ExecutionMode == ERigVMEditorExecutionModeType_Debug);	
+#endif
 	}
 
 	RefreshDetailView();
