RigVM: Improve string conversion code

#rnx

diff --git a/Engine/Plugins/Runtime/RigVM/Source/RigVM/Private/RigVMCore/RigVM.cpp b/Engine/Plugins/Runtime/RigVM/Source/RigVM/Private/RigVMCore/RigVM.cpp
--- a/Engine/Plugins/Runtime/RigVM/Source/RigVM/Private/RigVMCore/RigVM.cpp
+++ b/Engine/Plugins/Runtime/RigVM/Source/RigVM/Private/RigVMCore/RigVM.cpp
@@ -19,6 +19,7 @@
 #include "UObject/ObjectSaveContext.h"
 #include "RigVMHost.h"
 #include "RigVMCore/RigVMTrait.h"
+#include "RigVMStringUtils.h"
 
 #include UE_INLINE_GENERATED_CPP_BY_NAME(RigVM)
 
@@ -929,7 +930,7 @@ bool URigVM::CanExecuteEntry(const FRigVMExtendedExecuteContext& Context, const
 		EntryNamesBeingExecuted.Add(InEntryName.ToString());
 
 		static constexpr TCHAR RecursiveEntry[] = TEXT("Entry('%s') is being invoked recursively (%s).");
-		Context.GetPublicData<>().Logf(EMessageSeverity::Error, RecursiveEntry, *InEntryName.ToString(), *FString::Join(EntryNamesBeingExecuted, TEXT(" -> ")));
+		Context.GetPublicData<>().Logf(EMessageSeverity::Error, RecursiveEntry, *InEntryName.ToString(), *RigVMStringUtils::JoinStrings(EntryNamesBeingExecuted, TEXT(" -> ")));
 		return false;
 	}
 
@@ -1796,7 +1797,7 @@ ERigVMExecuteResult URigVM::ExecuteInstructions(FRigVMExtendedExecuteContext& Co
 					Labels.Add(GetOperandLabel(Context, Operand));
 				}
 
-				ContextPublicData.DebugMemoryString += FString::Printf(TEXT("Instruction %d: %s(%s)\n"), ContextPublicData.InstructionIndex, *FunctionNames[Op.FunctionIndex].ToString(), *FString::Join(Labels, TEXT(", ")));
+				ContextPublicData.DebugMemoryString += FString::Printf(TEXT("Instruction %d: %s(%s)\n"), ContextPublicData.InstructionIndex, *FunctionNames[Op.FunctionIndex].ToString(), *RigVMStringUtils::JoinStrings(Labels, TEXT(", ")));
 			}
 			else if(Instruction.OpCode == ERigVMOpCode::Copy)
 			{
@@ -2277,7 +2278,7 @@ ERigVMExecuteResult URigVM::ExecuteInstructions(FRigVMExtendedExecuteContext& Co
 				}
 				++LineIndex;
 			}
-			ContextPublicData.DebugMemoryString += FString::Join(CurrentWorkMemory, TEXT("\n")) + FString(TEXT("\n\n"));
+			ContextPublicData.DebugMemoryString += RigVMStringUtils::JoinStrings(CurrentWorkMemory, TEXT("\n")) + FString(TEXT("\n\n"));
 			ContextPublicData.PreviousWorkMemory = CurrentWorkMemory;
 		}
 #endif
@@ -2392,7 +2393,7 @@ TArray<FString> URigVM::DumpByteCodeAsTextArray(FRigVMExtendedExecuteContext& Co
 					Labels.Add(GetOperandLabel(Context, Operand, OperandFormatFunction));
 				}
 
-				ResultLine = FString::Printf(TEXT("%s(%s)"), *FunctionName, *FString::Join(Labels, TEXT(",")));
+				ResultLine = FString::Printf(TEXT("%s(%s)"), *FunctionName, *RigVMStringUtils::JoinStrings(Labels, TEXT(",")));
 				break;
 			}
 			case ERigVMOpCode::Zero:
@@ -2573,7 +2574,7 @@ TArray<FString> URigVM::DumpByteCodeAsTextArray(FRigVMExtendedExecuteContext& Co
 FString URigVM::DumpByteCodeAsText(FRigVMExtendedExecuteContext& Context, const TArray<int32>& InInstructionOrder, bool bIncludeLineNumbers)
 {
 	RefreshExternalPropertyPaths();
-	return FString::Join(DumpByteCodeAsTextArray(Context, InInstructionOrder, bIncludeLineNumbers), TEXT("\n"));
+	return RigVMStringUtils::JoinStrings(DumpByteCodeAsTextArray(Context, InInstructionOrder, bIncludeLineNumbers), TEXT("\n"));
 }
 
 FString URigVM::GetOperandLabel(FRigVMExtendedExecuteContext& Context, const FRigVMOperand& InOperand, TFunction<FString(const FString& RegisterName, const FString& RegisterOffsetName)> FormatFunction)
diff --git a/Engine/Plugins/Runtime/RigVM/Source/RigVM/Private/RigVMCore/RigVMByteCode.cpp b/Engine/Plugins/Runtime/RigVM/Source/RigVM/Private/RigVMCore/RigVMByteCode.cpp
--- a/Engine/Plugins/Runtime/RigVM/Source/RigVM/Private/RigVMCore/RigVMByteCode.cpp
+++ b/Engine/Plugins/Runtime/RigVM/Source/RigVM/Private/RigVMCore/RigVMByteCode.cpp
@@ -8,6 +8,7 @@
 #include "RigVMCore/RigVMTrait.h"
 #include "RigVMCore/RigVMTraitScope.h"
 #include "RigVMObjectVersion.h"
+#include "RigVMStringUtils.h"
 
 #include UE_INLINE_GENERATED_CPP_BY_NAME(RigVMByteCode)
 
@@ -1253,7 +1254,7 @@ FString FRigVMByteCode::DumpToText() const
 						OperandsContent.Add(FString::Printf(TEXT("\t%s"), *OperandContent));
 					}
 
-					Line += FString::Printf(TEXT("(\n%s\n)"), *FString::Join(OperandsContent, TEXT("\n")));
+					Line += RigVMStringUtils::JoinStrings(OperandsContent, TEXT("(\n"), TEXT("\n"), TEXT("\n)"));
 				}
 				break;
 			}
@@ -1354,7 +1355,7 @@ FString FRigVMByteCode::DumpToText() const
 					}
 					BranchInfoTexts.Add(FString::Printf(TEXT("%s (%u)"), *BranchInfo.Label.ToString(), BranchInfo.FirstInstruction));
 				}
-				Line += FString::Join(BranchInfoTexts, TEXT(", "));
+				Line += RigVMStringUtils::JoinStrings(BranchInfoTexts, TEXT(", "));
 				break;
 			}
 			case ERigVMOpCode::RunInstructions:
@@ -1394,7 +1395,7 @@ FString FRigVMByteCode::DumpToText() const
 		return FString();
 	}
 
-	return FString::Join(Lines, TEXT("\n"));
+	return RigVMStringUtils::JoinStrings(Lines, TEXT("\n"));
 }
 
 int32 FRigVMByteCode::AddBeginBlockOp(FRigVMOperand InCountArg, FRigVMOperand InIndexArg)
diff --git a/Engine/Plugins/Runtime/RigVM/Source/RigVM/Private/RigVMCore/RigVMDefines.cpp b/Engine/Plugins/Runtime/RigVM/Source/RigVM/Private/RigVMCore/RigVMDefines.cpp
--- a/Engine/Plugins/Runtime/RigVM/Source/RigVM/Private/RigVMCore/RigVMDefines.cpp
+++ b/Engine/Plugins/Runtime/RigVM/Source/RigVM/Private/RigVMCore/RigVMDefines.cpp
@@ -53,7 +53,7 @@ FRigVMArchiveTraceBracket::FRigVMArchiveTraceBracket(FArchive& InArchive, const
 		
 		ArchiveName = InArchive.GetArchiveName();
 		ArchiveWhiteSpace = WhiteSpace(Indentation * 2);
-		ArchivePrefix = FString::Printf(TEXT("  %s%s"), *ArchiveWhiteSpace, *InScope);
+		ArchivePrefix = TEXT("  ") + ArchiveWhiteSpace + InScope;
 
 		const FString ArchiveOffset = ArchiveOffsetToString(ArchivePos);
 		static const FString EntrySize = WhiteSpace(14);
diff --git a/Engine/Plugins/Runtime/RigVM/Source/RigVM/Private/RigVMCore/RigVMDispatchFactory.cpp b/Engine/Plugins/Runtime/RigVM/Source/RigVM/Private/RigVMCore/RigVMDispatchFactory.cpp
--- a/Engine/Plugins/Runtime/RigVM/Source/RigVM/Private/RigVMCore/RigVMDispatchFactory.cpp
+++ b/Engine/Plugins/Runtime/RigVM/Source/RigVM/Private/RigVMCore/RigVMDispatchFactory.cpp
@@ -4,6 +4,7 @@
 #include "RigVMCore/RigVMRegistry.h"
 #include "RigVMCore/RigVMStruct.h"
 #include "RigVMCore/RigVMMemoryStorage.h"
+#include "RigVMStringUtils.h"
 
 #include UE_INLINE_GENERATED_CPP_BY_NAME(RigVMDispatchFactory)
 
@@ -11,8 +12,7 @@ FCriticalSection FRigVMDispatchFactory::GetTemplateMutex;
 
 FName FRigVMDispatchFactory::GetFactoryName() const
 {
-	static constexpr TCHAR Format[] = TEXT("%s%s");
-	return *FString::Printf(Format, DispatchPrefix, *GetScriptStruct()->GetName());
+	return *(DispatchPrefix + GetScriptStruct()->GetName());
 }
 
 #if WITH_EDITOR
@@ -344,9 +344,8 @@ TArray<FRigVMTemplateArgumentInfo> FRigVMDispatchFactory::BuildArgumentListFromP
 
 FString FRigVMDispatchFactory::GetPermutationNameImpl(const FRigVMTemplateTypeMap& InTypes, const bool bLockRegistry) const
 {
-	static constexpr TCHAR Format[] = TEXT("%s::%s");
 	const FString TypePairStrings = FRigVMTemplate::GetStringFromArgumentTypes(InTypes, bLockRegistry);
-	return FString::Printf(Format, *GetFactoryName().ToString(), *TypePairStrings);
+	return RigVMStringUtils::JoinStrings(GetFactoryName().ToString(), TypePairStrings, TEXT("::"));
 }
 
 bool FRigVMDispatchFactory::CopyProperty(const FProperty* InTargetProperty, uint8* InTargetPtr,
diff --git a/Engine/Plugins/Runtime/RigVM/Source/RigVM/Private/RigVMCore/RigVMRegistry.cpp b/Engine/Plugins/Runtime/RigVM/Source/RigVM/Private/RigVMCore/RigVMRegistry.cpp
--- a/Engine/Plugins/Runtime/RigVM/Source/RigVM/Private/RigVMCore/RigVMRegistry.cpp
+++ b/Engine/Plugins/Runtime/RigVM/Source/RigVM/Private/RigVMCore/RigVMRegistry.cpp
@@ -5,6 +5,7 @@
 #include "RigVMCore/RigVMStruct.h"
 #include <RigVMCore/RigVMTrait.h>
 #include "RigVMTypeUtils.h"
+#include "RigVMStringUtils.h"
 #include "RigVMModule.h"
 #include "Animation/AttributeTypes.h"
 #include "AssetRegistry/AssetRegistryModule.h"
@@ -1535,7 +1536,7 @@ void FRigVMRegistry_NoLock::Register_NoLock(const TCHAR* InName, FRigVMFunctionP
 		FString MethodName;
 		if (FString(InName).Split(TEXT("::"), nullptr, &MethodName))
 		{
-			const FString TemplateName = FString::Printf(TEXT("%s::%s"), *TemplateMetadata, *MethodName);
+			const FString TemplateName = RigVMStringUtils::JoinStrings(TemplateMetadata, MethodName, TEXT("::"));
 			FRigVMTemplate Template(InStruct, TemplateName, Function.Index);
 			if (Template.IsValid())
 			{
@@ -1768,8 +1769,8 @@ const FRigVMFunction* FRigVMRegistry_NoLock::FindFunction_NoLock(UScriptStruct*
 {
 	check(InStruct);
 	check(InName);
-	
-	const FString FunctionName = FString::Printf(TEXT("%s::%s"), *InStruct->GetStructCPPName(), InName);
+
+	const FString FunctionName = RigVMStringUtils::JoinStrings(InStruct->GetStructCPPName(), InName, TEXT("::"));
 	return FindFunction_NoLock(*FunctionName, InResolvalInfo);
 }
 
@@ -1834,9 +1835,9 @@ const FRigVMTemplate* FRigVMRegistry_NoLock::FindTemplate_NoLock(const FName& In
 #if WITH_EDITOR
 	if(OriginalNotation.StartsWith(FRigVMDispatchFactory::DispatchPrefix))
 	{
-		const FString OriginalDispatchFactoryName = OriginalNotation
-			.Left(OriginalNotation.Find(TEXT("(")))
-			.RightChop(FCString::Strlen(FRigVMDispatchFactory::DispatchPrefix));
+		static const int32 PrefixLen = FCString::Strlen(FRigVMDispatchFactory::DispatchPrefix);
+		const int32 BraceIndex = OriginalNotation.Find(TEXT("("));
+		const FString OriginalDispatchFactoryName = OriginalNotation.Mid(PrefixLen, BraceIndex - PrefixLen); 
 
 		const FCoreRedirectObjectName OldObjectName(OriginalDispatchFactoryName);
 		TArray<const FCoreRedirect*> Redirects;
diff --git a/Engine/Plugins/Runtime/RigVM/Source/RigVM/Private/RigVMCore/RigVMStruct.cpp b/Engine/Plugins/Runtime/RigVM/Source/RigVM/Private/RigVMCore/RigVMStruct.cpp
--- a/Engine/Plugins/Runtime/RigVM/Source/RigVM/Private/RigVMCore/RigVMStruct.cpp
+++ b/Engine/Plugins/Runtime/RigVM/Source/RigVM/Private/RigVMCore/RigVMStruct.cpp
@@ -2,6 +2,7 @@
 
 #include "RigVMCore/RigVMStruct.h"
 #include "RigVMCore/RigVMRegistry.h"
+#include "RigVMStringUtils.h"
 #include "UObject/StructOnScope.h"
 #include "RigVMModule.h"
 
@@ -596,14 +597,7 @@ FString FRigVMStruct::ExportToFullyQualifiedText(const FProperty* InMemberProper
 			ElementValues.Add(ExportToFullyQualifiedText(ArrayProperty->Inner, ElementMemoryPtr));
 		}
 
-		if (ElementValues.Num() == 0)
-		{
-			DefaultValue = TEXT("()");
-		}
-		else
-		{
-			DefaultValue = FString::Printf(TEXT("(%s)"), *FString::Join(ElementValues, TEXT(",")));
-		}
+		DefaultValue = RigVMStringUtils::JoinStrings(ElementValues, TEXT(","), TEXT("("), TEXT(")"));
 	}
 	else
 	{
@@ -665,12 +659,7 @@ FString FRigVMStruct::ExportToFullyQualifiedText(const UScriptStruct* InStruct,
 			FieldValues.Add(FString::Printf(TEXT("%s=%s"), *PropertyName, *DefaultValue));
 		}
 
-		if (FieldValues.Num() == 0)
-		{
-			return TEXT("()");
-		}
-
-		return FString::Printf(TEXT("(%s)"), *FString::Join(FieldValues, TEXT(",")));
+		return RigVMStringUtils::JoinStrings(FieldValues, TEXT(","), TEXT("("), TEXT(")"));
 	}
 }
 
diff --git a/Engine/Plugins/Runtime/RigVM/Source/RigVM/Private/RigVMCore/RigVMTemplate.cpp b/Engine/Plugins/Runtime/RigVM/Source/RigVM/Private/RigVMCore/RigVMTemplate.cpp
--- a/Engine/Plugins/Runtime/RigVM/Source/RigVM/Private/RigVMCore/RigVMTemplate.cpp
+++ b/Engine/Plugins/Runtime/RigVM/Source/RigVM/Private/RigVMCore/RigVMTemplate.cpp
@@ -843,9 +843,8 @@ FName FRigVMTemplateArgumentInfo::ComputeTemplateNotation(const FName InTemplate
 	{
 		return NAME_None;
 	}
-	
-	const FString NotationStr = FString::Printf(TEXT("%s(%s)"), *InTemplateName.ToString(), *FString::Join(ArgumentNotations, TEXT(",")));
-	return *NotationStr;
+
+	return *RigVMStringUtils::MakeTemplateNotation(InTemplateName.ToString(), ArgumentNotations);
 }
 
 TArray<TRigVMTypeIndex> FRigVMTemplateArgumentInfo::GetTypesFromCategories(
@@ -939,7 +938,7 @@ FRigVMTemplate::FRigVMTemplate(UScriptStruct* InStruct, const FString& InTemplat
 
 	if (ArgumentNotations.Num() > 0)
 	{
-		const FString NotationStr = FString::Printf(TEXT("%s(%s)"), *InTemplateName, *FString::Join(ArgumentNotations, TEXT(",")));
+		const FString NotationStr = RigVMStringUtils::MakeTemplateNotation(InTemplateName, ArgumentNotations);
 		Notation = *NotationStr;
 		if (InFunctionIndex != INDEX_NONE)
 		{
@@ -1039,7 +1038,7 @@ const FString& FRigVMTemplate::GetDirectionPrefix(const ERigVMPinDirection InDir
 
 FString FRigVMTemplate::GetArgumentNotation(const FName InName, const ERigVMPinDirection InDirection)
 {
-	return FString::Printf(TEXT("%s%s"), *GetDirectionPrefix(InDirection), *InName.ToString());
+	return GetDirectionPrefix(InDirection) + InName.ToString();
 }
 
 void FRigVMTemplate::ComputeNotationFromArguments(const FString& InTemplateName)
@@ -1053,7 +1052,7 @@ void FRigVMTemplate::ComputeNotationFromArguments(const FString& InTemplateName)
 		}
 	}
 
-	const FString NotationStr = FString::Printf(TEXT("%s(%s)"), *InTemplateName, *FString::Join(ArgumentNotations, TEXT(",")));
+	const FString NotationStr = RigVMStringUtils::MakeTemplateNotation(InTemplateName, ArgumentNotations);
 	Notation = *NotationStr;
 }
 
@@ -1140,10 +1139,17 @@ FString FRigVMTemplate::GetStringFromArgumentTypes(const FTypeMap& InTypes, bool
 	{
 		const FRigVMTemplateArgumentType& Type = Registry.GetType_NoLock(Pair.Value);
 		static constexpr TCHAR Format[] = TEXT("%s:%s");
-		TypePairStrings.Add(FString::Printf(Format, *Pair.Key.ToString(), *Type.CPPType.ToString()));
+		FString PairString;
+		const FString KeyString = Pair.Key.ToString();
+		const FString CPPTypeString = Type.CPPType.ToString();
+		PairString.Reserve(KeyString.Len() + CPPTypeString.Len() + 1);
+		PairString.Append(KeyString);
+		PairString.AppendChar(TEXT(':'));
+		PairString.Append(CPPTypeString);
+		TypePairStrings.Add(PairString);
 	}
 			
-	return FString::Join(TypePairStrings, TEXT(","));
+	return RigVMStringUtils::JoinStrings(TypePairStrings, TEXT(","));
 }
 
 bool FRigVMTemplate::IsValid() const
@@ -2178,7 +2184,7 @@ FString FRigVMTemplate::GetKeywords() const
 		}
 	}
 
-	return FString::Join(KeywordsMetadata, TEXT(","));
+	return RigVMStringUtils::JoinStrings(KeywordsMetadata, TEXT(","));
 }
 
 #endif
diff --git a/Engine/Plugins/Runtime/RigVM/Source/RigVM/Private/RigVMFunctions/RigVMDispatch_CastObject.cpp b/Engine/Plugins/Runtime/RigVM/Source/RigVM/Private/RigVMFunctions/RigVMDispatch_CastObject.cpp
--- a/Engine/Plugins/Runtime/RigVM/Source/RigVM/Private/RigVMFunctions/RigVMDispatch_CastObject.cpp
+++ b/Engine/Plugins/Runtime/RigVM/Source/RigVM/Private/RigVMFunctions/RigVMDispatch_CastObject.cpp
@@ -72,10 +72,10 @@ FString FRigVMDispatch_CastObject::GetNodeTitle(const FRigVMTemplateTypeMap& InT
 			{
 				const FRigVMTemplateArgumentType& ValueType = Registry.GetType(ValueTypeIndex);
 				const FRigVMTemplateArgumentType& ResultType = Registry.GetType(ResultTypeIndex);
-				static constexpr TCHAR Format[] = TEXT("Cast to %s");
+				static const FString Prefix = TEXT("Cast to ");
 				if(ValueType.CPPTypeObject && ResultType.CPPTypeObject)
 				{
-					return FString::Printf(Format, *ResultType.CPPTypeObject->GetName());
+					return Prefix + ResultType.CPPTypeObject->GetName();
 				}
 			}
 		}
diff --git a/Engine/Plugins/Runtime/RigVM/Source/RigVM/Private/RigVMFunctions/RigVMFunction_String.cpp b/Engine/Plugins/Runtime/RigVM/Source/RigVM/Private/RigVMFunctions/RigVMFunction_String.cpp
--- a/Engine/Plugins/Runtime/RigVM/Source/RigVM/Private/RigVMFunctions/RigVMFunction_String.cpp
+++ b/Engine/Plugins/Runtime/RigVM/Source/RigVM/Private/RigVMFunctions/RigVMFunction_String.cpp
@@ -3,6 +3,7 @@
 #include "RigVMFunctions/RigVMFunction_String.h"
 #include "RigVMFunctions/RigVMFunctionDefines.h"
 #include "RigVMCore/RigVM.h"
+#include "RigVMStringUtils.h"
 
 #include UE_INLINE_GENERATED_CPP_BY_NAME(RigVMFunction_String)
 
@@ -145,7 +146,7 @@ FRigVMFunction_StringJoin_Execute()
 	Result.Reset();
 	if(!Values.IsEmpty())
 	{
-		Result = FString::Join(Values, *Separator);
+		Result = RigVMStringUtils::JoinStringsConst(Values, *Separator);
 	}
 }
 
diff --git a/Engine/Plugins/Runtime/RigVM/Source/RigVM/Private/RigVMStringUtils.cpp b/Engine/Plugins/Runtime/RigVM/Source/RigVM/Private/RigVMStringUtils.cpp
--- a/Engine/Plugins/Runtime/RigVM/Source/RigVM/Private/RigVMStringUtils.cpp
+++ b/Engine/Plugins/Runtime/RigVM/Source/RigVM/Private/RigVMStringUtils.cpp
@@ -2,58 +2,155 @@
 
 #include "RigVMStringUtils.h"
 
-bool RigVMStringUtils::SplitNodePathAtStart(const FString& InNodePath, FString& LeftMost, FString& Right)
+FString RigVMStringUtils::JoinStrings(const FString& InStringA, const FString& InStringB, const TCHAR* InSeparator, const TCHAR* InPrefix, const TCHAR* InSuffix)
 {
-	return InNodePath.Split(TEXT("|"), &LeftMost, &Right, ESearchCase::IgnoreCase, ESearchDir::FromStart);
+	const int32 SeparatorLen = FCString::Strlen(InSeparator);
+	const int32 PrefixLen = InPrefix ? FCString::Strlen(InPrefix) : 0;
+	const int32 SuffixLen = InSuffix ? FCString::Strlen(InSuffix) : 0;
+
+	const int32 RequiredLen = PrefixLen + InStringA.Len() + SeparatorLen + InStringB.Len() + SuffixLen;
+
+	FStringBuilderBase Builder;
+	Builder.Reserve(RequiredLen + 1 /* \0 */);
+
+	if(InPrefix)
+	{
+		Builder.Append(InPrefix);
+	}
+	Builder.Append(InStringA);
+	Builder.Append(InSeparator);
+	Builder.Append(InStringB);
+	if(InSuffix)
+	{
+		Builder.Append(InSuffix);
+	}
+	return Builder.ToString();
 }
 
-bool RigVMStringUtils::SplitNodePathAtEnd(const FString& InNodePath, FString& Left, FString& RightMost)
+FString RigVMStringUtils::JoinStrings(const TArray<FString>& InStrings, const TCHAR* InSeparator, const TCHAR* InPrefix, const TCHAR* InSuffix)
 {
-	return InNodePath.Split(TEXT("|"), &Left, &RightMost, ESearchCase::IgnoreCase, ESearchDir::FromEnd);
+	return JoinStringsConst(TArrayView<const FString>(InStrings.GetData(), InStrings.Num()), InSeparator, InPrefix, InSuffix);
 }
 
-bool RigVMStringUtils::SplitNodePath(const FString& InNodePath, TArray<FString>& Parts)
+FString RigVMStringUtils::JoinStrings(const TArrayView<FString>& InStrings, const TCHAR* InSeparator, const TCHAR* InPrefix, const TCHAR* InSuffix)
 {
-	int32 OriginalPartsCount = Parts.Num();
-	FString NodePathRemaining = InNodePath;
-	FString Left, Right;
-	Right = NodePathRemaining;
+	return JoinStringsConst(TArrayView<const FString>(InStrings.GetData(), InStrings.Num()), InSeparator, InPrefix, InSuffix);
+}
 
-	while (SplitNodePathAtStart(NodePathRemaining, Left, Right))
+FString RigVMStringUtils::JoinStringsConst(const TArrayView<const FString>& InStrings, const TCHAR* InSeparator, const TCHAR* InPrefix, const TCHAR* InSuffix)
+{
+	if(InStrings.IsEmpty())
+	{
+		if(InPrefix && InSuffix)
+		{
+			return FString(InPrefix) + FString(InSuffix);
+		}
+		if(InPrefix)
+		{
+			return InPrefix;
+		}
+		if(InSuffix)
+		{
+			return InSuffix;
+		}
+		return FString();
+	}
+	
+	const int32 SeparatorLen = FCString::Strlen(InSeparator);
+	const int32 PrefixLen = InPrefix ? FCString::Strlen(InPrefix) : 0;
+	const int32 SuffixLen = InSuffix ? FCString::Strlen(InSuffix) : 0;
+	int32 StringsTotalLen = 0;
+	for(const FString& String : InStrings)
 	{
-		Parts.Add(Left);
-		Left.Empty();
-		NodePathRemaining = Right;
+		StringsTotalLen += String.Len();
 	}
 
-	if (!Right.IsEmpty())
+	const int32 RequiredLen = PrefixLen + StringsTotalLen + SeparatorLen * (InStrings.Num() - 1) + SuffixLen;
+	FStringBuilderBase Builder;
+	Builder.Reserve(RequiredLen + 1 /* \0 */);
+
+	if(InPrefix)
 	{
-		Parts.Add(Right);
+		Builder.Append(InPrefix);
 	}
-
-	return Parts.Num() > OriginalPartsCount;
+	for(int32 Index = 0; Index < InStrings.Num(); Index++)
+	{
+		if(Index > 0)
+		{
+			Builder.Append(InSeparator);
+		}
+		Builder.Append(InStrings[Index]);
+	}
+	if(InSuffix)
+	{
+		Builder.Append(InSuffix);
+	}
+	return Builder.ToString();
 }
 
-FString RigVMStringUtils::JoinNodePath(const FString& Left, const FString& Right)
+// Splits a string up into its parts 
+bool RigVMStringUtils::SplitString(const FString& InString, const TCHAR* InSeparator, TArray<FString>& InOutParts)
 {
-	ensure(!Left.IsEmpty() && !Right.IsEmpty());
-	return Left + TEXT("|") + Right;
-}
+	const int32 OriginalPartsCount = InOutParts.Num();
+	FStringView Remaining = InString;
 
-FString RigVMStringUtils::JoinNodePath(const TArray<FString>& InParts)
-{
-	if (InParts.Num() == 0)
+	int32 Index = Remaining.Find(InSeparator, 0, ESearchCase::IgnoreCase);
+	while (Index != INDEX_NONE)
 	{
-		return FString();
+		InOutParts.Emplace(Remaining.Left(Index));
+		Remaining.RightChopInline(Index + 1);
+		Index = Remaining.Find(InSeparator, 0, ESearchCase::IgnoreCase);
 	}
 
-	FString Result = InParts[0];
-	for (int32 PartIndex = 1; PartIndex < InParts.Num(); PartIndex++)
+	if(!Remaining.IsEmpty())
 	{
-		Result += TEXT("|") + InParts[PartIndex];
+		InOutParts.Emplace(Remaining);
 	}
 
-	return Result;
+	return InOutParts.Num() > OriginalPartsCount;
+}
+
+// Creates a template notation such as "MyTemplate(float,bool,float)"
+FString RigVMStringUtils::MakeTemplateNotation(const FString& InTemplateName, const TArray<FString>& InArgumentNotations)
+{
+	const FString JoinedArguments = JoinStrings(InArgumentNotations, TEXT(","), TEXT("("), TEXT(")"));
+	
+	FStringBuilderBase Builder;
+	Builder.Reserve(InTemplateName.Len() + JoinedArguments.Len() + 1 /* \0 */);
+	Builder.Append(InTemplateName);
+	Builder.Append(JoinedArguments);
+	return Builder.ToString();
+}
+
+bool RigVMStringUtils::SplitNodePathAtStart(const FString& InNodePath, FString& LeftMost, FString& Right)
+{
+	return InNodePath.Split(TEXT("|"), &LeftMost, &Right, ESearchCase::IgnoreCase, ESearchDir::FromStart);
+}
+
+bool RigVMStringUtils::SplitNodePathAtEnd(const FString& InNodePath, FString& Left, FString& RightMost)
+{
+	return InNodePath.Split(TEXT("|"), &Left, &RightMost, ESearchCase::IgnoreCase, ESearchDir::FromEnd);
+}
+
+bool RigVMStringUtils::SplitNodePath(const FString& InNodePath, TArray<FString>& Parts)
+{
+	return SplitString(InNodePath, TEXT("|"), Parts);
+}
+
+FString RigVMStringUtils::JoinNodePath(const FString& Left, const FString& Right)
+{
+	ensure(!Left.IsEmpty() && !Right.IsEmpty());
+	FStringBuilderBase Builder;
+	Builder.Reserve(Left.Len() + 1 + Right.Len() + 1 /* \0 */);
+	Builder.Append(Left);
+	Builder.Append(TEXT("|"));
+	Builder.Append(Right);
+	return Builder.ToString();
+}
+
+FString RigVMStringUtils::JoinNodePath(const TArray<FString>& InParts)
+{
+	return JoinStrings(InParts, TEXT("|"));
 }
 
 bool RigVMStringUtils::SplitPinPathAtStart(const FString& InPinPath, FString& LeftMost, FString& Right)
@@ -68,56 +165,28 @@ bool RigVMStringUtils::SplitPinPathAtEnd(const FString& InPinPath, FString& Left
 
 bool RigVMStringUtils::SplitPinPath(const FString& InPinPath, TArray<FString>& Parts)
 {
-	int32 OriginalPartsCount = Parts.Num();
-	FString PinPathRemaining = InPinPath;
-	FString Left, Right;
-	while(SplitPinPathAtStart(PinPathRemaining, Left, Right))
-	{
-		Parts.Add(Left);
-		Left.Empty();
-		PinPathRemaining = Right;
-	}
-
-	if (!Right.IsEmpty())
-	{
-		Parts.Add(Right);
-	}
-
-	return Parts.Num() > OriginalPartsCount;
+	return SplitString(InPinPath, TEXT("."), Parts);
 }
 
 FString RigVMStringUtils::JoinPinPath(const FString& Left, const FString& Right)
 {
 	ensure(!Left.IsEmpty() && !Right.IsEmpty());
-	return Left + TEXT(".") + Right;
+	FStringBuilderBase Builder;
+	Builder.Reserve(Left.Len() + 1 + Right.Len() + 1 /* \0 */);
+	Builder.Append(Left);
+	Builder.Append(TEXT("."));
+	Builder.Append(Right);
+	return Builder.ToString();
 }
 
 FString RigVMStringUtils::JoinPinPath(const TArray<FString>& InParts)
 {
-	if (InParts.Num() == 0)
-	{
-		return FString();
-	}
-
-	FString Result = InParts[0];
-	for (int32 PartIndex = 1; PartIndex < InParts.Num(); PartIndex++)
-	{
-		Result += TEXT(".") + InParts[PartIndex];
-	}
-
-	return Result;
+	return JoinStrings(InParts, TEXT("."));
 }
 
 FString RigVMStringUtils::JoinDefaultValue(const TArray<FString>& InParts)
 {
-	static constexpr TCHAR EmptyBraces[] = TEXT("()");
-	if(InParts.IsEmpty())
-	{
-		return EmptyBraces;
-	}
-
-	static constexpr TCHAR Format[] = TEXT("(%s)");
-	return FString::Printf(Format, *FString::Join(InParts, TEXT(",")));
+	return JoinStrings(InParts, TEXT(","), TEXT("("), TEXT(")"));
 }
 
 TArray<FString> RigVMStringUtils::SplitDefaultValue(const FString& InDefaultValue)
diff --git a/Engine/Plugins/Runtime/RigVM/Source/RigVM/Public/RigVMCore/RigVMGraphFunctionDefinition.h b/Engine/Plugins/Runtime/RigVM/Source/RigVM/Public/RigVMCore/RigVMGraphFunctionDefinition.h
--- a/Engine/Plugins/Runtime/RigVM/Source/RigVM/Public/RigVMCore/RigVMGraphFunctionDefinition.h
+++ b/Engine/Plugins/Runtime/RigVM/Source/RigVM/Public/RigVMCore/RigVMGraphFunctionDefinition.h
@@ -11,6 +11,7 @@
 #include "UObject/UE5ReleaseStreamObjectVersion.h"
 #include "RigVMVariant.h"
 #include "RigVMNodeLayout.h"
+#include "RigVMStringUtils.h"
 #include "RigVMGraphFunctionDefinition.generated.h"
 
 class IRigVMGraphFunctionHost;
@@ -454,7 +455,7 @@ struct RIGVM_API FRigVMGraphFunctionHeader
 
 	FString GetHash() const
 	{
-		return FString::Printf(TEXT("%s:%s"), *LibraryPointer.HostObject.ToString(), *Name.ToString());
+		return RigVMStringUtils::JoinStrings(LibraryPointer.HostObject.ToString(), Name.ToString(), TEXT(":"));
 	}
 
 	friend uint32 GetTypeHash(const FRigVMGraphFunctionHeader& Header)
diff --git a/Engine/Plugins/Runtime/RigVM/Source/RigVM/Public/RigVMCore/RigVMMemoryStorage.h b/Engine/Plugins/Runtime/RigVM/Source/RigVM/Public/RigVMCore/RigVMMemoryStorage.h
--- a/Engine/Plugins/Runtime/RigVM/Source/RigVM/Public/RigVMCore/RigVMMemoryStorage.h
+++ b/Engine/Plugins/Runtime/RigVM/Source/RigVM/Public/RigVMCore/RigVMMemoryStorage.h
@@ -31,6 +31,7 @@
 #include "UObject/UObjectGlobals.h"
 #include "UObject/UnrealNames.h"
 #include "UObject/UnrealType.h"
+#include "RigVMStringUtils.h"
 
 #include "RigVMMemoryStorage.generated.h"
 
@@ -103,7 +104,7 @@ struct RIGVM_API FRigVMBranchInfo
 
 	static FString GetFixedArrayLabel(const FString& InArrayName, const FString& InElementName)
 	{
-		return FString::Printf(TEXT("%s_%s"), *InArrayName, *InElementName);
+		return RigVMStringUtils::JoinStrings(InArrayName, InElementName, TEXT("_"));
 	}
 
 	static FName GetFixedArrayLabel(const FName& InArrayName, const FName& InElementName)
diff --git a/Engine/Plugins/Runtime/RigVM/Source/RigVM/Public/RigVMStringUtils.h b/Engine/Plugins/Runtime/RigVM/Source/RigVM/Public/RigVMStringUtils.h
--- a/Engine/Plugins/Runtime/RigVM/Source/RigVM/Public/RigVMStringUtils.h
+++ b/Engine/Plugins/Runtime/RigVM/Source/RigVM/Public/RigVMStringUtils.h
@@ -6,6 +6,20 @@
 
 namespace RigVMStringUtils
 {
+	// Joins two strings given a separator 
+	RIGVM_API FString JoinStrings(const FString& InStringA, const FString& InStringB, const TCHAR* InSeparator, const TCHAR* InPrefix = nullptr, const TCHAR* InSuffix = nullptr);
+
+	// Joins a series of strings given a separator 
+	RIGVM_API FString JoinStrings(const TArray<FString>& InStrings, const TCHAR* InSeparator, const TCHAR* InPrefix = nullptr, const TCHAR* InSuffix = nullptr);
+	RIGVM_API FString JoinStrings(const TArrayView<FString>& InStrings, const TCHAR* InSeparator, const TCHAR* InPrefix = nullptr, const TCHAR* InSuffix = nullptr);
+	RIGVM_API FString JoinStringsConst(const TArrayView<const FString>& InStrings, const TCHAR* InSeparator, const TCHAR* InPrefix = nullptr, const TCHAR* InSuffix = nullptr);
+
+	// Splits a string up into its parts
+	RIGVM_API bool SplitString(const FString& InString, const TCHAR* InSeparator, TArray<FString>& InOutParts);
+
+	// Creates a template notation such as "MyTemplate(float,bool,float)"
+	RIGVM_API FString MakeTemplateNotation(const FString& InTemplateName, const TArray<FString>& InArgumentNotations);
+
 	// Splits a NodePath at the start, so for example "CollapseNodeA|CollapseNodeB|CollapseNodeC" becomes "CollapseNodeA" and "CollapseNodeB|CollapseNodeC"
 	RIGVM_API bool SplitNodePathAtStart(const FString& InNodePath, FString& LeftMost, FString& Right);
 
diff --git a/Engine/Plugins/Runtime/RigVM/Source/RigVM/Public/RigVMTypeUtils.h b/Engine/Plugins/Runtime/RigVM/Source/RigVM/Public/RigVMTypeUtils.h
--- a/Engine/Plugins/Runtime/RigVM/Source/RigVM/Public/RigVMTypeUtils.h
+++ b/Engine/Plugins/Runtime/RigVM/Source/RigVM/Public/RigVMTypeUtils.h
@@ -3,6 +3,7 @@
 #pragma once
 
 #include "CoreMinimal.h"
+#include "RigVMStringUtils.h"
 #include "RigVMCore/RigVMTypeIndex.h"
 #include "RigVMCore/RigVMUnknownType.h"
 #include "UObject/Interface.h"
@@ -131,17 +132,18 @@ namespace RigVMTypeUtils
 
 	inline FString ArrayTypeFromBaseType(const FString& InCPPType)
 	{
-		return FString::Printf(TArrayTemplate, *InCPPType);
+		return RigVMStringUtils::JoinStrings(FString(), InCPPType, TArrayPrefix, nullptr, TEXT(">"));
 	}
 
 	inline FString BaseTypeFromArrayType(const FString& InCPPType)
 	{
-		return InCPPType.RightChop(7).LeftChop(1).TrimStartAndEnd();
+		return InCPPType.Mid(7 /* TArray< */, InCPPType.Len() - 8 /* TArray< and > */);
 	}
 
 	inline FString GetUniqueStructTypeName(const FGuid& InStructGuid)
 	{
-		return FString::Printf(TEXT("FUserDefinedStruct_%s"), *InStructGuid.ToString());
+		static const FString UserDefinedStructPrefix = TEXT("FUserDefinedStruct_");
+		return UserDefinedStructPrefix + InStructGuid.ToString();
 	}
 
 	inline FString GetUniqueStructTypeName(const UScriptStruct* InScriptStruct)
@@ -395,17 +397,17 @@ namespace RigVMTypeUtils
 		if (CPPType.StartsWith(PrefixObjectPtr))
 		{
 			// Chop the prefix + the U indicating object class
-			CPPType = CPPType.RightChop(PrefixObjectPtr.Len() + 1).LeftChop(1);
+			CPPType = CPPType.Mid(PrefixObjectPtr.Len() + 1, CPPType.Len() - (PrefixObjectPtr.Len() + 2));
 		}
 		else if (CPPType.StartsWith(PrefixSubclassOf))
 		{
 			// Chop the prefix + the U indicating object class
-			CPPType = CPPType.RightChop(PrefixSubclassOf.Len() + 1).LeftChop(1);
+			CPPType = CPPType.Mid(PrefixSubclassOf.Len() + 1, CPPType.Len() - (PrefixSubclassOf.Len() + 2));
 		}
 		else if (CPPType.StartsWith(TScriptInterfacePrefix))
 		{
 			// Chop the prefix + the I indicating interface class
-			CPPType = CPPType.RightChop(PrefixScriptInterface.Len() + 1).LeftChop(1);
+			CPPType = CPPType.Mid(PrefixScriptInterface.Len() + 1, CPPType.Len() - (PrefixScriptInterface.Len() + 2));
 		}
 
 		UObject* CPPTypeObject = FindObjectGlobally(*CPPType, bUseRedirector);
