RigVM: Introduce FRigVMType and soft-rollout to template + registry

#rnx

diff --git a/Engine/Plugins/Animation/ControlRig/Source/ControlRig/Private/AnimNode_ControlRig.cpp b/Engine/Plugins/Animation/ControlRig/Source/ControlRig/Private/AnimNode_ControlRig.cpp
--- a/Engine/Plugins/Animation/ControlRig/Source/ControlRig/Private/AnimNode_ControlRig.cpp
+++ b/Engine/Plugins/Animation/ControlRig/Source/ControlRig/Private/AnimNode_ControlRig.cpp
@@ -772,16 +772,16 @@ void FAnimNode_ControlRig::AddVariableFunction(const FRigVMExternalVariable& InV
 	
 	const int32 VariableIndex = Variables.Add(InVariable);
 	
-	if (CastField<FBoolProperty>(InSourceProperty) && InVariable.TypeName == RigVMTypeUtils::BoolTypeName)
+	if (CastField<FBoolProperty>(InSourceProperty) && InVariable.TypeName == RigVMTypeUtils::TypeName::Bool)
 	{
 		return AddUpdateFunction([InSourceProperty, this, VariableIndex](const UObject* InSourceInstance)
 		{
 			SetVariableValue<bool>(Variables[VariableIndex], InSourceProperty, InSourceInstance);
 		});
 	}
-	else if (CastField<FFloatProperty>(InSourceProperty) && (InVariable.TypeName == RigVMTypeUtils::FloatTypeName || InVariable.TypeName == RigVMTypeUtils::DoubleTypeName))
+	else if (CastField<FFloatProperty>(InSourceProperty) && (InVariable.TypeName == RigVMTypeUtils::TypeName::Float || InVariable.TypeName == RigVMTypeUtils::TypeName::Double))
 	{
-		if(InVariable.TypeName == RigVMTypeUtils::FloatTypeName)
+		if(InVariable.TypeName == RigVMTypeUtils::TypeName::Float)
 		{
 			return AddUpdateFunction([InSourceProperty, this, VariableIndex](const UObject* InSourceInstance)
 			{
@@ -796,9 +796,9 @@ void FAnimNode_ControlRig::AddVariableFunction(const FRigVMExternalVariable& InV
 			});
 		}
 	}
-	else if (CastField<FDoubleProperty>(InSourceProperty) && (InVariable.TypeName == RigVMTypeUtils::FloatTypeName || InVariable.TypeName == RigVMTypeUtils::DoubleTypeName))
+	else if (CastField<FDoubleProperty>(InSourceProperty) && (InVariable.TypeName == RigVMTypeUtils::TypeName::Float || InVariable.TypeName == RigVMTypeUtils::TypeName::Double))
 	{
-		if(InVariable.TypeName == RigVMTypeUtils::FloatTypeName)
+		if(InVariable.TypeName == RigVMTypeUtils::TypeName::Float)
 		{
 			return AddUpdateFunction([InSourceProperty, this, VariableIndex](const UObject* InSourceInstance)
 			{
@@ -813,21 +813,21 @@ void FAnimNode_ControlRig::AddVariableFunction(const FRigVMExternalVariable& InV
 			});
 		}
 	}
-	else if (CastField<FIntProperty>(InSourceProperty) && InVariable.TypeName == RigVMTypeUtils::Int32TypeName)
+	else if (CastField<FIntProperty>(InSourceProperty) && InVariable.TypeName == RigVMTypeUtils::TypeName::Int32)
 	{
 		return AddUpdateFunction([InSourceProperty, this, VariableIndex](const UObject* InSourceInstance)
 		{
 			SetVariableValue<int32>(Variables[VariableIndex], InSourceProperty, InSourceInstance);
 		});
 	}
-	else if (CastField<FNameProperty>(InSourceProperty) && InVariable.TypeName == RigVMTypeUtils::FNameTypeName)
+	else if (CastField<FNameProperty>(InSourceProperty) && InVariable.TypeName == RigVMTypeUtils::TypeName::FName)
 	{
 		return AddUpdateFunction([InSourceProperty, this, VariableIndex](const UObject* InSourceInstance)
 		{
 			SetVariableValue<FName>(Variables[VariableIndex], InSourceProperty, InSourceInstance);
 		});
 	}
-	else if (CastField<FNameProperty>(InSourceProperty) && InVariable.TypeName == RigVMTypeUtils::FStringTypeName)
+	else if (CastField<FNameProperty>(InSourceProperty) && InVariable.TypeName == RigVMTypeUtils::TypeName::FString)
 	{
 		return AddUpdateFunction([InSourceProperty, this, VariableIndex](const UObject* InSourceInstance)
 		{
@@ -1115,15 +1115,15 @@ void FAnimNode_ControlRig::PropagateInputPropertiesNoCache(
 
 			const uint8* SrcPtr = CallerProperty->ContainerPtrToValuePtr<uint8>(InSourceInstance);
 
-			if (CastField<FBoolProperty>(CallerProperty) != nullptr && Variable.TypeName == RigVMTypeUtils::BoolTypeName)
+			if (CastField<FBoolProperty>(CallerProperty) != nullptr && Variable.TypeName == RigVMTypeUtils::TypeName::Bool)
 			{
 				const bool Value = *(const bool*)SrcPtr;
 				Variable.SetValue<bool>(Value);
 			}
-			else if (CastField<FFloatProperty>(CallerProperty) != nullptr && (Variable.TypeName == RigVMTypeUtils::FloatTypeName ||  Variable.TypeName == RigVMTypeUtils::DoubleTypeName))
+			else if (CastField<FFloatProperty>(CallerProperty) != nullptr && (Variable.TypeName == RigVMTypeUtils::TypeName::Float ||  Variable.TypeName == RigVMTypeUtils::TypeName::Double))
 			{
 				const float Value = *(const float*)SrcPtr;
-				if(Variable.TypeName == RigVMTypeUtils::FloatTypeName)
+				if(Variable.TypeName == RigVMTypeUtils::TypeName::Float)
 				{
 					Variable.SetValue<float>(Value);
 				}
@@ -1132,10 +1132,10 @@ void FAnimNode_ControlRig::PropagateInputPropertiesNoCache(
 					Variable.SetValue<double>(Value);
 				}
 			}
-			else if (CastField<FDoubleProperty>(CallerProperty) != nullptr && (Variable.TypeName == RigVMTypeUtils::FloatTypeName ||  Variable.TypeName == RigVMTypeUtils::DoubleTypeName))
+			else if (CastField<FDoubleProperty>(CallerProperty) != nullptr && (Variable.TypeName == RigVMTypeUtils::TypeName::Float ||  Variable.TypeName == RigVMTypeUtils::TypeName::Double))
 			{
 				const double Value = *(const double*)SrcPtr;
-				if(Variable.TypeName == RigVMTypeUtils::FloatTypeName)
+				if(Variable.TypeName == RigVMTypeUtils::TypeName::Float)
 				{
 					Variable.SetValue<float>((float)Value);
 				}
@@ -1144,17 +1144,17 @@ void FAnimNode_ControlRig::PropagateInputPropertiesNoCache(
 					Variable.SetValue<double>(Value);
 				}
 			}
-			else if (CastField<FIntProperty>(CallerProperty) != nullptr && Variable.TypeName == RigVMTypeUtils::Int32TypeName)
+			else if (CastField<FIntProperty>(CallerProperty) != nullptr && Variable.TypeName == RigVMTypeUtils::TypeName::Int32)
 			{
 				const int32 Value = *(const int32*)SrcPtr;
 				Variable.SetValue<int32>(Value);
 			}
-			else if (CastField<FNameProperty>(CallerProperty) != nullptr && Variable.TypeName == RigVMTypeUtils::FNameTypeName)
+			else if (CastField<FNameProperty>(CallerProperty) != nullptr && Variable.TypeName == RigVMTypeUtils::TypeName::FName)
 			{
 				const FName Value = *(const FName*)SrcPtr;
 				Variable.SetValue<FName>(Value);
 			}
-			else if (CastField<FNameProperty>(CallerProperty) != nullptr && Variable.TypeName == RigVMTypeUtils::FStringTypeName)
+			else if (CastField<FNameProperty>(CallerProperty) != nullptr && Variable.TypeName == RigVMTypeUtils::TypeName::FString)
 			{
 				const FString Value = *(const FString*)SrcPtr;
 				Variable.SetValue<FString>(Value);
diff --git a/Engine/Plugins/Animation/ControlRig/Source/ControlRig/Private/Units/Animation/RigUnit_AnimAttribute.cpp b/Engine/Plugins/Animation/ControlRig/Source/ControlRig/Private/Units/Animation/RigUnit_AnimAttribute.cpp
--- a/Engine/Plugins/Animation/ControlRig/Source/ControlRig/Private/Units/Animation/RigUnit_AnimAttribute.cpp
+++ b/Engine/Plugins/Animation/ControlRig/Source/ControlRig/Private/Units/Animation/RigUnit_AnimAttribute.cpp
@@ -32,9 +32,9 @@ bool FRigDispatch_AnimAttributeBase::IsTypeSupported(const TRigVMTypeIndex& InTy
 
 	// cpp type object can become invalid because users can choose to delete
 	// user defined structs
-	if (IsValid(InType.CPPTypeObject))
+	if (IsValid(InType.GetTypeObject()))
 	{
-		if (UScriptStruct* ScriptStruct = Cast<UScriptStruct>(InType.CPPTypeObject))
+		if (UScriptStruct* ScriptStruct = Cast<UScriptStruct>(InType.GetTypeObject()))
 		{
 			if (ScriptStruct->IsA<UUserDefinedStruct>())
 			{
@@ -291,9 +291,9 @@ FString FRigDispatch_AnimAttributeBase::GetNodeTitle(const FRigVMTemplateTypeMap
 			const FRigVMTemplateArgumentType& ValueType = FRigVMRegistry::Get().GetType(ValueTypeIndex);
 			FString ValueName;
 
-			if(ValueType.CPPTypeObject)
+			if(ValueType.GetTypeObject())
 			{
-				ValueName = ValueType.CPPTypeObject->GetName();
+				ValueName = ValueType.GetTypeObject()->GetName();
 			}
 			else
 			{
@@ -437,7 +437,7 @@ FRigVMFunctionPtr FRigDispatch_GetAnimAttribute::GetDispatchFunctionImpl(const F
 	
 	const FRigVMTemplateArgumentType& ValueType = Registry.GetType_NoLock(ValueTypeIndex);
 	
-	if (UScriptStruct* ScriptStruct = Cast<UScriptStruct>(ValueType.CPPTypeObject))
+	if (UScriptStruct* ScriptStruct = Cast<UScriptStruct>(ValueType.GetTypeObject()))
 	{
 		return &FRigDispatch_GetAnimAttribute::GetAnimAttributeDispatch;
 	}
@@ -526,7 +526,7 @@ FRigVMFunctionPtr FRigDispatch_SetAnimAttribute::GetDispatchFunctionImpl(const F
 	
 	const FRigVMTemplateArgumentType& ValueType = Registry.GetType_NoLock(ValueTypeIndex);
 	
-	if (UScriptStruct* ScriptStruct = Cast<UScriptStruct>(ValueType.CPPTypeObject))
+	if (UScriptStruct* ScriptStruct = Cast<UScriptStruct>(ValueType.GetTypeObject()))
 	{
 		return &FRigDispatch_SetAnimAttribute::SetAnimAttributeDispatch;
 	}
diff --git a/Engine/Plugins/Animation/ControlRig/Source/ControlRig/Private/Units/Hierarchy/RigUnit_Component.cpp b/Engine/Plugins/Animation/ControlRig/Source/ControlRig/Private/Units/Hierarchy/RigUnit_Component.cpp
--- a/Engine/Plugins/Animation/ControlRig/Source/ControlRig/Private/Units/Hierarchy/RigUnit_Component.cpp
+++ b/Engine/Plugins/Animation/ControlRig/Source/ControlRig/Private/Units/Hierarchy/RigUnit_Component.cpp
@@ -18,9 +18,9 @@ bool FRigDispatch_ComponentBase::IsTypeSupported(const TRigVMTypeIndex& InTypeIn
 	const FRigVMRegistry_NoLock& Registry = FRigVMRegistry_NoLock::GetForRead();
 	const FRigVMTemplateArgumentType& InType = Registry.GetType_NoLock(InTypeIndex);
 
-	if (IsValid(InType.CPPTypeObject))
+	if (IsValid(InType.GetTypeObject()))
 	{
-		if (const UScriptStruct* ScriptStruct = Cast<UScriptStruct>(InType.CPPTypeObject))
+		if (const UScriptStruct* ScriptStruct = Cast<UScriptStruct>(InType.GetTypeObject()))
 		{
 			return ScriptStruct != FRigBaseComponent::StaticStruct() &&
 				ScriptStruct->IsChildOf(FRigBaseComponent::StaticStruct());
diff --git a/Engine/Plugins/Animation/ControlRig/Source/ControlRig/Private/Units/Hierarchy/RigUnit_Metadata.cpp b/Engine/Plugins/Animation/ControlRig/Source/ControlRig/Private/Units/Hierarchy/RigUnit_Metadata.cpp
--- a/Engine/Plugins/Animation/ControlRig/Source/ControlRig/Private/Units/Hierarchy/RigUnit_Metadata.cpp
+++ b/Engine/Plugins/Animation/ControlRig/Source/ControlRig/Private/Units/Hierarchy/RigUnit_Metadata.cpp
@@ -25,13 +25,13 @@ FString FRigDispatch_MetadataBase::GetNodeTitle(const FRigVMTemplateTypeMap& InT
 		{
 			const FRigVMTemplateArgumentType& ValueType = FRigVMRegistry::Get().GetType(ValueTypeIndex);
 			FString ValueName;
-			if(ValueType.CPPTypeObject == FRigElementKey::StaticStruct())
+			if(ValueType.GetTypeObject() == FRigElementKey::StaticStruct())
 			{
 				ValueName = ItemArgName.ToString();
 			}
-			else if(ValueType.CPPTypeObject)
+			else if(ValueType.GetTypeObject())
 			{
-				ValueName = ValueType.CPPTypeObject->GetName();
+				ValueName = ValueType.GetTypeObject()->GetName();
 			}
 			else if(ValueTypeIndex == RigVMTypeUtils::TypeIndex::FName)
 			{
diff --git a/Engine/Plugins/Animation/ControlRig/Source/ControlRigEditor/Private/Editor/SRigHierarchyTreeView.cpp b/Engine/Plugins/Animation/ControlRig/Source/ControlRigEditor/Private/Editor/SRigHierarchyTreeView.cpp
--- a/Engine/Plugins/Animation/ControlRig/Source/ControlRigEditor/Private/Editor/SRigHierarchyTreeView.cpp
+++ b/Engine/Plugins/Animation/ControlRig/Source/ControlRigEditor/Private/Editor/SRigHierarchyTreeView.cpp
@@ -1281,13 +1281,13 @@ FLinearColor SRigHierarchyItem::GetColorForControlType(ERigControlType InControl
 	{
 		case ERigControlType::Bool:
 		{
-			PinType = RigVMTypeUtils::PinTypeFromCPPType(RigVMTypeUtils::BoolTypeName, nullptr);
+			PinType = RigVMTypeUtils::PinTypeFromCPPType(RigVMTypeUtils::TypeName::Bool, nullptr);
 			break;
 		}
 		case ERigControlType::Float:
 		case ERigControlType::ScaleFloat:
 		{
-			PinType = RigVMTypeUtils::PinTypeFromCPPType(RigVMTypeUtils::FloatTypeName, nullptr);
+			PinType = RigVMTypeUtils::PinTypeFromCPPType(RigVMTypeUtils::TypeName::Float, nullptr);
 			break;
 		}
 		case ERigControlType::Integer:
@@ -1298,7 +1298,7 @@ FLinearColor SRigHierarchyItem::GetColorForControlType(ERigControlType InControl
 			}
 			else
 			{
-				PinType = RigVMTypeUtils::PinTypeFromCPPType(RigVMTypeUtils::Int32TypeName, nullptr);
+				PinType = RigVMTypeUtils::PinTypeFromCPPType(RigVMTypeUtils::TypeName::Int32, nullptr);
 			}
 			break;
 		}
diff --git a/Engine/Plugins/Animation/ControlRig/Source/ControlRigEditor/Private/Graph/ControlRigGraphPanelPinFactory.cpp b/Engine/Plugins/Animation/ControlRig/Source/ControlRigEditor/Private/Graph/ControlRigGraphPanelPinFactory.cpp
--- a/Engine/Plugins/Animation/ControlRig/Source/ControlRigEditor/Private/Graph/ControlRigGraphPanelPinFactory.cpp
+++ b/Engine/Plugins/Animation/ControlRig/Source/ControlRigEditor/Private/Graph/ControlRigGraphPanelPinFactory.cpp
@@ -348,7 +348,7 @@ TSharedPtr<SGraphPin> FControlRigGraphPanelPinFactory::CreatePin_Internal(UEdGra
 								// find the pin that holds the control
 								for(URigVMPin* Pin : InPin->GetRootPin()->GetNode()->GetPins())
 								{
-									if(Pin->GetCPPType() == RigVMTypeUtils::FNameType && Pin->GetCustomWidgetName() == TEXT("ControlName"))
+									if(Pin->GetCPPType() == RigVMTypeUtils::TypeString::FName && Pin->GetCustomWidgetName() == TEXT("ControlName"))
 									{
 										const FString DefaultValue = Pin->GetDefaultValue();
 										const FName ControlName = DefaultValue.IsEmpty() ? FName(NAME_None) : FName(*DefaultValue);
diff --git a/Engine/Plugins/Experimental/AnimNext/Source/AnimNext/Private/Param/ParamType.cpp b/Engine/Plugins/Experimental/AnimNext/Source/AnimNext/Private/Param/ParamType.cpp
--- a/Engine/Plugins/Experimental/AnimNext/Source/AnimNext/Private/Param/ParamType.cpp
+++ b/Engine/Plugins/Experimental/AnimNext/Source/AnimNext/Private/Param/ParamType.cpp
@@ -24,23 +24,23 @@ FRigVMTemplateArgumentType FAnimNextParamType::ToRigVMTemplateArgument() const
 		case EValueType::None:
 			return FRigVMTemplateArgumentType();
 		case EValueType::Bool:
-			return FRigVMTemplateArgumentType(RigVMTypeUtils::BoolTypeName);
+			return FRigVMTemplateArgumentType(RigVMTypeUtils::TypeName::Bool);
 		case EValueType::Byte:
-			return FRigVMTemplateArgumentType(RigVMTypeUtils::UInt8Type);
+			return FRigVMTemplateArgumentType(RigVMTypeUtils::TypeName::UInt8);
 		case EValueType::Int32:
-			return FRigVMTemplateArgumentType(RigVMTypeUtils::Int32Type);
+			return FRigVMTemplateArgumentType(RigVMTypeUtils::TypeName::Int32);
 		case EValueType::Int64:
-			return FRigVMTemplateArgumentType(RigVMTypeUtils::Int64Type);
+			return FRigVMTemplateArgumentType(RigVMTypeUtils::TypeName::Int64);
 		case EValueType::Float:
-			return FRigVMTemplateArgumentType(RigVMTypeUtils::FloatType);
+			return FRigVMTemplateArgumentType(RigVMTypeUtils::TypeName::Float);
 		case EValueType::Double:
-			return FRigVMTemplateArgumentType(RigVMTypeUtils::DoubleType);
+			return FRigVMTemplateArgumentType(RigVMTypeUtils::TypeName::Double);
 		case EValueType::Name:
-			return FRigVMTemplateArgumentType(RigVMTypeUtils::FNameType);
+			return FRigVMTemplateArgumentType(RigVMTypeUtils::TypeName::FName);
 		case EValueType::String:
-			return FRigVMTemplateArgumentType(RigVMTypeUtils::FStringType);
+			return FRigVMTemplateArgumentType(RigVMTypeUtils::TypeName::FString);
 		case EValueType::Text:
-			return FRigVMTemplateArgumentType(RigVMTypeUtils::FTextType);
+			return FRigVMTemplateArgumentType(RigVMTypeUtils::TypeName::FText);
 		case EValueType::Enum:
 			if(const UEnum* Enum = Cast<UEnum>(ValueTypeObject.Get()))
 			{
@@ -89,70 +89,66 @@ FRigVMTemplateArgumentType FAnimNextParamType::ToRigVMTemplateArgument() const
 FAnimNextParamType FAnimNextParamType::FromRigVMTemplateArgument(const FRigVMTemplateArgumentType& RigVMType)
 {
 	FAnimNextParamType Type;	
-	const FString CPPTypeString = RigVMType.CPPType.ToString();
+	const FString CPPTypeString = RigVMType.GetTypeString();
 	
 	if (RigVMTypeUtils::IsArrayType(CPPTypeString))
 	{
 		Type.ContainerType = EPropertyBagContainerType::Array;
 	}
 
-	static const FName IntTypeName(TEXT("int")); // type used by some engine tests
-	static const FName Int64TypeName(TEXT("Int64"));
-	static const FName UInt64TypeName(TEXT("UInt64"));
-
 	const FName CPPType = *CPPTypeString;
 	
-	if (CPPType == RigVMTypeUtils::BoolTypeName)
+	if (CPPType == RigVMTypeUtils::TypeName::Bool)
 	{
 		Type.ValueType = EPropertyBagPropertyType::Bool;
 	}
-	else if (CPPType == RigVMTypeUtils::UInt8Type)
+	else if (CPPType == RigVMTypeUtils::TypeName::UInt8)
 	{
 		Type.ValueType = EPropertyBagPropertyType::Byte;
 	}
-	else if (CPPType == RigVMTypeUtils::Int32TypeName || CPPType == IntTypeName)
+	else if (CPPType == RigVMTypeUtils::TypeName::Int32 || CPPType == RigVMTypeUtils::TypeName::Int)
 	{
 		Type.ValueType = EPropertyBagPropertyType::Int32;
 	}
-	else if (CPPType == RigVMTypeUtils::UInt32TypeName)
+	else if (CPPType == RigVMTypeUtils::TypeName::UInt32)
 	{
 		Type.ValueType = EPropertyBagPropertyType::UInt32;
 	}
-	else if (CPPType == Int64TypeName)
+	else if (CPPType == RigVMTypeUtils::TypeName::Int64)
 	{
 		Type.ValueType = EPropertyBagPropertyType::Int64;
 	}
-	else if (CPPType == UInt64TypeName)
+	else if (CPPType == RigVMTypeUtils::TypeName::UInt64)
 	{
 		Type.ValueType = EPropertyBagPropertyType::UInt64;
 	}
-	else if (CPPType == RigVMTypeUtils::FloatTypeName)
+	else if (CPPType == RigVMTypeUtils::TypeName::Float)
 	{
 		Type.ValueType = EPropertyBagPropertyType::Float;
 	}
-	else if (CPPType == RigVMTypeUtils::DoubleTypeName)
+	else if (CPPType == RigVMTypeUtils::TypeName::Double)
 	{
 		Type.ValueType = EPropertyBagPropertyType::Double;
 	}
-	else if (CPPType == RigVMTypeUtils::FNameTypeName)
+	else if (CPPType == RigVMTypeUtils::TypeName::FName)
 	{
 		Type.ValueType = EPropertyBagPropertyType::Name;
 	}
-	else if (CPPType == RigVMTypeUtils::FStringTypeName)
+	else if (CPPType == RigVMTypeUtils::TypeName::FString)
 	{
 		Type.ValueType = EPropertyBagPropertyType::String;
 	}
-	else if (UScriptStruct* ScriptStruct = Cast<UScriptStruct>(RigVMType.CPPTypeObject))
+	else if (UScriptStruct* ScriptStruct = Cast<UScriptStruct>(RigVMType.GetTypeObject()))
 	{
 		Type.ValueType = EPropertyBagPropertyType::Struct;
 		Type.ValueTypeObject = ScriptStruct;
 	}
-	else if (UEnum* Enum = Cast<UEnum>(RigVMType.CPPTypeObject))
+	else if (UEnum* Enum = Cast<UEnum>(RigVMType.GetTypeObject()))
 	{
 		Type.ValueType = EPropertyBagPropertyType::Enum;
 		Type.ValueTypeObject = Enum;
 	}
-	else if (UObject* Object = Cast<UObject>(RigVMType.CPPTypeObject))
+	else if (UObject* Object = Cast<UObject>(RigVMType.GetTypeObject()))
 	{
 		Type.ValueType = EPropertyBagPropertyType::Object;	
 		Type.ValueTypeObject = Object;
diff --git a/Engine/Plugins/Experimental/AnimNext/Source/AnimNextTestSuite/Private/AnimNextVariablesTest.cpp b/Engine/Plugins/Experimental/AnimNext/Source/AnimNextTestSuite/Private/AnimNextVariablesTest.cpp
--- a/Engine/Plugins/Experimental/AnimNext/Source/AnimNextTestSuite/Private/AnimNextVariablesTest.cpp
+++ b/Engine/Plugins/Experimental/AnimNext/Source/AnimNextTestSuite/Private/AnimNextVariablesTest.cpp
@@ -253,11 +253,11 @@ bool FVariables::RunTest(const FString& InParameters)
 
 	URigVMController* Controller = EditorData->GetController(EventGraph->GetRigVMGraph());
 
-	URigVMVariableNode* VariableANode = Controller->AddVariableNode("A", RigVMTypeUtils::Int32Type, nullptr, true, TEXT(""));
+	URigVMVariableNode* VariableANode = Controller->AddVariableNode("A", RigVMTypeUtils::TypeString::Int32, nullptr, true, TEXT(""));
 	UE_RETURN_ON_ERROR(VariableANode != nullptr, TEXT("Could not add get variable node."));
-	URigVMVariableNode* VariableBNode = Controller->AddVariableNode("B", RigVMTypeUtils::Int32Type, nullptr, true, TEXT(""));
+	URigVMVariableNode* VariableBNode = Controller->AddVariableNode("B", RigVMTypeUtils::TypeString::Int32, nullptr, true, TEXT(""));
 	UE_RETURN_ON_ERROR(VariableBNode != nullptr, TEXT("Could not add get variable node."));
-	URigVMVariableNode* SetResultNode = Controller->AddVariableNode("Result", RigVMTypeUtils::Int32Type, nullptr, false, TEXT(""));
+	URigVMVariableNode* SetResultNode = Controller->AddVariableNode("Result", RigVMTypeUtils::TypeString::Int32, nullptr, false, TEXT(""));
 	UE_RETURN_ON_ERROR(SetResultNode != nullptr, TEXT("Could not add set variable node."));
 
 	URigVMUnitNode* TestOpUnitNode = Controller->AddUnitNode(FAnimNextTests_TestOperation::StaticStruct());
@@ -278,7 +278,7 @@ bool FVariables::RunTest(const FString& InParameters)
 	bool bLinkExec3Added = Controller->AddLink(SetResultNode->FindPin(FRigVMStruct::ExecuteContextName.ToString()), PrintResultUnitNode->FindPin(FRigVMStruct::ExecuteContextName.ToString()));
 	UE_RETURN_ON_ERROR(bLinkExec3Added, TEXT("Could not link variable node exec."));
 
-	URigVMVariableNode* GetResultNode = Controller->AddVariableNode("Result", RigVMTypeUtils::Int32Type, nullptr, true, TEXT(""));
+	URigVMVariableNode* GetResultNode = Controller->AddVariableNode("Result", RigVMTypeUtils::TypeString::Int32, nullptr, true, TEXT(""));
 	UE_RETURN_ON_ERROR(GetResultNode != nullptr, TEXT("Could not add get variable node."));
 	bool bLinkResult2Added = Controller->AddLink(GetResultNode->FindPin("Value"), PrintResultUnitNode->FindPin("Result"));
 	UE_RETURN_ON_ERROR(bLinkResult2Added, TEXT("Could not link variable node."));
@@ -368,11 +368,11 @@ bool FVariables_UOLBindings::RunTest(const FString& InParameters)
 
 	URigVMController* Controller = EditorData->GetController(EventGraph->GetRigVMGraph());
 
-	URigVMVariableNode* PropertyEntryNode = Controller->AddVariableNode("Property", RigVMTypeUtils::Int32Type, nullptr, true, TEXT(""));
+	URigVMVariableNode* PropertyEntryNode = Controller->AddVariableNode("Property", RigVMTypeUtils::TypeString::Int32, nullptr, true, TEXT(""));
 	UE_RETURN_ON_ERROR(PropertyEntryNode != nullptr, TEXT("Could not add get variable node."));
-	URigVMVariableNode* ObjectAccessorEntryNode = Controller->AddVariableNode("ObjectAccessor", RigVMTypeUtils::Int32Type, nullptr, true, TEXT(""));
+	URigVMVariableNode* ObjectAccessorEntryNode = Controller->AddVariableNode("ObjectAccessor", RigVMTypeUtils::TypeString::Int32, nullptr, true, TEXT(""));
 	UE_RETURN_ON_ERROR(ObjectAccessorEntryNode != nullptr, TEXT("Could not add get variable node."));
-	URigVMVariableNode* HoistedAccessorEntryNode = Controller->AddVariableNode("HoistedAccessor", RigVMTypeUtils::Int32Type, nullptr, true, TEXT(""));
+	URigVMVariableNode* HoistedAccessorEntryNode = Controller->AddVariableNode("HoistedAccessor", RigVMTypeUtils::TypeString::Int32, nullptr, true, TEXT(""));
 	UE_RETURN_ON_ERROR(HoistedAccessorEntryNode != nullptr, TEXT("Could not add get variable node."));
 
 	URigVMUnitNode* PrintPropertyUnitNode = Controller->AddUnitNode(FAnimNextTests_PrintResult::StaticStruct());
diff --git a/Engine/Plugins/Experimental/AnimNext/Source/AnimNextUncookedOnly/Private/UncookedOnlyUtils.cpp b/Engine/Plugins/Experimental/AnimNext/Source/AnimNextUncookedOnly/Private/UncookedOnlyUtils.cpp
--- a/Engine/Plugins/Experimental/AnimNext/Source/AnimNextUncookedOnly/Private/UncookedOnlyUtils.cpp
+++ b/Engine/Plugins/Experimental/AnimNext/Source/AnimNextUncookedOnly/Private/UncookedOnlyUtils.cpp
@@ -325,7 +325,7 @@ void FUtils::CompileVariableBindings(const FRigVMCompileSettings& InSettings, UA
 
 		TArray<IVariableBindingType::FBindingGraphInput>& Group = BindingGroups.FindOrAdd(BindingType.Get());
 		FRigVMTemplateArgumentType RigVMArg = Variable->GetType().ToRigVMTemplateArgument();
-		Group.Add({ Variable->GetVariableName(), RigVMArg.CPPType.ToString(), RigVMArg.CPPTypeObject, Binding });
+		Group.Add({ Variable->GetVariableName(), RigVMArg.GetTypeString(), RigVMArg.GetTypeObject(), Binding });
 	}
 
 	const bool bHasBindings = BindingGroups.Num() > 0;
@@ -588,58 +588,57 @@ FEdGraphPinType FUtils::GetPinTypeFromParamType(const FAnimNextParamType& InPara
 
 FRigVMTemplateArgumentType FUtils::GetRigVMArgTypeFromParamType(const FAnimNextParamType& InParamType)
 {
-	FRigVMTemplateArgumentType ArgType;
-
-	FString CPPTypeString;
+	FString TypeString;
+	UObject* TypeObject = nullptr;
 
 	// Value type
 	switch (InParamType.ValueType)
 	{
 	case EPropertyBagPropertyType::Bool:
-		CPPTypeString = RigVMTypeUtils::BoolType;
+		TypeString = RigVMTypeUtils::TypeString::Bool;
 		break;
 	case EPropertyBagPropertyType::Byte:
-		CPPTypeString = RigVMTypeUtils::UInt8Type;
+		TypeString = RigVMTypeUtils::TypeString::UInt8;
 		break;
 	case EPropertyBagPropertyType::Int32:
-		CPPTypeString = RigVMTypeUtils::UInt32Type;
+		TypeString = RigVMTypeUtils::TypeString::UInt32;
 		break;
 	case EPropertyBagPropertyType::Int64:
 		ensureMsgf(false, TEXT("Unhandled value type %d"), InParamType.ValueType);
 		break;
 	case EPropertyBagPropertyType::Float:
-		CPPTypeString = RigVMTypeUtils::FloatType;
+		TypeString = RigVMTypeUtils::TypeString::Float;
 		break;
 	case EPropertyBagPropertyType::Double:
-		CPPTypeString = RigVMTypeUtils::DoubleType;
+		TypeString = RigVMTypeUtils::TypeString::Double;
 		break;
 	case EPropertyBagPropertyType::Name:
-		CPPTypeString = RigVMTypeUtils::FNameType;
+		TypeString = RigVMTypeUtils::TypeString::FName;
 		break;
 	case EPropertyBagPropertyType::String:
-		CPPTypeString = RigVMTypeUtils::FStringType;
+		TypeString = RigVMTypeUtils::TypeString::FString;
 		break;
 	case EPropertyBagPropertyType::Text:
 		ensureMsgf(false, TEXT("Unhandled value type %d"), InParamType.ValueType);
 		break;
 	case EPropertyBagPropertyType::Enum:
-		CPPTypeString = RigVMTypeUtils::CPPTypeFromEnum(Cast<UEnum>(InParamType.ValueTypeObject.Get()));
-		ArgType.CPPTypeObject = const_cast<UObject*>(InParamType.ValueTypeObject.Get());
+		TypeString = RigVMTypeUtils::CPPTypeFromEnum(Cast<UEnum>(InParamType.ValueTypeObject.Get()));
+		TypeObject = const_cast<UObject*>(InParamType.ValueTypeObject.Get());
 		break;
 	case EPropertyBagPropertyType::Struct:
-		CPPTypeString = RigVMTypeUtils::GetUniqueStructTypeName(Cast<UScriptStruct>(InParamType.ValueTypeObject.Get()));
-		ArgType.CPPTypeObject = const_cast<UObject*>(InParamType.ValueTypeObject.Get());
+		TypeString = RigVMTypeUtils::GetUniqueStructTypeName(Cast<UScriptStruct>(InParamType.ValueTypeObject.Get()));
+		TypeObject = const_cast<UObject*>(InParamType.ValueTypeObject.Get());
 		break;
 	case EPropertyBagPropertyType::Object:
-		CPPTypeString = RigVMTypeUtils::CPPTypeFromObject(Cast<UClass>(InParamType.ValueTypeObject.Get()), RigVMTypeUtils::EClassArgType::AsObject);
-		ArgType.CPPTypeObject = const_cast<UObject*>(InParamType.ValueTypeObject.Get());
+		TypeString = RigVMTypeUtils::CPPTypeFromObject(Cast<UClass>(InParamType.ValueTypeObject.Get()), RigVMTypeUtils::EClassArgType::AsObject);
+		TypeObject = const_cast<UObject*>(InParamType.ValueTypeObject.Get());
 		break;
 	case EPropertyBagPropertyType::SoftObject:
 		ensureMsgf(false, TEXT("Unhandled value type %d"), InParamType.ValueType);
 		break;
 	case EPropertyBagPropertyType::Class:
-		CPPTypeString = RigVMTypeUtils::CPPTypeFromObject(Cast<UClass>(InParamType.ValueTypeObject.Get()), RigVMTypeUtils::EClassArgType::AsClass);
-		ArgType.CPPTypeObject = const_cast<UObject*>(InParamType.ValueTypeObject.Get());
+		TypeString = RigVMTypeUtils::CPPTypeFromObject(Cast<UClass>(InParamType.ValueTypeObject.Get()), RigVMTypeUtils::EClassArgType::AsClass);
+		TypeObject = const_cast<UObject*>(InParamType.ValueTypeObject.Get());
 		break;
 	case EPropertyBagPropertyType::SoftClass:
 		ensureMsgf(false, TEXT("Unhandled value type %d"), InParamType.ValueType);
@@ -655,16 +654,14 @@ FRigVMTemplateArgumentType FUtils::GetRigVMArgTypeFromParamType(const FAnimNextP
 	case FAnimNextParamType::EContainerType::None:
 		break;
 	case FAnimNextParamType::EContainerType::Array:
-		CPPTypeString = FString::Printf(RigVMTypeUtils::TArrayTemplate, *CPPTypeString);
+		TypeString = FString::Printf(RigVMTypeUtils::TArrayTemplate, *TypeString);
 		break;
 	default:
 		ensureMsgf(false, TEXT("Unhandled container type %d"), InParamType.ContainerType);
 		break;
 	}
 
-	ArgType.CPPType = *CPPTypeString;
-
-	return ArgType;
+	return FRigVMTemplateArgumentType(*TypeString, TypeObject);
 }
 
 void FUtils::SetupEventGraph(URigVMController* InController, UScriptStruct* InEventStruct, FName InEventName)
diff --git a/Engine/Plugins/Experimental/AnimNextAnimGraph/Source/AnimNextAnimGraphTestSuite/Private/AnimNextAnimGraphTraitGraphTest.cpp b/Engine/Plugins/Experimental/AnimNextAnimGraph/Source/AnimNextAnimGraphTestSuite/Private/AnimNextAnimGraphTraitGraphTest.cpp
--- a/Engine/Plugins/Experimental/AnimNextAnimGraph/Source/AnimNextAnimGraphTestSuite/Private/AnimNextAnimGraphTraitGraphTest.cpp
+++ b/Engine/Plugins/Experimental/AnimNextAnimGraph/Source/AnimNextAnimGraphTestSuite/Private/AnimNextAnimGraphTraitGraphTest.cpp
@@ -646,7 +646,7 @@ bool FAnimationAnimNextRuntimeTest_GraphExecuteLatent::RunTest(const FString& In
 		}
 
 		{
-			URigVMVariableNode* GetVariableNode = Controller->AddVariableNode(TEXT("TestIntVar"), RigVMTypeUtils::Int32Type, nullptr, true, TEXT(""));
+			URigVMVariableNode* GetVariableNode = Controller->AddVariableNode(TEXT("TestIntVar"), RigVMTypeUtils::TypeString::Int32, nullptr, true, TEXT(""));
 			UE_RETURN_ON_ERROR(GetVariableNode != nullptr, TEXT("FAnimationAnimNextRuntimeTest_GraphExecuteLatent -> Failed to create variable node"));
 
 			Controller->AddLink(
diff --git a/Engine/Plugins/Experimental/AnimNextAnimGraph/Source/AnimNextAnimGraphUncookedOnly/Private/Graph/AnimNextAnimationGraph_EditorData.cpp b/Engine/Plugins/Experimental/AnimNextAnimGraph/Source/AnimNextAnimGraphUncookedOnly/Private/Graph/AnimNextAnimationGraph_EditorData.cpp
--- a/Engine/Plugins/Experimental/AnimNextAnimGraph/Source/AnimNextAnimGraphUncookedOnly/Private/Graph/AnimNextAnimationGraph_EditorData.cpp
+++ b/Engine/Plugins/Experimental/AnimNextAnimGraph/Source/AnimNextAnimGraphUncookedOnly/Private/Graph/AnimNextAnimationGraph_EditorData.cpp
@@ -442,7 +442,7 @@ namespace UE::AnimNext::UncookedOnly::Private
 
 			FAnimNextGraphEvaluatorExecuteArgument Argument;
 			Argument.Name = Pin.Name.ToString();
-			Argument.CPPType = TypeArg.CPPType.ToString();
+			Argument.CPPType = TypeArg.GetTypeString();
 
 			ExecuteDefinition.Arguments.Add(Argument);
 		}
diff --git a/Engine/Plugins/Experimental/AnimNextStateTree/Source/AnimNextStateTree/Private/Conditions/AnimNextStateTreeRigVMConditionBase.cpp b/Engine/Plugins/Experimental/AnimNextStateTree/Source/AnimNextStateTree/Private/Conditions/AnimNextStateTreeRigVMConditionBase.cpp
--- a/Engine/Plugins/Experimental/AnimNextStateTree/Source/AnimNextStateTree/Private/Conditions/AnimNextStateTreeRigVMConditionBase.cpp
+++ b/Engine/Plugins/Experimental/AnimNextStateTree/Source/AnimNextStateTree/Private/Conditions/AnimNextStateTreeRigVMConditionBase.cpp
@@ -259,7 +259,7 @@ void FAnimNextStateTreeRigVMConditionBase::PostEditNodeChangeChainProperty(const
 				}
 				else if (Argument.Direction == ERigVMPinDirection::Output)
 				{
-					ensure(Argument.CPPType == RigVMTypeUtils::BoolTypeName);
+					ensure(Argument.CPPType == RigVMTypeUtils::TypeName::Bool);
 					ResultName = Argument.Name;
 				}
 			}
diff --git a/Engine/Plugins/Runtime/RigVM/Source/RigVM/Private/RigVMCore/RigVMExternalVariable.cpp b/Engine/Plugins/Runtime/RigVM/Source/RigVM/Private/RigVMCore/RigVMExternalVariable.cpp
--- a/Engine/Plugins/Runtime/RigVM/Source/RigVM/Private/RigVMCore/RigVMExternalVariable.cpp
+++ b/Engine/Plugins/Runtime/RigVM/Source/RigVM/Private/RigVMCore/RigVMExternalVariable.cpp
@@ -53,12 +53,12 @@ FRigVMExternalVariable RigVMTypeUtils::ExternalVariableFromPinType(const FName&
 
 	if (InPinType.PinCategory == UEdGraphSchema_K2::PC_Boolean)
 	{
-		ExternalVariable.TypeName = BoolTypeName;
+		ExternalVariable.TypeName = TypeName::Bool;
 		ExternalVariable.Size = sizeof(bool);
 	}
 	else if (InPinType.PinCategory == UEdGraphSchema_K2::PC_Int)
 	{
-		ExternalVariable.TypeName = Int32TypeName;
+		ExternalVariable.TypeName = TypeName::Int32;
 		ExternalVariable.Size = sizeof(int32);
 	}
 	else if (InPinType.PinCategory == UEdGraphSchema_K2::PC_Enum ||
@@ -71,7 +71,7 @@ FRigVMExternalVariable RigVMTypeUtils::ExternalVariableFromPinType(const FName&
 		}
 		else
 		{
-			ExternalVariable.TypeName = UInt8TypeName;
+			ExternalVariable.TypeName = TypeName::UInt8;
 		}
 		ExternalVariable.Size = sizeof(uint8);
 	}
@@ -79,12 +79,12 @@ FRigVMExternalVariable RigVMTypeUtils::ExternalVariableFromPinType(const FName&
 	{
 		if (InPinType.PinSubCategory == UEdGraphSchema_K2::PC_Float)
 		{
-			ExternalVariable.TypeName = FloatTypeName;
+			ExternalVariable.TypeName = TypeName::Float;
 			ExternalVariable.Size = sizeof(float);
 		}
 		else if (InPinType.PinSubCategory == UEdGraphSchema_K2::PC_Double)
 		{
-			ExternalVariable.TypeName = DoubleTypeName;
+			ExternalVariable.TypeName = TypeName::Double;
 			ExternalVariable.Size = sizeof(double);
 		}
 		else
@@ -94,17 +94,17 @@ FRigVMExternalVariable RigVMTypeUtils::ExternalVariableFromPinType(const FName&
 	}
 	else if (InPinType.PinCategory == UEdGraphSchema_K2::PC_Name)
 	{
-		ExternalVariable.TypeName = FNameTypeName;
+		ExternalVariable.TypeName = TypeName::FName;
 		ExternalVariable.Size = sizeof(FName);
 	}
 	else if (InPinType.PinCategory == UEdGraphSchema_K2::PC_String)
 	{
-		ExternalVariable.TypeName = FStringTypeName;
+		ExternalVariable.TypeName = TypeName::FString;
 		ExternalVariable.Size = sizeof(FString);
 	}
 	else if (InPinType.PinCategory == UEdGraphSchema_K2::PC_Text)
 	{
-		ExternalVariable.TypeName = FTextTypeName;
+		ExternalVariable.TypeName = TypeName::FText;
 		ExternalVariable.Size = sizeof(FText);
 	}
 	else if (InPinType.PinCategory == UEdGraphSchema_K2::PC_Struct)
@@ -148,37 +148,37 @@ FRigVMExternalVariable RigVMTypeUtils::ExternalVariableFromCPPTypePath(const FNa
 		CPPTypePath = BaseTypeFromArrayType(CPPTypePath);
 	}
 
-	if (CPPTypePath == BoolType)
+	if (CPPTypePath == TypeString::Bool)
 	{
 		Variable.TypeName = *CPPTypePath;
 		Variable.Size = sizeof(bool);
 	}
-	else if (CPPTypePath == FloatType)
+	else if (CPPTypePath == TypeString::Float)
 	{
 		Variable.TypeName = *CPPTypePath;
 		Variable.Size = sizeof(float);
 	}
-	else if (CPPTypePath == DoubleType)
+	else if (CPPTypePath == TypeString::Double)
 	{
 		Variable.TypeName = *CPPTypePath;
 		Variable.Size = sizeof(double);
 	}
-	else if (CPPTypePath == Int32Type)
+	else if (CPPTypePath == TypeString::Int32)
 	{
 		Variable.TypeName = *CPPTypePath;
 		Variable.Size = sizeof(int32);
 	}
-	else if (CPPTypePath == FStringType)
+	else if (CPPTypePath == TypeString::FString)
 	{
 		Variable.TypeName = *CPPTypePath;
 		Variable.Size = sizeof(FString);
 	}
-	else if (CPPTypePath == FNameType)
+	else if (CPPTypePath == TypeString::FName)
 	{
 		Variable.TypeName = *CPPTypePath;
 		Variable.Size = sizeof(FName);
 	}
-	else if (CPPTypePath == FTextType)
+	else if (CPPTypePath == TypeString::FText)
 	{
 		Variable.TypeName = *CPPTypePath;
 		Variable.Size = sizeof(FText);
@@ -222,37 +222,37 @@ FRigVMExternalVariable RigVMTypeUtils::ExternalVariableFromCPPType(const FName&
 		CPPType = BaseTypeFromArrayType(CPPType);
 	}
 
-	if (CPPType == BoolType)
+	if (CPPType == TypeString::Bool)
 	{
 		Variable.TypeName = *CPPType;
 		Variable.Size = sizeof(bool);
 	}
-	else if (CPPType == FloatType)
+	else if (CPPType == TypeString::Float)
 	{
 		Variable.TypeName = *CPPType;
 		Variable.Size = sizeof(float);
 	}
-	else if (CPPType == DoubleType)
+	else if (CPPType == TypeString::Double)
 	{
 		Variable.TypeName = *CPPType;
 		Variable.Size = sizeof(double);
 	}
-	else if (CPPType == Int32Type)
+	else if (CPPType == TypeString::Int32)
 	{
 		Variable.TypeName = *CPPType;
 		Variable.Size = sizeof(int32);
 	}
-	else if (CPPType == FStringType)
+	else if (CPPType == TypeString::FString)
 	{
 		Variable.TypeName = *CPPType;
 		Variable.Size = sizeof(FString);
 	}
-	else if (CPPType == FNameType)
+	else if (CPPType == TypeString::FName)
 	{
 		Variable.TypeName = *CPPType;
 		Variable.Size = sizeof(FName);
 	}
-	else if (CPPType == FTextType)
+	else if (CPPType == TypeString::FText)
 	{
 		Variable.TypeName = *CPPType;
 		Variable.Size = sizeof(FText);
diff --git a/Engine/Plugins/Runtime/RigVM/Source/RigVM/Private/RigVMCore/RigVMFunction.cpp b/Engine/Plugins/Runtime/RigVM/Source/RigVM/Private/RigVMCore/RigVMFunction.cpp
--- a/Engine/Plugins/Runtime/RigVM/Source/RigVM/Private/RigVMCore/RigVMFunction.cpp
+++ b/Engine/Plugins/Runtime/RigVM/Source/RigVM/Private/RigVMCore/RigVMFunction.cpp
@@ -153,7 +153,7 @@ const UScriptStruct* FRigVMFunction::GetExecuteContextStruct(bool bLockRegistry)
 		if(Struct->GetStringMetaDataHierarchical(FRigVMStruct::ExecuteContextName, &ExecuteContextName))
 		{
 			const FRigVMTemplateArgumentType& Type = FRigVMRegistry_RWLock::Get().FindTypeFromCPPType(ExecuteContextName, bLockRegistry);
-			if(const UScriptStruct* ExecuteContextStruct = Cast<UScriptStruct>(Type.CPPTypeObject))
+			if(const UScriptStruct* ExecuteContextStruct = Cast<UScriptStruct>(Type.GetTypeObject()))
 			{
 				return ExecuteContextStruct;
 			}
diff --git a/Engine/Plugins/Runtime/RigVM/Source/RigVM/Private/RigVMCore/RigVMMemoryStorage.cpp b/Engine/Plugins/Runtime/RigVM/Source/RigVM/Private/RigVMCore/RigVMMemoryStorage.cpp
--- a/Engine/Plugins/Runtime/RigVM/Source/RigVM/Private/RigVMCore/RigVMMemoryStorage.cpp
+++ b/Engine/Plugins/Runtime/RigVM/Source/RigVM/Private/RigVMCore/RigVMMemoryStorage.cpp
@@ -144,20 +144,20 @@ FRigVMPropertyDescription::FRigVMPropertyDescription(const FProperty* InProperty
 		CPPType = Property->GetCPPType(&ExtendedType);
 		CPPType += ExtendedType;
 
-		if(CPPType == RigVMTypeUtils::UInt8Type)
+		if(CPPType == RigVMTypeUtils::TypeString::UInt8)
 		{
 			if(CastField<FBoolProperty>(Property))
 			{
-				CPPType = RigVMTypeUtils::BoolType;
+				CPPType = RigVMTypeUtils::TypeString::Bool;
 			}
 		}
-		else if(CPPType == RigVMTypeUtils::UInt8ArrayType)
+		else if(CPPType == RigVMTypeUtils::TypeString::UInt8Array)
 		{
 			if(const FArrayProperty* ArrayProperty = CastField<FArrayProperty>(Property))
 			{
 				if(CastField<FBoolProperty>(ArrayProperty->Inner))
 				{
-					CPPType = RigVMTypeUtils::BoolArrayType;
+					CPPType = RigVMTypeUtils::TypeString::BoolArray;
 				}
 			}
 		}
@@ -781,37 +781,41 @@ FProperty* URigVMMemoryStorageGeneratorClass::AddProperty(URigVMMemoryStorageGen
 			static const FString IntString = TEXT("int");
 
 			const FString BaseCPPType = InProperty.GetTailCPPType();
-			if(BaseCPPType.Equals(RigVMTypeUtils::BoolType, ESearchCase::IgnoreCase))
+			if(BaseCPPType.Equals(RigVMTypeUtils::TypeString::Bool, ESearchCase::IgnoreCase))
 			{
 				FBoolProperty* BoolProperty = new FBoolProperty(PropertyOwner, InProperty.Name, RF_Public); 
 				BoolProperty->SetBoolSize(sizeof(bool), true);
 				(*ValuePropertyPtr) = BoolProperty;
 			}
-			else if(BaseCPPType.Equals(RigVMTypeUtils::Int32Type, ESearchCase::IgnoreCase) ||
+			else if(BaseCPPType.Equals(RigVMTypeUtils::TypeString::Int32, ESearchCase::IgnoreCase) ||
 				BaseCPPType.Equals(IntString, ESearchCase::IgnoreCase))
 			{
 				(*ValuePropertyPtr) = new FIntProperty(PropertyOwner, InProperty.Name, RF_Public);;
 			}
-			else if(BaseCPPType.Equals(RigVMTypeUtils::UInt32Type, ESearchCase::IgnoreCase))
+			else if(BaseCPPType.Equals(RigVMTypeUtils::TypeString::UInt32, ESearchCase::IgnoreCase))
 			{
 				(*ValuePropertyPtr) = new FUInt32Property(PropertyOwner, InProperty.Name, RF_Public);;
 			}
-			else if(BaseCPPType.Equals(RigVMTypeUtils::FloatType, ESearchCase::IgnoreCase))
+			else if(BaseCPPType.Equals(RigVMTypeUtils::TypeString::Float, ESearchCase::IgnoreCase))
 			{
 				(*ValuePropertyPtr) = new FFloatProperty(PropertyOwner, InProperty.Name, RF_Public);;
 			}
-			else if(BaseCPPType.Equals(RigVMTypeUtils::DoubleType, ESearchCase::IgnoreCase))
+			else if(BaseCPPType.Equals(RigVMTypeUtils::TypeString::Double, ESearchCase::IgnoreCase))
 			{
 				(*ValuePropertyPtr) = new FDoubleProperty(PropertyOwner, InProperty.Name, RF_Public);;
 			}
-			else if(BaseCPPType.Equals(RigVMTypeUtils::FStringType, ESearchCase::IgnoreCase))
+			else if(BaseCPPType.Equals(RigVMTypeUtils::TypeString::FString, ESearchCase::IgnoreCase))
 			{
 				(*ValuePropertyPtr) = new FStrProperty(PropertyOwner, InProperty.Name, RF_Public);;
 			}
-			else if(BaseCPPType.Equals(RigVMTypeUtils::FNameType, ESearchCase::IgnoreCase))
+			else if(BaseCPPType.Equals(RigVMTypeUtils::TypeString::FName, ESearchCase::IgnoreCase))
 			{
 				(*ValuePropertyPtr) = new FNameProperty(PropertyOwner, InProperty.Name, RF_Public);;
 			}
+			else if(BaseCPPType.Equals(RigVMTypeUtils::TypeString::FText, ESearchCase::IgnoreCase))
+			{
+				(*ValuePropertyPtr) = new FTextProperty(PropertyOwner, InProperty.Name, RF_Public);;
+			}
 			else
 			{
 				// we might have hit a base structure
diff --git a/Engine/Plugins/Runtime/RigVM/Source/RigVM/Private/RigVMCore/RigVMMemoryStorageStruct.cpp b/Engine/Plugins/Runtime/RigVM/Source/RigVM/Private/RigVMCore/RigVMMemoryStorageStruct.cpp
--- a/Engine/Plugins/Runtime/RigVM/Source/RigVM/Private/RigVMCore/RigVMMemoryStorageStruct.cpp
+++ b/Engine/Plugins/Runtime/RigVM/Source/RigVM/Private/RigVMCore/RigVMMemoryStorageStruct.cpp
@@ -619,47 +619,47 @@ void FRigVMMemoryStorageStruct::SetDefaultValues(const TArray<FRigVMPropertyDesc
 
 	FName VMType = *VMTypeString;
 
-	if (VMType == RigVMTypeUtils::BoolTypeName)
+	if (VMType == RigVMTypeUtils::TypeName::Bool)
 	{
 		OutBagPropertyType = EPropertyBagPropertyType::Bool;
 	}
-	else if (VMType == RigVMTypeUtils::UInt8TypeName)
+	else if (VMType == RigVMTypeUtils::TypeName::UInt8)
 	{
 		OutBagPropertyType = EPropertyBagPropertyType::Byte;
 	}
-	else if (VMType == RigVMTypeUtils::Int32TypeName || VMType == RigVMTypeUtils::IntTypeName)
+	else if (VMType == RigVMTypeUtils::TypeName::Int32 || VMType == RigVMTypeUtils::TypeName::Int)
 	{
 		OutBagPropertyType = EPropertyBagPropertyType::Int32;
 	}
-	else if (VMType == RigVMTypeUtils::UInt32TypeName)
+	else if (VMType == RigVMTypeUtils::TypeName::UInt32)
 	{
 		OutBagPropertyType = EPropertyBagPropertyType::UInt32;
 	}
-	else if (VMType == RigVMTypeUtils::Int64TypeName)
+	else if (VMType == RigVMTypeUtils::TypeName::Int64)
 	{
 		OutBagPropertyType = EPropertyBagPropertyType::Int64;
 	}
-	else if (VMType == RigVMTypeUtils::UInt64TypeName)
+	else if (VMType == RigVMTypeUtils::TypeName::UInt64)
 	{
 		OutBagPropertyType = EPropertyBagPropertyType::UInt64;
 	}
-	else if (VMType == RigVMTypeUtils::FloatTypeName)
+	else if (VMType == RigVMTypeUtils::TypeName::Float)
 	{
 		OutBagPropertyType = EPropertyBagPropertyType::Float;
 	}
-	else if (VMType == RigVMTypeUtils::DoubleTypeName)
+	else if (VMType == RigVMTypeUtils::TypeName::Double)
 	{
 		OutBagPropertyType = EPropertyBagPropertyType::Double;
 	}
-	else if (VMType == RigVMTypeUtils::FNameTypeName)
+	else if (VMType == RigVMTypeUtils::TypeName::FName)
 	{
 		OutBagPropertyType = EPropertyBagPropertyType::Name;
 	}
-	else if (VMType == RigVMTypeUtils::FStringTypeName)
+	else if (VMType == RigVMTypeUtils::TypeName::FString)
 	{
 		OutBagPropertyType = EPropertyBagPropertyType::String;
 	}
-	else if (VMType == RigVMTypeUtils::FTextTypeName)
+	else if (VMType == RigVMTypeUtils::TypeName::FText)
 	{
 		OutBagPropertyType = EPropertyBagPropertyType::Text;
 	}
diff --git a/Engine/Plugins/Runtime/RigVM/Source/RigVM/Private/RigVMCore/RigVMRegistry.cpp b/Engine/Plugins/Runtime/RigVM/Source/RigVM/Private/RigVMCore/RigVMRegistry.cpp
--- a/Engine/Plugins/Runtime/RigVM/Source/RigVM/Private/RigVMCore/RigVMRegistry.cpp
+++ b/Engine/Plugins/Runtime/RigVM/Source/RigVM/Private/RigVMCore/RigVMRegistry.cpp
@@ -51,39 +51,41 @@ void FRigVMRegistry_NoLock::AddReferencedObjects(FReferenceCollector& Collector)
 	{
 		// the Object needs to be checked for validity since it may be a user defined type (struct or enum)
 		// which is about to get removed. 
-		if (Type.Type.CPPTypeObject)
+		if (Type.ArgumentType.GetTypeObject())
 		{
 #if !UE_BUILD_SHIPPING
 			// in non shipping builds, immediately run IsValidLowLevelFast such that
 			// we can catch invalid types earlier via a direct crash more often
-			if (Type.Type.CPPTypeObject->IsValidLowLevelFast())
+			if (Type.ArgumentType.GetTypeObject()->IsValidLowLevelFast())
 			{
 				// By design, hold strong references only to non-native types
-				if (!Type.Type.CPPTypeObject->IsNative())
+				if (!Type.ArgumentType.GetTypeObject()->IsNative())
 				{
-					Collector.AddReferencedObject(Type.Type.CPPTypeObject);	
+					TWeakObjectPtr<UObject> WeakObjectPtr(Type.ArgumentType.GetTypeObject());
+					Collector.AddReferencedObject(WeakObjectPtr);	
 				}
 			}
 #else
 			// in shipping builds, try to be as safe as possible
-			if(IsValid(Type.Type.CPPTypeObject))
+			if(IsValid(Type.ArgumentType.GetTypeObject()))
 			{
-				if(Type.Type.CPPTypeObject->GetClass())
+				if(Type.ArgumentType.GetTypeObject()->GetClass())
 				{
-					if(Type.Type.CPPTypeObject->IsValidLowLevelFast() &&
-						!Type.Type.CPPTypeObject->IsNative() &&
-						!Type.Type.CPPTypeObject->IsUnreachable())
+					if(Type.ArgumentType.GetTypeObject()->IsValidLowLevelFast() &&
+						!Type.ArgumentType.GetTypeObject()->IsNative() &&
+						!Type.ArgumentType.GetTypeObject()->IsUnreachable())
 					{
 						// make sure the object is part of the GUObjectArray and can be retrieved
 						// so that GC doesn't crash after receiving the referenced object
-						const int32 ObjectIndex = GUObjectArray.ObjectToIndex(Type.Type.CPPTypeObject);
+						const int32 ObjectIndex = GUObjectArray.ObjectToIndex(Type.ArgumentType.GetTypeObject());
 						if(ObjectIndex != INDEX_NONE)
 						{
 							if(const FUObjectItem* Item = GUObjectArray.IndexToObject(ObjectIndex))
 							{
-								if(Item->Object == Type.Type.CPPTypeObject)
+								if(Item->Object == Type.ArgumentType.GetTypeObject())
 								{
-									Collector.AddReferencedObject(Type.Type.CPPTypeObject);
+									TWeakObjectPtr<UObject> WeakObjectPtr(Type.ArgumentType.GetTypeObject());
+									Collector.AddReferencedObject(WeakObjectPtr);	
 								}
 							}
 						}
@@ -135,25 +137,25 @@ uint32 FRigVMRegistry_NoLock::GetHashForType_NoLock(TRigVMTypeIndex InTypeIndex)
 	}
 
 	uint32 Hash;
-	if(const UScriptStruct* ScriptStruct = Cast<UScriptStruct>(TypeInfo.Type.CPPTypeObject))
+	if(const UScriptStruct* ScriptStruct = Cast<UScriptStruct>(TypeInfo.ArgumentType.GetTypeObject()))
 	{
 		Hash = GetHashForScriptStruct_NoLock(ScriptStruct, false);
 	}
-	else if(const UStruct* Struct = Cast<UStruct>(TypeInfo.Type.CPPTypeObject))
+	else if(const UStruct* Struct = Cast<UStruct>(TypeInfo.ArgumentType.GetTypeObject()))
 	{
 		Hash = GetHashForStruct_NoLock(Struct);
 	}
-	else if(const UEnum* Enum = Cast<UEnum>(TypeInfo.Type.CPPTypeObject))
+	else if(const UEnum* Enum = Cast<UEnum>(TypeInfo.ArgumentType.GetTypeObject()))
     {
     	Hash = GetHashForEnum_NoLock(Enum, false);
     }
     else
     {
-    	Hash = GetTypeHash(TypeInfo.Type.CPPType.ToString());
+    	Hash = GetTypeHash(TypeInfo.ArgumentType.GetTypeString());
     }
 
 	// for used defined structs - always recompute it
-	if(Cast<UUserDefinedStruct>(TypeInfo.Type.CPPTypeObject))
+	if(Cast<UUserDefinedStruct>(TypeInfo.ArgumentType.GetTypeObject()))
 	{
 		return Hash;
 	}
@@ -314,25 +316,56 @@ void FRigVMRegistry_NoLock::Initialize_NoLock()
 
 	RigVMTypeUtils::TypeIndex::Execute = FindOrAddType_NoLock(FRigVMTemplateArgumentType(FRigVMExecuteContext::StaticStruct()), false);
 	RigVMTypeUtils::TypeIndex::ExecuteArray = FindOrAddType_NoLock(FRigVMTemplateArgumentType(FRigVMExecuteContext::StaticStruct()).ConvertToArray(), false);
-	RigVMTypeUtils::TypeIndex::Bool = FindOrAddType_NoLock(FRigVMTemplateArgumentType(RigVMTypeUtils::BoolTypeName, nullptr), false);
-	RigVMTypeUtils::TypeIndex::Float = FindOrAddType_NoLock(FRigVMTemplateArgumentType(RigVMTypeUtils::FloatTypeName, nullptr), false);
-	RigVMTypeUtils::TypeIndex::Double = FindOrAddType_NoLock(FRigVMTemplateArgumentType(RigVMTypeUtils::DoubleTypeName, nullptr), false);
-	RigVMTypeUtils::TypeIndex::Int32 = FindOrAddType_NoLock(FRigVMTemplateArgumentType(RigVMTypeUtils::Int32TypeName, nullptr), false);
-	RigVMTypeUtils::TypeIndex::UInt32 = FindOrAddType_NoLock(FRigVMTemplateArgumentType(RigVMTypeUtils::UInt32TypeName, nullptr), false);
-	RigVMTypeUtils::TypeIndex::UInt8 = FindOrAddType_NoLock(FRigVMTemplateArgumentType(RigVMTypeUtils::UInt8TypeName, nullptr), false);
-	RigVMTypeUtils::TypeIndex::FName = FindOrAddType_NoLock(FRigVMTemplateArgumentType(RigVMTypeUtils::FNameTypeName, nullptr), false);
-	RigVMTypeUtils::TypeIndex::FString = FindOrAddType_NoLock(FRigVMTemplateArgumentType(RigVMTypeUtils::FStringTypeName, nullptr), false);
+	RigVMTypeUtils::TypeIndex::Bool = FindOrAddType_NoLock(FRigVMTemplateArgumentType(RigVMTypeUtils::TypeName::Bool, nullptr), false);
+	RigVMTypeUtils::TypeIndex::Float = FindOrAddType_NoLock(FRigVMTemplateArgumentType(RigVMTypeUtils::TypeName::Float, nullptr), false);
+	RigVMTypeUtils::TypeIndex::Double = FindOrAddType_NoLock(FRigVMTemplateArgumentType(RigVMTypeUtils::TypeName::Double, nullptr), false);
+	RigVMTypeUtils::TypeIndex::Int32 = FindOrAddType_NoLock(FRigVMTemplateArgumentType(RigVMTypeUtils::TypeName::Int32, nullptr), false);
+	RigVMTypeUtils::TypeIndex::Int64 = FindOrAddType_NoLock(FRigVMTemplateArgumentType(RigVMTypeUtils::TypeName::Int64, nullptr), false);
+	RigVMTypeUtils::TypeIndex::UInt8 = FindOrAddType_NoLock(FRigVMTemplateArgumentType(RigVMTypeUtils::TypeName::UInt8, nullptr), false);
+	RigVMTypeUtils::TypeIndex::UInt32 = FindOrAddType_NoLock(FRigVMTemplateArgumentType(RigVMTypeUtils::TypeName::UInt32, nullptr), false);
+	RigVMTypeUtils::TypeIndex::UInt64 = FindOrAddType_NoLock(FRigVMTemplateArgumentType(RigVMTypeUtils::TypeName::UInt64, nullptr), false);
+	RigVMTypeUtils::TypeIndex::FName = FindOrAddType_NoLock(FRigVMTemplateArgumentType(RigVMTypeUtils::TypeName::FName, nullptr), false);
+	RigVMTypeUtils::TypeIndex::FString = FindOrAddType_NoLock(FRigVMTemplateArgumentType(RigVMTypeUtils::TypeName::FString, nullptr), false);
+	RigVMTypeUtils::TypeIndex::FText = FindOrAddType_NoLock(FRigVMTemplateArgumentType(RigVMTypeUtils::TypeName::FText, nullptr), false);
 	RigVMTypeUtils::TypeIndex::WildCard = FindOrAddType_NoLock(FRigVMTemplateArgumentType(RigVMTypeUtils::GetWildCardCPPTypeName(), RigVMTypeUtils::GetWildCardCPPTypeObject()), false);
-	RigVMTypeUtils::TypeIndex::BoolArray = FindOrAddType_NoLock(FRigVMTemplateArgumentType(RigVMTypeUtils::BoolArrayTypeName, nullptr), false);
-	RigVMTypeUtils::TypeIndex::FloatArray = FindOrAddType_NoLock(FRigVMTemplateArgumentType(RigVMTypeUtils::FloatArrayTypeName, nullptr), false);
-	RigVMTypeUtils::TypeIndex::DoubleArray = FindOrAddType_NoLock(FRigVMTemplateArgumentType(RigVMTypeUtils::DoubleArrayTypeName, nullptr), false);
-	RigVMTypeUtils::TypeIndex::Int32Array = FindOrAddType_NoLock(FRigVMTemplateArgumentType(RigVMTypeUtils::Int32ArrayTypeName, nullptr), false);
-	RigVMTypeUtils::TypeIndex::UInt32Array = FindOrAddType_NoLock(FRigVMTemplateArgumentType(RigVMTypeUtils::UInt32ArrayTypeName, nullptr), false);
-	RigVMTypeUtils::TypeIndex::UInt8Array = FindOrAddType_NoLock(FRigVMTemplateArgumentType(RigVMTypeUtils::UInt8ArrayTypeName, nullptr), false);
-	RigVMTypeUtils::TypeIndex::FNameArray = FindOrAddType_NoLock(FRigVMTemplateArgumentType(RigVMTypeUtils::FNameArrayTypeName, nullptr), false);
-	RigVMTypeUtils::TypeIndex::FStringArray = FindOrAddType_NoLock(FRigVMTemplateArgumentType(RigVMTypeUtils::FStringArrayTypeName, nullptr), false);
+	RigVMTypeUtils::TypeIndex::BoolArray = FindOrAddType_NoLock(FRigVMTemplateArgumentType(RigVMTypeUtils::TypeName::BoolArray, nullptr), false);
+	RigVMTypeUtils::TypeIndex::FloatArray = FindOrAddType_NoLock(FRigVMTemplateArgumentType(RigVMTypeUtils::TypeName::FloatArray, nullptr), false);
+	RigVMTypeUtils::TypeIndex::DoubleArray = FindOrAddType_NoLock(FRigVMTemplateArgumentType(RigVMTypeUtils::TypeName::DoubleArray, nullptr), false);
+	RigVMTypeUtils::TypeIndex::Int32Array = FindOrAddType_NoLock(FRigVMTemplateArgumentType(RigVMTypeUtils::TypeName::Int32Array, nullptr), false);
+	RigVMTypeUtils::TypeIndex::Int64Array = FindOrAddType_NoLock(FRigVMTemplateArgumentType(RigVMTypeUtils::TypeName::Int64Array, nullptr), false);
+	RigVMTypeUtils::TypeIndex::UInt8Array = FindOrAddType_NoLock(FRigVMTemplateArgumentType(RigVMTypeUtils::TypeName::UInt8Array, nullptr), false);
+	RigVMTypeUtils::TypeIndex::UInt32Array = FindOrAddType_NoLock(FRigVMTemplateArgumentType(RigVMTypeUtils::TypeName::UInt32Array, nullptr), false);
+	RigVMTypeUtils::TypeIndex::UInt64Array = FindOrAddType_NoLock(FRigVMTemplateArgumentType(RigVMTypeUtils::TypeName::UInt64Array, nullptr), false);
+	RigVMTypeUtils::TypeIndex::FNameArray = FindOrAddType_NoLock(FRigVMTemplateArgumentType(RigVMTypeUtils::TypeName::FNameArray, nullptr), false);
+	RigVMTypeUtils::TypeIndex::FStringArray = FindOrAddType_NoLock(FRigVMTemplateArgumentType(RigVMTypeUtils::TypeName::FStringArray, nullptr), false);
+	RigVMTypeUtils::TypeIndex::FTextArray = FindOrAddType_NoLock(FRigVMTemplateArgumentType(RigVMTypeUtils::TypeName::FTextArray, nullptr), false);
 	RigVMTypeUtils::TypeIndex::WildCardArray = FindOrAddType_NoLock(FRigVMTemplateArgumentType(RigVMTypeUtils::GetWildCardArrayCPPTypeName(), RigVMTypeUtils::GetWildCardCPPTypeObject()), false);
 
+	RigVMTypeUtils::Types::Bool = FRigVMType(RigVMTypeUtils::TypeName::Bool, 0, RigVMTypeUtils::TypeIndex::Bool);
+	RigVMTypeUtils::Types::Float = FRigVMType(RigVMTypeUtils::TypeName::Float, 0, RigVMTypeUtils::TypeIndex::Float);
+	RigVMTypeUtils::Types::Double = FRigVMType(RigVMTypeUtils::TypeName::Double, 0, RigVMTypeUtils::TypeIndex::Double);
+	RigVMTypeUtils::Types::Int32 = FRigVMType(RigVMTypeUtils::TypeName::Int32, 0, RigVMTypeUtils::TypeIndex::Int32);
+	RigVMTypeUtils::Types::Int64 = FRigVMType(RigVMTypeUtils::TypeName::Int64, 0, RigVMTypeUtils::TypeIndex::Int64);
+	RigVMTypeUtils::Types::UInt8 = FRigVMType(RigVMTypeUtils::TypeName::UInt8, 0, RigVMTypeUtils::TypeIndex::UInt8);
+	RigVMTypeUtils::Types::UInt32 = FRigVMType(RigVMTypeUtils::TypeName::UInt32, 0, RigVMTypeUtils::TypeIndex::UInt32);
+	RigVMTypeUtils::Types::UInt64 = FRigVMType(RigVMTypeUtils::TypeName::UInt64, 0, RigVMTypeUtils::TypeIndex::UInt64);
+	RigVMTypeUtils::Types::FName = FRigVMType(RigVMTypeUtils::TypeName::FName, 0, RigVMTypeUtils::TypeIndex::FName);
+	RigVMTypeUtils::Types::FString = FRigVMType(RigVMTypeUtils::TypeName::FString, 0, RigVMTypeUtils::TypeIndex::FString);
+	RigVMTypeUtils::Types::FText = FRigVMType(RigVMTypeUtils::TypeName::FText, 0, RigVMTypeUtils::TypeIndex::FText);
+	RigVMTypeUtils::Types::WildCard = FRigVMType(RigVMTypeUtils::GetWildCardCPPTypeName(), 0, RigVMTypeUtils::TypeIndex::WildCard);
+	RigVMTypeUtils::Types::BoolArray = FRigVMType(RigVMTypeUtils::TypeName::Bool, 1, RigVMTypeUtils::TypeIndex::BoolArray);
+	RigVMTypeUtils::Types::FloatArray = FRigVMType(RigVMTypeUtils::TypeName::Float, 1, RigVMTypeUtils::TypeIndex::FloatArray);
+	RigVMTypeUtils::Types::DoubleArray = FRigVMType(RigVMTypeUtils::TypeName::Double, 1, RigVMTypeUtils::TypeIndex::DoubleArray);
+	RigVMTypeUtils::Types::Int32Array = FRigVMType(RigVMTypeUtils::TypeName::Int32, 1, RigVMTypeUtils::TypeIndex::Int32Array);
+	RigVMTypeUtils::Types::Int64Array = FRigVMType(RigVMTypeUtils::TypeName::Int64, 1, RigVMTypeUtils::TypeIndex::Int32Array);
+	RigVMTypeUtils::Types::UInt8Array = FRigVMType(RigVMTypeUtils::TypeName::UInt8, 1, RigVMTypeUtils::TypeIndex::UInt8Array);
+	RigVMTypeUtils::Types::UInt32Array = FRigVMType(RigVMTypeUtils::TypeName::UInt32, 1, RigVMTypeUtils::TypeIndex::UInt32Array);
+	RigVMTypeUtils::Types::UInt64Array = FRigVMType(RigVMTypeUtils::TypeName::UInt64, 1, RigVMTypeUtils::TypeIndex::UInt32Array);
+	RigVMTypeUtils::Types::FNameArray = FRigVMType(RigVMTypeUtils::TypeName::FName, 1, RigVMTypeUtils::TypeIndex::FNameArray);
+	RigVMTypeUtils::Types::FStringArray = FRigVMType(RigVMTypeUtils::TypeName::FString, 1, RigVMTypeUtils::TypeIndex::FStringArray);
+	RigVMTypeUtils::Types::FTextArray = FRigVMType(RigVMTypeUtils::TypeName::FText, 1, RigVMTypeUtils::TypeIndex::FTextArray);
+	RigVMTypeUtils::Types::WildCardArray = FRigVMType(RigVMTypeUtils::GetWildCardCPPTypeName(), 1, RigVMTypeUtils::TypeIndex::WildCardArray);
+
 	// register the default math types
 	for(UScriptStruct* MathType : GetMathTypes())
 	{
@@ -586,7 +619,7 @@ TRigVMTypeIndex FRigVMRegistry_NoLock::FindOrAddType_NoLock(const FRigVMTemplate
 			ElementType.ConvertToBaseElement();
 		}
 		
-		const UObject* CPPTypeObject = ElementType.CPPTypeObject;
+		const UObject* CPPTypeObject = ElementType.GetTypeObject();
 		if(!bForce && (CPPTypeObject != nullptr))
 		{
 			if(const UClass* Class = Cast<UClass>(CPPTypeObject))
@@ -634,15 +667,16 @@ TRigVMTypeIndex FRigVMRegistry_NoLock::FindOrAddType_NoLock(const FRigVMTemplate
 			}
 
 			FTypeInfo Info;
-			Info.Type = CurType;
+			Info.ArgumentType = CurType;
 			Info.bIsArray = ArrayDimension > 0;
 			Info.bIsExecute = bIsExecute;
 			
 			Index = Types.Add(Info);
+			Types[Index].ArgumentType.Type.Index = Index;
 #if UE_RIGVM_DEBUG_TYPEINDEX
-			Index.Name = Info.Type.CPPType;
+			Index.Name = Info.ArgumentType.GetTypeFName();
 #endif
-			TypeToIndex.Add(CurType, Index);
+			TypeToIndex.Add(CurType.GetType(), Index);
 
 			Indices.Add(Index);
 		}
@@ -970,7 +1004,7 @@ bool FRigVMRegistry_NoLock::RemoveType_NoLock(const FSoftObjectPath& InObjectPat
 			}
 
 			// remove the type from the registry entirely
-			TypeToIndex.Remove(GetType_NoLock(Index));
+			TypeToIndex.Remove(GetType_NoLock(Index).GetType());
 			Types[Index] = FTypeInfo();
 		}
 
@@ -1066,7 +1100,7 @@ void FRigVMRegistry_RWLock::OnAnimationAttributeTypesChanged(const UScriptStruct
 	}
 }
 
-TRigVMTypeIndex FRigVMRegistry_NoLock::GetTypeIndex_NoLock(const FRigVMTemplateArgumentType& InType) const
+TRigVMTypeIndex FRigVMRegistry_NoLock::GetTypeIndex_NoLock(const FRigVMType& InType) const
 {
 	if(const TRigVMTypeIndex* Index = TypeToIndex.Find(InType))
 	{
@@ -1075,11 +1109,20 @@ TRigVMTypeIndex FRigVMRegistry_NoLock::GetTypeIndex_NoLock(const FRigVMTemplateA
 	return INDEX_NONE;
 }
 
+TRigVMTypeIndex FRigVMRegistry_NoLock::GetTypeIndex_NoLock(const FRigVMTemplateArgumentType& InType) const
+{
+	if(const TRigVMTypeIndex* Index = TypeToIndex.Find(InType.GetType()))
+	{
+		return *Index;
+	}
+	return INDEX_NONE;
+}
+
 const FRigVMTemplateArgumentType& FRigVMRegistry_NoLock::GetType_NoLock(TRigVMTypeIndex InTypeIndex) const
 {
 	if((Types.IsValidIndex(InTypeIndex)))
 	{
-		return Types[InTypeIndex].Type;
+		return Types[InTypeIndex].ArgumentType;
 	}
 	static FRigVMTemplateArgumentType EmptyType;
 	return EmptyType;
@@ -1090,7 +1133,7 @@ const FRigVMTemplateArgumentType& FRigVMRegistry_NoLock::FindTypeFromCPPType_NoL
 	const int32 TypeIndex = GetTypeIndexFromCPPType_NoLock(InCPPType);
 	if(ensure(Types.IsValidIndex(TypeIndex)))
 	{
-		return Types[TypeIndex].Type;
+		return Types[TypeIndex].ArgumentType;
 	}
 
 	static FRigVMTemplateArgumentType EmptyType;
@@ -1106,7 +1149,7 @@ TRigVMTypeIndex FRigVMRegistry_NoLock::GetTypeIndexFromCPPType_NoLock(const FStr
 
 		auto Predicate = [CPPTypeName](const FTypeInfo& Info) -> bool
 		{
-			return Info.Type.CPPType == CPPTypeName;
+			return Info.ArgumentType.GetTypeFName() == CPPTypeName;
 		};
 
 		Result = Types.IndexOfByPredicate(Predicate);
@@ -1130,7 +1173,7 @@ TRigVMTypeIndex FRigVMRegistry_NoLock::GetTypeIndexFromCPPType_NoLock(const FStr
 			const FString NewCPPType = RigVMTypeUtils::PostProcessCPPType(InCPPType);
 			Result = Types.IndexOfByPredicate([NewCPPType](const FTypeInfo& Info) -> bool
 			{
-				return Info.Type.CPPType == *NewCPPType;
+				return Info.ArgumentType.GetTypeFName() == *NewCPPType;
 			});
 		}
 	}
diff --git a/Engine/Plugins/Runtime/RigVM/Source/RigVM/Private/RigVMCore/RigVMTemplate.cpp b/Engine/Plugins/Runtime/RigVM/Source/RigVM/Private/RigVMCore/RigVMTemplate.cpp
--- a/Engine/Plugins/Runtime/RigVM/Source/RigVM/Private/RigVMCore/RigVMTemplate.cpp
+++ b/Engine/Plugins/Runtime/RigVM/Source/RigVM/Private/RigVMCore/RigVMTemplate.cpp
@@ -16,8 +16,8 @@
 ////////////////////////////////////////////////////////////////////////////////////////////////
 
 FRigVMTemplateArgumentType::FRigVMTemplateArgumentType(const FName& InCPPType, UObject* InCPPTypeObject)
-		: CPPType(InCPPType)
-		, CPPTypeObject(InCPPTypeObject)
+: CPPType(InCPPType)
+, CPPTypeObject(InCPPTypeObject)
 {
 	// InCppType is unreliable because not all caller knows that
 	// we use generated unique names for user defined structs
@@ -30,6 +30,7 @@ FRigVMTemplateArgumentType::FRigVMTemplateArgumentType(const FName& InCPPType, U
 		UE_LOG(LogRigVM, Warning, TEXT("FRigVMTemplateArgumentType(): Input CPPType '%s' could not be resolved."), *InCPPTypeString);
 	}
 #endif
+	Type = FRigVMType(CPPType.ToString());
 }
 
 FRigVMTemplateArgument::FRigVMTemplateArgument()
@@ -705,7 +706,7 @@ TArray<FString> FRigVMTemplateArgument::GetSupportedTypeStrings(const TArray<int
 			// INDEX_NONE indicates deleted permutation
 			if (TypeIndex != INDEX_NONE)
 			{
-				const FString TypeString = Registry.GetType_NoLock(TypeIndex).CPPType.ToString();
+				const FString TypeString = Registry.GetType_NoLock(TypeIndex).GetTypeString();
 				SupportedTypes.AddUnique(TypeString);
 			}
 			return true;
@@ -719,7 +720,7 @@ TArray<FString> FRigVMTemplateArgument::GetSupportedTypeStrings(const TArray<int
 			// INDEX_NONE indicates deleted permutation
 			if (TypeIndex != INDEX_NONE)
 			{
-				const FString TypeString = Registry.GetType_NoLock(TypeIndex).CPPType.ToString();
+				const FString TypeString = Registry.GetType_NoLock(TypeIndex).GetTypeString();
 				SupportedTypes.AddUnique(TypeString);
 			}
 		}
@@ -1141,7 +1142,7 @@ FString FRigVMTemplate::GetStringFromArgumentTypes(const FTypeMap& InTypes, bool
 		static constexpr TCHAR Format[] = TEXT("%s:%s");
 		FString PairString;
 		const FString KeyString = Pair.Key.ToString();
-		const FString CPPTypeString = Type.CPPType.ToString();
+		const FString CPPTypeString = Type.GetTypeString();
 		PairString.Reserve(KeyString.Len() + CPPTypeString.Len() + 1);
 		PairString.Append(KeyString);
 		PairString.AppendChar(TEXT(':'));
@@ -1703,7 +1704,7 @@ const FRigVMFunction* FRigVMTemplate::GetOrCreatePermutation_NoLock(int32 InInde
 			for(const FRigVMTemplateArgument& Argument : Arguments)
 			{
 				const FRigVMTemplateArgumentType& Type = Registry.GetType_NoLock(Argument.GetTypeIndex_NoLock(InIndex));
-				FunctionArguments.Add(FRigVMFunctionArgument(Argument.Name.ToString(), Type.CPPType.ToString()));
+				FunctionArguments.Add(FRigVMFunctionArgument(Argument.Name.ToString(), Type.GetTypeString()));
 			}
 			
 			static constexpr TCHAR Format[] = TEXT("%s::%s");
diff --git a/Engine/Plugins/Runtime/RigVM/Source/RigVM/Private/RigVMCore/RigVMType.cpp b/Engine/Plugins/Runtime/RigVM/Source/RigVM/Private/RigVMCore/RigVMType.cpp
new file mode 100644
--- /dev/null
+++ b/Engine/Plugins/Runtime/RigVM/Source/RigVM/Private/RigVMCore/RigVMType.cpp
@@ -0,0 +1,67 @@
+// Copyright Epic Games, Inc. All Rights Reserved.
+
+#include "RigVMCore/RigVMType.h"
+#include "RigVMCore/RigVMRegistry.h"
+
+FRigVMType::FRigVMType(const FString& InCPPTypeString)
+: BaseType(NAME_None)
+, ArrayDimensions(0)
+{
+	static const int32 PrefixLength = GetArrayPrefix().Len();
+	static const int32 ArrayLengthReduction = PrefixLength + GetArraySuffix().Len(); 
+
+	FStringView View(InCPPTypeString);
+	while(View.StartsWith(GetArrayPrefix()))
+	{
+		View = View.Mid(PrefixLength, View.Len() - ArrayLengthReduction);
+		ArrayDimensions++;
+	}
+
+	BaseType = *FString(View);
+}
+
+const TRigVMTypeIndex& FRigVMType::GetIndex() const
+{
+	if(Index == INDEX_NONE)
+	{
+		const FRigVMRegistryReadLock Lock;
+		Index = Lock.GetRegistry().GetTypeIndex_NoLock(*this);
+	}
+	return Index;
+}
+
+const TRigVMTypeIndex& FRigVMType::GetIndex_NoLock() const
+{
+	if(Index == INDEX_NONE)
+	{
+		Index = FRigVMRegistry_NoLock::GetForRead().GetTypeIndex_NoLock(*this);
+	}
+	return Index;
+}
+
+FString FRigVMType::ToString() const
+{
+	const int32 ArrayNum = ArrayDimensions;
+	FString BaseTypeString = BaseType.ToString();
+	if(ArrayNum == 0)
+	{
+		return BaseTypeString;
+	}
+	FStringBuilderBase Builder;
+	Builder.Reserve(
+		BaseTypeString.Len() +
+		GetArrayPrefix().Len() * ArrayNum +
+		GetArraySuffix().Len() * ArrayNum
+	);
+
+	for(int32 ArrayIndex = 0; ArrayIndex < ArrayNum; ArrayIndex++)
+	{
+		Builder.Append(GetArrayPrefix());
+	}
+	Builder.Append(BaseTypeString);
+	for(int32 ArrayIndex = 0; ArrayIndex < ArrayNum; ArrayIndex++)
+	{
+		Builder.Append(GetArraySuffix());
+	}
+	return Builder.ToString();
+}
diff --git a/Engine/Plugins/Runtime/RigVM/Source/RigVM/Private/RigVMFunctions/RigVMDispatch_CastObject.cpp b/Engine/Plugins/Runtime/RigVM/Source/RigVM/Private/RigVMFunctions/RigVMDispatch_CastObject.cpp
--- a/Engine/Plugins/Runtime/RigVM/Source/RigVM/Private/RigVMFunctions/RigVMDispatch_CastObject.cpp
+++ b/Engine/Plugins/Runtime/RigVM/Source/RigVM/Private/RigVMFunctions/RigVMDispatch_CastObject.cpp
@@ -73,9 +73,9 @@ FString FRigVMDispatch_CastObject::GetNodeTitle(const FRigVMTemplateTypeMap& InT
 				const FRigVMTemplateArgumentType& ValueType = Registry.GetType(ValueTypeIndex);
 				const FRigVMTemplateArgumentType& ResultType = Registry.GetType(ResultTypeIndex);
 				static const FString Prefix = TEXT("Cast to ");
-				if(ValueType.CPPTypeObject && ResultType.CPPTypeObject)
+				if(ValueType.GetTypeObject() && ResultType.GetTypeObject())
 				{
-					return Prefix + ResultType.CPPTypeObject->GetName();
+					return Prefix + ResultType.GetTypeObject()->GetName();
 				}
 			}
 		}
@@ -92,9 +92,9 @@ FText FRigVMDispatch_CastObject::GetNodeTooltip(const FRigVMTemplateTypeMap& InT
 		const FRigVMRegistry& Registry = FRigVMRegistry::Get();
 		const FRigVMTemplateArgumentType& ValueType = Registry.GetType(*ValueTypeIndexPtr);
 		const FRigVMTemplateArgumentType& ResultType = Registry.GetType(*ResultTypeIndexPtr);
-		if(ValueType.CPPTypeObject && ValueType.CPPTypeObject->IsA<UClass>() && ResultType.CPPTypeObject && ResultType.CPPTypeObject->IsA<UClass>())
+		if(ValueType.GetTypeObject() && ValueType.GetTypeObject()->IsA<UClass>() && ResultType.GetTypeObject() && ResultType.GetTypeObject()->IsA<UClass>())
 		{
-			return FText::Format(LOCTEXT("CastToolTipFormat", "Cast from {0} to {1}"), CastChecked<UClass>(ValueType.CPPTypeObject)->GetDisplayNameText(), CastChecked<UClass>(ResultType.CPPTypeObject)->GetDisplayNameText());
+			return FText::Format(LOCTEXT("CastToolTipFormat", "Cast from {0} to {1}"), CastChecked<UClass>(ValueType.GetTypeObject())->GetDisplayNameText(), CastChecked<UClass>(ResultType.GetTypeObject())->GetDisplayNameText());
 		}
 	}
 
diff --git a/Engine/Plugins/Runtime/RigVM/Source/RigVM/Private/RigVMFunctions/RigVMDispatch_Constant.cpp b/Engine/Plugins/Runtime/RigVM/Source/RigVM/Private/RigVMFunctions/RigVMDispatch_Constant.cpp
--- a/Engine/Plugins/Runtime/RigVM/Source/RigVM/Private/RigVMFunctions/RigVMDispatch_Constant.cpp
+++ b/Engine/Plugins/Runtime/RigVM/Source/RigVM/Private/RigVMFunctions/RigVMDispatch_Constant.cpp
@@ -44,11 +44,11 @@ FString FRigVMDispatch_Constant::GetNodeTitle(const FRigVMTemplateTypeMap& InTyp
 			{
 				const FRigVMTemplateArgumentType& Type = Registry.GetType(TypeIndex);
 				static constexpr TCHAR Format[] = TEXT("%s %s");
-				if(Type.CPPTypeObject)
+				if(Type.GetTypeObject())
 				{
-					return FString::Printf(Format, *SuperNodeTitle, *Type.CPPTypeObject->GetName());
+					return FString::Printf(Format, *SuperNodeTitle, *Type.GetTypeObject()->GetName());
 				}
-				return FString::Printf(Format, *SuperNodeTitle, *Type.CPPType.ToString());
+				return FString::Printf(Format, *SuperNodeTitle, *Type.GetTypeString());
 			}
 		}
 	}
diff --git a/Engine/Plugins/Runtime/RigVM/Source/RigVM/Private/RigVMFunctions/RigVMDispatch_MakeStruct.cpp b/Engine/Plugins/Runtime/RigVM/Source/RigVM/Private/RigVMFunctions/RigVMDispatch_MakeStruct.cpp
--- a/Engine/Plugins/Runtime/RigVM/Source/RigVM/Private/RigVMFunctions/RigVMDispatch_MakeStruct.cpp
+++ b/Engine/Plugins/Runtime/RigVM/Source/RigVM/Private/RigVMFunctions/RigVMDispatch_MakeStruct.cpp
@@ -44,7 +44,7 @@ FString FRigVMDispatch_MakeStruct::GetNodeTitle(const FRigVMTemplateTypeMap& InT
 			const FRigVMRegistry& Registry = FRigVMRegistry::Get();
 			if(!Registry.IsWildCardType(StructType))
 			{
-				if(UScriptStruct* Struct = Cast<UScriptStruct>(Registry.GetType(StructType).CPPTypeObject))
+				if(const UScriptStruct* Struct = Cast<UScriptStruct>(Registry.GetType(StructType).GetTypeObject()))
 				{
 					static constexpr TCHAR Format[] = TEXT("%s %s");
 					return FString::Printf(Format, *SuperNodeTitle, *Struct->GetName());
@@ -81,7 +81,7 @@ FString FRigVMDispatch_MakeStruct::GetKeywords() const
 	for(const TRigVMTypeIndex& StructTypeIndex : StructTypeIndices)
 	{
 		const FRigVMTemplateArgumentType& Type = Registry.GetType(StructTypeIndex);
-		if(const UScriptStruct* ScriptStruct = Cast<UScriptStruct>(Type.CPPTypeObject))
+		if(const UScriptStruct* ScriptStruct = Cast<UScriptStruct>(Type.GetTypeObject()))
 		{
 			static constexpr TCHAR Format[] = TEXT(",Make%s");
 			Keywords += FString::Printf(Format, *ScriptStruct->GetName());
@@ -132,7 +132,7 @@ FString FRigVMDispatch_BreakStruct::GetKeywords() const
 	for(const TRigVMTypeIndex& StructTypeIndex : StructTypeIndices)
 	{
 		const FRigVMTemplateArgumentType& Type = Registry.GetType(StructTypeIndex);
-		if(const UScriptStruct* ScriptStruct = Cast<UScriptStruct>(Type.CPPTypeObject))
+		if(const UScriptStruct* ScriptStruct = Cast<UScriptStruct>(Type.GetTypeObject()))
 		{
 			static constexpr TCHAR Format[] = TEXT(",Break%s");
 			Keywords += FString::Printf(Format, *ScriptStruct->GetName());
diff --git a/Engine/Plugins/Runtime/RigVM/Source/RigVM/Private/RigVMHost.cpp b/Engine/Plugins/Runtime/RigVM/Source/RigVM/Private/RigVMHost.cpp
--- a/Engine/Plugins/Runtime/RigVM/Source/RigVM/Private/RigVMHost.cpp
+++ b/Engine/Plugins/Runtime/RigVM/Source/RigVM/Private/RigVMHost.cpp
@@ -1204,10 +1204,10 @@ TArray<const UObject*> URigVMHost::GetUserDefinedDependencies(const TArray<const
 		for (const TRigVMTypeIndex& TypeIndex : TypeIndices)
 		{
 			const FRigVMTemplateArgumentType& Type = Registry.GetType(TypeIndex);
-			if (Cast<UUserDefinedStruct>(Type.CPPTypeObject) ||
-				Cast<UUserDefinedEnum>(Type.CPPTypeObject))
+			if (Cast<UUserDefinedStruct>(Type.GetTypeObject()) ||
+				Cast<UUserDefinedEnum>(Type.GetTypeObject()))
 			{
-				Dependencies.AddUnique(Type.CPPTypeObject);
+				Dependencies.AddUnique(Type.GetTypeObject());
 			}
 		}
 	}
diff --git a/Engine/Plugins/Runtime/RigVM/Source/RigVM/Private/RigVMTypeUtils.cpp b/Engine/Plugins/Runtime/RigVM/Source/RigVM/Private/RigVMTypeUtils.cpp
--- a/Engine/Plugins/Runtime/RigVM/Source/RigVM/Private/RigVMTypeUtils.cpp
+++ b/Engine/Plugins/Runtime/RigVM/Source/RigVM/Private/RigVMTypeUtils.cpp
@@ -18,21 +18,52 @@ namespace RigVMTypeUtils
 	TRigVMTypeIndex TypeIndex::Bool = INDEX_NONE;	
 	TRigVMTypeIndex TypeIndex::Float = INDEX_NONE;	
 	TRigVMTypeIndex TypeIndex::Double = INDEX_NONE;	
-	TRigVMTypeIndex TypeIndex::Int32 = INDEX_NONE;	
+	TRigVMTypeIndex TypeIndex::Int32 = INDEX_NONE;
+	TRigVMTypeIndex TypeIndex::Int64 = INDEX_NONE;
 	TRigVMTypeIndex TypeIndex::UInt32 = INDEX_NONE;	
+	TRigVMTypeIndex TypeIndex::UInt64 = INDEX_NONE;	
 	TRigVMTypeIndex TypeIndex::UInt8 = INDEX_NONE;	
 	TRigVMTypeIndex TypeIndex::FName = INDEX_NONE;	
 	TRigVMTypeIndex TypeIndex::FString = INDEX_NONE;
+	TRigVMTypeIndex TypeIndex::FText = INDEX_NONE;
 	TRigVMTypeIndex TypeIndex::WildCard = INDEX_NONE;	
 	TRigVMTypeIndex TypeIndex::BoolArray = INDEX_NONE;	
 	TRigVMTypeIndex TypeIndex::FloatArray = INDEX_NONE;	
 	TRigVMTypeIndex TypeIndex::DoubleArray = INDEX_NONE;	
 	TRigVMTypeIndex TypeIndex::Int32Array = INDEX_NONE;	
-	TRigVMTypeIndex TypeIndex::UInt32Array = INDEX_NONE;	
+	TRigVMTypeIndex TypeIndex::Int64Array = INDEX_NONE;	
 	TRigVMTypeIndex TypeIndex::UInt8Array = INDEX_NONE;	
+	TRigVMTypeIndex TypeIndex::UInt32Array = INDEX_NONE;	
+	TRigVMTypeIndex TypeIndex::UInt64Array = INDEX_NONE;	
 	TRigVMTypeIndex TypeIndex::FNameArray = INDEX_NONE;	
-	TRigVMTypeIndex TypeIndex::FStringArray = INDEX_NONE;	
-	TRigVMTypeIndex TypeIndex::WildCardArray = INDEX_NONE;	
+	TRigVMTypeIndex TypeIndex::FStringArray = INDEX_NONE;
+	TRigVMTypeIndex TypeIndex::FTextArray = INDEX_NONE;
+	TRigVMTypeIndex TypeIndex::WildCardArray = INDEX_NONE;
+
+	FRigVMType Types::Bool;
+	FRigVMType Types::Float;
+	FRigVMType Types::Double;
+	FRigVMType Types::Int32;
+	FRigVMType Types::Int64;
+	FRigVMType Types::UInt8;
+	FRigVMType Types::UInt32;
+	FRigVMType Types::UInt64;
+	FRigVMType Types::FName;
+	FRigVMType Types::FString;
+	FRigVMType Types::FText;
+	FRigVMType Types::WildCard;
+	FRigVMType Types::BoolArray;
+	FRigVMType Types::FloatArray;
+	FRigVMType Types::DoubleArray;
+	FRigVMType Types::Int32Array;
+	FRigVMType Types::Int64Array;
+	FRigVMType Types::UInt8Array;
+	FRigVMType Types::UInt32Array;
+	FRigVMType Types::UInt64Array;
+	FRigVMType Types::FNameArray;
+	FRigVMType Types::FStringArray;
+	FRigVMType Types::FTextArray;
+	FRigVMType Types::WildCardArray;
 }
 
 UObject* RigVMTypeUtils::UserDefinedTypeFromCPPType(FString& InOutCPPType, const FRigVMUserDefinedTypeResolver* InTypeResolver)
diff --git a/Engine/Plugins/Runtime/RigVM/Source/RigVM/Public/RigVMCore/RigVMRegistry.h b/Engine/Plugins/Runtime/RigVM/Source/RigVM/Public/RigVMCore/RigVMRegistry.h
--- a/Engine/Plugins/Runtime/RigVM/Source/RigVM/Public/RigVMCore/RigVMRegistry.h
+++ b/Engine/Plugins/Runtime/RigVM/Source/RigVM/Public/RigVMCore/RigVMRegistry.h
@@ -11,7 +11,7 @@
 #include "RigVMDispatchFactory.h"
 #include "RigVMFunction.h"
 #include "RigVMTemplate.h"
-#include "RigVMTypeIndex.h"
+#include "RigVMType.h"
 #include "Templates/EnableIf.h"
 #include "Templates/IsEnum.h"
 #include "Templates/Models.h"
@@ -126,6 +126,9 @@ public:
 	// which also creates invalid permutations in templates that we should ignore
 	virtual bool RemoveType_NoLock(const FSoftObjectPath& InObjectPath, const UClass* InObjectClass);
 
+	// Returns the type index given a type
+	virtual TRigVMTypeIndex GetTypeIndex_NoLock(const FRigVMType& InType) const;
+
 	// Returns the type index given a type
 	virtual TRigVMTypeIndex GetTypeIndex_NoLock(const FRigVMTemplateArgumentType& InType) const;
 
@@ -328,7 +331,7 @@ protected:
 	struct FTypeInfo
 	{
 		FTypeInfo()
-			: Type()
+			: ArgumentType()
 			, BaseTypeIndex(INDEX_NONE)
 			, ArrayTypeIndex(INDEX_NONE)
 			, bIsArray(false)
@@ -336,7 +339,7 @@ protected:
 			, Hash(UINT32_MAX)
 		{}
 		
-		FRigVMTemplateArgumentType Type;
+		FRigVMTemplateArgumentType ArgumentType;
 		TRigVMTypeIndex BaseTypeIndex;
 		TRigVMTypeIndex ArrayTypeIndex;
 		bool bIsArray;
@@ -370,7 +373,7 @@ protected:
 
 	// memory for all (known) types
 	TArray<FTypeInfo> Types;
-	TMap<FRigVMTemplateArgumentType, TRigVMTypeIndex> TypeToIndex;
+	TMap<FRigVMType, TRigVMTypeIndex> TypeToIndex;
 
 	// memory for all functions
 	// We use TChunkedArray because we need the memory locations to be stable, since we only ever add and never remove.
@@ -669,6 +672,13 @@ public:
 		return Super::RemoveType_NoLock(InObjectPath, InObjectClass);
 	}
 
+	// Returns the type index given a type
+	TRigVMTypeIndex GetTypeIndex(const FRigVMType& InType, bool bLockRegistry = true) const
+	{
+		FConditionalReadScopeLock _(*this, bLockRegistry);
+		return Super::GetTypeIndex_NoLock(InType);
+	}
+
 	// Returns the type index given a type
 	TRigVMTypeIndex GetTypeIndex(const FRigVMTemplateArgumentType& InType, bool bLockRegistry = true) const
 	{
@@ -1053,6 +1063,13 @@ protected:
 		return Super::RemoveType_NoLock(InObjectPath, InObjectClass);
 	}
 
+	// Returns the type index given a type
+	// Note: Only call NoLock methods on the FRigVMRegistry_NoLock
+	virtual TRigVMTypeIndex GetTypeIndex_NoLock(const FRigVMType& InType) const override
+	{
+		return Super::GetTypeIndex_NoLock(InType);
+	}
+
 	// Returns the type index given a type
 	// Note: Only call NoLock methods on the FRigVMRegistry_NoLock
 	virtual TRigVMTypeIndex GetTypeIndex_NoLock(const FRigVMTemplateArgumentType& InType) const override
diff --git a/Engine/Plugins/Runtime/RigVM/Source/RigVM/Public/RigVMCore/RigVMTemplate.h b/Engine/Plugins/Runtime/RigVM/Source/RigVM/Public/RigVMCore/RigVMTemplate.h
--- a/Engine/Plugins/Runtime/RigVM/Source/RigVM/Public/RigVMCore/RigVMTemplate.h
+++ b/Engine/Plugins/Runtime/RigVM/Source/RigVM/Public/RigVMCore/RigVMTemplate.h
@@ -12,7 +12,7 @@
 #include "Math/Color.h"
 #include "Misc/AssertionMacros.h"
 #include "RigVMFunction.h"
-#include "RigVMTypeIndex.h"
+#include "RigVMType.h"
 #include "RigVMTypeUtils.h"
 #include "UObject/Class.h"
 #include "UObject/NameTypes.h"
@@ -48,18 +48,13 @@ struct RIGVM_API FRigVMTemplateArgumentType
 {
 	GENERATED_BODY()
 
-	UPROPERTY()
-	FName CPPType;
-	
-	UPROPERTY()
-	TObjectPtr<UObject> CPPTypeObject; 
-
 	FRigVMTemplateArgumentType()
 		: CPPType(NAME_None)
 		, CPPTypeObject(nullptr)
 	{
 		CPPType = RigVMTypeUtils::GetWildCardCPPTypeName();
 		CPPTypeObject = RigVMTypeUtils::GetWildCardCPPTypeObject();
+		Type = FRigVMType(CPPType.ToString());
 	}
 
 	FRigVMTemplateArgumentType(const FName& InCPPType, UObject* InCPPTypeObject = nullptr);
@@ -68,23 +63,46 @@ struct RIGVM_API FRigVMTemplateArgumentType
 	: CPPType(*RigVMTypeUtils::CPPTypeFromObject(InClass, InClassArgType))
 	, CPPTypeObject(InClass)
 	{
+		Type = FRigVMType(CPPType.ToString());
 	}
 
 	FRigVMTemplateArgumentType(UScriptStruct* InScriptStruct)
 	: CPPType(*RigVMTypeUtils::GetUniqueStructTypeName(InScriptStruct))
 	, CPPTypeObject(InScriptStruct)
 	{
+		Type = FRigVMType(CPPType.ToString());
 	}
 
 	FRigVMTemplateArgumentType(UEnum* InEnum)
 	: CPPType(*RigVMTypeUtils::CPPTypeFromEnum(InEnum))
 	, CPPTypeObject(InEnum)
 	{
+		Type = FRigVMType(CPPType.ToString());
 	}
 
 	bool IsValid() const
 	{
-		return !CPPType.IsNone();
+		return Type.IsValid();
+	}
+
+	FRigVMType GetType() const
+	{
+		return Type;
+	}
+
+	FString GetTypeString() const
+	{
+		return Type.ToString();
+	}
+	
+	FName GetTypeFName() const
+	{
+		return CPPType;
+	}
+
+	UObject* GetTypeObject() const
+	{
+		return CPPTypeObject.Get();
 	}
 
 	static FRigVMTemplateArgumentType Array()
@@ -94,17 +112,17 @@ struct RIGVM_API FRigVMTemplateArgumentType
 
 	bool operator == (const FRigVMTemplateArgumentType& InOther) const
 	{
-		return CPPType == InOther.CPPType;
+		return Type == InOther.Type;
 	}
 
 	bool operator != (const FRigVMTemplateArgumentType& InOther) const
 	{
-		return CPPType != InOther.CPPType;
+		return Type != InOther.Type;
 	}
 
 	RIGVM_API friend FORCEINLINE uint32 GetTypeHash(const FRigVMTemplateArgumentType& InType)
 	{
-		return GetTypeHash(InType.CPPType.ToString());
+		return GetTypeHash(InType.Type);
 	}
 
 	FName GetCPPTypeObjectPath() const
@@ -119,35 +137,45 @@ struct RIGVM_API FRigVMTemplateArgumentType
 	bool IsWildCard() const
 	{
 		return CPPTypeObject == RigVMTypeUtils::GetWildCardCPPTypeObject() ||
-			CPPType == RigVMTypeUtils::GetWildCardCPPTypeName() ||
-			CPPType == RigVMTypeUtils::GetWildCardArrayCPPTypeName();
+			Type.GetBaseType() == RigVMTypeUtils::GetWildCardCPPTypeName();
 	}
 
 	bool IsArray() const
 	{
-		return RigVMTypeUtils::IsArrayType(CPPType.ToString());
+		return Type.IsArray();
 	}
 
 	FString GetBaseCPPType() const
 	{
-		if(IsArray())
-		{
-			return RigVMTypeUtils::BaseTypeFromArrayType(CPPType.ToString());
-		}
-		return CPPType.ToString();
+		return Type.GetBaseType().ToString();
 	}
 
 	FRigVMTemplateArgumentType& ConvertToArray() 
 	{
-		CPPType = *RigVMTypeUtils::ArrayTypeFromBaseType(CPPType.ToString());
+		Type.ConvertToArray();
+		CPPType = *Type.ToString();
 		return *this;
 	}
 
 	FRigVMTemplateArgumentType& ConvertToBaseElement() 
 	{
-		CPPType = *RigVMTypeUtils::BaseTypeFromArrayType(CPPType.ToString());
+		Type.ConvertToArrayElement();
+		CPPType = *Type.ToString();
 		return *this;
 	}
+
+private:
+	
+	UPROPERTY()
+	FRigVMType Type;
+
+	UPROPERTY()
+	FName CPPType;
+
+	UPROPERTY()
+	TObjectPtr<UObject> CPPTypeObject; 
+
+	friend struct FRigVMRegistry_NoLock;
 };
 
 /**
diff --git a/Engine/Plugins/Runtime/RigVM/Source/RigVM/Public/RigVMCore/RigVMType.h b/Engine/Plugins/Runtime/RigVM/Source/RigVM/Public/RigVMCore/RigVMType.h
new file mode 100644
--- /dev/null
+++ b/Engine/Plugins/Runtime/RigVM/Source/RigVM/Public/RigVMCore/RigVMType.h
@@ -0,0 +1,151 @@
+// Copyright Epic Games, Inc. All Rights Reserved.
+
+#pragma once
+
+#include "RigVMTypeIndex.h"
+#include "RigVMType.generated.h"
+
+USTRUCT()
+struct RIGVM_API FRigVMType
+{
+public:
+	GENERATED_BODY()
+
+	FRigVMType()
+	: BaseType(NAME_None)
+	, Index(INDEX_NONE)
+	, ArrayDimensions(0)
+	{}
+	
+	FRigVMType(const FName& InBaseType, uint8 InArrayDimensions)
+		: BaseType(InBaseType)
+		, Index(INDEX_NONE)
+		, ArrayDimensions(InArrayDimensions)
+	{}
+
+	explicit FRigVMType(const FString& InCPPTypeString);
+
+	bool IsValid() const
+	{
+		return !BaseType.IsNone() && ArrayDimensions < 3;
+	}
+
+	const FName& GetBaseType() const
+	{ 
+		return BaseType; 
+	}
+
+	uint8 GetArrayDimensions() const
+	{ 
+		return ArrayDimensions;
+	}
+
+	bool IsArray() const
+	{
+		return ArrayDimensions > 0;
+	}
+
+	const TRigVMTypeIndex& GetIndex() const;
+	const TRigVMTypeIndex& GetIndex_NoLock() const;
+
+	FRigVMType AsBaseType() const
+	{
+		return FRigVMType(BaseType, 0);
+	}
+
+	FRigVMType AsArray() const
+	{
+		return FRigVMType(BaseType, ArrayDimensions + 1);
+	}
+
+	FRigVMType AsArrayElement() const
+	{
+		if(ArrayDimensions < 2)
+		{
+			return AsBaseType();
+		}
+		return FRigVMType(BaseType, ArrayDimensions - 1);
+	}
+
+	FRigVMType& ConvertToBaseType()
+	{
+		*this = AsBaseType();
+		return *this;
+	}
+
+	FRigVMType& ConvertToArray()
+	{
+		*this = AsArray();
+		return *this;
+	}
+
+	FRigVMType& ConvertToArrayElement()
+	{
+		*this = AsArrayElement();
+		return *this;
+	}
+
+	static const FString& GetArrayPrefix()
+	{
+		static const FString ArrayPrefix = TEXT("TArray<");
+		return ArrayPrefix;
+	}
+
+	static const FString& GetArraySuffix()
+	{
+		static const FString ArraySuffix = TEXT(">");
+		return ArraySuffix;
+	}
+
+	FString ToString() const;
+
+	bool operator ==(const FRigVMType& Other) const
+	{
+		return ArrayDimensions == Other.ArrayDimensions && BaseType == Other.BaseType;
+	}
+
+	bool operator !=(const FRigVMType& Other) const
+	{
+		return ArrayDimensions != Other.ArrayDimensions || BaseType != Other.BaseType;
+	}
+
+	bool operator >(const FRigVMType& Other) const
+	{
+		const int32 Cmp = BaseType.Compare(Other.BaseType);
+		if(Cmp == 0)
+		{
+			return ArrayDimensions > Other.ArrayDimensions;
+		}
+		return Cmp > 0;
+	}
+
+	bool operator <(const FRigVMType& Other) const
+	{
+		const int32 Cmp = BaseType.Compare(Other.BaseType);
+		if(Cmp == 0)
+		{
+			return ArrayDimensions < Other.ArrayDimensions;
+		}
+		return Cmp < 0;
+	}
+
+	friend uint32 GetTypeHash(const FRigVMType& InType)
+	{
+		return HashCombine(GetTypeHash(InType.BaseType), GetTypeHash(InType.ArrayDimensions));
+	}
+
+protected:
+
+	FRigVMType(const FName& InBaseType, uint8 InArrayDimensions, const TRigVMTypeIndex& InTypeIndex)
+	: BaseType(InBaseType)
+	, Index(InTypeIndex)
+	, ArrayDimensions(InArrayDimensions)
+	{}
+
+	FName BaseType;
+	mutable TRigVMTypeIndex Index;
+	uint8 ArrayDimensions;
+
+	friend struct FRigVMRegistry_NoLock;
+	friend struct FRigVMTemplateArgument;
+};
diff --git a/Engine/Plugins/Runtime/RigVM/Source/RigVM/Public/RigVMTypeUtils.h b/Engine/Plugins/Runtime/RigVM/Source/RigVM/Public/RigVMTypeUtils.h
--- a/Engine/Plugins/Runtime/RigVM/Source/RigVM/Public/RigVMTypeUtils.h
+++ b/Engine/Plugins/Runtime/RigVM/Source/RigVM/Public/RigVMTypeUtils.h
@@ -4,7 +4,7 @@
 
 #include "CoreMinimal.h"
 #include "RigVMStringUtils.h"
-#include "RigVMCore/RigVMTypeIndex.h"
+#include "RigVMCore/RigVMType.h"
 #include "RigVMCore/RigVMUnknownType.h"
 #include "UObject/Interface.h"
 #include "Engine/UserDefinedEnum.h"
@@ -56,49 +56,62 @@ namespace RigVMTypeUtils
 	constexpr TCHAR TSubclassOfTemplate[] = TEXT("TSubclassOf<%s%s>");
 	constexpr TCHAR TScriptInterfaceTemplate[] = TEXT("TScriptInterface<%s%s>");
 
-	const inline TCHAR* BoolType = TEXT("bool");
-	const inline TCHAR* FloatType = TEXT("float");
-	const inline TCHAR* DoubleType = TEXT("double");
-	const inline TCHAR* IntType = TEXT("int");
-	const inline TCHAR* Int32Type = TEXT("int32");
-	const inline TCHAR* Int64Type = TEXT("int64");
-	const inline TCHAR* UInt8Type = TEXT("uint8");
-	const inline TCHAR* UInt32Type = TEXT("uint32");
-	const inline TCHAR* UInt64Type = TEXT("uint64");
-	const inline TCHAR* FNameType = TEXT("FName");
-	const inline TCHAR* FStringType = TEXT("FString");
-	const inline TCHAR* FTextType = TEXT("FText");
-	const inline TCHAR* BoolArrayType = TEXT("TArray<bool>");
-	const inline TCHAR* FloatArrayType = TEXT("TArray<float>");
-	const inline TCHAR* DoubleArrayType = TEXT("TArray<double>");
-	const inline TCHAR* Int32ArrayType = TEXT("TArray<int32>");
-	const inline TCHAR* UInt32ArrayType = TEXT("TArray<uint32>");
-	const inline TCHAR* UInt8ArrayType = TEXT("TArray<uint8>");
-	const inline TCHAR* FNameArrayType = TEXT("TArray<FName>");
-	const inline TCHAR* FStringArrayType = TEXT("TArray<FString>");
-	const inline TCHAR* FTextArrayType = TEXT("TArray<FText>");
-
-	const FLazyName BoolTypeName(BoolType);
-	const FLazyName FloatTypeName(FloatType);
-	const FLazyName DoubleTypeName(DoubleType);
-	const FLazyName IntTypeName(IntType);
-	const FLazyName Int32TypeName(Int32Type);
-	const FLazyName Int64TypeName(Int64Type);
-	const FLazyName UInt8TypeName(UInt8Type);
-	const FLazyName UInt32TypeName(UInt32Type);
-	const FLazyName UInt64TypeName(UInt64Type);
-	const FLazyName FNameTypeName(FNameType);
-	const FLazyName FStringTypeName(FStringType);
-	const FLazyName FTextTypeName(FTextType);
-	const FLazyName BoolArrayTypeName(BoolArrayType);
-	const FLazyName FloatArrayTypeName(FloatArrayType);
-	const FLazyName DoubleArrayTypeName(DoubleArrayType);
-	const FLazyName Int32ArrayTypeName(Int32ArrayType);
-	const FLazyName UInt32ArrayTypeName(UInt32ArrayType);
-	const FLazyName UInt8ArrayTypeName(UInt8ArrayType);
-	const FLazyName FNameArrayTypeName(FNameArrayType);
-	const FLazyName FStringArrayTypeName(FStringArrayType);
-
+	class RIGVM_API TypeString
+	{
+	public:
+		static const inline TCHAR* Bool = TEXT("bool");
+		static const inline TCHAR* Float = TEXT("float");
+		static const inline TCHAR* Double = TEXT("double");
+		static const inline TCHAR* Int = TEXT("int");
+		static const inline TCHAR* Int32 = TEXT("int32");
+		static const inline TCHAR* Int64 = TEXT("int64");
+		static const inline TCHAR* UInt8 = TEXT("uint8");
+		static const inline TCHAR* UInt32 = TEXT("uint32");
+		static const inline TCHAR* UInt64 = TEXT("uint64");
+		static const inline TCHAR* FName = TEXT("FName");
+		static const inline TCHAR* FString = TEXT("FString");
+		static const inline TCHAR* FText = TEXT("FText");
+		static const inline TCHAR* BoolArray = TEXT("TArray<bool>");
+		static const inline TCHAR* FloatArray = TEXT("TArray<float>");
+		static const inline TCHAR* DoubleArray = TEXT("TArray<double>");
+		static const inline TCHAR* Int32Array = TEXT("TArray<int32>");
+		static const inline TCHAR* Int64Array = TEXT("TArray<int64>");
+		static const inline TCHAR* UInt8Array = TEXT("TArray<uint8>");
+		static const inline TCHAR* UInt32Array = TEXT("TArray<uint32>");
+		static const inline TCHAR* UInt64Array = TEXT("TArray<uint64>");
+		static const inline TCHAR* FNameArray = TEXT("TArray<FName>");
+		static const inline TCHAR* FStringArray = TEXT("TArray<FString>");
+		static const inline TCHAR* FTextArray = TEXT("TArray<FText>");
+	};
+	
+	class RIGVM_API TypeName
+	{
+	public:
+		static const inline FLazyName Bool = FLazyName(TypeString::Bool);
+		static const inline FLazyName Float = FLazyName(TypeString::Float);
+		static const inline FLazyName Double = FLazyName(TypeString::Double);
+		static const inline FLazyName Int = FLazyName(TypeString::Int);
+		static const inline FLazyName Int32 = FLazyName(TypeString::Int32);
+		static const inline FLazyName Int64 = FLazyName(TypeString::Int64);
+		static const inline FLazyName UInt8 = FLazyName(TypeString::UInt8);
+		static const inline FLazyName UInt32 = FLazyName(TypeString::UInt32);
+		static const inline FLazyName UInt64 = FLazyName(TypeString::UInt64);
+		static const inline FLazyName FName = FLazyName(TypeString::FName);
+		static const inline FLazyName FString = FLazyName(TypeString::FString);
+		static const inline FLazyName FText = FLazyName(TypeString::FText);
+		static const inline FLazyName BoolArray = FLazyName(TypeString::BoolArray);
+		static const inline FLazyName FloatArray = FLazyName(TypeString::FloatArray);
+		static const inline FLazyName DoubleArray = FLazyName(TypeString::DoubleArray);
+		static const inline FLazyName Int32Array = FLazyName(TypeString::Int32Array);
+		static const inline FLazyName Int64Array = FLazyName(TypeString::Int64Array);
+		static const inline FLazyName UInt8Array = FLazyName(TypeString::UInt8Array);
+		static const inline FLazyName UInt32Array = FLazyName(TypeString::UInt32Array);
+		static const inline FLazyName UInt64Array = FLazyName(TypeString::UInt64Array);
+		static const inline FLazyName FNameArray = FLazyName(TypeString::FNameArray);
+		static const inline FLazyName FStringArray = FLazyName(TypeString::FStringArray);
+		static const inline FLazyName FTextArray = FLazyName(TypeString::FTextArray);
+	};
+	
 	class RIGVM_API TypeIndex
 	{
 	public:
@@ -107,23 +120,58 @@ namespace RigVMTypeUtils
 		static TRigVMTypeIndex Bool;	
 		static TRigVMTypeIndex Float;	
 		static TRigVMTypeIndex Double;	
-		static TRigVMTypeIndex Int32;	
-		static TRigVMTypeIndex UInt32;	
+		static TRigVMTypeIndex Int32;
+		static TRigVMTypeIndex Int64;
+		static TRigVMTypeIndex UInt32;
+		static TRigVMTypeIndex UInt64;	
 		static TRigVMTypeIndex UInt8;	
 		static TRigVMTypeIndex FName;	
 		static TRigVMTypeIndex FString;
+		static TRigVMTypeIndex FText;
 		static TRigVMTypeIndex WildCard;	
 		static TRigVMTypeIndex BoolArray;	
 		static TRigVMTypeIndex FloatArray;	
 		static TRigVMTypeIndex DoubleArray;	
 		static TRigVMTypeIndex Int32Array;	
-		static TRigVMTypeIndex UInt32Array;	
+		static TRigVMTypeIndex Int64Array;	
 		static TRigVMTypeIndex UInt8Array;	
+		static TRigVMTypeIndex UInt32Array;	
+		static TRigVMTypeIndex UInt64Array;	
 		static TRigVMTypeIndex FNameArray;	
-		static TRigVMTypeIndex FStringArray;	
+		static TRigVMTypeIndex FStringArray;
+		static TRigVMTypeIndex FTextArray;
 		static TRigVMTypeIndex WildCardArray;	
 	};
 
+	class RIGVM_API Types
+	{
+	public:
+		static FRigVMType Bool;
+		static FRigVMType Float;
+		static FRigVMType Double;
+		static FRigVMType Int32;
+		static FRigVMType Int64;
+		static FRigVMType UInt8;
+		static FRigVMType UInt32;
+		static FRigVMType UInt64;
+		static FRigVMType FName;
+		static FRigVMType FString;
+		static FRigVMType FText;
+		static FRigVMType WildCard;
+		static FRigVMType BoolArray;
+		static FRigVMType FloatArray;
+		static FRigVMType DoubleArray;
+		static FRigVMType Int32Array;
+		static FRigVMType Int64Array;
+		static FRigVMType UInt8Array;
+		static FRigVMType UInt32Array;
+		static FRigVMType UInt64Array;
+		static FRigVMType FNameArray;
+		static FRigVMType FStringArray;
+		static FRigVMType FTextArray;
+		static FRigVMType WildCardArray;
+	};
+
 	// Returns true if the type specified is an array
 	inline bool IsArrayType(const FString& InCPPType)
 	{
@@ -187,7 +235,7 @@ namespace RigVMTypeUtils
 		return WildCardTypeObject;
 	}
 
-	static const FString& GetWildCardCPPType()
+	static const FString& GetWildCardCPPTypeString()
 	{
 		static const FString WildCardCPPType = GetUniqueStructTypeName(FRigVMUnknownType::StaticStruct()); 
 		return WildCardCPPType;
@@ -195,19 +243,19 @@ namespace RigVMTypeUtils
 
 	static const FLazyName& GetWildCardCPPTypeName()
 	{
-		static const FLazyName WildCardCPPTypeName(*GetWildCardCPPType()); 
+		static const FLazyName WildCardCPPTypeName(*GetWildCardCPPTypeString()); 
 		return WildCardCPPTypeName;
 	}
 
-	static const FString& GetWildCardArrayCPPType()
+	static const FString& GetWildCardArrayCPPTypeString()
 	{
-		static const FString WildCardArrayCPPType = ArrayTypeFromBaseType(GetWildCardCPPType()); 
+		static const FString WildCardArrayCPPType = ArrayTypeFromBaseType(GetWildCardCPPTypeString()); 
 		return WildCardArrayCPPType;
 	}
 
 	static const FLazyName& GetWildCardArrayCPPTypeName()
 	{
-		static const FLazyName WildCardArrayCPPTypeName(*GetWildCardArrayCPPType()); 
+		static const FLazyName WildCardArrayCPPTypeName(*GetWildCardArrayCPPTypeString()); 
 		return WildCardArrayCPPTypeName;
 	}
 
@@ -243,7 +291,10 @@ namespace RigVMTypeUtils
 			TEXT("TArray<TArray<FString>>"), 
 			TEXT("FName"), 
 			TEXT("TArray<FName>"),
-			TEXT("TArray<TArray<FName>>") 
+			TEXT("TArray<TArray<FName>>"),
+			TEXT("FText"), 
+			TEXT("TArray<FText>"),
+			TEXT("TArray<TArray<FText>>") 
 		};
 
 		if(!CPPTypesNotRequiringCPPTypeObject.Contains(InCPPType))
diff --git a/Engine/Plugins/Runtime/RigVM/Source/RigVMDeveloper/Private/RigVMCompiler/RigVMCodeGenerator.cpp b/Engine/Plugins/Runtime/RigVM/Source/RigVMDeveloper/Private/RigVMCompiler/RigVMCodeGenerator.cpp
--- a/Engine/Plugins/Runtime/RigVM/Source/RigVMDeveloper/Private/RigVMCompiler/RigVMCodeGenerator.cpp
+++ b/Engine/Plugins/Runtime/RigVM/Source/RigVMDeveloper/Private/RigVMCompiler/RigVMCodeGenerator.cpp
@@ -358,7 +358,7 @@ FString FRigVMCodeGenerator::DumpProperties(const FRigVMExtendedExecuteContext&
 					bUseConstExpr = false;
 				}
 
-				if (CPPType == RigVMTypeUtils::FNameType || CPPType == RigVMTypeUtils::FStringType)
+				if (CPPType == RigVMTypeUtils::TypeString::FName || CPPType == RigVMTypeUtils::TypeString::FString)
 				{
 					bUseConstExpr = false;
 				}
@@ -1754,9 +1754,9 @@ FString FRigVMCodeGenerator::GetOperandName(const FRigVMExtendedExecuteContext&
 				FString ExtendedType;
 				FString CPPType = CurrentProperty->GetCPPType(&ExtendedType);
 				CPPType += ExtendedType;
-				if(CPPType == RigVMTypeUtils::UInt8Type && CastField<FBoolProperty>(CurrentProperty))
+				if(CPPType == RigVMTypeUtils::TypeString::UInt8 && CastField<FBoolProperty>(CurrentProperty))
 				{
-					CPPType = RigVMTypeUtils::BoolType;
+					CPPType = RigVMTypeUtils::TypeString::Bool;
 				}
 
 				SegmentPath = Format(RigVM_GetArrayElementSafeFormat, *CPPType, *SegmentPath, *Left);
@@ -1823,7 +1823,7 @@ FString FRigVMCodeGenerator::SanitizeName(const FString& InName, const FString&
 		Name.ReplaceInline(RigVM_DoubleUnderscoreFormat, RigVM_SingleUnderscoreFormat);
 	}
 
-	if(CPPType == RigVMTypeUtils::BoolType && !Name.StartsWith(RigVM_BoolPropertyPrefix, ESearchCase::CaseSensitive))
+	if(CPPType == RigVMTypeUtils::TypeString::Bool && !Name.StartsWith(RigVM_BoolPropertyPrefix, ESearchCase::CaseSensitive))
 	{
 		Name = RigVM_BoolPropertyPrefix + Name;
 	}
@@ -1958,16 +1958,16 @@ FString FRigVMCodeGenerator::SanitizeValue(const FString& InValue, const FString
 			}
 		}
 
-		if (BaseCPPType == RigVMTypeUtils::BoolType)
+		if (BaseCPPType == RigVMTypeUtils::TypeString::Bool)
 		{
 			DefaultValue.ToLowerInline();
 		}
 
-		if (CPPType == RigVMTypeUtils::FNameType || CPPType == RigVMTypeUtils::FStringType)
+		if (CPPType == RigVMTypeUtils::TypeString::FName || CPPType == RigVMTypeUtils::TypeString::FString)
 		{
 			if(DefaultValue.IsEmpty())
 			{
-				if(CPPType == RigVMTypeUtils::FNameType)
+				if(CPPType == RigVMTypeUtils::TypeString::FName)
 				{
 					DefaultValue = RigVM_NameNoneFormat;
 				}
diff --git a/Engine/Plugins/Runtime/RigVM/Source/RigVMDeveloper/Private/RigVMCompiler/RigVMCompiler.cpp b/Engine/Plugins/Runtime/RigVM/Source/RigVMDeveloper/Private/RigVMCompiler/RigVMCompiler.cpp
--- a/Engine/Plugins/Runtime/RigVM/Source/RigVMDeveloper/Private/RigVMCompiler/RigVMCompiler.cpp
+++ b/Engine/Plugins/Runtime/RigVM/Source/RigVMDeveloper/Private/RigVMCompiler/RigVMCompiler.cpp
@@ -130,7 +130,7 @@ FRigVMOperand FRigVMCompilerWorkData::AddProperty(
 		}
 	}
 	
-	FRigVMPropertyDescription Description(InName, ArgumentType.CPPType.ToString(), ArgumentType.CPPTypeObject, InDefaultValue);
+	FRigVMPropertyDescription Description(InName, ArgumentType.GetTypeString(), ArgumentType.GetTypeObject(), InDefaultValue);
 	ensure(FindProperty(InMemoryType, Description.Name).IsValid() == false); // Warning : this check can not be done before the SanitizeName inside FRigVMPropertyDescription constructor
 
 	TArray<FRigVMPropertyDescription>& PropertyArray = PropertyDescriptions.FindOrAdd(InMemoryType);
@@ -872,14 +872,14 @@ bool URigVMCompiler::Compile(const FRigVMCompileSettings& InSettings, TArray<URi
 							{
 								if (Pin->GetTypeIndex() != ExpectedArgument.GetTypeIndex(0))
 								{
-									FString MissingPinMessage = FString::Printf(TEXT("Could not find pin %s of type %s in Node @@."), *ExpectedArgument.Name.ToString(), *Registry.GetType(ExpectedArgument.GetTypeIndex(0)).CPPType.ToString());
+									FString MissingPinMessage = FString::Printf(TEXT("Could not find pin %s of type %s in Node @@."), *ExpectedArgument.Name.ToString(), *Registry.GetType(ExpectedArgument.GetTypeIndex(0)).GetTypeString());
 									Settings.ASTSettings.Report(EMessageSeverity::Error, ModelNode, MissingPinMessage);
 									bEncounteredGraphError = true;
 								}
 							}
 							else
 							{
-								FString MissingPinMessage = FString::Printf(TEXT("Could not find pin %s of type %s in Node @@."), *ExpectedArgument.Name.ToString(), *Registry.GetType(ExpectedArgument.GetTypeIndex(0)).CPPType.ToString());
+								FString MissingPinMessage = FString::Printf(TEXT("Could not find pin %s of type %s in Node @@."), *ExpectedArgument.Name.ToString(), *Registry.GetType(ExpectedArgument.GetTypeIndex(0)).GetTypeString());
 								Settings.ASTSettings.Report(EMessageSeverity::Error, ModelNode, MissingPinMessage);
 								bEncounteredGraphError = true;
 							}
@@ -967,7 +967,7 @@ bool URigVMCompiler::Compile(const FRigVMCompileSettings& InSettings, TArray<URi
 							if(Registry.IsExecuteType(Argument->TypeIndex))
 							{
 								const FRigVMTemplateArgumentType& Type = Registry.GetType(Argument->TypeIndex);
-								if(UScriptStruct* ExecuteContextStruct = Cast<UScriptStruct>(Type.CPPTypeObject))
+								if(const UScriptStruct* ExecuteContextStruct = Cast<UScriptStruct>(Type.GetTypeObject()))
 								{
 									if(!ValidExecuteContextStructs.Contains(ExecuteContextStruct))
 									{
@@ -2099,7 +2099,7 @@ bool URigVMCompiler::TraverseCallExtern(const FRigVMCallExternExprAST* InExpr, F
 				if(!WorkData.TraitListLiterals.Contains(nullptr))
 				{
 					const FName Name = WorkData.GetUniquePropertyName(ERigVMMemoryType::Literal, TEXT("EmptyTraitList"));
-					const FRigVMOperand& ListOperand = WorkData.AddProperty(ERigVMMemoryType::Literal, Name, RigVMTypeUtils::ArrayTypeFromBaseType(RigVMTypeUtils::Int32Type), nullptr, TEXT("()"));
+					const FRigVMOperand& ListOperand = WorkData.AddProperty(ERigVMMemoryType::Literal, Name, RigVMTypeUtils::ArrayTypeFromBaseType(RigVMTypeUtils::TypeString::Int32), nullptr, TEXT("()"));
 					WorkData.TraitListLiterals.Add(nullptr, ListOperand);
 				}
 
@@ -2136,7 +2136,7 @@ bool URigVMCompiler::TraverseCallExtern(const FRigVMCallExternExprAST* InExpr, F
 					{
 						const FName Name = WorkData.GetUniquePropertyName(ERigVMMemoryType::Literal, TEXT("TraitList"));
 						const FString DefaultValue = FString::Printf(TEXT("(%s)"), *FString::Join(DefaultValues, TEXT(",")));
-						const FRigVMOperand& ListOperand = WorkData.AddProperty(ERigVMMemoryType::Literal, Name, RigVMTypeUtils::ArrayTypeFromBaseType(RigVMTypeUtils::Int32Type), nullptr, DefaultValue);
+						const FRigVMOperand& ListOperand = WorkData.AddProperty(ERigVMMemoryType::Literal, Name, RigVMTypeUtils::ArrayTypeFromBaseType(RigVMTypeUtils::TypeString::Int32), nullptr, DefaultValue);
 						WorkData.TraitListLiterals.Add(Node, ListOperand);
 					}
 				}
@@ -3347,8 +3347,8 @@ void URigVMCompiler::AddCopyOperator(const FRigVMCopyOp& InOp, const FRigVMAssig
 			{
 				const FRigVMRegistry& Registry = FRigVMRegistry::Get();
 				static constexpr TCHAR MissingCastMessage[] = TEXT("Cast (%s to %s) for Node @@ not found.");
-				const FString& SourceCPPType = Registry.GetType(SourceTypeIndex).CPPType.ToString();
-				const FString& TargetCPPType = Registry.GetType(TargetTypeIndex).CPPType.ToString();
+				const FString& SourceCPPType = Registry.GetType(SourceTypeIndex).GetTypeString();
+				const FString& TargetCPPType = Registry.GetType(TargetTypeIndex).GetTypeString();
 				WorkData.Settings.Report(EMessageSeverity::Error, InAssignExpr->GetTargetPin()->GetNode(),
 					FString::Printf(MissingCastMessage, *SourceCPPType, *TargetCPPType));
 				return;
diff --git a/Engine/Plugins/Runtime/RigVM/Source/RigVMDeveloper/Private/RigVMDeveloperTypeUtils.cpp b/Engine/Plugins/Runtime/RigVM/Source/RigVMDeveloper/Private/RigVMDeveloperTypeUtils.cpp
--- a/Engine/Plugins/Runtime/RigVM/Source/RigVMDeveloper/Private/RigVMDeveloperTypeUtils.cpp
+++ b/Engine/Plugins/Runtime/RigVM/Source/RigVMDeveloper/Private/RigVMDeveloperTypeUtils.cpp
@@ -47,9 +47,9 @@ FRigVMExternalVariable RigVMTypeUtils::ExternalVariableFromRigVMVariableDescript
 FEdGraphPinType RigVMTypeUtils::PinTypeFromTypeIndex(const TRigVMTypeIndex& InTypeIndex)
 {
 	const FRigVMTemplateArgumentType& Type = FRigVMRegistry::Get().GetType(InTypeIndex);
-	if(Type.CPPType.IsValid())
+	if(Type.GetType().IsValid())
 	{
-		return PinTypeFromCPPType(Type.CPPType, Type.CPPTypeObject);
+		return PinTypeFromCPPType(Type.GetTypeFName(), const_cast<UObject*>(Type.GetTypeObject()));
 	}
 	return FEdGraphPinType();
 }
@@ -68,32 +68,40 @@ FEdGraphPinType RigVMTypeUtils::PinTypeFromCPPType(const FName& InCPPType, UObje
 		PinType.ContainerType = EPinContainerType::Array;
 	}
 
-	if (BaseCPPType == BoolTypeName)
+	if (BaseCPPType == TypeName::Bool)
 	{
 		PinType.PinCategory = UEdGraphSchema_K2::PC_Boolean;
 	}
-	else if (BaseCPPType == Int32TypeName)
+	else if (BaseCPPType == TypeName::Int32)
 	{
 		PinType.PinCategory = UEdGraphSchema_K2::PC_Int;
 	}
-	else if (BaseCPPType == FloatTypeName)
+	else if (BaseCPPType == TypeName::Int64)
+	{
+		PinType.PinCategory = UEdGraphSchema_K2::PC_Int64;
+	}
+	else if (BaseCPPType == TypeName::Float)
 	{
 		PinType.PinCategory = UEdGraphSchema_K2::PC_Real;
 		PinType.PinSubCategory = UEdGraphSchema_K2::PC_Float;
 	}
-	else if (BaseCPPType == DoubleTypeName)
+	else if (BaseCPPType == TypeName::Double)
 	{
 		PinType.PinCategory = UEdGraphSchema_K2::PC_Real;
 		PinType.PinSubCategory = UEdGraphSchema_K2::PC_Double;
 	}
-	else if (BaseCPPType == FNameTypeName)
+	else if (BaseCPPType == TypeName::FName)
 	{
 		PinType.PinCategory = UEdGraphSchema_K2::PC_Name;
 	}
-	else if (BaseCPPType == FStringTypeName)
+	else if (BaseCPPType == TypeName::FString)
 	{
 		PinType.PinCategory = UEdGraphSchema_K2::PC_String;
 	}
+	else if (BaseCPPType == TypeName::FText)
+	{
+		PinType.PinCategory = UEdGraphSchema_K2::PC_Text;
+	}
 	else if (Cast<UScriptStruct>(InCPPTypeObject))
 	{
 		PinType.PinCategory = UEdGraphSchema_K2::PC_Struct;
@@ -120,32 +128,40 @@ FEdGraphPinType RigVMTypeUtils::PinTypeFromExternalVariable(const FRigVMExternal
 	PinType.ResetToDefaults();
 	PinType.PinCategory = NAME_None;
 
-	if (InExternalVariable.TypeName == BoolTypeName)
+	if (InExternalVariable.TypeName == TypeName::Bool)
 	{
 		PinType.PinCategory = UEdGraphSchema_K2::PC_Boolean;
 	}
-	else if (InExternalVariable.TypeName == Int32TypeName)
+	else if (InExternalVariable.TypeName == TypeName::Int32)
 	{
 		PinType.PinCategory = UEdGraphSchema_K2::PC_Int;
 	}
-	else if (InExternalVariable.TypeName == FloatTypeName)
+	else if (InExternalVariable.TypeName == TypeName::Int64)
+	{
+		PinType.PinCategory = UEdGraphSchema_K2::PC_Int64;
+	}
+	else if (InExternalVariable.TypeName == TypeName::Float)
 	{
 		PinType.PinCategory = UEdGraphSchema_K2::PC_Real;
 		PinType.PinSubCategory = UEdGraphSchema_K2::PC_Float;
 	}
-	else if (InExternalVariable.TypeName == DoubleTypeName)
+	else if (InExternalVariable.TypeName == TypeName::Double)
 	{
 		PinType.PinCategory = UEdGraphSchema_K2::PC_Real;
 		PinType.PinSubCategory = UEdGraphSchema_K2::PC_Double;
 	}
-	else if (InExternalVariable.TypeName == FNameTypeName)
+	else if (InExternalVariable.TypeName == TypeName::FName)
 	{
 		PinType.PinCategory = UEdGraphSchema_K2::PC_Name;
 	}
-	else if (InExternalVariable.TypeName == FStringTypeName)
+	else if (InExternalVariable.TypeName == TypeName::FString)
 	{
 		PinType.PinCategory = UEdGraphSchema_K2::PC_String;
 	}
+	else if (InExternalVariable.TypeName == TypeName::FText)
+	{
+		PinType.PinCategory = UEdGraphSchema_K2::PC_Text;
+	}
 	else if (Cast<UScriptStruct>(InExternalVariable.TypeObject))
 	{
 		PinType.PinCategory = UEdGraphSchema_K2::PC_Struct;
@@ -192,32 +208,40 @@ FEdGraphPinType RigVMTypeUtils::PinTypeFromRigVMVariableDescription(
 		PinType.ContainerType = EPinContainerType::None;
 	}
 
-	if (CurrentCPPType == BoolType)
+	if (CurrentCPPType == TypeString::Bool)
 	{
 		PinType.PinCategory = UEdGraphSchema_K2::PC_Boolean;
 	}
-	else if (CurrentCPPType == Int32Type)
+	else if (CurrentCPPType == TypeString::Int32)
 	{
 		PinType.PinCategory = UEdGraphSchema_K2::PC_Int;
 	}
-	else if (CurrentCPPType == FloatType)
+	else if (CurrentCPPType == TypeString::Int64)
+	{
+		PinType.PinCategory = UEdGraphSchema_K2::PC_Int64;
+	}
+	else if (CurrentCPPType == TypeString::Float)
 	{
 		PinType.PinCategory = UEdGraphSchema_K2::PC_Real;
 		PinType.PinSubCategory = UEdGraphSchema_K2::PC_Float;
 	}
-	else if (CurrentCPPType == DoubleType)
+	else if (CurrentCPPType == TypeString::Double)
 	{
 		PinType.PinCategory = UEdGraphSchema_K2::PC_Real;
 		PinType.PinSubCategory = UEdGraphSchema_K2::PC_Double;
 	}
-	else if (CurrentCPPType == FNameType)
+	else if (CurrentCPPType == TypeString::FName)
 	{
 		PinType.PinCategory = UEdGraphSchema_K2::PC_Name;
 	}
-	else if (CurrentCPPType == FStringType)
+	else if (CurrentCPPType == TypeString::FString)
 	{
 		PinType.PinCategory = UEdGraphSchema_K2::PC_String;
 	}
+	else if (CurrentCPPType == TypeString::FText)
+	{
+		PinType.PinCategory = UEdGraphSchema_K2::PC_Text;
+	}
 	else if (Cast<UScriptStruct>(InVariableDescription.CPPTypeObject))
 	{
 		PinType.PinCategory = UEdGraphSchema_K2::PC_Struct;
@@ -347,21 +371,25 @@ bool RigVMTypeUtils::CPPTypeFromPinType(const FEdGraphPinType& InPinType, FStrin
 	*OutCPPTypeObject = nullptr;
 	if (InPinType.PinCategory == UEdGraphSchema_K2::PC_Boolean)
 	{
-		OutCPPType = Prefix + BoolType + Suffix;
+		OutCPPType = Prefix + TypeString::Bool + Suffix;
 	}
 	else if (InPinType.PinCategory == UEdGraphSchema_K2::PC_Int)
 	{
-		OutCPPType = Prefix + Int32Type + Suffix;
+		OutCPPType = Prefix + TypeString::Int32 + Suffix;
+	}
+	else if (InPinType.PinCategory == UEdGraphSchema_K2::PC_Int64)
+	{
+		OutCPPType = Prefix + TypeString::Int64 + Suffix;
 	}
 	else if (InPinType.PinCategory == UEdGraphSchema_K2::PC_Real)
 	{
 		if (InPinType.PinSubCategory == UEdGraphSchema_K2::PC_Float)
 		{
-			OutCPPType = Prefix + FloatType + Suffix;
+			OutCPPType = Prefix + TypeString::Float + Suffix;
 		}
 		else if (InPinType.PinSubCategory == UEdGraphSchema_K2::PC_Double)
 		{
-			OutCPPType = Prefix + DoubleType + Suffix;
+			OutCPPType = Prefix + TypeString::Double + Suffix;
 		}
 		else
 		{
@@ -370,19 +398,23 @@ bool RigVMTypeUtils::CPPTypeFromPinType(const FEdGraphPinType& InPinType, FStrin
 	}
 	else if (InPinType.PinCategory == UEdGraphSchema_K2::PC_Float)
 	{
-		OutCPPType = Prefix + FloatType + Suffix;
+		OutCPPType = Prefix + TypeString::Float + Suffix;
 	}
 	else if (InPinType.PinCategory == UEdGraphSchema_K2::PC_Double)
 	{
-		OutCPPType = Prefix + DoubleType + Suffix;
+		OutCPPType = Prefix + TypeString::Double + Suffix;
 	}
 	else if (InPinType.PinCategory == UEdGraphSchema_K2::PC_Name)
 	{
-		OutCPPType = Prefix + FNameType + Suffix;
+		OutCPPType = Prefix + TypeString::FName + Suffix;
 	}
 	else if (InPinType.PinCategory == UEdGraphSchema_K2::PC_String)
 	{
-		OutCPPType = Prefix + FStringType + Suffix;
+		OutCPPType = Prefix + TypeString::FString + Suffix;
+	}
+	else if (InPinType.PinCategory == UEdGraphSchema_K2::PC_Text)
+	{
+		OutCPPType = Prefix + TypeString::FText + Suffix;
 	}
 	else if (InPinType.PinCategory == UEdGraphSchema_K2::PC_Struct)
 	{
@@ -412,7 +444,7 @@ bool RigVMTypeUtils::CPPTypeFromPinType(const FEdGraphPinType& InPinType, FStrin
 		}
 		else
 		{
-			OutCPPType = Prefix + UInt8Type + Suffix;
+			OutCPPType = Prefix + TypeString::UInt8 + Suffix;
 		}
 	}
 	else
@@ -479,29 +511,49 @@ bool RigVMTypeUtils::CPPTypeFromExternalVariable(const FRigVMExternalVariable& I
 	}
 
 	*OutCPPTypeObject = nullptr;
-	if (InExternalVariable.TypeName == BoolTypeName)
+	if (InExternalVariable.TypeName == TypeName::Bool)
+	{
+		OutCPPType = Prefix + TypeString::Bool + Suffix;
+	}
+	else if (InExternalVariable.TypeName == TypeName::Int32)
+	{
+		OutCPPType = Prefix + TypeString::Int32 + Suffix;
+	}
+	else if (InExternalVariable.TypeName == TypeName::Int64)
+	{
+		OutCPPType = Prefix + TypeString::Int64 + Suffix;
+	}
+	else if (InExternalVariable.TypeName == TypeName::UInt8)
+	{
+		OutCPPType = Prefix + TypeString::UInt8 + Suffix;
+	}
+	else if (InExternalVariable.TypeName == TypeName::UInt32)
 	{
-		OutCPPType = Prefix + BoolType + Suffix;
+		OutCPPType = Prefix + TypeString::UInt32 + Suffix;
 	}
-	else if (InExternalVariable.TypeName == Int32TypeName)
+	else if (InExternalVariable.TypeName == TypeName::UInt64)
 	{
-		OutCPPType = Prefix + Int32Type + Suffix;
+		OutCPPType = Prefix + TypeString::UInt64 + Suffix;
 	}
-	else if (InExternalVariable.TypeName == FloatTypeName)
+	else if (InExternalVariable.TypeName == TypeName::Float)
 	{
-		OutCPPType = Prefix + FloatType + Suffix;
+		OutCPPType = Prefix + TypeString::Float + Suffix;
 	}
-	else if (InExternalVariable.TypeName == DoubleTypeName)
+	else if (InExternalVariable.TypeName == TypeName::Double)
 	{
-		OutCPPType = Prefix + DoubleType + Suffix;
+		OutCPPType = Prefix + TypeString::Double + Suffix;
 	}
-	else if (InExternalVariable.TypeName == FNameTypeName)
+	else if (InExternalVariable.TypeName == TypeName::FName)
 	{
-		OutCPPType = Prefix + FNameType + Suffix;
+		OutCPPType = Prefix + TypeString::FName + Suffix;
 	}
-	else if (InExternalVariable.TypeName == FStringTypeName)
+	else if (InExternalVariable.TypeName == TypeName::FString)
 	{
-		OutCPPType = Prefix + FStringType + Suffix;
+		OutCPPType = Prefix + TypeString::FString + Suffix;
+	}
+	else if (InExternalVariable.TypeName == TypeName::FText)
+	{
+		OutCPPType = Prefix + TypeString::FText + Suffix;
 	}
 	else if (UScriptStruct* Struct = Cast<UScriptStruct>(InExternalVariable.TypeObject))
 	{
@@ -559,8 +611,8 @@ bool RigVMTypeUtils::AreCompatible(const FEdGraphPinType& InTypeA, const FEdGrap
 	// Every other case is evaluated by UEdGraphSchema_K2::ArePinTypesCompatible
 	if (SubPinTypeA.ContainerType == SubPinTypeB.ContainerType)
 	{
-		if ((SubPinTypeA.PinCategory == FloatTypeName && SubPinTypeB.PinCategory == DoubleTypeName) ||
-			(SubPinTypeA.PinCategory == DoubleTypeName && SubPinTypeB.PinCategory == FloatTypeName))
+		if ((SubPinTypeA.PinCategory == TypeName::Float && SubPinTypeB.PinCategory == TypeName::Double) ||
+			(SubPinTypeA.PinCategory == TypeName::Double && SubPinTypeB.PinCategory == TypeName::Float))
 		{
 			return true;
 		}
@@ -685,7 +737,7 @@ const FLazyName& RigVMTypeUtils::GetCastTemplateNotation()
 
 FText RigVMTypeUtils::GetDisplayTextForArgumentType(const FRigVMTemplateArgumentType& InType)
 {
-	if(UObject* CPPTypeObject = InType.CPPTypeObject)
+	if(const UObject* CPPTypeObject = InType.GetTypeObject())
 	{
 		if (const UClass* Class = Cast<UClass>(CPPTypeObject))
 		{
@@ -702,7 +754,7 @@ FText RigVMTypeUtils::GetDisplayTextForArgumentType(const FRigVMTemplateArgument
 	}
 	else
 	{
-		FString CPPType = InType.CPPType.ToString();
+		FString CPPType = InType.GetTypeString();
 		if(IsArrayType(CPPType))
 		{
 			CPPType = BaseTypeFromArrayType(CPPType);
@@ -713,27 +765,32 @@ FText RigVMTypeUtils::GetDisplayTextForArgumentType(const FRigVMTemplateArgument
 		static const FText Int32Label = LOCTEXT("Int32Label", "Integer");
 		static const FText FNameLabel = LOCTEXT("FNameLabel", "Name");
 		static const FText FStringLabel = LOCTEXT("FStringLabel", "String");
+		static const FText FTextLabel = LOCTEXT("FTextLabel", "Text");
 
-		if(CPPType == BoolType)
+		if(CPPType == TypeString::Bool)
 		{
 			return BoolLabel;
 		}
-		if(CPPType == FloatType || CPPType == DoubleType)
+		if(CPPType == TypeString::Float || CPPType == TypeString::Double)
 		{
 			return FloatLabel;
 		}
-		if(CPPType == Int32Type)
+		if(CPPType == TypeString::Int32)
 		{
 			return Int32Label;
 		}
-		if(CPPType == FNameType)
+		if(CPPType == TypeString::FName)
 		{
 			return FNameLabel;
 		}
-		if(CPPType == FStringType)
+		if(CPPType == TypeString::FString)
 		{
 			return FStringLabel;
 		}
+		if(CPPType == TypeString::FText)
+		{
+			return FTextLabel;
+		}
 
 		return FText::FromString(CPPType);
 	}
diff --git a/Engine/Plugins/Runtime/RigVM/Source/RigVMDeveloper/Private/RigVMModel/Nodes/RigVMTemplateNode.cpp b/Engine/Plugins/Runtime/RigVM/Source/RigVMDeveloper/Private/RigVMModel/Nodes/RigVMTemplateNode.cpp
--- a/Engine/Plugins/Runtime/RigVM/Source/RigVMDeveloper/Private/RigVMModel/Nodes/RigVMTemplateNode.cpp
+++ b/Engine/Plugins/Runtime/RigVM/Source/RigVMDeveloper/Private/RigVMModel/Nodes/RigVMTemplateNode.cpp
@@ -14,7 +14,7 @@ void FRigVMTemplatePreferredType::UpdateStringFromIndex()
 {
 	static constexpr TCHAR Format[] = TEXT("%s,%s");
 	const FRigVMTemplateArgumentType Type = FRigVMRegistry::Get().GetType(TypeIndex);
-	TypeString = FString::Printf(Format, *Type.CPPType.ToString(), *Type.GetCPPTypeObjectPath().ToString());
+	TypeString = FString::Printf(Format, *Type.GetTypeString(), *Type.GetCPPTypeObjectPath().ToString());
 }
 
 void FRigVMTemplatePreferredType::UpdateIndexFromString()
@@ -47,7 +47,7 @@ void FRigVMTemplatePreferredType::UpdateIndexFromString()
 	TypeIndex = FRigVMRegistry::Get().FindOrAddType(Type);
 
 	// We might have used redirectors. Update the TypeString with the new information.
-	if (TypeIndex != INDEX_NONE && OriginalCPPType != Type.CPPType.ToString())
+	if (TypeIndex != INDEX_NONE && OriginalCPPType != Type.GetTypeString())
 	{
 		UpdateStringFromIndex();
 	}
@@ -213,7 +213,7 @@ FText URigVMTemplateNode::GetToolTipTextForPin(const URigVMPin* InPin) const
 							continue;
 						}
 						TypesPrinted.Add(TypeIndex);
-						FString Type = FRigVMRegistry::Get().GetType(TypeIndex).CPPType.ToString();
+						FString Type = FRigVMRegistry::Get().GetType(TypeIndex).GetTypeString();
 						SupportedTypesJoined += Type + TEXT("\n");
 					}
 					Tooltip = TEXT("Supported Types:\n\n") + SupportedTypesJoined;
@@ -344,7 +344,7 @@ bool URigVMTemplateNode::IsSingleton() const
 bool URigVMTemplateNode::SupportsType(const URigVMPin* InPin, TRigVMTypeIndex InTypeIndex, TRigVMTypeIndex* OutTypeIndex)
 {
 	const FRigVMRegistry& Registry = FRigVMRegistry::Get();
-	static const FString WildCardCPPType = RigVMTypeUtils::GetWildCardCPPType();
+	static const FString WildCardCPPType = RigVMTypeUtils::GetWildCardCPPTypeString();
 	static const FString WildCardArrayCPPType = RigVMTypeUtils::ArrayTypeFromBaseType(WildCardCPPType);
 
 	const URigVMPin* RootPin = InPin->GetRootPin();
@@ -603,7 +603,7 @@ FString URigVMTemplateNode::GetInitialDefaultValueForPin(const FName& InRootPinN
 				else
 				{
 					const FRigVMTemplateArgumentType& Type = FRigVMRegistry::Get().GetType(TypeIndex);
-					if (UScriptStruct* ScriptStruct = Cast<UScriptStruct>(Type.CPPTypeObject))
+					if (const UScriptStruct* ScriptStruct = Cast<UScriptStruct>(Type.GetTypeObject()))
 					{
 						TArray<uint8, TAlignedHeapAllocator<16>> TempBuffer;
 						TempBuffer.AddUninitialized(ScriptStruct->GetStructureSize());
@@ -614,36 +614,36 @@ FString URigVMTemplateNode::GetInitialDefaultValueForPin(const FName& InRootPinN
 						ScriptStruct->ExportText(NewDefaultValue, TempBuffer.GetData(), nullptr, nullptr, PPF_None, nullptr);
 						ScriptStruct->DestroyStruct(TempBuffer.GetData());				
 					}
-					else if (const UEnum* Enum = Cast<UEnum>(Type.CPPTypeObject))
+					else if (const UEnum* Enum = Cast<UEnum>(Type.GetTypeObject()))
 					{
 						NewDefaultValue = Enum->GetNameStringByValue(0);
 					}
-					else if(const UClass* Class = Cast<UClass>(Type.CPPTypeObject))
+					else if(const UClass* Class = Cast<UClass>(Type.GetTypeObject()))
 					{
 						// not supporting objects yet
 						ensure(false);
 					}
-					else if (Type.CPPType == RigVMTypeUtils::FloatTypeName)
+					else if (Type.GetTypeFName() == RigVMTypeUtils::TypeName::Float)
 					{
 						NewDefaultValue = TEXT("0.000000");				
 					}
-					else if (Type.CPPType == RigVMTypeUtils::DoubleTypeName)
+					else if (Type.GetTypeFName() == RigVMTypeUtils::TypeName::Double)
 					{
 						NewDefaultValue = TEXT("0.000000");
 					}
-					else if (Type.CPPType == RigVMTypeUtils::Int32TypeName)
+					else if (Type.GetTypeFName() == RigVMTypeUtils::TypeName::Int32)
 					{
 						NewDefaultValue = TEXT("0");
 					}
-					else if (Type.CPPType == RigVMTypeUtils::BoolTypeName)
+					else if (Type.GetTypeFName() == RigVMTypeUtils::TypeName::Bool)
 					{
 						NewDefaultValue = TEXT("False");
 					}
-					else if (Type.CPPType == RigVMTypeUtils::FStringTypeName)
+					else if (Type.GetTypeFName() == RigVMTypeUtils::TypeName::FString)
 					{
 						NewDefaultValue = TEXT("");
 					}
-					else if (Type.CPPType == RigVMTypeUtils::FNameTypeName)
+					else if (Type.GetTypeFName() == RigVMTypeUtils::TypeName::FName)
 					{
 						NewDefaultValue = TEXT("");
 					}
diff --git a/Engine/Plugins/Runtime/RigVM/Source/RigVMDeveloper/Private/RigVMModel/RigVMController.cpp b/Engine/Plugins/Runtime/RigVM/Source/RigVMDeveloper/Private/RigVMModel/RigVMController.cpp
--- a/Engine/Plugins/Runtime/RigVM/Source/RigVMDeveloper/Private/RigVMModel/RigVMController.cpp
+++ b/Engine/Plugins/Runtime/RigVM/Source/RigVMDeveloper/Private/RigVMModel/RigVMController.cpp
@@ -752,7 +752,7 @@ int32 FRigVMPinInfoArray::AddPin(URigVMController* InController, int32 InParentI
 				const uint8* ElementDefaultValueMemory = nullptr;
 				FStructOnScope ElementDefaultValueMemoryScope;
 
-				if(UScriptStruct* ElementScriptStruct = Cast<UScriptStruct>(ElementType.CPPTypeObject))
+				if(const UScriptStruct* ElementScriptStruct = Cast<UScriptStruct>(ElementType.GetTypeObject()))
 				{
 					ElementDefaultValueMemoryScope = FStructOnScope(ElementScriptStruct);
 		
@@ -764,7 +764,7 @@ int32 FRigVMPinInfoArray::AddPin(URigVMController* InController, int32 InParentI
 				(void)AddPin(InController, Index, *FString::FormatAsNumber(ElementIndex), InDirection, ElementTypeIndex, ElementDefaultValue, InDefaultValueType, ElementDefaultValueMemory, InPreviousPinInfos, bAddSubPins);
 			}
 		}
-		else if(UScriptStruct* ScriptStruct = Cast<UScriptStruct>(Type.CPPTypeObject))
+		else if(UScriptStruct* ScriptStruct = Cast<UScriptStruct>(Type.GetTypeObject()))
 		{
 			const uint8* DefaultValueMemory = InDefaultValueMemory;
 			FStructOnScope DefaultValueMemoryScope;
@@ -1948,7 +1948,7 @@ URigVMVariableNode* URigVMController::AddVariableNode(const FName& InVariableNam
 	}
 
 	URigVMPin* VariablePin = NewObject<URigVMPin>(Node, URigVMVariableNode::VariableName);
-	VariablePin->CPPType = RigVMTypeUtils::FNameType;
+	VariablePin->CPPType = RigVMTypeUtils::TypeString::FName;
 	VariablePin->Direction = ERigVMPinDirection::Hidden;
 	VariablePin->DefaultValue = InVariableName.ToString();
 	VariablePin->CustomWidgetName = TEXT("VariableName");
@@ -9836,8 +9836,8 @@ bool URigVMController::SetPinDefaultValue(URigVMPin* InPin, const FString& InDef
 	check(InPin);
 
 	if(!InPin->IsUObject()
-		&& InPin->GetCPPType() != RigVMTypeUtils::FStringType
-		&& InPin->GetCPPType() != RigVMTypeUtils::FNameType
+		&& InPin->GetCPPType() != RigVMTypeUtils::TypeString::FString
+		&& InPin->GetCPPType() != RigVMTypeUtils::TypeString::FName
 		&& bValidatePinDefaults)
 	{
 		ensure(!InDefaultValue.IsEmpty());
@@ -10058,12 +10058,12 @@ FString URigVMController::GetPinInitialDefaultValue(const URigVMPin* InPin)
 	static const FString TObjectInitValue( TEXT("()") );
 	static const TMap<FString, FString> InitValues =
 	{
-		{ RigVMTypeUtils::BoolType,	TEXT("False") },
-		{ RigVMTypeUtils::Int32Type,	TEXT("0") },
-		{ RigVMTypeUtils::FloatType,	TEXT("0.000000") },
-		{ RigVMTypeUtils::DoubleType,	TEXT("0.000000") },
-		{ RigVMTypeUtils::FNameType,	FName(NAME_None).ToString() },
-		{ RigVMTypeUtils::FStringType,	TEXT("") }
+		{ RigVMTypeUtils::TypeString::Bool,		TEXT("False") },
+		{ RigVMTypeUtils::TypeString::Int32,	TEXT("0") },
+		{ RigVMTypeUtils::TypeString::Float,	TEXT("0.000000") },
+		{ RigVMTypeUtils::TypeString::Double,	TEXT("0.000000") },
+		{ RigVMTypeUtils::TypeString::FName,	FName(NAME_None).ToString() },
+		{ RigVMTypeUtils::TypeString::FString,	TEXT("") }
 	};
 
 	if (InPin->IsStruct())
@@ -12965,7 +12965,7 @@ bool URigVMController::ChangeExposedPinType(const FName& InPinName, const FStrin
 	}
 
 	// We do not allow unresolving exposed pins
-	if (InCPPType == RigVMTypeUtils::GetWildCardCPPType())
+	if (InCPPType == RigVMTypeUtils::GetWildCardCPPTypeString())
 	{
 		ReportError(TEXT("Cannot change exposed pin type to wildcard."));
 		return false;
@@ -16169,7 +16169,7 @@ URigVMEnumNode* URigVMController::AddEnumNode(const FName& InCPPTypeObjectPath,
 	AddNodePin(Node, EnumValuePin);
 
 	URigVMPin* EnumIndexPin = NewObject<URigVMPin>(Node, URigVMEnumNode::EnumIndexName);
-	EnumIndexPin->CPPType = RigVMTypeUtils::Int32Type;
+	EnumIndexPin->CPPType = RigVMTypeUtils::TypeString::Int32;
 	EnumIndexPin->Direction = ERigVMPinDirection::Output;
 	EnumIndexPin->DisplayName = TEXT("Result");
 	AddNodePin(Node, EnumIndexPin);
@@ -16371,7 +16371,7 @@ URigVMInvokeEntryNode* URigVMController::AddInvokeEntryNode(const FName& InEntry
 	AddNodePin(Node, ExecutePin);
 
 	URigVMPin* EntryNamePin = NewObject<URigVMPin>(Node, URigVMInvokeEntryNode::EntryName);
-	EntryNamePin->CPPType = RigVMTypeUtils::FNameType;
+	EntryNamePin->CPPType = RigVMTypeUtils::TypeString::FName;
 	EntryNamePin->Direction = ERigVMPinDirection::Input;
 	EntryNamePin->bIsConstant = true;
 	EntryNamePin->DefaultValue = InEntryName.ToString();
@@ -17117,8 +17117,8 @@ void URigVMController::AddPinsForTemplate(const FRigVMTemplate* InTemplate, cons
 			URigVMPin* Pin = NewObject<URigVMPin>(InNode, Arg->Name);
 			const FRigVMTemplateArgumentType Type = Registry.GetType(Arg->TypeIndex);
 	        
-			Pin->CPPType = Type.CPPType.ToString();
-			Pin->CPPTypeObject = Type.CPPTypeObject;
+			Pin->CPPType = Type.GetTypeString();
+			Pin->CPPTypeObject = Type.GetTypeObject();
 			if (Pin->CPPTypeObject)
 			{
 				Pin->CPPTypeObjectPath = *Pin->CPPTypeObject->GetPathName();
@@ -17156,8 +17156,8 @@ void URigVMController::AddPinsForTemplate(const FRigVMTemplate* InTemplate, cons
 		URigVMPin* Pin = NewObject<URigVMPin>(InNode, Arg->GetName());
 		const TRigVMTypeIndex& TypeIndex = InPinTypeMap.FindChecked(Arg->GetName());
 		const FRigVMTemplateArgumentType Type = FRigVMRegistry::Get().GetType(TypeIndex);
-		Pin->CPPType = Type.CPPType.ToString();
-		Pin->CPPTypeObject = Type.CPPTypeObject;
+		Pin->CPPType = Type.GetTypeString();
+		Pin->CPPTypeObject = Type.GetTypeObject();
 		if (Pin->CPPTypeObject)
 		{
 			Pin->CPPTypeObjectPath = *Pin->CPPTypeObject->GetPathName();
@@ -17945,7 +17945,7 @@ bool URigVMController::GenerateNewPinInfos(const FRigVMRegistry& Registry, URigV
 					const TRigVMTypeIndex Type0 = Arg->GetTypeIndex(0);
 					DefaultValue = Factory->GetArgumentDefaultValue(Arg->Name, Type0);
 					const FRigVMTemplateArgumentType& Type = Registry.GetType(Type0);
-					ArgumentScriptStruct = Cast<UScriptStruct>(Type.CPPTypeObject);
+					ArgumentScriptStruct = Cast<UScriptStruct>(Type.GetTypeObject());
 				}
 			}
 
@@ -18474,11 +18474,11 @@ void URigVMController::RepopulatePinsOnNode(const FRigVMRegistry& Registry, cons
 				UE_LOG(LogRigVMDeveloper, Display, TEXT("Changing pin '%s' type from %s to %s."),
 					*PinPath,
 					*Pin->CPPType,
-					*NewType.CPPType.ToString()
+					*NewType.GetTypeString()
 				);
 #endif
-				Pin->CPPType = NewType.CPPType.ToString();
-				Pin->CPPTypeObject = NewType.CPPTypeObject;
+				Pin->CPPType = NewType.GetTypeString();
+				Pin->CPPTypeObject = NewType.GetTypeObject();
 				check(Pin->CPPTypeObject == nullptr);
 				Pin->CPPTypeObjectPath = NAME_None;
 				Pin->LastKnownTypeIndex = NewPinInfo->TypeIndex;
@@ -19025,8 +19025,8 @@ URigVMPin* URigVMController::CreatePinFromPinInfo(const FRigVMRegistry& InRegist
 	else
 	{
 		const FRigVMTemplateArgumentType& Type = InRegistry.GetType(InPinInfo.TypeIndex);
-		Pin->CPPType = Type.CPPType.ToString();
-		Pin->CPPTypeObject = Type.CPPTypeObject;
+		Pin->CPPType = Type.GetTypeString();
+		Pin->CPPTypeObject = Type.GetTypeObject();
 		if (Pin->CPPTypeObject)
 		{
 			Pin->CPPTypeObjectPath = *Pin->CPPTypeObject->GetPathName();
@@ -19324,7 +19324,7 @@ void URigVMController::PostProcessDefaultValue(const URigVMPin* Pin, FString& Ou
 		{
 			OutDefaultValue = OutDefaultValue.LeftChop(1);
 		}
-		if(OutDefaultValue.IsEmpty() && Pin->GetCPPType() == RigVMTypeUtils::FNameType)
+		if(OutDefaultValue.IsEmpty() && Pin->GetCPPType() == RigVMTypeUtils::TypeString::FName)
 		{
 			OutDefaultValue = FName(NAME_None).ToString();
 		}
@@ -19643,8 +19643,8 @@ bool URigVMController::FullyResolveTemplateNode(URigVMTemplateNode* InNode, int3
 				const FRigVMTemplateArgumentType& Type = FRigVMRegistry::Get().GetType(TypeIndex); 
 				
 				Pin->Direction = MissingPin.GetDirection();
-				Pin->CPPType = Pin->LastKnownCPPType = Type.CPPType.ToString();
-				Pin->CPPTypeObject = Type.CPPTypeObject;
+				Pin->CPPType = Pin->LastKnownCPPType = Type.GetTypeString();
+				Pin->CPPTypeObject = Type.GetTypeObject();
 				Pin->CPPTypeObjectPath = Type.GetCPPTypeObjectPath();
 				Pin->LastKnownTypeIndex = TypeIndex;
 
@@ -20124,16 +20124,16 @@ bool URigVMController::UpdateTemplateNodePinTypes(URigVMTemplateNode* InNode, bo
 					ArrayType = Pin->IsArray() ? FRigVMTemplateArgument::EArrayType::EArrayType_ArrayValue : FRigVMTemplateArgument::EArrayType::EArrayType_SingleValue;
 				}
 				
-				FString CPPType = RigVMTypeUtils::GetWildCardCPPType();
+				FString CPPType = RigVMTypeUtils::GetWildCardCPPTypeString();
 				UObject* CPPObjectType = RigVMTypeUtils::GetWildCardCPPTypeObject();
 
 				if (ArrayType == FRigVMTemplateArgument::EArrayType_ArrayValue)
 				{
-					CPPType = RigVMTypeUtils::GetWildCardArrayCPPType();
+					CPPType = RigVMTypeUtils::GetWildCardArrayCPPTypeString();
 				}
 				else if(ArrayType == FRigVMTemplateArgument::EArrayType_Mixed)
 				{
-					CPPType = Pin->IsArray() ? RigVMTypeUtils::GetWildCardArrayCPPType() : RigVMTypeUtils::GetWildCardCPPType();
+					CPPType = Pin->IsArray() ? RigVMTypeUtils::GetWildCardArrayCPPTypeString() : RigVMTypeUtils::GetWildCardCPPTypeString();
 				}
 
 				// execute pins are no longer part of the template, avoid changing the type
@@ -20160,7 +20160,7 @@ bool URigVMController::UpdateTemplateNodePinTypes(URigVMTemplateNode* InNode, bo
 					bAnyTypeChanged = !Pin->IsExecuteContext();
 					if(bAnyTypeChanged)
 					{
-						const FString CPPType = Registry.GetType(FinalPinTypes[PinIndex]).CPPType.ToString();
+						const FString CPPType = Registry.GetType(FinalPinTypes[PinIndex]).GetTypeString();
 						ReportPinTypeChange(Pin, CPPType);
 					}
 					ChangePinType(Pin, FinalPinTypes[PinIndex], bSetupUndoRedo, false, false, false, bInitializeDefaultValue);
@@ -20458,7 +20458,7 @@ bool URigVMController::ChangePinType(URigVMPin* InPin, TRigVMTypeIndex InTypeInd
 				{
 					if(!TemplateNode->SupportsType(InPin, InTypeIndex))
 					{
-						ReportErrorf(TEXT("ChangePinType: %s doesn't support type '%s'."), *InPin->GetPinPath(), *FRigVMRegistry::Get().GetType(InTypeIndex).CPPType.ToString());
+						ReportErrorf(TEXT("ChangePinType: %s doesn't support type '%s'."), *InPin->GetPinPath(), *FRigVMRegistry::Get().GetType(InTypeIndex).GetTypeString());
 						return false;
 					}
 				}
@@ -20553,9 +20553,9 @@ bool URigVMController::ChangePinType(URigVMPin* InPin, TRigVMTypeIndex InTypeInd
 	const FString PreviousCPPType = InPin->CPPType;
 	
 	const FRigVMTemplateArgumentType Type = FRigVMRegistry::Get().GetType(TypeIndex);
-	InPin->CPPType = Type.CPPType.ToString();
+	InPin->CPPType = Type.GetTypeString();
 	InPin->CPPTypeObjectPath = Type.GetCPPTypeObjectPath();
-	InPin->CPPTypeObject = Type.CPPTypeObject;
+	InPin->CPPTypeObject = Type.GetTypeObject();
 	InPin->bIsDynamicArray = FRigVMRegistry::Get().IsArrayType(TypeIndex);
 
 	if (bInitializeDefaultValue)
@@ -20722,17 +20722,17 @@ bool URigVMController::ChangePinType(URigVMPin* InPin, TRigVMTypeIndex InTypeInd
 				{
 					// plain types within an array or at the root
 					FString SimpleTypeDefaultValue;
-					if(Pin->GetCPPType() == RigVMTypeUtils::BoolType)
+					if(Pin->GetCPPType() == RigVMTypeUtils::TypeString::Bool)
 					{
 						static const FString BoolDefaultValue = TEXT("False");
 						SimpleTypeDefaultValue = BoolDefaultValue;
 					}
-					else if(Pin->GetCPPType() == RigVMTypeUtils::FloatType || Pin->GetCPPType() == RigVMTypeUtils::DoubleType)
+					else if(Pin->GetCPPType() == RigVMTypeUtils::TypeString::Float || Pin->GetCPPType() == RigVMTypeUtils::TypeString::Double)
 					{
 						static const FString FloatingPointDefaultValue = TEXT("0.000000");
 						SimpleTypeDefaultValue = FloatingPointDefaultValue;
 					}
-					else if(Pin->GetCPPType() == RigVMTypeUtils::Int32Type)
+					else if(Pin->GetCPPType() == RigVMTypeUtils::TypeString::Int32)
 					{
 						static const FString IntegerDefaultValue = TEXT("0");
 						SimpleTypeDefaultValue = IntegerDefaultValue;
@@ -22728,7 +22728,7 @@ FRigVMClientPatchResult URigVMController::PatchReduceArrayFloatDoubleConvertsion
 						for (auto Pair : NewTypes)
 						{
 							const FRigVMTemplateArgumentType Type = Registry.GetType(Pair.Value);
-							ChangePinType(TemplateNode->FindPin(Pair.Key.ToString()), Type.CPPType.ToString(), Type.GetCPPTypeObjectPath(), false, false, false, false, false);
+							ChangePinType(TemplateNode->FindPin(Pair.Key.ToString()), Type.GetTypeString(), Type.GetCPPTypeObjectPath(), false, false, false, false, false);
 							bChangedType = true;
 							NodesModified.Add(TemplateNode->GetName());
 							Result.bChangedContent = true;
@@ -23053,7 +23053,7 @@ FRigVMClientPatchResult URigVMController::PatchPinDefaultValues()
 							// reset the value back to the current / updated default.
 							FString OriginalDefaultValue = Pin->GetOriginalDefaultValue();
 							PostProcessDefaultValue(Pin, OriginalDefaultValue);
-							if(!OriginalDefaultValue.IsEmpty() && (Pin->GetCPPType() != RigVMTypeUtils::FStringType))
+							if(!OriginalDefaultValue.IsEmpty() && (Pin->GetCPPType() != RigVMTypeUtils::TypeString::FString))
 							{
 								if(Pin->IsValidDefaultValue(OriginalDefaultValue))
 								{
@@ -23102,7 +23102,7 @@ ERigVMPinDefaultValueType URigVMController::GetDefaultValueType(const URigVMPin*
 		}
 
 		// if the node didn't provide a valid default value let's keep it at autodetect
-		if(OriginalDefaultValue.IsEmpty() && (InPin->GetCPPType() != RigVMTypeUtils::FStringType))
+		if(OriginalDefaultValue.IsEmpty() && (InPin->GetCPPType() != RigVMTypeUtils::TypeString::FString))
 		{
 			return ERigVMPinDefaultValueType::AutoDetect;
 		}
diff --git a/Engine/Plugins/Runtime/RigVM/Source/RigVMDeveloper/Private/RigVMModel/RigVMControllerActions.cpp b/Engine/Plugins/Runtime/RigVM/Source/RigVMDeveloper/Private/RigVMModel/RigVMControllerActions.cpp
--- a/Engine/Plugins/Runtime/RigVM/Source/RigVMDeveloper/Private/RigVMModel/RigVMControllerActions.cpp
+++ b/Engine/Plugins/Runtime/RigVM/Source/RigVMDeveloper/Private/RigVMModel/RigVMControllerActions.cpp
@@ -1395,7 +1395,7 @@ bool FRigVMSetPinDefaultValueAction::Undo()
 		// strings and wildcards allow to set an empty default
 		if(const URigVMPin* Pin = GetController()->GetGraph()->FindPin(PinPath))
 		{
-			if((Pin->GetCPPType() != RigVMTypeUtils::FStringType) &&
+			if((Pin->GetCPPType() != RigVMTypeUtils::TypeString::FString) &&
 				!Pin->IsWildCard())
 			{
 				return true;
@@ -1419,7 +1419,7 @@ bool FRigVMSetPinDefaultValueAction::Redo()
 		// strings and wildcards allow to set an empty default
 		if(const URigVMPin* Pin = GetController()->GetGraph()->FindPin(PinPath))
 		{
-			if(Pin->GetCPPType() == RigVMTypeUtils::FStringType)
+			if(Pin->GetCPPType() == RigVMTypeUtils::TypeString::FString)
 			{
 				bIsValidDefaultValue = true;
 			}
diff --git a/Engine/Plugins/Runtime/RigVM/Source/RigVMDeveloper/Private/RigVMModel/RigVMExternalDependency.cpp b/Engine/Plugins/Runtime/RigVM/Source/RigVMDeveloper/Private/RigVMModel/RigVMExternalDependency.cpp
--- a/Engine/Plugins/Runtime/RigVM/Source/RigVMDeveloper/Private/RigVMModel/RigVMExternalDependency.cpp
+++ b/Engine/Plugins/Runtime/RigVM/Source/RigVMDeveloper/Private/RigVMModel/RigVMExternalDependency.cpp
@@ -93,7 +93,7 @@ void IRigVMExternalDependencyManager::CollectExternalDependencies(TArray<FRigVME
 				const FRigVMTemplateArgumentType& ArgumentType = Registry.FindTypeFromCPPType(Argument.Type);
 				if(ArgumentType.IsValid())
 				{
-					CollectExternalDependenciesForCPPTypeObject(OutDependencies, InCategory, ArgumentType.CPPTypeObject.Get());
+					CollectExternalDependenciesForCPPTypeObject(OutDependencies, InCategory, ArgumentType.GetTypeObject());
 				}
 			}
 		}
diff --git a/Engine/Plugins/Runtime/RigVM/Source/RigVMDeveloper/Private/RigVMModel/RigVMSchema.cpp b/Engine/Plugins/Runtime/RigVM/Source/RigVMDeveloper/Private/RigVMModel/RigVMSchema.cpp
--- a/Engine/Plugins/Runtime/RigVM/Source/RigVMDeveloper/Private/RigVMModel/RigVMSchema.cpp
+++ b/Engine/Plugins/Runtime/RigVM/Source/RigVMDeveloper/Private/RigVMModel/RigVMSchema.cpp
@@ -27,7 +27,7 @@ bool URigVMSchema::SupportsType(URigVMController* InController, TRigVMTypeIndex
 	// filter out incompatible execute types
 	if(Registry->IsExecuteType(InTypeIndex))
 	{
-		const UStruct* Struct = CastChecked<UStruct>(Registry->GetType(InTypeIndex).CPPTypeObject);
+		const UStruct* Struct = CastChecked<UStruct>(Registry->GetType(InTypeIndex).GetTypeObject());
 		if(!ValidExecuteContextStructs.Contains(Struct))
 		{
 			static constexpr TCHAR Format[] = TEXT("ExecuteContext struct '%s' is not supported.");
@@ -56,7 +56,7 @@ bool URigVMSchema::SupportsUnitFunction(URigVMController* InController, const FR
 		{
 			if(!SupportsType(InController, ArgumentType))
 			{
-				const FString CPPTypeString = Registry->GetType(ArgumentType).CPPType.ToString();
+				const FString CPPTypeString = Registry->GetType(ArgumentType).GetTypeString();
 				static constexpr TCHAR Format[] = TEXT("Unit function '%s' is not supported since type '%s' is not supported.");
 				InController->ReportErrorf(Format, *InUnitFunction->GetName(), *CPPTypeString);
 				return false;
@@ -108,7 +108,7 @@ bool URigVMSchema::SupportsTemplate(URigVMController* InController, const FRigVM
 			const TRigVMTypeIndex& TypeIndex = InTemplate->GetExecuteArgument(Index, DispatchContext)->TypeIndex;
 			if(!SupportsType(InController, TypeIndex))
 			{
-				const FString CPPTypeString = Registry->GetType(TypeIndex).CPPType.ToString();
+				const FString CPPTypeString = Registry->GetType(TypeIndex).GetTypeString();
 				static constexpr TCHAR Format[] = TEXT("Template '%s' is not supported since type '%s' is not supported.");
 				InController->ReportErrorf(Format, *InTemplate->GetNotation().ToString(), *CPPTypeString);
 				return false;
diff --git a/Engine/Plugins/Runtime/RigVM/Source/RigVMEditor/Private/EdGraph/NodeSpawners/RigVMEdGraphEnumNodeSpawner.cpp b/Engine/Plugins/Runtime/RigVM/Source/RigVMEditor/Private/EdGraph/NodeSpawners/RigVMEdGraphEnumNodeSpawner.cpp
--- a/Engine/Plugins/Runtime/RigVM/Source/RigVMEditor/Private/EdGraph/NodeSpawners/RigVMEdGraphEnumNodeSpawner.cpp
+++ b/Engine/Plugins/Runtime/RigVM/Source/RigVMEditor/Private/EdGraph/NodeSpawners/RigVMEdGraphEnumNodeSpawner.cpp
@@ -59,7 +59,7 @@ UEdGraphNode* URigVMEdGraphEnumNodeSpawner::Invoke(UEdGraph* ParentGraph, FBindi
 	if (bIsTemplateNode)
 	{
 		TArray<FPinInfo> Pins;
-		Pins.Emplace(URigVMEnumNode::EnumValueName, ERigVMPinDirection::Output, RigVMTypeUtils::Int32TypeName, nullptr);
+		Pins.Emplace(URigVMEnumNode::EnumValueName, ERigVMPinDirection::Output, RigVMTypeUtils::TypeName::Int32, nullptr);
 		return SpawnTemplateNode(ParentGraph, Pins);
 	}
 
diff --git a/Engine/Plugins/Runtime/RigVM/Source/RigVMEditor/Private/EdGraph/NodeSpawners/RigVMEdGraphTemplateNodeSpawner.cpp b/Engine/Plugins/Runtime/RigVM/Source/RigVMEditor/Private/EdGraph/NodeSpawners/RigVMEdGraphTemplateNodeSpawner.cpp
--- a/Engine/Plugins/Runtime/RigVM/Source/RigVMEditor/Private/EdGraph/NodeSpawners/RigVMEdGraphTemplateNodeSpawner.cpp
+++ b/Engine/Plugins/Runtime/RigVM/Source/RigVMEditor/Private/EdGraph/NodeSpawners/RigVMEdGraphTemplateNodeSpawner.cpp
@@ -97,7 +97,7 @@ bool URigVMEdGraphTemplateNodeSpawner::IsTemplateNodeFilteredOut(FBlueprintActio
 					if (const URigVMPin* ModelPin = RigGraph->FindPin(Pin->GetName()))
 					{
 						// Only filter when the source pin is not a wildcard
-						if (ModelPin->GetCPPType() != RigVMTypeUtils::GetWildCardCPPType())
+						if (ModelPin->GetCPPType() != RigVMTypeUtils::GetWildCardCPPTypeString())
 						{
 							if (Template)
 							{
@@ -190,8 +190,8 @@ UEdGraphNode* URigVMEdGraphTemplateNodeSpawner::Invoke(UEdGraph* ParentGraph, FB
 				{
 					const TRigVMTypeIndex TypeIndex = Argument->GetTypeIndex(0);
 					const FRigVMTemplateArgumentType& Type = Registry.GetType(TypeIndex);
-					CPPType = Type.CPPType;
-					CPPTypeObject = Type.CPPTypeObject;
+					CPPType = Type.GetTypeFName();
+					CPPTypeObject = Type.GetTypeObject();
 				}
 				
 				Pins.Emplace(Argument->GetName(), Argument->GetDirection(), CPPType, CPPTypeObject);
diff --git a/Engine/Plugins/Runtime/RigVM/Source/RigVMEditor/Private/Editor/RigVMDetailsViewWrapperObject.cpp b/Engine/Plugins/Runtime/RigVM/Source/RigVMEditor/Private/Editor/RigVMDetailsViewWrapperObject.cpp
--- a/Engine/Plugins/Runtime/RigVM/Source/RigVMEditor/Private/Editor/RigVMDetailsViewWrapperObject.cpp
+++ b/Engine/Plugins/Runtime/RigVM/Source/RigVMEditor/Private/Editor/RigVMDetailsViewWrapperObject.cpp
@@ -460,28 +460,44 @@ UClass* URigVMDetailsViewWrapperObject::GetClassForNodes(TArray<URigVMNode*> InN
 			}
 		}
 
-		if(BaseCPPType.Equals(RigVMTypeUtils::BoolType, ESearchCase::IgnoreCase))
+		if(BaseCPPType.Equals(RigVMTypeUtils::TypeString::Bool, ESearchCase::IgnoreCase))
 		{
 			(*ElementProperty) = new FBoolProperty(PropertyOwner, Pin->GetFName(), RF_Public);;
 		}
-		else if(BaseCPPType.Equals(RigVMTypeUtils::Int32Type, ESearchCase::IgnoreCase) ||
-			BaseCPPType.Equals(RigVMTypeUtils::IntType, ESearchCase::IgnoreCase))
+		else if(BaseCPPType.Equals(RigVMTypeUtils::TypeString::Int32, ESearchCase::IgnoreCase) ||
+			BaseCPPType.Equals(RigVMTypeUtils::TypeString::Int, ESearchCase::IgnoreCase))
 		{
 			(*ElementProperty) = new FIntProperty(PropertyOwner, Pin->GetFName(), RF_Public);;
 		}
-		else if(BaseCPPType.Equals(RigVMTypeUtils::FloatType, ESearchCase::IgnoreCase))
+		else if(BaseCPPType.Equals(RigVMTypeUtils::TypeString::Int64, ESearchCase::IgnoreCase))
+		{
+			(*ElementProperty) = new FInt64Property(PropertyOwner, Pin->GetFName(), RF_Public);;
+		}
+		else if(BaseCPPType.Equals(RigVMTypeUtils::TypeString::UInt8, ESearchCase::IgnoreCase))
+		{
+			(*ElementProperty) = new FByteProperty(PropertyOwner, Pin->GetFName(), RF_Public);;
+		}
+		else if(BaseCPPType.Equals(RigVMTypeUtils::TypeString::UInt32, ESearchCase::IgnoreCase))
+		{
+			(*ElementProperty) = new FUInt32Property(PropertyOwner, Pin->GetFName(), RF_Public);;
+		}
+		else if(BaseCPPType.Equals(RigVMTypeUtils::TypeString::UInt64, ESearchCase::IgnoreCase))
+		{
+			(*ElementProperty) = new FUInt64Property(PropertyOwner, Pin->GetFName(), RF_Public);;
+		}
+		else if(BaseCPPType.Equals(RigVMTypeUtils::TypeString::Float, ESearchCase::IgnoreCase))
 		{
 			(*ElementProperty) = new FFloatProperty(PropertyOwner, Pin->GetFName(), RF_Public);;
 		}
-		else if(BaseCPPType.Equals(RigVMTypeUtils::DoubleType, ESearchCase::IgnoreCase))
+		else if(BaseCPPType.Equals(RigVMTypeUtils::TypeString::Double, ESearchCase::IgnoreCase))
 		{
 			(*ElementProperty) = new FDoubleProperty(PropertyOwner, Pin->GetFName(), RF_Public);;
 		}
-		else if(BaseCPPType.Equals(RigVMTypeUtils::FStringType, ESearchCase::IgnoreCase))
+		else if(BaseCPPType.Equals(RigVMTypeUtils::TypeString::FString, ESearchCase::IgnoreCase))
 		{
 			(*ElementProperty) = new FStrProperty(PropertyOwner, Pin->GetFName(), RF_Public);;
 		}
-		else if(BaseCPPType.Equals(RigVMTypeUtils::FNameType, ESearchCase::IgnoreCase))
+		else if(BaseCPPType.Equals(RigVMTypeUtils::TypeString::FName, ESearchCase::IgnoreCase))
 		{
 			(*ElementProperty) = new FNameProperty(PropertyOwner, Pin->GetFName(), RF_Public);;
 		}
diff --git a/Engine/Plugins/Runtime/RigVM/Source/RigVMEditor/Private/Editor/RigVMEditor.cpp b/Engine/Plugins/Runtime/RigVM/Source/RigVMEditor/Private/Editor/RigVMEditor.cpp
--- a/Engine/Plugins/Runtime/RigVM/Source/RigVMEditor/Private/Editor/RigVMEditor.cpp
+++ b/Engine/Plugins/Runtime/RigVM/Source/RigVMEditor/Private/Editor/RigVMEditor.cpp
@@ -2403,11 +2403,11 @@ TRigVMTypeIndex FRigVMEditor::OnRequestPinTypeSelectionDialog(const TArray<TRigV
 		const FRigVMTemplateArgumentType& ArgumentType = Registry.GetType(FinalType);
 		if (!TypeNames.ContainsByPredicate([&ArgumentType](const TSharedPtr<FName>& InName)
 		{
-			return *InName.Get() == ArgumentType.CPPType;
+			return *InName.Get() == ArgumentType.GetTypeFName();
 		}))
 		{
-			TypeNames.AddUnique(MakeShared<FName>(ArgumentType.CPPType));
-			TypeNameToIndex.Add(ArgumentType.CPPType, i);
+			TypeNames.AddUnique(MakeShared<FName>(ArgumentType.GetTypeFName()));
+			TypeNameToIndex.Add(ArgumentType.GetTypeFName(), i);
 		}
 	}
 	TSharedPtr< SWindow > Window = SNew(SWindow)
diff --git a/Engine/Plugins/Runtime/RigVM/Source/RigVMEditor/Private/Widgets/SRigVMGraphChangePinType.cpp b/Engine/Plugins/Runtime/RigVM/Source/RigVMEditor/Private/Widgets/SRigVMGraphChangePinType.cpp
--- a/Engine/Plugins/Runtime/RigVM/Source/RigVMEditor/Private/Widgets/SRigVMGraphChangePinType.cpp
+++ b/Engine/Plugins/Runtime/RigVM/Source/RigVMEditor/Private/Widgets/SRigVMGraphChangePinType.cpp
@@ -95,7 +95,7 @@ FLinearColor SRigVMGraphChangePinType::GetBindingColor() const
 	if(Types.Num() > 0)
 	{
 		const FRigVMTemplateArgumentType Type = FRigVMRegistry::Get().GetType(Types[0]);
-		const FEdGraphPinType PinType = RigVMTypeUtils::PinTypeFromCPPType(Type.CPPType, Type.CPPTypeObject);
+		const FEdGraphPinType PinType = RigVMTypeUtils::PinTypeFromCPPType(Type.GetTypeFName(), Type.GetTypeObject());
 		const URigVMEdGraphSchema* Schema = GetDefault<URigVMEdGraphSchema>();
 		return Schema->GetPinTypeColor(PinType);
 	}
@@ -212,21 +212,21 @@ void SRigVMGraphChangePinType::FillPinTypeMenu(FMenuBuilder& MenuBuilder)
 		{
 			const TRigVMTypeIndex& TypeIndex = SortedTypes[SortedIndex].Key;
 			const FRigVMTemplateArgumentType Type = FRigVMRegistry::Get().GetType(TypeIndex);
-			if (Type.CPPTypeObject != nullptr && !IsValid(Type.CPPTypeObject))
+			if (Type.GetTypeObject() != nullptr && !IsValid(Type.GetTypeObject()))
 			{
 				continue;
 			}
 			
-			const FEdGraphPinType PinType = RigVMTypeUtils::PinTypeFromCPPType(Type.CPPType, Type.CPPTypeObject);
+			const FEdGraphPinType PinType = RigVMTypeUtils::PinTypeFromCPPType(Type.GetTypeFName(), Type.GetTypeObject());
 
-			if(bHasAllTypes && Type.CPPTypeObject)
+			if(bHasAllTypes && Type.GetTypeObject())
 			{
-				if(Type.CPPTypeObject->IsA<UEnum>())
+				if(Type.GetTypeObject()->IsA<UEnum>())
 				{
 					continue;
 				}
 
-				if(UScriptStruct* ScriptStruct = Cast<UScriptStruct>(Type.CPPTypeObject))
+				if(UScriptStruct* ScriptStruct = Cast<UScriptStruct>(Type.GetTypeObject()))
 				{
 					const FRigVMTemplateArgumentType MathType(*RigVMTypeUtils::GetUniqueStructTypeName(ScriptStruct), ScriptStruct);
 					const TRigVMTypeIndex MathTypeIndex = FRigVMRegistry::Get().GetTypeIndex(MathType);
@@ -238,7 +238,7 @@ void SRigVMGraphChangePinType::FillPinTypeMenu(FMenuBuilder& MenuBuilder)
 				}
 			}
 
-			if(UScriptStruct* ScriptStruct = Cast<UScriptStruct>(Type.CPPTypeObject))
+			if(UScriptStruct* ScriptStruct = Cast<UScriptStruct>(Type.GetTypeObject()))
 			{
 				if(ScriptStruct->IsChildOf(FRigVMExecuteContext::StaticStruct()))
 				{
