Control Rig: Replace MyBlueprint with RigVMEditorGraphExplorer
#rnx

diff --git a/Engine/Plugins/Animation/ControlRig/Source/ControlRigEditor/Private/Editor/ControlRigEditorMode.cpp b/Engine/Plugins/Animation/ControlRig/Source/ControlRigEditor/Private/Editor/ControlRigEditorMode.cpp
--- a/Engine/Plugins/Animation/ControlRig/Source/ControlRigEditor/Private/Editor/ControlRigEditorMode.cpp
+++ b/Engine/Plugins/Animation/ControlRig/Source/ControlRigEditor/Private/Editor/ControlRigEditorMode.cpp
@@ -10,6 +10,7 @@
 #include "Editor/ModularRigModelTabSummoner.h"
 #include "Editor/RigModuleAssetBrowserTabSummoner.h"
 #include "Editor/RigVMExecutionStackTabSummoner.h"
+#include "Editor/RigVMEditorGraphExplorerTabSummoner.h"
 #include "Editor/RigCurveContainerTabSummoner.h"
 #include "Editor/RigValidationTabSummoner.h"
 #include "Editor/RigAnimAttributeTabSummoner.h"
@@ -22,6 +23,7 @@ FControlRigEditorMode::FControlRigEditorMode(const TSharedRef<FControlRigEditor>
 
 	TabFactories.RegisterFactory(MakeShared<FRigHierarchyTabSummoner>(InControlRigEditor));
 	TabFactories.RegisterFactory(MakeShared<FRigVMExecutionStackTabSummoner>(InControlRigEditor));
+	TabFactories.RegisterFactory(MakeShared<FRigVMEditorGraphExplorerTabSummoner>(InControlRigEditor));
 	TabFactories.RegisterFactory(MakeShared<FRigCurveContainerTabSummoner>(InControlRigEditor));
 	TabFactories.RegisterFactory(MakeShared<FRigValidationTabSummoner>(InControlRigEditor));
 	TabFactories.RegisterFactory(MakeShared<FRigAnimAttributeTabSummoner>(InControlRigEditor));
@@ -75,6 +77,7 @@ FControlRigEditorMode::FControlRigEditorMode(const TSharedRef<FControlRigEditor>
 							->AddTab(FRigVMExecutionStackTabSummoner::TabID, ETabState::OpenedTab)
 							->AddTab(FRigCurveContainerTabSummoner::TabID, ETabState::OpenedTab)
 							->AddTab(FBlueprintEditorTabs::MyBlueprintID, ETabState::OpenedTab)
+							->AddTab(FRigVMEditorGraphExplorerTabSummoner::TabID, ETabState::ClosedTab)
 						)
 					)
 					->Split
diff --git a/Engine/Plugins/Experimental/AnimNext/Source/AnimNextUncookedOnly/Internal/AnimNextRigVMAssetEditorData.h b/Engine/Plugins/Experimental/AnimNext/Source/AnimNextUncookedOnly/Internal/AnimNextRigVMAssetEditorData.h
--- a/Engine/Plugins/Experimental/AnimNext/Source/AnimNextUncookedOnly/Internal/AnimNextRigVMAssetEditorData.h
+++ b/Engine/Plugins/Experimental/AnimNext/Source/AnimNextUncookedOnly/Internal/AnimNextRigVMAssetEditorData.h
@@ -216,7 +216,7 @@ protected:
 	virtual void HandleRigVMGraphRemoved(const FRigVMClient* InClient, const FString& InNodePath) override;
 	virtual void HandleRigVMGraphRenamed(const FRigVMClient* InClient, const FString& InOldNodePath, const FString& InNewNodePath) override;
 	virtual void HandleConfigureRigVMController(const FRigVMClient* InClient, URigVMController* InControllerToConfigure) override;
-	virtual UObject* GetEditorObjectForRigVMGraph(URigVMGraph* InVMGraph) const override;
+	virtual UObject* GetEditorObjectForRigVMGraph(const URigVMGraph* InVMGraph) const override;
 	virtual URigVMGraph* GetRigVMGraphForEditorObject(UObject* InObject) const override;
 	virtual void RecompileVM() override;
 	virtual void RecompileVMIfRequired() override;
diff --git a/Engine/Plugins/Experimental/AnimNext/Source/AnimNextUncookedOnly/Private/AnimNextRigVMAssetEditorData.cpp b/Engine/Plugins/Experimental/AnimNext/Source/AnimNextUncookedOnly/Private/AnimNextRigVMAssetEditorData.cpp
--- a/Engine/Plugins/Experimental/AnimNext/Source/AnimNextUncookedOnly/Private/AnimNextRigVMAssetEditorData.cpp
+++ b/Engine/Plugins/Experimental/AnimNext/Source/AnimNextUncookedOnly/Private/AnimNextRigVMAssetEditorData.cpp
@@ -625,7 +625,7 @@ void UAnimNextRigVMAssetEditorData::HandleConfigureRigVMController(const FRigVMC
 #endif
 }
 
-UObject* UAnimNextRigVMAssetEditorData::GetEditorObjectForRigVMGraph(URigVMGraph* InVMGraph) const
+UObject* UAnimNextRigVMAssetEditorData::GetEditorObjectForRigVMGraph(const URigVMGraph* InVMGraph) const
 {
 	if(InVMGraph)
 	{
diff --git a/Engine/Plugins/Runtime/RigVM/Source/RigVMDeveloper/Private/RigVMBlueprint.cpp b/Engine/Plugins/Runtime/RigVM/Source/RigVMDeveloper/Private/RigVMBlueprint.cpp
--- a/Engine/Plugins/Runtime/RigVM/Source/RigVMDeveloper/Private/RigVMBlueprint.cpp
+++ b/Engine/Plugins/Runtime/RigVM/Source/RigVMDeveloper/Private/RigVMBlueprint.cpp
@@ -423,7 +423,7 @@ const IRigVMGraphFunctionHost* URigVMBlueprint::GetRigVMGraphFunctionHost() cons
 	return GetRigVMBlueprintGeneratedClass();
 }
 
-UObject* URigVMBlueprint::GetEditorObjectForRigVMGraph(URigVMGraph* InVMGraph) const
+UObject* URigVMBlueprint::GetEditorObjectForRigVMGraph(const URigVMGraph* InVMGraph) const
 {
 	if(InVMGraph)
 	{
@@ -2256,7 +2256,7 @@ URigVMController* URigVMBlueprint::GetTemplateController(bool bIsFunctionLibrary
 #endif
 }
 
-UEdGraph* URigVMBlueprint::GetEdGraph(URigVMGraph* InModel) const
+UEdGraph* URigVMBlueprint::GetEdGraph(const URigVMGraph* InModel) const
 {
 	return Cast<UEdGraph>(GetEditorObjectForRigVMGraph(InModel));
 }
diff --git a/Engine/Plugins/Runtime/RigVM/Source/RigVMDeveloper/Public/RigVMBlueprint.h b/Engine/Plugins/Runtime/RigVM/Source/RigVMDeveloper/Public/RigVMBlueprint.h
--- a/Engine/Plugins/Runtime/RigVM/Source/RigVMDeveloper/Public/RigVMBlueprint.h
+++ b/Engine/Plugins/Runtime/RigVM/Source/RigVMDeveloper/Public/RigVMBlueprint.h
@@ -325,7 +325,7 @@ public:
 	virtual const FRigVMClient* GetRigVMClient() const override;
 	virtual IRigVMGraphFunctionHost* GetRigVMGraphFunctionHost() override;
 	virtual const IRigVMGraphFunctionHost* GetRigVMGraphFunctionHost() const override;
-	virtual UObject* GetEditorObjectForRigVMGraph(URigVMGraph* InVMGraph) const override;
+	virtual UObject* GetEditorObjectForRigVMGraph(const URigVMGraph* InVMGraph) const override;
 	virtual URigVMGraph* GetRigVMGraphForEditorObject(UObject* InObject) const override;
 	virtual void HandleRigVMGraphAdded(const FRigVMClient* InClient, const FString& InNodePath) override;
 	virtual void HandleRigVMGraphRemoved(const FRigVMClient* InClient, const FString& InNodePath) override;
@@ -438,7 +438,7 @@ public:
 	URigVMController* GetTemplateController(bool bIsFunctionLibrary = false);
 
 #if WITH_EDITOR
-	UEdGraph* GetEdGraph(URigVMGraph* InModel) const;
+	UEdGraph* GetEdGraph(const URigVMGraph* InModel) const;
 	UEdGraph* GetEdGraph(const FString& InNodePath) const;
 #endif
 
diff --git a/Engine/Plugins/Runtime/RigVM/Source/RigVMDeveloper/Public/RigVMModel/RigVMClient.h b/Engine/Plugins/Runtime/RigVM/Source/RigVMDeveloper/Public/RigVMModel/RigVMClient.h
--- a/Engine/Plugins/Runtime/RigVM/Source/RigVMDeveloper/Public/RigVMModel/RigVMClient.h
+++ b/Engine/Plugins/Runtime/RigVM/Source/RigVMDeveloper/Public/RigVMModel/RigVMClient.h
@@ -66,7 +66,7 @@ public:
 	virtual const IRigVMGraphFunctionHost* GetRigVMGraphFunctionHost() const = 0;
 
 	// Returns the editor object corresponding with the supplied RigVM graph
-	virtual UObject* GetEditorObjectForRigVMGraph(URigVMGraph* InVMGraph) const = 0;
+	virtual UObject* GetEditorObjectForRigVMGraph(const URigVMGraph* InVMGraph) const = 0;
 
 	// Returns the RigVM graph corresponding with the supplied editor object
 	virtual URigVMGraph* GetRigVMGraphForEditorObject(UObject* InObject) const = 0;
diff --git a/Engine/Plugins/Runtime/RigVM/Source/RigVMEditor/Private/Editor/RigVMEditor.cpp b/Engine/Plugins/Runtime/RigVM/Source/RigVMEditor/Private/Editor/RigVMEditor.cpp
--- a/Engine/Plugins/Runtime/RigVM/Source/RigVMEditor/Private/Editor/RigVMEditor.cpp
+++ b/Engine/Plugins/Runtime/RigVM/Source/RigVMEditor/Private/Editor/RigVMEditor.cpp
@@ -50,6 +50,7 @@
 #include "Widgets/SRigVMBulkEditDialog.h"
 #include "AssetRegistry/AssetRegistryModule.h"
 #include "ToolMenuContext.h"
+#include "Widgets/SRigVMEditorGraphExplorer.h"
 
 #define LOCTEXT_NAMESPACE "RigVMEditor"
 
@@ -1306,6 +1307,7 @@ void FRigVMEditor::RefreshEditors(ERefreshBlueprintEditorReason::Type Reason)
 		Reason = ERefreshBlueprintEditorReason::BlueprintCompiled;
 	}
 	FBlueprintEditor::RefreshEditors(Reason);
+	GraphExplorerWidget->Refresh();
 }
 
 void FRigVMEditor::SetupGraphEditorEvents(UEdGraph* InGraph, SGraphEditor::FGraphEditorEvents& InEvents)
@@ -1524,6 +1526,11 @@ void FRigVMEditor::RedoAction()
 void FRigVMEditor::CreateDefaultTabContents(const TArray<UBlueprint*>& InBlueprints)
 {
 	FBlueprintEditor::CreateDefaultTabContents(InBlueprints);
+
+	if (IsEditingSingleBlueprint())
+	{
+		GraphExplorerWidget = SNew(SRigVMEditorGraphExplorer, SharedThis(this));
+	}
 }
 
 void FRigVMEditor::NewDocument_OnClicked(ECreatedDocumentType GraphType)
@@ -1541,7 +1548,6 @@ void FRigVMEditor::NewDocument_OnClicked(ECreatedDocumentType GraphType)
 						OpenDocument(NewGraph, FDocumentTracker::OpenNewDocument);
 						RenameNewlyAddedAction(FunctionNode->GetFName());
 					}
-
 				}
 			}
 		}
diff --git a/Engine/Plugins/Runtime/RigVM/Source/RigVMEditor/Private/Editor/RigVMEditorGraphExplorerTabSummoner.cpp b/Engine/Plugins/Runtime/RigVM/Source/RigVMEditor/Private/Editor/RigVMEditorGraphExplorerTabSummoner.cpp
new file mode 100644
--- /dev/null
+++ b/Engine/Plugins/Runtime/RigVM/Source/RigVMEditor/Private/Editor/RigVMEditorGraphExplorerTabSummoner.cpp
@@ -0,0 +1,27 @@
+// Copyright Epic Games, Inc. All Rights Reserved.
+
+#include "Editor/RigVMEditorGraphExplorerTabSummoner.h"
+#include "Widgets/SRigVMEditorGraphExplorer.h"
+#include "Editor/RigVMEditor.h"
+
+#define LOCTEXT_NAMESPACE "RigVMEditorGraphExplorerTabSummoner"
+
+FRigVMEditorGraphExplorerTabSummoner::FRigVMEditorGraphExplorerTabSummoner(const TSharedRef<FRigVMEditor>& InRigVMEditor)
+	: FWorkflowTabFactory(TabID, InRigVMEditor)
+	, RigVMEditor(InRigVMEditor)
+{
+	TabLabel = LOCTEXT("RigVMEditorGraphExplorerTabLabel", "Graph Explorer");
+	TabIcon = FSlateIcon(TEXT("RigVMEditorStyle"), "EditorGraphExplorer.TabIcon");
+
+	bIsSingleton = true;
+
+	ViewMenuDescription = LOCTEXT("RigVMEditorGraphExplorer_ViewMenu_Desc", "Graph Explorer");
+	ViewMenuTooltip = LOCTEXT("RigVMEditorGraphExplorer_ViewMenu_ToolTip", "Show the RigVM Editor Graph Explorer tab");
+}
+
+TSharedRef<SWidget> FRigVMEditorGraphExplorerTabSummoner::CreateTabBody(const FWorkflowTabSpawnInfo& Info) const
+{
+	return SNew(SRigVMEditorGraphExplorer, RigVMEditor.Pin().ToSharedRef());
+}
+
+#undef LOCTEXT_NAMESPACE 
diff --git a/Engine/Plugins/Runtime/RigVM/Source/RigVMEditor/Private/Editor/RigVMEditorMode.cpp b/Engine/Plugins/Runtime/RigVM/Source/RigVMEditor/Private/Editor/RigVMEditorMode.cpp
--- a/Engine/Plugins/Runtime/RigVM/Source/RigVMEditor/Private/Editor/RigVMEditorMode.cpp
+++ b/Engine/Plugins/Runtime/RigVM/Source/RigVMEditor/Private/Editor/RigVMEditorMode.cpp
@@ -4,6 +4,7 @@
 #include "BlueprintEditorTabs.h"
 #include "SBlueprintEditorToolbar.h"
 #include "Editor/RigVMExecutionStackTabSummoner.h"
+#include "Editor/RigVMEditorGraphExplorerTabSummoner.h"
 #include "ToolMenus.h"
 
 FRigVMEditorMode::FRigVMEditorMode(const TSharedRef<FRigVMEditor>& InRigVMEditor)
@@ -12,6 +13,7 @@ FRigVMEditorMode::FRigVMEditorMode(const TSharedRef<FRigVMEditor>& InRigVMEditor
 	RigVMBlueprintPtr = CastChecked<URigVMBlueprint>(InRigVMEditor->GetBlueprintObj());
 
 	TabFactories.RegisterFactory(MakeShared<FRigVMExecutionStackTabSummoner>(InRigVMEditor));
+	TabFactories.RegisterFactory(MakeShared<FRigVMEditorGraphExplorerTabSummoner>(InRigVMEditor));
 
 	TabLayout = FTabManager::NewLayout("Standalone_RigVMEditMode_Layout_v1.5")
 		->AddArea
@@ -30,6 +32,7 @@ FRigVMEditorMode::FRigVMEditorMode(const TSharedRef<FRigVMEditor>& InRigVMEditor
 					->SetSizeCoefficient(0.5f)
 					->AddTab(FRigVMExecutionStackTabSummoner::TabID, ETabState::OpenedTab)
 					->AddTab(FBlueprintEditorTabs::MyBlueprintID, ETabState::OpenedTab)
+					->AddTab(FRigVMEditorGraphExplorerTabSummoner::TabID, ETabState::ClosedTab)
 				)
 				->Split
 				(
diff --git a/Engine/Plugins/Runtime/RigVM/Source/RigVMEditor/Private/RigVMEditorModule.cpp b/Engine/Plugins/Runtime/RigVM/Source/RigVMEditor/Private/RigVMEditorModule.cpp
--- a/Engine/Plugins/Runtime/RigVM/Source/RigVMEditor/Private/RigVMEditorModule.cpp
+++ b/Engine/Plugins/Runtime/RigVM/Source/RigVMEditor/Private/RigVMEditorModule.cpp
@@ -10,6 +10,7 @@
 #include "AssetRegistry/AssetRegistryModule.h"
 #include "Editor/RigVMEditorCommands.h"
 #include "Editor/RigVMExecutionStackCommands.h"
+#include "Widgets/SRigVMEditorGraphExplorer.h"
 #include "Editor/RigVMEditorStyle.h"
 #include "Editor/RigVMGraphDetailCustomization.h"
 #include "Editor/RigVMVariableDetailCustomization.h"
@@ -62,6 +63,7 @@ void FRigVMEditorModule::StartupModule()
 	if(IsRigVMEditorModuleBase())
 	{
 		FRigVMExecutionStackCommands::Register();
+		FRigVMEditorGraphExplorerCommands::Register();
 		FRigVMEditorCommands::Register();
 		FRigVMEditorStyle::Register();
 
@@ -174,6 +176,7 @@ void FRigVMEditorModule::ShutdownModule()
 			FRigVMEditorStyle::Unregister();
 			FEdGraphUtilities::UnregisterVisualNodeFactory(EdGraphPanelNodeFactory);
 			FEdGraphUtilities::UnregisterVisualPinFactory(EdGraphPanelPinFactory);
+			FRigVMEditorGraphExplorerCommands::Unregister();
 
 			FBlueprintEditorUtils::OnRefreshAllNodesEvent.Remove(RefreshAllNodesDelegateHandle);
 			FBlueprintEditorUtils::OnReconstructAllNodesEvent.Remove(ReconstructAllNodesDelegateHandle);
diff --git a/Engine/Plugins/Runtime/RigVM/Source/RigVMEditor/Private/Widgets/SRigVMEditorGraphExplorer.cpp b/Engine/Plugins/Runtime/RigVM/Source/RigVMEditor/Private/Widgets/SRigVMEditorGraphExplorer.cpp
new file mode 100644
--- /dev/null
+++ b/Engine/Plugins/Runtime/RigVM/Source/RigVMEditor/Private/Widgets/SRigVMEditorGraphExplorer.cpp
@@ -0,0 +1,844 @@
+// Copyright Epic Games, Inc. All Rights Reserved.
+
+#include "Widgets/SRigVMEditorGraphExplorer.h"
+
+#include "Editor/RigVMEditor.h"
+#include "EdGraph/RigVMEdGraphSchema.h"
+#include "Widgets/SRigVMEditorGraphExplorerTreeView.h"
+
+#include "Framework/Commands/GenericCommands.h"
+#include "Widgets/Input/SSearchBox.h"
+#include "SPositiveActionButton.h"
+#include "EdGraphSchema_K2_Actions.h"
+
+// Ick.
+#include "GraphActionNode.h"
+#include "GraphEditorActions.h"
+#include "Kismet2/BlueprintEditorUtils.h"
+#include "RigVMModel/Nodes/RigVMAggregateNode.h"
+
+
+#define LOCTEXT_NAMESPACE "RigVMGraphExplorer"
+
+
+FRigVMEditorGraphExplorerCommands::FRigVMEditorGraphExplorerCommands()
+	: TCommands<FRigVMEditorGraphExplorerCommands>(
+          TEXT("RigVMEditorGraphExplorer"), NSLOCTEXT("Contexts", "Explorer", "Explorer"),
+          NAME_None, FAppStyle::GetAppStyleSetName())
+{
+}
+
+
+void FRigVMEditorGraphExplorerCommands::RegisterCommands()
+{
+	UI_COMMAND(OpenGraph, "Open Graph", "Opens up this graph in the editor.", EUserInterfaceActionType::Button, FInputChord());
+	UI_COMMAND(OpenGraphInNewTab, "Open Graph In New Tab", "Opens up this graph in a new tab.", EUserInterfaceActionType::Button, FInputChord());
+	UI_COMMAND(CreateGraph, "Add New Graph", "Create a new graph and show it in the editor.", EUserInterfaceActionType::Button, FInputChord());
+}
+
+
+void SRigVMEditorGraphExplorer::Construct(const FArguments& InArgs, TWeakPtr<FRigVMEditor> InRigVMEditor)
+{
+	RigVMEditor = InRigVMEditor;
+
+	TSharedPtr<FRigVMEditor> Editor = RigVMEditor.Pin();
+	if (Editor)
+	{
+		Editor->OnRefresh().AddRaw(this, &SRigVMEditorGraphExplorer::Refresh);
+	}
+
+	RegisterCommands();
+
+	CreateWidgets();
+}
+
+
+SRigVMEditorGraphExplorer::~SRigVMEditorGraphExplorer()
+{
+	TSharedPtr<FRigVMEditor> Editor = RigVMEditor.Pin();
+	if (Editor)
+	{
+		Editor->OnRefresh().RemoveAll(this);
+	}
+}
+
+void SRigVMEditorGraphExplorer::Refresh()
+{
+	TreeView->RefreshTreeView(true);
+}
+
+void SRigVMEditorGraphExplorer::Tick(const FGeometry& InAllottedGeometry, const double InCurrentTime, const float InDeltaTime)
+{
+	SCompoundWidget::Tick(InAllottedGeometry, InCurrentTime, InDeltaTime);
+}
+
+void SRigVMEditorGraphExplorer::RegisterCommands()
+{
+	TSharedPtr<FRigVMEditor> Editor = RigVMEditor.Pin();
+
+	if (Editor)
+	{
+		TSharedPtr<FUICommandList> ToolKitCommandList = Editor->GetToolkitCommands();
+
+		ToolKitCommandList->MapAction(FRigVMEditorGraphExplorerCommands::Get().OpenGraph,
+		    FExecuteAction::CreateSP(this, &SRigVMEditorGraphExplorer::OnOpenGraph, false),
+			FCanExecuteAction(), FGetActionCheckState(),
+			FIsActionButtonVisible::CreateSP(this, &SRigVMEditorGraphExplorer::CanOpenGraph));
+
+		ToolKitCommandList->MapAction(FRigVMEditorGraphExplorerCommands::Get().OpenGraphInNewTab,
+			FExecuteAction::CreateSP(this, &SRigVMEditorGraphExplorer::OnOpenGraph, true),
+			FCanExecuteAction(), FGetActionCheckState(),
+			FIsActionButtonVisible::CreateSP(this, &SRigVMEditorGraphExplorer::CanOpenGraph));
+
+		ToolKitCommandList->MapAction(FRigVMEditorGraphExplorerCommands::Get().CreateGraph,
+			FExecuteAction::CreateSP(this, &SRigVMEditorGraphExplorer::OnCreateGraph),
+			FCanExecuteAction::CreateSP(this, &SRigVMEditorGraphExplorer::CanCreateGraph)
+			);
+
+		ToolKitCommandList->MapAction(FGenericCommands::Get().Rename,
+			FExecuteAction::CreateSP(this, &SRigVMEditorGraphExplorer::OnRenameEntry),
+			FCanExecuteAction::CreateSP(this, &SRigVMEditorGraphExplorer::CanRenameEntry));
+
+		ToolKitCommandList->MapAction(FGenericCommands::Get().Copy,
+			FExecuteAction::CreateSP(this, &SRigVMEditorGraphExplorer::OnCopy),
+			FCanExecuteAction::CreateSP(this, &SRigVMEditorGraphExplorer::CanCopy));
+		
+		ToolKitCommandList->MapAction(FGenericCommands::Get().Cut,
+			FExecuteAction::CreateSP(this, &SRigVMEditorGraphExplorer::OnCut),
+			FCanExecuteAction::CreateSP(this, &SRigVMEditorGraphExplorer::CanCut));
+
+		ToolKitCommandList->MapAction(FGenericCommands::Get().Paste,
+			FExecuteAction::CreateSP(this, &SRigVMEditorGraphExplorer::OnPasteGeneric),
+			FCanExecuteAction(), FIsActionChecked(),
+			FIsActionButtonVisible::CreateSP(this, &SRigVMEditorGraphExplorer::CanPasteGeneric));
+
+		ToolKitCommandList->MapAction(FGenericCommands::Get().Delete,
+			FExecuteAction::CreateSP(this, &SRigVMEditorGraphExplorer::OnDeleteEntry),
+			FCanExecuteAction::CreateSP(this, &SRigVMEditorGraphExplorer::CanDeleteEntry));
+
+		ToolKitCommandList->MapAction( FGraphEditorCommands::Get().FindReferences,
+			FExecuteAction::CreateSP(this, &SRigVMEditorGraphExplorer::OnFindReference, /*bSearchAllBlueprints=*/false),
+			FCanExecuteAction(),
+			FIsActionChecked(),
+			FIsActionButtonVisible::CreateSP(this, &SRigVMEditorGraphExplorer::CanFindReference) );
+	}
+}
+
+void SRigVMEditorGraphExplorer::CreateWidgets()
+{
+	TSharedPtr<SWidget> AddNewMenu = SNew(SPositiveActionButton)
+		.ToolTipText(LOCTEXT("AddNewToolTip", "Add a new item."))
+		.OnGetMenuContent(this, &SRigVMEditorGraphExplorer::CreateAddNewMenuWidget)
+		.Icon(FAppStyle::GetBrush("Plus"))
+		.Text(LOCTEXT("AddNew", "Add"));
+
+	FMenuBuilder ViewOptions(true, nullptr);
+
+	ViewOptions.AddMenuEntry(
+	    LOCTEXT("ShowEmptySections", "Show Empty Sections"),
+	    LOCTEXT("ShowEmptySectionsTooltip", "Should we show empty sections? eg. Graphs, Functions...etc."),
+	    FSlateIcon(),
+	    FUIAction(
+	        FExecuteAction::CreateSP(this, &SRigVMEditorGraphExplorer::OnToggleShowEmptySections),
+	        FCanExecuteAction(),
+	        FIsActionChecked::CreateSP(this, &SRigVMEditorGraphExplorer::IsShowingEmptySections)),
+	    NAME_None,
+	    EUserInterfaceActionType::ToggleButton,
+	    TEXT("RigVMGraphExplorer_ShowEmptySections"));
+
+	SAssignNew(FilterBox, SSearchBox)
+	    .OnTextChanged_Lambda([this](const FText& InFilterText) { TreeView->FilterText = InFilterText; TreeView->RefreshTreeView(); } );
+
+
+	FRigVMEditorGraphExplorerTreeDelegates Delegates;
+	Delegates.OnGetRootGraphs = FOnGetRootGraphs::CreateSP(this, &SRigVMEditorGraphExplorer::GetRootGraphs);
+	Delegates.OnGetChildrenGraphs = FOnGetChildrenGraphs::CreateSP(this, &SRigVMEditorGraphExplorer::GetChildrenGraphs);
+	Delegates.OnGetEventNodesInGraph = FOnGetEventNodesInGraph::CreateSP(this, &SRigVMEditorGraphExplorer::GetEventNodesInGraph);
+	Delegates.OnGetGraphDisplayName = FOnGetGraphDisplayName::CreateSP(this, &SRigVMEditorGraphExplorer::GetGraphDisplayName);
+	Delegates.OnGetEventDisplayName = FOnGetEventDisplayName::CreateSP(this, &SRigVMEditorGraphExplorer::GetEventDisplayName);
+	Delegates.OnGetGraphIcon = FOnGetGraphIcon::CreateSP(this, &SRigVMEditorGraphExplorer::GetGraphIcon);
+	Delegates.OnGetGraphTooltip = FOnGetGraphTooltip::CreateSP(this, &SRigVMEditorGraphExplorer::GetGraphTooltip);
+	Delegates.OnGraphClicked = FOnGraphClicked::CreateSP(this, &SRigVMEditorGraphExplorer::OnGraphClicked);
+	Delegates.OnEventClicked = FOnEventClicked::CreateSP(this, &SRigVMEditorGraphExplorer::OnEventClicked);
+	Delegates.OnGraphDoubleClicked = FOnGraphDoubleClicked::CreateSP(this, &SRigVMEditorGraphExplorer::OnGraphDoubleClicked);
+	Delegates.OnEventDoubleClicked = FOnEventDoubleClicked::CreateSP(this, &SRigVMEditorGraphExplorer::OnEventDoubleClicked);
+	Delegates.OnCreateGraph = FOnCreateGraph::CreateSP(this, &SRigVMEditorGraphExplorer::OnCreateGraph);
+	Delegates.OnRenameGraph = FOnRenameGraph::CreateSP(this, &SRigVMEditorGraphExplorer::OnRenameGraph);
+	Delegates.OnCanRenameGraph = FOnCanRenameGraph::CreateSP(this, &SRigVMEditorGraphExplorer::OnCanRenameGraph);
+	Delegates.OnRequestContextMenu = FOnRequestContextMenu::CreateSP(this, &SRigVMEditorGraphExplorer::OnContextMenuOpening);
+	
+	Delegates.OnSelectionChanged = FOnGraphExplorerSelectionChanged::CreateSP(this, &SRigVMEditorGraphExplorer::HandleSelectionChanged);
+	
+	SAssignNew(TreeView, SRigVMEditorGraphExplorerTreeView)
+		.RigTreeDelegates(Delegates)
+	;
+
+	// now piece together all the content for this widget
+	ChildSlot
+	[
+		SNew(SVerticalBox)
+	    + SVerticalBox::Slot()
+	    .AutoHeight()
+	    [
+			SNew(SBorder)
+	        .Padding(4.0f)
+	        .BorderImage(FAppStyle::GetBrush("ToolPanel.GroupBorder"))
+	        [
+				SNew(SVerticalBox)
+	            + SVerticalBox::Slot()
+	            .AutoHeight()
+	            [
+					SNew(SHorizontalBox)
+	                + SHorizontalBox::Slot()
+	                .AutoWidth()
+	                .Padding(0, 0, 2, 0)
+	                [
+						AddNewMenu.ToSharedRef()
+					]
+	                + SHorizontalBox::Slot()
+	                .FillWidth(1.0f)
+	                .VAlign(VAlign_Center)
+	                [
+						FilterBox.ToSharedRef()
+					]
+	                + SHorizontalBox::Slot()
+	                .AutoWidth()
+	                .Padding(2, 0, 0, 0)
+	                [
+						SNew(SComboButton)
+		                .ButtonStyle(FAppStyle::Get(), "SimpleButton")
+	                    .ComboButtonStyle(FAppStyle::Get(), "ToolbarComboButton")
+	                    .ForegroundColor(FSlateColor::UseForeground())
+	                    .HasDownArrow(false)
+	                    .ContentPadding(0)
+	                    .AddMetaData<FTagMetaData>(FTagMetaData(TEXT("ViewOptions")))
+	                    .MenuContent()
+	                    [
+							ViewOptions.MakeWidget()
+						]
+	                    .ButtonContent()
+	                    [
+	                    	SNew(SImage)
+							.Image(FAppStyle::GetBrush("Icons.Settings"))
+							.ColorAndOpacity(FSlateColor::UseForeground())
+						]
+					]
+				]
+			]
+		]
+	    + SVerticalBox::Slot()
+	    .FillHeight(1.0f)
+	    [
+			TreeView.ToSharedRef()
+		]
+	];
+
+	Refresh();
+}
+
+
+TSharedRef<SWidget> SRigVMEditorGraphExplorer::CreateAddNewMenuWidget()
+{
+	FMenuBuilder MenuBuilder(/* bShouldCloseWindowAfterMenuSelection= */true, RigVMEditor.Pin()->GetToolkitCommands());
+
+	BuildAddNewMenu(MenuBuilder);
+
+	return MenuBuilder.MakeWidget();
+}
+
+
+void SRigVMEditorGraphExplorer::BuildAddNewMenu(FMenuBuilder& MenuBuilder)
+{
+	MenuBuilder.BeginSection("AddNewItem", LOCTEXT("AddOperations", "Add New"));
+	MenuBuilder.AddMenuEntry(FRigVMEditorGraphExplorerCommands::Get().CreateGraph);
+	
+	MenuBuilder.EndSection();
+}
+
+TArray<const URigVMGraph*> SRigVMEditorGraphExplorer::GetRootGraphs() const
+{
+	TArray<const URigVMGraph*> Graphs;
+	
+	TSharedPtr<FRigVMEditor> Editor = RigVMEditor.Pin();
+	if (!Editor)
+	{
+		return Graphs;
+	}
+
+	FRigVMClient* RigVMClient = Editor->GetRigVMBlueprint()->GetRigVMClient();
+	for (URigVMGraph* Graph : RigVMClient->GetModels())
+	{
+		Graphs.Add(Graph);
+	}
+
+	return Graphs;
+}
+
+TArray<const URigVMGraph*> SRigVMEditorGraphExplorer::GetChildrenGraphs(const FString& InParentGraphPath) const
+{
+	TArray<const URigVMGraph*> Children;
+	
+	TSharedPtr<FRigVMEditor> Editor = RigVMEditor.Pin();
+	if (!Editor)
+	{
+		return Children;
+	}
+
+	FRigVMClient* RigVMClient = Editor->GetRigVMBlueprint()->GetRigVMClient();
+	URigVMGraph* ParentGraph = RigVMClient->GetModel(InParentGraphPath);
+	if (!ParentGraph)
+	{
+		return Children;
+	}
+
+	TArray<URigVMGraph*> ContainedGraphs = ParentGraph->GetContainedGraphs();
+	Children.Reserve(ContainedGraphs.Num());
+	for (URigVMGraph* Child : ContainedGraphs)
+	{
+		// Do not show contained graphs of aggregate nodes
+		if (URigVMAggregateNode* CollapseNode = Cast<URigVMAggregateNode>(Child->GetOuter()))
+		{
+			continue;
+		}
+		Children.Add(Child);
+	}
+	
+	return Children;
+}
+
+TArray<URigVMNode*> SRigVMEditorGraphExplorer::GetEventNodesInGraph(const FString& InParentGraphPath) const
+{
+	TArray<URigVMNode*> Events;
+	
+	TSharedPtr<FRigVMEditor> Editor = RigVMEditor.Pin();
+	if (!Editor)
+	{
+		return Events;
+	}
+
+	FRigVMClient* RigVMClient = Editor->GetRigVMBlueprint()->GetRigVMClient();
+	URigVMGraph* ParentGraph = RigVMClient->GetModel(InParentGraphPath);
+	if (!ParentGraph)
+	{
+		return Events;
+	}
+
+	if (!ParentGraph->IsTopLevelGraph())
+	{
+		return Events;
+	}
+	
+	if (ParentGraph->IsA<URigVMFunctionLibrary>())
+	{
+		return Events;
+	}
+
+	for (URigVMNode* Node : ParentGraph->GetNodes())
+	{
+		if (Node->IsEvent())
+		{
+			Events.Add(Node);
+		}
+	}
+
+	return Events;
+}
+
+FText SRigVMEditorGraphExplorer::GetGraphDisplayName(const FString& InGraphPath) const
+{
+	FText DisplayName;
+	TSharedPtr<FRigVMEditor> Editor = RigVMEditor.Pin();
+	if (!Editor)
+	{
+		return DisplayName;
+	}
+
+	FRigVMClient* RigVMClient = Editor->GetRigVMBlueprint()->GetRigVMClient();
+	if (const URigVMGraph* Graph = RigVMClient->GetModel(InGraphPath))
+	{
+		const URigVMEdGraphSchema* EdSchema = GetDefault<URigVMEdGraphSchema>();
+		FGraphDisplayInfo DisplayInfo;
+		UEdGraph* EdGraph = Editor->GetRigVMBlueprint()->GetEdGraph(Graph);
+		EdSchema->GetGraphDisplayInformation(*EdGraph, DisplayInfo);
+
+		DisplayName = DisplayInfo.DisplayName;
+	}
+	
+	return DisplayName;
+}
+
+FText SRigVMEditorGraphExplorer::GetEventDisplayName(const FString& InNodePath) const
+{
+	FText DisplayName;
+	TSharedPtr<FRigVMEditor> Editor = RigVMEditor.Pin();
+	if (!Editor)
+	{
+		return DisplayName;
+	}
+
+	FRigVMClient* RigVMClient = Editor->GetRigVMBlueprint()->GetRigVMClient();
+	if (const URigVMNode* Node = RigVMClient->FindNode(InNodePath))
+	{
+		DisplayName = FText::FromName(Node->GetEventName());
+	}
+	
+	return DisplayName;
+}
+
+const FSlateBrush* SRigVMEditorGraphExplorer::GetGraphIcon(const FString& InGraphPath) const
+{
+	const FSlateBrush* Brush = nullptr;
+	TSharedPtr<FRigVMEditor> Editor = RigVMEditor.Pin();
+	if (!Editor)
+	{
+		return Brush;
+	}
+
+	FRigVMClient* RigVMClient = Editor->GetRigVMBlueprint()->GetRigVMClient();
+	if (const URigVMGraph* Graph = RigVMClient->GetModel(InGraphPath))
+	{
+		if (Graph->IsRootGraph())
+		{
+			return FAppStyle::GetBrush(TEXT("GraphEditor.EventGraph_16x"));
+		}
+		else
+		{
+			return FAppStyle::GetBrush(TEXT("GraphEditor.SubGraph_16x"));
+		}
+	}
+	return nullptr;
+}
+
+FText SRigVMEditorGraphExplorer::GetGraphTooltip(const FString& InGraphPath) const
+{
+	FText Tooltip;
+	TSharedPtr<FRigVMEditor> Editor = RigVMEditor.Pin();
+	if (!Editor)
+	{
+		return Tooltip;
+	}
+
+	FRigVMClient* RigVMClient = Editor->GetRigVMBlueprint()->GetRigVMClient();
+	if (const URigVMGraph* Graph = RigVMClient->GetModel(InGraphPath))
+	{
+		const URigVMEdGraphSchema* EdSchema = GetDefault<URigVMEdGraphSchema>();
+		FGraphDisplayInfo DisplayInfo;
+		UEdGraph* EdGraph = Editor->GetRigVMBlueprint()->GetEdGraph(Graph);
+		EdSchema->GetGraphDisplayInformation(*EdGraph, DisplayInfo);
+
+		Tooltip = DisplayInfo.Tooltip;
+	}
+	
+	return Tooltip;
+}
+
+void SRigVMEditorGraphExplorer::OnGraphClicked(const FString& InGraphPath)
+{
+	TSharedPtr<FRigVMEditor> Editor = RigVMEditor.Pin();
+	if (!Editor)
+	{
+		return;
+	}
+
+	FRigVMClient* RigVMClient = Editor->GetRigVMBlueprint()->GetRigVMClient();
+	if (const URigVMGraph* Graph = RigVMClient->GetModel(InGraphPath))
+	{
+		const URigVMEdGraphSchema* EdSchema = GetDefault<URigVMEdGraphSchema>();
+		if (UEdGraph* EdGraph = Editor->GetRigVMBlueprint()->GetEdGraph(Graph))
+		{
+			FGraphDisplayInfo DisplayInfo;
+			const UEdGraphSchema* Schema = EdGraph->GetSchema();
+			check(Schema != nullptr);
+			Schema->GetGraphDisplayInformation(*EdGraph, DisplayInfo);
+			TSharedRef<SKismetInspector> CurrentInspector = Editor->GetInspector();
+			CurrentInspector->ShowDetailsForSingleObject(EdGraph, SKismetInspector::FShowDetailsOptions(DisplayInfo.PlainName));
+		}
+	}
+}
+
+void SRigVMEditorGraphExplorer::OnEventClicked(const FString& InEventPath)
+{
+	TSharedPtr<FRigVMEditor> Editor = RigVMEditor.Pin();
+	if (!Editor)
+	{
+		return;
+	}
+	TSharedRef<SKismetInspector> CurrentInspector = Editor->GetInspector();
+	CurrentInspector->ShowDetailsForObjects(TArray<UObject*>());
+}
+
+void SRigVMEditorGraphExplorer::OnGraphDoubleClicked(const FString& InGraphPath)
+{
+	TSharedPtr<FRigVMEditor> Editor = RigVMEditor.Pin();
+	if (!Editor)
+	{
+		return;
+	}
+
+	FRigVMClient* RigVMClient = Editor->GetRigVMBlueprint()->GetRigVMClient();
+	if (const URigVMGraph* Graph = RigVMClient->GetModel(InGraphPath))
+	{
+		const URigVMEdGraphSchema* EdSchema = GetDefault<URigVMEdGraphSchema>();
+		if (UEdGraph* EdGraph = Editor->GetRigVMBlueprint()->GetEdGraph(Graph))
+		{
+			Editor->JumpToHyperlink(EdGraph);
+		}
+	}
+}
+
+void SRigVMEditorGraphExplorer::OnEventDoubleClicked(const FString& InEventPath)
+{
+	TSharedPtr<FRigVMEditor> Editor = RigVMEditor.Pin();
+	if (!Editor)
+	{
+		return;
+	}
+
+	URigVMBlueprint* Blueprint = Editor->GetRigVMBlueprint();
+	FRigVMClient* RigVMClient = Blueprint->GetRigVMClient();
+
+	if(const URigVMNode* ModelNode = RigVMClient->FindNode(InEventPath))
+	{
+		Blueprint->OnRequestJumpToHyperlink().Execute(ModelNode);
+	}
+}
+
+TSharedPtr<SWidget> SRigVMEditorGraphExplorer::OnContextMenuOpening()
+{
+	TSharedPtr<FRigVMEditor> Editor = RigVMEditor.Pin();
+	if (!Editor.IsValid())
+	{
+		return TSharedPtr<SWidget>();
+	}
+
+	FMenuBuilder MenuBuilder(/*Close after selection*/true, Editor->GetToolkitCommands());
+
+	TArray<FRigVMExplorerElementKey> Selection = TreeView->GetSelectedKeys();
+
+	if (Selection.Num() == 1)
+	{
+		switch (Selection[0].Type)
+		{
+			case Section:
+			{
+				break;
+			}
+			case Graph:
+			{
+				MenuBuilder.BeginSection("BasicOperations");
+				MenuBuilder.AddMenuEntry(FRigVMEditorGraphExplorerCommands::Get().OpenGraph);
+				MenuBuilder.AddMenuEntry(FRigVMEditorGraphExplorerCommands::Get().OpenGraphInNewTab);
+				MenuBuilder.AddMenuEntry(FGenericCommands::Get().Rename);
+				// MenuBuilder.AddSubMenu(
+				// 			LOCTEXT("FindReferences_Label", "Find References"),
+				// 			LOCTEXT("FindReferences_Tooltip", "Options for finding references to class members"),
+				// 			FNewMenuDelegate::CreateStatic(&FGraphEditorCommands::BuildFindReferencesMenu),
+				// 			false,
+				// 			FSlateIcon());
+				MenuBuilder.AddMenuEntry(FGraphEditorCommands::Get().FindReferences);
+				MenuBuilder.AddMenuEntry(FGenericCommands::Get().Cut);
+				MenuBuilder.AddMenuEntry(FGenericCommands::Get().Copy);
+				MenuBuilder.AddMenuEntry(FGenericCommands::Get().Duplicate);
+				MenuBuilder.AddMenuEntry(FGenericCommands::Get().Delete);
+				MenuBuilder.EndSection();
+				break;
+			}
+			case Event:
+			{
+				return nullptr;
+				break;
+			}
+		}
+		
+	}
+	else
+	{
+		BuildAddNewMenu(MenuBuilder);
+	}
+
+	return MenuBuilder.MakeWidget();
+}
+
+void SRigVMEditorGraphExplorer::OnCopy()
+{
+}
+
+bool SRigVMEditorGraphExplorer::CanCopy() const
+{
+	return false;
+}
+
+void SRigVMEditorGraphExplorer::OnCut()
+{
+}
+
+bool SRigVMEditorGraphExplorer::CanCut() const
+{
+	return false;
+}
+
+void SRigVMEditorGraphExplorer::OnPasteGeneric()
+{
+}
+
+bool SRigVMEditorGraphExplorer::CanPasteGeneric() const
+{
+	return false;
+}
+
+void SRigVMEditorGraphExplorer::OnToggleShowEmptySections()
+{
+	// FIXME: Move to preferences
+	bShowEmptySections = !bShowEmptySections;
+
+	Refresh();
+}
+
+
+bool SRigVMEditorGraphExplorer::IsShowingEmptySections() const
+{
+	return bShowEmptySections;
+}
+
+
+void SRigVMEditorGraphExplorer::OnOpenGraph(bool bOnNewTab)
+{
+	TSharedPtr<FRigVMEditor> Editor = RigVMEditor.Pin();
+
+	if (Editor)
+	{
+		UEdGraph* GraphToOpen = nullptr;
+		TArray<FRigVMExplorerElementKey> Selection = TreeView->GetSelectedKeys();
+		if (Selection.Num() == 1 && Selection[0].Type == ERigVMExplorerElementType::Graph)
+		{
+			const URigVMEdGraphSchema* Schema = GetDefault<URigVMEdGraphSchema>();
+			GraphToOpen = Editor->GetRigVMBlueprint()->GetEdGraph(Selection[0].Name);
+		}
+
+		if (GraphToOpen)
+		{
+			const FDocumentTracker::EOpenDocumentCause Cause =bOnNewTab
+			? FDocumentTracker::ForceOpenNewDocument
+			: FDocumentTracker::OpenNewDocument;
+			
+			Editor->OpenDocument(GraphToOpen, Cause);
+		}
+	}
+}
+
+
+bool SRigVMEditorGraphExplorer::CanOpenGraph()
+{
+	TArray<FRigVMExplorerElementKey> Selection = TreeView->GetSelectedKeys();
+	return RigVMEditor.IsValid() && Selection.Num() == 1 && Selection[0].Type == ERigVMExplorerElementType::Graph;
+}
+
+
+void SRigVMEditorGraphExplorer::OnCreateGraph()
+{
+	TSharedPtr<FRigVMEditor> Editor = RigVMEditor.Pin();
+	if (Editor)
+	{
+		Editor->NewDocument_OnClicked(FBlueprintEditor::CGT_NewEventGraph);
+	}
+}
+
+
+bool SRigVMEditorGraphExplorer::CanCreateGraph()
+{
+	TSharedPtr<FRigVMEditor> Editor = RigVMEditor.Pin();
+	if (Editor)
+	{
+		return Editor->InEditingMode();
+	}
+
+	return false;
+}
+
+void SRigVMEditorGraphExplorer::OnDeleteEntry()
+{
+	TSharedPtr<FRigVMEditor> Editor = RigVMEditor.Pin();
+
+	if (Editor)
+	{
+		TArray<FRigVMExplorerElementKey> Selected = TreeView->GetSelectedKeys();
+		const URigVMEdGraphSchema* Schema = GetDefault<URigVMEdGraphSchema>();
+		for (const FRigVMExplorerElementKey& Key : Selected)
+		{
+			switch (Key.Type)
+			{
+				case ERigVMExplorerElementType::Graph:
+				{
+					if (UEdGraph* EdGraph = Editor->GetRigVMBlueprint()->GetEdGraph(Key.Name))
+					{
+						Schema->TryDeleteGraph(EdGraph);
+					}
+					break;
+				}
+			}
+		}
+	}
+
+	Refresh();
+}
+
+
+bool SRigVMEditorGraphExplorer::CanDeleteEntry()
+{
+	TSharedPtr<FRigVMEditor> Editor = RigVMEditor.Pin();
+
+	if (!Editor)
+	{
+		return false;
+	}
+	
+	TArray<FRigVMExplorerElementKey> Selected = TreeView->GetSelectedKeys();
+	for (const FRigVMExplorerElementKey& Key : Selected)
+	{
+		switch (Key.Type)
+		{
+			case ERigVMExplorerElementType::Section:
+				{
+					return false;
+					break;
+				}
+			case ERigVMExplorerElementType::Graph:
+				{
+					break;
+				}
+			default:
+				{
+					return false;
+				}
+		}
+	}
+	
+	return true;
+}
+
+
+void SRigVMEditorGraphExplorer::OnRenameEntry()
+{
+	TArray<FRigVMExplorerElementKey> Selection = TreeView->GetSelectedKeys();
+	if (Selection.Num() != 1)
+	{
+		return;
+	}
+
+	TSharedPtr<FRigVMEditorGraphExplorerTreeElement> Element = TreeView->FindElement(Selection[0]);
+	if (Element.IsValid())
+	{
+		Element->RequestRename();
+	}
+}
+
+
+bool SRigVMEditorGraphExplorer::CanRenameEntry()
+{
+	return true;
+}
+
+bool SRigVMEditorGraphExplorer::OnRenameGraph(const FString& InOldPath, const FString& InNewPath)
+{
+	bool bResult = false;
+	TSharedPtr<FRigVMEditor> Editor = RigVMEditor.Pin();
+
+	if (Editor)
+	{
+		const URigVMEdGraphSchema* Schema = GetDefault<URigVMEdGraphSchema>();
+		if (UEdGraph* EdGraph = Editor->GetRigVMBlueprint()->GetEdGraph(InOldPath))
+		{
+			bResult = Schema->TryRenameGraph(EdGraph, *InNewPath);
+		}
+	}
+
+	Refresh();
+	return bResult;
+}
+
+bool SRigVMEditorGraphExplorer::OnCanRenameGraph(const FString& InOldPath, const FString& InNewPath, FText& OutErrorMessage)
+{
+	bool bResult = false;
+	TSharedPtr<FRigVMEditor> Editor = RigVMEditor.Pin();
+
+	if (Editor)
+	{
+		int32 Index = INDEX_NONE;
+		FString Prefix, NewPath = InNewPath;
+		if(InOldPath.FindLastChar(TEXT('|'), Index))
+		{
+			Prefix = InOldPath.Left(Index+1);
+		}
+		if (!Prefix.IsEmpty())
+		{
+			NewPath = FString::Printf(TEXT("%s%s"), *Prefix, *InNewPath);
+		}
+		if (Editor->GetRigVMBlueprint()->GetRigVMClient()->GetModel(NewPath))
+		{
+			OutErrorMessage = FText::FromString(TEXT("Name already in use."));
+			bResult = false;
+		}
+		else
+		{
+			bResult = true;
+		}
+	}
+
+	return bResult;
+}
+
+void SRigVMEditorGraphExplorer::OnFindReference(bool bSearchAllBlueprints)
+{
+	FString SearchTerm;
+
+	if (!RigVMEditor.IsValid())
+	{
+		return;
+	}
+
+	TSharedPtr<FRigVMEditor> Editor = RigVMEditor.Pin();
+	URigVMBlueprint* Blueprint = Editor->GetRigVMBlueprint();
+	const URigVMEdGraphSchema* Schema = GetDefault<URigVMEdGraphSchema>();
+
+	bool bUseQuotes = true;
+	TArray<FRigVMExplorerElementKey> Selection = TreeView->GetSelectedKeys();
+	if (Selection.Num() != 1)
+	{
+		return;
+	}
+
+	if (Selection[0].Type == ERigVMExplorerElementType::Graph)
+	{
+		bool bSearchTermGenerated = false;
+		if (UEdGraph* EdGraph = Editor->GetRigVMBlueprint()->GetEdGraph(Selection[0].Name))
+		{
+			FGraphDisplayInfo DisplayInfo;
+			Schema->GetGraphDisplayInformation(*EdGraph, DisplayInfo);
+			SearchTerm = DisplayInfo.DisplayName.ToString();
+		}
+	}
+	
+	if (!SearchTerm.IsEmpty())
+	{
+		if (bUseQuotes)
+		{
+			SearchTerm = FString::Printf(TEXT("\"%s\""), *SearchTerm);
+		}
+		
+		const bool bSetFindWithinBlueprint = !bSearchAllBlueprints;
+		Editor->SummonSearchUI(bSetFindWithinBlueprint, SearchTerm);
+	}
+}
+
+bool SRigVMEditorGraphExplorer::CanFindReference() const
+{
+	return true;
+}
+
+void SRigVMEditorGraphExplorer::HandleSelectionChanged(TSharedPtr<FRigVMEditorGraphExplorerTreeElement> Selection, ESelectInfo::Type SelectInfo)
+{
+	TreeView->RefreshTreeView();
+}
+
+
+#undef LOCTEXT_NAMESPACE
diff --git a/Engine/Plugins/Runtime/RigVM/Source/RigVMEditor/Private/Widgets/SRigVMEditorGraphExplorerTreeView.cpp b/Engine/Plugins/Runtime/RigVM/Source/RigVMEditor/Private/Widgets/SRigVMEditorGraphExplorerTreeView.cpp
new file mode 100644
--- /dev/null
+++ b/Engine/Plugins/Runtime/RigVM/Source/RigVMEditor/Private/Widgets/SRigVMEditorGraphExplorerTreeView.cpp
@@ -0,0 +1,586 @@
+// Copyright Epic Games, Inc. All Rights Reserved.
+
+#include "Widgets/SRigVMEditorGraphExplorerTreeView.h"
+#include "Editor/RigVMEditor.h"
+#include "Widgets/SRigVMEditorGraphExplorer.h"
+
+#include "Widgets/Text/SInlineEditableTextBlock.h"
+#include "Widgets/Text/SRichTextBlock.h"
+
+#define LOCTEXT_NAMESPACE "SRigVMEditorGraphExplorerTreeView"
+
+static const FString GraphsSectionName = TEXT("Graphs");
+
+void SRigVMEditorGraphExplorerItem::Construct(const FArguments& InArgs,
+                                              const TSharedRef<STableViewBase>& InOwnerTable,
+                                              TSharedRef<FRigVMEditorGraphExplorerTreeElement> InElement,
+                                              TSharedPtr<SRigVMEditorGraphExplorerTreeView> InTreeView)
+{
+	WeakExplorerElement = InElement;
+	Delegates = InTreeView->GetRigTreeDelegates();
+	OnAddClickedOnSection = InArgs._OnAddClickedOnSection;
+
+	FSlateFontInfo NameFont = FCoreStyle::GetDefaultFontStyle("Regular", 10);
+
+	if (InElement->Key.Type == ERigVMExplorerElementType::Section)
+	{
+		STableRow<TSharedPtr<FRigVMEditorGraphExplorerTreeElement>>::Construct(
+		STableRow<TSharedPtr<FRigVMEditorGraphExplorerTreeElement>>::FArguments()
+		.ShowWires(true)
+		.Padding(FMargin(0.0f, 2.0f, .0f, 0.0f))
+		.Content()
+		[
+			SNew(SBorder)
+			.BorderImage(FAppStyle::Get().GetBrush("Brushes.Header")) // SCategoryHeaderTableRow::GetBackgroundImage
+			.Padding(FMargin(3.0f, 5.0f))
+			[
+				SNew(SHorizontalBox)
+				+ SHorizontalBox::Slot()
+				.VAlign(VAlign_Center)
+				[
+					SNew(SRichTextBlock)
+					.Text(FText::FromString(InElement->Key.Name))
+					.TransformPolicy(ETextTransformPolicy::ToUpper)
+					.DecoratorStyleSet(&FAppStyle::Get())
+					.TextStyle(FAppStyle::Get(), "DetailsView.CategoryTextStyle")
+				]
+
+				+ SHorizontalBox::Slot()
+				.AutoWidth()
+				.VAlign(VAlign_Center)
+				.HAlign(HAlign_Right)
+				.Padding(FMargin(0,0,2,0))
+				[
+					SNew(SButton)
+					.ButtonStyle(FAppStyle::Get(), "SimpleButton")
+					.ContentPadding(FMargin(1, 0))
+					.ForegroundColor(FAppStyle::GetSlateColor("DefaultForeground"))
+					.OnClicked(this, &SRigVMEditorGraphExplorerItem::OnAddButtonClickedOnSection, InElement)
+					// .IsEnabled(this, &SRigVMEditorGraphExplorer::CanAddNewElementToSection, InSectionID)
+					.HAlign(HAlign_Center)
+					.VAlign(VAlign_Center)
+					[
+						SNew(SImage)
+						.Image(FAppStyle::GetBrush(TEXT("Icons.PlusCircle")))
+						.ColorAndOpacity(FSlateColor::UseForeground())
+						.ToolTipText(LOCTEXT("AddNewGraphTooltip", "Create a new graph"))
+					]
+				]
+			]
+			
+		], InOwnerTable);
+	}
+	else
+	{
+		STableRow<TSharedPtr<FRigVMEditorGraphExplorerTreeElement>>::Construct(
+		   STableRow<TSharedPtr<FRigVMEditorGraphExplorerTreeElement>>::FArguments()
+		   .ShowWires(true)
+		   .Content()
+		   [
+			   SNew(SHorizontalBox)
+			   + SHorizontalBox::Slot()
+			   .AutoWidth()
+			   .VAlign(VAlign_Center)
+			   [
+				   CreateIconWidget(InElement->Key)
+			   ]
+			   + SHorizontalBox::Slot()
+			   .FillWidth(1.f)
+			   .VAlign(VAlign_Center)
+			   .Padding(/* horizontal */ 3.0f, /* vertical */ 3.0f)
+			   [
+				   CreateTextSlotWidget(InElement->Key, InTreeView->FilterText)
+			   ]
+		   ], InOwnerTable);
+	}
+
+	if(WeakExplorerElement.IsValid() && InlineRenameWidget.IsValid())
+	{
+		WeakExplorerElement.Pin()->OnRenameRequested.BindSP(InlineRenameWidget.Get(), &SInlineEditableTextBlock::EnterEditingMode);
+	}
+}
+
+TSharedRef<SWidget> SRigVMEditorGraphExplorerItem::CreateIconWidget(const FRigVMExplorerElementKey& Key)
+{
+	TSharedPtr<SWidget> IconWidget;
+
+	if (Key.Type == ERigVMExplorerElementType::Graph)
+	{
+		IconWidget = SNew(SImage)
+			.Image(Delegates.GetGraphIcon(Key.Name));
+	}
+	else if(Key.Type == ERigVMExplorerElementType::Event)
+	{
+		static FSlateIcon EventIcon(FAppStyle::GetAppStyleSetName(), "GraphEditor.Event_16x");
+		IconWidget = SNew(SImage)
+			.Image(EventIcon.GetIcon());
+	}
+
+	if (IconWidget.IsValid())
+	{
+		return IconWidget.ToSharedRef();
+	}
+	else
+	{
+		return SNullWidget::NullWidget;
+	}
+}
+
+TSharedRef<SWidget> SRigVMEditorGraphExplorerItem::CreateTextSlotWidget(const FRigVMExplorerElementKey& Key, const FText& InHighlightText)
+{
+	TSharedPtr<SInlineEditableTextBlock> EditableTextElement = SNew(SInlineEditableTextBlock)
+		.Text(this, &SRigVMEditorGraphExplorerItem::GetDisplayText)
+		.HighlightText(InHighlightText)
+		.ToolTipText(this, &SRigVMEditorGraphExplorerItem::GetItemTooltip)
+		.OnVerifyTextChanged(this, &SRigVMEditorGraphExplorerItem::OnVerifyNameChanged)
+		.OnTextCommitted(this, &SRigVMEditorGraphExplorerItem::OnNameCommitted);
+
+	InlineRenameWidget = EditableTextElement.ToSharedRef();
+
+	return InlineRenameWidget.ToSharedRef();
+}
+
+FText SRigVMEditorGraphExplorerItem::GetDisplayText() const
+{
+	FText DisplayText;
+	if (!WeakExplorerElement.IsValid())
+	{
+		return DisplayText;
+	}
+	TSharedPtr<FRigVMEditorGraphExplorerTreeElement> Element = WeakExplorerElement.Pin();
+
+	switch (Element->Key.Type)
+	{
+		case ERigVMExplorerElementType::Graph:
+			{
+				return Delegates.GetGraphDisplayName(Element->Key.Name);
+			}
+		case ERigVMExplorerElementType::Event:
+			{
+				return Delegates.GetEventDisplayName(Element->Key.Name);
+			}
+	}
+	
+	return DisplayText;
+}
+
+FText SRigVMEditorGraphExplorerItem::GetItemTooltip() const
+{
+	FText Tooltip;
+	if (!WeakExplorerElement.IsValid())
+	{
+		return Tooltip;
+	}
+	TSharedPtr<FRigVMEditorGraphExplorerTreeElement> Element = WeakExplorerElement.Pin();
+
+	switch (Element->Key.Type)
+	{
+		case ERigVMExplorerElementType::Graph:
+			{
+				return Delegates.GetGraphTooltip(Element->Key.Name);
+			}
+		case ERigVMExplorerElementType::Event:
+			{
+				return Delegates.GetEventDisplayName(Element->Key.Name);
+			}
+	}
+	
+	return Tooltip;
+}
+
+bool SRigVMEditorGraphExplorerItem::OnVerifyNameChanged(const FText& InText, FText& OutErrorMessage)
+{
+	FString NewName = InText.ToString();
+	const FRigVMExplorerElementKey OldKey = WeakExplorerElement.Pin()->Key;
+
+	switch (OldKey.Type)
+	{
+		case Section:
+			{
+				break;
+			}
+		case Graph:
+			{
+				return Delegates.CanRenameGraph(OldKey.Name, *NewName, OutErrorMessage);
+				break;
+			}
+	}
+	return false;
+}
+
+void SRigVMEditorGraphExplorerItem::OnNameCommitted(const FText& InText, ETextCommit::Type InCommitType) const
+{
+	// for now only allow enter
+	if (InCommitType == ETextCommit::OnEnter)
+	{
+		FString NewName = InText.ToString();
+		const FRigVMExplorerElementKey OldKey = WeakExplorerElement.Pin()->Key;
+
+		switch (OldKey.Type)
+		{
+			case Section:
+				{
+					break;
+				}
+			case Graph:
+				{
+					Delegates.RenameGraph(OldKey.Name, *NewName);
+					break;
+				}
+		}
+	}
+}
+
+FReply SRigVMEditorGraphExplorerItem::OnAddButtonClickedOnSection(TSharedRef<FRigVMEditorGraphExplorerTreeElement> InElement)
+{
+	if (OnAddClickedOnSection.IsBound())
+	{
+		return OnAddClickedOnSection.Execute(InElement->Key);
+	}
+	return FReply::Unhandled();
+}
+
+TSharedRef<ITableRow> FRigVMEditorGraphExplorerTreeElement::MakeTreeRowWidget(const TSharedRef<STableViewBase>& InOwnerTable,
+	TSharedRef<FRigVMEditorGraphExplorerTreeElement> InElement, TSharedPtr<SRigVMEditorGraphExplorerTreeView> InTreeView)
+{
+	return SNew(SRigVMEditorGraphExplorerItem, InOwnerTable, InElement, InTreeView)
+	.OnAddClickedOnSection(InTreeView.Get(), &SRigVMEditorGraphExplorerTreeView::OnAddButtonClickedOnSection);
+}
+
+void FRigVMEditorGraphExplorerTreeElement::RequestRename()
+{
+	OnRenameRequested.ExecuteIfBound();
+}
+
+void SRigVMEditorGraphExplorerTreeView::Construct(const FArguments& InArgs)
+{
+	Delegates = InArgs._RigTreeDelegates;
+	
+	STreeView<TSharedPtr<FRigVMEditorGraphExplorerTreeElement>>::FArguments SuperArgs;
+	SuperArgs.TreeItemsSource(&RootElements);
+	SuperArgs.SelectionMode(ESelectionMode::Multi);
+	SuperArgs.OnGenerateRow(this, &SRigVMEditorGraphExplorerTreeView::MakeTableRowWidget);
+	SuperArgs.OnGetChildren(this, &SRigVMEditorGraphExplorerTreeView::HandleGetChildrenForTree);
+	SuperArgs.HighlightParentNodesForSelection(true);
+	SuperArgs.AllowInvisibleItemSelection(true);  //without this we deselect everything when we filter or we collapse
+	SuperArgs.OnGeneratePinnedRow(this, &SRigVMEditorGraphExplorerTreeView::MakeTableRowWidget);
+	SuperArgs.OnMouseButtonClick(this, &SRigVMEditorGraphExplorerTreeView::OnItemClicked);
+	SuperArgs.OnMouseButtonDoubleClick(this, &SRigVMEditorGraphExplorerTreeView::OnItemDoubleClicked);
+	SuperArgs.OnSelectionChanged(FOnGraphExplorerSelectionChanged::CreateRaw(&Delegates, &FRigVMEditorGraphExplorerTreeDelegates::HandleSelectionChanged));
+	SuperArgs.OnContextMenuOpening(FOnRequestContextMenu::CreateRaw(&Delegates, &FRigVMEditorGraphExplorerTreeDelegates::RequestContextMenu));
+
+	STreeView<TSharedPtr<FRigVMEditorGraphExplorerTreeElement>>::Construct(SuperArgs);
+}
+
+void SRigVMEditorGraphExplorerTreeView::RefreshTreeView(bool bRebuildContent)
+{
+	TArray<FRigVMExplorerElementKey> Selection;
+	FString FilterTextStr = FilterText.ToString();
+	
+	if(bRebuildContent)
+	{
+		// store expansion state
+		TMap<FRigVMExplorerElementKey, bool> ExpansionState;
+		for (TPair<FRigVMExplorerElementKey, TSharedPtr<FRigVMEditorGraphExplorerTreeElement>> Pair : ElementMap)
+		{
+			ExpansionState.FindOrAdd(Pair.Key) = IsItemExpanded(Pair.Value);
+		}
+
+		Selection = GetSelectedKeys();
+
+		// internally save expansion states before rebuilding the tree, so the states can be restored later
+		SaveAndClearSparseItemInfos();
+		
+		RootElements.Reset();
+		ElementMap.Reset();
+		ParentMap.Reset();
+
+		// Add Elements
+		{
+			// GRAPHS section
+			FRigVMExplorerElementKey SectionKey(ERigVMExplorerElementType::Section, GraphsSectionName);
+			TSharedPtr<FRigVMEditorGraphExplorerTreeElement> SectionItem = MakeShared<FRigVMEditorGraphExplorerTreeElement>(SectionKey, SharedThis(this));
+			ElementMap.Add(SectionKey, SectionItem);
+			RootElements.Add(SectionItem);
+			SetItemExpansion(SectionItem, true);
+
+			TArray<const URigVMGraph*> Graphs = Delegates.GetRootGraphs();
+			for (const URigVMGraph* Graph : Graphs)
+			{
+				FRigVMExplorerElementKey Key(ERigVMExplorerElementType::Graph, Graph->GetNodePath());
+				ParentMap.Add(Key, SectionKey);
+			}
+			
+			// Add Graphs to section
+			for (int32 i=0; i<Graphs.Num(); ++i)
+			{
+				const URigVMGraph* Graph = Graphs[i];
+				
+				FRigVMExplorerElementKey Key(ERigVMExplorerElementType::Graph, Graph->GetNodePath());
+				
+				TSharedPtr<FRigVMEditorGraphExplorerTreeElement> NewItem = MakeShared<FRigVMEditorGraphExplorerTreeElement>(Key, SharedThis(this));
+				ElementMap.Add(Key, NewItem);
+				SetItemExpansion(NewItem, true);
+
+				TArray<const URigVMGraph*> ChildrenGraphs = Delegates.GetChildrenGraphs(Key.Name);
+				for (const URigVMGraph* Child : ChildrenGraphs)
+				{
+					FRigVMExplorerElementKey ChildKey(ERigVMExplorerElementType::Graph, Child->GetNodePath());
+					ParentMap.Add(ChildKey, Key);
+					Graphs.Add(Child);
+				}
+
+				TArray<URigVMNode*> EventNodes = Delegates.GetEventNodesInGraph(Key.Name);
+				for (const URigVMNode* EventNode : EventNodes)
+				{
+					
+					FRigVMExplorerElementKey EventKey(ERigVMExplorerElementType::Event, EventNode->GetNodePath());
+					ParentMap.Add(EventKey, Key);
+					TSharedPtr<FRigVMEditorGraphExplorerTreeElement> EventItem = MakeShared<FRigVMEditorGraphExplorerTreeElement>(EventKey, SharedThis(this));
+					ElementMap.Add(EventKey, EventItem);
+					NewItem->Children.Add(EventItem);
+				}
+
+				FRigVMExplorerElementKey ParentKey = ParentMap.FindChecked(Key);
+				TSharedPtr<FRigVMEditorGraphExplorerTreeElement> ParentItem = ElementMap.FindChecked(ParentKey);
+
+				ParentItem->Children.Add(NewItem);
+			}
+		}
+
+		// Filter Items
+		if (!FilterText.IsEmpty())
+		{
+			for (TSharedPtr<FRigVMEditorGraphExplorerTreeElement> RootElement : RootElements)
+			{
+				struct Local
+				{
+					static bool ChildNameSatisfiesFilter(TSharedPtr<FRigVMEditorGraphExplorerTreeElement> InElement, const FString& InFilterText, FRigVMEditorGraphExplorerTreeDelegates& Delegates)
+					{
+						FText DisplayName;
+						switch (InElement->Key.Type)
+						{
+						case Graph:
+							{
+								DisplayName = Delegates.GetGraphDisplayName(InElement->Key.Name);
+								break;
+							}
+						case Event:
+							{
+								DisplayName = Delegates.GetEventDisplayName(InElement->Key.Name);
+								break;
+							}
+						}
+						
+						if (InFilterText.IsEmpty() || DisplayName.ToString().Contains(InFilterText))
+						{
+							return true;
+						}
+						return false;
+					}
+					
+					static bool HasVisibleChildren(TSharedPtr<FRigVMEditorGraphExplorerTreeElement> InElement, const FString& InFilterText, FRigVMEditorGraphExplorerTreeDelegates& Delegates)
+					{
+						TArray<TSharedPtr<FRigVMEditorGraphExplorerTreeElement>> NewChildren;
+						NewChildren.Reserve(InElement->Children.Num());
+						for (TSharedPtr<FRigVMEditorGraphExplorerTreeElement> Child : InElement->Children)
+						{
+							if (HasVisibleChildren(Child, InFilterText, Delegates) || ChildNameSatisfiesFilter(Child, InFilterText, Delegates))
+							{
+								NewChildren.Add(Child);
+							}
+						}
+						InElement->Children = NewChildren;
+						return !NewChildren.IsEmpty();
+					}
+				};
+				TArray<TSharedPtr<FRigVMEditorGraphExplorerTreeElement>> NewChildren;
+				NewChildren.Reserve(RootElement->Children.Num());
+				for (TSharedPtr<FRigVMEditorGraphExplorerTreeElement> Child : RootElement->Children)
+				{
+					if (Local::HasVisibleChildren(Child, FilterTextStr, Delegates) || Local::ChildNameSatisfiesFilter(Child, FilterTextStr, Delegates))
+					{
+						NewChildren.Add(Child);
+					}
+				}
+				RootElement->Children = NewChildren;
+			}
+		}
+
+
+		// expand all elements upon the initial construction of the tree
+		if (ExpansionState.Num() == 0)
+		{
+			for (TSharedPtr<FRigVMEditorGraphExplorerTreeElement> RootElement : RootElements)
+			{
+				SetExpansionRecursive(RootElement, false, true);
+			}
+		}
+		// expand any new items
+		else if (ExpansionState.Num() < ElementMap.Num())
+		{
+			for (const TPair<FRigVMExplorerElementKey, TSharedPtr<FRigVMEditorGraphExplorerTreeElement>>& Element : ElementMap)
+			{
+				if (!ExpansionState.Contains(Element.Key))
+				{
+					SetItemExpansion(Element.Value, true);
+				}
+			}
+		}
+
+		// restore infos
+		for (const TPair<FRigVMExplorerElementKey, TSharedPtr<FRigVMEditorGraphExplorerTreeElement>>& Pair : ElementMap)
+		{
+			RestoreSparseItemInfos(Pair.Value);
+		}
+	}
+
+
+	RequestTreeRefresh();
+	{
+		TGuardValue<bool> Guard(Delegates.bSuspendSelectionDelegate, true);
+		ClearSelection();
+
+		if(!Selection.IsEmpty())
+		{
+			TArray<FRigVMExplorerElementKey> SelectedElements;
+			for(const FRigVMExplorerElementKey& SelectedPath : Selection)
+			{
+				SelectedElements.Add(SelectedPath);
+			}
+			if(!SelectedElements.IsEmpty())
+			{
+				SetSelection(SelectedElements);
+			}
+		}
+
+		Selection = GetSelectedKeys();
+	}
+}
+
+TSharedRef<ITableRow> SRigVMEditorGraphExplorerTreeView::MakeTableRowWidget(TSharedPtr<FRigVMEditorGraphExplorerTreeElement> InItem,
+	const TSharedRef<STableViewBase>& OwnerTable)
+{
+	return InItem->MakeTreeRowWidget(OwnerTable, InItem.ToSharedRef(), SharedThis(this));
+}
+
+void SRigVMEditorGraphExplorerTreeView::SetExpansionRecursive(TSharedPtr<FRigVMEditorGraphExplorerTreeElement> InElement, bool bTowardsParent, bool bShouldBeExpanded)
+{
+	SetItemExpansion(InElement, bShouldBeExpanded);
+
+	if (bTowardsParent)
+	{
+		if (const FRigVMExplorerElementKey* ParentKey = ParentMap.Find(InElement->Key))
+		{
+			if (TSharedPtr<FRigVMEditorGraphExplorerTreeElement>* ParentItem = ElementMap.Find(*ParentKey))
+			{
+				SetExpansionRecursive(*ParentItem, bTowardsParent, bShouldBeExpanded);
+			}
+		}
+	}
+	else
+	{
+		for (int32 ChildIndex = 0; ChildIndex < InElement->Children.Num(); ++ChildIndex)
+		{
+			SetExpansionRecursive(InElement->Children[ChildIndex], bTowardsParent, bShouldBeExpanded);
+		}
+	}
+}
+
+void SRigVMEditorGraphExplorerTreeView::HandleGetChildrenForTree(TSharedPtr<FRigVMEditorGraphExplorerTreeElement> InItem,
+	TArray<TSharedPtr<FRigVMEditorGraphExplorerTreeElement>>& OutChildren)
+{
+	OutChildren = InItem.Get()->Children;
+}
+
+TArray<FRigVMExplorerElementKey> SRigVMEditorGraphExplorerTreeView::GetSelectedKeys() const
+{
+	TArray<FRigVMExplorerElementKey> Keys;
+	TArray<TSharedPtr<FRigVMEditorGraphExplorerTreeElement>> SelectedElements = GetSelectedItems();
+	for(const TSharedPtr<FRigVMEditorGraphExplorerTreeElement>& SelectedElement : SelectedElements)
+	{
+		if (SelectedElement.IsValid())
+		{
+			Keys.AddUnique(SelectedElement->Key);
+		}
+	}
+	return Keys;
+}
+
+void SRigVMEditorGraphExplorerTreeView::SetSelection(TArray<FRigVMExplorerElementKey>& InSelectedKeys)
+{
+	ClearSelection();
+	TArray<TSharedPtr<FRigVMEditorGraphExplorerTreeElement>> Selection;
+	for (const FRigVMExplorerElementKey& Key : InSelectedKeys)
+	{
+		if (TSharedPtr<FRigVMEditorGraphExplorerTreeElement>* Element = ElementMap.Find(Key))
+		{
+			Selection.Add(*Element);
+		}
+	}
+	SetItemSelection(Selection, true, ESelectInfo::Direct);
+}
+
+TSharedPtr<FRigVMEditorGraphExplorerTreeElement> SRigVMEditorGraphExplorerTreeView::FindElement(const FRigVMExplorerElementKey& Key)
+{
+	if (TSharedPtr<FRigVMEditorGraphExplorerTreeElement>* Element = ElementMap.Find(Key))
+	{
+		return *Element;
+	}
+	return nullptr;
+}
+
+void SRigVMEditorGraphExplorerTreeView::OnItemClicked(TSharedPtr<FRigVMEditorGraphExplorerTreeElement> InElement)
+{
+	switch(InElement->Key.Type)
+	{
+		case Section:
+			{
+				FRigVMExplorerElementKey SectionKey(ERigVMExplorerElementType::Section, GraphsSectionName);
+				TSharedPtr<FRigVMEditorGraphExplorerTreeElement> Section = ElementMap.FindChecked(SectionKey);
+				SetItemExpansion(Section, !IsItemExpanded(Section));
+				break;
+			}
+		case Graph:
+			{
+				Delegates.GraphClicked(InElement->Key.Name);
+				break;
+			}
+		case Event:
+			{
+				Delegates.EventClicked(InElement->Key.Name);
+			}
+	}
+}
+
+void SRigVMEditorGraphExplorerTreeView::OnItemDoubleClicked(TSharedPtr<FRigVMEditorGraphExplorerTreeElement> InElement)
+{
+	switch(InElement->Key.Type)
+	{
+		case Section:
+			{
+				break;
+			}
+		case Graph:
+			{
+				Delegates.GraphDoubleClicked(InElement->Key.Name);
+				break;
+			}
+		case Event:
+			{
+				Delegates.EventDoubleClicked(InElement->Key.Name);
+			}
+	}
+}
+
+FReply SRigVMEditorGraphExplorerTreeView::OnAddButtonClickedOnSection(const FRigVMExplorerElementKey& InSectionKey)
+{
+	if (InSectionKey.Name == GraphsSectionName)
+	{
+		Delegates.CreateGraph();
+		return FReply::Handled();
+	}
+	return FReply::Unhandled();
+}
+
+#undef LOCTEXT_NAMESPACE
\ No newline at end of file
diff --git a/Engine/Plugins/Runtime/RigVM/Source/RigVMEditor/Public/Editor/RigVMEditor.h b/Engine/Plugins/Runtime/RigVM/Source/RigVMEditor/Public/Editor/RigVMEditor.h
--- a/Engine/Plugins/Runtime/RigVM/Source/RigVMEditor/Public/Editor/RigVMEditor.h
+++ b/Engine/Plugins/Runtime/RigVM/Source/RigVMEditor/Public/Editor/RigVMEditor.h
@@ -9,6 +9,7 @@
 #include "Editor/RigVMDetailsViewWrapperObject.h"
 #include "AssetRegistry/AssetData.h"
 
+class SRigVMEditorGraphExplorer;
 class FRigVMEditor;
 
 DECLARE_MULTICAST_DELEGATE_TwoParams(FRigVMEditorClosed, const FRigVMEditor*, URigVMBlueprint*);
@@ -335,6 +336,8 @@ private:
 	/** Toolbox hosting widget */
 	TSharedPtr<SBorder> Toolbox;
 
+	TSharedPtr<SRigVMEditorGraphExplorer> GraphExplorerWidget;
+
 	bool bAnyErrorsLeft;
 	TMap<FString, FString> KnownInstructionLimitWarnings;
 	URigVMNode* HaltedAtNode;
@@ -371,4 +374,5 @@ protected:
 	FDelegateHandle PropertyChangedHandle;
 
 	friend class SRigVMExecutionStackView;
+	friend class SRigVMEditorGraphExplorer;
 };
diff --git a/Engine/Plugins/Runtime/RigVM/Source/RigVMEditor/Public/Editor/RigVMEditorGraphExplorerTabSummoner.h b/Engine/Plugins/Runtime/RigVM/Source/RigVMEditor/Public/Editor/RigVMEditorGraphExplorerTabSummoner.h
new file mode 100644
--- /dev/null
+++ b/Engine/Plugins/Runtime/RigVM/Source/RigVMEditor/Public/Editor/RigVMEditorGraphExplorerTabSummoner.h
@@ -0,0 +1,23 @@
+// Copyright Epic Games, Inc. All Rights Reserved.
+
+#pragma once
+
+#include "CoreMinimal.h"
+#include "WorkflowOrientedApp/WorkflowTabFactory.h"
+
+class FRigVMEditor;
+
+struct RIGVMEDITOR_API FRigVMEditorGraphExplorerTabSummoner : public FWorkflowTabFactory
+{
+public:
+
+	static inline const FLazyName TabID = FLazyName(TEXT("RigVM Graph Explorer"));
+	
+public:
+	FRigVMEditorGraphExplorerTabSummoner(const TSharedRef<FRigVMEditor>& InRigVMEditor);
+	
+	virtual TSharedRef<SWidget> CreateTabBody(const FWorkflowTabSpawnInfo& Info) const override;
+	
+protected:
+	TWeakPtr<FRigVMEditor> RigVMEditor;
+};
diff --git a/Engine/Plugins/Runtime/RigVM/Source/RigVMEditor/Public/Widgets/SRigVMEditorGraphExplorer.h b/Engine/Plugins/Runtime/RigVM/Source/RigVMEditor/Public/Widgets/SRigVMEditorGraphExplorer.h
new file mode 100644
--- /dev/null
+++ b/Engine/Plugins/Runtime/RigVM/Source/RigVMEditor/Public/Widgets/SRigVMEditorGraphExplorer.h
@@ -0,0 +1,117 @@
+// Copyright Epic Games, Inc. All Rights Reserved.
+
+#pragma once
+
+#include "CoreMinimal.h"
+#include "Widgets/DeclarativeSyntaxSupport.h"
+#include "Widgets/SCompoundWidget.h"
+#include "Framework/Commands/Commands.h"
+#include "RigVMModel/RigVMGraph.h"
+
+class SRigVMEditorGraphExplorerTreeView;
+//
+class URigVMGraph;
+class URigVMEdGraph;
+class FRigVMEditor;
+class SGraphActionMenu;
+struct FCreateWidgetForActionData;
+class FRigVMEditorGraphExplorerTreeElement;
+
+class FMenuBuilder;
+class SSearchBox;
+//
+
+class FRigVMEditorGraphExplorerCommands : public TCommands<FRigVMEditorGraphExplorerCommands>
+{
+public:
+	FRigVMEditorGraphExplorerCommands();
+
+	TSharedPtr<FUICommandInfo> OpenGraph;
+	TSharedPtr<FUICommandInfo> OpenGraphInNewTab;
+	TSharedPtr<FUICommandInfo> CreateGraph;
+	TSharedPtr<FUICommandInfo> DeleteGraph;
+
+	void RegisterCommands() override;
+};
+
+class SRigVMEditorGraphExplorer :
+	public SCompoundWidget
+{
+public:
+	SLATE_BEGIN_ARGS(SRigVMEditorGraphExplorer) {}
+	SLATE_END_ARGS()
+
+	void Construct(const FArguments& InArgs, TWeakPtr<FRigVMEditor> InRigVMEditor);
+	~SRigVMEditorGraphExplorer();
+
+	// Refresh the graph action menu.
+	void Refresh();
+
+	// SWidget overrides
+	void Tick(const FGeometry& InAllottedGeometry, const double InCurrentTime, const float InDeltaTime) override;
+
+private:
+	
+	void RegisterCommands();
+	void CreateWidgets();
+
+	// Add new menu
+	TSharedRef<SWidget> CreateAddNewMenuWidget();
+	void BuildAddNewMenu(FMenuBuilder& MenuBuilder);
+
+	TArray<const URigVMGraph*> GetRootGraphs() const;
+	TArray<const URigVMGraph*> GetChildrenGraphs(const FString& InParentGraphPath) const;
+	TArray<URigVMNode*> GetEventNodesInGraph(const FString& InParentGraphPath) const;
+	FText GetGraphDisplayName(const FString& InGraphPath) const;
+	FText GetEventDisplayName(const FString& InNodePath) const;
+	const FSlateBrush* GetGraphIcon(const FString& InGraphPath) const;
+	FText GetGraphTooltip(const FString& InGraphPath) const;
+	void OnGraphClicked(const FString& InGraphPath);
+	void OnEventClicked(const FString& InEventPath);
+	void OnGraphDoubleClicked(const FString& InGraphPath);
+	void OnEventDoubleClicked(const FString& InEventPath);
+	
+	TSharedPtr<SWidget> OnContextMenuOpening();
+	
+	void OnCopy();
+	bool CanCopy() const;
+	void OnCut();
+	bool CanCut() const;
+	void OnPasteGeneric();
+	bool CanPasteGeneric() const;
+
+
+
+	/** Support functions for view options for Show Empty Sections */
+	void OnToggleShowEmptySections();
+	bool IsShowingEmptySections() const;
+
+	// Command functions
+	void OnOpenGraph(bool bOnNewTab = false);
+	bool CanOpenGraph();
+
+	void OnCreateGraph();
+	bool CanCreateGraph();
+	
+	void OnDeleteEntry();
+	bool CanDeleteEntry();
+
+	void OnRenameEntry();
+	bool CanRenameEntry();
+	bool OnRenameGraph(const FString& InOldPath, const FString& InNewPath);
+	bool OnCanRenameGraph(const FString& InOldPath, const FString& InNewPath, FText& OutErrorMessage);
+
+	void OnFindReference(bool bSearchAllBlueprints);
+	bool CanFindReference() const;
+
+	void HandleSelectionChanged(TSharedPtr<FRigVMEditorGraphExplorerTreeElement> Selection, ESelectInfo::Type SelectInfo);
+
+	TWeakPtr<FRigVMEditor> RigVMEditor;
+
+	TSharedPtr<SRigVMEditorGraphExplorerTreeView> TreeView;
+
+	TSharedPtr<SSearchBox> FilterBox;
+
+	bool bShowEmptySections = true;
+};
+
diff --git a/Engine/Plugins/Runtime/RigVM/Source/RigVMEditor/Public/Widgets/SRigVMEditorGraphExplorerTreeView.h b/Engine/Plugins/Runtime/RigVM/Source/RigVMEditor/Public/Widgets/SRigVMEditorGraphExplorerTreeView.h
new file mode 100644
--- /dev/null
+++ b/Engine/Plugins/Runtime/RigVM/Source/RigVMEditor/Public/Widgets/SRigVMEditorGraphExplorerTreeView.h
@@ -0,0 +1,367 @@
+// Copyright Epic Games, Inc. All Rights Reserved.
+
+#pragma once
+
+#include "CoreMinimal.h"
+#include "SGraphPalette.h"
+#include "RigVMModel/RigVMGraph.h"
+#include "Widgets/Views/STreeView.h"
+
+class FRigVMEditor;
+class SSearchBox;
+class SRigVMEditorGraphExplorerTreeView;
+class SRigVMEditorGraphExplorerItem;
+class FRigVMEditorGraphExplorerTreeElement;
+class URigVMGraph;
+
+enum ERigVMExplorerElementType
+{
+	Section,
+	Graph,
+	Event,
+	Function,
+	Variable,
+	LocalVariable
+};
+
+struct FRigVMExplorerElementKey
+{
+	ERigVMExplorerElementType Type;
+	FString Name;
+
+	friend bool operator==(const FRigVMExplorerElementKey& A, const FRigVMExplorerElementKey& B)
+	{
+		return A.Type == B.Type && A.Name == B.Name;
+	}
+
+	friend uint32 GetTypeHash(const FRigVMExplorerElementKey& Key)
+	{
+		return HashCombine(GetTypeHash(Key.Type), GetTypeHash(Key.Name));
+	}
+};
+
+DECLARE_DELEGATE_RetVal(TArray<const URigVMGraph*>, FOnGetRootGraphs);
+DECLARE_DELEGATE( FOnCreateGraph);
+DECLARE_DELEGATE_RetVal_OneParam(TArray<const URigVMGraph*>, FOnGetChildrenGraphs, const FString&);
+DECLARE_DELEGATE_RetVal_OneParam(TArray<URigVMNode*>, FOnGetEventNodesInGraph, const FString&);
+DECLARE_DELEGATE_RetVal_OneParam(FText, FOnGetGraphDisplayName, const FString&);
+DECLARE_DELEGATE_RetVal_OneParam(FText, FOnGetEventDisplayName, const FString&);
+DECLARE_DELEGATE_RetVal_OneParam(FText, FOnGetGraphTooltip, const FString&);
+DECLARE_DELEGATE_RetVal_OneParam(const FSlateBrush*, FOnGetGraphIcon, const FString&);
+DECLARE_DELEGATE_OneParam(FOnGraphClicked, const FString&);
+DECLARE_DELEGATE_OneParam(FOnEventClicked, const FString&);
+DECLARE_DELEGATE_OneParam(FOnGraphDoubleClicked, const FString&);
+DECLARE_DELEGATE_OneParam(FOnEventDoubleClicked, const FString&);
+DECLARE_DELEGATE_RetVal_TwoParams(bool, FOnRenameGraph, const FString&, const FString&);
+DECLARE_DELEGATE_RetVal_ThreeParams(bool, FOnCanRenameGraph, const FString&, const FString&, FText&);
+DECLARE_DELEGATE_RetVal(TSharedPtr<SWidget>, FOnRequestContextMenu);
+DECLARE_DELEGATE_RetVal(bool, FOnHasContextMenu);
+
+typedef STreeView<TSharedPtr<FRigVMEditorGraphExplorerTreeElement>>::FOnSelectionChanged FOnGraphExplorerSelectionChanged;
+
+struct FRigVMEditorGraphExplorerTreeDelegates
+{
+public:
+	FOnGetRootGraphs OnGetRootGraphs;
+	FOnGetChildrenGraphs OnGetChildrenGraphs;
+	FOnGetEventNodesInGraph OnGetEventNodesInGraph;
+	FOnGetGraphDisplayName OnGetGraphDisplayName;
+	FOnGetGraphDisplayName OnGetEventDisplayName;
+	FOnGetGraphTooltip OnGetGraphTooltip;
+	FOnGetGraphIcon OnGetGraphIcon;
+	FOnGraphDoubleClicked OnGraphClicked;
+	FOnEventDoubleClicked OnEventClicked;
+	FOnGraphDoubleClicked OnGraphDoubleClicked;
+	FOnEventDoubleClicked OnEventDoubleClicked;
+	FOnCreateGraph OnCreateGraph;
+	FOnRenameGraph OnRenameGraph;
+	FOnCanRenameGraph OnCanRenameGraph;
+	FOnGraphExplorerSelectionChanged OnSelectionChanged;
+	FOnRequestContextMenu OnRequestContextMenu;
+	FOnHasContextMenu OnHasContextMenu;
+
+	FRigVMEditorGraphExplorerTreeDelegates()
+	{
+	}
+
+	TArray<const URigVMGraph*> GetRootGraphs() const
+	{
+		if (OnGetRootGraphs.IsBound())
+		{
+			return OnGetRootGraphs.Execute();
+		}
+		return TArray<const URigVMGraph*>();
+	}
+
+	TArray<const URigVMGraph*> GetChildrenGraphs(const FString& InParentPath) const
+	{
+		if (OnGetChildrenGraphs.IsBound())
+		{
+			return OnGetChildrenGraphs.Execute(InParentPath);
+		}
+		return TArray<const URigVMGraph*>();
+	}
+
+	TArray<URigVMNode*> GetEventNodesInGraph(const FString& InParentPath) const
+	{
+		if (OnGetEventNodesInGraph.IsBound())
+		{
+			return OnGetEventNodesInGraph.Execute(InParentPath);
+		}
+		return TArray<URigVMNode*>();
+	}
+
+	FText GetGraphDisplayName(const FString& InPath) const
+	{
+		if (OnGetGraphDisplayName.IsBound())
+		{
+			return OnGetGraphDisplayName.Execute(InPath);
+		}
+		return FText();
+	}
+
+	FText GetEventDisplayName(const FString& InPath) const
+	{
+		if (OnGetEventDisplayName.IsBound())
+		{
+			return OnGetEventDisplayName.Execute(InPath);
+		}
+		return FText();
+	}
+	
+	FText GetGraphTooltip(const FString& InPath) const
+	{
+		if (OnGetGraphTooltip.IsBound())
+		{
+			return OnGetGraphTooltip.Execute(InPath);
+		}
+		return FText();
+	}
+
+	const FSlateBrush* GetGraphIcon(const FString& InPath) const
+	{
+		if (OnGetGraphIcon.IsBound())
+		{
+			return OnGetGraphIcon.Execute(InPath);
+		}
+		return nullptr;
+	}
+
+	void GraphClicked(const FString& InPath) const
+	{
+		if (OnGraphClicked.IsBound())
+		{
+			OnGraphClicked.Execute(InPath);
+		}
+	}
+
+	void EventClicked(const FString& InPath) const
+	{
+		if (OnEventClicked.IsBound())
+		{
+			OnEventClicked.Execute(InPath);
+		}
+	}
+
+	void GraphDoubleClicked(const FString& InPath) const
+	{
+		if (OnGraphDoubleClicked.IsBound())
+		{
+			OnGraphDoubleClicked.Execute(InPath);
+		}
+	}
+
+	void EventDoubleClicked(const FString& InPath) const
+	{
+		if (OnEventDoubleClicked.IsBound())
+		{
+			OnEventDoubleClicked.Execute(InPath);
+		}
+	}
+
+	void CreateGraph() const
+	{
+		if (OnCreateGraph.IsBound())
+		{
+			OnCreateGraph.Execute();
+		}
+	}
+
+	void RenameGraph(const FString& InOldPath, const FString& InNewPath) const
+	{
+		if (OnRenameGraph.IsBound())
+		{
+			OnRenameGraph.Execute(InOldPath, InNewPath);
+		}
+	}
+
+	bool CanRenameGraph(const FString& InOldPath, const FString& InNewPath, FText& OutErrorMessage) const
+	{
+		if (OnCanRenameGraph.IsBound())
+		{
+			return OnCanRenameGraph.Execute(InOldPath, InNewPath, OutErrorMessage);
+		}
+		return false;
+	}
+
+	void HandleSelectionChanged(TSharedPtr<FRigVMEditorGraphExplorerTreeElement> Selection, ESelectInfo::Type SelectInfo)
+	{
+		if(bSuspendSelectionDelegate)
+		{
+			return;
+		}
+		TGuardValue<bool> Guard(bSuspendSelectionDelegate, true);
+		(void)OnSelectionChanged.ExecuteIfBound(Selection, SelectInfo);
+	}
+
+	TSharedPtr<SWidget> RequestContextMenu()
+	{
+		if (OnRequestContextMenu.IsBound())
+		{
+			return OnRequestContextMenu.Execute();
+		}
+		return nullptr;
+	}
+
+	bool HasContextMenu()
+	{
+		if (OnHasContextMenu.IsBound())
+		{
+			return OnHasContextMenu.Execute();
+		}
+		return false;
+	}
+	
+private:
+
+	bool bSuspendSelectionDelegate = false;
+
+	friend class SRigVMEditorGraphExplorerTreeView;
+};
+
+class SRigVMEditorGraphExplorerItem  : public STableRow<TSharedPtr<FRigVMEditorGraphExplorerTreeElement>>
+{
+	DECLARE_DELEGATE_RetVal_OneParam(FReply, FOnAddClickedOnSection, const FRigVMExplorerElementKey&);
+	
+public:
+	SLATE_BEGIN_ARGS(SRigVMEditorGraphExplorerItem)
+	{}
+		SLATE_EVENT( FOnAddClickedOnSection, OnAddClickedOnSection )
+	SLATE_END_ARGS()
+
+	void Construct(const FArguments& InArgs,
+		const TSharedRef<STableViewBase>& InOwnerTable,
+		TSharedRef<FRigVMEditorGraphExplorerTreeElement> InElement,
+		TSharedPtr<SRigVMEditorGraphExplorerTreeView> InTreeView);
+
+	TSharedRef<SWidget> CreateIconWidget(const FRigVMExplorerElementKey& Key);
+	TSharedRef<SWidget> CreateTextSlotWidget(const FRigVMExplorerElementKey& Key, const FText& InHighlightText);
+	FText GetDisplayText() const;
+	FText GetItemTooltip() const;
+	bool OnVerifyNameChanged(const FText& InText, FText& OutErrorMessage);
+	void OnNameCommitted(const FText& InText, ETextCommit::Type InCommitType) const;
+	FReply OnAddButtonClickedOnSection(TSharedRef<FRigVMEditorGraphExplorerTreeElement> InElement);
+	
+	private:
+	
+	TWeakPtr<FRigVMEditorGraphExplorerTreeElement> WeakExplorerElement;
+	TSharedPtr<SInlineEditableTextBlock> InlineRenameWidget;
+	FRigVMEditorGraphExplorerTreeDelegates Delegates;
+
+	FOnAddClickedOnSection OnAddClickedOnSection;
+};
+
+
+/** An item in the tree */
+class FRigVMEditorGraphExplorerTreeElement : public TSharedFromThis<FRigVMEditorGraphExplorerTreeElement>
+{
+public:
+	FRigVMEditorGraphExplorerTreeElement(const FRigVMExplorerElementKey& InKey, TWeakPtr<SRigVMEditorGraphExplorerTreeView> InTreeView)
+		: Key(InKey) {}
+
+	FRigVMExplorerElementKey Key;
+	TArray<TSharedPtr<FRigVMEditorGraphExplorerTreeElement>> Children;
+
+	TSharedRef<ITableRow> MakeTreeRowWidget(const TSharedRef<STableViewBase>& InOwnerTable, TSharedRef<FRigVMEditorGraphExplorerTreeElement> InRigTreeElement, TSharedPtr<SRigVMEditorGraphExplorerTreeView> InTreeView);
+
+	void RequestRename();
+
+	/** Delegate for when the context menu requests a rename */
+	DECLARE_DELEGATE(FOnRenameRequested);
+	FOnRenameRequested OnRenameRequested;
+
+};
+
+class SRigVMEditorGraphExplorerTreeView : public STreeView<TSharedPtr<FRigVMEditorGraphExplorerTreeElement>>
+{
+public:
+
+	SLATE_BEGIN_ARGS(SRigVMEditorGraphExplorerTreeView)
+	{}
+		SLATE_ARGUMENT(FRigVMEditorGraphExplorerTreeDelegates, RigTreeDelegates)
+	SLATE_END_ARGS()
+
+	void Construct(const FArguments& InArgs);
+	virtual ~SRigVMEditorGraphExplorerTreeView() {}
+
+	/** Save a snapshot of the internal map that tracks item expansion before tree reconstruction */
+	void SaveAndClearSparseItemInfos()
+	{
+		// Only save the info if there is something to save (do not overwrite info with an empty map)
+		if (!SparseItemInfos.IsEmpty())
+		{
+			OldSparseItemInfos = SparseItemInfos;
+		}
+		ClearExpandedItems();
+	}
+
+	/** Restore the expansion infos map from the saved snapshot after tree reconstruction */
+	void RestoreSparseItemInfos(TSharedPtr<FRigVMEditorGraphExplorerTreeElement> ItemPtr)
+	{
+		for (const auto& Pair : OldSparseItemInfos)
+		{
+			if (Pair.Key->Key == ItemPtr->Key)
+			{
+				// the SparseItemInfos now reference the new element, but keep the same expansion state
+				SparseItemInfos.Add(ItemPtr, Pair.Value);
+				break;
+			}
+		}
+	}
+
+	void RefreshTreeView(bool bRebuildContent = true);
+	TSharedRef<ITableRow> MakeTableRowWidget(TSharedPtr<FRigVMEditorGraphExplorerTreeElement> InItem, const TSharedRef<STableViewBase>& OwnerTable);
+	TSharedPtr< SWidget > CreateContextMenu();
+	void SetExpansionRecursive(TSharedPtr<FRigVMEditorGraphExplorerTreeElement> InElement, bool bTowardsParent, bool bShouldBeExpanded);
+	void HandleGetChildrenForTree(TSharedPtr<FRigVMEditorGraphExplorerTreeElement> InItem, TArray<TSharedPtr<FRigVMEditorGraphExplorerTreeElement>>& OutChildren);
+	TArray<FRigVMExplorerElementKey> GetSelectedKeys() const;
+	void SetSelection(TArray<FRigVMExplorerElementKey>& InSelectedKeys);
+	TSharedPtr<FRigVMEditorGraphExplorerTreeElement> FindElement(const FRigVMExplorerElementKey& Key);
+	
+	const TArray<TSharedPtr<FRigVMEditorGraphExplorerTreeElement>>& GetRootElements() const { return RootElements; }
+	FRigVMEditorGraphExplorerTreeDelegates& GetRigTreeDelegates() { return Delegates; }
+	void OnItemClicked(TSharedPtr<FRigVMEditorGraphExplorerTreeElement> InElement);
+	void OnItemDoubleClicked(TSharedPtr<FRigVMEditorGraphExplorerTreeElement> InElement);
+	FReply OnAddButtonClickedOnSection(const FRigVMExplorerElementKey& InSectionKey);
+
+private:
+
+	/** A temporary snapshot of the SparseItemInfos in STreeView, used during RefreshTreeView() */
+	TSparseItemMap OldSparseItemInfos;
+
+	/** Backing array for tree view */
+	TArray<TSharedPtr<FRigVMEditorGraphExplorerTreeElement>> RootElements;
+
+	/** A map for looking up items based on their key */
+	TMap<FRigVMExplorerElementKey, TSharedPtr<FRigVMEditorGraphExplorerTreeElement>> ElementMap;
+
+	/** A map for looking up a parent based on their key */
+	TMap<FRigVMExplorerElementKey, FRigVMExplorerElementKey> ParentMap;
+	
+	FRigVMEditorGraphExplorerTreeDelegates Delegates;
+
+	FText FilterText;
+
+	friend class SRigVMEditorGraphExplorer;
+	friend class SRigVMEditorGraphExplorerItem;
+};
\ No newline at end of file
