Control Rig: Fix compilation error in mac
#rnx

diff --git a/Engine/Plugins/Runtime/RigVM/Source/RigVMEditor/Public/Widgets/SRigVMEditorGraphExplorerTreeView.h b/Engine/Plugins/Runtime/RigVM/Source/RigVMEditor/Public/Widgets/SRigVMEditorGraphExplorerTreeView.h
--- a/Engine/Plugins/Runtime/RigVM/Source/RigVMEditor/Public/Widgets/SRigVMEditorGraphExplorerTreeView.h
+++ b/Engine/Plugins/Runtime/RigVM/Source/RigVMEditor/Public/Widgets/SRigVMEditorGraphExplorerTreeView.h
@@ -26,6 +26,10 @@ enum ERigVMExplorerElementType
 
 struct FRigVMExplorerElementKey
 {
+
+	FRigVMExplorerElementKey(const ERigVMExplorerElementType& InType, const FString& InName)
+		: Type(InType), Name(InName){}
+	
 	ERigVMExplorerElementType Type;
 	FString Name;
 
