Control Rig: Add functions to graph explorer
#rnx

diff --git a/Engine/Plugins/Animation/ControlRig/Source/ControlRigEditor/Private/Editor/ControlRigEditor.cpp b/Engine/Plugins/Animation/ControlRig/Source/ControlRigEditor/Private/Editor/ControlRigEditor.cpp
--- a/Engine/Plugins/Animation/ControlRig/Source/ControlRigEditor/Private/Editor/ControlRigEditor.cpp
+++ b/Engine/Plugins/Animation/ControlRig/Source/ControlRigEditor/Private/Editor/ControlRigEditor.cpp
@@ -5049,7 +5049,10 @@ void FControlRigEditor::OnGraphNodeDropToPerform(TSharedPtr<FGraphNodeDragDropOp
 			);
 
 		}
+		return;
 	}
+
+	FRigVMEditor::OnGraphNodeDropToPerform(InDragDropOp, InGraph, InNodePosition, InScreenPosition);
 }
 
 void FControlRigEditor::HandleMakeElementGetterSetter(ERigElementGetterSetterType Type, bool bIsGetter, TArray<FRigElementKey> Keys, UEdGraph* Graph, FVector2D NodePosition)
diff --git a/Engine/Plugins/Animation/ControlRig/Source/ControlRigEditor/Private/Editor/ControlRigEditor.h b/Engine/Plugins/Animation/ControlRig/Source/ControlRigEditor/Private/Editor/ControlRigEditor.h
--- a/Engine/Plugins/Animation/ControlRig/Source/ControlRigEditor/Private/Editor/ControlRigEditor.h
+++ b/Engine/Plugins/Animation/ControlRig/Source/ControlRigEditor/Private/Editor/ControlRigEditor.h
@@ -170,7 +170,7 @@ public:
 
 	const FName RigHierarchyToGraphDragAndDropMenuName = TEXT("ControlRigEditor.RigHierarchyToGraphDragAndDropMenu");
 	void CreateRigHierarchyToGraphDragAndDropMenu() const;
-	void OnGraphNodeDropToPerform(TSharedPtr<FGraphNodeDragDropOp> InDragDropOp, UEdGraph* InGraph, const FVector2D& InNodePosition, const FVector2D& InScreenPosition);
+	virtual void OnGraphNodeDropToPerform(TSharedPtr<FGraphNodeDragDropOp> InDragDropOp, UEdGraph* InGraph, const FVector2D& InNodePosition, const FVector2D& InScreenPosition) override;
 
 	FPersonaViewportKeyDownDelegate& GetKeyDownDelegate() { return OnKeyDownDelegate; }
 
diff --git a/Engine/Plugins/Runtime/RigVM/Source/RigVMDeveloper/Private/EdGraph/RigVMEdGraphSchema.cpp b/Engine/Plugins/Runtime/RigVM/Source/RigVMDeveloper/Private/EdGraph/RigVMEdGraphSchema.cpp
--- a/Engine/Plugins/Runtime/RigVM/Source/RigVMDeveloper/Private/EdGraph/RigVMEdGraphSchema.cpp
+++ b/Engine/Plugins/Runtime/RigVM/Source/RigVMDeveloper/Private/EdGraph/RigVMEdGraphSchema.cpp
@@ -1523,7 +1523,7 @@ EGraphType URigVMEdGraphSchema::GetGraphType(const UEdGraph* TestEdGraph) const
 	return Super::GetGraphType(TestEdGraph);
 }
 
-FReply URigVMEdGraphSchema::TrySetGraphCategory(const UEdGraph* InGraph, const FText& InCategory)
+FReply URigVMEdGraphSchema::TrySetGraphCategory(const UEdGraph* InGraph, const FText& InCategory) const
 {
 	if (const URigVMEdGraph* RigGraph = Cast<URigVMEdGraph>((UEdGraph*)InGraph))
 	{
diff --git a/Engine/Plugins/Runtime/RigVM/Source/RigVMDeveloper/Public/EdGraph/RigVMEdGraphSchema.h b/Engine/Plugins/Runtime/RigVM/Source/RigVMDeveloper/Public/EdGraph/RigVMEdGraphSchema.h
--- a/Engine/Plugins/Runtime/RigVM/Source/RigVMDeveloper/Public/EdGraph/RigVMEdGraphSchema.h
+++ b/Engine/Plugins/Runtime/RigVM/Source/RigVMDeveloper/Public/EdGraph/RigVMEdGraphSchema.h
@@ -308,7 +308,7 @@ public:
 	virtual TSharedPtr<FEdGraphSchemaAction> MakeActionFromVariableDescription(const UEdGraph* InEdGraph, const FBPVariableDescription& Variable) const override;
 	virtual FText GetGraphCategory(const UEdGraph* InGraph) const override;
 	virtual EGraphType GetGraphType(const UEdGraph* TestEdGraph) const override;
-	virtual FReply TrySetGraphCategory(const UEdGraph* InGraph, const FText& InCategory) override;
+	virtual FReply TrySetGraphCategory(const UEdGraph* InGraph, const FText& InCategory) const override;
 	virtual bool TryDeleteGraph(UEdGraph* GraphToDelete) const override;
 	virtual bool TryRenameGraph(UEdGraph* GraphToRename, const FName& InNewName) const override;
 	virtual bool TryToGetChildEvents(const UEdGraph* Graph, const int32 SectionId, TArray<TSharedPtr<FEdGraphSchemaAction>>& Actions, const FText& ParentCategory) const override;
diff --git a/Engine/Plugins/Runtime/RigVM/Source/RigVMEditor/Private/Editor/RigVMEditor.cpp b/Engine/Plugins/Runtime/RigVM/Source/RigVMEditor/Private/Editor/RigVMEditor.cpp
--- a/Engine/Plugins/Runtime/RigVM/Source/RigVMEditor/Private/Editor/RigVMEditor.cpp
+++ b/Engine/Plugins/Runtime/RigVM/Source/RigVMEditor/Private/Editor/RigVMEditor.cpp
@@ -3628,6 +3628,45 @@ void FRigVMEditor::SwapAssetReferences()
 	SwapFunctionsDialog->ShowNormal();
 }
 
+void FRigVMEditor::OnGraphNodeDropToPerform(TSharedPtr<FGraphNodeDragDropOp> InDragDropOp, UEdGraph* InGraph, const FVector2D& InNodePosition,
+	const FVector2D& InScreenPosition)
+{
+	DECLARE_SCOPE_HIERARCHICAL_COUNTER_FUNC()
+
+	if (InDragDropOp->IsOfType<FRigVMGraphExplorerDragDropOp>())
+	{
+		TSharedPtr<FRigVMGraphExplorerDragDropOp> ExplorerOp = StaticCastSharedPtr<FRigVMGraphExplorerDragDropOp>(InDragDropOp);
+
+		if (ExplorerOp->GetElement().Type == Function && FocusedGraphEdPtr.IsValid())
+		{
+			if (URigVMBlueprint* TargetRigBlueprint = Cast<URigVMBlueprint>(FBlueprintEditorUtils::FindBlueprintForGraph(InGraph)))
+			{
+				if (URigVMBlueprint* SourceRigBlueprint = ExplorerOp->GetBlueprint())
+				{
+					if (URigVMLibraryNode* FunctionDefinitionNode = SourceRigBlueprint->GetRigVMClient()->GetFunctionLibrary()->FindFunction(*ExplorerOp->GetElement().Name))
+					{
+						if(URigVMController* TargetController = TargetRigBlueprint->GetController(InGraph))
+						{
+#if WITH_EDITOR
+							if(SourceRigBlueprint != TargetRigBlueprint)
+							{
+								if(!SourceRigBlueprint->IsFunctionPublic(FunctionDefinitionNode->GetFName()))
+								{
+									FRigVMGraphFunctionIdentifier FunctionIdentifier = FunctionDefinitionNode->GetFunctionIdentifier();
+									TargetRigBlueprint->BroadcastRequestLocalizeFunctionDialog(FunctionIdentifier);
+									FunctionDefinitionNode = TargetRigBlueprint->GetLocalFunctionLibrary()->FindPreviouslyLocalizedFunction(FunctionIdentifier);
+								}
+							}
+#endif
+							TargetController->AddFunctionReferenceNode(FunctionDefinitionNode, InNodePosition, FString(), true, true);
+						}
+					}
+				}
+			}
+		}
+	}
+}
+
 void FRigVMEditor::UpdateGraphCompilerErrors()
 {
 	DECLARE_SCOPE_HIERARCHICAL_COUNTER_FUNC()
diff --git a/Engine/Plugins/Runtime/RigVM/Source/RigVMEditor/Private/Widgets/SRigVMEditorGraphExplorer.cpp b/Engine/Plugins/Runtime/RigVM/Source/RigVMEditor/Private/Widgets/SRigVMEditorGraphExplorer.cpp
--- a/Engine/Plugins/Runtime/RigVM/Source/RigVMEditor/Private/Widgets/SRigVMEditorGraphExplorer.cpp
+++ b/Engine/Plugins/Runtime/RigVM/Source/RigVMEditor/Private/Widgets/SRigVMEditorGraphExplorer.cpp
@@ -14,13 +14,72 @@
 // Ick.
 #include "GraphActionNode.h"
 #include "GraphEditorActions.h"
+#include "SPinTypeSelector.h"
 #include "Kismet2/BlueprintEditorUtils.h"
 #include "RigVMModel/Nodes/RigVMAggregateNode.h"
-
+#include "Framework/Application/SlateApplication.h"
+#include "Widgets/Layout/SScaleBox.h"
+#include "HAL/PlatformApplicationMisc.h"
+#include "ScopedTransaction.h"
 
 #define LOCTEXT_NAMESPACE "RigVMGraphExplorer"
 
 
+TSharedRef<FRigVMGraphExplorerDragDropOp> FRigVMGraphExplorerDragDropOp::New(const ::FRigVMExplorerElementKey& InElement, TObjectPtr<URigVMBlueprint> InBlueprint)
+{
+	TSharedRef<FRigVMGraphExplorerDragDropOp> Operation = MakeShared<FRigVMGraphExplorerDragDropOp>();
+	Operation->Element = InElement;
+	Operation->SourceBlueprint = InBlueprint;
+	Operation->Construct();
+	return Operation;
+}
+
+void FRigVMGraphExplorerDragDropOp::Construct()
+{
+	// Create the drag-drop decorator window
+	CursorDecoratorWindow = SWindow::MakeCursorDecorator();
+	const bool bShowImmediately = false;
+	FSlateApplication::Get().AddWindow(CursorDecoratorWindow.ToSharedRef(), bShowImmediately);
+
+	const FSlateBrush* PrimarySymbol = FAppStyle::GetBrush(TEXT("Graph.ConnectorFeedback.NewNode"));
+	const FSlateBrush* SecondarySymbol = nullptr;
+	FSlateColor PrimaryColor = FLinearColor::White;
+	FSlateColor SecondaryColor = FLinearColor::White;
+
+	//Create feedback message with the function name.
+	TSharedRef<SWidget> TypeImage = SPinTypeSelector::ConstructPinTypeImage(PrimarySymbol, PrimaryColor, SecondarySymbol, SecondaryColor, TSharedPtr<SToolTip>());
+
+	CursorDecoratorWindow->ShowWindow();
+	CursorDecoratorWindow->SetContent
+	(
+		SNew(SBorder)
+		. BorderImage(FAppStyle::GetBrush("Graph.ConnectorFeedback.Border"))
+		[
+			SNew(SHorizontalBox)
+			+SHorizontalBox::Slot()
+			.AutoWidth()
+			.Padding(3.0f)
+			[
+				SNew(SScaleBox)
+				.Stretch(EStretch::ScaleToFit)
+				[
+					TypeImage
+				]
+			]
+			+SHorizontalBox::Slot()
+			.AutoWidth()
+			.Padding(3.0f)
+			.MaxWidth(500)
+			.VAlign(VAlign_Center)
+			[
+				SNew(STextBlock)
+				.WrapTextAt( 480 )
+				.Text( FText::FromString(GetElement().Name ))
+			]
+		]
+	);
+}
+
 FRigVMEditorGraphExplorerCommands::FRigVMEditorGraphExplorerCommands()
 	: TCommands<FRigVMEditorGraphExplorerCommands>(
           TEXT("RigVMEditorGraphExplorer"), NSLOCTEXT("Contexts", "Explorer", "Explorer"),
@@ -34,6 +93,9 @@ void FRigVMEditorGraphExplorerCommands::RegisterCommands()
 	UI_COMMAND(OpenGraph, "Open Graph", "Opens up this graph in the editor.", EUserInterfaceActionType::Button, FInputChord());
 	UI_COMMAND(OpenGraphInNewTab, "Open Graph In New Tab", "Opens up this graph in a new tab.", EUserInterfaceActionType::Button, FInputChord());
 	UI_COMMAND(CreateGraph, "Add New Graph", "Create a new graph and show it in the editor.", EUserInterfaceActionType::Button, FInputChord());
+	UI_COMMAND(CreateFunction, "Add New Function", "Create a new function and show it in the editor.", EUserInterfaceActionType::Button, FInputChord());
+	UI_COMMAND(AddFunctionVariant, "Add Variant (Experimental)", "Creates a new variant of a function.", EUserInterfaceActionType::Button, FInputChord());
+	UI_COMMAND(PasteFunction, "Paste Function", "Pastes the function.", EUserInterfaceActionType::Button, FInputChord());
 }
 
 
@@ -95,6 +157,16 @@ void SRigVMEditorGraphExplorer::RegisterCommands()
 			FCanExecuteAction::CreateSP(this, &SRigVMEditorGraphExplorer::CanCreateGraph)
 			);
 
+		ToolKitCommandList->MapAction(FRigVMEditorGraphExplorerCommands::Get().CreateFunction,
+			FExecuteAction::CreateSP(this, &SRigVMEditorGraphExplorer::OnCreateFunction),
+			FCanExecuteAction::CreateSP(this, &SRigVMEditorGraphExplorer::CanCreateFunction)
+			);
+
+		ToolKitCommandList->MapAction(FRigVMEditorGraphExplorerCommands::Get().AddFunctionVariant,
+			FExecuteAction::CreateSP(this, &SRigVMEditorGraphExplorer::OnAddFunctionVariant),
+			FCanExecuteAction::CreateSP(this, &SRigVMEditorGraphExplorer::CanAddFunctionVariant)
+			);
+
 		ToolKitCommandList->MapAction(FGenericCommands::Get().Rename,
 			FExecuteAction::CreateSP(this, &SRigVMEditorGraphExplorer::OnRenameEntry),
 			FCanExecuteAction::CreateSP(this, &SRigVMEditorGraphExplorer::CanRenameEntry));
@@ -112,6 +184,11 @@ void SRigVMEditorGraphExplorer::RegisterCommands()
 			FCanExecuteAction(), FIsActionChecked(),
 			FIsActionButtonVisible::CreateSP(this, &SRigVMEditorGraphExplorer::CanPasteGeneric));
 
+		ToolKitCommandList->MapAction(FRigVMEditorGraphExplorerCommands::Get().PasteFunction,
+			FExecuteAction::CreateSP(this, &SRigVMEditorGraphExplorer::OnPasteFunction),
+			FCanExecuteAction(), FIsActionChecked(),
+			FIsActionButtonVisible::CreateSP(this, &SRigVMEditorGraphExplorer::CanPasteFunction));
+
 		ToolKitCommandList->MapAction(FGenericCommands::Get().Delete,
 			FExecuteAction::CreateSP(this, &SRigVMEditorGraphExplorer::OnDeleteEntry),
 			FCanExecuteAction::CreateSP(this, &SRigVMEditorGraphExplorer::CanDeleteEntry));
@@ -151,23 +228,32 @@ void SRigVMEditorGraphExplorer::CreateWidgets()
 
 
 	FRigVMEditorGraphExplorerTreeDelegates Delegates;
-	Delegates.OnGetRootGraphs = FOnGetRootGraphs::CreateSP(this, &SRigVMEditorGraphExplorer::GetRootGraphs);
-	Delegates.OnGetChildrenGraphs = FOnGetChildrenGraphs::CreateSP(this, &SRigVMEditorGraphExplorer::GetChildrenGraphs);
-	Delegates.OnGetEventNodesInGraph = FOnGetEventNodesInGraph::CreateSP(this, &SRigVMEditorGraphExplorer::GetEventNodesInGraph);
-	Delegates.OnGetGraphDisplayName = FOnGetGraphDisplayName::CreateSP(this, &SRigVMEditorGraphExplorer::GetGraphDisplayName);
-	Delegates.OnGetEventDisplayName = FOnGetEventDisplayName::CreateSP(this, &SRigVMEditorGraphExplorer::GetEventDisplayName);
-	Delegates.OnGetGraphIcon = FOnGetGraphIcon::CreateSP(this, &SRigVMEditorGraphExplorer::GetGraphIcon);
-	Delegates.OnGetGraphTooltip = FOnGetGraphTooltip::CreateSP(this, &SRigVMEditorGraphExplorer::GetGraphTooltip);
-	Delegates.OnGraphClicked = FOnGraphClicked::CreateSP(this, &SRigVMEditorGraphExplorer::OnGraphClicked);
-	Delegates.OnEventClicked = FOnEventClicked::CreateSP(this, &SRigVMEditorGraphExplorer::OnEventClicked);
-	Delegates.OnGraphDoubleClicked = FOnGraphDoubleClicked::CreateSP(this, &SRigVMEditorGraphExplorer::OnGraphDoubleClicked);
-	Delegates.OnEventDoubleClicked = FOnEventDoubleClicked::CreateSP(this, &SRigVMEditorGraphExplorer::OnEventDoubleClicked);
-	Delegates.OnCreateGraph = FOnCreateGraph::CreateSP(this, &SRigVMEditorGraphExplorer::OnCreateGraph);
-	Delegates.OnRenameGraph = FOnRenameGraph::CreateSP(this, &SRigVMEditorGraphExplorer::OnRenameGraph);
-	Delegates.OnCanRenameGraph = FOnCanRenameGraph::CreateSP(this, &SRigVMEditorGraphExplorer::OnCanRenameGraph);
-	Delegates.OnRequestContextMenu = FOnRequestContextMenu::CreateSP(this, &SRigVMEditorGraphExplorer::OnContextMenuOpening);
+	Delegates.OnGetRootGraphs = FRigVMGraphExplorer_OnGetRootGraphs::CreateSP(this, &SRigVMEditorGraphExplorer::GetRootGraphs);
+	Delegates.OnGetChildrenGraphs = FRigVMGraphExplorer_OnGetChildrenGraphs::CreateSP(this, &SRigVMEditorGraphExplorer::GetChildrenGraphs);
+	Delegates.OnGetEventNodesInGraph = FRigVMGraphExplorer_OnGetEventNodesInGraph::CreateSP(this, &SRigVMEditorGraphExplorer::GetEventNodesInGraph);
+	Delegates.OnGetFunctions = FRigVMGraphExplorer_OnGetFunctions::CreateSP(this, &SRigVMEditorGraphExplorer::GetFunctions);
+	Delegates.OnGetGraphDisplayName = FRigVMGraphExplorer_OnGetGraphDisplayName::CreateSP(this, &SRigVMEditorGraphExplorer::GetGraphDisplayName);
+	Delegates.OnGetEventDisplayName = FRigVMGraphExplorer_OnGetEventDisplayName::CreateSP(this, &SRigVMEditorGraphExplorer::GetEventDisplayName);
+	Delegates.OnGetGraphIcon = FRigVMGraphExplorer_OnGetGraphIcon::CreateSP(this, &SRigVMEditorGraphExplorer::GetGraphIcon);
+	Delegates.OnGetGraphTooltip = FRigVMGraphExplorer_OnGetGraphTooltip::CreateSP(this, &SRigVMEditorGraphExplorer::GetGraphTooltip);
+	Delegates.OnGraphClicked = FRigVMGraphExplorer_OnGraphClicked::CreateSP(this, &SRigVMEditorGraphExplorer::OnGraphClicked);
+	Delegates.OnEventClicked = FRigVMGraphExplorer_OnEventClicked::CreateSP(this, &SRigVMEditorGraphExplorer::OnEventClicked);
+	Delegates.OnFunctionClicked = FRigVMGraphExplorer_OnFunctionClicked::CreateSP(this, &SRigVMEditorGraphExplorer::OnFunctionClicked);
+	Delegates.OnGraphDoubleClicked = FRigVMGraphExplorer_OnGraphDoubleClicked::CreateSP(this, &SRigVMEditorGraphExplorer::OnGraphDoubleClicked);
+	Delegates.OnEventDoubleClicked = FRigVMGraphExplorer_OnEventDoubleClicked::CreateSP(this, &SRigVMEditorGraphExplorer::OnEventDoubleClicked);
+	Delegates.OnFunctionDoubleClicked = FRigVMGraphExplorer_OnFunctionDoubleClicked::CreateSP(this, &SRigVMEditorGraphExplorer::OnFunctionDoubleClicked);
+	Delegates.OnCreateGraph = FRigVMGraphExplorer_OnCreateGraph::CreateSP(this, &SRigVMEditorGraphExplorer::OnCreateGraph);
+	Delegates.OnCreateFunction = FRigVMGraphExplorer_OnCreateFunction::CreateSP(this, &SRigVMEditorGraphExplorer::OnCreateFunction);
+	Delegates.OnRenameGraph = FRigVMGraphExplorer_OnRenameGraph::CreateSP(this, &SRigVMEditorGraphExplorer::OnRenameGraph);
+	Delegates.OnRenameFunction = FRigVMGraphExplorer_OnRenameFunction::CreateSP(this, &SRigVMEditorGraphExplorer::OnRenameFunction);
+	Delegates.OnCanRenameGraph = FRigVMGraphExplorer_OnCanRenameGraph::CreateSP(this, &SRigVMEditorGraphExplorer::OnCanRenameGraph);
+	Delegates.OnCanRenameFunction = FRigVMGraphExplorer_OnCanRenameFunction::CreateSP(this, &SRigVMEditorGraphExplorer::OnCanRenameFunction);
+	Delegates.OnSetFunctionCategory = FRigVMGraphExplorer_OnSetFunctionCategory::CreateSP(this, &SRigVMEditorGraphExplorer::OnSetFunctionCategory);
+	Delegates.OnGetFunctionCategory = FRigVMGraphExplorer_OnGetFunctionCategory::CreateSP(this, &SRigVMEditorGraphExplorer::OnGetFunctionCategory);
+	Delegates.OnRequestContextMenu = FRigVMGraphExplorer_OnRequestContextMenu::CreateSP(this, &SRigVMEditorGraphExplorer::OnContextMenuOpening);
+	Delegates.OnDragDetected = FOnDragDetected::CreateSP(this, &SRigVMEditorGraphExplorer::OnDragDetected);
 	
-	Delegates.OnSelectionChanged = FOnGraphExplorerSelectionChanged::CreateSP(this, &SRigVMEditorGraphExplorer::HandleSelectionChanged);
+	Delegates.OnSelectionChanged = FRigVMGraphExplorer_OnSelectionChanged::CreateSP(this, &SRigVMEditorGraphExplorer::HandleSelectionChanged);
 	
 	SAssignNew(TreeView, SRigVMEditorGraphExplorerTreeView)
 		.RigTreeDelegates(Delegates)
@@ -251,6 +337,7 @@ void SRigVMEditorGraphExplorer::BuildAddNewMenu(FMenuBuilder& MenuBuilder)
 {
 	MenuBuilder.BeginSection("AddNewItem", LOCTEXT("AddOperations", "Add New"));
 	MenuBuilder.AddMenuEntry(FRigVMEditorGraphExplorerCommands::Get().CreateGraph);
+	MenuBuilder.AddMenuEntry(FRigVMEditorGraphExplorerCommands::Get().PasteFunction);
 	
 	MenuBuilder.EndSection();
 }
@@ -344,6 +431,20 @@ TArray<URigVMNode*> SRigVMEditorGraphExplorer::GetEventNodesInGraph(const FStrin
 	return Events;
 }
 
+TArray<URigVMLibraryNode*> SRigVMEditorGraphExplorer::GetFunctions() const
+{
+	TArray<URigVMLibraryNode*> Functions;
+	
+	TSharedPtr<FRigVMEditor> Editor = RigVMEditor.Pin();
+	if (!Editor)
+	{
+		return Functions;
+	}
+
+	FRigVMClient* RigVMClient = Editor->GetRigVMBlueprint()->GetRigVMClient();
+	return RigVMClient->GetFunctionLibrary()->GetFunctions();
+}
+
 FText SRigVMEditorGraphExplorer::GetGraphDisplayName(const FString& InGraphPath) const
 {
 	FText DisplayName;
@@ -467,6 +568,32 @@ void SRigVMEditorGraphExplorer::OnEventClicked(const FString& InEventPath)
 	CurrentInspector->ShowDetailsForObjects(TArray<UObject*>());
 }
 
+void SRigVMEditorGraphExplorer::OnFunctionClicked(const FString& InFunctionPath)
+{
+	TSharedPtr<FRigVMEditor> Editor = RigVMEditor.Pin();
+	if (!Editor)
+	{
+		return;
+	}
+
+	FRigVMClient* RigVMClient = Editor->GetRigVMBlueprint()->GetRigVMClient();
+	URigVMFunctionLibrary* FunctionLibrary = RigVMClient->GetFunctionLibrary();
+	
+	if (const URigVMLibraryNode* FunctionNode = FunctionLibrary->FindFunction(*InFunctionPath))
+	{
+		const URigVMEdGraphSchema* EdSchema = GetDefault<URigVMEdGraphSchema>();
+		if (UEdGraph* EdGraph = Editor->GetRigVMBlueprint()->GetEdGraph(FunctionNode->GetContainedGraph()))
+		{
+			FGraphDisplayInfo DisplayInfo;
+			const UEdGraphSchema* Schema = EdGraph->GetSchema();
+			check(Schema != nullptr);
+			Schema->GetGraphDisplayInformation(*EdGraph, DisplayInfo);
+			TSharedRef<SKismetInspector> CurrentInspector = Editor->GetInspector();
+			CurrentInspector->ShowDetailsForSingleObject(EdGraph, SKismetInspector::FShowDetailsOptions(DisplayInfo.PlainName));
+		}
+	}
+}
+
 void SRigVMEditorGraphExplorer::OnGraphDoubleClicked(const FString& InGraphPath)
 {
 	TSharedPtr<FRigVMEditor> Editor = RigVMEditor.Pin();
@@ -503,6 +630,70 @@ void SRigVMEditorGraphExplorer::OnEventDoubleClicked(const FString& InEventPath)
 	}
 }
 
+void SRigVMEditorGraphExplorer::OnFunctionDoubleClicked(const FString& InFunctionPath)
+{
+	TSharedPtr<FRigVMEditor> Editor = RigVMEditor.Pin();
+	if (!Editor)
+	{
+		return;
+	}
+
+	FRigVMClient* RigVMClient = Editor->GetRigVMBlueprint()->GetRigVMClient();
+	URigVMFunctionLibrary* FunctionLibrary = RigVMClient->GetFunctionLibrary();
+	
+	if (const URigVMLibraryNode* FunctionNode = FunctionLibrary->FindFunction(*InFunctionPath))
+	{
+		const URigVMEdGraphSchema* EdSchema = GetDefault<URigVMEdGraphSchema>();
+		if (UEdGraph* EdGraph = Editor->GetRigVMBlueprint()->GetEdGraph(FunctionNode->GetContainedGraph()))
+		{
+			Editor->JumpToHyperlink(EdGraph);
+		}
+	}
+}
+
+bool SRigVMEditorGraphExplorer::OnSetFunctionCategory(const FString& InFunctionPath, const FString& InCategory)
+{
+	TSharedPtr<FRigVMEditor> Editor = RigVMEditor.Pin();
+	if (!Editor.IsValid())
+	{
+		return false;
+	}
+
+	const URigVMEdGraphSchema* Schema = GetDefault<URigVMEdGraphSchema>();
+	if (URigVMFunctionLibrary* Library = Editor->GetRigVMBlueprint()->GetLocalFunctionLibrary())
+	{
+		if (URigVMLibraryNode* Function = Library->FindFunction(*InFunctionPath))
+		{
+			if (UEdGraph* EdGraph = Editor->GetRigVMBlueprint()->GetEdGraph(Function->GetContainedGraph()))
+			{
+				Schema->TrySetGraphCategory(EdGraph, FText::FromString(InCategory));
+				return true;
+			}
+		}
+	}
+	
+	return false;
+}
+
+FString SRigVMEditorGraphExplorer::OnGetFunctionCategory(const FString& InFunctionPath) const
+{
+	TSharedPtr<FRigVMEditor> Editor = RigVMEditor.Pin();
+	if (!Editor.IsValid())
+	{
+		return FString();
+	}
+
+	const URigVMEdGraphSchema* Schema = GetDefault<URigVMEdGraphSchema>();
+	if (URigVMFunctionLibrary* Library = Editor->GetRigVMBlueprint()->GetLocalFunctionLibrary())
+	{
+		if (URigVMLibraryNode* Function = Library->FindFunction(*InFunctionPath))
+		{
+			return Function->GetNodeCategory();
+		}
+	}
+	return FString();
+}
+
 TSharedPtr<SWidget> SRigVMEditorGraphExplorer::OnContextMenuOpening()
 {
 	TSharedPtr<FRigVMEditor> Editor = RigVMEditor.Pin();
@@ -538,7 +729,6 @@ TSharedPtr<SWidget> SRigVMEditorGraphExplorer::OnContextMenuOpening()
 				MenuBuilder.AddMenuEntry(FGraphEditorCommands::Get().FindReferences);
 				MenuBuilder.AddMenuEntry(FGenericCommands::Get().Cut);
 				MenuBuilder.AddMenuEntry(FGenericCommands::Get().Copy);
-				MenuBuilder.AddMenuEntry(FGenericCommands::Get().Duplicate);
 				MenuBuilder.AddMenuEntry(FGenericCommands::Get().Delete);
 				MenuBuilder.EndSection();
 				break;
@@ -548,6 +738,20 @@ TSharedPtr<SWidget> SRigVMEditorGraphExplorer::OnContextMenuOpening()
 				return nullptr;
 				break;
 			}
+			case Function:
+			{
+				MenuBuilder.BeginSection("BasicOperations");
+				MenuBuilder.AddMenuEntry(FRigVMEditorGraphExplorerCommands::Get().OpenGraph);
+				MenuBuilder.AddMenuEntry(FRigVMEditorGraphExplorerCommands::Get().OpenGraphInNewTab);
+				MenuBuilder.AddMenuEntry(FGenericCommands::Get().Rename);
+				MenuBuilder.AddMenuEntry(FGraphEditorCommands::Get().FindReferences);
+				MenuBuilder.AddMenuEntry(FGenericCommands::Get().Cut);
+				MenuBuilder.AddMenuEntry(FGenericCommands::Get().Copy);
+				MenuBuilder.AddMenuEntry(FGenericCommands::Get().Delete);
+				MenuBuilder.AddMenuEntry(FRigVMEditorGraphExplorerCommands::Get().AddFunctionVariant);
+				MenuBuilder.EndSection();
+				break;
+			}
 		}
 		
 	}
@@ -559,31 +763,125 @@ TSharedPtr<SWidget> SRigVMEditorGraphExplorer::OnContextMenuOpening()
 	return MenuBuilder.MakeWidget();
 }
 
+FReply SRigVMEditorGraphExplorer::OnDragDetected(const FGeometry& MyGeometry, const FPointerEvent& MouseEvent)
+{
+	TArray<FRigVMExplorerElementKey> DraggedElements = TreeView->GetSelectedKeys();
+	if (MouseEvent.IsMouseButtonDown(EKeys::LeftMouseButton) && DraggedElements.Num() > 0)
+	{
+		if (RigVMEditor.IsValid() && DraggedElements.Num() == 1)
+		{
+			TSharedRef<FRigVMGraphExplorerDragDropOp> DragDropOp = FRigVMGraphExplorerDragDropOp::New(MoveTemp(DraggedElements[0]), RigVMEditor.Pin()->GetRigVMBlueprint());
+			DragDropOp->OnPerformDropToGraph.BindSP(RigVMEditor.Pin().Get(), &FRigVMEditor::OnGraphNodeDropToPerform);
+			return FReply::Handled().BeginDragDrop(DragDropOp);
+		}
+	}
+
+	return FReply::Unhandled();
+}
+
 void SRigVMEditorGraphExplorer::OnCopy()
 {
+	TSharedPtr<FRigVMEditor> Editor = RigVMEditor.Pin();
+	if (!Editor)
+	{
+		return;
+	}
+	TArray<FRigVMExplorerElementKey> Selection = TreeView->GetSelectedKeys();
+	if (Selection.Num() != 1 || Selection[0].Type != Function)
+	{
+		return;
+	}
+
+	FString OutputString;
+	const URigVMEdGraphSchema* Schema = GetDefault<URigVMEdGraphSchema>();
+	URigVMBlueprint* Blueprint = Editor->GetRigVMBlueprint();
+	URigVMFunctionLibrary* FunctionLibrary = Blueprint->GetLocalFunctionLibrary();
+	if (const URigVMLibraryNode* FunctionNode = FunctionLibrary->FindFunction(*Selection[0].Name))
+	{
+		if (UEdGraph* EdGraph = Editor->GetRigVMBlueprint()->GetEdGraph(FunctionNode->GetContainedGraph()))
+		{
+			Blueprint->ExportGraphToText(EdGraph, OutputString);
+		}
+	}
+
+	if (!OutputString.IsEmpty())
+	{
+		FPlatformApplicationMisc::ClipboardCopy(OutputString.GetCharArray().GetData());
+	}
 }
 
 bool SRigVMEditorGraphExplorer::CanCopy() const
 {
+	TArray<FRigVMExplorerElementKey> Selection = TreeView->GetSelectedKeys();
+	if (Selection.Num() == 1 && Selection[0].Type == Function)
+	{
+		return true;
+	}
 	return false;
 }
 
 void SRigVMEditorGraphExplorer::OnCut()
 {
+	OnCopy();
+	OnDeleteEntry();
 }
 
 bool SRigVMEditorGraphExplorer::CanCut() const
 {
-	return false;
+	return CanCopy() && CanDeleteEntry();
 }
 
 void SRigVMEditorGraphExplorer::OnPasteGeneric()
 {
+	// prioritize pasting as a member variable if possible
+	// if (CanPasteVariable())
+	// {
+	// 	OnPasteVariable();
+	// }
+	// else if (CanPasteLocalVariable())
+	// {
+	// 	OnPasteLocalVariable();
+	// }
+	// else 
+	if (CanPasteFunction())
+	{
+		OnPasteFunction();
+	}
 }
 
 bool SRigVMEditorGraphExplorer::CanPasteGeneric() const
 {
-	return false;
+	return //CanPasteVariable() 
+	//|| CanPasteLocalVariable() 
+	CanPasteFunction() ;
+}
+
+void SRigVMEditorGraphExplorer::OnPasteFunction()
+{
+	TSharedPtr<FRigVMEditor> Editor = RigVMEditor.Pin();
+	if (!Editor)
+	{
+		return;
+	}
+
+	FString ClipboardText;
+	FPlatformApplicationMisc::ClipboardPaste(ClipboardText);
+
+	Editor->GetRigVMBlueprint()->TryImportGraphFromText(ClipboardText);
+}
+
+bool SRigVMEditorGraphExplorer::CanPasteFunction() const
+{
+	TSharedPtr<FRigVMEditor> Editor = RigVMEditor.Pin();
+	if (!Editor)
+	{
+		return false;
+	}
+
+	FString ClipboardText;
+	FPlatformApplicationMisc::ClipboardPaste(ClipboardText);
+
+	return Editor->GetRigVMBlueprint()->CanImportGraphFromText(ClipboardText);
 }
 
 void SRigVMEditorGraphExplorer::OnToggleShowEmptySections()
@@ -627,7 +925,7 @@ void SRigVMEditorGraphExplorer::OnOpenGraph(bool bOnNewTab)
 }
 
 
-bool SRigVMEditorGraphExplorer::CanOpenGraph()
+bool SRigVMEditorGraphExplorer::CanOpenGraph() const
 {
 	TArray<FRigVMExplorerElementKey> Selection = TreeView->GetSelectedKeys();
 	return RigVMEditor.IsValid() && Selection.Num() == 1 && Selection[0].Type == ERigVMExplorerElementType::Graph;
@@ -644,7 +942,7 @@ void SRigVMEditorGraphExplorer::OnCreateGraph()
 }
 
 
-bool SRigVMEditorGraphExplorer::CanCreateGraph()
+bool SRigVMEditorGraphExplorer::CanCreateGraph() const
 {
 	TSharedPtr<FRigVMEditor> Editor = RigVMEditor.Pin();
 	if (Editor)
@@ -655,6 +953,63 @@ bool SRigVMEditorGraphExplorer::CanCreateGraph()
 	return false;
 }
 
+void SRigVMEditorGraphExplorer::OnCreateFunction()
+{
+	TSharedPtr<FRigVMEditor> Editor = RigVMEditor.Pin();
+	if (Editor)
+	{
+		Editor->NewDocument_OnClicked(FBlueprintEditor::CGT_NewFunctionGraph);
+	}
+}
+
+bool SRigVMEditorGraphExplorer::CanCreateFunction() const
+{
+	return true;
+}
+
+void SRigVMEditorGraphExplorer::OnAddFunctionVariant()
+{
+	TSharedPtr<FRigVMEditor> Editor = RigVMEditor.Pin();
+
+	if (!Editor)
+	{
+		return;
+	}
+	
+	TArray<FRigVMExplorerElementKey> Selected = TreeView->GetSelectedKeys();
+	if (Selected.Num() != 1 || Selected[0].Type != Function)
+	{
+		return;
+	}
+
+	URigVMFunctionLibrary* FunctionLibrary = Editor->GetRigVMBlueprint()->GetLocalFunctionLibrary();
+	if (const URigVMLibraryNode* FunctionNode = FunctionLibrary->FindFunction(*Selected[0].Name))
+	{
+		if (UEdGraph* EdGraph = Editor->GetRigVMBlueprint()->GetEdGraph(FunctionNode->GetContainedGraph()))
+		{
+			Editor->AddNewFunctionVariant(EdGraph);
+		}
+	}
+}
+
+bool SRigVMEditorGraphExplorer::CanAddFunctionVariant() const
+{
+	TSharedPtr<FRigVMEditor> Editor = RigVMEditor.Pin();
+
+	if (!Editor)
+	{
+		return false;
+	}
+	
+	TArray<FRigVMExplorerElementKey> Selected = TreeView->GetSelectedKeys();
+	if (Selected.Num() != 1 || Selected[0].Type != Function)
+	{
+		return false;
+	}
+
+	return true;
+}
+
 void SRigVMEditorGraphExplorer::OnDeleteEntry()
 {
 	TSharedPtr<FRigVMEditor> Editor = RigVMEditor.Pin();
@@ -675,6 +1030,20 @@ void SRigVMEditorGraphExplorer::OnDeleteEntry()
 					}
 					break;
 				}
+				case ERigVMExplorerElementType::Function:
+				{
+					if (URigVMFunctionLibrary* Library = Editor->GetRigVMBlueprint()->GetLocalFunctionLibrary())
+					{
+						if (URigVMLibraryNode* Function = Library->FindFunction(*Key.Name))
+						{
+							if (UEdGraph* EdGraph = Editor->GetRigVMBlueprint()->GetEdGraph(Function->GetContainedGraph()))
+							{
+								Schema->TryDeleteGraph(EdGraph);
+							}
+						}
+					}
+					break;
+				}
 			}
 		}
 	}
@@ -683,7 +1052,7 @@ void SRigVMEditorGraphExplorer::OnDeleteEntry()
 }
 
 
-bool SRigVMEditorGraphExplorer::CanDeleteEntry()
+bool SRigVMEditorGraphExplorer::CanDeleteEntry() const
 {
 	TSharedPtr<FRigVMEditor> Editor = RigVMEditor.Pin();
 
@@ -698,11 +1067,13 @@ bool SRigVMEditorGraphExplorer::CanDeleteEntry()
 		switch (Key.Type)
 		{
 			case ERigVMExplorerElementType::Section:
+			case ERigVMExplorerElementType::Event:
 				{
 					return false;
 					break;
 				}
 			case ERigVMExplorerElementType::Graph:
+			case ERigVMExplorerElementType::Function:
 				{
 					break;
 				}
@@ -733,7 +1104,7 @@ void SRigVMEditorGraphExplorer::OnRenameEntry()
 }
 
 
-bool SRigVMEditorGraphExplorer::CanRenameEntry()
+bool SRigVMEditorGraphExplorer::CanRenameEntry() const
 {
 	return true;
 }
@@ -756,7 +1127,7 @@ bool SRigVMEditorGraphExplorer::OnRenameGraph(const FString& InOldPath, const FS
 	return bResult;
 }
 
-bool SRigVMEditorGraphExplorer::OnCanRenameGraph(const FString& InOldPath, const FString& InNewPath, FText& OutErrorMessage)
+bool SRigVMEditorGraphExplorer::OnCanRenameGraph(const FString& InOldPath, const FString& InNewPath, FText& OutErrorMessage) const
 {
 	bool bResult = false;
 	TSharedPtr<FRigVMEditor> Editor = RigVMEditor.Pin();
@@ -787,6 +1158,50 @@ bool SRigVMEditorGraphExplorer::OnCanRenameGraph(const FString& InOldPath, const
 	return bResult;
 }
 
+bool SRigVMEditorGraphExplorer::OnRenameFunction(const FString& InOldPath, const FString& InNewPath)
+{
+	bool bResult = false;
+	TSharedPtr<FRigVMEditor> Editor = RigVMEditor.Pin();
+
+	if (Editor)
+	{
+		const URigVMEdGraphSchema* Schema = GetDefault<URigVMEdGraphSchema>();
+		URigVMFunctionLibrary* FunctionLibrary = Editor->GetRigVMBlueprint()->GetLocalFunctionLibrary();
+		if (const URigVMLibraryNode* FunctionNode = FunctionLibrary->FindFunction(*InOldPath))
+		{
+			if (UEdGraph* EdGraph = Editor->GetRigVMBlueprint()->GetEdGraph(FunctionNode->GetContainedGraph()))
+			{
+				bResult = Schema->TryRenameGraph(EdGraph, *InNewPath);
+			}
+		}
+	}
+
+	Refresh();
+	return bResult;
+}
+
+bool SRigVMEditorGraphExplorer::OnCanRenameFunction(const FString& InOldPath, const FString& InNewPath, FText& OutErrorMessage) const
+{
+	bool bResult = false;
+	TSharedPtr<FRigVMEditor> Editor = RigVMEditor.Pin();
+
+	if (Editor)
+	{
+		
+		if (Editor->GetRigVMBlueprint()->GetRigVMClient()->GetFunctionLibrary()->FindFunction(*InNewPath))
+		{
+			OutErrorMessage = FText::FromString(TEXT("Name already in use."));
+			bResult = false;
+		}
+		else
+		{
+			bResult = true;
+		}
+	}
+
+	return bResult;
+}
+
 void SRigVMEditorGraphExplorer::OnFindReference(bool bSearchAllBlueprints)
 {
 	FString SearchTerm;
diff --git a/Engine/Plugins/Runtime/RigVM/Source/RigVMEditor/Private/Widgets/SRigVMEditorGraphExplorerTreeView.cpp b/Engine/Plugins/Runtime/RigVM/Source/RigVMEditor/Private/Widgets/SRigVMEditorGraphExplorerTreeView.cpp
--- a/Engine/Plugins/Runtime/RigVM/Source/RigVMEditor/Private/Widgets/SRigVMEditorGraphExplorerTreeView.cpp
+++ b/Engine/Plugins/Runtime/RigVM/Source/RigVMEditor/Private/Widgets/SRigVMEditorGraphExplorerTreeView.cpp
@@ -9,7 +9,8 @@
 
 #define LOCTEXT_NAMESPACE "SRigVMEditorGraphExplorerTreeView"
 
-static const FString GraphsSectionName = TEXT("Graphs");
+static const FText GraphsSectionName = LOCTEXT("Graphs", "Graphs");
+static const FText FunctionsSectionName = LOCTEXT("Functions", "Functions");
 
 void SRigVMEditorGraphExplorerItem::Construct(const FArguments& InArgs,
                                               const TSharedRef<STableViewBase>& InOwnerTable,
@@ -26,7 +27,7 @@ void SRigVMEditorGraphExplorerItem::Construct(const FArguments& InArgs,
 	{
 		STableRow<TSharedPtr<FRigVMEditorGraphExplorerTreeElement>>::Construct(
 		STableRow<TSharedPtr<FRigVMEditorGraphExplorerTreeElement>>::FArguments()
-		.ShowWires(true)
+		.ShowWires(false)
 		.Padding(FMargin(0.0f, 2.0f, .0f, 0.0f))
 		.Content()
 		[
@@ -70,11 +71,39 @@ void SRigVMEditorGraphExplorerItem::Construct(const FArguments& InArgs,
 			
 		], InOwnerTable);
 	}
+	else if(InElement->Key.Type == ERigVMExplorerElementType::FunctionCategory
+		|| InElement->Key.Type == ERigVMExplorerElementType::VariableCategory)
+	{
+		int32 LastSeparator;
+		InElement->Key.Name.FindLastChar('|', LastSeparator);
+		FString CategoryName = InElement->Key.Name.RightChop(LastSeparator+1);
+		CategoryName[0] = FChar::ToUpper(CategoryName[0]);
+		STableRow<TSharedPtr<FRigVMEditorGraphExplorerTreeElement>>::Construct(
+		STableRow<TSharedPtr<FRigVMEditorGraphExplorerTreeElement>>::FArguments()
+		.ShowWires(false)
+		.Padding(FMargin(0.0f, 2.0f, .0f, 0.0f))
+		.Content()
+		[
+			SNew(SHorizontalBox)
+			+SHorizontalBox::Slot()
+			.VAlign(VAlign_Center)
+			[
+				SAssignNew(InlineRenameWidget, SInlineEditableTextBlock)
+				.Font( FCoreStyle::GetDefaultFontStyle("Bold", 9) )
+				.Text( FText::FromString(CategoryName ))
+				.ToolTipText( FText::FromString(CategoryName ))
+				.HighlightText( InTreeView->FilterText )
+				.OnVerifyTextChanged(this, &SRigVMEditorGraphExplorerItem::OnVerifyNameChanged)
+				.OnTextCommitted(this, &SRigVMEditorGraphExplorerItem::OnNameCommitted)
+			]
+		], InOwnerTable);
+	}
 	else
 	{
 		STableRow<TSharedPtr<FRigVMEditorGraphExplorerTreeElement>>::Construct(
 		   STableRow<TSharedPtr<FRigVMEditorGraphExplorerTreeElement>>::FArguments()
-		   .ShowWires(true)
+		   .ShowWires(false)
+		   .OnDragDetected(Delegates.OnDragDetected)
 		   .Content()
 		   [
 			   SNew(SHorizontalBox)
@@ -104,16 +133,28 @@ TSharedRef<SWidget> SRigVMEditorGraphExplorerItem::CreateIconWidget(const FRigVM
 {
 	TSharedPtr<SWidget> IconWidget;
 
-	if (Key.Type == ERigVMExplorerElementType::Graph)
+	switch (Key.Type)
 	{
-		IconWidget = SNew(SImage)
-			.Image(Delegates.GetGraphIcon(Key.Name));
-	}
-	else if(Key.Type == ERigVMExplorerElementType::Event)
-	{
-		static FSlateIcon EventIcon(FAppStyle::GetAppStyleSetName(), "GraphEditor.Event_16x");
-		IconWidget = SNew(SImage)
-			.Image(EventIcon.GetIcon());
+		case Graph:
+			{
+				IconWidget = SNew(SImage)
+							.Image(Delegates.GetGraphIcon(Key.Name));
+				break;
+			}
+		case Event:
+			{
+				static FSlateIcon EventIcon(FAppStyle::GetAppStyleSetName(), "GraphEditor.Event_16x");
+				IconWidget = SNew(SImage)
+							.Image(EventIcon.GetIcon());
+				break;
+			}
+		case Function:
+			{
+				static FSlateIcon FunctionIcon(FAppStyle::GetAppStyleSetName(), "Kismet.AllClasses.FunctionIcon");
+				IconWidget = SNew(SImage)
+							.Image(FunctionIcon.GetIcon());
+				break;
+			}
 	}
 
 	if (IconWidget.IsValid())
@@ -159,6 +200,10 @@ FText SRigVMEditorGraphExplorerItem::GetDisplayText() const
 			{
 				return Delegates.GetEventDisplayName(Element->Key.Name);
 			}
+		case ERigVMExplorerElementType::Function:
+			{
+				return FText::FromString(Element->Key.Name);
+			}
 	}
 	
 	return DisplayText;
@@ -199,22 +244,39 @@ bool SRigVMEditorGraphExplorerItem::OnVerifyNameChanged(const FText& InText, FTe
 			{
 				break;
 			}
+		case FunctionCategory:
+		case VariableCategory:
+			{
+				return true;
+				break;
+			}
 		case Graph:
 			{
 				return Delegates.CanRenameGraph(OldKey.Name, *NewName, OutErrorMessage);
 				break;
 			}
+		case Function:
+			{
+				return Delegates.CanRenameFunction(OldKey.Name, *NewName, OutErrorMessage);
+				break;
+			}
 	}
 	return false;
 }
 
 void SRigVMEditorGraphExplorerItem::OnNameCommitted(const FText& InText, ETextCommit::Type InCommitType) const
 {
+	if (!WeakExplorerElement.IsValid())
+	{
+		return;
+	}
+	
 	// for now only allow enter
 	if (InCommitType == ETextCommit::OnEnter)
 	{
 		FString NewName = InText.ToString();
-		const FRigVMExplorerElementKey OldKey = WeakExplorerElement.Pin()->Key;
+		TSharedPtr<FRigVMEditorGraphExplorerTreeElement> Element = WeakExplorerElement.Pin();
+		const FRigVMExplorerElementKey OldKey = Element->Key;
 
 		switch (OldKey.Type)
 		{
@@ -222,11 +284,48 @@ void SRigVMEditorGraphExplorerItem::OnNameCommitted(const FText& InText, ETextCo
 				{
 					break;
 				}
+			case FunctionCategory:
+				{
+					int32 LastSeparator;
+					FString OldCategoryPath = Element->Key.Name;
+					Element->Key.Name.FindLastChar('|', LastSeparator);
+					FString Prefix;
+					if (LastSeparator != INDEX_NONE)
+					{
+						Prefix = Element->Key.Name.LeftChop(Element->Key.Name.Len() - LastSeparator);
+					}
+					if (!Prefix.IsEmpty() && !InText.IsEmpty())
+					{
+						Prefix.Append(TEXT("|"));
+					}
+					FString NewCategoryPath = FString::Printf(TEXT("%s%s"), *Prefix, *InText.ToString());
+					TArray<TSharedPtr<FRigVMEditorGraphExplorerTreeElement>> SubElements = Element->Children;
+					for (int32 i=0; i<SubElements.Num(); ++i)
+					{
+						TSharedPtr<FRigVMEditorGraphExplorerTreeElement> SubElement = SubElements[i];
+						if (SubElement->Key.Type == ERigVMExplorerElementType::Function)
+						{
+							FString Category = Delegates.GetFunctionCategory(SubElement->Key.Name);
+							Category.RemoveFromStart(OldCategoryPath);
+							Category.InsertAt(0, NewCategoryPath);
+							Category.RemoveFromStart(TEXT("|"));
+							Category.RemoveFromEnd(TEXT("|"));
+							Delegates.SetFunctionCategory(SubElement->Key.Name, Category);
+						}
+						SubElements.Append(SubElement->Children);
+					}
+					break;
+				}
 			case Graph:
 				{
 					Delegates.RenameGraph(OldKey.Name, *NewName);
 					break;
 				}
+			case Function:
+				{
+					Delegates.RenameFunction(OldKey.Name, *NewName);
+					break;
+				}
 		}
 	}
 }
@@ -257,6 +356,7 @@ void SRigVMEditorGraphExplorerTreeView::Construct(const FArguments& InArgs)
 	Delegates = InArgs._RigTreeDelegates;
 	
 	STreeView<TSharedPtr<FRigVMEditorGraphExplorerTreeElement>>::FArguments SuperArgs;
+	SuperArgs.SelectionMode(ESelectionMode::Single);
 	SuperArgs.TreeItemsSource(&RootElements);
 	SuperArgs.SelectionMode(ESelectionMode::Multi);
 	SuperArgs.OnGenerateRow(this, &SRigVMEditorGraphExplorerTreeView::MakeTableRowWidget);
@@ -266,8 +366,9 @@ void SRigVMEditorGraphExplorerTreeView::Construct(const FArguments& InArgs)
 	SuperArgs.OnGeneratePinnedRow(this, &SRigVMEditorGraphExplorerTreeView::MakeTableRowWidget);
 	SuperArgs.OnMouseButtonClick(this, &SRigVMEditorGraphExplorerTreeView::OnItemClicked);
 	SuperArgs.OnMouseButtonDoubleClick(this, &SRigVMEditorGraphExplorerTreeView::OnItemDoubleClicked);
-	SuperArgs.OnSelectionChanged(FOnGraphExplorerSelectionChanged::CreateRaw(&Delegates, &FRigVMEditorGraphExplorerTreeDelegates::HandleSelectionChanged));
-	SuperArgs.OnContextMenuOpening(FOnRequestContextMenu::CreateRaw(&Delegates, &FRigVMEditorGraphExplorerTreeDelegates::RequestContextMenu));
+	SuperArgs.OnSelectionChanged(FRigVMGraphExplorer_OnSelectionChanged::CreateRaw(&Delegates, &FRigVMEditorGraphExplorerTreeDelegates::HandleSelectionChanged));
+	SuperArgs.OnContextMenuOpening(FRigVMGraphExplorer_OnRequestContextMenu::CreateRaw(&Delegates, &FRigVMEditorGraphExplorerTreeDelegates::RequestContextMenu));
+	
 
 	STreeView<TSharedPtr<FRigVMEditorGraphExplorerTreeElement>>::Construct(SuperArgs);
 }
@@ -298,53 +399,101 @@ void SRigVMEditorGraphExplorerTreeView::RefreshTreeView(bool bRebuildContent)
 		// Add Elements
 		{
 			// GRAPHS section
-			FRigVMExplorerElementKey SectionKey(ERigVMExplorerElementType::Section, GraphsSectionName);
-			TSharedPtr<FRigVMEditorGraphExplorerTreeElement> SectionItem = MakeShared<FRigVMEditorGraphExplorerTreeElement>(SectionKey, SharedThis(this));
-			ElementMap.Add(SectionKey, SectionItem);
-			RootElements.Add(SectionItem);
-			SetItemExpansion(SectionItem, true);
-
-			TArray<const URigVMGraph*> Graphs = Delegates.GetRootGraphs();
-			for (const URigVMGraph* Graph : Graphs)
 			{
-				FRigVMExplorerElementKey Key(ERigVMExplorerElementType::Graph, Graph->GetNodePath());
-				ParentMap.Add(Key, SectionKey);
-			}
+				FRigVMExplorerElementKey SectionKey(ERigVMExplorerElementType::Section, GraphsSectionName.ToString());
+				TSharedPtr<FRigVMEditorGraphExplorerTreeElement> SectionItem = MakeShared<FRigVMEditorGraphExplorerTreeElement>(SectionKey, SharedThis(this));
+				ElementMap.Add(SectionKey, SectionItem);
+				RootElements.Add(SectionItem);
+				SetItemExpansion(SectionItem, true);
+
+				TArray<const URigVMGraph*> Graphs = Delegates.GetRootGraphs();
+				for (const URigVMGraph* Graph : Graphs)
+				{
+					FRigVMExplorerElementKey Key(ERigVMExplorerElementType::Graph, Graph->GetNodePath());
+					ParentMap.Add(Key, SectionKey);
+				}
 			
-			// Add Graphs to section
-			for (int32 i=0; i<Graphs.Num(); ++i)
-			{
-				const URigVMGraph* Graph = Graphs[i];
+				// Add Graphs to section
+				for (int32 i=0; i<Graphs.Num(); ++i)
+				{
+					const URigVMGraph* Graph = Graphs[i];
 				
-				FRigVMExplorerElementKey Key(ERigVMExplorerElementType::Graph, Graph->GetNodePath());
+					FRigVMExplorerElementKey Key(ERigVMExplorerElementType::Graph, Graph->GetNodePath());
 				
-				TSharedPtr<FRigVMEditorGraphExplorerTreeElement> NewItem = MakeShared<FRigVMEditorGraphExplorerTreeElement>(Key, SharedThis(this));
-				ElementMap.Add(Key, NewItem);
-				SetItemExpansion(NewItem, true);
+					TSharedPtr<FRigVMEditorGraphExplorerTreeElement> NewItem = MakeShared<FRigVMEditorGraphExplorerTreeElement>(Key, SharedThis(this));
+					ElementMap.Add(Key, NewItem);
+					SetItemExpansion(NewItem, true);
 
-				TArray<const URigVMGraph*> ChildrenGraphs = Delegates.GetChildrenGraphs(Key.Name);
-				for (const URigVMGraph* Child : ChildrenGraphs)
-				{
-					FRigVMExplorerElementKey ChildKey(ERigVMExplorerElementType::Graph, Child->GetNodePath());
-					ParentMap.Add(ChildKey, Key);
-					Graphs.Add(Child);
+					TArray<const URigVMGraph*> ChildrenGraphs = Delegates.GetChildrenGraphs(Key.Name);
+					for (const URigVMGraph* Child : ChildrenGraphs)
+					{
+						FRigVMExplorerElementKey ChildKey(ERigVMExplorerElementType::Graph, Child->GetNodePath());
+						ParentMap.Add(ChildKey, Key);
+						Graphs.Add(Child);
+					}
+
+					TArray<URigVMNode*> EventNodes = Delegates.GetEventNodesInGraph(Key.Name);
+					for (const URigVMNode* EventNode : EventNodes)
+					{
+					
+						FRigVMExplorerElementKey EventKey(ERigVMExplorerElementType::Event, EventNode->GetNodePath());
+						ParentMap.Add(EventKey, Key);
+						TSharedPtr<FRigVMEditorGraphExplorerTreeElement> EventItem = MakeShared<FRigVMEditorGraphExplorerTreeElement>(EventKey, SharedThis(this));
+						ElementMap.Add(EventKey, EventItem);
+						NewItem->Children.Add(EventItem);
+					}
+
+					FRigVMExplorerElementKey ParentKey = ParentMap.FindChecked(Key);
+					TSharedPtr<FRigVMEditorGraphExplorerTreeElement> ParentItem = ElementMap.FindChecked(ParentKey);
+
+					ParentItem->Children.Add(NewItem);
 				}
+			}
 
-				TArray<URigVMNode*> EventNodes = Delegates.GetEventNodesInGraph(Key.Name);
-				for (const URigVMNode* EventNode : EventNodes)
+			// FUNCTIONS section
+			{
+				FRigVMExplorerElementKey SectionKey(ERigVMExplorerElementType::Section, FunctionsSectionName.ToString());
+				TSharedPtr<FRigVMEditorGraphExplorerTreeElement> SectionItem = MakeShared<FRigVMEditorGraphExplorerTreeElement>(SectionKey, SharedThis(this));
+				ElementMap.Add(SectionKey, SectionItem);
+				RootElements.Add(SectionItem);
+				SetItemExpansion(SectionItem, true);
+
+				TArray<URigVMLibraryNode*> Functions = Delegates.GetFunctions();
+				for (const URigVMLibraryNode* Function : Functions)
 				{
+					const FString Category = Function->GetNodeCategory();
+					TArray<FString> SingleCategories;
+					Category.ParseIntoArray(SingleCategories, TEXT("|"));
+					FString CategoryParentPath;
+					TSharedPtr<FRigVMEditorGraphExplorerTreeElement> PreviousParentItem = SectionItem;
+					for (const FString& SingleCategory : SingleCategories)
+					{
+						if (!CategoryParentPath.IsEmpty())
+						{
+							CategoryParentPath.Append(TEXT("|"));
+						}
+						CategoryParentPath.Append(SingleCategory);
+						FRigVMExplorerElementKey ParentKey(ERigVMExplorerElementType::FunctionCategory, CategoryParentPath);
+						TSharedPtr<FRigVMEditorGraphExplorerTreeElement> ParentItem = FindElement(ParentKey);
+						if (!ParentItem)
+						{
+							ParentItem = MakeShared<FRigVMEditorGraphExplorerTreeElement>(ParentKey, SharedThis(this));
+							ElementMap.Add(ParentKey, ParentItem);
+							ParentMap.Add(ParentKey, PreviousParentItem->Key);
+							PreviousParentItem->Children.Add(ParentItem);
+							SetItemExpansion(ParentItem, true);
+						}
+						PreviousParentItem = ParentItem;
+					}
 					
-					FRigVMExplorerElementKey EventKey(ERigVMExplorerElementType::Event, EventNode->GetNodePath());
-					ParentMap.Add(EventKey, Key);
-					TSharedPtr<FRigVMEditorGraphExplorerTreeElement> EventItem = MakeShared<FRigVMEditorGraphExplorerTreeElement>(EventKey, SharedThis(this));
-					ElementMap.Add(EventKey, EventItem);
-					NewItem->Children.Add(EventItem);
+					FRigVMExplorerElementKey Key(ERigVMExplorerElementType::Function, Function->GetNodePath());
+				
+					TSharedPtr<FRigVMEditorGraphExplorerTreeElement> NewItem = MakeShared<FRigVMEditorGraphExplorerTreeElement>(Key, SharedThis(this));
+					ElementMap.Add(Key, NewItem);
+					ParentMap.Add(Key, PreviousParentItem->Key);
+					PreviousParentItem->Children.Add(NewItem);
+					SetItemExpansion(NewItem, true);
 				}
-
-				FRigVMExplorerElementKey ParentKey = ParentMap.FindChecked(Key);
-				TSharedPtr<FRigVMEditorGraphExplorerTreeElement> ParentItem = ElementMap.FindChecked(ParentKey);
-
-				ParentItem->Children.Add(NewItem);
 			}
 		}
 
@@ -536,7 +685,7 @@ void SRigVMEditorGraphExplorerTreeView::OnItemClicked(TSharedPtr<FRigVMEditorGra
 	{
 		case Section:
 			{
-				FRigVMExplorerElementKey SectionKey(ERigVMExplorerElementType::Section, GraphsSectionName);
+				FRigVMExplorerElementKey SectionKey(ERigVMExplorerElementType::Section, InElement->Key.Name);
 				TSharedPtr<FRigVMEditorGraphExplorerTreeElement> Section = ElementMap.FindChecked(SectionKey);
 				SetItemExpansion(Section, !IsItemExpanded(Section));
 				break;
@@ -550,6 +699,10 @@ void SRigVMEditorGraphExplorerTreeView::OnItemClicked(TSharedPtr<FRigVMEditorGra
 			{
 				Delegates.EventClicked(InElement->Key.Name);
 			}
+		case Function:
+			{
+				Delegates.FunctionClicked(InElement->Key.Name);
+			}
 	}
 }
 
@@ -570,16 +723,25 @@ void SRigVMEditorGraphExplorerTreeView::OnItemDoubleClicked(TSharedPtr<FRigVMEdi
 			{
 				Delegates.EventDoubleClicked(InElement->Key.Name);
 			}
+		case Function:
+			{
+				Delegates.FunctionDoubleClicked(InElement->Key.Name);
+			}
 	}
 }
 
 FReply SRigVMEditorGraphExplorerTreeView::OnAddButtonClickedOnSection(const FRigVMExplorerElementKey& InSectionKey)
 {
-	if (InSectionKey.Name == GraphsSectionName)
+	if (InSectionKey.Name == GraphsSectionName.ToString())
 	{
 		Delegates.CreateGraph();
 		return FReply::Handled();
 	}
+	else if (InSectionKey.Name == FunctionsSectionName.ToString())
+	{
+		Delegates.CreateFunction();
+		return FReply::Handled();
+	}
 	return FReply::Unhandled();
 }
 
diff --git a/Engine/Plugins/Runtime/RigVM/Source/RigVMEditor/Public/Editor/RigVMEditor.h b/Engine/Plugins/Runtime/RigVM/Source/RigVMEditor/Public/Editor/RigVMEditor.h
--- a/Engine/Plugins/Runtime/RigVM/Source/RigVMEditor/Public/Editor/RigVMEditor.h
+++ b/Engine/Plugins/Runtime/RigVM/Source/RigVMEditor/Public/Editor/RigVMEditor.h
@@ -9,6 +9,7 @@
 #include "Editor/RigVMDetailsViewWrapperObject.h"
 #include "AssetRegistry/AssetData.h"
 
+class FGraphNodeDragDropOp;
 class SRigVMEditorGraphExplorer;
 class FRigVMEditor;
 
@@ -316,6 +317,8 @@ protected:
 	virtual void SwapFunctionForAssets(const TArray<FAssetData>& InAssets, bool bSetupUndo);
 	virtual void SwapAssetReferences();
 
+	virtual void OnGraphNodeDropToPerform(TSharedPtr<FGraphNodeDragDropOp> InDragDropOp, UEdGraph* InGraph, const FVector2D& InNodePosition, const FVector2D& InScreenPosition);
+
 	/** Once the log is collected update the graph */
 	void UpdateGraphCompilerErrors();
 
diff --git a/Engine/Plugins/Runtime/RigVM/Source/RigVMEditor/Public/Widgets/SRigVMEditorGraphExplorer.h b/Engine/Plugins/Runtime/RigVM/Source/RigVMEditor/Public/Widgets/SRigVMEditorGraphExplorer.h
--- a/Engine/Plugins/Runtime/RigVM/Source/RigVMEditor/Public/Widgets/SRigVMEditorGraphExplorer.h
+++ b/Engine/Plugins/Runtime/RigVM/Source/RigVMEditor/Public/Widgets/SRigVMEditorGraphExplorer.h
@@ -3,10 +3,12 @@
 #pragma once
 
 #include "CoreMinimal.h"
+#include "DragAndDrop/GraphNodeDragDropOp.h"
 #include "Widgets/DeclarativeSyntaxSupport.h"
 #include "Widgets/SCompoundWidget.h"
 #include "Framework/Commands/Commands.h"
 #include "RigVMModel/RigVMGraph.h"
+#include "SRigVMEditorGraphExplorerTreeView.h"
 
 class SRigVMEditorGraphExplorerTreeView;
 //
@@ -21,6 +23,31 @@ class FMenuBuilder;
 class SSearchBox;
 //
 
+class FRigVMGraphExplorerDragDropOp : public FGraphNodeDragDropOp
+{
+public:
+	DRAG_DROP_OPERATOR_TYPE(FRigVMExplorerDragDropOp, FGraphNodeDragDropOp)
+
+	static TSharedRef<FRigVMGraphExplorerDragDropOp> New(const FRigVMExplorerElementKey& InElement, TObjectPtr<URigVMBlueprint> InBlueprint);
+
+	const FRigVMExplorerElementKey& GetElement() const
+	{
+		return Element;
+	}
+
+	TObjectPtr<URigVMBlueprint> GetBlueprint() const
+	{
+		return SourceBlueprint;
+	}
+
+	virtual void Construct() override;
+
+private:
+
+	FRigVMExplorerElementKey Element;
+	TObjectPtr<URigVMBlueprint> SourceBlueprint;
+};
+
 class FRigVMEditorGraphExplorerCommands : public TCommands<FRigVMEditorGraphExplorerCommands>
 {
 public:
@@ -29,7 +56,9 @@ public:
 	TSharedPtr<FUICommandInfo> OpenGraph;
 	TSharedPtr<FUICommandInfo> OpenGraphInNewTab;
 	TSharedPtr<FUICommandInfo> CreateGraph;
-	TSharedPtr<FUICommandInfo> DeleteGraph;
+	TSharedPtr<FUICommandInfo> CreateFunction;
+	TSharedPtr<FUICommandInfo> AddFunctionVariant;
+	TSharedPtr<FUICommandInfo> PasteFunction;
 
 	void RegisterCommands() override;
 };
@@ -62,16 +91,22 @@ private:
 	TArray<const URigVMGraph*> GetRootGraphs() const;
 	TArray<const URigVMGraph*> GetChildrenGraphs(const FString& InParentGraphPath) const;
 	TArray<URigVMNode*> GetEventNodesInGraph(const FString& InParentGraphPath) const;
+	TArray<URigVMLibraryNode*> GetFunctions() const;
 	FText GetGraphDisplayName(const FString& InGraphPath) const;
 	FText GetEventDisplayName(const FString& InNodePath) const;
 	const FSlateBrush* GetGraphIcon(const FString& InGraphPath) const;
 	FText GetGraphTooltip(const FString& InGraphPath) const;
 	void OnGraphClicked(const FString& InGraphPath);
 	void OnEventClicked(const FString& InEventPath);
+	void OnFunctionClicked(const FString& InGraphPath);
 	void OnGraphDoubleClicked(const FString& InGraphPath);
 	void OnEventDoubleClicked(const FString& InEventPath);
+	void OnFunctionDoubleClicked(const FString& InFunctionPath);
+	bool OnSetFunctionCategory(const FString& InFunctionPath, const FString& InCategory);
+	FString OnGetFunctionCategory(const FString& InPath) const;
 	
 	TSharedPtr<SWidget> OnContextMenuOpening();
+	FReply OnDragDetected(const FGeometry& MyGeometry, const FPointerEvent& MouseEvent);
 	
 	void OnCopy();
 	bool CanCopy() const;
@@ -79,6 +114,8 @@ private:
 	bool CanCut() const;
 	void OnPasteGeneric();
 	bool CanPasteGeneric() const;
+	void OnPasteFunction();
+	bool CanPasteFunction() const;
 
 
 
@@ -88,18 +125,26 @@ private:
 
 	// Command functions
 	void OnOpenGraph(bool bOnNewTab = false);
-	bool CanOpenGraph();
+	bool CanOpenGraph() const;
 
 	void OnCreateGraph();
-	bool CanCreateGraph();
+	bool CanCreateGraph() const;
+
+	void OnCreateFunction();
+	bool CanCreateFunction() const;
+
+	void OnAddFunctionVariant();
+	bool CanAddFunctionVariant() const;
 	
 	void OnDeleteEntry();
-	bool CanDeleteEntry();
+	bool CanDeleteEntry() const;
 
 	void OnRenameEntry();
-	bool CanRenameEntry();
+	bool CanRenameEntry() const;
 	bool OnRenameGraph(const FString& InOldPath, const FString& InNewPath);
-	bool OnCanRenameGraph(const FString& InOldPath, const FString& InNewPath, FText& OutErrorMessage);
+	bool OnCanRenameGraph(const FString& InOldPath, const FString& InNewPath, FText& OutErrorMessage) const;
+	bool OnRenameFunction(const FString& InOldPath, const FString& InNewPath);
+	bool OnCanRenameFunction(const FString& InOldPath, const FString& InNewPath, FText& OutErrorMessage) const;
 
 	void OnFindReference(bool bSearchAllBlueprints);
 	bool CanFindReference() const;
diff --git a/Engine/Plugins/Runtime/RigVM/Source/RigVMEditor/Public/Widgets/SRigVMEditorGraphExplorerTreeView.h b/Engine/Plugins/Runtime/RigVM/Source/RigVMEditor/Public/Widgets/SRigVMEditorGraphExplorerTreeView.h
--- a/Engine/Plugins/Runtime/RigVM/Source/RigVMEditor/Public/Widgets/SRigVMEditorGraphExplorerTreeView.h
+++ b/Engine/Plugins/Runtime/RigVM/Source/RigVMEditor/Public/Widgets/SRigVMEditorGraphExplorerTreeView.h
@@ -16,7 +16,10 @@ class URigVMGraph;
 
 enum ERigVMExplorerElementType
 {
+	Invalid,
 	Section,
+	FunctionCategory,
+	VariableCategory,
 	Graph,
 	Event,
 	Function,
@@ -26,6 +29,8 @@ enum ERigVMExplorerElementType
 
 struct FRigVMExplorerElementKey
 {
+	FRigVMExplorerElementKey()
+		: Type(ERigVMExplorerElementType::Invalid), Name(FString()){}
 
 	FRigVMExplorerElementKey(const ERigVMExplorerElementType& InType, const FString& InName)
 		: Type(InType), Name(InName){}
@@ -44,45 +49,62 @@ struct FRigVMExplorerElementKey
 	}
 };
 
-DECLARE_DELEGATE_RetVal(TArray<const URigVMGraph*>, FOnGetRootGraphs);
-DECLARE_DELEGATE( FOnCreateGraph);
-DECLARE_DELEGATE_RetVal_OneParam(TArray<const URigVMGraph*>, FOnGetChildrenGraphs, const FString&);
-DECLARE_DELEGATE_RetVal_OneParam(TArray<URigVMNode*>, FOnGetEventNodesInGraph, const FString&);
-DECLARE_DELEGATE_RetVal_OneParam(FText, FOnGetGraphDisplayName, const FString&);
-DECLARE_DELEGATE_RetVal_OneParam(FText, FOnGetEventDisplayName, const FString&);
-DECLARE_DELEGATE_RetVal_OneParam(FText, FOnGetGraphTooltip, const FString&);
-DECLARE_DELEGATE_RetVal_OneParam(const FSlateBrush*, FOnGetGraphIcon, const FString&);
-DECLARE_DELEGATE_OneParam(FOnGraphClicked, const FString&);
-DECLARE_DELEGATE_OneParam(FOnEventClicked, const FString&);
-DECLARE_DELEGATE_OneParam(FOnGraphDoubleClicked, const FString&);
-DECLARE_DELEGATE_OneParam(FOnEventDoubleClicked, const FString&);
-DECLARE_DELEGATE_RetVal_TwoParams(bool, FOnRenameGraph, const FString&, const FString&);
-DECLARE_DELEGATE_RetVal_ThreeParams(bool, FOnCanRenameGraph, const FString&, const FString&, FText&);
-DECLARE_DELEGATE_RetVal(TSharedPtr<SWidget>, FOnRequestContextMenu);
-DECLARE_DELEGATE_RetVal(bool, FOnHasContextMenu);
-
-typedef STreeView<TSharedPtr<FRigVMEditorGraphExplorerTreeElement>>::FOnSelectionChanged FOnGraphExplorerSelectionChanged;
+
+DECLARE_DELEGATE_RetVal(TArray<const URigVMGraph*>, FRigVMGraphExplorer_OnGetRootGraphs);
+DECLARE_DELEGATE( FRigVMGraphExplorer_OnCreateGraph);
+DECLARE_DELEGATE( FRigVMGraphExplorer_OnCreateFunction);
+DECLARE_DELEGATE_RetVal_OneParam(TArray<const URigVMGraph*>, FRigVMGraphExplorer_OnGetChildrenGraphs, const FString&);
+DECLARE_DELEGATE_RetVal_OneParam(TArray<URigVMNode*>, FRigVMGraphExplorer_OnGetEventNodesInGraph, const FString&);
+DECLARE_DELEGATE_RetVal(TArray<URigVMLibraryNode*>, FRigVMGraphExplorer_OnGetFunctions);
+DECLARE_DELEGATE_RetVal_OneParam(FText, FRigVMGraphExplorer_OnGetGraphDisplayName, const FString&);
+DECLARE_DELEGATE_RetVal_OneParam(FText, FRigVMGraphExplorer_OnGetEventDisplayName, const FString&);
+DECLARE_DELEGATE_RetVal_OneParam(FText, FRigVMGraphExplorer_OnGetGraphTooltip, const FString&);
+DECLARE_DELEGATE_RetVal_OneParam(const FSlateBrush*, FRigVMGraphExplorer_OnGetGraphIcon, const FString&);
+DECLARE_DELEGATE_OneParam(FRigVMGraphExplorer_OnGraphClicked, const FString&);
+DECLARE_DELEGATE_OneParam(FRigVMGraphExplorer_OnEventClicked, const FString&);
+DECLARE_DELEGATE_OneParam(FRigVMGraphExplorer_OnFunctionClicked, const FString&);
+DECLARE_DELEGATE_OneParam(FRigVMGraphExplorer_OnGraphDoubleClicked, const FString&);
+DECLARE_DELEGATE_OneParam(FRigVMGraphExplorer_OnEventDoubleClicked, const FString&);
+DECLARE_DELEGATE_OneParam(FRigVMGraphExplorer_OnFunctionDoubleClicked, const FString&);
+DECLARE_DELEGATE_RetVal_TwoParams(bool, FRigVMGraphExplorer_OnRenameGraph, const FString&, const FString&);
+DECLARE_DELEGATE_RetVal_TwoParams(bool, FRigVMGraphExplorer_OnRenameFunction, const FString&, const FString&);
+DECLARE_DELEGATE_RetVal_ThreeParams(bool, FRigVMGraphExplorer_OnCanRenameGraph, const FString&, const FString&, FText&);
+DECLARE_DELEGATE_RetVal_ThreeParams(bool, FRigVMGraphExplorer_OnCanRenameFunction, const FString&, const FString&, FText&);
+DECLARE_DELEGATE_RetVal_TwoParams(bool, FRigVMGraphExplorer_OnSetFunctionCategory, const FString&, const FString&);
+DECLARE_DELEGATE_RetVal_OneParam(FString, FRigVMGraphExplorer_OnGetFunctionCategory, const FString&);
+DECLARE_DELEGATE_RetVal(TSharedPtr<SWidget>, FRigVMGraphExplorer_OnRequestContextMenu);
+
+typedef STreeView<TSharedPtr<FRigVMEditorGraphExplorerTreeElement>>::FOnSelectionChanged FRigVMGraphExplorer_OnSelectionChanged;
 
 struct FRigVMEditorGraphExplorerTreeDelegates
 {
 public:
-	FOnGetRootGraphs OnGetRootGraphs;
-	FOnGetChildrenGraphs OnGetChildrenGraphs;
-	FOnGetEventNodesInGraph OnGetEventNodesInGraph;
-	FOnGetGraphDisplayName OnGetGraphDisplayName;
-	FOnGetGraphDisplayName OnGetEventDisplayName;
-	FOnGetGraphTooltip OnGetGraphTooltip;
-	FOnGetGraphIcon OnGetGraphIcon;
-	FOnGraphDoubleClicked OnGraphClicked;
-	FOnEventDoubleClicked OnEventClicked;
-	FOnGraphDoubleClicked OnGraphDoubleClicked;
-	FOnEventDoubleClicked OnEventDoubleClicked;
-	FOnCreateGraph OnCreateGraph;
-	FOnRenameGraph OnRenameGraph;
-	FOnCanRenameGraph OnCanRenameGraph;
-	FOnGraphExplorerSelectionChanged OnSelectionChanged;
-	FOnRequestContextMenu OnRequestContextMenu;
-	FOnHasContextMenu OnHasContextMenu;
+	FRigVMGraphExplorer_OnGetRootGraphs OnGetRootGraphs;
+	FRigVMGraphExplorer_OnGetChildrenGraphs OnGetChildrenGraphs;
+	FRigVMGraphExplorer_OnGetEventNodesInGraph OnGetEventNodesInGraph;
+	FRigVMGraphExplorer_OnGetFunctions OnGetFunctions;
+	FRigVMGraphExplorer_OnGetGraphDisplayName OnGetGraphDisplayName;
+	FRigVMGraphExplorer_OnGetGraphDisplayName OnGetEventDisplayName;
+	FRigVMGraphExplorer_OnGetGraphTooltip OnGetGraphTooltip;
+	FRigVMGraphExplorer_OnGetGraphIcon OnGetGraphIcon;
+	FRigVMGraphExplorer_OnGraphClicked OnGraphClicked;
+	FRigVMGraphExplorer_OnEventClicked OnEventClicked;
+	FRigVMGraphExplorer_OnFunctionClicked OnFunctionClicked;
+	FRigVMGraphExplorer_OnGraphDoubleClicked OnGraphDoubleClicked;
+	FRigVMGraphExplorer_OnEventDoubleClicked OnEventDoubleClicked;
+	FRigVMGraphExplorer_OnFunctionDoubleClicked OnFunctionDoubleClicked;
+	FRigVMGraphExplorer_OnCreateGraph OnCreateGraph;
+	FRigVMGraphExplorer_OnCreateFunction OnCreateFunction;
+	FRigVMGraphExplorer_OnRenameGraph OnRenameGraph;
+	FRigVMGraphExplorer_OnCanRenameGraph OnCanRenameGraph;
+	FRigVMGraphExplorer_OnRenameFunction OnRenameFunction;
+	FRigVMGraphExplorer_OnCanRenameFunction OnCanRenameFunction;
+	FRigVMGraphExplorer_OnSetFunctionCategory OnSetFunctionCategory;
+	FRigVMGraphExplorer_OnGetFunctionCategory OnGetFunctionCategory;
+	
+	FRigVMGraphExplorer_OnSelectionChanged OnSelectionChanged;
+	FRigVMGraphExplorer_OnRequestContextMenu OnRequestContextMenu;
+	FOnDragDetected OnDragDetected;
 
 	FRigVMEditorGraphExplorerTreeDelegates()
 	{
@@ -115,6 +137,15 @@ public:
 		return TArray<URigVMNode*>();
 	}
 
+	TArray<URigVMLibraryNode*> GetFunctions() const
+	{
+		if (OnGetFunctions.IsBound())
+		{
+			return OnGetFunctions.Execute();
+		}
+		return TArray<URigVMLibraryNode*>();
+	}
+
 	FText GetGraphDisplayName(const FString& InPath) const
 	{
 		if (OnGetGraphDisplayName.IsBound())
@@ -167,6 +198,14 @@ public:
 		}
 	}
 
+	void FunctionClicked(const FString& InPath) const
+	{
+		if (OnFunctionClicked.IsBound())
+		{
+			OnFunctionClicked.Execute(InPath);
+		}
+	}
+
 	void GraphDoubleClicked(const FString& InPath) const
 	{
 		if (OnGraphDoubleClicked.IsBound())
@@ -183,6 +222,14 @@ public:
 		}
 	}
 
+	void FunctionDoubleClicked(const FString& InPath) const
+	{
+		if (OnFunctionDoubleClicked.IsBound())
+		{
+			OnFunctionDoubleClicked.Execute(InPath);
+		}
+	}
+
 	void CreateGraph() const
 	{
 		if (OnCreateGraph.IsBound())
@@ -191,6 +238,14 @@ public:
 		}
 	}
 
+	void CreateFunction() const
+	{
+		if (OnCreateFunction.IsBound())
+		{
+			OnCreateFunction.Execute();
+		}
+	}
+
 	void RenameGraph(const FString& InOldPath, const FString& InNewPath) const
 	{
 		if (OnRenameGraph.IsBound())
@@ -208,6 +263,40 @@ public:
 		return false;
 	}
 
+	void RenameFunction(const FString& InOldPath, const FString& InNewPath) const
+	{
+		if (OnRenameFunction.IsBound())
+		{
+			OnRenameFunction.Execute(InOldPath, InNewPath);
+		}
+	}
+
+	bool CanRenameFunction(const FString& InOldPath, const FString& InNewPath, FText& OutErrorMessage) const
+	{
+		if (OnCanRenameFunction.IsBound())
+		{
+			return OnCanRenameFunction.Execute(InOldPath, InNewPath, OutErrorMessage);
+		}
+		return false;
+	}
+
+	void SetFunctionCategory(const FString& InPath, const FString& InCategory) const
+	{
+		if (OnSetFunctionCategory.IsBound())
+		{
+			OnSetFunctionCategory.Execute(InPath, InCategory);
+		}
+	}
+
+	FString GetFunctionCategory(const FString& InPath) const
+	{
+		if (OnGetFunctionCategory.IsBound())
+		{
+			return OnGetFunctionCategory.Execute(InPath);
+		}
+		return FString(); 
+	}
+
 	void HandleSelectionChanged(TSharedPtr<FRigVMEditorGraphExplorerTreeElement> Selection, ESelectInfo::Type SelectInfo)
 	{
 		if(bSuspendSelectionDelegate)
@@ -227,15 +316,6 @@ public:
 		return nullptr;
 	}
 
-	bool HasContextMenu()
-	{
-		if (OnHasContextMenu.IsBound())
-		{
-			return OnHasContextMenu.Execute();
-		}
-		return false;
-	}
-	
 private:
 
 	bool bSuspendSelectionDelegate = false;
diff --git a/Engine/Source/Runtime/Engine/Classes/EdGraph/EdGraphSchema.h b/Engine/Source/Runtime/Engine/Classes/EdGraph/EdGraphSchema.h
--- a/Engine/Source/Runtime/Engine/Classes/EdGraph/EdGraphSchema.h
+++ b/Engine/Source/Runtime/Engine/Classes/EdGraph/EdGraphSchema.h
@@ -1150,7 +1150,7 @@ class UEdGraphSchema : public UObject
 	 * @param	InGraph				Graph to set the category for
 	 * @param	InCategory			Pipe "|" separated category for the graph.
 	 */
-	virtual FReply TrySetGraphCategory(const UEdGraph* InGraph, const FText& InCategory) { return FReply::Unhandled(); }
+	virtual FReply TrySetGraphCategory(const UEdGraph* InGraph, const FText& InCategory) const { return FReply::Unhandled(); }
 
 	/** Called when asset(s) are dropped onto a graph background. */
 	virtual void DroppedAssetsOnGraph(const TArray<struct FAssetData>& Assets, const FVector2D& GraphPosition, UEdGraph* Graph) const {}
