Control Rig: Fix nonunity compilation error
#rnx

diff --git a/Engine/Plugins/Runtime/RigVM/Source/RigVMEditor/Public/Widgets/SRigVMEditorGraphExplorer.h b/Engine/Plugins/Runtime/RigVM/Source/RigVMEditor/Public/Widgets/SRigVMEditorGraphExplorer.h
--- a/Engine/Plugins/Runtime/RigVM/Source/RigVMEditor/Public/Widgets/SRigVMEditorGraphExplorer.h
+++ b/Engine/Plugins/Runtime/RigVM/Source/RigVMEditor/Public/Widgets/SRigVMEditorGraphExplorer.h
@@ -15,6 +15,7 @@ class SRigVMEditorGraphExplorerTreeView;
 class URigVMGraph;
 class URigVMEdGraph;
 class FRigVMEditor;
+class URigVMBlueprint;
 class SGraphActionMenu;
 struct FCreateWidgetForActionData;
 class FRigVMEditorGraphExplorerTreeElement;
