RigVM / ControlRig : Added debug draw settings to debug funcions and debug nodes

diff --git a/Engine/Plugins/Animation/ControlRig/Source/ControlRig/Private/AnimNode_ControlRigBase.cpp b/Engine/Plugins/Animation/ControlRig/Source/ControlRig/Private/AnimNode_ControlRigBase.cpp
--- a/Engine/Plugins/Animation/ControlRig/Source/ControlRig/Private/AnimNode_ControlRigBase.cpp
+++ b/Engine/Plugins/Animation/ControlRig/Source/ControlRig/Private/AnimNode_ControlRigBase.cpp
@@ -443,7 +443,7 @@ void FAnimNode_ControlRigBase::QueueControlRigDrawInstructions(UControlRig* Cont
 				{
 					for (const FVector& Point : Instruction.Positions)
 					{
-						Proxy->AnimDrawDebugPoint(InstructionTransform.TransformPosition(Point), Instruction.Thickness, Instruction.Color.ToFColor(true), false, -1.f, SDPG_Foreground);
+						Proxy->AnimDrawDebugPoint(InstructionTransform.TransformPosition(Point), Instruction.Thickness, Instruction.Color.ToFColor(true), false, Instruction.Lifetime, Instruction.DepthPriority);
 					}
 					break;
 				}
@@ -453,7 +453,7 @@ void FAnimNode_ControlRigBase::QueueControlRigDrawInstructions(UControlRig* Cont
 
 					for (int32 PointIndex = 0; PointIndex < Points.Num() - 1; PointIndex += 2)
 					{
-						Proxy->AnimDrawDebugLine(InstructionTransform.TransformPosition(Points[PointIndex]), InstructionTransform.TransformPosition(Points[PointIndex + 1]), Instruction.Color.ToFColor(true), false, -1.f, Instruction.Thickness, SDPG_Foreground);
+						Proxy->AnimDrawDebugLine(InstructionTransform.TransformPosition(Points[PointIndex]), InstructionTransform.TransformPosition(Points[PointIndex + 1]), Instruction.Color.ToFColor(true), false, Instruction.Lifetime, Instruction.Thickness, Instruction.DepthPriority);
 					}
 					break;
 				}
@@ -463,7 +463,7 @@ void FAnimNode_ControlRigBase::QueueControlRigDrawInstructions(UControlRig* Cont
 
 					for (int32 PointIndex = 0; PointIndex < Points.Num() - 1; PointIndex++)
 					{
-						Proxy->AnimDrawDebugLine(InstructionTransform.TransformPosition(Points[PointIndex]), InstructionTransform.TransformPosition(Points[PointIndex + 1]), Instruction.Color.ToFColor(true), false, -1.f, Instruction.Thickness, SDPG_Foreground);
+						Proxy->AnimDrawDebugLine(InstructionTransform.TransformPosition(Points[PointIndex]), InstructionTransform.TransformPosition(Points[PointIndex + 1]), Instruction.Color.ToFColor(true), false, Instruction.Lifetime, Instruction.Thickness, Instruction.DepthPriority);
 					}
 					break;
 				}
diff --git a/Engine/Plugins/Animation/ControlRig/Source/ControlRig/Private/Units/Debug/RigUnit_DebugBezier.cpp b/Engine/Plugins/Animation/ControlRig/Source/ControlRig/Private/Units/Debug/RigUnit_DebugBezier.cpp
--- a/Engine/Plugins/Animation/ControlRig/Source/ControlRig/Private/Units/Debug/RigUnit_DebugBezier.cpp
+++ b/Engine/Plugins/Animation/ControlRig/Source/ControlRig/Private/Units/Debug/RigUnit_DebugBezier.cpp
@@ -18,7 +18,8 @@ FRigUnit_DebugBezier_Execute()
 		Detail, 
 		FRigElementKey(Space, ERigElementType::Bone), 
 		WorldOffset, 
-		bEnabled);
+		bEnabled,
+		DebugDrawSettings);
 }
 
 FRigVMStructUpgradeInfo FRigUnit_DebugBezier::GetUpgradeInfo() const
diff --git a/Engine/Plugins/Animation/ControlRig/Source/ControlRig/Private/Units/Debug/RigUnit_DebugHierarchy.cpp b/Engine/Plugins/Animation/ControlRig/Source/ControlRig/Private/Units/Debug/RigUnit_DebugHierarchy.cpp
--- a/Engine/Plugins/Animation/ControlRig/Source/ControlRig/Private/Units/Debug/RigUnit_DebugHierarchy.cpp
+++ b/Engine/Plugins/Animation/ControlRig/Source/ControlRig/Private/Units/Debug/RigUnit_DebugHierarchy.cpp
@@ -3,6 +3,7 @@
 #include "Units/Debug/RigUnit_DebugHierarchy.h"
 #include "Units/RigUnitContext.h"
 #include "ControlRigDefines.h"
+#include "RigVMCore/RigVMDebugDrawSettings.h"
 
 #include UE_INLINE_GENERATED_CPP_BY_NAME(RigUnit_DebugHierarchy)
 
@@ -13,7 +14,7 @@ FRigUnit_DebugHierarchy_Execute()
 	URigHierarchy* Hierarchy = ExecuteContext.Hierarchy;
 	if (Hierarchy)
 	{
-		DrawHierarchy(ExecuteContext, WorldOffset, Hierarchy, EControlRigDrawHierarchyMode::Axes, Scale, Color, Thickness, nullptr, &Items);
+		DrawHierarchy(ExecuteContext, WorldOffset, Hierarchy, EControlRigDrawHierarchyMode::Axes, Scale, Color, Thickness, nullptr, &Items, DebugDrawSettings);
 	}
 }
 
@@ -24,11 +25,11 @@ FRigUnit_DebugPose_Execute()
 	URigHierarchy* Hierarchy = ExecuteContext.Hierarchy;
 	if (Hierarchy)
 	{
-		FRigUnit_DebugHierarchy::DrawHierarchy(ExecuteContext, WorldOffset, Hierarchy, EControlRigDrawHierarchyMode::Axes, Scale, Color, Thickness, &Pose, &Items);
+		FRigUnit_DebugHierarchy::DrawHierarchy(ExecuteContext, WorldOffset, Hierarchy, EControlRigDrawHierarchyMode::Axes, Scale, Color, Thickness, &Pose, &Items, DebugDrawSettings);
 	}
 }
 
-void FRigUnit_DebugHierarchy::DrawHierarchy(const FRigVMExecuteContext& InContext, const FTransform& WorldOffset, URigHierarchy* Hierarchy, EControlRigDrawHierarchyMode::Type Mode, float Scale, const FLinearColor& Color, float Thickness, const FRigPose* InPose, const TArrayView<const FRigElementKey>* InItems)
+void FRigUnit_DebugHierarchy::DrawHierarchy(const FRigVMExecuteContext& InContext, const FTransform& WorldOffset, URigHierarchy* Hierarchy, EControlRigDrawHierarchyMode::Type Mode, float Scale, const FLinearColor& Color, float Thickness, const FRigPose* InPose, const TArrayView<const FRigElementKey>* InItems, const FRigVMDebugDrawSettings& DebugDrawSettings)
 {
 	FRigVMDrawInterface* DrawInterface = InContext.GetDrawInterface();
 	if(DrawInterface == nullptr)
@@ -57,10 +58,10 @@ void FRigUnit_DebugHierarchy::DrawHierarchy(const FRigVMExecuteContext& InContex
 	{
 		case EControlRigDrawHierarchyMode::Axes:
 		{
-			FRigVMDrawInstruction InstructionX(ERigVMDrawSettings::Lines, FLinearColor::Red, Thickness, WorldOffset);
-			FRigVMDrawInstruction InstructionY(ERigVMDrawSettings::Lines, FLinearColor::Green, Thickness, WorldOffset);
-			FRigVMDrawInstruction InstructionZ(ERigVMDrawSettings::Lines, FLinearColor::Blue, Thickness, WorldOffset);
-			FRigVMDrawInstruction InstructionParent(ERigVMDrawSettings::Lines, Color, Thickness, WorldOffset);
+			FRigVMDrawInstruction InstructionX(ERigVMDrawSettings::Lines, FLinearColor::Red, Thickness, WorldOffset, DebugDrawSettings.DepthPriority, DebugDrawSettings.Lifetime);
+			FRigVMDrawInstruction InstructionY(ERigVMDrawSettings::Lines, FLinearColor::Green, Thickness, WorldOffset, DebugDrawSettings.DepthPriority, DebugDrawSettings.Lifetime);
+			FRigVMDrawInstruction InstructionZ(ERigVMDrawSettings::Lines, FLinearColor::Blue, Thickness, WorldOffset, DebugDrawSettings.DepthPriority, DebugDrawSettings.Lifetime);
+			FRigVMDrawInstruction InstructionParent(ERigVMDrawSettings::Lines, Color, Thickness, WorldOffset, DebugDrawSettings.DepthPriority, DebugDrawSettings.Lifetime);
 			InstructionX.Positions.Reserve(Hierarchy->Num() * 2);
 			InstructionY.Positions.Reserve(Hierarchy->Num() * 2);
 			InstructionZ.Positions.Reserve(Hierarchy->Num() * 2);
diff --git a/Engine/Plugins/Animation/ControlRig/Source/ControlRig/Private/Units/Debug/RigUnit_DebugLine.cpp b/Engine/Plugins/Animation/ControlRig/Source/ControlRig/Private/Units/Debug/RigUnit_DebugLine.cpp
--- a/Engine/Plugins/Animation/ControlRig/Source/ControlRig/Private/Units/Debug/RigUnit_DebugLine.cpp
+++ b/Engine/Plugins/Animation/ControlRig/Source/ControlRig/Private/Units/Debug/RigUnit_DebugLine.cpp
@@ -16,7 +16,8 @@ FRigUnit_DebugLine_Execute()
 		Thickness,
 		FRigElementKey(Space, ERigElementType::Bone), 
 		WorldOffset, 
-		bEnabled);
+		bEnabled,
+		DebugDrawSettings);
 }
 
 FRigVMStructUpgradeInfo FRigUnit_DebugLine::GetUpgradeInfo() const
@@ -49,7 +50,7 @@ FRigUnit_DebugLineItemSpace_Execute()
 		DrawB = Transform.TransformPosition(DrawB);
 	}
 
-	ExecuteContext.GetDrawInterface()->DrawLine(WorldOffset, DrawA, DrawB, Color, Thickness);
+	ExecuteContext.GetDrawInterface()->DrawLine(WorldOffset, DrawA, DrawB, Color, Thickness, DebugDrawSettings.DepthPriority, DebugDrawSettings.Lifetime);
 }
 
 FRigVMStructUpgradeInfo FRigUnit_DebugLineItemSpace::GetUpgradeInfo() const
diff --git a/Engine/Plugins/Animation/ControlRig/Source/ControlRig/Private/Units/Debug/RigUnit_DebugLineStrip.cpp b/Engine/Plugins/Animation/ControlRig/Source/ControlRig/Private/Units/Debug/RigUnit_DebugLineStrip.cpp
--- a/Engine/Plugins/Animation/ControlRig/Source/ControlRig/Private/Units/Debug/RigUnit_DebugLineStrip.cpp
+++ b/Engine/Plugins/Animation/ControlRig/Source/ControlRig/Private/Units/Debug/RigUnit_DebugLineStrip.cpp
@@ -15,7 +15,8 @@ FRigUnit_DebugLineStrip_Execute()
 		Thickness,
 		FRigElementKey(Space, ERigElementType::Bone), 
 		WorldOffset, 
-		bEnabled);
+		bEnabled,
+		DebugDrawSettings);
 }
 
 FRigVMStructUpgradeInfo FRigUnit_DebugLineStrip::GetUpgradeInfo() const
@@ -48,11 +49,11 @@ FRigUnit_DebugLineStripItemSpace_Execute()
 		{
 			PointsTransformed.Add(Transform.TransformPosition(Point));
 		}
-		ExecuteContext.GetDrawInterface()->DrawLineStrip(WorldOffset, PointsTransformed, Color, Thickness);
+		ExecuteContext.GetDrawInterface()->DrawLineStrip(WorldOffset, PointsTransformed, Color, Thickness, DebugDrawSettings.DepthPriority, DebugDrawSettings.Lifetime);
 	}
 	else
 	{
-		ExecuteContext.GetDrawInterface()->DrawLineStrip(WorldOffset, TArrayView<const FVector>(Points.GetData(), Points.Num()), Color, Thickness);
+		ExecuteContext.GetDrawInterface()->DrawLineStrip(WorldOffset, TArrayView<const FVector>(Points.GetData(), Points.Num()), Color, Thickness, DebugDrawSettings.DepthPriority, DebugDrawSettings.Lifetime);
 	}
 }
 
diff --git a/Engine/Plugins/Animation/ControlRig/Source/ControlRig/Private/Units/Debug/RigUnit_DebugPrimitives.cpp b/Engine/Plugins/Animation/ControlRig/Source/ControlRig/Private/Units/Debug/RigUnit_DebugPrimitives.cpp
--- a/Engine/Plugins/Animation/ControlRig/Source/ControlRig/Private/Units/Debug/RigUnit_DebugPrimitives.cpp
+++ b/Engine/Plugins/Animation/ControlRig/Source/ControlRig/Private/Units/Debug/RigUnit_DebugPrimitives.cpp
@@ -16,7 +16,8 @@ FRigUnit_DebugRectangle_Execute()
 		Thickness,
 		FRigElementKey(Space, ERigElementType::Bone), 
 		WorldOffset, 
-		bEnabled);
+		bEnabled,
+		DebugDrawSettings);
 }
 
 FRigVMStructUpgradeInfo FRigUnit_DebugRectangle::GetUpgradeInfo() const
@@ -60,7 +61,7 @@ FRigUnit_DebugRectangleItemSpace_Execute()
 		DrawTransform = DrawTransform * ExecuteContext.Hierarchy->GetGlobalTransform(Space);
 	}
 
-	ExecuteContext.GetDrawInterface()->DrawRectangle(WorldOffset, DrawTransform, Scale, Color, Thickness);
+	ExecuteContext.GetDrawInterface()->DrawRectangle(WorldOffset, DrawTransform, Scale, Color, Thickness, DebugDrawSettings.DepthPriority, DebugDrawSettings.Lifetime);
 }
 
 FRigUnit_DebugArc_Execute()
@@ -76,7 +77,8 @@ FRigUnit_DebugArc_Execute()
 		Detail,
 		FRigElementKey(Space, ERigElementType::Bone), 
 		WorldOffset, 
-		bEnabled);
+		bEnabled,
+		DebugDrawSettings);
 }
 
 FRigVMStructUpgradeInfo FRigUnit_DebugArc::GetUpgradeInfo() const
@@ -126,5 +128,5 @@ FRigUnit_DebugArcItemSpace_Execute()
 		DrawTransform = DrawTransform * ExecuteContext.Hierarchy->GetGlobalTransform(Space);
 	}
 
-	ExecuteContext.GetDrawInterface()->DrawArc(WorldOffset, DrawTransform, Radius, FMath::DegreesToRadians(MinimumDegrees), FMath::DegreesToRadians(MaximumDegrees), Color, Thickness, Detail);
+	ExecuteContext.GetDrawInterface()->DrawArc(WorldOffset, DrawTransform, Radius, FMath::DegreesToRadians(MinimumDegrees), FMath::DegreesToRadians(MaximumDegrees), Color, Thickness, Detail, DebugDrawSettings.DepthPriority, DebugDrawSettings.Lifetime);
 }
diff --git a/Engine/Plugins/Animation/ControlRig/Source/ControlRig/Private/Units/Debug/RigUnit_DebugTransform.cpp b/Engine/Plugins/Animation/ControlRig/Source/ControlRig/Private/Units/Debug/RigUnit_DebugTransform.cpp
--- a/Engine/Plugins/Animation/ControlRig/Source/ControlRig/Private/Units/Debug/RigUnit_DebugTransform.cpp
+++ b/Engine/Plugins/Animation/ControlRig/Source/ControlRig/Private/Units/Debug/RigUnit_DebugTransform.cpp
@@ -26,18 +26,18 @@ FRigUnit_DebugTransform_Execute()
 	{
 		case ERigUnitDebugTransformMode::Axes:
 		{
-			ExecuteContext.GetDrawInterface()->DrawAxes(WorldOffset, DrawTransform, Scale, Thickness);
+			ExecuteContext.GetDrawInterface()->DrawAxes(WorldOffset, DrawTransform, Scale, Thickness, DebugDrawSettings.DepthPriority, DebugDrawSettings.Lifetime);
 			break;
 		}
 		case ERigUnitDebugTransformMode::Point:
 		{
-			ExecuteContext.GetDrawInterface()->DrawPoint(WorldOffset, DrawTransform.GetLocation(), Scale, Color);
+			ExecuteContext.GetDrawInterface()->DrawPoint(WorldOffset, DrawTransform.GetLocation(), Scale, Color, DebugDrawSettings.DepthPriority, DebugDrawSettings.Lifetime);
 			break;
 		}
 		case ERigUnitDebugTransformMode::Box:
 		{
 			DrawTransform.SetScale3D(DrawTransform.GetScale3D() * Scale);
-			ExecuteContext.GetDrawInterface()->DrawBox(WorldOffset, DrawTransform, Color, Thickness);
+			ExecuteContext.GetDrawInterface()->DrawBox(WorldOffset, DrawTransform, Color, Thickness, DebugDrawSettings.DepthPriority, DebugDrawSettings.Lifetime);
 			break;
 		}
 	}
@@ -69,7 +69,8 @@ FRigUnit_DebugTransformMutable_Execute()
 		Scale,
 		FRigElementKey(Space, ERigElementType::Bone), 
 		WorldOffset, 
-		bEnabled);
+		bEnabled,
+		DebugDrawSettings);
 }
 
 FRigVMStructUpgradeInfo FRigUnit_DebugTransformMutable::GetUpgradeInfo() const
@@ -120,18 +121,18 @@ FRigUnit_DebugTransformMutableItemSpace_Execute()
 	{
 		case ERigUnitDebugTransformMode::Axes:
 		{
-			ExecuteContext.GetDrawInterface()->DrawAxes(WorldOffset, DrawTransform, Scale, Thickness);
+			ExecuteContext.GetDrawInterface()->DrawAxes(WorldOffset, DrawTransform, Scale, Thickness, DebugDrawSettings.DepthPriority, DebugDrawSettings.Lifetime);
 			break;
 		}
 		case ERigUnitDebugTransformMode::Point:
 		{
-			ExecuteContext.GetDrawInterface()->DrawPoint(WorldOffset, DrawTransform.GetLocation(), Scale, Color);
+			ExecuteContext.GetDrawInterface()->DrawPoint(WorldOffset, DrawTransform.GetLocation(), Scale, Color, DebugDrawSettings.DepthPriority, DebugDrawSettings.Lifetime);
 			break;
 		}
 		case ERigUnitDebugTransformMode::Box:
 		{
 			DrawTransform.SetScale3D(DrawTransform.GetScale3D() * Scale);
-			ExecuteContext.GetDrawInterface()->DrawBox(WorldOffset, DrawTransform, Color, Thickness);
+			ExecuteContext.GetDrawInterface()->DrawBox(WorldOffset, DrawTransform, Color, Thickness, DebugDrawSettings.DepthPriority, DebugDrawSettings.Lifetime);
 			break;
 		}
 	}
@@ -168,18 +169,18 @@ FRigUnit_DebugTransformArrayMutable_Execute()
 		{
 			case ERigUnitDebugTransformMode::Axes:
 			{
-				ExecuteContext.GetDrawInterface()->DrawAxes(WorldOffset, DrawTransform, Scale, Thickness);
+				ExecuteContext.GetDrawInterface()->DrawAxes(WorldOffset, DrawTransform, Scale, Thickness, DebugDrawSettings.DepthPriority, DebugDrawSettings.Lifetime);
 				break;
 			}
 			case ERigUnitDebugTransformMode::Point:
 			{
-				ExecuteContext.GetDrawInterface()->DrawPoint(WorldOffset, DrawTransform.GetLocation(), Scale, Color);
+				ExecuteContext.GetDrawInterface()->DrawPoint(WorldOffset, DrawTransform.GetLocation(), Scale, Color, DebugDrawSettings.DepthPriority, DebugDrawSettings.Lifetime);
 				break;
 			}
 			case ERigUnitDebugTransformMode::Box:
 			{
 				DrawTransform.SetScale3D(DrawTransform.GetScale3D() * Scale);
-				ExecuteContext.GetDrawInterface()->DrawBox(WorldOffset, DrawTransform, Color, Thickness);
+				ExecuteContext.GetDrawInterface()->DrawBox(WorldOffset, DrawTransform, Color, Thickness, DebugDrawSettings.DepthPriority, DebugDrawSettings.Lifetime);
 				break;
 			}
 		}
@@ -235,7 +236,8 @@ FRigUnit_DebugTransformArrayMutableItemSpace_Execute()
 			Scale,
 			Space, 
 			WorldOffset, 
-			bEnabled);
+			bEnabled,
+			DebugDrawSettings);
 	}
 
 	if(ParentIndices.Num() == Transforms.Num())
@@ -249,7 +251,9 @@ FRigUnit_DebugTransformArrayMutableItemSpace_Execute()
 					Transforms[Index].GetTranslation(),
 					Transforms[ParentIndices[Index]].GetTranslation(),
 					Color,
-					Thickness
+					Thickness,
+					DebugDrawSettings.DepthPriority,
+					DebugDrawSettings.Lifetime
 				);
 			}
 		}
diff --git a/Engine/Plugins/Animation/ControlRig/Source/ControlRig/Private/Units/Debug/RigUnit_VisualDebug.cpp b/Engine/Plugins/Animation/ControlRig/Source/ControlRig/Private/Units/Debug/RigUnit_VisualDebug.cpp
--- a/Engine/Plugins/Animation/ControlRig/Source/ControlRig/Private/Units/Debug/RigUnit_VisualDebug.cpp
+++ b/Engine/Plugins/Animation/ControlRig/Source/ControlRig/Private/Units/Debug/RigUnit_VisualDebug.cpp
@@ -8,7 +8,7 @@
 
 FRigUnit_VisualDebugVector_Execute()
 {
-	FRigUnit_VisualDebugVectorItemSpace::StaticExecute(ExecuteContext, Value, bEnabled, Mode, Color, Thickness, Scale, FRigElementKey(BoneSpace, ERigElementType::Bone));
+	FRigUnit_VisualDebugVectorItemSpace::StaticExecute(ExecuteContext, Value, bEnabled, Mode, Color, Thickness, Scale, FRigElementKey(BoneSpace, ERigElementType::Bone), DebugDrawSettings);
 }
 
 FRigVMStructUpgradeInfo FRigUnit_VisualDebugVector::GetUpgradeInfo() const
@@ -44,12 +44,12 @@ FRigUnit_VisualDebugVectorItemSpace_Execute()
 	{
 		case ERigUnitVisualDebugPointMode::Point:
 		{
-			ExecuteContext.GetDrawInterface()->DrawPoint(WorldOffset, Value, Thickness, Color);
+			ExecuteContext.GetDrawInterface()->DrawPoint(WorldOffset, Value, Thickness, Color, DebugDrawSettings.DepthPriority, DebugDrawSettings.Lifetime);
 			break;
 		}
 		case ERigUnitVisualDebugPointMode::Vector:
 		{
-			ExecuteContext.GetDrawInterface()->DrawLine(WorldOffset, FVector::ZeroVector, Value * Scale, Color, Thickness);
+			ExecuteContext.GetDrawInterface()->DrawLine(WorldOffset, FVector::ZeroVector, Value * Scale, Color, Thickness, DebugDrawSettings.DepthPriority, DebugDrawSettings.Lifetime);
 			break;
 		}
 		default:
@@ -61,7 +61,7 @@ FRigUnit_VisualDebugVectorItemSpace_Execute()
 
 FRigUnit_VisualDebugQuat_Execute()
 {
-	FRigUnit_VisualDebugQuatItemSpace::StaticExecute(ExecuteContext, Value, bEnabled, Thickness, Scale, FRigElementKey(BoneSpace, ERigElementType::Bone));
+	FRigUnit_VisualDebugQuatItemSpace::StaticExecute(ExecuteContext, Value, bEnabled, Thickness, Scale, FRigElementKey(BoneSpace, ERigElementType::Bone), DebugDrawSettings);
 }
 
 FRigVMStructUpgradeInfo FRigUnit_VisualDebugVectorItemSpace::GetUpgradeInfo() const
@@ -97,12 +97,12 @@ FRigUnit_VisualDebugQuatItemSpace_Execute()
     FTransform Transform = FTransform::Identity;
     Transform.SetRotation(Value);
 
-	FRigUnit_VisualDebugTransformItemSpace::StaticExecute(ExecuteContext, Transform, bEnabled, Thickness, Scale, Space);
+	FRigUnit_VisualDebugTransformItemSpace::StaticExecute(ExecuteContext, Transform, bEnabled, Thickness, Scale, Space, DebugDrawSettings);
 }
 
 FRigUnit_VisualDebugTransform_Execute()
 {
-	FRigUnit_VisualDebugTransformItemSpace::StaticExecute(ExecuteContext, Value, bEnabled, Thickness, Scale, FRigElementKey(BoneSpace, ERigElementType::Bone));
+	FRigUnit_VisualDebugTransformItemSpace::StaticExecute(ExecuteContext, Value, bEnabled, Thickness, Scale, FRigElementKey(BoneSpace, ERigElementType::Bone), DebugDrawSettings);
 }
 
 FRigVMStructUpgradeInfo FRigUnit_VisualDebugQuatItemSpace::GetUpgradeInfo() const
@@ -156,5 +156,5 @@ FRigUnit_VisualDebugTransformItemSpace_Execute()
 		WorldOffset = ExecuteContext.Hierarchy->GetGlobalTransform(Space);
 	}
 
-	ExecuteContext.GetDrawInterface()->DrawAxes(WorldOffset, Value, Scale, Thickness);
+	ExecuteContext.GetDrawInterface()->DrawAxes(WorldOffset, Value, Scale, Thickness, DebugDrawSettings.DepthPriority, DebugDrawSettings.Lifetime);
 }
diff --git a/Engine/Plugins/Animation/ControlRig/Source/ControlRig/Public/Units/Debug/RigUnit_DebugBase.h b/Engine/Plugins/Animation/ControlRig/Source/ControlRig/Public/Units/Debug/RigUnit_DebugBase.h
--- a/Engine/Plugins/Animation/ControlRig/Source/ControlRig/Public/Units/Debug/RigUnit_DebugBase.h
+++ b/Engine/Plugins/Animation/ControlRig/Source/ControlRig/Public/Units/Debug/RigUnit_DebugBase.h
@@ -3,16 +3,23 @@
 #pragma once
 
 #include "Units/RigUnit.h"
+#include "RigVMCore/RigVMDebugDrawSettings.h"
 #include "RigUnit_DebugBase.generated.h"
 
 USTRUCT(meta=(Abstract, Category="Debug", NodeColor = "0.83077 0.846873 0.049707"))
 struct CONTROLRIG_API FRigUnit_DebugBase : public FRigUnit
 {
 	GENERATED_BODY()
+
+	UPROPERTY(meta = (Input, DetailsOnly, DisplayName = "Draw Settings"))
+	FRigVMDebugDrawSettings DebugDrawSettings;
 };
 
 USTRUCT(meta=(Abstract, Category="Debug", NodeColor = "0.83077 0.846873 0.049707"))
 struct CONTROLRIG_API FRigUnit_DebugBaseMutable : public FRigUnitMutable
 {
 	GENERATED_BODY()
+
+	UPROPERTY(meta = (Input, DetailsOnly, DisplayName = "Draw Settings"))
+	FRigVMDebugDrawSettings DebugDrawSettings;
 };
diff --git a/Engine/Plugins/Animation/ControlRig/Source/ControlRig/Public/Units/Debug/RigUnit_DebugHierarchy.h b/Engine/Plugins/Animation/ControlRig/Source/ControlRig/Public/Units/Debug/RigUnit_DebugHierarchy.h
--- a/Engine/Plugins/Animation/ControlRig/Source/ControlRig/Public/Units/Debug/RigUnit_DebugHierarchy.h
+++ b/Engine/Plugins/Animation/ControlRig/Source/ControlRig/Public/Units/Debug/RigUnit_DebugHierarchy.h
@@ -63,7 +63,10 @@ struct CONTROLRIG_API FRigUnit_DebugHierarchy : public FRigVMFunction_DebugBase
 	UPROPERTY(meta = (Input))
 	bool bEnabled;
 
-	static void DrawHierarchy(const FRigVMExecuteContext& InContext, const FTransform& WorldOffset, URigHierarchy* Hierarchy, EControlRigDrawHierarchyMode::Type Mode, float Scale, const FLinearColor& Color, float Thickness, const FRigPose* InPose, const TArrayView<const FRigElementKey>* InItems);
+	static void DrawHierarchy(const FRigVMExecuteContext& InContext, const FTransform& WorldOffset, URigHierarchy* Hierarchy, EControlRigDrawHierarchyMode::Type Mode, float Scale, const FLinearColor& Color, float Thickness, const FRigPose* InPose, const TArrayView<const FRigElementKey>* InItems, const FRigVMDebugDrawSettings& DebugDrawSettings);
+	
+	UE_DEPRECATED(5.5, "Please use DrawHierarchy with DrawDebugSettings")
+	static void DrawHierarchy(const FRigVMExecuteContext& InContext, const FTransform& WorldOffset, URigHierarchy* Hierarchy, EControlRigDrawHierarchyMode::Type Mode, float Scale, const FLinearColor& Color, float Thickness, const FRigPose* InPose, const TArrayView<const FRigElementKey>* InItems) {}
 };
 
 /**
diff --git a/Engine/Plugins/Runtime/RigVM/Source/RigVM/Private/RigVMCore/RigVMDrawInterface.cpp b/Engine/Plugins/Runtime/RigVM/Source/RigVM/Private/RigVMCore/RigVMDrawInterface.cpp
--- a/Engine/Plugins/Runtime/RigVM/Source/RigVM/Private/RigVMCore/RigVMDrawInterface.cpp
+++ b/Engine/Plugins/Runtime/RigVM/Source/RigVM/Private/RigVMCore/RigVMDrawInterface.cpp
@@ -18,68 +18,68 @@ void FRigVMDrawInterface::DrawInstruction(const FRigVMDrawInstruction& InInstruc
 	Instructions.Add(InInstruction);
 }
 
-void FRigVMDrawInterface::DrawPoint(const FTransform& WorldOffset, const FVector& Position, float Size, const FLinearColor& Color)
+void FRigVMDrawInterface::DrawPoint(const FTransform& WorldOffset, const FVector& Position, float Size, const FLinearColor& Color, ESceneDepthPriorityGroup DepthPriority, float Lifetime)
 {
 	if (!IsEnabled())
 	{
 		return;
 	}
 
-	FRigVMDrawInstruction Instruction(ERigVMDrawSettings::Points, Color, Size, WorldOffset);
+	FRigVMDrawInstruction Instruction(ERigVMDrawSettings::Points, Color, Size, WorldOffset, DepthPriority, Lifetime);
 	Instruction.Positions.Add(Position);
 	Instructions.Add(Instruction);
 }
 
-void FRigVMDrawInterface::DrawPoints(const FTransform& WorldOffset, const TArrayView<const FVector>& Points, float Size, const FLinearColor& Color)
+void FRigVMDrawInterface::DrawPoints(const FTransform& WorldOffset, const TArrayView<const FVector>& Points, float Size, const FLinearColor& Color, ESceneDepthPriorityGroup DepthPriority, float Lifetime)
 {
 	if (!IsEnabled())
 	{
 		return;
 	}
 
-	FRigVMDrawInstruction Instruction(ERigVMDrawSettings::Points, Color, Size, WorldOffset);
+	FRigVMDrawInstruction Instruction(ERigVMDrawSettings::Points, Color, Size, WorldOffset, DepthPriority, Lifetime);
 	Instruction.Positions.Append(Points.GetData(), Points.Num());
 	Instructions.Add(Instruction);
 }
 
-void FRigVMDrawInterface::DrawLine(const FTransform& WorldOffset, const FVector& LineStart, const FVector& LineEnd, const FLinearColor& Color, float Thickness)
+void FRigVMDrawInterface::DrawLine(const FTransform& WorldOffset, const FVector& LineStart, const FVector& LineEnd, const FLinearColor& Color, float Thickness, ESceneDepthPriorityGroup DepthPriority, float Lifetime)
 {
 	if (!IsEnabled())
 	{
 		return;
 	}
 
-	FRigVMDrawInstruction Instruction(ERigVMDrawSettings::Lines, Color, Thickness, WorldOffset);
+	FRigVMDrawInstruction Instruction(ERigVMDrawSettings::Lines, Color, Thickness, WorldOffset, DepthPriority, Lifetime);
 	Instruction.Positions.Add(LineStart);
 	Instruction.Positions.Add(LineEnd);
 	Instructions.Add(Instruction);
 }
 
-void FRigVMDrawInterface::DrawLines(const FTransform& WorldOffset, const TArrayView<const FVector>& Positions, const FLinearColor& Color, float Thickness)
+void FRigVMDrawInterface::DrawLines(const FTransform& WorldOffset, const TArrayView<const FVector>& Positions, const FLinearColor& Color, float Thickness, ESceneDepthPriorityGroup DepthPriority, float Lifetime)
 {
 	if (!IsEnabled())
 	{
 		return;
 	}
 
-	FRigVMDrawInstruction Instruction(ERigVMDrawSettings::Lines, Color, Thickness, WorldOffset);
+	FRigVMDrawInstruction Instruction(ERigVMDrawSettings::Lines, Color, Thickness, WorldOffset, DepthPriority, Lifetime);
 	Instruction.Positions.Append(Positions.GetData(), Positions.Num());
 	Instructions.Add(Instruction);
 }
 
-void FRigVMDrawInterface::DrawLineStrip(const FTransform& WorldOffset, const TArrayView<const FVector>& Positions, const FLinearColor& Color, float Thickness)
+void FRigVMDrawInterface::DrawLineStrip(const FTransform& WorldOffset, const TArrayView<const FVector>& Positions, const FLinearColor& Color, float Thickness, ESceneDepthPriorityGroup DepthPriority, float Lifetime)
 {
 	if (!IsEnabled())
 	{
 		return;
 	}
 
-	FRigVMDrawInstruction Instruction(ERigVMDrawSettings::LineStrip, Color, Thickness, WorldOffset);
+	FRigVMDrawInstruction Instruction(ERigVMDrawSettings::LineStrip, Color, Thickness, WorldOffset, DepthPriority, Lifetime);
 	Instruction.Positions.Append(Positions.GetData(), Positions.Num());
 	Instructions.Add(Instruction);
 }
 
-void FRigVMDrawInterface::DrawBox(const FTransform& WorldOffset, const FTransform& Transform, const FLinearColor& Color, float Thickness)
+void FRigVMDrawInterface::DrawBox(const FTransform& WorldOffset, const FTransform& Transform, const FLinearColor& Color, float Thickness, ESceneDepthPriorityGroup DepthPriority, float Lifetime)
 {
 	if (!IsEnabled())
 	{
@@ -88,7 +88,7 @@ void FRigVMDrawInterface::DrawBox(const FTransform& WorldOffset, const FTransfor
 
 	FTransform DrawTransform = Transform * WorldOffset;
 
-	FRigVMDrawInstruction Instruction(ERigVMDrawSettings::Lines, Color, Thickness, DrawTransform);
+	FRigVMDrawInstruction Instruction(ERigVMDrawSettings::Lines, Color, Thickness, DrawTransform, DepthPriority, Lifetime);
 
 	Instruction.Positions.Add(FVector(0.5f, 0.5f, 0.5f));
 	Instruction.Positions.Add(FVector(0.5f, -0.5f, 0.5f));
@@ -120,7 +120,7 @@ void FRigVMDrawInterface::DrawBox(const FTransform& WorldOffset, const FTransfor
 	Instructions.Add(Instruction);
 }
 
-void FRigVMDrawInterface::DrawAxes(const FTransform& WorldOffset, const FTransform& Transform, float Size, float Thickness)
+void FRigVMDrawInterface::DrawAxes(const FTransform& WorldOffset, const FTransform& Transform, float Size, float Thickness, ESceneDepthPriorityGroup DepthPriority, float Lifetime)
 {
 	if (!IsEnabled())
 	{
@@ -132,7 +132,7 @@ void FRigVMDrawInterface::DrawAxes(const FTransform& WorldOffset, const FTransfo
 	DrawLine(WorldOffset, Transform.GetLocation(), Transform.TransformPosition(FVector(0.f, 0.f, Size)), FLinearColor::Blue, Thickness);
 }
 
-void FRigVMDrawInterface::DrawAxes(const FTransform& WorldOffset, const TArrayView<const FTransform>& Transforms, float Size, float Thickness)
+void FRigVMDrawInterface::DrawAxes(const FTransform& WorldOffset, const TArrayView<const FTransform>& Transforms, float Size, float Thickness, ESceneDepthPriorityGroup DepthPriority, float Lifetime)
 {
 	if (!IsEnabled())
 	{
@@ -144,9 +144,9 @@ void FRigVMDrawInterface::DrawAxes(const FTransform& WorldOffset, const TArrayVi
 		return;
 	}
 
-	FRigVMDrawInstruction InstructionX(ERigVMDrawSettings::Lines, FLinearColor::Red, Thickness, WorldOffset);
-	FRigVMDrawInstruction InstructionY(ERigVMDrawSettings::Lines, FLinearColor::Green, Thickness, WorldOffset);
-	FRigVMDrawInstruction InstructionZ(ERigVMDrawSettings::Lines, FLinearColor::Blue, Thickness, WorldOffset);
+	FRigVMDrawInstruction InstructionX(ERigVMDrawSettings::Lines, FLinearColor::Red, Thickness, WorldOffset, DepthPriority, Lifetime);
+	FRigVMDrawInstruction InstructionY(ERigVMDrawSettings::Lines, FLinearColor::Green, Thickness, WorldOffset, DepthPriority, Lifetime);
+	FRigVMDrawInstruction InstructionZ(ERigVMDrawSettings::Lines, FLinearColor::Blue, Thickness, WorldOffset, DepthPriority, Lifetime);
 
 	InstructionX.Positions.Reserve(Transforms.Num() * 2);
 	InstructionY.Positions.Reserve(Transforms.Num() * 2);
@@ -168,14 +168,14 @@ void FRigVMDrawInterface::DrawAxes(const FTransform& WorldOffset, const TArrayVi
 }
 
 void FRigVMDrawInterface::DrawAxes(const FTransform& WorldOffset, const FTransform& Transform,
-	const FLinearColor& InColor, float Size, float Thickness)
+	const FLinearColor& InColor, float Size, float Thickness, ESceneDepthPriorityGroup DepthPriority, float Lifetime)
 {
 	if (!IsEnabled())
 	{
 		return;
 	}
 
-	FRigVMDrawInstruction Instruction(ERigVMDrawSettings::Lines, InColor, Thickness, WorldOffset);
+	FRigVMDrawInstruction Instruction(ERigVMDrawSettings::Lines, InColor, Thickness, WorldOffset, DepthPriority, Lifetime);
 
 	Instruction.Positions.Reserve(6);
 	Instruction.Positions.Add(Transform.GetLocation());
@@ -189,7 +189,7 @@ void FRigVMDrawInterface::DrawAxes(const FTransform& WorldOffset, const FTransfo
 }
 
 void FRigVMDrawInterface::DrawAxes(const FTransform& WorldOffset, const TArrayView<const FTransform>& Transforms,
-	const FLinearColor& InColor, float Size, float Thickness)
+	const FLinearColor& InColor, float Size, float Thickness, ESceneDepthPriorityGroup DepthPriority, float Lifetime)
 {
 	if (!IsEnabled())
 	{
@@ -201,7 +201,7 @@ void FRigVMDrawInterface::DrawAxes(const FTransform& WorldOffset, const TArrayVi
 		return;
 	}
 
-	FRigVMDrawInstruction Instruction(ERigVMDrawSettings::Lines, InColor, Thickness, WorldOffset);
+	FRigVMDrawInstruction Instruction(ERigVMDrawSettings::Lines, InColor, Thickness, WorldOffset, DepthPriority, Lifetime);
 
 	Instruction.Positions.Reserve(Transforms.Num() * 6);
 
@@ -218,7 +218,7 @@ void FRigVMDrawInterface::DrawAxes(const FTransform& WorldOffset, const TArrayVi
 	Instructions.Add(Instruction);
 }
 
-void FRigVMDrawInterface::DrawRectangle(const FTransform& WorldOffset, const FTransform& Transform, float Size, const FLinearColor& Color, float Thickness)
+void FRigVMDrawInterface::DrawRectangle(const FTransform& WorldOffset, const FTransform& Transform, float Size, const FLinearColor& Color, float Thickness, ESceneDepthPriorityGroup DepthPriority, float Lifetime)
 {
 	if (!IsEnabled())
 	{
@@ -227,7 +227,7 @@ void FRigVMDrawInterface::DrawRectangle(const FTransform& WorldOffset, const FTr
 
 	FTransform DrawTransform = Transform * WorldOffset;
 
-	FRigVMDrawInstruction Instruction(ERigVMDrawSettings::LineStrip, Color, Thickness, DrawTransform);
+	FRigVMDrawInstruction Instruction(ERigVMDrawSettings::LineStrip, Color, Thickness, DrawTransform, DepthPriority, Lifetime);
 
 	float Extent = Size * 0.5f;
 	Instruction.Positions.Reserve(5);
@@ -240,7 +240,7 @@ void FRigVMDrawInterface::DrawRectangle(const FTransform& WorldOffset, const FTr
 	Instructions.Add(Instruction);
 }
 
-void FRigVMDrawInterface::DrawArc(const FTransform& WorldOffset, const FTransform& Transform, float Radius, float MinimumAngle, float MaximumAngle, const FLinearColor& Color, float Thickness, int32 Detail)
+void FRigVMDrawInterface::DrawArc(const FTransform& WorldOffset, const FTransform& Transform, float Radius, float MinimumAngle, float MaximumAngle, const FLinearColor& Color, float Thickness, int32 Detail, ESceneDepthPriorityGroup DepthPriority, float Lifetime)
 {
 	if (!IsEnabled())
 	{
@@ -251,7 +251,7 @@ void FRigVMDrawInterface::DrawArc(const FTransform& WorldOffset, const FTransfor
 	
 	FTransform DrawTransform = Transform * WorldOffset;
 
-	FRigVMDrawInstruction Instruction(ERigVMDrawSettings::LineStrip, Color, Thickness, DrawTransform);
+	FRigVMDrawInstruction Instruction(ERigVMDrawSettings::LineStrip, Color, Thickness, DrawTransform, DepthPriority, Lifetime);
 	Instruction.Positions.Reserve(Count);
 
 	FVector V = FVector(Radius, 0.f, 0.f);
@@ -274,13 +274,13 @@ void FRigVMDrawInterface::DrawArc(const FTransform& WorldOffset, const FTransfor
 }
 
 void FRigVMDrawInterface::DrawCircle(const FTransform& WorldOffset, const FTransform& Transform, float Radius,
-	const FLinearColor& Color, float Thickness, int32 Detail)
+	const FLinearColor& Color, float Thickness, int32 Detail, ESceneDepthPriorityGroup DepthPriority, float Lifetime)
 {
 	DrawArc(WorldOffset, Transform, Radius, 0, TWO_PI, Color, Thickness, Detail);
 }
 
 void FRigVMDrawInterface::DrawSphere(const FTransform& WorldOffset, const FTransform& Transform, float Radius, 
-	const FLinearColor& Color, float Thickness, int32 Detail)
+	const FLinearColor& Color, float Thickness, int32 Detail, ESceneDepthPriorityGroup DepthPriority, float Lifetime)
 {
 	// Circles are drawn in the XY plane
 	FTransform R1(FQuat(FVector(1.f, 0.f, 0.f), HALF_PI));
@@ -291,7 +291,7 @@ void FRigVMDrawInterface::DrawSphere(const FTransform& WorldOffset, const FTrans
 }
 
 void FRigVMDrawInterface::DrawHemisphere(const FTransform& WorldOffset, const FTransform& Transform, float Radius,
-	const FLinearColor& Color, float Thickness, int32 Detail)
+	const FLinearColor& Color, float Thickness, int32 Detail, ESceneDepthPriorityGroup DepthPriority, float Lifetime)
 {
 	// Circles are drawn in the XY plane
 	FTransform R1(FQuat(FVector(1.f, 0.f, 0.f), HALF_PI));
@@ -302,7 +302,7 @@ void FRigVMDrawInterface::DrawHemisphere(const FTransform& WorldOffset, const FT
 }
 
 void FRigVMDrawInterface::DrawCapsule(const FTransform& WorldOffset, const FTransform& Transform,
-	float Radius, float Length, const FLinearColor& Color, float Thickness, int32 Detail)
+	float Radius, float Length, const FLinearColor& Color, float Thickness, int32 Detail, ESceneDepthPriorityGroup DepthPriority, float Lifetime)
 {
 	FTransform T1(FVector(0, 0, Length * 0.5));
 	FTransform T2(FQuat(FVector(0.f, 1.f, 0.f), PI), FVector(0, 0, -Length * 0.5));
@@ -325,7 +325,9 @@ void FRigVMDrawInterface::DrawCapsule(const FTransform& WorldOffset, const FTran
 }
 
 
-void FRigVMDrawInterface::DrawCone(const FTransform& WorldOffset, const FTransform& ConeOffset, float Angle1, float Angle2, uint32 NumSides, bool bDrawSideLines, const FLinearColor& SideLineColor, FMaterialRenderProxy* const MaterialRenderProxy, float SideLineThickness)
+void FRigVMDrawInterface::DrawCone(const FTransform& WorldOffset, const FTransform& ConeOffset, float Angle1, float Angle2, uint32 NumSides,
+	bool bDrawSideLines, const FLinearColor& SideLineColor, FMaterialRenderProxy* const MaterialRenderProxy, float SideLineThickness,
+	ESceneDepthPriorityGroup DepthPriority, float Lifetime)
 {
 	if (!IsEnabled())
 	{
@@ -343,6 +345,8 @@ void FRigVMDrawInterface::DrawCone(const FTransform& WorldOffset, const FTransfo
 	MeshInstruction.MeshVerts = MeshVerts;
 	MeshInstruction.MeshIndices = MeshIndices;
 	MeshInstruction.MaterialRenderProxy = MaterialRenderProxy;
+	MeshInstruction.DepthPriority = DepthPriority;
+	MeshInstruction.Lifetime = Lifetime;
 	Instructions.Add(MeshInstruction);
 
 	if (bDrawSideLines)
@@ -362,14 +366,14 @@ void FRigVMDrawInterface::DrawCone(const FTransform& WorldOffset, const FTransfo
 }
 
 void FRigVMDrawInterface::DrawArrow(const FTransform& WorldOffset, const FVector& Direction, const FVector& Side,
-	const FLinearColor& Color, float Thickness)
+	const FLinearColor& Color, float Thickness, ESceneDepthPriorityGroup DepthPriority, float Lifetime)
 {
 	if (!IsEnabled())
 	{
 		return;
 	}
 
-	FRigVMDrawInstruction Instruction(ERigVMDrawSettings::Lines, Color, Thickness, WorldOffset);
+	FRigVMDrawInstruction Instruction(ERigVMDrawSettings::Lines, Color, Thickness, WorldOffset, DepthPriority, Lifetime);
 
 	Instruction.Positions.SetNum(6);
 	Instruction.Positions[0] = FVector::ZeroVector;
@@ -383,7 +387,7 @@ void FRigVMDrawInterface::DrawArrow(const FTransform& WorldOffset, const FVector
 }
 
 void FRigVMDrawInterface::DrawPlane(const FTransform& WorldOffset, const FVector2D& Scale, const FLinearColor& MeshColor,
-	bool bDrawLines, const FLinearColor& LineColor, FMaterialRenderProxy* const MaterialRenderProxy)
+	bool bDrawLines, const FLinearColor& LineColor, FMaterialRenderProxy* const MaterialRenderProxy, ESceneDepthPriorityGroup DepthPriority, float Lifetime)
 {
 	if (!IsEnabled())
 	{
@@ -403,12 +407,14 @@ void FRigVMDrawInterface::DrawPlane(const FTransform& WorldOffset, const FVector
 	};
 	MeshInstruction.MeshIndices = { 0, 1, 2, 2, 3, 0 };
 	MeshInstruction.MaterialRenderProxy = MaterialRenderProxy;
+	MeshInstruction.DepthPriority = DepthPriority;
+	MeshInstruction.Lifetime = Lifetime;
 
 	Instructions.Add(MeshInstruction);
 
 	if (bDrawLines)
 	{
-		FRigVMDrawInstruction LinesInstruction(ERigVMDrawSettings::LineStrip, LineColor, 0.f, WorldOffset);
+		FRigVMDrawInstruction LinesInstruction(ERigVMDrawSettings::LineStrip, LineColor, 0.f, WorldOffset, DepthPriority, Lifetime);
 
 		LinesInstruction.Positions.Reserve(5);
 		LinesInstruction.Positions.Add(FVector(MeshInstruction.MeshVerts[0].Position));
diff --git a/Engine/Plugins/Runtime/RigVM/Source/RigVM/Private/RigVMFunctions/Debug/RigVMFunction_DebugLine.cpp b/Engine/Plugins/Runtime/RigVM/Source/RigVM/Private/RigVMFunctions/Debug/RigVMFunction_DebugLine.cpp
--- a/Engine/Plugins/Runtime/RigVM/Source/RigVM/Private/RigVMFunctions/Debug/RigVMFunction_DebugLine.cpp
+++ b/Engine/Plugins/Runtime/RigVM/Source/RigVM/Private/RigVMFunctions/Debug/RigVMFunction_DebugLine.cpp
@@ -11,5 +11,5 @@ FRigVMFunction_DebugLineNoSpace_Execute()
 		return;
 	}
 
-	ExecuteContext.GetDrawInterface()->DrawLine(WorldOffset, A, B, Color, Thickness);
+	ExecuteContext.GetDrawInterface()->DrawLine(WorldOffset, A, B, Color, Thickness, DebugDrawSettings.DepthPriority, DebugDrawSettings.Lifetime);
 }
diff --git a/Engine/Plugins/Runtime/RigVM/Source/RigVM/Private/RigVMFunctions/Debug/RigVMFunction_DebugLineStrip.cpp b/Engine/Plugins/Runtime/RigVM/Source/RigVM/Private/RigVMFunctions/Debug/RigVMFunction_DebugLineStrip.cpp
--- a/Engine/Plugins/Runtime/RigVM/Source/RigVM/Private/RigVMFunctions/Debug/RigVMFunction_DebugLineStrip.cpp
+++ b/Engine/Plugins/Runtime/RigVM/Source/RigVM/Private/RigVMFunctions/Debug/RigVMFunction_DebugLineStrip.cpp
@@ -11,6 +11,6 @@ FRigVMFunction_DebugLineStripNoSpace_Execute()
 		return;
 	}
 
-	ExecuteContext.GetDrawInterface()->DrawLineStrip(WorldOffset, TArrayView<const FVector>(Points.GetData(), Points.Num()), Color, Thickness);
+	ExecuteContext.GetDrawInterface()->DrawLineStrip(WorldOffset, TArrayView<const FVector>(Points.GetData(), Points.Num()), Color, Thickness, DebugDrawSettings.DepthPriority, DebugDrawSettings.Lifetime);
 }
 
diff --git a/Engine/Plugins/Runtime/RigVM/Source/RigVM/Private/RigVMFunctions/Debug/RigVMFunction_DebugPoint.cpp b/Engine/Plugins/Runtime/RigVM/Source/RigVM/Private/RigVMFunctions/Debug/RigVMFunction_DebugPoint.cpp
--- a/Engine/Plugins/Runtime/RigVM/Source/RigVM/Private/RigVMFunctions/Debug/RigVMFunction_DebugPoint.cpp
+++ b/Engine/Plugins/Runtime/RigVM/Source/RigVM/Private/RigVMFunctions/Debug/RigVMFunction_DebugPoint.cpp
@@ -19,12 +19,12 @@ FRigVMFunction_DebugPoint_Execute()
 	{
 		case ERigUnitDebugPointMode::Point:
 		{
-			ExecuteContext.GetDrawInterface()->DrawPoint(WorldOffset, DrawVector, Scale, Color);
+			ExecuteContext.GetDrawInterface()->DrawPoint(WorldOffset, DrawVector, Scale, Color, DebugDrawSettings.DepthPriority, DebugDrawSettings.Lifetime);
 			break;
 		}
 		case ERigUnitDebugPointMode::Vector:
 		{
-			ExecuteContext.GetDrawInterface()->DrawLine(WorldOffset, Center, DrawVector, Color, Thickness);
+			ExecuteContext.GetDrawInterface()->DrawLine(WorldOffset, Center, DrawVector, Color, Thickness, DebugDrawSettings.DepthPriority, DebugDrawSettings.Lifetime);
 			break;
 		}
 	}
@@ -60,12 +60,12 @@ FRigVMFunction_DebugPointMutable_Execute()
 	{
 		case ERigUnitDebugPointMode::Point:
 		{
-			ExecuteContext.GetDrawInterface()->DrawPoint(WorldOffset, DrawVector, Scale, Color);
+			ExecuteContext.GetDrawInterface()->DrawPoint(WorldOffset, DrawVector, Scale, Color, DebugDrawSettings.DepthPriority, DebugDrawSettings.Lifetime);
 			break;
 		}
 		case ERigUnitDebugPointMode::Vector:
 		{
-			ExecuteContext.GetDrawInterface()->DrawLine(WorldOffset, Center, DrawVector, Color, Thickness);
+			ExecuteContext.GetDrawInterface()->DrawLine(WorldOffset, Center, DrawVector, Color, Thickness, DebugDrawSettings.DepthPriority, DebugDrawSettings.Lifetime);
 			break;
 		}
 	}
diff --git a/Engine/Plugins/Runtime/RigVM/Source/RigVM/Private/RigVMFunctions/Debug/RigVMFunction_DebugPrimitives.cpp b/Engine/Plugins/Runtime/RigVM/Source/RigVM/Private/RigVMFunctions/Debug/RigVMFunction_DebugPrimitives.cpp
--- a/Engine/Plugins/Runtime/RigVM/Source/RigVM/Private/RigVMFunctions/Debug/RigVMFunction_DebugPrimitives.cpp
+++ b/Engine/Plugins/Runtime/RigVM/Source/RigVM/Private/RigVMFunctions/Debug/RigVMFunction_DebugPrimitives.cpp
@@ -13,7 +13,8 @@ FRigVMFunction_DebugRectangle_Execute()
 		Scale,
 		Thickness,
 		WorldOffset, 
-		bEnabled);
+		bEnabled,
+		DebugDrawSettings);
 }
 
 FRigVMStructUpgradeInfo FRigVMFunction_DebugRectangle::GetUpgradeInfo() const
@@ -37,7 +38,7 @@ FRigVMFunction_DebugRectangleNoSpace_Execute()
 		return;
 	}
 
-	ExecuteContext.GetDrawInterface()->DrawRectangle(WorldOffset, Transform, Scale, Color, Thickness);
+	ExecuteContext.GetDrawInterface()->DrawRectangle(WorldOffset, Transform, Scale, Color, Thickness, DebugDrawSettings.DepthPriority, DebugDrawSettings.Lifetime);
 }
 
 FRigVMFunction_DebugArc_Execute()
@@ -52,7 +53,8 @@ FRigVMFunction_DebugArc_Execute()
 		Thickness,
 		Detail,
 		WorldOffset, 
-		bEnabled);
+		bEnabled,
+		DebugDrawSettings);
 }
 
 FRigVMStructUpgradeInfo FRigVMFunction_DebugArc::GetUpgradeInfo() const
@@ -79,7 +81,7 @@ FRigVMFunction_DebugArcNoSpace_Execute()
 		return;
 	}
 
-	ExecuteContext.GetDrawInterface()->DrawArc(WorldOffset, Transform, Radius, FMath::DegreesToRadians(MinimumDegrees), FMath::DegreesToRadians(MaximumDegrees), Color, Thickness, Detail);
+	ExecuteContext.GetDrawInterface()->DrawArc(WorldOffset, Transform, Radius, FMath::DegreesToRadians(MinimumDegrees), FMath::DegreesToRadians(MaximumDegrees), Color, Thickness, Detail, DebugDrawSettings.DepthPriority, DebugDrawSettings.Lifetime);
 }
 
 FRigVMFunction_DebugBoxNoSpace_Execute()
@@ -92,5 +94,5 @@ FRigVMFunction_DebugBoxNoSpace_Execute()
 	FTransform Transform = FTransform::Identity;
 	Transform.SetTranslation(Box.GetCenter());
 	Transform.SetScale3D(Box.GetSize());
-	ExecuteContext.GetDrawInterface()->DrawBox(WorldOffset, Transform, Color, Thickness);
+	ExecuteContext.GetDrawInterface()->DrawBox(WorldOffset, Transform, Color, Thickness, DebugDrawSettings.DepthPriority, DebugDrawSettings.Lifetime);
 }
diff --git a/Engine/Plugins/Runtime/RigVM/Source/RigVM/Private/RigVMFunctions/Debug/RigVMFunction_DebugTransform.cpp b/Engine/Plugins/Runtime/RigVM/Source/RigVM/Private/RigVMFunctions/Debug/RigVMFunction_DebugTransform.cpp
--- a/Engine/Plugins/Runtime/RigVM/Source/RigVM/Private/RigVMFunctions/Debug/RigVMFunction_DebugTransform.cpp
+++ b/Engine/Plugins/Runtime/RigVM/Source/RigVM/Private/RigVMFunctions/Debug/RigVMFunction_DebugTransform.cpp
@@ -18,18 +18,18 @@ FRigVMFunction_DebugTransformMutableNoSpace_Execute()
 	{
 		case ERigUnitDebugTransformMode::Axes:
 		{
-			ExecuteContext.GetDrawInterface()->DrawAxes(WorldOffset, DrawTransform, Scale, Thickness);
+			ExecuteContext.GetDrawInterface()->DrawAxes(WorldOffset, DrawTransform, Scale, Thickness, DebugDrawSettings.DepthPriority, DebugDrawSettings.Lifetime);
 			break;
 		}
 		case ERigUnitDebugTransformMode::Point:
 		{
-			ExecuteContext.GetDrawInterface()->DrawPoint(WorldOffset, DrawTransform.GetLocation(), Scale, Color);
+			ExecuteContext.GetDrawInterface()->DrawPoint(WorldOffset, DrawTransform.GetLocation(), Scale, Color, DebugDrawSettings.DepthPriority, DebugDrawSettings.Lifetime);
 			break;
 		}
 		case ERigUnitDebugTransformMode::Box:
 		{
 			DrawTransform.SetScale3D(DrawTransform.GetScale3D() * Scale);
-			ExecuteContext.GetDrawInterface()->DrawBox(WorldOffset, DrawTransform, Color, Thickness);
+			ExecuteContext.GetDrawInterface()->DrawBox(WorldOffset, DrawTransform, Color, Thickness, DebugDrawSettings.DepthPriority, DebugDrawSettings.Lifetime);
 			break;
 		}
 	}
@@ -52,7 +52,8 @@ FRigVMFunction_DebugTransformArrayMutableNoSpace_Execute()
 			Thickness,
 			Scale,
 			WorldOffset, 
-			bEnabled);
+			bEnabled,
+			DebugDrawSettings);
 	}
 
 	if(ParentIndices.Num() == Transforms.Num())
@@ -66,7 +67,9 @@ FRigVMFunction_DebugTransformArrayMutableNoSpace_Execute()
 					Transforms[Index].GetTranslation(),
 					Transforms[ParentIndices[Index]].GetTranslation(),
 					Color,
-					Thickness
+					Thickness,
+					DebugDrawSettings.DepthPriority,
+					DebugDrawSettings.Lifetime
 				);
 			}
 		}
diff --git a/Engine/Plugins/Runtime/RigVM/Source/RigVM/Private/RigVMFunctions/Debug/RigVMFunction_VisualDebug.cpp b/Engine/Plugins/Runtime/RigVM/Source/RigVM/Private/RigVMFunctions/Debug/RigVMFunction_VisualDebug.cpp
--- a/Engine/Plugins/Runtime/RigVM/Source/RigVM/Private/RigVMFunctions/Debug/RigVMFunction_VisualDebug.cpp
+++ b/Engine/Plugins/Runtime/RigVM/Source/RigVM/Private/RigVMFunctions/Debug/RigVMFunction_VisualDebug.cpp
@@ -6,7 +6,7 @@
 
 FRigVMFunction_VisualDebugVector_Execute()
 {
-	FRigVMFunction_VisualDebugVectorNoSpace::StaticExecute(ExecuteContext, Value, bEnabled, Mode, Color, Thickness, Scale);
+	FRigVMFunction_VisualDebugVectorNoSpace::StaticExecute(ExecuteContext, Value, bEnabled, Mode, Color, Thickness, Scale, DebugDrawSettings);
 }
 
 FRigVMStructUpgradeInfo FRigVMFunction_VisualDebugVector::GetUpgradeInfo() const
@@ -36,12 +36,12 @@ FRigVMFunction_VisualDebugVectorNoSpace_Execute()
 	{
 		case ERigUnitVisualDebugPointMode::Point:
 		{
-			ExecuteContext.GetDrawInterface()->DrawPoint(WorldOffset, Value, Thickness, Color);
+			ExecuteContext.GetDrawInterface()->DrawPoint(WorldOffset, Value, Thickness, Color, DebugDrawSettings.DepthPriority, DebugDrawSettings.Lifetime);
 			break;
 		}
 		case ERigUnitVisualDebugPointMode::Vector:
 		{
-			ExecuteContext.GetDrawInterface()->DrawLine(WorldOffset, FVector::ZeroVector, Value * Scale, Color, Thickness);
+			ExecuteContext.GetDrawInterface()->DrawLine(WorldOffset, FVector::ZeroVector, Value * Scale, Color, Thickness, DebugDrawSettings.DepthPriority, DebugDrawSettings.Lifetime);
 			break;
 		}
 		default:
@@ -53,7 +53,7 @@ FRigVMFunction_VisualDebugVectorNoSpace_Execute()
 
 FRigVMFunction_VisualDebugQuat_Execute()
 {
-	FRigVMFunction_VisualDebugQuatNoSpace::StaticExecute(ExecuteContext, Value, bEnabled, Thickness, Scale);
+	FRigVMFunction_VisualDebugQuatNoSpace::StaticExecute(ExecuteContext, Value, bEnabled, Thickness, Scale, DebugDrawSettings);
 }
 
 FRigVMStructUpgradeInfo FRigVMFunction_VisualDebugQuat::GetUpgradeInfo() const
@@ -74,12 +74,12 @@ FRigVMFunction_VisualDebugQuatNoSpace_Execute()
     FTransform Transform = FTransform::Identity;
     Transform.SetRotation(Value);
 
-	FRigVMFunction_VisualDebugTransformNoSpace::StaticExecute(ExecuteContext, Transform, bEnabled, Thickness, Scale);
+	FRigVMFunction_VisualDebugTransformNoSpace::StaticExecute(ExecuteContext, Transform, bEnabled, Thickness, Scale, DebugDrawSettings);
 }
 
 FRigVMFunction_VisualDebugTransform_Execute()
 {
-	FRigVMFunction_VisualDebugTransformNoSpace::StaticExecute(ExecuteContext, Value, bEnabled, Thickness, Scale);
+	FRigVMFunction_VisualDebugTransformNoSpace::StaticExecute(ExecuteContext, Value, bEnabled, Thickness, Scale, DebugDrawSettings);
 }
 
 FRigVMStructUpgradeInfo FRigVMFunction_VisualDebugTransform::GetUpgradeInfo() const
@@ -104,5 +104,5 @@ FRigVMFunction_VisualDebugTransformNoSpace_Execute()
 
 	static const FTransform WorldOffset = FTransform::Identity;
 
-	ExecuteContext.GetDrawInterface()->DrawAxes(WorldOffset, Value, Scale, Thickness);
+	ExecuteContext.GetDrawInterface()->DrawAxes(WorldOffset, Value, Scale, Thickness, DebugDrawSettings.DepthPriority, DebugDrawSettings.Lifetime);
 }
diff --git a/Engine/Plugins/Runtime/RigVM/Source/RigVM/Public/RigVMCore/RigVMDebugDrawSettings.h b/Engine/Plugins/Runtime/RigVM/Source/RigVM/Public/RigVMCore/RigVMDebugDrawSettings.h
new file mode 100644
--- /dev/null
+++ b/Engine/Plugins/Runtime/RigVM/Source/RigVM/Public/RigVMCore/RigVMDebugDrawSettings.h
@@ -0,0 +1,18 @@
+// Copyright Epic Games, Inc. All Rights Reserved.
+
+#pragma once
+
+#include "Engine/EngineTypes.h"
+#include "RigVMDebugDrawSettings.generated.h"
+
+USTRUCT()
+struct RIGVM_API FRigVMDebugDrawSettings
+{
+	GENERATED_BODY()
+
+	UPROPERTY(meta = (Input, DisplayPriority = MAX_int32), Category = "DebugSettings")
+	TEnumAsByte<ESceneDepthPriorityGroup> DepthPriority = ESceneDepthPriorityGroup::SDPG_Foreground;
+
+	UPROPERTY(meta = (Input, EditConditionHides, EditCondition = "DepthPriority == ESceneDepthPriorityGroup::SDPG_World", DisplayPriority = MAX_int32), Category = "DebugSettings")
+	float Lifetime = -1.f;
+};
diff --git a/Engine/Plugins/Runtime/RigVM/Source/RigVM/Public/RigVMCore/RigVMDrawInstruction.h b/Engine/Plugins/Runtime/RigVM/Source/RigVM/Public/RigVMCore/RigVMDrawInstruction.h
--- a/Engine/Plugins/Runtime/RigVM/Source/RigVM/Public/RigVMCore/RigVMDrawInstruction.h
+++ b/Engine/Plugins/Runtime/RigVM/Source/RigVM/Public/RigVMCore/RigVMDrawInstruction.h
@@ -4,6 +4,7 @@
 
 #include "CoreMinimal.h"
 #include "DynamicMeshBuilder.h"
+#include "Engine/EngineTypes.h"
 #include "RigVMDrawInstruction.generated.h"
 
 class FMaterialRenderProxy;
@@ -26,42 +27,43 @@ struct RIGVM_API FRigVMDrawInstruction
 	GENERATED_BODY()
 
 	UPROPERTY(EditAnywhere, Category = "DrawInstruction")
-	FName Name;
+	FName Name = NAME_None;
 
 	UPROPERTY(EditAnywhere, Category = "DrawInstruction")
-	TEnumAsByte<ERigVMDrawSettings::Primitive> PrimitiveType;
+	TEnumAsByte<ERigVMDrawSettings::Primitive> PrimitiveType = ERigVMDrawSettings::Points;
 
 	UPROPERTY(EditAnywhere, Category = "DrawInstruction")
 	TArray<FVector> Positions;
 
 	UPROPERTY(EditAnywhere, Category = "DrawInstruction")
-	FLinearColor Color;
+	FLinearColor Color = FLinearColor::Red;
 
 	UPROPERTY(EditAnywhere, Category = "DrawInstruction")
-	float Thickness;
+	float Thickness = 0.f;
 
 	UPROPERTY(EditAnywhere, Category = "DrawInstruction")
-	FTransform Transform;
+	FTransform Transform = FTransform::Identity;
+
+	UPROPERTY(EditAnywhere, Category = "DrawInstruction")
+	TEnumAsByte<ESceneDepthPriorityGroup> DepthPriority = ESceneDepthPriorityGroup::SDPG_Foreground;
+
+	UPROPERTY(EditAnywhere, Category = "DrawInstruction")
+	float Lifetime = -1.f;
 
 	// This is to draw cone, they're not UPROPERTY
 	TArray<FDynamicMeshVertex> MeshVerts;
 	TArray<uint32> MeshIndices;
 	FMaterialRenderProxy* MaterialRenderProxy = nullptr;
 
-	FRigVMDrawInstruction()
-		: Name(NAME_None)
-		, PrimitiveType(ERigVMDrawSettings::Points)
-		, Color(FLinearColor::Red)
-		, Thickness(0.f)
-		, Transform(FTransform::Identity)
-	{}
+	FRigVMDrawInstruction() = default;
 
-	FRigVMDrawInstruction(ERigVMDrawSettings::Primitive InPrimitiveType, const FLinearColor& InColor, float InThickness = 0.f, FTransform InTransform = FTransform::Identity)
-		: Name(NAME_None)
-		, PrimitiveType(InPrimitiveType)
+	FRigVMDrawInstruction(ERigVMDrawSettings::Primitive InPrimitiveType, const FLinearColor& InColor, float InThickness = 0.f, FTransform InTransform = FTransform::Identity, ESceneDepthPriorityGroup InDepthPriority = ESceneDepthPriorityGroup::SDPG_Foreground, float InLifetime = -1.f)
+		: PrimitiveType(InPrimitiveType)
 		, Color(InColor)
 		, Thickness(InThickness)
 		, Transform(InTransform)
+		, DepthPriority(InDepthPriority)
+		, Lifetime(InLifetime)
 	{}
 
 	bool IsValid() const
diff --git a/Engine/Plugins/Runtime/RigVM/Source/RigVM/Public/RigVMCore/RigVMDrawInterface.h b/Engine/Plugins/Runtime/RigVM/Source/RigVM/Public/RigVMCore/RigVMDrawInterface.h
--- a/Engine/Plugins/Runtime/RigVM/Source/RigVM/Public/RigVMCore/RigVMDrawInterface.h
+++ b/Engine/Plugins/Runtime/RigVM/Source/RigVM/Public/RigVMCore/RigVMDrawInterface.h
@@ -14,27 +14,27 @@ public:
 	GENERATED_BODY()
 
 	void DrawInstruction(const FRigVMDrawInstruction& InInstruction);
-	void DrawPoint(const FTransform& WorldOffset, const FVector& Position, float Size, const FLinearColor& Color);
-	void DrawPoints(const FTransform& WorldOffset, const TArrayView<const FVector>& Points, float Size, const FLinearColor& Color);
-	void DrawLine(const FTransform& WorldOffset, const FVector& LineStart, const FVector& LineEnd, const FLinearColor& Color, float Thickness = 0.f);
-	void DrawLines(const FTransform& WorldOffset, const TArrayView<const FVector>& Positions, const FLinearColor& Color, float Thickness = 0.f);
-	void DrawLineStrip(const FTransform& WorldOffset, const TArrayView<const FVector>& Positions, const FLinearColor& Color, float Thickness = 0.f);
-	void DrawBox(const FTransform& WorldOffset, const FTransform& Transform, const FLinearColor& Color, float Thickness = 0.f);
-	void DrawSphere(const FTransform& WorldOffset, const FTransform& Transform, float Radius, const FLinearColor& Color, float Thickness = 0.f, int32 Detail = 12);
+	void DrawPoint(const FTransform& WorldOffset, const FVector& Position, float Size, const FLinearColor& Color, ESceneDepthPriorityGroup DepthPriority = ESceneDepthPriorityGroup::SDPG_Foreground, float Lifetime = -1.f);
+	void DrawPoints(const FTransform& WorldOffset, const TArrayView<const FVector>& Points, float Size, const FLinearColor& Color, ESceneDepthPriorityGroup DepthPriority = ESceneDepthPriorityGroup::SDPG_Foreground, float Lifetime = -1.f);
+	void DrawLine(const FTransform& WorldOffset, const FVector& LineStart, const FVector& LineEnd, const FLinearColor& Color, float Thickness = 0.f, ESceneDepthPriorityGroup DepthPriority = ESceneDepthPriorityGroup::SDPG_Foreground, float Lifetime = -1.f);
+	void DrawLines(const FTransform& WorldOffset, const TArrayView<const FVector>& Positions, const FLinearColor& Color, float Thickness = 0.f, ESceneDepthPriorityGroup DepthPriority = ESceneDepthPriorityGroup::SDPG_Foreground, float Lifetime = -1.f);
+	void DrawLineStrip(const FTransform& WorldOffset, const TArrayView<const FVector>& Positions, const FLinearColor& Color, float Thickness = 0.f, ESceneDepthPriorityGroup DepthPriority = ESceneDepthPriorityGroup::SDPG_Foreground, float Lifetime = -1.f);
+	void DrawBox(const FTransform& WorldOffset, const FTransform& Transform, const FLinearColor& Color, float Thickness = 0.f, ESceneDepthPriorityGroup DepthPriority = ESceneDepthPriorityGroup::SDPG_Foreground, float Lifetime = -1.f);
+	void DrawSphere(const FTransform& WorldOffset, const FTransform& Transform, float Radius, const FLinearColor& Color, float Thickness = 0.f, int32 Detail = 12, ESceneDepthPriorityGroup DepthPriority = ESceneDepthPriorityGroup::SDPG_Foreground, float Lifetime = -1.f);
 	// Hemispheres are drawn so that they extend up along +ve Z
-	void DrawHemisphere(const FTransform& WorldOffset, const FTransform& Transform, float Radius, const FLinearColor& Color, float Thickness = 0.f, int32 Detail = 12);
-	void DrawCapsule(const FTransform& WorldOffset, const FTransform& Transform, float Radius, float Length, const FLinearColor& Color, float Thickness = 0.f, int32 Detail = 12);
-	void DrawAxes(const FTransform& WorldOffset, const FTransform& Transform, float Size, float Thickness = 0.f);
-	void DrawAxes(const FTransform& WorldOffset, const TArrayView<const FTransform>& Transforms, float Size, float Thickness = 0.f);
-	void DrawAxes(const FTransform& WorldOffset, const FTransform& Transform, const FLinearColor& InColor, float Size, float Thickness = 0.f);
-	void DrawAxes(const FTransform& WorldOffset, const TArrayView<const FTransform>& Transforms, const FLinearColor& InColor, float Size, float Thickness = 0.f);
-	void DrawRectangle(const FTransform& WorldOffset, const FTransform& Transform, float Size, const FLinearColor& Color, float Thickness);
-	void DrawArc(const FTransform& WorldOffset, const FTransform& Transform, float Radius, float MinimumAngle, float MaximumAngle, const FLinearColor& Color, float Thickness, int32 Detail);
+	void DrawHemisphere(const FTransform& WorldOffset, const FTransform& Transform, float Radius, const FLinearColor& Color, float Thickness = 0.f, int32 Detail = 12, ESceneDepthPriorityGroup DepthPriority = ESceneDepthPriorityGroup::SDPG_Foreground, float Lifetime = -1.f);
+	void DrawCapsule(const FTransform& WorldOffset, const FTransform& Transform, float Radius, float Length, const FLinearColor& Color, float Thickness = 0.f, int32 Detail = 12, ESceneDepthPriorityGroup DepthPriority = ESceneDepthPriorityGroup::SDPG_Foreground, float Lifetime = -1.f);
+	void DrawAxes(const FTransform& WorldOffset, const FTransform& Transform, float Size, float Thickness = 0.f, ESceneDepthPriorityGroup DepthPriority = ESceneDepthPriorityGroup::SDPG_Foreground, float Lifetime = -1.f);
+	void DrawAxes(const FTransform& WorldOffset, const TArrayView<const FTransform>& Transforms, float Size, float Thickness = 0.f, ESceneDepthPriorityGroup DepthPriority = ESceneDepthPriorityGroup::SDPG_Foreground, float Lifetime = -1.f);
+	void DrawAxes(const FTransform& WorldOffset, const FTransform& Transform, const FLinearColor& InColor, float Size, float Thickness = 0.f, ESceneDepthPriorityGroup DepthPriority = ESceneDepthPriorityGroup::SDPG_Foreground, float Lifetime = -1.f);
+	void DrawAxes(const FTransform& WorldOffset, const TArrayView<const FTransform>& Transforms, const FLinearColor& InColor, float Size, float Thickness = 0.f, ESceneDepthPriorityGroup DepthPriority = ESceneDepthPriorityGroup::SDPG_Foreground, float Lifetime = -1.f);
+	void DrawRectangle(const FTransform& WorldOffset, const FTransform& Transform, float Size, const FLinearColor& Color, float Thickness, ESceneDepthPriorityGroup DepthPriority = ESceneDepthPriorityGroup::SDPG_Foreground, float Lifetime = -1.f);
+	void DrawArc(const FTransform& WorldOffset, const FTransform& Transform, float Radius, float MinimumAngle, float MaximumAngle, const FLinearColor& Color, float Thickness, int32 Detail, ESceneDepthPriorityGroup DepthPriority = ESceneDepthPriorityGroup::SDPG_Foreground, float Lifetime = -1.f);
 	// Circles are drawn in the XY plane
-	void DrawCircle(const FTransform& WorldOffset, const FTransform& Transform, float Radius, const FLinearColor& Color, float Thickness, int32 Detail);
-	void DrawCone(const FTransform& WorldOffset, const FTransform& ConeOffset, float Angle1, float Angle2, uint32 NumSides, bool bDrawSideLines, const FLinearColor& SideLineColor, FMaterialRenderProxy* const MaterialRenderProxy, float SideLineThickness = 1.0f);
-	void DrawArrow(const FTransform& WorldOffset, const FVector& Direction, const FVector& Side, const FLinearColor& InColor, float Thickness);
-	void DrawPlane(const FTransform& WorldOffset, const FVector2D& Scale, const FLinearColor& MeshColor, bool bDrawLines, const FLinearColor& LineColor, FMaterialRenderProxy* const MaterialRenderProxy);
+	void DrawCircle(const FTransform& WorldOffset, const FTransform& Transform, float Radius, const FLinearColor& Color, float Thickness, int32 Detail, ESceneDepthPriorityGroup DepthPriority = ESceneDepthPriorityGroup::SDPG_Foreground, float Lifetime = -1.f);
+	void DrawCone(const FTransform& WorldOffset, const FTransform& ConeOffset, float Angle1, float Angle2, uint32 NumSides, bool bDrawSideLines, const FLinearColor& SideLineColor, FMaterialRenderProxy* const MaterialRenderProxy, float SideLineThickness = 1.0f, ESceneDepthPriorityGroup DepthPriority = ESceneDepthPriorityGroup::SDPG_Foreground, float Lifetime = -1.f);
+	void DrawArrow(const FTransform& WorldOffset, const FVector& Direction, const FVector& Side, const FLinearColor& InColor, float Thickness, ESceneDepthPriorityGroup DepthPriority = ESceneDepthPriorityGroup::SDPG_Foreground, float Lifetime = -1.f);
+	void DrawPlane(const FTransform& WorldOffset, const FVector2D& Scale, const FLinearColor& MeshColor, bool bDrawLines, const FLinearColor& LineColor, FMaterialRenderProxy* const MaterialRenderProxy, ESceneDepthPriorityGroup DepthPriority = ESceneDepthPriorityGroup::SDPG_Foreground, float Lifetime = -1.f);
 
 	bool IsEnabled() const;
 };
diff --git a/Engine/Plugins/Runtime/RigVM/Source/RigVM/Public/RigVMFunctions/Debug/RigVMFunction_DebugBase.h b/Engine/Plugins/Runtime/RigVM/Source/RigVM/Public/RigVMFunctions/Debug/RigVMFunction_DebugBase.h
--- a/Engine/Plugins/Runtime/RigVM/Source/RigVM/Public/RigVMFunctions/Debug/RigVMFunction_DebugBase.h
+++ b/Engine/Plugins/Runtime/RigVM/Source/RigVM/Public/RigVMFunctions/Debug/RigVMFunction_DebugBase.h
@@ -3,16 +3,23 @@
 #pragma once
 
 #include "RigVMCore/RigVMStruct.h"
+#include "RigVMCore/RigVMDebugDrawSettings.h"
 #include "RigVMFunction_DebugBase.generated.h"
 
 USTRUCT(meta=(Abstract, Category="Debug", NodeColor = "0.83077 0.846873 0.049707"))
 struct RIGVM_API FRigVMFunction_DebugBase : public FRigVMStruct
 {
 	GENERATED_BODY()
+
+	UPROPERTY(meta = (Input, DetailsOnly, DisplayName = "Draw Settings"))
+	FRigVMDebugDrawSettings DebugDrawSettings;
 };
 
 USTRUCT(meta=(Abstract, Category="Debug", NodeColor = "0.83077 0.846873 0.049707"))
 struct RIGVM_API FRigVMFunction_DebugBaseMutable : public FRigVMStructMutable
 {
 	GENERATED_BODY()
+
+	UPROPERTY(meta = (Input, DetailsOnly, DisplayName = "Draw Settings"))
+	FRigVMDebugDrawSettings DebugDrawSettings;
 };
