Control Rig: Add variables to graph explorer
#rnx

diff --git a/Engine/Plugins/Animation/ControlRig/Source/ControlRigEditor/Private/Editor/ControlRigEditor.cpp b/Engine/Plugins/Animation/ControlRig/Source/ControlRigEditor/Private/Editor/ControlRigEditor.cpp
--- a/Engine/Plugins/Animation/ControlRig/Source/ControlRigEditor/Private/Editor/ControlRigEditor.cpp
+++ b/Engine/Plugins/Animation/ControlRig/Source/ControlRigEditor/Private/Editor/ControlRigEditor.cpp
@@ -5012,7 +5012,7 @@ void FControlRigEditor::CreateRigHierarchyToGraphDragAndDropMenu() const
 	}
 }
 
-void FControlRigEditor::OnGraphNodeDropToPerform(TSharedPtr<FGraphNodeDragDropOp> InDragDropOp, UEdGraph* InGraph, const FVector2D& InNodePosition, const FVector2D& InScreenPosition)
+void FControlRigEditor::OnGraphNodeDropToPerform(TSharedPtr<FDragDropOperation> InDragDropOp, UEdGraph* InGraph, const FVector2D& InNodePosition, const FVector2D& InScreenPosition)
 {
 	DECLARE_SCOPE_HIERARCHICAL_COUNTER_FUNC()
 
diff --git a/Engine/Plugins/Animation/ControlRig/Source/ControlRigEditor/Private/Editor/ControlRigEditor.h b/Engine/Plugins/Animation/ControlRig/Source/ControlRigEditor/Private/Editor/ControlRigEditor.h
--- a/Engine/Plugins/Animation/ControlRig/Source/ControlRigEditor/Private/Editor/ControlRigEditor.h
+++ b/Engine/Plugins/Animation/ControlRig/Source/ControlRigEditor/Private/Editor/ControlRigEditor.h
@@ -170,7 +170,7 @@ public:
 
 	const FName RigHierarchyToGraphDragAndDropMenuName = TEXT("ControlRigEditor.RigHierarchyToGraphDragAndDropMenu");
 	void CreateRigHierarchyToGraphDragAndDropMenu() const;
-	virtual void OnGraphNodeDropToPerform(TSharedPtr<FGraphNodeDragDropOp> InDragDropOp, UEdGraph* InGraph, const FVector2D& InNodePosition, const FVector2D& InScreenPosition) override;
+	virtual void OnGraphNodeDropToPerform(TSharedPtr<FDragDropOperation> InDragDropOp, UEdGraph* InGraph, const FVector2D& InNodePosition, const FVector2D& InScreenPosition) override;
 
 	FPersonaViewportKeyDownDelegate& GetKeyDownDelegate() { return OnKeyDownDelegate; }
 
diff --git a/Engine/Plugins/Runtime/RigVM/Source/RigVMDeveloper/Private/EdGraph/RigVMEdGraphSchema.cpp b/Engine/Plugins/Runtime/RigVM/Source/RigVMDeveloper/Private/EdGraph/RigVMEdGraphSchema.cpp
--- a/Engine/Plugins/Runtime/RigVM/Source/RigVMDeveloper/Private/EdGraph/RigVMEdGraphSchema.cpp
+++ b/Engine/Plugins/Runtime/RigVM/Source/RigVMDeveloper/Private/EdGraph/RigVMEdGraphSchema.cpp
@@ -2051,7 +2051,41 @@ bool URigVMEdGraphSchema::CanVariableBeDropped(UEdGraph* InGraph, FProperty* InV
 	return ExternalVariable.IsValid(true /* allow nullptr */);
 }
 
-bool URigVMEdGraphSchema::RequestVariableDropOnPanel(UEdGraph* InGraph, FProperty* InVariableToDrop, const FVector2D& InDropPosition, const FVector2D& InScreenPosition)
+bool URigVMEdGraphSchema::RequestFunctionDropOnPanel(UEdGraph* InGraph, const FRigVMGraphFunctionIdentifier& InFunction, const FVector2D& InDropPosition, const FVector2D& InScreenPosition) const
+{
+	if (URigVMBlueprint* TargetRigBlueprint = Cast<URigVMBlueprint>(FBlueprintEditorUtils::FindBlueprintForGraph(InGraph)))
+	{
+		if(URigVMController* TargetController = TargetRigBlueprint->GetController(InGraph))
+		{
+			FRigVMGraphFunctionHeader Header = FRigVMGraphFunctionHeader::FindGraphFunctionHeader(InFunction);
+			if (IRigVMClientHost* SourceClientHost = Cast<IRigVMClientHost>(InFunction.HostObject.ResolveObject()))
+			{
+				if (IRigVMGraphFunctionHost* SourceFunctionHost = SourceClientHost->GetRigVMGraphFunctionHost())
+				{
+					if (FRigVMGraphFunctionStore* SourceStore = SourceFunctionHost->GetRigVMGraphFunctionStore())
+					{
+#if WITH_EDITOR
+						if(InFunction.HostObject != TargetRigBlueprint->GetRigVMBlueprintGeneratedClass())
+						{
+							if (!SourceStore->IsFunctionPublic(InFunction))
+							{
+								TargetRigBlueprint->BroadcastRequestLocalizeFunctionDialog(InFunction);
+								URigVMLibraryNode* FunctionDefinitionNode = TargetRigBlueprint->GetLocalFunctionLibrary()->FindPreviouslyLocalizedFunction(InFunction);
+								Header = FunctionDefinitionNode->GetFunctionHeader();
+							}
+						}
+#endif
+					}
+				}
+			}
+			return TargetController->AddFunctionReferenceNodeFromDescription(Header, InDropPosition, FString(), true, true) != nullptr;
+		}
+	}
+	
+	return false;
+}
+
+bool URigVMEdGraphSchema::RequestVariableDropOnPanel(UEdGraph* InGraph, FProperty* InVariableToDrop, const FVector2D& InDropPosition, const FVector2D& InScreenPosition) const
 {
 #if WITH_EDITOR
 	if (CanVariableBeDropped(InGraph, InVariableToDrop))
@@ -2069,7 +2103,7 @@ bool URigVMEdGraphSchema::RequestVariableDropOnPanel(UEdGraph* InGraph, FPropert
 	return false;
 }
 
-bool URigVMEdGraphSchema::RequestVariableDropOnPin(UEdGraph* InGraph, FProperty* InVariableToDrop, UEdGraphPin* InPin, const FVector2D& InDropPosition, const FVector2D& InScreenPosition)
+bool URigVMEdGraphSchema::RequestVariableDropOnPin(UEdGraph* InGraph, FProperty* InVariableToDrop, UEdGraphPin* InPin, const FVector2D& InDropPosition, const FVector2D& InScreenPosition) const
 {
 #if WITH_EDITOR
 	if (CanVariableBeDropped(InGraph, InVariableToDrop))
diff --git a/Engine/Plugins/Runtime/RigVM/Source/RigVMDeveloper/Private/RigVMBlueprint.cpp b/Engine/Plugins/Runtime/RigVM/Source/RigVMDeveloper/Private/RigVMBlueprint.cpp
--- a/Engine/Plugins/Runtime/RigVM/Source/RigVMDeveloper/Private/RigVMBlueprint.cpp
+++ b/Engine/Plugins/Runtime/RigVM/Source/RigVMDeveloper/Private/RigVMBlueprint.cpp
@@ -2691,6 +2691,7 @@ TArray<FRigVMGraphVariableDescription> URigVMBlueprint::GetMemberVariables() con
 		FRigVMGraphVariableDescription NewVariable;
 		NewVariable.Name = BPVariable.VarName;
 		NewVariable.DefaultValue = BPVariable.DefaultValue;
+		NewVariable.Category = BPVariable.Category;
 		FString CPPType;
 		UObject* CPPTypeObject;
 		RigVMTypeUtils::CPPTypeFromPinType(BPVariable.VarType, CPPType, &CPPTypeObject);
diff --git a/Engine/Plugins/Runtime/RigVM/Source/RigVMDeveloper/Public/EdGraph/RigVMEdGraphSchema.h b/Engine/Plugins/Runtime/RigVM/Source/RigVMDeveloper/Public/EdGraph/RigVMEdGraphSchema.h
--- a/Engine/Plugins/Runtime/RigVM/Source/RigVMDeveloper/Public/EdGraph/RigVMEdGraphSchema.h
+++ b/Engine/Plugins/Runtime/RigVM/Source/RigVMDeveloper/Public/EdGraph/RigVMEdGraphSchema.h
@@ -298,8 +298,9 @@ public:
 	virtual bool MarkBlueprintDirtyFromNewNode(UBlueprint* InBlueprint, UEdGraphNode* InEdGraphNode) const override;
 	virtual bool SafeDeleteNodeFromGraph(UEdGraph* Graph, UEdGraphNode* Node) const override;
 	virtual bool CanVariableBeDropped(UEdGraph* InGraph, FProperty* InVariableToDrop) const override;
-	virtual bool RequestVariableDropOnPanel(UEdGraph* InGraph, FProperty* InVariableToDrop, const FVector2D& InDropPosition, const FVector2D& InScreenPosition) override;
-	virtual bool RequestVariableDropOnPin(UEdGraph* InGraph, FProperty* InVariableToDrop, UEdGraphPin* InPin, const FVector2D& InDropPosition, const FVector2D& InScreenPosition) override;
+	virtual bool RequestFunctionDropOnPanel(UEdGraph* InGraph, const FRigVMGraphFunctionIdentifier& InFunction, const FVector2D& InDropPosition, const FVector2D& InScreenPosition) const;
+	virtual bool RequestVariableDropOnPanel(UEdGraph* InGraph, FProperty* InVariableToDrop, const FVector2D& InDropPosition, const FVector2D& InScreenPosition) const override;
+	virtual bool RequestVariableDropOnPin(UEdGraph* InGraph, FProperty* InVariableToDrop, UEdGraphPin* InPin, const FVector2D& InDropPosition, const FVector2D& InScreenPosition) const override;
 	virtual bool IsStructEditable(UStruct* InStruct) const;
 	virtual void SetNodePosition(UEdGraphNode* Node, const FVector2D& Position) const override;
 	void SetNodePosition(UEdGraphNode* Node, const FVector2D& Position, bool bSetupUndo) const;
diff --git a/Engine/Plugins/Runtime/RigVM/Source/RigVMDeveloper/Public/RigVMModel/RigVMVariableDescription.h b/Engine/Plugins/Runtime/RigVM/Source/RigVMDeveloper/Public/RigVMModel/RigVMVariableDescription.h
--- a/Engine/Plugins/Runtime/RigVM/Source/RigVMDeveloper/Public/RigVMModel/RigVMVariableDescription.h
+++ b/Engine/Plugins/Runtime/RigVM/Source/RigVMDeveloper/Public/RigVMModel/RigVMVariableDescription.h
@@ -42,6 +42,10 @@ public:
 	UPROPERTY(EditAnywhere, BlueprintReadOnly, Category = RigVMGraphVariableDescription)
 	FString DefaultValue;
 
+	// The category of the variable
+	UPROPERTY(EditAnywhere, BlueprintReadOnly, Category = RigVMGraphVariableDescription)
+	FText Category;
+
 	// Returns nullptr external variable matching this description
 	FRigVMExternalVariable ToExternalVariable() const
 	{
diff --git a/Engine/Plugins/Runtime/RigVM/Source/RigVMEditor/Private/Editor/RigVMEditor.cpp b/Engine/Plugins/Runtime/RigVM/Source/RigVMEditor/Private/Editor/RigVMEditor.cpp
--- a/Engine/Plugins/Runtime/RigVM/Source/RigVMEditor/Private/Editor/RigVMEditor.cpp
+++ b/Engine/Plugins/Runtime/RigVM/Source/RigVMEditor/Private/Editor/RigVMEditor.cpp
@@ -2160,6 +2160,7 @@ void FRigVMEditor::OnWrappedPropertyChangedChainEvent(URigVMDetailsViewWrapperOb
 					}
 				}
 				RefreshMyBlueprint();
+				GraphExplorerWidget->Refresh();
 				GetRigVMBlueprint()->RequestAutoVMRecompilation();
 			}
 			else if (PropertyPath == TEXT("CPPType") || PropertyPath == TEXT("CPPTypeObject"))
@@ -3628,8 +3629,7 @@ void FRigVMEditor::SwapAssetReferences()
 	SwapFunctionsDialog->ShowNormal();
 }
 
-void FRigVMEditor::OnGraphNodeDropToPerform(TSharedPtr<FGraphNodeDragDropOp> InDragDropOp, UEdGraph* InGraph, const FVector2D& InNodePosition,
-	const FVector2D& InScreenPosition)
+void FRigVMEditor::OnGraphNodeDropToPerform(TSharedPtr<FDragDropOperation> InDragDropOp, UEdGraph* InGraph, const FVector2D& InNodePosition, const FVector2D& InScreenPosition)
 {
 	DECLARE_SCOPE_HIERARCHICAL_COUNTER_FUNC()
 
@@ -3637,32 +3637,27 @@ void FRigVMEditor::OnGraphNodeDropToPerform(TSharedPtr<FGraphNodeDragDropOp> InD
 	{
 		TSharedPtr<FRigVMGraphExplorerDragDropOp> ExplorerOp = StaticCastSharedPtr<FRigVMGraphExplorerDragDropOp>(InDragDropOp);
 
-		if (ExplorerOp->GetElement().Type == Function && FocusedGraphEdPtr.IsValid())
+		if (FocusedGraphEdPtr.IsValid())
 		{
-			if (URigVMBlueprint* TargetRigBlueprint = Cast<URigVMBlueprint>(FBlueprintEditorUtils::FindBlueprintForGraph(InGraph)))
+			const URigVMEdGraphSchema* Schema = CastChecked<URigVMEdGraphSchema>(GetRigVMBlueprint()->GetRigVMEdGraphSchemaClass()->GetDefaultObject());
+			if (ExplorerOp->GetElement().Type == Function)
 			{
-				if (URigVMBlueprint* SourceRigBlueprint = ExplorerOp->GetBlueprint())
+				if (URigVMBlueprint* TargetRigBlueprint = Cast<URigVMBlueprint>(FBlueprintEditorUtils::FindBlueprintForGraph(InGraph)))
 				{
-					if (URigVMLibraryNode* FunctionDefinitionNode = SourceRigBlueprint->GetRigVMClient()->GetFunctionLibrary()->FindFunction(*ExplorerOp->GetElement().Name))
+					if (URigVMLibraryNode* LibraryNode = ExplorerOp->GetBlueprint()->GetLocalFunctionLibrary()->FindFunction(*ExplorerOp->GetElement().Name))
 					{
-						if(URigVMController* TargetController = TargetRigBlueprint->GetController(InGraph))
-						{
-#if WITH_EDITOR
-							if(SourceRigBlueprint != TargetRigBlueprint)
-							{
-								if(!SourceRigBlueprint->IsFunctionPublic(FunctionDefinitionNode->GetFName()))
-								{
-									FRigVMGraphFunctionIdentifier FunctionIdentifier = FunctionDefinitionNode->GetFunctionIdentifier();
-									TargetRigBlueprint->BroadcastRequestLocalizeFunctionDialog(FunctionIdentifier);
-									FunctionDefinitionNode = TargetRigBlueprint->GetLocalFunctionLibrary()->FindPreviouslyLocalizedFunction(FunctionIdentifier);
-								}
-							}
-#endif
-							TargetController->AddFunctionReferenceNode(FunctionDefinitionNode, InNodePosition, FString(), true, true);
-						}
+						Schema->RequestFunctionDropOnPanel(InGraph, LibraryNode->GetFunctionIdentifier(), InNodePosition, InScreenPosition);
 					}
 				}
 			}
+			else if (ExplorerOp->GetElement().Type == Variable)
+			{
+				if (URigVMBlueprint* TargetRigBlueprint = Cast<URigVMBlueprint>(FBlueprintEditorUtils::FindBlueprintForGraph(InGraph)))
+				{
+					FProperty* Property = TargetRigBlueprint->SkeletonGeneratedClass->FindPropertyByName(*ExplorerOp->GetElement().Name);
+					Schema->RequestVariableDropOnPanel(InGraph, Property, InNodePosition, InScreenPosition);
+				}
+			}
 		}
 	}
 }
diff --git a/Engine/Plugins/Runtime/RigVM/Source/RigVMEditor/Private/Widgets/SRigVMEditorGraphExplorer.cpp b/Engine/Plugins/Runtime/RigVM/Source/RigVMEditor/Private/Widgets/SRigVMEditorGraphExplorer.cpp
--- a/Engine/Plugins/Runtime/RigVM/Source/RigVMEditor/Private/Widgets/SRigVMEditorGraphExplorer.cpp
+++ b/Engine/Plugins/Runtime/RigVM/Source/RigVMEditor/Private/Widgets/SRigVMEditorGraphExplorer.cpp
@@ -80,6 +80,38 @@ void FRigVMGraphExplorerDragDropOp::Construct()
 	);
 }
 
+FReply FRigVMGraphExplorerDragDropOp::DroppedOnPanel(const TSharedRef<SWidget>& Panel, FVector2D ScreenPosition, FVector2D GraphPosition, UEdGraph& Graph)
+{
+	if (URigVMEdGraph* TargetRigGraph = Cast<URigVMEdGraph>(&Graph))
+	{
+		if(URigVMBlueprint* Blueprint = TargetRigGraph->GetBlueprint())
+		{
+			if (GEditor->GetEditorSubsystem<UAssetEditorSubsystem>()->OpenEditorForAsset(Blueprint))
+			{
+				FRigVMEditor* Editor = static_cast<FRigVMEditor*>(GEditor->GetEditorSubsystem<UAssetEditorSubsystem>()->FindEditorForAsset(Blueprint, /*bFocusIfOpen =*/true));
+				Editor->OnGraphNodeDropToPerform(SharedThis(this), TargetRigGraph, GraphPosition, ScreenPosition);
+			}
+		}
+	}
+	return FGraphEditorDragDropAction::DroppedOnPanel(Panel, ScreenPosition, GraphPosition, Graph);
+}
+
+FReply FRigVMGraphExplorerDragDropOp::DroppedOnPin(FVector2D ScreenPosition, FVector2D GraphPosition)
+{
+	UEdGraph* TargetGraph = GetHoveredGraph();
+	UEdGraphPin* TargetPin = GetHoveredPin();
+	TObjectPtr<URigVMBlueprint> Blueprint = GetBlueprint();
+	
+	const URigVMEdGraphSchema* EdSchema = GetDefault<URigVMEdGraphSchema>();
+	FProperty* Property = Blueprint->SkeletonGeneratedClass->FindPropertyByName(*GetElement().Name);
+	if (EdSchema->RequestVariableDropOnPin(TargetGraph, Property, TargetPin, GraphPosition, ScreenPosition))
+	{
+		return FReply::Handled();
+	}
+	
+	return FGraphEditorDragDropAction::DroppedOnPin(ScreenPosition, GraphPosition);
+}
+
 FRigVMEditorGraphExplorerCommands::FRigVMEditorGraphExplorerCommands()
 	: TCommands<FRigVMEditorGraphExplorerCommands>(
           TEXT("RigVMEditorGraphExplorer"), NSLOCTEXT("Contexts", "Explorer", "Explorer"),
@@ -94,13 +126,16 @@ void FRigVMEditorGraphExplorerCommands::RegisterCommands()
 	UI_COMMAND(OpenGraphInNewTab, "Open Graph In New Tab", "Opens up this graph in a new tab.", EUserInterfaceActionType::Button, FInputChord());
 	UI_COMMAND(CreateGraph, "Add New Graph", "Create a new graph and show it in the editor.", EUserInterfaceActionType::Button, FInputChord());
 	UI_COMMAND(CreateFunction, "Add New Function", "Create a new function and show it in the editor.", EUserInterfaceActionType::Button, FInputChord());
+	UI_COMMAND(CreateVariable, "Add New Variable", "Create a new member variable.", EUserInterfaceActionType::Button, FInputChord());
 	UI_COMMAND(AddFunctionVariant, "Add Variant (Experimental)", "Creates a new variant of a function.", EUserInterfaceActionType::Button, FInputChord());
 	UI_COMMAND(PasteFunction, "Paste Function", "Pastes the function.", EUserInterfaceActionType::Button, FInputChord());
+	UI_COMMAND(PasteVariable, "Paste Variable", "Pastes the variable.", EUserInterfaceActionType::Button, FInputChord());
 }
 
 
 void SRigVMEditorGraphExplorer::Construct(const FArguments& InArgs, TWeakPtr<FRigVMEditor> InRigVMEditor)
 {
+	bNeedsRefresh = false;
 	RigVMEditor = InRigVMEditor;
 
 	TSharedPtr<FRigVMEditor> Editor = RigVMEditor.Pin();
@@ -126,12 +161,18 @@ SRigVMEditorGraphExplorer::~SRigVMEditorGraphExplorer()
 
 void SRigVMEditorGraphExplorer::Refresh()
 {
+	bNeedsRefresh = false;
 	TreeView->RefreshTreeView(true);
 }
 
 void SRigVMEditorGraphExplorer::Tick(const FGeometry& InAllottedGeometry, const double InCurrentTime, const float InDeltaTime)
 {
 	SCompoundWidget::Tick(InAllottedGeometry, InCurrentTime, InDeltaTime);
+
+	if(bNeedsRefresh)
+	{
+		Refresh();
+	}
 }
 
 void SRigVMEditorGraphExplorer::RegisterCommands()
@@ -162,6 +203,11 @@ void SRigVMEditorGraphExplorer::RegisterCommands()
 			FCanExecuteAction::CreateSP(this, &SRigVMEditorGraphExplorer::CanCreateFunction)
 			);
 
+		ToolKitCommandList->MapAction(FRigVMEditorGraphExplorerCommands::Get().CreateVariable,
+			FExecuteAction::CreateSP(this, &SRigVMEditorGraphExplorer::OnCreateVariable),
+			FCanExecuteAction::CreateSP(this, &SRigVMEditorGraphExplorer::CanCreateVariable)
+			);
+
 		ToolKitCommandList->MapAction(FRigVMEditorGraphExplorerCommands::Get().AddFunctionVariant,
 			FExecuteAction::CreateSP(this, &SRigVMEditorGraphExplorer::OnAddFunctionVariant),
 			FCanExecuteAction::CreateSP(this, &SRigVMEditorGraphExplorer::CanAddFunctionVariant)
@@ -179,6 +225,10 @@ void SRigVMEditorGraphExplorer::RegisterCommands()
 			FExecuteAction::CreateSP(this, &SRigVMEditorGraphExplorer::OnCut),
 			FCanExecuteAction::CreateSP(this, &SRigVMEditorGraphExplorer::CanCut));
 
+		ToolKitCommandList->MapAction(FGenericCommands::Get().Duplicate,
+			FExecuteAction::CreateSP(this, &SRigVMEditorGraphExplorer::OnDuplicate),
+			FCanExecuteAction::CreateSP(this, &SRigVMEditorGraphExplorer::CanDuplicate));
+
 		ToolKitCommandList->MapAction(FGenericCommands::Get().Paste,
 			FExecuteAction::CreateSP(this, &SRigVMEditorGraphExplorer::OnPasteGeneric),
 			FCanExecuteAction(), FIsActionChecked(),
@@ -189,6 +239,11 @@ void SRigVMEditorGraphExplorer::RegisterCommands()
 			FCanExecuteAction(), FIsActionChecked(),
 			FIsActionButtonVisible::CreateSP(this, &SRigVMEditorGraphExplorer::CanPasteFunction));
 
+		ToolKitCommandList->MapAction(FRigVMEditorGraphExplorerCommands::Get().PasteVariable,
+			FExecuteAction::CreateSP(this, &SRigVMEditorGraphExplorer::OnPasteVariable),
+			FCanExecuteAction(), FIsActionChecked(),
+			FIsActionButtonVisible::CreateSP(this, &SRigVMEditorGraphExplorer::CanPasteVariable));
+
 		ToolKitCommandList->MapAction(FGenericCommands::Get().Delete,
 			FExecuteAction::CreateSP(this, &SRigVMEditorGraphExplorer::OnDeleteEntry),
 			FCanExecuteAction::CreateSP(this, &SRigVMEditorGraphExplorer::CanDeleteEntry));
@@ -232,6 +287,7 @@ void SRigVMEditorGraphExplorer::CreateWidgets()
 	Delegates.OnGetChildrenGraphs = FRigVMGraphExplorer_OnGetChildrenGraphs::CreateSP(this, &SRigVMEditorGraphExplorer::GetChildrenGraphs);
 	Delegates.OnGetEventNodesInGraph = FRigVMGraphExplorer_OnGetEventNodesInGraph::CreateSP(this, &SRigVMEditorGraphExplorer::GetEventNodesInGraph);
 	Delegates.OnGetFunctions = FRigVMGraphExplorer_OnGetFunctions::CreateSP(this, &SRigVMEditorGraphExplorer::GetFunctions);
+	Delegates.OnGetVariables = FRigVMGraphExplorer_OnGetVariables::CreateSP(this, &SRigVMEditorGraphExplorer::GetVariables);
 	Delegates.OnGetGraphDisplayName = FRigVMGraphExplorer_OnGetGraphDisplayName::CreateSP(this, &SRigVMEditorGraphExplorer::GetGraphDisplayName);
 	Delegates.OnGetEventDisplayName = FRigVMGraphExplorer_OnGetEventDisplayName::CreateSP(this, &SRigVMEditorGraphExplorer::GetEventDisplayName);
 	Delegates.OnGetGraphIcon = FRigVMGraphExplorer_OnGetGraphIcon::CreateSP(this, &SRigVMEditorGraphExplorer::GetGraphIcon);
@@ -239,19 +295,29 @@ void SRigVMEditorGraphExplorer::CreateWidgets()
 	Delegates.OnGraphClicked = FRigVMGraphExplorer_OnGraphClicked::CreateSP(this, &SRigVMEditorGraphExplorer::OnGraphClicked);
 	Delegates.OnEventClicked = FRigVMGraphExplorer_OnEventClicked::CreateSP(this, &SRigVMEditorGraphExplorer::OnEventClicked);
 	Delegates.OnFunctionClicked = FRigVMGraphExplorer_OnFunctionClicked::CreateSP(this, &SRigVMEditorGraphExplorer::OnFunctionClicked);
+	Delegates.OnVariableClicked = FRigVMGraphExplorer_OnVariableClicked::CreateSP(this, &SRigVMEditorGraphExplorer::OnVariableClicked);
 	Delegates.OnGraphDoubleClicked = FRigVMGraphExplorer_OnGraphDoubleClicked::CreateSP(this, &SRigVMEditorGraphExplorer::OnGraphDoubleClicked);
 	Delegates.OnEventDoubleClicked = FRigVMGraphExplorer_OnEventDoubleClicked::CreateSP(this, &SRigVMEditorGraphExplorer::OnEventDoubleClicked);
 	Delegates.OnFunctionDoubleClicked = FRigVMGraphExplorer_OnFunctionDoubleClicked::CreateSP(this, &SRigVMEditorGraphExplorer::OnFunctionDoubleClicked);
 	Delegates.OnCreateGraph = FRigVMGraphExplorer_OnCreateGraph::CreateSP(this, &SRigVMEditorGraphExplorer::OnCreateGraph);
 	Delegates.OnCreateFunction = FRigVMGraphExplorer_OnCreateFunction::CreateSP(this, &SRigVMEditorGraphExplorer::OnCreateFunction);
+	Delegates.OnCreateVariable = FRigVMGraphExplorer_OnCreateVariable::CreateSP(this, &SRigVMEditorGraphExplorer::OnCreateVariable);
 	Delegates.OnRenameGraph = FRigVMGraphExplorer_OnRenameGraph::CreateSP(this, &SRigVMEditorGraphExplorer::OnRenameGraph);
 	Delegates.OnRenameFunction = FRigVMGraphExplorer_OnRenameFunction::CreateSP(this, &SRigVMEditorGraphExplorer::OnRenameFunction);
 	Delegates.OnCanRenameGraph = FRigVMGraphExplorer_OnCanRenameGraph::CreateSP(this, &SRigVMEditorGraphExplorer::OnCanRenameGraph);
 	Delegates.OnCanRenameFunction = FRigVMGraphExplorer_OnCanRenameFunction::CreateSP(this, &SRigVMEditorGraphExplorer::OnCanRenameFunction);
+	Delegates.OnRenameVariable = FRigVMGraphExplorer_OnRenameVariable::CreateSP(this, &SRigVMEditorGraphExplorer::OnRenameVariable);
+	Delegates.OnCanRenameVariable = FRigVMGraphExplorer_OnCanRenameVariable::CreateSP(this, &SRigVMEditorGraphExplorer::OnCanRenameVariable);
 	Delegates.OnSetFunctionCategory = FRigVMGraphExplorer_OnSetFunctionCategory::CreateSP(this, &SRigVMEditorGraphExplorer::OnSetFunctionCategory);
 	Delegates.OnGetFunctionCategory = FRigVMGraphExplorer_OnGetFunctionCategory::CreateSP(this, &SRigVMEditorGraphExplorer::OnGetFunctionCategory);
+	Delegates.OnSetVariableCategory = FRigVMGraphExplorer_OnSetVariableCategory::CreateSP(this, &SRigVMEditorGraphExplorer::OnSetVariableCategory);
+	Delegates.OnGetVariableCategory = FRigVMGraphExplorer_OnGetVariableCategory::CreateSP(this, &SRigVMEditorGraphExplorer::OnGetVariableCategory);
 	Delegates.OnRequestContextMenu = FRigVMGraphExplorer_OnRequestContextMenu::CreateSP(this, &SRigVMEditorGraphExplorer::OnContextMenuOpening);
 	Delegates.OnDragDetected = FOnDragDetected::CreateSP(this, &SRigVMEditorGraphExplorer::OnDragDetected);
+	Delegates.OnGetVariablePinType = FRigVMGraphExplorer_OnGetVariablePinType::CreateSP(this, &SRigVMEditorGraphExplorer::OnGetVariablePinType);
+	Delegates.OnSetVariablePinType = FRigVMGraphExplorer_OnSetVariablePinType::CreateSP(this, &SRigVMEditorGraphExplorer::OnSetVariablePinType);
+	Delegates.OnIsVariablePublic = FRigVMGraphExplorer_OnIsVariablePublic::CreateSP(this, &SRigVMEditorGraphExplorer::OnIsVariablePublic);
+	Delegates.OnToggleVariablePublic = FRigVMGraphExplorer_OnToggleVariablePublic::CreateSP(this, &SRigVMEditorGraphExplorer::OnToggleVariablePublic);
 	
 	Delegates.OnSelectionChanged = FRigVMGraphExplorer_OnSelectionChanged::CreateSP(this, &SRigVMEditorGraphExplorer::HandleSelectionChanged);
 	
@@ -337,7 +403,10 @@ void SRigVMEditorGraphExplorer::BuildAddNewMenu(FMenuBuilder& MenuBuilder)
 {
 	MenuBuilder.BeginSection("AddNewItem", LOCTEXT("AddOperations", "Add New"));
 	MenuBuilder.AddMenuEntry(FRigVMEditorGraphExplorerCommands::Get().CreateGraph);
+	MenuBuilder.AddMenuEntry(FRigVMEditorGraphExplorerCommands::Get().CreateFunction);
+	MenuBuilder.AddMenuEntry(FRigVMEditorGraphExplorerCommands::Get().CreateVariable);
 	MenuBuilder.AddMenuEntry(FRigVMEditorGraphExplorerCommands::Get().PasteFunction);
+	MenuBuilder.AddMenuEntry(FRigVMEditorGraphExplorerCommands::Get().PasteVariable);
 	
 	MenuBuilder.EndSection();
 }
@@ -445,6 +514,18 @@ TArray<URigVMLibraryNode*> SRigVMEditorGraphExplorer::GetFunctions() const
 	return RigVMClient->GetFunctionLibrary()->GetFunctions();
 }
 
+TArray<FRigVMGraphVariableDescription> SRigVMEditorGraphExplorer::GetVariables() const
+{
+	TArray<FRigVMGraphVariableDescription> Variables;
+	TSharedPtr<FRigVMEditor> Editor = RigVMEditor.Pin();
+	if (!Editor)
+	{
+		return Variables;
+	}
+
+	return Editor->GetRigVMBlueprint()->GetMemberVariables();
+}
+
 FText SRigVMEditorGraphExplorer::GetGraphDisplayName(const FString& InGraphPath) const
 {
 	FText DisplayName;
@@ -594,6 +675,24 @@ void SRigVMEditorGraphExplorer::OnFunctionClicked(const FString& InFunctionPath)
 	}
 }
 
+void SRigVMEditorGraphExplorer::OnVariableClicked(const FString& InVariable)
+{
+	TSharedPtr<FRigVMEditor> Editor = RigVMEditor.Pin();
+	if (!Editor)
+	{
+		return;
+	}
+
+	FRigVMClient* RigVMClient = Editor->GetRigVMBlueprint()->GetRigVMClient();
+	TSharedRef<SKismetInspector> CurrentInspector = Editor->GetInspector();
+	SKismetInspector::FShowDetailsOptions Options(FText::FromString(InVariable));
+	Options.bForceRefresh = true;
+
+	FProperty* Prop = Editor->GetRigVMBlueprint()->SkeletonGeneratedClass->FindPropertyByName(*InVariable);
+	UPropertyWrapper* PropWrap = (Prop ? Prop->GetUPropertyWrapper() : nullptr);
+	CurrentInspector->ShowDetailsForSingleObject(PropWrap, Options);
+}
+
 void SRigVMEditorGraphExplorer::OnGraphDoubleClicked(const FString& InGraphPath)
 {
 	TSharedPtr<FRigVMEditor> Editor = RigVMEditor.Pin();
@@ -694,6 +793,112 @@ FString SRigVMEditorGraphExplorer::OnGetFunctionCategory(const FString& InFuncti
 	return FString();
 }
 
+bool SRigVMEditorGraphExplorer::OnSetVariableCategory(const FString& InVariable, const FString& InCategory)
+{
+	TSharedPtr<FRigVMEditor> Editor = RigVMEditor.Pin();
+	if (!Editor.IsValid())
+	{
+		return false;
+	}
+
+	FBlueprintEditorUtils::SetBlueprintVariableCategory(Editor->GetRigVMBlueprint(), *InVariable, nullptr, FText::FromString(InCategory), true);
+	bNeedsRefresh = true;
+	return true;
+}
+
+FString SRigVMEditorGraphExplorer::OnGetVariableCategory(const FString& InVariable) const
+{
+	TSharedPtr<FRigVMEditor> Editor = RigVMEditor.Pin();
+	if (!Editor.IsValid())
+	{
+		return FString();
+	}
+
+	return FBlueprintEditorUtils::GetBlueprintVariableCategory(Editor->GetRigVMBlueprint(), *InVariable, nullptr).ToString();
+}
+
+FEdGraphPinType SRigVMEditorGraphExplorer::OnGetVariablePinType(const FString& InVariableName)
+{
+	TSharedPtr<FRigVMEditor> Editor = RigVMEditor.Pin();
+	if (!Editor.IsValid())
+	{
+		return FEdGraphPinType();
+	}
+
+	URigVMBlueprint* Blueprint = Editor->GetRigVMBlueprint();
+	for (FRigVMGraphVariableDescription& Variable : Blueprint->GetMemberVariables())
+	{
+		if (Variable.Name == InVariableName)
+		{
+			return Variable.ToPinType();
+		}
+	}
+
+	return FEdGraphPinType();
+}
+
+bool SRigVMEditorGraphExplorer::OnSetVariablePinType(const FString& InVariableName, const FEdGraphPinType& InType)
+{
+	TSharedPtr<FRigVMEditor> Editor = RigVMEditor.Pin();
+	if (!Editor.IsValid())
+	{
+		return false;
+	}
+
+	URigVMBlueprint* Blueprint = Editor->GetRigVMBlueprint();
+	
+	for (FRigVMGraphVariableDescription& Variable : Blueprint->GetMemberVariables())
+	{
+		if (Variable.Name == InVariableName)
+		{
+			FBlueprintEditorUtils::ChangeMemberVariableType(Blueprint, *InVariableName, InType);
+			return true;
+		}
+	}
+
+	return false;
+}
+
+bool SRigVMEditorGraphExplorer::OnIsVariablePublic(const FString& InVariableName) const
+{
+	TSharedPtr<FRigVMEditor> Editor = RigVMEditor.Pin();
+	if (!Editor.IsValid())
+	{
+		return false;
+	}
+
+	URigVMBlueprint* Blueprint = Editor->GetRigVMBlueprint();
+	for (FRigVMGraphVariableDescription& Variable : Blueprint->GetMemberVariables())
+	{
+		if (Variable.Name == InVariableName)
+		{
+			FProperty* Property = Blueprint->SkeletonGeneratedClass->FindPropertyByName(*InVariableName);
+			if (Property)
+			{
+				return Property->HasAnyPropertyFlags(CPF_DisableEditOnInstance) ? true : false;
+			}
+			return true;
+		}
+	}
+
+	return false;
+}
+
+bool SRigVMEditorGraphExplorer::OnToggleVariablePublic(const FString& InVariableName) const
+{
+	TSharedPtr<FRigVMEditor> Editor = RigVMEditor.Pin();
+	if (!Editor.IsValid())
+	{
+		return false;
+	}
+
+	// Toggle the flag on the blueprint's version of the variable description, based on state
+	const bool bVariableIsExposed = OnIsVariablePublic(InVariableName);
+	FBlueprintEditorUtils::SetBlueprintOnlyEditableFlag(Editor->GetRigVMBlueprint(), *InVariableName, !bVariableIsExposed);
+	
+	return true;
+}
+
 TSharedPtr<SWidget> SRigVMEditorGraphExplorer::OnContextMenuOpening()
 {
 	TSharedPtr<FRigVMEditor> Editor = RigVMEditor.Pin();
@@ -720,12 +925,6 @@ TSharedPtr<SWidget> SRigVMEditorGraphExplorer::OnContextMenuOpening()
 				MenuBuilder.AddMenuEntry(FRigVMEditorGraphExplorerCommands::Get().OpenGraph);
 				MenuBuilder.AddMenuEntry(FRigVMEditorGraphExplorerCommands::Get().OpenGraphInNewTab);
 				MenuBuilder.AddMenuEntry(FGenericCommands::Get().Rename);
-				// MenuBuilder.AddSubMenu(
-				// 			LOCTEXT("FindReferences_Label", "Find References"),
-				// 			LOCTEXT("FindReferences_Tooltip", "Options for finding references to class members"),
-				// 			FNewMenuDelegate::CreateStatic(&FGraphEditorCommands::BuildFindReferencesMenu),
-				// 			false,
-				// 			FSlateIcon());
 				MenuBuilder.AddMenuEntry(FGraphEditorCommands::Get().FindReferences);
 				MenuBuilder.AddMenuEntry(FGenericCommands::Get().Cut);
 				MenuBuilder.AddMenuEntry(FGenericCommands::Get().Copy);
@@ -752,6 +951,18 @@ TSharedPtr<SWidget> SRigVMEditorGraphExplorer::OnContextMenuOpening()
 				MenuBuilder.EndSection();
 				break;
 			}
+			case Variable:
+			{
+				MenuBuilder.BeginSection("BasicOperations");
+				MenuBuilder.AddMenuEntry(FGenericCommands::Get().Rename);
+				MenuBuilder.AddMenuEntry(FGraphEditorCommands::Get().FindReferences);
+				MenuBuilder.AddMenuEntry(FGenericCommands::Get().Cut);
+				MenuBuilder.AddMenuEntry(FGenericCommands::Get().Copy);
+				MenuBuilder.AddMenuEntry(FGenericCommands::Get().Duplicate);
+				MenuBuilder.AddMenuEntry(FGenericCommands::Get().Delete);
+				MenuBuilder.EndSection();
+				break;
+			}
 		}
 		
 	}
@@ -771,7 +982,6 @@ FReply SRigVMEditorGraphExplorer::OnDragDetected(const FGeometry& MyGeometry, co
 		if (RigVMEditor.IsValid() && DraggedElements.Num() == 1)
 		{
 			TSharedRef<FRigVMGraphExplorerDragDropOp> DragDropOp = FRigVMGraphExplorerDragDropOp::New(MoveTemp(DraggedElements[0]), RigVMEditor.Pin()->GetRigVMBlueprint());
-			DragDropOp->OnPerformDropToGraph.BindSP(RigVMEditor.Pin().Get(), &FRigVMEditor::OnGraphNodeDropToPerform);
 			return FReply::Handled().BeginDragDrop(DragDropOp);
 		}
 	}
@@ -787,7 +997,12 @@ void SRigVMEditorGraphExplorer::OnCopy()
 		return;
 	}
 	TArray<FRigVMExplorerElementKey> Selection = TreeView->GetSelectedKeys();
-	if (Selection.Num() != 1 || Selection[0].Type != Function)
+	if (Selection.Num() != 1)
+	{
+		return;
+	}
+
+	if (Selection[0].Type != Function && Selection[0].Type != Variable)
 	{
 		return;
 	}
@@ -795,14 +1010,51 @@ void SRigVMEditorGraphExplorer::OnCopy()
 	FString OutputString;
 	const URigVMEdGraphSchema* Schema = GetDefault<URigVMEdGraphSchema>();
 	URigVMBlueprint* Blueprint = Editor->GetRigVMBlueprint();
-	URigVMFunctionLibrary* FunctionLibrary = Blueprint->GetLocalFunctionLibrary();
-	if (const URigVMLibraryNode* FunctionNode = FunctionLibrary->FindFunction(*Selection[0].Name))
+
+	switch (Selection[0].Type)
 	{
-		if (UEdGraph* EdGraph = Editor->GetRigVMBlueprint()->GetEdGraph(FunctionNode->GetContainedGraph()))
+		case Function:
+		{
+			URigVMFunctionLibrary* FunctionLibrary = Blueprint->GetLocalFunctionLibrary();
+			if (const URigVMLibraryNode* FunctionNode = FunctionLibrary->FindFunction(*Selection[0].Name))
+			{
+				if (UEdGraph* EdGraph = Editor->GetRigVMBlueprint()->GetEdGraph(FunctionNode->GetContainedGraph()))
+				{
+					Blueprint->ExportGraphToText(EdGraph, OutputString);
+				}
+			}
+			break;
+		}
+		case Variable:
 		{
-			Blueprint->ExportGraphToText(EdGraph, OutputString);
+			int32 VarIndex = FBlueprintEditorUtils::FindNewVariableIndex(Blueprint, *Selection[0].Name);
+			if (VarIndex != INDEX_NONE)
+			{
+				// make a copy of the Variable description so we can set the default value
+				FBPVariableDescription Description = Blueprint->NewVariables[VarIndex];
+
+				//Grab property of blueprint's current CDO
+				UClass* GeneratedClass = Blueprint->GeneratedClass;
+				UObject* GeneratedCDO = GeneratedClass->GetDefaultObject();
+				FProperty* TargetProperty = FindFProperty<FProperty>(GeneratedClass, Description.VarName);
+
+				if (TargetProperty)
+				{
+					// Grab the address of where the property is actually stored (UObject* base, plus the offset defined in the property)
+					void* OldPropertyAddr = TargetProperty->ContainerPtrToValuePtr<void>(GeneratedCDO);
+					if (OldPropertyAddr)
+					{
+						TargetProperty->ExportTextItem_Direct(Description.DefaultValue, OldPropertyAddr, OldPropertyAddr, nullptr, PPF_SerializedAsImportText);
+					}
+				}
+
+				FBPVariableDescription::StaticStruct()->ExportText(OutputString, &Description, &Description, nullptr, 0, nullptr, false);
+				OutputString = TEXT("BPVar") + OutputString;
+			}
+			break;
 		}
 	}
+	
 
 	if (!OutputString.IsEmpty())
 	{
@@ -813,9 +1065,12 @@ void SRigVMEditorGraphExplorer::OnCopy()
 bool SRigVMEditorGraphExplorer::CanCopy() const
 {
 	TArray<FRigVMExplorerElementKey> Selection = TreeView->GetSelectedKeys();
-	if (Selection.Num() == 1 && Selection[0].Type == Function)
+	if (Selection.Num() == 1)
 	{
-		return true;
+		if (Selection[0].Type == Function || Selection[0].Type == Variable)
+		{
+			return true;
+		}
 	}
 	return false;
 }
@@ -831,13 +1086,46 @@ bool SRigVMEditorGraphExplorer::CanCut() const
 	return CanCopy() && CanDeleteEntry();
 }
 
+void SRigVMEditorGraphExplorer::OnDuplicate()
+{
+	OnCopy();
+	OnPasteGeneric();
+}
+
+bool SRigVMEditorGraphExplorer::CanDuplicate() const
+{
+	TSharedPtr<FRigVMEditor> Editor = RigVMEditor.Pin();
+	if (!Editor)
+	{
+		return false;
+	}
+	
+	TArray<FRigVMExplorerElementKey> Selection = TreeView->GetSelectedKeys();
+	if (Selection.Num() != 1)
+	{
+		return false;
+	}
+
+	const URigVMEdGraphSchema* Schema = GetDefault<URigVMEdGraphSchema>();
+	switch (Selection[0].Type)
+	{
+		case Function:
+		case Variable:
+		{
+			return true;
+		}
+	}
+	
+	return false;
+}
+
 void SRigVMEditorGraphExplorer::OnPasteGeneric()
 {
 	// prioritize pasting as a member variable if possible
-	// if (CanPasteVariable())
-	// {
-	// 	OnPasteVariable();
-	// }
+	if (CanPasteVariable())
+	{
+		OnPasteVariable();
+	}
 	// else if (CanPasteLocalVariable())
 	// {
 	// 	OnPasteLocalVariable();
@@ -851,9 +1139,9 @@ void SRigVMEditorGraphExplorer::OnPasteGeneric()
 
 bool SRigVMEditorGraphExplorer::CanPasteGeneric() const
 {
-	return //CanPasteVariable() 
+	return CanPasteVariable() 
 	//|| CanPasteLocalVariable() 
-	CanPasteFunction() ;
+	|| CanPasteFunction() ;
 }
 
 void SRigVMEditorGraphExplorer::OnPasteFunction()
@@ -884,6 +1172,65 @@ bool SRigVMEditorGraphExplorer::CanPasteFunction() const
 	return Editor->GetRigVMBlueprint()->CanImportGraphFromText(ClipboardText);
 }
 
+void SRigVMEditorGraphExplorer::OnPasteVariable()
+{
+	TSharedPtr<FRigVMEditor> Editor = RigVMEditor.Pin();
+	if (!Editor)
+	{
+		return;
+	}
+	URigVMBlueprint* Blueprint = Editor->GetRigVMBlueprint();
+	
+	FString ClipboardText;
+	FPlatformApplicationMisc::ClipboardPaste(ClipboardText);
+	if (!ensure(ClipboardText.StartsWith(TEXT("BPVar"), ESearchCase::CaseSensitive)))
+	{
+		return;
+	}
+
+	FBPVariableDescription Description;
+	FStringOutputDevice Errors;
+	const TCHAR* Import = ClipboardText.GetCharArray().GetData() + FCString::Strlen(TEXT("BPVar"));
+	FBPVariableDescription::StaticStruct()->ImportText(Import, &Description, nullptr, PPF_None, &Errors, FBPVariableDescription::StaticStruct()->GetName());
+	if (Errors.IsEmpty())
+	{
+		FBPVariableDescription NewVar = FBlueprintEditorUtils::DuplicateVariableDescription(Blueprint, Description);
+		if (NewVar.VarGuid.IsValid())
+		{
+			FScopedTransaction Transaction(FText::Format(LOCTEXT("PasteVariable", "Paste Variable: {0}"), FText::FromName(NewVar.VarName)));
+			Blueprint->Modify();
+			Blueprint->NewVariables.Add(NewVar);
+
+			// Potentially adjust variable names for any child blueprints
+			FBlueprintEditorUtils::ValidateBlueprintChildVariables(Blueprint, NewVar.VarName);
+			FBlueprintEditorUtils::MarkBlueprintAsStructurallyModified(Blueprint);
+		}
+	}
+}
+
+bool SRigVMEditorGraphExplorer::CanPasteVariable() const
+{
+	TSharedPtr<FRigVMEditor> Editor = RigVMEditor.Pin();
+	if (!Editor)
+	{
+		return false;
+	}
+
+	FString ClipboardText;
+	FPlatformApplicationMisc::ClipboardPaste(ClipboardText);
+	if (ClipboardText.StartsWith(TEXT("BPVar"), ESearchCase::CaseSensitive))
+	{
+		FBPVariableDescription Description;
+		FStringOutputDevice Errors;
+		const TCHAR* Import = ClipboardText.GetCharArray().GetData() + FCString::Strlen(TEXT("BPVar"));
+		FBPVariableDescription::StaticStruct()->ImportText(Import, &Description, nullptr, 0, &Errors, FBPVariableDescription::StaticStruct()->GetName());
+
+		return Errors.IsEmpty();
+	}
+
+	return false;
+}
+
 void SRigVMEditorGraphExplorer::OnToggleShowEmptySections()
 {
 	// FIXME: Move to preferences
@@ -967,6 +1314,20 @@ bool SRigVMEditorGraphExplorer::CanCreateFunction() const
 	return true;
 }
 
+void SRigVMEditorGraphExplorer::OnCreateVariable()
+{
+	TSharedPtr<FRigVMEditor> Editor = RigVMEditor.Pin();
+	if (Editor)
+	{
+		Editor->OnAddNewVariable();
+	}
+}
+
+bool SRigVMEditorGraphExplorer::CanCreateVariable() const
+{
+	return true;
+}
+
 void SRigVMEditorGraphExplorer::OnAddFunctionVariant()
 {
 	TSharedPtr<FRigVMEditor> Editor = RigVMEditor.Pin();
@@ -1022,7 +1383,7 @@ void SRigVMEditorGraphExplorer::OnDeleteEntry()
 		{
 			switch (Key.Type)
 			{
-				case ERigVMExplorerElementType::Graph:
+				case Graph:
 				{
 					if (UEdGraph* EdGraph = Editor->GetRigVMBlueprint()->GetEdGraph(Key.Name))
 					{
@@ -1030,7 +1391,7 @@ void SRigVMEditorGraphExplorer::OnDeleteEntry()
 					}
 					break;
 				}
-				case ERigVMExplorerElementType::Function:
+				case Function:
 				{
 					if (URigVMFunctionLibrary* Library = Editor->GetRigVMBlueprint()->GetLocalFunctionLibrary())
 					{
@@ -1044,6 +1405,15 @@ void SRigVMEditorGraphExplorer::OnDeleteEntry()
 					}
 					break;
 				}
+				case Variable:
+				{
+					const FScopedTransaction Transaction( LOCTEXT( "RemoveVariable", "Remove Variable" ) );
+					URigVMBlueprint* Blueprint = Editor->GetRigVMBlueprint();
+
+					Blueprint->Modify();
+					FBlueprintEditorUtils::RemoveMemberVariable(Blueprint, *Key.Name);
+					break;
+				}
 			}
 		}
 	}
@@ -1068,12 +1438,15 @@ bool SRigVMEditorGraphExplorer::CanDeleteEntry() const
 		{
 			case ERigVMExplorerElementType::Section:
 			case ERigVMExplorerElementType::Event:
+			case ERigVMExplorerElementType::FunctionCategory:
+			case ERigVMExplorerElementType::VariableCategory:
 				{
 					return false;
 					break;
 				}
 			case ERigVMExplorerElementType::Graph:
 			case ERigVMExplorerElementType::Function:
+			case ERigVMExplorerElementType::Variable:
 				{
 					break;
 				}
@@ -1202,6 +1575,52 @@ bool SRigVMEditorGraphExplorer::OnCanRenameFunction(const FString& InOldPath, co
 	return bResult;
 }
 
+bool SRigVMEditorGraphExplorer::OnRenameVariable(const FString& InOldName, const FString& InNewName)
+{
+	TSharedPtr<FRigVMEditor> Editor = RigVMEditor.Pin();
+	if (!Editor)
+	{
+		return false;
+	}
+	
+	// Check if the name is unchanged
+	if (InNewName.Equals(InOldName, ESearchCase::CaseSensitive))
+	{
+		return false;
+	}
+
+	const FScopedTransaction Transaction( LOCTEXT( "RenameVariable", "Rename Variable" ) );
+
+	URigVMBlueprint* Blueprint = Editor->GetRigVMBlueprint();
+	Blueprint->Modify();
+
+	FBlueprintEditorUtils::RenameMemberVariable(Blueprint, *InOldName, *InNewName);
+	return true;
+}
+
+bool SRigVMEditorGraphExplorer::OnCanRenameVariable(const FString& InOldName, const FString& InNewName, FText& OutErrorMessage) const
+{
+	if (InOldName == InNewName)
+	{
+		return true;
+	}
+
+	for (FRigVMGraphVariableDescription& Variable : GetVariables())
+	{
+		if (Variable.Name == InOldName)
+		{
+			continue;
+		}
+
+		if (Variable.Name == InNewName)
+		{
+			OutErrorMessage = FText::FromString(TEXT("Name already in use."));
+			return false;
+		}
+	}
+	return true;
+}
+
 void SRigVMEditorGraphExplorer::OnFindReference(bool bSearchAllBlueprints)
 {
 	FString SearchTerm;
@@ -1222,14 +1641,23 @@ void SRigVMEditorGraphExplorer::OnFindReference(bool bSearchAllBlueprints)
 		return;
 	}
 
-	if (Selection[0].Type == ERigVMExplorerElementType::Graph)
+	switch (Selection[0].Type)
 	{
-		bool bSearchTermGenerated = false;
-		if (UEdGraph* EdGraph = Editor->GetRigVMBlueprint()->GetEdGraph(Selection[0].Name))
+		case Graph:
 		{
-			FGraphDisplayInfo DisplayInfo;
-			Schema->GetGraphDisplayInformation(*EdGraph, DisplayInfo);
-			SearchTerm = DisplayInfo.DisplayName.ToString();
+			if (UEdGraph* EdGraph = Editor->GetRigVMBlueprint()->GetEdGraph(Selection[0].Name))
+			{
+				FGraphDisplayInfo DisplayInfo;
+				Schema->GetGraphDisplayInformation(*EdGraph, DisplayInfo);
+				SearchTerm = DisplayInfo.DisplayName.ToString();
+			}
+			break;
+		}
+		case Function:
+		case Variable:
+		{
+			SearchTerm = Selection[0].Name;
+			break;
 		}
 	}
 	
diff --git a/Engine/Plugins/Runtime/RigVM/Source/RigVMEditor/Private/Widgets/SRigVMEditorGraphExplorerTreeView.cpp b/Engine/Plugins/Runtime/RigVM/Source/RigVMEditor/Private/Widgets/SRigVMEditorGraphExplorerTreeView.cpp
--- a/Engine/Plugins/Runtime/RigVM/Source/RigVMEditor/Private/Widgets/SRigVMEditorGraphExplorerTreeView.cpp
+++ b/Engine/Plugins/Runtime/RigVM/Source/RigVMEditor/Private/Widgets/SRigVMEditorGraphExplorerTreeView.cpp
@@ -1,16 +1,20 @@
 // Copyright Epic Games, Inc. All Rights Reserved.
 
 #include "Widgets/SRigVMEditorGraphExplorerTreeView.h"
+
+#include "SPinTypeSelector.h"
 #include "Editor/RigVMEditor.h"
 #include "Widgets/SRigVMEditorGraphExplorer.h"
 
 #include "Widgets/Text/SInlineEditableTextBlock.h"
 #include "Widgets/Text/SRichTextBlock.h"
+#include "Widgets/Input/SCheckBox.h"
 
 #define LOCTEXT_NAMESPACE "SRigVMEditorGraphExplorerTreeView"
 
 static const FText GraphsSectionName = LOCTEXT("Graphs", "Graphs");
 static const FText FunctionsSectionName = LOCTEXT("Functions", "Functions");
+static const FText VariablesSectionName = LOCTEXT("Variables", "Variables");
 
 void SRigVMEditorGraphExplorerItem::Construct(const FArguments& InArgs,
                                               const TSharedRef<STableViewBase>& InOwnerTable,
@@ -23,7 +27,7 @@ void SRigVMEditorGraphExplorerItem::Construct(const FArguments& InArgs,
 
 	FSlateFontInfo NameFont = FCoreStyle::GetDefaultFontStyle("Regular", 10);
 
-	if (InElement->Key.Type == ERigVMExplorerElementType::Section)
+	if (InElement->Key.Type == Section)
 	{
 		STableRow<TSharedPtr<FRigVMEditorGraphExplorerTreeElement>>::Construct(
 		STableRow<TSharedPtr<FRigVMEditorGraphExplorerTreeElement>>::FArguments()
@@ -71,8 +75,8 @@ void SRigVMEditorGraphExplorerItem::Construct(const FArguments& InArgs,
 			
 		], InOwnerTable);
 	}
-	else if(InElement->Key.Type == ERigVMExplorerElementType::FunctionCategory
-		|| InElement->Key.Type == ERigVMExplorerElementType::VariableCategory)
+	else if(InElement->Key.Type == FunctionCategory
+		|| InElement->Key.Type == VariableCategory)
 	{
 		int32 LastSeparator;
 		InElement->Key.Name.FindLastChar('|', LastSeparator);
@@ -98,6 +102,92 @@ void SRigVMEditorGraphExplorerItem::Construct(const FArguments& InArgs,
 			]
 		], InOwnerTable);
 	}
+	else if (InElement->Key.Type == Variable)
+	{
+		// TArray<TSharedPtr<IPinTypeSelectorFilter>> CustomPinTypeFilters;
+		// if (BlueprintEditorPtr.IsValid())
+		// {
+		// 	BlueprintEditorPtr.Pin()->GetPinTypeSelectorFilters(CustomPinTypeFilters);
+		// }
+		
+		STableRow<TSharedPtr<FRigVMEditorGraphExplorerTreeElement>>::Construct(
+		   STableRow<TSharedPtr<FRigVMEditorGraphExplorerTreeElement>>::FArguments()
+		   .ShowWires(false)
+		   .OnDragDetected(Delegates.OnDragDetected)
+		   .Content()
+		   [
+			   // if (ActionAccessSpecifier != EAccessSpecifier::None && GraphAction->GetTypeId() != FEdGraphSchemaAction_K2LocalVar::StaticGetTypeId())
+			   // {
+				  //  CreateAccessSpecifierLambda();
+			   // }
+
+			   SNew(SHorizontalBox)
+			   +SHorizontalBox::Slot()
+				   .FillWidth(0.6f)
+				   .VAlign(VAlign_Center)
+				   .Padding(3.0f, 0.0f)
+				   [
+					   CreateTextSlotWidget(InElement->Key, InTreeView->FilterText)
+				   ]
+
+				+SHorizontalBox::Slot()
+				   .FillWidth(0.4f)
+				   .HAlign(HAlign_Left)
+				   .VAlign(VAlign_Center)
+				   [
+						SNew(SPinTypeSelector, FGetPinTypeTree::CreateUObject(GetDefault<UEdGraphSchema_K2>(), &UEdGraphSchema_K2::GetVariableTypeTree))
+						.Schema(GetDefault<URigVMEdGraphSchema>())
+						.TargetPinType_Lambda([this, InElement]() { return Delegates.GetVariablePinType(InElement->Key.Name); })
+						.OnPinTypeChanged_Lambda([this, InElement](const FEdGraphPinType& InType) { Delegates.SetVariablePinType(InElement->Key.Name, InType); })
+						.TypeTreeFilter(ETypeTreeFilter::None)
+						.SelectorType(SPinTypeSelector::ESelectorType::Partial)
+						//.CustomFilters(CustomPinTypeFilters)
+				   ]
+
+				+SHorizontalBox::Slot()
+				   .AutoWidth()
+				   .Padding(FMargin(6.0f, 0.0f, 3.0f, 0.0f))
+				   .HAlign(HAlign_Right)
+				   .VAlign(VAlign_Center)
+				   [
+					   SNew(SBorder)
+					   .Padding(0.0f)
+					   .BorderImage(FStyleDefaults::GetNoBrush())
+					   [
+						   SNew(SCheckBox)
+						   .ToolTipText_Lambda([this, InElement]()
+						   {
+						   	FText ToolTipText = FText::GetEmpty();
+						   	if (Delegates.IsVariablePublic(InElement->Key.Name))
+							{
+								ToolTipText = LOCTEXT("VariablePrivacy_is_public_Tooltip", "Variable is public and is editable on each instance of this Blueprint.");
+							}
+							else
+							{
+								ToolTipText = LOCTEXT("VariablePrivacy_not_public_Tooltip", "Variable is not public and will not be editable on an instance of this Blueprint.");
+							}
+							return ToolTipText;
+						   })
+						   .OnCheckStateChanged_Lambda([this, InElement](ECheckBoxState) { Delegates.ToggleVariablePublic(InElement->Key.Name); })
+						   .IsChecked_Lambda([this, InElement]() { return Delegates.IsVariablePublic(InElement->Key.Name) ? ECheckBoxState::Checked : ECheckBoxState::Unchecked; })
+						   .Style(FAppStyle::Get(), "TransparentCheckBox")
+						   [
+							   SNew(SImage)
+							   .Image_Lambda([this, InElement]()
+							   {
+								   if (Delegates.IsVariablePublic(InElement->Key.Name))
+								   {
+									   return FAppStyle::GetBrush( "Kismet.VariableList.HideForInstance" );
+								   }
+							   		return FAppStyle::GetBrush( "Kismet.VariableList.ExposeForInstance" );
+							   })
+							   .ColorAndOpacity(FSlateColor::UseForeground())
+						   ]
+					   ]
+				   ]
+
+		   ], InOwnerTable);
+	}
 	else
 	{
 		STableRow<TSharedPtr<FRigVMEditorGraphExplorerTreeElement>>::Construct(
@@ -192,15 +282,16 @@ FText SRigVMEditorGraphExplorerItem::GetDisplayText() const
 
 	switch (Element->Key.Type)
 	{
-		case ERigVMExplorerElementType::Graph:
+		case Graph:
 			{
 				return Delegates.GetGraphDisplayName(Element->Key.Name);
 			}
-		case ERigVMExplorerElementType::Event:
+		case Event:
 			{
 				return Delegates.GetEventDisplayName(Element->Key.Name);
 			}
-		case ERigVMExplorerElementType::Function:
+		case Function:
+		case Variable:
 			{
 				return FText::FromString(Element->Key.Name);
 			}
@@ -220,11 +311,11 @@ FText SRigVMEditorGraphExplorerItem::GetItemTooltip() const
 
 	switch (Element->Key.Type)
 	{
-		case ERigVMExplorerElementType::Graph:
+		case Graph:
 			{
 				return Delegates.GetGraphTooltip(Element->Key.Name);
 			}
-		case ERigVMExplorerElementType::Event:
+		case Event:
 			{
 				return Delegates.GetEventDisplayName(Element->Key.Name);
 			}
@@ -260,6 +351,11 @@ bool SRigVMEditorGraphExplorerItem::OnVerifyNameChanged(const FText& InText, FTe
 				return Delegates.CanRenameFunction(OldKey.Name, *NewName, OutErrorMessage);
 				break;
 			}
+		case Variable:
+			{
+				return Delegates.CanRenameVariable(OldKey.Name, *NewName, OutErrorMessage);
+				break;
+			}
 	}
 	return false;
 }
@@ -285,6 +381,7 @@ void SRigVMEditorGraphExplorerItem::OnNameCommitted(const FText& InText, ETextCo
 					break;
 				}
 			case FunctionCategory:
+			case VariableCategory:
 				{
 					int32 LastSeparator;
 					FString OldCategoryPath = Element->Key.Name;
@@ -303,14 +400,23 @@ void SRigVMEditorGraphExplorerItem::OnNameCommitted(const FText& InText, ETextCo
 					for (int32 i=0; i<SubElements.Num(); ++i)
 					{
 						TSharedPtr<FRigVMEditorGraphExplorerTreeElement> SubElement = SubElements[i];
-						if (SubElement->Key.Type == ERigVMExplorerElementType::Function)
+						if (SubElement->Key.Type == Function || SubElement->Key.Type == Variable)
 						{
-							FString Category = Delegates.GetFunctionCategory(SubElement->Key.Name);
+							FString Category = OldKey.Type == FunctionCategory ?
+								Delegates.GetFunctionCategory(SubElement->Key.Name)
+									: Delegates.GetVariableCategory(SubElement->Key.Name);
 							Category.RemoveFromStart(OldCategoryPath);
 							Category.InsertAt(0, NewCategoryPath);
 							Category.RemoveFromStart(TEXT("|"));
 							Category.RemoveFromEnd(TEXT("|"));
-							Delegates.SetFunctionCategory(SubElement->Key.Name, Category);
+							if (OldKey.Type == FunctionCategory)
+							{
+								Delegates.SetFunctionCategory(SubElement->Key.Name, Category);
+							}
+							else if (OldKey.Type == VariableCategory)
+							{
+								Delegates.SetVariableCategory(SubElement->Key.Name, Category);
+							}
 						}
 						SubElements.Append(SubElement->Children);
 					}
@@ -326,6 +432,11 @@ void SRigVMEditorGraphExplorerItem::OnNameCommitted(const FText& InText, ETextCo
 					Delegates.RenameFunction(OldKey.Name, *NewName);
 					break;
 				}
+			case Variable:
+				{
+					Delegates.RenameVariable(OldKey.Name, *NewName);
+					break;
+				}
 		}
 	}
 }
@@ -358,7 +469,6 @@ void SRigVMEditorGraphExplorerTreeView::Construct(const FArguments& InArgs)
 	STreeView<TSharedPtr<FRigVMEditorGraphExplorerTreeElement>>::FArguments SuperArgs;
 	SuperArgs.SelectionMode(ESelectionMode::Single);
 	SuperArgs.TreeItemsSource(&RootElements);
-	SuperArgs.SelectionMode(ESelectionMode::Multi);
 	SuperArgs.OnGenerateRow(this, &SRigVMEditorGraphExplorerTreeView::MakeTableRowWidget);
 	SuperArgs.OnGetChildren(this, &SRigVMEditorGraphExplorerTreeView::HandleGetChildrenForTree);
 	SuperArgs.HighlightParentNodesForSelection(true);
@@ -400,7 +510,7 @@ void SRigVMEditorGraphExplorerTreeView::RefreshTreeView(bool bRebuildContent)
 		{
 			// GRAPHS section
 			{
-				FRigVMExplorerElementKey SectionKey(ERigVMExplorerElementType::Section, GraphsSectionName.ToString());
+				FRigVMExplorerElementKey SectionKey(Section, GraphsSectionName.ToString());
 				TSharedPtr<FRigVMEditorGraphExplorerTreeElement> SectionItem = MakeShared<FRigVMEditorGraphExplorerTreeElement>(SectionKey, SharedThis(this));
 				ElementMap.Add(SectionKey, SectionItem);
 				RootElements.Add(SectionItem);
@@ -436,7 +546,7 @@ void SRigVMEditorGraphExplorerTreeView::RefreshTreeView(bool bRebuildContent)
 					for (const URigVMNode* EventNode : EventNodes)
 					{
 					
-						FRigVMExplorerElementKey EventKey(ERigVMExplorerElementType::Event, EventNode->GetNodePath());
+						FRigVMExplorerElementKey EventKey(Event, EventNode->GetNodePath());
 						ParentMap.Add(EventKey, Key);
 						TSharedPtr<FRigVMEditorGraphExplorerTreeElement> EventItem = MakeShared<FRigVMEditorGraphExplorerTreeElement>(EventKey, SharedThis(this));
 						ElementMap.Add(EventKey, EventItem);
@@ -452,7 +562,7 @@ void SRigVMEditorGraphExplorerTreeView::RefreshTreeView(bool bRebuildContent)
 
 			// FUNCTIONS section
 			{
-				FRigVMExplorerElementKey SectionKey(ERigVMExplorerElementType::Section, FunctionsSectionName.ToString());
+				FRigVMExplorerElementKey SectionKey(Section, FunctionsSectionName.ToString());
 				TSharedPtr<FRigVMEditorGraphExplorerTreeElement> SectionItem = MakeShared<FRigVMEditorGraphExplorerTreeElement>(SectionKey, SharedThis(this));
 				ElementMap.Add(SectionKey, SectionItem);
 				RootElements.Add(SectionItem);
@@ -473,7 +583,7 @@ void SRigVMEditorGraphExplorerTreeView::RefreshTreeView(bool bRebuildContent)
 							CategoryParentPath.Append(TEXT("|"));
 						}
 						CategoryParentPath.Append(SingleCategory);
-						FRigVMExplorerElementKey ParentKey(ERigVMExplorerElementType::FunctionCategory, CategoryParentPath);
+						FRigVMExplorerElementKey ParentKey(FunctionCategory, CategoryParentPath);
 						TSharedPtr<FRigVMEditorGraphExplorerTreeElement> ParentItem = FindElement(ParentKey);
 						if (!ParentItem)
 						{
@@ -495,6 +605,55 @@ void SRigVMEditorGraphExplorerTreeView::RefreshTreeView(bool bRebuildContent)
 					SetItemExpansion(NewItem, true);
 				}
 			}
+
+			// VARIABLES section
+			{
+				FRigVMExplorerElementKey SectionKey(Section, VariablesSectionName.ToString());
+				TSharedPtr<FRigVMEditorGraphExplorerTreeElement> SectionItem = MakeShared<FRigVMEditorGraphExplorerTreeElement>(SectionKey, SharedThis(this));
+				ElementMap.Add(SectionKey, SectionItem);
+				RootElements.Add(SectionItem);
+				SetItemExpansion(SectionItem, true);
+
+				TArray<FRigVMGraphVariableDescription> Variables = Delegates.GetVariables();
+				for (const FRigVMGraphVariableDescription& Variable : Variables)
+				{
+					const FString Category = Variable.Category.ToString();
+					TSharedPtr<FRigVMEditorGraphExplorerTreeElement> PreviousParentItem = SectionItem;
+					if (!Category.Equals(UEdGraphSchema_K2::VR_DefaultCategory.ToString()))
+					{
+						TArray<FString> SingleCategories;
+						Category.ParseIntoArray(SingleCategories, TEXT("|"));
+						FString CategoryParentPath;
+						for (const FString& SingleCategory : SingleCategories)
+						{
+							if (!CategoryParentPath.IsEmpty())
+							{
+								CategoryParentPath.Append(TEXT("|"));
+							}
+							CategoryParentPath.Append(SingleCategory);
+							FRigVMExplorerElementKey ParentKey(VariableCategory, CategoryParentPath);
+							TSharedPtr<FRigVMEditorGraphExplorerTreeElement> ParentItem = FindElement(ParentKey);
+							if (!ParentItem)
+							{
+								ParentItem = MakeShared<FRigVMEditorGraphExplorerTreeElement>(ParentKey, SharedThis(this));
+								ElementMap.Add(ParentKey, ParentItem);
+								ParentMap.Add(ParentKey, PreviousParentItem->Key);
+								PreviousParentItem->Children.Add(ParentItem);
+								SetItemExpansion(ParentItem, true);
+							}
+							PreviousParentItem = ParentItem;
+						}
+					}
+					
+					FRigVMExplorerElementKey Key(ERigVMExplorerElementType::Variable, Variable.Name.ToString());
+				
+					TSharedPtr<FRigVMEditorGraphExplorerTreeElement> NewItem = MakeShared<FRigVMEditorGraphExplorerTreeElement>(Key, SharedThis(this));
+					ElementMap.Add(Key, NewItem);
+					ParentMap.Add(Key, PreviousParentItem->Key);
+					PreviousParentItem->Children.Add(NewItem);
+					SetItemExpansion(NewItem, true);
+				}
+			}
 		}
 
 		// Filter Items
@@ -519,6 +678,14 @@ void SRigVMEditorGraphExplorerTreeView::RefreshTreeView(bool bRebuildContent)
 								DisplayName = Delegates.GetEventDisplayName(InElement->Key.Name);
 								break;
 							}
+						case Function:
+						case Variable:
+						case FunctionCategory:
+						case VariableCategory:
+							{
+								DisplayName = FText::FromString(InElement->Key.Name);
+								break;
+							}
 						}
 						
 						if (InFilterText.IsEmpty() || DisplayName.ToString().Contains(InFilterText))
@@ -685,7 +852,7 @@ void SRigVMEditorGraphExplorerTreeView::OnItemClicked(TSharedPtr<FRigVMEditorGra
 	{
 		case Section:
 			{
-				FRigVMExplorerElementKey SectionKey(ERigVMExplorerElementType::Section, InElement->Key.Name);
+				FRigVMExplorerElementKey SectionKey(Section, InElement->Key.Name);
 				TSharedPtr<FRigVMEditorGraphExplorerTreeElement> Section = ElementMap.FindChecked(SectionKey);
 				SetItemExpansion(Section, !IsItemExpanded(Section));
 				break;
@@ -698,10 +865,17 @@ void SRigVMEditorGraphExplorerTreeView::OnItemClicked(TSharedPtr<FRigVMEditorGra
 		case Event:
 			{
 				Delegates.EventClicked(InElement->Key.Name);
+				break;
 			}
 		case Function:
 			{
 				Delegates.FunctionClicked(InElement->Key.Name);
+				break;
+			}
+		case Variable:
+			{
+				Delegates.VariableClicked(InElement->Key.Name);
+				break;
 			}
 	}
 }
@@ -711,6 +885,7 @@ void SRigVMEditorGraphExplorerTreeView::OnItemDoubleClicked(TSharedPtr<FRigVMEdi
 	switch(InElement->Key.Type)
 	{
 		case Section:
+		case Variable:
 			{
 				break;
 			}
@@ -742,6 +917,11 @@ FReply SRigVMEditorGraphExplorerTreeView::OnAddButtonClickedOnSection(const FRig
 		Delegates.CreateFunction();
 		return FReply::Handled();
 	}
+	else if (InSectionKey.Name == VariablesSectionName.ToString())
+	{
+		Delegates.CreateVariable();
+		return FReply::Handled();
+	}
 	return FReply::Unhandled();
 }
 
diff --git a/Engine/Plugins/Runtime/RigVM/Source/RigVMEditor/Public/Editor/RigVMEditor.h b/Engine/Plugins/Runtime/RigVM/Source/RigVMEditor/Public/Editor/RigVMEditor.h
--- a/Engine/Plugins/Runtime/RigVM/Source/RigVMEditor/Public/Editor/RigVMEditor.h
+++ b/Engine/Plugins/Runtime/RigVM/Source/RigVMEditor/Public/Editor/RigVMEditor.h
@@ -9,7 +9,7 @@
 #include "Editor/RigVMDetailsViewWrapperObject.h"
 #include "AssetRegistry/AssetData.h"
 
-class FGraphNodeDragDropOp;
+class FRigVMGraphExplorerDragDropOp;
 class SRigVMEditorGraphExplorer;
 class FRigVMEditor;
 
@@ -281,6 +281,8 @@ public:
 	virtual FSlateIcon GetEventQueueIcon(const TArray<FName>& InEventQueue) const;
 	FSlateIcon GetEventQueueIcon() const;
 
+	virtual void OnGraphNodeDropToPerform(TSharedPtr<FDragDropOperation> InDragDropOp, UEdGraph* InGraph, const FVector2D& InNodePosition, const FVector2D& InScreenPosition);
+
 protected:
 
 	void SetExecutionMode(const ERigVMEditorExecutionModeType InExecutionMode);
@@ -317,7 +319,6 @@ protected:
 	virtual void SwapFunctionForAssets(const TArray<FAssetData>& InAssets, bool bSetupUndo);
 	virtual void SwapAssetReferences();
 
-	virtual void OnGraphNodeDropToPerform(TSharedPtr<FGraphNodeDragDropOp> InDragDropOp, UEdGraph* InGraph, const FVector2D& InNodePosition, const FVector2D& InScreenPosition);
 
 	/** Once the log is collected update the graph */
 	void UpdateGraphCompilerErrors();
diff --git a/Engine/Plugins/Runtime/RigVM/Source/RigVMEditor/Public/Widgets/SRigVMEditorGraphExplorer.h b/Engine/Plugins/Runtime/RigVM/Source/RigVMEditor/Public/Widgets/SRigVMEditorGraphExplorer.h
--- a/Engine/Plugins/Runtime/RigVM/Source/RigVMEditor/Public/Widgets/SRigVMEditorGraphExplorer.h
+++ b/Engine/Plugins/Runtime/RigVM/Source/RigVMEditor/Public/Widgets/SRigVMEditorGraphExplorer.h
@@ -3,6 +3,7 @@
 #pragma once
 
 #include "CoreMinimal.h"
+#include "GraphEditorDragDropAction.h"
 #include "DragAndDrop/GraphNodeDragDropOp.h"
 #include "Widgets/DeclarativeSyntaxSupport.h"
 #include "Widgets/SCompoundWidget.h"
@@ -24,10 +25,10 @@ class FMenuBuilder;
 class SSearchBox;
 //
 
-class FRigVMGraphExplorerDragDropOp : public FGraphNodeDragDropOp
+class FRigVMGraphExplorerDragDropOp : public FGraphEditorDragDropAction 
 {
 public:
-	DRAG_DROP_OPERATOR_TYPE(FRigVMExplorerDragDropOp, FGraphNodeDragDropOp)
+	DRAG_DROP_OPERATOR_TYPE(FRigVMExplorerDragDropOp, FGraphEditorDragDropAction)
 
 	static TSharedRef<FRigVMGraphExplorerDragDropOp> New(const FRigVMExplorerElementKey& InElement, TObjectPtr<URigVMBlueprint> InBlueprint);
 
@@ -43,6 +44,9 @@ public:
 
 	virtual void Construct() override;
 
+	virtual FReply DroppedOnPanel(const TSharedRef<SWidget>& Panel, FVector2D ScreenPosition, FVector2D GraphPosition, UEdGraph& Graph) override;
+	virtual FReply DroppedOnPin(FVector2D ScreenPosition, FVector2D GraphPosition) override;
+
 private:
 
 	FRigVMExplorerElementKey Element;
@@ -58,8 +62,10 @@ public:
 	TSharedPtr<FUICommandInfo> OpenGraphInNewTab;
 	TSharedPtr<FUICommandInfo> CreateGraph;
 	TSharedPtr<FUICommandInfo> CreateFunction;
+	TSharedPtr<FUICommandInfo> CreateVariable;
 	TSharedPtr<FUICommandInfo> AddFunctionVariant;
 	TSharedPtr<FUICommandInfo> PasteFunction;
+	TSharedPtr<FUICommandInfo> PasteVariable;
 
 	void RegisterCommands() override;
 };
@@ -81,6 +87,7 @@ public:
 	void Tick(const FGeometry& InAllottedGeometry, const double InCurrentTime, const float InDeltaTime) override;
 
 private:
+	bool bNeedsRefresh;
 	
 	void RegisterCommands();
 	void CreateWidgets();
@@ -93,18 +100,26 @@ private:
 	TArray<const URigVMGraph*> GetChildrenGraphs(const FString& InParentGraphPath) const;
 	TArray<URigVMNode*> GetEventNodesInGraph(const FString& InParentGraphPath) const;
 	TArray<URigVMLibraryNode*> GetFunctions() const;
+	TArray<FRigVMGraphVariableDescription> GetVariables() const;
 	FText GetGraphDisplayName(const FString& InGraphPath) const;
 	FText GetEventDisplayName(const FString& InNodePath) const;
 	const FSlateBrush* GetGraphIcon(const FString& InGraphPath) const;
 	FText GetGraphTooltip(const FString& InGraphPath) const;
 	void OnGraphClicked(const FString& InGraphPath);
 	void OnEventClicked(const FString& InEventPath);
-	void OnFunctionClicked(const FString& InGraphPath);
+	void OnFunctionClicked(const FString& InFunctionPath);
+	void OnVariableClicked(const FString& InVariable);
 	void OnGraphDoubleClicked(const FString& InGraphPath);
 	void OnEventDoubleClicked(const FString& InEventPath);
 	void OnFunctionDoubleClicked(const FString& InFunctionPath);
 	bool OnSetFunctionCategory(const FString& InFunctionPath, const FString& InCategory);
 	FString OnGetFunctionCategory(const FString& InPath) const;
+	bool OnSetVariableCategory(const FString& InVariable, const FString& InCategory);
+	FString OnGetVariableCategory(const FString& InVariable) const;
+	FEdGraphPinType OnGetVariablePinType(const FString& InVariableName);
+	bool OnSetVariablePinType(const FString& InVariableName, const FEdGraphPinType& InType);
+	bool OnIsVariablePublic(const FString& InVariableName) const;
+	bool OnToggleVariablePublic(const FString& InVariableName) const;
 	
 	TSharedPtr<SWidget> OnContextMenuOpening();
 	FReply OnDragDetected(const FGeometry& MyGeometry, const FPointerEvent& MouseEvent);
@@ -113,10 +128,14 @@ private:
 	bool CanCopy() const;
 	void OnCut();
 	bool CanCut() const;
+	void OnDuplicate();
+	bool CanDuplicate() const;
 	void OnPasteGeneric();
 	bool CanPasteGeneric() const;
 	void OnPasteFunction();
 	bool CanPasteFunction() const;
+	void OnPasteVariable();
+	bool CanPasteVariable() const;
 
 
 
@@ -134,6 +153,9 @@ private:
 	void OnCreateFunction();
 	bool CanCreateFunction() const;
 
+	void OnCreateVariable();
+	bool CanCreateVariable() const;
+
 	void OnAddFunctionVariant();
 	bool CanAddFunctionVariant() const;
 	
@@ -146,6 +168,8 @@ private:
 	bool OnCanRenameGraph(const FString& InOldPath, const FString& InNewPath, FText& OutErrorMessage) const;
 	bool OnRenameFunction(const FString& InOldPath, const FString& InNewPath);
 	bool OnCanRenameFunction(const FString& InOldPath, const FString& InNewPath, FText& OutErrorMessage) const;
+	bool OnRenameVariable(const FString& InOldName, const FString& InNewName);
+	bool OnCanRenameVariable(const FString& InOldName, const FString& InNewName, FText& OutErrorMessage) const;
 
 	void OnFindReference(bool bSearchAllBlueprints);
 	bool CanFindReference() const;
diff --git a/Engine/Plugins/Runtime/RigVM/Source/RigVMEditor/Public/Widgets/SRigVMEditorGraphExplorerTreeView.h b/Engine/Plugins/Runtime/RigVM/Source/RigVMEditor/Public/Widgets/SRigVMEditorGraphExplorerTreeView.h
--- a/Engine/Plugins/Runtime/RigVM/Source/RigVMEditor/Public/Widgets/SRigVMEditorGraphExplorerTreeView.h
+++ b/Engine/Plugins/Runtime/RigVM/Source/RigVMEditor/Public/Widgets/SRigVMEditorGraphExplorerTreeView.h
@@ -53,9 +53,11 @@ struct FRigVMExplorerElementKey
 DECLARE_DELEGATE_RetVal(TArray<const URigVMGraph*>, FRigVMGraphExplorer_OnGetRootGraphs);
 DECLARE_DELEGATE( FRigVMGraphExplorer_OnCreateGraph);
 DECLARE_DELEGATE( FRigVMGraphExplorer_OnCreateFunction);
+DECLARE_DELEGATE( FRigVMGraphExplorer_OnCreateVariable);
 DECLARE_DELEGATE_RetVal_OneParam(TArray<const URigVMGraph*>, FRigVMGraphExplorer_OnGetChildrenGraphs, const FString&);
 DECLARE_DELEGATE_RetVal_OneParam(TArray<URigVMNode*>, FRigVMGraphExplorer_OnGetEventNodesInGraph, const FString&);
 DECLARE_DELEGATE_RetVal(TArray<URigVMLibraryNode*>, FRigVMGraphExplorer_OnGetFunctions);
+DECLARE_DELEGATE_RetVal(TArray<FRigVMGraphVariableDescription>, FRigVMGraphExplorer_OnGetVariables);
 DECLARE_DELEGATE_RetVal_OneParam(FText, FRigVMGraphExplorer_OnGetGraphDisplayName, const FString&);
 DECLARE_DELEGATE_RetVal_OneParam(FText, FRigVMGraphExplorer_OnGetEventDisplayName, const FString&);
 DECLARE_DELEGATE_RetVal_OneParam(FText, FRigVMGraphExplorer_OnGetGraphTooltip, const FString&);
@@ -63,16 +65,25 @@ DECLARE_DELEGATE_RetVal_OneParam(const FSlateBrush*, FRigVMGraphExplorer_OnGetGr
 DECLARE_DELEGATE_OneParam(FRigVMGraphExplorer_OnGraphClicked, const FString&);
 DECLARE_DELEGATE_OneParam(FRigVMGraphExplorer_OnEventClicked, const FString&);
 DECLARE_DELEGATE_OneParam(FRigVMGraphExplorer_OnFunctionClicked, const FString&);
+DECLARE_DELEGATE_OneParam(FRigVMGraphExplorer_OnVariableClicked, const FString&);
 DECLARE_DELEGATE_OneParam(FRigVMGraphExplorer_OnGraphDoubleClicked, const FString&);
 DECLARE_DELEGATE_OneParam(FRigVMGraphExplorer_OnEventDoubleClicked, const FString&);
 DECLARE_DELEGATE_OneParam(FRigVMGraphExplorer_OnFunctionDoubleClicked, const FString&);
 DECLARE_DELEGATE_RetVal_TwoParams(bool, FRigVMGraphExplorer_OnRenameGraph, const FString&, const FString&);
 DECLARE_DELEGATE_RetVal_TwoParams(bool, FRigVMGraphExplorer_OnRenameFunction, const FString&, const FString&);
+DECLARE_DELEGATE_RetVal_TwoParams(bool, FRigVMGraphExplorer_OnRenameVariable, const FString&, const FString&);
 DECLARE_DELEGATE_RetVal_ThreeParams(bool, FRigVMGraphExplorer_OnCanRenameGraph, const FString&, const FString&, FText&);
 DECLARE_DELEGATE_RetVal_ThreeParams(bool, FRigVMGraphExplorer_OnCanRenameFunction, const FString&, const FString&, FText&);
+DECLARE_DELEGATE_RetVal_ThreeParams(bool, FRigVMGraphExplorer_OnCanRenameVariable, const FString&, const FString&, FText&);
 DECLARE_DELEGATE_RetVal_TwoParams(bool, FRigVMGraphExplorer_OnSetFunctionCategory, const FString&, const FString&);
 DECLARE_DELEGATE_RetVal_OneParam(FString, FRigVMGraphExplorer_OnGetFunctionCategory, const FString&);
+DECLARE_DELEGATE_RetVal_TwoParams(bool, FRigVMGraphExplorer_OnSetVariableCategory, const FString&, const FString&);
+DECLARE_DELEGATE_RetVal_OneParam(FString, FRigVMGraphExplorer_OnGetVariableCategory, const FString&);
 DECLARE_DELEGATE_RetVal(TSharedPtr<SWidget>, FRigVMGraphExplorer_OnRequestContextMenu);
+DECLARE_DELEGATE_RetVal_OneParam(FEdGraphPinType, FRigVMGraphExplorer_OnGetVariablePinType, const FString&);
+DECLARE_DELEGATE_RetVal_TwoParams(bool, FRigVMGraphExplorer_OnSetVariablePinType, const FString&, const FEdGraphPinType&);
+DECLARE_DELEGATE_RetVal_OneParam(bool, FRigVMGraphExplorer_OnIsVariablePublic, const FString&);
+DECLARE_DELEGATE_RetVal_OneParam(bool, FRigVMGraphExplorer_OnToggleVariablePublic, const FString&);
 
 typedef STreeView<TSharedPtr<FRigVMEditorGraphExplorerTreeElement>>::FOnSelectionChanged FRigVMGraphExplorer_OnSelectionChanged;
 
@@ -83,6 +94,7 @@ public:
 	FRigVMGraphExplorer_OnGetChildrenGraphs OnGetChildrenGraphs;
 	FRigVMGraphExplorer_OnGetEventNodesInGraph OnGetEventNodesInGraph;
 	FRigVMGraphExplorer_OnGetFunctions OnGetFunctions;
+	FRigVMGraphExplorer_OnGetVariables OnGetVariables;
 	FRigVMGraphExplorer_OnGetGraphDisplayName OnGetGraphDisplayName;
 	FRigVMGraphExplorer_OnGetGraphDisplayName OnGetEventDisplayName;
 	FRigVMGraphExplorer_OnGetGraphTooltip OnGetGraphTooltip;
@@ -90,17 +102,27 @@ public:
 	FRigVMGraphExplorer_OnGraphClicked OnGraphClicked;
 	FRigVMGraphExplorer_OnEventClicked OnEventClicked;
 	FRigVMGraphExplorer_OnFunctionClicked OnFunctionClicked;
+	FRigVMGraphExplorer_OnVariableClicked OnVariableClicked;
 	FRigVMGraphExplorer_OnGraphDoubleClicked OnGraphDoubleClicked;
 	FRigVMGraphExplorer_OnEventDoubleClicked OnEventDoubleClicked;
 	FRigVMGraphExplorer_OnFunctionDoubleClicked OnFunctionDoubleClicked;
 	FRigVMGraphExplorer_OnCreateGraph OnCreateGraph;
 	FRigVMGraphExplorer_OnCreateFunction OnCreateFunction;
+	FRigVMGraphExplorer_OnCreateVariable OnCreateVariable;
 	FRigVMGraphExplorer_OnRenameGraph OnRenameGraph;
 	FRigVMGraphExplorer_OnCanRenameGraph OnCanRenameGraph;
 	FRigVMGraphExplorer_OnRenameFunction OnRenameFunction;
 	FRigVMGraphExplorer_OnCanRenameFunction OnCanRenameFunction;
+	FRigVMGraphExplorer_OnRenameVariable OnRenameVariable;
+	FRigVMGraphExplorer_OnCanRenameVariable OnCanRenameVariable;
 	FRigVMGraphExplorer_OnSetFunctionCategory OnSetFunctionCategory;
 	FRigVMGraphExplorer_OnGetFunctionCategory OnGetFunctionCategory;
+	FRigVMGraphExplorer_OnSetVariableCategory OnSetVariableCategory;
+	FRigVMGraphExplorer_OnGetVariableCategory OnGetVariableCategory;
+	FRigVMGraphExplorer_OnGetVariablePinType OnGetVariablePinType;
+	FRigVMGraphExplorer_OnSetVariablePinType OnSetVariablePinType;
+	FRigVMGraphExplorer_OnIsVariablePublic OnIsVariablePublic;
+	FRigVMGraphExplorer_OnToggleVariablePublic OnToggleVariablePublic;
 	
 	FRigVMGraphExplorer_OnSelectionChanged OnSelectionChanged;
 	FRigVMGraphExplorer_OnRequestContextMenu OnRequestContextMenu;
@@ -146,6 +168,15 @@ public:
 		return TArray<URigVMLibraryNode*>();
 	}
 
+	TArray<FRigVMGraphVariableDescription> GetVariables() const
+	{
+		if (OnGetVariables.IsBound())
+		{
+			return OnGetVariables.Execute();
+		}
+		return TArray<FRigVMGraphVariableDescription>();
+	}
+
 	FText GetGraphDisplayName(const FString& InPath) const
 	{
 		if (OnGetGraphDisplayName.IsBound())
@@ -206,6 +237,14 @@ public:
 		}
 	}
 
+	void VariableClicked(const FString& InPath) const
+	{
+		if (OnVariableClicked.IsBound())
+		{
+			OnVariableClicked.Execute(InPath);
+		}
+	}
+
 	void GraphDoubleClicked(const FString& InPath) const
 	{
 		if (OnGraphDoubleClicked.IsBound())
@@ -246,6 +285,14 @@ public:
 		}
 	}
 
+	void CreateVariable() const
+	{
+		if (OnCreateVariable.IsBound())
+		{
+			OnCreateVariable.Execute();
+		}
+	}
+
 	void RenameGraph(const FString& InOldPath, const FString& InNewPath) const
 	{
 		if (OnRenameGraph.IsBound())
@@ -280,6 +327,23 @@ public:
 		return false;
 	}
 
+	void RenameVariable(const FString& InOldName, const FString& InNewName) const
+	{
+		if (OnRenameVariable.IsBound())
+		{
+			OnRenameVariable.Execute(InOldName, InNewName);
+		}
+	}
+
+	bool CanRenameVariable(const FString& InOldName, const FString& InNewName, FText& OutErrorMessage) const
+	{
+		if (OnCanRenameVariable.IsBound())
+		{
+			return OnCanRenameVariable.Execute(InOldName, InNewName, OutErrorMessage);
+		}
+		return false;
+	}
+
 	void SetFunctionCategory(const FString& InPath, const FString& InCategory) const
 	{
 		if (OnSetFunctionCategory.IsBound())
@@ -297,6 +361,23 @@ public:
 		return FString(); 
 	}
 
+	void SetVariableCategory(const FString& InPath, const FString& InCategory) const
+	{
+		if (OnSetVariableCategory.IsBound())
+		{
+			OnSetVariableCategory.Execute(InPath, InCategory);
+		}
+	}
+
+	FString GetVariableCategory(const FString& InPath) const
+	{
+		if (OnGetVariableCategory.IsBound())
+		{
+			return OnGetVariableCategory.Execute(InPath);
+		}
+		return FString(); 
+	}
+
 	void HandleSelectionChanged(TSharedPtr<FRigVMEditorGraphExplorerTreeElement> Selection, ESelectInfo::Type SelectInfo)
 	{
 		if(bSuspendSelectionDelegate)
@@ -316,6 +397,42 @@ public:
 		return nullptr;
 	}
 
+	FEdGraphPinType GetVariablePinType(const FString& InVariable)
+	{
+		if (OnGetVariablePinType.IsBound())
+		{
+			return OnGetVariablePinType.Execute(InVariable);
+		}
+		return FEdGraphPinType();
+	}
+
+	bool SetVariablePinType(const FString& InVariable, const FEdGraphPinType& InType)
+	{
+		if (OnSetVariablePinType.IsBound())
+		{
+			return OnSetVariablePinType.Execute(InVariable, InType);
+		}
+		return false;
+	}
+
+	bool IsVariablePublic(const FString& InVariable)
+	{
+		if (OnIsVariablePublic.IsBound())
+		{
+			return OnIsVariablePublic.Execute(InVariable);
+		}
+		return false;
+	}
+
+	bool ToggleVariablePublic(const FString& InVariable)
+	{
+		if (OnToggleVariablePublic.IsBound())
+		{
+			return OnToggleVariablePublic.Execute(InVariable);
+		}
+		return false;
+	}
+
 private:
 
 	bool bSuspendSelectionDelegate = false;
diff --git a/Engine/Source/Editor/GraphEditor/Public/DragAndDrop/GraphNodeDragDropOp.h b/Engine/Source/Editor/GraphEditor/Public/DragAndDrop/GraphNodeDragDropOp.h
--- a/Engine/Source/Editor/GraphEditor/Public/DragAndDrop/GraphNodeDragDropOp.h
+++ b/Engine/Source/Editor/GraphEditor/Public/DragAndDrop/GraphNodeDragDropOp.h
@@ -9,7 +9,7 @@ class FGraphNodeDragDropOp : public FDecoratedDragDropOp
 public:
 	DRAG_DROP_OPERATOR_TYPE(FGraphNodeDragDropOp, FDecoratedDragDropOp)
 
-	DECLARE_DELEGATE_FourParams( FOnPerformDropToGraph, TSharedPtr<FGraphNodeDragDropOp>, class UEdGraph*, const FVector2D&, const FVector2D&);
+	DECLARE_DELEGATE_FourParams( FOnPerformDropToGraph, TSharedPtr<FDragDropOperation>, class UEdGraph*, const FVector2D&, const FVector2D&);
 
 	FOnPerformDropToGraph OnPerformDropToGraph;
 };
diff --git a/Engine/Source/Runtime/Engine/Classes/EdGraph/EdGraphSchema.h b/Engine/Source/Runtime/Engine/Classes/EdGraph/EdGraphSchema.h
--- a/Engine/Source/Runtime/Engine/Classes/EdGraph/EdGraphSchema.h
+++ b/Engine/Source/Runtime/Engine/Classes/EdGraph/EdGraphSchema.h
@@ -1276,7 +1276,7 @@ class UEdGraphSchema : public UObject
 	 * @param	InScreenPosition The position inside of the screen
 	 * @return the response to rejection or performing the drop
 	 */
-	virtual bool RequestVariableDropOnPanel(UEdGraph* InGraph, FProperty* InVariableToDrop, const FVector2D& InDropPosition, const FVector2D& InScreenPosition) { return false; }
+	virtual bool RequestVariableDropOnPanel(UEdGraph* InGraph, FProperty* InVariableToDrop, const FVector2D& InDropPosition, const FVector2D& InScreenPosition) const { return false; }
 
 	/**
 	 * Request to drop a variable on a node
@@ -1287,7 +1287,7 @@ class UEdGraphSchema : public UObject
 	 * @param	InScreenPosition The position inside of the screen
 	 * @return the response to rejection or performing the drop
 	 */
-	virtual bool RequestVariableDropOnNode(UEdGraph* InGraph, FProperty* InVariableToDrop, UEdGraphNode* InNode, const FVector2D& InDropPosition, const FVector2D& InScreenPosition) { return false; }
+	virtual bool RequestVariableDropOnNode(UEdGraph* InGraph, FProperty* InVariableToDrop, UEdGraphNode* InNode, const FVector2D& InDropPosition, const FVector2D& InScreenPosition) const { return false; }
 
 	/**
 	 * Request to drop a variable on a pin
@@ -1298,7 +1298,7 @@ class UEdGraphSchema : public UObject
 	 * @param	InScreenPosition The position inside of the screen
 	 * @return the response to rejection or performing the drop
 	 */
-	virtual bool RequestVariableDropOnPin(UEdGraph* InGraph, FProperty* InVariableToDrop, UEdGraphPin* InPin, const FVector2D& InDropPosition, const FVector2D& InScreenPosition) { return false; }
+	virtual bool RequestVariableDropOnPin(UEdGraph* InGraph, FProperty* InVariableToDrop, UEdGraphPin* InPin, const FVector2D& InDropPosition, const FVector2D& InScreenPosition) const { return false; }
 
 	/**
 	 * Returns true if the types and directions of two pins are schema compatible. Handles
