Control Rig: Fix crash on unit test Import_Skeletal_Mesh_To_Existing_Control_Rig
#rnx

diff --git a/Engine/Plugins/Runtime/RigVM/Source/RigVMEditor/Private/Editor/RigVMEditor.cpp b/Engine/Plugins/Runtime/RigVM/Source/RigVMEditor/Private/Editor/RigVMEditor.cpp
--- a/Engine/Plugins/Runtime/RigVM/Source/RigVMEditor/Private/Editor/RigVMEditor.cpp
+++ b/Engine/Plugins/Runtime/RigVM/Source/RigVMEditor/Private/Editor/RigVMEditor.cpp
@@ -3644,9 +3644,12 @@ void FRigVMEditor::OnGraphNodeDropToPerform(TSharedPtr<FDragDropOperation> InDra
 			{
 				if (URigVMBlueprint* TargetRigBlueprint = Cast<URigVMBlueprint>(FBlueprintEditorUtils::FindBlueprintForGraph(InGraph)))
 				{
-					if (URigVMLibraryNode* LibraryNode = ExplorerOp->GetBlueprint()->GetLocalFunctionLibrary()->FindFunction(*ExplorerOp->GetElement().Name))
+					if (URigVMFunctionLibrary* Library = ExplorerOp->GetBlueprint()->GetLocalFunctionLibrary())
 					{
-						Schema->RequestFunctionDropOnPanel(InGraph, LibraryNode->GetFunctionIdentifier(), InNodePosition, InScreenPosition);
+						if (URigVMLibraryNode* LibraryNode = Library->FindFunction(*ExplorerOp->GetElement().Name))
+						{
+							Schema->RequestFunctionDropOnPanel(InGraph, LibraryNode->GetFunctionIdentifier(), InNodePosition, InScreenPosition);
+						}
 					}
 				}
 			}
diff --git a/Engine/Plugins/Runtime/RigVM/Source/RigVMEditor/Private/Widgets/SRigVMEditorGraphExplorer.cpp b/Engine/Plugins/Runtime/RigVM/Source/RigVMEditor/Private/Widgets/SRigVMEditorGraphExplorer.cpp
--- a/Engine/Plugins/Runtime/RigVM/Source/RigVMEditor/Private/Widgets/SRigVMEditorGraphExplorer.cpp
+++ b/Engine/Plugins/Runtime/RigVM/Source/RigVMEditor/Private/Widgets/SRigVMEditorGraphExplorer.cpp
@@ -511,7 +511,12 @@ TArray<URigVMLibraryNode*> SRigVMEditorGraphExplorer::GetFunctions() const
 	}
 
 	FRigVMClient* RigVMClient = Editor->GetRigVMBlueprint()->GetRigVMClient();
-	return RigVMClient->GetFunctionLibrary()->GetFunctions();
+	if (URigVMFunctionLibrary* Library = RigVMClient->GetFunctionLibrary())
+	{
+		return Library->GetFunctions();
+	}
+
+	return Functions;
 }
 
 TArray<FRigVMGraphVariableDescription> SRigVMEditorGraphExplorer::GetVariables() const
@@ -658,19 +663,20 @@ void SRigVMEditorGraphExplorer::OnFunctionClicked(const FString& InFunctionPath)
 	}
 
 	FRigVMClient* RigVMClient = Editor->GetRigVMBlueprint()->GetRigVMClient();
-	URigVMFunctionLibrary* FunctionLibrary = RigVMClient->GetFunctionLibrary();
-	
-	if (const URigVMLibraryNode* FunctionNode = FunctionLibrary->FindFunction(*InFunctionPath))
+	if (URigVMFunctionLibrary* FunctionLibrary = RigVMClient->GetFunctionLibrary())
 	{
-		const URigVMEdGraphSchema* EdSchema = GetDefault<URigVMEdGraphSchema>();
-		if (UEdGraph* EdGraph = Editor->GetRigVMBlueprint()->GetEdGraph(FunctionNode->GetContainedGraph()))
+		if (const URigVMLibraryNode* FunctionNode = FunctionLibrary->FindFunction(*InFunctionPath))
 		{
-			FGraphDisplayInfo DisplayInfo;
-			const UEdGraphSchema* Schema = EdGraph->GetSchema();
-			check(Schema != nullptr);
-			Schema->GetGraphDisplayInformation(*EdGraph, DisplayInfo);
-			TSharedRef<SKismetInspector> CurrentInspector = Editor->GetInspector();
-			CurrentInspector->ShowDetailsForSingleObject(EdGraph, SKismetInspector::FShowDetailsOptions(DisplayInfo.PlainName));
+			const URigVMEdGraphSchema* EdSchema = GetDefault<URigVMEdGraphSchema>();
+			if (UEdGraph* EdGraph = Editor->GetRigVMBlueprint()->GetEdGraph(FunctionNode->GetContainedGraph()))
+			{
+				FGraphDisplayInfo DisplayInfo;
+				const UEdGraphSchema* Schema = EdGraph->GetSchema();
+				check(Schema != nullptr);
+				Schema->GetGraphDisplayInformation(*EdGraph, DisplayInfo);
+				TSharedRef<SKismetInspector> CurrentInspector = Editor->GetInspector();
+				CurrentInspector->ShowDetailsForSingleObject(EdGraph, SKismetInspector::FShowDetailsOptions(DisplayInfo.PlainName));
+			}
 		}
 	}
 }
@@ -738,14 +744,15 @@ void SRigVMEditorGraphExplorer::OnFunctionDoubleClicked(const FString& InFunctio
 	}
 
 	FRigVMClient* RigVMClient = Editor->GetRigVMBlueprint()->GetRigVMClient();
-	URigVMFunctionLibrary* FunctionLibrary = RigVMClient->GetFunctionLibrary();
-	
-	if (const URigVMLibraryNode* FunctionNode = FunctionLibrary->FindFunction(*InFunctionPath))
+	if (URigVMFunctionLibrary* FunctionLibrary = RigVMClient->GetFunctionLibrary())
 	{
-		const URigVMEdGraphSchema* EdSchema = GetDefault<URigVMEdGraphSchema>();
-		if (UEdGraph* EdGraph = Editor->GetRigVMBlueprint()->GetEdGraph(FunctionNode->GetContainedGraph()))
+		if (const URigVMLibraryNode* FunctionNode = FunctionLibrary->FindFunction(*InFunctionPath))
 		{
-			Editor->JumpToHyperlink(EdGraph);
+			const URigVMEdGraphSchema* EdSchema = GetDefault<URigVMEdGraphSchema>();
+			if (UEdGraph* EdGraph = Editor->GetRigVMBlueprint()->GetEdGraph(FunctionNode->GetContainedGraph()))
+			{
+				Editor->JumpToHyperlink(EdGraph);
+			}
 		}
 	}
 }
@@ -1015,12 +1022,14 @@ void SRigVMEditorGraphExplorer::OnCopy()
 	{
 		case Function:
 		{
-			URigVMFunctionLibrary* FunctionLibrary = Blueprint->GetLocalFunctionLibrary();
-			if (const URigVMLibraryNode* FunctionNode = FunctionLibrary->FindFunction(*Selection[0].Name))
+			if (URigVMFunctionLibrary* FunctionLibrary = Blueprint->GetLocalFunctionLibrary())
 			{
-				if (UEdGraph* EdGraph = Editor->GetRigVMBlueprint()->GetEdGraph(FunctionNode->GetContainedGraph()))
+				if (const URigVMLibraryNode* FunctionNode = FunctionLibrary->FindFunction(*Selection[0].Name))
 				{
-					Blueprint->ExportGraphToText(EdGraph, OutputString);
+					if (UEdGraph* EdGraph = Editor->GetRigVMBlueprint()->GetEdGraph(FunctionNode->GetContainedGraph()))
+					{
+						Blueprint->ExportGraphToText(EdGraph, OutputString);
+					}
 				}
 			}
 			break;
@@ -1343,12 +1352,14 @@ void SRigVMEditorGraphExplorer::OnAddFunctionVariant()
 		return;
 	}
 
-	URigVMFunctionLibrary* FunctionLibrary = Editor->GetRigVMBlueprint()->GetLocalFunctionLibrary();
-	if (const URigVMLibraryNode* FunctionNode = FunctionLibrary->FindFunction(*Selected[0].Name))
+	if (URigVMFunctionLibrary* FunctionLibrary = Editor->GetRigVMBlueprint()->GetLocalFunctionLibrary())
 	{
-		if (UEdGraph* EdGraph = Editor->GetRigVMBlueprint()->GetEdGraph(FunctionNode->GetContainedGraph()))
+		if (const URigVMLibraryNode* FunctionNode = FunctionLibrary->FindFunction(*Selected[0].Name))
 		{
-			Editor->AddNewFunctionVariant(EdGraph);
+			if (UEdGraph* EdGraph = Editor->GetRigVMBlueprint()->GetEdGraph(FunctionNode->GetContainedGraph()))
+			{
+				Editor->AddNewFunctionVariant(EdGraph);
+			}
 		}
 	}
 }
@@ -1539,12 +1550,14 @@ bool SRigVMEditorGraphExplorer::OnRenameFunction(const FString& InOldPath, const
 	if (Editor)
 	{
 		const URigVMEdGraphSchema* Schema = GetDefault<URigVMEdGraphSchema>();
-		URigVMFunctionLibrary* FunctionLibrary = Editor->GetRigVMBlueprint()->GetLocalFunctionLibrary();
-		if (const URigVMLibraryNode* FunctionNode = FunctionLibrary->FindFunction(*InOldPath))
+		if (URigVMFunctionLibrary* FunctionLibrary = Editor->GetRigVMBlueprint()->GetLocalFunctionLibrary())
 		{
-			if (UEdGraph* EdGraph = Editor->GetRigVMBlueprint()->GetEdGraph(FunctionNode->GetContainedGraph()))
+			if (const URigVMLibraryNode* FunctionNode = FunctionLibrary->FindFunction(*InOldPath))
 			{
-				bResult = Schema->TryRenameGraph(EdGraph, *InNewPath);
+				if (UEdGraph* EdGraph = Editor->GetRigVMBlueprint()->GetEdGraph(FunctionNode->GetContainedGraph()))
+				{
+					bResult = Schema->TryRenameGraph(EdGraph, *InNewPath);
+				}
 			}
 		}
 	}
@@ -1560,15 +1573,17 @@ bool SRigVMEditorGraphExplorer::OnCanRenameFunction(const FString& InOldPath, co
 
 	if (Editor)
 	{
-		
-		if (Editor->GetRigVMBlueprint()->GetRigVMClient()->GetFunctionLibrary()->FindFunction(*InNewPath))
+		if (URigVMFunctionLibrary* Library = Editor->GetRigVMBlueprint()->GetRigVMClient()->GetFunctionLibrary())
 		{
-			OutErrorMessage = FText::FromString(TEXT("Name already in use."));
-			bResult = false;
-		}
-		else
-		{
-			bResult = true;
+			if (Library->FindFunction(*InNewPath))
+			{
+				OutErrorMessage = FText::FromString(TEXT("Name already in use."));
+				bResult = false;
+			}
+			else
+			{
+				bResult = true;
+			}
 		}
 	}
 
