Control Rig: Add local variables to graph explorer
#rnx

diff --git a/Engine/Plugins/Runtime/RigVM/Source/RigVMDeveloper/Private/EdGraph/RigVMEdGraphSchema.cpp b/Engine/Plugins/Runtime/RigVM/Source/RigVMDeveloper/Private/EdGraph/RigVMEdGraphSchema.cpp
--- a/Engine/Plugins/Runtime/RigVM/Source/RigVMDeveloper/Private/EdGraph/RigVMEdGraphSchema.cpp
+++ b/Engine/Plugins/Runtime/RigVM/Source/RigVMDeveloper/Private/EdGraph/RigVMEdGraphSchema.cpp
@@ -2113,17 +2113,35 @@ bool URigVMEdGraphSchema::RequestVariableDropOnPin(UEdGraph* InGraph, FProperty*
 			if (const URigVMPin* ModelPin = Graph->GetModel()->FindPin(InPin->GetName()))
 			{
 				const FRigVMExternalVariable ExternalVariable = FRigVMExternalVariable::Make(InVariableToDrop, nullptr);
-				if (ModelPin->CanBeBoundToVariable(ExternalVariable))
+				return RequestVariableDropOnPin(InGraph, ExternalVariable, InPin, InDropPosition, InScreenPosition);
+			}
+		}
+	}
+#endif
+
+	return false;
+}
+
+bool URigVMEdGraphSchema::RequestVariableDropOnPin(UEdGraph* InGraph, const FRigVMExternalVariable& InVariableToDrop, UEdGraphPin* InPin, const FVector2D& InDropPosition, const FVector2D& InScreenPosition) const
+{
+#if WITH_EDITOR
+	if (InVariableToDrop.IsValid(true /* allow nullptr */))
+	{
+		if(const URigVMEdGraph* Graph = Cast<URigVMEdGraph>(InGraph))
+		{
+			if (const URigVMPin* ModelPin = Graph->GetModel()->FindPin(InPin->GetName()))
+			{
+				if (ModelPin->CanBeBoundToVariable(InVariableToDrop))
 				{
 					const FModifierKeysState KeyState = FSlateApplication::Get().GetModifierKeys();
 					if (KeyState.IsAltDown())
 					{
-						return Graph->GetController()->BindPinToVariable(ModelPin->GetPinPath(), InVariableToDrop->GetName(), true, true);
+						return Graph->GetController()->BindPinToVariable(ModelPin->GetPinPath(), InVariableToDrop.Name.ToString(), true, true);
 					}
 					else
 					{
 						Graph->GetController()->OpenUndoBracket(TEXT("Bind Variable to Pin"));
-						if (const URigVMVariableNode* VariableNode = Graph->GetController()->AddVariableNode(ExternalVariable.Name, ExternalVariable.TypeName.ToString(), ExternalVariable.TypeObject, true, FString(), InDropPosition + FVector2D(0.f, -34.f)))
+						if (const URigVMVariableNode* VariableNode = Graph->GetController()->AddVariableNode(InVariableToDrop.Name, InVariableToDrop.TypeName.ToString(), InVariableToDrop.TypeObject, true, FString(), InDropPosition + FVector2D(0.f, -34.f)))
 						{
 							Graph->GetController()->AddLink(VariableNode->FindPin(TEXT("Value"))->GetPinPath(), ModelPin->GetPinPath(), true);
 						}
@@ -2135,7 +2153,6 @@ bool URigVMEdGraphSchema::RequestVariableDropOnPin(UEdGraph* InGraph, FProperty*
 		}
 	}
 #endif
-
 	return false;
 }
 
diff --git a/Engine/Plugins/Runtime/RigVM/Source/RigVMDeveloper/Public/EdGraph/RigVMEdGraphSchema.h b/Engine/Plugins/Runtime/RigVM/Source/RigVMDeveloper/Public/EdGraph/RigVMEdGraphSchema.h
--- a/Engine/Plugins/Runtime/RigVM/Source/RigVMDeveloper/Public/EdGraph/RigVMEdGraphSchema.h
+++ b/Engine/Plugins/Runtime/RigVM/Source/RigVMDeveloper/Public/EdGraph/RigVMEdGraphSchema.h
@@ -301,6 +301,7 @@ public:
 	virtual bool RequestFunctionDropOnPanel(UEdGraph* InGraph, const FRigVMGraphFunctionIdentifier& InFunction, const FVector2D& InDropPosition, const FVector2D& InScreenPosition) const;
 	virtual bool RequestVariableDropOnPanel(UEdGraph* InGraph, FProperty* InVariableToDrop, const FVector2D& InDropPosition, const FVector2D& InScreenPosition) const override;
 	virtual bool RequestVariableDropOnPin(UEdGraph* InGraph, FProperty* InVariableToDrop, UEdGraphPin* InPin, const FVector2D& InDropPosition, const FVector2D& InScreenPosition) const override;
+	virtual bool RequestVariableDropOnPin(UEdGraph* InGraph, const FRigVMExternalVariable& InVariableToDrop, UEdGraphPin* InPin, const FVector2D& InDropPosition, const FVector2D& InScreenPosition) const;
 	virtual bool IsStructEditable(UStruct* InStruct) const;
 	virtual void SetNodePosition(UEdGraphNode* Node, const FVector2D& Position) const override;
 	void SetNodePosition(UEdGraphNode* Node, const FVector2D& Position, bool bSetupUndo) const;
diff --git a/Engine/Plugins/Runtime/RigVM/Source/RigVMEditor/Private/Editor/RigVMEditor.cpp b/Engine/Plugins/Runtime/RigVM/Source/RigVMEditor/Private/Editor/RigVMEditor.cpp
--- a/Engine/Plugins/Runtime/RigVM/Source/RigVMEditor/Private/Editor/RigVMEditor.cpp
+++ b/Engine/Plugins/Runtime/RigVM/Source/RigVMEditor/Private/Editor/RigVMEditor.cpp
@@ -473,6 +473,15 @@ bool FRigVMEditor::InEditingMode() const
 	return true;
 }
 
+void FRigVMEditor::OnGraphEditorFocused(const TSharedRef<SGraphEditor>& InGraphEditor)
+{
+	FBlueprintEditor::OnGraphEditorFocused(InGraphEditor);
+	if (GraphExplorerWidget.IsValid())
+	{
+		GraphExplorerWidget->Refresh();
+	}
+}
+
 void FRigVMEditor::Tick(float DeltaTime)
 {
 	FBlueprintEditor::Tick(DeltaTime);
@@ -1523,16 +1532,6 @@ void FRigVMEditor::RedoAction()
 	GEditor->RedoTransaction();
 }
 
-void FRigVMEditor::CreateDefaultTabContents(const TArray<UBlueprint*>& InBlueprints)
-{
-	FBlueprintEditor::CreateDefaultTabContents(InBlueprints);
-
-	if (IsEditingSingleBlueprint())
-	{
-		GraphExplorerWidget = SNew(SRigVMEditorGraphExplorer, SharedThis(this));
-	}
-}
-
 void FRigVMEditor::NewDocument_OnClicked(ECreatedDocumentType GraphType)
 {
 	if (GraphType == FBlueprintEditor::CGT_NewFunctionGraph)
@@ -3640,7 +3639,7 @@ void FRigVMEditor::OnGraphNodeDropToPerform(TSharedPtr<FDragDropOperation> InDra
 		if (FocusedGraphEdPtr.IsValid())
 		{
 			const URigVMEdGraphSchema* Schema = CastChecked<URigVMEdGraphSchema>(GetRigVMBlueprint()->GetRigVMEdGraphSchemaClass()->GetDefaultObject());
-			if (ExplorerOp->GetElement().Type == Function)
+			if (ExplorerOp->GetElement().Type == ERigVMExplorerElementType::Function)
 			{
 				if (URigVMBlueprint* TargetRigBlueprint = Cast<URigVMBlueprint>(FBlueprintEditorUtils::FindBlueprintForGraph(InGraph)))
 				{
@@ -3653,7 +3652,7 @@ void FRigVMEditor::OnGraphNodeDropToPerform(TSharedPtr<FDragDropOperation> InDra
 					}
 				}
 			}
-			else if (ExplorerOp->GetElement().Type == Variable)
+			else if (ExplorerOp->GetElement().Type == ERigVMExplorerElementType::Variable)
 			{
 				if (URigVMBlueprint* TargetRigBlueprint = Cast<URigVMBlueprint>(FBlueprintEditorUtils::FindBlueprintForGraph(InGraph)))
 				{
@@ -3661,6 +3660,62 @@ void FRigVMEditor::OnGraphNodeDropToPerform(TSharedPtr<FDragDropOperation> InDra
 					Schema->RequestVariableDropOnPanel(InGraph, Property, InNodePosition, InScreenPosition);
 				}
 			}
+			else if (ExplorerOp->GetElement().Type == ERigVMExplorerElementType::LocalVariable)
+			{
+				if (URigVMBlueprint* TargetRigBlueprint = Cast<URigVMBlueprint>(FBlueprintEditorUtils::FindBlueprintForGraph(InGraph)))
+				{
+					if (URigVMGraph* Graph = TargetRigBlueprint->GetFocusedModel())
+					{
+						for (FRigVMGraphVariableDescription LocalVariable : Graph->GetLocalVariables())
+						{
+							if (LocalVariable.Name == ExplorerOp->GetElement().Name)
+							{
+								if (URigVMController* Controller = TargetRigBlueprint->GetRigVMClient()->GetController(Graph))
+								{
+									FMenuBuilder MenuBuilder(true, nullptr);
+									const FText VariableNameText = FText::FromName( LocalVariable.Name );
+
+									MenuBuilder.BeginSection("BPVariableDroppedOn", VariableNameText );
+
+									MenuBuilder.AddMenuEntry(
+										FText::Format( LOCTEXT("CreateGetVariable", "Get {0}"), VariableNameText ),
+										FText::Format( LOCTEXT("CreateVariableGetterToolTip", "Create Getter for variable '{0}'\n(Ctrl-drag to automatically create a getter)"), VariableNameText ),
+										FSlateIcon(),
+										FUIAction(
+										FExecuteAction::CreateLambda([Controller, LocalVariable, InNodePosition]()
+										{
+											Controller->AddVariableNode(LocalVariable.Name, LocalVariable.CPPType, LocalVariable.CPPTypeObject, true, LocalVariable.DefaultValue, InNodePosition, FString(), true, true);
+										}),
+										FCanExecuteAction()));
+
+									MenuBuilder.AddMenuEntry(
+										FText::Format( LOCTEXT("CreateSetVariable", "Set {0}"), VariableNameText ),
+										FText::Format( LOCTEXT("CreateVariableSetterToolTip", "Create Setter for variable '{0}'\n(Alt-drag to automatically create a setter)"), VariableNameText ),
+										FSlateIcon(),
+										FUIAction(
+										FExecuteAction::CreateLambda([Controller, LocalVariable, InNodePosition]()
+										{
+											Controller->AddVariableNode(LocalVariable.Name, LocalVariable.CPPType, LocalVariable.CPPTypeObject, false, LocalVariable.DefaultValue, InNodePosition, FString(), true, true);
+										}),
+										FCanExecuteAction()));
+
+									TSharedRef< SWidget > PanelWidget = GetGraphEditor(InGraph).ToSharedRef();
+									// Show dialog to choose getter vs setter
+									FSlateApplication::Get().PushMenu(
+										PanelWidget,
+										FWidgetPath(),
+										MenuBuilder.MakeWidget(),
+										InScreenPosition,
+										FPopupTransitionEffect( FPopupTransitionEffect::ContextMenu)
+										);
+
+									MenuBuilder.EndSection();
+								}
+							}
+						}
+					}
+				}
+			}
 		}
 	}
 }
diff --git a/Engine/Plugins/Runtime/RigVM/Source/RigVMEditor/Private/Editor/RigVMEditorGraphExplorerTabSummoner.cpp b/Engine/Plugins/Runtime/RigVM/Source/RigVMEditor/Private/Editor/RigVMEditorGraphExplorerTabSummoner.cpp
--- a/Engine/Plugins/Runtime/RigVM/Source/RigVMEditor/Private/Editor/RigVMEditorGraphExplorerTabSummoner.cpp
+++ b/Engine/Plugins/Runtime/RigVM/Source/RigVMEditor/Private/Editor/RigVMEditorGraphExplorerTabSummoner.cpp
@@ -21,7 +21,9 @@ FRigVMEditorGraphExplorerTabSummoner::FRigVMEditorGraphExplorerTabSummoner(const
 
 TSharedRef<SWidget> FRigVMEditorGraphExplorerTabSummoner::CreateTabBody(const FWorkflowTabSpawnInfo& Info) const
 {
-	return SNew(SRigVMEditorGraphExplorer, RigVMEditor.Pin().ToSharedRef());
+	TSharedRef<SRigVMEditorGraphExplorer> Explorer = SNew(SRigVMEditorGraphExplorer, RigVMEditor.Pin().ToSharedRef());
+	RigVMEditor.Pin()->SetGraphExplorerWidget(Explorer);
+	return Explorer;
 }
 
 #undef LOCTEXT_NAMESPACE 
diff --git a/Engine/Plugins/Runtime/RigVM/Source/RigVMEditor/Private/Widgets/SRigVMEditorGraphExplorer.cpp b/Engine/Plugins/Runtime/RigVM/Source/RigVMEditor/Private/Widgets/SRigVMEditorGraphExplorer.cpp
--- a/Engine/Plugins/Runtime/RigVM/Source/RigVMEditor/Private/Widgets/SRigVMEditorGraphExplorer.cpp
+++ b/Engine/Plugins/Runtime/RigVM/Source/RigVMEditor/Private/Widgets/SRigVMEditorGraphExplorer.cpp
@@ -103,10 +103,30 @@ FReply FRigVMGraphExplorerDragDropOp::DroppedOnPin(FVector2D ScreenPosition, FVe
 	TObjectPtr<URigVMBlueprint> Blueprint = GetBlueprint();
 	
 	const URigVMEdGraphSchema* EdSchema = GetDefault<URigVMEdGraphSchema>();
-	FProperty* Property = Blueprint->SkeletonGeneratedClass->FindPropertyByName(*GetElement().Name);
-	if (EdSchema->RequestVariableDropOnPin(TargetGraph, Property, TargetPin, GraphPosition, ScreenPosition))
+
+	if (Element.Type == ERigVMExplorerElementType::Variable)
+	{
+		if (FProperty* Property = Blueprint->SkeletonGeneratedClass->FindPropertyByName(*GetElement().Name))
+		{
+			if (EdSchema->RequestVariableDropOnPin(TargetGraph, Property, TargetPin, GraphPosition, ScreenPosition))
+			{
+				return FReply::Handled();
+			}
+		}
+	}
+	else if (Element.Type == ERigVMExplorerElementType::LocalVariable)
 	{
-		return FReply::Handled();
+		for (FRigVMGraphVariableDescription& Variable : Blueprint->GetFocusedModel()->GetLocalVariables())
+		{
+			if (Variable.Name == Element.Name)
+			{
+				if (EdSchema->RequestVariableDropOnPin(TargetGraph, Variable.ToExternalVariable(), TargetPin, GraphPosition, ScreenPosition))
+				{
+					return FReply::Handled();
+				}
+				break;
+			}
+		}
 	}
 	
 	return FGraphEditorDragDropAction::DroppedOnPin(ScreenPosition, GraphPosition);
@@ -127,9 +147,11 @@ void FRigVMEditorGraphExplorerCommands::RegisterCommands()
 	UI_COMMAND(CreateGraph, "Add New Graph", "Create a new graph and show it in the editor.", EUserInterfaceActionType::Button, FInputChord());
 	UI_COMMAND(CreateFunction, "Add New Function", "Create a new function and show it in the editor.", EUserInterfaceActionType::Button, FInputChord());
 	UI_COMMAND(CreateVariable, "Add New Variable", "Create a new member variable.", EUserInterfaceActionType::Button, FInputChord());
+	UI_COMMAND(CreateLocalVariable, "Add New Local Variable", "Create a new local variable to the function.", EUserInterfaceActionType::Button, FInputChord());
 	UI_COMMAND(AddFunctionVariant, "Add Variant (Experimental)", "Creates a new variant of a function.", EUserInterfaceActionType::Button, FInputChord());
 	UI_COMMAND(PasteFunction, "Paste Function", "Pastes the function.", EUserInterfaceActionType::Button, FInputChord());
 	UI_COMMAND(PasteVariable, "Paste Variable", "Pastes the variable.", EUserInterfaceActionType::Button, FInputChord());
+	UI_COMMAND(PasteLocalVariable, "Paste Local Variable", "Pastes the local variable.", EUserInterfaceActionType::Button, FInputChord());
 }
 
 
@@ -161,8 +183,7 @@ SRigVMEditorGraphExplorer::~SRigVMEditorGraphExplorer()
 
 void SRigVMEditorGraphExplorer::Refresh()
 {
-	bNeedsRefresh = false;
-	TreeView->RefreshTreeView(true);
+	bNeedsRefresh = true;
 }
 
 void SRigVMEditorGraphExplorer::Tick(const FGeometry& InAllottedGeometry, const double InCurrentTime, const float InDeltaTime)
@@ -171,7 +192,8 @@ void SRigVMEditorGraphExplorer::Tick(const FGeometry& InAllottedGeometry, const
 
 	if(bNeedsRefresh)
 	{
-		Refresh();
+		TreeView->RefreshTreeView(true);
+		bNeedsRefresh = false;
 	}
 }
 
@@ -208,6 +230,11 @@ void SRigVMEditorGraphExplorer::RegisterCommands()
 			FCanExecuteAction::CreateSP(this, &SRigVMEditorGraphExplorer::CanCreateVariable)
 			);
 
+		ToolKitCommandList->MapAction(FRigVMEditorGraphExplorerCommands::Get().CreateLocalVariable,
+			FExecuteAction::CreateSP(this, &SRigVMEditorGraphExplorer::OnCreateLocalVariable),
+			FCanExecuteAction::CreateSP(this, &SRigVMEditorGraphExplorer::CanCreateLocalVariable)
+			);
+
 		ToolKitCommandList->MapAction(FRigVMEditorGraphExplorerCommands::Get().AddFunctionVariant,
 			FExecuteAction::CreateSP(this, &SRigVMEditorGraphExplorer::OnAddFunctionVariant),
 			FCanExecuteAction::CreateSP(this, &SRigVMEditorGraphExplorer::CanAddFunctionVariant)
@@ -244,6 +271,11 @@ void SRigVMEditorGraphExplorer::RegisterCommands()
 			FCanExecuteAction(), FIsActionChecked(),
 			FIsActionButtonVisible::CreateSP(this, &SRigVMEditorGraphExplorer::CanPasteVariable));
 
+		ToolKitCommandList->MapAction(FRigVMEditorGraphExplorerCommands::Get().PasteLocalVariable,
+			FExecuteAction::CreateSP(this, &SRigVMEditorGraphExplorer::OnPasteLocalVariable),
+			FCanExecuteAction(), FIsActionChecked(),
+			FIsActionButtonVisible::CreateSP(this, &SRigVMEditorGraphExplorer::CanPasteLocalVariable));
+
 		ToolKitCommandList->MapAction(FGenericCommands::Get().Delete,
 			FExecuteAction::CreateSP(this, &SRigVMEditorGraphExplorer::OnDeleteEntry),
 			FCanExecuteAction::CreateSP(this, &SRigVMEditorGraphExplorer::CanDeleteEntry));
@@ -288,6 +320,7 @@ void SRigVMEditorGraphExplorer::CreateWidgets()
 	Delegates.OnGetEventNodesInGraph = FRigVMGraphExplorer_OnGetEventNodesInGraph::CreateSP(this, &SRigVMEditorGraphExplorer::GetEventNodesInGraph);
 	Delegates.OnGetFunctions = FRigVMGraphExplorer_OnGetFunctions::CreateSP(this, &SRigVMEditorGraphExplorer::GetFunctions);
 	Delegates.OnGetVariables = FRigVMGraphExplorer_OnGetVariables::CreateSP(this, &SRigVMEditorGraphExplorer::GetVariables);
+	Delegates.OnGetLocalVariables = FRigVMGraphExplorer_OnGetVariables::CreateSP(this, &SRigVMEditorGraphExplorer::GetLocalVariables);
 	Delegates.OnGetGraphDisplayName = FRigVMGraphExplorer_OnGetGraphDisplayName::CreateSP(this, &SRigVMEditorGraphExplorer::GetGraphDisplayName);
 	Delegates.OnGetEventDisplayName = FRigVMGraphExplorer_OnGetEventDisplayName::CreateSP(this, &SRigVMEditorGraphExplorer::GetEventDisplayName);
 	Delegates.OnGetGraphIcon = FRigVMGraphExplorer_OnGetGraphIcon::CreateSP(this, &SRigVMEditorGraphExplorer::GetGraphIcon);
@@ -302,6 +335,7 @@ void SRigVMEditorGraphExplorer::CreateWidgets()
 	Delegates.OnCreateGraph = FRigVMGraphExplorer_OnCreateGraph::CreateSP(this, &SRigVMEditorGraphExplorer::OnCreateGraph);
 	Delegates.OnCreateFunction = FRigVMGraphExplorer_OnCreateFunction::CreateSP(this, &SRigVMEditorGraphExplorer::OnCreateFunction);
 	Delegates.OnCreateVariable = FRigVMGraphExplorer_OnCreateVariable::CreateSP(this, &SRigVMEditorGraphExplorer::OnCreateVariable);
+	Delegates.OnCreateLocalVariable = FRigVMGraphExplorer_OnCreateVariable::CreateSP(this, &SRigVMEditorGraphExplorer::OnCreateLocalVariable);
 	Delegates.OnRenameGraph = FRigVMGraphExplorer_OnRenameGraph::CreateSP(this, &SRigVMEditorGraphExplorer::OnRenameGraph);
 	Delegates.OnRenameFunction = FRigVMGraphExplorer_OnRenameFunction::CreateSP(this, &SRigVMEditorGraphExplorer::OnRenameFunction);
 	Delegates.OnCanRenameGraph = FRigVMGraphExplorer_OnCanRenameGraph::CreateSP(this, &SRigVMEditorGraphExplorer::OnCanRenameGraph);
@@ -318,6 +352,8 @@ void SRigVMEditorGraphExplorer::CreateWidgets()
 	Delegates.OnSetVariablePinType = FRigVMGraphExplorer_OnSetVariablePinType::CreateSP(this, &SRigVMEditorGraphExplorer::OnSetVariablePinType);
 	Delegates.OnIsVariablePublic = FRigVMGraphExplorer_OnIsVariablePublic::CreateSP(this, &SRigVMEditorGraphExplorer::OnIsVariablePublic);
 	Delegates.OnToggleVariablePublic = FRigVMGraphExplorer_OnToggleVariablePublic::CreateSP(this, &SRigVMEditorGraphExplorer::OnToggleVariablePublic);
+	Delegates.OnIsFunctionFocused = FRigVMGraphExplorer_OnIsFunctionFocused::CreateSP(this, &SRigVMEditorGraphExplorer::IsFunctionFocused);
+	Delegates.OnGetCustomPinFilters = FRigVMGraphExplorer_OnGetCustomPinFilters::CreateSP(this, &SRigVMEditorGraphExplorer::GetCustomPinFilters);
 	
 	Delegates.OnSelectionChanged = FRigVMGraphExplorer_OnSelectionChanged::CreateSP(this, &SRigVMEditorGraphExplorer::HandleSelectionChanged);
 	
@@ -404,9 +440,11 @@ void SRigVMEditorGraphExplorer::BuildAddNewMenu(FMenuBuilder& MenuBuilder)
 	MenuBuilder.BeginSection("AddNewItem", LOCTEXT("AddOperations", "Add New"));
 	MenuBuilder.AddMenuEntry(FRigVMEditorGraphExplorerCommands::Get().CreateGraph);
 	MenuBuilder.AddMenuEntry(FRigVMEditorGraphExplorerCommands::Get().CreateFunction);
-	MenuBuilder.AddMenuEntry(FRigVMEditorGraphExplorerCommands::Get().CreateVariable);
 	MenuBuilder.AddMenuEntry(FRigVMEditorGraphExplorerCommands::Get().PasteFunction);
+	MenuBuilder.AddMenuEntry(FRigVMEditorGraphExplorerCommands::Get().CreateVariable);
 	MenuBuilder.AddMenuEntry(FRigVMEditorGraphExplorerCommands::Get().PasteVariable);
+	MenuBuilder.AddMenuEntry(FRigVMEditorGraphExplorerCommands::Get().CreateLocalVariable);
+	MenuBuilder.AddMenuEntry(FRigVMEditorGraphExplorerCommands::Get().PasteLocalVariable);
 	
 	MenuBuilder.EndSection();
 }
@@ -531,6 +569,26 @@ TArray<FRigVMGraphVariableDescription> SRigVMEditorGraphExplorer::GetVariables()
 	return Editor->GetRigVMBlueprint()->GetMemberVariables();
 }
 
+TArray<FRigVMGraphVariableDescription> SRigVMEditorGraphExplorer::GetLocalVariables() const
+{
+	TArray<FRigVMGraphVariableDescription> Variables;
+	TSharedPtr<FRigVMEditor> Editor = RigVMEditor.Pin();
+	if (!Editor)
+	{
+		return Variables;
+	}
+
+	if (URigVMBlueprint* Blueprint = Editor->GetRigVMBlueprint())
+	{
+		if (URigVMGraph* Function = Blueprint->GetFocusedModel())
+		{
+			return Function->GetLocalVariables(false);
+		}
+	}
+
+	return Variables;
+}
+
 FText SRigVMEditorGraphExplorer::GetGraphDisplayName(const FString& InGraphPath) const
 {
 	FText DisplayName;
@@ -681,7 +739,7 @@ void SRigVMEditorGraphExplorer::OnFunctionClicked(const FString& InFunctionPath)
 	}
 }
 
-void SRigVMEditorGraphExplorer::OnVariableClicked(const FString& InVariable)
+void SRigVMEditorGraphExplorer::OnVariableClicked(const FRigVMExplorerElementKey& InVariable)
 {
 	TSharedPtr<FRigVMEditor> Editor = RigVMEditor.Pin();
 	if (!Editor)
@@ -689,14 +747,24 @@ void SRigVMEditorGraphExplorer::OnVariableClicked(const FString& InVariable)
 		return;
 	}
 
-	FRigVMClient* RigVMClient = Editor->GetRigVMBlueprint()->GetRigVMClient();
-	TSharedRef<SKismetInspector> CurrentInspector = Editor->GetInspector();
-	SKismetInspector::FShowDetailsOptions Options(FText::FromString(InVariable));
-	Options.bForceRefresh = true;
+	URigVMBlueprint* Blueprint = Editor->GetRigVMBlueprint();
+	if (InVariable.Type == ERigVMExplorerElementType::Variable)
+	{
+		FRigVMClient* RigVMClient = Blueprint->GetRigVMClient();
+		TSharedRef<SKismetInspector> CurrentInspector = Editor->GetInspector();
+		SKismetInspector::FShowDetailsOptions Options(FText::FromString(InVariable.Name));
+		Options.bForceRefresh = true;
 
-	FProperty* Prop = Editor->GetRigVMBlueprint()->SkeletonGeneratedClass->FindPropertyByName(*InVariable);
-	UPropertyWrapper* PropWrap = (Prop ? Prop->GetUPropertyWrapper() : nullptr);
-	CurrentInspector->ShowDetailsForSingleObject(PropWrap, Options);
+		FProperty* Prop = Blueprint->SkeletonGeneratedClass->FindPropertyByName(*InVariable.Name);
+		UPropertyWrapper* PropWrap = (Prop ? Prop->GetUPropertyWrapper() : nullptr);
+		CurrentInspector->ShowDetailsForSingleObject(PropWrap, Options);
+	}
+	else if (InVariable.Type == ERigVMExplorerElementType::LocalVariable)
+	{
+		URigVMGraph* Graph = Blueprint->GetFocusedModel();
+		UEdGraph* EdGraph = Blueprint->GetEdGraph(Graph);
+		Editor->SelectLocalVariable(EdGraph, *InVariable.Name);
+	}
 }
 
 void SRigVMEditorGraphExplorer::OnGraphDoubleClicked(const FString& InGraphPath)
@@ -824,7 +892,7 @@ FString SRigVMEditorGraphExplorer::OnGetVariableCategory(const FString& InVariab
 	return FBlueprintEditorUtils::GetBlueprintVariableCategory(Editor->GetRigVMBlueprint(), *InVariable, nullptr).ToString();
 }
 
-FEdGraphPinType SRigVMEditorGraphExplorer::OnGetVariablePinType(const FString& InVariableName)
+FEdGraphPinType SRigVMEditorGraphExplorer::OnGetVariablePinType(const FRigVMExplorerElementKey& InVariable)
 {
 	TSharedPtr<FRigVMEditor> Editor = RigVMEditor.Pin();
 	if (!Editor.IsValid())
@@ -833,18 +901,34 @@ FEdGraphPinType SRigVMEditorGraphExplorer::OnGetVariablePinType(const FString& I
 	}
 
 	URigVMBlueprint* Blueprint = Editor->GetRigVMBlueprint();
-	for (FRigVMGraphVariableDescription& Variable : Blueprint->GetMemberVariables())
+	if (InVariable.Type == ERigVMExplorerElementType::Variable)
 	{
-		if (Variable.Name == InVariableName)
+		for (FRigVMGraphVariableDescription& Variable : Blueprint->GetMemberVariables())
 		{
-			return Variable.ToPinType();
+			if (Variable.Name == InVariable.Name)
+			{
+				return Variable.ToPinType();
+			}
+		}
+	}
+	else if(InVariable.Type == ERigVMExplorerElementType::LocalVariable)
+	{
+		if (URigVMGraph* Function = Blueprint->GetFocusedModel())
+		{
+			for (FRigVMGraphVariableDescription& Variable : Function->GetLocalVariables())
+			{
+				if (Variable.Name == InVariable.Name)
+				{
+					return Variable.ToPinType();
+				}
+			}
 		}
 	}
 
 	return FEdGraphPinType();
 }
 
-bool SRigVMEditorGraphExplorer::OnSetVariablePinType(const FString& InVariableName, const FEdGraphPinType& InType)
+bool SRigVMEditorGraphExplorer::OnSetVariablePinType(const FRigVMExplorerElementKey& InVariable, const FEdGraphPinType& InType)
 {
 	TSharedPtr<FRigVMEditor> Editor = RigVMEditor.Pin();
 	if (!Editor.IsValid())
@@ -853,13 +937,29 @@ bool SRigVMEditorGraphExplorer::OnSetVariablePinType(const FString& InVariableNa
 	}
 
 	URigVMBlueprint* Blueprint = Editor->GetRigVMBlueprint();
-	
-	for (FRigVMGraphVariableDescription& Variable : Blueprint->GetMemberVariables())
+
+	if (InVariable.Type == ERigVMExplorerElementType::Variable)
 	{
-		if (Variable.Name == InVariableName)
+		for (FRigVMGraphVariableDescription& Variable : Blueprint->GetMemberVariables())
 		{
-			FBlueprintEditorUtils::ChangeMemberVariableType(Blueprint, *InVariableName, InType);
-			return true;
+			if (Variable.Name == InVariable.Name)
+			{
+				FBlueprintEditorUtils::ChangeMemberVariableType(Blueprint, *InVariable.Name, InType);
+				return true;
+			}
+		}
+	}
+	else if(InVariable.Type == ERigVMExplorerElementType::LocalVariable)
+	{
+		if (URigVMGraph* Function = Blueprint->GetFocusedModel())
+		{
+			if (URigVMController* Controller = Blueprint->GetRigVMClient()->GetController(Function))
+			{
+				FString NewCPPType;
+				UObject* NewCPPTypeObject = nullptr;
+				RigVMTypeUtils::CPPTypeFromPinType(InType, NewCPPType, &NewCPPTypeObject);
+				Controller->SetLocalVariableType(*InVariable.Name, NewCPPType, NewCPPTypeObject, true, true);
+			}
 		}
 	}
 
@@ -906,6 +1006,37 @@ bool SRigVMEditorGraphExplorer::OnToggleVariablePublic(const FString& InVariable
 	return true;
 }
 
+bool SRigVMEditorGraphExplorer::IsFunctionFocused() const
+{
+	TSharedPtr<FRigVMEditor> Editor = RigVMEditor.Pin();
+	if (!Editor.IsValid())
+	{
+		return false;
+	}
+
+	if (URigVMGraph* Graph = Editor->GetRigVMBlueprint()->GetFocusedModel())
+	{
+		if (URigVMCollapseNode* FunctionNode = Cast<URigVMCollapseNode>(Graph->GetOuter()))
+		{
+			if (URigVMFunctionLibrary* Library = Cast<URigVMFunctionLibrary>(FunctionNode->GetOuter()))
+			{
+				return true;
+			}
+		}
+	}
+	return false;
+}
+
+TArray<TSharedPtr<IPinTypeSelectorFilter>> SRigVMEditorGraphExplorer::GetCustomPinFilters() const
+{
+	TArray<TSharedPtr<IPinTypeSelectorFilter>> CustomPinTypeFilters;
+	if (RigVMEditor.IsValid())
+	{
+		RigVMEditor.Pin()->GetPinTypeSelectorFilters(CustomPinTypeFilters);
+	}
+	return CustomPinTypeFilters;
+}
+
 TSharedPtr<SWidget> SRigVMEditorGraphExplorer::OnContextMenuOpening()
 {
 	TSharedPtr<FRigVMEditor> Editor = RigVMEditor.Pin();
@@ -922,11 +1053,11 @@ TSharedPtr<SWidget> SRigVMEditorGraphExplorer::OnContextMenuOpening()
 	{
 		switch (Selection[0].Type)
 		{
-			case Section:
+			case ERigVMExplorerElementType::Section:
 			{
 				break;
 			}
-			case Graph:
+			case ERigVMExplorerElementType::Graph:
 			{
 				MenuBuilder.BeginSection("BasicOperations");
 				MenuBuilder.AddMenuEntry(FRigVMEditorGraphExplorerCommands::Get().OpenGraph);
@@ -939,12 +1070,12 @@ TSharedPtr<SWidget> SRigVMEditorGraphExplorer::OnContextMenuOpening()
 				MenuBuilder.EndSection();
 				break;
 			}
-			case Event:
+			case ERigVMExplorerElementType::Event:
 			{
 				return nullptr;
 				break;
 			}
-			case Function:
+			case ERigVMExplorerElementType::Function:
 			{
 				MenuBuilder.BeginSection("BasicOperations");
 				MenuBuilder.AddMenuEntry(FRigVMEditorGraphExplorerCommands::Get().OpenGraph);
@@ -958,7 +1089,8 @@ TSharedPtr<SWidget> SRigVMEditorGraphExplorer::OnContextMenuOpening()
 				MenuBuilder.EndSection();
 				break;
 			}
-			case Variable:
+			case ERigVMExplorerElementType::Variable:
+			case ERigVMExplorerElementType::LocalVariable:
 			{
 				MenuBuilder.BeginSection("BasicOperations");
 				MenuBuilder.AddMenuEntry(FGenericCommands::Get().Rename);
@@ -1009,7 +1141,7 @@ void SRigVMEditorGraphExplorer::OnCopy()
 		return;
 	}
 
-	if (Selection[0].Type != Function && Selection[0].Type != Variable)
+	if (Selection[0].Type != ERigVMExplorerElementType::Function && Selection[0].Type != ERigVMExplorerElementType::Variable && Selection[0].Type != ERigVMExplorerElementType::LocalVariable)
 	{
 		return;
 	}
@@ -1020,7 +1152,7 @@ void SRigVMEditorGraphExplorer::OnCopy()
 
 	switch (Selection[0].Type)
 	{
-		case Function:
+		case ERigVMExplorerElementType::Function:
 		{
 			if (URigVMFunctionLibrary* FunctionLibrary = Blueprint->GetLocalFunctionLibrary())
 			{
@@ -1034,7 +1166,7 @@ void SRigVMEditorGraphExplorer::OnCopy()
 			}
 			break;
 		}
-		case Variable:
+		case ERigVMExplorerElementType::Variable:
 		{
 			int32 VarIndex = FBlueprintEditorUtils::FindNewVariableIndex(Blueprint, *Selection[0].Name);
 			if (VarIndex != INDEX_NONE)
@@ -1062,6 +1194,27 @@ void SRigVMEditorGraphExplorer::OnCopy()
 			}
 			break;
 		}
+		case ERigVMExplorerElementType::LocalVariable:
+		{
+			if (URigVMGraph* Graph = Blueprint->GetFocusedModel())
+			{
+				if (const UEdGraph* FocusedGraph = Blueprint->GetEdGraph(Graph))
+				{
+					TArray<FBPVariableDescription> LocalVariables;
+					Schema->GetLocalVariables(FocusedGraph, LocalVariables);
+					for (const FBPVariableDescription& VariableDescription : LocalVariables)
+					{
+						if (VariableDescription.VarName == *Selection[0].Name)
+						{
+							FBPVariableDescription::StaticStruct()->ExportText(OutputString, &VariableDescription, &VariableDescription, nullptr, 0, nullptr, false);
+							OutputString = TEXT("BPVar") + OutputString;
+							break;
+						}
+					}
+				}
+			}
+			break;
+		}
 	}
 	
 
@@ -1076,7 +1229,7 @@ bool SRigVMEditorGraphExplorer::CanCopy() const
 	TArray<FRigVMExplorerElementKey> Selection = TreeView->GetSelectedKeys();
 	if (Selection.Num() == 1)
 	{
-		if (Selection[0].Type == Function || Selection[0].Type == Variable)
+		if (Selection[0].Type == ERigVMExplorerElementType::Function || Selection[0].Type == ERigVMExplorerElementType::Variable || Selection[0].Type == ERigVMExplorerElementType::LocalVariable)
 		{
 			return true;
 		}
@@ -1118,8 +1271,9 @@ bool SRigVMEditorGraphExplorer::CanDuplicate() const
 	const URigVMEdGraphSchema* Schema = GetDefault<URigVMEdGraphSchema>();
 	switch (Selection[0].Type)
 	{
-		case Function:
-		case Variable:
+		case ERigVMExplorerElementType::Function:
+		case ERigVMExplorerElementType::Variable:
+		case ERigVMExplorerElementType::LocalVariable:
 		{
 			return true;
 		}
@@ -1130,17 +1284,33 @@ bool SRigVMEditorGraphExplorer::CanDuplicate() const
 
 void SRigVMEditorGraphExplorer::OnPasteGeneric()
 {
+	TArray<FRigVMExplorerElementKey> Selected = TreeView->GetSelectedKeys();
+	if (Selected.Num() == 1)
+	{
+		if (Selected[0].Type == ERigVMExplorerElementType::Variable && CanPasteVariable())
+		{
+			OnPasteVariable();
+			return;
+		}
+		if (Selected[0].Type == ERigVMExplorerElementType::LocalVariable && CanPasteLocalVariable())
+		{
+			OnPasteLocalVariable();
+			return;
+		}
+	}
+
+	// try any of the other options
+
 	// prioritize pasting as a member variable if possible
 	if (CanPasteVariable())
 	{
 		OnPasteVariable();
 	}
-	// else if (CanPasteLocalVariable())
-	// {
-	// 	OnPasteLocalVariable();
-	// }
-	// else 
-	if (CanPasteFunction())
+	else if (CanPasteLocalVariable())
+	{
+		OnPasteLocalVariable();
+	}
+	else if (CanPasteFunction())
 	{
 		OnPasteFunction();
 	}
@@ -1149,7 +1319,7 @@ void SRigVMEditorGraphExplorer::OnPasteGeneric()
 bool SRigVMEditorGraphExplorer::CanPasteGeneric() const
 {
 	return CanPasteVariable() 
-	//|| CanPasteLocalVariable() 
+	|| CanPasteLocalVariable() 
 	|| CanPasteFunction() ;
 }
 
@@ -1240,6 +1410,38 @@ bool SRigVMEditorGraphExplorer::CanPasteVariable() const
 	return false;
 }
 
+void SRigVMEditorGraphExplorer::OnPasteLocalVariable()
+{
+	TSharedPtr<FRigVMEditor> Editor = RigVMEditor.Pin();
+	if (!Editor)
+	{
+		return;
+	}
+	URigVMBlueprint* Blueprint = Editor->GetRigVMBlueprint();
+	
+	FString ClipboardText;
+	FPlatformApplicationMisc::ClipboardPaste(ClipboardText);
+	if (!ensure(ClipboardText.StartsWith(TEXT("BPVar"), ESearchCase::CaseSensitive)))
+	{
+		return;
+	}
+
+	FBPVariableDescription Description;
+	FStringOutputDevice Errors;
+	const TCHAR* Import = ClipboardText.GetCharArray().GetData() + FCString::Strlen(TEXT("BPVar"));
+	FBPVariableDescription::StaticStruct()->ImportText(Import, &Description, nullptr, PPF_None, &Errors, FBPVariableDescription::StaticStruct()->GetName());
+	if (Errors.IsEmpty())
+	{
+		Editor->OnPasteNewLocalVariable(Description);
+		Refresh();
+	}
+}
+
+bool SRigVMEditorGraphExplorer::CanPasteLocalVariable() const
+{
+	return CanPasteVariable();
+}
+
 void SRigVMEditorGraphExplorer::OnToggleShowEmptySections()
 {
 	// FIXME: Move to preferences
@@ -1337,6 +1539,29 @@ bool SRigVMEditorGraphExplorer::CanCreateVariable() const
 	return true;
 }
 
+void SRigVMEditorGraphExplorer::OnCreateLocalVariable()
+{
+	TSharedPtr<FRigVMEditor> Editor = RigVMEditor.Pin();
+	if (!Editor)
+	{
+		return;
+	}
+
+	Editor->OnAddNewLocalVariable();
+	bNeedsRefresh = true;
+}
+
+bool SRigVMEditorGraphExplorer::CanCreateLocalVariable() const
+{
+	TSharedPtr<FRigVMEditor> Editor = RigVMEditor.Pin();
+	if (!Editor)
+	{
+		return false;
+	}
+
+	return Editor->CanAddNewLocalVariable();
+}
+
 void SRigVMEditorGraphExplorer::OnAddFunctionVariant()
 {
 	TSharedPtr<FRigVMEditor> Editor = RigVMEditor.Pin();
@@ -1347,7 +1572,7 @@ void SRigVMEditorGraphExplorer::OnAddFunctionVariant()
 	}
 	
 	TArray<FRigVMExplorerElementKey> Selected = TreeView->GetSelectedKeys();
-	if (Selected.Num() != 1 || Selected[0].Type != Function)
+	if (Selected.Num() != 1 || Selected[0].Type != ERigVMExplorerElementType::Function)
 	{
 		return;
 	}
@@ -1374,7 +1599,7 @@ bool SRigVMEditorGraphExplorer::CanAddFunctionVariant() const
 	}
 	
 	TArray<FRigVMExplorerElementKey> Selected = TreeView->GetSelectedKeys();
-	if (Selected.Num() != 1 || Selected[0].Type != Function)
+	if (Selected.Num() != 1 || Selected[0].Type != ERigVMExplorerElementType::Function)
 	{
 		return false;
 	}
@@ -1388,27 +1613,28 @@ void SRigVMEditorGraphExplorer::OnDeleteEntry()
 
 	if (Editor)
 	{
+		URigVMBlueprint* Blueprint = Editor->GetRigVMBlueprint();
 		TArray<FRigVMExplorerElementKey> Selected = TreeView->GetSelectedKeys();
 		const URigVMEdGraphSchema* Schema = GetDefault<URigVMEdGraphSchema>();
 		for (const FRigVMExplorerElementKey& Key : Selected)
 		{
 			switch (Key.Type)
 			{
-				case Graph:
+				case ERigVMExplorerElementType::Graph:
 				{
-					if (UEdGraph* EdGraph = Editor->GetRigVMBlueprint()->GetEdGraph(Key.Name))
+					if (UEdGraph* EdGraph = Blueprint->GetEdGraph(Key.Name))
 					{
 						Schema->TryDeleteGraph(EdGraph);
 					}
 					break;
 				}
-				case Function:
+				case ERigVMExplorerElementType::Function:
 				{
-					if (URigVMFunctionLibrary* Library = Editor->GetRigVMBlueprint()->GetLocalFunctionLibrary())
+					if (URigVMFunctionLibrary* Library = Blueprint->GetLocalFunctionLibrary())
 					{
 						if (URigVMLibraryNode* Function = Library->FindFunction(*Key.Name))
 						{
-							if (UEdGraph* EdGraph = Editor->GetRigVMBlueprint()->GetEdGraph(Function->GetContainedGraph()))
+							if (UEdGraph* EdGraph = Blueprint->GetEdGraph(Function->GetContainedGraph()))
 							{
 								Schema->TryDeleteGraph(EdGraph);
 							}
@@ -1416,15 +1642,25 @@ void SRigVMEditorGraphExplorer::OnDeleteEntry()
 					}
 					break;
 				}
-				case Variable:
+				case ERigVMExplorerElementType::Variable:
 				{
 					const FScopedTransaction Transaction( LOCTEXT( "RemoveVariable", "Remove Variable" ) );
-					URigVMBlueprint* Blueprint = Editor->GetRigVMBlueprint();
 
 					Blueprint->Modify();
 					FBlueprintEditorUtils::RemoveMemberVariable(Blueprint, *Key.Name);
 					break;
 				}
+				case ERigVMExplorerElementType::LocalVariable:
+				{
+					if (URigVMGraph* FocusedGraph = Editor->GetFocusedModel())
+					{
+						if (URigVMController* Controller = Blueprint->GetRigVMClient()->GetController(FocusedGraph))
+						{
+							Controller->RemoveLocalVariable(*Key.Name, true, true);
+						}
+					}
+					break;
+				}
 			}
 		}
 	}
@@ -1458,17 +1694,15 @@ bool SRigVMEditorGraphExplorer::CanDeleteEntry() const
 			case ERigVMExplorerElementType::Graph:
 			case ERigVMExplorerElementType::Function:
 			case ERigVMExplorerElementType::Variable:
+			case ERigVMExplorerElementType::LocalVariable:
 				{
+					return true;
 					break;
 				}
-			default:
-				{
-					return false;
-				}
 		}
 	}
 	
-	return true;
+	return false;
 }
 
 
@@ -1590,7 +1824,7 @@ bool SRigVMEditorGraphExplorer::OnCanRenameFunction(const FString& InOldPath, co
 	return bResult;
 }
 
-bool SRigVMEditorGraphExplorer::OnRenameVariable(const FString& InOldName, const FString& InNewName)
+bool SRigVMEditorGraphExplorer::OnRenameVariable(const FRigVMExplorerElementKey& InOldKey, const FString& InNewName)
 {
 	TSharedPtr<FRigVMEditor> Editor = RigVMEditor.Pin();
 	if (!Editor)
@@ -1599,7 +1833,7 @@ bool SRigVMEditorGraphExplorer::OnRenameVariable(const FString& InOldName, const
 	}
 	
 	// Check if the name is unchanged
-	if (InNewName.Equals(InOldName, ESearchCase::CaseSensitive))
+	if (InNewName.Equals(InOldKey.Name, ESearchCase::CaseSensitive))
 	{
 		return false;
 	}
@@ -1607,22 +1841,36 @@ bool SRigVMEditorGraphExplorer::OnRenameVariable(const FString& InOldName, const
 	const FScopedTransaction Transaction( LOCTEXT( "RenameVariable", "Rename Variable" ) );
 
 	URigVMBlueprint* Blueprint = Editor->GetRigVMBlueprint();
-	Blueprint->Modify();
 
-	FBlueprintEditorUtils::RenameMemberVariable(Blueprint, *InOldName, *InNewName);
+	if (InOldKey.Type == ERigVMExplorerElementType::Variable)
+	{
+		Blueprint->Modify();
+		FBlueprintEditorUtils::RenameMemberVariable(Blueprint, *InOldKey.Name, *InNewName);
+	}
+	else if (InOldKey.Type == ERigVMExplorerElementType::LocalVariable)
+	{
+		if (URigVMGraph* Graph = Blueprint->GetFocusedModel())
+		{
+			if (URigVMController* Controller = Blueprint->GetRigVMClient()->GetController(Graph))
+			{
+				Controller->RenameLocalVariable(*InOldKey.Name, *InNewName, true, true);
+			}
+		}
+	}
 	return true;
 }
 
-bool SRigVMEditorGraphExplorer::OnCanRenameVariable(const FString& InOldName, const FString& InNewName, FText& OutErrorMessage) const
+bool SRigVMEditorGraphExplorer::OnCanRenameVariable(const FRigVMExplorerElementKey& InOldKey, const FString& InNewName, FText& OutErrorMessage) const
 {
-	if (InOldName == InNewName)
+	if (InNewName.Equals(InOldKey.Name, ESearchCase::CaseSensitive))
 	{
 		return true;
 	}
 
-	for (FRigVMGraphVariableDescription& Variable : GetVariables())
+	TArray<FRigVMGraphVariableDescription> Variables = (InOldKey.Type == ERigVMExplorerElementType::Variable) ? GetVariables() : GetLocalVariables();
+	for (FRigVMGraphVariableDescription& Variable : Variables)
 	{
-		if (Variable.Name == InOldName)
+		if (Variable.Name == InOldKey.Name)
 		{
 			continue;
 		}
@@ -1658,7 +1906,7 @@ void SRigVMEditorGraphExplorer::OnFindReference(bool bSearchAllBlueprints)
 
 	switch (Selection[0].Type)
 	{
-		case Graph:
+		case ERigVMExplorerElementType::Graph:
 		{
 			if (UEdGraph* EdGraph = Editor->GetRigVMBlueprint()->GetEdGraph(Selection[0].Name))
 			{
@@ -1668,8 +1916,9 @@ void SRigVMEditorGraphExplorer::OnFindReference(bool bSearchAllBlueprints)
 			}
 			break;
 		}
-		case Function:
-		case Variable:
+		case ERigVMExplorerElementType::Function:
+		case ERigVMExplorerElementType::Variable:
+		case ERigVMExplorerElementType::LocalVariable:
 		{
 			SearchTerm = Selection[0].Name;
 			break;
diff --git a/Engine/Plugins/Runtime/RigVM/Source/RigVMEditor/Private/Widgets/SRigVMEditorGraphExplorerTreeView.cpp b/Engine/Plugins/Runtime/RigVM/Source/RigVMEditor/Private/Widgets/SRigVMEditorGraphExplorerTreeView.cpp
--- a/Engine/Plugins/Runtime/RigVM/Source/RigVMEditor/Private/Widgets/SRigVMEditorGraphExplorerTreeView.cpp
+++ b/Engine/Plugins/Runtime/RigVM/Source/RigVMEditor/Private/Widgets/SRigVMEditorGraphExplorerTreeView.cpp
@@ -15,6 +15,7 @@
 static const FText GraphsSectionName = LOCTEXT("Graphs", "Graphs");
 static const FText FunctionsSectionName = LOCTEXT("Functions", "Functions");
 static const FText VariablesSectionName = LOCTEXT("Variables", "Variables");
+static const FText LocalVariablesSectionName = LOCTEXT("LocalVariables", "Local Variables");
 
 void SRigVMEditorGraphExplorerItem::Construct(const FArguments& InArgs,
                                               const TSharedRef<STableViewBase>& InOwnerTable,
@@ -27,7 +28,7 @@ void SRigVMEditorGraphExplorerItem::Construct(const FArguments& InArgs,
 
 	FSlateFontInfo NameFont = FCoreStyle::GetDefaultFontStyle("Regular", 10);
 
-	if (InElement->Key.Type == Section)
+	if (InElement->Key.Type == ERigVMExplorerElementType::Section)
 	{
 		STableRow<TSharedPtr<FRigVMEditorGraphExplorerTreeElement>>::Construct(
 		STableRow<TSharedPtr<FRigVMEditorGraphExplorerTreeElement>>::FArguments()
@@ -75,8 +76,8 @@ void SRigVMEditorGraphExplorerItem::Construct(const FArguments& InArgs,
 			
 		], InOwnerTable);
 	}
-	else if(InElement->Key.Type == FunctionCategory
-		|| InElement->Key.Type == VariableCategory)
+	else if(InElement->Key.Type == ERigVMExplorerElementType::FunctionCategory
+		|| InElement->Key.Type == ERigVMExplorerElementType::VariableCategory)
 	{
 		int32 LastSeparator;
 		InElement->Key.Name.FindLastChar('|', LastSeparator);
@@ -102,13 +103,9 @@ void SRigVMEditorGraphExplorerItem::Construct(const FArguments& InArgs,
 			]
 		], InOwnerTable);
 	}
-	else if (InElement->Key.Type == Variable)
+	else if (InElement->Key.Type == ERigVMExplorerElementType::Variable || InElement->Key.Type == ERigVMExplorerElementType::LocalVariable)
 	{
-		// TArray<TSharedPtr<IPinTypeSelectorFilter>> CustomPinTypeFilters;
-		// if (BlueprintEditorPtr.IsValid())
-		// {
-		// 	BlueprintEditorPtr.Pin()->GetPinTypeSelectorFilters(CustomPinTypeFilters);
-		// }
+		TArray<TSharedPtr<IPinTypeSelectorFilter>> CustomPinTypeFilters = Delegates.GetCustomPinFilters();
 		
 		STableRow<TSharedPtr<FRigVMEditorGraphExplorerTreeElement>>::Construct(
 		   STableRow<TSharedPtr<FRigVMEditorGraphExplorerTreeElement>>::FArguments()
@@ -116,11 +113,6 @@ void SRigVMEditorGraphExplorerItem::Construct(const FArguments& InArgs,
 		   .OnDragDetected(Delegates.OnDragDetected)
 		   .Content()
 		   [
-			   // if (ActionAccessSpecifier != EAccessSpecifier::None && GraphAction->GetTypeId() != FEdGraphSchemaAction_K2LocalVar::StaticGetTypeId())
-			   // {
-				  //  CreateAccessSpecifierLambda();
-			   // }
-
 			   SNew(SHorizontalBox)
 			   +SHorizontalBox::Slot()
 				   .FillWidth(0.6f)
@@ -137,11 +129,11 @@ void SRigVMEditorGraphExplorerItem::Construct(const FArguments& InArgs,
 				   [
 						SNew(SPinTypeSelector, FGetPinTypeTree::CreateUObject(GetDefault<UEdGraphSchema_K2>(), &UEdGraphSchema_K2::GetVariableTypeTree))
 						.Schema(GetDefault<URigVMEdGraphSchema>())
-						.TargetPinType_Lambda([this, InElement]() { return Delegates.GetVariablePinType(InElement->Key.Name); })
-						.OnPinTypeChanged_Lambda([this, InElement](const FEdGraphPinType& InType) { Delegates.SetVariablePinType(InElement->Key.Name, InType); })
+						.TargetPinType_Lambda([this, InElement]() { return Delegates.GetVariablePinType(InElement->Key); })
+						.OnPinTypeChanged_Lambda([this, InElement](const FEdGraphPinType& InType) { Delegates.SetVariablePinType(InElement->Key, InType); })
 						.TypeTreeFilter(ETypeTreeFilter::None)
 						.SelectorType(SPinTypeSelector::ESelectorType::Partial)
-						//.CustomFilters(CustomPinTypeFilters)
+						.CustomFilters(CustomPinTypeFilters)
 				   ]
 
 				+SHorizontalBox::Slot()
@@ -152,6 +144,7 @@ void SRigVMEditorGraphExplorerItem::Construct(const FArguments& InArgs,
 				   [
 					   SNew(SBorder)
 					   .Padding(0.0f)
+					   .Visibility(InElement->Key.Type == ERigVMExplorerElementType::Variable ? EVisibility::Visible : EVisibility::Collapsed)
 					   .BorderImage(FStyleDefaults::GetNoBrush())
 					   [
 						   SNew(SCheckBox)
@@ -225,20 +218,20 @@ TSharedRef<SWidget> SRigVMEditorGraphExplorerItem::CreateIconWidget(const FRigVM
 
 	switch (Key.Type)
 	{
-		case Graph:
+		case ERigVMExplorerElementType::Graph:
 			{
 				IconWidget = SNew(SImage)
 							.Image(Delegates.GetGraphIcon(Key.Name));
 				break;
 			}
-		case Event:
+		case ERigVMExplorerElementType::Event:
 			{
 				static FSlateIcon EventIcon(FAppStyle::GetAppStyleSetName(), "GraphEditor.Event_16x");
 				IconWidget = SNew(SImage)
 							.Image(EventIcon.GetIcon());
 				break;
 			}
-		case Function:
+		case ERigVMExplorerElementType::Function:
 			{
 				static FSlateIcon FunctionIcon(FAppStyle::GetAppStyleSetName(), "Kismet.AllClasses.FunctionIcon");
 				IconWidget = SNew(SImage)
@@ -282,16 +275,17 @@ FText SRigVMEditorGraphExplorerItem::GetDisplayText() const
 
 	switch (Element->Key.Type)
 	{
-		case Graph:
+		case ERigVMExplorerElementType::Graph:
 			{
 				return Delegates.GetGraphDisplayName(Element->Key.Name);
 			}
-		case Event:
+		case ERigVMExplorerElementType::Event:
 			{
 				return Delegates.GetEventDisplayName(Element->Key.Name);
 			}
-		case Function:
-		case Variable:
+		case ERigVMExplorerElementType::Function:
+		case ERigVMExplorerElementType::Variable:
+		case ERigVMExplorerElementType::LocalVariable:
 			{
 				return FText::FromString(Element->Key.Name);
 			}
@@ -311,11 +305,11 @@ FText SRigVMEditorGraphExplorerItem::GetItemTooltip() const
 
 	switch (Element->Key.Type)
 	{
-		case Graph:
+		case ERigVMExplorerElementType::Graph:
 			{
 				return Delegates.GetGraphTooltip(Element->Key.Name);
 			}
-		case Event:
+		case ERigVMExplorerElementType::Event:
 			{
 				return Delegates.GetEventDisplayName(Element->Key.Name);
 			}
@@ -331,29 +325,30 @@ bool SRigVMEditorGraphExplorerItem::OnVerifyNameChanged(const FText& InText, FTe
 
 	switch (OldKey.Type)
 	{
-		case Section:
+		case ERigVMExplorerElementType::Section:
 			{
 				break;
 			}
-		case FunctionCategory:
-		case VariableCategory:
+		case ERigVMExplorerElementType::FunctionCategory:
+		case ERigVMExplorerElementType::VariableCategory:
 			{
 				return true;
 				break;
 			}
-		case Graph:
+		case ERigVMExplorerElementType::Graph:
 			{
 				return Delegates.CanRenameGraph(OldKey.Name, *NewName, OutErrorMessage);
 				break;
 			}
-		case Function:
+		case ERigVMExplorerElementType::Function:
 			{
 				return Delegates.CanRenameFunction(OldKey.Name, *NewName, OutErrorMessage);
 				break;
 			}
-		case Variable:
+		case ERigVMExplorerElementType::Variable:
+		case ERigVMExplorerElementType::LocalVariable:
 			{
-				return Delegates.CanRenameVariable(OldKey.Name, *NewName, OutErrorMessage);
+				return Delegates.CanRenameVariable(OldKey, *NewName, OutErrorMessage);
 				break;
 			}
 	}
@@ -376,12 +371,12 @@ void SRigVMEditorGraphExplorerItem::OnNameCommitted(const FText& InText, ETextCo
 
 		switch (OldKey.Type)
 		{
-			case Section:
+			case ERigVMExplorerElementType::Section:
 				{
 					break;
 				}
-			case FunctionCategory:
-			case VariableCategory:
+			case ERigVMExplorerElementType::FunctionCategory:
+			case ERigVMExplorerElementType::VariableCategory:
 				{
 					int32 LastSeparator;
 					FString OldCategoryPath = Element->Key.Name;
@@ -400,20 +395,20 @@ void SRigVMEditorGraphExplorerItem::OnNameCommitted(const FText& InText, ETextCo
 					for (int32 i=0; i<SubElements.Num(); ++i)
 					{
 						TSharedPtr<FRigVMEditorGraphExplorerTreeElement> SubElement = SubElements[i];
-						if (SubElement->Key.Type == Function || SubElement->Key.Type == Variable)
+						if (SubElement->Key.Type == ERigVMExplorerElementType::Function || SubElement->Key.Type == ERigVMExplorerElementType::Variable || SubElement->Key.Type == ERigVMExplorerElementType::LocalVariable)
 						{
-							FString Category = OldKey.Type == FunctionCategory ?
+							FString Category = OldKey.Type == ERigVMExplorerElementType::FunctionCategory ?
 								Delegates.GetFunctionCategory(SubElement->Key.Name)
 									: Delegates.GetVariableCategory(SubElement->Key.Name);
 							Category.RemoveFromStart(OldCategoryPath);
 							Category.InsertAt(0, NewCategoryPath);
 							Category.RemoveFromStart(TEXT("|"));
 							Category.RemoveFromEnd(TEXT("|"));
-							if (OldKey.Type == FunctionCategory)
+							if (OldKey.Type == ERigVMExplorerElementType::FunctionCategory)
 							{
 								Delegates.SetFunctionCategory(SubElement->Key.Name, Category);
 							}
-							else if (OldKey.Type == VariableCategory)
+							else if (OldKey.Type == ERigVMExplorerElementType::VariableCategory)
 							{
 								Delegates.SetVariableCategory(SubElement->Key.Name, Category);
 							}
@@ -422,19 +417,20 @@ void SRigVMEditorGraphExplorerItem::OnNameCommitted(const FText& InText, ETextCo
 					}
 					break;
 				}
-			case Graph:
+			case ERigVMExplorerElementType::Graph:
 				{
 					Delegates.RenameGraph(OldKey.Name, *NewName);
 					break;
 				}
-			case Function:
+			case ERigVMExplorerElementType::Function:
 				{
 					Delegates.RenameFunction(OldKey.Name, *NewName);
 					break;
 				}
-			case Variable:
+			case ERigVMExplorerElementType::Variable:
+			case ERigVMExplorerElementType::LocalVariable:
 				{
-					Delegates.RenameVariable(OldKey.Name, *NewName);
+					Delegates.RenameVariable(OldKey, *NewName);
 					break;
 				}
 		}
@@ -510,7 +506,7 @@ void SRigVMEditorGraphExplorerTreeView::RefreshTreeView(bool bRebuildContent)
 		{
 			// GRAPHS section
 			{
-				FRigVMExplorerElementKey SectionKey(Section, GraphsSectionName.ToString());
+				FRigVMExplorerElementKey SectionKey(ERigVMExplorerElementType::Section, GraphsSectionName.ToString());
 				TSharedPtr<FRigVMEditorGraphExplorerTreeElement> SectionItem = MakeShared<FRigVMEditorGraphExplorerTreeElement>(SectionKey, SharedThis(this));
 				ElementMap.Add(SectionKey, SectionItem);
 				RootElements.Add(SectionItem);
@@ -546,7 +542,7 @@ void SRigVMEditorGraphExplorerTreeView::RefreshTreeView(bool bRebuildContent)
 					for (const URigVMNode* EventNode : EventNodes)
 					{
 					
-						FRigVMExplorerElementKey EventKey(Event, EventNode->GetNodePath());
+						FRigVMExplorerElementKey EventKey(ERigVMExplorerElementType::Event, EventNode->GetNodePath());
 						ParentMap.Add(EventKey, Key);
 						TSharedPtr<FRigVMEditorGraphExplorerTreeElement> EventItem = MakeShared<FRigVMEditorGraphExplorerTreeElement>(EventKey, SharedThis(this));
 						ElementMap.Add(EventKey, EventItem);
@@ -562,7 +558,7 @@ void SRigVMEditorGraphExplorerTreeView::RefreshTreeView(bool bRebuildContent)
 
 			// FUNCTIONS section
 			{
-				FRigVMExplorerElementKey SectionKey(Section, FunctionsSectionName.ToString());
+				FRigVMExplorerElementKey SectionKey(ERigVMExplorerElementType::Section, FunctionsSectionName.ToString());
 				TSharedPtr<FRigVMEditorGraphExplorerTreeElement> SectionItem = MakeShared<FRigVMEditorGraphExplorerTreeElement>(SectionKey, SharedThis(this));
 				ElementMap.Add(SectionKey, SectionItem);
 				RootElements.Add(SectionItem);
@@ -583,7 +579,7 @@ void SRigVMEditorGraphExplorerTreeView::RefreshTreeView(bool bRebuildContent)
 							CategoryParentPath.Append(TEXT("|"));
 						}
 						CategoryParentPath.Append(SingleCategory);
-						FRigVMExplorerElementKey ParentKey(FunctionCategory, CategoryParentPath);
+						FRigVMExplorerElementKey ParentKey(ERigVMExplorerElementType::FunctionCategory, CategoryParentPath);
 						TSharedPtr<FRigVMEditorGraphExplorerTreeElement> ParentItem = FindElement(ParentKey);
 						if (!ParentItem)
 						{
@@ -608,7 +604,7 @@ void SRigVMEditorGraphExplorerTreeView::RefreshTreeView(bool bRebuildContent)
 
 			// VARIABLES section
 			{
-				FRigVMExplorerElementKey SectionKey(Section, VariablesSectionName.ToString());
+				FRigVMExplorerElementKey SectionKey(ERigVMExplorerElementType::Section, VariablesSectionName.ToString());
 				TSharedPtr<FRigVMEditorGraphExplorerTreeElement> SectionItem = MakeShared<FRigVMEditorGraphExplorerTreeElement>(SectionKey, SharedThis(this));
 				ElementMap.Add(SectionKey, SectionItem);
 				RootElements.Add(SectionItem);
@@ -631,7 +627,7 @@ void SRigVMEditorGraphExplorerTreeView::RefreshTreeView(bool bRebuildContent)
 								CategoryParentPath.Append(TEXT("|"));
 							}
 							CategoryParentPath.Append(SingleCategory);
-							FRigVMExplorerElementKey ParentKey(VariableCategory, CategoryParentPath);
+							FRigVMExplorerElementKey ParentKey(ERigVMExplorerElementType::VariableCategory, CategoryParentPath);
 							TSharedPtr<FRigVMEditorGraphExplorerTreeElement> ParentItem = FindElement(ParentKey);
 							if (!ParentItem)
 							{
@@ -654,6 +650,28 @@ void SRigVMEditorGraphExplorerTreeView::RefreshTreeView(bool bRebuildContent)
 					SetItemExpansion(NewItem, true);
 				}
 			}
+
+			// LOCAL VARIABLES section
+			if (Delegates.IsFunctionFocused())
+			{
+				FRigVMExplorerElementKey SectionKey(ERigVMExplorerElementType::Section, LocalVariablesSectionName.ToString());
+				TSharedPtr<FRigVMEditorGraphExplorerTreeElement> SectionItem = MakeShared<FRigVMEditorGraphExplorerTreeElement>(SectionKey, SharedThis(this));
+				ElementMap.Add(SectionKey, SectionItem);
+				RootElements.Add(SectionItem);
+				SetItemExpansion(SectionItem, true);
+
+				TArray<FRigVMGraphVariableDescription> LocalVariables = Delegates.GetLocalVariables();
+				for (const FRigVMGraphVariableDescription& Variable : LocalVariables)
+				{
+					FRigVMExplorerElementKey Key(ERigVMExplorerElementType::LocalVariable, Variable.Name.ToString());
+				
+					TSharedPtr<FRigVMEditorGraphExplorerTreeElement> NewItem = MakeShared<FRigVMEditorGraphExplorerTreeElement>(Key, SharedThis(this));
+					ElementMap.Add(Key, NewItem);
+					ParentMap.Add(Key, SectionItem->Key);
+					SectionItem->Children.Add(NewItem);
+					SetItemExpansion(NewItem, true);
+				}
+			}
 		}
 
 		// Filter Items
@@ -668,20 +686,21 @@ void SRigVMEditorGraphExplorerTreeView::RefreshTreeView(bool bRebuildContent)
 						FText DisplayName;
 						switch (InElement->Key.Type)
 						{
-						case Graph:
+						case ERigVMExplorerElementType::Graph:
 							{
 								DisplayName = Delegates.GetGraphDisplayName(InElement->Key.Name);
 								break;
 							}
-						case Event:
+						case ERigVMExplorerElementType::Event:
 							{
 								DisplayName = Delegates.GetEventDisplayName(InElement->Key.Name);
 								break;
 							}
-						case Function:
-						case Variable:
-						case FunctionCategory:
-						case VariableCategory:
+						case ERigVMExplorerElementType::Function:
+						case ERigVMExplorerElementType::Variable:
+						case ERigVMExplorerElementType::LocalVariable:
+						case ERigVMExplorerElementType::FunctionCategory:
+						case ERigVMExplorerElementType::VariableCategory:
 							{
 								DisplayName = FText::FromString(InElement->Key.Name);
 								break;
@@ -850,31 +869,32 @@ void SRigVMEditorGraphExplorerTreeView::OnItemClicked(TSharedPtr<FRigVMEditorGra
 {
 	switch(InElement->Key.Type)
 	{
-		case Section:
+		case ERigVMExplorerElementType::Section:
 			{
-				FRigVMExplorerElementKey SectionKey(Section, InElement->Key.Name);
+				FRigVMExplorerElementKey SectionKey(ERigVMExplorerElementType::Section, InElement->Key.Name);
 				TSharedPtr<FRigVMEditorGraphExplorerTreeElement> Section = ElementMap.FindChecked(SectionKey);
 				SetItemExpansion(Section, !IsItemExpanded(Section));
 				break;
 			}
-		case Graph:
+		case ERigVMExplorerElementType::Graph:
 			{
 				Delegates.GraphClicked(InElement->Key.Name);
 				break;
 			}
-		case Event:
+		case ERigVMExplorerElementType::Event:
 			{
 				Delegates.EventClicked(InElement->Key.Name);
 				break;
 			}
-		case Function:
+		case ERigVMExplorerElementType::Function:
 			{
 				Delegates.FunctionClicked(InElement->Key.Name);
 				break;
 			}
-		case Variable:
+		case ERigVMExplorerElementType::Variable:
+		case ERigVMExplorerElementType::LocalVariable:
 			{
-				Delegates.VariableClicked(InElement->Key.Name);
+				Delegates.VariableClicked(InElement->Key);
 				break;
 			}
 	}
@@ -884,21 +904,22 @@ void SRigVMEditorGraphExplorerTreeView::OnItemDoubleClicked(TSharedPtr<FRigVMEdi
 {
 	switch(InElement->Key.Type)
 	{
-		case Section:
-		case Variable:
+		case ERigVMExplorerElementType::Section:
+		case ERigVMExplorerElementType::Variable:
+		case ERigVMExplorerElementType::LocalVariable:
 			{
 				break;
 			}
-		case Graph:
+		case ERigVMExplorerElementType::Graph:
 			{
 				Delegates.GraphDoubleClicked(InElement->Key.Name);
 				break;
 			}
-		case Event:
+		case ERigVMExplorerElementType::Event:
 			{
 				Delegates.EventDoubleClicked(InElement->Key.Name);
 			}
-		case Function:
+		case ERigVMExplorerElementType::Function:
 			{
 				Delegates.FunctionDoubleClicked(InElement->Key.Name);
 			}
@@ -922,6 +943,11 @@ FReply SRigVMEditorGraphExplorerTreeView::OnAddButtonClickedOnSection(const FRig
 		Delegates.CreateVariable();
 		return FReply::Handled();
 	}
+	else if (InSectionKey.Name == LocalVariablesSectionName.ToString())
+	{
+		Delegates.CreateLocalVariable();
+		return FReply::Handled();
+	}
 	return FReply::Unhandled();
 }
 
diff --git a/Engine/Plugins/Runtime/RigVM/Source/RigVMEditor/Public/Editor/RigVMEditor.h b/Engine/Plugins/Runtime/RigVM/Source/RigVMEditor/Public/Editor/RigVMEditor.h
--- a/Engine/Plugins/Runtime/RigVM/Source/RigVMEditor/Public/Editor/RigVMEditor.h
+++ b/Engine/Plugins/Runtime/RigVM/Source/RigVMEditor/Public/Editor/RigVMEditor.h
@@ -65,6 +65,7 @@ public:
 	virtual UBlueprint* GetBlueprintObj() const override;
 	virtual TSubclassOf<UEdGraphSchema> GetDefaultSchemaClass() const override;
 	virtual bool InEditingMode() const override;
+	virtual void OnGraphEditorFocused(const TSharedRef<SGraphEditor>& InGraphEditor) override;
 
 	//  FTickableEditorObject Interface
 	virtual void Tick(float DeltaTime) override;
@@ -157,7 +158,6 @@ protected:
 	virtual void SaveAsset_Execute() override;
 	virtual void SaveAssetAs_Execute() override;
 	virtual bool IsInAScriptingMode() const override { return true; }
-	virtual void CreateDefaultTabContents(const TArray<UBlueprint*>& InBlueprints) override;
 	virtual void NewDocument_OnClicked(ECreatedDocumentType GraphType) override;
 	virtual bool IsSectionVisible(NodeSectionID::Type InSectionID) const override;
 	virtual bool AreEventGraphsAllowed() const override;
@@ -342,6 +342,10 @@ private:
 
 	TSharedPtr<SRigVMEditorGraphExplorer> GraphExplorerWidget;
 
+public:
+	void SetGraphExplorerWidget(TSharedPtr<SRigVMEditorGraphExplorer> InWidget) { GraphExplorerWidget = InWidget; }
+
+private:
 	bool bAnyErrorsLeft;
 	TMap<FString, FString> KnownInstructionLimitWarnings;
 	URigVMNode* HaltedAtNode;
diff --git a/Engine/Plugins/Runtime/RigVM/Source/RigVMEditor/Public/Widgets/SRigVMEditorGraphExplorer.h b/Engine/Plugins/Runtime/RigVM/Source/RigVMEditor/Public/Widgets/SRigVMEditorGraphExplorer.h
--- a/Engine/Plugins/Runtime/RigVM/Source/RigVMEditor/Public/Widgets/SRigVMEditorGraphExplorer.h
+++ b/Engine/Plugins/Runtime/RigVM/Source/RigVMEditor/Public/Widgets/SRigVMEditorGraphExplorer.h
@@ -63,9 +63,11 @@ public:
 	TSharedPtr<FUICommandInfo> CreateGraph;
 	TSharedPtr<FUICommandInfo> CreateFunction;
 	TSharedPtr<FUICommandInfo> CreateVariable;
+	TSharedPtr<FUICommandInfo> CreateLocalVariable;
 	TSharedPtr<FUICommandInfo> AddFunctionVariant;
 	TSharedPtr<FUICommandInfo> PasteFunction;
 	TSharedPtr<FUICommandInfo> PasteVariable;
+	TSharedPtr<FUICommandInfo> PasteLocalVariable;
 
 	void RegisterCommands() override;
 };
@@ -101,6 +103,7 @@ private:
 	TArray<URigVMNode*> GetEventNodesInGraph(const FString& InParentGraphPath) const;
 	TArray<URigVMLibraryNode*> GetFunctions() const;
 	TArray<FRigVMGraphVariableDescription> GetVariables() const;
+	TArray<FRigVMGraphVariableDescription> GetLocalVariables() const;
 	FText GetGraphDisplayName(const FString& InGraphPath) const;
 	FText GetEventDisplayName(const FString& InNodePath) const;
 	const FSlateBrush* GetGraphIcon(const FString& InGraphPath) const;
@@ -108,7 +111,7 @@ private:
 	void OnGraphClicked(const FString& InGraphPath);
 	void OnEventClicked(const FString& InEventPath);
 	void OnFunctionClicked(const FString& InFunctionPath);
-	void OnVariableClicked(const FString& InVariable);
+	void OnVariableClicked(const FRigVMExplorerElementKey& InVariable);
 	void OnGraphDoubleClicked(const FString& InGraphPath);
 	void OnEventDoubleClicked(const FString& InEventPath);
 	void OnFunctionDoubleClicked(const FString& InFunctionPath);
@@ -116,10 +119,12 @@ private:
 	FString OnGetFunctionCategory(const FString& InPath) const;
 	bool OnSetVariableCategory(const FString& InVariable, const FString& InCategory);
 	FString OnGetVariableCategory(const FString& InVariable) const;
-	FEdGraphPinType OnGetVariablePinType(const FString& InVariableName);
-	bool OnSetVariablePinType(const FString& InVariableName, const FEdGraphPinType& InType);
+	FEdGraphPinType OnGetVariablePinType(const FRigVMExplorerElementKey& InVariable);
+	bool OnSetVariablePinType(const FRigVMExplorerElementKey& InVariable, const FEdGraphPinType& InType);
 	bool OnIsVariablePublic(const FString& InVariableName) const;
 	bool OnToggleVariablePublic(const FString& InVariableName) const;
+	bool IsFunctionFocused() const;
+	TArray<TSharedPtr<IPinTypeSelectorFilter>> GetCustomPinFilters() const;
 	
 	TSharedPtr<SWidget> OnContextMenuOpening();
 	FReply OnDragDetected(const FGeometry& MyGeometry, const FPointerEvent& MouseEvent);
@@ -136,6 +141,8 @@ private:
 	bool CanPasteFunction() const;
 	void OnPasteVariable();
 	bool CanPasteVariable() const;
+	void OnPasteLocalVariable();
+	bool CanPasteLocalVariable() const;
 
 
 
@@ -156,6 +163,9 @@ private:
 	void OnCreateVariable();
 	bool CanCreateVariable() const;
 
+	void OnCreateLocalVariable();
+	bool CanCreateLocalVariable() const;
+
 	void OnAddFunctionVariant();
 	bool CanAddFunctionVariant() const;
 	
@@ -168,8 +178,8 @@ private:
 	bool OnCanRenameGraph(const FString& InOldPath, const FString& InNewPath, FText& OutErrorMessage) const;
 	bool OnRenameFunction(const FString& InOldPath, const FString& InNewPath);
 	bool OnCanRenameFunction(const FString& InOldPath, const FString& InNewPath, FText& OutErrorMessage) const;
-	bool OnRenameVariable(const FString& InOldName, const FString& InNewName);
-	bool OnCanRenameVariable(const FString& InOldName, const FString& InNewName, FText& OutErrorMessage) const;
+	bool OnRenameVariable(const FRigVMExplorerElementKey& InOldKey, const FString& InNewName);
+	bool OnCanRenameVariable(const FRigVMExplorerElementKey& InOldKey, const FString& InNewName, FText& OutErrorMessage) const;
 
 	void OnFindReference(bool bSearchAllBlueprints);
 	bool CanFindReference() const;
diff --git a/Engine/Plugins/Runtime/RigVM/Source/RigVMEditor/Public/Widgets/SRigVMEditorGraphExplorerTreeView.h b/Engine/Plugins/Runtime/RigVM/Source/RigVMEditor/Public/Widgets/SRigVMEditorGraphExplorerTreeView.h
--- a/Engine/Plugins/Runtime/RigVM/Source/RigVMEditor/Public/Widgets/SRigVMEditorGraphExplorerTreeView.h
+++ b/Engine/Plugins/Runtime/RigVM/Source/RigVMEditor/Public/Widgets/SRigVMEditorGraphExplorerTreeView.h
@@ -13,6 +13,7 @@ class SRigVMEditorGraphExplorerTreeView;
 class SRigVMEditorGraphExplorerItem;
 class FRigVMEditorGraphExplorerTreeElement;
 class URigVMGraph;
+class IPinTypeSelectorFilter;
 
 enum ERigVMExplorerElementType
 {
@@ -58,6 +59,7 @@ DECLARE_DELEGATE_RetVal_OneParam(TArray<const URigVMGraph*>, FRigVMGraphExplorer
 DECLARE_DELEGATE_RetVal_OneParam(TArray<URigVMNode*>, FRigVMGraphExplorer_OnGetEventNodesInGraph, const FString&);
 DECLARE_DELEGATE_RetVal(TArray<URigVMLibraryNode*>, FRigVMGraphExplorer_OnGetFunctions);
 DECLARE_DELEGATE_RetVal(TArray<FRigVMGraphVariableDescription>, FRigVMGraphExplorer_OnGetVariables);
+DECLARE_DELEGATE_RetVal(bool, FRigVMGraphExplorer_OnIsFunctionFocused);
 DECLARE_DELEGATE_RetVal_OneParam(FText, FRigVMGraphExplorer_OnGetGraphDisplayName, const FString&);
 DECLARE_DELEGATE_RetVal_OneParam(FText, FRigVMGraphExplorer_OnGetEventDisplayName, const FString&);
 DECLARE_DELEGATE_RetVal_OneParam(FText, FRigVMGraphExplorer_OnGetGraphTooltip, const FString&);
@@ -65,25 +67,26 @@ DECLARE_DELEGATE_RetVal_OneParam(const FSlateBrush*, FRigVMGraphExplorer_OnGetGr
 DECLARE_DELEGATE_OneParam(FRigVMGraphExplorer_OnGraphClicked, const FString&);
 DECLARE_DELEGATE_OneParam(FRigVMGraphExplorer_OnEventClicked, const FString&);
 DECLARE_DELEGATE_OneParam(FRigVMGraphExplorer_OnFunctionClicked, const FString&);
-DECLARE_DELEGATE_OneParam(FRigVMGraphExplorer_OnVariableClicked, const FString&);
+DECLARE_DELEGATE_OneParam(FRigVMGraphExplorer_OnVariableClicked, const FRigVMExplorerElementKey&);
 DECLARE_DELEGATE_OneParam(FRigVMGraphExplorer_OnGraphDoubleClicked, const FString&);
 DECLARE_DELEGATE_OneParam(FRigVMGraphExplorer_OnEventDoubleClicked, const FString&);
 DECLARE_DELEGATE_OneParam(FRigVMGraphExplorer_OnFunctionDoubleClicked, const FString&);
 DECLARE_DELEGATE_RetVal_TwoParams(bool, FRigVMGraphExplorer_OnRenameGraph, const FString&, const FString&);
 DECLARE_DELEGATE_RetVal_TwoParams(bool, FRigVMGraphExplorer_OnRenameFunction, const FString&, const FString&);
-DECLARE_DELEGATE_RetVal_TwoParams(bool, FRigVMGraphExplorer_OnRenameVariable, const FString&, const FString&);
+DECLARE_DELEGATE_RetVal_TwoParams(bool, FRigVMGraphExplorer_OnRenameVariable, const FRigVMExplorerElementKey&, const FString&);
 DECLARE_DELEGATE_RetVal_ThreeParams(bool, FRigVMGraphExplorer_OnCanRenameGraph, const FString&, const FString&, FText&);
 DECLARE_DELEGATE_RetVal_ThreeParams(bool, FRigVMGraphExplorer_OnCanRenameFunction, const FString&, const FString&, FText&);
-DECLARE_DELEGATE_RetVal_ThreeParams(bool, FRigVMGraphExplorer_OnCanRenameVariable, const FString&, const FString&, FText&);
+DECLARE_DELEGATE_RetVal_ThreeParams(bool, FRigVMGraphExplorer_OnCanRenameVariable, const FRigVMExplorerElementKey&, const FString&, FText&);
 DECLARE_DELEGATE_RetVal_TwoParams(bool, FRigVMGraphExplorer_OnSetFunctionCategory, const FString&, const FString&);
 DECLARE_DELEGATE_RetVal_OneParam(FString, FRigVMGraphExplorer_OnGetFunctionCategory, const FString&);
 DECLARE_DELEGATE_RetVal_TwoParams(bool, FRigVMGraphExplorer_OnSetVariableCategory, const FString&, const FString&);
 DECLARE_DELEGATE_RetVal_OneParam(FString, FRigVMGraphExplorer_OnGetVariableCategory, const FString&);
 DECLARE_DELEGATE_RetVal(TSharedPtr<SWidget>, FRigVMGraphExplorer_OnRequestContextMenu);
-DECLARE_DELEGATE_RetVal_OneParam(FEdGraphPinType, FRigVMGraphExplorer_OnGetVariablePinType, const FString&);
-DECLARE_DELEGATE_RetVal_TwoParams(bool, FRigVMGraphExplorer_OnSetVariablePinType, const FString&, const FEdGraphPinType&);
+DECLARE_DELEGATE_RetVal_OneParam(FEdGraphPinType, FRigVMGraphExplorer_OnGetVariablePinType, const FRigVMExplorerElementKey&);
+DECLARE_DELEGATE_RetVal_TwoParams(bool, FRigVMGraphExplorer_OnSetVariablePinType, const FRigVMExplorerElementKey&, const FEdGraphPinType&);
 DECLARE_DELEGATE_RetVal_OneParam(bool, FRigVMGraphExplorer_OnIsVariablePublic, const FString&);
 DECLARE_DELEGATE_RetVal_OneParam(bool, FRigVMGraphExplorer_OnToggleVariablePublic, const FString&);
+DECLARE_DELEGATE_RetVal(TArray<TSharedPtr<IPinTypeSelectorFilter>>, FRigVMGraphExplorer_OnGetCustomPinFilters);
 
 typedef STreeView<TSharedPtr<FRigVMEditorGraphExplorerTreeElement>>::FOnSelectionChanged FRigVMGraphExplorer_OnSelectionChanged;
 
@@ -95,6 +98,8 @@ public:
 	FRigVMGraphExplorer_OnGetEventNodesInGraph OnGetEventNodesInGraph;
 	FRigVMGraphExplorer_OnGetFunctions OnGetFunctions;
 	FRigVMGraphExplorer_OnGetVariables OnGetVariables;
+	FRigVMGraphExplorer_OnGetVariables OnGetLocalVariables;
+	FRigVMGraphExplorer_OnIsFunctionFocused OnIsFunctionFocused;
 	FRigVMGraphExplorer_OnGetGraphDisplayName OnGetGraphDisplayName;
 	FRigVMGraphExplorer_OnGetGraphDisplayName OnGetEventDisplayName;
 	FRigVMGraphExplorer_OnGetGraphTooltip OnGetGraphTooltip;
@@ -109,6 +114,7 @@ public:
 	FRigVMGraphExplorer_OnCreateGraph OnCreateGraph;
 	FRigVMGraphExplorer_OnCreateFunction OnCreateFunction;
 	FRigVMGraphExplorer_OnCreateVariable OnCreateVariable;
+	FRigVMGraphExplorer_OnCreateVariable OnCreateLocalVariable;
 	FRigVMGraphExplorer_OnRenameGraph OnRenameGraph;
 	FRigVMGraphExplorer_OnCanRenameGraph OnCanRenameGraph;
 	FRigVMGraphExplorer_OnRenameFunction OnRenameFunction;
@@ -123,6 +129,7 @@ public:
 	FRigVMGraphExplorer_OnSetVariablePinType OnSetVariablePinType;
 	FRigVMGraphExplorer_OnIsVariablePublic OnIsVariablePublic;
 	FRigVMGraphExplorer_OnToggleVariablePublic OnToggleVariablePublic;
+	FRigVMGraphExplorer_OnGetCustomPinFilters OnGetCustomPinFilters;
 	
 	FRigVMGraphExplorer_OnSelectionChanged OnSelectionChanged;
 	FRigVMGraphExplorer_OnRequestContextMenu OnRequestContextMenu;
@@ -177,6 +184,24 @@ public:
 		return TArray<FRigVMGraphVariableDescription>();
 	}
 
+	TArray<FRigVMGraphVariableDescription> GetLocalVariables() const
+	{
+		if (OnGetLocalVariables.IsBound())
+		{
+			return OnGetLocalVariables.Execute();
+		}
+		return TArray<FRigVMGraphVariableDescription>();
+	}
+
+	bool IsFunctionFocused() const
+	{
+		if (OnIsFunctionFocused.IsBound())
+		{
+			return OnIsFunctionFocused.Execute();
+		}
+		return false;
+	}
+
 	FText GetGraphDisplayName(const FString& InPath) const
 	{
 		if (OnGetGraphDisplayName.IsBound())
@@ -237,11 +262,11 @@ public:
 		}
 	}
 
-	void VariableClicked(const FString& InPath) const
+	void VariableClicked(const FRigVMExplorerElementKey& InVariable) const
 	{
 		if (OnVariableClicked.IsBound())
 		{
-			OnVariableClicked.Execute(InPath);
+			OnVariableClicked.Execute(InVariable);
 		}
 	}
 
@@ -293,6 +318,14 @@ public:
 		}
 	}
 
+	void CreateLocalVariable() const
+	{
+		if (OnCreateLocalVariable.IsBound())
+		{
+			OnCreateLocalVariable.Execute();
+		}
+	}
+
 	void RenameGraph(const FString& InOldPath, const FString& InNewPath) const
 	{
 		if (OnRenameGraph.IsBound())
@@ -327,19 +360,19 @@ public:
 		return false;
 	}
 
-	void RenameVariable(const FString& InOldName, const FString& InNewName) const
+	void RenameVariable(const FRigVMExplorerElementKey& InOldKey, const FString& InNewName) const
 	{
 		if (OnRenameVariable.IsBound())
 		{
-			OnRenameVariable.Execute(InOldName, InNewName);
+			OnRenameVariable.Execute(InOldKey, InNewName);
 		}
 	}
 
-	bool CanRenameVariable(const FString& InOldName, const FString& InNewName, FText& OutErrorMessage) const
+	bool CanRenameVariable(const FRigVMExplorerElementKey& InOldKey, const FString& InNewName, FText& OutErrorMessage) const
 	{
 		if (OnCanRenameVariable.IsBound())
 		{
-			return OnCanRenameVariable.Execute(InOldName, InNewName, OutErrorMessage);
+			return OnCanRenameVariable.Execute(InOldKey, InNewName, OutErrorMessage);
 		}
 		return false;
 	}
@@ -397,20 +430,20 @@ public:
 		return nullptr;
 	}
 
-	FEdGraphPinType GetVariablePinType(const FString& InVariable)
+	FEdGraphPinType GetVariablePinType(const FRigVMExplorerElementKey& InVariableKey)
 	{
 		if (OnGetVariablePinType.IsBound())
 		{
-			return OnGetVariablePinType.Execute(InVariable);
+			return OnGetVariablePinType.Execute(InVariableKey);
 		}
 		return FEdGraphPinType();
 	}
 
-	bool SetVariablePinType(const FString& InVariable, const FEdGraphPinType& InType)
+	bool SetVariablePinType(const FRigVMExplorerElementKey& InVariableKey, const FEdGraphPinType& InType)
 	{
 		if (OnSetVariablePinType.IsBound())
 		{
-			return OnSetVariablePinType.Execute(InVariable, InType);
+			return OnSetVariablePinType.Execute(InVariableKey, InType);
 		}
 		return false;
 	}
@@ -433,6 +466,15 @@ public:
 		return false;
 	}
 
+	TArray<TSharedPtr<IPinTypeSelectorFilter>> GetCustomPinFilters()
+	{
+		if (OnGetCustomPinFilters.IsBound())
+		{
+			return OnGetCustomPinFilters.Execute();
+		}
+		return TArray<TSharedPtr<IPinTypeSelectorFilter>>();
+	}
+
 private:
 
 	bool bSuspendSelectionDelegate = false;
