Control Rig: Avoid crash when opening editor with default layout
#rnx

diff --git a/Engine/Plugins/Runtime/RigVM/Source/RigVMEditor/Private/Editor/RigVMEditor.cpp b/Engine/Plugins/Runtime/RigVM/Source/RigVMEditor/Private/Editor/RigVMEditor.cpp
--- a/Engine/Plugins/Runtime/RigVM/Source/RigVMEditor/Private/Editor/RigVMEditor.cpp
+++ b/Engine/Plugins/Runtime/RigVM/Source/RigVMEditor/Private/Editor/RigVMEditor.cpp
@@ -1316,7 +1316,10 @@ void FRigVMEditor::RefreshEditors(ERefreshBlueprintEditorReason::Type Reason)
 		Reason = ERefreshBlueprintEditorReason::BlueprintCompiled;
 	}
 	FBlueprintEditor::RefreshEditors(Reason);
-	GraphExplorerWidget->Refresh();
+	if (GraphExplorerWidget.IsValid())
+	{
+		GraphExplorerWidget->Refresh();
+	}
 }
 
 void FRigVMEditor::SetupGraphEditorEvents(UEdGraph* InGraph, SGraphEditor::FGraphEditorEvents& InEvents)
