RigVM: Variant management related backend API for BuildData + unit test

#rnx

diff --git a/Engine/Plugins/Runtime/RigVM/Source/RigVM/Public/RigVMCore/RigVMVariant.h b/Engine/Plugins/Runtime/RigVM/Source/RigVM/Public/RigVMCore/RigVMVariant.h
--- a/Engine/Plugins/Runtime/RigVM/Source/RigVM/Public/RigVMCore/RigVMVariant.h
+++ b/Engine/Plugins/Runtime/RigVM/Source/RigVM/Public/RigVMCore/RigVMVariant.h
@@ -129,6 +129,11 @@ struct RIGVM_API FRigVMVariant
 	FRigVMVariant()
 		: Guid(FGuid()) {}
 
+	bool IsValid() const
+	{
+		return Guid.IsValid();
+	}
+
 	friend FArchive& operator<<(FArchive& Ar, FRigVMVariant& Data)
 	{
 		Ar.UsingCustomVersion(FRigVMObjectVersion::GUID);
@@ -165,6 +170,11 @@ struct RIGVM_API FRigVMVariantRef
 		return Variant.Guid == Other.Variant.Guid && ObjectPath == Other.ObjectPath; 
 	}
 
+	bool IsValid() const
+	{
+		return ObjectPath.IsValid() && Variant.IsValid();
+	}
+
 	friend uint32 GetTypeHash(const FRigVMVariantRef& InVariantRef)
 	{
 		return HashCombine(GetTypeHash(InVariantRef.ObjectPath.ToString()), GetTypeHash(InVariantRef.Variant.Guid));
diff --git a/Engine/Plugins/Runtime/RigVM/Source/RigVMDeveloper/Private/RigVMBlueprint.cpp b/Engine/Plugins/Runtime/RigVM/Source/RigVMDeveloper/Private/RigVMBlueprint.cpp
--- a/Engine/Plugins/Runtime/RigVM/Source/RigVMDeveloper/Private/RigVMBlueprint.cpp
+++ b/Engine/Plugins/Runtime/RigVM/Source/RigVMDeveloper/Private/RigVMBlueprint.cpp
@@ -2831,6 +2831,11 @@ FRigVMVariant URigVMBlueprint::GetAssetVariantBP() const
 	return GetAssetVariant();
 }
 
+FRigVMVariantRef URigVMBlueprint::GetAssetVariantRef() const
+{
+	return FRigVMVariantRef(FSoftObjectPath(this), GetAssetVariant());
+}
+
 bool URigVMBlueprint::SplitAssetVariant()
 {
 	if(GetMatchingVariants().IsEmpty())
diff --git a/Engine/Plugins/Runtime/RigVM/Source/RigVMDeveloper/Private/RigVMModel/Nodes/RigVMLibraryNode.cpp b/Engine/Plugins/Runtime/RigVM/Source/RigVMDeveloper/Private/RigVMModel/Nodes/RigVMLibraryNode.cpp
--- a/Engine/Plugins/Runtime/RigVM/Source/RigVMDeveloper/Private/RigVMModel/Nodes/RigVMLibraryNode.cpp
+++ b/Engine/Plugins/Runtime/RigVM/Source/RigVMDeveloper/Private/RigVMModel/Nodes/RigVMLibraryNode.cpp
@@ -338,6 +338,11 @@ FRigVMVariant URigVMLibraryNode::GetFunctionVariant() const
 	return GetFunctionHeader().Variant;
 }
 
+FRigVMVariantRef URigVMLibraryNode::GetFunctionVariantRef() const
+{
+	return FRigVMVariantRef(FSoftObjectPath(this), GetFunctionVariant());
+}
+
 TArray<FRigVMVariantRef> URigVMLibraryNode::GetMatchingVariants() const
 {
 	const FRigVMGraphFunctionHeader Header = GetFunctionHeader();
diff --git a/Engine/Plugins/Runtime/RigVM/Source/RigVMDeveloper/Private/RigVMModel/RigVMBuildData.cpp b/Engine/Plugins/Runtime/RigVM/Source/RigVMDeveloper/Private/RigVMModel/RigVMBuildData.cpp
--- a/Engine/Plugins/Runtime/RigVM/Source/RigVMDeveloper/Private/RigVMModel/RigVMBuildData.cpp
+++ b/Engine/Plugins/Runtime/RigVM/Source/RigVMDeveloper/Private/RigVMModel/RigVMBuildData.cpp
@@ -10,6 +10,12 @@
 #include "UObject/StrongObjectPtr.h"
 #include "RigVMBlueprint.h"
 
+#if WITH_EDITOR
+#include "Editor.h"
+#include "Editor/EditorEngine.h"
+#include "Subsystems/EditorAssetSubsystem.h"
+#endif
+
 #include UE_INLINE_GENERATED_CPP_BY_NAME(RigVMBuildData)
 
 static const FLazyName FunctionReferenceNodeDataName = TEXT("FunctionReferenceNodeData");
@@ -520,7 +526,7 @@ void URigVMBuildData::ClearInvalidReferences()
 	}
 }
 
-TArray<FRigVMVariantRef> URigVMBuildData::GatherAllFunctionVariantRefs()
+TArray<FRigVMVariantRef> URigVMBuildData::GatherAllFunctionVariantRefs() const
 {
 	TArray<FRigVMVariantRef> Result;
 	const FAssetRegistryModule& AssetRegistryModule = FModuleManager::LoadModuleChecked<FAssetRegistryModule>("AssetRegistry");
@@ -560,7 +566,7 @@ TArray<FRigVMVariantRef> URigVMBuildData::GatherAllFunctionVariantRefs()
 	return Result;
 }
 
-TArray<FRigVMVariantRef> URigVMBuildData::GatherFunctionVariantRefsForAsset(const FAssetData& InAssetData)
+TArray<FRigVMVariantRef> URigVMBuildData::GatherFunctionVariantRefsForAsset(const FAssetData& InAssetData) const
 {
 	TArray<FRigVMVariantRef> Result;
 
@@ -582,14 +588,56 @@ TArray<FRigVMVariantRef> URigVMBuildData::GatherFunctionVariantRefsForAsset(cons
 	return Result;
 }
 
-TArray<FRigVMVariantRef> URigVMBuildData::FindFunctionVariantRefs(const FGuid& InGuid)
+TArray<FRigVMVariantRef> URigVMBuildData::FindFunctionVariantRefs(const FGuid& InGuid) const
 {
 	TArray<FRigVMVariantRef> Result = GatherAllFunctionVariantRefs();
 	Result = Result.FilterByPredicate([InGuid](const FRigVMVariantRef& VariantRef) { return VariantRef.Variant.Guid == InGuid; });
 	return Result;
 }
 
-TArray<FRigVMVariantRef> URigVMBuildData::GatherAllAssetVariantRefs()
+FRigVMGraphFunctionIdentifier URigVMBuildData::GetFunctionIdentifierForVariant(const FRigVMVariantRef& InVariantRef) const
+{
+	FString SubPathString = InVariantRef.ObjectPath.GetSubPathString();
+	if(!SubPathString.IsEmpty())
+	{
+		int32 LastSlashIndex = INDEX_NONE;
+		if(SubPathString.FindLastChar(TEXT('.'), LastSlashIndex))
+		{
+			SubPathString = SubPathString.Mid(LastSlashIndex + 1);
+		}
+		
+		const FRigVMGraphFunctionHeader Header = FRigVMGraphFunctionHeader::FindGraphFunctionHeader(InVariantRef.ObjectPath, *SubPathString, nullptr, nullptr);
+		if(Header.IsValid())
+		{
+			return Header.LibraryPointer;
+		}
+	}
+	return FRigVMGraphFunctionIdentifier();
+}
+
+FRigVMVariantRef URigVMBuildData::CreateFunctionVariant(const FRigVMGraphFunctionIdentifier& InFunctionIdentifier, FName InName)
+{
+	const FSoftObjectPath LibraryNodePath = FSoftObjectPath(InFunctionIdentifier.GetLibraryNodePath());
+	if(const URigVMLibraryNode* LibraryNode = Cast<URigVMLibraryNode>(LibraryNodePath.TryLoad()))
+	{
+		if(URigVMFunctionLibrary* FunctionLibrary = Cast<URigVMFunctionLibrary>(LibraryNode->GetGraph()))
+		{
+			if(URigVMBlueprint* RigVMBlueprint = FunctionLibrary->GetTypedOuter<URigVMBlueprint>())
+			{
+				if(URigVMController* Controller = RigVMBlueprint->GetOrCreateController(FunctionLibrary))
+				{
+					if(const URigVMLibraryNode* VariantNode = Controller->CreateFunctionVariant(InFunctionIdentifier.GetFunctionFName(), InName))
+					{
+						return VariantNode->GetFunctionVariantRef();
+					}
+				}
+			}
+		}
+	}
+	return FRigVMVariantRef();
+}
+
+TArray<FRigVMVariantRef> URigVMBuildData::GatherAllAssetVariantRefs() const
 {
 	TArray<FRigVMVariantRef> Result;
 	const FAssetRegistryModule& AssetRegistryModule = FModuleManager::LoadModuleChecked<FAssetRegistryModule>("AssetRegistry");
@@ -601,42 +649,183 @@ TArray<FRigVMVariantRef> URigVMBuildData::GatherAllAssetVariantRefs()
 
 	for (const FAssetData& Asset : Assets)
 	{
-		FRigVMVariant AssetVariant;
-		if(Asset.IsAssetLoaded())
+		Result.Add(GetVariantRefForAsset(Asset));
+	}
+
+	return Result;
+}
+
+TArray<FRigVMVariantRef> URigVMBuildData::FindAssetVariantRefs(const FGuid& InGuid) const
+{
+	TArray<FRigVMVariantRef> Result = GatherAllAssetVariantRefs();
+	Result = Result.FilterByPredicate([InGuid](const FRigVMVariantRef& VariantRef) { return VariantRef.Variant.Guid == InGuid; });
+	return Result;
+}
+
+FRigVMVariantRef URigVMBuildData::CreateAssetVariant(const FAssetData& InAssetData, FName InName)
+{
+	FRigVMVariantRef TargetVariantRef;
+
+#if WITH_EDITOR
+	const FRigVMVariantRef SourceVariantRef = GetVariantRefForAsset(InAssetData);
+	if(!SourceVariantRef.IsValid())
+	{
+		return SourceVariantRef;
+	}
+	
+	UEditorAssetSubsystem* EditorAssetSubsystem = GEditor->GetEditorSubsystem<UEditorAssetSubsystem>();
+
+	const FString SourcePackageLongName = InAssetData.GetSoftObjectPath().GetWithoutSubPath().ToString();
+	FString SourcePackageDirectory, SourcePackagePath, SourcePackageName;
+	FPackageName::SplitLongPackageName(SourcePackageLongName, SourcePackageDirectory, SourcePackagePath, SourcePackageName);
+
+	int32 LastPeriodPos = INDEX_NONE;
+	if(SourcePackageName.FindChar(TEXT('.'), LastPeriodPos))
+	{
+		SourcePackageName = SourcePackageName.Left(LastPeriodPos);
+	}
+
+	FAssetRegistryModule& AssetRegistryModule = FModuleManager::LoadModuleChecked<FAssetRegistryModule>("AssetRegistry");
+
+	int32 Suffix = 1;
+	const FString DestinationAssetNameBase = InName.IsNone() ? SourcePackageName : InName.ToString();
+	FString DestinationAssetName = DestinationAssetNameBase;
+	FString DestinationPackageLongName;
+	do
+	{
+		DestinationPackageLongName = SourcePackageDirectory + SourcePackagePath + DestinationAssetName + TEXT(".") + DestinationAssetName;
+		DestinationAssetName = DestinationAssetNameBase + TEXT("_") + FString::FromInt(++Suffix);
+	}
+	while (AssetRegistryModule.Get().GetAssetByObjectPath(FSoftObjectPath(DestinationPackageLongName)).IsValid());
+	
+	if(UObject* DuplicatedAsset = EditorAssetSubsystem->DuplicateAsset(SourcePackageLongName, DestinationPackageLongName))
+	{
+		if(URigVMBlueprint* RigVMBlueprint = Cast<URigVMBlueprint>(DuplicatedAsset))
 		{
-			if(const URigVMBlueprint* Blueprint = Cast<URigVMBlueprint>(Asset.GetAsset()))
-			{
-				AssetVariant = Blueprint->GetAssetVariant();
-			}
+			// since we duplicated the asset the variant guid will be the same too,
+			// but to be sure we'll run the join code path as well.
+			(void)RigVMBlueprint->JoinAssetVariant(SourceVariantRef.Variant.Guid);
+			TargetVariantRef = RigVMBlueprint->GetAssetVariantRef();
 		}
+	}
+#endif
+	return TargetVariantRef;
+}
+
+FAssetData URigVMBuildData::GetAssetDataForPath(const FSoftObjectPath& InObjectPath) const
+{
+#if WITH_EDITOR
+	UEditorAssetSubsystem* EditorAssetSubsystem = GEditor->GetEditorSubsystem<UEditorAssetSubsystem>();
+	return EditorAssetSubsystem->FindAssetData(InObjectPath.GetWithoutSubPath().ToString());
+#else
+	return FAssetData();
+#endif
+}
+
+FRigVMVariantRef URigVMBuildData::GetVariantRefForAsset(const FAssetData& InAssetData) const
+{
+	FRigVMVariant AssetVariant;
+	if(InAssetData.IsAssetLoaded())
+	{
+		if(const URigVMBlueprint* Blueprint = Cast<URigVMBlueprint>(InAssetData.GetAsset()))
+		{
+			AssetVariant = Blueprint->GetAssetVariant();
+		}
+	}
 		
-		if (!AssetVariant.Guid.IsValid())
+	if (!AssetVariant.Guid.IsValid())
+	{
+		static const FName AssetVariantPropertyName = GET_MEMBER_NAME_CHECKED(URigVMBlueprint, AssetVariant);
+		const FProperty* AssetVariantProperty = CastField<FProperty>(URigVMBlueprint::StaticClass()->FindPropertyByName(AssetVariantPropertyName));
+		const FString VariantStr = InAssetData.GetTagValueRef<FString>(AssetVariantPropertyName);
+		if(!VariantStr.IsEmpty())
 		{
-			static const FName AssetVariantPropertyName = GET_MEMBER_NAME_CHECKED(URigVMBlueprint, AssetVariant);
-			const FProperty* AssetVariantProperty = CastField<FProperty>(URigVMBlueprint::StaticClass()->FindPropertyByName(AssetVariantPropertyName));
-			const FString VariantStr = Asset.GetTagValueRef<FString>(AssetVariantPropertyName);
-			if(!VariantStr.IsEmpty())
-			{
-				AssetVariantProperty->ImportText_Direct(*VariantStr, &AssetVariant, nullptr, EPropertyPortFlags::PPF_None);
-			}
+			AssetVariantProperty->ImportText_Direct(*VariantStr, &AssetVariant, nullptr, EPropertyPortFlags::PPF_None);
 		}
+	}
+
+	if (!AssetVariant.Guid.IsValid())
+	{
+		AssetVariant.Guid = FRigVMVariant::GenerateGUID(InAssetData.PackageName.ToString());
+	}
+		
+	return FRigVMVariantRef(InAssetData.ToSoftObjectPath(), AssetVariant);
+}
 
-		if (!AssetVariant.Guid.IsValid())
+FAssetData URigVMBuildData::GetAssetDataForVariant(const FRigVMVariantRef& InVariantRef) const
+{
+	return GetAssetDataForPath(InVariantRef.ObjectPath);
+}
+
+FRigVMVariantRef URigVMBuildData::SplitVariantFromSet(const FRigVMVariantRef& InVariantRef)
+{
+	FRigVMVariantRef Result = InVariantRef;
+	if(Result.IsValid())
+	{
+		FAssetData AssetData = GetAssetDataForVariant(Result);
+		if(URigVMBlueprint* RigVMBlueprint = Cast<URigVMBlueprint>(AssetData.GetAsset()))
 		{
-			AssetVariant.Guid = FRigVMVariant::GenerateGUID(Asset.PackageName.ToString());
+			if(UObject* Subject = InVariantRef.ObjectPath.TryLoad())
+			{
+				if(Subject == RigVMBlueprint)
+				{
+					if(RigVMBlueprint->SplitAssetVariant())
+					{
+						Result.Variant = RigVMBlueprint->GetAssetVariant();
+					}
+				}
+				else if(URigVMLibraryNode* FunctionNode = Cast<URigVMLibraryNode>(Subject))
+				{
+					if(URigVMFunctionLibrary* FunctionLibrary = Cast<URigVMFunctionLibrary>(FunctionNode->GetGraph()))
+					{
+						if(URigVMController* Controller = RigVMBlueprint->GetOrCreateController(FunctionLibrary))
+						{
+							if(Controller->SplitFunctionVariant(FunctionNode->GetFName()))
+							{
+								Result.Variant = FunctionNode->GetFunctionVariant();
+							}
+						}
+					}
+				}
+			}
 		}
-		
-		FRigVMVariantRef VariantRef(Asset.ToSoftObjectPath(), AssetVariant);
-		Result.Add(VariantRef);
 	}
-
 	return Result;
 }
 
-TArray<FRigVMVariantRef> URigVMBuildData::FindAssetVariantRefs(const FGuid& InGuid)
+FRigVMVariantRef URigVMBuildData::JoinVariantSet(const FRigVMVariantRef& InVariantRef, const FGuid& InGuid)
 {
-	TArray<FRigVMVariantRef> Result = GatherAllAssetVariantRefs();
-	Result = Result.FilterByPredicate([InGuid](const FRigVMVariantRef& VariantRef) { return VariantRef.Variant.Guid == InGuid; });
+	FRigVMVariantRef Result = InVariantRef;
+	if(Result.IsValid() && InGuid.IsValid())
+	{
+		FAssetData AssetData = GetAssetDataForVariant(Result);
+		if(URigVMBlueprint* RigVMBlueprint = Cast<URigVMBlueprint>(AssetData.GetAsset()))
+		{
+			if(UObject* Subject = InVariantRef.ObjectPath.TryLoad())
+			{
+				if(Subject == RigVMBlueprint)
+				{
+					if(RigVMBlueprint->JoinAssetVariant(InGuid))
+					{
+						Result.Variant = RigVMBlueprint->GetAssetVariant();
+					}
+				}
+				else if(URigVMLibraryNode* FunctionNode = Cast<URigVMLibraryNode>(Subject))
+				{
+					if(URigVMFunctionLibrary* FunctionLibrary = Cast<URigVMFunctionLibrary>(FunctionNode->GetGraph()))
+					{
+						if(URigVMController* Controller = RigVMBlueprint->GetOrCreateController(FunctionLibrary))
+						{
+							if(Controller->JoinFunctionVariant(FunctionNode->GetFName(), InGuid))
+							{
+								Result.Variant = FunctionNode->GetFunctionVariant();
+							}
+						}
+					}
+				}
+			}
+		}
+	}
 	return Result;
 }
 
diff --git a/Engine/Plugins/Runtime/RigVM/Source/RigVMDeveloper/Public/RigVMBlueprint.h b/Engine/Plugins/Runtime/RigVM/Source/RigVMDeveloper/Public/RigVMBlueprint.h
--- a/Engine/Plugins/Runtime/RigVM/Source/RigVMDeveloper/Public/RigVMBlueprint.h
+++ b/Engine/Plugins/Runtime/RigVM/Source/RigVMDeveloper/Public/RigVMBlueprint.h
@@ -570,18 +570,21 @@ public:
 
 	const FRigVMVariant& GetAssetVariant() const { return AssetVariant; }
 
-	UFUNCTION(BlueprintPure, Category = "Variables", meta = (DisplayName = "GetAssetVariant", ScriptName = "GetAssetVariant"))
+	UFUNCTION(BlueprintPure, Category = "Variants", meta = (DisplayName = "GetAssetVariant", ScriptName = "GetAssetVariant"))
 	FRigVMVariant GetAssetVariantBP() const;
 
+	UFUNCTION(BlueprintPure, Category = "Variants")
+	FRigVMVariantRef GetAssetVariantRef() const;
+
 	/** Resets the asset's guid to a new one and splits it from the former variant set */
-	UFUNCTION(BlueprintCallable, Category = "Variables")
+	UFUNCTION(BlueprintCallable, Category = "Variants")
 	bool SplitAssetVariant();
 
 	/** Merges the asset's guid with a provided one to join the variant set */
-	UFUNCTION(BlueprintCallable, Category = "Variables")
+	UFUNCTION(BlueprintCallable, Category = "Variants")
 	bool JoinAssetVariant(const FGuid& InGuid);
 
-	UFUNCTION(BlueprintPure, Category = "Variables")
+	UFUNCTION(BlueprintPure, Category = "Variants")
 	TArray<FRigVMVariantRef> GetMatchingVariants() const;
 
 #endif
diff --git a/Engine/Plugins/Runtime/RigVM/Source/RigVMDeveloper/Public/RigVMModel/Nodes/RigVMLibraryNode.h b/Engine/Plugins/Runtime/RigVM/Source/RigVMDeveloper/Public/RigVMModel/Nodes/RigVMLibraryNode.h
--- a/Engine/Plugins/Runtime/RigVM/Source/RigVMDeveloper/Public/RigVMModel/Nodes/RigVMLibraryNode.h
+++ b/Engine/Plugins/Runtime/RigVM/Source/RigVMDeveloper/Public/RigVMModel/Nodes/RigVMLibraryNode.h
@@ -64,6 +64,9 @@ public:
 	UFUNCTION(BlueprintPure, Category = RigVMLibraryNode)
 	FRigVMVariant GetFunctionVariant() const;
 
+	UFUNCTION(BlueprintPure, Category = RigVMLibraryNode)
+	FRigVMVariantRef GetFunctionVariantRef() const;
+
 	UFUNCTION(BlueprintPure, Category = RigVMLibraryNode)
 	TArray<FRigVMVariantRef> GetMatchingVariants() const;
 
diff --git a/Engine/Plugins/Runtime/RigVM/Source/RigVMDeveloper/Public/RigVMModel/RigVMBuildData.h b/Engine/Plugins/Runtime/RigVM/Source/RigVMDeveloper/Public/RigVMModel/RigVMBuildData.h
--- a/Engine/Plugins/Runtime/RigVM/Source/RigVMDeveloper/Public/RigVMModel/RigVMBuildData.h
+++ b/Engine/Plugins/Runtime/RigVM/Source/RigVMDeveloper/Public/RigVMModel/RigVMBuildData.h
@@ -64,7 +64,7 @@ private:
  * for the RigVM graph to improve the user experience.
  * This object is never serialized.
  */
-UCLASS()
+UCLASS(BlueprintType)
 class RIGVMDEVELOPER_API URigVMBuildData : public UObject
 {
 	GENERATED_BODY()
@@ -72,6 +72,7 @@ class RIGVMDEVELOPER_API URigVMBuildData : public UObject
 public:
 
 	// Returns the singleton build data
+	UFUNCTION(BlueprintPure, Category=RigVMBuildData)
 	static URigVMBuildData* Get();
 	
 	// Looks for all function references (in RigVMClientHost metadata) and initializes the URigVMBuildData
@@ -120,26 +121,65 @@ public:
 	void SetIsRunningUnitTest(bool bIsRunning) { bIsRunningUnitTest = bIsRunning; }
 
 	// Will find all public function variant refs, and private function variant refs from loaded assets
-	TArray<FRigVMVariantRef> GatherAllFunctionVariantRefs();
+	UFUNCTION(BlueprintPure, Category=RigVMBuildData)
+	TArray<FRigVMVariantRef> GatherAllFunctionVariantRefs() const;
 
 	// Will find the public function variant refs inside this asset, and private function variant refs if the asset is loaded
-	TArray<FRigVMVariantRef> GatherFunctionVariantRefsForAsset(const FAssetData& InAssetData);
+	UFUNCTION(BlueprintPure, Category=RigVMBuildData)
+	TArray<FRigVMVariantRef> GatherFunctionVariantRefsForAsset(const FAssetData& InAssetData) const;
 
 	// Will find all the function variants matching the given variant guid
-	TArray<FRigVMVariantRef> FindFunctionVariantRefs(const FGuid& InGuid);
+	UFUNCTION(BlueprintPure, Category=RigVMBuildData)
+	TArray<FRigVMVariantRef> FindFunctionVariantRefs(const FGuid& InGuid) const;
+
+	// Returns the function identifier given a variant reference
+	UFUNCTION(BlueprintPure, Category=RigVMBuildData)
+	FRigVMGraphFunctionIdentifier GetFunctionIdentifierForVariant(const FRigVMVariantRef& InVariantRef) const;
+
+	// Creates a new variant of a graph function within the same asset
+	UFUNCTION(BlueprintCallable, Category=RigVMBuildData)
+	FRigVMVariantRef CreateFunctionVariant(const FRigVMGraphFunctionIdentifier& InFunctionIdentifier, FName InName = NAME_None);
 
 	// Will find all asset variant refs
-	TArray<FRigVMVariantRef> GatherAllAssetVariantRefs();
+	UFUNCTION(BlueprintPure, Category=RigVMBuildData)
+	TArray<FRigVMVariantRef> GatherAllAssetVariantRefs() const;
 
 	// Will find all the asset variants matching the given variant guid
-	TArray<FRigVMVariantRef> FindAssetVariantRefs(const FGuid& InGuid);
+	UFUNCTION(BlueprintPure, Category=RigVMBuildData)
+	TArray<FRigVMVariantRef> FindAssetVariantRefs(const FGuid& InGuid) const;
+
+	// Creates a new variant of an asset
+	UFUNCTION(BlueprintCallable, Category=RigVMBuildData)
+	FRigVMVariantRef CreateAssetVariant(const FAssetData& InAssetData, FName InName = NAME_None);
+
+	// Returns the asset data given an object path
+	UFUNCTION(BlueprintPure, Category=RigVMBuildData)
+	FAssetData GetAssetDataForPath(const FSoftObjectPath& InObjectPath) const;
+
+	// Returns the asset data given a variant reference
+	UFUNCTION(BlueprintPure, Category=RigVMBuildData)
+	FRigVMVariantRef GetVariantRefForAsset(const FAssetData& InAssetData) const;
+	
+	// Returns the asset data given a variant reference
+	UFUNCTION(BlueprintPure, Category=RigVMBuildData)
+	FAssetData GetAssetDataForVariant(const FRigVMVariantRef& InVariantRef) const;
+
+	// Splits a variant from its variant set (by assigning a new, unique guid
+	UFUNCTION(BlueprintPure, Category=RigVMBuildData)
+	FRigVMVariantRef SplitVariantFromSet(const FRigVMVariantRef& InVariantRef);
+
+	// Joins a variant with another variant's set
+	UFUNCTION(BlueprintPure, Category=RigVMBuildData)
+	FRigVMVariantRef JoinVariantSet(const FRigVMVariantRef& InVariantRef, const FGuid& InGuid);
 
 #if WITH_EDITOR
 	// Returns all known public function identifiers used in the project
+	UFUNCTION(BlueprintPure, Category=RigVMBuildData)
 	TArray<FRigVMGraphFunctionIdentifier> GetAllFunctionIdentifiers(bool bOnlyPublic = true) const;
 #endif
 
 	// Returns all known public function identifiers used in the project
+	UFUNCTION(BlueprintPure, Category=RigVMBuildData)
 	TArray<FRigVMGraphFunctionIdentifier> GetUsedFunctionIdentifiers(bool bOnlyPublic = true) const;
 
 	// Returns all known function references
