RigVM: Variant management related frontend

#rnx

diff --git a/Engine/Plugins/Runtime/RigVM/Source/RigVMEditor/Private/Editor/RigVMGraphDetailCustomization.cpp b/Engine/Plugins/Runtime/RigVM/Source/RigVMEditor/Private/Editor/RigVMGraphDetailCustomization.cpp
--- a/Engine/Plugins/Runtime/RigVM/Source/RigVMEditor/Private/Editor/RigVMGraphDetailCustomization.cpp
+++ b/Engine/Plugins/Runtime/RigVM/Source/RigVMEditor/Private/Editor/RigVMGraphDetailCustomization.cpp
@@ -1016,6 +1016,7 @@ void FRigVMGraphDetailCustomizationImpl::CustomizeDetails(IDetailLayoutBuilder&
 				SNew(SRigVMVariantWidget)
 				.Context(VariantContext)
 				.Variant(this, &FRigVMGraphDetailCustomizationImpl::GetVariant)
+				.SubjectVariantRef(this, &FRigVMGraphDetailCustomizationImpl::GetSubjectVariantRef)
 				.VariantRefs(this, &FRigVMGraphDetailCustomizationImpl::GetVariantRefs)
 				.OnVariantChanged(this, &FRigVMGraphDetailCustomizationImpl::OnVariantChanged)
 				.OnBrowseVariantRef(this, &FRigVMGraphDetailCustomizationImpl::OnBrowseVariantRef)
@@ -1436,6 +1437,21 @@ FRigVMVariant FRigVMGraphDetailCustomizationImpl::GetVariant() const
 	return FRigVMVariant();
 }
 
+FRigVMVariantRef FRigVMGraphDetailCustomizationImpl::GetSubjectVariantRef() const
+{
+	if (WeakModel.IsValid() && RigVMClientHost.IsValid())
+	{
+		if (const URigVMGraph* Model = WeakModel.Get())
+		{
+			if (const URigVMLibraryNode* LibraryNode = Cast<URigVMLibraryNode>(Model->GetOuter()))
+			{
+				return FRigVMVariantRef(LibraryNode, GetVariant());
+			}
+		}
+	}
+	return FRigVMVariantRef();
+}
+
 TArray<FRigVMVariantRef> FRigVMGraphDetailCustomizationImpl::GetVariantRefs() const
 {
 	if (WeakModel.IsValid() && RigVMClientHost.IsValid())
diff --git a/Engine/Plugins/Runtime/RigVM/Source/RigVMEditor/Private/Editor/RigVMVariantDetailCustomization.cpp b/Engine/Plugins/Runtime/RigVM/Source/RigVMEditor/Private/Editor/RigVMVariantDetailCustomization.cpp
--- a/Engine/Plugins/Runtime/RigVM/Source/RigVMEditor/Private/Editor/RigVMVariantDetailCustomization.cpp
+++ b/Engine/Plugins/Runtime/RigVM/Source/RigVMEditor/Private/Editor/RigVMVariantDetailCustomization.cpp
@@ -54,6 +54,7 @@ void FRigVMVariantDetailCustomization::CustomizeHeader(TSharedRef<IPropertyHandl
 		SNew(SRigVMVariantWidget)
 		.Context(VariantContext)
 		.Variant(this, &FRigVMVariantDetailCustomization::GetVariant)
+		.SubjectVariantRef(this, &FRigVMVariantDetailCustomization::GetSubjectVariantRef)
 		.VariantRefs(this, &FRigVMVariantDetailCustomization::GetVariantRefs)
 		.OnVariantChanged(this, &FRigVMVariantDetailCustomization::OnVariantChanged)
 		.OnBrowseVariantRef(this, &FRigVMVariantDetailCustomization::OnBrowseVariantRef)
@@ -79,6 +80,15 @@ FRigVMVariant FRigVMVariantDetailCustomization::GetVariant() const
 	return FRigVMVariant();
 }
 
+FRigVMVariantRef FRigVMVariantDetailCustomization::GetSubjectVariantRef() const
+{
+	if (BlueprintBeingCustomized)
+	{
+		return BlueprintBeingCustomized->GetAssetVariantRef();
+	}
+	return FRigVMVariantRef();
+}
+
 TArray<FRigVMVariantRef> FRigVMVariantDetailCustomization::GetVariantRefs() const
 {
 	if (BlueprintBeingCustomized)
diff --git a/Engine/Plugins/Runtime/RigVM/Source/RigVMEditor/Private/Widgets/SRigVMVariantWidget.cpp b/Engine/Plugins/Runtime/RigVM/Source/RigVMEditor/Private/Widgets/SRigVMVariantWidget.cpp
--- a/Engine/Plugins/Runtime/RigVM/Source/RigVMEditor/Private/Widgets/SRigVMVariantWidget.cpp
+++ b/Engine/Plugins/Runtime/RigVM/Source/RigVMEditor/Private/Widgets/SRigVMVariantWidget.cpp
@@ -10,6 +10,9 @@
 #include "Editor/RigVMEditorTools.h"
 #include "AssetThumbnail.h"
 #include "RigVMModel/RigVMBuildData.h"
+#include "Framework/Application/SlateApplication.h"
+#include "HAL/PlatformApplicationMisc.h"
+#include "Framework/MultiBox/MultiBoxBuilder.h"
 
 #define LOCTEXT_NAMESPACE "SRigVMVariantWidget"
 
@@ -81,6 +84,48 @@ void SRigVMVariantToolTipWithTags::OnClosed()
 	ResetContentWidget();
 }
 
+void SRigVMVariantGuidWidget::Construct(const FArguments& InArgs)
+{
+	OnContextMenu = InArgs._OnContextMenu;
+	
+	SBox::FArguments SuperArguments;
+	SuperArguments.Padding(0);
+	SuperArguments.HAlign(HAlign_Left);
+	SuperArguments.VAlign(VAlign_Center);
+
+	TAttribute<FGuid> Guid = InArgs._Guid;
+	SuperArguments.Content()
+	[
+		SNew(STextBlock)
+		.Font(IDetailLayoutBuilder::GetDetailFont())
+		.Text_Lambda([this, Guid]()
+		{
+			return FText::FromString(Guid.Get().ToString(EGuidFormats::DigitsWithHyphensLower));
+		})
+	];
+	
+	SBox::Construct(SuperArguments);
+}
+
+FReply SRigVMVariantGuidWidget::OnMouseButtonDown(const FGeometry& MyGeometry, const FPointerEvent& MouseEvent)
+{
+	if(OnContextMenu.IsBound() && MouseEvent.IsMouseButtonDown(EKeys::RightMouseButton))
+	{
+		if(const TSharedPtr<SWidget> ContextMenuWidget = OnContextMenu.Execute())
+		{
+			FSlateApplication::Get().PushMenu(
+				SharedThis(this),
+				FWidgetPath(),
+				ContextMenuWidget.ToSharedRef(),
+				MouseEvent.GetScreenSpacePosition(),
+				FPopupTransitionEffect(FPopupTransitionEffect::ContextMenu)
+			);
+			return FReply::Handled();
+		}
+	}
+	return SBox::OnMouseButtonDown(MyGeometry, MouseEvent);
+}
+
 SRigVMVariantWidget::SRigVMVariantWidget()
 	: VariantRefHash(UINT32_MAX)
 {
@@ -94,18 +139,25 @@ void SRigVMVariantWidget::Construct(
 	const FArguments& InArgs)
 {
 	VariantAttribute = InArgs._Variant;
+	SubjectVariantRefAttribute = InArgs._SubjectVariantRef;
 	
 	OnVariantChanged = InArgs._OnVariantChanged;
 
 	VariantRefsAttribute = InArgs._VariantRefs;
 	OnCreateVariantRefRow = InArgs._OnCreateVariantRefRow;
 	OnBrowseVariantRef = InArgs._OnBrowseVariantRef;
+	OnVariantRefContextMenu = InArgs._OnVariantRefContextMenu;
 
 	if(!OnCreateVariantRefRow.IsBound())
 	{
 		OnCreateVariantRefRow.BindSP(this, &SRigVMVariantWidget::CreateDefaultVariantRefRow);
 	}
 
+	if(!OnVariantRefContextMenu.IsBound())
+	{
+		OnVariantRefContextMenu.BindSP(this, &SRigVMVariantWidget::CreateDefaultVariantRefContextMenu);
+	}
+
 	ContextAttribute = InArgs._Context;
 	if(!ContextAttribute.IsSet() && !ContextAttribute.IsBound())
 	{
@@ -129,12 +181,19 @@ void SRigVMVariantWidget::Construct(
 		.Padding(0, 8, 0, 0)
 		.HAlign(HAlign_Left)
 		[
-			SNew(STextBlock)
-			.Font(IDetailLayoutBuilder::GetDetailFont())
-			.Text_Lambda([this]()
+			SNew(SRigVMVariantGuidWidget)
+			.Guid_Lambda([this]()
 			{
-				const FRigVMVariant Variant = VariantAttribute.Get();
-				return FText::FromString(Variant.Guid.ToString(EGuidFormats::DigitsWithHyphensLower));
+				return VariantAttribute.Get().Guid;
+			})
+			.OnContextMenu_Lambda([this]() -> TSharedPtr<SWidget>
+			{
+				const FRigVMVariantRef SubjectVariantRef = SubjectVariantRefAttribute.Get();
+				if(SubjectVariantRef.IsValid())
+				{
+					return OnVariantRefContextMenu.Execute(SubjectVariantRef);
+				}
+				return SNullWidget::NullWidget;
 			})
 		]
 		+ SVerticalBox::Slot()
@@ -198,6 +257,7 @@ void SRigVMVariantWidget::Construct(
 		.TreeItemsSource(&VariantTreeRowInfos)
 		.OnGenerateRow(this, &SRigVMVariantWidget::GenerateVariantTreeRow)
 		.OnGetChildren(this, &SRigVMVariantWidget::GetChildrenForVariantInfo)
+		.OnContextMenuOpening(this, &SRigVMVariantWidget::OnVariantRefTreeContextMenu)
 	];
 
 	SetCanTick(true);
@@ -269,9 +329,14 @@ EVisibility SRigVMVariantWidget::GetVariantRefListVisibility() const
 
 TSharedRef<ITableRow> SRigVMVariantWidget::GenerateVariantTreeRow(TSharedPtr<FVariantTreeRowInfo> InRowInfo, const TSharedRef<STableViewBase>& OwnerTable)
 {
-	return SNew(SRigVMVariantRefTreeRow, OwnerTable)
-		.Content(CreateDefaultVariantRefRow(InRowInfo->VariantRef));
-	
+	TSharedRef<ITableRow> RowWidget = SNew(SRigVMVariantRefTreeRow, OwnerTable)
+		.Content(OnCreateVariantRefRow.Execute(InRowInfo->VariantRef));
+
+	if(InRowInfo)
+	{
+		InRowInfo->RowWidget = RowWidget.ToWeakPtr();
+	}
+	return RowWidget;
 }
 
 void SRigVMVariantWidget::GetChildrenForVariantInfo(TSharedPtr<FVariantTreeRowInfo> InInfo, TArray<TSharedPtr<FVariantTreeRowInfo>>& OutChildren)
@@ -496,4 +561,98 @@ const FSlateBrush* SRigVMVariantWidget::GetThumbnailBorder(TSharedRef<SBorder> I
 	return InThumbnailBorder->IsHovered() ? FAppStyle::Get().GetBrush(HoveredBorderName) : FAppStyle::Get().GetBrush(RegularBorderName);
 }
 
+TSharedPtr<SWidget> SRigVMVariantWidget::OnVariantRefTreeContextMenu()
+{
+	const FVector2f MousePosition = FSlateApplication::Get().GetCursorPos();
+
+	for(const TSharedPtr<FVariantTreeRowInfo>& VariantTreeRowInfo : VariantTreeRowInfos)
+	{
+		if(VariantTreeRowInfo->RowWidget.IsValid())
+		{
+			if(const TSharedPtr<SWidget> Content = VariantTreeRowInfo->RowWidget.Pin()->GetContent())
+			{
+				const FGeometry TickSpaceGeometry = Content->GetTickSpaceGeometry();
+				if(TickSpaceGeometry.IsUnderLocation(MousePosition))
+				{
+					return OnVariantRefContextMenu.Execute(VariantTreeRowInfo->VariantRef);
+				}
+			}
+		}
+	}
+	return SNullWidget::NullWidget;
+}
+
+TSharedPtr<SWidget> SRigVMVariantWidget::CreateDefaultVariantRefContextMenu(const FRigVMVariantRef& InVariantRef)
+{
+	FMenuBuilder MenuBuilder(true, nullptr);
+	{
+		if (InVariantRef.IsValid())
+		{
+			URigVMBuildData* BuildData = URigVMBuildData::Get();
+			TArray<FRigVMVariantRef> MatchingVariants = BuildData->FindAssetVariantRefs(InVariantRef.Variant.Guid);
+			if(MatchingVariants.IsEmpty())
+			{
+				MatchingVariants = BuildData->FindFunctionVariantRefs(InVariantRef.Variant.Guid);
+			}
+
+			MenuBuilder.AddMenuEntry(
+			 LOCTEXT("CopyGuidToClipboardLabel", "Copy Variant Guid"),
+			 LOCTEXT("CopyGuidToClipboardTooltip", "Copies the variant guid to the clipboard"),
+			 FSlateIcon{},
+			 FUIAction(FExecuteAction::CreateLambda([this, InVariantRef]()
+				{
+				 	FPlatformApplicationMisc::ClipboardCopy(*InVariantRef.Variant.Guid.ToString(EGuidFormats::DigitsWithHyphensLower));
+				}),
+				FCanExecuteAction())
+			);
+
+			MenuBuilder.AddMenuEntry(
+				 LOCTEXT("SplitVariantRefLabel", "Split Variant from Set"),
+				 LOCTEXT("SplitVariantRefTooltip", "Removes this variant from the set and applies a unique GUID to it"),
+				 FSlateIcon{},
+				 FUIAction(FExecuteAction::CreateLambda([this, InVariantRef, BuildData]()
+					{
+				 		BuildData->SplitVariantFromSet(InVariantRef);
+				 		RebuildVariantRefList();
+					}),
+					FCanExecuteAction::CreateLambda([MatchingVariants]() -> bool
+					{
+						return MatchingVariants.Num() > 1;
+					}))
+				);
+
+			// for the main subject also offer to join another set
+			const FRigVMVariantRef SubjectVariantRef = SubjectVariantRefAttribute.Get();
+			if(SubjectVariantRef == InVariantRef)
+			{
+				MenuBuilder.AddMenuEntry(
+					 LOCTEXT("JoinVariantRefLabel", "Join Variant Set (from Clipboard)"),
+					 LOCTEXT("JoinVariantRefTooltip", "Joins the Variant Set given the GUID on the clipboard"),
+					 FSlateIcon{},
+					 FUIAction(FExecuteAction::CreateLambda([this, InVariantRef, BuildData]()
+						{
+					 		FString TextToImport;
+							FPlatformApplicationMisc::ClipboardPaste(TextToImport);
+						 	const FGuid GuidToCheck(TextToImport);
+					 		BuildData->JoinVariantSet(InVariantRef, GuidToCheck);
+							RebuildVariantRefList();
+						}),
+						FCanExecuteAction::CreateLambda([InVariantRef]() -> bool
+						{
+							FString TextToImport;
+							FPlatformApplicationMisc::ClipboardPaste(TextToImport);
+							if(!TextToImport.IsEmpty())
+							{
+								const FGuid GuidToCheck(TextToImport);
+								return GuidToCheck.IsValid() && GuidToCheck != InVariantRef.Variant.Guid;
+							}
+							return false;
+						}))
+					);
+			}
+		}
+	}
+	return MenuBuilder.MakeWidget();
+}
+
 #undef LOCTEXT_NAMESPACE // SRigVMVariantWidget
diff --git a/Engine/Plugins/Runtime/RigVM/Source/RigVMEditor/Public/Editor/RigVMGraphDetailCustomization.h b/Engine/Plugins/Runtime/RigVM/Source/RigVMEditor/Public/Editor/RigVMGraphDetailCustomization.h
--- a/Engine/Plugins/Runtime/RigVM/Source/RigVMEditor/Public/Editor/RigVMGraphDetailCustomization.h
+++ b/Engine/Plugins/Runtime/RigVM/Source/RigVMEditor/Public/Editor/RigVMGraphDetailCustomization.h
@@ -214,6 +214,7 @@ private:
 	TSharedRef<ITableRow> HandleGenerateRowAccessSpecifier( TSharedPtr<FRigVMStringWithTag> SpecifierName, const TSharedRef<STableViewBase>& OwnerTable );
 	bool IsValidFunction() const;
 	FRigVMVariant GetVariant() const;
+	FRigVMVariantRef GetSubjectVariantRef() const;
 	TArray<FRigVMVariantRef> GetVariantRefs() const;
 
 private:
diff --git a/Engine/Plugins/Runtime/RigVM/Source/RigVMEditor/Public/Editor/RigVMVariantDetailCustomization.h b/Engine/Plugins/Runtime/RigVM/Source/RigVMEditor/Public/Editor/RigVMVariantDetailCustomization.h
--- a/Engine/Plugins/Runtime/RigVM/Source/RigVMEditor/Public/Editor/RigVMVariantDetailCustomization.h
+++ b/Engine/Plugins/Runtime/RigVM/Source/RigVMEditor/Public/Editor/RigVMVariantDetailCustomization.h
@@ -24,6 +24,7 @@ public:
 	virtual void CustomizeChildren(TSharedRef<class IPropertyHandle> InStructPropertyHandle, class IDetailChildrenBuilder& StructBuilder, IPropertyTypeCustomizationUtils& StructCustomizationUtils) override;
 
 	FRigVMVariant GetVariant() const;
+	FRigVMVariantRef GetSubjectVariantRef() const;
 	TArray<FRigVMVariantRef> GetVariantRefs() const;
 	
 protected:
diff --git a/Engine/Plugins/Runtime/RigVM/Source/RigVMEditor/Public/Widgets/SRigVMVariantWidget.h b/Engine/Plugins/Runtime/RigVM/Source/RigVMEditor/Public/Widgets/SRigVMVariantWidget.h
--- a/Engine/Plugins/Runtime/RigVM/Source/RigVMEditor/Public/Widgets/SRigVMVariantWidget.h
+++ b/Engine/Plugins/Runtime/RigVM/Source/RigVMEditor/Public/Widgets/SRigVMVariantWidget.h
@@ -13,6 +13,8 @@
 DECLARE_DELEGATE_OneParam(FRigVMVariantWidget_OnVariantChanged, const FRigVMVariant&);
 DECLARE_DELEGATE_RetVal_OneParam(TSharedPtr<SWidget>, FRigVMVariantWidget_OnCreateVariantRefRow, const FRigVMVariantRef&);
 DECLARE_DELEGATE_OneParam(FRigVMVariantWidget_OnBrowseVariantRef, const FRigVMVariantRef&);
+DECLARE_DELEGATE_RetVal_OneParam(TSharedPtr<SWidget>, FRigVMVariantWidget_OnVariantRefContextMenu, const FRigVMVariantRef&);
+DECLARE_DELEGATE_RetVal(TSharedPtr<SWidget>, FRigVMVariantWidget_OnContextMenu);
 
 struct RIGVMEDITOR_API FRigVMVariantWidgetContext
 {
@@ -45,6 +47,23 @@ private:
 	SToolTip::FArguments SuperClassArgs;
 };
 
+class SRigVMVariantGuidWidget : public SBox
+{
+public:
+	SLATE_BEGIN_ARGS(SRigVMVariantGuidWidget) {}
+		SLATE_ATTRIBUTE(FGuid, Guid)
+		SLATE_EVENT(FRigVMVariantWidget_OnContextMenu, OnContextMenu)
+	SLATE_END_ARGS()
+
+	void Construct(const FArguments& InArgs);
+
+	virtual FReply OnMouseButtonDown(const FGeometry& MyGeometry, const FPointerEvent& MouseEvent) override;
+
+private:
+	
+	FRigVMVariantWidget_OnContextMenu OnContextMenu;
+};
+
 class RIGVMEDITOR_API SRigVMVariantWidget : public SBox
 {
 public:
@@ -57,11 +76,13 @@ public:
 	{
 	}
 	SLATE_ATTRIBUTE(FRigVMVariant, Variant)
+	SLATE_ATTRIBUTE(FRigVMVariantRef, SubjectVariantRef)
 	SLATE_ATTRIBUTE(TArray<FRigVMVariantRef>, VariantRefs)
 	SLATE_ATTRIBUTE(FRigVMVariantWidgetContext, Context)
 	SLATE_EVENT(FRigVMVariantWidget_OnVariantChanged, OnVariantChanged)
 	SLATE_EVENT(FRigVMVariantWidget_OnCreateVariantRefRow, OnCreateVariantRefRow);
 	SLATE_EVENT(FRigVMVariantWidget_OnBrowseVariantRef, OnBrowseVariantRef)
+	SLATE_EVENT(FRigVMVariantWidget_OnVariantRefContextMenu, OnVariantRefContextMenu)
 	SLATE_ATTRIBUTE(float, MaxVariantRefListHeight)
 	SLATE_EVENT(FRigVMVariant_OnGetTags, OnGetTags)
 	SLATE_EVENT(FRigVMVariant_OnAddTag, OnAddTag)
@@ -83,6 +104,7 @@ private:
 	{
 		FRigVMVariantRef VariantRef;
 		TArray<TSharedPtr<FVariantTreeRowInfo>> NestedInfos;
+		TWeakPtr<ITableRow> RowWidget;
 	};
 
 	class RIGVMEDITOR_API SRigVMVariantRefTreeRow
@@ -93,6 +115,7 @@ private:
 		SLATE_BEGIN_ARGS(SRigVMVariantRefTreeRow)
 		{}
 		SLATE_ARGUMENT(TSharedPtr<SWidget>, Content)
+		SLATE_EVENT(FRigVMVariantWidget_OnVariantRefContextMenu, OnVariantRefContextMenu)
 		SLATE_END_ARGS()
 
 		virtual ~SRigVMVariantRefTreeRow() override;
@@ -106,8 +129,11 @@ private:
 	TSharedPtr<SWidget> CreateDefaultVariantRefRow(const FRigVMVariantRef& InVariantRef) const;
 	void RebuildVariantRefList();
 	const FSlateBrush* GetThumbnailBorder(TSharedRef<SBorder> InThumbnailBorder) const;
+	TSharedPtr<SWidget> OnVariantRefTreeContextMenu();
+	TSharedPtr<SWidget> CreateDefaultVariantRefContextMenu(const FRigVMVariantRef& InVariantRef);
 
 	TAttribute<FRigVMVariant> VariantAttribute;
+	TAttribute<FRigVMVariantRef> SubjectVariantRefAttribute;
 	FRigVMVariantWidget_OnVariantChanged OnVariantChanged;
 
 	TSharedPtr<SRigVMVariantTagWidget> TagWidget;
@@ -115,6 +141,7 @@ private:
 	TAttribute<TArray<FRigVMVariantRef>> VariantRefsAttribute;
 	FRigVMVariantWidget_OnCreateVariantRefRow OnCreateVariantRefRow;
 	FRigVMVariantWidget_OnBrowseVariantRef OnBrowseVariantRef;
+	FRigVMVariantWidget_OnVariantRefContextMenu OnVariantRefContextMenu;
 	TArray<FRigVMVariantRef> VariantRefs;
 	TArray<TSharedPtr<FVariantTreeRowInfo>> VariantTreeRowInfos;
 	uint32 VariantRefHash;
