Prevent RigVM customizations from reading / writing to random object data

#tests Unit, editor

diff --git a/Engine/Plugins/Runtime/RigVM/Source/RigVMEditor/Private/Editor/RigVMGraphDetailCustomization.cpp b/Engine/Plugins/Runtime/RigVM/Source/RigVMEditor/Private/Editor/RigVMGraphDetailCustomization.cpp
--- a/Engine/Plugins/Runtime/RigVM/Source/RigVMEditor/Private/Editor/RigVMGraphDetailCustomization.cpp
+++ b/Engine/Plugins/Runtime/RigVM/Source/RigVMEditor/Private/Editor/RigVMGraphDetailCustomization.cpp
@@ -2859,8 +2859,6 @@ void FRigVMGraphEnumDetailCustomization::CustomizeHeader(TSharedRef<IPropertyHan
 
 	for (UObject* Object : Objects)
 	{
-		ObjectsBeingCustomized.Add(Object);
-
 		if(BlueprintBeingCustomized == nullptr)
 		{
 			BlueprintBeingCustomized = Object->GetTypedOuter<URigVMBlueprint>();
@@ -2870,6 +2868,11 @@ void FRigVMGraphEnumDetailCustomization::CustomizeHeader(TSharedRef<IPropertyHan
 		{
 			GraphBeingCustomized = Object->GetTypedOuter<URigVMGraph>();
 		}
+
+		if (BlueprintBeingCustomized || GraphBeingCustomized)
+		{
+			ObjectsBeingCustomized.Add(Object);
+		}
 	}
 
 	FProperty* Property = InPropertyHandle->GetProperty();
@@ -3584,8 +3587,6 @@ void FRigVMGraphMathTypeDetailCustomization::CustomizeHeader(TSharedRef<IPropert
 
 	for (UObject* Object : Objects)
 	{
-		ObjectsBeingCustomized.Add(Object);
-
 		if(BlueprintBeingCustomized == nullptr)
 		{
 			BlueprintBeingCustomized = Object->GetTypedOuter<URigVMBlueprint>();
@@ -3595,6 +3596,11 @@ void FRigVMGraphMathTypeDetailCustomization::CustomizeHeader(TSharedRef<IPropert
 		{
 			GraphBeingCustomized = Object->GetTypedOuter<URigVMGraph>();
 		}
+
+		if (BlueprintBeingCustomized || GraphBeingCustomized)
+		{
+			ObjectsBeingCustomized.Add(Object);
+		}
 	}
 
 	FProperty* Property = InPropertyHandle->GetProperty();
@@ -3644,8 +3650,6 @@ void FRigVMGraphMathTypeDetailCustomization::CustomizeChildren(TSharedRef<IPrope
 
 	for (UObject* Object : Objects)
 	{
-		ObjectsBeingCustomized.Add(Object);
-
 		if(BlueprintBeingCustomized == nullptr)
 		{
 			BlueprintBeingCustomized = Object->GetTypedOuter<URigVMBlueprint>();
@@ -3655,6 +3659,11 @@ void FRigVMGraphMathTypeDetailCustomization::CustomizeChildren(TSharedRef<IPrope
 		{
 			GraphBeingCustomized = Object->GetTypedOuter<URigVMGraph>();
 		}
+
+		if (BlueprintBeingCustomized || GraphBeingCustomized)
+		{
+			ObjectsBeingCustomized.Add(Object);
+		}
 	}
 
 	StructsBeingCustomized.Reset();
