Allow RigVM Debug Draw Settings to show persistence when using Foreground batcher

diff --git a/Engine/Plugins/Runtime/RigVM/Source/RigVM/Public/RigVMCore/RigVMDebugDrawSettings.h b/Engine/Plugins/Runtime/RigVM/Source/RigVM/Public/RigVMCore/RigVMDebugDrawSettings.h
--- a/Engine/Plugins/Runtime/RigVM/Source/RigVM/Public/RigVMCore/RigVMDebugDrawSettings.h
+++ b/Engine/Plugins/Runtime/RigVM/Source/RigVM/Public/RigVMCore/RigVMDebugDrawSettings.h
@@ -13,6 +13,6 @@ struct RIGVM_API FRigVMDebugDrawSettings
 	UPROPERTY(meta = (Input, DisplayPriority = MAX_int32), Category = "DebugSettings")
 	TEnumAsByte<ESceneDepthPriorityGroup> DepthPriority = ESceneDepthPriorityGroup::SDPG_Foreground;
 
-	UPROPERTY(meta = (Input, EditConditionHides, EditCondition = "DepthPriority == ESceneDepthPriorityGroup::SDPG_World", DisplayPriority = MAX_int32), Category = "DebugSettings")
+	UPROPERTY(meta = (Input, DisplayPriority = MAX_int32), Category = "DebugSettings")
 	float Lifetime = -1.f;
 };
