Control Rig: Fix issue with template auto resolval

#rnx

diff --git a/Engine/Plugins/Runtime/RigVM/Source/RigVMDeveloper/Private/RigVMModel/RigVMController.cpp b/Engine/Plugins/Runtime/RigVM/Source/RigVMDeveloper/Private/RigVMModel/RigVMController.cpp
--- a/Engine/Plugins/Runtime/RigVM/Source/RigVMDeveloper/Private/RigVMModel/RigVMController.cpp
+++ b/Engine/Plugins/Runtime/RigVM/Source/RigVMDeveloper/Private/RigVMModel/RigVMController.cpp
@@ -12077,7 +12077,7 @@ bool URigVMController::AddLink(URigVMPin* OutputPin, URigVMPin* InputPin, bool b
 #if WITH_EDITOR
 		if (!bSuspendTemplateComputation)
 		{
-			auto ResolveTemplateNodeToCommonTypes = [this](URigVMPin* Pin)
+			auto ResolveTemplateNodeToCommonTypes = [this, bSetupUndoRedo](URigVMPin* Pin)
 			{
 				if(!Pin->IsExecuteContext())
 				{
@@ -12111,7 +12111,7 @@ bool URigVMController::AddLink(URigVMPin* OutputPin, URigVMPin* InputPin, bool b
 				if(PreferredPermutation != INDEX_NONE)
 				{
 					const TGuardValue<bool> DisableRegisterUseOfTemplate(bRegisterTemplateNodeUsage, false);
-					if(FullyResolveTemplateNode(TemplateNode, PreferredPermutation, true))
+					if(FullyResolveTemplateNode(TemplateNode, PreferredPermutation, bSetupUndoRedo))
 					{
 						static const FString Message = TEXT("Template node was automatically resolved to commonly used types.");
 						SendUserFacingNotification(Message, 0.f, TemplateNode, TEXT("MessageLog.Note"));
@@ -19610,16 +19610,16 @@ bool URigVMController::FullyResolveTemplateNode(URigVMTemplateNode* InNode, int3
 						return false;
 					}
 				}
-				if(!ChangePinType(Pin, ExpectedTypeIndex, bSetupUndoRedo, false, true, true))
-				{
-					if(bSetupUndoRedo)
-					{
-						CancelUndoBracket();
-					}
-					return false;
-				}
 			}
 		}
+		if(!ChangePinType(Pin, ExpectedTypeIndex, bSetupUndoRedo, false, true, true))
+		{
+			if(bSetupUndoRedo)
+			{
+				CancelUndoBracket();
+			}
+			return false;
+		}
 	}
 
 	// remove obsolete pins
diff --git a/Engine/Plugins/Runtime/RigVM/Source/RigVMDeveloper/Public/RigVMModel/RigVMController.h b/Engine/Plugins/Runtime/RigVM/Source/RigVMDeveloper/Public/RigVMModel/RigVMController.h
--- a/Engine/Plugins/Runtime/RigVM/Source/RigVMDeveloper/Public/RigVMModel/RigVMController.h
+++ b/Engine/Plugins/Runtime/RigVM/Source/RigVMDeveloper/Public/RigVMModel/RigVMController.h
@@ -1718,7 +1718,7 @@ struct FRigVMController_CommonTypePerTemplate
  * Default settings for the RigVM Controller
  */
 UCLASS(config = EditorSettings)
-class  URigVMControllerSettings : public UObject
+class RIGVMDEVELOPER_API URigVMControllerSettings : public UObject
 {
 public:
 	URigVMControllerSettings(const FObjectInitializer& Initializer);
