RigVM: Check nullptr when finding execute pins

#rnx

diff --git a/Engine/Plugins/Runtime/RigVM/Source/RigVMDeveloper/Private/RigVMModel/RigVMController.cpp b/Engine/Plugins/Runtime/RigVM/Source/RigVMDeveloper/Private/RigVMModel/RigVMController.cpp
--- a/Engine/Plugins/Runtime/RigVM/Source/RigVMDeveloper/Private/RigVMModel/RigVMController.cpp
+++ b/Engine/Plugins/Runtime/RigVM/Source/RigVMDeveloper/Private/RigVMModel/RigVMController.cpp
@@ -22883,13 +22883,16 @@ FRigVMClientPatchResult URigVMController::PatchFunctionsWithInvalidReturnPaths()
 		bool bReturnExecuteIsLinked = false;
 		if (bReturnIsMutable)
 		{
-			URigVMPin* ReturnExecutePin = FindExecutePin(ReturnNode);
-			bReturnExecuteIsLinked = ReturnExecutePin->IsLinked();
+			if(URigVMPin* ReturnExecutePin = FindExecutePin(ReturnNode))
+			{
+				bReturnExecuteIsLinked = ReturnExecutePin->IsLinked();
+			}
 		}
 
 		if (bEntryIsMutable && (!bReturnIsMutable || !bReturnExecuteIsLinked))
 		{
 			URigVMPin* EntryExecutePin = FindExecutePin(EntryNode);
+			check(EntryExecutePin);
 			
 			URigVMPin* OtherExecutePin = nullptr;
 			if (EntryExecutePin->IsLinked())
@@ -22903,7 +22906,9 @@ FRigVMClientPatchResult URigVMController::PatchFunctionsWithInvalidReturnPaths()
 			RemoveExposedPin(ExecutePinName, false);
 			AddExposedPin(ExecutePinName, ERigVMPinDirection::IO, ExecuteCPPType, ExecuteCPPTypeObjectPath, FString(), false);
 			EntryExecutePin = FindExecutePin(EntryNode);
+			check(EntryExecutePin);
 			URigVMPin* ReturnExecutePin = FindExecutePin(ReturnNode);
+			check(ReturnExecutePin);
 			
 			if (OtherExecutePin)
 			{
