Fix for RigVM VM assigning a temporary to a TArrayView.
Changed the TArrayView for a TArray, as the normal use is assigning a temporary array created by a function that contains three pointers.

diff --git a/Engine/Plugins/Runtime/RigVM/Source/RigVM/Private/RigVMCore/RigVM.cpp b/Engine/Plugins/Runtime/RigVM/Source/RigVM/Private/RigVMCore/RigVM.cpp
--- a/Engine/Plugins/Runtime/RigVM/Source/RigVM/Private/RigVMCore/RigVM.cpp
+++ b/Engine/Plugins/Runtime/RigVM/Source/RigVM/Private/RigVMCore/RigVM.cpp
@@ -1696,7 +1696,7 @@ ERigVMExecuteResult URigVM::ExecuteVM(FRigVMExtendedExecuteContext& Context, con
 #if WITH_EDITOR
 	if(bIsRootEntry)
 	{
-		Context.CurrentVMMemory = TArrayView<FRigVMMemoryStorageStruct*>();
+		Context.CurrentVMMemory.Reset();
 		
 		if (FRigVMDebugInfo* RigVMDebugInfo = Context.GetRigVMDebugInfo())
 		{
diff --git a/Engine/Plugins/Runtime/RigVM/Source/RigVM/Private/RigVMCore/RigVMExecuteContext.cpp b/Engine/Plugins/Runtime/RigVM/Source/RigVM/Private/RigVMCore/RigVMExecuteContext.cpp
--- a/Engine/Plugins/Runtime/RigVM/Source/RigVM/Private/RigVMCore/RigVMExecuteContext.cpp
+++ b/Engine/Plugins/Runtime/RigVM/Source/RigVM/Private/RigVMCore/RigVMExecuteContext.cpp
@@ -94,7 +94,7 @@ void FRigVMExtendedExecuteContext::Reset()
 	WorkMemoryStorage = FRigVMMemoryStorageStruct();
 	DebugMemoryStorage = FRigVMMemoryStorageStruct();
 
-	CurrentVMMemory = TArrayView<FRigVMMemoryStorageStruct*>();
+	CurrentVMMemory.Reset();
 	ExecutionReachedExit().Clear();
 
 	CachedMemoryHandles.Reset();
diff --git a/Engine/Plugins/Runtime/RigVM/Source/RigVM/Public/RigVMCore/RigVMExecuteContext.h b/Engine/Plugins/Runtime/RigVM/Source/RigVM/Public/RigVMCore/RigVMExecuteContext.h
--- a/Engine/Plugins/Runtime/RigVM/Source/RigVM/Public/RigVMCore/RigVMExecuteContext.h
+++ b/Engine/Plugins/Runtime/RigVM/Source/RigVM/Public/RigVMCore/RigVMExecuteContext.h
@@ -823,7 +823,7 @@ struct RIGVM_API FRigVMExtendedExecuteContext
 	FName CurrentEntryName = NAME_None;
 	bool bCurrentlyRunningRootEntry = false;
 
-	TArrayView<FRigVMMemoryStorageStruct*> CurrentVMMemory;
+	TArray<FRigVMMemoryStorageStruct*> CurrentVMMemory;
 
 #if WITH_EDITORONLY_DATA
 	// changes to the layout of cached memory array should be reflected in GetContainerIndex()
