[AutoRTFM] Replace atomic race check with a transactionally-safe lock.

Previously, this code would bump an atomic to indicate that work was taking place, and then elsewhere it would check-fail if the atomic was non-zero.

This code is directly reachable from Verse, and in AutoRTFM we don't want other threads to be able to view side effects from transactional work at all. I've replaced this atomic with a transactionally-safe lock which will enforce our transactional restrictions properly.

I also experimented with replacing the atomic value with a MT recursive access detector, where the ExecuteVM function was designated as performing a write access, and CopyDataForSerialization was assigned read access. This immediately violated the access detector. I didn't attempt to analyze the failure in detail.

diff --git a/Engine/Plugins/Runtime/RigVM/Source/RigVM/Private/RigVMCore/RigVM.cpp b/Engine/Plugins/Runtime/RigVM/Source/RigVM/Private/RigVMCore/RigVM.cpp
--- a/Engine/Plugins/Runtime/RigVM/Source/RigVM/Private/RigVMCore/RigVM.cpp
+++ b/Engine/Plugins/Runtime/RigVM/Source/RigVM/Private/RigVMCore/RigVM.cpp
@@ -10,6 +10,8 @@
 #include "UObject/FortniteMainBranchObjectVersion.h"
 #include "RigVMObjectVersion.h"
 #include "HAL/PlatformTLS.h"
+#include "Misc/Optional.h"
+#include "Misc/ScopeLock.h"
 #include "Async/ParallelFor.h"
 #include "Engine/UserDefinedEnum.h"
 #include "GenericPlatform/GenericPlatformSurvey.h"
@@ -284,7 +286,7 @@ void URigVM::Load(FArchive& Ar)
 void URigVM::CopyDataForSerialization(URigVM* InVM)
 {
 	check(InVM);
-	check(ActiveExecutions.load() == 0);
+	UE::TScopeLock Lock(ActiveExecutionLock);
 
 	Reset_Internal();
 
@@ -1535,24 +1537,18 @@ ERigVMExecuteResult URigVM::ExecuteVM(FRigVMExtendedExecuteContext& Context, con
 {
 	// if this the first entry being executed - get ready for execution
 	const bool bIsRootEntry = Context.EntriesBeingExecuted.IsEmpty();
-	ON_SCOPE_EXIT
-	{
-		if (bIsRootEntry)
-		{
-			ActiveExecutions--;
-		}
-	};
 
+	TOptional<UE::TScopeLock<FTransactionallySafeCriticalSection>> Lock;
 	TGuardValue<FName> EntryNameGuard(Context.CurrentEntryName, InEntryName);
 	TGuardValue<bool> RootEntryGuard(Context.bCurrentlyRunningRootEntry, bIsRootEntry);
 
 	FRigVMByteCode& ByteCode = GetByteCode();
 
-	if(bIsRootEntry)
+	if (bIsRootEntry)
 	{
 		Context.CurrentExecuteResult = ERigVMExecuteResult::Succeeded;
-		
-		ActiveExecutions++;
+
+		Lock.Emplace(ActiveExecutionLock);  // The lock will be held for the duration of ExecuteVM.
 
 		if (Instructions.Num() == 0)
 		{
diff --git a/Engine/Plugins/Runtime/RigVM/Source/RigVM/Public/RigVMCore/RigVM.h b/Engine/Plugins/Runtime/RigVM/Source/RigVM/Public/RigVMCore/RigVM.h
--- a/Engine/Plugins/Runtime/RigVM/Source/RigVM/Public/RigVMCore/RigVM.h
+++ b/Engine/Plugins/Runtime/RigVM/Source/RigVM/Public/RigVMCore/RigVM.h
@@ -14,6 +14,7 @@
 #include "Math/Vector.h"
 #include "Math/Vector2D.h"
 #include "Misc/AssertionMacros.h"
+#include "Misc/TransactionallySafeCriticalSection.h"
 #include "RigVMByteCode.h"
 #include "RigVMCore/RigVMExternalVariable.h"
 #include "RigVMCore/RigVMFunction.h"
@@ -1012,14 +1013,9 @@ protected:
 	uint32 NumExecutions_DEPRECATED = 0;
 #endif
 
-	std::atomic<int32> ActiveExecutions;
+	FTransactionallySafeCriticalSection ActiveExecutionLock;
 
 public:
-	int32 GetActiveExecutions() const
-	{
-		return ActiveExecutions.load();
-	}
-
 	bool ValidateBytecode();
 	void RefreshArgumentNameCaches();
 
