Control Rig: New RigVMEditor, which does not inherit from BlueprintEditor

diff --git a/Engine/Plugins/Runtime/RigVM/Source/RigVMEditor/Private/Editor/RigVMEditor.cpp b/Engine/Plugins/Runtime/RigVM/Source/RigVMEditor/Private/Editor/RigVMEditor.cpp
--- a/Engine/Plugins/Runtime/RigVM/Source/RigVMEditor/Private/Editor/RigVMEditor.cpp
+++ b/Engine/Plugins/Runtime/RigVM/Source/RigVMEditor/Private/Editor/RigVMEditor.cpp
@@ -1,60 +1,54 @@
 // Copyright Epic Games, Inc. All Rights Reserved.
 
 #include "Editor/RigVMEditor.h"
-#include "Editor/RigVMEditorTools.h"
-#include "Editor/RigVMEditorMenuContext.h"
-#include "HAL/PlatformApplicationMisc.h"
-#include "SMyBlueprint.h"
-#include "EdGraphNode_Comment.h"
-#include "RigVMFunctions/RigVMFunction_ControlFlow.h"
-#include "Editor.h"
-#include "Editor/Transactor.h"
-#include "BlueprintEditorTabs.h"
-#include "Kismet2/BlueprintEditorUtils.h"
-#include "Widgets/Notifications/SNotificationList.h"
-#include "Framework/Notifications/NotificationManager.h"
-#include "Editor/RigVMEditorCommands.h"
+
+#include "BlueprintActionDatabase.h"
 #include "BlueprintCompilationManager.h"
-#include "SBlueprintEditorToolbar.h"
-#include "BlueprintEditorModes.h"
-#include "AssetEditorModeManager.h"
-#include "RigVMCore/RigVMMemoryStorageStruct.h"
+#include "EdGraphNode_Comment.h"
+#include "EditorModeManager.h"
+#include "EulerTransform.h"
+#include "InstancedPropertyBagStructureDataProvider.h"
+#include "RigVMBlueprint.h"
 #include "RigVMBlueprintUtils.h"
+#include "RigVMEditorCommands.h"
 #include "RigVMPythonUtils.h"
-#include "EulerTransform.h"
-#include "RigVMEditorModule.h"
+#include "AssetRegistry/AssetRegistryModule.h"
+#include "AssetRegistry/IAssetRegistry.h"
+#include "Editor/RigVMEditorMenuContext.h"
+#include "Editor/RigVMEditorTools.h"
 #include "Editor/RigVMGraphDetailCustomization.h"
+#include "Editor/RigVMLegacyEditor.h"
+#include "Editor/RigVMNewEditor.h"
+#include "Editor/Transactor.h"
+#include "Framework/Commands/GenericCommands.h"
+#include "Framework/Notifications/NotificationManager.h"
+#include "Kismet2/BlueprintEditorUtils.h"
 #include "Kismet2/KismetDebugUtilities.h"
 #include "Kismet2/WatchedPin.h"
-#include "SBlueprintEditorSelectedDebugObjectWidget.h"
-#include "IMessageLogListing.h"
-#include "SGraphPanel.h"
-#include "BlueprintActionDatabase.h"
+#include "RigVMFunctions/RigVMFunction_ControlFlow.h"
 #include "RigVMModel/Nodes/RigVMAggregateNode.h"
-#include "Widgets/SRigVMGraphFunctionLocalizationWidget.h"
-#include "Widgets/SRigVMGraphFunctionBulkEditWidget.h"
+#include "Widgets/SRigVMBulkEditDialog.h"
 #include "Widgets/SRigVMGraphBreakLinksWidget.h"
-#include "Widgets/SRigVMGraphChangePinType.h"
-#include "Framework/Commands/GenericCommands.h"
+#include "Widgets/SRigVMGraphFunctionBulkEditWidget.h"
+#include "Widgets/SRigVMSwapAssetReferencesWidget.h"
+#include "Widgets/SRigVMSwapFunctionsWidget.h"
+#include "Widgets/Notifications/SNotificationList.h"
 #include "Stats/StatsHierarchical.h"
-#include "Widgets/Layout/SBorder.h"
-#include "Styling/AppStyle.h"
+#include "ToolMenuContext.h"
 #include "Widgets/Docking/SDockTab.h"
 #include "Widgets/Layout/SScrollBox.h"
-#include "Framework/Docking/TabManager.h"
-#include "ScopedTransaction.h"
-#include "Editor/RigVMEditorMode.h"
-#include "InstancedPropertyBagStructureDataProvider.h"
-#include "Widgets/SRigVMSwapFunctionsWidget.h"
-#include "Widgets/SRigVMSwapAssetReferencesWidget.h"
-#include "Widgets/SRigVMBulkEditDialog.h"
-#include "AssetRegistry/AssetRegistryModule.h"
-#include "ToolMenuContext.h"
-#include "Widgets/SRigVMEditorGraphExplorer.h"
+#include "IStructureDetailsView.h"
+
+#if WITH_RIGVMLEGACYEDITOR
+#include "SKismetInspector.h"
+#else
+#include "Editor/SRigVMDetailsInspector.h"
+#endif
 
 #define LOCTEXT_NAMESPACE "RigVMEditor"
 
-FRigVMEditor::FRigVMEditor()
+
+FRigVMEditorBase::FRigVMEditorBase()
 	: bAnyErrorsLeft(false)
 	, KnownInstructionLimitWarnings()
 	, HaltedAtNode(nullptr)
@@ -71,7 +65,7 @@ FRigVMEditor::FRigVMEditor()
 {
 }
 
-FRigVMEditor::~FRigVMEditor()
+void FRigVMEditorBase::UnbindEditor()
 {
 	URigVMBlueprint* RigVMBlueprint = GetRigVMBlueprint();
 	RigVMEditorClosedDelegate.Broadcast(this, RigVMBlueprint);
@@ -84,7 +78,7 @@ FRigVMEditor::~FRigVMEditor()
 	}
 	
 	FEditorDelegates::EndPIE.RemoveAll(this);
-    FEditorDelegates::CancelPIE.RemoveAll(this);
+	FEditorDelegates::CancelPIE.RemoveAll(this);
 
 	if (RigVMBlueprint)
 	{
@@ -128,25 +122,51 @@ FRigVMEditor::~FRigVMEditor()
 	}
 }
 
-void FRigVMEditor::InitRigVMEditor(const EToolkitMode::Type Mode, const TSharedPtr< class IToolkitHost >& InitToolkitHost, class URigVMBlueprint* InRigVMBlueprint)
+FRigVMEditorBase* FRigVMEditorBase::GetFromAssetEditorInstance(IAssetEditorInstance* Instance)
+{
+	FRigVMEditorBase* Editor = nullptr;
+	FWorkflowCentricApplication* App = static_cast<FWorkflowCentricApplication*>(Instance);
+
+	TSharedRef<FAssetEditorToolkit> SharedApp = App->AsShared();
+#if WITH_RIGVMLEGACYEDITOR
+	if (SharedApp->IsBlueprintEditor())
+	{
+		TSharedPtr<FRigVMLegacyEditor> LegacyEditor = StaticCastSharedPtr<FRigVMLegacyEditor>(App->AsShared().ToSharedPtr());
+		Editor = LegacyEditor.Get();
+	}
+	else
+	{
+		TSharedPtr<FRigVMNewEditor> NewEditor = StaticCastSharedPtr<FRigVMNewEditor>(App->AsShared().ToSharedPtr());
+		Editor = NewEditor.Get();
+	}
+#else
+	TSharedPtr<FRigVMNewEditor> NewEditor = StaticCastSharedPtr<FRigVMNewEditor>(App->AsShared().ToSharedPtr());
+	Editor = NewEditor.Get();
+#endif
+	return Editor;
+}
+
+void FRigVMEditorBase::InitRigVMEditor(const EToolkitMode::Type Mode, const TSharedPtr< class IToolkitHost >& InitToolkitHost, class URigVMBlueprint* InRigVMBlueprint)
 {
 	DECLARE_SCOPE_HIERARCHICAL_COUNTER_FUNC()
 
 	check(InRigVMBlueprint);
 
+#if WITH_RIGVMLEGACYEDITOR
 	FBlueprintCompilationManager::FlushCompilationQueue(nullptr);
+#else
+	// TODO: sara-s remove once blueprint backend is replaced
+	// FRigVMBlueprintCompilationManager::FlushCompilationQueue(nullptr);
+#endif
 
 	Toolbox = SNew(SBorder)
 		.BorderImage(FAppStyle::GetBrush("ToolPanel.GroupBorder"))
 		.Padding(0.f);
 
-	if (!Toolbar.IsValid())
-	{
-		Toolbar = MakeShareable(new FBlueprintEditorToolbar(SharedThis(this)));
-	}
+	CreateEditorToolbar();
 
-	FEditorDelegates::EndPIE.AddRaw(this, &FRigVMEditor::OnPIEStopped);
-	FEditorDelegates::CancelPIE.AddRaw(this, &FRigVMEditor::OnPIEStopped, false);
+	FEditorDelegates::EndPIE.AddRaw(this, &FRigVMEditorBase::OnPIEStopped);
+	FEditorDelegates::CancelPIE.AddRaw(this, &FRigVMEditorBase::OnPIEStopped, false);
 
 	// Build up a list of objects being edited in this asset editor
 	TArray<UObject*> ObjectsBeingEdited;
@@ -156,8 +176,8 @@ void FRigVMEditor::InitRigVMEditor(const EToolkitMode::Type Mode, const TSharedP
 	const bool bCreateDefaultStandaloneMenu = true;
 	const bool bCreateDefaultToolbar = true;
 	InitAssetEditor(Mode, InitToolkitHost, GetEditorAppName(), FTabManager::FLayout::NullLayout, bCreateDefaultStandaloneMenu, bCreateDefaultToolbar, ObjectsBeingEdited);
-	GEditor->GetEditorSubsystem<UAssetEditorSubsystem>()->OnAssetEditorRequestedOpen().AddSP(this, &FRigVMEditor::HandleAssetRequestedOpen);
-	GEditor->GetEditorSubsystem<UAssetEditorSubsystem>()->OnAssetEditorRequestClose().AddSP(this, &FRigVMEditor::HandleAssetRequestClose);
+	GEditor->GetEditorSubsystem<UAssetEditorSubsystem>()->OnAssetEditorRequestedOpen().AddSP(StaticCastSharedRef<FRigVMEditorBase>(SharedRef()), &FRigVMEditorBase::HandleAssetRequestedOpen);
+	GEditor->GetEditorSubsystem<UAssetEditorSubsystem>()->OnAssetEditorRequestClose().AddSP(StaticCastSharedRef<FRigVMEditorBase>(SharedRef()), &FRigVMEditorBase::HandleAssetRequestClose);
 
 	CreateDefaultCommands();
 
@@ -203,10 +223,10 @@ void FRigVMEditor::InitRigVMEditor(const EToolkitMode::Type Mode, const TSharedP
 
 	URigVMBlueprint::sCurrentlyOpenedRigVMBlueprints.AddUnique(InRigVMBlueprint);
 
-	InRigVMBlueprint->OnModified().AddSP(this, &FRigVMEditor::HandleModifiedEvent);
-	InRigVMBlueprint->OnVMCompiled().AddSP(this, &FRigVMEditor::HandleVMCompiledEvent);
-	InRigVMBlueprint->OnRequestInspectObject().AddSP(this, &FRigVMEditor::SetDetailObjects);
-	InRigVMBlueprint->OnRequestInspectMemoryStorage().AddSP(this, &FRigVMEditor::SetMemoryStorageDetails);
+	InRigVMBlueprint->OnModified().AddSP(StaticCastSharedRef<FRigVMEditorBase>(SharedRef()), &FRigVMEditorBase::HandleModifiedEvent);
+	InRigVMBlueprint->OnVMCompiled().AddSP(StaticCastSharedRef<FRigVMEditorBase>(SharedRef()), &FRigVMEditorBase::HandleVMCompiledEvent);
+	InRigVMBlueprint->OnRequestInspectObject().AddSP(StaticCastSharedRef<FRigVMEditorBase>(SharedRef()), &FRigVMEditorBase::SetDetailObjects);
+	InRigVMBlueprint->OnRequestInspectMemoryStorage().AddSP(StaticCastSharedRef<FRigVMEditorBase>(SharedRef()), &FRigVMEditorBase::SetMemoryStorageDetails);
 
 	BindCommands();
 
@@ -229,8 +249,8 @@ void FRigVMEditor::InitRigVMEditor(const EToolkitMode::Type Mode, const TSharedP
 		SetCurrentMode(GetEditorModeName());
 
 		// Activate our edit mode
-		GetEditorModeManager().SetDefaultMode(GetEditorModeName());
-		GetEditorModeManager().ActivateMode(GetEditorModeName());
+		GetToolkitEditorModeManager().SetDefaultMode(GetEditorModeName());
+		GetToolkitEditorModeManager().ActivateMode(GetEditorModeName());
 	}
 
 	{
@@ -258,7 +278,7 @@ void FRigVMEditor::InitRigVMEditor(const EToolkitMode::Type Mode, const TSharedP
 					bBroughtGraphToFront = true;
 				}
 
-				RigGraph->OnGraphNodeClicked.AddSP(this, &FRigVMEditor::OnGraphNodeClicked);
+				RigGraph->OnGraphNodeClicked.AddSP(StaticCastSharedRef<FRigVMEditorBase>(SharedRef()), &FRigVMEditorBase::OnGraphNodeClicked);
 				ActiveTabNodePath = RigGraph->ModelNodePath;
 			}
 		}
@@ -311,19 +331,19 @@ void FRigVMEditor::InitRigVMEditor(const EToolkitMode::Type Mode, const TSharedP
 			}
 		}
 
-		InRigVMBlueprint->OnRefreshEditor().AddSP(this, &FRigVMEditor::HandleRefreshEditorFromBlueprint);
-		InRigVMBlueprint->OnVariableDropped().AddSP(this, &FRigVMEditor::HandleVariableDroppedFromBlueprint);
-		InRigVMBlueprint->OnBreakpointAdded().AddSP(this, &FRigVMEditor::HandleBreakpointAdded);
-
-		InRigVMBlueprint->OnNodeDoubleClicked().AddSP(this, &FRigVMEditor::OnNodeDoubleClicked);
-		InRigVMBlueprint->OnGraphImported().AddSP(this, &FRigVMEditor::OnGraphImported);
-		InRigVMBlueprint->OnRequestLocalizeFunctionDialog().AddSP(this, &FRigVMEditor::OnRequestLocalizeFunctionDialog);
-		InRigVMBlueprint->OnRequestBulkEditDialog().BindSP(this, &FRigVMEditor::OnRequestBulkEditDialog);
-		InRigVMBlueprint->OnRequestBreakLinksDialog().BindSP(this, &FRigVMEditor::OnRequestBreakLinksDialog);
-		InRigVMBlueprint->OnRequestPinTypeSelectionDialog().BindSP(this, &FRigVMEditor::OnRequestPinTypeSelectionDialog);
-		InRigVMBlueprint->OnRequestJumpToHyperlink().BindSP(this, &FRigVMEditor::HandleJumpToHyperlink);
+		InRigVMBlueprint->OnRefreshEditor().AddSP(StaticCastSharedRef<FRigVMEditorBase>(SharedRef()), &FRigVMEditorBase::HandleRefreshEditorFromBlueprint);
+		InRigVMBlueprint->OnVariableDropped().AddSP(StaticCastSharedRef<FRigVMEditorBase>(SharedRef()), &FRigVMEditorBase::HandleVariableDroppedFromBlueprint);
+		InRigVMBlueprint->OnBreakpointAdded().AddSP(StaticCastSharedRef<FRigVMEditorBase>(SharedRef()), &FRigVMEditorBase::HandleBreakpointAdded);
+
+		InRigVMBlueprint->OnNodeDoubleClicked().AddSP(StaticCastSharedRef<FRigVMEditorBase>(SharedRef()), &FRigVMEditorBase::OnNodeDoubleClicked);
+		InRigVMBlueprint->OnGraphImported().AddSP(StaticCastSharedRef<FRigVMEditorBase>(SharedRef()), &FRigVMEditorBase::OnGraphImported);
+		InRigVMBlueprint->OnRequestLocalizeFunctionDialog().AddSP(StaticCastSharedRef<FRigVMEditorBase>(SharedRef()), &FRigVMEditorBase::OnRequestLocalizeFunctionDialog);
+		InRigVMBlueprint->OnRequestBulkEditDialog().BindSP(StaticCastSharedRef<FRigVMEditorBase>(SharedRef()), &FRigVMEditorBase::OnRequestBulkEditDialog);
+		InRigVMBlueprint->OnRequestBreakLinksDialog().BindSP(StaticCastSharedRef<FRigVMEditorBase>(SharedRef()), &FRigVMEditorBase::OnRequestBreakLinksDialog);
+		InRigVMBlueprint->OnRequestPinTypeSelectionDialog().BindSP(StaticCastSharedRef<FRigVMEditorBase>(SharedRef()), &FRigVMEditorBase::OnRequestPinTypeSelectionDialog);
+		InRigVMBlueprint->OnRequestJumpToHyperlink().BindSP(StaticCastSharedRef<FRigVMEditorBase>(SharedRef()), &FRigVMEditorBase::HandleJumpToHyperlink);
 #if WITH_EDITOR
-		InRigVMBlueprint->OnGetFocusedGraph().BindSP(this, &FRigVMEditor::GetFocusedModel);
+		InRigVMBlueprint->OnGetFocusedGraph().BindSP(StaticCastSharedRef<FRigVMEditorBase>(SharedRef()), &FRigVMEditorBase::GetFocusedModel);
 #endif
 	}
 
@@ -369,7 +389,7 @@ void FRigVMEditor::InitRigVMEditor(const EToolkitMode::Type Mode, const TSharedP
 		}
 		
 		// Always show the myblueprint tab
-		static const FTabId MyBlueprintTabId(FBlueprintEditorTabs::MyBlueprintID);
+		static const FTabId MyBlueprintTabId(GetGraphExplorerWidgetID());
 		if (!TabManager->FindExistingLiveTab(MyBlueprintTabId).IsValid())
 		{
 			TabManager->TryInvokeTab(MyBlueprintTabId);
@@ -397,23 +417,37 @@ void FRigVMEditor::InitRigVMEditor(const EToolkitMode::Type Mode, const TSharedP
 
 	for(UScriptStruct* StructToCustomize : StructsToCustomize)
 	{
-		Inspector->GetPropertyView()->RegisterInstancedCustomPropertyTypeLayout(StructToCustomize->GetFName(),
-			FOnGetPropertyTypeCustomizationInstance::CreateLambda([=]()
-			{
-				return FRigVMGraphMathTypeDetailCustomization::MakeInstance();
-			}));
+#if WITH_RIGVMLEGACYEDITOR
+		if (TSharedPtr<SKismetInspector> Inspector = GetInspector())
+#else
+		if (TSharedPtr<SRigVMDetailsInspector> Inspector = GetInspector())
+#endif
+		{
+			Inspector->GetPropertyView()->RegisterInstancedCustomPropertyTypeLayout(StructToCustomize->GetFName(),
+			   FOnGetPropertyTypeCustomizationInstance::CreateLambda([=]()
+			   {
+				   return FRigVMGraphMathTypeDetailCustomization::MakeInstance();
+			   }));
+		}
 	}
 
-	Inspector->GetPropertyView()->RegisterInstancedCustomPropertyTypeLayout(UEnum::StaticClass()->GetFName(),
-		FOnGetPropertyTypeCustomizationInstance::CreateLambda([=]()
-		{
-			return FRigVMGraphEnumDetailCustomization::MakeInstance();
-		}));
+#if WITH_RIGVMLEGACYEDITOR
+	if (TSharedPtr<SKismetInspector> Inspector = GetInspector())
+#else
+	if (TSharedPtr<SRigVMDetailsInspector> Inspector = GetInspector())
+#endif
+	{
+		Inspector->GetPropertyView()->RegisterInstancedCustomPropertyTypeLayout(UEnum::StaticClass()->GetFName(),
+		   FOnGetPropertyTypeCustomizationInstance::CreateLambda([=]()
+		   {
+			   return FRigVMGraphEnumDetailCustomization::MakeInstance();
+		   }));
+	}
 
-	PropertyChangedHandle = FCoreUObjectDelegates::OnObjectPropertyChanged.AddSP(this, &FRigVMEditor::OnPropertyChanged);
+	PropertyChangedHandle = FCoreUObjectDelegates::OnObjectPropertyChanged.AddSP(StaticCastSharedRef<FRigVMEditorBase>(SharedRef()), &FRigVMEditorBase::OnPropertyChanged);
 }
 
-void FRigVMEditor::HandleAssetRequestedOpen(UObject* InObject)
+void FRigVMEditorBase::HandleAssetRequestedOpen(UObject* InObject)
 {
 	if (InObject == GetRigVMBlueprint())
 	{
@@ -421,7 +455,7 @@ void FRigVMEditor::HandleAssetRequestedOpen(UObject* InObject)
 	}
 }
 
-void FRigVMEditor::HandleAssetRequestClose(UObject* InObject, EAssetEditorCloseReason InReason)
+void FRigVMEditorBase::HandleAssetRequestClose(UObject* InObject, EAssetEditorCloseReason InReason)
 {
 	if (InObject == GetRigVMBlueprint())
 	{
@@ -429,23 +463,12 @@ void FRigVMEditor::HandleAssetRequestClose(UObject* InObject, EAssetEditorCloseR
 	}
 }
 
-const FName FRigVMEditor::GetEditorAppName() const
-{
-	static const FLazyName AppName(TEXT("RigVMEditorApp"));
-	return AppName;
-}
-
-const FName FRigVMEditor::GetEditorModeName() const
+const FName FRigVMEditorBase::GetEditorModeName() const
 {
 	return FRigVMEditorModes::RigVMEditorMode;
 }
 
-TSharedPtr<FApplicationMode> FRigVMEditor::CreateEditorMode()
-{
-	return MakeShareable(new FRigVMEditorMode(SharedThis(this)));
-}
-
-UBlueprint* FRigVMEditor::GetBlueprintObj() const
+UBlueprint* FRigVMEditorBase::GetBlueprintObj() const
 {
 	const TArray<UObject*>& EditingObjs = GetEditingObjects();
 	for (UObject* Obj : EditingObjs)
@@ -458,7 +481,7 @@ UBlueprint* FRigVMEditor::GetBlueprintObj() const
 	return nullptr;
 }
 
-TSubclassOf<UEdGraphSchema> FRigVMEditor::GetDefaultSchemaClass() const
+TSubclassOf<UEdGraphSchema> FRigVMEditorBase::GetDefaultSchemaClass() const
 {
 	if(const URigVMBlueprint* RigVMBlueprint = GetRigVMBlueprint())
 	{
@@ -467,34 +490,31 @@ TSubclassOf<UEdGraphSchema> FRigVMEditor::GetDefaultSchemaClass() const
 	return URigVMEdGraphSchema::StaticClass();
 }
 
-bool FRigVMEditor::InEditingMode() const
+bool FRigVMEditorBase::InEditingMode() const
 {
 	// always allow editing - also during PIE.
 	return true;
 }
 
-void FRigVMEditor::OnGraphEditorFocused(const TSharedRef<SGraphEditor>& InGraphEditor)
+void FRigVMEditorBase::OnGraphEditorFocused(const TSharedRef<SGraphEditor>& InGraphEditor)
 {
-	FBlueprintEditor::OnGraphEditorFocused(InGraphEditor);
 	if (GraphExplorerWidget.IsValid())
 	{
 		GraphExplorerWidget->Refresh();
 	}
 }
 
-void FRigVMEditor::Tick(float DeltaTime)
+void FRigVMEditorBase::Tick(float DeltaTime)
 {
-	FBlueprintEditor::Tick(DeltaTime);
-
 	// tick the  rigvm host
 	if (const URigVMBlueprint* Blueprint = GetRigVMBlueprint())
 	{
 		if (LastDebuggedHost != GetCustomDebugObjectLabel(Blueprint->GetObjectBeingDebugged()))
 		{
-			TArray<FCustomDebugObject> DebugList;
+			TArray<FRigVMCustomDebugObject> DebugList;
 			GetCustomDebugObjects(DebugList);
 
-			for (const FCustomDebugObject& DebugObject : DebugList)
+			for (const FRigVMCustomDebugObject& DebugObject : DebugList)
 			{
 				if (DebugObject.NameOverride == LastDebuggedHost)
 				{
@@ -506,48 +526,46 @@ void FRigVMEditor::Tick(float DeltaTime)
 	}
 }
 
-void FRigVMEditor::BringToolkitToFront()
+void FRigVMEditorBase::BringToolkitToFront()
 {
-	if (ToolkitHost.IsValid())
+	if (IsHosted())
 	{
-		FBlueprintEditor::BringToolkitToFront();
+		BringToolkitToFrontImpl();
 	}
 }
 
-FName FRigVMEditor::GetToolkitFName() const
+FName FRigVMEditorBase::GetToolkitFName() const
 {
 	return FName("RigVMEditor");
 }
 
-FName FRigVMEditor::GetToolkitContextFName() const
+FName FRigVMEditorBase::GetToolkitContextFName() const
 {
 	return GetToolkitFName();
 }
 
-FText FRigVMEditor::GetBaseToolkitName() const
+FText FRigVMEditorBase::GetBaseToolkitName() const
 {
 	return LOCTEXT("AppLabel", "RigVM Editor");
 }
 
-FText FRigVMEditor::GetToolkitToolTipText() const
+FText FRigVMEditorBase::GetToolkitToolTipText() const
 {
 	return FAssetEditorToolkit::GetToolTipTextForObject(GetBlueprintObj());
 }
 
-FString FRigVMEditor::GetWorldCentricTabPrefix() const
+FString FRigVMEditorBase::GetWorldCentricTabPrefix() const
 {
 	return LOCTEXT("WorldCentricTabPrefix", "RigVM Editor ").ToString();
 }
 
-FLinearColor FRigVMEditor::GetWorldCentricTabColorScale() const
+FLinearColor FRigVMEditorBase::GetWorldCentricTabColorScale() const
 {
 	return FLinearColor( 0.5f, 0.25f, 0.35f, 0.5f );
 }
 
-void FRigVMEditor::InitToolMenuContext(FToolMenuContext& MenuContext)
+void FRigVMEditorBase::InitToolMenuContextImpl(FToolMenuContext& MenuContext)
 {
-	FBlueprintEditor::InitToolMenuContext(MenuContext);
-
 	if (URigVMBlueprint* RigVMBlueprint = GetRigVMBlueprint())
 	{
 		URigVMGraph* Model = nullptr;
@@ -573,13 +591,13 @@ void FRigVMEditor::InitToolMenuContext(FToolMenuContext& MenuContext)
 		
 		URigVMEditorMenuContext* RigVMEditorMenuContext = NewObject<URigVMEditorMenuContext>();
 		const FRigVMEditorGraphMenuContext GraphMenuContext = FRigVMEditorGraphMenuContext(Model, Node, Pin);
-		RigVMEditorMenuContext->Init(SharedThis(this), GraphMenuContext);
+		RigVMEditorMenuContext->Init(StaticCastSharedRef<FRigVMEditorBase>(SharedRef()), GraphMenuContext);
 
 		MenuContext.AddObject(RigVMEditorMenuContext);
 	}
 }
 
-bool FRigVMEditor::TransactionObjectAffectsBlueprint(UObject* InTransactedObject)
+bool FRigVMEditorBase::TransactionObjectAffectsBlueprintImpl(UObject* InTransactedObject)
 {
 	URigVMBlueprint* RigVMBlueprint = GetRigVMBlueprint();
 	if (RigVMBlueprint == nullptr)
@@ -591,10 +609,10 @@ bool FRigVMEditor::TransactionObjectAffectsBlueprint(UObject* InTransactedObject
 	{
 		return false;
 	}
-	return FBlueprintEditor::TransactionObjectAffectsBlueprint(InTransactedObject);
+	return true;
 }
 
-bool FRigVMEditor::CanAddNewLocalVariable() const
+bool FRigVMEditorBase::CanAddNewLocalVariable() const
 {
 	const URigVMGraph* Graph = GetFocusedModel();
 	const URigVMGraph* ParentGraph = Graph->GetParentGraph();				
@@ -605,7 +623,7 @@ bool FRigVMEditor::CanAddNewLocalVariable() const
 	return false;
 }
 
-void FRigVMEditor::OnAddNewLocalVariable()
+void FRigVMEditorBase::OnAddNewLocalVariable()
 {
 	if (!CanAddNewLocalVariable())
 	{
@@ -613,7 +631,7 @@ void FRigVMEditor::OnAddNewLocalVariable()
 	}
 
 	FRigVMGraphVariableDescription LastTypeVar;
-	LastTypeVar.ChangeType(MyBlueprintWidget->GetLastPinTypeUsed());
+	LastTypeVar.ChangeType(GetLastPinTypeUsed());
 	FRigVMGraphVariableDescription NewVar = GetFocusedController()->AddLocalVariable(TEXT("NewLocalVar"), LastTypeVar.CPPType, LastTypeVar.CPPTypeObject, LastTypeVar.DefaultValue, true, true);
 	if(NewVar.Name.IsNone())
 	{
@@ -625,7 +643,7 @@ void FRigVMEditor::OnAddNewLocalVariable()
 	}
 }
 
-void FRigVMEditor::OnPasteNewLocalVariable(const FBPVariableDescription& VariableDescription)
+void FRigVMEditorBase::OnPasteNewLocalVariable(const FBPVariableDescription& VariableDescription)
 {
 	if (!CanAddNewLocalVariable())
 	{
@@ -641,7 +659,7 @@ void FRigVMEditor::OnPasteNewLocalVariable(const FBPVariableDescription& Variabl
 	}
 }
 
-void FRigVMEditor::DeleteSelectedNodes()
+void FRigVMEditorBase::DeleteSelectedNodes()
 {
 	DECLARE_SCOPE_HIERARCHICAL_COUNTER_FUNC()
 
@@ -702,30 +720,30 @@ void FRigVMEditor::DeleteSelectedNodes()
 	GetFocusedController()->CloseUndoBracket();
 }
 
-bool FRigVMEditor::CanDeleteNodes() const
+bool FRigVMEditorBase::CanDeleteNodes() const
 {
 	return true;
 }
 
-void FRigVMEditor::CopySelectedNodes()
+void FRigVMEditorBase::CopySelectedNodes()
 {
 	FString ExportedText = GetFocusedController()->ExportSelectedNodesToText();
 	FPlatformApplicationMisc::ClipboardCopy(*ExportedText);
 }
 
-bool FRigVMEditor::CanCopyNodes() const
+bool FRigVMEditorBase::CanCopyNodes() const
 {
 	return GetFocusedModel()->GetSelectNodes().Num() > 0;
 }
 
-bool FRigVMEditor::CanPasteNodes() const
+bool FRigVMEditorBase::CanPasteNodes() const
 {
 	FString TextToImport;
 	FPlatformApplicationMisc::ClipboardPaste(TextToImport);
 	return GetFocusedController()->CanImportNodesFromText(TextToImport);
 }
 
-FReply FRigVMEditor::OnSpawnGraphNodeByShortcut(FInputChord InChord, const FVector2D& InPosition,
+FReply FRigVMEditorBase::OnSpawnGraphNodeByShortcut(FInputChord InChord, const FVector2D& InPosition,
                                                      UEdGraph* InGraph)
 {
 	if(!InChord.HasAnyModifierKeys())
@@ -745,17 +763,17 @@ FReply FRigVMEditor::OnSpawnGraphNodeByShortcut(FInputChord InChord, const FVect
 	return FReply::Unhandled();
 }
 
-void FRigVMEditor::JumpToHyperlink(const UObject* ObjectReference, bool bRequestRename)
+bool FRigVMEditorBase::JumpToHyperlinkImpl(const UObject* ObjectReference, bool bRequestRename)
 {
 	if(const URigVMEdGraph* Graph = Cast<URigVMEdGraph>(ObjectReference))
 	{
 		OpenGraphAndBringToFront((UEdGraph*)Graph, true);
-		return;
+		return true;
 	}
-	FBlueprintEditor::JumpToHyperlink(ObjectReference, bRequestRename);
+	return false;
 }
 
-void FRigVMEditor::AddNewFunctionVariant(const UEdGraph* InOriginalFunction)
+void FRigVMEditorBase::AddNewFunctionVariant(const UEdGraph* InOriginalFunction)
 {
 #if WITH_EDITOR
 	if (GEditor)
@@ -784,21 +802,19 @@ void FRigVMEditor::AddNewFunctionVariant(const UEdGraph* InOriginalFunction)
 	}
 }
 
-void FRigVMEditor::PostUndo(bool bSuccess)
+void FRigVMEditorBase::PostUndoImpl(bool bSuccess)
 {
 	const FTransaction* Transaction = GEditor->Trans->GetTransaction(GEditor->Trans->GetQueueLength() - GEditor->Trans->GetUndoCount());
-	FBlueprintEditor::PostUndo(bSuccess);
 	PostTransaction(bSuccess, Transaction, false);
 }
 
-void FRigVMEditor::PostRedo(bool bSuccess)
+void FRigVMEditorBase::PostRedoImpl(bool bSuccess)
 {
 	const FTransaction* Transaction = GEditor->Trans->GetTransaction(GEditor->Trans->GetQueueLength() - GEditor->Trans->GetUndoCount() - 1);
-	FBlueprintEditor::PostRedo(bSuccess);
 	PostTransaction(bSuccess, Transaction, true);
 }
 
-void FRigVMEditor::OnStartWatchingPin()
+void FRigVMEditorBase::OnStartWatchingPin()
 {
 	if (UEdGraphPin* Pin = GetCurrentlySelectedPin())
 	{
@@ -806,7 +822,7 @@ void FRigVMEditor::OnStartWatchingPin()
 	}
 }
 
-bool FRigVMEditor::CanStartWatchingPin() const
+bool FRigVMEditorBase::CanStartWatchingPin() const
 {
 	if (UEdGraphPin* Pin = GetCurrentlySelectedPin())
 	{
@@ -819,7 +835,7 @@ bool FRigVMEditor::CanStartWatchingPin() const
 	return false;
 }
 
-void FRigVMEditor::OnStopWatchingPin()
+void FRigVMEditorBase::OnStopWatchingPin()
 {
 	if (UEdGraphPin* Pin = GetCurrentlySelectedPin())
 	{
@@ -827,7 +843,7 @@ void FRigVMEditor::OnStopWatchingPin()
 	}
 }
 
-bool FRigVMEditor::CanStopWatchingPin() const
+bool FRigVMEditorBase::CanStopWatchingPin() const
 {
 	if (UEdGraphPin* Pin = GetCurrentlySelectedPin())
 	{
@@ -839,7 +855,7 @@ bool FRigVMEditor::CanStopWatchingPin() const
 	return false;
 }
 
-void FRigVMEditor::OnToolkitHostingStarted(const TSharedRef<class IToolkit>& Toolkit)
+void FRigVMEditorBase::OnToolkitHostingStarted(const TSharedRef<class IToolkit>& Toolkit)
 {
 	TSharedPtr<SWidget> InlineContent = Toolkit->GetInlineContent();
 	if (InlineContent.IsValid())
@@ -848,17 +864,17 @@ void FRigVMEditor::OnToolkitHostingStarted(const TSharedRef<class IToolkit>& Too
 	}
 }
 
-void FRigVMEditor::OnToolkitHostingFinished(const TSharedRef<class IToolkit>& Toolkit)
+void FRigVMEditorBase::OnToolkitHostingFinished(const TSharedRef<class IToolkit>& Toolkit)
 {
 	Toolbox->SetContent(SNullWidget::NullWidget);
 }
 
-TStatId FRigVMEditor::GetStatId() const
+TStatId FRigVMEditorBase::GetStatId() const
 {
-	RETURN_QUICK_DECLARE_CYCLE_STAT(FRigVMEditor, STATGROUP_Tickables);
+	RETURN_QUICK_DECLARE_CYCLE_STAT(FRigVMEditorBase, STATGROUP_Tickables);
 }
 
-void FRigVMEditor::PasteNodes()
+void FRigVMEditorBase::PasteNodes()
 {
 	DECLARE_SCOPE_HIERARCHICAL_COUNTER_FUNC()
 
@@ -866,7 +882,7 @@ void FRigVMEditor::PasteNodes()
 
 	FVector2D PasteLocation = FSlateApplication::Get().GetCursorPos();
 
-	TSharedPtr<SDockTab> ActiveTab = DocumentManager->GetActiveTab();
+	TSharedPtr<SDockTab> ActiveTab = GetDocumentManager()->GetActiveTab();
 	if (ActiveTab.IsValid())
 	{
 		TSharedPtr<SGraphEditor> GraphEditor = StaticCastSharedRef<SGraphEditor>(ActiveTab->GetContent());
@@ -894,12 +910,12 @@ void FRigVMEditor::PasteNodes()
 
 }
 
-URigVMBlueprint* FRigVMEditor::GetRigVMBlueprint() const
+URigVMBlueprint* FRigVMEditorBase::GetRigVMBlueprint() const
 {
 	return Cast<URigVMBlueprint>(GetBlueprintObj());
 }
 
-URigVMHost* FRigVMEditor::GetRigVMHost() const
+URigVMHost* FRigVMEditorBase::GetRigVMHost() const
 {
 	if (URigVMBlueprint* RigVMBlueprint = GetRigVMBlueprint())
 	{
@@ -911,17 +927,17 @@ URigVMHost* FRigVMEditor::GetRigVMHost() const
 	return nullptr;
 }
 
-UObject* FRigVMEditor::GetOuterForHost() const
+UObject* FRigVMEditorBase::GetOuterForHost() const
 {
 	return GetRigVMBlueprint();
 }
 
-UClass* FRigVMEditor::GetDetailWrapperClass() const
+UClass* FRigVMEditorBase::GetDetailWrapperClass() const
 {
 	return URigVMDetailsViewWrapperObject::StaticClass();
 }
 
-bool FRigVMEditor::SelectLocalVariable(const UEdGraph* Graph, const FName& VariableName)
+bool FRigVMEditorBase::SelectLocalVariable(const UEdGraph* Graph, const FName& VariableName)
 {
 	if (const URigVMEdGraph* RigVMEdGraph = Cast<URigVMEdGraph>(Graph))
 	{
@@ -933,7 +949,7 @@ bool FRigVMEditor::SelectLocalVariable(const UEdGraph* Graph, const FName& Varia
 				{
 					URigVMDetailsViewWrapperObject* WrapperObject = URigVMDetailsViewWrapperObject::MakeInstance(
 						GetDetailWrapperClass(), GetBlueprintObj(), Variable.StaticStruct(), (uint8*)&Variable, RigVMGraph);
-					WrapperObject->GetWrappedPropertyChangedChainEvent().AddSP(this, &FRigVMEditor::OnWrappedPropertyChangedChainEvent);
+					WrapperObject->GetWrappedPropertyChangedChainEvent().AddSP(StaticCastSharedRef<FRigVMEditorBase>(SharedRef()), &FRigVMEditorBase::OnWrappedPropertyChangedChainEvent);
 					WrapperObject->AddToRoot();
 
 					TArray<UObject*> Objects = {WrapperObject};
@@ -946,122 +962,26 @@ bool FRigVMEditor::SelectLocalVariable(const UEdGraph* Graph, const FName& Varia
 	return false;
 }
 
-void FRigVMEditor::CreateDefaultCommands()
+void FRigVMEditorBase::CreateDefaultCommandsImpl()
 {
 	if (GetBlueprintObj())
 	{
-		FBlueprintEditor::CreateDefaultCommands();
+		// FBlueprintEditor::CreateDefaultCommands();
 	}
 	else
 	{
-		ToolkitCommands->MapAction( FGenericCommands::Get().Undo, 
-			FExecuteAction::CreateSP( this, &FRigVMEditor::UndoAction ));
-		ToolkitCommands->MapAction( FGenericCommands::Get().Redo, 
-			FExecuteAction::CreateSP( this, &FRigVMEditor::RedoAction ));
+		GetToolkitCommands()->MapAction( FGenericCommands::Get().Undo, 
+			FExecuteAction::CreateSP( StaticCastSharedRef<FRigVMEditorBase>(SharedRef()), &FRigVMEditorBase::UndoAction ));
+		GetToolkitCommands()->MapAction( FGenericCommands::Get().Redo, 
+			FExecuteAction::CreateSP( StaticCastSharedRef<FRigVMEditorBase>(SharedRef()), &FRigVMEditorBase::RedoAction ));
 	}
 }
 
-void FRigVMEditor::OnCreateGraphEditorCommands(TSharedPtr<FUICommandList> GraphEditorCommandsList)
+void FRigVMEditorBase::OnCreateGraphEditorCommands(TSharedPtr<FUICommandList> GraphEditorCommandsList)
 {
 }
 
-struct FRigVMEditorZoomLevelsContainer : public FZoomLevelsContainer
-{
-	struct FRigVMEditorZoomLevelEntry
-	{
-	public:
-		FRigVMEditorZoomLevelEntry(float InZoomAmount, const FText& InDisplayText, EGraphRenderingLOD::Type InLOD)
-			: DisplayText(FText::Format(NSLOCTEXT("GraphEditor", "Zoom", "Zoom {0}"), InDisplayText))
-		, ZoomAmount(InZoomAmount)
-		, LOD(InLOD)
-		{
-		}
-
-	public:
-		FText DisplayText;
-		float ZoomAmount;
-		EGraphRenderingLOD::Type LOD;
-	};
-	
-	FRigVMEditorZoomLevelsContainer()
-	{
-		ZoomLevels.Reserve(22);
-		ZoomLevels.Add(FRigVMEditorZoomLevelEntry(0.025f, FText::FromString(TEXT("-14")), EGraphRenderingLOD::LowestDetail));
-		ZoomLevels.Add(FRigVMEditorZoomLevelEntry(0.070f, FText::FromString(TEXT("-13")), EGraphRenderingLOD::LowestDetail));
-		ZoomLevels.Add(FRigVMEditorZoomLevelEntry(0.100f, FText::FromString(TEXT("-12")), EGraphRenderingLOD::LowestDetail));
-		ZoomLevels.Add(FRigVMEditorZoomLevelEntry(0.125f, FText::FromString(TEXT("-11")), EGraphRenderingLOD::LowestDetail));
-		ZoomLevels.Add(FRigVMEditorZoomLevelEntry(0.150f, FText::FromString(TEXT("-10")), EGraphRenderingLOD::LowestDetail));
-		ZoomLevels.Add(FRigVMEditorZoomLevelEntry(0.175f, FText::FromString(TEXT("-9")), EGraphRenderingLOD::LowestDetail));
-		ZoomLevels.Add(FRigVMEditorZoomLevelEntry(0.200f, FText::FromString(TEXT("-8")), EGraphRenderingLOD::LowestDetail));
-		ZoomLevels.Add(FRigVMEditorZoomLevelEntry(0.225f, FText::FromString(TEXT("-7")), EGraphRenderingLOD::LowDetail));
-		ZoomLevels.Add(FRigVMEditorZoomLevelEntry(0.250f, FText::FromString(TEXT("-6")), EGraphRenderingLOD::LowDetail));
-		ZoomLevels.Add(FRigVMEditorZoomLevelEntry(0.375f, FText::FromString(TEXT("-5")), EGraphRenderingLOD::MediumDetail));
-		ZoomLevels.Add(FRigVMEditorZoomLevelEntry(0.500f, FText::FromString(TEXT("-4")), EGraphRenderingLOD::MediumDetail));
-		ZoomLevels.Add(FRigVMEditorZoomLevelEntry(0.675f, FText::FromString(TEXT("-3")), EGraphRenderingLOD::MediumDetail));
-		ZoomLevels.Add(FRigVMEditorZoomLevelEntry(0.750f, FText::FromString(TEXT("-2")), EGraphRenderingLOD::DefaultDetail));
-		ZoomLevels.Add(FRigVMEditorZoomLevelEntry(0.875f, FText::FromString(TEXT("-1")), EGraphRenderingLOD::DefaultDetail));
-		ZoomLevels.Add(FRigVMEditorZoomLevelEntry(1.000f, FText::FromString(TEXT("1:1")), EGraphRenderingLOD::DefaultDetail));
-		ZoomLevels.Add(FRigVMEditorZoomLevelEntry(1.250f, FText::FromString(TEXT("+1")), EGraphRenderingLOD::DefaultDetail));
-		ZoomLevels.Add(FRigVMEditorZoomLevelEntry(1.375f, FText::FromString(TEXT("+2")), EGraphRenderingLOD::DefaultDetail));
-		ZoomLevels.Add(FRigVMEditorZoomLevelEntry(1.500f, FText::FromString(TEXT("+3")), EGraphRenderingLOD::FullyZoomedIn));
-		ZoomLevels.Add(FRigVMEditorZoomLevelEntry(1.675f, FText::FromString(TEXT("+4")), EGraphRenderingLOD::FullyZoomedIn));
-		ZoomLevels.Add(FRigVMEditorZoomLevelEntry(1.750f, FText::FromString(TEXT("+5")), EGraphRenderingLOD::FullyZoomedIn));
-		ZoomLevels.Add(FRigVMEditorZoomLevelEntry(1.875f, FText::FromString(TEXT("+6")), EGraphRenderingLOD::FullyZoomedIn));
-		ZoomLevels.Add(FRigVMEditorZoomLevelEntry(2.000f, FText::FromString(TEXT("+7")), EGraphRenderingLOD::FullyZoomedIn));
-	}
-
-	float GetZoomAmount(int32 InZoomLevel) const override
-	{
-		checkSlow(ZoomLevels.IsValidIndex(InZoomLevel));
-		return ZoomLevels[InZoomLevel].ZoomAmount;
-	}
-
-	int32 GetNearestZoomLevel(float InZoomAmount) const override
-	{
-		for (int32 ZoomLevelIndex=0; ZoomLevelIndex < GetNumZoomLevels(); ++ZoomLevelIndex)
-		{
-			if (InZoomAmount <= GetZoomAmount(ZoomLevelIndex))
-			{
-				return ZoomLevelIndex;
-			}
-		}
-
-		return GetDefaultZoomLevel();
-	}
-	
-	FText GetZoomText(int32 InZoomLevel) const override
-	{
-		checkSlow(ZoomLevels.IsValidIndex(InZoomLevel));
-		return ZoomLevels[InZoomLevel].DisplayText;
-	}
-	
-	int32 GetNumZoomLevels() const override
-	{
-		return ZoomLevels.Num();
-	}
-	
-	int32 GetDefaultZoomLevel() const override
-	{
-		return 14;
-	}
-
-	EGraphRenderingLOD::Type GetLOD(int32 InZoomLevel) const override
-	{
-		checkSlow(ZoomLevels.IsValidIndex(InZoomLevel));
-		return ZoomLevels[InZoomLevel].LOD;
-	}
-
-	TArray<FRigVMEditorZoomLevelEntry> ZoomLevels;
-};
-
-TSharedRef<SGraphEditor> FRigVMEditor::CreateGraphEditorWidget(TSharedRef<FTabInfo> InTabInfo, UEdGraph* InGraph)
-{
-	TSharedRef<SGraphEditor> GraphEditor = FBlueprintEditor::CreateGraphEditorWidget(InTabInfo, InGraph);
-	GraphEditor->GetGraphPanel()->SetZoomLevelsContainer<FRigVMEditorZoomLevelsContainer>();
-	return GraphEditor;
-}
-
-void FRigVMEditor::Compile()
+void FRigVMEditorBase::Compile()
 {
 	{
 		DECLARE_SCOPE_HIERARCHICAL_COUNTER_FUNC()
@@ -1093,7 +1013,7 @@ void FRigVMEditor::Compile()
 		SetHost(nullptr);
 		{
 			TGuardValue<bool> GuardCompileReEntry(bIsCompilingThroughUI, true);
-			FBlueprintEditor::Compile();
+			CompileImpl();
 			RigVMBlueprint->InitializeArchetypeInstances();
 			UpdateRigVMHost();
 		}
@@ -1121,10 +1041,10 @@ void FRigVMEditor::Compile()
 			}
 		}
 
-		TArray<FCustomDebugObject> DebugList;
+		TArray<FRigVMCustomDebugObject> DebugList;
 		GetCustomDebugObjects(DebugList);
 
-		for (const FCustomDebugObject& DebugObject : DebugList)
+		for (const FRigVMCustomDebugObject& DebugObject : DebugList)
 		{
 			if (DebugObject.NameOverride == LastDebuggedObjectName)
 			{
@@ -1161,7 +1081,7 @@ void FRigVMEditor::Compile()
 		if (bAnyVariableValueChanged)
 		{
 			// Go over all the instances to update the default values from CDO
-			for (const FCustomDebugObject& DebugObject : DebugList)
+			for (const FRigVMCustomDebugObject& DebugObject : DebugList)
 			{
 				if (URigVMHost* DebuggedHost = Cast<URigVMHost>(DebugObject.Object))
 				{
@@ -1179,23 +1099,23 @@ void FRigVMEditor::Compile()
 	// FStatsHierarchical::DumpMeasurements(LogForMeasurements);
 }
 
-void FRigVMEditor::SaveAsset_Execute()
+void FRigVMEditorBase::SaveAsset_Execute()
 {
 	LastDebuggedHost = GetCustomDebugObjectLabel(GetBlueprintObj()->GetObjectBeingDebugged());
-	FBlueprintEditor::SaveAsset_Execute();
+	SaveAsset_Execute_Impl();
 
 	UpdateRigVMHost();
 }
 
-void FRigVMEditor::SaveAssetAs_Execute()
+void FRigVMEditorBase::SaveAssetAs_Execute()
 {
 	LastDebuggedHost = GetCustomDebugObjectLabel(GetBlueprintObj()->GetObjectBeingDebugged());
-	FBlueprintEditor::SaveAssetAs_Execute();
+	SaveAssetAs_Execute_Impl();
 
 	UpdateRigVMHost();
 }
 
-bool FRigVMEditor::IsEditable(UEdGraph* InGraph) const
+bool FRigVMEditorBase::IsEditable(UEdGraph* InGraph) const
 {
 	if(!IsGraphInCurrentBlueprint(InGraph))
 	{
@@ -1223,20 +1143,20 @@ bool FRigVMEditor::IsEditable(UEdGraph* InGraph) const
 		}
 	}
 
-	return FBlueprintEditor::IsEditable(InGraph);
+	return IsEditableImpl(InGraph);
 }
 
-bool FRigVMEditor::IsCompilingEnabled() const
+bool FRigVMEditorBase::IsCompilingEnabled() const
 {
 	return true;
 }
 
-FText FRigVMEditor::GetGraphDecorationString(UEdGraph* InGraph) const
+FText FRigVMEditorBase::GetGraphDecorationString(UEdGraph* InGraph) const
 {
 	return FText::GetEmpty();
 }
 
-void FRigVMEditor::OnSelectedNodesChangedImpl(const TSet<UObject*>& NewSelection)
+void FRigVMEditorBase::OnSelectedNodesChangedImpl(const TSet<UObject*>& NewSelection)
 {
 	DECLARE_SCOPE_HIERARCHICAL_COUNTER_FUNC()
 
@@ -1272,7 +1192,7 @@ void FRigVMEditor::OnSelectedNodesChangedImpl(const TSet<UObject*>& NewSelection
 	}
 }
 
-void FRigVMEditor::OnBlueprintChangedImpl(UBlueprint* InBlueprint, bool bIsJustBeingCompiled)
+void FRigVMEditorBase::OnBlueprintChangedImpl(UBlueprint* InBlueprint, bool bIsJustBeingCompiled)
 {
 	DECLARE_SCOPE_HIERARCHICAL_COUNTER_FUNC()
 
@@ -1281,7 +1201,7 @@ void FRigVMEditor::OnBlueprintChangedImpl(UBlueprint* InBlueprint, bool bIsJustB
 		return;
 	}
 
-	FBlueprintEditor::OnBlueprintChangedImpl(InBlueprint, bIsJustBeingCompiled);
+	OnBlueprintChangedInnerImpl(InBlueprint, bIsJustBeingCompiled);
 
 	if(InBlueprint == GetBlueprintObj())
 	{
@@ -1291,10 +1211,10 @@ void FRigVMEditor::OnBlueprintChangedImpl(UBlueprint* InBlueprint, bool bIsJustB
 
 			if (!LastDebuggedHost.IsEmpty())
 			{
-				TArray<FCustomDebugObject> DebugList;
+				TArray<FRigVMCustomDebugObject> DebugList;
 				GetCustomDebugObjects(DebugList);
 
-				for (const FCustomDebugObject& DebugObject : DebugList)
+				for (const FRigVMCustomDebugObject& DebugObject : DebugList)
 				{
 					if (DebugObject.NameOverride == LastDebuggedHost)
 					{
@@ -1307,37 +1227,37 @@ void FRigVMEditor::OnBlueprintChangedImpl(UBlueprint* InBlueprint, bool bIsJustB
 		}
 	}}
 
-void FRigVMEditor::RefreshEditors(ERefreshBlueprintEditorReason::Type Reason)
+void FRigVMEditorBase::RefreshEditors(ERefreshRigVMEditorReason::Type Reason)
 {
-	if(Reason == ERefreshBlueprintEditorReason::UnknownReason)
+	if(Reason == ERefreshRigVMEditorReason::UnknownReason)
 	{
 		// we mark the reason as just compiled since we don't want to
 		// update the graph(s) all the time during compilation
-		Reason = ERefreshBlueprintEditorReason::BlueprintCompiled;
+		Reason = ERefreshRigVMEditorReason::BlueprintCompiled;
 	}
-	FBlueprintEditor::RefreshEditors(Reason);
+	RefreshEditorsImpl(Reason);
 	if (GraphExplorerWidget.IsValid())
 	{
 		GraphExplorerWidget->Refresh();
 	}
 }
 
-void FRigVMEditor::SetupGraphEditorEvents(UEdGraph* InGraph, SGraphEditor::FGraphEditorEvents& InEvents)
+void FRigVMEditorBase::SetupGraphEditorEvents(UEdGraph* InGraph, SGraphEditor::FGraphEditorEvents& InEvents)
 {
 	DECLARE_SCOPE_HIERARCHICAL_COUNTER_FUNC()
 
-	FBlueprintEditor::SetupGraphEditorEvents(InGraph, InEvents);
+	SetupGraphEditorEventsImpl(InGraph, InEvents);
 
-	InEvents.OnCreateActionMenu = SGraphEditor::FOnCreateActionMenu::CreateSP(this, &FRigVMEditor::HandleCreateGraphActionMenu);
-	InEvents.OnTextCommitted = FOnNodeTextCommitted::CreateSP(this, &FRigVMEditor::OnNodeTitleCommitted);
+	InEvents.OnCreateActionMenu = SGraphEditor::FOnCreateActionMenu::CreateSP(StaticCastSharedRef<FRigVMEditorBase>(SharedRef()), &FRigVMEditorBase::HandleCreateGraphActionMenu);
+	InEvents.OnTextCommitted = FOnNodeTextCommitted::CreateSP(StaticCastSharedRef<FRigVMEditorBase>(SharedRef()), &FRigVMEditorBase::OnNodeTitleCommitted);
 }
 
-FActionMenuContent FRigVMEditor::HandleCreateGraphActionMenu(UEdGraph* InGraph, const FVector2D& InNodePosition, const TArray<UEdGraphPin*>& InDraggedPins, bool bAutoExpand, SGraphEditor::FActionMenuClosed InOnMenuClosed)
+FActionMenuContent FRigVMEditorBase::HandleCreateGraphActionMenu(UEdGraph* InGraph, const FVector2D& InNodePosition, const TArray<UEdGraphPin*>& InDraggedPins, bool bAutoExpand, SGraphEditor::FActionMenuClosed InOnMenuClosed)
 {
-	return FBlueprintEditor::OnCreateGraphActionMenu(InGraph, InNodePosition, InDraggedPins, bAutoExpand, InOnMenuClosed);
+	return OnCreateGraphActionMenu(InGraph, InNodePosition, InDraggedPins, bAutoExpand, InOnMenuClosed);
 }
 
-void FRigVMEditor::OnNodeTitleCommitted(const FText& NewText, ETextCommit::Type CommitInfo, UEdGraphNode* NodeBeingChanged)
+void FRigVMEditorBase::OnNodeTitleCommitted(const FText& NewText, ETextCommit::Type CommitInfo, UEdGraphNode* NodeBeingChanged)
 {
 	DECLARE_SCOPE_HIERARCHICAL_COUNTER_FUNC()
 
@@ -1350,96 +1270,98 @@ void FRigVMEditor::OnNodeTitleCommitted(const FText& NewText, ETextCommit::Type
 	}
 }
 
-void FRigVMEditor::FocusInspectorOnGraphSelection(const TSet<UObject*>& NewSelection, bool bForceRefresh)
+void FRigVMEditorBase::FocusInspectorOnGraphSelection(const TSet<UObject*>& NewSelection, bool bForceRefresh)
 {
 	// nothing to do here
 }
 
-void FRigVMEditor::AddReferencedObjects(FReferenceCollector& Collector)
+void FRigVMEditorBase::AddReferencedObjects(FReferenceCollector& Collector)
 {
-	FBlueprintEditor::AddReferencedObjects(Collector);
+	AddReferencedObjectsImpl(Collector);
 
 	TWeakObjectPtr<URigVMHost> RigVMHost(GetRigVMHost());
 	Collector.AddReferencedObject(RigVMHost);
 }
 
-void FRigVMEditor::BindCommands()
+void FRigVMEditorBase::BindCommands()
 {
 	DECLARE_SCOPE_HIERARCHICAL_COUNTER_FUNC()
 
+	TSharedRef<FRigVMEditorBase> ThisRef = StaticCastSharedRef<FRigVMEditorBase>(SharedRef());
+
 	GetToolkitCommands()->MapAction(
 		FRigVMEditorCommands::Get().AutoCompileGraph,
-		FExecuteAction::CreateSP(this, &FRigVMEditor::ToggleAutoCompileGraph), 
-		FIsActionChecked::CreateSP(this, &FRigVMEditor::CanAutoCompileGraph),
-		FIsActionChecked::CreateSP(this, &FRigVMEditor::IsAutoCompileGraphOn));
+		FExecuteAction::CreateSP(ThisRef, &FRigVMEditorBase::ToggleAutoCompileGraph), 
+		FIsActionChecked::CreateSP(ThisRef, &FRigVMEditorBase::CanAutoCompileGraph),
+		FIsActionChecked::CreateSP(ThisRef, &FRigVMEditorBase::IsAutoCompileGraphOn));
 
 	GetToolkitCommands()->MapAction(
 		FRigVMEditorCommands::Get().ToggleEventQueue,
-		FExecuteAction::CreateSP(this, &FRigVMEditor::ToggleEventQueue),
+		FExecuteAction::CreateSP(ThisRef, &FRigVMEditorBase::ToggleEventQueue),
 		FCanExecuteAction());
 
 	GetToolkitCommands()->MapAction(
 		FRigVMEditorCommands::Get().ToggleExecutionMode,
-		FExecuteAction::CreateSP(this, &FRigVMEditor::ToggleExecutionMode),
+		FExecuteAction::CreateSP(ThisRef, &FRigVMEditorBase::ToggleExecutionMode),
 		FCanExecuteAction());
 
 	GetToolkitCommands()->MapAction(
 		FRigVMEditorCommands::Get().ReleaseMode,
-		FExecuteAction::CreateSP(this, &FRigVMEditor::SetExecutionMode, ERigVMEditorExecutionModeType_Release),
+		FExecuteAction::CreateSP(ThisRef, &FRigVMEditorBase::SetExecutionMode, ERigVMEditorExecutionModeType_Release),
 		FCanExecuteAction());
 
 	GetToolkitCommands()->MapAction(
 		FRigVMEditorCommands::Get().DebugMode,
-		FExecuteAction::CreateSP(this, &FRigVMEditor::SetExecutionMode, ERigVMEditorExecutionModeType_Debug),
+		FExecuteAction::CreateSP(ThisRef, &FRigVMEditorBase::SetExecutionMode, ERigVMEditorExecutionModeType_Debug),
 		FCanExecuteAction());
 
 	GetToolkitCommands()->MapAction(
 		FRigVMEditorCommands::Get().ResumeExecution,
-		FExecuteAction::CreateSP(this, &FRigVMEditor::HandleBreakpointActionRequested, ERigVMBreakpointAction::Resume),			
-		FIsActionChecked::CreateSP(this, &FRigVMEditor::IsHaltedAtBreakpoint));
+		FExecuteAction::CreateSP(ThisRef, &FRigVMEditorBase::HandleBreakpointActionRequested, ERigVMBreakpointAction::Resume),			
+		FIsActionChecked::CreateSP(ThisRef, &FRigVMEditorBase::IsHaltedAtBreakpoint));
 
 	GetToolkitCommands()->MapAction(
 		FRigVMEditorCommands::Get().ShowCurrentStatement,
-		FExecuteAction::CreateSP(this, &FRigVMEditor::HandleShowCurrentStatement),
-		FIsActionChecked::CreateSP(this, &FRigVMEditor::IsHaltedAtBreakpoint));
+		FExecuteAction::CreateSP(ThisRef, &FRigVMEditorBase::HandleShowCurrentStatement),
+		FIsActionChecked::CreateSP(ThisRef, &FRigVMEditorBase::IsHaltedAtBreakpoint));
 
 	GetToolkitCommands()->MapAction(
 		FRigVMEditorCommands::Get().StepOver,
-		FExecuteAction::CreateSP(this, &FRigVMEditor::HandleBreakpointActionRequested, ERigVMBreakpointAction::StepOver),
-		FIsActionChecked::CreateSP(this, &FRigVMEditor::IsHaltedAtBreakpoint));
+		FExecuteAction::CreateSP(ThisRef, &FRigVMEditorBase::HandleBreakpointActionRequested, ERigVMBreakpointAction::StepOver),
+		FIsActionChecked::CreateSP(ThisRef, &FRigVMEditorBase::IsHaltedAtBreakpoint));
 
 	GetToolkitCommands()->MapAction(
 		FRigVMEditorCommands::Get().StepInto,
-		FExecuteAction::CreateSP(this, &FRigVMEditor::HandleBreakpointActionRequested, ERigVMBreakpointAction::StepInto),
-		FIsActionChecked::CreateSP(this, &FRigVMEditor::IsHaltedAtBreakpoint));
+		FExecuteAction::CreateSP(ThisRef, &FRigVMEditorBase::HandleBreakpointActionRequested, ERigVMBreakpointAction::StepInto),
+		FIsActionChecked::CreateSP(ThisRef, &FRigVMEditorBase::IsHaltedAtBreakpoint));
 
 	GetToolkitCommands()->MapAction(
 		FRigVMEditorCommands::Get().StepOut,
-		FExecuteAction::CreateSP(this, &FRigVMEditor::HandleBreakpointActionRequested, ERigVMBreakpointAction::StepOut),
-		FIsActionChecked::CreateSP(this, &FRigVMEditor::IsHaltedAtBreakpoint));
+		FExecuteAction::CreateSP(ThisRef, &FRigVMEditorBase::HandleBreakpointActionRequested, ERigVMBreakpointAction::StepOut),
+		FIsActionChecked::CreateSP(ThisRef, &FRigVMEditorBase::IsHaltedAtBreakpoint));
 
 	GetToolkitCommands()->MapAction(
 		FRigVMEditorCommands::Get().FrameSelection,
-		FExecuteAction::CreateSP(this, &FRigVMEditor::FrameSelection),
+		FExecuteAction::CreateSP(ThisRef, &FRigVMEditorBase::FrameSelection),
 		FCanExecuteAction());
 
 	GetToolkitCommands()->MapAction(
 		FRigVMEditorCommands::Get().SwapFunctionWithinAsset,
-		FExecuteAction::CreateSP(this, &FRigVMEditor::SwapFunctionWithinAsset),
+		FExecuteAction::CreateSP(ThisRef, &FRigVMEditorBase::SwapFunctionWithinAsset),
 		FCanExecuteAction());
 
 	GetToolkitCommands()->MapAction(
 		FRigVMEditorCommands::Get().SwapFunctionAcrossProject,
-		FExecuteAction::CreateSP(this, &FRigVMEditor::SwapFunctionAcrossProject),
+		FExecuteAction::CreateSP(ThisRef, &FRigVMEditorBase::SwapFunctionAcrossProject),
 		FCanExecuteAction());
 
 	GetToolkitCommands()->MapAction(
 		FRigVMEditorCommands::Get().SwapAssetReferences,
-		FExecuteAction::CreateSP(this, &FRigVMEditor::SwapAssetReferences),
+		FExecuteAction::CreateSP(ThisRef, &FRigVMEditorBase::SwapAssetReferences),
 		FCanExecuteAction());
 }
 
-void FRigVMEditor::ToggleAutoCompileGraph()
+void FRigVMEditorBase::ToggleAutoCompileGraph()
 {
 	if (URigVMBlueprint* RigVMBlueprint = GetRigVMBlueprint())
 	{
@@ -1451,7 +1373,7 @@ void FRigVMEditor::ToggleAutoCompileGraph()
 	}
 }
 
-bool FRigVMEditor::IsAutoCompileGraphOn() const
+bool FRigVMEditorBase::IsAutoCompileGraphOn() const
 {
 	if (URigVMBlueprint* RigVMBlueprint = GetRigVMBlueprint())
 	{
@@ -1460,30 +1382,30 @@ bool FRigVMEditor::IsAutoCompileGraphOn() const
 	return false;
 }
 
-void FRigVMEditor::ToggleEventQueue()
+void FRigVMEditorBase::ToggleEventQueue()
 {
 	SetEventQueue(LastEventQueue);
 }
 
-void FRigVMEditor::ToggleExecutionMode()
+void FRigVMEditorBase::ToggleExecutionMode()
 {
 	SetExecutionMode((ExecutionMode == ERigVMEditorExecutionModeType_Debug) ?
 		ERigVMEditorExecutionModeType_Release
 		: ERigVMEditorExecutionModeType_Debug);
 }
 
-TSharedRef<SWidget> FRigVMEditor::GenerateEventQueueMenuContent()
+TSharedRef<SWidget> FRigVMEditorBase::GenerateEventQueueMenuContent()
 {
 	FMenuBuilder MenuBuilder(true, GetToolkitCommands());
 	GenerateEventQueueMenuContent(MenuBuilder);
 	return MenuBuilder.MakeWidget();
 }
 
-void FRigVMEditor::GenerateEventQueueMenuContent(FMenuBuilder& MenuBuilder)
+void FRigVMEditorBase::GenerateEventQueueMenuContent(FMenuBuilder& MenuBuilder)
 {
 }
 
-TSharedRef<SWidget> FRigVMEditor::GenerateExecutionModeMenuContent()
+TSharedRef<SWidget> FRigVMEditorBase::GenerateExecutionModeMenuContent()
 {
 	FMenuBuilder MenuBuilder(true, GetToolkitCommands());
 	MenuBuilder.BeginSection(TEXT("Events"));
@@ -1493,7 +1415,7 @@ TSharedRef<SWidget> FRigVMEditor::GenerateExecutionModeMenuContent()
 	return MenuBuilder.MakeWidget();
 }
 
-FMenuBuilder FRigVMEditor::GenerateBulkEditMenu()
+FMenuBuilder FRigVMEditorBase::GenerateBulkEditMenu()
 {
 	FMenuBuilder MenuBuilder(true, GetToolkitCommands());
 	MenuBuilder.BeginSection(TEXT("Functions"), LOCTEXT("Functions", "Functions"));
@@ -1506,38 +1428,41 @@ FMenuBuilder FRigVMEditor::GenerateBulkEditMenu()
 	return MenuBuilder;
 }
 
-TSharedRef<SWidget> FRigVMEditor::GenerateBulkEditMenuContent()
+TSharedRef<SWidget> FRigVMEditorBase::GenerateBulkEditMenuContent()
 {
 	FMenuBuilder MenuBuilder = GenerateBulkEditMenu();
 	return MenuBuilder.MakeWidget();
 }
 
-void FRigVMEditor::OnActiveTabChanged( TSharedPtr<SDockTab> PreviouslyActive, TSharedPtr<SDockTab> NewlyActivated )
+void FRigVMEditorBase::OnActiveTabChanged( TSharedPtr<SDockTab> PreviouslyActive, TSharedPtr<SDockTab> NewlyActivated )
 {
 	if (!NewlyActivated.IsValid())
 	{
 		TArray<UObject*> ObjArray;
-		Inspector->ShowDetailsForObjects(ObjArray);
-	}
-	else 
-	{
-		FBlueprintEditor::OnActiveTabChanged(PreviouslyActive, NewlyActivated);
+#if WITH_RIGVMLEGACYEDITOR
+		if (TSharedPtr<SKismetInspector> Inspector = GetInspector())
+#else
+		if (TSharedPtr<SRigVMDetailsInspector> Inspector = GetInspector())
+#endif
+		{
+			Inspector->ShowDetailsForObjects(ObjArray);
+		}
 	}
 }
 
-void FRigVMEditor::UndoAction()
+void FRigVMEditorBase::UndoAction()
 {
 	GEditor->UndoTransaction();
 }
 
-void FRigVMEditor::RedoAction()
+void FRigVMEditorBase::RedoAction()
 {
 	GEditor->RedoTransaction();
 }
 
-void FRigVMEditor::NewDocument_OnClicked(ECreatedDocumentType GraphType)
+void FRigVMEditorBase::NewDocument_OnClicked(ECreatedDocumentType GraphType)
 {
-	if (GraphType == FBlueprintEditor::CGT_NewFunctionGraph)
+	if (GraphType == FRigVMNewEditor::CGT_NewFunctionGraph)
 	{
 		if (URigVMBlueprint* RigVMBlueprint = GetRigVMBlueprint())
 		{
@@ -1554,7 +1479,7 @@ void FRigVMEditor::NewDocument_OnClicked(ECreatedDocumentType GraphType)
 			}
 		}
 	}
-	else if(GraphType == FBlueprintEditor::CGT_NewEventGraph)
+	else if(GraphType == FRigVMNewEditor::CGT_NewEventGraph)
 	{
 		if (URigVMBlueprint* RigVMBlueprint = GetRigVMBlueprint())
 		{
@@ -1573,21 +1498,21 @@ void FRigVMEditor::NewDocument_OnClicked(ECreatedDocumentType GraphType)
 	}
 }
 
-bool FRigVMEditor::IsSectionVisible(NodeSectionID::Type InSectionID) const
+bool FRigVMEditorBase::IsSectionVisibleImpl(RigVMNodeSectionID::Type InSectionID) const
 {
 	switch (InSectionID)
 	{
-		case NodeSectionID::GRAPH:
-		case NodeSectionID::VARIABLE:
-		case NodeSectionID::FUNCTION:
+		case RigVMNodeSectionID::GRAPH:
+		case RigVMNodeSectionID::VARIABLE:
+		case RigVMNodeSectionID::FUNCTION:
 		{
 			return true;
 		}
-		case NodeSectionID::LOCAL_VARIABLE:
+		case RigVMNodeSectionID::LOCAL_VARIABLE:
 		{
 			if(const URigVMGraph* Graph = GetFocusedModel())
 			{
-				const URigVMGraph* ParentGraph = Graph->GetParentGraph();				
+				const URigVMGraph* ParentGraph = Graph->GetParentGraph();
 				if (ParentGraph && ParentGraph->IsA<URigVMFunctionLibrary>())
 				{
 					return true;
@@ -1602,34 +1527,34 @@ bool FRigVMEditor::IsSectionVisible(NodeSectionID::Type InSectionID) const
 	return false;
 }
 
-bool FRigVMEditor::AreEventGraphsAllowed() const
+bool FRigVMEditorBase::AreEventGraphsAllowed() const
 {
 	if(const URigVMBlueprint* RigVMBlueprint = GetRigVMBlueprint())
 	{
 		return RigVMBlueprint->SupportsEventGraphs();
 	}
-	return FBlueprintEditor::AreEventGraphsAllowed();
+	return false;
 }
 
-bool FRigVMEditor::AreMacrosAllowed() const
+bool FRigVMEditorBase::AreMacrosAllowed() const
 {
 	if(const URigVMBlueprint* RigVMBlueprint = GetRigVMBlueprint())
 	{
 		return RigVMBlueprint->SupportsMacros();
 	}
-	return FBlueprintEditor::AreMacrosAllowed();
+	return false;
 }
 
-bool FRigVMEditor::AreDelegatesAllowed() const
+bool FRigVMEditorBase::AreDelegatesAllowed() const
 {
 	if(const URigVMBlueprint* RigVMBlueprint = GetRigVMBlueprint())
 	{
 		return RigVMBlueprint->SupportsDelegates();
 	}
-	return FBlueprintEditor::AreDelegatesAllowed();
+	return false;
 }
 
-bool FRigVMEditor::NewDocument_IsVisibleForType(ECreatedDocumentType GraphType) const
+bool FRigVMEditorBase::NewDocument_IsVisibleForTypeImpl(ECreatedDocumentType GraphType) const
 {
 	switch(GraphType)
 	{
@@ -1643,12 +1568,12 @@ bool FRigVMEditor::NewDocument_IsVisibleForType(ECreatedDocumentType GraphType)
 			break;
 		}
 	}
-	return FBlueprintEditor::NewDocument_IsVisibleForType(GraphType);
+	return true;
 }
 
-FGraphAppearanceInfo FRigVMEditor::GetGraphAppearance(UEdGraph* InGraph) const
+FGraphAppearanceInfo FRigVMEditorBase::GetGraphAppearance(UEdGraph* InGraph) const
 {
-	FGraphAppearanceInfo AppearanceInfo = FBlueprintEditor::GetGraphAppearance(InGraph);
+	FGraphAppearanceInfo AppearanceInfo = GetGraphAppearanceImpl(InGraph);
 
 	if (const URigVMBlueprint* RigVMBlueprint = GetRigVMBlueprint())
 	{
@@ -1678,7 +1603,7 @@ FGraphAppearanceInfo FRigVMEditor::GetGraphAppearance(UEdGraph* InGraph) const
 	return AppearanceInfo;
 }
 
-void FRigVMEditor::HandleModifiedEvent(ERigVMGraphNotifType InNotifType, URigVMGraph* InGraph, UObject* InSubject)
+void FRigVMEditorBase::HandleModifiedEvent(ERigVMGraphNotifType InNotifType, URigVMGraph* InGraph, UObject* InSubject)
 {
 	DECLARE_SCOPE_HIERARCHICAL_COUNTER_FUNC()
 
@@ -1757,7 +1682,7 @@ void FRigVMEditor::HandleModifiedEvent(ERigVMGraphNotifType InNotifType, URigVMG
 										// we are ok with not reacting to errors here
 										if(Property && PropertyStorage)
 										{
-											FRigVMPinDefaultValueImportErrorContext ErrorPipe;										
+											FRigVMPinDefaultValueImportErrorContext ErrorPipe;
 											Property->ImportText_Direct(*DefaultValue, PropertyStorage, nullptr, PPF_None, &ErrorPipe);
 										}
 									}
@@ -1842,12 +1767,12 @@ public:
 	ERigVMMemoryType MemoryType;
 };
 
-void FRigVMEditor::HandleVMCompiledEvent(UObject* InCompiledObject, URigVM* InVM, FRigVMExtendedExecuteContext& InContext)
+void FRigVMEditorBase::HandleVMCompiledEvent(UObject* InCompiledObject, URigVM* InVM, FRigVMExtendedExecuteContext& InContext)
 {
 	if(URigVMBlueprint* RigVMBlueprint = Cast<URigVMBlueprint>(InCompiledObject))
 	{
-		CompilerResultsListing->ClearMessages();
-		CompilerResultsListing->AddMessages(RigVMBlueprint->GetCompileLog().Messages);
+		GetCompilerResultsListing()->ClearMessages();
+		GetCompilerResultsListing()->AddMessages(RigVMBlueprint->GetCompileLog().Messages);
 		RigVMBlueprint->GetCompileLog().Messages.Reset();
 		RigVMBlueprint->GetCompileLog().NumErrors = RigVMBlueprint->GetCompileLog().NumWarnings = 0;
 	}
@@ -1886,7 +1811,7 @@ void FRigVMEditor::HandleVMCompiledEvent(UObject* InCompiledObject, URigVM* InVM
 	UpdateGraphCompilerErrors();
 }
 
-void FRigVMEditor::HandleVMExecutedEvent(URigVMHost* InHost, const FName& InEventName)
+void FRigVMEditorBase::HandleVMExecutedEvent(URigVMHost* InHost, const FName& InEventName)
 {
 	if (URigVMBlueprint* RigVMBlueprint = GetRigVMBlueprint())
 	{
@@ -1999,7 +1924,7 @@ void FRigVMEditor::HandleVMExecutedEvent(URigVMHost* InHost, const FName& InEven
 	UpdateGraphCompilerErrors();
 }
 
-void FRigVMEditor::HandleVMExecutionHalted(const int32 InstructionIndex, UObject* InNodeObject, const FName& InEntryName)
+void FRigVMEditorBase::HandleVMExecutionHalted(const int32 InstructionIndex, UObject* InNodeObject, const FName& InEntryName)
 {
 	if (HaltedAtNode == InNodeObject)
 	{
@@ -2029,14 +1954,21 @@ void FRigVMEditor::HandleVMExecutionHalted(const int32 InstructionIndex, UObject
 	}
 	else 
 	{
-		if (InEntryName == GetRigVMHost()->GetEventQueue().Last())
+		if (URigVMHost* RigVMHost = GetRigVMHost())
 		{
-			SetHaltedNode(nullptr);
+			const TArray<FName>& EventQueue = RigVMHost->GetEventQueue();
+			if (!EventQueue.IsEmpty())
+			{
+				if (InEntryName == EventQueue.Last())
+				{
+					SetHaltedNode(nullptr);
+				}
+			}
 		}
 	}
 }
 
-void FRigVMEditor::SetHaltedNode(URigVMNode* Node)
+void FRigVMEditorBase::SetHaltedNode(URigVMNode* Node)
 {
 	if (HaltedAtNode)
 	{
@@ -2049,9 +1981,9 @@ void FRigVMEditor::SetHaltedNode(URigVMNode* Node)
 	}
 }
 
-void FRigVMEditor::NotifyPreChange(FProperty* PropertyAboutToChange)
+void FRigVMEditorBase::NotifyPreChange(FProperty* PropertyAboutToChange)
 {
-	FBlueprintEditor::NotifyPreChange(PropertyAboutToChange);
+	NotifyPreChangeImpl(PropertyAboutToChange);
 
 	if (URigVMBlueprint* RigVMBlueprint = GetRigVMBlueprint())
 	{
@@ -2059,9 +1991,9 @@ void FRigVMEditor::NotifyPreChange(FProperty* PropertyAboutToChange)
 	}
 }
 
-void FRigVMEditor::NotifyPostChange(const FPropertyChangedEvent& PropertyChangedEvent, FProperty* PropertyThatChanged)
+void FRigVMEditorBase::NotifyPostChange(const FPropertyChangedEvent& PropertyChangedEvent, FProperty* PropertyThatChanged)
 {
-	FBlueprintEditor::NotifyPostChange(PropertyChangedEvent, PropertyThatChanged);
+	NotifyPostChangeImpl(PropertyChangedEvent, PropertyThatChanged);
 
 	// we need to listen to changes for variables on the blueprint here since
 	// OnFinishedChangingProperties is called only for top level property changes.
@@ -2088,7 +2020,7 @@ void FRigVMEditor::NotifyPostChange(const FPropertyChangedEvent& PropertyChanged
 	}
 }
 
-void FRigVMEditor::OnFinishedChangingProperties(const FPropertyChangedEvent& PropertyChangedEvent)
+void FRigVMEditorBase::OnFinishedChangingProperties(const FPropertyChangedEvent& PropertyChangedEvent)
 {
 	DECLARE_SCOPE_HIERARCHICAL_COUNTER_FUNC()
 
@@ -2109,7 +2041,7 @@ void FRigVMEditor::OnFinishedChangingProperties(const FPropertyChangedEvent& Pro
 	}
 }
 
-void FRigVMEditor::OnPropertyChanged(UObject* InObject, FPropertyChangedEvent& InEvent)
+void FRigVMEditorBase::OnPropertyChanged(UObject* InObject, FPropertyChangedEvent& InEvent)
 {
 	URigVMBlueprint* RigVMBlueprint = GetRigVMBlueprint();
 
@@ -2119,12 +2051,12 @@ void FRigVMEditor::OnPropertyChanged(UObject* InObject, FPropertyChangedEvent& I
 		if(InEvent.MemberProperty == RigVMBlueprint->GetClass()->FindPropertyByName(GET_MEMBER_NAME_CHECKED(URigVMBlueprint, RigVMClient)) ||
 			InEvent.MemberProperty == RigVMBlueprint->GetClass()->FindPropertyByName(GET_MEMBER_NAME_CHECKED(URigVMBlueprint, UbergraphPages)))
 		{
-			DocumentManager->CleanInvalidTabs();
+			GetDocumentManager()->CleanInvalidTabs();
 		}
 	}
 }
 
-void FRigVMEditor::OnWrappedPropertyChangedChainEvent(URigVMDetailsViewWrapperObject* InWrapperObject,
+void FRigVMEditorBase::OnWrappedPropertyChangedChainEvent(URigVMDetailsViewWrapperObject* InWrapperObject,
 	const FString& InPropertyPath, FPropertyChangedChainEvent& InPropertyChangedChainEvent)
 {
 	check(InWrapperObject);
@@ -2144,25 +2076,22 @@ void FRigVMEditor::OnWrappedPropertyChangedChainEvent(URigVMDetailsViewWrapperOb
 			const FRigVMGraphVariableDescription VariableDescription = InWrapperObject->GetContent<FRigVMGraphVariableDescription>();
 			URigVMGraph* Graph = CastChecked<URigVMGraph>(InWrapperObject->GetSubject());
 			URigVMController* Controller = RigVMBlueprint->GetController(Graph);
-			if (PropertyPath == TEXT("Name") && MyBlueprintWidget.IsValid())
+			if (PropertyPath == TEXT("Name"))
 			{
-				if (FEdGraphSchemaAction_BlueprintVariableBase* VariableAcion = MyBlueprintWidget->SelectionAsBlueprintVariable())
+				const FName OldVariableName = GetSelectedVariableName();
+				if (!OldVariableName.IsNone())
 				{
-					const FName OldVariableName = VariableAcion->GetVariableName();
-					if (!OldVariableName.IsNone())
+					for (FRigVMGraphVariableDescription& Variable : Graph->GetLocalVariables())
 					{
-						for (FRigVMGraphVariableDescription& Variable : Graph->GetLocalVariables())
+						if (Variable.Name == OldVariableName)
 						{
-							if (Variable.Name == OldVariableName)
-							{
-								Controller->RenameLocalVariable(OldVariableName, VariableDescription.Name);
-								break;
-							}
+							Controller->RenameLocalVariable(OldVariableName, VariableDescription.Name);
+							break;
 						}
 					}
 				}
-				RefreshMyBlueprint();
-				GraphExplorerWidget->Refresh();
+
+				RefreshEditors();
 				GetRigVMBlueprint()->RequestAutoVMRecompilation();
 			}
 			else if (PropertyPath == TEXT("CPPType") || PropertyPath == TEXT("CPPTypeObject"))
@@ -2325,7 +2254,7 @@ void FRigVMEditor::OnWrappedPropertyChangedChainEvent(URigVMDetailsViewWrapperOb
 	}
 }
 
-void FRigVMEditor::OnRequestLocalizeFunctionDialog(FRigVMGraphFunctionIdentifier& InFunction,
+void FRigVMEditorBase::OnRequestLocalizeFunctionDialog(FRigVMGraphFunctionIdentifier& InFunction,
 	URigVMController* InTargetController,
 	IRigVMGraphFunctionHost* InTargetFunctionHost,
 	bool bForce)
@@ -2333,7 +2262,7 @@ void FRigVMEditor::OnRequestLocalizeFunctionDialog(FRigVMGraphFunctionIdentifier
 	UE::RigVM::Editor::Tools::OnRequestLocalizeFunctionDialog(InFunction, InTargetController, InTargetFunctionHost, bForce);
 }
 
-FRigVMController_BulkEditResult FRigVMEditor::OnRequestBulkEditDialog(URigVMBlueprint* InBlueprint,
+FRigVMController_BulkEditResult FRigVMEditorBase::OnRequestBulkEditDialog(URigVMBlueprint* InBlueprint,
 	URigVMController* InController, URigVMLibraryNode* InFunction, ERigVMControllerBulkEditType InEditType)
 {
 	if (bAllowBulkEdits)
@@ -2368,7 +2297,7 @@ FRigVMController_BulkEditResult FRigVMEditor::OnRequestBulkEditDialog(URigVMBlue
 	return Result;
 }
 
-bool FRigVMEditor::OnRequestBreakLinksDialog(TArray<URigVMLink*> InLinks)
+bool FRigVMEditorBase::OnRequestBreakLinksDialog(TArray<URigVMLink*> InLinks)
 {
 	if(InLinks.Num() == 0)
 	{
@@ -2385,7 +2314,7 @@ bool FRigVMEditor::OnRequestBreakLinksDialog(TArray<URigVMLink*> InLinks)
 	return BreakLinksDialog->ShowModal() == EAppReturnType::Ok; 
 }
 
-TRigVMTypeIndex FRigVMEditor::OnRequestPinTypeSelectionDialog(const TArray<TRigVMTypeIndex>& InTypes)
+TRigVMTypeIndex FRigVMEditorBase::OnRequestPinTypeSelectionDialog(const TArray<TRigVMTypeIndex>& InTypes)
 {
 		if(InTypes.Num() == 0)
 	{
@@ -2500,7 +2429,7 @@ TRigVMTypeIndex FRigVMEditor::OnRequestPinTypeSelectionDialog(const TArray<TRigV
 	return Answer;
 }
 
-void FRigVMEditor::HandleJumpToHyperlink(const UObject* InSubject)
+void FRigVMEditorBase::HandleJumpToHyperlink(const UObject* InSubject)
 {
 	URigVMBlueprint* RigBlueprint = GetRigVMBlueprint();
 	if(RigBlueprint == nullptr)
@@ -2549,10 +2478,10 @@ void FRigVMEditor::HandleJumpToHyperlink(const UObject* InSubject)
 			{
 				if (GEditor->GetEditorSubsystem<UAssetEditorSubsystem>()->OpenEditorForAsset(OtherBlueprint))
 				{
-					FRigVMEditor* OtherEditor = static_cast<FRigVMEditor*>(GEditor->GetEditorSubsystem<UAssetEditorSubsystem>()->FindEditorForAsset(OtherBlueprint, /*bFocusIfOpen =*/true));
-					if(OtherEditor)
+					IAssetEditorInstance* OtherEditor = GEditor->GetEditorSubsystem<UAssetEditorSubsystem>()->FindEditorForAsset(OtherBlueprint, /*bFocusIfOpen =*/true);
+					if (FRigVMEditorBase* OtherRigVMEditor = FRigVMEditorBase::GetFromAssetEditorInstance(OtherEditor))
 					{
-						OtherEditor->HandleJumpToHyperlink(NodeToJumpTo);
+						OtherRigVMEditor->HandleJumpToHyperlink(NodeToJumpTo);
 						return;
 					}
 				}
@@ -2582,7 +2511,7 @@ void FRigVMEditor::HandleJumpToHyperlink(const UObject* InSubject)
 	}
 }
 
-bool FRigVMEditor::UpdateDefaultValueForVariable(FBPVariableDescription& InVariable, bool bUseCDO)
+bool FRigVMEditorBase::UpdateDefaultValueForVariable(FBPVariableDescription& InVariable, bool bUseCDO)
 {
 	bool bAnyValueChanged = false;
 	if (URigVMBlueprint* RigVMBlueprint = GetRigVMBlueprint())
@@ -2609,7 +2538,7 @@ bool FRigVMEditor::UpdateDefaultValueForVariable(FBPVariableDescription& InVaria
 	return bAnyValueChanged;
 }
 
-void FRigVMEditor::UpdateRigVMHost()
+void FRigVMEditorBase::UpdateRigVMHost()
 {
 	DECLARE_SCOPE_HIERARCHICAL_COUNTER_FUNC()
 
@@ -2662,14 +2591,14 @@ void FRigVMEditor::UpdateRigVMHost()
 			Blueprint->ClearFlags(RF_Transient);
 		}
 
-		RigVMHost->OnInitialized_AnyThread().AddSP(this, &FRigVMEditor::HandleVMExecutedEvent);
-		RigVMHost->OnExecuted_AnyThread().AddSP(this, &FRigVMEditor::HandleVMExecutedEvent);
+		RigVMHost->OnInitialized_AnyThread().AddSP(StaticCastSharedRef<FRigVMEditorBase>(SharedRef()), &FRigVMEditorBase::HandleVMExecutedEvent);
+		RigVMHost->OnExecuted_AnyThread().AddSP(StaticCastSharedRef<FRigVMEditorBase>(SharedRef()), &FRigVMEditorBase::HandleVMExecutedEvent);
 		RigVMHost->RequestInit();
-		RigVMHost->GetDebugInfo().ExecutionHalted().AddSP(this, &FRigVMEditor::HandleVMExecutionHalted);
+		RigVMHost->GetDebugInfo().ExecutionHalted().AddSP(StaticCastSharedRef<FRigVMEditorBase>(SharedRef()), &FRigVMEditorBase::HandleVMExecutionHalted);
 	}
 }
 
-void FRigVMEditor::CacheNameLists()
+void FRigVMEditorBase::CacheNameLists()
 {
 	if (URigVMBlueprint* RigVMBlueprint = GetRigVMBlueprint())
 	{
@@ -2688,9 +2617,9 @@ void FRigVMEditor::CacheNameLists()
 	}
 }
 
-void FRigVMEditor::OnCreateComment()
+void FRigVMEditorBase::OnCreateComment()
 {
-	TSharedPtr<SGraphEditor> GraphEditor = FocusedGraphEdPtr.Pin();
+	TSharedPtr<SGraphEditor> GraphEditor = GetFocusedGraphEditor().Pin();
 	if (GraphEditor.IsValid())
 	{
 		if (UEdGraph* Graph = GraphEditor->GetCurrentGraph())
@@ -2717,33 +2646,41 @@ void FRigVMEditor::OnCreateComment()
 	}
 }
 
-TArray<TWeakObjectPtr<UObject>> FRigVMEditor::GetSelectedObjects() const
+TArray<TWeakObjectPtr<UObject>> FRigVMEditorBase::GetSelectedObjects() const
 {
 	// if the inspector shows wrapped objects - look in that array instead.
 	// with recent weak object pointer changes on the property detail view
 	// we cannot rely on the GetSelectedObjects being valid after blueprint compilation.
-	if(WrapperObjects.Num() == Inspector->GetSelectedObjects().Num())
+#if WITH_RIGVMLEGACYEDITOR
+	if (TSharedPtr<SKismetInspector> Inspector = GetInspector())
+#else
+	if (TSharedPtr<SRigVMDetailsInspector> Inspector = GetInspector())
+#endif
 	{
-		TArray<TWeakObjectPtr<UObject>> WeakWrapperObjects;
-		for(const TStrongObjectPtr<URigVMDetailsViewWrapperObject>& WrapperObjectPtr : WrapperObjects)
+		if(WrapperObjects.Num() == Inspector->GetSelectedObjects().Num())
 		{
-			URigVMDetailsViewWrapperObject* WrapperObject = WrapperObjectPtr.Get();
-			if(IsValid(WrapperObject->GetSubject()))
+			TArray<TWeakObjectPtr<UObject>> WeakWrapperObjects;
+			for(const TStrongObjectPtr<URigVMDetailsViewWrapperObject>& WrapperObjectPtr : WrapperObjects)
 			{
-				WeakWrapperObjects.Add(WrapperObject);
+				URigVMDetailsViewWrapperObject* WrapperObject = WrapperObjectPtr.Get();
+				if(IsValid(WrapperObject->GetSubject()))
+				{
+					WeakWrapperObjects.Add(WrapperObject);
+				}
 			}
+			return WeakWrapperObjects;
 		}
-		return WeakWrapperObjects;
+		return Inspector->GetSelectedObjects();
 	}
-	return Inspector->GetSelectedObjects();
+	return TArray<TWeakObjectPtr<UObject>>();
 }
 
-void FRigVMEditor::SetDetailObjects(const TArray<UObject*>& InObjects)
+void FRigVMEditorBase::SetDetailObjects(const TArray<UObject*>& InObjects)
 {
 	SetDetailObjects(InObjects, true);
 }
 
-void FRigVMEditor::SetDetailObjects(const TArray<UObject*>& InObjects, bool bChangeUISelectionState)
+void FRigVMEditorBase::SetDetailObjects(const TArray<UObject*>& InObjects, bool bChangeUISelectionState)
 {
 	if(bSuspendDetailsPanelRefresh)
 	{
@@ -2776,7 +2713,11 @@ void FRigVMEditor::SetDetailObjects(const TArray<UObject*>& InObjects, bool bCha
 			}
 
 			GetTabManager()->InsertNewDocumentTab(
+#if WITH_RIGVMLEGACYEDITOR
 				FBlueprintEditorTabs::DetailsID,
+#else
+				FRigVMNewEditorTabs::DetailsID(),
+#endif
 				TabId,
 				FTabManager::FLastMajorOrNomadTab(TEXT("RigVMMemoryDetails")),
 				DockTab
@@ -2862,7 +2803,7 @@ void FRigVMEditor::SetDetailObjects(const TArray<UObject*>& InObjects, bool bCha
 			// create the wrapper object
 			if(URigVMDetailsViewWrapperObject* WrapperObject = URigVMDetailsViewWrapperObject::MakeInstance(GetDetailWrapperClass(), GetBlueprintObj(), ModelNodes, ModelNode))
 			{
-				WrapperObject->GetWrappedPropertyChangedChainEvent().AddSP(this, &FRigVMEditor::OnWrappedPropertyChangedChainEvent);
+				WrapperObject->GetWrappedPropertyChangedChainEvent().AddSP(StaticCastSharedRef<FRigVMEditorBase>(SharedRef()), &FRigVMEditorBase::OnWrappedPropertyChangedChainEvent);
 				WrapperObject->AddToRoot();
 
 				// todo: use transform widget for transforms
@@ -2885,12 +2826,21 @@ void FRigVMEditor::SetDetailObjects(const TArray<UObject*>& InObjects, bool bCha
 		}
 	}
 
+#if WITH_RIGVMLEGACYEDITOR
 	SKismetInspector::FShowDetailsOptions Options;
 	Options.bForceRefresh = true;
-	Inspector->ShowDetailsForObjects(FilteredObjects, Options);
+	if (TSharedPtr<SKismetInspector> Inspector = GetInspector())
+#else
+	SRigVMDetailsInspector::FShowDetailsOptions Options;
+	Options.bForceRefresh = true;
+	if (TSharedPtr<SRigVMDetailsInspector> Inspector = GetInspector())
+#endif
+	{
+		Inspector->ShowDetailsForObjects(FilteredObjects, Options);
+	}
 }
 
-void FRigVMEditor::SetMemoryStorageDetails(const TArray<FRigVMMemoryStorageStruct*>& InStructs)
+void FRigVMEditorBase::SetMemoryStorageDetails(const TArray<FRigVMMemoryStorageStruct*>& InStructs)
 {
 	if (bSuspendDetailsPanelRefresh)
 	{
@@ -2927,7 +2877,11 @@ void FRigVMEditor::SetMemoryStorageDetails(const TArray<FRigVMMemoryStorageStruc
 			}
 
 			GetTabManager()->InsertNewDocumentTab(
+#if WITH_RIGVMLEGACYEDITOR
 				FBlueprintEditorTabs::DetailsID,
+#else
+				FRigVMNewEditorTabs::DetailsID(),
+#endif
 				TabId,
 				FTabManager::FLastMajorOrNomadTab(TEXT("RigVMMemoryDetails")),
 				DockTab
@@ -2937,7 +2891,7 @@ void FRigVMEditor::SetMemoryStorageDetails(const TArray<FRigVMMemoryStorageStruc
 	}
 }
 
-void FRigVMEditor::SetDetailViewForGraph(URigVMGraph* InGraph)
+void FRigVMEditorBase::SetDetailViewForGraph(URigVMGraph* InGraph)
 {
 	check(InGraph);
 	
@@ -2961,7 +2915,7 @@ void FRigVMEditor::SetDetailViewForGraph(URigVMGraph* InGraph)
 	SetDetailObjects(SelectedNodes);
 }
 
-void FRigVMEditor::SetDetailViewForFocusedGraph()
+void FRigVMEditorBase::SetDetailViewForFocusedGraph()
 {
 	if(bSuspendDetailsPanelRefresh)
 	{
@@ -2977,7 +2931,7 @@ void FRigVMEditor::SetDetailViewForFocusedGraph()
 	SetDetailViewForGraph(Model);
 }
 
-void FRigVMEditor::SetDetailViewForLocalVariable()
+void FRigVMEditorBase::SetDetailViewForLocalVariable()
 {
 	FName VariableName;
 	TArray< TWeakObjectPtr<UObject> > SelectedObjects = GetSelectedObjects();
@@ -2996,7 +2950,7 @@ void FRigVMEditor::SetDetailViewForLocalVariable()
 	SelectLocalVariable(GetFocusedGraph(), VariableName);
 }
 
-void FRigVMEditor::RefreshDetailView()
+void FRigVMEditorBase::RefreshDetailView()
 {
 	if(bSuspendDetailsPanelRefresh)
 	{
@@ -3020,7 +2974,7 @@ void FRigVMEditor::RefreshDetailView()
 	}
 }
 
-bool FRigVMEditor::DetailViewShowsAnyRigUnit() const
+bool FRigVMEditorBase::DetailViewShowsAnyRigUnit() const
 {
 	if (DetailViewShowsStruct(FRigVMStruct::StaticStruct()))
 	{
@@ -3046,12 +3000,12 @@ bool FRigVMEditor::DetailViewShowsAnyRigUnit() const
 	return false;
 }
 
-bool FRigVMEditor::DetailViewShowsLocalVariable() const
+bool FRigVMEditorBase::DetailViewShowsLocalVariable() const
 {
 	return DetailViewShowsStruct(FRigVMGraphVariableDescription::StaticStruct());
 }
 
-bool FRigVMEditor::DetailViewShowsStruct(UScriptStruct* InStruct) const
+bool FRigVMEditorBase::DetailViewShowsStruct(UScriptStruct* InStruct) const
 {
 	TArray< TWeakObjectPtr<UObject> > SelectedObjects = GetSelectedObjects();
 	for (TWeakObjectPtr<UObject> SelectedObject : SelectedObjects)
@@ -3073,7 +3027,7 @@ bool FRigVMEditor::DetailViewShowsStruct(UScriptStruct* InStruct) const
 	return false;
 }
 
-void FRigVMEditor::ClearDetailObject(bool bChangeUISelectionState)
+void FRigVMEditorBase::ClearDetailObject(bool bChangeUISelectionState)
 {
 	if(bSuspendDetailsPanelRefresh)
 	{
@@ -3081,18 +3035,29 @@ void FRigVMEditor::ClearDetailObject(bool bChangeUISelectionState)
 	}
 
 	ClearDetailsViewWrapperObjects();
-	
-	Inspector->GetPropertyView()->SetObjects(TArray<UObject*>(), true); // clear property view synchronously
-	Inspector->ShowDetailsForObjects(TArray<UObject*>());
-	Inspector->ShowSingleStruct(TSharedPtr<FStructOnScope>());
+#if WITH_RIGVMLEGACYEDITOR
+	TSharedPtr<SKismetInspector> Inspector = GetInspector();
+#else
+	TSharedPtr<SRigVMDetailsInspector> Inspector = GetInspector();
+#endif
+	if(Inspector.IsValid())
+	{
+		Inspector->GetPropertyView()->SetObjects(TArray<UObject*>(), true); // clear property view synchronously
+		Inspector->ShowDetailsForObjects(TArray<UObject*>());
+		Inspector->ShowSingleStruct(TSharedPtr<FStructOnScope>());
+	}
 
 	if (bChangeUISelectionState)
 	{
+#if WITH_RIGVMLEGACYEDITOR
 		SetUISelectionState(FBlueprintEditor::SelectionState_Graph);
+#else
+		SetUISelectionState(FRigVMNewEditor::SelectionState_Graph());
+#endif
 	}
 }
 
-void FRigVMEditor::ClearDetailsViewWrapperObjects()
+void FRigVMEditorBase::ClearDetailsViewWrapperObjects()
 {
 	for(const TStrongObjectPtr<URigVMDetailsViewWrapperObject>& WrapperObjectPtr : WrapperObjects)
 	{
@@ -3107,7 +3072,7 @@ void FRigVMEditor::ClearDetailsViewWrapperObjects()
 	WrapperObjects.Reset();
 }
 
-void FRigVMEditor::SetHost(URigVMHost* InHost)
+void FRigVMEditorBase::SetHost(URigVMHost* InHost)
 {
 	if (URigVMBlueprint* RigVMBlueprint = GetRigVMBlueprint())
 	{
@@ -3124,7 +3089,7 @@ void FRigVMEditor::SetHost(URigVMHost* InHost)
 	}
 }
 
-URigVMGraph* FRigVMEditor::GetFocusedModel() const
+URigVMGraph* FRigVMEditorBase::GetFocusedModel() const
 {
 	const URigVMBlueprint* Blueprint = GetRigVMBlueprint();
 	if (Blueprint == nullptr)
@@ -3136,7 +3101,7 @@ URigVMGraph* FRigVMEditor::GetFocusedModel() const
 	return Blueprint->GetModel(EdGraph);
 }
 
-URigVMController* FRigVMEditor::GetFocusedController() const
+URigVMController* FRigVMEditorBase::GetFocusedController() const
 {
 	URigVMBlueprint* Blueprint = GetRigVMBlueprint();
 	if (Blueprint == nullptr)
@@ -3146,10 +3111,10 @@ URigVMController* FRigVMEditor::GetFocusedController() const
 	return Blueprint->GetOrCreateController(GetFocusedModel());
 }
 
-TSharedPtr<SGraphEditor> FRigVMEditor::GetGraphEditor(UEdGraph* InEdGraph) const
+TSharedPtr<SGraphEditor> FRigVMEditorBase::GetGraphEditor(UEdGraph* InEdGraph) const
 {
 	TArray< TSharedPtr<SDockTab> > GraphEditorTabs;
-	DocumentManager->FindAllTabsForFactory(GraphEditorTabFactoryPtr, /*out*/ GraphEditorTabs);
+	GetDocumentManager()->FindAllTabsForFactory(GetGraphEditorTabFactory(), /*out*/ GraphEditorTabs);
 
 	for (const TSharedPtr<SDockTab>& GraphEditorTab : GraphEditorTabs)
 	{
@@ -3163,7 +3128,7 @@ TSharedPtr<SGraphEditor> FRigVMEditor::GetGraphEditor(UEdGraph* InEdGraph) const
 	return TSharedPtr<SGraphEditor>();
 }
 
-void FRigVMEditor::ExtendMenu()
+void FRigVMEditorBase::ExtendMenu()
 {
 	if(MenuExtender.IsValid())
 	{
@@ -3180,7 +3145,7 @@ void FRigVMEditor::ExtendMenu()
 	AddMenuExtender(RigVMEditorModule.GetMenuExtensibilityManager()->GetAllExtenders(GetToolkitCommands(), GetEditingObjects()));
 }
 
-void FRigVMEditor::ExtendToolbar()
+void FRigVMEditorBase::ExtendToolbar()
 {
 	DECLARE_SCOPE_HIERARCHICAL_COUNTER_FUNC()
 
@@ -3204,7 +3169,7 @@ void FRigVMEditor::ExtendToolbar()
 	{
 		if(ToolbarExtenderDelegate.IsBound())
 		{
-			AddToolbarExtender(ToolbarExtenderDelegate.Execute(GetToolkitCommands(), SharedThis(this)));
+			AddToolbarExtender(ToolbarExtenderDelegate.Execute(GetToolkitCommands(), StaticCastSharedRef<FRigVMEditorBase>(SharedRef())));
 		}
 	}
 
@@ -3212,26 +3177,26 @@ void FRigVMEditor::ExtendToolbar()
 		"Asset",
 		EExtensionHook::After,
 		GetToolkitCommands(),
-		FToolBarExtensionDelegate::CreateSP(this, &FRigVMEditor::FillToolbar, true)
+		FToolBarExtensionDelegate::CreateSP(StaticCastSharedRef<FRigVMEditorBase>(SharedRef()), &FRigVMEditorBase::FillToolbar, true)
 	);
 }
 
-void FRigVMEditor::FillToolbar(FToolBarBuilder& ToolbarBuilder, bool bEndSection)
+void FRigVMEditorBase::FillToolbar(FToolBarBuilder& ToolbarBuilder, bool bEndSection)
 {
 	ToolbarBuilder.BeginSection("Toolbar");
 	{
 		ToolbarBuilder.AddToolBarButton(
 			FRigVMEditorCommands::Get().ToggleEventQueue,
 			NAME_None, 
-			TAttribute<FText>::Create(TAttribute<FText>::FGetter::CreateSP(this, &FRigVMEditor::GetEventQueueLabel)),
+			TAttribute<FText>::Create(TAttribute<FText>::FGetter::CreateSP(StaticCastSharedRef<FRigVMEditorBase>(SharedRef()), &FRigVMEditorBase::GetEventQueueLabel)),
 			TAttribute<FText>(), 
-			TAttribute<FSlateIcon>::Create(TAttribute<FSlateIcon>::FGetter::CreateSP(this, &FRigVMEditor::GetEventQueueIcon))
+			TAttribute<FSlateIcon>::Create(TAttribute<FSlateIcon>::FGetter::CreateSP(StaticCastSharedRef<FRigVMEditorBase>(SharedRef()), &FRigVMEditorBase::GetEventQueueIcon))
 		);
 
 		FUIAction DefaultAction;
 		ToolbarBuilder.AddComboButton(
 			DefaultAction,
-			FOnGetContent::CreateSP(this, &FRigVMEditor::GenerateEventQueueMenuContent),
+			FOnGetContent::CreateSP(StaticCastSharedRef<FRigVMEditorBase>(SharedRef()), &FRigVMEditorBase::GenerateEventQueueMenuContent),
 			LOCTEXT("EventQueue_Label", "Available Events"),
 			LOCTEXT("EventQueue_ToolTip", "Pick between different events / modes for testing the Control Rig"),
 			FSlateIcon(FAppStyle::GetAppStyleSetName(), "LevelEditor.Recompile"),
@@ -3240,21 +3205,21 @@ void FRigVMEditor::FillToolbar(FToolBarBuilder& ToolbarBuilder, bool bEndSection
 		ToolbarBuilder.AddToolBarButton(FRigVMEditorCommands::Get().AutoCompileGraph,
 			NAME_None, TAttribute<FText>(), TAttribute<FText>(), FSlateIcon(FRigVMEditorStyle::Get().GetStyleSetName(), "RigVM.AutoCompileGraph"));
 
-		ToolbarBuilder.AddWidget(SNew(SBlueprintEditorSelectedDebugObjectWidget, SharedThis(this)));
+		AddSelectedDebugObjectWidget(ToolbarBuilder);
 
 		ToolbarBuilder.AddSeparator();
 		ToolbarBuilder.AddToolBarButton(
 			FRigVMEditorCommands::Get().ToggleExecutionMode,
 			NAME_None, 
-			TAttribute<FText>::Create(TAttribute<FText>::FGetter::CreateSP(this, &FRigVMEditor::GetExecutionModeLabel)),
+			TAttribute<FText>::Create(TAttribute<FText>::FGetter::CreateSP(StaticCastSharedRef<FRigVMEditorBase>(SharedRef()), &FRigVMEditorBase::GetExecutionModeLabel)),
 			TAttribute<FText>(), 
-			TAttribute<FSlateIcon>::Create(TAttribute<FSlateIcon>::FGetter::CreateSP(this, &FRigVMEditor::GetExecutionModeIcon))
+			TAttribute<FSlateIcon>::Create(TAttribute<FSlateIcon>::FGetter::CreateSP(StaticCastSharedRef<FRigVMEditorBase>(SharedRef()), &FRigVMEditorBase::GetExecutionModeIcon))
 		);
 		
 		FUIAction DefaultExecutionMode;
 		ToolbarBuilder.AddComboButton(
 			DefaultExecutionMode,
-			FOnGetContent::CreateSP(this, &FRigVMEditor::GenerateExecutionModeMenuContent),
+			FOnGetContent::CreateSP(StaticCastSharedRef<FRigVMEditorBase>(SharedRef()), &FRigVMEditorBase::GenerateExecutionModeMenuContent),
 			LOCTEXT("ExecutionMode_Label", "Execution Modes"),
 			LOCTEXT("ExecutionMode_ToolTip", "Pick between different execution modes for testing the Control Rig"),
 			FSlateIcon(FAppStyle::GetAppStyleSetName(), "LevelEditor.Recompile"),
@@ -3284,7 +3249,7 @@ void FRigVMEditor::FillToolbar(FToolBarBuilder& ToolbarBuilder, bool bEndSection
 		FUIAction DefaultBulkEditAction;
 		ToolbarBuilder.AddComboButton(
 			DefaultBulkEditAction,
-			FOnGetContent::CreateSP(this, &FRigVMEditor::GenerateBulkEditMenuContent),
+			FOnGetContent::CreateSP(StaticCastSharedRef<FRigVMEditorBase>(SharedRef()), &FRigVMEditorBase::GenerateBulkEditMenuContent),
 			LOCTEXT("BulkEdit_Label", "Bulk Edit"),
 			LOCTEXT("BulkEdit_ToolTip", "Perform changes across many nodes / assets"),
 			FSlateIcon(FAppStyle::GetAppStyleSetName(), "LevelEditor.Recompile"),
@@ -3297,7 +3262,7 @@ void FRigVMEditor::FillToolbar(FToolBarBuilder& ToolbarBuilder, bool bEndSection
 	}
 }
 
-void FRigVMEditor::HandleHideItem()
+void FRigVMEditorBase::HandleHideItem()
 {
 	DECLARE_SCOPE_HIERARCHICAL_COUNTER_FUNC()
 
@@ -3324,12 +3289,12 @@ void FRigVMEditor::HandleHideItem()
 	}
 }
 
-bool FRigVMEditor::CanHideItem() const
+bool FRigVMEditorBase::CanHideItem() const
 {
-	return GetNumberOfSelectedNodes() > 0;
+	return GetSelectedNodes().Num() > 0;
 }
 
-void FRigVMEditor::UpdateStaleWatchedPins()
+void FRigVMEditorBase::UpdateStaleWatchedPins()
 {
 	URigVMBlueprint* RigVMBlueprint = GetRigVMBlueprint();
 	if (RigVMBlueprint == nullptr)
@@ -3399,12 +3364,12 @@ void FRigVMEditor::UpdateStaleWatchedPins()
 	}
 }
 
-void FRigVMEditor::HandleRefreshEditorFromBlueprint(URigVMBlueprint* InBlueprint)
+void FRigVMEditorBase::HandleRefreshEditorFromBlueprint(URigVMBlueprint* InBlueprint)
 {
 	Compile();
 }
 
-void FRigVMEditor::HandleVariableDroppedFromBlueprint(UObject* InSubject, FProperty* InVariableToDrop, const FVector2D& InDropPosition, const FVector2D& InScreenPosition)
+void FRigVMEditorBase::HandleVariableDroppedFromBlueprint(UObject* InSubject, FProperty* InVariableToDrop, const FVector2D& InDropPosition, const FVector2D& InScreenPosition)
 {
 	URigVMBlueprint* Blueprint = Cast<URigVMBlueprint>(GetBlueprintObj());
 	if (Blueprint == nullptr)
@@ -3456,7 +3421,7 @@ void FRigVMEditor::HandleVariableDroppedFromBlueprint(UObject* InSubject, FPrope
 
 	MenuBuilder.EndSection();
 
-	TSharedRef<SWidget> GraphEditorPanel = FocusedGraphEdPtr.Pin().ToSharedRef();
+	TSharedRef<SWidget> GraphEditorPanel = GetFocusedGraphEditor().Pin().ToSharedRef();
 
 	// Show dialog to choose getter vs setter
 	FSlateApplication::Get().PushMenu(
@@ -3468,12 +3433,12 @@ void FRigVMEditor::HandleVariableDroppedFromBlueprint(UObject* InSubject, FPrope
 	);
 }
 
-void FRigVMEditor::HandleBreakpointAdded()
+void FRigVMEditorBase::HandleBreakpointAdded()
 {
 	SetExecutionMode(ERigVMEditorExecutionModeType_Debug);
 }
 
-void FRigVMEditor::OnGraphNodeClicked(URigVMEdGraphNode* InNode)
+void FRigVMEditorBase::OnGraphNodeClicked(URigVMEdGraphNode* InNode)
 {
 	if (InNode)
 	{
@@ -3484,7 +3449,7 @@ void FRigVMEditor::OnGraphNodeClicked(URigVMEdGraphNode* InNode)
 	}
 }
 
-void FRigVMEditor::OnNodeDoubleClicked(URigVMBlueprint* InBlueprint, URigVMNode* InNode)
+void FRigVMEditorBase::OnNodeDoubleClicked(URigVMBlueprint* InBlueprint, URigVMNode* InNode)
 {
 	ensure(GetRigVMBlueprint() == InBlueprint);
 
@@ -3506,25 +3471,13 @@ void FRigVMEditor::OnNodeDoubleClicked(URigVMBlueprint* InBlueprint, URigVMNode*
 			}
 			else
 			{
-				if(const URigVMCollapseNode* FunctionLibraryNode = Cast<URigVMCollapseNode>(ContainedGraph->GetOuter()))
-				{
-					if(const URigVMFunctionLibrary* FunctionLibrary = FunctionLibraryNode->GetLibrary())
-					{
-						if(const URigVMBlueprint* FunctionBlueprint = Cast<URigVMBlueprint>(FunctionLibrary->GetOuter()))
-						{
-							if (const UEdGraph* FunctionEdGraph = FunctionBlueprint->GetEdGraph(ContainedGraph))
-							{
-								FKismetEditorUtilities::BringKismetToFocusAttentionOnObject(FunctionEdGraph);
-							}
-						}
-					}
-				}
+				UE_LOG(LogRigVMEditor, Warning, TEXT("Could not open graph (%s)"), *LibraryNode->GetFunctionIdentifier().GetLibraryNodePath());
 			}
 		}
 	}
 }
 
-void FRigVMEditor::OnGraphImported(UEdGraph* InEdGraph)
+void FRigVMEditorBase::OnGraphImported(UEdGraph* InEdGraph)
 {
 	check(InEdGraph);
 
@@ -3532,7 +3485,7 @@ void FRigVMEditor::OnGraphImported(UEdGraph* InEdGraph)
 	RenameNewlyAddedAction(InEdGraph->GetFName());
 }
 
-bool FRigVMEditor::OnActionMatchesName(FEdGraphSchemaAction* InAction, const FName& InName) const
+bool FRigVMEditorBase::OnActionMatchesName(FEdGraphSchemaAction* InAction, const FName& InName) const
 {
 	if (InAction->GetMenuDescription().ToString() == InName.ToString())
 	{
@@ -3541,7 +3494,7 @@ bool FRigVMEditor::OnActionMatchesName(FEdGraphSchemaAction* InAction, const FNa
 	return false;
 }
 
-void FRigVMEditor::HandleShowCurrentStatement()
+void FRigVMEditorBase::HandleShowCurrentStatement()
 {
 	if(HaltedAtNode)
 	{
@@ -3558,7 +3511,7 @@ void FRigVMEditor::HandleShowCurrentStatement()
 	}
 }
 
-void FRigVMEditor::HandleBreakpointActionRequested(const ERigVMBreakpointAction BreakpointAction)
+void FRigVMEditorBase::HandleBreakpointActionRequested(const ERigVMBreakpointAction BreakpointAction)
 {
 	if(URigVMHost* DebuggedHost = Cast<URigVMHost>(GetBlueprintObj()->GetObjectBeingDebugged()))
 	{
@@ -3566,14 +3519,14 @@ void FRigVMEditor::HandleBreakpointActionRequested(const ERigVMBreakpointAction
 	}
 }
 
-bool FRigVMEditor::IsHaltedAtBreakpoint() const
+bool FRigVMEditorBase::IsHaltedAtBreakpoint() const
 {
 	return HaltedAtNode != nullptr;
 }
 
-void FRigVMEditor::FrameSelection()
+void FRigVMEditorBase::FrameSelection()
 {
-	if (SGraphEditor* GraphEd = FocusedGraphEdPtr.Pin().Get())
+	if (SGraphEditor* GraphEd = GetFocusedGraphEditor().Pin().Get())
 	{
 		if(URigVMGraph* Model = GetFocusedModel())
 		{
@@ -3583,13 +3536,13 @@ void FRigVMEditor::FrameSelection()
 	}
 }
 
-void FRigVMEditor::SwapFunctionWithinAsset()
+void FRigVMEditorBase::SwapFunctionWithinAsset()
 {
 	const FAssetData Asset = UE::RigVM::Editor::Tools::FindAssetFromAnyPath(GetRigVMBlueprint()->GetPathName(), true);
 	SwapFunctionForAssets({Asset}, true);
 }
 
-void FRigVMEditor::SwapFunctionAcrossProject()
+void FRigVMEditorBase::SwapFunctionAcrossProject()
 {
 	const IAssetRegistry& AssetRegistry = FModuleManager::LoadModuleChecked<FAssetRegistryModule>(TEXT("AssetRegistry")).Get();
 	TArray<FAssetData> AllAssets;
@@ -3597,7 +3550,7 @@ void FRigVMEditor::SwapFunctionAcrossProject()
 	SwapFunctionForAssets(AllAssets, false);
 }
 
-void FRigVMEditor::SwapFunctionForAssets(const TArray<FAssetData>& InAssets, bool bSetupUndo)
+void FRigVMEditorBase::SwapFunctionForAssets(const TArray<FAssetData>& InAssets, bool bSetupUndo)
 {
 	SRigVMSwapFunctionsWidget::FArguments WidgetArgs;
 	WidgetArgs
@@ -3613,7 +3566,7 @@ void FRigVMEditor::SwapFunctionForAssets(const TArray<FAssetData>& InAssets, boo
 	SwapFunctionsDialog->ShowNormal();
 }
 
-void FRigVMEditor::SwapAssetReferences()
+void FRigVMEditorBase::SwapAssetReferences()
 {
 	const FAssetData Asset = UE::RigVM::Editor::Tools::FindAssetFromAnyPath(GetRigVMBlueprint()->GetPathName(), true);
 	
@@ -3631,7 +3584,7 @@ void FRigVMEditor::SwapAssetReferences()
 	SwapFunctionsDialog->ShowNormal();
 }
 
-void FRigVMEditor::OnGraphNodeDropToPerform(TSharedPtr<FDragDropOperation> InDragDropOp, UEdGraph* InGraph, const FVector2D& InNodePosition, const FVector2D& InScreenPosition)
+void FRigVMEditorBase::OnGraphNodeDropToPerform(TSharedPtr<FDragDropOperation> InDragDropOp, UEdGraph* InGraph, const FVector2D& InNodePosition, const FVector2D& InScreenPosition)
 {
 	DECLARE_SCOPE_HIERARCHICAL_COUNTER_FUNC()
 
@@ -3639,7 +3592,7 @@ void FRigVMEditor::OnGraphNodeDropToPerform(TSharedPtr<FDragDropOperation> InDra
 	{
 		TSharedPtr<FRigVMGraphExplorerDragDropOp> ExplorerOp = StaticCastSharedPtr<FRigVMGraphExplorerDragDropOp>(InDragDropOp);
 
-		if (FocusedGraphEdPtr.IsValid())
+		if (GetFocusedGraphEditor().IsValid())
 		{
 			const URigVMEdGraphSchema* Schema = CastChecked<URigVMEdGraphSchema>(GetRigVMBlueprint()->GetRigVMEdGraphSchemaClass()->GetDefaultObject());
 			if (ExplorerOp->GetElement().Type == ERigVMExplorerElementType::Function)
@@ -3723,7 +3676,7 @@ void FRigVMEditor::OnGraphNodeDropToPerform(TSharedPtr<FDragDropOperation> InDra
 	}
 }
 
-void FRigVMEditor::UpdateGraphCompilerErrors()
+void FRigVMEditorBase::UpdateGraphCompilerErrors()
 {
 	DECLARE_SCOPE_HIERARCHICAL_COUNTER_FUNC()
 
@@ -3824,17 +3777,17 @@ void FRigVMEditor::UpdateGraphCompilerErrors()
 
 }
 
-bool FRigVMEditor::IsPIERunning()
+bool FRigVMEditorBase::IsPIERunning()
 {
 	return GEditor && (GEditor->PlayWorld != nullptr);
 }
 
-TArray<FName> FRigVMEditor::GetDefaultEventQueue() const
+TArray<FName> FRigVMEditorBase::GetDefaultEventQueue() const
 {
 	return TArray<FName>();
 }
 
-TArray<FName> FRigVMEditor::GetEventQueue() const
+TArray<FName> FRigVMEditorBase::GetEventQueue() const
 {
 	if (const URigVMHost* CurrentHost = GetRigVMHost())
 	{
@@ -3844,12 +3797,12 @@ TArray<FName> FRigVMEditor::GetEventQueue() const
 	return GetDefaultEventQueue();
 }
 
-void FRigVMEditor::SetEventQueue(TArray<FName> InEventQueue)
+void FRigVMEditorBase::SetEventQueue(TArray<FName> InEventQueue)
 {
 	return SetEventQueue(InEventQueue, false);
 }
 
-void FRigVMEditor::SetEventQueue(TArray<FName> InEventQueue, bool bCompile)
+void FRigVMEditorBase::SetEventQueue(TArray<FName> InEventQueue, bool bCompile)
 {
 	if (GetEventQueue() == InEventQueue)
 	{
@@ -3869,17 +3822,17 @@ void FRigVMEditor::SetEventQueue(TArray<FName> InEventQueue, bool bCompile)
 	}
 }
 
-FSlateIcon FRigVMEditor::GetEventQueueIcon(const TArray<FName>& InEventQueue) const
+FSlateIcon FRigVMEditorBase::GetEventQueueIcon(const TArray<FName>& InEventQueue) const
 {
 	return FSlateIcon();
 }
 
-FSlateIcon FRigVMEditor::GetEventQueueIcon() const
+FSlateIcon FRigVMEditorBase::GetEventQueueIcon() const
 {
 	return GetEventQueueIcon(GetEventQueue());
 }
 
-void FRigVMEditor::SetExecutionMode(const ERigVMEditorExecutionModeType InExecutionMode)
+void FRigVMEditorBase::SetExecutionMode(const ERigVMEditorExecutionModeType InExecutionMode)
 {
 	if (ExecutionMode == InExecutionMode)
 	{
@@ -3900,12 +3853,12 @@ void FRigVMEditor::SetExecutionMode(const ERigVMEditorExecutionModeType InExecut
 	RefreshDetailView();
 }
 
-int32 FRigVMEditor::GetExecutionModeComboValue() const
+int32 FRigVMEditorBase::GetExecutionModeComboValue() const
 {
 	return (int32) ExecutionMode;
 }
 
-FText FRigVMEditor::GetExecutionModeLabel() const
+FText FRigVMEditorBase::GetExecutionModeLabel() const
 {
 	if (ExecutionMode == ERigVMEditorExecutionModeType_Debug)
 	{
@@ -3914,7 +3867,7 @@ FText FRigVMEditor::GetExecutionModeLabel() const
 	return FText::FromString(TEXT("ReleaseMode"));
 }
 
-FSlateIcon FRigVMEditor::GetExecutionModeIcon(ERigVMEditorExecutionModeType InExecutionMode)
+FSlateIcon FRigVMEditorBase::GetExecutionModeIcon(ERigVMEditorExecutionModeType InExecutionMode)
 {
 	if (InExecutionMode == ERigVMEditorExecutionModeType_Debug)
 	{
@@ -3923,12 +3876,12 @@ FSlateIcon FRigVMEditor::GetExecutionModeIcon(ERigVMEditorExecutionModeType InEx
 	return FSlateIcon(FRigVMEditorStyle::Get().GetStyleSetName(), "RigVM.ReleaseMode");
 }
 
-FSlateIcon FRigVMEditor::GetExecutionModeIcon() const
+FSlateIcon FRigVMEditorBase::GetExecutionModeIcon() const
 {
 	return GetExecutionModeIcon(ExecutionMode);
 }
 
-void FRigVMEditor::GetCustomDebugObjects(TArray<FCustomDebugObject>& DebugList) const
+void FRigVMEditorBase::GetCustomDebugObjects(TArray<FRigVMCustomDebugObject>& DebugList) const
 {
 	URigVMBlueprint* RigVMBlueprint = GetRigVMBlueprint();
 	if (RigVMBlueprint == nullptr)
@@ -3940,7 +3893,7 @@ void FRigVMEditor::GetCustomDebugObjects(TArray<FCustomDebugObject>& DebugList)
 	{
 		if(IsValid(CurrentHost))
 		{
-			FCustomDebugObject DebugObject;
+			FRigVMCustomDebugObject DebugObject;
 			DebugObject.Object = CurrentHost;
 			DebugObject.NameOverride = GetCustomDebugObjectLabel(CurrentHost);
 			DebugList.Add(DebugObject);
@@ -4032,7 +3985,7 @@ void FRigVMEditor::GetCustomDebugObjects(TArray<FCustomDebugObject>& DebugList)
 							continue;
 						}
 
-						FCustomDebugObject DebugObject;
+						FRigVMCustomDebugObject DebugObject;
 						DebugObject.Object = InstancedHost;
 						DebugObject.NameOverride = GetCustomDebugObjectLabel(InstancedHost);
 						DebugList.Add(DebugObject);
@@ -4043,7 +3996,7 @@ void FRigVMEditor::GetCustomDebugObjects(TArray<FCustomDebugObject>& DebugList)
 	}
 }
 
-void FRigVMEditor::HandleSetObjectBeingDebugged(UObject* InObject)
+void FRigVMEditorBase::HandleSetObjectBeingDebugged(UObject* InObject)
 {
 	if(URigVMHost* PreviouslyDebuggedHost = Cast<URigVMHost>(GetBlueprintObj()->GetObjectBeingDebugged()))
 	{
@@ -4084,9 +4037,9 @@ void FRigVMEditor::HandleSetObjectBeingDebugged(UObject* InObject)
 	if(DebuggedHost)
 	{
 		DebuggedHost->SetLog(&RigVMLog);
-		DebuggedHost->OnExecuted_AnyThread().AddSP(this, &FRigVMEditor::HandleVMExecutedEvent);
+		DebuggedHost->OnExecuted_AnyThread().AddSP(StaticCastSharedRef<FRigVMEditorBase>(SharedRef()), &FRigVMEditorBase::HandleVMExecutedEvent);
 #if WITH_EDITOR
-		DebuggedHost->GetDebugInfo().ExecutionHalted().AddSP(this, &FRigVMEditor::HandleVMExecutionHalted);
+		DebuggedHost->GetDebugInfo().ExecutionHalted().AddSP(StaticCastSharedRef<FRigVMEditorBase>(SharedRef()), &FRigVMEditorBase::HandleVMExecutionHalted);
 		DebuggedHost->SetIsInDebugMode(ExecutionMode == ERigVMEditorExecutionModeType_Debug);	
 #endif
 	}
@@ -4095,7 +4048,7 @@ void FRigVMEditor::HandleSetObjectBeingDebugged(UObject* InObject)
 	LastDebuggedHost = GetCustomDebugObjectLabel(DebuggedHost);
 }
 
-FString FRigVMEditor::GetCustomDebugObjectLabel(UObject* ObjectBeingDebugged) const
+FString FRigVMEditorBase::GetCustomDebugObjectLabel(UObject* ObjectBeingDebugged) const
 {
 	if (ObjectBeingDebugged == nullptr)
 	{
@@ -4124,7 +4077,7 @@ FString FRigVMEditor::GetCustomDebugObjectLabel(UObject* ObjectBeingDebugged) co
 	return GetBlueprintObj()->GetName();
 }
 
-void FRigVMEditor::OnPIEStopped(bool bSimulation)
+void FRigVMEditorBase::OnPIEStopped(bool bSimulation)
 {
 	if(URigVMBlueprint* Blueprint = GetRigVMBlueprint())
 	{
@@ -4132,4 +4085,9 @@ void FRigVMEditor::OnPIEStopped(bool bSimulation)
 	}
 }
 
-#undef LOCTEXT_NAMESPACE 
+bool FRigVMEditorBase::IsEditingSingleBlueprint() const
+{
+	return GetBlueprintObj() != nullptr;
+}
+
+#undef LOCTEXT_NAMESPACE
\ No newline at end of file
diff --git a/Engine/Plugins/Runtime/RigVM/Source/RigVMEditor/Private/Editor/RigVMEditorGraphExplorerTabSummoner.cpp b/Engine/Plugins/Runtime/RigVM/Source/RigVMEditor/Private/Editor/RigVMEditorGraphExplorerTabSummoner.cpp
--- a/Engine/Plugins/Runtime/RigVM/Source/RigVMEditor/Private/Editor/RigVMEditorGraphExplorerTabSummoner.cpp
+++ b/Engine/Plugins/Runtime/RigVM/Source/RigVMEditor/Private/Editor/RigVMEditorGraphExplorerTabSummoner.cpp
@@ -2,12 +2,12 @@
 
 #include "Editor/RigVMEditorGraphExplorerTabSummoner.h"
 #include "Widgets/SRigVMEditorGraphExplorer.h"
-#include "Editor/RigVMEditor.h"
+#include "Editor/RigVMNewEditor.h"
 
 #define LOCTEXT_NAMESPACE "RigVMEditorGraphExplorerTabSummoner"
 
-FRigVMEditorGraphExplorerTabSummoner::FRigVMEditorGraphExplorerTabSummoner(const TSharedRef<FRigVMEditor>& InRigVMEditor)
-	: FWorkflowTabFactory(TabID, InRigVMEditor)
+FRigVMEditorGraphExplorerTabSummoner::FRigVMEditorGraphExplorerTabSummoner(const TSharedRef<FRigVMEditorBase>& InRigVMEditor)
+	: FWorkflowTabFactory(TabID, InRigVMEditor.Get().GetHostingApp())
 	, RigVMEditor(InRigVMEditor)
 {
 	TabLabel = LOCTEXT("RigVMEditorGraphExplorerTabLabel", "Graph Explorer");
diff --git a/Engine/Plugins/Runtime/RigVM/Source/RigVMEditor/Private/Editor/RigVMEditorMenuContext.cpp b/Engine/Plugins/Runtime/RigVM/Source/RigVMEditor/Private/Editor/RigVMEditorMenuContext.cpp
--- a/Engine/Plugins/Runtime/RigVM/Source/RigVMEditor/Private/Editor/RigVMEditorMenuContext.cpp
+++ b/Engine/Plugins/Runtime/RigVM/Source/RigVMEditor/Private/Editor/RigVMEditorMenuContext.cpp
@@ -3,12 +3,12 @@
 #include "Editor/RigVMEditorMenuContext.h"
 
 #include "RigVMBlueprint.h"
-#include "Editor/RigVMEditor.h"
+#include "Editor/RigVMNewEditor.h"
 #include "Framework/Application/SlateApplication.h"
 
 #include UE_INLINE_GENERATED_CPP_BY_NAME(RigVMEditorMenuContext)
 
-void URigVMEditorMenuContext::Init(TWeakPtr<FRigVMEditor> InRigVMEditor, const FRigVMEditorGraphMenuContext& InGraphMenuContext)
+void URigVMEditorMenuContext::Init(TWeakPtr<FRigVMEditorBase> InRigVMEditor, const FRigVMEditorGraphMenuContext& InGraphMenuContext)
 {
 	WeakRigVMEditor = InRigVMEditor;
 	GraphMenuContext = InGraphMenuContext;
@@ -16,7 +16,7 @@ void URigVMEditorMenuContext::Init(TWeakPtr<FRigVMEditor> InRigVMEditor, const F
 
 URigVMBlueprint* URigVMEditorMenuContext::GetRigVMBlueprint() const
 {
-	if (const TSharedPtr<FRigVMEditor> Editor = WeakRigVMEditor.Pin())
+	if (const TSharedPtr<FRigVMEditorBase> Editor = WeakRigVMEditor.Pin())
 	{
 		return Editor->GetRigVMBlueprint();
 	}
@@ -46,9 +46,9 @@ FRigVMEditorGraphMenuContext URigVMEditorMenuContext::GetGraphMenuContext()
 	return GraphMenuContext;
 }
 
-FRigVMEditor* URigVMEditorMenuContext::GetRigVMEditor() const
+FRigVMEditorBase* URigVMEditorMenuContext::GetRigVMEditor() const
 {
-	if (const TSharedPtr<FRigVMEditor> Editor = WeakRigVMEditor.Pin())
+	if (const TSharedPtr<FRigVMEditorBase> Editor = WeakRigVMEditor.Pin())
 	{
 		return Editor.Get();
 	}
diff --git a/Engine/Plugins/Runtime/RigVM/Source/RigVMEditor/Private/Editor/RigVMExecutionStackTabSummoner.cpp b/Engine/Plugins/Runtime/RigVM/Source/RigVMEditor/Private/Editor/RigVMExecutionStackTabSummoner.cpp
--- a/Engine/Plugins/Runtime/RigVM/Source/RigVMEditor/Private/Editor/RigVMExecutionStackTabSummoner.cpp
+++ b/Engine/Plugins/Runtime/RigVM/Source/RigVMEditor/Private/Editor/RigVMExecutionStackTabSummoner.cpp
@@ -2,12 +2,12 @@
 
 #include "Editor/RigVMExecutionStackTabSummoner.h"
 #include "Widgets/SRigVMExecutionStackView.h"
-#include "Editor/RigVMEditor.h"
+#include "Editor/RigVMNewEditor.h"
 
 #define LOCTEXT_NAMESPACE "RigVMExecutionStackTabSummoner"
 
-FRigVMExecutionStackTabSummoner::FRigVMExecutionStackTabSummoner(const TSharedRef<FRigVMEditor>& InRigVMEditor)
-	: FWorkflowTabFactory(TabID, InRigVMEditor)
+FRigVMExecutionStackTabSummoner::FRigVMExecutionStackTabSummoner(const TSharedRef<FRigVMEditorBase>& InRigVMEditor)
+	: FWorkflowTabFactory(TabID, InRigVMEditor->GetHostingApp())
 	, RigVMEditor(InRigVMEditor)
 {
 	TabLabel = LOCTEXT("RigVMExecutionStackTabLabel", "Execution Stack");
