Control Rig: New RigVMEditor, which does not inherit from BlueprintEditor

diff --git a/Engine/Plugins/Runtime/RigVM/Source/RigVMEditor/Private/Editor/RigVMGraphDetailCustomization.cpp b/Engine/Plugins/Runtime/RigVM/Source/RigVMEditor/Private/Editor/RigVMGraphDetailCustomization.cpp
--- a/Engine/Plugins/Runtime/RigVM/Source/RigVMEditor/Private/Editor/RigVMGraphDetailCustomization.cpp
+++ b/Engine/Plugins/Runtime/RigVM/Source/RigVMEditor/Private/Editor/RigVMGraphDetailCustomization.cpp
@@ -35,6 +35,9 @@
 #include "ContentBrowserModule.h"
 #include "IContentBrowserSingleton.h"
 #include "UObject/GarbageCollectionSchema.h"
+#if WITH_RIGVMLEGACYEDITOR
+#include "Editor/RigVMLegacyEditor.h"
+#endif
 
 #define LOCTEXT_NAMESPACE "RigVMGraphDetailCustomization"
 
@@ -43,7 +46,7 @@ static const FText RigVMGraphDetailCustomizationMultipleValues = LOCTEXT("Multip
 FRigVMFunctionArgumentGroupLayout::FRigVMFunctionArgumentGroupLayout(
 	const TWeakObjectPtr<URigVMGraph>& InGraph,
 	const TWeakInterfacePtr<IRigVMClientHost>& InRigVMClientHost,
-	const TWeakPtr<FRigVMEditor>& InEditor,
+	const TWeakPtr<FRigVMEditorBase>& InEditor,
 	bool bInputs)
 	: GraphPtr(InGraph)
 	, WeakRigVMClientHost(InRigVMClientHost)
@@ -150,7 +153,7 @@ void FRigVMFunctionArgumentGroupLayout::HandleModifiedEvent(ERigVMGraphNotifType
 class FRigVMFunctionArgumentPinTypeSelectorFilter : public IPinTypeSelectorFilter
 {
 public:
-	FRigVMFunctionArgumentPinTypeSelectorFilter(const TWeakPtr<FRigVMEditor>& InRigVMEditor, const TWeakObjectPtr<URigVMGraph>& InGraph)
+	FRigVMFunctionArgumentPinTypeSelectorFilter(const TWeakPtr<FRigVMEditorBase>& InRigVMEditor, const TWeakObjectPtr<URigVMGraph>& InGraph)
 		: RigVMEditorPtr(InRigVMEditor), GraphPtr(InGraph)
 	{
 	}
@@ -205,7 +208,7 @@ public:
 
 private:
 
-	TWeakPtr<FRigVMEditor> RigVMEditorPtr;
+	TWeakPtr<FRigVMEditorBase> RigVMEditorPtr;
 	
 	TWeakObjectPtr<URigVMGraph> GraphPtr;
 };
@@ -725,16 +728,26 @@ void FRigVMFunctionArgumentDefaultNode::HandleModifiedEvent(ERigVMGraphNotifType
 	}
 }
 
-TSharedPtr<IDetailCustomization> FRigVMGraphDetailCustomization::MakeInstance(TSharedPtr<IBlueprintEditor> InBlueprintEditor, const UClass* InExpectedBlueprintClass)
+#if WITH_RIGVMLEGACYEDITOR
+TSharedPtr<IDetailCustomization> FRigVMGraphDetailCustomization::MakeInstance(TSharedPtr<IBlueprintEditor> InEditor, const UClass* InExpectedBlueprintClass)
+#else
+TSharedPtr<IDetailCustomization> FRigVMGraphDetailCustomization::MakeInstance(TSharedPtr<FRigVMEditorBase> InEditor, const UClass* InExpectedBlueprintClass)
+#endif
 {
-	const TArray<UObject*>* Objects = (InBlueprintEditor.IsValid() ? InBlueprintEditor->GetObjectsCurrentlyBeingEdited() : nullptr);
+	const TArray<UObject*>* Objects = (InEditor.IsValid() ? InEditor->GetObjectsCurrentlyBeingEdited() : nullptr);
 	if (Objects && Objects->Num() == 1)
 	{
 		if (URigVMBlueprint* RigVMBlueprint = Cast<URigVMBlueprint>((*Objects)[0]))
 		{
 			if(RigVMBlueprint->GetClass() == InExpectedBlueprintClass)
 			{
-				return MakeShareable(new FRigVMGraphDetailCustomization(StaticCastSharedPtr<FRigVMEditor>(InBlueprintEditor), RigVMBlueprint));
+#if WITH_RIGVMLEGACYEDITOR
+				// TODO: need to change how the Kismet Inspector calls the customization with a ptr to IBlueprintEditor
+				TSharedPtr<FRigVMLegacyEditor> SharedPtr = StaticCastSharedPtr<FRigVMLegacyEditor>(InEditor);
+				return MakeShareable(new FRigVMGraphDetailCustomization(SharedPtr, RigVMBlueprint));
+#else
+				return MakeShareable(new FRigVMGraphDetailCustomization(InEditor, RigVMBlueprint));
+#endif
 			}
 		}
 	}
@@ -742,7 +755,7 @@ TSharedPtr<IDetailCustomization> FRigVMGraphDetailCustomization::MakeInstance(TS
 	return nullptr;
 }
 
-FRigVMGraphDetailCustomization::FRigVMGraphDetailCustomization(TSharedPtr<FRigVMEditor> RigVMigEditor, URigVMBlueprint* RigVMBlueprint)
+FRigVMGraphDetailCustomization::FRigVMGraphDetailCustomization(TSharedPtr<FRigVMEditorBase> RigVMigEditor, URigVMBlueprint* RigVMBlueprint)
 	: RigVMEditorPtr(RigVMigEditor)
 	, RigVMBlueprintPtr(RigVMBlueprint)
 	, RigVMGraphDetailCustomizationImpl(MakeShared<FRigVMGraphDetailCustomizationImpl>())
@@ -785,7 +798,7 @@ void FRigVMGraphDetailCustomizationImpl::CustomizeDetails(IDetailLayoutBuilder&
 	URigVMGraph* InModel,
 	URigVMController* InController,
 	IRigVMClientHost* InRigVMClientHost,
-	TWeakPtr<FRigVMEditor> InEditor)
+	TWeakPtr<FRigVMEditorBase> InEditor)
 {
 	WeakModel = InModel;
 	WeakController = InController;
@@ -1486,7 +1499,7 @@ void FRigVMGraphDetailCustomizationImpl::OnBrowseVariantRef(const FRigVMVariantR
 			
 				if(IAssetEditorInstance* Editor = GEditor->GetEditorSubsystem<UAssetEditorSubsystem>()->FindEditorForAsset(Blueprint, true))
 				{
-					if(FRigVMEditor* RigVMEditor = static_cast<FRigVMEditor*>(Editor))
+					if (FRigVMEditorBase* RigVMEditor = FRigVMEditorBase::GetFromAssetEditorInstance(Editor))
 					{
 						RigVMEditor->HandleJumpToHyperlink(LibraryNode);
 					}
@@ -2006,7 +2019,7 @@ void FRigVMWrappedNodeDetailCustomization::CustomizeDetails(IDetailLayoutBuilder
 			
 								if(IAssetEditorInstance* Editor = GEditor->GetEditorSubsystem<UAssetEditorSubsystem>()->FindEditorForAsset(Blueprint, true))
 								{
-									if(FRigVMEditor* RigVMEditor = static_cast<FRigVMEditor*>(Editor))
+									if (FRigVMEditorBase* RigVMEditor = FRigVMEditorBase::GetFromAssetEditorInstance(Editor))
 									{
 										RigVMEditor->HandleJumpToHyperlink(LibraryNode);
 										return FReply::Handled();
diff --git a/Engine/Plugins/Runtime/RigVM/Source/RigVMEditor/Private/Editor/RigVMGraphEditorSummoner.cpp b/Engine/Plugins/Runtime/RigVM/Source/RigVMEditor/Private/Editor/RigVMGraphEditorSummoner.cpp
new file mode 100644
--- /dev/null
+++ b/Engine/Plugins/Runtime/RigVM/Source/RigVMEditor/Private/Editor/RigVMGraphEditorSummoner.cpp
@@ -0,0 +1,63 @@
+// Copyright Epic Games, Inc. All Rights Reserved.
+
+#include "Editor/RigVMGraphEditorSummoner.h"
+#include "Editor/RigVMNewEditor.h"
+#include "GraphEditor.h"
+#include "Widgets/Docking/SDockTab.h"
+
+#define LOCTEXT_NAMESPACE "RigVMGraphEditorSummoner"
+
+FText FRigVMLocalKismetCallbacks::GetGraphDisplayName(const UEdGraph* Graph)
+{
+	if (Graph)
+	{
+		if (const UEdGraphSchema* Schema = Graph->GetSchema())
+		{
+			FGraphDisplayInfo Info;
+			Schema->GetGraphDisplayInformation(*Graph, /*out*/ Info);
+
+			return Info.DisplayName;
+		}
+		else
+		{
+			// if we don't have a schema, we're dealing with a malformed (or incomplete graph)...
+			// possibly in the midst of some transaction - here we return the object's outer path 
+			// so we can at least get some context as to which graph we're referring
+			return FText::FromString(Graph->GetPathName());
+		}
+	}
+
+	return LOCTEXT("UnknownGraphName", "UNKNOWN");
+}
+
+FRigVMGraphEditorSummoner::FRigVMGraphEditorSummoner(TSharedPtr<FRigVMNewEditor> InEditorPtr, FOnCreateGraphEditorWidget CreateGraphEditorWidgetCallback)
+: FDocumentTabFactoryForObjects<UEdGraph>(TabID, InEditorPtr)
+, BlueprintEditorPtr(InEditorPtr)
+, OnCreateGraphEditorWidget(CreateGraphEditorWidgetCallback)
+{
+}
+
+void FRigVMGraphEditorSummoner::OnTabActivated(TSharedPtr<SDockTab> Tab) const
+{
+	TSharedRef<SGraphEditor> GraphEditor = StaticCastSharedRef<SGraphEditor>(Tab->GetContent());
+	BlueprintEditorPtr.Pin()->OnGraphEditorFocused(GraphEditor);
+}
+
+void FRigVMGraphEditorSummoner::OnTabBackgrounded(TSharedPtr<SDockTab> Tab) const
+{
+	TSharedRef<SGraphEditor> GraphEditor = StaticCastSharedRef<SGraphEditor>(Tab->GetContent());
+	BlueprintEditorPtr.Pin()->OnGraphEditorBackgrounded(GraphEditor);
+}
+
+TSharedRef<SWidget> FRigVMGraphEditorSummoner::CreateTabBodyForObject(const FWorkflowTabSpawnInfo& Info, UEdGraph* DocumentID) const
+{
+	check(Info.TabInfo.IsValid());
+	return OnCreateGraphEditorWidget.Execute(Info.TabInfo.ToSharedRef(), DocumentID);
+}
+
+const FSlateBrush* FRigVMGraphEditorSummoner::GetTabIconForObject(const FWorkflowTabSpawnInfo& Info, UEdGraph* DocumentID) const
+{
+	return FRigVMNewEditor::GetGlyphForGraph(DocumentID, false);
+}
+
+#undef LOCTEXT_NAMESPACE
\ No newline at end of file
diff --git a/Engine/Plugins/Runtime/RigVM/Source/RigVMEditor/Private/Editor/RigVMNewEditor.cpp b/Engine/Plugins/Runtime/RigVM/Source/RigVMEditor/Private/Editor/RigVMNewEditor.cpp
new file mode 100644
--- /dev/null
+++ b/Engine/Plugins/Runtime/RigVM/Source/RigVMEditor/Private/Editor/RigVMNewEditor.cpp
@@ -0,0 +1,2163 @@
+// Copyright Epic Games, Inc. All Rights Reserved.
+
+#include "Editor/RigVMNewEditor.h"
+
+#include "BlueprintActionDatabase.h"
+#include "BlueprintEditorTabs.h"
+#include "IMessageLogListing.h"
+#include "K2Node_Composite.h"
+#include "RigVMSettings.h"
+#include "SGraphPanel.h"
+#include "SRigVMGraphTitleBar.h"
+#include "UnrealEdGlobals.h"
+#include "Editor/SRigVMActionMenu.h"
+#include "Editor/UnrealEdEngine.h"
+#include "Framework/Commands/GenericCommands.h"
+#include "Framework/Notifications/NotificationManager.h"
+#include "Kismet2/BlueprintEditorUtils.h"
+#include "Kismet2/DebuggerCommands.h"
+#include "Kismet2/KismetDebugUtilities.h"
+#include "Widgets/Notifications/SNotificationList.h"
+#include "UObject/MetaData.h"
+#include "Widgets/Docking/SDockTab.h"
+#include "ScopedTransaction.h"
+#include "Framework/Application/SlateApplication.h"
+#include "Engine/BlueprintGeneratedClass.h"
+#include "Stats/StatsHierarchical.h"
+#include "Logging/MessageLog.h"
+#include "GraphEditor.h"
+#include "Editor/RigVMGraphEditorSummoner.h"
+#include "Widgets/Input/SEditableTextBox.h"
+#if WITH_RIGVMLEGACYEDITOR
+#include "SKismetInspector.h"
+#else
+#include "Editor/SRigVMDetailsInspector.h"
+#endif
+
+
+#define LOCTEXT_NAMESPACE "RigVMNewEditor"
+
+const FName FRigVMNewEditorTabs::DetailsID() {			static FName ID = TEXT( "Inspector" ); return ID; };
+const FName FRigVMNewEditorTabs::CompilerResultsID() {	static FName ID = TEXT( "CompilerResults" ); return ID; };
+
+const FName FRigVMNewEditor::SelectionState_GraphExplorer() {	static FName State = (TEXT("GraphExplorer")); return State; };
+const FName FRigVMNewEditor::SelectionState_Graph() {			static FName State = (TEXT("Graph")); return State; };
+const FName FRigVMNewEditor::SelectionState_ClassSettings() {	static FName State = (TEXT("ClassSettings")); return State; };
+const FName FRigVMNewEditor::SelectionState_ClassDefaults() {	static FName State = (TEXT("ClassDefaults")); return State; };
+
+namespace RigVMNewEditorImpl
+{
+	static const float InstructionFadeDuration = 0.5f;
+	
+	/**
+	 * Searches through a blueprint, looking for the most severe error'ing node.
+	 * 
+	 * @param  Blueprint	The blueprint to search through.
+	 * @param  Severity		Defines the severity of the error/warning to search for.
+	 * @return The first node found with the specified error.
+	 */
+	static UEdGraphNode* FindNodeWithError(UBlueprint* Blueprint, EMessageSeverity::Type Severity = EMessageSeverity::Error);
+
+	/**
+	 * Searches through an error log, looking for the most severe error'ing node.
+	 * 
+	 * @param  ErrorLog		The error log you want to search through.
+	 * @param  Severity		Defines the severity of the error/warning to search for.
+	 * @return The first node found with the specified error.
+	 */
+	static UEdGraphNode* FindNodeWithError(FCompilerResultsLog const& ErrorLog, EMessageSeverity::Type Severity = EMessageSeverity::Error);
+
+	/**
+	 * Utility function that will check to see if the specified graph has any
+	 * nodes that were default, pre-placed, in the graph.
+	 *
+	 * @param  InGraph  The graph to check.
+	 * @return True if the graph has any pre-placed nodes, otherwise false.
+	 */
+	static bool GraphHasDefaultNode(UEdGraph const* InGraph);
+
+	/**
+	 * Utility function that will check to see if the specified graph has any 
+	 * nodes other than those that come default, pre-placed, in the graph.
+	 *
+	 * @param  InGraph  The graph to check.
+	 * @return True if the graph has any nodes added by the user, otherwise false.
+	 */
+	static bool GraphHasUserPlacedNodes(UEdGraph const* InGraph);
+}
+
+static UEdGraphNode* RigVMNewEditorImpl::FindNodeWithError(UBlueprint* Blueprint, EMessageSeverity::Type Severity/* = EMessageSeverity::Error*/)
+{
+	TArray<UEdGraph*> Graphs;
+	Blueprint->GetAllGraphs(Graphs);
+
+	UEdGraphNode* ChoiceNode = nullptr;
+	for (UEdGraph* Graph : Graphs)
+	{
+		for (UEdGraphNode* Node : Graph->Nodes)
+		{
+			if (Node && Node->bHasCompilerMessage && !Node->ErrorMsg.IsEmpty() && (Node->ErrorType <= Severity))
+			{
+				if ((ChoiceNode == nullptr) || (ChoiceNode->ErrorType > Node->ErrorType))
+				{
+					ChoiceNode = Node;
+					if (ChoiceNode->ErrorType == 0)
+					{
+						break;
+					}
+				}
+			}
+		}
+	}
+	return ChoiceNode;
+}
+
+static UEdGraphNode* RigVMNewEditorImpl::FindNodeWithError(FCompilerResultsLog const& ErrorLog, EMessageSeverity::Type Severity/* = EMessageSeverity::Error*/)
+{
+	UEdGraphNode* ChoiceNode = nullptr;
+	for (TWeakObjectPtr<UEdGraphNode> NodePtr : ErrorLog.AnnotatedNodes)
+	{
+		UEdGraphNode* Node = NodePtr.Get();
+		if ((Node != nullptr) && (Node->ErrorType <= Severity))
+		{
+			if ((ChoiceNode == nullptr) || (Node->ErrorType < ChoiceNode->ErrorType))
+			{
+				ChoiceNode = Node;
+				if (ChoiceNode->ErrorType == 0)
+				{
+					break;
+				}
+			}
+		}
+	}
+
+	return ChoiceNode;
+}
+
+static bool RigVMNewEditorImpl::GraphHasDefaultNode(UEdGraph const* InGraph)
+{
+	bool bHasDefaultNodes = false;
+
+	for (UEdGraphNode const* Node : InGraph->Nodes)
+	{
+		if (Node == nullptr)
+		{
+			continue;
+		}
+
+		if (Node->GetOutermost()->GetMetaData()->HasValue(Node, FNodeMetadata::DefaultGraphNode) && Node->IsNodeEnabled())
+		{
+			bHasDefaultNodes = true;
+			break;
+		}
+	}
+
+	return bHasDefaultNodes;
+}
+
+static bool RigVMNewEditorImpl::GraphHasUserPlacedNodes(UEdGraph const* InGraph)
+{
+	bool bHasUserPlacedNodes = false;
+
+	for (UEdGraphNode const* Node : InGraph->Nodes)
+	{
+		if (Node == nullptr)
+		{
+			continue;
+		}
+
+		if (!Node->GetOutermost()->GetMetaData()->HasValue(Node, FNodeMetadata::DefaultGraphNode))
+		{
+			bHasUserPlacedNodes = true;
+			break;
+		}
+	}
+
+	return bHasUserPlacedNodes;
+}
+
+FRigVMNewEditor::FRigVMNewEditor()
+	: FRigVMEditorBase()
+	, bIsActionMenuContextSensitive(true)
+{
+	DocumentManager = MakeShareable(new FDocumentTracker);
+}
+
+void FRigVMNewEditor::OnClose()
+{
+	FRigVMEditorBase::UnbindEditor();
+	FWorkflowCentricApplication::OnClose();
+}
+
+void FRigVMNewEditor::InitAssetEditor(const EToolkitMode::Type Mode, const TSharedPtr<IToolkitHost>& InitToolkitHost, const FName AppIdentifier,
+                                      const TSharedRef<FTabManager::FLayout>& StandaloneDefaultLayout, const bool bCreateDefaultStandaloneMenu, const bool bCreateDefaultToolbar,
+                                      const TArray<UObject*>& ObjectsToEdit, const bool bInIsToolbarFocusable, const bool bInUseSmallToolbarIcons,
+                                      const TOptional<EAssetOpenMethod>& InOpenMethod)
+{
+	FWorkflowCentricApplication::InitAssetEditor(Mode, InitToolkitHost, AppIdentifier, StandaloneDefaultLayout, bCreateDefaultStandaloneMenu, bCreateDefaultToolbar, ObjectsToEdit, bInIsToolbarFocusable, bInUseSmallToolbarIcons, InOpenMethod);
+}
+
+void FRigVMNewEditor::CommonInitialization(const TArray<UBlueprint*>& InitBlueprints, bool bShouldOpenInDefaultsMode)
+{
+	TSharedPtr<FRigVMNewEditor> ThisPtr = SharedThis(this);
+
+	// @todo TabManagement
+	DocumentManager->Initialize(ThisPtr);
+
+	// Register the document factories
+	{
+		// TODO sara-s
+		//DocumentManager->RegisterDocumentFactory(MakeShareable(new FTimelineEditorSummoner(ThisPtr)));
+
+		TSharedRef<FRigVMNewEditor> SharedRef = StaticCastSharedRef<FRigVMNewEditor>(SharedThis(this));
+		TSharedRef<FDocumentTabFactory> GraphEditorFactory = MakeShareable(new FRigVMGraphEditorSummoner(ThisPtr,
+			FRigVMGraphEditorSummoner::FOnCreateGraphEditorWidget::CreateSP(SharedRef, &FRigVMNewEditor::CreateGraphEditorWidget)
+			));
+
+		// Also store off a reference to the grapheditor factory so we can find all the tabs spawned by it later.
+		GraphEditorTabFactoryPtr = GraphEditorFactory;
+		DocumentManager->RegisterDocumentFactory(GraphEditorFactory);
+	}
+
+	// Create a namespace helper to keep track of imports for all BPs being edited.
+	//ImportedNamespaceHelper = MakeShared<FBlueprintNamespaceHelper>();
+
+	// Add each Blueprint instance to be edited into the namespace helper's context.
+	// for (const UBlueprint* BP : InitBlueprints)
+	// {
+	// 	ImportedNamespaceHelper->AddBlueprint(BP);
+	// }
+
+	// Create imported namespace type filters for value editing.
+	// ImportedClassViewerFilter = MakeShared<BlueprintEditorImpl::FImportedClassViewerFilterProxy>(ImportedNamespaceHelper->GetClassViewerFilter());
+	// ImportedPinTypeSelectorFilter = MakeShared<BlueprintEditorImpl::FImportedPinTypeSelectorFilterProxy>(ImportedNamespaceHelper->GetPinTypeSelectorFilter());
+	// PermissionsPinTypeSelectorFilter = MakeShared<BlueprintEditorImpl::FPermissionsPinTypeSelectorFilter>(InitBlueprints);
+
+	// Make sure we know when tabs become active to update details tab
+	//OnActiveTabChangedDelegateHandle = FGlobalTabmanager::Get()->OnActiveTabChanged_Subscribe( FOnActiveTabChanged::FDelegate::CreateRaw(this, &FBlueprintEditor::OnActiveTabChanged) );
+
+	if (InitBlueprints.Num() == 1)
+	{
+		if (!bShouldOpenInDefaultsMode)
+		{
+			// Load blueprint libraries
+			// if (ShouldLoadBPLibrariesFromAssetRegistry())
+			// {
+			// 	LoadLibrariesFromAssetRegistry();
+			// }
+
+			// Init the action DB for the context menu/palette if not already constructed
+			FBlueprintActionDatabase::Get();
+		}
+
+		//FLoadObjectsFromAssetRegistryHelper::Load<UUserDefinedEnum>(UserDefinedEnumerators);
+
+		UBlueprint* InitBlueprint = InitBlueprints[0];
+
+		// Update the blueprint if required
+		EBlueprintStatus OldStatus = InitBlueprint->Status;
+		//EnsureBlueprintIsUpToDate(InitBlueprint); // TODO sara-s
+		UPackage* BpPackage = InitBlueprint->GetOutermost();
+		bBlueprintModifiedOnOpen = (InitBlueprint->Status != OldStatus) && !BpPackage->HasAnyPackageFlags(PKG_NewlyCreated);
+
+		// Flag the blueprint as having been opened
+		InitBlueprint->bIsNewlyCreated = false;
+
+		// When the blueprint that we are observing changes, it will notify this wrapper widget.
+		InitBlueprint->OnChanged().AddSP(this, &FRigVMNewEditor::OnBlueprintChanged); // TODO sara-s
+		// InitBlueprint->OnCompiled().AddSP(this, &FRigVMNewEditor::OnBlueprintCompiled); // TODO sara-s
+		InitBlueprint->OnSetObjectBeingDebugged().AddSP(this, &FRigVMNewEditor::HandleSetObjectBeingDebugged);
+	}
+
+	bWasOpenedInDefaultsMode = bShouldOpenInDefaultsMode;
+
+	CreateDefaultTabContents(InitBlueprints);
+
+	// FCoreUObjectDelegates::OnPreObjectPropertyChanged.AddSP(this, &FRigVMNewEditor::OnPreObjectPropertyChanged); // TODO sara-s
+	// FCoreUObjectDelegates::OnObjectPropertyChanged.AddSP(this, &FRigVMNewEditor::OnPostObjectPropertyChanged); // TODO sara-s
+
+	// FKismetEditorUtilities::OnBlueprintUnloaded.AddSP(this, &FRigVMNewEditor::OnBlueprintUnloaded); // TODO sara-s
+}
+
+void FRigVMNewEditor::OnBlueprintChanged(UBlueprint* InBlueprint)
+{
+	if (InBlueprint)
+	{
+		// Notify that the blueprint has been changed (update Content browser, etc)
+		InBlueprint->PostEditChange();
+
+		// Call PostEditChange() on any Actors that are based on this Blueprint 
+		FBlueprintEditorUtils::PostEditChangeBlueprintActors(InBlueprint);
+
+		// Refresh the graphs
+		//TODO sara-s
+		ERefreshRigVMEditorReason::Type Reason = ERefreshRigVMEditorReason::UnknownReason;// = bIsJustBeingCompiled ? ERefreshBlueprintEditorReason::BlueprintCompiled : ERefreshBlueprintEditorReason::UnknownReason;
+		RefreshEditors(Reason);
+
+		// In case objects were deleted, which should close the tab
+		if (GetCurrentMode() == FRigVMNewEditorApplicationModes::StandardRigVMEditorMode())
+		{
+			SaveEditedObjectState();
+		}
+	}
+}
+
+void FRigVMNewEditor::SaveEditedObjectState()
+{
+	check(IsEditingSingleBlueprint());
+
+	// Clear currently edited documents
+	GetBlueprintObj()->LastEditedDocuments.Empty();
+
+	// Ask all open documents to save their state, which will update LastEditedDocuments
+	DocumentManager->SaveAllState();
+}
+
+void FRigVMNewEditor::SetCurrentMode(FName NewMode)
+{
+	// Clear the selection state when the mode changes.
+	SetUISelectionState(NAME_None);
+
+	// TODO sara-s: Replicate what blueprint editor is doing
+	//OnModeSetData.Broadcast( NewMode );
+	
+	FWorkflowCentricApplication::SetCurrentMode(NewMode);
+}
+
+void FRigVMNewEditor::PostLayoutBlueprintEditorInitialization()
+{
+	// TODO sara-s: Replicate what blueprint editor is doing
+	if (UBlueprint* Blueprint = GetBlueprintObj())
+	{
+		// Refresh the graphs
+		RefreshEditors();
+		
+		// EnsureBlueprintIsUpToDate may have updated the blueprint so show notifications to user.
+		if (bBlueprintModifiedOnOpen)
+		{
+			bBlueprintModifiedOnOpen = false;
+
+			if (FocusedGraphEdPtr.IsValid())
+			{
+				FNotificationInfo Info( NSLOCTEXT("Kismet", "Blueprint Modified", "Blueprint requires updating. Please resave.") );
+				Info.Image = FAppStyle::GetBrush(TEXT("Icons.Info"));
+				Info.bFireAndForget = true;
+				Info.bUseSuccessFailIcons = false;
+				Info.ExpireDuration = 5.0f;
+
+				FocusedGraphEdPtr.Pin()->AddNotification(Info, true);
+			}
+
+			// Fire log message
+			FString BlueprintName;
+			Blueprint->GetName(BlueprintName);
+
+			FFormatNamedArguments Args;
+			Args.Add( TEXT("BlueprintName"), FText::FromString( BlueprintName ) );
+			LogSimpleMessage( FText::Format( LOCTEXT("Blueprint Modified Long", "Blueprint \"{BlueprintName}\" was updated to fix issues detected on load. Please resave."), Args ) );
+		}
+
+		// Determine if the current "mode" supports invoking the Compiler Results tab.
+#if WITH_RIGVMLEGACYEDITOR
+		const bool bCanInvokeCompilerResultsTab = TabManager->HasTabSpawner(FBlueprintEditorTabs::CompilerResultsID);
+#else
+		const bool bCanInvokeCompilerResultsTab = TabManager->HasTabSpawner(FRigVMNewEditorTabs::CompilerResultsID());
+#endif
+
+		// If we have a warning/error, open output log if the current mode allows us to invoke it.
+		const bool bIsBlueprintInWarningOrErrorState = !Blueprint->IsUpToDate() || (Blueprint->Status == BS_UpToDateWithWarnings);
+		if (bIsBlueprintInWarningOrErrorState && bCanInvokeCompilerResultsTab)
+		{
+#if WITH_RIGVMLEGACYEDITOR
+			TabManager->TryInvokeTab(FBlueprintEditorTabs::CompilerResultsID);
+#else
+			TabManager->TryInvokeTab(FRigVMNewEditorTabs::CompilerResultsID());
+#endif
+		}
+		else
+		{
+			// Toolkit modes that don't include this tab may have been incorrectly saved with layout information for restoring it
+			// as an "unrecognized" tab, due to having previously invoked it above without checking to see if the layout can open
+			// it first. To correct this, we check if the tab was restored from a saved layout here, and close it if not supported.
+#if WITH_RIGVMLEGACYEDITOR
+			TSharedPtr<SDockTab> TabPtr = TabManager->FindExistingLiveTab(FBlueprintEditorTabs::CompilerResultsID);
+#else
+			TSharedPtr<SDockTab> TabPtr = TabManager->FindExistingLiveTab(FRigVMNewEditorTabs::CompilerResultsID());
+#endif
+			if (TabPtr.IsValid() && !bCanInvokeCompilerResultsTab)
+			{
+				TabPtr->RequestCloseTab();
+			}
+		}
+	}
+}
+
+TSharedPtr<SGraphEditor> FRigVMNewEditor::OpenGraphAndBringToFront(UEdGraph* Graph, bool bSetFocus)
+{
+	if (!IsValid(Graph))
+	{
+		return TSharedPtr<SGraphEditor>();
+	}
+
+	// First, switch back to standard mode
+	SetCurrentMode(FRigVMNewEditorApplicationModes::StandardRigVMEditorMode());
+
+	// This will either reuse an existing tab or spawn a new one
+	TSharedPtr<SDockTab> TabWithGraph = OpenDocument(Graph, FDocumentTracker::OpenNewDocument);
+	if (TabWithGraph.IsValid())
+	{
+
+		// We know that the contents of the opened tabs will be a graph editor.
+		TSharedRef<SGraphEditor> NewGraphEditor = StaticCastSharedRef<SGraphEditor>(TabWithGraph->GetContent());
+
+		// Handover the keyboard focus to the new graph editor widget.
+		if (bSetFocus)
+		{
+			NewGraphEditor->CaptureKeyboard();
+		}
+
+		return NewGraphEditor;
+	}
+	else
+	{
+		return TSharedPtr<SGraphEditor>();
+	}
+}
+
+bool FRigVMNewEditor::FindOpenTabsContainingDocument(const UObject* DocumentID, TArray<TSharedPtr<SDockTab>>& Results)
+{
+	// TODO sara-s: Replicate what blueprint editor is doing
+	return false;
+}
+
+TSharedPtr<SDockTab> FRigVMNewEditor::OpenDocument(const UObject* DocumentID, FDocumentTracker::EOpenDocumentCause Cause)
+{
+	TSharedRef<FTabPayload_UObject> Payload = FTabPayload_UObject::Make(DocumentID);
+	return DocumentManager->OpenDocument(Payload, Cause);
+}
+
+void FRigVMNewEditor::CloseDocumentTab(const UObject* DocumentID)
+{
+	TSharedRef<FTabPayload_UObject> Payload = FTabPayload_UObject::Make(DocumentID);
+	DocumentManager->CloseTab(Payload);
+}
+
+TSharedPtr<FApplicationMode> FRigVMNewEditor::CreateEditorMode()
+{
+	return MakeShareable(new FRigVMNewEditorMode(SharedThis(this)));
+}
+
+const FName FRigVMNewEditor::GetEditorAppName() const
+{
+	static const FLazyName AppName(TEXT("RigVMNewEditorApp"));
+	return AppName;
+}
+
+void FRigVMNewEditor::OnGraphEditorFocused(const TSharedRef<SGraphEditor>& InGraphEditor)
+{
+	// Update the graph editor that is currently focused
+	FocusedGraphEdPtr = InGraphEditor;
+
+	// TODO sara-s
+	//InGraphEditor->SetPinVisibility(PinVisibility);
+
+	// Update the inspector as well, to show selection from the focused graph editor
+	FGraphPanelSelectionSet SelectedNodes = GetSelectedNodes();
+	FocusInspectorOnGraphSelection(SelectedNodes, /*bForceRefresh=*/ true);
+
+	// During undo, garbage graphs can be temporarily brought into focus, ensure that before a refresh of the MyBlueprint window that the graph is owned by a Blueprint
+	if ( FocusedGraphEdPtr.IsValid() && GraphExplorerWidget.IsValid() )
+	{
+		// The focused graph can be garbage as well
+		TWeakObjectPtr< UEdGraph > FocusedGraphPtr = FocusedGraphEdPtr.Pin()->GetCurrentGraph();
+		UEdGraph* FocusedGraph = FocusedGraphPtr.Get();
+
+		if ( FocusedGraph != nullptr )
+		{
+			if ( FBlueprintEditorUtils::FindBlueprintForGraph(FocusedGraph) )
+			{
+				GraphExplorerWidget->Refresh();
+			}
+		}
+	}
+
+	// TODO sara-s
+	// if (bHideUnrelatedNodes && SelectedNodes.Num() <= 0)
+	// {
+	// 	ResetAllNodesUnrelatedStates();
+	// }
+
+	// TODO sara-sc
+	// If the bookmarks view is active, check whether or not we're restricting the view to the current graph. If we are, update the tree to reflect the focused graph context.
+	// if (BookmarksWidget.IsValid()
+	// 	&& GetDefault<UBlueprintEditorSettings>()->bShowBookmarksForCurrentDocumentOnlyInTab)
+	// {
+	// 	BookmarksWidget->RefreshBookmarksTree();
+	// }
+
+	FRigVMEditorBase::OnGraphEditorFocused(InGraphEditor);
+}
+
+void FRigVMNewEditor::OnGraphEditorBackgrounded(const TSharedRef<SGraphEditor>& InGraphEditor)
+{
+	// If the newly active document tab isn't a graph we want to make sure we clear the focused graph pointer.
+	// Several other UI reads that, like the MyBlueprints view uses it to determine if it should show the "Local Variable" section.
+	FocusedGraphEdPtr = nullptr;
+
+	if ( GraphExplorerWidget.IsValid() == true )
+	{
+		GraphExplorerWidget->Refresh();
+	}
+}
+
+void FRigVMNewEditor::RegisterTabSpawners(const TSharedRef<FTabManager>& InTabManager)
+{
+	//@TODO: Can't we do this sooner?
+	DocumentManager->SetTabManager(InTabManager);
+
+	FWorkflowCentricApplication::RegisterTabSpawners(InTabManager);
+}
+
+void FRigVMNewEditor::Tick(float DeltaTime)
+{
+	// TODO sara-s: Replicate what blueprint editor is doing
+	FRigVMEditorBase::Tick(DeltaTime);
+}
+
+void FRigVMNewEditor::GetPinTypeSelectorFilters(TArray<TSharedPtr<IPinTypeSelectorFilter>>& OutFilters) const
+{
+	// TODO sara-s
+	// OutFilters.Add(ImportedPinTypeSelectorFilter);
+	// OutFilters.Add(PermissionsPinTypeSelectorFilter);
+}
+
+void FRigVMNewEditor::OnAddNewVariable()
+{
+	const FScopedTransaction Transaction( LOCTEXT("AddVariable", "Add Variable") );
+
+	FName VarName = FBlueprintEditorUtils::FindUniqueKismetName(GetBlueprintObj(), TEXT("NewVar"));
+
+	bool bSuccess = GraphExplorerWidget.IsValid() && FBlueprintEditorUtils::AddMemberVariable(GetBlueprintObj(), VarName, GraphExplorerWidget->GetLastPinTypeUsed());
+
+	if(!bSuccess)
+	{
+		LogSimpleMessage( LOCTEXT("AddVariable_Error", "Adding new variable failed.") );
+	}
+	else
+	{
+		RenameNewlyAddedAction(VarName);
+	}
+}
+
+void FRigVMNewEditor::SummonSearchUI(bool bSetFindWithinBlueprint, FString NewSearchTerms, bool bSelectFirstResult)
+{
+	// TSharedPtr<SFindInBlueprints> FindResultsToUse = nullptr;
+
+	// TODO: sara-s
+	// if (bSetFindWithinBlueprint)
+	// {
+	// 	FindResultsToUse = FindResults;
+	// 	TabManager->TryInvokeTab(FRigVMNewEditorTabs::FindResultsID);
+	// }
+	// else
+	// {
+	// 	FindResultsToUse = FFindInBlueprintSearchManager::Get().GetGlobalFindResults();
+	// }
+	//
+	// if (FindResultsToUse.IsValid())
+	// {
+	// 	FindResultsToUse->FocusForUse(bSetFindWithinBlueprint, NewSearchTerms, bSelectFirstResult);
+	// }
+}
+
+void FRigVMNewEditor::ZoomToSelection_Clicked()
+{
+	if (SGraphEditor* GraphEd = FocusedGraphEdPtr.Pin().Get())
+	{
+		GraphEd->ZoomToFit(/*bOnlySelection=*/ true);
+	}
+}
+
+void FRigVMNewEditor::RestoreEditedObjectState()
+{
+	check(IsEditingSingleBlueprint());	
+
+	UBlueprint* Blueprint = GetBlueprintObj();
+	if (Blueprint->LastEditedDocuments.Num() == 0)
+	{
+		if (FBlueprintEditorUtils::SupportsConstructionScript(Blueprint))
+		{
+			Blueprint->LastEditedDocuments.Add(FBlueprintEditorUtils::FindUserConstructionScript(Blueprint));
+		}
+
+		if (Blueprint->SupportsEventGraphs())
+		{
+			Blueprint->LastEditedDocuments.Add(FBlueprintEditorUtils::FindEventGraph(Blueprint));
+		}
+	}
+
+	TSet<FSoftObjectPath> PathsToRemove;
+	for (int32 i = 0; i < Blueprint->LastEditedDocuments.Num(); i++)
+	{
+		if (UObject* Obj = Blueprint->LastEditedDocuments[i].EditedObjectPath.ResolveObject())
+		{
+			if(UEdGraph* Graph = Cast<UEdGraph>(Obj))
+			{
+				if (FBlueprintEditorUtils::IsEventGraph(Graph) && !Blueprint->SupportsEventGraphs())
+				{
+					continue;
+				}
+
+				struct LocalStruct
+				{
+					static TSharedPtr<SDockTab> OpenGraphTree(FRigVMNewEditor* InBlueprintEditor, UEdGraph* InGraph)
+					{
+						FDocumentTracker::EOpenDocumentCause OpenCause = FDocumentTracker::QuickNavigateCurrentDocument;
+
+						for (UObject* OuterObject = InGraph->GetOuter(); OuterObject; OuterObject = OuterObject->GetOuter())
+						{
+							if (OuterObject->IsA<UBlueprint>())
+							{
+								// reached up to the blueprint for the graph, we are done climbing the tree
+								OpenCause = FDocumentTracker::RestorePreviousDocument;
+								break;
+							}
+							else if(UEdGraph* OuterGraph = Cast<UEdGraph>(OuterObject))
+							{
+								// Found another graph, open it up
+								OpenGraphTree(InBlueprintEditor, OuterGraph);
+								break;
+							}
+						}
+
+						return InBlueprintEditor->OpenDocument(InGraph, OpenCause);
+					}
+				};
+				TSharedPtr<SDockTab> TabWithGraph = LocalStruct::OpenGraphTree(this, Graph);
+				if (TabWithGraph.IsValid())
+				{
+					TSharedRef<SGraphEditor> GraphEditor = StaticCastSharedRef<SGraphEditor>(TabWithGraph->GetContent());
+					GraphEditor->SetViewLocation(Blueprint->LastEditedDocuments[i].SavedViewOffset, Blueprint->LastEditedDocuments[i].SavedZoomAmount);
+				}
+			}
+			else
+			{
+				TSharedPtr<SDockTab> TabWithGraph = OpenDocument(Obj, FDocumentTracker::RestorePreviousDocument);
+			}
+		}
+		else
+		{
+			PathsToRemove.Add(Blueprint->LastEditedDocuments[i].EditedObjectPath);
+		}
+	}
+
+	// Older assets may have neglected to clean up this array when referenced objects were deleted, so
+	// we'll check for that now. This is done to ensure we don't store invalid object paths indefinitely.
+	if (PathsToRemove.Num() > 0)
+	{
+		Blueprint->LastEditedDocuments.RemoveAll([&PathsToRemove](const FEditedDocumentInfo& Entry)
+		{
+			return PathsToRemove.Contains(Entry.EditedObjectPath);
+		});
+	}
+}
+
+void FRigVMNewEditor::SetupViewForBlueprintEditingMode()
+{
+	// Make sure the defaults tab is pointing to the defaults
+	StartEditingDefaults(/*bAutoFocus=*/ true);
+
+	// Make sure the inspector is always on top
+	//@TODO: This is necessary right now because of a bug in restoring layouts not remembering which tab is on top (to get it right initially), but do we want this behavior always?
+	TryInvokingDetailsTab();
+}
+
+void FRigVMNewEditor::InitToolMenuContext(FToolMenuContext& MenuContext)
+{
+	// TODO sara-s: Replicate what blueprint editor is doing
+	FWorkflowCentricApplication::InitToolMenuContext(MenuContext);
+	FRigVMEditorBase::InitToolMenuContextImpl(MenuContext);
+}
+
+bool FRigVMNewEditor::TransactionObjectAffectsBlueprint(UObject* InTransactedObject)
+{
+	// TODO sara-s: Replicate what blueprint editor is doing
+	return FRigVMEditorBase::TransactionObjectAffectsBlueprintImpl(InTransactedObject);
+}
+
+FEdGraphPinType FRigVMNewEditor::GetLastPinTypeUsed()
+{
+	if (GraphExplorerWidget.IsValid())
+	{
+		return GraphExplorerWidget->GetLastPinTypeUsed();
+	}
+	return FEdGraphPinType();
+}
+
+void FRigVMNewEditor::LogSimpleMessage(const FText& MessageText)
+{
+	FNotificationInfo Info( MessageText );
+	Info.ExpireDuration = 3.0f;
+	Info.bUseLargeFont = false;
+	TSharedPtr<SNotificationItem> Notification = FSlateNotificationManager::Get().AddNotification(Info);
+	if ( Notification.IsValid() )
+	{
+		Notification->SetCompletionState( SNotificationItem::CS_Fail );
+	}
+}
+
+void FRigVMNewEditor::RenameNewlyAddedAction(FName InActionName)
+{
+	// TODO sara-s: Replicate what blueprint editor is doing
+}
+
+FGraphPanelSelectionSet FRigVMNewEditor::GetSelectedNodes() const
+{
+	FGraphPanelSelectionSet CurrentSelection;
+	TSharedPtr<SGraphEditor> FocusedGraphEd = FocusedGraphEdPtr.Pin();
+	if (FocusedGraphEd.IsValid())
+	{
+		CurrentSelection = FocusedGraphEd->GetSelectedNodes();
+	}
+	return CurrentSelection;
+}
+
+void FRigVMNewEditor::SetUISelectionState(FName SelectionOwner)
+{
+	if ( SelectionOwner != CurrentUISelection )
+	{
+		ClearSelectionStateFor(CurrentUISelection);
+	
+		CurrentUISelection = SelectionOwner;
+	}
+}
+
+void FRigVMNewEditor::AnalyticsTrackNodeEvent(UBlueprint* Blueprint, UEdGraphNode* GraphNode, bool bNodeDelete) const
+{
+	// TODO sara-s
+	// if(Blueprint && GraphNode && FEngineAnalytics::IsAvailable())
+	// {
+	// 	// we'd like to see if this was happening in normal blueprint editor or persona 
+	// 	//const FString EditorName = Cast<UAnimBlueprint>(Blueprint) != nullptr ? TEXT("Persona") : TEXT("BlueprintEditor");
+	// 	const FString EditorName = TEXT("RigVMNewEditor");
+	//
+	// 	// Build Node Details
+	// 	const UGeneralProjectSettings& ProjectSettings = *GetDefault<UGeneralProjectSettings>();
+	// 	FString ProjectID = ProjectSettings.ProjectID.ToString();
+	// 	TArray<FAnalyticsEventAttribute> NodeAttributes;
+	// 	NodeAttributes.Add(FAnalyticsEventAttribute(TEXT("ProjectId"), ProjectID));
+	// 	NodeAttributes.Add(FAnalyticsEventAttribute(TEXT("BlueprintId"), Blueprint->GetBlueprintGuid().ToString()));
+	// 	TArray<TKeyValuePair<FString, FString>> Attributes;
+	//
+	// 	if (UK2Node* K2Node = Cast<UK2Node>(GraphNode))
+	// 	{
+	// 		K2Node->GetNodeAttributes(Attributes);
+	// 	}
+	// 	else if (UEdGraphNode_Comment* CommentNode = Cast<UEdGraphNode_Comment>(GraphNode))
+	// 	{
+	// 		Attributes.Add(TKeyValuePair<FString, FString>(TEXT("Type"), TEXT("Comment")));
+	// 		Attributes.Add(TKeyValuePair<FString, FString>(TEXT("Class"), CommentNode->GetClass()->GetName()));
+	// 		Attributes.Add(TKeyValuePair<FString, FString>(TEXT("Name"), CommentNode->GetName()));
+	// 	}
+	// 	if (Attributes.Num() > 0)
+	// 	{
+	// 		// Build Node Attributes
+	// 		for (const TKeyValuePair<FString, FString>& Attribute : Attributes)
+	// 		{
+	// 			NodeAttributes.Add(FAnalyticsEventAttribute(Attribute.Key, Attribute.Value));
+	// 		}
+	// 		// Send Analytics event 
+	// 		FString EventType = bNodeDelete ?	FString::Printf(TEXT("Editor.Usage.%s.NodeDeleted"), *EditorName) :
+	// 											FString::Printf(TEXT("Editor.Usage.%s.NodeCreated"), *EditorName);
+	// 		FEngineAnalytics::GetProvider().RecordEvent(EventType, NodeAttributes);
+	// 	}
+	// }
+}
+
+void FRigVMNewEditor::JumpToHyperlink(const UObject* ObjectReference, bool bRequestRename)
+{
+	if (FRigVMEditorBase::JumpToHyperlinkImpl(ObjectReference, bRequestRename))
+	{
+		return;
+	}
+	
+	SetCurrentMode(FRigVMNewEditorApplicationModes::StandardRigVMEditorMode());
+
+	if (const UEdGraphNode* Node = Cast<const UEdGraphNode>(ObjectReference))
+	{
+		if (bRequestRename)
+		{
+			IsNodeTitleVisible(Node, bRequestRename);
+		}
+		else
+		{
+			JumpToNode(Node, false);
+		}
+	}
+	else if (const UEdGraph* Graph = Cast<const UEdGraph>(ObjectReference))
+	{
+		// Navigating into things should re-use the current tab when it makes sense
+		FDocumentTracker::EOpenDocumentCause OpenMode = FDocumentTracker::OpenNewDocument;
+		if ((Graph->GetSchema()->GetGraphType(Graph) == GT_Ubergraph) || Cast<UK2Node>(Graph->GetOuter()) || Cast<UEdGraph>(Graph->GetOuter()))
+		{
+			// Ubergraphs directly reuse the current graph
+			OpenMode = FDocumentTracker::NavigatingCurrentDocument;
+		}
+		else
+		{
+			// Walk up the outer chain to see if any tabs have a parent of this document open for edit, and if so
+			// we should reuse that one and drill in deeper instead
+			for (UObject* WalkPtr = const_cast<UEdGraph*>(Graph); WalkPtr != nullptr; WalkPtr = WalkPtr->GetOuter())
+			{
+				TArray< TSharedPtr<SDockTab> > TabResults;
+				if (FindOpenTabsContainingDocument(WalkPtr, /*out*/ TabResults))
+				{
+					// See if the parent was active
+					bool bIsActive = false;
+					for (TSharedPtr<SDockTab> Tab : TabResults)
+					{
+						if (Tab->IsActive())
+						{
+							bIsActive = true;
+							break;
+						}
+					}
+
+					if (bIsActive)
+					{
+						OpenMode = FDocumentTracker::NavigatingCurrentDocument;
+						break;
+					}
+				}
+			}
+		}
+
+		// Force it to open in a new document if shift is pressed
+		const bool bIsShiftPressed = FSlateApplication::Get().GetModifierKeys().IsShiftDown();
+		if (bIsShiftPressed)
+		{
+			OpenMode = FDocumentTracker::ForceOpenNewDocument;
+		}
+
+		// Open the document
+		OpenDocument(Graph, OpenMode);
+	}
+	else if (const AActor* ReferencedActor = Cast<const AActor>(ObjectReference))
+	{
+		// Check if the world is active in the editor. It's possible to open level BPs without formally opening
+		// the levels through Find-in-Blueprints
+		bool bInOpenWorld = false;
+		const TIndirectArray<FWorldContext>& WorldContextList = GEditor->GetWorldContexts();
+		const UWorld* ReferencedActorOwningWorld = ReferencedActor->GetWorld();
+		for (const FWorldContext& WorldContext : WorldContextList)
+		{
+			if (WorldContext.World() == ReferencedActorOwningWorld)
+			{
+				bInOpenWorld = true;
+				break;
+			}
+		}
+
+		// Clear the selection even if we couldn't find it, so the existing selection doesn't get mistaken for the desired to be selected actor
+		GEditor->SelectNone(false, false);
+
+		if (bInOpenWorld)
+		{
+			// Select the in-level actor
+			GEditor->SelectActor(const_cast<AActor*>(ReferencedActor), true, true, true);
+
+			// Point the camera at it
+			GUnrealEd->Exec(ReferencedActor->GetWorld(), TEXT("CAMERA ALIGN ACTIVEVIEWPORTONLY"));
+		}
+	}
+	else if(const UFunction* Function = Cast<const UFunction>(ObjectReference))
+	{
+		UBlueprint* BP = GetBlueprintObj();
+		if(BP)
+		{
+			if (UEdGraph* FunctionGraph = FBlueprintEditorUtils::FindScopeGraph(BP, Function))
+			{
+				OpenDocument(FunctionGraph, FDocumentTracker::OpenNewDocument);
+			}
+		}
+	}
+	else if(const UBlueprintGeneratedClass* Class = Cast<const UBlueprintGeneratedClass>(ObjectReference))
+	{
+		GEditor->GetEditorSubsystem<UAssetEditorSubsystem>()->OpenEditorForAsset(Class->ClassGeneratedBy);
+	}
+	// else if (const UTimelineTemplate* Timeline = Cast<const UTimelineTemplate>(ObjectReference))
+	// {
+	// 	OpenDocument(Timeline, FDocumentTracker::OpenNewDocument);
+	// }
+	else if ((ObjectReference != nullptr) && ObjectReference->IsAsset())
+	{
+		GEditor->GetEditorSubsystem<UAssetEditorSubsystem>()->OpenEditorForAsset(const_cast<UObject*>(ObjectReference));
+	}
+	else
+	{
+		UE_LOG(LogBlueprint, Warning, TEXT("Unknown type of hyperlinked object (%s), cannot focus it"), *GetNameSafe(ObjectReference));
+	}
+
+	//@TODO: Hacky way to ensure a message is seen when hitting an exception and doing intraframe debugging
+	const FText ExceptionMessage = FKismetDebugUtilities::GetAndClearLastExceptionMessage();
+	if (!ExceptionMessage.IsEmpty())
+	{
+		LogSimpleMessage( ExceptionMessage );
+	}
+}
+
+void FRigVMNewEditor::PostUndo(bool bSuccess)
+{
+	FEditorUndoClient::PostUndo(bSuccess);
+	FRigVMEditorBase::PostUndoImpl(bSuccess);
+}
+
+void FRigVMNewEditor::PostRedo(bool bSuccess)
+{
+	FEditorUndoClient::PostRedo(bSuccess);
+	FRigVMEditorBase::PostRedoImpl(bSuccess);
+}
+
+UEdGraphPin* FRigVMNewEditor::GetCurrentlySelectedPin() const
+{
+	TSharedPtr<SGraphEditor> FocusedGraphEd = FocusedGraphEdPtr.Pin();
+	if (FocusedGraphEd.IsValid())
+	{
+		return FocusedGraphEd->GetGraphPinForMenu();
+	}
+
+	return nullptr;
+}
+
+void FRigVMNewEditor::CreateDefaultCommands()
+{
+	// TODO sara-s: Replicate what blueprint editor is doing
+	if (GetBlueprintObj())
+	{
+		//FBlueprintEditor::CreateDefaultCommands();
+	}
+
+	FRigVMEditorBase::CreateDefaultCommandsImpl();
+}
+
+TSharedRef<SGraphEditor> FRigVMNewEditor::CreateGraphEditorWidget(TSharedRef<FTabInfo> InTabInfo, UEdGraph* InGraph)
+{
+	// TODO sara-s Populate the graph editor commands 
+
+	check((InGraph != nullptr) && IsEditingSingleBlueprint());
+
+	// No need to regenerate the commands.
+	if(!GraphEditorCommands.IsValid())
+	{
+		GraphEditorCommands = MakeShareable( new FUICommandList );
+		{
+			// GraphEditorCommands->MapAction(FGraphEditorCommands::Get().PromoteToVariable,
+			// 	FExecuteAction::CreateSP( this, &FBlueprintEditor::OnPromoteToVariable, true ),
+			// 	FCanExecuteAction::CreateSP( this, &FBlueprintEditor::CanPromoteToVariable, true )
+			// 	);
+			//
+			// GraphEditorCommands->MapAction(FGraphEditorCommands::Get().PromoteToLocalVariable,
+			// 	FExecuteAction::CreateSP( this, &FBlueprintEditor::OnPromoteToVariable, false ),
+			// 	FCanExecuteAction::CreateSP( this, &FBlueprintEditor::CanPromoteToVariable, false )
+			// 	);
+			//
+			// GraphEditorCommands->MapAction(FGraphEditorCommands::Get().SplitStructPin,
+			// 	FExecuteAction::CreateSP( this, &FBlueprintEditor::OnSplitStructPin ),
+			// 	FCanExecuteAction::CreateSP( this, &FBlueprintEditor::CanSplitStructPin )
+			// 	);
+			//
+			// GraphEditorCommands->MapAction(FGraphEditorCommands::Get().RecombineStructPin,
+			// 	FExecuteAction::CreateSP( this, &FBlueprintEditor::OnRecombineStructPin ),
+			// 	FCanExecuteAction::CreateSP( this, &FBlueprintEditor::CanRecombineStructPin )
+			// 	);
+			//
+			// GraphEditorCommands->MapAction( FGraphEditorCommands::Get().AddExecutionPin,
+			// 	FExecuteAction::CreateSP( this, &FBlueprintEditor::OnAddExecutionPin ),
+			// 	FCanExecuteAction::CreateSP( this, &FBlueprintEditor::CanAddExecutionPin )
+			// 	);
+			//
+			// GraphEditorCommands->MapAction( FGraphEditorCommands::Get().InsertExecutionPinBefore,
+			// 	FExecuteAction::CreateSP( this, &FBlueprintEditor::OnInsertExecutionPinBefore ),
+			// 	FCanExecuteAction::CreateSP( this, &FBlueprintEditor::CanInsertExecutionPin )
+			// 	);
+			//
+			// GraphEditorCommands->MapAction(FGraphEditorCommands::Get().InsertExecutionPinAfter,
+			// 	FExecuteAction::CreateSP(this, &FBlueprintEditor::OnInsertExecutionPinAfter),
+			// 	FCanExecuteAction::CreateSP(this, &FBlueprintEditor::CanInsertExecutionPin)
+			// );
+			//
+			// GraphEditorCommands->MapAction( FGraphEditorCommands::Get().RemoveExecutionPin,
+			// 	FExecuteAction::CreateSP( this, &FBlueprintEditor::OnRemoveExecutionPin ),
+			// 	FCanExecuteAction::CreateSP( this, &FBlueprintEditor::CanRemoveExecutionPin )
+			// 	);
+			//
+			// GraphEditorCommands->MapAction(FGraphEditorCommands::Get().RemoveThisStructVarPin,
+			// 	FExecuteAction::CreateSP(this, &FBlueprintEditor::OnRemoveThisStructVarPin),
+			// 	FCanExecuteAction::CreateSP(this, &FBlueprintEditor::CanRemoveThisStructVarPin)
+			// 	);
+			//
+			// GraphEditorCommands->MapAction(FGraphEditorCommands::Get().RemoveOtherStructVarPins,
+			// 	FExecuteAction::CreateSP(this, &FBlueprintEditor::OnRemoveOtherStructVarPins),
+			// 	FCanExecuteAction::CreateSP(this, &FBlueprintEditor::CanRemoveOtherStructVarPins)
+			// 	);
+			//
+			// GraphEditorCommands->MapAction(FGraphEditorCommands::Get().RestoreAllStructVarPins,
+			// 	FExecuteAction::CreateSP(this, &FBlueprintEditor::OnRestoreAllStructVarPins),
+			// 	FCanExecuteAction::CreateSP(this, &FBlueprintEditor::CanRestoreAllStructVarPins)
+			// 	);
+			//
+			// GraphEditorCommands->MapAction(FGraphEditorCommands::Get().ResetPinToDefaultValue,
+			// 	FExecuteAction::CreateSP(this, &FBlueprintEditor::OnResetPinToDefaultValue),
+			// 	FCanExecuteAction::CreateSP(this, &FBlueprintEditor::CanResetPinToDefaultValue)
+			// );
+			//
+			// GraphEditorCommands->MapAction( FGraphEditorCommands::Get().AddOptionPin,
+			// 	FExecuteAction::CreateSP( this, &FBlueprintEditor::OnAddOptionPin ),
+			// 	FCanExecuteAction::CreateSP( this, &FBlueprintEditor::CanAddOptionPin )
+			// 	);
+			//
+			// GraphEditorCommands->MapAction( FGraphEditorCommands::Get().RemoveOptionPin,
+			// 	FExecuteAction::CreateSP( this, &FBlueprintEditor::OnRemoveOptionPin ),
+			// 	FCanExecuteAction::CreateSP( this, &FBlueprintEditor::CanRemoveOptionPin )
+			// 	);
+			//
+			// GraphEditorCommands->MapAction( FGraphEditorCommands::Get().ChangePinType,
+			// 	FExecuteAction::CreateSP( this, &FBlueprintEditor::OnChangePinType ),
+			// 	FCanExecuteAction::CreateSP( this, &FBlueprintEditor::CanChangePinType )
+			// 	);
+			//
+			// GraphEditorCommands->MapAction( FGraphEditorCommands::Get().AddParentNode,
+			// 	FExecuteAction::CreateSP( this, &FBlueprintEditor::OnAddParentNode ),
+			// 	FCanExecuteAction::CreateSP( this, &FBlueprintEditor::CanAddParentNode )
+			// 	);
+			//
+			// GraphEditorCommands->MapAction( FGraphEditorCommands::Get().CreateMatchingFunction,
+			// 	FExecuteAction::CreateSP( this, &FBlueprintEditor::OnCreateMatchingFunction ),
+			// 	FCanExecuteAction::CreateSP( this, &FBlueprintEditor::CanCreateMatchingFunction )
+			// 	);
+			//
+			// // Debug actions
+			// GraphEditorCommands->MapAction( FGraphEditorCommands::Get().AddBreakpoint,
+			// 	FExecuteAction::CreateSP( this, &FBlueprintEditor::OnAddBreakpoint ),
+			// 	FCanExecuteAction::CreateSP( this, &FBlueprintEditor::CanAddBreakpoint ),
+			// 	FIsActionChecked(),
+			// 	FIsActionButtonVisible::CreateSP( this, &FBlueprintEditor::CanAddBreakpoint )
+			// 	);
+			//
+			// GraphEditorCommands->MapAction( FGraphEditorCommands::Get().RemoveBreakpoint,
+			// 	FExecuteAction::CreateSP( this, &FBlueprintEditor::OnRemoveBreakpoint ),
+			// 	FCanExecuteAction::CreateSP( this, &FBlueprintEditor::CanRemoveBreakpoint ),
+			// 	FIsActionChecked(),
+			// 	FIsActionButtonVisible::CreateSP( this, &FBlueprintEditor::CanRemoveBreakpoint )
+			// 	);
+			//
+			// GraphEditorCommands->MapAction( FGraphEditorCommands::Get().EnableBreakpoint,
+			// 	FExecuteAction::CreateSP( this, &FBlueprintEditor::OnEnableBreakpoint ),
+			// 	FCanExecuteAction::CreateSP( this, &FBlueprintEditor::CanEnableBreakpoint ),
+			// 	FIsActionChecked(),
+			// 	FIsActionButtonVisible::CreateSP( this, &FBlueprintEditor::CanEnableBreakpoint )
+			// 	);
+			//
+			// GraphEditorCommands->MapAction( FGraphEditorCommands::Get().DisableBreakpoint,
+			// 	FExecuteAction::CreateSP( this, &FBlueprintEditor::OnDisableBreakpoint ),
+			// 	FCanExecuteAction::CreateSP( this, &FBlueprintEditor::CanDisableBreakpoint ),
+			// 	FIsActionChecked(),
+			// 	FIsActionButtonVisible::CreateSP( this, &FBlueprintEditor::CanDisableBreakpoint )
+			// 	);
+			//
+			// GraphEditorCommands->MapAction( FGraphEditorCommands::Get().ToggleBreakpoint,
+			// 	FExecuteAction::CreateSP( this, &FBlueprintEditor::OnToggleBreakpoint ),
+			// 	FCanExecuteAction::CreateSP( this, &FBlueprintEditor::CanToggleBreakpoint ),
+			// 	FIsActionChecked(),
+			// 	FIsActionButtonVisible::CreateSP( this, &FBlueprintEditor::CanToggleBreakpoint )
+			// 	);
+			//
+			// // Encapsulation commands
+			// GraphEditorCommands->MapAction( FGraphEditorCommands::Get().CollapseNodes,
+			// 	FExecuteAction::CreateSP( this, &FBlueprintEditor::OnCollapseNodes ),
+			// 	FCanExecuteAction::CreateSP( this, &FBlueprintEditor::CanCollapseNodes )
+			// 	);
+			//
+			// GraphEditorCommands->MapAction( FGraphEditorCommands::Get().CollapseSelectionToFunction,
+			// 	FExecuteAction::CreateSP( this, &FBlueprintEditor::OnCollapseSelectionToFunction ),
+			// 	FCanExecuteAction::CreateSP( this, &FBlueprintEditor::CanCollapseSelectionToFunction ),
+			// 	FIsActionChecked(),
+			// 	FIsActionButtonVisible::CreateSP( this, &FBlueprintEditor::NewDocument_IsVisibleForType, CGT_NewFunctionGraph )
+			// 	);
+			//
+			// GraphEditorCommands->MapAction( FGraphEditorCommands::Get().CollapseSelectionToMacro,
+			// 	FExecuteAction::CreateSP( this, &FBlueprintEditor::OnCollapseSelectionToMacro ),
+			// 	FCanExecuteAction::CreateSP( this, &FBlueprintEditor::CanCollapseSelectionToMacro ),
+			// 	FIsActionChecked(),
+			// 	FIsActionButtonVisible::CreateSP( this, &FBlueprintEditor::NewDocument_IsVisibleForType, CGT_NewMacroGraph )
+			// 	);
+			//
+			// GraphEditorCommands->MapAction(FGraphEditorCommands::Get().ConvertFunctionToEvent,
+			// 	FExecuteAction::CreateSP(this, &FBlueprintEditor::OnConvertFunctionToEvent),
+			// 	FCanExecuteAction::CreateSP(this, &FBlueprintEditor::CanConvertFunctionToEvent),
+			// 	FIsActionChecked()
+			// );
+			//
+			// GraphEditorCommands->MapAction(FGraphEditorCommands::Get().ConvertEventToFunction,
+			// 	FExecuteAction::CreateSP(this, &FBlueprintEditor::OnConvertEventToFunction),
+			// 	FCanExecuteAction::CreateSP(this, &FBlueprintEditor::CanConvertEventToFunction),
+			// 	FIsActionChecked()
+			// );
+			//
+			// GraphEditorCommands->MapAction( FGraphEditorCommands::Get().PromoteSelectionToFunction,
+			// 	FExecuteAction::CreateSP( this, &FBlueprintEditor::OnPromoteSelectionToFunction ),
+			// 	FCanExecuteAction::CreateSP( this, &FBlueprintEditor::CanPromoteSelectionToFunction ),
+			// 	FIsActionChecked(),
+			// 	FIsActionButtonVisible::CreateSP( this, &FBlueprintEditor::NewDocument_IsVisibleForType, CGT_NewFunctionGraph )
+			// 	);
+			//
+			// GraphEditorCommands->MapAction( FGraphEditorCommands::Get().PromoteSelectionToMacro,
+			// 	FExecuteAction::CreateSP( this, &FBlueprintEditor::OnPromoteSelectionToMacro ),
+			// 	FCanExecuteAction::CreateSP( this, &FBlueprintEditor::CanPromoteSelectionToMacro ),
+			// 	FIsActionChecked(),
+			// 	FIsActionButtonVisible::CreateSP( this, &FBlueprintEditor::NewDocument_IsVisibleForType, CGT_NewMacroGraph )
+			// 	);
+			//
+			// // Alignment Commands
+			// GraphEditorCommands->MapAction( FGraphEditorCommands::Get().AlignNodesTop,
+			// 	FExecuteAction::CreateSP( this, &FBlueprintEditor::OnAlignTop )
+			// 	);
+			//
+			// GraphEditorCommands->MapAction( FGraphEditorCommands::Get().AlignNodesMiddle,
+			// 	FExecuteAction::CreateSP( this, &FBlueprintEditor::OnAlignMiddle )
+			// 	);
+			//
+			// GraphEditorCommands->MapAction( FGraphEditorCommands::Get().AlignNodesBottom,
+			// 	FExecuteAction::CreateSP( this, &FBlueprintEditor::OnAlignBottom )
+			// 	);
+			//
+			// GraphEditorCommands->MapAction( FGraphEditorCommands::Get().AlignNodesLeft,
+			// 	FExecuteAction::CreateSP( this, &FBlueprintEditor::OnAlignLeft )
+			// 	);
+			//
+			// GraphEditorCommands->MapAction( FGraphEditorCommands::Get().AlignNodesCenter,
+			// 	FExecuteAction::CreateSP( this, &FBlueprintEditor::OnAlignCenter )
+			// 	);
+			//
+			// GraphEditorCommands->MapAction( FGraphEditorCommands::Get().AlignNodesRight,
+			// 	FExecuteAction::CreateSP( this, &FBlueprintEditor::OnAlignRight )
+			// 	);
+			//
+			// GraphEditorCommands->MapAction( FGraphEditorCommands::Get().StraightenConnections,
+			// 	FExecuteAction::CreateSP( this, &FBlueprintEditor::OnStraightenConnections )
+			// 	);
+			//
+			// // Distribution Commands
+			// GraphEditorCommands->MapAction( FGraphEditorCommands::Get().DistributeNodesHorizontally,
+			// 	FExecuteAction::CreateSP( this, &FBlueprintEditor::OnDistributeNodesH )
+			// 	);
+			//
+			// GraphEditorCommands->MapAction( FGraphEditorCommands::Get().DistributeNodesVertically,
+			// 	FExecuteAction::CreateSP( this, &FBlueprintEditor::OnDistributeNodesV )
+			// 	);
+			//
+			// GraphEditorCommands->MapAction( FGenericCommands::Get().Rename,
+			// 	FExecuteAction::CreateSP( this, &FBlueprintEditor::OnRenameNode ),
+			// 	FCanExecuteAction::CreateSP( this, &FBlueprintEditor::CanRenameNodes )
+			// 	);
+			//
+			// GraphEditorCommands->MapAction( FGraphEditorCommands::Get().ExpandNodes,
+			// 	FExecuteAction::CreateSP( this, &FBlueprintEditor::OnExpandNodes ),
+			// 	FCanExecuteAction::CreateSP( this, &FBlueprintEditor::CanExpandNodes ),
+			// 	FIsActionChecked(),
+			// 	FIsActionButtonVisible::CreateSP( this, &FBlueprintEditor::CanExpandNodes )
+			// 	);
+			//
+			// // Editing commands
+			// GraphEditorCommands->MapAction( FGenericCommands::Get().SelectAll,
+			// 	FExecuteAction::CreateSP( this, &FBlueprintEditor::SelectAllNodes ),
+			// 	FCanExecuteAction::CreateSP( this, &FBlueprintEditor::CanSelectAllNodes )
+			// 	);
+			//
+			GraphEditorCommands->MapAction( FGenericCommands::Get().Delete,
+				FExecuteAction::CreateSP( this, &FRigVMEditorBase::DeleteSelectedNodes ),
+				FCanExecuteAction::CreateSP( this, &FRigVMEditorBase::CanDeleteNodes )
+				);
+			//
+			// GraphEditorCommands->MapAction(FGraphEditorCommands::Get().DeleteAndReconnectNodes,
+			// 	FExecuteAction::CreateSP(this, &FBlueprintEditor::DeleteSelectedNodes),
+			// 	FCanExecuteAction::CreateSP(this, &FBlueprintEditor::CanDeleteNodes)
+			// );
+			//
+			// GraphEditorCommands->MapAction( FGenericCommands::Get().Copy,
+			// 	FExecuteAction::CreateSP( this, &FBlueprintEditor::CopySelectedNodes ),
+			// 	FCanExecuteAction::CreateSP( this, &FBlueprintEditor::CanCopyNodes )
+			// 	);
+			//
+			// GraphEditorCommands->MapAction( FGenericCommands::Get().Cut,
+			// 	FExecuteAction::CreateSP( this, &FBlueprintEditor::CutSelectedNodes ),
+			// 	FCanExecuteAction::CreateSP( this, &FBlueprintEditor::CanCutNodes )
+			// 	);
+			//
+			// GraphEditorCommands->MapAction( FGenericCommands::Get().Paste,
+			// 	FExecuteAction::CreateSP( this, &FBlueprintEditor::PasteGeneric ),
+			// 	FCanExecuteAction::CreateSP( this, &FBlueprintEditor::CanPasteGeneric )
+			// 	);
+			//
+			// GraphEditorCommands->MapAction( FGenericCommands::Get().Duplicate,
+			// 	FExecuteAction::CreateSP( this, &FBlueprintEditor::DuplicateNodes ),
+			// 	FCanExecuteAction::CreateSP( this, &FBlueprintEditor::CanDuplicateNodes )
+			// 	);
+			//
+			// GraphEditorCommands->MapAction( FGraphEditorCommands::Get().SelectReferenceInLevel,
+			// 	FExecuteAction::CreateSP( this, &FBlueprintEditor::OnSelectReferenceInLevel ),
+			// 	FCanExecuteAction::CreateSP( this, &FBlueprintEditor::CanSelectReferenceInLevel ),
+			// 	FIsActionChecked(),
+			// 	FIsActionButtonVisible::CreateSP( this, &FBlueprintEditor::CanSelectReferenceInLevel )
+			// 	);
+			//
+			// GraphEditorCommands->MapAction( FGraphEditorCommands::Get().AssignReferencedActor,
+			// 	FExecuteAction::CreateSP( this, &FBlueprintEditor::OnAssignReferencedActor ),
+			// 	FCanExecuteAction::CreateSP( this, &FBlueprintEditor::CanAssignReferencedActor ) );
+			//
+			// GraphEditorCommands->MapAction( FGraphEditorCommands::Get().StartWatchingPin,
+			// 	FExecuteAction::CreateSP( this, &FBlueprintEditor::OnStartWatchingPin ),
+			// 	FCanExecuteAction::CreateSP( this, &FBlueprintEditor::CanStartWatchingPin )
+			// 	);
+			//
+			// GraphEditorCommands->MapAction( FGraphEditorCommands::Get().StopWatchingPin,
+			// 	FExecuteAction::CreateSP( this, &FBlueprintEditor::OnStopWatchingPin ),
+			// 	FCanExecuteAction::CreateSP( this, &FBlueprintEditor::CanStopWatchingPin )
+			// 	);
+			//
+			// GraphEditorCommands->MapAction( FGraphEditorCommands::Get().CreateComment,
+			// 	FExecuteAction::CreateSP( this, &FBlueprintEditor::OnCreateComment )
+			// 	);
+			// 	
+			// GraphEditorCommands->MapAction( FGraphEditorCommands::Get().CreateCustomEvent,
+			// 	FExecuteAction::CreateSP( this, &FBlueprintEditor::OnCreateCustomEvent )
+			// 	);
+			//
+			// GraphEditorCommands->MapAction(FGraphEditorCommands::Get().ShowAllPins,
+			// 	FExecuteAction::CreateSP(this, &FBlueprintEditor::SetPinVisibility, SGraphEditor::Pin_Show)
+			// 	);
+			//
+			// GraphEditorCommands->MapAction(FGraphEditorCommands::Get().HideNoConnectionPins,
+			// 	FExecuteAction::CreateSP(this, &FBlueprintEditor::SetPinVisibility, SGraphEditor::Pin_HideNoConnection)
+			// 	);
+			//
+			// GraphEditorCommands->MapAction(FGraphEditorCommands::Get().HideNoConnectionNoDefaultPins,
+			// 	FExecuteAction::CreateSP(this, &FBlueprintEditor::SetPinVisibility, SGraphEditor::Pin_HideNoConnectionNoDefault)
+			// 	);
+			//
+			// GraphEditorCommands->MapAction( FGraphEditorCommands::Get().FindReferences,
+			// 	FExecuteAction::CreateSP( this, &FBlueprintEditor::OnFindReferences, /*bSearchAllBlueprints=*/false, EGetFindReferenceSearchStringFlags::Legacy),
+			// 	FCanExecuteAction::CreateSP( this, &FBlueprintEditor::CanFindReferences )
+			// 	);
+			//
+			// GraphEditorCommands->MapAction( FGraphEditorCommands::Get().FindReferencesByNameLocal,
+			// 	FExecuteAction::CreateSP( this, &FBlueprintEditor::OnFindReferences, /*bSearchAllBlueprints=*/false, EGetFindReferenceSearchStringFlags::None ),
+			// 	FCanExecuteAction::CreateSP( this, &FBlueprintEditor::CanFindReferences )
+			// 	);
+			//
+			// GraphEditorCommands->MapAction( FGraphEditorCommands::Get().FindReferencesByNameGlobal,
+			// 	FExecuteAction::CreateSP( this, &FBlueprintEditor::OnFindReferences, /*bSearchAllBlueprints=*/true, EGetFindReferenceSearchStringFlags::None ),
+			// 	FCanExecuteAction::CreateSP( this, &FBlueprintEditor::CanFindReferences )
+			// 	);
+			//
+			// GraphEditorCommands->MapAction( FGraphEditorCommands::Get().FindReferencesByClassMemberLocal,
+			// 	FExecuteAction::CreateSP( this, &FBlueprintEditor::OnFindReferences, /*bSearchAllBlueprints=*/false, EGetFindReferenceSearchStringFlags::UseSearchSyntax ),
+			// 	FCanExecuteAction::CreateSP( this, &FBlueprintEditor::CanFindReferences )
+			// 	);
+			//
+			// GraphEditorCommands->MapAction( FGraphEditorCommands::Get().FindReferencesByClassMemberGlobal,
+			// 	FExecuteAction::CreateSP( this, &FBlueprintEditor::OnFindReferences, /*bSearchAllBlueprints=*/true, EGetFindReferenceSearchStringFlags::UseSearchSyntax ),
+			// 	FCanExecuteAction::CreateSP( this, &FBlueprintEditor::CanFindReferences )
+			// 	);
+			//
+			// GraphEditorCommands->MapAction( FGraphEditorCommands::Get().GoToDefinition,
+			// 	FExecuteAction::CreateSP( this, &FBlueprintEditor::OnGoToDefinition ),
+			// 	FCanExecuteAction::CreateSP( this, &FBlueprintEditor::CanGoToDefinition )
+			// 	);
+			//
+			// GraphEditorCommands->MapAction(FGraphEditorCommands::Get().GoToDocumentation,
+			// 	FExecuteAction::CreateSP(this, &FBlueprintEditor::OnGoToDocumentation),
+			// 	FCanExecuteAction::CreateSP(this, &FBlueprintEditor::CanGoToDocumentation)
+			// 	);
+			//
+			// GraphEditorCommands->MapAction(FGraphEditorCommands::Get().EnableNodes,
+			// 	FExecuteAction(),
+			// 	FCanExecuteAction(),
+			// 	FGetActionCheckState::CreateSP(this, &FBlueprintEditor::GetEnabledCheckBoxStateForSelectedNodes)
+			// 	);
+			//
+			// GraphEditorCommands->MapAction(FGraphEditorCommands::Get().DisableNodes,
+			// 	FExecuteAction::CreateSP(this, &FBlueprintEditor::OnSetEnabledStateForSelectedNodes, ENodeEnabledState::Disabled),
+			// 	FCanExecuteAction(),
+			// 	FGetActionCheckState::CreateSP(this, &FBlueprintEditor::CheckEnabledStateForSelectedNodes, ENodeEnabledState::Disabled)
+			// 	);
+			//
+			// GraphEditorCommands->MapAction(FGraphEditorCommands::Get().EnableNodes_Always,
+			// 	FExecuteAction::CreateSP(this, &FBlueprintEditor::OnSetEnabledStateForSelectedNodes, ENodeEnabledState::Enabled),
+			// 	FCanExecuteAction(),
+			// 	FGetActionCheckState::CreateSP(this, &FBlueprintEditor::CheckEnabledStateForSelectedNodes, ENodeEnabledState::Enabled)
+			// 	);
+			//
+			// GraphEditorCommands->MapAction(FGraphEditorCommands::Get().EnableNodes_DevelopmentOnly,
+			// 	FExecuteAction::CreateSP(this, &FBlueprintEditor::OnSetEnabledStateForSelectedNodes, ENodeEnabledState::DevelopmentOnly),
+			// 	FCanExecuteAction(),
+			// 	FGetActionCheckState::CreateSP(this, &FBlueprintEditor::CheckEnabledStateForSelectedNodes, ENodeEnabledState::DevelopmentOnly)
+			// 	);
+
+			OnCreateGraphEditorCommands(GraphEditorCommands);
+		}
+	}
+
+	// Create the title bar widget
+	TSharedPtr<SWidget> TitleBarWidget = SNew(SRigVMGraphTitleBar)
+		.EdGraphObj(InGraph)
+		.Editor(SharedThis(this))
+		.HistoryNavigationWidget(InTabInfo->CreateHistoryNavigationWidget());
+
+	SGraphEditor::FGraphEditorEvents InEvents;
+	SetupGraphEditorEvents(InGraph, InEvents);
+
+	// Append play world commands
+	GraphEditorCommands->Append( FPlayWorldCommands::GlobalPlayWorldActions.ToSharedRef() );
+
+	TSharedRef<SGraphEditor> Editor = SNew(SGraphEditor)
+		.AdditionalCommands(GraphEditorCommands)
+		//.IsEditable(this, &FBlueprintEditor::IsEditable, InGraph)
+		//.DisplayAsReadOnly(this, &FBlueprintEditor::IsGraphReadOnly, InGraph)
+		.TitleBar(TitleBarWidget)
+		//.Appearance(this, &FBlueprintEditor::GetGraphAppearance, InGraph)
+		.GraphToEdit(InGraph)
+		.GraphEvents(InEvents)
+		//.OnNavigateHistoryBack(FSimpleDelegate::CreateSP(this, &FBlueprintEditor::NavigateTab, FDocumentTracker::NavigateBackwards))
+		//.OnNavigateHistoryForward(FSimpleDelegate::CreateSP(this, &FBlueprintEditor::NavigateTab, FDocumentTracker::NavigateForwards))
+		.AssetEditorToolkit(this->GetHostingApp());
+		//@TODO: Crashes in command list code during the callback .OnGraphModuleReloaded(FEdGraphEvent::CreateSP(this, &FBlueprintEditor::ChangeOpenGraphInDocumentEditorWidget, WeakParent))
+		;
+		
+	//OnSetPinVisibility.AddSP(&Editor.Get(), &SGraphEditor::SetPinVisibility);
+
+	FVector2D ViewOffset = FVector2D::ZeroVector;
+	float ZoomAmount = INDEX_NONE;
+
+	TSharedPtr<SDockTab> ActiveTab = DocumentManager->GetActiveTab();
+	if(ActiveTab.IsValid())
+	{
+		// Check if the graph is already opened in the current tab, if it is we want to start at the same position to stop the graph from jumping around oddly
+		TSharedPtr<SGraphEditor> GraphEditor = StaticCastSharedRef<SGraphEditor>(ActiveTab->GetContent());
+
+		if(GraphEditor.IsValid() && GraphEditor->GetCurrentGraph() == InGraph)
+		{
+			GraphEditor->GetViewLocation(ViewOffset, ZoomAmount);
+		}
+	}
+
+	Editor->SetViewLocation(ViewOffset, ZoomAmount);
+	
+	Editor->GetGraphPanel()->SetZoomLevelsContainer<FRigVMEditorZoomLevelsContainer>();
+	return Editor;
+}
+
+void FRigVMNewEditor::CompileImpl()
+{
+	DECLARE_SCOPE_HIERARCHICAL_COUNTER_FUNC()
+
+	UBlueprint* BlueprintObj = GetBlueprintObj();
+	if (BlueprintObj)
+	{
+		FMessageLog BlueprintLog("BlueprintLog");
+
+		FFormatNamedArguments Arguments;
+		Arguments.Add(TEXT("BlueprintName"), FText::FromString(BlueprintObj->GetName()));
+		BlueprintLog.NewPage(FText::Format(LOCTEXT("CompilationPageLabel", "Compile {BlueprintName}"), Arguments));
+
+		FCompilerResultsLog LogResults;
+		LogResults.SetSourcePath(BlueprintObj->GetPathName());
+		LogResults.BeginEvent(TEXT("Compile"));
+		
+		// TODO: sara-s remove once blueprint backend is replaced
+		// {
+		// 	EBlueprintCompileOptions CompileOptions = EBlueprintCompileOptions::None;
+		// 	// if( bSaveIntermediateBuildProducts )
+		// 	// {
+		// 	// 	CompileOptions |= EBlueprintCompileOptions::SaveIntermediateProducts;
+		// 	// }
+		//
+		// 	// if (bIsReparentingBlueprint)
+		// 	// {
+		// 	// 	CompileOptions |= (EBlueprintCompileOptions::UseDeltaSerializationDuringReinstancing | EBlueprintCompileOptions::SkipNewVariableDefaultsDetection);
+		// 	// }
+		//
+		// 	// If compilation is enabled during PIE/simulation, references to the CDO might be held by a script variable.
+		// 	// Thus, we set the flag to direct the compiler to allow those references to be replaced during reinstancing.
+		// 	if (GEditor->PlayWorld != nullptr)
+		// 	{
+		// 		CompileOptions |= EBlueprintCompileOptions::IncludeCDOInReferenceReplacement;
+		// 	}
+		//
+		// 	FKismetEditorUtilities::CompileBlueprint(BlueprintObj, CompileOptions, &LogResults);
+		// }
+
+		LogResults.EndEvent();
+
+		// CachedNumWarnings = LogResults.NumWarnings;
+		// CachedNumErrors = LogResults.NumErrors;
+
+		const bool bForceMessageDisplay = (LogResults.NumWarnings > 0 || LogResults.NumErrors > 0) && !BlueprintObj->bIsRegeneratingOnLoad;
+		//DumpMessagesToCompilerLog(LogResults.Messages, bForceMessageDisplay);
+
+		if (BlueprintObj->UpgradeNotesLog.IsValid())
+		{
+			CompilerResultsListing->AddMessages(BlueprintObj->UpgradeNotesLog->Messages);
+		}
+
+		// send record when player clicks compile and send the result
+		// this will make sure how the users activity is
+		//AnalyticsTrackCompileEvent(BlueprintObj, LogResults.NumErrors, LogResults.NumWarnings);
+
+		RefreshInspector();
+	}
+}
+
+bool FRigVMNewEditor::IsGraphInCurrentBlueprint(const UEdGraph* InGraph) const
+{
+	bool bEditable = true;
+
+	UBlueprint* EditingBP = GetBlueprintObj();
+	if(EditingBP)
+	{
+		TArray<UEdGraph*> Graphs;
+		EditingBP->GetAllGraphs(Graphs);
+		bEditable &= Graphs.Contains(InGraph);
+	}
+
+	return bEditable;
+}
+
+bool FRigVMNewEditor::IsEditableImpl(UEdGraph* InGraph) const
+{
+	return InEditingMode() && !FBlueprintEditorUtils::IsGraphReadOnly(InGraph);
+}
+
+UEdGraph* FRigVMNewEditor::GetFocusedGraph() const
+{
+	if (FocusedGraphEdPtr.IsValid())
+	{
+		if (UEdGraph* Graph = FocusedGraphEdPtr.Pin()->GetCurrentGraph())
+		{
+			if (IsValid(Graph))
+			{
+				return Graph;
+			}
+		}
+	}
+	return nullptr;
+}
+
+void FRigVMNewEditor::JumpToNode(const UEdGraphNode* Node, bool bRequestRename)
+{
+	TSharedPtr<SGraphEditor> GraphEditor;
+	if(bRequestRename)
+	{
+		// If we are renaming, the graph will be open already, just grab the tab and it's content and jump to the node.
+		TSharedPtr<SDockTab> ActiveTab = DocumentManager->GetActiveTab();
+		check(ActiveTab.IsValid());
+		GraphEditor = StaticCastSharedRef<SGraphEditor>(ActiveTab->GetContent());
+	}
+	else
+	{
+		// Open a graph editor and jump to the node
+		GraphEditor = OpenGraphAndBringToFront(Node->GetGraph());
+	}
+
+	if (GraphEditor.IsValid())
+	{
+		GraphEditor->JumpToNode(Node, bRequestRename);
+	}
+}
+
+void FRigVMNewEditor::JumpToPin(const UEdGraphPin* Pin)
+{
+	if (!Pin->IsPendingKill())
+	{
+		// Open a graph editor and jump to the pin
+		TSharedPtr<SGraphEditor> GraphEditor = OpenGraphAndBringToFront(Pin->GetOwningNode()->GetGraph());
+		if (GraphEditor.IsValid())
+		{
+			GraphEditor->JumpToPin(Pin);
+		}
+	}
+}
+
+void FRigVMNewEditor::OnBlueprintChangedInnerImpl(UBlueprint* InBlueprint, bool bIsJustBeingCompiled)
+{
+	if (InBlueprint)
+	{
+		// Notify that the blueprint has been changed (update Content browser, etc)
+		InBlueprint->PostEditChange();
+
+		// Call PostEditChange() on any Actors that are based on this Blueprint 
+		FBlueprintEditorUtils::PostEditChangeBlueprintActors(InBlueprint);
+
+		// Refresh the graphs
+		ERefreshRigVMEditorReason::Type Reason = bIsJustBeingCompiled ? ERefreshRigVMEditorReason::BlueprintCompiled : ERefreshRigVMEditorReason::UnknownReason;
+		RefreshEditors(Reason);
+
+		// In case objects were deleted, which should close the tab
+		if (GetCurrentMode() == FRigVMNewEditorApplicationModes::StandardRigVMEditorMode())
+		{
+			// Copied from FBlueprintEditor::SaveEditedObjectState()
+            {
+            	check(IsEditingSingleBlueprint());
+            
+            	// Clear currently edited documents
+            	GetBlueprintObj()->LastEditedDocuments.Empty();
+            
+            	// Ask all open documents to save their state, which will update LastEditedDocuments
+            	DocumentManager->SaveAllState();
+            }
+		}
+	}
+}
+
+void FRigVMNewEditor::RefreshEditorsImpl(ERefreshRigVMEditorReason::Type Reason)
+{
+	// TODO sara-s: Replicate what blueprint editor is doing
+	bool bForceFocusOnSelectedNodes = false;
+
+	// if (CurrentUISelection == SelectionState_MyBlueprint)
+	// {
+	// 	// Handled below, here to avoid tripping the ensure
+	// }
+	// else if (CurrentUISelection == SelectionState_Components)
+	// {
+	// 	if(SubobjectEditor.IsValid())
+	// 	{
+	// 		SubobjectEditor->RefreshSelectionDetails();
+	// 	}
+	// }
+	// else if (CurrentUISelection == SelectionState_Graph)
+	// {
+	// 	bForceFocusOnSelectedNodes = true;
+	// }
+	// else if (CurrentUISelection == SelectionState_ClassSettings)
+	// {
+	// 	// No need for a refresh, the Blueprint object didn't change
+	// }
+	// else if (CurrentUISelection == SelectionState_ClassDefaults)
+	// {
+	// 	StartEditingDefaults(/*bAutoFocus=*/ false, true);
+	// }
+
+	// Remove any tabs are that are pending kill or otherwise invalid UObject pointers.
+	DocumentManager->CleanInvalidTabs();
+
+	//@TODO: Should determine when we need to do the invalid/refresh business and if the graph node selection change
+	// under non-compiles is necessary (except when the selection mode is appropriate, as already detected above)
+	if (Reason != ERefreshRigVMEditorReason::BlueprintCompiled)
+	{
+		DocumentManager->RefreshAllTabs();
+
+		bForceFocusOnSelectedNodes = true;
+	}
+
+	if (bForceFocusOnSelectedNodes)
+	{
+		FocusInspectorOnGraphSelection(GetSelectedNodes(), /*bForceRefresh=*/ true);
+	}
+
+	// if (ReplaceReferencesWidget.IsValid())
+	// {
+	// 	ReplaceReferencesWidget->Refresh();
+	// }
+
+	if (GraphExplorerWidget.IsValid())
+	{
+		GraphExplorerWidget->Refresh();
+	}
+
+	// if(SubobjectEditor.IsValid())
+	// {
+	// 	SubobjectEditor->RefreshComponentTypesList();
+	// 	SubobjectEditor->UpdateTree();
+	//
+	// 	// Note: Don't pass 'true' here because we don't want the preview actor to be reconstructed until after Blueprint modification is complete.
+	// 	UpdateSubobjectPreview();
+	// }
+	//
+	// if (BookmarksWidget.IsValid())
+	// {
+	// 	BookmarksWidget->RefreshBookmarksTree();
+	// }
+
+	// Note: There is an optimization inside of ShowDetailsForSingleObject() that skips the refresh if the object being selected is the same as the previous object.
+	// The SKismetInspector class is shared between both Defaults mode and Components mode, but in Defaults mode the object selected is always going to be the CDO. Given
+	// that the selection does not really change, we force it to refresh and skip the optimization. Otherwise, some things may not work correctly in Defaults mode. For
+	// example, transform details are customized and the rotation value is cached at customization time; if we don't force refresh here, then after an undo of a previous
+	// rotation edit, transform details won't be re-customized and thus the cached rotation value will be stale, resulting in an invalid rotation value on the next edit.
+	//@TODO: Probably not always necessary
+	//RefreshStandAloneDefaultsEditor();
+
+	// Update associated controls like the function editor
+	//BroadcastRefresh();
+}
+
+void FRigVMNewEditor::SetupGraphEditorEventsImpl(UEdGraph* InGraph, SGraphEditor::FGraphEditorEvents& InEvents)
+{
+	// TODO sara-s: Replicate what blueprint editor is doing
+	// InEvents.OnSelectionChanged = SGraphEditor::FOnSelectionChanged::CreateSP( this, &FRigVMNewEditor::OnSelectedNodesChanged );
+	// InEvents.OnDropActor = SGraphEditor::FOnDropActor::CreateSP( this, &FRigVMNewEditor::OnGraphEditorDropActor );
+	// InEvents.OnDropStreamingLevel = SGraphEditor::FOnDropStreamingLevel::CreateSP( this, &FRigVMNewEditor::OnGraphEditorDropStreamingLevel );
+	// InEvents.OnNodeDoubleClicked = FSingleNodeEvent::CreateSP(this, &FRigVMNewEditor::OnNodeDoubleClicked);
+	// InEvents.OnVerifyTextCommit = FOnNodeVerifyTextCommit::CreateSP(this, &FRigVMNewEditor::OnNodeVerifyTitleCommit);
+	// InEvents.OnTextCommitted = FOnNodeTextCommitted::CreateSP(this, &FRigVMNewEditor::OnNodeTitleCommitted);
+	// InEvents.OnSpawnNodeByShortcut = SGraphEditor::FOnSpawnNodeByShortcut::CreateSP(this, &FRigVMNewEditor::OnSpawnGraphNodeByShortcut, InGraph);
+	// InEvents.OnNodeSpawnedByKeymap = SGraphEditor::FOnNodeSpawnedByKeymap::CreateSP(this, &FRigVMNewEditor::OnNodeSpawnedByKeymap );
+	// InEvents.OnDisallowedPinConnection = SGraphEditor::FOnDisallowedPinConnection::CreateSP(this, &FRigVMNewEditor::OnDisallowedPinConnection);
+	// InEvents.OnDoubleClicked = SGraphEditor::FOnDoubleClicked::CreateSP(this, &FRigVMNewEditor::NavigateToParentGraphByDoubleClick);
+	//
+	// // Custom menu for K2 schemas
+	if(InGraph->Schema != nullptr && InGraph->Schema->IsChildOf(UEdGraphSchema_K2::StaticClass()))
+	{
+		InEvents.OnCreateActionMenu = SGraphEditor::FOnCreateActionMenu::CreateSP(this, &FRigVMNewEditor::OnCreateGraphActionMenu);
+	}
+}
+
+FActionMenuContent FRigVMNewEditor::OnCreateGraphActionMenu(UEdGraph* InGraph, const FVector2D& InNodePosition,
+	const TArray<UEdGraphPin*>& InDraggedPins, bool bAutoExpand, SGraphEditor::FActionMenuClosed InOnMenuClosed)
+{
+	// TODO sara-s: Replicate what blueprint editor is doing
+	HasOpenActionMenu = InGraph;
+	// if (!BlueprintEditorImpl::GraphHasUserPlacedNodes(InGraph))
+	// {
+	// 	InstructionsFadeCountdown = BlueprintEditorImpl::InstructionFadeDuration;
+	// }
+
+	TSharedRef<SRigVMActionMenu> ActionMenu = 
+		SNew(SRigVMActionMenu, SharedThis(this))
+		.GraphObj(InGraph)
+		.NewNodePosition(InNodePosition)
+		.DraggedFromPins(InDraggedPins);
+
+	return FActionMenuContent( ActionMenu, ActionMenu->GetFilterTextBox() );
+}
+
+void FRigVMNewEditor::AddReferencedObjectsImpl(FReferenceCollector& Collector)
+{
+	
+}
+
+bool FRigVMNewEditor::IsSectionVisible(RigVMNodeSectionID::Type InSectionID) const
+{
+	return FRigVMEditorBase::IsSectionVisibleImpl(InSectionID);
+}
+
+bool FRigVMNewEditor::NewDocument_IsVisibleForType(ECreatedDocumentType GraphType) const
+{
+	return FRigVMEditorBase::NewDocument_IsVisibleForTypeImpl(GraphType);
+}
+
+FGraphAppearanceInfo FRigVMNewEditor::GetGraphAppearanceImpl(UEdGraph* InGraph) const
+{
+	// Create the appearance info
+	FGraphAppearanceInfo AppearanceInfo;
+
+	UBlueprint* Blueprint = (InGraph != nullptr) ? FBlueprintEditorUtils::FindBlueprintForGraph(InGraph) : GetBlueprintObj();
+	if (Blueprint != nullptr)
+	{
+		if (FBlueprintEditorUtils::IsEditorUtilityBlueprint(Blueprint))
+		{
+			AppearanceInfo.CornerText = LOCTEXT("EditorUtilityAppearanceCornerText", "EDITOR UTILITY");
+		}
+		else
+		{
+			switch (Blueprint->BlueprintType)
+			{
+			case BPTYPE_LevelScript:
+				AppearanceInfo.CornerText = LOCTEXT("AppearanceCornerText_LevelScript", "LEVEL BLUEPRINT");
+				break;
+			case BPTYPE_MacroLibrary:
+				AppearanceInfo.CornerText = LOCTEXT("AppearanceCornerText_Macro", "MACRO");
+				break;
+			case BPTYPE_Interface:
+				AppearanceInfo.CornerText = LOCTEXT("AppearanceCornerText_Interface", "INTERFACE");
+				break;
+			default:
+				AppearanceInfo.CornerText = LOCTEXT("AppearanceCornerText_Blueprint", "BLUEPRINT");
+				break;
+			}
+		}
+	}
+
+	UEdGraph const* EditingGraph = GetFocusedGraph();
+	if (InGraph && RigVMNewEditorImpl::GraphHasDefaultNode(InGraph))
+	{
+		AppearanceInfo.InstructionText = LOCTEXT("AppearanceInstructionText_DefaultGraph", "Drag Off Pins to Create/Connect New Nodes.");
+	}
+	else // if the graph is empty...
+		{
+		AppearanceInfo.InstructionText = LOCTEXT("AppearanceInstructionText_EmptyGraph", "Right-Click to Create New Nodes.");
+		}
+	auto InstructionOpacityDelegate = TAttribute<float>::FGetter::CreateSP(this, &FRigVMNewEditor::GetInstructionTextOpacity, InGraph);
+	AppearanceInfo.InstructionFade.Bind(InstructionOpacityDelegate);
+
+	// Copy from FBlueprintEditor::GetPIEStatus
+	{
+		UBlueprint* CurrentBlueprint = GetBlueprintObj();
+		UWorld *DebugWorld = nullptr;
+		ENetMode NetMode = NM_Standalone;
+		if (CurrentBlueprint)
+		{
+			DebugWorld = CurrentBlueprint->GetWorldBeingDebugged();
+			if (DebugWorld)
+			{
+				NetMode = DebugWorld->GetNetMode();
+			}
+			else
+			{
+				UObject* ObjOuter = CurrentBlueprint->GetObjectBeingDebugged();
+				while(DebugWorld == nullptr && ObjOuter != nullptr)
+				{
+					ObjOuter = ObjOuter->GetOuter();
+					DebugWorld = Cast<UWorld>(ObjOuter);
+				}
+
+				if (DebugWorld)
+				{
+					// Redirect through streaming levels to find the owning world; this ensures that we always use the appropriate NetMode for the context string below.
+					if (DebugWorld->PersistentLevel != nullptr && DebugWorld->PersistentLevel->OwningWorld != nullptr)
+					{
+						DebugWorld = DebugWorld->PersistentLevel->OwningWorld;
+					}
+
+					NetMode = DebugWorld->GetNetMode();
+				}
+			}
+		}
+
+		if (NetMode == NM_ListenServer || NetMode == NM_DedicatedServer)
+		{
+			AppearanceInfo.PIENotifyText = LOCTEXT("PIEStatusServerSimulating", "SERVER - SIMULATING");
+		}
+		else if (NetMode == NM_Client)
+		{
+			FWorldContext* PIEContext = GEngine->GetWorldContextFromWorld(DebugWorld);
+			if (PIEContext && PIEContext->PIEInstance > 1)
+			{
+				AppearanceInfo.PIENotifyText = FText::Format(LOCTEXT("PIEStatusClientSimulatingFormat", "CLIENT {0} - SIMULATING"), FText::AsNumber(PIEContext->PIEInstance - 1));
+			}
+		
+			AppearanceInfo.PIENotifyText = LOCTEXT("PIEStatusClientSimulating", "CLIENT - SIMULATING");
+		}
+
+		AppearanceInfo.PIENotifyText = LOCTEXT("PIEStatusSimulating", "SIMULATING");
+	}
+
+	return AppearanceInfo;
+}
+
+void FRigVMNewEditor::NotifyPreChangeImpl(FProperty* PropertyAboutToChange)
+{
+	// this only delivers message to the "FOCUSED" one, not every one
+	// internally it will only deliver the message to the selected node, not all nodes
+	FString PropertyName = PropertyAboutToChange->GetName();
+	if (FocusedGraphEdPtr.IsValid())
+	{
+		FocusedGraphEdPtr.Pin()->NotifyPrePropertyChange(PropertyName);
+	}
+}
+
+void FRigVMNewEditor::NotifyPostChangeImpl(const FPropertyChangedEvent& PropertyChangedEvent, FProperty* PropertyThatChanged)
+{
+	FString PropertyName = PropertyThatChanged->GetName();
+	if (FocusedGraphEdPtr.IsValid())
+	{
+		FocusedGraphEdPtr.Pin()->NotifyPostPropertyChange(PropertyChangedEvent, PropertyName);
+	}
+	
+	if (IsEditingSingleBlueprint())
+	{
+		UBlueprint* Blueprint = GetBlueprintObj();
+		UPackage* BlueprintPackage = Blueprint->GetOutermost();
+
+		// if any of the objects being edited are in our package, mark us as dirty
+		bool bPropertyInBlueprint = false;
+		for (int32 ObjectIndex = 0; ObjectIndex < PropertyChangedEvent.GetNumObjectsBeingEdited(); ++ObjectIndex)
+		{
+			const UObject* Object = PropertyChangedEvent.GetObjectBeingEdited(ObjectIndex);
+			if (Object && Object->GetOutermost() == BlueprintPackage)
+			{
+				bPropertyInBlueprint = true;
+				break;
+			}
+		}
+
+		if (bPropertyInBlueprint)
+		{
+			// Note: if change type is "interactive," hold off on applying the change (e.g. this will occur if the user is scrubbing a spinbox value; we don't want to apply the change until the mouse is released, for performance reasons)
+			if (PropertyChangedEvent.ChangeType != EPropertyChangeType::Interactive)
+			{
+				FBlueprintEditorUtils::MarkBlueprintAsModified(Blueprint, PropertyChangedEvent);
+
+				// Call PostEditChange() on any Actors that might be based on this Blueprint
+				FBlueprintEditorUtils::PostEditChangeBlueprintActors(Blueprint);
+			}
+
+			// TODO sara-s
+			// Force updates to occur immediately during interactive mode (otherwise the preview won't refresh because it won't be ticking)
+			//UpdateSubobjectPreview(PropertyChangedEvent.ChangeType == EPropertyChangeType::Interactive);
+		}
+	}
+}
+
+FName FRigVMNewEditor::GetSelectedVariableName()
+{
+	FName VariableName;
+	if (GraphExplorerWidget.IsValid())
+	{
+		return GraphExplorerWidget->GetSelectedVariableName();
+	}
+	return VariableName;
+}
+
+bool FRigVMNewEditor::IsNodeTitleVisible(const UEdGraphNode* Node, bool bRequestRename)
+{
+	TSharedPtr<SGraphEditor> GraphEditor;
+	if(bRequestRename)
+	{
+		// If we are renaming, the graph will be open already, just grab the tab and it's content and jump to the node.
+		TSharedPtr<SDockTab> ActiveTab = DocumentManager->GetActiveTab();
+		check(ActiveTab.IsValid());
+		GraphEditor = StaticCastSharedRef<SGraphEditor>(ActiveTab->GetContent());
+	}
+	else
+	{
+		// Open a graph editor and jump to the node
+		GraphEditor = OpenGraphAndBringToFront(Node->GetGraph());
+	}
+
+	bool bVisible = false;
+	if (GraphEditor.IsValid())
+	{
+		bVisible = GraphEditor->IsNodeTitleVisible(Node, bRequestRename);
+	}
+	return bVisible;
+}
+
+void FRigVMNewEditor::EditClassDefaults_Clicked()
+{
+	StartEditingDefaults(true, true);
+}
+
+void FRigVMNewEditor::EditGlobalOptions_Clicked()
+{
+	SetUISelectionState(FRigVMNewEditor::SelectionState_ClassSettings());
+
+	if (bWasOpenedInDefaultsMode)
+	{
+		RefreshStandAloneDefaultsEditor();
+	}
+	else
+	{
+		UBlueprint* Blueprint = GetBlueprintObj();
+		if (Blueprint != nullptr)
+		{
+			// Show details for the Blueprint instance we're editing
+			if (Inspector.IsValid())
+			{
+				Inspector->ShowDetailsForSingleObject(Blueprint);
+			}
+
+			TryInvokingDetailsTab();
+		}
+	}
+}
+
+void FRigVMNewEditor::TryInvokingDetailsTab(bool bFlash)
+{
+	if ( TabManager->HasTabSpawner(FRigVMNewEditorTabs::DetailsID()) )
+	{
+		TSharedPtr<SDockTab> BlueprintTab = FGlobalTabmanager::Get()->GetMajorTabForTabManager(TabManager.ToSharedRef());
+
+		// We don't want to force this tab into existence when the blueprint editor isn't in the foreground and actively
+		// being interacted with.  So we make sure the window it's in is focused and the tab is in the foreground.
+		if ( BlueprintTab.IsValid() && BlueprintTab->IsForeground() )
+		{
+			TSharedPtr<SWindow> ParentWindow = BlueprintTab->GetParentWindow();
+			if ( ParentWindow.IsValid() && ParentWindow->HasFocusedDescendants() )
+			{
+				if ( !Inspector.IsValid() || !Inspector->GetOwnerTab().IsValid() || Inspector->GetOwnerTab()->GetDockArea().IsValid() )
+				{
+					// Show the details panel if it doesn't exist.
+					TabManager->TryInvokeTab(FRigVMNewEditorTabs::DetailsID());
+
+					if ( bFlash && Inspector.IsValid())
+					{
+						TSharedPtr<SDockTab> OwnerTab = Inspector->GetOwnerTab();
+						if ( OwnerTab.IsValid() )
+						{
+							OwnerTab->FlashTab();
+						}
+					}
+				}
+			}
+		}
+	}
+}
+
+void FRigVMNewEditor::RefreshInspector()
+{
+	if (Inspector.IsValid())
+	{
+	    Inspector->GetPropertyView()->ForceRefresh();
+	}
+}
+
+void FRigVMNewEditor::RefreshStandAloneDefaultsEditor()
+{
+	// Update the details panel
+#if WITH_RIGVMLEGACYEDITOR
+	SKismetInspector::FShowDetailsOptions Options(FText::GetEmpty(), true);
+#else
+	SRigVMDetailsInspector::FShowDetailsOptions Options(FText::GetEmpty(), true);
+#endif
+
+	TArray<UObject*> DefaultObjects;
+	for ( int32 i = 0; i < GetEditingObjects().Num(); ++i )
+	{
+		if (UBlueprint* Blueprint = Cast<UBlueprint>(GetEditingObjects()[i]))
+		{
+			if (CurrentUISelection == FRigVMNewEditor::SelectionState_ClassSettings())
+			{
+				DefaultObjects.Add(Blueprint);
+			}
+			else if (Blueprint->GeneratedClass)
+			{
+				DefaultObjects.Add(Blueprint->GeneratedClass->GetDefaultObject());
+			}
+		}
+	}
+
+	if ( DefaultObjects.Num() && DefaultEditor.IsValid() )
+	{
+		DefaultEditor->ShowDetailsForObjects(DefaultObjects);
+	}
+}
+
+void FRigVMNewEditor::StartEditingDefaults(bool bAutoFocus, bool bForceRefresh)
+{
+	SetUISelectionState(FRigVMNewEditor::SelectionState_ClassDefaults());
+
+	if (IsEditingSingleBlueprint())
+	{
+		if (GetBlueprintObj()->GeneratedClass != nullptr)
+		{
+			// TODO sara-s
+			// if (SubobjectEditor.IsValid() && GetBlueprintObj()->GeneratedClass->IsChildOf<AActor>())
+			// {
+			// 	SubobjectEditor->SelectRoot();
+			// }
+			// else
+			{
+				UObject* DefaultObject = GetBlueprintObj()->GeneratedClass->GetDefaultObject();
+
+				// Update the details panel
+				FString Title;
+				DefaultObject->GetName(Title);
+#if WITH_RIGVMLEGACYEDITOR
+				SKismetInspector::FShowDetailsOptions Options(FText::FromString(Title), bForceRefresh);
+				Options.bShowComponents = false;
+#else
+				SRigVMDetailsInspector::FShowDetailsOptions Options(FText::FromString(Title), bForceRefresh);
+#endif
+
+				if (Inspector.IsValid())
+				{
+					Inspector->ShowDetailsForSingleObject(DefaultObject, Options);
+				}
+
+				if ( bAutoFocus )
+				{
+					TryInvokingDetailsTab();
+				}
+			}
+		}
+	}
+	
+	RefreshStandAloneDefaultsEditor();
+}
+
+float FRigVMNewEditor::GetInstructionTextOpacity(UEdGraph* InGraph) const
+{
+	bool bGraphReadOnly = true;
+	if (InGraph)
+	{
+		bGraphReadOnly = !InGraph->bEditable;
+
+		if (!bGraphReadOnly)
+		{
+			const UBlueprint* BlueprintForGraph = FBlueprintEditorUtils::FindBlueprintForGraph(InGraph);
+			bool const bIsInterface = ((BlueprintForGraph != nullptr) && (BlueprintForGraph->BlueprintType == BPTYPE_Interface));
+			bool const bIsDelegate  = FBlueprintEditorUtils::IsDelegateSignatureGraph(InGraph);
+			bool const bIsMathExpression = FBlueprintEditorUtils::IsMathExpressionGraph(InGraph);
+
+			bGraphReadOnly = bIsInterface || bIsDelegate || bIsMathExpression;
+		}
+	}
+	
+	URigVMEditorSettings const* Settings = GetDefault<URigVMEditorSettings>();
+	if ((InGraph == nullptr) || !IsEditable(InGraph) || bGraphReadOnly)
+	{
+		return 0.0f;
+	}
+	else if ((InstructionsFadeCountdown > 0.0f) || (HasOpenActionMenu == InGraph))
+	{
+		return InstructionsFadeCountdown / RigVMNewEditorImpl::InstructionFadeDuration;
+	}
+	else if (RigVMNewEditorImpl::GraphHasUserPlacedNodes(InGraph))
+	{
+		return 0.0f;
+	}
+	return 1.0f;
+}
+
+void FRigVMNewEditor::ClearSelectionStateFor(FName SelectionOwner)
+{
+	if ( SelectionOwner == SelectionState_Graph() )
+	{
+		TArray< TSharedPtr<SDockTab> > GraphEditorTabs;
+		DocumentManager->FindAllTabsForFactory(GraphEditorTabFactoryPtr, /*out*/ GraphEditorTabs);
+
+		for (TSharedPtr<SDockTab>& GraphEditorTab : GraphEditorTabs)
+		{
+			TSharedRef<SGraphEditor> Editor = StaticCastSharedRef<SGraphEditor>((GraphEditorTab)->GetContent());
+
+			Editor->ClearSelectionSet();
+		}
+	}
+	else if ( SelectionOwner == SelectionState_GraphExplorer() )
+	{
+		if ( GraphExplorerWidget.IsValid() )
+		{
+			GraphExplorerWidget->ClearSelection();
+		}
+	}
+}
+
+void FRigVMNewEditor::CreateDefaultTabContents(const TArray<UBlueprint*> InBlueprints)
+{
+	UBlueprint* InBlueprint = InBlueprints.Num() == 1 ? InBlueprints[0] : nullptr;
+
+	// Cache off whether or not this is an interface, since it is used to govern multiple widget's behavior
+	const bool bIsInterface = (InBlueprint && InBlueprint->BlueprintType == BPTYPE_Interface);
+	const bool bIsMacro = (InBlueprint && InBlueprint->BlueprintType == BPTYPE_MacroLibrary);
+
+	if (InBlueprint)
+	{
+		// TODO sara-s
+		// this->BookmarksWidget =
+		// 	SNew(SBlueprintBookmarks)
+		// 		.EditorContext(SharedThis(this));
+	}
+
+	if (IsEditingSingleBlueprint())
+	{
+		//this->ReplaceReferencesWidget = SNew(SReplaceNodeReferences, SharedThis(this)); // TODO sara-s
+	}
+	
+	CompilerResultsListing = FCompilerResultsLog::GetBlueprintMessageLog(InBlueprint);
+	//CompilerResultsListing->OnMessageTokenClicked().AddSP(this, &FBlueprintEditor::OnLogTokenClicked); // TODO sara-s
+
+	// TODO sara-s
+	// FMessageLogModule& MessageLogModule = FModuleManager::LoadModuleChecked<FMessageLogModule>("MessageLog");
+	// CompilerResults = MessageLogModule.CreateLogListingWidget( CompilerResultsListing.ToSharedRef() );
+	// FindResults = SNew(SFindInBlueprints, SharedThis(this));
+	
+	// this->Inspector = 
+	// 	SNew(SRigVMDetailsInspector)
+	// 	.Editor(SharedRef())
+	// 	. HideNameArea(true)
+	// 	. ViewIdentifier(FName("BlueprintInspector"))
+	// 	//. OnFinishedChangingProperties( FOnFinishedChangingProperties::FDelegate::CreateSP(this, &FBlueprintEditor::OnFinishedChangingProperties) ); // TODO sara-s
+	// ;
+
+	if ( InBlueprints.Num() > 0 )
+	{
+		// Don't show the object name in defaults mode.
+		const bool bHideNameArea = bWasOpenedInDefaultsMode;
+
+		// this->DefaultEditor = 
+		// 	SNew(SRigVMDetailsInspector)
+		// 	.Editor(SharedRef())
+		// 	. ViewIdentifier(FName("BlueprintDefaults"))
+		// 	. IsEnabled(!bIsInterface)
+		// 	//. ShowPublicViewControl(this, &FBlueprintEditor::ShouldShowPublicViewControl) // TODO sara-s
+		// 	. ShowTitleArea(false)
+		// 	. HideNameArea(bHideNameArea)
+		// 	//. OnFinishedChangingProperties( FOnFinishedChangingProperties::FDelegate::CreateSP( this, &FBlueprintEditor::OnFinishedChangingProperties ) ); // TODO sara-s
+		// ;
+	}
+
+	if (InBlueprint && 
+		InBlueprint->ParentClass &&
+		InBlueprint->ParentClass->IsChildOf(AActor::StaticClass()) && 
+		InBlueprint->SimpleConstructionScript )
+	{
+		//CreateSubobjectEditors(); // TODO sara-s
+	}
+}
+
+const FSlateBrush* FRigVMNewEditor::GetGlyphForGraph(const UEdGraph* Graph, bool bInLargeIcon)
+{
+	const FSlateBrush* ReturnValue = FAppStyle::GetBrush( bInLargeIcon ? TEXT("GraphEditor.Function_24x") : TEXT("GraphEditor.Function_16x") );
+
+	check(Graph != nullptr);
+	const UEdGraphSchema* Schema = Graph->GetSchema();
+	if (Schema != nullptr)
+	{
+		const EGraphType GraphType = Schema->GetGraphType(Graph);
+		switch (GraphType)
+		{
+		default:
+		case GT_StateMachine:
+		case GT_Macro:
+		case GT_Animation:
+			check(false);
+			break;
+		case GT_Function:
+			{
+				bool bIsSubGraph = false;
+				{
+					if( Graph->GetOuter() )
+					{
+						//Check whether the outer is a composite node
+						if( Graph->GetOuter()->IsA( UK2Node_Composite::StaticClass() ) )
+						{
+							bIsSubGraph = true;
+						}
+					}
+				}
+				
+				//Check for subgraph
+				if( bIsSubGraph )
+				{
+					ReturnValue = FAppStyle::GetBrush( bInLargeIcon ? TEXT("GraphEditor.SubGraph_24x") : TEXT("GraphEditor.SubGraph_16x") );
+				}
+				else
+				{
+					ReturnValue = FAppStyle::GetBrush( bInLargeIcon ? TEXT("GraphEditor.Function_24x") : TEXT("GraphEditor.Function_16x") );
+				}
+			}
+			break;
+		case GT_Ubergraph:
+			{
+				ReturnValue = FAppStyle::GetBrush( bInLargeIcon ? TEXT("GraphEditor.EventGraph_24x") : TEXT("GraphEditor.EventGraph_16x") );
+			}
+			break;
+		}
+	}
+
+	return ReturnValue;
+}
+
+
+#undef LOCTEXT_NAMESPACE
diff --git a/Engine/Plugins/Runtime/RigVM/Source/RigVMEditor/Private/RigVMEditorModule.cpp b/Engine/Plugins/Runtime/RigVM/Source/RigVMEditor/Private/RigVMEditorModule.cpp
--- a/Engine/Plugins/Runtime/RigVM/Source/RigVMEditor/Private/RigVMEditorModule.cpp
+++ b/Engine/Plugins/Runtime/RigVM/Source/RigVMEditor/Private/RigVMEditorModule.cpp
@@ -75,9 +75,11 @@ void FRigVMEditorModule::StartupModule()
 		ReconstructAllNodesDelegateHandle = FBlueprintEditorUtils::OnReconstructAllNodesEvent.AddStatic(&FRigVMBlueprintUtils::HandleReconstructAllNodes);
 		RefreshAllNodesDelegateHandle = FBlueprintEditorUtils::OnRefreshAllNodesEvent.AddStatic(&FRigVMBlueprintUtils::HandleRefreshAllNodes);
 
+#if WITH_RIGVMLEGACYEDITOR
 		// Register Blueprint editor variable customization
 		FBlueprintEditorModule& BlueprintEditorModule = FModuleManager::LoadModuleChecked<FBlueprintEditorModule>("Kismet");
 		BlueprintVariableCustomizationHandle = BlueprintEditorModule.RegisterVariableCustomization(FProperty::StaticClass(), FOnGetVariableCustomizationInstance::CreateStatic(&FRigVMVariableDetailCustomization::MakeInstance));
+#endif
 
 		IAssetRegistry* AssetRegistry = IAssetRegistry::Get();
 		if (ensure(AssetRegistry))
@@ -154,8 +156,10 @@ void FRigVMEditorModule::StartupModuleCommon()
 	const URigVMBlueprint* BlueprintCDO = GetRigVMBlueprintCDO();
 	const UClass* EdGraphSchemaClass = BlueprintCDO->GetRigVMEdGraphSchemaClass();
 	const URigVMEdGraphSchema* SchemaCDO = CastChecked<URigVMEdGraphSchema>(EdGraphSchemaClass->GetDefaultObject());
+#if WITH_RIGVMLEGACYEDITOR
 	FBlueprintEditorModule& BlueprintEditorModule = FModuleManager::LoadModuleChecked<FBlueprintEditorModule>("Kismet");
 	BlueprintEditorModule.RegisterGraphCustomization(SchemaCDO, FOnGetGraphCustomizationInstance::CreateStatic(&FRigVMGraphDetailCustomization::MakeInstance, BlueprintClass));
+#endif
 
 	// Register to fixup newly created BPs
 	FKismetEditorUtilities::RegisterOnBlueprintCreatedCallback(this, URigVMHost::StaticClass(), FKismetEditorUtilities::FOnBlueprintCreated::CreateRaw(this, &FRigVMEditorModule::HandleNewBlueprintCreated));
@@ -182,8 +186,10 @@ void FRigVMEditorModule::ShutdownModule()
 			FBlueprintEditorUtils::OnReconstructAllNodesEvent.Remove(ReconstructAllNodesDelegateHandle);
 
 			// Unregister Blueprint editor variable customization
+#if WITH_RIGVMLEGACYEDITOR
 			FBlueprintEditorModule& BlueprintEditorModule = FModuleManager::LoadModuleChecked<FBlueprintEditorModule>("Kismet");
 			BlueprintEditorModule.UnregisterVariableCustomization(FProperty::StaticClass(), BlueprintVariableCustomizationHandle);
+#endif
 		}
 	}
 
@@ -197,8 +203,10 @@ void FRigVMEditorModule::ShutdownModuleCommon()
 		const URigVMBlueprint* BlueprintCDO = GetRigVMBlueprintCDO();
 		const UClass* EdGraphSchemaClass = BlueprintCDO->GetRigVMEdGraphSchemaClass();
 		const URigVMEdGraphSchema* SchemaCDO = CastChecked<URigVMEdGraphSchema>(EdGraphSchemaClass->GetDefaultObject());
+#if WITH_RIGVMLEGACYEDITOR
 		FBlueprintEditorModule& BlueprintEditorModule = FModuleManager::LoadModuleChecked<FBlueprintEditorModule>("Kismet");
 		BlueprintEditorModule.UnregisterGraphCustomization(SchemaCDO);
+#endif
 	}
 
 	// Unregister to fixup newly created BPs
diff --git a/Engine/Plugins/Runtime/RigVM/Source/RigVMEditor/Private/Widgets/SRigVMBulkEditWidget.cpp b/Engine/Plugins/Runtime/RigVM/Source/RigVMEditor/Private/Widgets/SRigVMBulkEditWidget.cpp
--- a/Engine/Plugins/Runtime/RigVM/Source/RigVMEditor/Private/Widgets/SRigVMBulkEditWidget.cpp
+++ b/Engine/Plugins/Runtime/RigVM/Source/RigVMEditor/Private/Widgets/SRigVMBulkEditWidget.cpp
@@ -15,6 +15,8 @@
 #include "Misc/UObjectToken.h"
 #include "Dialogs/Dialogs.h"
 #include "Framework/Application/SlateApplication.h"
+#include "Editor/RigVMNewEditor.h"
+
 #define LOCTEXT_NAMESPACE "SRigVMBulkEditWidget"
 
 //////////////////////////////////////////////////////////////////////////
@@ -278,7 +280,7 @@ void SRigVMBulkEditWidget::Tick(const FGeometry& AllottedGeometry, const double
 														
 												if(IAssetEditorInstance* Editor = GEditor->GetEditorSubsystem<UAssetEditorSubsystem>()->FindEditorForAsset(Blueprint, true))
 												{
-													if(FRigVMEditor* RigVMEditor = static_cast<FRigVMEditor*>(Editor))
+													if (FRigVMEditorBase* RigVMEditor = FRigVMEditorBase::GetFromAssetEditorInstance(Editor))
 													{
 														RigVMEditor->HandleJumpToHyperlink(WeakSubject.Get());
 													}
