Control Rig: New RigVMEditor, which does not inherit from BlueprintEditor

diff --git a/Engine/Plugins/Runtime/RigVM/Source/RigVMEditor/Private/Widgets/SRigVMEditorGraphExplorer.cpp b/Engine/Plugins/Runtime/RigVM/Source/RigVMEditor/Private/Widgets/SRigVMEditorGraphExplorer.cpp
--- a/Engine/Plugins/Runtime/RigVM/Source/RigVMEditor/Private/Widgets/SRigVMEditorGraphExplorer.cpp
+++ b/Engine/Plugins/Runtime/RigVM/Source/RigVMEditor/Private/Widgets/SRigVMEditorGraphExplorer.cpp
@@ -2,7 +2,7 @@
 
 #include "Widgets/SRigVMEditorGraphExplorer.h"
 
-#include "Editor/RigVMEditor.h"
+#include "Editor/RigVMNewEditor.h"
 #include "EdGraph/RigVMEdGraphSchema.h"
 #include "Widgets/SRigVMEditorGraphExplorerTreeView.h"
 
@@ -14,6 +14,7 @@
 // Ick.
 #include "GraphActionNode.h"
 #include "GraphEditorActions.h"
+#include "RigVMBlueprint.h"
 #include "SPinTypeSelector.h"
 #include "Kismet2/BlueprintEditorUtils.h"
 #include "RigVMModel/Nodes/RigVMAggregateNode.h"
@@ -21,6 +22,14 @@
 #include "Widgets/Layout/SScaleBox.h"
 #include "HAL/PlatformApplicationMisc.h"
 #include "ScopedTransaction.h"
+#if WITH_RIGVMLEGACYEDITOR
+#include "SKismetInspector.h"
+#include "Editor/RigVMLegacyEditor.h"
+#else
+#include "Editor/SRigVMDetailsInspector.h"
+#endif
+#include "EdGraph/RigVMEdGraph.h"
+#include "WorkflowOrientedApp/WorkflowTabManager.h"
 
 #define LOCTEXT_NAMESPACE "RigVMGraphExplorer"
 
@@ -86,10 +95,42 @@ FReply FRigVMGraphExplorerDragDropOp::DroppedOnPanel(const TSharedRef<SWidget>&
 	{
 		if(URigVMBlueprint* Blueprint = TargetRigGraph->GetBlueprint())
 		{
-			if (GEditor->GetEditorSubsystem<UAssetEditorSubsystem>()->OpenEditorForAsset(Blueprint))
+			//  Find the appropiate asset editor where the drop is happening
+			const TArray<IAssetEditorInstance*> AssetEditors = GEditor->GetEditorSubsystem<UAssetEditorSubsystem>()->FindEditorsForAsset(Blueprint);
+			IAssetEditorInstance* FocusedAssetEditor = nullptr;
+			for (IAssetEditorInstance* AssetEditor : AssetEditors)
 			{
-				FRigVMEditor* Editor = static_cast<FRigVMEditor*>(GEditor->GetEditorSubsystem<UAssetEditorSubsystem>()->FindEditorForAsset(Blueprint, /*bFocusIfOpen =*/true));
-				Editor->OnGraphNodeDropToPerform(SharedThis(this), TargetRigGraph, GraphPosition, ScreenPosition);
+				if (FRigVMEditorBase* RigVMEditor = FRigVMEditorBase::GetFromAssetEditorInstance(AssetEditor))
+				{
+					TWeakPtr<SGraphEditor> GraphEditor = RigVMEditor->GetFocusedGraphEditor();
+					TSharedPtr<SWidget> Widget = Panel->GetParentWidget();
+					while (Widget)
+					{
+						if (Widget == GraphEditor)
+						{
+							break;
+						}
+						Widget = Widget->GetParentWidget();
+					}
+					if (Widget)
+					{
+						FocusedAssetEditor = AssetEditor;
+						break;
+					}
+				}
+			}
+
+			if (!FocusedAssetEditor)
+			{
+				GEditor->GetEditorSubsystem<UAssetEditorSubsystem>()->OpenEditorForAsset(Blueprint);
+				FocusedAssetEditor = GEditor->GetEditorSubsystem<UAssetEditorSubsystem>()->FindEditorForAsset(Blueprint, /*bFocusIfOpen =*/true);
+			}
+			if (FocusedAssetEditor)
+			{
+				if (FRigVMEditorBase* RigVMEditor = FRigVMEditorBase::GetFromAssetEditorInstance(FocusedAssetEditor))
+				{
+					RigVMEditor->OnGraphNodeDropToPerform(SharedThis(this), TargetRigGraph, GraphPosition, ScreenPosition);
+				}
 			}
 		}
 	}
@@ -155,12 +196,12 @@ void FRigVMEditorGraphExplorerCommands::RegisterCommands()
 }
 
 
-void SRigVMEditorGraphExplorer::Construct(const FArguments& InArgs, TWeakPtr<FRigVMEditor> InRigVMEditor)
+void SRigVMEditorGraphExplorer::Construct(const FArguments& InArgs, TWeakPtr<FRigVMEditorBase> InRigVMEditor)
 {
 	bNeedsRefresh = false;
 	RigVMEditor = InRigVMEditor;
 
-	TSharedPtr<FRigVMEditor> Editor = RigVMEditor.Pin();
+	TSharedPtr<FRigVMEditorBase> Editor = RigVMEditor.Pin();
 	if (Editor)
 	{
 		Editor->OnRefresh().AddRaw(this, &SRigVMEditorGraphExplorer::Refresh);
@@ -169,12 +210,15 @@ void SRigVMEditorGraphExplorer::Construct(const FArguments& InArgs, TWeakPtr<FRi
 	RegisterCommands();
 
 	CreateWidgets();
+
+	LastPinType.ResetToDefaults();
+	LastPinType.PinCategory = TEXT("bool");
 }
 
 
 SRigVMEditorGraphExplorer::~SRigVMEditorGraphExplorer()
 {
-	TSharedPtr<FRigVMEditor> Editor = RigVMEditor.Pin();
+	TSharedPtr<FRigVMEditorBase> Editor = RigVMEditor.Pin();
 	if (Editor)
 	{
 		Editor->OnRefresh().RemoveAll(this);
@@ -197,9 +241,30 @@ void SRigVMEditorGraphExplorer::Tick(const FGeometry& InAllottedGeometry, const
 	}
 }
 
+FName SRigVMEditorGraphExplorer::GetSelectedVariableName() const
+{
+	TArray<FRigVMExplorerElementKey> Selected = TreeView->GetSelectedKeys();
+	if (Selected.Num() != 1)
+	{
+		return NAME_None;
+	}
+
+	if (Selected[0].Type != ERigVMExplorerElementType::Variable || Selected[0].Type != ERigVMExplorerElementType::LocalVariable)
+	{
+		return NAME_None;
+	}
+
+	return *Selected[0].Name;
+}
+
+void SRigVMEditorGraphExplorer::ClearSelection()
+{
+	TreeView->ClearSelection();
+}
+
 void SRigVMEditorGraphExplorer::RegisterCommands()
 {
-	TSharedPtr<FRigVMEditor> Editor = RigVMEditor.Pin();
+	TSharedPtr<FRigVMEditorBase> Editor = RigVMEditor.Pin();
 
 	if (Editor)
 	{
@@ -453,7 +518,7 @@ TArray<const URigVMGraph*> SRigVMEditorGraphExplorer::GetRootGraphs() const
 {
 	TArray<const URigVMGraph*> Graphs;
 	
-	TSharedPtr<FRigVMEditor> Editor = RigVMEditor.Pin();
+	TSharedPtr<FRigVMEditorBase> Editor = RigVMEditor.Pin();
 	if (!Editor)
 	{
 		return Graphs;
@@ -472,7 +537,7 @@ TArray<const URigVMGraph*> SRigVMEditorGraphExplorer::GetChildrenGraphs(const FS
 {
 	TArray<const URigVMGraph*> Children;
 	
-	TSharedPtr<FRigVMEditor> Editor = RigVMEditor.Pin();
+	TSharedPtr<FRigVMEditorBase> Editor = RigVMEditor.Pin();
 	if (!Editor)
 	{
 		return Children;
@@ -504,7 +569,7 @@ TArray<URigVMNode*> SRigVMEditorGraphExplorer::GetEventNodesInGraph(const FStrin
 {
 	TArray<URigVMNode*> Events;
 	
-	TSharedPtr<FRigVMEditor> Editor = RigVMEditor.Pin();
+	TSharedPtr<FRigVMEditorBase> Editor = RigVMEditor.Pin();
 	if (!Editor)
 	{
 		return Events;
@@ -542,7 +607,7 @@ TArray<URigVMLibraryNode*> SRigVMEditorGraphExplorer::GetFunctions() const
 {
 	TArray<URigVMLibraryNode*> Functions;
 	
-	TSharedPtr<FRigVMEditor> Editor = RigVMEditor.Pin();
+	TSharedPtr<FRigVMEditorBase> Editor = RigVMEditor.Pin();
 	if (!Editor)
 	{
 		return Functions;
@@ -560,7 +625,7 @@ TArray<URigVMLibraryNode*> SRigVMEditorGraphExplorer::GetFunctions() const
 TArray<FRigVMGraphVariableDescription> SRigVMEditorGraphExplorer::GetVariables() const
 {
 	TArray<FRigVMGraphVariableDescription> Variables;
-	TSharedPtr<FRigVMEditor> Editor = RigVMEditor.Pin();
+	TSharedPtr<FRigVMEditorBase> Editor = RigVMEditor.Pin();
 	if (!Editor)
 	{
 		return Variables;
@@ -572,7 +637,7 @@ TArray<FRigVMGraphVariableDescription> SRigVMEditorGraphExplorer::GetVariables()
 TArray<FRigVMGraphVariableDescription> SRigVMEditorGraphExplorer::GetLocalVariables() const
 {
 	TArray<FRigVMGraphVariableDescription> Variables;
-	TSharedPtr<FRigVMEditor> Editor = RigVMEditor.Pin();
+	TSharedPtr<FRigVMEditorBase> Editor = RigVMEditor.Pin();
 	if (!Editor)
 	{
 		return Variables;
@@ -592,7 +657,7 @@ TArray<FRigVMGraphVariableDescription> SRigVMEditorGraphExplorer::GetLocalVariab
 FText SRigVMEditorGraphExplorer::GetGraphDisplayName(const FString& InGraphPath) const
 {
 	FText DisplayName;
-	TSharedPtr<FRigVMEditor> Editor = RigVMEditor.Pin();
+	TSharedPtr<FRigVMEditorBase> Editor = RigVMEditor.Pin();
 	if (!Editor)
 	{
 		return DisplayName;
@@ -615,7 +680,7 @@ FText SRigVMEditorGraphExplorer::GetGraphDisplayName(const FString& InGraphPath)
 FText SRigVMEditorGraphExplorer::GetEventDisplayName(const FString& InNodePath) const
 {
 	FText DisplayName;
-	TSharedPtr<FRigVMEditor> Editor = RigVMEditor.Pin();
+	TSharedPtr<FRigVMEditorBase> Editor = RigVMEditor.Pin();
 	if (!Editor)
 	{
 		return DisplayName;
@@ -633,7 +698,7 @@ FText SRigVMEditorGraphExplorer::GetEventDisplayName(const FString& InNodePath)
 const FSlateBrush* SRigVMEditorGraphExplorer::GetGraphIcon(const FString& InGraphPath) const
 {
 	const FSlateBrush* Brush = nullptr;
-	TSharedPtr<FRigVMEditor> Editor = RigVMEditor.Pin();
+	TSharedPtr<FRigVMEditorBase> Editor = RigVMEditor.Pin();
 	if (!Editor)
 	{
 		return Brush;
@@ -657,7 +722,7 @@ const FSlateBrush* SRigVMEditorGraphExplorer::GetGraphIcon(const FString& InGrap
 FText SRigVMEditorGraphExplorer::GetGraphTooltip(const FString& InGraphPath) const
 {
 	FText Tooltip;
-	TSharedPtr<FRigVMEditor> Editor = RigVMEditor.Pin();
+	TSharedPtr<FRigVMEditorBase> Editor = RigVMEditor.Pin();
 	if (!Editor)
 	{
 		return Tooltip;
@@ -679,7 +744,7 @@ FText SRigVMEditorGraphExplorer::GetGraphTooltip(const FString& InGraphPath) con
 
 void SRigVMEditorGraphExplorer::OnGraphClicked(const FString& InGraphPath)
 {
-	TSharedPtr<FRigVMEditor> Editor = RigVMEditor.Pin();
+	TSharedPtr<FRigVMEditorBase> Editor = RigVMEditor.Pin();
 	if (!Editor)
 	{
 		return;
@@ -695,26 +760,42 @@ void SRigVMEditorGraphExplorer::OnGraphClicked(const FString& InGraphPath)
 			const UEdGraphSchema* Schema = EdGraph->GetSchema();
 			check(Schema != nullptr);
 			Schema->GetGraphDisplayInformation(*EdGraph, DisplayInfo);
-			TSharedRef<SKismetInspector> CurrentInspector = Editor->GetInspector();
-			CurrentInspector->ShowDetailsForSingleObject(EdGraph, SKismetInspector::FShowDetailsOptions(DisplayInfo.PlainName));
+#if WITH_RIGVMLEGACYEDITOR
+			if (TSharedPtr<SKismetInspector> CurrentInspector = Editor->GetInspector())
+			{
+				CurrentInspector->ShowDetailsForSingleObject(EdGraph, SKismetInspector::FShowDetailsOptions(DisplayInfo.PlainName));
+			}
+#else
+			if (TSharedPtr<SRigVMDetailsInspector> CurrentInspector = Editor->GetInspector())
+			{
+				CurrentInspector->ShowDetailsForSingleObject(EdGraph, SRigVMDetailsInspector::FShowDetailsOptions(DisplayInfo.PlainName));
+			}
+#endif
 		}
 	}
 }
 
 void SRigVMEditorGraphExplorer::OnEventClicked(const FString& InEventPath)
 {
-	TSharedPtr<FRigVMEditor> Editor = RigVMEditor.Pin();
+	TSharedPtr<FRigVMEditorBase> Editor = RigVMEditor.Pin();
 	if (!Editor)
 	{
 		return;
 	}
-	TSharedRef<SKismetInspector> CurrentInspector = Editor->GetInspector();
-	CurrentInspector->ShowDetailsForObjects(TArray<UObject*>());
+#if WITH_RIGVMLEGACYEDITOR
+	TSharedPtr<SKismetInspector> CurrentInspector = Editor->GetInspector();
+#else
+	TSharedPtr<SRigVMDetailsInspector> CurrentInspector = Editor->GetInspector();
+#endif
+	if (CurrentInspector.IsValid())
+	{
+		CurrentInspector->ShowDetailsForObjects(TArray<UObject*>());
+	}
 }
 
 void SRigVMEditorGraphExplorer::OnFunctionClicked(const FString& InFunctionPath)
 {
-	TSharedPtr<FRigVMEditor> Editor = RigVMEditor.Pin();
+	TSharedPtr<FRigVMEditorBase> Editor = RigVMEditor.Pin();
 	if (!Editor)
 	{
 		return;
@@ -732,8 +813,17 @@ void SRigVMEditorGraphExplorer::OnFunctionClicked(const FString& InFunctionPath)
 				const UEdGraphSchema* Schema = EdGraph->GetSchema();
 				check(Schema != nullptr);
 				Schema->GetGraphDisplayInformation(*EdGraph, DisplayInfo);
-				TSharedRef<SKismetInspector> CurrentInspector = Editor->GetInspector();
-				CurrentInspector->ShowDetailsForSingleObject(EdGraph, SKismetInspector::FShowDetailsOptions(DisplayInfo.PlainName));
+#if WITH_RIGVMLEGACYEDITOR
+				if (TSharedPtr<SKismetInspector> CurrentInspector = Editor->GetInspector())
+				{
+					CurrentInspector->ShowDetailsForSingleObject(EdGraph, SKismetInspector::FShowDetailsOptions(DisplayInfo.PlainName));
+				}
+#else
+				if (TSharedPtr<SRigVMDetailsInspector> CurrentInspector = Editor->GetInspector())
+				{
+					CurrentInspector->ShowDetailsForSingleObject(EdGraph, SRigVMDetailsInspector::FShowDetailsOptions(DisplayInfo.PlainName));
+				}
+#endif
 			}
 		}
 	}
@@ -741,7 +831,7 @@ void SRigVMEditorGraphExplorer::OnFunctionClicked(const FString& InFunctionPath)
 
 void SRigVMEditorGraphExplorer::OnVariableClicked(const FRigVMExplorerElementKey& InVariable)
 {
-	TSharedPtr<FRigVMEditor> Editor = RigVMEditor.Pin();
+	TSharedPtr<FRigVMEditorBase> Editor = RigVMEditor.Pin();
 	if (!Editor)
 	{
 		return;
@@ -751,13 +841,21 @@ void SRigVMEditorGraphExplorer::OnVariableClicked(const FRigVMExplorerElementKey
 	if (InVariable.Type == ERigVMExplorerElementType::Variable)
 	{
 		FRigVMClient* RigVMClient = Blueprint->GetRigVMClient();
-		TSharedRef<SKismetInspector> CurrentInspector = Editor->GetInspector();
+#if WITH_RIGVMLEGACYEDITOR
+		TSharedPtr<SKismetInspector> CurrentInspector = Editor->GetInspector();
 		SKismetInspector::FShowDetailsOptions Options(FText::FromString(InVariable.Name));
+#else
+		TSharedPtr<SRigVMDetailsInspector> CurrentInspector = Editor->GetInspector();
+		SRigVMDetailsInspector::FShowDetailsOptions Options(FText::FromString(InVariable.Name));
+#endif
 		Options.bForceRefresh = true;
 
 		FProperty* Prop = Blueprint->SkeletonGeneratedClass->FindPropertyByName(*InVariable.Name);
 		UPropertyWrapper* PropWrap = (Prop ? Prop->GetUPropertyWrapper() : nullptr);
-		CurrentInspector->ShowDetailsForSingleObject(PropWrap, Options);
+		if (CurrentInspector.IsValid())
+		{
+			CurrentInspector->ShowDetailsForSingleObject(PropWrap, Options);
+		}
 	}
 	else if (InVariable.Type == ERigVMExplorerElementType::LocalVariable)
 	{
@@ -769,7 +867,7 @@ void SRigVMEditorGraphExplorer::OnVariableClicked(const FRigVMExplorerElementKey
 
 void SRigVMEditorGraphExplorer::OnGraphDoubleClicked(const FString& InGraphPath)
 {
-	TSharedPtr<FRigVMEditor> Editor = RigVMEditor.Pin();
+	TSharedPtr<FRigVMEditorBase> Editor = RigVMEditor.Pin();
 	if (!Editor)
 	{
 		return;
@@ -788,7 +886,7 @@ void SRigVMEditorGraphExplorer::OnGraphDoubleClicked(const FString& InGraphPath)
 
 void SRigVMEditorGraphExplorer::OnEventDoubleClicked(const FString& InEventPath)
 {
-	TSharedPtr<FRigVMEditor> Editor = RigVMEditor.Pin();
+	TSharedPtr<FRigVMEditorBase> Editor = RigVMEditor.Pin();
 	if (!Editor)
 	{
 		return;
@@ -805,7 +903,7 @@ void SRigVMEditorGraphExplorer::OnEventDoubleClicked(const FString& InEventPath)
 
 void SRigVMEditorGraphExplorer::OnFunctionDoubleClicked(const FString& InFunctionPath)
 {
-	TSharedPtr<FRigVMEditor> Editor = RigVMEditor.Pin();
+	TSharedPtr<FRigVMEditorBase> Editor = RigVMEditor.Pin();
 	if (!Editor)
 	{
 		return;
@@ -827,7 +925,7 @@ void SRigVMEditorGraphExplorer::OnFunctionDoubleClicked(const FString& InFunctio
 
 bool SRigVMEditorGraphExplorer::OnSetFunctionCategory(const FString& InFunctionPath, const FString& InCategory)
 {
-	TSharedPtr<FRigVMEditor> Editor = RigVMEditor.Pin();
+	TSharedPtr<FRigVMEditorBase> Editor = RigVMEditor.Pin();
 	if (!Editor.IsValid())
 	{
 		return false;
@@ -851,7 +949,7 @@ bool SRigVMEditorGraphExplorer::OnSetFunctionCategory(const FString& InFunctionP
 
 FString SRigVMEditorGraphExplorer::OnGetFunctionCategory(const FString& InFunctionPath) const
 {
-	TSharedPtr<FRigVMEditor> Editor = RigVMEditor.Pin();
+	TSharedPtr<FRigVMEditorBase> Editor = RigVMEditor.Pin();
 	if (!Editor.IsValid())
 	{
 		return FString();
@@ -870,7 +968,7 @@ FString SRigVMEditorGraphExplorer::OnGetFunctionCategory(const FString& InFuncti
 
 bool SRigVMEditorGraphExplorer::OnSetVariableCategory(const FString& InVariable, const FString& InCategory)
 {
-	TSharedPtr<FRigVMEditor> Editor = RigVMEditor.Pin();
+	TSharedPtr<FRigVMEditorBase> Editor = RigVMEditor.Pin();
 	if (!Editor.IsValid())
 	{
 		return false;
@@ -883,7 +981,7 @@ bool SRigVMEditorGraphExplorer::OnSetVariableCategory(const FString& InVariable,
 
 FString SRigVMEditorGraphExplorer::OnGetVariableCategory(const FString& InVariable) const
 {
-	TSharedPtr<FRigVMEditor> Editor = RigVMEditor.Pin();
+	TSharedPtr<FRigVMEditorBase> Editor = RigVMEditor.Pin();
 	if (!Editor.IsValid())
 	{
 		return FString();
@@ -894,7 +992,7 @@ FString SRigVMEditorGraphExplorer::OnGetVariableCategory(const FString& InVariab
 
 FEdGraphPinType SRigVMEditorGraphExplorer::OnGetVariablePinType(const FRigVMExplorerElementKey& InVariable)
 {
-	TSharedPtr<FRigVMEditor> Editor = RigVMEditor.Pin();
+	TSharedPtr<FRigVMEditorBase> Editor = RigVMEditor.Pin();
 	if (!Editor.IsValid())
 	{
 		return FEdGraphPinType();
@@ -930,7 +1028,7 @@ FEdGraphPinType SRigVMEditorGraphExplorer::OnGetVariablePinType(const FRigVMExpl
 
 bool SRigVMEditorGraphExplorer::OnSetVariablePinType(const FRigVMExplorerElementKey& InVariable, const FEdGraphPinType& InType)
 {
-	TSharedPtr<FRigVMEditor> Editor = RigVMEditor.Pin();
+	TSharedPtr<FRigVMEditorBase> Editor = RigVMEditor.Pin();
 	if (!Editor.IsValid())
 	{
 		return false;
@@ -945,6 +1043,7 @@ bool SRigVMEditorGraphExplorer::OnSetVariablePinType(const FRigVMExplorerElement
 			if (Variable.Name == InVariable.Name)
 			{
 				FBlueprintEditorUtils::ChangeMemberVariableType(Blueprint, *InVariable.Name, InType);
+				SetLastPinTypeUsed(InType);
 				return true;
 			}
 		}
@@ -959,6 +1058,8 @@ bool SRigVMEditorGraphExplorer::OnSetVariablePinType(const FRigVMExplorerElement
 				UObject* NewCPPTypeObject = nullptr;
 				RigVMTypeUtils::CPPTypeFromPinType(InType, NewCPPType, &NewCPPTypeObject);
 				Controller->SetLocalVariableType(*InVariable.Name, NewCPPType, NewCPPTypeObject, true, true);
+				SetLastPinTypeUsed(InType);
+				return true;
 			}
 		}
 	}
@@ -968,7 +1069,7 @@ bool SRigVMEditorGraphExplorer::OnSetVariablePinType(const FRigVMExplorerElement
 
 bool SRigVMEditorGraphExplorer::OnIsVariablePublic(const FString& InVariableName) const
 {
-	TSharedPtr<FRigVMEditor> Editor = RigVMEditor.Pin();
+	TSharedPtr<FRigVMEditorBase> Editor = RigVMEditor.Pin();
 	if (!Editor.IsValid())
 	{
 		return false;
@@ -993,7 +1094,7 @@ bool SRigVMEditorGraphExplorer::OnIsVariablePublic(const FString& InVariableName
 
 bool SRigVMEditorGraphExplorer::OnToggleVariablePublic(const FString& InVariableName) const
 {
-	TSharedPtr<FRigVMEditor> Editor = RigVMEditor.Pin();
+	TSharedPtr<FRigVMEditorBase> Editor = RigVMEditor.Pin();
 	if (!Editor.IsValid())
 	{
 		return false;
@@ -1008,7 +1109,7 @@ bool SRigVMEditorGraphExplorer::OnToggleVariablePublic(const FString& InVariable
 
 bool SRigVMEditorGraphExplorer::IsFunctionFocused() const
 {
-	TSharedPtr<FRigVMEditor> Editor = RigVMEditor.Pin();
+	TSharedPtr<FRigVMEditorBase> Editor = RigVMEditor.Pin();
 	if (!Editor.IsValid())
 	{
 		return false;
@@ -1039,7 +1140,7 @@ TArray<TSharedPtr<IPinTypeSelectorFilter>> SRigVMEditorGraphExplorer::GetCustomP
 
 TSharedPtr<SWidget> SRigVMEditorGraphExplorer::OnContextMenuOpening()
 {
-	TSharedPtr<FRigVMEditor> Editor = RigVMEditor.Pin();
+	TSharedPtr<FRigVMEditorBase> Editor = RigVMEditor.Pin();
 	if (!Editor.IsValid())
 	{
 		return TSharedPtr<SWidget>();
@@ -1130,7 +1231,7 @@ FReply SRigVMEditorGraphExplorer::OnDragDetected(const FGeometry& MyGeometry, co
 
 void SRigVMEditorGraphExplorer::OnCopy()
 {
-	TSharedPtr<FRigVMEditor> Editor = RigVMEditor.Pin();
+	TSharedPtr<FRigVMEditorBase> Editor = RigVMEditor.Pin();
 	if (!Editor)
 	{
 		return;
@@ -1256,7 +1357,7 @@ void SRigVMEditorGraphExplorer::OnDuplicate()
 
 bool SRigVMEditorGraphExplorer::CanDuplicate() const
 {
-	TSharedPtr<FRigVMEditor> Editor = RigVMEditor.Pin();
+	TSharedPtr<FRigVMEditorBase> Editor = RigVMEditor.Pin();
 	if (!Editor)
 	{
 		return false;
@@ -1325,7 +1426,7 @@ bool SRigVMEditorGraphExplorer::CanPasteGeneric() const
 
 void SRigVMEditorGraphExplorer::OnPasteFunction()
 {
-	TSharedPtr<FRigVMEditor> Editor = RigVMEditor.Pin();
+	TSharedPtr<FRigVMEditorBase> Editor = RigVMEditor.Pin();
 	if (!Editor)
 	{
 		return;
@@ -1339,7 +1440,7 @@ void SRigVMEditorGraphExplorer::OnPasteFunction()
 
 bool SRigVMEditorGraphExplorer::CanPasteFunction() const
 {
-	TSharedPtr<FRigVMEditor> Editor = RigVMEditor.Pin();
+	TSharedPtr<FRigVMEditorBase> Editor = RigVMEditor.Pin();
 	if (!Editor)
 	{
 		return false;
@@ -1353,7 +1454,7 @@ bool SRigVMEditorGraphExplorer::CanPasteFunction() const
 
 void SRigVMEditorGraphExplorer::OnPasteVariable()
 {
-	TSharedPtr<FRigVMEditor> Editor = RigVMEditor.Pin();
+	TSharedPtr<FRigVMEditorBase> Editor = RigVMEditor.Pin();
 	if (!Editor)
 	{
 		return;
@@ -1389,7 +1490,7 @@ void SRigVMEditorGraphExplorer::OnPasteVariable()
 
 bool SRigVMEditorGraphExplorer::CanPasteVariable() const
 {
-	TSharedPtr<FRigVMEditor> Editor = RigVMEditor.Pin();
+	TSharedPtr<FRigVMEditorBase> Editor = RigVMEditor.Pin();
 	if (!Editor)
 	{
 		return false;
@@ -1412,7 +1513,7 @@ bool SRigVMEditorGraphExplorer::CanPasteVariable() const
 
 void SRigVMEditorGraphExplorer::OnPasteLocalVariable()
 {
-	TSharedPtr<FRigVMEditor> Editor = RigVMEditor.Pin();
+	TSharedPtr<FRigVMEditorBase> Editor = RigVMEditor.Pin();
 	if (!Editor)
 	{
 		return;
@@ -1459,7 +1560,7 @@ bool SRigVMEditorGraphExplorer::IsShowingEmptySections() const
 
 void SRigVMEditorGraphExplorer::OnOpenGraph(bool bOnNewTab)
 {
-	TSharedPtr<FRigVMEditor> Editor = RigVMEditor.Pin();
+	TSharedPtr<FRigVMEditorBase> Editor = RigVMEditor.Pin();
 
 	if (Editor)
 	{
@@ -1492,17 +1593,17 @@ bool SRigVMEditorGraphExplorer::CanOpenGraph() const
 
 void SRigVMEditorGraphExplorer::OnCreateGraph()
 {
-	TSharedPtr<FRigVMEditor> Editor = RigVMEditor.Pin();
+	TSharedPtr<FRigVMEditorBase> Editor = RigVMEditor.Pin();
 	if (Editor)
 	{
-		Editor->NewDocument_OnClicked(FBlueprintEditor::CGT_NewEventGraph);
+		Editor->NewDocument_OnClicked(FRigVMEditorBase::CGT_NewEventGraph);
 	}
 }
 
 
 bool SRigVMEditorGraphExplorer::CanCreateGraph() const
 {
-	TSharedPtr<FRigVMEditor> Editor = RigVMEditor.Pin();
+	TSharedPtr<FRigVMEditorBase> Editor = RigVMEditor.Pin();
 	if (Editor)
 	{
 		return Editor->InEditingMode();
@@ -1513,10 +1614,10 @@ bool SRigVMEditorGraphExplorer::CanCreateGraph() const
 
 void SRigVMEditorGraphExplorer::OnCreateFunction()
 {
-	TSharedPtr<FRigVMEditor> Editor = RigVMEditor.Pin();
+	TSharedPtr<FRigVMEditorBase> Editor = RigVMEditor.Pin();
 	if (Editor)
 	{
-		Editor->NewDocument_OnClicked(FBlueprintEditor::CGT_NewFunctionGraph);
+		Editor->NewDocument_OnClicked(FRigVMEditorBase::CGT_NewFunctionGraph);
 	}
 }
 
@@ -1527,7 +1628,7 @@ bool SRigVMEditorGraphExplorer::CanCreateFunction() const
 
 void SRigVMEditorGraphExplorer::OnCreateVariable()
 {
-	TSharedPtr<FRigVMEditor> Editor = RigVMEditor.Pin();
+	TSharedPtr<FRigVMEditorBase> Editor = RigVMEditor.Pin();
 	if (Editor)
 	{
 		Editor->OnAddNewVariable();
@@ -1541,7 +1642,7 @@ bool SRigVMEditorGraphExplorer::CanCreateVariable() const
 
 void SRigVMEditorGraphExplorer::OnCreateLocalVariable()
 {
-	TSharedPtr<FRigVMEditor> Editor = RigVMEditor.Pin();
+	TSharedPtr<FRigVMEditorBase> Editor = RigVMEditor.Pin();
 	if (!Editor)
 	{
 		return;
@@ -1553,7 +1654,7 @@ void SRigVMEditorGraphExplorer::OnCreateLocalVariable()
 
 bool SRigVMEditorGraphExplorer::CanCreateLocalVariable() const
 {
-	TSharedPtr<FRigVMEditor> Editor = RigVMEditor.Pin();
+	TSharedPtr<FRigVMEditorBase> Editor = RigVMEditor.Pin();
 	if (!Editor)
 	{
 		return false;
@@ -1564,7 +1665,7 @@ bool SRigVMEditorGraphExplorer::CanCreateLocalVariable() const
 
 void SRigVMEditorGraphExplorer::OnAddFunctionVariant()
 {
-	TSharedPtr<FRigVMEditor> Editor = RigVMEditor.Pin();
+	TSharedPtr<FRigVMEditorBase> Editor = RigVMEditor.Pin();
 
 	if (!Editor)
 	{
@@ -1591,7 +1692,7 @@ void SRigVMEditorGraphExplorer::OnAddFunctionVariant()
 
 bool SRigVMEditorGraphExplorer::CanAddFunctionVariant() const
 {
-	TSharedPtr<FRigVMEditor> Editor = RigVMEditor.Pin();
+	TSharedPtr<FRigVMEditorBase> Editor = RigVMEditor.Pin();
 
 	if (!Editor)
 	{
@@ -1609,7 +1710,7 @@ bool SRigVMEditorGraphExplorer::CanAddFunctionVariant() const
 
 void SRigVMEditorGraphExplorer::OnDeleteEntry()
 {
-	TSharedPtr<FRigVMEditor> Editor = RigVMEditor.Pin();
+	TSharedPtr<FRigVMEditorBase> Editor = RigVMEditor.Pin();
 
 	if (Editor)
 	{
@@ -1671,7 +1772,7 @@ void SRigVMEditorGraphExplorer::OnDeleteEntry()
 
 bool SRigVMEditorGraphExplorer::CanDeleteEntry() const
 {
-	TSharedPtr<FRigVMEditor> Editor = RigVMEditor.Pin();
+	TSharedPtr<FRigVMEditorBase> Editor = RigVMEditor.Pin();
 
 	if (!Editor)
 	{
@@ -1730,7 +1831,7 @@ bool SRigVMEditorGraphExplorer::CanRenameEntry() const
 bool SRigVMEditorGraphExplorer::OnRenameGraph(const FString& InOldPath, const FString& InNewPath)
 {
 	bool bResult = false;
-	TSharedPtr<FRigVMEditor> Editor = RigVMEditor.Pin();
+	TSharedPtr<FRigVMEditorBase> Editor = RigVMEditor.Pin();
 
 	if (Editor)
 	{
@@ -1748,7 +1849,7 @@ bool SRigVMEditorGraphExplorer::OnRenameGraph(const FString& InOldPath, const FS
 bool SRigVMEditorGraphExplorer::OnCanRenameGraph(const FString& InOldPath, const FString& InNewPath, FText& OutErrorMessage) const
 {
 	bool bResult = false;
-	TSharedPtr<FRigVMEditor> Editor = RigVMEditor.Pin();
+	TSharedPtr<FRigVMEditorBase> Editor = RigVMEditor.Pin();
 
 	if (Editor)
 	{
@@ -1779,7 +1880,7 @@ bool SRigVMEditorGraphExplorer::OnCanRenameGraph(const FString& InOldPath, const
 bool SRigVMEditorGraphExplorer::OnRenameFunction(const FString& InOldPath, const FString& InNewPath)
 {
 	bool bResult = false;
-	TSharedPtr<FRigVMEditor> Editor = RigVMEditor.Pin();
+	TSharedPtr<FRigVMEditorBase> Editor = RigVMEditor.Pin();
 
 	if (Editor)
 	{
@@ -1803,7 +1904,7 @@ bool SRigVMEditorGraphExplorer::OnRenameFunction(const FString& InOldPath, const
 bool SRigVMEditorGraphExplorer::OnCanRenameFunction(const FString& InOldPath, const FString& InNewPath, FText& OutErrorMessage) const
 {
 	bool bResult = false;
-	TSharedPtr<FRigVMEditor> Editor = RigVMEditor.Pin();
+	TSharedPtr<FRigVMEditorBase> Editor = RigVMEditor.Pin();
 
 	if (Editor)
 	{
@@ -1826,7 +1927,7 @@ bool SRigVMEditorGraphExplorer::OnCanRenameFunction(const FString& InOldPath, co
 
 bool SRigVMEditorGraphExplorer::OnRenameVariable(const FRigVMExplorerElementKey& InOldKey, const FString& InNewName)
 {
-	TSharedPtr<FRigVMEditor> Editor = RigVMEditor.Pin();
+	TSharedPtr<FRigVMEditorBase> Editor = RigVMEditor.Pin();
 	if (!Editor)
 	{
 		return false;
@@ -1893,7 +1994,7 @@ void SRigVMEditorGraphExplorer::OnFindReference(bool bSearchAllBlueprints)
 		return;
 	}
 
-	TSharedPtr<FRigVMEditor> Editor = RigVMEditor.Pin();
+	TSharedPtr<FRigVMEditorBase> Editor = RigVMEditor.Pin();
 	URigVMBlueprint* Blueprint = Editor->GetRigVMBlueprint();
 	const URigVMEdGraphSchema* Schema = GetDefault<URigVMEdGraphSchema>();
 
diff --git a/Engine/Plugins/Runtime/RigVM/Source/RigVMEditor/Private/Widgets/SRigVMEditorGraphExplorerTreeView.cpp b/Engine/Plugins/Runtime/RigVM/Source/RigVMEditor/Private/Widgets/SRigVMEditorGraphExplorerTreeView.cpp
--- a/Engine/Plugins/Runtime/RigVM/Source/RigVMEditor/Private/Widgets/SRigVMEditorGraphExplorerTreeView.cpp
+++ b/Engine/Plugins/Runtime/RigVM/Source/RigVMEditor/Private/Widgets/SRigVMEditorGraphExplorerTreeView.cpp
@@ -3,7 +3,9 @@
 #include "Widgets/SRigVMEditorGraphExplorerTreeView.h"
 
 #include "SPinTypeSelector.h"
-#include "Editor/RigVMEditor.h"
+#include "EdGraph/RigVMEdGraphSchema.h"
+#include "Editor/RigVMNewEditor.h"
+#include "RigVMModel/Nodes/RigVMLibraryNode.h"
 #include "Widgets/SRigVMEditorGraphExplorer.h"
 
 #include "Widgets/Text/SInlineEditableTextBlock.h"
diff --git a/Engine/Plugins/Runtime/RigVM/Source/RigVMEditor/Private/Widgets/SRigVMExecutionStackView.cpp b/Engine/Plugins/Runtime/RigVM/Source/RigVMEditor/Private/Widgets/SRigVMExecutionStackView.cpp
--- a/Engine/Plugins/Runtime/RigVM/Source/RigVMEditor/Private/Widgets/SRigVMExecutionStackView.cpp
+++ b/Engine/Plugins/Runtime/RigVM/Source/RigVMEditor/Private/Widgets/SRigVMExecutionStackView.cpp
@@ -7,7 +7,7 @@
 #include "Widgets/Layout/SSpacer.h"
 #include "Widgets/Input/SSearchBox.h"
 #include "Widgets/Input/SNumericEntryBox.h"
-#include "Editor/RigVMEditor.h"
+#include "Editor/RigVMNewEditor.h"
 #include "Editor/RigVMExecutionStackCommands.h"
 #include "RigVMHost.h"
 #include "RigVMBlueprintGeneratedClass.h"
@@ -247,22 +247,22 @@ FText SRigStackItem::GetDurationText() const
 
 SRigVMExecutionStackView::~SRigVMExecutionStackView()
 {
-	if (RigVMEditor.IsValid())
+	if (TSharedPtr<FRigVMEditorBase> Editor = WeakEditor.Pin())
 	{
-		if (OnModelModified.IsValid() && RigVMEditor.Pin()->GetRigVMBlueprint())
+		if (OnModelModified.IsValid() && Editor->GetRigVMBlueprint())
 		{
-			RigVMEditor.Pin()->GetRigVMBlueprint()->OnModified().Remove(OnModelModified);
+			Editor->GetRigVMBlueprint()->OnModified().Remove(OnModelModified);
 		}
 		if (OnHostInitializedHandle.IsValid())
 		{
-			if(RigVMEditor.Pin()->GetRigVMHost())
+			if(Editor->GetRigVMHost())
 			{
-				RigVMEditor.Pin()->GetRigVMHost()->OnInitialized_AnyThread().Remove(OnHostInitializedHandle);
+				Editor->GetRigVMHost()->OnInitialized_AnyThread().Remove(OnHostInitializedHandle);
 			}
 		}
 		if (OnPreviewHostUpdatedHandle.IsValid())
 		{
-			RigVMEditor.Pin()->OnPreviewHostUpdated().Remove(OnPreviewHostUpdatedHandle);
+			Editor->OnPreviewHostUpdated().Remove(OnPreviewHostUpdatedHandle);
 		}
 	}
 	if (RigVMBlueprint.IsValid() && OnVMCompiledHandle.IsValid())
@@ -271,10 +271,10 @@ SRigVMExecutionStackView::~SRigVMExecutionStackView()
 	}
 }
 
-void SRigVMExecutionStackView::Construct( const FArguments& InArgs, TSharedRef<FRigVMEditor> InRigVMEditor)
+void SRigVMExecutionStackView::Construct( const FArguments& InArgs, TSharedRef<FRigVMEditorBase> InRigVMEditor)
 {
-	RigVMEditor = InRigVMEditor;
-	RigVMBlueprint = RigVMEditor.Pin()->GetRigVMBlueprint();
+	WeakEditor = InRigVMEditor;
+	RigVMBlueprint = InRigVMEditor->GetRigVMBlueprint();
 	CommandList = MakeShared<FUICommandList>();
 	bSuspendModelNotifications = false;
 	bSuspendControllerSelection = false;
@@ -348,18 +348,18 @@ void SRigVMExecutionStackView::Construct( const FArguments& InArgs, TSharedRef<F
 		OnVMCompiledHandle = RigVMBlueprint->OnVMCompiled().AddSP(this, &SRigVMExecutionStackView::OnVMCompiled);
 		OnModelModified = RigVMBlueprint->OnModified().AddSP(this, &SRigVMExecutionStackView::HandleModifiedEvent);
 
-		if (RigVMEditor.IsValid())
+		if (TSharedPtr<FRigVMEditorBase> Editor = WeakEditor.Pin())
 		{
-			if (URigVMHost* RigVMHost = RigVMEditor.Pin()->GetRigVMHost())
+			if (URigVMHost* RigVMHost = Editor->GetRigVMHost())
 			{
 				OnVMCompiled(RigVMBlueprint.Get(),  RigVMHost->GetVM(),  RigVMHost->GetRigVMExtendedExecuteContext());
 			}
 		}
 	}
 
-	if (RigVMEditor.IsValid())
+	if (TSharedPtr<FRigVMEditorBase> Editor = WeakEditor.Pin())
 	{
-		OnPreviewHostUpdatedHandle = RigVMEditor.Pin()->OnPreviewHostUpdated().AddSP(this, &SRigVMExecutionStackView::HandlePreviewHostUpdated);
+		OnPreviewHostUpdatedHandle = Editor->OnPreviewHostUpdated().AddSP(this, &SRigVMExecutionStackView::HandlePreviewHostUpdated);
 	}
 }
 
@@ -385,7 +385,11 @@ void SRigVMExecutionStackView::OnSelectionChanged(TSharedPtr<FRigStackEntry> Sel
 			return;
 		}
 
-		URigVMHost* Host = RigVMEditor.Pin()->GetRigVMHost();
+		URigVMHost* Host = nullptr;
+		if (TSharedPtr<FRigVMEditorBase> Editor = WeakEditor.Pin())
+		{
+			Host = Editor->GetRigVMHost();
+		}
 		if (Host == nullptr || Host->GetVM() == nullptr)
 		{
 			return;
@@ -454,8 +458,6 @@ void SRigVMExecutionStackView::PopulateStackView(URigVM* InVM, FRigVMExtendedExe
 {
 	if (InVM && InVMContext)
 	{
-		URigVMBlueprint* Blueprint = RigVMEditor.Pin()->GetRigVMBlueprint();
-
 		const FRigVMInstructionArray Instructions = InVM->GetInstructions();
 		const FRigVMByteCode& ByteCode = InVM->GetByteCode();
 
@@ -738,9 +740,9 @@ void SRigVMExecutionStackView::RefreshTreeView(URigVM* InVM, FRigVMExtendedExecu
 	if (InVM && InVMContext)
 	{
 		// fill the children from the log
-		if (RigVMEditor.IsValid())
+		if (TSharedPtr<FRigVMEditorBase> Editor = WeakEditor.Pin())
 		{
-			URigVMHost* Host = RigVMEditor.Pin()->GetRigVMHost();
+			URigVMHost* Host = Editor->GetRigVMHost();
 			if(Host && Host->GetLog())
 			{
 				const TArray<FRigVMLog::FLogEntry>& LogEntries = Host->GetLog()->Entries;
@@ -828,30 +830,33 @@ void SRigVMExecutionStackView::HandleFocusOnSelectedGraphNode()
 	TArray<TSharedPtr<FRigStackEntry>> SelectedItems = TreeView->GetSelectedItems();
 	if (SelectedItems.Num() > 0)
 	{
-		URigVMHost* Host = RigVMEditor.Pin()->GetRigVMHost();
-		if (Host == nullptr || Host->GetVM() == nullptr)
-		{
-			return;
-		}
-
-		const FRigVMByteCode& ByteCode = Host->GetVM()->GetByteCode();
-		UObject* Subject = ByteCode.GetSubjectForInstruction(SelectedItems[0]->InstructionIndex);
-		if (URigVMNode* SelectedNode = Cast<URigVMNode>(Subject))
+		if (TSharedPtr<FRigVMEditorBase> Editor = WeakEditor.Pin())
 		{
-			URigVMGraph* GraphToFocus = SelectedNode->GetGraph();
-			if (GraphToFocus && GraphToFocus->GetTypedOuter<URigVMAggregateNode>())
+			URigVMHost* Host = Editor->GetRigVMHost();
+			if (Host == nullptr || Host->GetVM() == nullptr)
 			{
-				if(URigVMGraph* ParentGraph = GraphToFocus->GetParentGraph())
-				{
-					GraphToFocus = ParentGraph;
-				} 
+				return;
 			}
-			
-			if (UEdGraph* EdGraph = RigVMBlueprint->GetEdGraph(GraphToFocus))
+
+			const FRigVMByteCode& ByteCode = Host->GetVM()->GetByteCode();
+			UObject* Subject = ByteCode.GetSubjectForInstruction(SelectedItems[0]->InstructionIndex);
+			if (URigVMNode* SelectedNode = Cast<URigVMNode>(Subject))
 			{
-				RigVMEditor.Pin()->OpenGraphAndBringToFront(EdGraph, true);
-				RigVMEditor.Pin()->ZoomToSelection_Clicked();
-				RigVMEditor.Pin()->HandleModifiedEvent(ERigVMGraphNotifType::NodeSelected, GraphToFocus, SelectedNode);
+				URigVMGraph* GraphToFocus = SelectedNode->GetGraph();
+				if (GraphToFocus && GraphToFocus->GetTypedOuter<URigVMAggregateNode>())
+				{
+					if(URigVMGraph* ParentGraph = GraphToFocus->GetParentGraph())
+					{
+						GraphToFocus = ParentGraph;
+					} 
+				}
+				
+				if (UEdGraph* EdGraph = RigVMBlueprint->GetEdGraph(GraphToFocus))
+				{
+					Editor->OpenGraphAndBringToFront(EdGraph, true);
+					Editor->ZoomToSelection_Clicked();
+					Editor->HandleModifiedEvent(ERigVMGraphNotifType::NodeSelected, GraphToFocus, SelectedNode);
+				}
 			}
 		}
 	}
@@ -928,8 +933,12 @@ void SRigVMExecutionStackView::HandleSelectTargetInstructions()
 TArray<TSharedPtr<FRigStackEntry>> SRigVMExecutionStackView::GetTargetItems(const TArray<TSharedPtr<FRigStackEntry>>& InItems) const
 {
 	TArray<TSharedPtr<FRigStackEntry>> TargetItems;
-	
-	URigVMHost* Host = RigVMEditor.Pin()->GetRigVMHost();
+
+	URigVMHost* Host = nullptr;
+	if (TSharedPtr<FRigVMEditorBase> Editor = WeakEditor.Pin())
+	{
+		Host = Editor->GetRigVMHost();
+	}
 	if (Host == nullptr || Host->GetVM() == nullptr)
 	{
 		return TargetItems;
@@ -1037,11 +1046,14 @@ void SRigVMExecutionStackView::OnVMCompiled(UObject* InCompiledObject, URigVM* I
 {
 	RefreshTreeView(InCompiledVM, &InVMContext);
 
-	if (RigVMEditor.IsValid() && !OnHostInitializedHandle.IsValid())
+	if (TSharedPtr<FRigVMEditorBase> Editor = WeakEditor.Pin())
 	{
-		if(URigVMHost* Host = RigVMEditor.Pin()->GetRigVMHost())
+		if (!OnHostInitializedHandle.IsValid())
 		{
-			OnHostInitializedHandle = Host->OnInitialized_AnyThread().AddSP(this, &SRigVMExecutionStackView::HandleHostInitializedEvent);
+			if(URigVMHost* Host = Editor->GetRigVMHost())
+			{
+				OnHostInitializedHandle = Host->OnInitialized_AnyThread().AddSP(this, &SRigVMExecutionStackView::HandleHostInitializedEvent);
+			}
 		}
 	}
 }
@@ -1053,10 +1065,13 @@ void SRigVMExecutionStackView::HandleExecutionHalted(const int32 InHaltedAtInstr
 		return;
 	}
 
-	if (InHaltedAtInstruction == INDEX_NONE && InEntryName == RigVMEditor.Pin()->GetRigVMHost()->GetEventQueue().Last())
+	if (TSharedPtr<FRigVMEditorBase> Editor = WeakEditor.Pin())
 	{
-		HaltedAtInstruction = InHaltedAtInstruction;
-		return;
+		if (InHaltedAtInstruction == INDEX_NONE && InEntryName == Editor->GetRigVMHost()->GetEventQueue().Last())
+		{
+			HaltedAtInstruction = InHaltedAtInstruction;
+			return;
+		}
 	}
 	
 	if (InHaltedAtInstruction != INDEX_NONE && Operators.Num() > InHaltedAtInstruction)
@@ -1070,10 +1085,13 @@ void SRigVMExecutionStackView::HandleExecutionHalted(const int32 InHaltedAtInstr
 void SRigVMExecutionStackView::OnFilterTextChanged(const FText& SearchText)
 {
 	FilterText = SearchText;
-	URigVMHost* RigVMHost = RigVMEditor.Pin()->GetRigVMHost();
-	if (RigVMHost != nullptr)
+	if (TSharedPtr<FRigVMEditorBase> Editor = WeakEditor.Pin())
 	{
-		RefreshTreeView(RigVMHost->GetVM(), &RigVMHost->GetRigVMExtendedExecuteContext());
+		URigVMHost* RigVMHost = Editor->GetRigVMHost();
+		if (RigVMHost != nullptr)
+		{
+			RefreshTreeView(RigVMHost->GetVM(), &RigVMHost->GetRigVMExtendedExecuteContext());
+		}
 	}
 }
 
@@ -1084,7 +1102,11 @@ void SRigVMExecutionStackView::HandleModifiedEvent(ERigVMGraphNotifType InNotifT
 		return;
 	}
 
-	URigVMHost* Host = RigVMEditor.Pin()->GetRigVMHost();
+	URigVMHost* Host = nullptr;
+	if (TSharedPtr<FRigVMEditorBase> Editor = WeakEditor.Pin())
+	{
+		Host = Editor->GetRigVMHost();
+	}
 	if (Host == nullptr || Host->GetVM() == nullptr)
 	{
 		return;
@@ -1156,12 +1178,12 @@ void SRigVMExecutionStackView::HandleHostInitializedEvent(URigVMHost* InHost, co
 		}
 	}
 	
-	if (RigVMEditor.IsValid())
+	if (TSharedPtr<FRigVMEditorBase> Editor = WeakEditor.Pin())
 	{
 		InHost->OnInitialized_AnyThread().Remove(OnHostInitializedHandle);
 		OnHostInitializedHandle.Reset();
 
-		if (URigVMBlueprint* RigBlueprint = RigVMEditor.Pin()->GetRigVMBlueprint())
+		if (URigVMBlueprint* RigBlueprint = Editor->GetRigVMBlueprint())
 		{
 			TArray<URigVMGraph*> Models = RigBlueprint->GetAllModels();
 			for (URigVMGraph* Model : Models)
@@ -1178,7 +1200,7 @@ void SRigVMExecutionStackView::HandleHostInitializedEvent(URigVMHost* InHost, co
 	}
 }
 
-void SRigVMExecutionStackView::HandlePreviewHostUpdated(FRigVMEditor* InEditor)
+void SRigVMExecutionStackView::HandlePreviewHostUpdated(FRigVMEditorBase* InEditor)
 {
 	if(URigVMHost* RigVMHost = InEditor->GetRigVMHost())
 	{
@@ -1188,34 +1210,37 @@ void SRigVMExecutionStackView::HandlePreviewHostUpdated(FRigVMEditor* InEditor)
 
 void SRigVMExecutionStackView::HandleItemMouseDoubleClick(TSharedPtr<FRigStackEntry> InItem)
 {
-	if (!RigVMEditor.IsValid() || !RigVMBlueprint.IsValid())
+	if (TSharedPtr<FRigVMEditorBase> Editor = WeakEditor.Pin())
 	{
-		return;
-	}
+		if (!RigVMBlueprint.IsValid())
+		{
+			return;
+		}
 	
-	URigVMHost* Host = Cast<URigVMHost>(RigVMBlueprint->GetObjectBeingDebugged());
-	if (!Host || !Host->GetVM())
-	{
-		return;
-	}
+		URigVMHost* Host = Cast<URigVMHost>(RigVMBlueprint->GetObjectBeingDebugged());
+		if (!Host || !Host->GetVM())
+		{
+			return;
+		}
 
-	const FRigVMByteCode& ByteCode = Host->GetVM()->GetByteCode();
-	if (URigVMNode* Subject = Cast<URigVMNode>(ByteCode.GetSubjectForInstruction(InItem->InstructionIndex)))
-	{
-		URigVMGraph* GraphToFocus = Subject->GetGraph();
-		if (GraphToFocus && GraphToFocus->GetTypedOuter<URigVMAggregateNode>())
+		const FRigVMByteCode& ByteCode = Host->GetVM()->GetByteCode();
+		if (URigVMNode* Subject = Cast<URigVMNode>(ByteCode.GetSubjectForInstruction(InItem->InstructionIndex)))
 		{
-			if(URigVMGraph* ParentGraph = GraphToFocus->GetParentGraph())
+			URigVMGraph* GraphToFocus = Subject->GetGraph();
+			if (GraphToFocus && GraphToFocus->GetTypedOuter<URigVMAggregateNode>())
 			{
-				GraphToFocus = ParentGraph;
-			} 
-		}
+				if(URigVMGraph* ParentGraph = GraphToFocus->GetParentGraph())
+				{
+					GraphToFocus = ParentGraph;
+				} 
+			}
 		
-		if(URigVMEdGraph* EdGraph = Cast<URigVMEdGraph>(RigVMBlueprint->GetEdGraph(GraphToFocus)))
-		{
-			if(const UEdGraphNode* Node = EdGraph->FindNodeForModelNodeName(Subject->GetFName()))
+			if(URigVMEdGraph* EdGraph = Cast<URigVMEdGraph>(RigVMBlueprint->GetEdGraph(GraphToFocus)))
 			{
-				RigVMEditor.Pin()->JumpToHyperlink(Node, false);
+				if(const UEdGraphNode* Node = EdGraph->FindNodeForModelNodeName(Subject->GetFName()))
+				{
+					Editor->JumpToHyperlink(Node, false);
+				}
 			}
 		}
 	}
diff --git a/Engine/Plugins/Runtime/RigVM/Source/RigVMEditor/Private/Widgets/SRigVMGraphBreakLinksWidget.cpp b/Engine/Plugins/Runtime/RigVM/Source/RigVMEditor/Private/Widgets/SRigVMGraphBreakLinksWidget.cpp
--- a/Engine/Plugins/Runtime/RigVM/Source/RigVMEditor/Private/Widgets/SRigVMGraphBreakLinksWidget.cpp
+++ b/Engine/Plugins/Runtime/RigVM/Source/RigVMEditor/Private/Widgets/SRigVMGraphBreakLinksWidget.cpp
@@ -11,10 +11,12 @@
 #include "Styling/AppStyle.h"
 #include "HAL/ConsoleManager.h"
 #include "RigVMModel/RigVMController.h"
-#include "Editor/RigVMEditor.h"
+#include "Editor/RigVMNewEditor.h"
 #include "ContentBrowserModule.h"
 #include "IContentBrowserSingleton.h"
 #include "Misc/ScopedSlowTask.h"
+#include "Widgets/Views/STableRow.h"
+#include "Widgets/Views/SListView.h"
 
 #define LOCTEXT_NAMESPACE "SRigVMGraphBreakLinksWidget"
 
diff --git a/Engine/Plugins/Runtime/RigVM/Source/RigVMEditor/Private/Widgets/SRigVMSwapAssetReferencesWidget.cpp b/Engine/Plugins/Runtime/RigVM/Source/RigVMEditor/Private/Widgets/SRigVMSwapAssetReferencesWidget.cpp
--- a/Engine/Plugins/Runtime/RigVM/Source/RigVMEditor/Private/Widgets/SRigVMSwapAssetReferencesWidget.cpp
+++ b/Engine/Plugins/Runtime/RigVM/Source/RigVMEditor/Private/Widgets/SRigVMSwapAssetReferencesWidget.cpp
@@ -8,9 +8,8 @@
 #include "AssetRegistry/AssetRegistryModule.h"
 #include "Styling/AppStyle.h"
 #include "Algo/Count.h"
-#include "BlueprintEditor.h"
 #include "RigVMEditorBlueprintLibrary.h"
-#include "Editor/RigVMEditor.h"
+#include "Editor/RigVMNewEditor.h"
 #include "Misc/UObjectToken.h"
 #include "RigVMModel/RigVMControllerActions.h"
 
@@ -442,7 +441,7 @@ FReply SRigVMSwapAssetReferencesWidget::OnNodeDoubleClicked(TSharedRef<FRigVMTre
 						GEditor->GetEditorSubsystem<UAssetEditorSubsystem>()->OpenEditorForAsset(TopLevelObject);
 						if(IAssetEditorInstance* Editor = GEditor->GetEditorSubsystem<UAssetEditorSubsystem>()->FindEditorForAsset(TopLevelObject, true))
 						{
-							if(FRigVMEditor* RigVMEditor = static_cast<FRigVMEditor*>(Editor))
+							if (FRigVMEditorBase* RigVMEditor = FRigVMEditorBase::GetFromAssetEditorInstance(Editor))
 							{
 								RigVMEditor->HandleJumpToHyperlink(ObjectReference);
 								return FReply::Handled();
diff --git a/Engine/Plugins/Runtime/RigVM/Source/RigVMEditor/Private/Widgets/SRigVMSwapFunctionsWidget.cpp b/Engine/Plugins/Runtime/RigVM/Source/RigVMEditor/Private/Widgets/SRigVMSwapFunctionsWidget.cpp
--- a/Engine/Plugins/Runtime/RigVM/Source/RigVMEditor/Private/Widgets/SRigVMSwapFunctionsWidget.cpp
+++ b/Engine/Plugins/Runtime/RigVM/Source/RigVMEditor/Private/Widgets/SRigVMSwapFunctionsWidget.cpp
@@ -8,8 +8,7 @@
 #include "AssetRegistry/AssetRegistryModule.h"
 #include "Styling/AppStyle.h"
 #include "Algo/Count.h"
-#include "BlueprintEditor.h"
-#include "Editor/RigVMEditor.h"
+#include "Editor/RigVMNewEditor.h"
 #include "Misc/UObjectToken.h"
 #include "RigVMModel/RigVMControllerActions.h"
 
@@ -425,7 +424,7 @@ bool FRigVMSwapFunctionTask::Execute(const TSharedRef<FRigVMTreePhase>& InPhase)
 				
 				if(IAssetEditorInstance* Editor = GEditor->GetEditorSubsystem<UAssetEditorSubsystem>()->FindEditorForAsset(Blueprint, true))
 				{
-					if(FRigVMEditor* RigVMEditor = static_cast<FRigVMEditor*>(Editor))
+					if (FRigVMEditorBase* RigVMEditor = FRigVMEditorBase::GetFromAssetEditorInstance(Editor))
 					{
 						RigVMEditor->HandleJumpToHyperlink(WeakRefNode.Get());
 					}
@@ -833,7 +832,7 @@ FReply SRigVMSwapFunctionsWidget::OnNodeDoubleClicked(TSharedRef<FRigVMTreeNode>
 						GEditor->GetEditorSubsystem<UAssetEditorSubsystem>()->OpenEditorForAsset(TopLevelObject);
 						if(IAssetEditorInstance* Editor = GEditor->GetEditorSubsystem<UAssetEditorSubsystem>()->FindEditorForAsset(TopLevelObject, true))
 						{
-							if(FRigVMEditor* RigVMEditor = static_cast<FRigVMEditor*>(Editor))
+							if (FRigVMEditorBase* RigVMEditor = FRigVMEditorBase::GetFromAssetEditorInstance(Editor))
 							{
 								RigVMEditor->HandleJumpToHyperlink(ObjectReference);
 								return FReply::Handled();
diff --git a/Engine/Plugins/Runtime/RigVM/Source/RigVMEditor/Public/EdGraph/NodeSpawners/RigVMEdGraphFunctionRefNodeSpawner.h b/Engine/Plugins/Runtime/RigVM/Source/RigVMEditor/Public/EdGraph/NodeSpawners/RigVMEdGraphFunctionRefNodeSpawner.h
--- a/Engine/Plugins/Runtime/RigVM/Source/RigVMEditor/Public/EdGraph/NodeSpawners/RigVMEdGraphFunctionRefNodeSpawner.h
+++ b/Engine/Plugins/Runtime/RigVM/Source/RigVMEditor/Public/EdGraph/NodeSpawners/RigVMEdGraphFunctionRefNodeSpawner.h
@@ -64,5 +64,5 @@ private:
 	static URigVMEdGraphNode* SpawnNode(UEdGraph* ParentGraph, UBlueprint* Blueprint, FRigVMGraphFunctionHeader& InFunction, FVector2D const Location);
 
 	friend class UEngineTestControlRig;
-	friend class FRigVMEditor;
+	friend class FRigVMEditorBase;
 };
diff --git a/Engine/Plugins/Runtime/RigVM/Source/RigVMEditor/Public/Editor/RigVMEditor.h b/Engine/Plugins/Runtime/RigVM/Source/RigVMEditor/Public/Editor/RigVMEditor.h
--- a/Engine/Plugins/Runtime/RigVM/Source/RigVMEditor/Public/Editor/RigVMEditor.h
+++ b/Engine/Plugins/Runtime/RigVM/Source/RigVMEditor/Public/Editor/RigVMEditor.h
@@ -3,17 +3,54 @@
 #pragma once
 
 #include "CoreMinimal.h"
-#include "BlueprintEditor.h"
-#include "RigVMHost.h"
-#include "RigVMBlueprint.h"
+#include "RigVMEditorGraphExplorerTabSummoner.h"
+#include "RigVMNewEditorMode.h"
+#include "RigVMCore/RigVMExecuteContext.h"
+#include "RigVMCore/RigVMGraphFunctionDefinition.h"
+#include "RigVMModel/RigVMController.h"
+#include "WorkflowOrientedApp/WorkflowCentricApplication.h"
+#include "RigVMModel/RigVMNotifications.h"
+#include "Widgets/SRigVMEditorGraphExplorer.h"
+#include "WorkflowOrientedApp/WorkflowTabManager.h"
+#include "WorkflowOrientedApp/WorkflowUObjectDocuments.h"
+#include "GraphEditor.h"
+#include "Misc/NotifyHook.h"
+#include "TickableEditorObject.h"
+#include "EditorUndoClient.h"
+#include "RigVMEditorModule.h"
+#include "SNodePanel.h"
 #include "Editor/RigVMDetailsViewWrapperObject.h"
-#include "AssetRegistry/AssetData.h"
 
-class FRigVMGraphExplorerDragDropOp;
+#define LOCTEXT_NAMESPACE "RigVMEditor"
+
+class FDocumentTracker;
 class SRigVMEditorGraphExplorer;
-class FRigVMEditor;
+class URigVM;
+class FTabInfo;
+class FRigVMEditorBase;
+class URigVMController;
+class URigVMHost;
+class FTransaction;
+class URigVMBlueprint;
+class SGraphEditor;
+
+DECLARE_MULTICAST_DELEGATE_TwoParams(FRigVMEditorClosed, const FRigVMEditorBase*, URigVMBlueprint*);
 
-DECLARE_MULTICAST_DELEGATE_TwoParams(FRigVMEditorClosed, const FRigVMEditor*, URigVMBlueprint*);
+
+/* Enums to use when grouping the blueprint members in the list panel. The order here will determine the order in the list */
+namespace RigVMNodeSectionID
+{
+	// Keep the values as they are defined in NodeSectionID, which is defined in BlueprintEditor.h
+	// TODO: Once there is no need for FRigVMLegacyEditor, improve the definition of this enum, including a uint8 definition
+	enum Type
+	{
+		NONE = 0,
+		GRAPH = 1,					// Graph
+		FUNCTION = 4,				// Functions
+		VARIABLE = 8,				// Variables
+		LOCAL_VARIABLE = 12			// Local variables
+	};
+};
 
 struct FRigVMEditorModes
 {
@@ -37,12 +74,144 @@ private:
 	FRigVMEditorModes() {}
 };
 
-class RIGVMEDITOR_API FRigVMEditor : public FBlueprintEditor
+struct FRigVMCustomDebugObject
+{
+public:
+	// Custom object to include, regardless of the current debugging World
+	UObject* Object;
+
+	// Override for the object name (if not empty)
+	FString NameOverride;
+
+public:
+	FRigVMCustomDebugObject()
+		: Object(nullptr)
+	{
+	}
+
+	FRigVMCustomDebugObject(UObject* InObject, const FString& InLabel)
+		: Object(InObject)
+		, NameOverride(InLabel)
+	{
+	}
+};
+
+class RIGVMEDITOR_API IRigVMEditor
+{
+public:
+	
+	virtual TSharedRef<IRigVMEditor> SharedRef() = 0;
+
+	virtual TSharedPtr<FAssetEditorToolkit> GetHostingApp() = 0;
+	virtual void InitAssetEditor(const EToolkitMode::Type Mode, const TSharedPtr<IToolkitHost>& InitToolkitHost, const FName AppIdentifier, const TSharedRef<FTabManager::FLayout>& StandaloneDefaultLayout, const bool bCreateDefaultStandaloneMenu, const bool bCreateDefaultToolbar, const TArray<UObject*>& ObjectsToEdit, const bool bInIsToolbarFocusable = false, const bool bInUseSmallToolbarIcons = false, const TOptional<EAssetOpenMethod>& InOpenMethod = TOptional<EAssetOpenMethod>()) = 0;
+	virtual void CreateEditorToolbar() = 0;
+	virtual void CommonInitialization(const TArray<UBlueprint*>& InitBlueprints, bool bShouldOpenInDefaultsMode) = 0;
+	virtual TSharedPtr<FDocumentTracker> GetDocumentManager() const = 0;
+	virtual void AddApplicationMode(FName ModeName, TSharedRef<FApplicationMode> Mode) = 0;
+	virtual void RegenerateMenusAndToolbars() = 0;
+	virtual void SetCurrentMode(FName NewMode) = 0;
+	virtual FEditorModeTools& GetToolkitEditorModeManager() const = 0;
+	virtual void PostLayoutBlueprintEditorInitialization() = 0;
+	virtual TSharedPtr<SGraphEditor> OpenGraphAndBringToFront(UEdGraph* Graph, bool bSetFocus = true) = 0;
+	virtual bool FindOpenTabsContainingDocument(const UObject* DocumentID, /*inout*/ TArray< TSharedPtr<SDockTab> >& Results) = 0;
+	virtual TSharedPtr<SDockTab> OpenDocument(const UObject* DocumentID, FDocumentTracker::EOpenDocumentCause Cause) = 0;
+	virtual void CloseDocumentTab(const UObject* DocumentID) = 0;
+	virtual TSharedPtr<FTabManager> GetTabManager() = 0;
+#if WITH_RIGVMLEGACYEDITOR
+	virtual TSharedPtr<class SKismetInspector> GetInspector() const = 0;
+#else
+	virtual TSharedPtr<class SRigVMDetailsInspector> GetInspector() const = 0;
+#endif
+	virtual TSharedPtr<FApplicationMode> CreateEditorMode() = 0;
+	virtual const FName GetEditorAppName() const = 0;
+	virtual const TArray< UObject* >& GetEditingObjects() const = 0;
+	virtual const TSharedRef<IToolkitHost> GetToolkitHost() const = 0;
+	virtual bool IsHosted() const = 0;
+	virtual void BringToolkitToFrontImpl() = 0;
+	virtual void InitToolMenuContext(FToolMenuContext& MenuContext) = 0;
+	virtual TSharedRef<FUICommandList> GetToolkitCommands() = 0;
+	virtual TWeakPtr<class SGraphEditor> GetFocusedGraphEditor() = 0;
+	virtual TWeakPtr<FDocumentTabFactory> GetGraphEditorTabFactory() const = 0;
+	virtual bool TransactionObjectAffectsBlueprint(UObject* InTransactedObject) = 0;
+	virtual FEdGraphPinType GetLastPinTypeUsed() = 0;
+	virtual void LogSimpleMessage(const FText& MessageText) = 0;
+	virtual void RenameNewlyAddedAction(FName InActionName) = 0;
+	virtual FGraphPanelSelectionSet GetSelectedNodes() const = 0;
+	virtual void SetUISelectionState(FName SelectionOwner) = 0;
+	virtual void AnalyticsTrackNodeEvent(UBlueprint* Blueprint, UEdGraphNode *GraphNode, bool bNodeDelete = false) const = 0;
+	virtual void JumpToHyperlink(const UObject* ObjectReference, bool bRequestRename = false) = 0;
+	virtual void PostUndo(bool bSuccess) = 0;
+	virtual void PostRedo(bool bSuccess) = 0;
+	virtual UEdGraphPin* GetCurrentlySelectedPin() const = 0;
+	virtual void CreateDefaultCommands() = 0;
+	virtual TSharedRef<SGraphEditor> CreateGraphEditorWidget(TSharedRef<class FTabInfo> InTabInfo, class UEdGraph* InGraph) = 0;
+	virtual void CompileImpl() = 0;
+	virtual void SaveAsset_Execute_Impl() = 0;
+	virtual void SaveAssetAs_Execute_Impl() = 0;
+	virtual bool IsGraphInCurrentBlueprint(const UEdGraph* InGraph) const = 0;
+	virtual bool IsEditableImpl(UEdGraph* InGraph) const = 0;
+	virtual UEdGraph* GetFocusedGraph() const = 0;
+	virtual void JumpToNode(const class UEdGraphNode* Node, bool bRequestRename = false) = 0;
+	virtual void JumpToPin(const class UEdGraphPin* Pin) = 0;
+	virtual void AddToolbarExtender(TSharedPtr<FExtender> Extender) = 0;
+	virtual void RemoveToolbarExtender(TSharedPtr<FExtender> Extender) = 0;
+	virtual void AddMenuExtender(TSharedPtr<FExtender> Extender) = 0;
+	virtual void RemoveMenuExtender(TSharedPtr<FExtender> Extender) = 0;
+	virtual TSharedPtr<class IMessageLogListing> GetCompilerResultsListing() = 0;
+	virtual void OnBlueprintChangedInnerImpl(UBlueprint* InBlueprint, bool bIsJustBeingCompiled = false) = 0;
+	virtual void RefreshEditorsImpl(ERefreshRigVMEditorReason::Type Reason = ERefreshRigVMEditorReason::UnknownReason) = 0;
+	virtual void SetupGraphEditorEventsImpl(UEdGraph* InGraph, SGraphEditor::FGraphEditorEvents& InEvents) = 0;
+	virtual FActionMenuContent OnCreateGraphActionMenu(UEdGraph* InGraph, const FVector2D& InNodePosition, const TArray<UEdGraphPin*>& InDraggedPins, bool bAutoExpand, SGraphEditor::FActionMenuClosed InOnMenuClosed) = 0;
+	virtual void AddReferencedObjectsImpl(FReferenceCollector& Collector) = 0;
+	virtual FGraphAppearanceInfo GetGraphAppearanceImpl(class UEdGraph* InGraph) const = 0;
+	virtual void NotifyPreChangeImpl( FProperty* PropertyAboutToChange ) = 0;
+	virtual void NotifyPostChangeImpl( const FPropertyChangedEvent& PropertyChangedEvent, FProperty* PropertyThatChanged) = 0;
+	virtual FName GetSelectedVariableName() = 0;
+	virtual bool IsNodeTitleVisible(const UEdGraphNode* Node, bool bRequestRename) = 0;
+	virtual void EditClassDefaults_Clicked() = 0;
+	virtual void EditGlobalOptions_Clicked() = 0;
+	virtual void TryInvokingDetailsTab(bool bFlash = true) = 0;
+	virtual FName GetGraphExplorerWidgetID() = 0;
+	virtual void RefreshInspector() = 0;
+	virtual void RefreshStandAloneDefaultsEditor() = 0;
+	virtual void GetPinTypeSelectorFilters(TArray<TSharedPtr<class IPinTypeSelectorFilter>>& OutFilters) const = 0;
+	virtual void OnAddNewVariable() = 0;
+	virtual void SummonSearchUI(bool bSetFindWithinBlueprint, FString NewSearchTerms = FString(), bool bSelectFirstResult = false) = 0;
+	virtual void ZoomToSelection_Clicked() = 0;
+	virtual void RestoreEditedObjectState() = 0;
+	virtual void SetupViewForBlueprintEditingMode() = 0;
+	virtual bool GetIsContextSensitive() = 0;
+	virtual void SetIsContextSensitive(const bool bIsContextSensitive) = 0;
+	virtual void RegisterToolbarTab(const TSharedRef<class FTabManager>& TabManager) = 0;
+	virtual const TArray< UObject* >* GetObjectsCurrentlyBeingEdited() const = 0;
+	virtual void AddSelectedDebugObjectWidget(FToolBarBuilder& ToolbarBuilder) = 0;
+};
+
+class RIGVMEDITOR_API FRigVMEditorBase : public IRigVMEditor
 {
+protected:
+	// Type of new document/graph being created by a menu item
+	enum ECreatedDocumentType
+	{
+		CGT_NewVariable,
+		CGT_NewFunctionGraph,
+		CGT_NewMacroGraph,
+		CGT_NewAnimationLayer,
+		CGT_NewEventGraph,
+		CGT_NewLocalVariable
+	};
+	
+	FRigVMEditorBase();
+	virtual ~FRigVMEditorBase(){}
+
+
+	void UnbindEditor();
+
 public:
+	static FRigVMEditorBase* GetFromAssetEditorInstance(IAssetEditorInstance* Instance);
 
-	FRigVMEditor();
-	virtual ~FRigVMEditor();
+	DECLARE_EVENT(FRigVMEditorBase, FOnRefreshEvent);
+	FOnRefreshEvent OnRefresh() { return OnRefreshEvent; }
 
 	/**
 	 * Edits the specified asset(s)
@@ -52,56 +221,60 @@ public:
 	 * @param	InRigVMBlueprint	The blueprint object to start editing.
 	 */
 	virtual void InitRigVMEditor(const EToolkitMode::Type Mode, const TSharedPtr< class IToolkitHost >& InitToolkitHost, class URigVMBlueprint* InRigVMBlueprint);
+	
+protected:
 
 	virtual void HandleAssetRequestedOpen(UObject* InObject);
 	virtual void HandleAssetRequestClose(UObject* InObject, EAssetEditorCloseReason InReason);
 	bool bRequestedReopen = false;
 
-	virtual const FName GetEditorAppName() const;
 	virtual const FName GetEditorModeName() const;
-	virtual TSharedPtr<FApplicationMode> CreateEditorMode();
 
 	// FBlueprintEditor interface
-	virtual UBlueprint* GetBlueprintObj() const override;
-	virtual TSubclassOf<UEdGraphSchema> GetDefaultSchemaClass() const override;
-	virtual bool InEditingMode() const override;
-	virtual void OnGraphEditorFocused(const TSharedRef<SGraphEditor>& InGraphEditor) override;
+public:
+	virtual UBlueprint* GetBlueprintObj() const;
+protected:
+	virtual TSubclassOf<UEdGraphSchema> GetDefaultSchemaClass() const;
+	virtual bool InEditingMode() const;
+	virtual void OnGraphEditorFocused(const TSharedRef<SGraphEditor>& InGraphEditor);
 
 	//  FTickableEditorObject Interface
-	virtual void Tick(float DeltaTime) override;
+	virtual void Tick(float DeltaTime);
 
 	// IToolkit Interface
-	virtual void BringToolkitToFront() override;
-	virtual FName GetToolkitFName() const override;
-	virtual FName GetToolkitContextFName() const override;
-	virtual FText GetBaseToolkitName() const override;
-	virtual FText GetToolkitToolTipText() const override;
-	virtual FString GetWorldCentricTabPrefix() const override;
-	virtual FLinearColor GetWorldCentricTabColorScale() const override;	
-	virtual void InitToolMenuContext(FToolMenuContext& MenuContext) override;
+	virtual void BringToolkitToFront();
+	virtual FName GetToolkitFName() const;
+	virtual FName GetToolkitContextFName() const;
+	virtual FText GetBaseToolkitName() const;
+	virtual FText GetToolkitToolTipText() const;
+	virtual FString GetWorldCentricTabPrefix() const;
+	virtual FLinearColor GetWorldCentricTabColorScale() const;	
+	virtual void InitToolMenuContextImpl(FToolMenuContext& MenuContext);
 
 	// BlueprintEditor interface
-	virtual bool TransactionObjectAffectsBlueprint(UObject* InTransactedObject) override;
-	virtual bool CanAddNewLocalVariable() const override;
-	virtual void OnAddNewLocalVariable() override;
-	virtual void OnPasteNewLocalVariable(const FBPVariableDescription& VariableDescription) override;
-	virtual void DeleteSelectedNodes() override;
-	virtual bool CanDeleteNodes() const override;
-	virtual void CopySelectedNodes() override;
-	virtual bool CanCopyNodes() const override;
-	virtual void PasteNodes() override;
-	virtual bool CanPasteNodes() const override;
-	virtual bool IsNativeParentClassCodeLinkEnabled() const override { return false; }
-	virtual bool ReparentBlueprint_IsVisible() const override { return false; }
-	virtual FReply OnSpawnGraphNodeByShortcut(FInputChord InChord, const FVector2D& InPosition, UEdGraph* InGraph) override;
-	virtual bool ShouldLoadBPLibrariesFromAssetRegistry() override { return false; }
-	virtual void JumpToHyperlink(const UObject* ObjectReference, bool bRequestRename = false) override;
+	virtual bool TransactionObjectAffectsBlueprintImpl(UObject* InTransactedObject);
+	virtual bool CanAddNewLocalVariable() const;
+	virtual void OnAddNewLocalVariable();
+	virtual void OnPasteNewLocalVariable(const FBPVariableDescription& VariableDescription);
+public:
+	virtual void DeleteSelectedNodes();
+	virtual bool CanDeleteNodes() const;
+protected:
+	virtual void CopySelectedNodes();
+	virtual bool CanCopyNodes() const;
+	virtual void PasteNodes();
+	virtual bool CanPasteNodes() const;
+	virtual bool IsNativeParentClassCodeLinkEnabled() const { return false; }
+	virtual bool ReparentBlueprint_IsVisible() const { return false; }
+	virtual FReply OnSpawnGraphNodeByShortcut(FInputChord InChord, const FVector2D& InPosition, UEdGraph* InGraph);
+	virtual bool ShouldLoadBPLibrariesFromAssetRegistry() { return false; }
+	virtual bool JumpToHyperlinkImpl(const UObject* ObjectReference, bool bRequestRename = false);
 	virtual bool ShouldOpenGraphByDefault() const { return true; }
-	virtual void AddNewFunctionVariant(const UEdGraph* InOriginalFunction) override;
+	virtual void AddNewFunctionVariant(const UEdGraph* InOriginalFunction);
 
 	// FEditorUndoClient Interface
-	virtual void PostUndo(bool bSuccess) override;
-	virtual void PostRedo(bool bSuccess) override;
+	virtual void PostUndoImpl(bool bSuccess);
+	virtual void PostRedoImpl(bool bSuccess);
 
 	virtual void PostTransaction(bool bSuccess, const FTransaction* Transaction, bool bIsRedo) {};
 
@@ -112,17 +285,20 @@ public:
 	virtual bool CanStopWatchingPin() const;
 
 	// IToolkitHost Interface
-	virtual void OnToolkitHostingStarted(const TSharedRef<class IToolkit>& Toolkit) override;
-	virtual void OnToolkitHostingFinished(const TSharedRef<class IToolkit>& Toolkit) override;
+	virtual void OnToolkitHostingStarted(const TSharedRef<class IToolkit>& Toolkit);
+	virtual void OnToolkitHostingFinished(const TSharedRef<class IToolkit>& Toolkit);
 
 	//  FTickableEditorObject Interface
-	virtual TStatId GetStatId() const override;
+	virtual TStatId GetStatId() const;
 
+public:
 	// returns the blueprint being edited
 	URigVMBlueprint* GetRigVMBlueprint() const;
 
 	// returns the currently debugged / viewed host
 	URigVMHost* GetRigVMHost() const;
+	
+protected:
 
 	virtual UObject* GetOuterForHost() const;
 
@@ -132,15 +308,17 @@ public:
 	// allows the editor to fill an empty graph
 	virtual void CreateEmptyGraphContent(URigVMController* InController) {}
 
-	DECLARE_EVENT_OneParam(FRigVMEditor, FPreviewHostUpdated, FRigVMEditor*);
+	DECLARE_EVENT_OneParam(FRigVMEditorBase, FPreviewHostUpdated, FRigVMEditorBase*);
 	FPreviewHostUpdated& OnPreviewHostUpdated() { return PreviewHostUpdated;  }
 
+public:
 	FRigVMEditorClosed& OnEditorClosed() { return RigVMEditorClosedDelegate; }
+protected:
 
 	/** Get the toolbox hosting widget */
 	TSharedRef<SBorder> GetToolbox() { return Toolbox.ToSharedRef(); }
 
-	virtual bool SelectLocalVariable(const UEdGraph* Graph, const FName& VariableName) override;
+	virtual bool SelectLocalVariable(const UEdGraph* Graph, const FName& VariableName);
 
 protected:
 
@@ -151,29 +329,26 @@ protected:
 	};
 
 	// FBlueprintEditor Interface
-	virtual void CreateDefaultCommands() override;
+	virtual void CreateDefaultCommandsImpl();
 	virtual void OnCreateGraphEditorCommands(TSharedPtr<FUICommandList> GraphEditorCommandsList);
-	virtual TSharedRef<SGraphEditor> CreateGraphEditorWidget(TSharedRef<FTabInfo> InTabInfo, UEdGraph* InGraph) override;
-	virtual void Compile() override;
-	virtual void SaveAsset_Execute() override;
-	virtual void SaveAssetAs_Execute() override;
-	virtual bool IsInAScriptingMode() const override { return true; }
-	virtual void NewDocument_OnClicked(ECreatedDocumentType GraphType) override;
-	virtual bool IsSectionVisible(NodeSectionID::Type InSectionID) const override;
-	virtual bool AreEventGraphsAllowed() const override;
-	virtual bool AreMacrosAllowed() const override;
-	virtual bool AreDelegatesAllowed() const override;
-	virtual bool NewDocument_IsVisibleForType(ECreatedDocumentType GraphType) const;
-	virtual FGraphAppearanceInfo GetGraphAppearance(class UEdGraph* InGraph) const override;
-	virtual bool IsEditable(UEdGraph* InGraph) const override;
-	virtual bool IsCompilingEnabled() const override;
-	virtual FText GetGraphDecorationString(UEdGraph* InGraph) const override;
-	virtual void OnActiveTabChanged( TSharedPtr<SDockTab> PreviouslyActive, TSharedPtr<SDockTab> NewlyActivated ) override;
-	virtual void OnSelectedNodesChangedImpl(const TSet<class UObject*>& NewSelection) override;
-	virtual void OnBlueprintChangedImpl(UBlueprint* InBlueprint, bool bIsJustBeingCompiled) override;
-	virtual void RefreshEditors(ERefreshBlueprintEditorReason::Type Reason = ERefreshBlueprintEditorReason::UnknownReason) override;
-	virtual void SetupGraphEditorEvents(UEdGraph* InGraph, SGraphEditor::FGraphEditorEvents& InEvents) override;
-	virtual void FocusInspectorOnGraphSelection(const TSet<class UObject*>& NewSelection, bool bForceRefresh = false) override;
+	virtual void Compile();
+	virtual void SaveAsset_Execute();
+	virtual void SaveAssetAs_Execute();
+	virtual bool IsInAScriptingMode() const { return true; }
+	virtual void NewDocument_OnClicked(ECreatedDocumentType GraphType);
+	bool IsSectionVisibleImpl(RigVMNodeSectionID::Type InSectionID) const;
+	virtual bool AreEventGraphsAllowed() const;
+	virtual bool AreMacrosAllowed() const;
+	virtual bool AreDelegatesAllowed() const;
+	bool NewDocument_IsVisibleForTypeImpl(ECreatedDocumentType GraphType) const;
+	virtual FGraphAppearanceInfo GetGraphAppearance(class UEdGraph* InGraph) const;
+	virtual bool IsCompilingEnabled() const;
+	virtual void OnActiveTabChanged( TSharedPtr<SDockTab> PreviouslyActive, TSharedPtr<SDockTab> NewlyActivated );
+	virtual void OnSelectedNodesChangedImpl(const TSet<class UObject*>& NewSelection);
+	virtual void OnBlueprintChangedImpl(UBlueprint* InBlueprint, bool bIsJustBeingCompiled);
+	virtual void RefreshEditors(ERefreshRigVMEditorReason::Type Reason = ERefreshRigVMEditorReason::UnknownReason);
+	virtual void SetupGraphEditorEvents(UEdGraph* InGraph, SGraphEditor::FGraphEditorEvents& InEvents);
+	virtual void FocusInspectorOnGraphSelection(const TSet<class UObject*>& NewSelection, bool bForceRefresh = false);
 
 	virtual void HandleModifiedEvent(ERigVMGraphNotifType InNotifType, URigVMGraph* InGraph, UObject* InSubject);
 	UE_DEPRECATED(5.4, "Please use HandleVMCompiledEvent with ExtendedExecuteContext param.")
@@ -184,10 +359,10 @@ protected:
 	void SetHaltedNode(URigVMNode* Node);
 	
 	// FNotifyHook Interface
-	virtual void NotifyPreChange(FProperty* PropertyAboutToChange) override;
-	virtual void NotifyPostChange(const FPropertyChangedEvent& PropertyChangedEvent, FProperty* PropertyThatChanged) override;
+	virtual void NotifyPreChange(FProperty* PropertyAboutToChange);
+	virtual void NotifyPostChange(const FPropertyChangedEvent& PropertyChangedEvent, FProperty* PropertyThatChanged);
 	/** delegate for changing property */
-	virtual void OnFinishedChangingProperties(const FPropertyChangedEvent& PropertyChangedEvent) override;
+	virtual void OnFinishedChangingProperties(const FPropertyChangedEvent& PropertyChangedEvent);
 	void OnPropertyChanged(UObject* InObject, FPropertyChangedEvent& InEvent);
 	virtual void OnWrappedPropertyChangedChainEvent(URigVMDetailsViewWrapperObject* InWrapperObject, const FString& InPropertyPath, FPropertyChangedChainEvent& InPropertyChangedChainEvent);
 	void OnRequestLocalizeFunctionDialog(FRigVMGraphFunctionIdentifier& InFunction, URigVMController* InTargetController, IRigVMGraphFunctionHost* InTargetFunctionHost, bool bForce);
@@ -197,6 +372,8 @@ protected:
 
 public:
 	void HandleJumpToHyperlink(const UObject* InSubject);
+	virtual FText GetGraphDecorationString(UEdGraph* InGraph) const;
+	virtual bool IsEditable(UEdGraph* InGraph) const;
 
 protected:
 	bool UpdateDefaultValueForVariable(FBPVariableDescription& InVariable, bool bUseCDO);
@@ -211,7 +388,7 @@ protected:
 	virtual void CacheNameLists();
 
 	// FGCObject Interface
-	virtual void AddReferencedObjects( FReferenceCollector& Collector ) override;
+	virtual void AddReferencedObjects( FReferenceCollector& Collector );
 
 	virtual void BindCommands();
 
@@ -236,7 +413,7 @@ protected:
 	/** Redo Action **/
 	void RedoAction();
 	
-	virtual void OnCreateComment() override;
+	virtual void OnCreateComment();
 
 	bool IsDetailsPanelRefreshSuspended() const { return bSuspendDetailsPanelRefresh; }
 	bool& GetSuspendDetailsPanelRefreshFlag() { return bSuspendDetailsPanelRefresh; }
@@ -291,10 +468,10 @@ protected:
 	static FSlateIcon GetExecutionModeIcon(const ERigVMEditorExecutionModeType InExecutionMode);
 	FSlateIcon GetExecutionModeIcon() const;
 	
-	virtual void GetCustomDebugObjects(TArray<FCustomDebugObject>& DebugList) const override;
-	virtual bool OnlyShowCustomDebugObjects() const override { return true; }
-	virtual void HandleSetObjectBeingDebugged(UObject* InObject) override;
-	virtual FString GetCustomDebugObjectLabel(UObject* ObjectBeingDebugged) const override;
+	virtual void GetCustomDebugObjects(TArray<FRigVMCustomDebugObject>& DebugList) const;
+	virtual bool OnlyShowCustomDebugObjects() const { return true; }
+	virtual void HandleSetObjectBeingDebugged(UObject* InObject);
+	virtual FString GetCustomDebugObjectLabel(UObject* ObjectBeingDebugged) const;
 
 	/** Handle hiding items in the graph */
 	void HandleHideItem();
@@ -309,7 +486,7 @@ protected:
 	virtual void OnGraphNodeClicked(URigVMEdGraphNode* InNode);
 	virtual void OnNodeDoubleClicked(URigVMBlueprint* InBlueprint, URigVMNode* InNode);
 	virtual void OnGraphImported(UEdGraph* InEdGraph);
-	virtual bool OnActionMatchesName(FEdGraphSchemaAction* InAction, const FName& InName) const override;
+	virtual bool OnActionMatchesName(FEdGraphSchemaAction* InAction, const FName& InName) const;
 	virtual void HandleShowCurrentStatement();
 	virtual void HandleBreakpointActionRequested(const ERigVMBreakpointAction BreakpointAction);
 	virtual bool IsHaltedAtBreakpoint() const;
@@ -335,16 +512,19 @@ private:
 
 	FPreviewHostUpdated PreviewHostUpdated;
 
-	FRigVMEditorClosed RigVMEditorClosedDelegate;
-
 	/** Toolbox hosting widget */
 	TSharedPtr<SBorder> Toolbox;
 
+protected:
 	TSharedPtr<SRigVMEditorGraphExplorer> GraphExplorerWidget;
 
+	FRigVMEditorClosed RigVMEditorClosedDelegate;
+
 public:
 	void SetGraphExplorerWidget(TSharedPtr<SRigVMEditorGraphExplorer> InWidget) { GraphExplorerWidget = InWidget; }
 
+	bool IsEditingSingleBlueprint() const;
+	
 private:
 	bool bAnyErrorsLeft;
 	TMap<FString, FString> KnownInstructionLimitWarnings;
@@ -381,6 +561,99 @@ protected:
 
 	FDelegateHandle PropertyChangedHandle;
 
+	FOnRefreshEvent OnRefreshEvent;
+
 	friend class SRigVMExecutionStackView;
 	friend class SRigVMEditorGraphExplorer;
 };
+
+struct FRigVMEditorZoomLevelsContainer : public FZoomLevelsContainer
+{
+	struct FRigVMEditorZoomLevelEntry
+	{
+	public:
+		FRigVMEditorZoomLevelEntry(float InZoomAmount, const FText& InDisplayText, EGraphRenderingLOD::Type InLOD)
+			: DisplayText(FText::Format(NSLOCTEXT("GraphEditor", "Zoom", "Zoom {0}"), InDisplayText))
+		, ZoomAmount(InZoomAmount)
+		, LOD(InLOD)
+		{
+		}
+
+	public:
+		FText DisplayText;
+		float ZoomAmount;
+		EGraphRenderingLOD::Type LOD;
+	};
+	
+	FRigVMEditorZoomLevelsContainer()
+	{
+		ZoomLevels.Reserve(22);
+		ZoomLevels.Add(FRigVMEditorZoomLevelEntry(0.025f, FText::FromString(TEXT("-14")), EGraphRenderingLOD::LowestDetail));
+		ZoomLevels.Add(FRigVMEditorZoomLevelEntry(0.070f, FText::FromString(TEXT("-13")), EGraphRenderingLOD::LowestDetail));
+		ZoomLevels.Add(FRigVMEditorZoomLevelEntry(0.100f, FText::FromString(TEXT("-12")), EGraphRenderingLOD::LowestDetail));
+		ZoomLevels.Add(FRigVMEditorZoomLevelEntry(0.125f, FText::FromString(TEXT("-11")), EGraphRenderingLOD::LowestDetail));
+		ZoomLevels.Add(FRigVMEditorZoomLevelEntry(0.150f, FText::FromString(TEXT("-10")), EGraphRenderingLOD::LowestDetail));
+		ZoomLevels.Add(FRigVMEditorZoomLevelEntry(0.175f, FText::FromString(TEXT("-9")), EGraphRenderingLOD::LowestDetail));
+		ZoomLevels.Add(FRigVMEditorZoomLevelEntry(0.200f, FText::FromString(TEXT("-8")), EGraphRenderingLOD::LowestDetail));
+		ZoomLevels.Add(FRigVMEditorZoomLevelEntry(0.225f, FText::FromString(TEXT("-7")), EGraphRenderingLOD::LowDetail));
+		ZoomLevels.Add(FRigVMEditorZoomLevelEntry(0.250f, FText::FromString(TEXT("-6")), EGraphRenderingLOD::LowDetail));
+		ZoomLevels.Add(FRigVMEditorZoomLevelEntry(0.375f, FText::FromString(TEXT("-5")), EGraphRenderingLOD::MediumDetail));
+		ZoomLevels.Add(FRigVMEditorZoomLevelEntry(0.500f, FText::FromString(TEXT("-4")), EGraphRenderingLOD::MediumDetail));
+		ZoomLevels.Add(FRigVMEditorZoomLevelEntry(0.675f, FText::FromString(TEXT("-3")), EGraphRenderingLOD::MediumDetail));
+		ZoomLevels.Add(FRigVMEditorZoomLevelEntry(0.750f, FText::FromString(TEXT("-2")), EGraphRenderingLOD::DefaultDetail));
+		ZoomLevels.Add(FRigVMEditorZoomLevelEntry(0.875f, FText::FromString(TEXT("-1")), EGraphRenderingLOD::DefaultDetail));
+		ZoomLevels.Add(FRigVMEditorZoomLevelEntry(1.000f, FText::FromString(TEXT("1:1")), EGraphRenderingLOD::DefaultDetail));
+		ZoomLevels.Add(FRigVMEditorZoomLevelEntry(1.250f, FText::FromString(TEXT("+1")), EGraphRenderingLOD::DefaultDetail));
+		ZoomLevels.Add(FRigVMEditorZoomLevelEntry(1.375f, FText::FromString(TEXT("+2")), EGraphRenderingLOD::DefaultDetail));
+		ZoomLevels.Add(FRigVMEditorZoomLevelEntry(1.500f, FText::FromString(TEXT("+3")), EGraphRenderingLOD::FullyZoomedIn));
+		ZoomLevels.Add(FRigVMEditorZoomLevelEntry(1.675f, FText::FromString(TEXT("+4")), EGraphRenderingLOD::FullyZoomedIn));
+		ZoomLevels.Add(FRigVMEditorZoomLevelEntry(1.750f, FText::FromString(TEXT("+5")), EGraphRenderingLOD::FullyZoomedIn));
+		ZoomLevels.Add(FRigVMEditorZoomLevelEntry(1.875f, FText::FromString(TEXT("+6")), EGraphRenderingLOD::FullyZoomedIn));
+		ZoomLevels.Add(FRigVMEditorZoomLevelEntry(2.000f, FText::FromString(TEXT("+7")), EGraphRenderingLOD::FullyZoomedIn));
+	}
+
+	float GetZoomAmount(int32 InZoomLevel) const override
+	{
+		checkSlow(ZoomLevels.IsValidIndex(InZoomLevel));
+		return ZoomLevels[InZoomLevel].ZoomAmount;
+	}
+
+	int32 GetNearestZoomLevel(float InZoomAmount) const override
+	{
+		for (int32 ZoomLevelIndex=0; ZoomLevelIndex < GetNumZoomLevels(); ++ZoomLevelIndex)
+		{
+			if (InZoomAmount <= GetZoomAmount(ZoomLevelIndex))
+			{
+				return ZoomLevelIndex;
+			}
+		}
+
+		return GetDefaultZoomLevel();
+	}
+	
+	FText GetZoomText(int32 InZoomLevel) const override
+	{
+		checkSlow(ZoomLevels.IsValidIndex(InZoomLevel));
+		return ZoomLevels[InZoomLevel].DisplayText;
+	}
+	
+	int32 GetNumZoomLevels() const override
+	{
+		return ZoomLevels.Num();
+	}
+	
+	int32 GetDefaultZoomLevel() const override
+	{
+		return 14;
+	}
+
+	EGraphRenderingLOD::Type GetLOD(int32 InZoomLevel) const override
+	{
+		checkSlow(ZoomLevels.IsValidIndex(InZoomLevel));
+		return ZoomLevels[InZoomLevel].LOD;
+	}
+
+	TArray<FRigVMEditorZoomLevelEntry> ZoomLevels;
+};
+
+#undef LOCTEXT_NAMESPACE
\ No newline at end of file
diff --git a/Engine/Plugins/Runtime/RigVM/Source/RigVMEditor/Public/Editor/RigVMEditorGraphExplorerTabSummoner.h b/Engine/Plugins/Runtime/RigVM/Source/RigVMEditor/Public/Editor/RigVMEditorGraphExplorerTabSummoner.h
--- a/Engine/Plugins/Runtime/RigVM/Source/RigVMEditor/Public/Editor/RigVMEditorGraphExplorerTabSummoner.h
+++ b/Engine/Plugins/Runtime/RigVM/Source/RigVMEditor/Public/Editor/RigVMEditorGraphExplorerTabSummoner.h
@@ -5,7 +5,7 @@
 #include "CoreMinimal.h"
 #include "WorkflowOrientedApp/WorkflowTabFactory.h"
 
-class FRigVMEditor;
+class FRigVMEditorBase;
 
 struct RIGVMEDITOR_API FRigVMEditorGraphExplorerTabSummoner : public FWorkflowTabFactory
 {
@@ -14,10 +14,10 @@ public:
 	static inline const FLazyName TabID = FLazyName(TEXT("RigVM Graph Explorer"));
 	
 public:
-	FRigVMEditorGraphExplorerTabSummoner(const TSharedRef<FRigVMEditor>& InRigVMEditor);
+	FRigVMEditorGraphExplorerTabSummoner(const TSharedRef<FRigVMEditorBase>& InRigVMEditor);
 	
 	virtual TSharedRef<SWidget> CreateTabBody(const FWorkflowTabSpawnInfo& Info) const override;
 	
 protected:
-	TWeakPtr<FRigVMEditor> RigVMEditor;
+	TWeakPtr<FRigVMEditorBase> RigVMEditor;
 };
diff --git a/Engine/Plugins/Runtime/RigVM/Source/RigVMEditor/Public/Editor/RigVMEditorMenuContext.h b/Engine/Plugins/Runtime/RigVM/Source/RigVMEditor/Public/Editor/RigVMEditorMenuContext.h
--- a/Engine/Plugins/Runtime/RigVM/Source/RigVMEditor/Public/Editor/RigVMEditorMenuContext.h
+++ b/Engine/Plugins/Runtime/RigVM/Source/RigVMEditor/Public/Editor/RigVMEditorMenuContext.h
@@ -9,7 +9,7 @@
 
 class URigVMHost;
 class URigVMBlueprint;
-class FRigVMEditor;
+class FRigVMEditorBase;
 class URigVMGraph;
 class URigVMNode;
 class URigVMPin;
@@ -57,7 +57,7 @@ public:
 	 * @param InRigVMEditor 	    The RigVM Editor hosting the menus
 	 * @param InGraphMenuContext 	Additional context for specific menus
 	*/
-	void Init(TWeakPtr<FRigVMEditor> InRigVMEditor, const FRigVMEditorGraphMenuContext& InGraphMenuContext = FRigVMEditorGraphMenuContext());
+	void Init(TWeakPtr<FRigVMEditorBase> InRigVMEditor, const FRigVMEditorGraphMenuContext& InGraphMenuContext = FRigVMEditorGraphMenuContext());
 
 	/** Get the rigvm blueprint that we are editing */
 	UFUNCTION(BlueprintCallable, Category = RigVMEditor)
@@ -75,10 +75,10 @@ public:
 	UFUNCTION(BlueprintCallable, Category = RigVMEditor)
 	FRigVMEditorGraphMenuContext GetGraphMenuContext();
 	
-	FRigVMEditor* GetRigVMEditor() const;
+	FRigVMEditorBase* GetRigVMEditor() const;
 
 private:
-	TWeakPtr<FRigVMEditor> WeakRigVMEditor;
+	TWeakPtr<FRigVMEditorBase> WeakRigVMEditor;
 
 	FRigVMEditorGraphMenuContext GraphMenuContext;
 };
\ No newline at end of file
diff --git a/Engine/Plugins/Runtime/RigVM/Source/RigVMEditor/Public/Editor/RigVMExecutionStackTabSummoner.h b/Engine/Plugins/Runtime/RigVM/Source/RigVMEditor/Public/Editor/RigVMExecutionStackTabSummoner.h
--- a/Engine/Plugins/Runtime/RigVM/Source/RigVMEditor/Public/Editor/RigVMExecutionStackTabSummoner.h
+++ b/Engine/Plugins/Runtime/RigVM/Source/RigVMEditor/Public/Editor/RigVMExecutionStackTabSummoner.h
@@ -5,7 +5,7 @@
 #include "CoreMinimal.h"
 #include "WorkflowOrientedApp/WorkflowTabFactory.h"
 
-class FRigVMEditor;
+class FRigVMEditorBase;
 
 struct RIGVMEDITOR_API FRigVMExecutionStackTabSummoner : public FWorkflowTabFactory
 {
@@ -14,10 +14,10 @@ public:
 	static inline const FLazyName TabID = FLazyName(TEXT("Execution Stack"));
 	
 public:
-	FRigVMExecutionStackTabSummoner(const TSharedRef<FRigVMEditor>& InRigVMEditor);
+	FRigVMExecutionStackTabSummoner(const TSharedRef<FRigVMEditorBase>& InRigVMEditor);
 	
 	virtual TSharedRef<SWidget> CreateTabBody(const FWorkflowTabSpawnInfo& Info) const override;
 	
 protected:
-	TWeakPtr<FRigVMEditor> RigVMEditor;
+	TWeakPtr<FRigVMEditorBase> RigVMEditor;
 };
diff --git a/Engine/Plugins/Runtime/RigVM/Source/RigVMEditor/Public/Editor/RigVMGraphDetailCustomization.h b/Engine/Plugins/Runtime/RigVM/Source/RigVMEditor/Public/Editor/RigVMGraphDetailCustomization.h
--- a/Engine/Plugins/Runtime/RigVM/Source/RigVMEditor/Public/Editor/RigVMGraphDetailCustomization.h
+++ b/Engine/Plugins/Runtime/RigVM/Source/RigVMEditor/Public/Editor/RigVMGraphDetailCustomization.h
@@ -8,7 +8,8 @@
 #include "UObject/WeakObjectPtr.h"
 #include "EdGraph/RigVMEdGraph.h"
 #include "RigVMBlueprint.h"
-#include "Editor/RigVMEditor.h"
+#include "Editor/RigVMNewEditor.h"
+#include "Editor/RigVMLegacyEditor.h"
 #include "SGraphPin.h"
 #include "Widgets/SRigVMGraphNode.h"
 #include "Widgets/Colors/SColorBlock.h"
@@ -29,7 +30,7 @@ public:
 	FRigVMFunctionArgumentGroupLayout(
 		const TWeakObjectPtr<URigVMGraph>& InGraph,
 		const TWeakInterfacePtr<IRigVMClientHost>& InRigVMClientHost,
-		const TWeakPtr<FRigVMEditor>& InEditor,
+		const TWeakPtr<FRigVMEditorBase>& InEditor,
 		bool bInputs);
 	virtual ~FRigVMFunctionArgumentGroupLayout();
 
@@ -49,7 +50,7 @@ private:
 
 	TWeakObjectPtr<URigVMGraph> GraphPtr;
 	TWeakInterfacePtr<IRigVMClientHost> WeakRigVMClientHost;
-	TWeakPtr<FRigVMEditor> RigVMEditorPtr;
+	TWeakPtr<FRigVMEditorBase> RigVMEditorPtr;
 	bool bIsInputGroup;
 	FSimpleDelegate OnRebuildChildren;
 };
@@ -62,7 +63,7 @@ public:
 		const TWeakObjectPtr<URigVMPin>& InPin,
 		const TWeakObjectPtr<URigVMGraph>& InGraph, 
 		const TWeakInterfacePtr<IRigVMClientHost>& InRigVMClientHost,
-		const TWeakPtr<FRigVMEditor>& InEditor)
+		const TWeakPtr<FRigVMEditorBase>& InEditor)
 		: PinPtr(InPin)
 		, GraphPtr(InGraph)
 		, WeakRigVMClientHost(InRigVMClientHost)
@@ -116,7 +117,7 @@ private:
 	TWeakInterfacePtr<IRigVMClientHost> WeakRigVMClientHost;
 
 	/** The editor we are editing */
-	TWeakPtr<FRigVMEditor> RigVMEditorPtr;
+	TWeakPtr<FRigVMEditorBase> RigVMEditorPtr;
 
 	/** Holds a weak pointer to the argument name widget, used for error notifications */
 	TWeakPtr<SEditableTextBox> ArgumentNameWidget;
@@ -162,16 +163,20 @@ class RIGVMEDITOR_API FRigVMGraphDetailCustomization : public IDetailCustomizati
 {
 public:
 	/** Makes a new instance of this detail layout class for a specific detail view requesting it */
+#if WITH_RIGVMLEGACYEDITOR
 	static TSharedPtr<IDetailCustomization> MakeInstance(TSharedPtr<IBlueprintEditor> InBlueprintEditor, const UClass* InExpectedBlueprintClass);
+#else
+	static TSharedPtr<IDetailCustomization> MakeInstance(TSharedPtr<FRigVMEditorBase> InBlueprintEditor, const UClass* InExpectedBlueprintClass);
+#endif
 
-	FRigVMGraphDetailCustomization(TSharedPtr<FRigVMEditor> RigVMigEditor, URigVMBlueprint* RigVMBlueprint);
+	FRigVMGraphDetailCustomization(TSharedPtr<FRigVMEditorBase> RigVMigEditor, URigVMBlueprint* RigVMBlueprint);
 
 	// IDetailCustomization interface
 	virtual void CustomizeDetails(IDetailLayoutBuilder& DetailLayout) override;
 
 private:
-	/** The Blueprint editor we are embedded in */
-	TWeakPtr<FRigVMEditor> RigVMEditorPtr;
+	/** The editor we are embedded in */
+	TWeakPtr<FRigVMEditorBase> RigVMEditorPtr;
 
 	/** The graph we are editing */
 	TWeakObjectPtr<URigVMEdGraph> GraphPtr;
@@ -189,7 +194,7 @@ public:
 		URigVMGraph* Model,
 		URigVMController* Controller,
 		IRigVMClientHost* InRigVMClientHost,
-		TWeakPtr<FRigVMEditor> InEditor);
+		TWeakPtr<FRigVMEditorBase> InEditor);
 
 private:
 
@@ -254,7 +259,7 @@ private:
 	TWeakObjectPtr<URigVMController> WeakController;
 
 	/** The editor we are embedded in */
-	TWeakPtr<FRigVMEditor> RigVMEditorPtr;
+	TWeakPtr<FRigVMEditorBase> RigVMEditorPtr;
 
 	/** The asset host we are editing */
 	TWeakInterfacePtr<IRigVMClientHost> RigVMClientHost;
diff --git a/Engine/Plugins/Runtime/RigVM/Source/RigVMEditor/Public/Editor/RigVMGraphEditorSummoner.h b/Engine/Plugins/Runtime/RigVM/Source/RigVMEditor/Public/Editor/RigVMGraphEditorSummoner.h
new file mode 100644
--- /dev/null
+++ b/Engine/Plugins/Runtime/RigVM/Source/RigVMEditor/Public/Editor/RigVMGraphEditorSummoner.h
@@ -0,0 +1,47 @@
+// Copyright Epic Games, Inc. All Rights Reserved.
+
+#pragma once
+
+#include "WorkflowOrientedApp/WorkflowUObjectDocuments.h"
+#include "EdGraph/EdGraph.h"
+
+#define LOCTEXT_NAMESPACE "RigVMGraphEditorSummoner"
+
+class UEdGraphSchema;
+struct FGraphDisplayInfo;
+class SGraphEditor;
+
+struct FRigVMLocalKismetCallbacks
+{
+	static FText GetGraphDisplayName(const UEdGraph* Graph);
+};
+
+struct FRigVMGraphEditorSummoner : public FDocumentTabFactoryForObjects<UEdGraph>
+{
+	static inline const FLazyName TabID = FLazyName(TEXT("RigVM Graph Editor"));
+	
+public:
+	DECLARE_DELEGATE_RetVal_TwoParams(TSharedRef<SGraphEditor>, FOnCreateGraphEditorWidget, TSharedRef<FTabInfo>, UEdGraph*);
+
+	FRigVMGraphEditorSummoner(TSharedPtr<class FRigVMNewEditor> InEditorPtr, FOnCreateGraphEditorWidget CreateGraphEditorWidgetCallback);
+
+	virtual void OnTabActivated(TSharedPtr<SDockTab> Tab) const override;
+	
+	virtual void OnTabBackgrounded(TSharedPtr<SDockTab> Tab) const override;
+
+	protected:
+	virtual TAttribute<FText> ConstructTabNameForObject(UEdGraph* DocumentID) const override
+	{
+		return TAttribute<FText>::Create(TAttribute<FText>::FGetter::CreateStatic(&FRigVMLocalKismetCallbacks::GetGraphDisplayName, (const UEdGraph*)DocumentID));
+	}
+
+	virtual TSharedRef<SWidget> CreateTabBodyForObject(const FWorkflowTabSpawnInfo& Info, UEdGraph* DocumentID) const override;
+
+	virtual const FSlateBrush* GetTabIconForObject(const FWorkflowTabSpawnInfo& Info, UEdGraph* DocumentID) const override;
+
+	protected:
+	TWeakPtr<class FRigVMNewEditor> BlueprintEditorPtr;
+	FOnCreateGraphEditorWidget OnCreateGraphEditorWidget;
+};
+
+#undef LOCTEXT_NAMESPACE
\ No newline at end of file
diff --git a/Engine/Plugins/Runtime/RigVM/Source/RigVMEditor/Public/Editor/RigVMNewEditor.h b/Engine/Plugins/Runtime/RigVM/Source/RigVMEditor/Public/Editor/RigVMNewEditor.h
new file mode 100644
--- /dev/null
+++ b/Engine/Plugins/Runtime/RigVM/Source/RigVMEditor/Public/Editor/RigVMNewEditor.h
@@ -0,0 +1,189 @@
+// Copyright Epic Games, Inc. All Rights Reserved.
+
+#pragma once
+
+#include "RigVMEditor.h"
+#include "WorkflowOrientedApp/WorkflowCentricApplication.h"
+
+struct RIGVMEDITOR_API FRigVMNewEditorTabs
+{
+	// Tab identifiers
+	static const FName DetailsID();
+	static const FName CompilerResultsID();
+};
+
+class RIGVMEDITOR_API FRigVMNewEditor : public FWorkflowCentricApplication, public FRigVMEditorBase, public FGCObject, public FNotifyHook, public FTickableEditorObject, public FEditorUndoClient, public FNoncopyable
+{
+public:
+	FRigVMNewEditor();
+	virtual void OnClose() override;
+
+	virtual TSharedRef<IRigVMEditor> SharedRef() override { return StaticCastSharedRef<IRigVMEditor>(SharedThis(this)); }
+
+protected:
+	
+	virtual TSharedPtr<FAssetEditorToolkit> GetHostingApp() override { return AsShared(); }
+	virtual void InitAssetEditor(const EToolkitMode::Type Mode, const TSharedPtr<IToolkitHost>& InitToolkitHost, const FName AppIdentifier, const TSharedRef<FTabManager::FLayout>& StandaloneDefaultLayout, const bool bCreateDefaultStandaloneMenu, const bool bCreateDefaultToolbar, const TArray<UObject*>& ObjectsToEdit, const bool bInIsToolbarFocusable, const bool bInUseSmallToolbarIcons, const TOptional<EAssetOpenMethod>& InOpenMethod) override;
+	virtual void CreateEditorToolbar() override {}
+	virtual void CommonInitialization(const TArray<UBlueprint*>& InitBlueprints, bool bShouldOpenInDefaultsMode) override;
+	void OnBlueprintChanged(UBlueprint* InBlueprint);
+	void SaveEditedObjectState();
+	virtual TSharedPtr<FDocumentTracker> GetDocumentManager() const override { return DocumentManager; }
+	virtual void AddApplicationMode(FName ModeName, TSharedRef<FApplicationMode> Mode) override { FWorkflowCentricApplication::AddApplicationMode(ModeName, Mode); }
+	virtual void RegenerateMenusAndToolbars() override { FWorkflowCentricApplication::RegenerateMenusAndToolbars(); }
+	virtual void SetCurrentMode(FName NewMode) override;
+	virtual FEditorModeTools& GetToolkitEditorModeManager() const override { return FWorkflowCentricApplication::GetEditorModeManager(); }
+	virtual void PostLayoutBlueprintEditorInitialization() override;
+	virtual TSharedPtr<SGraphEditor> OpenGraphAndBringToFront(UEdGraph* Graph, bool bSetFocus = true) override;
+	virtual bool FindOpenTabsContainingDocument(const UObject* DocumentID, /*inout*/ TArray< TSharedPtr<SDockTab> >& Results) override;
+	virtual TSharedPtr<SDockTab> OpenDocument(const UObject* DocumentID, FDocumentTracker::EOpenDocumentCause Cause) override;
+	virtual void CloseDocumentTab(const UObject* DocumentID) override;
+	virtual TSharedPtr<FTabManager> GetTabManager() override { return FWorkflowCentricApplication::GetTabManager(); }
+#if WITH_RIGVMLEGACYEDITOR
+	virtual TSharedPtr<SKismetInspector> GetInspector() const override { return Inspector.IsValid() ? Inspector.ToSharedRef().ToSharedPtr() : nullptr; }
+#else
+	virtual TSharedPtr<SRigVMDetailsInspector> GetInspector() const override { return Inspector.IsValid() ? Inspector.ToSharedRef().ToSharedPtr() : nullptr; }
+#endif
+	virtual TSharedPtr<FApplicationMode> CreateEditorMode() override;
+	virtual const FName GetEditorAppName() const override;
+	virtual const TArray< UObject* >& GetEditingObjects() const override { return FWorkflowCentricApplication::GetEditingObjects(); }
+	virtual const TSharedRef<IToolkitHost> GetToolkitHost() const override { return FWorkflowCentricApplication::GetToolkitHost(); }
+	virtual bool IsHosted() const override { return FWorkflowCentricApplication::IsHosted(); }
+	virtual void BringToolkitToFrontImpl() override { FWorkflowCentricApplication::BringToolkitToFront(); }
+	virtual void InitToolMenuContext(FToolMenuContext& MenuContext) override;
+	virtual TSharedRef<FUICommandList> GetToolkitCommands() override { return ToolkitCommands; }
+	virtual TWeakPtr<SGraphEditor> GetFocusedGraphEditor() override { return FocusedGraphEdPtr; }
+	virtual TWeakPtr<FDocumentTabFactory> GetGraphEditorTabFactory() const override { return GraphEditorTabFactoryPtr; }
+	virtual bool TransactionObjectAffectsBlueprint(UObject* InTransactedObject) override;
+	virtual FEdGraphPinType GetLastPinTypeUsed() override;
+	virtual void LogSimpleMessage(const FText& MessageText) override;
+	virtual void RenameNewlyAddedAction(FName InActionName) override;
+	virtual FGraphPanelSelectionSet GetSelectedNodes() const override;
+	virtual void SetUISelectionState(FName SelectionOwner) override;
+	virtual void AnalyticsTrackNodeEvent(UBlueprint* Blueprint, UEdGraphNode* GraphNode, bool bNodeDelete) const override;
+	virtual void JumpToHyperlink(const UObject* ObjectReference, bool bRequestRename) override;
+	virtual void PostUndo(bool bSuccess) override;
+	virtual void PostRedo(bool bSuccess) override;
+	virtual UEdGraphPin* GetCurrentlySelectedPin() const override;
+	virtual void CreateDefaultCommands() override;
+	virtual TSharedRef<SGraphEditor> CreateGraphEditorWidget(TSharedRef<FTabInfo> InTabInfo, UEdGraph* InGraph) override;
+	virtual void CompileImpl() override;
+	virtual void SaveAsset_Execute_Impl() override { FWorkflowCentricApplication::SaveAsset_Execute(); }
+	virtual void SaveAssetAs_Execute_Impl() override { FWorkflowCentricApplication::SaveAssetAs_Execute(); }
+	virtual bool IsGraphInCurrentBlueprint(const UEdGraph* InGraph) const override;
+	virtual bool IsEditableImpl(UEdGraph* InGraph) const override;
+	virtual UEdGraph* GetFocusedGraph() const override;
+	virtual void JumpToNode(const UEdGraphNode* Node, bool bRequestRename) override;
+	virtual void JumpToPin(const UEdGraphPin* Pin) override;
+	virtual void AddToolbarExtender(TSharedPtr<FExtender> Extender) override { FWorkflowCentricApplication::AddToolbarExtender(Extender); }
+	virtual void RemoveToolbarExtender(TSharedPtr<FExtender> Extender) override { FWorkflowCentricApplication::RemoveToolbarExtender(Extender); };
+	virtual void AddMenuExtender(TSharedPtr<FExtender> Extender) override { FWorkflowCentricApplication::AddMenuExtender(Extender); }
+	virtual void RemoveMenuExtender(TSharedPtr<FExtender> Extender) override { FWorkflowCentricApplication::RemoveMenuExtender(Extender); }
+	virtual TSharedPtr<IMessageLogListing> GetCompilerResultsListing() override { return CompilerResultsListing; }
+	virtual void OnBlueprintChangedInnerImpl(UBlueprint* InBlueprint, bool bIsJustBeingCompiled) override;
+	virtual void RefreshEditorsImpl(ERefreshRigVMEditorReason::Type Reason) override;
+	virtual void SetupGraphEditorEventsImpl(UEdGraph* InGraph, SGraphEditor::FGraphEditorEvents& InEvents) override;
+	virtual FActionMenuContent OnCreateGraphActionMenu(UEdGraph* InGraph, const FVector2D& InNodePosition, const TArray<UEdGraphPin*>& InDraggedPins, bool bAutoExpand, SGraphEditor::FActionMenuClosed InOnMenuClosed) override;
+	virtual void AddReferencedObjectsImpl(FReferenceCollector& Collector) override;
+	virtual bool IsSectionVisible(RigVMNodeSectionID::Type InSectionID) const; 
+	virtual bool NewDocument_IsVisibleForType(ECreatedDocumentType GraphType) const;
+	virtual FGraphAppearanceInfo GetGraphAppearanceImpl(UEdGraph* InGraph) const override;
+	virtual void NotifyPreChangeImpl(FProperty* PropertyAboutToChange) override;
+	virtual void NotifyPostChangeImpl(const FPropertyChangedEvent& PropertyChangedEvent, FProperty* PropertyThatChanged) override;
+	virtual FName GetSelectedVariableName() override;
+	virtual bool IsNodeTitleVisible(const UEdGraphNode* Node, bool bRequestRename) override;
+	virtual void EditClassDefaults_Clicked() override;
+	virtual void EditGlobalOptions_Clicked() override;
+	virtual void TryInvokingDetailsTab(bool bFlash = true) override;
+	virtual FName GetGraphExplorerWidgetID() override { return FRigVMEditorGraphExplorerTabSummoner::TabID; }
+	virtual void RefreshInspector() override;
+	virtual void RefreshStandAloneDefaultsEditor() override;
+	virtual void Tick(float DeltaTime) override;
+	virtual void GetPinTypeSelectorFilters(TArray<TSharedPtr<IPinTypeSelectorFilter>>& OutFilters) const override;
+	virtual void OnAddNewVariable() override;
+	virtual void ZoomToSelection_Clicked() override;
+public:
+	virtual void SummonSearchUI(bool bSetFindWithinBlueprint, FString NewSearchTerms, bool bSelectFirstResult) override;
+	virtual void RestoreEditedObjectState() override;
+	virtual void SetupViewForBlueprintEditingMode() override;
+	virtual void OnGraphEditorFocused(const TSharedRef<SGraphEditor>& InGraphEditor) override;
+	virtual void OnGraphEditorBackgrounded(const TSharedRef<SGraphEditor>& InGraphEditor);
+	virtual bool GetIsContextSensitive() override { return bIsActionMenuContextSensitive; }
+	virtual void SetIsContextSensitive(const bool bIsContextSensitive) override { bIsActionMenuContextSensitive = bIsContextSensitive; }
+	virtual void RegisterToolbarTab(const TSharedRef<FTabManager>& InTabManager) override { FAssetEditorToolkit::RegisterTabSpawners(InTabManager); }
+	virtual const TArray<UObject*>* GetObjectsCurrentlyBeingEdited() const override { return FAssetEditorToolkit::GetObjectsCurrentlyBeingEdited(); }
+	virtual void AddSelectedDebugObjectWidget(FToolBarBuilder& ToolbarBuilder) override {}
+
+	//~ Begin IToolkit Interface
+	virtual void RegisterTabSpawners(const TSharedRef<class FTabManager>& InTabManager) override;
+	//~ End IToolkit Interface
+
+
+	virtual FLinearColor GetWorldCentricTabColorScale() const override { return FRigVMEditorBase::GetWorldCentricTabColorScale(); }
+	virtual void AddReferencedObjects(FReferenceCollector& Collector) override { return FRigVMEditorBase::AddReferencedObjects(Collector); }
+	virtual FString GetReferencerName() const override { return TEXT("FRigVMNewEditor"); }
+	virtual TStatId GetStatId() const override { RETURN_QUICK_DECLARE_CYCLE_STAT(FRigVMNewEditor, STATGROUP_Tickables); }
+	virtual void StartEditingDefaults(bool bAutoFocus = true, bool bForceRefresh = false);
+
+	float GetInstructionTextOpacity(UEdGraph* InGraph) const;
+	virtual void ClearSelectionStateFor(FName SelectionOwner);
+	void CreateDefaultTabContents(const TArray<UBlueprint*> InBlueprints);
+
+public:
+	static const FSlateBrush* GetGlyphForGraph(const UEdGraph* Graph, bool bInLargeIcon);
+
+	static const FName SelectionState_GraphExplorer();
+	static const FName SelectionState_Graph();
+	static const FName SelectionState_ClassSettings();
+	static const FName SelectionState_ClassDefaults();
+
+protected:
+	
+	TSharedPtr<FDocumentTracker> DocumentManager;
+	/** Node inspector widget */
+#if WITH_RIGVMLEGACYEDITOR
+	TSharedPtr<class SKismetInspector> Inspector;
+#else
+	TSharedPtr<class SRigVMDetailsInspector> Inspector;
+#endif
+
+	/** Currently focused graph editor */
+	TWeakPtr<class SGraphEditor> FocusedGraphEdPtr;
+
+	/** Factory that spawns graph editors; used to look up all tabs spawned by it. */
+	TWeakPtr<FDocumentTabFactory> GraphEditorTabFactoryPtr;
+
+	/** The current UI selection state of this editor */
+	FName CurrentUISelection;
+
+	/** Command list for the graph editor */
+	TSharedPtr<FUICommandList> GraphEditorCommands;
+
+	/** Compiler results log, with the log listing that it reflects */
+	//TSharedPtr<class SWidget> CompilerResults;
+	TSharedPtr<class IMessageLogListing> CompilerResultsListing;
+
+	/** When set, flags which graph has a action menu currently open (if null, no graphs do). */
+	UEdGraph* HasOpenActionMenu;
+	
+	/** Used to nicely fade instruction text, when the context menu is opened. */
+	float InstructionsFadeCountdown;
+
+	/** defaults inspector widget */
+#if WITH_RIGVMLEGACYEDITOR
+	TSharedPtr<class SKismetInspector> DefaultEditor;
+#else
+	TSharedPtr<class SRigVMDetailsInspector> DefaultEditor;
+#endif
+	
+	/** True if the editor was opened in defaults mode */
+	bool bWasOpenedInDefaultsMode;
+
+	/** Did we update the blueprint when it opened */
+	bool bBlueprintModifiedOnOpen;
+
+	/** Whether the graph action menu should be sensitive to the pins dragged off of */
+	bool bIsActionMenuContextSensitive;
+};
+
+
diff --git a/Engine/Plugins/Runtime/RigVM/Source/RigVMEditor/Public/Editor/RigVMNewEditorMode.h b/Engine/Plugins/Runtime/RigVM/Source/RigVMEditor/Public/Editor/RigVMNewEditorMode.h
new file mode 100644
--- /dev/null
+++ b/Engine/Plugins/Runtime/RigVM/Source/RigVMEditor/Public/Editor/RigVMNewEditorMode.h
@@ -0,0 +1,42 @@
+// Copyright Epic Games, Inc. All Rights Reserved.
+
+#pragma once
+
+#include "CoreMinimal.h"
+#include "WorkflowOrientedApp/WorkflowTabManager.h"
+#include "WorkflowOrientedApp/ApplicationMode.h"
+
+class URigVMBlueprint;
+class FRigVMNewEditor;
+
+struct RIGVMEDITOR_API FRigVMNewEditorApplicationModes
+{
+	// Mode identifiers
+	static const FName StandardRigVMEditorMode();
+	static const FName RigVMDefaultsMode();
+	
+private:
+	FRigVMNewEditorApplicationModes() {}
+};
+
+class RIGVMEDITOR_API FRigVMNewEditorMode : public FApplicationMode
+{
+public:
+	FRigVMNewEditorMode(const TSharedRef<FRigVMNewEditor>& InRigVMEditor); 
+
+	// FApplicationMode interface
+	virtual void RegisterTabFactories(TSharedPtr<FTabManager> InTabManager) override;
+
+	void RegisterToolbarTab(const TSharedRef<class FTabManager>& TabManager);
+
+	virtual void PostActivateMode() override;
+
+protected:
+	// Set of spawnable tabs
+	FWorkflowAllowedTabSet TabFactories;
+
+	TWeakPtr<FRigVMNewEditor> Editor;
+private:
+	
+	TWeakObjectPtr<URigVMBlueprint> RigVMBlueprintPtr;
+};
diff --git a/Engine/Plugins/Runtime/RigVM/Source/RigVMEditor/Public/Editor/SRigVMDetailsInspector.h b/Engine/Plugins/Runtime/RigVM/Source/RigVMEditor/Public/Editor/SRigVMDetailsInspector.h
new file mode 100644
--- /dev/null
+++ b/Engine/Plugins/Runtime/RigVM/Source/RigVMEditor/Public/Editor/SRigVMDetailsInspector.h
@@ -0,0 +1,89 @@
+// Copyright Epic Games, Inc. All Rights Reserved.
+
+#pragma once
+
+#if !WITH_RIGVMLEGACYEDITOR
+#include "Editor/RigVMEditor.h"
+
+class RIGVMEDITOR_API SRigVMDetailsInspector : public SCompoundWidget, public FGCObject
+{
+	SLATE_BEGIN_ARGS( SRigVMDetailsInspector )
+		: _ShowPublicViewControl(false)
+		, _HideNameArea(true)
+		, _SetNotifyHook(true)
+		, _ShowTitleArea(false)
+		, _ShowLocalVariables(false)
+		, _ScrollbarAlignment(HAlign_Right)
+		, _ShowSectionSelector(false)
+	{}
+
+	SLATE_ARGUMENT(TWeakPtr<FRigVMEditorBase>, Editor)
+	SLATE_ARGUMENT(TWeakPtr<SRigVMEditorGraphExplorer>, GraphExplorer)
+	SLATE_ATTRIBUTE( bool, ShowPublicViewControl )
+	SLATE_ARGUMENT( bool, HideNameArea )
+	SLATE_ARGUMENT( FIsPropertyEditingEnabled, IsPropertyEditingEnabledDelegate )
+	SLATE_ARGUMENT( FOnFinishedChangingProperties::FDelegate, OnFinishedChangingProperties )
+	SLATE_ARGUMENT( FName, ViewIdentifier)
+	SLATE_ARGUMENT( bool, SetNotifyHook)
+	SLATE_ARGUMENT( bool, ShowTitleArea)
+	SLATE_ARGUMENT( bool, ShowLocalVariables)
+	SLATE_ARGUMENT( TSharedPtr<SScrollBar>, ExternalScrollbar)
+	SLATE_ARGUMENT( EHorizontalAlignment, ScrollbarAlignment)
+	SLATE_ARGUMENT( bool, ShowSectionSelector)
+	SLATE_END_ARGS()
+
+	void Construct(const FArguments& InArgs){}
+
+	/** Options for ShowDetails */
+	struct FShowDetailsOptions
+	{
+		FText ForcedTitle;
+		bool bForceRefresh;
+
+		FShowDetailsOptions()
+			:ForcedTitle()
+			,bForceRefresh(false)
+		{}
+		
+		FShowDetailsOptions(const FText& InForcedTitle, bool bInForceRefresh = false)
+			:ForcedTitle(InForcedTitle)
+			,bForceRefresh(bInForceRefresh)
+		{}
+	};
+
+
+ 	/** Update the inspector window to show information on the supplied object */
+	void ShowDetailsForSingleObject(UObject* Object, const FShowDetailsOptions& Options = FShowDetailsOptions()){}
+
+ 	/** Update the inspector window to show information on the supplied objects */
+	void ShowDetailsForObjects(const TArray<UObject*>& PropertyObjects, const FShowDetailsOptions& Options = FShowDetailsOptions()){}
+
+ 	/** Update the inspector window to show single struct. This invalidates ShowDetailsForObjects */
+	void ShowSingleStruct(TSharedPtr<FStructOnScope> InStructToDisplay){}
+
+	TSharedPtr<class IDetailsView> GetPropertyView() const { return PropertyView; }
+
+	void SetOwnerTab(TSharedRef<SDockTab> Tab);
+	TSharedPtr<SDockTab> GetOwnerTab() const { return OwnerTab.Pin(); }
+
+	/** returns the list of selected objects */
+	const TArray< TWeakObjectPtr<UObject> >& GetSelectedObjects() const { return SelectedObjects; }
+
+	virtual void AddReferencedObjects(FReferenceCollector& Collector) override { Collector.AddReferencedObjects(RefreshPropertyObjects); }
+	virtual FString GetReferencerName() const override { return TEXT("SRigVMDetailsInspector"); }
+
+protected:
+
+	/** The tab that owns this details view. */
+	TWeakPtr<SDockTab> OwnerTab;
+
+	/** Property viewing widget */
+	TSharedPtr<class IDetailsView> PropertyView;
+
+	/** Selected objects for this detail view */
+	TArray< TWeakObjectPtr<UObject> > SelectedObjects;
+
+	/** Holds the property objects that need to be displayed by the inspector starting on the next tick */
+	TArray<TObjectPtr<UObject>> RefreshPropertyObjects;
+};
+#endif
\ No newline at end of file
