Control Rig: New RigVMEditor, which does not inherit from BlueprintEditor

diff --git a/Engine/Plugins/Runtime/RigVM/Source/RigVMEditor/Public/RigVMEditorModule.h b/Engine/Plugins/Runtime/RigVM/Source/RigVMEditor/Public/RigVMEditorModule.h
--- a/Engine/Plugins/Runtime/RigVM/Source/RigVMEditor/Public/RigVMEditorModule.h
+++ b/Engine/Plugins/Runtime/RigVM/Source/RigVMEditor/Public/RigVMEditorModule.h
@@ -9,14 +9,35 @@
 #include "CoreMinimal.h"
 #include "Modules/ModuleInterface.h"
 #include "Logging/LogMacros.h"
-#include "Editor/RigVMEditor.h"
 #include "EdGraph/RigVMEdGraphPanelNodeFactory.h"
 #include "EdGraph/RigVMEdGraphPanelPinFactory.h"
 #include "Toolkits/AssetEditorToolkit.h"
 #include "Kismet2/StructureEditorUtils.h"
+#include "Modules/ModuleManager.h"
+
+class URigVMUserWorkflowOptions;
+class URigVMPin;
+class URigVMNode;
+class URigVMEdGraphNode;
+class IRigVMClientHost;
+class URigVMBlueprint;
+class URigVMEdGraphSchema;
+class FRigVMEditorBase;
+class FBlueprintActionDatabaseRegistrar;
 
 DECLARE_LOG_CATEGORY_EXTERN(LogRigVMEditor, Log, All);
 
+/** Describes the reason for Refreshing the editor */
+namespace ERefreshRigVMEditorReason
+{
+	enum Type
+	{
+		BlueprintCompiled,
+		PostUndo,
+		UnknownReason
+	};
+}
+
 // shallow interface declaration for use within RigVMDeveloper 
 class IRigVMEditorModule : public IModuleInterface, public IHasMenuExtensibility, public IHasToolBarExtensibility, public FStructureEditorUtils::INotifyOnStructChanged
 {
@@ -69,7 +90,7 @@ public:
 	virtual void PostChange(const UUserDefinedStruct* Changed, FStructureEditorUtils::EStructureEditorChangeInfo ChangedType) override;
 
 	/** Get all toolbar extenders */
-	DECLARE_DELEGATE_RetVal_TwoParams(TSharedRef<FExtender>, FRigVMEditorToolbarExtender, const TSharedRef<FUICommandList> /*InCommandList*/, TSharedRef<FRigVMEditor> /*InRigVMEditor*/);
+	DECLARE_DELEGATE_RetVal_TwoParams(TSharedRef<FExtender>, FRigVMEditorToolbarExtender, const TSharedRef<FUICommandList> /*InCommandList*/, TSharedRef<FRigVMEditorBase> /*InRigVMEditor*/);
 	virtual TArray<FRigVMEditorToolbarExtender>& GetAllRigVMEditorToolbarExtenders();
 
 	virtual void GetContextMenuActions(const URigVMEdGraphSchema* Schema, class UToolMenu* Menu, class UGraphNodeContextMenuContext* Context) const override;
diff --git a/Engine/Plugins/Runtime/RigVM/Source/RigVMEditor/Public/Widgets/SRigVMEditorGraphExplorer.h b/Engine/Plugins/Runtime/RigVM/Source/RigVMEditor/Public/Widgets/SRigVMEditorGraphExplorer.h
--- a/Engine/Plugins/Runtime/RigVM/Source/RigVMEditor/Public/Widgets/SRigVMEditorGraphExplorer.h
+++ b/Engine/Plugins/Runtime/RigVM/Source/RigVMEditor/Public/Widgets/SRigVMEditorGraphExplorer.h
@@ -11,11 +11,12 @@
 #include "RigVMModel/RigVMGraph.h"
 #include "SRigVMEditorGraphExplorerTreeView.h"
 
+class FRigVMEditorBase;
 class SRigVMEditorGraphExplorerTreeView;
 //
 class URigVMGraph;
 class URigVMEdGraph;
-class FRigVMEditor;
+class FRigVMEditorBase;
 class URigVMBlueprint;
 class SGraphActionMenu;
 struct FCreateWidgetForActionData;
@@ -79,7 +80,7 @@ public:
 	SLATE_BEGIN_ARGS(SRigVMEditorGraphExplorer) {}
 	SLATE_END_ARGS()
 
-	void Construct(const FArguments& InArgs, TWeakPtr<FRigVMEditor> InRigVMEditor);
+	void Construct(const FArguments& InArgs, TWeakPtr<FRigVMEditorBase> InRigVMEditor);
 	~SRigVMEditorGraphExplorer();
 
 	// Refresh the graph action menu.
@@ -88,8 +89,17 @@ public:
 	// SWidget overrides
 	void Tick(const FGeometry& InAllottedGeometry, const double InCurrentTime, const float InDeltaTime) override;
 
+	void SetLastPinTypeUsed(const FEdGraphPinType& InType) { LastPinType = InType; }
+	FEdGraphPinType GetLastPinTypeUsed() const { return LastPinType; }
+
+	FName GetSelectedVariableName() const;
+
+	void ClearSelection();
+
 private:
 	bool bNeedsRefresh;
+
+	FEdGraphPinType LastPinType;
 	
 	void RegisterCommands();
 	void CreateWidgets();
@@ -186,7 +196,7 @@ private:
 
 	void HandleSelectionChanged(TSharedPtr<FRigVMEditorGraphExplorerTreeElement> Selection, ESelectInfo::Type SelectInfo);
 
-	TWeakPtr<FRigVMEditor> RigVMEditor;
+	TWeakPtr<FRigVMEditorBase> RigVMEditor;
 
 	TSharedPtr<SRigVMEditorGraphExplorerTreeView> TreeView;
 
diff --git a/Engine/Plugins/Runtime/RigVM/Source/RigVMEditor/Public/Widgets/SRigVMEditorGraphExplorerTreeView.h b/Engine/Plugins/Runtime/RigVM/Source/RigVMEditor/Public/Widgets/SRigVMEditorGraphExplorerTreeView.h
--- a/Engine/Plugins/Runtime/RigVM/Source/RigVMEditor/Public/Widgets/SRigVMEditorGraphExplorerTreeView.h
+++ b/Engine/Plugins/Runtime/RigVM/Source/RigVMEditor/Public/Widgets/SRigVMEditorGraphExplorerTreeView.h
@@ -7,7 +7,7 @@
 #include "RigVMModel/RigVMGraph.h"
 #include "Widgets/Views/STreeView.h"
 
-class FRigVMEditor;
+class FRigVMEditorBase;
 class SSearchBox;
 class SRigVMEditorGraphExplorerTreeView;
 class SRigVMEditorGraphExplorerItem;
@@ -15,28 +15,32 @@ class FRigVMEditorGraphExplorerTreeElement;
 class URigVMGraph;
 class IPinTypeSelectorFilter;
 
-enum ERigVMExplorerElementType
+namespace ERigVMExplorerElementType
 {
-	Invalid,
-	Section,
-	FunctionCategory,
-	VariableCategory,
-	Graph,
-	Event,
-	Function,
-	Variable,
-	LocalVariable
-};
+	enum Type
+	{
+		Invalid,
+		Section,
+		FunctionCategory,
+		VariableCategory,
+		Graph,
+		Event,
+		Function,
+		Variable,
+		LocalVariable
+	};
+}
 
 struct FRigVMExplorerElementKey
 {
+
 	FRigVMExplorerElementKey()
 		: Type(ERigVMExplorerElementType::Invalid), Name(FString()){}
 
-	FRigVMExplorerElementKey(const ERigVMExplorerElementType& InType, const FString& InName)
+	FRigVMExplorerElementKey(const ERigVMExplorerElementType::Type& InType, const FString& InName)
 		: Type(InType), Name(InName){}
 	
-	ERigVMExplorerElementType Type;
+	ERigVMExplorerElementType::Type Type;
 	FString Name;
 
 	friend bool operator==(const FRigVMExplorerElementKey& A, const FRigVMExplorerElementKey& B)
diff --git a/Engine/Plugins/Runtime/RigVM/Source/RigVMEditor/Public/Widgets/SRigVMExecutionStackView.h b/Engine/Plugins/Runtime/RigVM/Source/RigVMEditor/Public/Widgets/SRigVMExecutionStackView.h
--- a/Engine/Plugins/Runtime/RigVM/Source/RigVMEditor/Public/Widgets/SRigVMExecutionStackView.h
+++ b/Engine/Plugins/Runtime/RigVM/Source/RigVMEditor/Public/Widgets/SRigVMExecutionStackView.h
@@ -10,7 +10,7 @@
 #include "RigVMCore/RigVM.h"
 #include "RigVMBlueprint.h"
 
-class FRigVMEditor;
+class FRigVMEditorBase;
 class SRigVMExecutionStackView;
 class FUICommandList;
 class SSearchBox;
@@ -75,7 +75,7 @@ public:
 
 	~SRigVMExecutionStackView();
 
-	void Construct(const FArguments& InArgs, TSharedRef<FRigVMEditor> InRigVMEditor);
+	void Construct(const FArguments& InArgs, TSharedRef<FRigVMEditorBase> InRigVMEditor);
 
 	/** Set Selection Changed */
 	void OnSelectionChanged(TSharedPtr<FRigStackEntry> Selection, ESelectInfo::Type SelectInfo);
@@ -126,7 +126,7 @@ private:
 	bool bSuspendControllerSelection;
 	void HandleModifiedEvent(ERigVMGraphNotifType InNotifType, URigVMGraph* InGraph, UObject* InSubject);
 	void HandleHostInitializedEvent(URigVMHost* InHost, const FName& InEventName);
-	void HandlePreviewHostUpdated(FRigVMEditor* InEditor);
+	void HandlePreviewHostUpdated(FRigVMEditorBase* InEditor);
 	void HandleItemMouseDoubleClick(TSharedPtr<FRigStackEntry> InItem);
 
 	/** Populate the execution stack with descriptive names for each instruction */
@@ -137,7 +137,7 @@ private:
 	/** Command list we bind to */
 	TSharedPtr<FUICommandList> CommandList;
 
-	TWeakPtr<FRigVMEditor> RigVMEditor;
+	TWeakPtr<FRigVMEditorBase> WeakEditor;
 	TWeakObjectPtr<URigVMBlueprint> RigVMBlueprint;
 
 	TArray<TSharedPtr<FRigStackEntry>> Operators;
diff --git a/Engine/Plugins/Runtime/RigVM/Source/RigVMEditor/RigVMEditor.Build.cs b/Engine/Plugins/Runtime/RigVM/Source/RigVMEditor/RigVMEditor.Build.cs
--- a/Engine/Plugins/Runtime/RigVM/Source/RigVMEditor/RigVMEditor.Build.cs
+++ b/Engine/Plugins/Runtime/RigVM/Source/RigVMEditor/RigVMEditor.Build.cs
@@ -21,7 +21,6 @@ public class RigVMEditor : ModuleRules
                 "SlateCore",
                 "InputCore",
                 "ToolWidgets",
-                "Kismet",
                 "KismetCompiler",
                 "KismetWidgets",
                 "AnimationWidgets",
@@ -35,5 +34,12 @@ public class RigVMEditor : ModuleRules
 				"ContentBrowser",
 			}
 		);
+        
+        bool bAddBlueprintEditorDependency = true;
+        PublicDefinitions.Add("WITH_RIGVMLEGACYEDITOR=" + (bAddBlueprintEditorDependency ? '1' : '0'));
+        if (bAddBlueprintEditorDependency)
+        {
+			PublicDependencyModuleNames.Add("Kismet");
+		}
 	}
 }
diff --git a/Engine/Source/Editor/KismetCompiler/Public/KismetCompiler.h b/Engine/Source/Editor/KismetCompiler/Public/KismetCompiler.h
--- a/Engine/Source/Editor/KismetCompiler/Public/KismetCompiler.h
+++ b/Engine/Source/Editor/KismetCompiler/Public/KismetCompiler.h
@@ -678,6 +678,7 @@ private:
 	 * we have decided to make the compilation manager a friend:
 	 */
 	friend struct FBlueprintCompilationManagerImpl;
+	friend struct FRigVMBlueprintCompilationManagerImpl;
 };
 
 //////////////////////////////////////////////////////////////////////////
diff --git a/Engine/Source/Editor/UnrealEd/Public/Kismet2/KismetReinstanceUtilities.h b/Engine/Source/Editor/UnrealEd/Public/Kismet2/KismetReinstanceUtilities.h
--- a/Engine/Source/Editor/UnrealEd/Public/Kismet2/KismetReinstanceUtilities.h
+++ b/Engine/Source/Editor/UnrealEd/Public/Kismet2/KismetReinstanceUtilities.h
@@ -107,6 +107,7 @@ class FBlueprintCompileReinstancer : public TSharedFromThis<FBlueprintCompileRei
 {
 protected:
 	friend struct FBlueprintCompilationManagerImpl;
+	friend struct FRigVMBlueprintCompilationManagerImpl;
 	friend struct FReinstancingJob;
 
 	static UNREALED_API TSet<TWeakObjectPtr<UBlueprint>> CompiledBlueprintsToSave;
