[Backout] - CL39678024
#rnx
Original CL Desc
-----------------------------------------------------------------
Control Rig: New RigVMEditor, which does not inherit from BlueprintEditor
#rnx

diff --git a/Engine/Plugins/Animation/ControlRig/Source/ControlRigDeveloper/Private/Graph/ControlRigGraphNode.cpp b/Engine/Plugins/Animation/ControlRig/Source/ControlRigDeveloper/Private/Graph/ControlRigGraphNode.cpp
--- a/Engine/Plugins/Animation/ControlRig/Source/ControlRigDeveloper/Private/Graph/ControlRigGraphNode.cpp
+++ b/Engine/Plugins/Animation/ControlRig/Source/ControlRigDeveloper/Private/Graph/ControlRigGraphNode.cpp
@@ -4,9 +4,6 @@
 #include "Graph/ControlRigGraph.h"
 #include "Graph/ControlRigGraphSchema.h"
 #include "IControlRigEditorModule.h"
-#if !WITH_RIGVMLEGACYEDITOR
-#include "RigVMDeveloper/Private/Kismet/RigVMFindInBlueprintManager.h"
-#endif
 
 #include UE_INLINE_GENERATED_CPP_BY_NAME(ControlRigGraphNode)
 
@@ -19,15 +16,9 @@ UControlRigGraphNode::UControlRigGraphNode()
 
 #if WITH_EDITOR
 
-#if WITH_RIGVMLEGACYEDITOR
 void UControlRigGraphNode::AddPinSearchMetaDataInfo(const UEdGraphPin* Pin, TArray<FSearchTagDataPair>& OutTaggedMetaData) const
 {
 	Super::AddPinSearchMetaDataInfo(Pin, OutTaggedMetaData);
-#else
-void UControlRigGraphNode::AddRigVMPinSearchMetaDataInfo(const UEdGraphPin* Pin, TArray<UBlueprintExtension::FSearchTagDataPair>& OutTaggedMetaData) const
-{
-	Super::AddRigVMPinSearchMetaDataInfo(Pin, OutTaggedMetaData);
-#endif
 
 	if(const URigVMPin* ModelPin = FindModelPinFromGraphPin(Pin))
 	{
diff --git a/Engine/Plugins/Animation/ControlRig/Source/ControlRigDeveloper/Public/Graph/ControlRigGraphNode.h b/Engine/Plugins/Animation/ControlRig/Source/ControlRigDeveloper/Public/Graph/ControlRigGraphNode.h
--- a/Engine/Plugins/Animation/ControlRig/Source/ControlRigDeveloper/Public/Graph/ControlRigGraphNode.h
+++ b/Engine/Plugins/Animation/ControlRig/Source/ControlRigDeveloper/Public/Graph/ControlRigGraphNode.h
@@ -18,9 +18,6 @@ class FBlueprintActionDatabaseRegistrar;
 class UEdGraph;
 struct FSlateIcon;
 class UControlRigBlueprint;
-#if WITH_RIGVMLEGACYEDITOR
-struct FSearchTagDataPair;
-#endif
 
 /** Base class for animation ControlRig-related nodes */
 UCLASS()
@@ -43,11 +40,7 @@ public:
 	UControlRigGraphNode();
 
 #if WITH_EDITOR
-#if WITH_RIGVMLEGACYEDITOR
-	void AddPinSearchMetaDataInfo(const UEdGraphPin* Pin, TArray<FSearchTagDataPair>& OutTaggedMetaData) const;
-#else
-	void AddRigVMPinSearchMetaDataInfo(const UEdGraphPin* Pin, TArray<UBlueprintExtension::FSearchTagDataPair>& OutTaggedMetaData) const;
-#endif
+	virtual void AddPinSearchMetaDataInfo(const UEdGraphPin* Pin, TArray<FSearchTagDataPair>& OutTaggedMetaData) const override;
 #endif
 
 private:
diff --git a/Engine/Plugins/Animation/ControlRig/Source/ControlRigEditor/Private/ControlRigDrawingDetails.cpp b/Engine/Plugins/Animation/ControlRig/Source/ControlRigEditor/Private/ControlRigDrawingDetails.cpp
--- a/Engine/Plugins/Animation/ControlRig/Source/ControlRigEditor/Private/ControlRigDrawingDetails.cpp
+++ b/Engine/Plugins/Animation/ControlRig/Source/ControlRigEditor/Private/ControlRigDrawingDetails.cpp
@@ -14,12 +14,7 @@
 #include "FbxImporter.h"
 #include "ScopedTransaction.h"
 #include "Dialogs/Dialogs.h"
-
-#if WITH_RIGVMLEGACYEDITOR
 #include "SKismetInspector.h"
-#else
-#include "Editor/SRigVMDetailsInspector.h"
-#endif
 
 #include UE_INLINE_GENERATED_CPP_BY_NAME(ControlRigDrawingDetails)
 
@@ -112,11 +107,8 @@ FReply FControlRigDrawContainerDetails::OnImportCurvesFromFBXClicked()
 
 		FControlRigDrawContainerImportFbxSettings Settings;
 		TSharedPtr<FStructOnScope> StructToDisplay = MakeShareable(new FStructOnScope(FControlRigDrawContainerImportFbxSettings::StaticStruct(), (uint8*)&Settings));
-#if WITH_RIGVMLEGACYEDITOR
+
 		TSharedRef<SKismetInspector> KismetInspector = SNew(SKismetInspector);
-#else
-		TSharedRef<SRigVMDetailsInspector> KismetInspector = SNew(SRigVMDetailsInspector);
-#endif
 		KismetInspector->ShowSingleStruct(StructToDisplay);
 
 		SGenericDialogWidget::FArguments DialogArguments;
diff --git a/Engine/Plugins/Animation/ControlRig/Source/ControlRigEditor/Private/ControlRigEditorModule.cpp b/Engine/Plugins/Animation/ControlRig/Source/ControlRigEditor/Private/ControlRigEditorModule.cpp
--- a/Engine/Plugins/Animation/ControlRig/Source/ControlRigEditor/Private/ControlRigEditorModule.cpp
+++ b/Engine/Plugins/Animation/ControlRig/Source/ControlRigEditor/Private/ControlRigEditorModule.cpp
@@ -54,6 +54,7 @@
 #include "ControlRigThumbnailRenderer.h"
 #include "RigVMModel/RigVMController.h"
 #include "ControlRigBlueprint.h"
+#include "SKismetInspector.h"
 #include "Dialogs/Dialogs.h"
 #include "Settings/ControlRigSettings.h"
 #include "EditMode/ControlRigControlsProxy.h"
@@ -99,12 +100,6 @@
 #include "SequencerUtilities.h"
 #include "Bindings/MovieSceneSpawnableActorBinding.h"
 
-#if WITH_RIGVMLEGACYEDITOR
-#include "SKismetInspector.h"
-#else
-#include "Editor/SRigVMDetailsInspector.h"
-#endif
-
 #define LOCTEXT_NAMESPACE "ControlRigEditorModule"
 
 DEFINE_LOG_CATEGORY(LogControlRigEditor);
@@ -671,11 +666,8 @@ void FControlRigEditorModule::GetNodeContextMenuActions(IRigVMClientHost* RigVMC
 
 				FRigVMMirrorSettings Settings;
 				TSharedPtr<FStructOnScope> StructToDisplay = MakeShareable(new FStructOnScope(FRigVMMirrorSettings::StaticStruct(), (uint8*)&Settings));
-#if WITH_RIGVMLEGACYEDITOR
+
 				TSharedRef<SKismetInspector> KismetInspector = SNew(SKismetInspector);
-#else
-				TSharedRef<SRigVMDetailsInspector> KismetInspector = SNew(SRigVMDetailsInspector);
-#endif
 				KismetInspector->ShowSingleStruct(StructToDisplay);
 
 				TSharedRef<SCustomDialog> MirrorDialog = SNew(SCustomDialog)
diff --git a/Engine/Plugins/Animation/ControlRig/Source/ControlRigEditor/Private/Editor/ControlRigEditor.cpp b/Engine/Plugins/Animation/ControlRig/Source/ControlRigEditor/Private/Editor/ControlRigEditor.cpp
--- a/Engine/Plugins/Animation/ControlRig/Source/ControlRigEditor/Private/Editor/ControlRigEditor.cpp
+++ b/Engine/Plugins/Animation/ControlRig/Source/ControlRigEditor/Private/Editor/ControlRigEditor.cpp
@@ -6,6 +6,7 @@
 #include "ControlRigBlueprint.h"
 #include "SBlueprintEditorToolbar.h"
 #include "Editor/ControlRigEditorMode.h"
+#include "SKismetInspector.h"
 #include "SEnumCombo.h"
 #include "SceneView.h"
 #include "Widgets/Layout/SScrollBox.h"
@@ -123,12 +124,6 @@
 #include "ITransportControl.h"
 #include "PropertyPath.h"
 
-#if WITH_RIGVMLEGACYEDITOR
-#include "SKismetInspector.h"
-#else
-#include "Editor/SRigVMDetailsInspector.h"
-#endif
-
 #define LOCTEXT_NAMESPACE "ControlRigEditor"
 
 namespace UE::ControlRigEditor::Private
@@ -161,10 +156,10 @@ UControlRigBlueprint* GetControlRigBlueprintFromContext(const UUnrealEdViewportT
 			{
 				if (AnimViewportClient->GetAssetEditorToolkit()->IsBlueprintEditor())
 				{
-					const TSharedRef<FControlRigEditor>& Editor =
-						StaticCastSharedRef<FControlRigEditor>(AnimViewportClient->GetAssetEditorToolkit());
+					const TSharedRef<FBlueprintEditor>& BlueprintEditor =
+						StaticCastSharedRef<FBlueprintEditor>(AnimViewportClient->GetAssetEditorToolkit());
 
-					return Cast<UControlRigBlueprint>(Editor->FRigVMEditorBase::GetBlueprintObj());
+					return Cast<UControlRigBlueprint>(BlueprintEditor->GetBlueprintObj());
 				}
 			}
 		}
@@ -324,7 +319,7 @@ UObject* FControlRigEditor::GetOuterForHost() const
 	{
 		return EditorSkelComp;
 	}
-	return FRigVMEditorBase::GetOuterForHost();
+	return FRigVMEditor::GetOuterForHost();
 }
 
 UClass* FControlRigEditor::GetDetailWrapperClass() const
@@ -332,7 +327,6 @@ UClass* FControlRigEditor::GetDetailWrapperClass() const
 	return UControlRigWrapperObject::StaticClass();
 }
 
-#if WITH_RIGVMLEGACYEDITOR
 bool FControlRigEditor::IsSectionVisible(NodeSectionID::Type InSectionID) const
 {
 	if(!IControlRigEditor::IsSectionVisible(InSectionID))
@@ -364,42 +358,7 @@ bool FControlRigEditor::IsSectionVisible(NodeSectionID::Type InSectionID) const
 	return true;
 }
 
-#else
-bool FControlRigEditor::IsSectionVisible(RigVMNodeSectionID::Type InSectionID) const
-{
-	if(!IControlRigEditor::IsSectionVisible(InSectionID))
-    {
-    	return false;
-    }
-
-	if(const UControlRigBlueprint* RigBlueprint = GetControlRigBlueprint())
-	{
-		if(IsModularRig())
-		{
-			switch (InSectionID)
-			{
-				case RigVMNodeSectionID::GRAPH:
-				{
-					return RigBlueprint->SupportsEventGraphs();
-				}
-				case RigVMNodeSectionID::FUNCTION:
-				{
-					return RigBlueprint->SupportsFunctions();
-				}
-				default:
-				{
-					break;
-				}
-			}
-		}
-	}
-	return true;
-}
-#endif
-
-
-#if WITH_RIGVMLEGACYEDITOR
-bool FControlRigEditor::NewDocument_IsVisibleForType(FBlueprintEditor::ECreatedDocumentType GraphType) const
+bool FControlRigEditor::NewDocument_IsVisibleForType(ECreatedDocumentType GraphType) const
 {
 	if(!IControlRigEditor::NewDocument_IsVisibleForType(GraphType))
 	{
@@ -412,11 +371,11 @@ bool FControlRigEditor::NewDocument_IsVisibleForType(FBlueprintEditor::ECreatedD
 		{
 			switch(GraphType)
 			{
-				case FBlueprintEditor::CGT_NewEventGraph:
+				case CGT_NewEventGraph:
 				{
 					return RigBlueprint->SupportsEventGraphs();
 				}
-				case FBlueprintEditor::CGT_NewFunctionGraph:
+				case CGT_NewFunctionGraph:
 				{
 					return RigBlueprint->SupportsFunctions();
 				}
@@ -430,40 +389,6 @@ bool FControlRigEditor::NewDocument_IsVisibleForType(FBlueprintEditor::ECreatedD
 
 	return true;
 }
-#else
-bool FControlRigEditor::NewDocument_IsVisibleForType(FRigVMEditorBase::ECreatedDocumentType GraphType) const
-{
-	if(!IControlRigEditor::NewDocument_IsVisibleForType(GraphType))
-	{
-		return false;
-	}
-
-	if(const UControlRigBlueprint* RigBlueprint = GetControlRigBlueprint())
-	{
-		if(IsModularRig())
-		{
-			switch(GraphType)
-			{
-				case FRigVMEditorBase::CGT_NewEventGraph:
-				{
-					return RigBlueprint->SupportsEventGraphs();
-				}
-				case FRigVMEditorBase::CGT_NewFunctionGraph:
-				{
-					return RigBlueprint->SupportsFunctions();
-				}
-				default:
-				{
-					break;
-				}
-			}
-		}
-	}
-
-	return true;
-}
-#endif
-
 
 FReply FControlRigEditor::OnViewportDrop(const FGeometry& MyGeometry, const FDragDropEvent& DragDropEvent)
 {
@@ -549,7 +474,7 @@ void FControlRigEditor::InitRigVMEditor(const EToolkitMode::Type Mode, const TSh
 {
 	DECLARE_SCOPE_HIERARCHICAL_COUNTER_FUNC()
 
-	IControlRigEditor::InitRigVMEditor(Mode, InitToolkitHost, InRigVMBlueprint);
+	FRigVMEditor::InitRigVMEditor(Mode, InitToolkitHost, InRigVMBlueprint);
 
 	UControlRigBlueprint* ControlRigBlueprint = CastChecked<UControlRigBlueprint>(InRigVMBlueprint);
 
@@ -925,10 +850,8 @@ bool FControlRigEditor::RecordReplay(double InRecordingDuration)
 				}
 				else
 				{
-#if WITH_RIGVMLEGACYEDITOR
 					static constexpr TCHAR PrintStringFormat[] = TEXT("Recording starts in... %d");
 					UKismetSystemLibrary::PrintString(GetPreviewScene()->GetWorld(), FString::Printf(PrintStringFormat, SecondsLeft), true, false, FLinearColor::Red, 1.f);
-#endif
 				}
 			}), 1.f, true, 0.01f);
 		}
@@ -972,7 +895,7 @@ EControlRigReplayPlaybackMode FControlRigEditor::GetReplayPlaybackMode() const
 
 void FControlRigEditor::FillToolbar(FToolBarBuilder& ToolbarBuilder, bool bEndSection)
 {
-	IControlRigEditor::FillToolbar(ToolbarBuilder, false);
+	FRigVMEditor::FillToolbar(ToolbarBuilder, false);
 	
 	{
 		if(CVarControlRigHierarchyEnableModules.GetValueOnAnyThread())
@@ -1216,7 +1139,7 @@ void FControlRigEditor::SetEventQueue(TArray<FName> InEventQueue, bool bCompile)
 	}
 
 	TArray<FRigElementKey> PreviousSelection;
-	if (UControlRigBlueprint* RigBlueprint = Cast<UControlRigBlueprint>(FRigVMEditorBase::GetBlueprintObj()))
+	if (UControlRigBlueprint* RigBlueprint = Cast<UControlRigBlueprint>(GetBlueprintObj()))
 	{
 		if(bCompile)
 		{
@@ -1239,13 +1162,13 @@ void FControlRigEditor::SetEventQueue(TArray<FName> InEventQueue, bool bCompile)
 		}
 	}
 
-	IControlRigEditor::SetEventQueue(InEventQueue, bCompile);
+	FRigVMEditor::SetEventQueue(InEventQueue, bCompile);
 
 	if (UControlRig* ControlRig = GetControlRig())
 	{
 		if (InEventQueue.Num() > 0)
 		{
-			if (UControlRigBlueprint* RigBlueprint = Cast<UControlRigBlueprint>(FRigVMEditorBase::GetBlueprintObj()))
+			if (UControlRigBlueprint* RigBlueprint = Cast<UControlRigBlueprint>(GetBlueprintObj()))
 			{
 				RigBlueprint->Validator->SetControlRig(ControlRig);
 
@@ -1302,7 +1225,7 @@ int32 FControlRigEditor::GetEventQueueComboValue() const
 	{
 		return 3;
 	}
-	return IControlRigEditor::GetEventQueueComboValue();
+	return FRigVMEditor::GetEventQueueComboValue();
 }
 
 FText FControlRigEditor::GetEventQueueLabel() const
@@ -1365,10 +1288,10 @@ FSlateIcon FControlRigEditor::GetEventQueueIcon(const TArray<FName>& InEventQueu
 
 void FControlRigEditor::HandleSetObjectBeingDebugged(UObject* InObject)
 {
-	FRigVMEditorBase::HandleSetObjectBeingDebugged(InObject);
+	FRigVMEditor::HandleSetObjectBeingDebugged(InObject);
 	
 	UControlRig* DebuggedControlRig = Cast<UControlRig>(InObject);
-	if(UControlRig* PreviouslyDebuggedControlRig = Cast<UControlRig>(FRigVMEditorBase::GetBlueprintObj()->GetObjectBeingDebugged()))
+	if(UControlRig* PreviouslyDebuggedControlRig = Cast<UControlRig>(GetBlueprintObj()->GetObjectBeingDebugged()))
 	{
 		if(!URigVMHost::IsGarbageOrDestroyed(PreviouslyDebuggedControlRig))
 		{
@@ -1381,7 +1304,7 @@ void FControlRigEditor::HandleSetObjectBeingDebugged(UObject* InObject)
 		}
 	}
 
-	if (UControlRigBlueprint* RigBlueprint = Cast<UControlRigBlueprint>(FRigVMEditorBase::GetBlueprintObj()))
+	if (UControlRigBlueprint* RigBlueprint = Cast<UControlRigBlueprint>(GetBlueprintObj()))
 	{
 		RigBlueprint->Validator->SetControlRig(DebuggedControlRig);
 	}
@@ -1421,7 +1344,7 @@ void FControlRigEditor::HandleSetObjectBeingDebugged(UObject* InObject)
 			if(bShouldExecute)
 			{
 				DebuggedControlRig->SetBoneInitialTransformsFromSkeletalMeshComponent(EditorSkelComp);
-				if(UControlRigBlueprint* RigBlueprint = Cast<UControlRigBlueprint>(FRigVMEditorBase::GetBlueprintObj()))
+				if(UControlRigBlueprint* RigBlueprint = Cast<UControlRigBlueprint>(GetBlueprintObj()))
 				{
 					// copy the initial transforms back to the blueprint
 					// no need to call modify here since this code only modifies the bp if the preview mesh changed
@@ -1505,7 +1428,7 @@ void FControlRigEditor::SetDetailViewForRigElements(const TArray<FRigElementKey>
 
 	ClearDetailObject();
 
-	UControlRigBlueprint* RigBlueprint = Cast<UControlRigBlueprint>(FRigVMEditorBase::GetBlueprintObj());
+	UControlRigBlueprint* RigBlueprint = Cast<UControlRigBlueprint>(GetBlueprintObj());
 	URigHierarchy* HierarchyBeingDebugged = GetHierarchyBeingDebugged();
 	TArray<UObject*> Objects;
 
@@ -1517,7 +1440,7 @@ void FControlRigEditor::SetDetailViewForRigElements(const TArray<FRigElementKey>
 			continue;
 		}
 
-		URigVMDetailsViewWrapperObject* WrapperObject = URigVMDetailsViewWrapperObject::MakeInstance(GetDetailWrapperClass(), FRigVMEditorBase::GetBlueprintObj(), Element->GetScriptStruct(), (uint8*)Element, HierarchyBeingDebugged);
+		URigVMDetailsViewWrapperObject* WrapperObject = URigVMDetailsViewWrapperObject::MakeInstance(GetDetailWrapperClass(), GetBlueprintObj(), Element->GetScriptStruct(), (uint8*)Element, HierarchyBeingDebugged);
 		WrapperObject->GetWrappedPropertyChangedChainEvent().AddSP(this, &FControlRigEditor::OnWrappedPropertyChangedChainEvent);
 		WrapperObject->AddToRoot();
 
@@ -1554,7 +1477,7 @@ void FControlRigEditor::RefreshDetailView()
 		return;
 	}
 
-	IControlRigEditor::RefreshDetailView();
+	FRigVMEditor::RefreshDetailView();
 }
 
 bool FControlRigEditor::DetailViewShowsAnyRigElement() const
@@ -1564,7 +1487,7 @@ bool FControlRigEditor::DetailViewShowsAnyRigElement() const
 
 bool FControlRigEditor::DetailViewShowsRigElement(FRigElementKey InKey) const
 {
-	TArray< TWeakObjectPtr<UObject> > SelectedObjects = GetSelectedObjectsFromDetailView();
+	TArray< TWeakObjectPtr<UObject> > SelectedObjects = Inspector->GetSelectedObjects();
 	for (TWeakObjectPtr<UObject> SelectedObject : SelectedObjects)
 	{
 		if (SelectedObject.IsValid())
@@ -1591,7 +1514,7 @@ TArray<FRigElementKey> FControlRigEditor::GetSelectedRigElementsFromDetailView()
 {
 	TArray<FRigElementKey> Elements;
 	
-	TArray< TWeakObjectPtr<UObject> > SelectedObjects = GetSelectedObjectsFromDetailView();
+	TArray< TWeakObjectPtr<UObject> > SelectedObjects = Inspector->GetSelectedObjects();
 	for (TWeakObjectPtr<UObject> SelectedObject : SelectedObjects)
 	{
 		if (SelectedObject.IsValid())
@@ -1612,15 +1535,6 @@ TArray<FRigElementKey> FControlRigEditor::GetSelectedRigElementsFromDetailView()
 	return Elements;
 }
 
-TArray<TWeakObjectPtr<UObject>> FControlRigEditor::GetSelectedObjectsFromDetailView() const
-{
-	if (Inspector.IsValid())
-	{
-		return Inspector->GetSelectedObjects();
-	}
-	return TArray<TWeakObjectPtr<UObject>>();
-}
-
 void FControlRigEditor::SetDetailViewForRigModules()
 {
 	SetDetailViewForRigModules(ModulesSelected);
@@ -1635,7 +1549,7 @@ void FControlRigEditor::SetDetailViewForRigModules(const TArray<FString> InKeys)
 
 	ClearDetailObject();
 
-	UControlRigBlueprint* RigBlueprint = Cast<UControlRigBlueprint>(FRigVMEditorBase::GetBlueprintObj());
+	UControlRigBlueprint* RigBlueprint = Cast<UControlRigBlueprint>(GetBlueprintObj());
 	UModularRig* RigBeingDebugged = Cast<UModularRig>(RigBlueprint->GetDebuggedControlRig());
 
 	if (!RigBeingDebugged)
@@ -1676,7 +1590,7 @@ bool FControlRigEditor::DetailViewShowsAnyRigModule() const
 
 bool FControlRigEditor::DetailViewShowsRigModule(FString InKey) const
 {
-	TArray< TWeakObjectPtr<UObject> > SelectedObjects = GetSelectedObjectsFromDetailView();
+	TArray< TWeakObjectPtr<UObject> > SelectedObjects = Inspector->GetSelectedObjects();
 	for (TWeakObjectPtr<UObject> SelectedObject : SelectedObjects)
 	{
 		if (SelectedObject.IsValid())
@@ -1716,7 +1630,7 @@ void FControlRigEditor::Compile()
 			return;
 		}
 
-		TArray< TWeakObjectPtr<UObject> > SelectedObjects = GetSelectedObjectsFromDetailView();
+		TArray< TWeakObjectPtr<UObject> > SelectedObjects = Inspector->GetSelectedObjects();
 
 		if(IsConstructionModeEnabled())
 		{
@@ -1734,7 +1648,7 @@ void FControlRigEditor::Compile()
 		ResetAllBoneModification(); 
 
 		{
-			FRigVMEditorBase::Compile();
+			FRigVMEditor::Compile();
 		}
 
 		ControlRigBlueprint->RecompileModularRig();
@@ -1840,7 +1754,7 @@ void FControlRigEditor::HandleModifiedEvent(ERigVMGraphNotifType InNotifType, UR
 
 void FControlRigEditor::OnCreateGraphEditorCommands(TSharedPtr<FUICommandList> GraphEditorCommandsList)
 {
-	FRigVMEditorBase::OnCreateGraphEditorCommands(GraphEditorCommandsList);
+	IControlRigEditor::OnCreateGraphEditorCommands(GraphEditorCommandsList);
 
 	GraphEditorCommandsList->MapAction(
 		FControlRigEditorCommands::Get().RequestDirectManipulationPosition,
@@ -1874,10 +1788,10 @@ void FControlRigEditor::HandleVMCompiledEvent(UObject* InCompiledObject, URigVM*
 
 void FControlRigEditor::SaveAsset_Execute()
 {
-	FRigVMEditorBase::SaveAsset_Execute();
+	FRigVMEditor::SaveAsset_Execute();
 
 	// Save the new state of the hierarchy in the default object, so that it has the correct values on load
-	UControlRigBlueprint* RigBlueprint = Cast<UControlRigBlueprint>(FRigVMEditorBase::GetBlueprintObj());
+	UControlRigBlueprint* RigBlueprint = Cast<UControlRigBlueprint>(GetBlueprintObj());
 	if(const UControlRig* ControlRig = GetControlRig())
 	{
 		UControlRig* CDO = ControlRig->GetClass()->GetDefaultObject<UControlRig>();
@@ -1892,10 +1806,10 @@ void FControlRigEditor::SaveAsset_Execute()
 
 void FControlRigEditor::SaveAssetAs_Execute()
 {
-	FRigVMEditorBase::SaveAssetAs_Execute();
+	FRigVMEditor::SaveAssetAs_Execute();
 
 	// Save the new state of the hierarchy in the default object, so that it has the correct values on load
-	UControlRigBlueprint* RigBlueprint = Cast<UControlRigBlueprint>(FRigVMEditorBase::GetBlueprintObj());
+	UControlRigBlueprint* RigBlueprint = Cast<UControlRigBlueprint>(GetBlueprintObj());
 	if(const UControlRig* ControlRig = GetControlRig())
 	{
 		UControlRig* CDO = ControlRig->GetClass()->GetDefaultObject<UControlRig>();
@@ -1910,7 +1824,7 @@ void FControlRigEditor::SaveAssetAs_Execute()
 
 bool FControlRigEditor::IsModularRig() const
 {
-	if(UControlRigBlueprint* RigBlueprint = Cast<UControlRigBlueprint>(FRigVMEditorBase::GetBlueprintObj()))
+	if(UControlRigBlueprint* RigBlueprint = Cast<UControlRigBlueprint>(GetBlueprintObj()))
 	{
 		return RigBlueprint->IsModularRig();
 	}
@@ -1919,7 +1833,7 @@ bool FControlRigEditor::IsModularRig() const
 
 bool FControlRigEditor::IsRigModule() const
 {
-	if(UControlRigBlueprint* RigBlueprint = Cast<UControlRigBlueprint>(FRigVMEditorBase::GetBlueprintObj()))
+	if(UControlRigBlueprint* RigBlueprint = Cast<UControlRigBlueprint>(GetBlueprintObj()))
 	{
 		return RigBlueprint->IsControlRigModule();
 	}
@@ -1943,7 +1857,7 @@ FString FControlRigEditor::GetWorldCentricTabPrefix() const
 
 FReply FControlRigEditor::OnSpawnGraphNodeByShortcut(FInputChord InChord, const FVector2D& InPosition, UEdGraph* InGraph)
 {
-	const FReply SuperReply = FRigVMEditorBase::OnSpawnGraphNodeByShortcut(InChord, InPosition, InGraph);
+	const FReply SuperReply = FRigVMEditor::OnSpawnGraphNodeByShortcut(InChord, InPosition, InGraph);
 	if(SuperReply.IsEventHandled())
 	{
 		return SuperReply;
@@ -1990,7 +1904,7 @@ void FControlRigEditor::PostTransaction(bool bSuccess, const FTransaction* Trans
 {
 	EnsureValidRigElementsInDetailPanel();
 
-	if (UControlRigBlueprint* RigBlueprint = Cast<UControlRigBlueprint>(FRigVMEditorBase::GetBlueprintObj()))
+	if (UControlRigBlueprint* RigBlueprint = Cast<UControlRigBlueprint>(GetBlueprintObj()))
 	{
 		// Do not compile here. ControlRigBlueprint::PostTransacted decides when it is necessary to compile depending
 		// on the properties that are affected.
@@ -2028,7 +1942,7 @@ void FControlRigEditor::EnsureValidRigElementsInDetailPanel()
 	UControlRigBlueprint* ControlRigBP = GetControlRigBlueprint();
 	URigHierarchy* Hierarchy = ControlRigBP->Hierarchy; 
 
-	TArray< TWeakObjectPtr<UObject> > SelectedObjects = GetSelectedObjectsFromDetailView();
+	TArray< TWeakObjectPtr<UObject> > SelectedObjects = Inspector->GetSelectedObjects();
 	for (TWeakObjectPtr<UObject> SelectedObject : SelectedObjects)
 	{
 		if (SelectedObject.IsValid())
@@ -2072,13 +1986,13 @@ void FControlRigEditor::OnAnimInitialized()
 
 void FControlRigEditor::HandleVMExecutedEvent(URigVMHost* InHost, const FName& InEventName)
 {
-	IControlRigEditor::HandleVMExecutedEvent(InHost, InEventName);
+	FRigVMEditor::HandleVMExecutedEvent(InHost, InEventName);
 
 	if (UControlRigBlueprint* ControlRigBP = GetControlRigBlueprint())
 	{
 		URigHierarchy* Hierarchy = GetHierarchyBeingDebugged(); 
 
-		TArray< TWeakObjectPtr<UObject> > SelectedObjects = GetSelectedObjectsFromDetailView();
+		TArray< TWeakObjectPtr<UObject> > SelectedObjects = Inspector->GetSelectedObjects();
 		for (TWeakObjectPtr<UObject> SelectedObject : SelectedObjects)
 		{
 			if (SelectedObject.IsValid())
@@ -2187,7 +2101,7 @@ void FControlRigEditor::CreateEditorModeManager()
 
 void FControlRigEditor::Tick(float DeltaTime)
 {
-	IControlRigEditor::Tick(DeltaTime);
+	FRigVMEditor::Tick(DeltaTime);
 
 	bool bDrawHierarchyBones = false;
 
@@ -2251,7 +2165,7 @@ void FControlRigEditor::HandleViewportCreated(const TSharedRef<class IPersonaVie
 	// TODO: this is duplicated code from FAnimBlueprintEditor, would be nice to consolidate. 
 	auto GetCompilationStateText = [this]()
 	{
-		if (UBlueprint* Blueprint = FRigVMEditorBase::GetBlueprintObj())
+		if (UBlueprint* Blueprint = GetBlueprintObj())
 		{
 			switch (Blueprint->Status)
 			{
@@ -2300,7 +2214,7 @@ void FControlRigEditor::HandleViewportCreated(const TSharedRef<class IPersonaVie
 
 	auto CompileBlueprint = [this]()
 	{
-		if (UBlueprint* Blueprint = FRigVMEditorBase::GetBlueprintObj())
+		if (UBlueprint* Blueprint = GetBlueprintObj())
 		{
 			if (!Blueprint->IsUpToDate())
 			{
@@ -2313,7 +2227,7 @@ void FControlRigEditor::HandleViewportCreated(const TSharedRef<class IPersonaVie
 
 	auto GetErrorSeverity = [this]()
 	{
-		if (UBlueprint* Blueprint = FRigVMEditorBase::GetBlueprintObj())
+		if (UBlueprint* Blueprint = GetBlueprintObj())
 		{
 			return (Blueprint->Status == BS_Error) ? EMessageSeverity::Error : EMessageSeverity::Warning;
 		}
@@ -2323,7 +2237,7 @@ void FControlRigEditor::HandleViewportCreated(const TSharedRef<class IPersonaVie
 
 	auto GetIcon = [this]()
 	{
-		if (UBlueprint* Blueprint = FRigVMEditorBase::GetBlueprintObj())
+		if (UBlueprint* Blueprint = GetBlueprintObj())
 		{
 			return (Blueprint->Status == BS_Error) ? FEditorFontGlyphs::Exclamation_Triangle : FEditorFontGlyphs::Eye;
 		}
@@ -2810,7 +2724,7 @@ void FControlRigEditor::HandleViewportCreated(const TSharedRef<class IPersonaVie
 	{
 		AnimViewportClient->OnSetBoneSize.BindLambda([this](float InBoneSize)
 		{
-			if (UControlRigBlueprint* RigBlueprint = Cast<UControlRigBlueprint>(FRigVMEditorBase::GetBlueprintObj()))
+			if (UControlRigBlueprint* RigBlueprint = Cast<UControlRigBlueprint>(GetBlueprintObj()))
 			{
 				RigBlueprint->Modify();
 				RigBlueprint->DebugBoneRadius = InBoneSize;
@@ -2819,7 +2733,7 @@ void FControlRigEditor::HandleViewportCreated(const TSharedRef<class IPersonaVie
 		
 		AnimViewportClient->OnGetBoneSize.BindLambda([this]() -> float
 		{
-			if (UControlRigBlueprint* RigBlueprint = Cast<UControlRigBlueprint>(FRigVMEditorBase::GetBlueprintObj()))
+			if (UControlRigBlueprint* RigBlueprint = Cast<UControlRigBlueprint>(GetBlueprintObj()))
 			{
 				return RigBlueprint->DebugBoneRadius;
 			}
@@ -3016,7 +2930,7 @@ void FControlRigEditor::HandlePreviewSceneCreated(const TSharedRef<IPersonaPrevi
 	if(FWorldContext* WorldContext = GEngine->GetWorldContextFromWorld(InPersonaPreviewScene->GetWorld()))
 	{
 		static constexpr TCHAR Format[] = TEXT("ControlRigEditor (%s)");
-		WorldContext->CustomDescription = FString::Printf(Format, *FRigVMEditorBase::GetBlueprintObj()->GetName());
+		WorldContext->CustomDescription = FString::Printf(Format, *GetBlueprintObj()->GetName());
 	}
 
 	// create ground mesh actor
@@ -3071,9 +2985,9 @@ void FControlRigEditor::UpdateRigVMHost()
 {
 	DECLARE_SCOPE_HIERARCHICAL_COUNTER_FUNC()
 
-	IControlRigEditor::UpdateRigVMHost();
+	FRigVMEditor::UpdateRigVMHost();
 
-	UControlRigBlueprint* Blueprint = Cast<UControlRigBlueprint>(FRigVMEditorBase::GetBlueprintObj());
+	UControlRigBlueprint* Blueprint = Cast<UControlRigBlueprint>(GetBlueprintObj());
 	if(UClass* Class = Blueprint->GeneratedClass)
 	{
 		UControlRigSkeletalMeshComponent* EditorSkelComp = Cast<UControlRigSkeletalMeshComponent>(GetPersonaToolkit()->GetPreviewScene()->GetPreviewMeshComponent());
@@ -3140,7 +3054,7 @@ void FControlRigEditor::CacheNameLists()
 {
 	DECLARE_SCOPE_HIERARCHICAL_COUNTER_FUNC()
 
-	IControlRigEditor::CacheNameLists();
+	FRigVMEditor::CacheNameLists();
 
 	if (UControlRigBlueprint* ControlRigBP = GetControlRigBlueprint())
 	{
@@ -3435,7 +3349,7 @@ void FControlRigEditor::GenerateEventQueueMenuContent(FMenuBuilder& MenuBuilder)
     MenuBuilder.AddMenuEntry(FControlRigEditorCommands::Get().BackwardsAndForwardsSolveEvent, TEXT("BackwardsAndForwards"), TAttribute<FText>(), TAttribute<FText>(), GetEventQueueIcon(BackwardsAndForwardsSolveEventQueue));
     MenuBuilder.EndSection();
 
-    if (const UControlRigBlueprint* RigBlueprint = Cast<UControlRigBlueprint>(FRigVMEditorBase::GetBlueprintObj()))
+    if (const UControlRigBlueprint* RigBlueprint = Cast<UControlRigBlueprint>(GetBlueprintObj()))
     {
     	URigVMEdGraphSchema* Schema = CastChecked<URigVMEdGraphSchema>(RigBlueprint->GetRigVMEdGraphSchemaClass()->GetDefaultObject());
     	
@@ -3465,11 +3379,7 @@ void FControlRigEditor::GenerateEventQueueMenuContent(FMenuBuilder& MenuBuilder)
     			FText::FromString(FString::Printf(TEXT("Runs the user defined %s"), *EventNameStr)),
     			GetEventQueueIcon({EntryName}),
     			FUIAction(
-#if WITH_RIGVMLEGACYEDITOR
-    				FExecuteAction::CreateSP(this, &FRigVMLegacyEditor::SetEventQueue, TArray<FName>({EntryName})),
-#else
-    				FExecuteAction::CreateSP(this, &FRigVMNewEditor::SetEventQueue, TArray<FName>({EntryName})),
-#endif
+    				FExecuteAction::CreateSP(this, &FRigVMEditor::SetEventQueue, TArray<FName>({EntryName})),
     				FCanExecuteAction()
     			)
     		);
@@ -3531,7 +3441,7 @@ FTransform FControlRigEditor::GetRigElementTransform(const FRigElementKey& InEle
 
 	if(bOnDebugInstance)
 	{
-		UControlRig* DebuggedControlRig = Cast<UControlRig>(FRigVMEditorBase::GetBlueprintObj()->GetObjectBeingDebugged());
+		UControlRig* DebuggedControlRig = Cast<UControlRig>(GetBlueprintObj()->GetObjectBeingDebugged());
 		if (DebuggedControlRig == nullptr)
 		{
 			DebuggedControlRig = GetControlRig();
@@ -3641,7 +3551,7 @@ void FControlRigEditor::SetRigElementTransform(const FRigElementKey& InElement,
 
 void FControlRigEditor::OnFinishedChangingProperties(const FPropertyChangedEvent& PropertyChangedEvent)
 {
-	FRigVMEditorBase::OnFinishedChangingProperties(PropertyChangedEvent);
+	FRigVMEditor::OnFinishedChangingProperties(PropertyChangedEvent);
 	
 	DECLARE_SCOPE_HIERARCHICAL_COUNTER_FUNC()
 
@@ -3668,7 +3578,7 @@ void FControlRigEditor::OnFinishedChangingProperties(const FPropertyChangedEvent
 
 void FControlRigEditor::OnWrappedPropertyChangedChainEvent(URigVMDetailsViewWrapperObject* InWrapperObject, const FString& InPropertyPath, FPropertyChangedChainEvent& InPropertyChangedChainEvent)
 {
-	IControlRigEditor::OnWrappedPropertyChangedChainEvent(InWrapperObject, InPropertyPath, InPropertyChangedChainEvent);
+	FRigVMEditor::OnWrappedPropertyChangedChainEvent(InWrapperObject, InPropertyPath, InPropertyChangedChainEvent);
 	
 	check(InWrapperObject);
 	check(!GetWrapperObjects().IsEmpty());
@@ -3911,7 +3821,7 @@ EVisibility FControlRigEditor::GetDirectManipulationVisibility() const
 
 FText FControlRigEditor::GetDirectionManipulationText() const
 {
-	if (UControlRig* DebuggedControlRig = Cast<UControlRig>(FRigVMEditorBase::GetBlueprintObj()->GetObjectBeingDebugged()))
+	if (UControlRig* DebuggedControlRig = Cast<UControlRig>(GetBlueprintObj()->GetObjectBeingDebugged()))
 	{
 		if (URigHierarchy* Hierarchy = DebuggedControlRig->GetHierarchy())
 		{
@@ -3943,7 +3853,7 @@ void FControlRigEditor::OnDirectManipulationChanged(TSharedPtr<FString> NewValue
 		return;
 	}
 	
-	UControlRigBlueprint* ControlRigBlueprint = CastChecked<UControlRigBlueprint>(FRigVMEditorBase::GetBlueprintObj());
+	UControlRigBlueprint* ControlRigBlueprint = CastChecked<UControlRigBlueprint>(GetBlueprintObj());
 	if(ControlRigBlueprint == nullptr)
 	{
 		return;
@@ -3976,7 +3886,7 @@ const TArray<FRigDirectManipulationTarget> FControlRigEditor::GetDirectManipulat
 {
 	if(DirectManipulationSubject.IsValid())
 	{
-		if (UControlRig* DebuggedControlRig = Cast<UControlRig>(FRigVMEditorBase::GetBlueprintObj()->GetObjectBeingDebugged()))
+		if (UControlRig* DebuggedControlRig = Cast<UControlRig>(GetBlueprintObj()->GetObjectBeingDebugged()))
 		{
 			if(const URigVMUnitNode* Node = DirectManipulationSubject.Get())
 			{
@@ -4070,42 +3980,26 @@ FReply FControlRigEditor::OnNavigateToConnectorWarning() const
 
 void FControlRigEditor::BindCommands()
 {
-	IControlRigEditor::BindCommands();
+	FRigVMEditor::BindCommands();
 
 	GetToolkitCommands()->MapAction(
 		FControlRigEditorCommands::Get().ConstructionEvent,
-#if WITH_RIGVMLEGACYEDITOR
-		FExecuteAction::CreateSP(this, &FRigVMLegacyEditor::SetEventQueue, TArray<FName>(ConstructionEventQueue)),
-#else
-		FExecuteAction::CreateSP(this, &FRigVMNewEditor::SetEventQueue, TArray<FName>(ConstructionEventQueue)),
-#endif
+		FExecuteAction::CreateSP(this, &FRigVMEditor::SetEventQueue, TArray<FName>(ConstructionEventQueue)),
 		FCanExecuteAction());
 
 	GetToolkitCommands()->MapAction(
 		FControlRigEditorCommands::Get().ForwardsSolveEvent,
-#if WITH_RIGVMLEGACYEDITOR
-		FExecuteAction::CreateSP(this, &FRigVMLegacyEditor::SetEventQueue, TArray<FName>(ForwardsSolveEventQueue)),
-#else
-		FExecuteAction::CreateSP(this, &FRigVMNewEditor::SetEventQueue, TArray<FName>(ForwardsSolveEventQueue)),
-#endif
+		FExecuteAction::CreateSP(this, &FRigVMEditor::SetEventQueue, TArray<FName>(ForwardsSolveEventQueue)),
 		FCanExecuteAction());
 
 	GetToolkitCommands()->MapAction(
 		FControlRigEditorCommands::Get().BackwardsSolveEvent,
-#if WITH_RIGVMLEGACYEDITOR
-		FExecuteAction::CreateSP(this, &FRigVMLegacyEditor::SetEventQueue, TArray<FName>(BackwardsSolveEventQueue)),
-#else
-		FExecuteAction::CreateSP(this, &FRigVMNewEditor::SetEventQueue, TArray<FName>(BackwardsSolveEventQueue)),
-#endif
+		FExecuteAction::CreateSP(this, &FRigVMEditor::SetEventQueue, TArray<FName>(BackwardsSolveEventQueue)),
 		FCanExecuteAction());
 
 	GetToolkitCommands()->MapAction(
 		FControlRigEditorCommands::Get().BackwardsAndForwardsSolveEvent,
-#if WITH_RIGVMLEGACYEDITOR
-		FExecuteAction::CreateSP(this, &FRigVMLegacyEditor::SetEventQueue, TArray<FName>(BackwardsAndForwardsSolveEventQueue)),
-#else
-		FExecuteAction::CreateSP(this, &FRigVMNewEditor::SetEventQueue, TArray<FName>(BackwardsAndForwardsSolveEventQueue)),
-#endif
+		FExecuteAction::CreateSP(this, &FRigVMEditor::SetEventQueue, TArray<FName>(BackwardsAndForwardsSolveEventQueue)),
 		FCanExecuteAction());
 
 	GetToolkitCommands()->MapAction(
@@ -4216,7 +4110,7 @@ void FControlRigEditor::OnHierarchyChanged()
 
 void FControlRigEditor::OnHierarchyModified(ERigHierarchyNotification InNotif, URigHierarchy* InHierarchy, const FRigNotificationSubject& InSubject)
 {
-	UControlRigBlueprint* RigBlueprint = Cast<UControlRigBlueprint>(FRigVMEditorBase::GetBlueprintObj());
+	UControlRigBlueprint* RigBlueprint = Cast<UControlRigBlueprint>(GetBlueprintObj());
 	if(RigBlueprint == nullptr)
 	{
 		return;
@@ -4514,7 +4408,7 @@ void FControlRigEditor::OnHierarchyModified_AnyThread(ERigHierarchyNotification
 						else
 						{
 							// only clear the details if we are not looking at a transient control
-							if (UControlRig* DebuggedControlRig = Cast<UControlRig>(FRigVMEditorBase::GetBlueprintObj()->GetObjectBeingDebugged()))
+							if (UControlRig* DebuggedControlRig = Cast<UControlRig>(GetBlueprintObj()->GetObjectBeingDebugged()))
 							{
 								if(DebuggedControlRig->RigUnitManipulationInfos.IsEmpty())
 								{
@@ -4855,7 +4749,7 @@ void FControlRigEditor::UpdateBoneModification(FName BoneName, const FTransform&
 		}
 		
 		TMap<FName, FTransform>* TransformOverrideMap = &ControlRig->TransformOverrideForUserCreatedBones;
-		if (UControlRig* DebuggedControlRig = Cast<UControlRig>(FRigVMEditorBase::GetBlueprintObj()->GetObjectBeingDebugged()))
+		if (UControlRig* DebuggedControlRig = Cast<UControlRig>(GetBlueprintObj()->GetObjectBeingDebugged()))
 		{
 			TransformOverrideMap = &DebuggedControlRig->TransformOverrideForUserCreatedBones;
 		}
@@ -4877,7 +4771,7 @@ void FControlRigEditor::RemoveBoneModification(FName BoneName)
 		}
 
 		TMap<FName, FTransform>* TransformOverrideMap = &ControlRig->TransformOverrideForUserCreatedBones;
-		if (UControlRig* DebuggedControlRig = Cast<UControlRig>(FRigVMEditorBase::GetBlueprintObj()->GetObjectBeingDebugged()))
+		if (UControlRig* DebuggedControlRig = Cast<UControlRig>(GetBlueprintObj()->GetObjectBeingDebugged()))
 		{
 			TransformOverrideMap = &DebuggedControlRig->TransformOverrideForUserCreatedBones;
 		}
@@ -4896,7 +4790,7 @@ void FControlRigEditor::ResetAllBoneModification()
 		}
 
 		TMap<FName, FTransform>* TransformOverrideMap = &ControlRig->TransformOverrideForUserCreatedBones;
-		if (UControlRig* DebuggedControlRig = Cast<UControlRig>(FRigVMEditorBase::GetBlueprintObj()->GetObjectBeingDebugged()))
+		if (UControlRig* DebuggedControlRig = Cast<UControlRig>(GetBlueprintObj()->GetObjectBeingDebugged()))
 		{
 			TransformOverrideMap = &DebuggedControlRig->TransformOverrideForUserCreatedBones;
 		}
@@ -5280,7 +5174,7 @@ void FControlRigEditor::OnGraphNodeDropToPerform(TSharedPtr<FDragDropOperation>
 		return;
 	}
 
-	IControlRigEditor::OnGraphNodeDropToPerform(InDragDropOp, InGraph, InNodePosition, InScreenPosition);
+	FRigVMEditor::OnGraphNodeDropToPerform(InDragDropOp, InGraph, InNodePosition, InScreenPosition);
 }
 
 void FControlRigEditor::HandleMakeElementGetterSetter(ERigElementGetterSetterType Type, bool bIsGetter, TArray<FRigElementKey> Keys, UEdGraph* Graph, FVector2D NodePosition)
@@ -5796,13 +5690,13 @@ void FControlRigEditor::HandleMakeElementGetterSetter(ERigElementGetterSetterTyp
 
 void FControlRigEditor::HandleOnControlModified(UControlRig* Subject, FRigControlElement* ControlElement, const FRigControlModifiedContext& Context)
 {
-	UControlRig* DebuggedControlRig = Cast<UControlRig>(FRigVMEditorBase::GetBlueprintObj()->GetObjectBeingDebugged());
+	UControlRig* DebuggedControlRig = Cast<UControlRig>(GetBlueprintObj()->GetObjectBeingDebugged());
 	if (Subject != DebuggedControlRig)
 	{
 		return;
 	}
 
-	UControlRigBlueprint* Blueprint = Cast<UControlRigBlueprint>(FRigVMEditorBase::GetBlueprintObj());
+	UControlRigBlueprint* Blueprint = Cast<UControlRigBlueprint>(GetBlueprintObj());
 	if (Blueprint == nullptr)
 	{
 		return;
@@ -5953,7 +5847,7 @@ void FControlRigEditor::HandleOnControlModified(UControlRig* Subject, FRigContro
 void FControlRigEditor::HandleRefreshEditorFromBlueprint(URigVMBlueprint* InBlueprint)
 {
 	OnHierarchyChanged();
-	IControlRigEditor::HandleRefreshEditorFromBlueprint(InBlueprint);
+	FRigVMEditor::HandleRefreshEditorFromBlueprint(InBlueprint);
 }
 
 UToolMenu* FControlRigEditor::HandleOnGetViewportContextMenuDelegate()
diff --git a/Engine/Plugins/Animation/ControlRig/Source/ControlRigEditor/Private/Editor/ControlRigEditor.h b/Engine/Plugins/Animation/ControlRig/Source/ControlRigEditor/Private/Editor/ControlRigEditor.h
--- a/Engine/Plugins/Animation/ControlRig/Source/ControlRigEditor/Private/Editor/ControlRigEditor.h
+++ b/Engine/Plugins/Animation/ControlRig/Source/ControlRigEditor/Private/Editor/ControlRigEditor.h
@@ -77,7 +77,7 @@ public:
 
 	virtual UObject* GetOuterForHost() const override;
 
-	// FRigVMEditorBase interface
+	// FRigVMEditor interface
 	virtual UClass* GetDetailWrapperClass() const;
 	virtual void Compile() override;
 	virtual void HandleModifiedEvent(ERigVMGraphNotifType InNotifType, URigVMGraph* InGraph, UObject* InSubject) override;
@@ -110,13 +110,8 @@ public:
 
 	// BlueprintEditor interface
 	virtual FReply OnSpawnGraphNodeByShortcut(FInputChord InChord, const FVector2D& InPosition, UEdGraph* InGraph) override;
-#if WITH_RIGVMLEGACYEDITOR
 	virtual bool IsSectionVisible(NodeSectionID::Type InSectionID) const override;
-	virtual bool NewDocument_IsVisibleForType(FBlueprintEditor::ECreatedDocumentType GraphType) const override;
-#else
-	virtual bool IsSectionVisible(RigVMNodeSectionID::Type InSectionID) const override;
-	virtual bool NewDocument_IsVisibleForType(FRigVMEditorBase::ECreatedDocumentType GraphType) const override;
-#endif
+	virtual bool NewDocument_IsVisibleForType(ECreatedDocumentType GraphType) const override;
 
 	virtual void PostTransaction(bool bSuccess, const FTransaction* Transaction, bool bIsRedo) override;
 
@@ -139,7 +134,6 @@ public:
 	bool DetailViewShowsAnyRigElement() const;
 	bool DetailViewShowsRigElement(FRigElementKey InKey) const;
 	TArray<FRigElementKey> GetSelectedRigElementsFromDetailView() const;
-	TArray< TWeakObjectPtr<UObject> > GetSelectedObjectsFromDetailView() const;
 
 	void SetDetailViewForRigModules();
 	void SetDetailViewForRigModules(const TArray<FString> InKeys);
diff --git a/Engine/Plugins/Animation/ControlRig/Source/ControlRigEditor/Private/Editor/ControlRigEditorMode.cpp b/Engine/Plugins/Animation/ControlRig/Source/ControlRigEditor/Private/Editor/ControlRigEditorMode.cpp
--- a/Engine/Plugins/Animation/ControlRig/Source/ControlRigEditor/Private/Editor/ControlRigEditorMode.cpp
+++ b/Engine/Plugins/Animation/ControlRig/Source/ControlRigEditor/Private/Editor/ControlRigEditorMode.cpp
@@ -1,10 +1,8 @@
 // Copyright Epic Games, Inc. All Rights Reserved.
 
-#include "Editor/ControlRigEditorMode.h"
-#if WITH_RIGVMLEGACYEDITOR
+#include "Editor/ControlRigEditorMode.h" 
 #include "BlueprintEditorTabs.h"
 #include "SBlueprintEditorToolbar.h"
-#endif 
 #include "PersonaModule.h"
 #include "IPersonaToolkit.h"
 #include "PersonaTabs.h"
@@ -19,13 +17,9 @@
 #include "ToolMenus.h"
 
 FControlRigEditorMode::FControlRigEditorMode(const TSharedRef<FControlRigEditor>& InControlRigEditor, bool bCreateDefaultLayout)
-#if WITH_RIGVMLEGACYEDITOR
 	: FBlueprintEditorApplicationMode(InControlRigEditor, FControlRigEditorModes::ControlRigEditorMode, FControlRigEditorModes::GetLocalizedMode, false, false)
-#else
-	: FRigVMNewEditorMode(InControlRigEditor)
-#endif
 {
-	ControlRigBlueprintPtr = CastChecked<UControlRigBlueprint>(InControlRigEditor->FRigVMEditorBase::GetBlueprintObj());
+	ControlRigBlueprintPtr = CastChecked<UControlRigBlueprint>(InControlRigEditor->GetBlueprintObj());
 
 	TabFactories.RegisterFactory(MakeShared<FRigHierarchyTabSummoner>(InControlRigEditor));
 	TabFactories.RegisterFactory(MakeShared<FRigVMExecutionStackTabSummoner>(InControlRigEditor));
@@ -37,9 +31,7 @@ FControlRigEditorMode::FControlRigEditorMode(const TSharedRef<FControlRigEditor>
 	FPersonaModule& PersonaModule = FModuleManager::LoadModuleChecked<FPersonaModule>("Persona");
 
  	FPersonaViewportArgs ViewportArgs(InControlRigEditor->GetPersonaToolkit()->GetPreviewScene());
-#if WITH_RIGVMLEGACYEDITOR
  	ViewportArgs.BlueprintEditor = InControlRigEditor;
-#endif
  	ViewportArgs.bShowStats = false;
 	ViewportArgs.bShowPlaySpeedMenu = false;
 	ViewportArgs.bShowTimeline = true;
@@ -86,9 +78,7 @@ FControlRigEditorMode::FControlRigEditorMode(const TSharedRef<FControlRigEditor>
 							->AddTab(FRigHierarchyTabSummoner::TabID, ETabState::OpenedTab)
 							->AddTab(FRigVMExecutionStackTabSummoner::TabID, ETabState::OpenedTab)
 							->AddTab(FRigCurveContainerTabSummoner::TabID, ETabState::OpenedTab)
-#if WITH_RIGVMLEGACYEDITOR
 							->AddTab(FBlueprintEditorTabs::MyBlueprintID, ETabState::OpenedTab)
-#endif
 							->AddTab(FRigVMEditorGraphExplorerTabSummoner::TabID, ETabState::ClosedTab)
 						)
 					)
@@ -110,10 +100,8 @@ FControlRigEditorMode::FControlRigEditorMode(const TSharedRef<FControlRigEditor>
 							// Middle bottom - compiler results & find
 							FTabManager::NewStack()
 							->SetSizeCoefficient(0.2f)
-#if WITH_RIGVMLEGACYEDITOR
 							->AddTab(FBlueprintEditorTabs::CompilerResultsID, ETabState::ClosedTab)
 							->AddTab(FBlueprintEditorTabs::FindResultsID, ETabState::ClosedTab)
-#endif
 						)
 					)
 					->Split
@@ -128,45 +116,33 @@ FControlRigEditorMode::FControlRigEditorMode(const TSharedRef<FControlRigEditor>
 							FTabManager::NewStack()
 							->SetHideTabWell(false)
 							->SetSizeCoefficient(1.f)
-#if WITH_RIGVMLEGACYEDITOR
 							->AddTab(FBlueprintEditorTabs::DetailsID, ETabState::OpenedTab)
-#endif
 							->AddTab(FPersonaTabs::AdvancedPreviewSceneSettingsID, ETabState::OpenedTab)
 							->AddTab(FRigAnimAttributeTabSummoner::TabID, ETabState::OpenedTab)
-#if WITH_RIGVMLEGACYEDITOR
 							->SetForegroundTab(FBlueprintEditorTabs::DetailsID)
-#endif
 						)
 					)
 				)
 			);
 	}
-
-#if WITH_RIGVMLEGACYEDITOR
+	
 	if (UToolMenu* Toolbar = InControlRigEditor->RegisterModeToolbarIfUnregistered(GetModeName()))
 	{
 		InControlRigEditor->GetToolbarBuilder()->AddCompileToolbar(Toolbar);
 		InControlRigEditor->GetToolbarBuilder()->AddScriptingToolbar(Toolbar);
 		InControlRigEditor->GetToolbarBuilder()->AddBlueprintGlobalOptionsToolbar(Toolbar);
 	}
-#endif
 }
 
 void FControlRigEditorMode::RegisterTabFactories(TSharedPtr<FTabManager> InTabManager)
 {
-#if WITH_RIGVMLEGACYEDITOR
 	TSharedPtr<FBlueprintEditor> BP = MyBlueprintEditor.Pin();
-#else
-	TSharedPtr<FRigVMNewEditor> BP = Editor.Pin();
-#endif
 	
 	BP->RegisterToolbarTab(InTabManager.ToSharedRef());
 
 	// Mode-specific setup
-#if WITH_RIGVMLEGACYEDITOR
 	BP->PushTabFactories(CoreTabFactories);
 	BP->PushTabFactories(BlueprintEditorTabFactories);
-#endif
 	BP->PushTabFactories(TabFactories);
 }
 
@@ -199,9 +175,7 @@ FModularRigEditorMode::FModularRigEditorMode(const TSharedRef<FControlRigEditor>
 						->SetHideTabWell(false)
 						->SetSizeCoefficient(0.5f)
 						->AddTab(FRigModuleAssetBrowserTabSummoner::TabID, ETabState::OpenedTab)
-#if WITH_RIGVMLEGACYEDITOR
 						->SetForegroundTab(FBlueprintEditorTabs::DetailsID)
-#endif
 					)
 					->Split
 					(
@@ -211,9 +185,7 @@ FModularRigEditorMode::FModularRigEditorMode(const TSharedRef<FControlRigEditor>
 						->SetSizeCoefficient(0.5f)
 						->SetForegroundTab(FModularRigModelTabSummoner::TabID)
 						->AddTab(FModularRigModelTabSummoner::TabID, ETabState::OpenedTab)
-#if WITH_RIGVMLEGACYEDITOR
 						->AddTab(FBlueprintEditorTabs::MyBlueprintID, ETabState::OpenedTab)
-#endif
 					)
 				)
 				->Split
@@ -230,10 +202,8 @@ FModularRigEditorMode::FModularRigEditorMode(const TSharedRef<FControlRigEditor>
 					FTabManager::NewStack()
 					->SetSizeCoefficient(0.2f)
 					->SetHideTabWell(false)
-#if WITH_RIGVMLEGACYEDITOR
 					->SetForegroundTab(FBlueprintEditorTabs::DetailsID)
 					->AddTab(FBlueprintEditorTabs::DetailsID, ETabState::OpenedTab)
-#endif
 					->AddTab(FPersonaTabs::AdvancedPreviewSceneSettingsID, ETabState::OpenedTab)
 					->AddTab(FRigHierarchyTabSummoner::TabID, ETabState::OpenedTab)
 				)
@@ -243,22 +213,16 @@ FModularRigEditorMode::FModularRigEditorMode(const TSharedRef<FControlRigEditor>
 
 void FModularRigEditorMode::RegisterTabFactories(TSharedPtr<FTabManager> InTabManager)
 {
-#if WITH_RIGVMLEGACYEDITOR
 	TSharedPtr<FBlueprintEditor> BP = MyBlueprintEditor.Pin();
-#else
-	TSharedPtr<FRigVMNewEditor> BP = Editor.Pin();
-#endif
 	
 	BP->RegisterToolbarTab(InTabManager.ToSharedRef());
 
 	static const TArray<FName> DisallowedTabs = {
-#if WITH_RIGVMLEGACYEDITOR
 		FBlueprintEditorTabs::PaletteID,
 		FBlueprintEditorTabs::ReplaceNodeReferencesID,
 		FBlueprintEditorTabs::CompilerResultsID,
 		FBlueprintEditorTabs::FindResultsID,
 		FBlueprintEditorTabs::BookmarksID,
-#endif
 		FRigVMExecutionStackTabSummoner::TabID
 	};
 
@@ -275,9 +239,7 @@ void FModularRigEditorMode::RegisterTabFactories(TSharedPtr<FTabManager> InTabMa
 	};
 
 	// Mode-specific setup
-#if WITH_RIGVMLEGACYEDITOR
 	PushTabFactories(CoreTabFactories);
 	PushTabFactories(BlueprintEditorTabFactories);
-#endif
 	PushTabFactories(TabFactories);
 }
diff --git a/Engine/Plugins/Animation/ControlRig/Source/ControlRigEditor/Private/Editor/ControlRigEditorMode.h b/Engine/Plugins/Animation/ControlRig/Source/ControlRigEditor/Private/Editor/ControlRigEditorMode.h
--- a/Engine/Plugins/Animation/ControlRig/Source/ControlRigEditor/Private/Editor/ControlRigEditorMode.h
+++ b/Engine/Plugins/Animation/ControlRig/Source/ControlRigEditor/Private/Editor/ControlRigEditorMode.h
@@ -6,19 +6,11 @@
 #include "WorkflowOrientedApp/WorkflowTabManager.h"
 #include "Editor/ControlRigEditor.h"
 #include "ControlRigBlueprint.h"
-#if WITH_RIGVMLEGACYEDITOR
 #include "BlueprintEditorModes.h"
-#else
-#include "Editor/RigVMNewEditorMode.h" 
-#endif
 
 class UControlRigBlueprint;
 
-#if WITH_RIGVMLEGACYEDITOR
 class FControlRigEditorMode : public FBlueprintEditorApplicationMode
-#else
-class FControlRigEditorMode : public FRigVMNewEditorMode
-#endif
 {
 public:
 	FControlRigEditorMode(const TSharedRef<FControlRigEditor>& InControlRigEditor, bool bCreateDefaultLayout = true);
diff --git a/Engine/Plugins/Animation/ControlRig/Source/ControlRigEditor/Private/Editor/RigInfluenceMapTabSummoner.cpp b/Engine/Plugins/Animation/ControlRig/Source/ControlRigEditor/Private/Editor/RigInfluenceMapTabSummoner.cpp
--- a/Engine/Plugins/Animation/ControlRig/Source/ControlRigEditor/Private/Editor/RigInfluenceMapTabSummoner.cpp
+++ b/Engine/Plugins/Animation/ControlRig/Source/ControlRigEditor/Private/Editor/RigInfluenceMapTabSummoner.cpp
@@ -4,13 +4,8 @@
 #include "Editor/SRigHierarchy.h"
 #include "Editor/RigVMEditorStyle.h"
 #include "Editor/ControlRigEditor.h"
-#include "ControlRigBlueprint.h"
-
-#if WITH_RIGVMLEGACYEDITOR
 #include "SKismetInspector.h"
-#else
-#include "Editor/SRigVMDetailsInspector.h"
-#endif
+#include "ControlRigBlueprint.h"
 
 #define LOCTEXT_NAMESPACE "RigInfluenceMapTabSummoner"
 
@@ -31,15 +26,11 @@ FRigInfluenceMapTabSummoner::FRigInfluenceMapTabSummoner(const TSharedRef<FContr
 
 TSharedRef<SWidget> FRigInfluenceMapTabSummoner::CreateTabBody(const FWorkflowTabSpawnInfo& Info) const
 {
-#if WITH_RIGVMLEGACYEDITOR
 	TSharedRef<SKismetInspector> KismetInspector = SNew(SKismetInspector);
-#else
-	TSharedRef<SRigVMDetailsInspector> KismetInspector = SNew(SRigVMDetailsInspector);
-#endif
 
-	if (TSharedPtr<FControlRigEditor> Editor = ControlRigEditor.Pin())
+	if (ControlRigEditor.IsValid())
 	{
-		if (UControlRigBlueprint* RigBlueprint = Cast<UControlRigBlueprint>(Editor->FRigVMEditorBase::GetBlueprintObj()))
+		if (UControlRigBlueprint* RigBlueprint = Cast<UControlRigBlueprint>(ControlRigEditor.Pin()->GetBlueprintObj()))
 		{
 			TSharedPtr<FStructOnScope> StructToDisplay = MakeShareable(new FStructOnScope(
 				FRigInfluenceMapPerEvent::StaticStruct(), (uint8*)&RigBlueprint->Influences));
diff --git a/Engine/Plugins/Animation/ControlRig/Source/ControlRigEditor/Private/Editor/RigValidationTabSummoner.cpp b/Engine/Plugins/Animation/ControlRig/Source/ControlRigEditor/Private/Editor/RigValidationTabSummoner.cpp
--- a/Engine/Plugins/Animation/ControlRig/Source/ControlRigEditor/Private/Editor/RigValidationTabSummoner.cpp
+++ b/Engine/Plugins/Animation/ControlRig/Source/ControlRigEditor/Private/Editor/RigValidationTabSummoner.cpp
@@ -27,20 +27,16 @@ FRigValidationTabSummoner::FRigValidationTabSummoner(const TSharedRef<FControlRi
 TSharedRef<SWidget> FRigValidationTabSummoner::CreateTabBody(const FWorkflowTabSpawnInfo& Info) const
 {
 	ensure(WeakControlRigEditor.IsValid());
-	if (TSharedPtr<FRigVMEditorBase> Editor = WeakControlRigEditor.Pin())
-	{
-		UControlRigBlueprint* RigBlueprint = Cast<UControlRigBlueprint>(Editor->FRigVMEditorBase::GetBlueprintObj());
-		check(RigBlueprint);
-	
-		UControlRigValidator* Validator = RigBlueprint->Validator;
-		check(Validator);
 
-		TSharedRef<SControlRigValidationWidget> ValidationWidget = SNew(SControlRigValidationWidget, Validator);
-		Validator->SetControlRig(Cast<UControlRig>(RigBlueprint->GetObjectBeingDebugged()));
-		return ValidationWidget;
-	}
+	UControlRigBlueprint* RigBlueprint = Cast<UControlRigBlueprint>(WeakControlRigEditor.Pin()->GetBlueprintObj());
+	check(RigBlueprint);
+	
+	UControlRigValidator* Validator = RigBlueprint->Validator;
+	check(Validator);
 
-	return SNullWidget::NullWidget;
+	TSharedRef<SControlRigValidationWidget> ValidationWidget = SNew(SControlRigValidationWidget, Validator);
+	Validator->SetControlRig(Cast<UControlRig>(RigBlueprint->GetObjectBeingDebugged()));
+	return ValidationWidget;
 }
 
 #undef LOCTEXT_NAMESPACE 
diff --git a/Engine/Plugins/Animation/ControlRig/Source/ControlRigEditor/Private/Editor/SControlRigValidationWidget.cpp b/Engine/Plugins/Animation/ControlRig/Source/ControlRigEditor/Private/Editor/SControlRigValidationWidget.cpp
--- a/Engine/Plugins/Animation/ControlRig/Source/ControlRigEditor/Private/Editor/SControlRigValidationWidget.cpp
+++ b/Engine/Plugins/Animation/ControlRig/Source/ControlRigEditor/Private/Editor/SControlRigValidationWidget.cpp
@@ -16,16 +16,9 @@
 #include "UserInterface/SMessageLogListing.h"
 #include "Framework/Application/SlateApplication.h"
 #include "ControlRigBlueprint.h"
-#include "Editor/RigVMNewEditor.h"
 #include "UObject/UObjectIterator.h"
 #include "Rigs/RigHierarchyController.h"
 
-#if WITH_RIGVMLEGACYEDITOR
-#include "SKismetInspector.h"
-#else
-#include "Editor/SRigVMDetailsInspector.h"
-#endif
-
 #define LOCTEXT_NAMESPACE "SControlRigValidationWidget"
 
 //////////////////////////////////////////////////////////////
@@ -71,11 +64,7 @@ void SControlRigValidationPassTableRow::Construct(const FArguments& InArgs, cons
 			.Padding(24.0f, 4.0f, 0.f, 0.f)
 			.AutoHeight()
 			[
-#if WITH_RIGVMLEGACYEDITOR
 				SAssignNew(KismetInspector, SKismetInspector)
-#else
-				SAssignNew(KismetInspector, SRigVMDetailsInspector)
-#endif
 				.Visibility(InValidationWidget, &SControlRigValidationWidget::IsClassVisible, InPassItem->Class)
 				.ShowTitleArea(false)
 			]
diff --git a/Engine/Plugins/Animation/ControlRig/Source/ControlRigEditor/Private/Editor/SControlRigValidationWidget.h b/Engine/Plugins/Animation/ControlRig/Source/ControlRigEditor/Private/Editor/SControlRigValidationWidget.h
--- a/Engine/Plugins/Animation/ControlRig/Source/ControlRigEditor/Private/Editor/SControlRigValidationWidget.h
+++ b/Engine/Plugins/Animation/ControlRig/Source/ControlRigEditor/Private/Editor/SControlRigValidationWidget.h
@@ -4,19 +4,11 @@
 
 #include "CoreMinimal.h"
 #include "Widgets/Views/SListView.h"
+#include "SKismetInspector.h"
 #include "Model/MessageLogListingModel.h"
 #include "Presentation/MessageLogListingViewModel.h"
 #include "ControlRigValidationPass.h"
 
-#if WITH_RIGVMLEGACYEDITOR
-#include "SKismetInspector.h"
-#else
-#include "Editor/SRigVMDetailsInspector.h"
-#endif
-
-#if !WITH_RIGVMLEGACYEDITOR
-class SRigVMDetailsInspector;
-#endif
 class SControlRigValidationWidget;
 
 class FControlRigValidationPassItem : public TSharedFromThis<FControlRigValidationPassItem>
@@ -45,11 +37,8 @@ class SControlRigValidationPassTableRow : public STableRow<TSharedPtr<FControlRi
 	void RefreshDetails(UControlRigValidator* InValidator, UClass* InClass);
 
 private:
-#if WITH_RIGVMLEGACYEDITOR
+
 	TSharedPtr<SKismetInspector> KismetInspector;
-#else
-	TSharedPtr<SRigVMDetailsInspector> KismetInspector;
-#endif
 };
 
 class SControlRigValidationWidget : public SCompoundWidget
diff --git a/Engine/Plugins/Animation/ControlRig/Source/ControlRigEditor/Private/Editor/SModularRigModel.cpp b/Engine/Plugins/Animation/ControlRig/Source/ControlRigEditor/Private/Editor/SModularRigModel.cpp
--- a/Engine/Plugins/Animation/ControlRig/Source/ControlRigEditor/Private/Editor/SModularRigModel.cpp
+++ b/Engine/Plugins/Animation/ControlRig/Source/ControlRigEditor/Private/Editor/SModularRigModel.cpp
@@ -34,6 +34,7 @@
 #include "HAL/PlatformTime.h"
 #include "Dialogs/Dialogs.h"
 #include "IPersonaToolkit.h"
+#include "SKismetInspector.h"
 #include "Types/WidgetActiveTimerDelegate.h"
 #include "Dialog/SCustomDialog.h"
 #include "EditMode/ControlRigEditMode.h"
@@ -57,12 +58,6 @@
 #include "Widgets/SRigVMBulkEditDialog.h"
 #include "Widgets/SRigVMSwapAssetReferencesWidget.h"
 
-#if WITH_RIGVMLEGACYEDITOR
-#include "SKismetInspector.h"
-#else
-#include "Editor/SRigVMDetailsInspector.h"
-#endif
-
 #define LOCTEXT_NAMESPACE "SModularRigModel"
 
 ///////////////////////////////////////////////////////////
@@ -242,7 +237,7 @@ void SModularRigModel::Construct(const FArguments& InArgs, TSharedRef<FControlRi
 	}
 }
 
-void SModularRigModel::OnEditorClose(const FRigVMEditorBase* InEditor, URigVMBlueprint* InBlueprint)
+void SModularRigModel::OnEditorClose(const FRigVMEditor* InEditor, URigVMBlueprint* InBlueprint)
 {
 	if (InEditor)
 	{
@@ -927,19 +922,15 @@ void SModularRigModel::HandleMirrorModules(const TArray<FString>& InPaths)
 	{
 		FRigVMMirrorSettings Settings;
 		TSharedPtr<FStructOnScope> StructToDisplay = MakeShareable(new FStructOnScope(FRigVMMirrorSettings::StaticStruct(), (uint8*)&Settings));
-
-#if WITH_RIGVMLEGACYEDITOR
-		TSharedRef<SKismetInspector> DetailsInspector = SNew(SKismetInspector);
-#else
-		TSharedRef<SRigVMDetailsInspector> DetailsInspector = SNew(SRigVMDetailsInspector);
-#endif
-		DetailsInspector->ShowSingleStruct(StructToDisplay);
+		
+		TSharedRef<SKismetInspector> KismetInspector = SNew(SKismetInspector);
+		KismetInspector->ShowSingleStruct(StructToDisplay);
 		
 		TSharedRef<SCustomDialog> MirrorDialog = SNew(SCustomDialog)
 			.Title(FText(LOCTEXT("ControlModularModelMirror", "Mirror Selected Modules")))
 			.Content()
 			[
-				DetailsInspector
+				KismetInspector
 			]
 			.Buttons({
 				SCustomDialog::FButton(LOCTEXT("OK", "OK")),
diff --git a/Engine/Plugins/Animation/ControlRig/Source/ControlRigEditor/Private/Editor/SModularRigModel.h b/Engine/Plugins/Animation/ControlRig/Source/ControlRigEditor/Private/Editor/SModularRigModel.h
--- a/Engine/Plugins/Animation/ControlRig/Source/ControlRigEditor/Private/Editor/SModularRigModel.h
+++ b/Engine/Plugins/Animation/ControlRig/Source/ControlRigEditor/Private/Editor/SModularRigModel.h
@@ -7,7 +7,7 @@
 #include "Editor/SModularRigTreeView.h"
 #include "ControlRigBlueprint.h"
 #include "DragAndDrop/GraphNodeDragDropOp.h"
-#include "Editor/RigVMNewEditor.h"
+#include "Editor/RigVMEditor.h"
 #include "ControlRigDragOps.h"
 
 class SModularRigModel;
@@ -43,7 +43,7 @@ public:
 
 private:
 
-	void OnEditorClose(const FRigVMEditorBase* InEditor, URigVMBlueprint* InBlueprint);
+	void OnEditorClose(const FRigVMEditor* InEditor, URigVMBlueprint* InBlueprint);
 
 	/** Bind commands that this widget handles */
 	void BindCommands();
diff --git a/Engine/Plugins/Animation/ControlRig/Source/ControlRigEditor/Private/Editor/SRigCurveContainer.cpp b/Engine/Plugins/Animation/ControlRig/Source/ControlRigEditor/Private/Editor/SRigCurveContainer.cpp
--- a/Engine/Plugins/Animation/ControlRig/Source/ControlRigEditor/Private/Editor/SRigCurveContainer.cpp
+++ b/Engine/Plugins/Animation/ControlRig/Source/ControlRigEditor/Private/Editor/SRigCurveContainer.cpp
@@ -310,7 +310,7 @@ TSharedPtr<SWidget> SRigCurveContainer::OnGetContextMenuContent() const
 	return MenuBuilder.MakeWidget();
 }
 
-void SRigCurveContainer::OnEditorClose(const FRigVMEditorBase* InEditor, URigVMBlueprint* InBlueprint)
+void SRigCurveContainer::OnEditorClose(const FRigVMEditor* InEditor, URigVMBlueprint* InBlueprint)
 {
 	if (UControlRigBlueprint* BP = Cast<UControlRigBlueprint>(InBlueprint))
 	{
diff --git a/Engine/Plugins/Animation/ControlRig/Source/ControlRigEditor/Private/Editor/SRigCurveContainer.h b/Engine/Plugins/Animation/ControlRig/Source/ControlRigEditor/Private/Editor/SRigCurveContainer.h
--- a/Engine/Plugins/Animation/ControlRig/Source/ControlRigEditor/Private/Editor/SRigCurveContainer.h
+++ b/Engine/Plugins/Animation/ControlRig/Source/ControlRigEditor/Private/Editor/SRigCurveContainer.h
@@ -18,7 +18,7 @@ class SInlineEditableTextBlock;
 class FControlRigEditor;
 struct FAssetData;
 class URigVMBlueprint;
-class FRigVMEditorBase;
+class FRigVMEditor;
 
 class FDisplayedRigCurveInfo
 {
@@ -213,7 +213,7 @@ private:
 	/** Handler for context menus */
 	TSharedPtr<SWidget> OnGetContextMenuContent() const;
 
-	void OnEditorClose(const FRigVMEditorBase* InEditor, URigVMBlueprint* InBlueprint);
+	void OnEditorClose(const FRigVMEditor* InEditor, URigVMBlueprint* InBlueprint);
 
 	/**
 	* Clears and rebuilds the table, according to an optional search string
diff --git a/Engine/Plugins/Animation/ControlRig/Source/ControlRigEditor/Private/Editor/SRigHierarchy.cpp b/Engine/Plugins/Animation/ControlRig/Source/ControlRigEditor/Private/Editor/SRigHierarchy.cpp
--- a/Engine/Plugins/Animation/ControlRig/Source/ControlRigEditor/Private/Editor/SRigHierarchy.cpp
+++ b/Engine/Plugins/Animation/ControlRig/Source/ControlRigEditor/Private/Editor/SRigHierarchy.cpp
@@ -34,6 +34,7 @@
 #include "HAL/PlatformTime.h"
 #include "Dialogs/Dialogs.h"
 #include "IPersonaToolkit.h"
+#include "SKismetInspector.h"
 #include "Types/WidgetActiveTimerDelegate.h"
 #include "Dialog/SCustomDialog.h"
 #include "EditMode/ControlRigEditMode.h"
@@ -54,12 +55,6 @@
 #include "Preferences/PersonaOptions.h"
 #include "Editor/SModularRigModel.h"
 
-#if WITH_RIGVMLEGACYEDITOR
-#include "SKismetInspector.h"
-#else
-#include "Editor/SRigVMDetailsInspector.h"
-#endif
-
 #include UE_INLINE_GENERATED_CPP_BY_NAME(SRigHierarchy)
 
 #define LOCTEXT_NAMESPACE "SRigHierarchy"
@@ -255,7 +250,7 @@ void SRigHierarchy::Construct(const FArguments& InArgs, TSharedRef<FControlRigEd
 	}));
 }
 
-void SRigHierarchy::OnEditorClose(const FRigVMEditorBase* InEditor, URigVMBlueprint* InBlueprint)
+void SRigHierarchy::OnEditorClose(const FRigVMEditor* InEditor, URigVMBlueprint* InBlueprint)
 {
 	if (InEditor)
 	{
@@ -538,11 +533,8 @@ EVisibility SRigHierarchy::IsSearchbarVisible() const
 FReply SRigHierarchy::OnImportSkeletonClicked()
 {
 	TSharedPtr<FStructOnScope> StructToDisplay = MakeShareable(new FStructOnScope(FRigHierarchyImportSettings::StaticStruct(), (uint8*)&ImportSettings));
-#if WITH_RIGVMLEGACYEDITOR
+
 	TSharedRef<SKismetInspector> KismetInspector = SNew(SKismetInspector);
-#else
-	TSharedRef<SRigVMDetailsInspector> KismetInspector = SNew(SRigVMDetailsInspector);
-#endif
 	KismetInspector->ShowSingleStruct(StructToDisplay);
 
 	SGenericDialogWidget::FArguments DialogArguments;
@@ -2371,11 +2363,8 @@ void SRigHierarchy::HandleMirrorItem()
 		
 		FRigVMMirrorSettings Settings;
 		TSharedPtr<FStructOnScope> StructToDisplay = MakeShareable(new FStructOnScope(FRigVMMirrorSettings::StaticStruct(), (uint8*)&Settings));
-#if WITH_RIGVMLEGACYEDITOR
+
 		TSharedRef<SKismetInspector> KismetInspector = SNew(SKismetInspector);
-#else
-		TSharedRef<SRigVMDetailsInspector> KismetInspector = SNew(SRigVMDetailsInspector);
-#endif
 		KismetInspector->ShowSingleStruct(StructToDisplay);
 
 		TSharedRef<SCustomDialog> MirrorDialog = SNew(SCustomDialog)
diff --git a/Engine/Plugins/Animation/ControlRig/Source/ControlRigEditor/Private/Editor/SRigHierarchy.h b/Engine/Plugins/Animation/ControlRig/Source/ControlRigEditor/Private/Editor/SRigHierarchy.h
--- a/Engine/Plugins/Animation/ControlRig/Source/ControlRigEditor/Private/Editor/SRigHierarchy.h
+++ b/Engine/Plugins/Animation/ControlRig/Source/ControlRigEditor/Private/Editor/SRigHierarchy.h
@@ -8,7 +8,7 @@
 #include "Editor/SRigHierarchyTreeView.h"
 #include "Units/RigUnitContext.h"
 #include "ControlRigBlueprint.h"
-#include "Editor/RigVMNewEditor.h"
+#include "Editor/RigVMEditor.h"
 #include "ControlRigSchematicModel.h"
 #include "ControlRigDragOps.h"
 #include "SRigHierarchy.generated.h"
@@ -59,7 +59,7 @@ public:
 
 private:
 
-	void OnEditorClose(const FRigVMEditorBase* InEditor, URigVMBlueprint* InBlueprint);
+	void OnEditorClose(const FRigVMEditor* InEditor, URigVMBlueprint* InBlueprint);
 
 	/** Bind commands that this widget handles */
 	void BindCommands();
diff --git a/Engine/Plugins/Animation/ControlRig/Source/ControlRigEditor/Public/EditMode/ControlRigEditMode.h b/Engine/Plugins/Animation/ControlRig/Source/ControlRigEditor/Public/EditMode/ControlRigEditMode.h
--- a/Engine/Plugins/Animation/ControlRig/Source/ControlRigEditor/Public/EditMode/ControlRigEditMode.h
+++ b/Engine/Plugins/Animation/ControlRig/Source/ControlRigEditor/Public/EditMode/ControlRigEditMode.h
@@ -22,7 +22,7 @@
 #include "ControlRigEditMode.generated.h"
 
 class FEditorViewportClient;
-class FRigVMEditorBase;
+class FRigVMEditor;
 class FViewport;
 class UActorFactory;
 struct FViewportClick;
diff --git a/Engine/Plugins/Animation/ControlRig/Source/ControlRigEditor/Public/IControlRigEditor.h b/Engine/Plugins/Animation/ControlRig/Source/ControlRigEditor/Public/IControlRigEditor.h
--- a/Engine/Plugins/Animation/ControlRig/Source/ControlRigEditor/Public/IControlRigEditor.h
+++ b/Engine/Plugins/Animation/ControlRig/Source/ControlRigEditor/Public/IControlRigEditor.h
@@ -3,23 +3,11 @@
 #pragma once
 
 #include "CoreMinimal.h"
-
-#if WITH_RIGVMLEGACYEDITOR
-#include "Editor/RigVMLegacyEditor.h"
-#else
-#include "Editor/RigVMNewEditor.h"
-#endif
+#include "Editor/RigVMEditor.h"
 
 class UControlRigBlueprint;
 class FBlueprintActionDatabaseRegistrar;
 
-#if WITH_RIGVMLEGACYEDITOR
-class IControlRigEditor : public FRigVMLegacyEditor
-{
-};
-
-#else
-class IControlRigEditor : public FRigVMNewEditor
+class IControlRigEditor : public FRigVMEditor
 {
 };
-#endif
\ No newline at end of file
diff --git a/Engine/Plugins/Experimental/AnimNext/Source/AnimNextEditor/Private/Common/AnimNextEdGraphNodeCustomization.cpp b/Engine/Plugins/Experimental/AnimNext/Source/AnimNextEditor/Private/Common/AnimNextEdGraphNodeCustomization.cpp
--- a/Engine/Plugins/Experimental/AnimNext/Source/AnimNextEditor/Private/Common/AnimNextEdGraphNodeCustomization.cpp
+++ b/Engine/Plugins/Experimental/AnimNext/Source/AnimNextEditor/Private/Common/AnimNextEdGraphNodeCustomization.cpp
@@ -108,7 +108,7 @@ void FAnimNextEdGraphNodeCustomization::CustomizeObjects(IDetailLayoutBuilder& D
 						continue;
 					}
 
-					TWeakPtr<FRigVMEditorBase> WeakEditor; // passing an empty RigVMEditor to FRigVMCollapseGraphLayout
+					TWeakPtr<FRigVMEditor> WeakEditor; // passing an empty RigVMEditor to FRigVMCollapseGraphLayout
 					IRigVMClientHost* RigVMClientHost = EdGraphNode->GetController()->GetClientHost();
 					check(RigVMClientHost);
 
diff --git a/Engine/Plugins/Experimental/AnimNext/Source/AnimNextEditor/Private/Graph/AnimNextEdGraphCustomization.cpp b/Engine/Plugins/Experimental/AnimNext/Source/AnimNextEditor/Private/Graph/AnimNextEdGraphCustomization.cpp
--- a/Engine/Plugins/Experimental/AnimNext/Source/AnimNextEditor/Private/Graph/AnimNextEdGraphCustomization.cpp
+++ b/Engine/Plugins/Experimental/AnimNext/Source/AnimNextEditor/Private/Graph/AnimNextEdGraphCustomization.cpp
@@ -29,7 +29,7 @@ void FAnimNextEdGraphCustomization::CustomizeDetails(IDetailLayoutBuilder& Detai
 	{
 		if(URigVMGraph* Model = EdGraph->GetModel())
 		{
-			TWeakPtr<FRigVMEditorBase> WeakEditor; // passing an empty RigVMEditor to FRigVMCollapseGraphLayout
+			TWeakPtr<FRigVMEditor> WeakEditor; // passing an empty RigVMEditor to FRigVMCollapseGraphLayout
 			IRigVMClientHost* RigVMClientHost = EdGraph->GetController()->GetClientHost();
 			check(RigVMClientHost);
 
diff --git a/Engine/Plugins/Runtime/RigVM/Source/RigVM/Private/RigVMSettings.cpp b/Engine/Plugins/Runtime/RigVM/Source/RigVM/Private/RigVMSettings.cpp
--- a/Engine/Plugins/Runtime/RigVM/Source/RigVM/Private/RigVMSettings.cpp
+++ b/Engine/Plugins/Runtime/RigVM/Source/RigVM/Private/RigVMSettings.cpp
@@ -8,9 +8,6 @@
 
 URigVMEditorSettings::URigVMEditorSettings(const FObjectInitializer& ObjectInitializer)
     : Super(ObjectInitializer)
-#if WITH_EDITORONLY_DATA
-	, bEnableContextMenuTimeSlicing(true)
-#endif
 {
 #if WITH_EDITORONLY_DATA
 	bAutoLinkMutableNodes = false;
diff --git a/Engine/Plugins/Runtime/RigVM/Source/RigVM/Public/RigVMSettings.h b/Engine/Plugins/Runtime/RigVM/Source/RigVM/Public/RigVMSettings.h
--- a/Engine/Plugins/Runtime/RigVM/Source/RigVM/Public/RigVMSettings.h
+++ b/Engine/Plugins/Runtime/RigVM/Source/RigVM/Public/RigVMSettings.h
@@ -33,9 +33,6 @@ public:
 	UPROPERTY(EditAnywhere, config, Category = Interaction)
 	bool bAutoLinkMutableNodes;
 
-	/** When the Blueprint graph context menu is invoked (e.g. by right-clicking in the graph or dragging off a pin), do not block the UI while populating the available actions list. */
-	UPROPERTY(EditAnywhere, config, Category = Workflow, meta = (DisplayName = "Enable Non-Blocking Context Menu"))
-	bool bEnableContextMenuTimeSlicing;
 #endif
 };
 
diff --git a/Engine/Plugins/Runtime/RigVM/Source/RigVMDeveloper/Private/EdGraph/RigVMEdGraphNode.cpp b/Engine/Plugins/Runtime/RigVM/Source/RigVMDeveloper/Private/EdGraph/RigVMEdGraphNode.cpp
--- a/Engine/Plugins/Runtime/RigVM/Source/RigVMDeveloper/Private/EdGraph/RigVMEdGraphNode.cpp
+++ b/Engine/Plugins/Runtime/RigVM/Source/RigVMDeveloper/Private/EdGraph/RigVMEdGraphNode.cpp
@@ -10,12 +10,6 @@
 #include "KismetCompiler.h"
 #include "BlueprintNodeSpawner.h"
 #include "BlueprintActionDatabaseRegistrar.h"
-#if WITH_RIGVMLEGACYEDITOR
-#include "FindInBlueprintManager.h"
-#else
-#include "Kismet/RigVMBlueprintCompilationManager.h"
-#include "Kismet/RigVMFindInBlueprintManager.h"
-#endif
 #include "RigVMEditorModule.h"
 #include "Textures/SlateIcon.h"
 #include "RigVMBlueprint.h"
@@ -2168,7 +2162,6 @@ void URigVMEdGraphNode::BeginDestroy()
 }
 
 #if WITH_EDITOR
-#if WITH_RIGVMLEGACYEDITOR
 void URigVMEdGraphNode::AddPinSearchMetaDataInfo(const UEdGraphPin* Pin, TArray<FSearchTagDataPair>& OutTaggedMetaData) const
 {
 	Super::AddPinSearchMetaDataInfo(Pin, OutTaggedMetaData);
@@ -2182,62 +2175,6 @@ void URigVMEdGraphNode::AddPinSearchMetaDataInfo(const UEdGraphPin* Pin, TArray<
 		}
 	}
 }
-#else
-void URigVMEdGraphNode::AddRigVMSearchMetaDataInfo(TArray<struct UBlueprintExtension::FSearchTagDataPair>& OutTaggedMetaData) const
-{
-	// Copied from EdGraphNode::AddSearchMetaDataInfo
-	if (GetSchema() == nullptr)
-	{
-		ensure(false);
-		return;
-	}
-
-	// Searchable - Primary label for the item in the search results
-	OutTaggedMetaData.Add(UBlueprintExtension::FSearchTagDataPair(FRigVMFindInBlueprintSearchTags::FiB_Name, GetNodeTitle(ENodeTitleType::ListView)));
-
-	// Searchable - As well as being searchable, this displays in the tooltip for the node
-	OutTaggedMetaData.Add(UBlueprintExtension::FSearchTagDataPair(FRigVMFindInBlueprintSearchTags::FiB_ClassName, FText::FromString(GetClass()->GetName())));
-
-	// Non-searchable - Used to lookup the node when attempting to jump to it
-	OutTaggedMetaData.Add(UBlueprintExtension::FSearchTagDataPair(FRigVMFindInBlueprintSearchTags::FiB_NodeGuid, FText::FromString(NodeGuid.ToString(EGuidFormats::Digits))));
-
-	// Non-searchable - Important for matching pin types with icons and colors, stored here so that each pin does not store it
-	OutTaggedMetaData.Add(UBlueprintExtension::FSearchTagDataPair(FRigVMFindInBlueprintSearchTags::FiB_SchemaName, FText::FromString(GetSchema()->GetClass()->GetName())));
-
-	// Non-Searchable - Used to display the icon and color for this node for better visual identification.
-	FLinearColor GlyphColor = FLinearColor::White;
-	FSlateIcon Icon = GetIconAndTint(GlyphColor);
-	OutTaggedMetaData.Add(UBlueprintExtension::FSearchTagDataPair(FRigVMFindInBlueprintSearchTags::FiB_Glyph, FText::FromName(Icon.GetStyleName())));
-	OutTaggedMetaData.Add(UBlueprintExtension::FSearchTagDataPair(FRigVMFindInBlueprintSearchTags::FiB_GlyphStyleSet, FText::FromName(Icon.GetStyleSetName())));
-	OutTaggedMetaData.Add(UBlueprintExtension::FSearchTagDataPair(FRigVMFindInBlueprintSearchTags::FiB_GlyphColor, FText::FromString(GlyphColor.ToString())));
-	OutTaggedMetaData.Add(UBlueprintExtension::FSearchTagDataPair(FRigVMFindInBlueprintSearchTags::FiB_Comment, FText::FromString(NodeComment)));
-}
-
-void URigVMEdGraphNode::AddRigVMPinSearchMetaDataInfo(const UEdGraphPin* Pin, TArray<UBlueprintExtension::FSearchTagDataPair>& OutTaggedMetaData) const
-{
-	// Copied from EdGraphNode::AddPinSearchMetaDataInfo
-	{
-		//Super::AddPinSearchMetaDataInfo(Pin, OutTaggedMetaData);
-		// Searchable - Primary label for the item in the search results
-		OutTaggedMetaData.Add(UBlueprintExtension::FSearchTagDataPair(FRigVMFindInBlueprintSearchTags::FiB_Name, Pin->GetSchema()->GetPinDisplayName(Pin)));
-
-		// Searchable - The pin's default value as a text string
-		OutTaggedMetaData.Add(UBlueprintExtension::FSearchTagDataPair(FRigVMFindInBlueprintSearchTags::FiB_DefaultValue, Pin->GetDefaultAsText()));
-	}
-	
-	if(const URigVMPin* ModelPin = FindModelPinFromGraphPin(Pin))
-	{
-		OutTaggedMetaData.Emplace(FText::FromString(TEXT("Data Type")), FText::FromString(ModelPin->GetCPPType()));
-		
-		if (ModelPin->IsBoundToVariable())
-		{
-			OutTaggedMetaData.Add(UBlueprintExtension::FSearchTagDataPair(FText::FromString(TEXT("Binding")), FText::FromString(ModelPin->GetBoundVariablePath())));
-		}
-	}
-}
-
-#endif
-
 #endif
 
 FEdGraphPinType URigVMEdGraphNode::GetPinTypeForModelPin(const URigVMPin* InModelPin)
diff --git a/Engine/Plugins/Runtime/RigVM/Source/RigVMDeveloper/Private/EdGraph/RigVMEdGraphSchema.cpp b/Engine/Plugins/Runtime/RigVM/Source/RigVMDeveloper/Private/EdGraph/RigVMEdGraphSchema.cpp
--- a/Engine/Plugins/Runtime/RigVM/Source/RigVMDeveloper/Private/EdGraph/RigVMEdGraphSchema.cpp
+++ b/Engine/Plugins/Runtime/RigVM/Source/RigVMDeveloper/Private/EdGraph/RigVMEdGraphSchema.cpp
@@ -28,11 +28,7 @@
 #include "RigVMFunctions/RigVMDispatch_Constant.h"
 #include "RigVMFunctions/RigVMDispatch_MakeStruct.h"
 #include "RigVMModel/Nodes/RigVMDispatchNode.h"
-#if WITH_RIGVMLEGACYEDITOR
 #include "BlueprintEditor.h"
-#else
-#include "Editor/RigVMNewEditor.h"
-#endif
 
 #include UE_INLINE_GENERATED_CPP_BY_NAME(RigVMEdGraphSchema)
 
@@ -977,11 +973,6 @@ TSharedPtr<INameValidatorInterface> URigVMEdGraphSchema::GetNameValidator(const
 	return MakeShareable(new FRigVMNameValidator(BlueprintObj, ValidationScope, OriginalName));		
 }
 
-bool URigVMEdGraphSchema::CanPromotePinToVariable(const UEdGraphPin& Pin, bool bCond) const
-{
-	return true;
-}
-
 bool URigVMEdGraphSchema::SupportsPinType(const UScriptStruct* ScriptStruct) const
 {
 	if(!ScriptStruct)
@@ -1477,11 +1468,8 @@ TSharedPtr<FEdGraphSchemaAction> URigVMEdGraphSchema::MakeActionFromVariableDesc
 		{
 			Category = FText::GetEmpty();
 		}
-#if WITH_RIGVMLEGACYEDITOR
+
 		TSharedPtr<FRigVMEdGraphSchemaAction_LocalVar> Action = MakeShareable(new FRigVMEdGraphSchemaAction_LocalVar(Category, FText::FromName(Variable.VarName), FText::GetEmpty(), 0, NodeSectionID::LOCAL_VARIABLE));
-#else
-		TSharedPtr<FRigVMEdGraphSchemaAction_LocalVar> Action = MakeShareable(new FRigVMEdGraphSchemaAction_LocalVar(Category, FText::FromName(Variable.VarName), FText::GetEmpty(), 0, RigVMNodeSectionID::LOCAL_VARIABLE));
-#endif
 		Action->SetVariableInfo(Variable.VarName, RigGraph, Variable.VarType.PinCategory == UEdGraphSchema_K2::PC_Boolean);
 		return Action;
 	}
diff --git a/Engine/Plugins/Runtime/RigVM/Source/RigVMDeveloper/Private/Kismet/RigVMBlueprintAssetHandler.cpp b/Engine/Plugins/Runtime/RigVM/Source/RigVMDeveloper/Private/Kismet/RigVMBlueprintAssetHandler.cpp
deleted file mode 100644
diff --git a/Engine/Plugins/Runtime/RigVM/Source/RigVMDeveloper/Private/Kismet/RigVMBlueprintAssetHandler.h b/Engine/Plugins/Runtime/RigVM/Source/RigVMDeveloper/Private/Kismet/RigVMBlueprintAssetHandler.h
deleted file mode 100644
diff --git a/Engine/Plugins/Runtime/RigVM/Source/RigVMDeveloper/Private/Kismet/RigVMBlueprintCompilationManager.cpp b/Engine/Plugins/Runtime/RigVM/Source/RigVMDeveloper/Private/Kismet/RigVMBlueprintCompilationManager.cpp
deleted file mode 100644
diff --git a/Engine/Plugins/Runtime/RigVM/Source/RigVMDeveloper/Private/Kismet/RigVMBlueprintCompilationManager.h b/Engine/Plugins/Runtime/RigVM/Source/RigVMDeveloper/Private/Kismet/RigVMBlueprintCompilationManager.h
deleted file mode 100644
diff --git a/Engine/Plugins/Runtime/RigVM/Source/RigVMDeveloper/Private/Kismet/RigVMBlueprintCompilerExtension.cpp b/Engine/Plugins/Runtime/RigVM/Source/RigVMDeveloper/Private/Kismet/RigVMBlueprintCompilerExtension.cpp
deleted file mode 100644
diff --git a/Engine/Plugins/Runtime/RigVM/Source/RigVMDeveloper/Private/Kismet/RigVMBlueprintCompilerExtension.h b/Engine/Plugins/Runtime/RigVM/Source/RigVMDeveloper/Private/Kismet/RigVMBlueprintCompilerExtension.h
deleted file mode 100644
diff --git a/Engine/Plugins/Runtime/RigVM/Source/RigVMDeveloper/Private/Kismet/RigVMFiBSearchInstance.cpp b/Engine/Plugins/Runtime/RigVM/Source/RigVMDeveloper/Private/Kismet/RigVMFiBSearchInstance.cpp
deleted file mode 100644
diff --git a/Engine/Plugins/Runtime/RigVM/Source/RigVMDeveloper/Private/Kismet/RigVMFiBSearchInstance.h b/Engine/Plugins/Runtime/RigVM/Source/RigVMDeveloper/Private/Kismet/RigVMFiBSearchInstance.h
deleted file mode 100644
diff --git a/Engine/Plugins/Runtime/RigVM/Source/RigVMDeveloper/Private/Kismet/RigVMFindInBlueprintManager.cpp b/Engine/Plugins/Runtime/RigVM/Source/RigVMDeveloper/Private/Kismet/RigVMFindInBlueprintManager.cpp
deleted file mode 100644
diff --git a/Engine/Plugins/Runtime/RigVM/Source/RigVMDeveloper/Private/Kismet/RigVMFindInBlueprintManager.h b/Engine/Plugins/Runtime/RigVM/Source/RigVMDeveloper/Private/Kismet/RigVMFindInBlueprintManager.h
deleted file mode 100644
diff --git a/Engine/Plugins/Runtime/RigVM/Source/RigVMDeveloper/Private/Kismet/RigVMFindInBlueprints.cpp b/Engine/Plugins/Runtime/RigVM/Source/RigVMDeveloper/Private/Kismet/RigVMFindInBlueprints.cpp
deleted file mode 100644
diff --git a/Engine/Plugins/Runtime/RigVM/Source/RigVMDeveloper/Private/Kismet/RigVMFindInBlueprints.h b/Engine/Plugins/Runtime/RigVM/Source/RigVMDeveloper/Private/Kismet/RigVMFindInBlueprints.h
deleted file mode 100644
diff --git a/Engine/Plugins/Runtime/RigVM/Source/RigVMDeveloper/Private/Kismet/RigVMImaginaryBlueprintData.cpp b/Engine/Plugins/Runtime/RigVM/Source/RigVMDeveloper/Private/Kismet/RigVMImaginaryBlueprintData.cpp
deleted file mode 100644
diff --git a/Engine/Plugins/Runtime/RigVM/Source/RigVMDeveloper/Private/Kismet/RigVMImaginaryBlueprintData.h b/Engine/Plugins/Runtime/RigVM/Source/RigVMDeveloper/Private/Kismet/RigVMImaginaryBlueprintData.h
deleted file mode 100644
diff --git a/Engine/Plugins/Runtime/RigVM/Source/RigVMDeveloper/Private/RigVMBlueprint.cpp b/Engine/Plugins/Runtime/RigVM/Source/RigVMDeveloper/Private/RigVMBlueprint.cpp
--- a/Engine/Plugins/Runtime/RigVM/Source/RigVMDeveloper/Private/RigVMBlueprint.cpp
+++ b/Engine/Plugins/Runtime/RigVM/Source/RigVMDeveloper/Private/RigVMBlueprint.cpp
@@ -38,9 +38,6 @@
 #include "CookOnTheSide/CookOnTheFlyServer.h"
 #include "RigVMEditorModule.h"
 #include "ScopedTransaction.h"
-#if !WITH_RIGVMLEGACYEDITOR
-#include "Kismet/RigVMBlueprintCompilationManager.h"
-#endif
 #endif//WITH_EDITOR
 
 #define LOCTEXT_NAMESPACE "RigVMBlueprint"
@@ -2712,13 +2709,8 @@ FName URigVMBlueprint::AddMemberVariable(const FName& InName, const FString& InC
 	FName Result = AddHostMemberVariableFromExternal(Variable, InDefaultValue);
 	if (!Result.IsNone())
 	{
-#if WITH_RIGVMLEGACYEDITOR
 		FBPCompileRequest Request(this, EBlueprintCompileOptions::None, nullptr);
 		FBlueprintCompilationManager::CompileSynchronously(Request);
-#else
-		FRigVMBPCompileRequest Request(this, EBlueprintCompileOptions::None, nullptr);
-		FRigVMBlueprintCompilationManager::CompileSynchronously(Request);
-#endif
 	}
 	return Result;
 }
diff --git a/Engine/Plugins/Runtime/RigVM/Source/RigVMDeveloper/Public/EdGraph/RigVMEdGraph.h b/Engine/Plugins/Runtime/RigVM/Source/RigVMDeveloper/Public/EdGraph/RigVMEdGraph.h
--- a/Engine/Plugins/Runtime/RigVM/Source/RigVMDeveloper/Public/EdGraph/RigVMEdGraph.h
+++ b/Engine/Plugins/Runtime/RigVM/Source/RigVMDeveloper/Public/EdGraph/RigVMEdGraph.h
@@ -208,7 +208,7 @@ private:
 #endif
 	friend class URigVMEdGraphNode;
 	friend class URigVMBlueprint;
-	friend class FRigVMEditorBase;
+	friend class FRigVMEditor;
 	friend class SRigVMGraphNode;
 	friend class URigVMBlueprint;
 
diff --git a/Engine/Plugins/Runtime/RigVM/Source/RigVMDeveloper/Public/EdGraph/RigVMEdGraphNode.h b/Engine/Plugins/Runtime/RigVM/Source/RigVMDeveloper/Public/EdGraph/RigVMEdGraphNode.h
--- a/Engine/Plugins/Runtime/RigVM/Source/RigVMDeveloper/Public/EdGraph/RigVMEdGraphNode.h
+++ b/Engine/Plugins/Runtime/RigVM/Source/RigVMDeveloper/Public/EdGraph/RigVMEdGraphNode.h
@@ -6,7 +6,6 @@
 #include "CoreMinimal.h"
 #include "UObject/WeakObjectPtrTemplates.h"
 #include "RigVMModel/RigVMGraph.h"
-#include "Blueprint/BlueprintExtension.h"
 #include "RigVMEdGraphNode.generated.h"
 
 class FBlueprintActionDatabaseRegistrar;
@@ -27,7 +26,7 @@ class RIGVMDEVELOPER_API URigVMEdGraphNode : public UEdGraphNode
 	friend class URigVMBlueprint;
 	friend class FRigVMEdGraphTraverser;
 	friend class FRigVMEdGraphPanelPinFactory;
-	friend class FRigVMEditorBase;
+	friend class FRigVMEditor;
 	friend struct FRigVMBlueprintUtils;
 	friend class SRigVMEdGraphPinCurveFloat;
 
@@ -126,12 +125,7 @@ public:
 	virtual bool ShouldDrawNodeAsControlPointOnly(int32& OutInputPinIndex, int32& OutOutputPinIndex) const override;
 	virtual void BeginDestroy() override;
 #if WITH_EDITOR
-#if WITH_RIGVMLEGACYEDITOR
 	virtual void AddPinSearchMetaDataInfo(const UEdGraphPin* Pin, TArray<FSearchTagDataPair>& OutTaggedMetaData) const override;
-#else
-	void AddRigVMSearchMetaDataInfo(TArray<struct UBlueprintExtension::FSearchTagDataPair>& OutTaggedMetaData) const;
-	void AddRigVMPinSearchMetaDataInfo(const UEdGraphPin* Pin, TArray<UBlueprintExtension::FSearchTagDataPair>& OutTaggedMetaData) const;
-#endif
 #endif
 	virtual FString GetPinMetaData(FName InPinName, FName InKey) override;
 
diff --git a/Engine/Plugins/Runtime/RigVM/Source/RigVMDeveloper/Public/EdGraph/RigVMEdGraphSchema.h b/Engine/Plugins/Runtime/RigVM/Source/RigVMDeveloper/Public/EdGraph/RigVMEdGraphSchema.h
--- a/Engine/Plugins/Runtime/RigVM/Source/RigVMDeveloper/Public/EdGraph/RigVMEdGraphSchema.h
+++ b/Engine/Plugins/Runtime/RigVM/Source/RigVMDeveloper/Public/EdGraph/RigVMEdGraphSchema.h
@@ -321,7 +321,6 @@ public:
 	virtual void SetPinBeingDroppedOnNode(UEdGraphPin* InSourcePin) const override { PinBeingDropped = InSourcePin; }
 	virtual void InsertAdditionalActions(TArray<UBlueprint*> InBlueprints, TArray<UEdGraph*> InGraphs, TArray<UEdGraphPin*> InPins, FGraphActionListBuilderBase& OutAllActions) const override;
 	virtual TSharedPtr<INameValidatorInterface> GetNameValidator(const UBlueprint* BlueprintObj, const FName& OriginalName, const UStruct* ValidationScope, const FName& ActionTypeId) const override;
-	bool CanPromotePinToVariable(const UEdGraphPin& Pin, bool bCond) const;
 
 	/** Returns true if the schema supports the script type */
 	bool SupportsPinType(const UScriptStruct* ScriptStruct) const;
diff --git a/Engine/Plugins/Runtime/RigVM/Source/RigVMDeveloper/Public/RigVMBlueprint.h b/Engine/Plugins/Runtime/RigVM/Source/RigVMDeveloper/Public/RigVMBlueprint.h
--- a/Engine/Plugins/Runtime/RigVM/Source/RigVMDeveloper/Public/RigVMBlueprint.h
+++ b/Engine/Plugins/Runtime/RigVM/Source/RigVMDeveloper/Public/RigVMBlueprint.h
@@ -849,9 +849,7 @@ private:
 #endif
 
 	friend class FRigVMBlueprintCompilerContext;
-	friend class FRigVMLegacyEditor;
-	friend class FRigVMNewEditor;
-	friend class FRigVMEditorBase;
+	friend class FRigVMEditor;
 	friend class FRigVMEditorModule;
 	friend class URigVMEdGraphSchema;
 	friend struct FRigVMEdGraphSchemaAction_PromoteToVariable;
diff --git a/Engine/Plugins/Runtime/RigVM/Source/RigVMDeveloper/RigVMDeveloper.Build.cs b/Engine/Plugins/Runtime/RigVM/Source/RigVMDeveloper/RigVMDeveloper.Build.cs
--- a/Engine/Plugins/Runtime/RigVM/Source/RigVMDeveloper/RigVMDeveloper.Build.cs
+++ b/Engine/Plugins/Runtime/RigVM/Source/RigVMDeveloper/RigVMDeveloper.Build.cs
@@ -30,29 +30,9 @@ public class RigVMDeveloper : ModuleRules
 					"MessageLog",
 					"BlueprintGraph",
 					"GraphEditor",
+					"Kismet",
                 }
 			);
-
-            bool bAddBlueprintEditorDependency = true;
-            PublicDefinitions.Add("WITH_RIGVMLEGACYEDITOR=" + (bAddBlueprintEditorDependency ? '1' : '0'));
-            if (bAddBlueprintEditorDependency)
-            {
-	            PrivateDependencyModuleNames.Add("Kismet");
-            }
-            else
-            {
-	            PrivateDependencyModuleNames.AddRange(new string[]
-	            {
-		            "ApplicationCore",
-		            "FieldNotification",
-		            "InputCore",
-		            "Json",
-		            "JsonUtilities",
-		            "ToolWidgets",
-		            "TraceLog",
-		            "WorkspaceMenuStructure"
-	            });
-            }
             
             PrivateIncludePathModuleNames.Add("RigVMEditor");
             DynamicallyLoadedModuleNames.Add("RigVMEditor");
diff --git a/Engine/Plugins/Runtime/RigVM/Source/RigVMEditor/Private/Editor/RigVMActionMenuBuilder.cpp b/Engine/Plugins/Runtime/RigVM/Source/RigVMEditor/Private/Editor/RigVMActionMenuBuilder.cpp
deleted file mode 100644
diff --git a/Engine/Plugins/Runtime/RigVM/Source/RigVMEditor/Private/Editor/RigVMActionMenuItem.cpp b/Engine/Plugins/Runtime/RigVM/Source/RigVMEditor/Private/Editor/RigVMActionMenuItem.cpp
deleted file mode 100644
diff --git a/Engine/Plugins/Runtime/RigVM/Source/RigVMEditor/Private/Editor/RigVMActionMenuUtils.cpp b/Engine/Plugins/Runtime/RigVM/Source/RigVMEditor/Private/Editor/RigVMActionMenuUtils.cpp
deleted file mode 100644
diff --git a/Engine/Plugins/Runtime/RigVM/Source/RigVMEditor/Private/Editor/RigVMEditor.cpp b/Engine/Plugins/Runtime/RigVM/Source/RigVMEditor/Private/Editor/RigVMEditor.cpp
--- a/Engine/Plugins/Runtime/RigVM/Source/RigVMEditor/Private/Editor/RigVMEditor.cpp
+++ b/Engine/Plugins/Runtime/RigVM/Source/RigVMEditor/Private/Editor/RigVMEditor.cpp
@@ -1,54 +1,60 @@
 // Copyright Epic Games, Inc. All Rights Reserved.
 
 #include "Editor/RigVMEditor.h"
-
-#include "BlueprintActionDatabase.h"
-#include "BlueprintCompilationManager.h"
+#include "Editor/RigVMEditorTools.h"
+#include "Editor/RigVMEditorMenuContext.h"
+#include "HAL/PlatformApplicationMisc.h"
+#include "SMyBlueprint.h"
 #include "EdGraphNode_Comment.h"
-#include "EditorModeManager.h"
-#include "EulerTransform.h"
-#include "InstancedPropertyBagStructureDataProvider.h"
-#include "RigVMBlueprint.h"
+#include "RigVMFunctions/RigVMFunction_ControlFlow.h"
+#include "Editor.h"
+#include "Editor/Transactor.h"
+#include "BlueprintEditorTabs.h"
+#include "Kismet2/BlueprintEditorUtils.h"
+#include "Widgets/Notifications/SNotificationList.h"
+#include "Framework/Notifications/NotificationManager.h"
+#include "Editor/RigVMEditorCommands.h"
+#include "BlueprintCompilationManager.h"
+#include "SBlueprintEditorToolbar.h"
+#include "BlueprintEditorModes.h"
+#include "AssetEditorModeManager.h"
+#include "RigVMCore/RigVMMemoryStorageStruct.h"
 #include "RigVMBlueprintUtils.h"
-#include "RigVMEditorCommands.h"
 #include "RigVMPythonUtils.h"
-#include "AssetRegistry/AssetRegistryModule.h"
-#include "AssetRegistry/IAssetRegistry.h"
-#include "Editor/RigVMEditorMenuContext.h"
-#include "Editor/RigVMEditorTools.h"
+#include "EulerTransform.h"
+#include "RigVMEditorModule.h"
 #include "Editor/RigVMGraphDetailCustomization.h"
-#include "Editor/RigVMLegacyEditor.h"
-#include "Editor/RigVMNewEditor.h"
-#include "Editor/Transactor.h"
-#include "Framework/Commands/GenericCommands.h"
-#include "Framework/Notifications/NotificationManager.h"
-#include "Kismet2/BlueprintEditorUtils.h"
 #include "Kismet2/KismetDebugUtilities.h"
 #include "Kismet2/WatchedPin.h"
-#include "RigVMFunctions/RigVMFunction_ControlFlow.h"
+#include "SBlueprintEditorSelectedDebugObjectWidget.h"
+#include "IMessageLogListing.h"
+#include "SGraphPanel.h"
+#include "BlueprintActionDatabase.h"
 #include "RigVMModel/Nodes/RigVMAggregateNode.h"
-#include "Widgets/SRigVMBulkEditDialog.h"
-#include "Widgets/SRigVMGraphBreakLinksWidget.h"
+#include "Widgets/SRigVMGraphFunctionLocalizationWidget.h"
 #include "Widgets/SRigVMGraphFunctionBulkEditWidget.h"
-#include "Widgets/SRigVMSwapAssetReferencesWidget.h"
-#include "Widgets/SRigVMSwapFunctionsWidget.h"
-#include "Widgets/Notifications/SNotificationList.h"
+#include "Widgets/SRigVMGraphBreakLinksWidget.h"
+#include "Widgets/SRigVMGraphChangePinType.h"
+#include "Framework/Commands/GenericCommands.h"
 #include "Stats/StatsHierarchical.h"
-#include "ToolMenuContext.h"
+#include "Widgets/Layout/SBorder.h"
+#include "Styling/AppStyle.h"
 #include "Widgets/Docking/SDockTab.h"
 #include "Widgets/Layout/SScrollBox.h"
-#include "IStructureDetailsView.h"
-
-#if WITH_RIGVMLEGACYEDITOR
-#include "SKismetInspector.h"
-#else
-#include "Editor/SRigVMDetailsInspector.h"
-#endif
+#include "Framework/Docking/TabManager.h"
+#include "ScopedTransaction.h"
+#include "Editor/RigVMEditorMode.h"
+#include "InstancedPropertyBagStructureDataProvider.h"
+#include "Widgets/SRigVMSwapFunctionsWidget.h"
+#include "Widgets/SRigVMSwapAssetReferencesWidget.h"
+#include "Widgets/SRigVMBulkEditDialog.h"
+#include "AssetRegistry/AssetRegistryModule.h"
+#include "ToolMenuContext.h"
+#include "Widgets/SRigVMEditorGraphExplorer.h"
 
 #define LOCTEXT_NAMESPACE "RigVMEditor"
 
-
-FRigVMEditorBase::FRigVMEditorBase()
+FRigVMEditor::FRigVMEditor()
 	: bAnyErrorsLeft(false)
 	, KnownInstructionLimitWarnings()
 	, HaltedAtNode(nullptr)
@@ -65,7 +71,7 @@ FRigVMEditorBase::FRigVMEditorBase()
 {
 }
 
-void FRigVMEditorBase::UnbindEditor()
+FRigVMEditor::~FRigVMEditor()
 {
 	URigVMBlueprint* RigVMBlueprint = GetRigVMBlueprint();
 	RigVMEditorClosedDelegate.Broadcast(this, RigVMBlueprint);
@@ -78,7 +84,7 @@ void FRigVMEditorBase::UnbindEditor()
 	}
 	
 	FEditorDelegates::EndPIE.RemoveAll(this);
-	FEditorDelegates::CancelPIE.RemoveAll(this);
+    FEditorDelegates::CancelPIE.RemoveAll(this);
 
 	if (RigVMBlueprint)
 	{
@@ -122,51 +128,25 @@ void FRigVMEditorBase::UnbindEditor()
 	}
 }
 
-FRigVMEditorBase* FRigVMEditorBase::GetFromAssetEditorInstance(IAssetEditorInstance* Instance)
-{
-	FRigVMEditorBase* Editor = nullptr;
-	FWorkflowCentricApplication* App = static_cast<FWorkflowCentricApplication*>(Instance);
-
-	TSharedRef<FAssetEditorToolkit> SharedApp = App->AsShared();
-#if WITH_RIGVMLEGACYEDITOR
-	if (SharedApp->IsBlueprintEditor())
-	{
-		TSharedPtr<FRigVMLegacyEditor> LegacyEditor = StaticCastSharedPtr<FRigVMLegacyEditor>(App->AsShared().ToSharedPtr());
-		Editor = LegacyEditor.Get();
-	}
-	else
-	{
-		TSharedPtr<FRigVMNewEditor> NewEditor = StaticCastSharedPtr<FRigVMNewEditor>(App->AsShared().ToSharedPtr());
-		Editor = NewEditor.Get();
-	}
-#else
-	TSharedPtr<FRigVMNewEditor> NewEditor = StaticCastSharedPtr<FRigVMNewEditor>(App->AsShared().ToSharedPtr());
-	Editor = NewEditor.Get();
-#endif
-	return Editor;
-}
-
-void FRigVMEditorBase::InitRigVMEditor(const EToolkitMode::Type Mode, const TSharedPtr< class IToolkitHost >& InitToolkitHost, class URigVMBlueprint* InRigVMBlueprint)
+void FRigVMEditor::InitRigVMEditor(const EToolkitMode::Type Mode, const TSharedPtr< class IToolkitHost >& InitToolkitHost, class URigVMBlueprint* InRigVMBlueprint)
 {
 	DECLARE_SCOPE_HIERARCHICAL_COUNTER_FUNC()
 
 	check(InRigVMBlueprint);
 
-#if WITH_RIGVMLEGACYEDITOR
 	FBlueprintCompilationManager::FlushCompilationQueue(nullptr);
-#else
-	// TODO: sara-s remove once blueprint backend is replaced
-	// FRigVMBlueprintCompilationManager::FlushCompilationQueue(nullptr);
-#endif
 
 	Toolbox = SNew(SBorder)
 		.BorderImage(FAppStyle::GetBrush("ToolPanel.GroupBorder"))
 		.Padding(0.f);
 
-	CreateEditorToolbar();
+	if (!Toolbar.IsValid())
+	{
+		Toolbar = MakeShareable(new FBlueprintEditorToolbar(SharedThis(this)));
+	}
 
-	FEditorDelegates::EndPIE.AddRaw(this, &FRigVMEditorBase::OnPIEStopped);
-	FEditorDelegates::CancelPIE.AddRaw(this, &FRigVMEditorBase::OnPIEStopped, false);
+	FEditorDelegates::EndPIE.AddRaw(this, &FRigVMEditor::OnPIEStopped);
+	FEditorDelegates::CancelPIE.AddRaw(this, &FRigVMEditor::OnPIEStopped, false);
 
 	// Build up a list of objects being edited in this asset editor
 	TArray<UObject*> ObjectsBeingEdited;
@@ -176,8 +156,8 @@ void FRigVMEditorBase::InitRigVMEditor(const EToolkitMode::Type Mode, const TSha
 	const bool bCreateDefaultStandaloneMenu = true;
 	const bool bCreateDefaultToolbar = true;
 	InitAssetEditor(Mode, InitToolkitHost, GetEditorAppName(), FTabManager::FLayout::NullLayout, bCreateDefaultStandaloneMenu, bCreateDefaultToolbar, ObjectsBeingEdited);
-	GEditor->GetEditorSubsystem<UAssetEditorSubsystem>()->OnAssetEditorRequestedOpen().AddSP(StaticCastSharedRef<FRigVMEditorBase>(SharedRef()), &FRigVMEditorBase::HandleAssetRequestedOpen);
-	GEditor->GetEditorSubsystem<UAssetEditorSubsystem>()->OnAssetEditorRequestClose().AddSP(StaticCastSharedRef<FRigVMEditorBase>(SharedRef()), &FRigVMEditorBase::HandleAssetRequestClose);
+	GEditor->GetEditorSubsystem<UAssetEditorSubsystem>()->OnAssetEditorRequestedOpen().AddSP(this, &FRigVMEditor::HandleAssetRequestedOpen);
+	GEditor->GetEditorSubsystem<UAssetEditorSubsystem>()->OnAssetEditorRequestClose().AddSP(this, &FRigVMEditor::HandleAssetRequestClose);
 
 	CreateDefaultCommands();
 
@@ -223,10 +203,10 @@ void FRigVMEditorBase::InitRigVMEditor(const EToolkitMode::Type Mode, const TSha
 
 	URigVMBlueprint::sCurrentlyOpenedRigVMBlueprints.AddUnique(InRigVMBlueprint);
 
-	InRigVMBlueprint->OnModified().AddSP(StaticCastSharedRef<FRigVMEditorBase>(SharedRef()), &FRigVMEditorBase::HandleModifiedEvent);
-	InRigVMBlueprint->OnVMCompiled().AddSP(StaticCastSharedRef<FRigVMEditorBase>(SharedRef()), &FRigVMEditorBase::HandleVMCompiledEvent);
-	InRigVMBlueprint->OnRequestInspectObject().AddSP(StaticCastSharedRef<FRigVMEditorBase>(SharedRef()), &FRigVMEditorBase::SetDetailObjects);
-	InRigVMBlueprint->OnRequestInspectMemoryStorage().AddSP(StaticCastSharedRef<FRigVMEditorBase>(SharedRef()), &FRigVMEditorBase::SetMemoryStorageDetails);
+	InRigVMBlueprint->OnModified().AddSP(this, &FRigVMEditor::HandleModifiedEvent);
+	InRigVMBlueprint->OnVMCompiled().AddSP(this, &FRigVMEditor::HandleVMCompiledEvent);
+	InRigVMBlueprint->OnRequestInspectObject().AddSP(this, &FRigVMEditor::SetDetailObjects);
+	InRigVMBlueprint->OnRequestInspectMemoryStorage().AddSP(this, &FRigVMEditor::SetMemoryStorageDetails);
 
 	BindCommands();
 
@@ -249,8 +229,8 @@ void FRigVMEditorBase::InitRigVMEditor(const EToolkitMode::Type Mode, const TSha
 		SetCurrentMode(GetEditorModeName());
 
 		// Activate our edit mode
-		GetToolkitEditorModeManager().SetDefaultMode(GetEditorModeName());
-		GetToolkitEditorModeManager().ActivateMode(GetEditorModeName());
+		GetEditorModeManager().SetDefaultMode(GetEditorModeName());
+		GetEditorModeManager().ActivateMode(GetEditorModeName());
 	}
 
 	{
@@ -278,7 +258,7 @@ void FRigVMEditorBase::InitRigVMEditor(const EToolkitMode::Type Mode, const TSha
 					bBroughtGraphToFront = true;
 				}
 
-				RigGraph->OnGraphNodeClicked.AddSP(StaticCastSharedRef<FRigVMEditorBase>(SharedRef()), &FRigVMEditorBase::OnGraphNodeClicked);
+				RigGraph->OnGraphNodeClicked.AddSP(this, &FRigVMEditor::OnGraphNodeClicked);
 				ActiveTabNodePath = RigGraph->ModelNodePath;
 			}
 		}
@@ -331,19 +311,19 @@ void FRigVMEditorBase::InitRigVMEditor(const EToolkitMode::Type Mode, const TSha
 			}
 		}
 
-		InRigVMBlueprint->OnRefreshEditor().AddSP(StaticCastSharedRef<FRigVMEditorBase>(SharedRef()), &FRigVMEditorBase::HandleRefreshEditorFromBlueprint);
-		InRigVMBlueprint->OnVariableDropped().AddSP(StaticCastSharedRef<FRigVMEditorBase>(SharedRef()), &FRigVMEditorBase::HandleVariableDroppedFromBlueprint);
-		InRigVMBlueprint->OnBreakpointAdded().AddSP(StaticCastSharedRef<FRigVMEditorBase>(SharedRef()), &FRigVMEditorBase::HandleBreakpointAdded);
-
-		InRigVMBlueprint->OnNodeDoubleClicked().AddSP(StaticCastSharedRef<FRigVMEditorBase>(SharedRef()), &FRigVMEditorBase::OnNodeDoubleClicked);
-		InRigVMBlueprint->OnGraphImported().AddSP(StaticCastSharedRef<FRigVMEditorBase>(SharedRef()), &FRigVMEditorBase::OnGraphImported);
-		InRigVMBlueprint->OnRequestLocalizeFunctionDialog().AddSP(StaticCastSharedRef<FRigVMEditorBase>(SharedRef()), &FRigVMEditorBase::OnRequestLocalizeFunctionDialog);
-		InRigVMBlueprint->OnRequestBulkEditDialog().BindSP(StaticCastSharedRef<FRigVMEditorBase>(SharedRef()), &FRigVMEditorBase::OnRequestBulkEditDialog);
-		InRigVMBlueprint->OnRequestBreakLinksDialog().BindSP(StaticCastSharedRef<FRigVMEditorBase>(SharedRef()), &FRigVMEditorBase::OnRequestBreakLinksDialog);
-		InRigVMBlueprint->OnRequestPinTypeSelectionDialog().BindSP(StaticCastSharedRef<FRigVMEditorBase>(SharedRef()), &FRigVMEditorBase::OnRequestPinTypeSelectionDialog);
-		InRigVMBlueprint->OnRequestJumpToHyperlink().BindSP(StaticCastSharedRef<FRigVMEditorBase>(SharedRef()), &FRigVMEditorBase::HandleJumpToHyperlink);
+		InRigVMBlueprint->OnRefreshEditor().AddSP(this, &FRigVMEditor::HandleRefreshEditorFromBlueprint);
+		InRigVMBlueprint->OnVariableDropped().AddSP(this, &FRigVMEditor::HandleVariableDroppedFromBlueprint);
+		InRigVMBlueprint->OnBreakpointAdded().AddSP(this, &FRigVMEditor::HandleBreakpointAdded);
+
+		InRigVMBlueprint->OnNodeDoubleClicked().AddSP(this, &FRigVMEditor::OnNodeDoubleClicked);
+		InRigVMBlueprint->OnGraphImported().AddSP(this, &FRigVMEditor::OnGraphImported);
+		InRigVMBlueprint->OnRequestLocalizeFunctionDialog().AddSP(this, &FRigVMEditor::OnRequestLocalizeFunctionDialog);
+		InRigVMBlueprint->OnRequestBulkEditDialog().BindSP(this, &FRigVMEditor::OnRequestBulkEditDialog);
+		InRigVMBlueprint->OnRequestBreakLinksDialog().BindSP(this, &FRigVMEditor::OnRequestBreakLinksDialog);
+		InRigVMBlueprint->OnRequestPinTypeSelectionDialog().BindSP(this, &FRigVMEditor::OnRequestPinTypeSelectionDialog);
+		InRigVMBlueprint->OnRequestJumpToHyperlink().BindSP(this, &FRigVMEditor::HandleJumpToHyperlink);
 #if WITH_EDITOR
-		InRigVMBlueprint->OnGetFocusedGraph().BindSP(StaticCastSharedRef<FRigVMEditorBase>(SharedRef()), &FRigVMEditorBase::GetFocusedModel);
+		InRigVMBlueprint->OnGetFocusedGraph().BindSP(this, &FRigVMEditor::GetFocusedModel);
 #endif
 	}
 
@@ -389,7 +369,7 @@ void FRigVMEditorBase::InitRigVMEditor(const EToolkitMode::Type Mode, const TSha
 		}
 		
 		// Always show the myblueprint tab
-		static const FTabId MyBlueprintTabId(GetGraphExplorerWidgetID());
+		static const FTabId MyBlueprintTabId(FBlueprintEditorTabs::MyBlueprintID);
 		if (!TabManager->FindExistingLiveTab(MyBlueprintTabId).IsValid())
 		{
 			TabManager->TryInvokeTab(MyBlueprintTabId);
@@ -417,37 +397,23 @@ void FRigVMEditorBase::InitRigVMEditor(const EToolkitMode::Type Mode, const TSha
 
 	for(UScriptStruct* StructToCustomize : StructsToCustomize)
 	{
-#if WITH_RIGVMLEGACYEDITOR
-		if (TSharedPtr<SKismetInspector> Inspector = GetInspector())
-#else
-		if (TSharedPtr<SRigVMDetailsInspector> Inspector = GetInspector())
-#endif
-		{
-			Inspector->GetPropertyView()->RegisterInstancedCustomPropertyTypeLayout(StructToCustomize->GetFName(),
-			   FOnGetPropertyTypeCustomizationInstance::CreateLambda([=]()
-			   {
-				   return FRigVMGraphMathTypeDetailCustomization::MakeInstance();
-			   }));
-		}
+		Inspector->GetPropertyView()->RegisterInstancedCustomPropertyTypeLayout(StructToCustomize->GetFName(),
+			FOnGetPropertyTypeCustomizationInstance::CreateLambda([=]()
+			{
+				return FRigVMGraphMathTypeDetailCustomization::MakeInstance();
+			}));
 	}
 
-#if WITH_RIGVMLEGACYEDITOR
-	if (TSharedPtr<SKismetInspector> Inspector = GetInspector())
-#else
-	if (TSharedPtr<SRigVMDetailsInspector> Inspector = GetInspector())
-#endif
-	{
-		Inspector->GetPropertyView()->RegisterInstancedCustomPropertyTypeLayout(UEnum::StaticClass()->GetFName(),
-		   FOnGetPropertyTypeCustomizationInstance::CreateLambda([=]()
-		   {
-			   return FRigVMGraphEnumDetailCustomization::MakeInstance();
-		   }));
-	}
+	Inspector->GetPropertyView()->RegisterInstancedCustomPropertyTypeLayout(UEnum::StaticClass()->GetFName(),
+		FOnGetPropertyTypeCustomizationInstance::CreateLambda([=]()
+		{
+			return FRigVMGraphEnumDetailCustomization::MakeInstance();
+		}));
 
-	PropertyChangedHandle = FCoreUObjectDelegates::OnObjectPropertyChanged.AddSP(StaticCastSharedRef<FRigVMEditorBase>(SharedRef()), &FRigVMEditorBase::OnPropertyChanged);
+	PropertyChangedHandle = FCoreUObjectDelegates::OnObjectPropertyChanged.AddSP(this, &FRigVMEditor::OnPropertyChanged);
 }
 
-void FRigVMEditorBase::HandleAssetRequestedOpen(UObject* InObject)
+void FRigVMEditor::HandleAssetRequestedOpen(UObject* InObject)
 {
 	if (InObject == GetRigVMBlueprint())
 	{
@@ -455,7 +421,7 @@ void FRigVMEditorBase::HandleAssetRequestedOpen(UObject* InObject)
 	}
 }
 
-void FRigVMEditorBase::HandleAssetRequestClose(UObject* InObject, EAssetEditorCloseReason InReason)
+void FRigVMEditor::HandleAssetRequestClose(UObject* InObject, EAssetEditorCloseReason InReason)
 {
 	if (InObject == GetRigVMBlueprint())
 	{
@@ -463,12 +429,23 @@ void FRigVMEditorBase::HandleAssetRequestClose(UObject* InObject, EAssetEditorCl
 	}
 }
 
-const FName FRigVMEditorBase::GetEditorModeName() const
+const FName FRigVMEditor::GetEditorAppName() const
+{
+	static const FLazyName AppName(TEXT("RigVMEditorApp"));
+	return AppName;
+}
+
+const FName FRigVMEditor::GetEditorModeName() const
 {
 	return FRigVMEditorModes::RigVMEditorMode;
 }
 
-UBlueprint* FRigVMEditorBase::GetBlueprintObj() const
+TSharedPtr<FApplicationMode> FRigVMEditor::CreateEditorMode()
+{
+	return MakeShareable(new FRigVMEditorMode(SharedThis(this)));
+}
+
+UBlueprint* FRigVMEditor::GetBlueprintObj() const
 {
 	const TArray<UObject*>& EditingObjs = GetEditingObjects();
 	for (UObject* Obj : EditingObjs)
@@ -481,7 +458,7 @@ UBlueprint* FRigVMEditorBase::GetBlueprintObj() const
 	return nullptr;
 }
 
-TSubclassOf<UEdGraphSchema> FRigVMEditorBase::GetDefaultSchemaClass() const
+TSubclassOf<UEdGraphSchema> FRigVMEditor::GetDefaultSchemaClass() const
 {
 	if(const URigVMBlueprint* RigVMBlueprint = GetRigVMBlueprint())
 	{
@@ -490,31 +467,34 @@ TSubclassOf<UEdGraphSchema> FRigVMEditorBase::GetDefaultSchemaClass() const
 	return URigVMEdGraphSchema::StaticClass();
 }
 
-bool FRigVMEditorBase::InEditingMode() const
+bool FRigVMEditor::InEditingMode() const
 {
 	// always allow editing - also during PIE.
 	return true;
 }
 
-void FRigVMEditorBase::OnGraphEditorFocused(const TSharedRef<SGraphEditor>& InGraphEditor)
+void FRigVMEditor::OnGraphEditorFocused(const TSharedRef<SGraphEditor>& InGraphEditor)
 {
+	FBlueprintEditor::OnGraphEditorFocused(InGraphEditor);
 	if (GraphExplorerWidget.IsValid())
 	{
 		GraphExplorerWidget->Refresh();
 	}
 }
 
-void FRigVMEditorBase::Tick(float DeltaTime)
+void FRigVMEditor::Tick(float DeltaTime)
 {
+	FBlueprintEditor::Tick(DeltaTime);
+
 	// tick the  rigvm host
 	if (const URigVMBlueprint* Blueprint = GetRigVMBlueprint())
 	{
 		if (LastDebuggedHost != GetCustomDebugObjectLabel(Blueprint->GetObjectBeingDebugged()))
 		{
-			TArray<FRigVMCustomDebugObject> DebugList;
+			TArray<FCustomDebugObject> DebugList;
 			GetCustomDebugObjects(DebugList);
 
-			for (const FRigVMCustomDebugObject& DebugObject : DebugList)
+			for (const FCustomDebugObject& DebugObject : DebugList)
 			{
 				if (DebugObject.NameOverride == LastDebuggedHost)
 				{
@@ -526,46 +506,48 @@ void FRigVMEditorBase::Tick(float DeltaTime)
 	}
 }
 
-void FRigVMEditorBase::BringToolkitToFront()
+void FRigVMEditor::BringToolkitToFront()
 {
-	if (IsHosted())
+	if (ToolkitHost.IsValid())
 	{
-		BringToolkitToFrontImpl();
+		FBlueprintEditor::BringToolkitToFront();
 	}
 }
 
-FName FRigVMEditorBase::GetToolkitFName() const
+FName FRigVMEditor::GetToolkitFName() const
 {
 	return FName("RigVMEditor");
 }
 
-FName FRigVMEditorBase::GetToolkitContextFName() const
+FName FRigVMEditor::GetToolkitContextFName() const
 {
 	return GetToolkitFName();
 }
 
-FText FRigVMEditorBase::GetBaseToolkitName() const
+FText FRigVMEditor::GetBaseToolkitName() const
 {
 	return LOCTEXT("AppLabel", "RigVM Editor");
 }
 
-FText FRigVMEditorBase::GetToolkitToolTipText() const
+FText FRigVMEditor::GetToolkitToolTipText() const
 {
 	return FAssetEditorToolkit::GetToolTipTextForObject(GetBlueprintObj());
 }
 
-FString FRigVMEditorBase::GetWorldCentricTabPrefix() const
+FString FRigVMEditor::GetWorldCentricTabPrefix() const
 {
 	return LOCTEXT("WorldCentricTabPrefix", "RigVM Editor ").ToString();
 }
 
-FLinearColor FRigVMEditorBase::GetWorldCentricTabColorScale() const
+FLinearColor FRigVMEditor::GetWorldCentricTabColorScale() const
 {
 	return FLinearColor( 0.5f, 0.25f, 0.35f, 0.5f );
 }
 
-void FRigVMEditorBase::InitToolMenuContextImpl(FToolMenuContext& MenuContext)
+void FRigVMEditor::InitToolMenuContext(FToolMenuContext& MenuContext)
 {
+	FBlueprintEditor::InitToolMenuContext(MenuContext);
+
 	if (URigVMBlueprint* RigVMBlueprint = GetRigVMBlueprint())
 	{
 		URigVMGraph* Model = nullptr;
@@ -591,13 +573,13 @@ void FRigVMEditorBase::InitToolMenuContextImpl(FToolMenuContext& MenuContext)
 		
 		URigVMEditorMenuContext* RigVMEditorMenuContext = NewObject<URigVMEditorMenuContext>();
 		const FRigVMEditorGraphMenuContext GraphMenuContext = FRigVMEditorGraphMenuContext(Model, Node, Pin);
-		RigVMEditorMenuContext->Init(StaticCastSharedRef<FRigVMEditorBase>(SharedRef()), GraphMenuContext);
+		RigVMEditorMenuContext->Init(SharedThis(this), GraphMenuContext);
 
 		MenuContext.AddObject(RigVMEditorMenuContext);
 	}
 }
 
-bool FRigVMEditorBase::TransactionObjectAffectsBlueprintImpl(UObject* InTransactedObject)
+bool FRigVMEditor::TransactionObjectAffectsBlueprint(UObject* InTransactedObject)
 {
 	URigVMBlueprint* RigVMBlueprint = GetRigVMBlueprint();
 	if (RigVMBlueprint == nullptr)
@@ -609,10 +591,10 @@ bool FRigVMEditorBase::TransactionObjectAffectsBlueprintImpl(UObject* InTransact
 	{
 		return false;
 	}
-	return true;
+	return FBlueprintEditor::TransactionObjectAffectsBlueprint(InTransactedObject);
 }
 
-bool FRigVMEditorBase::CanAddNewLocalVariable() const
+bool FRigVMEditor::CanAddNewLocalVariable() const
 {
 	const URigVMGraph* Graph = GetFocusedModel();
 	const URigVMGraph* ParentGraph = Graph->GetParentGraph();				
@@ -623,7 +605,7 @@ bool FRigVMEditorBase::CanAddNewLocalVariable() const
 	return false;
 }
 
-void FRigVMEditorBase::OnAddNewLocalVariable()
+void FRigVMEditor::OnAddNewLocalVariable()
 {
 	if (!CanAddNewLocalVariable())
 	{
@@ -631,7 +613,7 @@ void FRigVMEditorBase::OnAddNewLocalVariable()
 	}
 
 	FRigVMGraphVariableDescription LastTypeVar;
-	LastTypeVar.ChangeType(GetLastPinTypeUsed());
+	LastTypeVar.ChangeType(MyBlueprintWidget->GetLastPinTypeUsed());
 	FRigVMGraphVariableDescription NewVar = GetFocusedController()->AddLocalVariable(TEXT("NewLocalVar"), LastTypeVar.CPPType, LastTypeVar.CPPTypeObject, LastTypeVar.DefaultValue, true, true);
 	if(NewVar.Name.IsNone())
 	{
@@ -643,7 +625,7 @@ void FRigVMEditorBase::OnAddNewLocalVariable()
 	}
 }
 
-void FRigVMEditorBase::OnPasteNewLocalVariable(const FBPVariableDescription& VariableDescription)
+void FRigVMEditor::OnPasteNewLocalVariable(const FBPVariableDescription& VariableDescription)
 {
 	if (!CanAddNewLocalVariable())
 	{
@@ -659,7 +641,7 @@ void FRigVMEditorBase::OnPasteNewLocalVariable(const FBPVariableDescription& Var
 	}
 }
 
-void FRigVMEditorBase::DeleteSelectedNodes()
+void FRigVMEditor::DeleteSelectedNodes()
 {
 	DECLARE_SCOPE_HIERARCHICAL_COUNTER_FUNC()
 
@@ -720,30 +702,30 @@ void FRigVMEditorBase::DeleteSelectedNodes()
 	GetFocusedController()->CloseUndoBracket();
 }
 
-bool FRigVMEditorBase::CanDeleteNodes() const
+bool FRigVMEditor::CanDeleteNodes() const
 {
 	return true;
 }
 
-void FRigVMEditorBase::CopySelectedNodes()
+void FRigVMEditor::CopySelectedNodes()
 {
 	FString ExportedText = GetFocusedController()->ExportSelectedNodesToText();
 	FPlatformApplicationMisc::ClipboardCopy(*ExportedText);
 }
 
-bool FRigVMEditorBase::CanCopyNodes() const
+bool FRigVMEditor::CanCopyNodes() const
 {
 	return GetFocusedModel()->GetSelectNodes().Num() > 0;
 }
 
-bool FRigVMEditorBase::CanPasteNodes() const
+bool FRigVMEditor::CanPasteNodes() const
 {
 	FString TextToImport;
 	FPlatformApplicationMisc::ClipboardPaste(TextToImport);
 	return GetFocusedController()->CanImportNodesFromText(TextToImport);
 }
 
-FReply FRigVMEditorBase::OnSpawnGraphNodeByShortcut(FInputChord InChord, const FVector2D& InPosition,
+FReply FRigVMEditor::OnSpawnGraphNodeByShortcut(FInputChord InChord, const FVector2D& InPosition,
                                                      UEdGraph* InGraph)
 {
 	if(!InChord.HasAnyModifierKeys())
@@ -763,17 +745,17 @@ FReply FRigVMEditorBase::OnSpawnGraphNodeByShortcut(FInputChord InChord, const F
 	return FReply::Unhandled();
 }
 
-bool FRigVMEditorBase::JumpToHyperlinkImpl(const UObject* ObjectReference, bool bRequestRename)
+void FRigVMEditor::JumpToHyperlink(const UObject* ObjectReference, bool bRequestRename)
 {
 	if(const URigVMEdGraph* Graph = Cast<URigVMEdGraph>(ObjectReference))
 	{
 		OpenGraphAndBringToFront((UEdGraph*)Graph, true);
-		return true;
+		return;
 	}
-	return false;
+	FBlueprintEditor::JumpToHyperlink(ObjectReference, bRequestRename);
 }
 
-void FRigVMEditorBase::AddNewFunctionVariant(const UEdGraph* InOriginalFunction)
+void FRigVMEditor::AddNewFunctionVariant(const UEdGraph* InOriginalFunction)
 {
 #if WITH_EDITOR
 	if (GEditor)
@@ -802,19 +784,21 @@ void FRigVMEditorBase::AddNewFunctionVariant(const UEdGraph* InOriginalFunction)
 	}
 }
 
-void FRigVMEditorBase::PostUndoImpl(bool bSuccess)
+void FRigVMEditor::PostUndo(bool bSuccess)
 {
 	const FTransaction* Transaction = GEditor->Trans->GetTransaction(GEditor->Trans->GetQueueLength() - GEditor->Trans->GetUndoCount());
+	FBlueprintEditor::PostUndo(bSuccess);
 	PostTransaction(bSuccess, Transaction, false);
 }
 
-void FRigVMEditorBase::PostRedoImpl(bool bSuccess)
+void FRigVMEditor::PostRedo(bool bSuccess)
 {
 	const FTransaction* Transaction = GEditor->Trans->GetTransaction(GEditor->Trans->GetQueueLength() - GEditor->Trans->GetUndoCount() - 1);
+	FBlueprintEditor::PostRedo(bSuccess);
 	PostTransaction(bSuccess, Transaction, true);
 }
 
-void FRigVMEditorBase::OnStartWatchingPin()
+void FRigVMEditor::OnStartWatchingPin()
 {
 	if (UEdGraphPin* Pin = GetCurrentlySelectedPin())
 	{
@@ -822,7 +806,7 @@ void FRigVMEditorBase::OnStartWatchingPin()
 	}
 }
 
-bool FRigVMEditorBase::CanStartWatchingPin() const
+bool FRigVMEditor::CanStartWatchingPin() const
 {
 	if (UEdGraphPin* Pin = GetCurrentlySelectedPin())
 	{
@@ -835,7 +819,7 @@ bool FRigVMEditorBase::CanStartWatchingPin() const
 	return false;
 }
 
-void FRigVMEditorBase::OnStopWatchingPin()
+void FRigVMEditor::OnStopWatchingPin()
 {
 	if (UEdGraphPin* Pin = GetCurrentlySelectedPin())
 	{
@@ -843,7 +827,7 @@ void FRigVMEditorBase::OnStopWatchingPin()
 	}
 }
 
-bool FRigVMEditorBase::CanStopWatchingPin() const
+bool FRigVMEditor::CanStopWatchingPin() const
 {
 	if (UEdGraphPin* Pin = GetCurrentlySelectedPin())
 	{
@@ -855,7 +839,7 @@ bool FRigVMEditorBase::CanStopWatchingPin() const
 	return false;
 }
 
-void FRigVMEditorBase::OnToolkitHostingStarted(const TSharedRef<class IToolkit>& Toolkit)
+void FRigVMEditor::OnToolkitHostingStarted(const TSharedRef<class IToolkit>& Toolkit)
 {
 	TSharedPtr<SWidget> InlineContent = Toolkit->GetInlineContent();
 	if (InlineContent.IsValid())
@@ -864,17 +848,17 @@ void FRigVMEditorBase::OnToolkitHostingStarted(const TSharedRef<class IToolkit>&
 	}
 }
 
-void FRigVMEditorBase::OnToolkitHostingFinished(const TSharedRef<class IToolkit>& Toolkit)
+void FRigVMEditor::OnToolkitHostingFinished(const TSharedRef<class IToolkit>& Toolkit)
 {
 	Toolbox->SetContent(SNullWidget::NullWidget);
 }
 
-TStatId FRigVMEditorBase::GetStatId() const
+TStatId FRigVMEditor::GetStatId() const
 {
-	RETURN_QUICK_DECLARE_CYCLE_STAT(FRigVMEditorBase, STATGROUP_Tickables);
+	RETURN_QUICK_DECLARE_CYCLE_STAT(FRigVMEditor, STATGROUP_Tickables);
 }
 
-void FRigVMEditorBase::PasteNodes()
+void FRigVMEditor::PasteNodes()
 {
 	DECLARE_SCOPE_HIERARCHICAL_COUNTER_FUNC()
 
@@ -882,7 +866,7 @@ void FRigVMEditorBase::PasteNodes()
 
 	FVector2D PasteLocation = FSlateApplication::Get().GetCursorPos();
 
-	TSharedPtr<SDockTab> ActiveTab = GetDocumentManager()->GetActiveTab();
+	TSharedPtr<SDockTab> ActiveTab = DocumentManager->GetActiveTab();
 	if (ActiveTab.IsValid())
 	{
 		TSharedPtr<SGraphEditor> GraphEditor = StaticCastSharedRef<SGraphEditor>(ActiveTab->GetContent());
@@ -910,12 +894,12 @@ void FRigVMEditorBase::PasteNodes()
 
 }
 
-URigVMBlueprint* FRigVMEditorBase::GetRigVMBlueprint() const
+URigVMBlueprint* FRigVMEditor::GetRigVMBlueprint() const
 {
 	return Cast<URigVMBlueprint>(GetBlueprintObj());
 }
 
-URigVMHost* FRigVMEditorBase::GetRigVMHost() const
+URigVMHost* FRigVMEditor::GetRigVMHost() const
 {
 	if (URigVMBlueprint* RigVMBlueprint = GetRigVMBlueprint())
 	{
@@ -927,17 +911,17 @@ URigVMHost* FRigVMEditorBase::GetRigVMHost() const
 	return nullptr;
 }
 
-UObject* FRigVMEditorBase::GetOuterForHost() const
+UObject* FRigVMEditor::GetOuterForHost() const
 {
 	return GetRigVMBlueprint();
 }
 
-UClass* FRigVMEditorBase::GetDetailWrapperClass() const
+UClass* FRigVMEditor::GetDetailWrapperClass() const
 {
 	return URigVMDetailsViewWrapperObject::StaticClass();
 }
 
-bool FRigVMEditorBase::SelectLocalVariable(const UEdGraph* Graph, const FName& VariableName)
+bool FRigVMEditor::SelectLocalVariable(const UEdGraph* Graph, const FName& VariableName)
 {
 	if (const URigVMEdGraph* RigVMEdGraph = Cast<URigVMEdGraph>(Graph))
 	{
@@ -949,7 +933,7 @@ bool FRigVMEditorBase::SelectLocalVariable(const UEdGraph* Graph, const FName& V
 				{
 					URigVMDetailsViewWrapperObject* WrapperObject = URigVMDetailsViewWrapperObject::MakeInstance(
 						GetDetailWrapperClass(), GetBlueprintObj(), Variable.StaticStruct(), (uint8*)&Variable, RigVMGraph);
-					WrapperObject->GetWrappedPropertyChangedChainEvent().AddSP(StaticCastSharedRef<FRigVMEditorBase>(SharedRef()), &FRigVMEditorBase::OnWrappedPropertyChangedChainEvent);
+					WrapperObject->GetWrappedPropertyChangedChainEvent().AddSP(this, &FRigVMEditor::OnWrappedPropertyChangedChainEvent);
 					WrapperObject->AddToRoot();
 
 					TArray<UObject*> Objects = {WrapperObject};
@@ -962,26 +946,122 @@ bool FRigVMEditorBase::SelectLocalVariable(const UEdGraph* Graph, const FName& V
 	return false;
 }
 
-void FRigVMEditorBase::CreateDefaultCommandsImpl()
+void FRigVMEditor::CreateDefaultCommands()
 {
 	if (GetBlueprintObj())
 	{
-		// FBlueprintEditor::CreateDefaultCommands();
+		FBlueprintEditor::CreateDefaultCommands();
 	}
 	else
 	{
-		GetToolkitCommands()->MapAction( FGenericCommands::Get().Undo, 
-			FExecuteAction::CreateSP( StaticCastSharedRef<FRigVMEditorBase>(SharedRef()), &FRigVMEditorBase::UndoAction ));
-		GetToolkitCommands()->MapAction( FGenericCommands::Get().Redo, 
-			FExecuteAction::CreateSP( StaticCastSharedRef<FRigVMEditorBase>(SharedRef()), &FRigVMEditorBase::RedoAction ));
+		ToolkitCommands->MapAction( FGenericCommands::Get().Undo, 
+			FExecuteAction::CreateSP( this, &FRigVMEditor::UndoAction ));
+		ToolkitCommands->MapAction( FGenericCommands::Get().Redo, 
+			FExecuteAction::CreateSP( this, &FRigVMEditor::RedoAction ));
 	}
 }
 
-void FRigVMEditorBase::OnCreateGraphEditorCommands(TSharedPtr<FUICommandList> GraphEditorCommandsList)
+void FRigVMEditor::OnCreateGraphEditorCommands(TSharedPtr<FUICommandList> GraphEditorCommandsList)
 {
 }
 
-void FRigVMEditorBase::Compile()
+struct FRigVMEditorZoomLevelsContainer : public FZoomLevelsContainer
+{
+	struct FRigVMEditorZoomLevelEntry
+	{
+	public:
+		FRigVMEditorZoomLevelEntry(float InZoomAmount, const FText& InDisplayText, EGraphRenderingLOD::Type InLOD)
+			: DisplayText(FText::Format(NSLOCTEXT("GraphEditor", "Zoom", "Zoom {0}"), InDisplayText))
+		, ZoomAmount(InZoomAmount)
+		, LOD(InLOD)
+		{
+		}
+
+	public:
+		FText DisplayText;
+		float ZoomAmount;
+		EGraphRenderingLOD::Type LOD;
+	};
+	
+	FRigVMEditorZoomLevelsContainer()
+	{
+		ZoomLevels.Reserve(22);
+		ZoomLevels.Add(FRigVMEditorZoomLevelEntry(0.025f, FText::FromString(TEXT("-14")), EGraphRenderingLOD::LowestDetail));
+		ZoomLevels.Add(FRigVMEditorZoomLevelEntry(0.070f, FText::FromString(TEXT("-13")), EGraphRenderingLOD::LowestDetail));
+		ZoomLevels.Add(FRigVMEditorZoomLevelEntry(0.100f, FText::FromString(TEXT("-12")), EGraphRenderingLOD::LowestDetail));
+		ZoomLevels.Add(FRigVMEditorZoomLevelEntry(0.125f, FText::FromString(TEXT("-11")), EGraphRenderingLOD::LowestDetail));
+		ZoomLevels.Add(FRigVMEditorZoomLevelEntry(0.150f, FText::FromString(TEXT("-10")), EGraphRenderingLOD::LowestDetail));
+		ZoomLevels.Add(FRigVMEditorZoomLevelEntry(0.175f, FText::FromString(TEXT("-9")), EGraphRenderingLOD::LowestDetail));
+		ZoomLevels.Add(FRigVMEditorZoomLevelEntry(0.200f, FText::FromString(TEXT("-8")), EGraphRenderingLOD::LowestDetail));
+		ZoomLevels.Add(FRigVMEditorZoomLevelEntry(0.225f, FText::FromString(TEXT("-7")), EGraphRenderingLOD::LowDetail));
+		ZoomLevels.Add(FRigVMEditorZoomLevelEntry(0.250f, FText::FromString(TEXT("-6")), EGraphRenderingLOD::LowDetail));
+		ZoomLevels.Add(FRigVMEditorZoomLevelEntry(0.375f, FText::FromString(TEXT("-5")), EGraphRenderingLOD::MediumDetail));
+		ZoomLevels.Add(FRigVMEditorZoomLevelEntry(0.500f, FText::FromString(TEXT("-4")), EGraphRenderingLOD::MediumDetail));
+		ZoomLevels.Add(FRigVMEditorZoomLevelEntry(0.675f, FText::FromString(TEXT("-3")), EGraphRenderingLOD::MediumDetail));
+		ZoomLevels.Add(FRigVMEditorZoomLevelEntry(0.750f, FText::FromString(TEXT("-2")), EGraphRenderingLOD::DefaultDetail));
+		ZoomLevels.Add(FRigVMEditorZoomLevelEntry(0.875f, FText::FromString(TEXT("-1")), EGraphRenderingLOD::DefaultDetail));
+		ZoomLevels.Add(FRigVMEditorZoomLevelEntry(1.000f, FText::FromString(TEXT("1:1")), EGraphRenderingLOD::DefaultDetail));
+		ZoomLevels.Add(FRigVMEditorZoomLevelEntry(1.250f, FText::FromString(TEXT("+1")), EGraphRenderingLOD::DefaultDetail));
+		ZoomLevels.Add(FRigVMEditorZoomLevelEntry(1.375f, FText::FromString(TEXT("+2")), EGraphRenderingLOD::DefaultDetail));
+		ZoomLevels.Add(FRigVMEditorZoomLevelEntry(1.500f, FText::FromString(TEXT("+3")), EGraphRenderingLOD::FullyZoomedIn));
+		ZoomLevels.Add(FRigVMEditorZoomLevelEntry(1.675f, FText::FromString(TEXT("+4")), EGraphRenderingLOD::FullyZoomedIn));
+		ZoomLevels.Add(FRigVMEditorZoomLevelEntry(1.750f, FText::FromString(TEXT("+5")), EGraphRenderingLOD::FullyZoomedIn));
+		ZoomLevels.Add(FRigVMEditorZoomLevelEntry(1.875f, FText::FromString(TEXT("+6")), EGraphRenderingLOD::FullyZoomedIn));
+		ZoomLevels.Add(FRigVMEditorZoomLevelEntry(2.000f, FText::FromString(TEXT("+7")), EGraphRenderingLOD::FullyZoomedIn));
+	}
+
+	float GetZoomAmount(int32 InZoomLevel) const override
+	{
+		checkSlow(ZoomLevels.IsValidIndex(InZoomLevel));
+		return ZoomLevels[InZoomLevel].ZoomAmount;
+	}
+
+	int32 GetNearestZoomLevel(float InZoomAmount) const override
+	{
+		for (int32 ZoomLevelIndex=0; ZoomLevelIndex < GetNumZoomLevels(); ++ZoomLevelIndex)
+		{
+			if (InZoomAmount <= GetZoomAmount(ZoomLevelIndex))
+			{
+				return ZoomLevelIndex;
+			}
+		}
+
+		return GetDefaultZoomLevel();
+	}
+	
+	FText GetZoomText(int32 InZoomLevel) const override
+	{
+		checkSlow(ZoomLevels.IsValidIndex(InZoomLevel));
+		return ZoomLevels[InZoomLevel].DisplayText;
+	}
+	
+	int32 GetNumZoomLevels() const override
+	{
+		return ZoomLevels.Num();
+	}
+	
+	int32 GetDefaultZoomLevel() const override
+	{
+		return 14;
+	}
+
+	EGraphRenderingLOD::Type GetLOD(int32 InZoomLevel) const override
+	{
+		checkSlow(ZoomLevels.IsValidIndex(InZoomLevel));
+		return ZoomLevels[InZoomLevel].LOD;
+	}
+
+	TArray<FRigVMEditorZoomLevelEntry> ZoomLevels;
+};
+
+TSharedRef<SGraphEditor> FRigVMEditor::CreateGraphEditorWidget(TSharedRef<FTabInfo> InTabInfo, UEdGraph* InGraph)
+{
+	TSharedRef<SGraphEditor> GraphEditor = FBlueprintEditor::CreateGraphEditorWidget(InTabInfo, InGraph);
+	GraphEditor->GetGraphPanel()->SetZoomLevelsContainer<FRigVMEditorZoomLevelsContainer>();
+	return GraphEditor;
+}
+
+void FRigVMEditor::Compile()
 {
 	{
 		DECLARE_SCOPE_HIERARCHICAL_COUNTER_FUNC()
@@ -1013,7 +1093,7 @@ void FRigVMEditorBase::Compile()
 		SetHost(nullptr);
 		{
 			TGuardValue<bool> GuardCompileReEntry(bIsCompilingThroughUI, true);
-			CompileImpl();
+			FBlueprintEditor::Compile();
 			RigVMBlueprint->InitializeArchetypeInstances();
 			UpdateRigVMHost();
 		}
@@ -1041,10 +1121,10 @@ void FRigVMEditorBase::Compile()
 			}
 		}
 
-		TArray<FRigVMCustomDebugObject> DebugList;
+		TArray<FCustomDebugObject> DebugList;
 		GetCustomDebugObjects(DebugList);
 
-		for (const FRigVMCustomDebugObject& DebugObject : DebugList)
+		for (const FCustomDebugObject& DebugObject : DebugList)
 		{
 			if (DebugObject.NameOverride == LastDebuggedObjectName)
 			{
@@ -1081,7 +1161,7 @@ void FRigVMEditorBase::Compile()
 		if (bAnyVariableValueChanged)
 		{
 			// Go over all the instances to update the default values from CDO
-			for (const FRigVMCustomDebugObject& DebugObject : DebugList)
+			for (const FCustomDebugObject& DebugObject : DebugList)
 			{
 				if (URigVMHost* DebuggedHost = Cast<URigVMHost>(DebugObject.Object))
 				{
@@ -1099,23 +1179,23 @@ void FRigVMEditorBase::Compile()
 	// FStatsHierarchical::DumpMeasurements(LogForMeasurements);
 }
 
-void FRigVMEditorBase::SaveAsset_Execute()
+void FRigVMEditor::SaveAsset_Execute()
 {
 	LastDebuggedHost = GetCustomDebugObjectLabel(GetBlueprintObj()->GetObjectBeingDebugged());
-	SaveAsset_Execute_Impl();
+	FBlueprintEditor::SaveAsset_Execute();
 
 	UpdateRigVMHost();
 }
 
-void FRigVMEditorBase::SaveAssetAs_Execute()
+void FRigVMEditor::SaveAssetAs_Execute()
 {
 	LastDebuggedHost = GetCustomDebugObjectLabel(GetBlueprintObj()->GetObjectBeingDebugged());
-	SaveAssetAs_Execute_Impl();
+	FBlueprintEditor::SaveAssetAs_Execute();
 
 	UpdateRigVMHost();
 }
 
-bool FRigVMEditorBase::IsEditable(UEdGraph* InGraph) const
+bool FRigVMEditor::IsEditable(UEdGraph* InGraph) const
 {
 	if(!IsGraphInCurrentBlueprint(InGraph))
 	{
@@ -1143,20 +1223,20 @@ bool FRigVMEditorBase::IsEditable(UEdGraph* InGraph) const
 		}
 	}
 
-	return IsEditableImpl(InGraph);
+	return FBlueprintEditor::IsEditable(InGraph);
 }
 
-bool FRigVMEditorBase::IsCompilingEnabled() const
+bool FRigVMEditor::IsCompilingEnabled() const
 {
 	return true;
 }
 
-FText FRigVMEditorBase::GetGraphDecorationString(UEdGraph* InGraph) const
+FText FRigVMEditor::GetGraphDecorationString(UEdGraph* InGraph) const
 {
 	return FText::GetEmpty();
 }
 
-void FRigVMEditorBase::OnSelectedNodesChangedImpl(const TSet<UObject*>& NewSelection)
+void FRigVMEditor::OnSelectedNodesChangedImpl(const TSet<UObject*>& NewSelection)
 {
 	DECLARE_SCOPE_HIERARCHICAL_COUNTER_FUNC()
 
@@ -1192,7 +1272,7 @@ void FRigVMEditorBase::OnSelectedNodesChangedImpl(const TSet<UObject*>& NewSelec
 	}
 }
 
-void FRigVMEditorBase::OnBlueprintChangedImpl(UBlueprint* InBlueprint, bool bIsJustBeingCompiled)
+void FRigVMEditor::OnBlueprintChangedImpl(UBlueprint* InBlueprint, bool bIsJustBeingCompiled)
 {
 	DECLARE_SCOPE_HIERARCHICAL_COUNTER_FUNC()
 
@@ -1201,7 +1281,7 @@ void FRigVMEditorBase::OnBlueprintChangedImpl(UBlueprint* InBlueprint, bool bIsJ
 		return;
 	}
 
-	OnBlueprintChangedInnerImpl(InBlueprint, bIsJustBeingCompiled);
+	FBlueprintEditor::OnBlueprintChangedImpl(InBlueprint, bIsJustBeingCompiled);
 
 	if(InBlueprint == GetBlueprintObj())
 	{
@@ -1211,10 +1291,10 @@ void FRigVMEditorBase::OnBlueprintChangedImpl(UBlueprint* InBlueprint, bool bIsJ
 
 			if (!LastDebuggedHost.IsEmpty())
 			{
-				TArray<FRigVMCustomDebugObject> DebugList;
+				TArray<FCustomDebugObject> DebugList;
 				GetCustomDebugObjects(DebugList);
 
-				for (const FRigVMCustomDebugObject& DebugObject : DebugList)
+				for (const FCustomDebugObject& DebugObject : DebugList)
 				{
 					if (DebugObject.NameOverride == LastDebuggedHost)
 					{
@@ -1227,37 +1307,37 @@ void FRigVMEditorBase::OnBlueprintChangedImpl(UBlueprint* InBlueprint, bool bIsJ
 		}
 	}}
 
-void FRigVMEditorBase::RefreshEditors(ERefreshRigVMEditorReason::Type Reason)
+void FRigVMEditor::RefreshEditors(ERefreshBlueprintEditorReason::Type Reason)
 {
-	if(Reason == ERefreshRigVMEditorReason::UnknownReason)
+	if(Reason == ERefreshBlueprintEditorReason::UnknownReason)
 	{
 		// we mark the reason as just compiled since we don't want to
 		// update the graph(s) all the time during compilation
-		Reason = ERefreshRigVMEditorReason::BlueprintCompiled;
+		Reason = ERefreshBlueprintEditorReason::BlueprintCompiled;
 	}
-	RefreshEditorsImpl(Reason);
+	FBlueprintEditor::RefreshEditors(Reason);
 	if (GraphExplorerWidget.IsValid())
 	{
 		GraphExplorerWidget->Refresh();
 	}
 }
 
-void FRigVMEditorBase::SetupGraphEditorEvents(UEdGraph* InGraph, SGraphEditor::FGraphEditorEvents& InEvents)
+void FRigVMEditor::SetupGraphEditorEvents(UEdGraph* InGraph, SGraphEditor::FGraphEditorEvents& InEvents)
 {
 	DECLARE_SCOPE_HIERARCHICAL_COUNTER_FUNC()
 
-	SetupGraphEditorEventsImpl(InGraph, InEvents);
+	FBlueprintEditor::SetupGraphEditorEvents(InGraph, InEvents);
 
-	InEvents.OnCreateActionMenu = SGraphEditor::FOnCreateActionMenu::CreateSP(StaticCastSharedRef<FRigVMEditorBase>(SharedRef()), &FRigVMEditorBase::HandleCreateGraphActionMenu);
-	InEvents.OnTextCommitted = FOnNodeTextCommitted::CreateSP(StaticCastSharedRef<FRigVMEditorBase>(SharedRef()), &FRigVMEditorBase::OnNodeTitleCommitted);
+	InEvents.OnCreateActionMenu = SGraphEditor::FOnCreateActionMenu::CreateSP(this, &FRigVMEditor::HandleCreateGraphActionMenu);
+	InEvents.OnTextCommitted = FOnNodeTextCommitted::CreateSP(this, &FRigVMEditor::OnNodeTitleCommitted);
 }
 
-FActionMenuContent FRigVMEditorBase::HandleCreateGraphActionMenu(UEdGraph* InGraph, const FVector2D& InNodePosition, const TArray<UEdGraphPin*>& InDraggedPins, bool bAutoExpand, SGraphEditor::FActionMenuClosed InOnMenuClosed)
+FActionMenuContent FRigVMEditor::HandleCreateGraphActionMenu(UEdGraph* InGraph, const FVector2D& InNodePosition, const TArray<UEdGraphPin*>& InDraggedPins, bool bAutoExpand, SGraphEditor::FActionMenuClosed InOnMenuClosed)
 {
-	return OnCreateGraphActionMenu(InGraph, InNodePosition, InDraggedPins, bAutoExpand, InOnMenuClosed);
+	return FBlueprintEditor::OnCreateGraphActionMenu(InGraph, InNodePosition, InDraggedPins, bAutoExpand, InOnMenuClosed);
 }
 
-void FRigVMEditorBase::OnNodeTitleCommitted(const FText& NewText, ETextCommit::Type CommitInfo, UEdGraphNode* NodeBeingChanged)
+void FRigVMEditor::OnNodeTitleCommitted(const FText& NewText, ETextCommit::Type CommitInfo, UEdGraphNode* NodeBeingChanged)
 {
 	DECLARE_SCOPE_HIERARCHICAL_COUNTER_FUNC()
 
@@ -1270,98 +1350,96 @@ void FRigVMEditorBase::OnNodeTitleCommitted(const FText& NewText, ETextCommit::T
 	}
 }
 
-void FRigVMEditorBase::FocusInspectorOnGraphSelection(const TSet<UObject*>& NewSelection, bool bForceRefresh)
+void FRigVMEditor::FocusInspectorOnGraphSelection(const TSet<UObject*>& NewSelection, bool bForceRefresh)
 {
 	// nothing to do here
 }
 
-void FRigVMEditorBase::AddReferencedObjects(FReferenceCollector& Collector)
+void FRigVMEditor::AddReferencedObjects(FReferenceCollector& Collector)
 {
-	AddReferencedObjectsImpl(Collector);
+	FBlueprintEditor::AddReferencedObjects(Collector);
 
 	TWeakObjectPtr<URigVMHost> RigVMHost(GetRigVMHost());
 	Collector.AddReferencedObject(RigVMHost);
 }
 
-void FRigVMEditorBase::BindCommands()
+void FRigVMEditor::BindCommands()
 {
 	DECLARE_SCOPE_HIERARCHICAL_COUNTER_FUNC()
 
-	TSharedRef<FRigVMEditorBase> ThisRef = StaticCastSharedRef<FRigVMEditorBase>(SharedRef());
-
 	GetToolkitCommands()->MapAction(
 		FRigVMEditorCommands::Get().AutoCompileGraph,
-		FExecuteAction::CreateSP(ThisRef, &FRigVMEditorBase::ToggleAutoCompileGraph), 
-		FIsActionChecked::CreateSP(ThisRef, &FRigVMEditorBase::CanAutoCompileGraph),
-		FIsActionChecked::CreateSP(ThisRef, &FRigVMEditorBase::IsAutoCompileGraphOn));
+		FExecuteAction::CreateSP(this, &FRigVMEditor::ToggleAutoCompileGraph), 
+		FIsActionChecked::CreateSP(this, &FRigVMEditor::CanAutoCompileGraph),
+		FIsActionChecked::CreateSP(this, &FRigVMEditor::IsAutoCompileGraphOn));
 
 	GetToolkitCommands()->MapAction(
 		FRigVMEditorCommands::Get().ToggleEventQueue,
-		FExecuteAction::CreateSP(ThisRef, &FRigVMEditorBase::ToggleEventQueue),
+		FExecuteAction::CreateSP(this, &FRigVMEditor::ToggleEventQueue),
 		FCanExecuteAction());
 
 	GetToolkitCommands()->MapAction(
 		FRigVMEditorCommands::Get().ToggleExecutionMode,
-		FExecuteAction::CreateSP(ThisRef, &FRigVMEditorBase::ToggleExecutionMode),
+		FExecuteAction::CreateSP(this, &FRigVMEditor::ToggleExecutionMode),
 		FCanExecuteAction());
 
 	GetToolkitCommands()->MapAction(
 		FRigVMEditorCommands::Get().ReleaseMode,
-		FExecuteAction::CreateSP(ThisRef, &FRigVMEditorBase::SetExecutionMode, ERigVMEditorExecutionModeType_Release),
+		FExecuteAction::CreateSP(this, &FRigVMEditor::SetExecutionMode, ERigVMEditorExecutionModeType_Release),
 		FCanExecuteAction());
 
 	GetToolkitCommands()->MapAction(
 		FRigVMEditorCommands::Get().DebugMode,
-		FExecuteAction::CreateSP(ThisRef, &FRigVMEditorBase::SetExecutionMode, ERigVMEditorExecutionModeType_Debug),
+		FExecuteAction::CreateSP(this, &FRigVMEditor::SetExecutionMode, ERigVMEditorExecutionModeType_Debug),
 		FCanExecuteAction());
 
 	GetToolkitCommands()->MapAction(
 		FRigVMEditorCommands::Get().ResumeExecution,
-		FExecuteAction::CreateSP(ThisRef, &FRigVMEditorBase::HandleBreakpointActionRequested, ERigVMBreakpointAction::Resume),			
-		FIsActionChecked::CreateSP(ThisRef, &FRigVMEditorBase::IsHaltedAtBreakpoint));
+		FExecuteAction::CreateSP(this, &FRigVMEditor::HandleBreakpointActionRequested, ERigVMBreakpointAction::Resume),			
+		FIsActionChecked::CreateSP(this, &FRigVMEditor::IsHaltedAtBreakpoint));
 
 	GetToolkitCommands()->MapAction(
 		FRigVMEditorCommands::Get().ShowCurrentStatement,
-		FExecuteAction::CreateSP(ThisRef, &FRigVMEditorBase::HandleShowCurrentStatement),
-		FIsActionChecked::CreateSP(ThisRef, &FRigVMEditorBase::IsHaltedAtBreakpoint));
+		FExecuteAction::CreateSP(this, &FRigVMEditor::HandleShowCurrentStatement),
+		FIsActionChecked::CreateSP(this, &FRigVMEditor::IsHaltedAtBreakpoint));
 
 	GetToolkitCommands()->MapAction(
 		FRigVMEditorCommands::Get().StepOver,
-		FExecuteAction::CreateSP(ThisRef, &FRigVMEditorBase::HandleBreakpointActionRequested, ERigVMBreakpointAction::StepOver),
-		FIsActionChecked::CreateSP(ThisRef, &FRigVMEditorBase::IsHaltedAtBreakpoint));
+		FExecuteAction::CreateSP(this, &FRigVMEditor::HandleBreakpointActionRequested, ERigVMBreakpointAction::StepOver),
+		FIsActionChecked::CreateSP(this, &FRigVMEditor::IsHaltedAtBreakpoint));
 
 	GetToolkitCommands()->MapAction(
 		FRigVMEditorCommands::Get().StepInto,
-		FExecuteAction::CreateSP(ThisRef, &FRigVMEditorBase::HandleBreakpointActionRequested, ERigVMBreakpointAction::StepInto),
-		FIsActionChecked::CreateSP(ThisRef, &FRigVMEditorBase::IsHaltedAtBreakpoint));
+		FExecuteAction::CreateSP(this, &FRigVMEditor::HandleBreakpointActionRequested, ERigVMBreakpointAction::StepInto),
+		FIsActionChecked::CreateSP(this, &FRigVMEditor::IsHaltedAtBreakpoint));
 
 	GetToolkitCommands()->MapAction(
 		FRigVMEditorCommands::Get().StepOut,
-		FExecuteAction::CreateSP(ThisRef, &FRigVMEditorBase::HandleBreakpointActionRequested, ERigVMBreakpointAction::StepOut),
-		FIsActionChecked::CreateSP(ThisRef, &FRigVMEditorBase::IsHaltedAtBreakpoint));
+		FExecuteAction::CreateSP(this, &FRigVMEditor::HandleBreakpointActionRequested, ERigVMBreakpointAction::StepOut),
+		FIsActionChecked::CreateSP(this, &FRigVMEditor::IsHaltedAtBreakpoint));
 
 	GetToolkitCommands()->MapAction(
 		FRigVMEditorCommands::Get().FrameSelection,
-		FExecuteAction::CreateSP(ThisRef, &FRigVMEditorBase::FrameSelection),
+		FExecuteAction::CreateSP(this, &FRigVMEditor::FrameSelection),
 		FCanExecuteAction());
 
 	GetToolkitCommands()->MapAction(
 		FRigVMEditorCommands::Get().SwapFunctionWithinAsset,
-		FExecuteAction::CreateSP(ThisRef, &FRigVMEditorBase::SwapFunctionWithinAsset),
+		FExecuteAction::CreateSP(this, &FRigVMEditor::SwapFunctionWithinAsset),
 		FCanExecuteAction());
 
 	GetToolkitCommands()->MapAction(
 		FRigVMEditorCommands::Get().SwapFunctionAcrossProject,
-		FExecuteAction::CreateSP(ThisRef, &FRigVMEditorBase::SwapFunctionAcrossProject),
+		FExecuteAction::CreateSP(this, &FRigVMEditor::SwapFunctionAcrossProject),
 		FCanExecuteAction());
 
 	GetToolkitCommands()->MapAction(
 		FRigVMEditorCommands::Get().SwapAssetReferences,
-		FExecuteAction::CreateSP(ThisRef, &FRigVMEditorBase::SwapAssetReferences),
+		FExecuteAction::CreateSP(this, &FRigVMEditor::SwapAssetReferences),
 		FCanExecuteAction());
 }
 
-void FRigVMEditorBase::ToggleAutoCompileGraph()
+void FRigVMEditor::ToggleAutoCompileGraph()
 {
 	if (URigVMBlueprint* RigVMBlueprint = GetRigVMBlueprint())
 	{
@@ -1373,7 +1451,7 @@ void FRigVMEditorBase::ToggleAutoCompileGraph()
 	}
 }
 
-bool FRigVMEditorBase::IsAutoCompileGraphOn() const
+bool FRigVMEditor::IsAutoCompileGraphOn() const
 {
 	if (URigVMBlueprint* RigVMBlueprint = GetRigVMBlueprint())
 	{
@@ -1382,30 +1460,30 @@ bool FRigVMEditorBase::IsAutoCompileGraphOn() const
 	return false;
 }
 
-void FRigVMEditorBase::ToggleEventQueue()
+void FRigVMEditor::ToggleEventQueue()
 {
 	SetEventQueue(LastEventQueue);
 }
 
-void FRigVMEditorBase::ToggleExecutionMode()
+void FRigVMEditor::ToggleExecutionMode()
 {
 	SetExecutionMode((ExecutionMode == ERigVMEditorExecutionModeType_Debug) ?
 		ERigVMEditorExecutionModeType_Release
 		: ERigVMEditorExecutionModeType_Debug);
 }
 
-TSharedRef<SWidget> FRigVMEditorBase::GenerateEventQueueMenuContent()
+TSharedRef<SWidget> FRigVMEditor::GenerateEventQueueMenuContent()
 {
 	FMenuBuilder MenuBuilder(true, GetToolkitCommands());
 	GenerateEventQueueMenuContent(MenuBuilder);
 	return MenuBuilder.MakeWidget();
 }
 
-void FRigVMEditorBase::GenerateEventQueueMenuContent(FMenuBuilder& MenuBuilder)
+void FRigVMEditor::GenerateEventQueueMenuContent(FMenuBuilder& MenuBuilder)
 {
 }
 
-TSharedRef<SWidget> FRigVMEditorBase::GenerateExecutionModeMenuContent()
+TSharedRef<SWidget> FRigVMEditor::GenerateExecutionModeMenuContent()
 {
 	FMenuBuilder MenuBuilder(true, GetToolkitCommands());
 	MenuBuilder.BeginSection(TEXT("Events"));
@@ -1415,7 +1493,7 @@ TSharedRef<SWidget> FRigVMEditorBase::GenerateExecutionModeMenuContent()
 	return MenuBuilder.MakeWidget();
 }
 
-FMenuBuilder FRigVMEditorBase::GenerateBulkEditMenu()
+FMenuBuilder FRigVMEditor::GenerateBulkEditMenu()
 {
 	FMenuBuilder MenuBuilder(true, GetToolkitCommands());
 	MenuBuilder.BeginSection(TEXT("Functions"), LOCTEXT("Functions", "Functions"));
@@ -1428,41 +1506,38 @@ FMenuBuilder FRigVMEditorBase::GenerateBulkEditMenu()
 	return MenuBuilder;
 }
 
-TSharedRef<SWidget> FRigVMEditorBase::GenerateBulkEditMenuContent()
+TSharedRef<SWidget> FRigVMEditor::GenerateBulkEditMenuContent()
 {
 	FMenuBuilder MenuBuilder = GenerateBulkEditMenu();
 	return MenuBuilder.MakeWidget();
 }
 
-void FRigVMEditorBase::OnActiveTabChanged( TSharedPtr<SDockTab> PreviouslyActive, TSharedPtr<SDockTab> NewlyActivated )
+void FRigVMEditor::OnActiveTabChanged( TSharedPtr<SDockTab> PreviouslyActive, TSharedPtr<SDockTab> NewlyActivated )
 {
 	if (!NewlyActivated.IsValid())
 	{
 		TArray<UObject*> ObjArray;
-#if WITH_RIGVMLEGACYEDITOR
-		if (TSharedPtr<SKismetInspector> Inspector = GetInspector())
-#else
-		if (TSharedPtr<SRigVMDetailsInspector> Inspector = GetInspector())
-#endif
-		{
-			Inspector->ShowDetailsForObjects(ObjArray);
-		}
+		Inspector->ShowDetailsForObjects(ObjArray);
+	}
+	else 
+	{
+		FBlueprintEditor::OnActiveTabChanged(PreviouslyActive, NewlyActivated);
 	}
 }
 
-void FRigVMEditorBase::UndoAction()
+void FRigVMEditor::UndoAction()
 {
 	GEditor->UndoTransaction();
 }
 
-void FRigVMEditorBase::RedoAction()
+void FRigVMEditor::RedoAction()
 {
 	GEditor->RedoTransaction();
 }
 
-void FRigVMEditorBase::NewDocument_OnClicked(ECreatedDocumentType GraphType)
+void FRigVMEditor::NewDocument_OnClicked(ECreatedDocumentType GraphType)
 {
-	if (GraphType == FRigVMNewEditor::CGT_NewFunctionGraph)
+	if (GraphType == FBlueprintEditor::CGT_NewFunctionGraph)
 	{
 		if (URigVMBlueprint* RigVMBlueprint = GetRigVMBlueprint())
 		{
@@ -1479,7 +1554,7 @@ void FRigVMEditorBase::NewDocument_OnClicked(ECreatedDocumentType GraphType)
 			}
 		}
 	}
-	else if(GraphType == FRigVMNewEditor::CGT_NewEventGraph)
+	else if(GraphType == FBlueprintEditor::CGT_NewEventGraph)
 	{
 		if (URigVMBlueprint* RigVMBlueprint = GetRigVMBlueprint())
 		{
@@ -1498,21 +1573,21 @@ void FRigVMEditorBase::NewDocument_OnClicked(ECreatedDocumentType GraphType)
 	}
 }
 
-bool FRigVMEditorBase::IsSectionVisibleImpl(RigVMNodeSectionID::Type InSectionID) const
+bool FRigVMEditor::IsSectionVisible(NodeSectionID::Type InSectionID) const
 {
 	switch (InSectionID)
 	{
-		case RigVMNodeSectionID::GRAPH:
-		case RigVMNodeSectionID::VARIABLE:
-		case RigVMNodeSectionID::FUNCTION:
+		case NodeSectionID::GRAPH:
+		case NodeSectionID::VARIABLE:
+		case NodeSectionID::FUNCTION:
 		{
 			return true;
 		}
-		case RigVMNodeSectionID::LOCAL_VARIABLE:
+		case NodeSectionID::LOCAL_VARIABLE:
 		{
 			if(const URigVMGraph* Graph = GetFocusedModel())
 			{
-				const URigVMGraph* ParentGraph = Graph->GetParentGraph();
+				const URigVMGraph* ParentGraph = Graph->GetParentGraph();				
 				if (ParentGraph && ParentGraph->IsA<URigVMFunctionLibrary>())
 				{
 					return true;
@@ -1527,34 +1602,34 @@ bool FRigVMEditorBase::IsSectionVisibleImpl(RigVMNodeSectionID::Type InSectionID
 	return false;
 }
 
-bool FRigVMEditorBase::AreEventGraphsAllowed() const
+bool FRigVMEditor::AreEventGraphsAllowed() const
 {
 	if(const URigVMBlueprint* RigVMBlueprint = GetRigVMBlueprint())
 	{
 		return RigVMBlueprint->SupportsEventGraphs();
 	}
-	return false;
+	return FBlueprintEditor::AreEventGraphsAllowed();
 }
 
-bool FRigVMEditorBase::AreMacrosAllowed() const
+bool FRigVMEditor::AreMacrosAllowed() const
 {
 	if(const URigVMBlueprint* RigVMBlueprint = GetRigVMBlueprint())
 	{
 		return RigVMBlueprint->SupportsMacros();
 	}
-	return false;
+	return FBlueprintEditor::AreMacrosAllowed();
 }
 
-bool FRigVMEditorBase::AreDelegatesAllowed() const
+bool FRigVMEditor::AreDelegatesAllowed() const
 {
 	if(const URigVMBlueprint* RigVMBlueprint = GetRigVMBlueprint())
 	{
 		return RigVMBlueprint->SupportsDelegates();
 	}
-	return false;
+	return FBlueprintEditor::AreDelegatesAllowed();
 }
 
-bool FRigVMEditorBase::NewDocument_IsVisibleForTypeImpl(ECreatedDocumentType GraphType) const
+bool FRigVMEditor::NewDocument_IsVisibleForType(ECreatedDocumentType GraphType) const
 {
 	switch(GraphType)
 	{
@@ -1568,12 +1643,12 @@ bool FRigVMEditorBase::NewDocument_IsVisibleForTypeImpl(ECreatedDocumentType Gra
 			break;
 		}
 	}
-	return true;
+	return FBlueprintEditor::NewDocument_IsVisibleForType(GraphType);
 }
 
-FGraphAppearanceInfo FRigVMEditorBase::GetGraphAppearance(UEdGraph* InGraph) const
+FGraphAppearanceInfo FRigVMEditor::GetGraphAppearance(UEdGraph* InGraph) const
 {
-	FGraphAppearanceInfo AppearanceInfo = GetGraphAppearanceImpl(InGraph);
+	FGraphAppearanceInfo AppearanceInfo = FBlueprintEditor::GetGraphAppearance(InGraph);
 
 	if (const URigVMBlueprint* RigVMBlueprint = GetRigVMBlueprint())
 	{
@@ -1603,7 +1678,7 @@ FGraphAppearanceInfo FRigVMEditorBase::GetGraphAppearance(UEdGraph* InGraph) con
 	return AppearanceInfo;
 }
 
-void FRigVMEditorBase::HandleModifiedEvent(ERigVMGraphNotifType InNotifType, URigVMGraph* InGraph, UObject* InSubject)
+void FRigVMEditor::HandleModifiedEvent(ERigVMGraphNotifType InNotifType, URigVMGraph* InGraph, UObject* InSubject)
 {
 	DECLARE_SCOPE_HIERARCHICAL_COUNTER_FUNC()
 
@@ -1682,7 +1757,7 @@ void FRigVMEditorBase::HandleModifiedEvent(ERigVMGraphNotifType InNotifType, URi
 										// we are ok with not reacting to errors here
 										if(Property && PropertyStorage)
 										{
-											FRigVMPinDefaultValueImportErrorContext ErrorPipe;
+											FRigVMPinDefaultValueImportErrorContext ErrorPipe;										
 											Property->ImportText_Direct(*DefaultValue, PropertyStorage, nullptr, PPF_None, &ErrorPipe);
 										}
 									}
@@ -1767,12 +1842,12 @@ public:
 	ERigVMMemoryType MemoryType;
 };
 
-void FRigVMEditorBase::HandleVMCompiledEvent(UObject* InCompiledObject, URigVM* InVM, FRigVMExtendedExecuteContext& InContext)
+void FRigVMEditor::HandleVMCompiledEvent(UObject* InCompiledObject, URigVM* InVM, FRigVMExtendedExecuteContext& InContext)
 {
 	if(URigVMBlueprint* RigVMBlueprint = Cast<URigVMBlueprint>(InCompiledObject))
 	{
-		GetCompilerResultsListing()->ClearMessages();
-		GetCompilerResultsListing()->AddMessages(RigVMBlueprint->GetCompileLog().Messages);
+		CompilerResultsListing->ClearMessages();
+		CompilerResultsListing->AddMessages(RigVMBlueprint->GetCompileLog().Messages);
 		RigVMBlueprint->GetCompileLog().Messages.Reset();
 		RigVMBlueprint->GetCompileLog().NumErrors = RigVMBlueprint->GetCompileLog().NumWarnings = 0;
 	}
@@ -1811,7 +1886,7 @@ void FRigVMEditorBase::HandleVMCompiledEvent(UObject* InCompiledObject, URigVM*
 	UpdateGraphCompilerErrors();
 }
 
-void FRigVMEditorBase::HandleVMExecutedEvent(URigVMHost* InHost, const FName& InEventName)
+void FRigVMEditor::HandleVMExecutedEvent(URigVMHost* InHost, const FName& InEventName)
 {
 	if (URigVMBlueprint* RigVMBlueprint = GetRigVMBlueprint())
 	{
@@ -1924,7 +1999,7 @@ void FRigVMEditorBase::HandleVMExecutedEvent(URigVMHost* InHost, const FName& In
 	UpdateGraphCompilerErrors();
 }
 
-void FRigVMEditorBase::HandleVMExecutionHalted(const int32 InstructionIndex, UObject* InNodeObject, const FName& InEntryName)
+void FRigVMEditor::HandleVMExecutionHalted(const int32 InstructionIndex, UObject* InNodeObject, const FName& InEntryName)
 {
 	if (HaltedAtNode == InNodeObject)
 	{
@@ -1954,21 +2029,14 @@ void FRigVMEditorBase::HandleVMExecutionHalted(const int32 InstructionIndex, UOb
 	}
 	else 
 	{
-		if (URigVMHost* RigVMHost = GetRigVMHost())
+		if (InEntryName == GetRigVMHost()->GetEventQueue().Last())
 		{
-			const TArray<FName>& EventQueue = RigVMHost->GetEventQueue();
-			if (!EventQueue.IsEmpty())
-			{
-				if (InEntryName == EventQueue.Last())
-				{
-					SetHaltedNode(nullptr);
-				}
-			}
+			SetHaltedNode(nullptr);
 		}
 	}
 }
 
-void FRigVMEditorBase::SetHaltedNode(URigVMNode* Node)
+void FRigVMEditor::SetHaltedNode(URigVMNode* Node)
 {
 	if (HaltedAtNode)
 	{
@@ -1981,9 +2049,9 @@ void FRigVMEditorBase::SetHaltedNode(URigVMNode* Node)
 	}
 }
 
-void FRigVMEditorBase::NotifyPreChange(FProperty* PropertyAboutToChange)
+void FRigVMEditor::NotifyPreChange(FProperty* PropertyAboutToChange)
 {
-	NotifyPreChangeImpl(PropertyAboutToChange);
+	FBlueprintEditor::NotifyPreChange(PropertyAboutToChange);
 
 	if (URigVMBlueprint* RigVMBlueprint = GetRigVMBlueprint())
 	{
@@ -1991,9 +2059,9 @@ void FRigVMEditorBase::NotifyPreChange(FProperty* PropertyAboutToChange)
 	}
 }
 
-void FRigVMEditorBase::NotifyPostChange(const FPropertyChangedEvent& PropertyChangedEvent, FProperty* PropertyThatChanged)
+void FRigVMEditor::NotifyPostChange(const FPropertyChangedEvent& PropertyChangedEvent, FProperty* PropertyThatChanged)
 {
-	NotifyPostChangeImpl(PropertyChangedEvent, PropertyThatChanged);
+	FBlueprintEditor::NotifyPostChange(PropertyChangedEvent, PropertyThatChanged);
 
 	// we need to listen to changes for variables on the blueprint here since
 	// OnFinishedChangingProperties is called only for top level property changes.
@@ -2020,7 +2088,7 @@ void FRigVMEditorBase::NotifyPostChange(const FPropertyChangedEvent& PropertyCha
 	}
 }
 
-void FRigVMEditorBase::OnFinishedChangingProperties(const FPropertyChangedEvent& PropertyChangedEvent)
+void FRigVMEditor::OnFinishedChangingProperties(const FPropertyChangedEvent& PropertyChangedEvent)
 {
 	DECLARE_SCOPE_HIERARCHICAL_COUNTER_FUNC()
 
@@ -2041,7 +2109,7 @@ void FRigVMEditorBase::OnFinishedChangingProperties(const FPropertyChangedEvent&
 	}
 }
 
-void FRigVMEditorBase::OnPropertyChanged(UObject* InObject, FPropertyChangedEvent& InEvent)
+void FRigVMEditor::OnPropertyChanged(UObject* InObject, FPropertyChangedEvent& InEvent)
 {
 	URigVMBlueprint* RigVMBlueprint = GetRigVMBlueprint();
 
@@ -2051,12 +2119,12 @@ void FRigVMEditorBase::OnPropertyChanged(UObject* InObject, FPropertyChangedEven
 		if(InEvent.MemberProperty == RigVMBlueprint->GetClass()->FindPropertyByName(GET_MEMBER_NAME_CHECKED(URigVMBlueprint, RigVMClient)) ||
 			InEvent.MemberProperty == RigVMBlueprint->GetClass()->FindPropertyByName(GET_MEMBER_NAME_CHECKED(URigVMBlueprint, UbergraphPages)))
 		{
-			GetDocumentManager()->CleanInvalidTabs();
+			DocumentManager->CleanInvalidTabs();
 		}
 	}
 }
 
-void FRigVMEditorBase::OnWrappedPropertyChangedChainEvent(URigVMDetailsViewWrapperObject* InWrapperObject,
+void FRigVMEditor::OnWrappedPropertyChangedChainEvent(URigVMDetailsViewWrapperObject* InWrapperObject,
 	const FString& InPropertyPath, FPropertyChangedChainEvent& InPropertyChangedChainEvent)
 {
 	check(InWrapperObject);
@@ -2076,22 +2144,25 @@ void FRigVMEditorBase::OnWrappedPropertyChangedChainEvent(URigVMDetailsViewWrapp
 			const FRigVMGraphVariableDescription VariableDescription = InWrapperObject->GetContent<FRigVMGraphVariableDescription>();
 			URigVMGraph* Graph = CastChecked<URigVMGraph>(InWrapperObject->GetSubject());
 			URigVMController* Controller = RigVMBlueprint->GetController(Graph);
-			if (PropertyPath == TEXT("Name"))
+			if (PropertyPath == TEXT("Name") && MyBlueprintWidget.IsValid())
 			{
-				const FName OldVariableName = GetSelectedVariableName();
-				if (!OldVariableName.IsNone())
+				if (FEdGraphSchemaAction_BlueprintVariableBase* VariableAcion = MyBlueprintWidget->SelectionAsBlueprintVariable())
 				{
-					for (FRigVMGraphVariableDescription& Variable : Graph->GetLocalVariables())
+					const FName OldVariableName = VariableAcion->GetVariableName();
+					if (!OldVariableName.IsNone())
 					{
-						if (Variable.Name == OldVariableName)
+						for (FRigVMGraphVariableDescription& Variable : Graph->GetLocalVariables())
 						{
-							Controller->RenameLocalVariable(OldVariableName, VariableDescription.Name);
-							break;
+							if (Variable.Name == OldVariableName)
+							{
+								Controller->RenameLocalVariable(OldVariableName, VariableDescription.Name);
+								break;
+							}
 						}
 					}
 				}
-
-				RefreshEditors();
+				RefreshMyBlueprint();
+				GraphExplorerWidget->Refresh();
 				GetRigVMBlueprint()->RequestAutoVMRecompilation();
 			}
 			else if (PropertyPath == TEXT("CPPType") || PropertyPath == TEXT("CPPTypeObject"))
@@ -2254,7 +2325,7 @@ void FRigVMEditorBase::OnWrappedPropertyChangedChainEvent(URigVMDetailsViewWrapp
 	}
 }
 
-void FRigVMEditorBase::OnRequestLocalizeFunctionDialog(FRigVMGraphFunctionIdentifier& InFunction,
+void FRigVMEditor::OnRequestLocalizeFunctionDialog(FRigVMGraphFunctionIdentifier& InFunction,
 	URigVMController* InTargetController,
 	IRigVMGraphFunctionHost* InTargetFunctionHost,
 	bool bForce)
@@ -2262,7 +2333,7 @@ void FRigVMEditorBase::OnRequestLocalizeFunctionDialog(FRigVMGraphFunctionIdenti
 	UE::RigVM::Editor::Tools::OnRequestLocalizeFunctionDialog(InFunction, InTargetController, InTargetFunctionHost, bForce);
 }
 
-FRigVMController_BulkEditResult FRigVMEditorBase::OnRequestBulkEditDialog(URigVMBlueprint* InBlueprint,
+FRigVMController_BulkEditResult FRigVMEditor::OnRequestBulkEditDialog(URigVMBlueprint* InBlueprint,
 	URigVMController* InController, URigVMLibraryNode* InFunction, ERigVMControllerBulkEditType InEditType)
 {
 	if (bAllowBulkEdits)
@@ -2297,7 +2368,7 @@ FRigVMController_BulkEditResult FRigVMEditorBase::OnRequestBulkEditDialog(URigVM
 	return Result;
 }
 
-bool FRigVMEditorBase::OnRequestBreakLinksDialog(TArray<URigVMLink*> InLinks)
+bool FRigVMEditor::OnRequestBreakLinksDialog(TArray<URigVMLink*> InLinks)
 {
 	if(InLinks.Num() == 0)
 	{
@@ -2314,7 +2385,7 @@ bool FRigVMEditorBase::OnRequestBreakLinksDialog(TArray<URigVMLink*> InLinks)
 	return BreakLinksDialog->ShowModal() == EAppReturnType::Ok; 
 }
 
-TRigVMTypeIndex FRigVMEditorBase::OnRequestPinTypeSelectionDialog(const TArray<TRigVMTypeIndex>& InTypes)
+TRigVMTypeIndex FRigVMEditor::OnRequestPinTypeSelectionDialog(const TArray<TRigVMTypeIndex>& InTypes)
 {
 		if(InTypes.Num() == 0)
 	{
@@ -2429,7 +2500,7 @@ TRigVMTypeIndex FRigVMEditorBase::OnRequestPinTypeSelectionDialog(const TArray<T
 	return Answer;
 }
 
-void FRigVMEditorBase::HandleJumpToHyperlink(const UObject* InSubject)
+void FRigVMEditor::HandleJumpToHyperlink(const UObject* InSubject)
 {
 	URigVMBlueprint* RigBlueprint = GetRigVMBlueprint();
 	if(RigBlueprint == nullptr)
@@ -2478,10 +2549,10 @@ void FRigVMEditorBase::HandleJumpToHyperlink(const UObject* InSubject)
 			{
 				if (GEditor->GetEditorSubsystem<UAssetEditorSubsystem>()->OpenEditorForAsset(OtherBlueprint))
 				{
-					IAssetEditorInstance* OtherEditor = GEditor->GetEditorSubsystem<UAssetEditorSubsystem>()->FindEditorForAsset(OtherBlueprint, /*bFocusIfOpen =*/true);
-					if (FRigVMEditorBase* OtherRigVMEditor = FRigVMEditorBase::GetFromAssetEditorInstance(OtherEditor))
+					FRigVMEditor* OtherEditor = static_cast<FRigVMEditor*>(GEditor->GetEditorSubsystem<UAssetEditorSubsystem>()->FindEditorForAsset(OtherBlueprint, /*bFocusIfOpen =*/true));
+					if(OtherEditor)
 					{
-						OtherRigVMEditor->HandleJumpToHyperlink(NodeToJumpTo);
+						OtherEditor->HandleJumpToHyperlink(NodeToJumpTo);
 						return;
 					}
 				}
@@ -2511,7 +2582,7 @@ void FRigVMEditorBase::HandleJumpToHyperlink(const UObject* InSubject)
 	}
 }
 
-bool FRigVMEditorBase::UpdateDefaultValueForVariable(FBPVariableDescription& InVariable, bool bUseCDO)
+bool FRigVMEditor::UpdateDefaultValueForVariable(FBPVariableDescription& InVariable, bool bUseCDO)
 {
 	bool bAnyValueChanged = false;
 	if (URigVMBlueprint* RigVMBlueprint = GetRigVMBlueprint())
@@ -2538,7 +2609,7 @@ bool FRigVMEditorBase::UpdateDefaultValueForVariable(FBPVariableDescription& InV
 	return bAnyValueChanged;
 }
 
-void FRigVMEditorBase::UpdateRigVMHost()
+void FRigVMEditor::UpdateRigVMHost()
 {
 	DECLARE_SCOPE_HIERARCHICAL_COUNTER_FUNC()
 
@@ -2591,14 +2662,14 @@ void FRigVMEditorBase::UpdateRigVMHost()
 			Blueprint->ClearFlags(RF_Transient);
 		}
 
-		RigVMHost->OnInitialized_AnyThread().AddSP(StaticCastSharedRef<FRigVMEditorBase>(SharedRef()), &FRigVMEditorBase::HandleVMExecutedEvent);
-		RigVMHost->OnExecuted_AnyThread().AddSP(StaticCastSharedRef<FRigVMEditorBase>(SharedRef()), &FRigVMEditorBase::HandleVMExecutedEvent);
+		RigVMHost->OnInitialized_AnyThread().AddSP(this, &FRigVMEditor::HandleVMExecutedEvent);
+		RigVMHost->OnExecuted_AnyThread().AddSP(this, &FRigVMEditor::HandleVMExecutedEvent);
 		RigVMHost->RequestInit();
-		RigVMHost->GetDebugInfo().ExecutionHalted().AddSP(StaticCastSharedRef<FRigVMEditorBase>(SharedRef()), &FRigVMEditorBase::HandleVMExecutionHalted);
+		RigVMHost->GetDebugInfo().ExecutionHalted().AddSP(this, &FRigVMEditor::HandleVMExecutionHalted);
 	}
 }
 
-void FRigVMEditorBase::CacheNameLists()
+void FRigVMEditor::CacheNameLists()
 {
 	if (URigVMBlueprint* RigVMBlueprint = GetRigVMBlueprint())
 	{
@@ -2617,9 +2688,9 @@ void FRigVMEditorBase::CacheNameLists()
 	}
 }
 
-void FRigVMEditorBase::OnCreateComment()
+void FRigVMEditor::OnCreateComment()
 {
-	TSharedPtr<SGraphEditor> GraphEditor = GetFocusedGraphEditor().Pin();
+	TSharedPtr<SGraphEditor> GraphEditor = FocusedGraphEdPtr.Pin();
 	if (GraphEditor.IsValid())
 	{
 		if (UEdGraph* Graph = GraphEditor->GetCurrentGraph())
@@ -2646,41 +2717,33 @@ void FRigVMEditorBase::OnCreateComment()
 	}
 }
 
-TArray<TWeakObjectPtr<UObject>> FRigVMEditorBase::GetSelectedObjects() const
+TArray<TWeakObjectPtr<UObject>> FRigVMEditor::GetSelectedObjects() const
 {
 	// if the inspector shows wrapped objects - look in that array instead.
 	// with recent weak object pointer changes on the property detail view
 	// we cannot rely on the GetSelectedObjects being valid after blueprint compilation.
-#if WITH_RIGVMLEGACYEDITOR
-	if (TSharedPtr<SKismetInspector> Inspector = GetInspector())
-#else
-	if (TSharedPtr<SRigVMDetailsInspector> Inspector = GetInspector())
-#endif
+	if(WrapperObjects.Num() == Inspector->GetSelectedObjects().Num())
 	{
-		if(WrapperObjects.Num() == Inspector->GetSelectedObjects().Num())
+		TArray<TWeakObjectPtr<UObject>> WeakWrapperObjects;
+		for(const TStrongObjectPtr<URigVMDetailsViewWrapperObject>& WrapperObjectPtr : WrapperObjects)
 		{
-			TArray<TWeakObjectPtr<UObject>> WeakWrapperObjects;
-			for(const TStrongObjectPtr<URigVMDetailsViewWrapperObject>& WrapperObjectPtr : WrapperObjects)
+			URigVMDetailsViewWrapperObject* WrapperObject = WrapperObjectPtr.Get();
+			if(IsValid(WrapperObject->GetSubject()))
 			{
-				URigVMDetailsViewWrapperObject* WrapperObject = WrapperObjectPtr.Get();
-				if(IsValid(WrapperObject->GetSubject()))
-				{
-					WeakWrapperObjects.Add(WrapperObject);
-				}
+				WeakWrapperObjects.Add(WrapperObject);
 			}
-			return WeakWrapperObjects;
 		}
-		return Inspector->GetSelectedObjects();
+		return WeakWrapperObjects;
 	}
-	return TArray<TWeakObjectPtr<UObject>>();
+	return Inspector->GetSelectedObjects();
 }
 
-void FRigVMEditorBase::SetDetailObjects(const TArray<UObject*>& InObjects)
+void FRigVMEditor::SetDetailObjects(const TArray<UObject*>& InObjects)
 {
 	SetDetailObjects(InObjects, true);
 }
 
-void FRigVMEditorBase::SetDetailObjects(const TArray<UObject*>& InObjects, bool bChangeUISelectionState)
+void FRigVMEditor::SetDetailObjects(const TArray<UObject*>& InObjects, bool bChangeUISelectionState)
 {
 	if(bSuspendDetailsPanelRefresh)
 	{
@@ -2713,11 +2776,7 @@ void FRigVMEditorBase::SetDetailObjects(const TArray<UObject*>& InObjects, bool
 			}
 
 			GetTabManager()->InsertNewDocumentTab(
-#if WITH_RIGVMLEGACYEDITOR
 				FBlueprintEditorTabs::DetailsID,
-#else
-				FRigVMNewEditorTabs::DetailsID(),
-#endif
 				TabId,
 				FTabManager::FLastMajorOrNomadTab(TEXT("RigVMMemoryDetails")),
 				DockTab
@@ -2803,7 +2862,7 @@ void FRigVMEditorBase::SetDetailObjects(const TArray<UObject*>& InObjects, bool
 			// create the wrapper object
 			if(URigVMDetailsViewWrapperObject* WrapperObject = URigVMDetailsViewWrapperObject::MakeInstance(GetDetailWrapperClass(), GetBlueprintObj(), ModelNodes, ModelNode))
 			{
-				WrapperObject->GetWrappedPropertyChangedChainEvent().AddSP(StaticCastSharedRef<FRigVMEditorBase>(SharedRef()), &FRigVMEditorBase::OnWrappedPropertyChangedChainEvent);
+				WrapperObject->GetWrappedPropertyChangedChainEvent().AddSP(this, &FRigVMEditor::OnWrappedPropertyChangedChainEvent);
 				WrapperObject->AddToRoot();
 
 				// todo: use transform widget for transforms
@@ -2826,21 +2885,12 @@ void FRigVMEditorBase::SetDetailObjects(const TArray<UObject*>& InObjects, bool
 		}
 	}
 
-#if WITH_RIGVMLEGACYEDITOR
 	SKismetInspector::FShowDetailsOptions Options;
 	Options.bForceRefresh = true;
-	if (TSharedPtr<SKismetInspector> Inspector = GetInspector())
-#else
-	SRigVMDetailsInspector::FShowDetailsOptions Options;
-	Options.bForceRefresh = true;
-	if (TSharedPtr<SRigVMDetailsInspector> Inspector = GetInspector())
-#endif
-	{
-		Inspector->ShowDetailsForObjects(FilteredObjects, Options);
-	}
+	Inspector->ShowDetailsForObjects(FilteredObjects, Options);
 }
 
-void FRigVMEditorBase::SetMemoryStorageDetails(const TArray<FRigVMMemoryStorageStruct*>& InStructs)
+void FRigVMEditor::SetMemoryStorageDetails(const TArray<FRigVMMemoryStorageStruct*>& InStructs)
 {
 	if (bSuspendDetailsPanelRefresh)
 	{
@@ -2877,11 +2927,7 @@ void FRigVMEditorBase::SetMemoryStorageDetails(const TArray<FRigVMMemoryStorageS
 			}
 
 			GetTabManager()->InsertNewDocumentTab(
-#if WITH_RIGVMLEGACYEDITOR
 				FBlueprintEditorTabs::DetailsID,
-#else
-				FRigVMNewEditorTabs::DetailsID(),
-#endif
 				TabId,
 				FTabManager::FLastMajorOrNomadTab(TEXT("RigVMMemoryDetails")),
 				DockTab
@@ -2891,7 +2937,7 @@ void FRigVMEditorBase::SetMemoryStorageDetails(const TArray<FRigVMMemoryStorageS
 	}
 }
 
-void FRigVMEditorBase::SetDetailViewForGraph(URigVMGraph* InGraph)
+void FRigVMEditor::SetDetailViewForGraph(URigVMGraph* InGraph)
 {
 	check(InGraph);
 	
@@ -2915,7 +2961,7 @@ void FRigVMEditorBase::SetDetailViewForGraph(URigVMGraph* InGraph)
 	SetDetailObjects(SelectedNodes);
 }
 
-void FRigVMEditorBase::SetDetailViewForFocusedGraph()
+void FRigVMEditor::SetDetailViewForFocusedGraph()
 {
 	if(bSuspendDetailsPanelRefresh)
 	{
@@ -2931,7 +2977,7 @@ void FRigVMEditorBase::SetDetailViewForFocusedGraph()
 	SetDetailViewForGraph(Model);
 }
 
-void FRigVMEditorBase::SetDetailViewForLocalVariable()
+void FRigVMEditor::SetDetailViewForLocalVariable()
 {
 	FName VariableName;
 	TArray< TWeakObjectPtr<UObject> > SelectedObjects = GetSelectedObjects();
@@ -2950,7 +2996,7 @@ void FRigVMEditorBase::SetDetailViewForLocalVariable()
 	SelectLocalVariable(GetFocusedGraph(), VariableName);
 }
 
-void FRigVMEditorBase::RefreshDetailView()
+void FRigVMEditor::RefreshDetailView()
 {
 	if(bSuspendDetailsPanelRefresh)
 	{
@@ -2974,7 +3020,7 @@ void FRigVMEditorBase::RefreshDetailView()
 	}
 }
 
-bool FRigVMEditorBase::DetailViewShowsAnyRigUnit() const
+bool FRigVMEditor::DetailViewShowsAnyRigUnit() const
 {
 	if (DetailViewShowsStruct(FRigVMStruct::StaticStruct()))
 	{
@@ -3000,12 +3046,12 @@ bool FRigVMEditorBase::DetailViewShowsAnyRigUnit() const
 	return false;
 }
 
-bool FRigVMEditorBase::DetailViewShowsLocalVariable() const
+bool FRigVMEditor::DetailViewShowsLocalVariable() const
 {
 	return DetailViewShowsStruct(FRigVMGraphVariableDescription::StaticStruct());
 }
 
-bool FRigVMEditorBase::DetailViewShowsStruct(UScriptStruct* InStruct) const
+bool FRigVMEditor::DetailViewShowsStruct(UScriptStruct* InStruct) const
 {
 	TArray< TWeakObjectPtr<UObject> > SelectedObjects = GetSelectedObjects();
 	for (TWeakObjectPtr<UObject> SelectedObject : SelectedObjects)
@@ -3027,7 +3073,7 @@ bool FRigVMEditorBase::DetailViewShowsStruct(UScriptStruct* InStruct) const
 	return false;
 }
 
-void FRigVMEditorBase::ClearDetailObject(bool bChangeUISelectionState)
+void FRigVMEditor::ClearDetailObject(bool bChangeUISelectionState)
 {
 	if(bSuspendDetailsPanelRefresh)
 	{
@@ -3035,29 +3081,18 @@ void FRigVMEditorBase::ClearDetailObject(bool bChangeUISelectionState)
 	}
 
 	ClearDetailsViewWrapperObjects();
-#if WITH_RIGVMLEGACYEDITOR
-	TSharedPtr<SKismetInspector> Inspector = GetInspector();
-#else
-	TSharedPtr<SRigVMDetailsInspector> Inspector = GetInspector();
-#endif
-	if(Inspector.IsValid())
-	{
-		Inspector->GetPropertyView()->SetObjects(TArray<UObject*>(), true); // clear property view synchronously
-		Inspector->ShowDetailsForObjects(TArray<UObject*>());
-		Inspector->ShowSingleStruct(TSharedPtr<FStructOnScope>());
-	}
+	
+	Inspector->GetPropertyView()->SetObjects(TArray<UObject*>(), true); // clear property view synchronously
+	Inspector->ShowDetailsForObjects(TArray<UObject*>());
+	Inspector->ShowSingleStruct(TSharedPtr<FStructOnScope>());
 
 	if (bChangeUISelectionState)
 	{
-#if WITH_RIGVMLEGACYEDITOR
 		SetUISelectionState(FBlueprintEditor::SelectionState_Graph);
-#else
-		SetUISelectionState(FRigVMNewEditor::SelectionState_Graph());
-#endif
 	}
 }
 
-void FRigVMEditorBase::ClearDetailsViewWrapperObjects()
+void FRigVMEditor::ClearDetailsViewWrapperObjects()
 {
 	for(const TStrongObjectPtr<URigVMDetailsViewWrapperObject>& WrapperObjectPtr : WrapperObjects)
 	{
@@ -3072,7 +3107,7 @@ void FRigVMEditorBase::ClearDetailsViewWrapperObjects()
 	WrapperObjects.Reset();
 }
 
-void FRigVMEditorBase::SetHost(URigVMHost* InHost)
+void FRigVMEditor::SetHost(URigVMHost* InHost)
 {
 	if (URigVMBlueprint* RigVMBlueprint = GetRigVMBlueprint())
 	{
@@ -3089,7 +3124,7 @@ void FRigVMEditorBase::SetHost(URigVMHost* InHost)
 	}
 }
 
-URigVMGraph* FRigVMEditorBase::GetFocusedModel() const
+URigVMGraph* FRigVMEditor::GetFocusedModel() const
 {
 	const URigVMBlueprint* Blueprint = GetRigVMBlueprint();
 	if (Blueprint == nullptr)
@@ -3101,7 +3136,7 @@ URigVMGraph* FRigVMEditorBase::GetFocusedModel() const
 	return Blueprint->GetModel(EdGraph);
 }
 
-URigVMController* FRigVMEditorBase::GetFocusedController() const
+URigVMController* FRigVMEditor::GetFocusedController() const
 {
 	URigVMBlueprint* Blueprint = GetRigVMBlueprint();
 	if (Blueprint == nullptr)
@@ -3111,10 +3146,10 @@ URigVMController* FRigVMEditorBase::GetFocusedController() const
 	return Blueprint->GetOrCreateController(GetFocusedModel());
 }
 
-TSharedPtr<SGraphEditor> FRigVMEditorBase::GetGraphEditor(UEdGraph* InEdGraph) const
+TSharedPtr<SGraphEditor> FRigVMEditor::GetGraphEditor(UEdGraph* InEdGraph) const
 {
 	TArray< TSharedPtr<SDockTab> > GraphEditorTabs;
-	GetDocumentManager()->FindAllTabsForFactory(GetGraphEditorTabFactory(), /*out*/ GraphEditorTabs);
+	DocumentManager->FindAllTabsForFactory(GraphEditorTabFactoryPtr, /*out*/ GraphEditorTabs);
 
 	for (const TSharedPtr<SDockTab>& GraphEditorTab : GraphEditorTabs)
 	{
@@ -3128,7 +3163,7 @@ TSharedPtr<SGraphEditor> FRigVMEditorBase::GetGraphEditor(UEdGraph* InEdGraph) c
 	return TSharedPtr<SGraphEditor>();
 }
 
-void FRigVMEditorBase::ExtendMenu()
+void FRigVMEditor::ExtendMenu()
 {
 	if(MenuExtender.IsValid())
 	{
@@ -3145,7 +3180,7 @@ void FRigVMEditorBase::ExtendMenu()
 	AddMenuExtender(RigVMEditorModule.GetMenuExtensibilityManager()->GetAllExtenders(GetToolkitCommands(), GetEditingObjects()));
 }
 
-void FRigVMEditorBase::ExtendToolbar()
+void FRigVMEditor::ExtendToolbar()
 {
 	DECLARE_SCOPE_HIERARCHICAL_COUNTER_FUNC()
 
@@ -3169,7 +3204,7 @@ void FRigVMEditorBase::ExtendToolbar()
 	{
 		if(ToolbarExtenderDelegate.IsBound())
 		{
-			AddToolbarExtender(ToolbarExtenderDelegate.Execute(GetToolkitCommands(), StaticCastSharedRef<FRigVMEditorBase>(SharedRef())));
+			AddToolbarExtender(ToolbarExtenderDelegate.Execute(GetToolkitCommands(), SharedThis(this)));
 		}
 	}
 
@@ -3177,26 +3212,26 @@ void FRigVMEditorBase::ExtendToolbar()
 		"Asset",
 		EExtensionHook::After,
 		GetToolkitCommands(),
-		FToolBarExtensionDelegate::CreateSP(StaticCastSharedRef<FRigVMEditorBase>(SharedRef()), &FRigVMEditorBase::FillToolbar, true)
+		FToolBarExtensionDelegate::CreateSP(this, &FRigVMEditor::FillToolbar, true)
 	);
 }
 
-void FRigVMEditorBase::FillToolbar(FToolBarBuilder& ToolbarBuilder, bool bEndSection)
+void FRigVMEditor::FillToolbar(FToolBarBuilder& ToolbarBuilder, bool bEndSection)
 {
 	ToolbarBuilder.BeginSection("Toolbar");
 	{
 		ToolbarBuilder.AddToolBarButton(
 			FRigVMEditorCommands::Get().ToggleEventQueue,
 			NAME_None, 
-			TAttribute<FText>::Create(TAttribute<FText>::FGetter::CreateSP(StaticCastSharedRef<FRigVMEditorBase>(SharedRef()), &FRigVMEditorBase::GetEventQueueLabel)),
+			TAttribute<FText>::Create(TAttribute<FText>::FGetter::CreateSP(this, &FRigVMEditor::GetEventQueueLabel)),
 			TAttribute<FText>(), 
-			TAttribute<FSlateIcon>::Create(TAttribute<FSlateIcon>::FGetter::CreateSP(StaticCastSharedRef<FRigVMEditorBase>(SharedRef()), &FRigVMEditorBase::GetEventQueueIcon))
+			TAttribute<FSlateIcon>::Create(TAttribute<FSlateIcon>::FGetter::CreateSP(this, &FRigVMEditor::GetEventQueueIcon))
 		);
 
 		FUIAction DefaultAction;
 		ToolbarBuilder.AddComboButton(
 			DefaultAction,
-			FOnGetContent::CreateSP(StaticCastSharedRef<FRigVMEditorBase>(SharedRef()), &FRigVMEditorBase::GenerateEventQueueMenuContent),
+			FOnGetContent::CreateSP(this, &FRigVMEditor::GenerateEventQueueMenuContent),
 			LOCTEXT("EventQueue_Label", "Available Events"),
 			LOCTEXT("EventQueue_ToolTip", "Pick between different events / modes for testing the Control Rig"),
 			FSlateIcon(FAppStyle::GetAppStyleSetName(), "LevelEditor.Recompile"),
@@ -3205,21 +3240,21 @@ void FRigVMEditorBase::FillToolbar(FToolBarBuilder& ToolbarBuilder, bool bEndSec
 		ToolbarBuilder.AddToolBarButton(FRigVMEditorCommands::Get().AutoCompileGraph,
 			NAME_None, TAttribute<FText>(), TAttribute<FText>(), FSlateIcon(FRigVMEditorStyle::Get().GetStyleSetName(), "RigVM.AutoCompileGraph"));
 
-		AddSelectedDebugObjectWidget(ToolbarBuilder);
+		ToolbarBuilder.AddWidget(SNew(SBlueprintEditorSelectedDebugObjectWidget, SharedThis(this)));
 
 		ToolbarBuilder.AddSeparator();
 		ToolbarBuilder.AddToolBarButton(
 			FRigVMEditorCommands::Get().ToggleExecutionMode,
 			NAME_None, 
-			TAttribute<FText>::Create(TAttribute<FText>::FGetter::CreateSP(StaticCastSharedRef<FRigVMEditorBase>(SharedRef()), &FRigVMEditorBase::GetExecutionModeLabel)),
+			TAttribute<FText>::Create(TAttribute<FText>::FGetter::CreateSP(this, &FRigVMEditor::GetExecutionModeLabel)),
 			TAttribute<FText>(), 
-			TAttribute<FSlateIcon>::Create(TAttribute<FSlateIcon>::FGetter::CreateSP(StaticCastSharedRef<FRigVMEditorBase>(SharedRef()), &FRigVMEditorBase::GetExecutionModeIcon))
+			TAttribute<FSlateIcon>::Create(TAttribute<FSlateIcon>::FGetter::CreateSP(this, &FRigVMEditor::GetExecutionModeIcon))
 		);
 		
 		FUIAction DefaultExecutionMode;
 		ToolbarBuilder.AddComboButton(
 			DefaultExecutionMode,
-			FOnGetContent::CreateSP(StaticCastSharedRef<FRigVMEditorBase>(SharedRef()), &FRigVMEditorBase::GenerateExecutionModeMenuContent),
+			FOnGetContent::CreateSP(this, &FRigVMEditor::GenerateExecutionModeMenuContent),
 			LOCTEXT("ExecutionMode_Label", "Execution Modes"),
 			LOCTEXT("ExecutionMode_ToolTip", "Pick between different execution modes for testing the Control Rig"),
 			FSlateIcon(FAppStyle::GetAppStyleSetName(), "LevelEditor.Recompile"),
@@ -3249,7 +3284,7 @@ void FRigVMEditorBase::FillToolbar(FToolBarBuilder& ToolbarBuilder, bool bEndSec
 		FUIAction DefaultBulkEditAction;
 		ToolbarBuilder.AddComboButton(
 			DefaultBulkEditAction,
-			FOnGetContent::CreateSP(StaticCastSharedRef<FRigVMEditorBase>(SharedRef()), &FRigVMEditorBase::GenerateBulkEditMenuContent),
+			FOnGetContent::CreateSP(this, &FRigVMEditor::GenerateBulkEditMenuContent),
 			LOCTEXT("BulkEdit_Label", "Bulk Edit"),
 			LOCTEXT("BulkEdit_ToolTip", "Perform changes across many nodes / assets"),
 			FSlateIcon(FAppStyle::GetAppStyleSetName(), "LevelEditor.Recompile"),
@@ -3262,7 +3297,7 @@ void FRigVMEditorBase::FillToolbar(FToolBarBuilder& ToolbarBuilder, bool bEndSec
 	}
 }
 
-void FRigVMEditorBase::HandleHideItem()
+void FRigVMEditor::HandleHideItem()
 {
 	DECLARE_SCOPE_HIERARCHICAL_COUNTER_FUNC()
 
@@ -3289,12 +3324,12 @@ void FRigVMEditorBase::HandleHideItem()
 	}
 }
 
-bool FRigVMEditorBase::CanHideItem() const
+bool FRigVMEditor::CanHideItem() const
 {
-	return GetSelectedNodes().Num() > 0;
+	return GetNumberOfSelectedNodes() > 0;
 }
 
-void FRigVMEditorBase::UpdateStaleWatchedPins()
+void FRigVMEditor::UpdateStaleWatchedPins()
 {
 	URigVMBlueprint* RigVMBlueprint = GetRigVMBlueprint();
 	if (RigVMBlueprint == nullptr)
@@ -3364,12 +3399,12 @@ void FRigVMEditorBase::UpdateStaleWatchedPins()
 	}
 }
 
-void FRigVMEditorBase::HandleRefreshEditorFromBlueprint(URigVMBlueprint* InBlueprint)
+void FRigVMEditor::HandleRefreshEditorFromBlueprint(URigVMBlueprint* InBlueprint)
 {
 	Compile();
 }
 
-void FRigVMEditorBase::HandleVariableDroppedFromBlueprint(UObject* InSubject, FProperty* InVariableToDrop, const FVector2D& InDropPosition, const FVector2D& InScreenPosition)
+void FRigVMEditor::HandleVariableDroppedFromBlueprint(UObject* InSubject, FProperty* InVariableToDrop, const FVector2D& InDropPosition, const FVector2D& InScreenPosition)
 {
 	URigVMBlueprint* Blueprint = Cast<URigVMBlueprint>(GetBlueprintObj());
 	if (Blueprint == nullptr)
@@ -3421,7 +3456,7 @@ void FRigVMEditorBase::HandleVariableDroppedFromBlueprint(UObject* InSubject, FP
 
 	MenuBuilder.EndSection();
 
-	TSharedRef<SWidget> GraphEditorPanel = GetFocusedGraphEditor().Pin().ToSharedRef();
+	TSharedRef<SWidget> GraphEditorPanel = FocusedGraphEdPtr.Pin().ToSharedRef();
 
 	// Show dialog to choose getter vs setter
 	FSlateApplication::Get().PushMenu(
@@ -3433,12 +3468,12 @@ void FRigVMEditorBase::HandleVariableDroppedFromBlueprint(UObject* InSubject, FP
 	);
 }
 
-void FRigVMEditorBase::HandleBreakpointAdded()
+void FRigVMEditor::HandleBreakpointAdded()
 {
 	SetExecutionMode(ERigVMEditorExecutionModeType_Debug);
 }
 
-void FRigVMEditorBase::OnGraphNodeClicked(URigVMEdGraphNode* InNode)
+void FRigVMEditor::OnGraphNodeClicked(URigVMEdGraphNode* InNode)
 {
 	if (InNode)
 	{
@@ -3449,7 +3484,7 @@ void FRigVMEditorBase::OnGraphNodeClicked(URigVMEdGraphNode* InNode)
 	}
 }
 
-void FRigVMEditorBase::OnNodeDoubleClicked(URigVMBlueprint* InBlueprint, URigVMNode* InNode)
+void FRigVMEditor::OnNodeDoubleClicked(URigVMBlueprint* InBlueprint, URigVMNode* InNode)
 {
 	ensure(GetRigVMBlueprint() == InBlueprint);
 
@@ -3471,13 +3506,25 @@ void FRigVMEditorBase::OnNodeDoubleClicked(URigVMBlueprint* InBlueprint, URigVMN
 			}
 			else
 			{
-				UE_LOG(LogRigVMEditor, Warning, TEXT("Could not open graph (%s)"), *LibraryNode->GetFunctionIdentifier().GetLibraryNodePath());
+				if(const URigVMCollapseNode* FunctionLibraryNode = Cast<URigVMCollapseNode>(ContainedGraph->GetOuter()))
+				{
+					if(const URigVMFunctionLibrary* FunctionLibrary = FunctionLibraryNode->GetLibrary())
+					{
+						if(const URigVMBlueprint* FunctionBlueprint = Cast<URigVMBlueprint>(FunctionLibrary->GetOuter()))
+						{
+							if (const UEdGraph* FunctionEdGraph = FunctionBlueprint->GetEdGraph(ContainedGraph))
+							{
+								FKismetEditorUtilities::BringKismetToFocusAttentionOnObject(FunctionEdGraph);
+							}
+						}
+					}
+				}
 			}
 		}
 	}
 }
 
-void FRigVMEditorBase::OnGraphImported(UEdGraph* InEdGraph)
+void FRigVMEditor::OnGraphImported(UEdGraph* InEdGraph)
 {
 	check(InEdGraph);
 
@@ -3485,7 +3532,7 @@ void FRigVMEditorBase::OnGraphImported(UEdGraph* InEdGraph)
 	RenameNewlyAddedAction(InEdGraph->GetFName());
 }
 
-bool FRigVMEditorBase::OnActionMatchesName(FEdGraphSchemaAction* InAction, const FName& InName) const
+bool FRigVMEditor::OnActionMatchesName(FEdGraphSchemaAction* InAction, const FName& InName) const
 {
 	if (InAction->GetMenuDescription().ToString() == InName.ToString())
 	{
@@ -3494,7 +3541,7 @@ bool FRigVMEditorBase::OnActionMatchesName(FEdGraphSchemaAction* InAction, const
 	return false;
 }
 
-void FRigVMEditorBase::HandleShowCurrentStatement()
+void FRigVMEditor::HandleShowCurrentStatement()
 {
 	if(HaltedAtNode)
 	{
@@ -3511,7 +3558,7 @@ void FRigVMEditorBase::HandleShowCurrentStatement()
 	}
 }
 
-void FRigVMEditorBase::HandleBreakpointActionRequested(const ERigVMBreakpointAction BreakpointAction)
+void FRigVMEditor::HandleBreakpointActionRequested(const ERigVMBreakpointAction BreakpointAction)
 {
 	if(URigVMHost* DebuggedHost = Cast<URigVMHost>(GetBlueprintObj()->GetObjectBeingDebugged()))
 	{
@@ -3519,14 +3566,14 @@ void FRigVMEditorBase::HandleBreakpointActionRequested(const ERigVMBreakpointAct
 	}
 }
 
-bool FRigVMEditorBase::IsHaltedAtBreakpoint() const
+bool FRigVMEditor::IsHaltedAtBreakpoint() const
 {
 	return HaltedAtNode != nullptr;
 }
 
-void FRigVMEditorBase::FrameSelection()
+void FRigVMEditor::FrameSelection()
 {
-	if (SGraphEditor* GraphEd = GetFocusedGraphEditor().Pin().Get())
+	if (SGraphEditor* GraphEd = FocusedGraphEdPtr.Pin().Get())
 	{
 		if(URigVMGraph* Model = GetFocusedModel())
 		{
@@ -3536,13 +3583,13 @@ void FRigVMEditorBase::FrameSelection()
 	}
 }
 
-void FRigVMEditorBase::SwapFunctionWithinAsset()
+void FRigVMEditor::SwapFunctionWithinAsset()
 {
 	const FAssetData Asset = UE::RigVM::Editor::Tools::FindAssetFromAnyPath(GetRigVMBlueprint()->GetPathName(), true);
 	SwapFunctionForAssets({Asset}, true);
 }
 
-void FRigVMEditorBase::SwapFunctionAcrossProject()
+void FRigVMEditor::SwapFunctionAcrossProject()
 {
 	const IAssetRegistry& AssetRegistry = FModuleManager::LoadModuleChecked<FAssetRegistryModule>(TEXT("AssetRegistry")).Get();
 	TArray<FAssetData> AllAssets;
@@ -3550,7 +3597,7 @@ void FRigVMEditorBase::SwapFunctionAcrossProject()
 	SwapFunctionForAssets(AllAssets, false);
 }
 
-void FRigVMEditorBase::SwapFunctionForAssets(const TArray<FAssetData>& InAssets, bool bSetupUndo)
+void FRigVMEditor::SwapFunctionForAssets(const TArray<FAssetData>& InAssets, bool bSetupUndo)
 {
 	SRigVMSwapFunctionsWidget::FArguments WidgetArgs;
 	WidgetArgs
@@ -3566,7 +3613,7 @@ void FRigVMEditorBase::SwapFunctionForAssets(const TArray<FAssetData>& InAssets,
 	SwapFunctionsDialog->ShowNormal();
 }
 
-void FRigVMEditorBase::SwapAssetReferences()
+void FRigVMEditor::SwapAssetReferences()
 {
 	const FAssetData Asset = UE::RigVM::Editor::Tools::FindAssetFromAnyPath(GetRigVMBlueprint()->GetPathName(), true);
 	
@@ -3584,7 +3631,7 @@ void FRigVMEditorBase::SwapAssetReferences()
 	SwapFunctionsDialog->ShowNormal();
 }
 
-void FRigVMEditorBase::OnGraphNodeDropToPerform(TSharedPtr<FDragDropOperation> InDragDropOp, UEdGraph* InGraph, const FVector2D& InNodePosition, const FVector2D& InScreenPosition)
+void FRigVMEditor::OnGraphNodeDropToPerform(TSharedPtr<FDragDropOperation> InDragDropOp, UEdGraph* InGraph, const FVector2D& InNodePosition, const FVector2D& InScreenPosition)
 {
 	DECLARE_SCOPE_HIERARCHICAL_COUNTER_FUNC()
 
@@ -3592,7 +3639,7 @@ void FRigVMEditorBase::OnGraphNodeDropToPerform(TSharedPtr<FDragDropOperation> I
 	{
 		TSharedPtr<FRigVMGraphExplorerDragDropOp> ExplorerOp = StaticCastSharedPtr<FRigVMGraphExplorerDragDropOp>(InDragDropOp);
 
-		if (GetFocusedGraphEditor().IsValid())
+		if (FocusedGraphEdPtr.IsValid())
 		{
 			const URigVMEdGraphSchema* Schema = CastChecked<URigVMEdGraphSchema>(GetRigVMBlueprint()->GetRigVMEdGraphSchemaClass()->GetDefaultObject());
 			if (ExplorerOp->GetElement().Type == ERigVMExplorerElementType::Function)
@@ -3676,7 +3723,7 @@ void FRigVMEditorBase::OnGraphNodeDropToPerform(TSharedPtr<FDragDropOperation> I
 	}
 }
 
-void FRigVMEditorBase::UpdateGraphCompilerErrors()
+void FRigVMEditor::UpdateGraphCompilerErrors()
 {
 	DECLARE_SCOPE_HIERARCHICAL_COUNTER_FUNC()
 
@@ -3777,17 +3824,17 @@ void FRigVMEditorBase::UpdateGraphCompilerErrors()
 
 }
 
-bool FRigVMEditorBase::IsPIERunning()
+bool FRigVMEditor::IsPIERunning()
 {
 	return GEditor && (GEditor->PlayWorld != nullptr);
 }
 
-TArray<FName> FRigVMEditorBase::GetDefaultEventQueue() const
+TArray<FName> FRigVMEditor::GetDefaultEventQueue() const
 {
 	return TArray<FName>();
 }
 
-TArray<FName> FRigVMEditorBase::GetEventQueue() const
+TArray<FName> FRigVMEditor::GetEventQueue() const
 {
 	if (const URigVMHost* CurrentHost = GetRigVMHost())
 	{
@@ -3797,12 +3844,12 @@ TArray<FName> FRigVMEditorBase::GetEventQueue() const
 	return GetDefaultEventQueue();
 }
 
-void FRigVMEditorBase::SetEventQueue(TArray<FName> InEventQueue)
+void FRigVMEditor::SetEventQueue(TArray<FName> InEventQueue)
 {
 	return SetEventQueue(InEventQueue, false);
 }
 
-void FRigVMEditorBase::SetEventQueue(TArray<FName> InEventQueue, bool bCompile)
+void FRigVMEditor::SetEventQueue(TArray<FName> InEventQueue, bool bCompile)
 {
 	if (GetEventQueue() == InEventQueue)
 	{
@@ -3822,17 +3869,17 @@ void FRigVMEditorBase::SetEventQueue(TArray<FName> InEventQueue, bool bCompile)
 	}
 }
 
-FSlateIcon FRigVMEditorBase::GetEventQueueIcon(const TArray<FName>& InEventQueue) const
+FSlateIcon FRigVMEditor::GetEventQueueIcon(const TArray<FName>& InEventQueue) const
 {
 	return FSlateIcon();
 }
 
-FSlateIcon FRigVMEditorBase::GetEventQueueIcon() const
+FSlateIcon FRigVMEditor::GetEventQueueIcon() const
 {
 	return GetEventQueueIcon(GetEventQueue());
 }
 
-void FRigVMEditorBase::SetExecutionMode(const ERigVMEditorExecutionModeType InExecutionMode)
+void FRigVMEditor::SetExecutionMode(const ERigVMEditorExecutionModeType InExecutionMode)
 {
 	if (ExecutionMode == InExecutionMode)
 	{
@@ -3853,12 +3900,12 @@ void FRigVMEditorBase::SetExecutionMode(const ERigVMEditorExecutionModeType InEx
 	RefreshDetailView();
 }
 
-int32 FRigVMEditorBase::GetExecutionModeComboValue() const
+int32 FRigVMEditor::GetExecutionModeComboValue() const
 {
 	return (int32) ExecutionMode;
 }
 
-FText FRigVMEditorBase::GetExecutionModeLabel() const
+FText FRigVMEditor::GetExecutionModeLabel() const
 {
 	if (ExecutionMode == ERigVMEditorExecutionModeType_Debug)
 	{
@@ -3867,7 +3914,7 @@ FText FRigVMEditorBase::GetExecutionModeLabel() const
 	return FText::FromString(TEXT("ReleaseMode"));
 }
 
-FSlateIcon FRigVMEditorBase::GetExecutionModeIcon(ERigVMEditorExecutionModeType InExecutionMode)
+FSlateIcon FRigVMEditor::GetExecutionModeIcon(ERigVMEditorExecutionModeType InExecutionMode)
 {
 	if (InExecutionMode == ERigVMEditorExecutionModeType_Debug)
 	{
@@ -3876,12 +3923,12 @@ FSlateIcon FRigVMEditorBase::GetExecutionModeIcon(ERigVMEditorExecutionModeType
 	return FSlateIcon(FRigVMEditorStyle::Get().GetStyleSetName(), "RigVM.ReleaseMode");
 }
 
-FSlateIcon FRigVMEditorBase::GetExecutionModeIcon() const
+FSlateIcon FRigVMEditor::GetExecutionModeIcon() const
 {
 	return GetExecutionModeIcon(ExecutionMode);
 }
 
-void FRigVMEditorBase::GetCustomDebugObjects(TArray<FRigVMCustomDebugObject>& DebugList) const
+void FRigVMEditor::GetCustomDebugObjects(TArray<FCustomDebugObject>& DebugList) const
 {
 	URigVMBlueprint* RigVMBlueprint = GetRigVMBlueprint();
 	if (RigVMBlueprint == nullptr)
@@ -3893,7 +3940,7 @@ void FRigVMEditorBase::GetCustomDebugObjects(TArray<FRigVMCustomDebugObject>& De
 	{
 		if(IsValid(CurrentHost))
 		{
-			FRigVMCustomDebugObject DebugObject;
+			FCustomDebugObject DebugObject;
 			DebugObject.Object = CurrentHost;
 			DebugObject.NameOverride = GetCustomDebugObjectLabel(CurrentHost);
 			DebugList.Add(DebugObject);
@@ -3985,7 +4032,7 @@ void FRigVMEditorBase::GetCustomDebugObjects(TArray<FRigVMCustomDebugObject>& De
 							continue;
 						}
 
-						FRigVMCustomDebugObject DebugObject;
+						FCustomDebugObject DebugObject;
 						DebugObject.Object = InstancedHost;
 						DebugObject.NameOverride = GetCustomDebugObjectLabel(InstancedHost);
 						DebugList.Add(DebugObject);
@@ -3996,7 +4043,7 @@ void FRigVMEditorBase::GetCustomDebugObjects(TArray<FRigVMCustomDebugObject>& De
 	}
 }
 
-void FRigVMEditorBase::HandleSetObjectBeingDebugged(UObject* InObject)
+void FRigVMEditor::HandleSetObjectBeingDebugged(UObject* InObject)
 {
 	if(URigVMHost* PreviouslyDebuggedHost = Cast<URigVMHost>(GetBlueprintObj()->GetObjectBeingDebugged()))
 	{
@@ -4037,9 +4084,9 @@ void FRigVMEditorBase::HandleSetObjectBeingDebugged(UObject* InObject)
 	if(DebuggedHost)
 	{
 		DebuggedHost->SetLog(&RigVMLog);
-		DebuggedHost->OnExecuted_AnyThread().AddSP(StaticCastSharedRef<FRigVMEditorBase>(SharedRef()), &FRigVMEditorBase::HandleVMExecutedEvent);
+		DebuggedHost->OnExecuted_AnyThread().AddSP(this, &FRigVMEditor::HandleVMExecutedEvent);
 #if WITH_EDITOR
-		DebuggedHost->GetDebugInfo().ExecutionHalted().AddSP(StaticCastSharedRef<FRigVMEditorBase>(SharedRef()), &FRigVMEditorBase::HandleVMExecutionHalted);
+		DebuggedHost->GetDebugInfo().ExecutionHalted().AddSP(this, &FRigVMEditor::HandleVMExecutionHalted);
 		DebuggedHost->SetIsInDebugMode(ExecutionMode == ERigVMEditorExecutionModeType_Debug);	
 #endif
 	}
@@ -4048,7 +4095,7 @@ void FRigVMEditorBase::HandleSetObjectBeingDebugged(UObject* InObject)
 	LastDebuggedHost = GetCustomDebugObjectLabel(DebuggedHost);
 }
 
-FString FRigVMEditorBase::GetCustomDebugObjectLabel(UObject* ObjectBeingDebugged) const
+FString FRigVMEditor::GetCustomDebugObjectLabel(UObject* ObjectBeingDebugged) const
 {
 	if (ObjectBeingDebugged == nullptr)
 	{
@@ -4077,7 +4124,7 @@ FString FRigVMEditorBase::GetCustomDebugObjectLabel(UObject* ObjectBeingDebugged
 	return GetBlueprintObj()->GetName();
 }
 
-void FRigVMEditorBase::OnPIEStopped(bool bSimulation)
+void FRigVMEditor::OnPIEStopped(bool bSimulation)
 {
 	if(URigVMBlueprint* Blueprint = GetRigVMBlueprint())
 	{
@@ -4085,9 +4132,4 @@ void FRigVMEditorBase::OnPIEStopped(bool bSimulation)
 	}
 }
 
-bool FRigVMEditorBase::IsEditingSingleBlueprint() const
-{
-	return GetBlueprintObj() != nullptr;
-}
-
-#undef LOCTEXT_NAMESPACE
\ No newline at end of file
+#undef LOCTEXT_NAMESPACE 
diff --git a/Engine/Plugins/Runtime/RigVM/Source/RigVMEditor/Private/Editor/RigVMEditorGraphExplorerTabSummoner.cpp b/Engine/Plugins/Runtime/RigVM/Source/RigVMEditor/Private/Editor/RigVMEditorGraphExplorerTabSummoner.cpp
--- a/Engine/Plugins/Runtime/RigVM/Source/RigVMEditor/Private/Editor/RigVMEditorGraphExplorerTabSummoner.cpp
+++ b/Engine/Plugins/Runtime/RigVM/Source/RigVMEditor/Private/Editor/RigVMEditorGraphExplorerTabSummoner.cpp
@@ -2,12 +2,12 @@
 
 #include "Editor/RigVMEditorGraphExplorerTabSummoner.h"
 #include "Widgets/SRigVMEditorGraphExplorer.h"
-#include "Editor/RigVMNewEditor.h"
+#include "Editor/RigVMEditor.h"
 
 #define LOCTEXT_NAMESPACE "RigVMEditorGraphExplorerTabSummoner"
 
-FRigVMEditorGraphExplorerTabSummoner::FRigVMEditorGraphExplorerTabSummoner(const TSharedRef<FRigVMEditorBase>& InRigVMEditor)
-	: FWorkflowTabFactory(TabID, InRigVMEditor.Get().GetHostingApp())
+FRigVMEditorGraphExplorerTabSummoner::FRigVMEditorGraphExplorerTabSummoner(const TSharedRef<FRigVMEditor>& InRigVMEditor)
+	: FWorkflowTabFactory(TabID, InRigVMEditor)
 	, RigVMEditor(InRigVMEditor)
 {
 	TabLabel = LOCTEXT("RigVMEditorGraphExplorerTabLabel", "Graph Explorer");
diff --git a/Engine/Plugins/Runtime/RigVM/Source/RigVMEditor/Private/Editor/RigVMEditorMenuContext.cpp b/Engine/Plugins/Runtime/RigVM/Source/RigVMEditor/Private/Editor/RigVMEditorMenuContext.cpp
--- a/Engine/Plugins/Runtime/RigVM/Source/RigVMEditor/Private/Editor/RigVMEditorMenuContext.cpp
+++ b/Engine/Plugins/Runtime/RigVM/Source/RigVMEditor/Private/Editor/RigVMEditorMenuContext.cpp
@@ -3,12 +3,12 @@
 #include "Editor/RigVMEditorMenuContext.h"
 
 #include "RigVMBlueprint.h"
-#include "Editor/RigVMNewEditor.h"
+#include "Editor/RigVMEditor.h"
 #include "Framework/Application/SlateApplication.h"
 
 #include UE_INLINE_GENERATED_CPP_BY_NAME(RigVMEditorMenuContext)
 
-void URigVMEditorMenuContext::Init(TWeakPtr<FRigVMEditorBase> InRigVMEditor, const FRigVMEditorGraphMenuContext& InGraphMenuContext)
+void URigVMEditorMenuContext::Init(TWeakPtr<FRigVMEditor> InRigVMEditor, const FRigVMEditorGraphMenuContext& InGraphMenuContext)
 {
 	WeakRigVMEditor = InRigVMEditor;
 	GraphMenuContext = InGraphMenuContext;
@@ -16,7 +16,7 @@ void URigVMEditorMenuContext::Init(TWeakPtr<FRigVMEditorBase> InRigVMEditor, con
 
 URigVMBlueprint* URigVMEditorMenuContext::GetRigVMBlueprint() const
 {
-	if (const TSharedPtr<FRigVMEditorBase> Editor = WeakRigVMEditor.Pin())
+	if (const TSharedPtr<FRigVMEditor> Editor = WeakRigVMEditor.Pin())
 	{
 		return Editor->GetRigVMBlueprint();
 	}
@@ -46,9 +46,9 @@ FRigVMEditorGraphMenuContext URigVMEditorMenuContext::GetGraphMenuContext()
 	return GraphMenuContext;
 }
 
-FRigVMEditorBase* URigVMEditorMenuContext::GetRigVMEditor() const
+FRigVMEditor* URigVMEditorMenuContext::GetRigVMEditor() const
 {
-	if (const TSharedPtr<FRigVMEditorBase> Editor = WeakRigVMEditor.Pin())
+	if (const TSharedPtr<FRigVMEditor> Editor = WeakRigVMEditor.Pin())
 	{
 		return Editor.Get();
 	}
diff --git a/Engine/Plugins/Runtime/RigVM/Source/RigVMEditor/Private/Editor/RigVMEditorMode.cpp b/Engine/Plugins/Runtime/RigVM/Source/RigVMEditor/Private/Editor/RigVMEditorMode.cpp
new file mode 100644
--- /dev/null
+++ b/Engine/Plugins/Runtime/RigVM/Source/RigVMEditor/Private/Editor/RigVMEditorMode.cpp
@@ -0,0 +1,96 @@
+// Copyright Epic Games, Inc. All Rights Reserved.
+
+#include "Editor/RigVMEditorMode.h" 
+#include "BlueprintEditorTabs.h"
+#include "SBlueprintEditorToolbar.h"
+#include "Editor/RigVMExecutionStackTabSummoner.h"
+#include "Editor/RigVMEditorGraphExplorerTabSummoner.h"
+#include "ToolMenus.h"
+
+FRigVMEditorMode::FRigVMEditorMode(const TSharedRef<FRigVMEditor>& InRigVMEditor)
+	: FBlueprintEditorApplicationMode(InRigVMEditor, FRigVMEditorModes::RigVMEditorMode, FRigVMEditorModes::GetLocalizedMode, false, false)
+{
+	RigVMBlueprintPtr = CastChecked<URigVMBlueprint>(InRigVMEditor->GetBlueprintObj());
+
+	TabFactories.RegisterFactory(MakeShared<FRigVMExecutionStackTabSummoner>(InRigVMEditor));
+	TabFactories.RegisterFactory(MakeShared<FRigVMEditorGraphExplorerTabSummoner>(InRigVMEditor));
+
+	TabLayout = FTabManager::NewLayout("Standalone_RigVMEditMode_Layout_v1.5")
+		->AddArea
+		(
+			// Main application area
+			FTabManager::NewPrimaryArea()
+			->SetOrientation(Orient_Vertical)
+			->Split
+			(
+				FTabManager::NewSplitter()
+				->SetOrientation(Orient_Horizontal)
+				->Split
+				(
+					//	Left bottom - rig/hierarchy
+					FTabManager::NewStack()
+					->SetSizeCoefficient(0.5f)
+					->AddTab(FRigVMExecutionStackTabSummoner::TabID, ETabState::OpenedTab)
+					->AddTab(FBlueprintEditorTabs::MyBlueprintID, ETabState::OpenedTab)
+					->AddTab(FRigVMEditorGraphExplorerTabSummoner::TabID, ETabState::ClosedTab)
+				)
+				->Split
+				(
+					// Middle 
+					FTabManager::NewSplitter()
+					->SetOrientation(Orient_Vertical)
+					->SetSizeCoefficient(0.6f)
+					->Split
+					(
+						// Middle top - document edit area
+						FTabManager::NewStack()
+						->SetSizeCoefficient(0.8f)
+						->AddTab("Document", ETabState::ClosedTab)
+					)
+					->Split
+					(
+						// Middle bottom - compiler results & find
+						FTabManager::NewStack()
+						->SetSizeCoefficient(0.2f)
+						->AddTab(FBlueprintEditorTabs::CompilerResultsID, ETabState::ClosedTab)
+						->AddTab(FBlueprintEditorTabs::FindResultsID, ETabState::ClosedTab)
+					)
+				)
+				->Split
+				(
+					// Right side
+					FTabManager::NewSplitter()
+					->SetOrientation(Orient_Vertical)
+					->SetSizeCoefficient(0.2f)
+					->Split
+					(
+						// Right top
+						FTabManager::NewStack()
+						->SetHideTabWell(false)
+						->SetSizeCoefficient(1.f)
+						->AddTab(FBlueprintEditorTabs::DetailsID, ETabState::OpenedTab)
+						->SetForegroundTab(FBlueprintEditorTabs::DetailsID)
+					)
+				)
+			)
+		);
+
+	if (UToolMenu* Toolbar = InRigVMEditor->RegisterModeToolbarIfUnregistered(GetModeName()))
+	{
+		InRigVMEditor->GetToolbarBuilder()->AddCompileToolbar(Toolbar);
+		InRigVMEditor->GetToolbarBuilder()->AddScriptingToolbar(Toolbar);
+		InRigVMEditor->GetToolbarBuilder()->AddBlueprintGlobalOptionsToolbar(Toolbar);
+	}
+}
+
+void FRigVMEditorMode::RegisterTabFactories(TSharedPtr<FTabManager> InTabManager)
+{
+	TSharedPtr<FBlueprintEditor> BP = MyBlueprintEditor.Pin();
+	
+	BP->RegisterToolbarTab(InTabManager.ToSharedRef());
+
+	// Mode-specific setup
+	BP->PushTabFactories(CoreTabFactories);
+	BP->PushTabFactories(BlueprintEditorTabFactories);
+	BP->PushTabFactories(TabFactories);
+}
diff --git a/Engine/Plugins/Runtime/RigVM/Source/RigVMEditor/Private/Editor/RigVMExecutionStackTabSummoner.cpp b/Engine/Plugins/Runtime/RigVM/Source/RigVMEditor/Private/Editor/RigVMExecutionStackTabSummoner.cpp
--- a/Engine/Plugins/Runtime/RigVM/Source/RigVMEditor/Private/Editor/RigVMExecutionStackTabSummoner.cpp
+++ b/Engine/Plugins/Runtime/RigVM/Source/RigVMEditor/Private/Editor/RigVMExecutionStackTabSummoner.cpp
@@ -2,12 +2,12 @@
 
 #include "Editor/RigVMExecutionStackTabSummoner.h"
 #include "Widgets/SRigVMExecutionStackView.h"
-#include "Editor/RigVMNewEditor.h"
+#include "Editor/RigVMEditor.h"
 
 #define LOCTEXT_NAMESPACE "RigVMExecutionStackTabSummoner"
 
-FRigVMExecutionStackTabSummoner::FRigVMExecutionStackTabSummoner(const TSharedRef<FRigVMEditorBase>& InRigVMEditor)
-	: FWorkflowTabFactory(TabID, InRigVMEditor->GetHostingApp())
+FRigVMExecutionStackTabSummoner::FRigVMExecutionStackTabSummoner(const TSharedRef<FRigVMEditor>& InRigVMEditor)
+	: FWorkflowTabFactory(TabID, InRigVMEditor)
 	, RigVMEditor(InRigVMEditor)
 {
 	TabLabel = LOCTEXT("RigVMExecutionStackTabLabel", "Execution Stack");
diff --git a/Engine/Plugins/Runtime/RigVM/Source/RigVMEditor/Private/Editor/RigVMGraphDetailCustomization.cpp b/Engine/Plugins/Runtime/RigVM/Source/RigVMEditor/Private/Editor/RigVMGraphDetailCustomization.cpp
--- a/Engine/Plugins/Runtime/RigVM/Source/RigVMEditor/Private/Editor/RigVMGraphDetailCustomization.cpp
+++ b/Engine/Plugins/Runtime/RigVM/Source/RigVMEditor/Private/Editor/RigVMGraphDetailCustomization.cpp
@@ -35,9 +35,6 @@
 #include "ContentBrowserModule.h"
 #include "IContentBrowserSingleton.h"
 #include "UObject/GarbageCollectionSchema.h"
-#if WITH_RIGVMLEGACYEDITOR
-#include "Editor/RigVMLegacyEditor.h"
-#endif
 
 #define LOCTEXT_NAMESPACE "RigVMGraphDetailCustomization"
 
@@ -46,7 +43,7 @@ static const FText RigVMGraphDetailCustomizationMultipleValues = LOCTEXT("Multip
 FRigVMFunctionArgumentGroupLayout::FRigVMFunctionArgumentGroupLayout(
 	const TWeakObjectPtr<URigVMGraph>& InGraph,
 	const TWeakInterfacePtr<IRigVMClientHost>& InRigVMClientHost,
-	const TWeakPtr<FRigVMEditorBase>& InEditor,
+	const TWeakPtr<FRigVMEditor>& InEditor,
 	bool bInputs)
 	: GraphPtr(InGraph)
 	, WeakRigVMClientHost(InRigVMClientHost)
@@ -153,7 +150,7 @@ void FRigVMFunctionArgumentGroupLayout::HandleModifiedEvent(ERigVMGraphNotifType
 class FRigVMFunctionArgumentPinTypeSelectorFilter : public IPinTypeSelectorFilter
 {
 public:
-	FRigVMFunctionArgumentPinTypeSelectorFilter(const TWeakPtr<FRigVMEditorBase>& InRigVMEditor, const TWeakObjectPtr<URigVMGraph>& InGraph)
+	FRigVMFunctionArgumentPinTypeSelectorFilter(const TWeakPtr<FRigVMEditor>& InRigVMEditor, const TWeakObjectPtr<URigVMGraph>& InGraph)
 		: RigVMEditorPtr(InRigVMEditor), GraphPtr(InGraph)
 	{
 	}
@@ -208,7 +205,7 @@ public:
 
 private:
 
-	TWeakPtr<FRigVMEditorBase> RigVMEditorPtr;
+	TWeakPtr<FRigVMEditor> RigVMEditorPtr;
 	
 	TWeakObjectPtr<URigVMGraph> GraphPtr;
 };
@@ -728,26 +725,16 @@ void FRigVMFunctionArgumentDefaultNode::HandleModifiedEvent(ERigVMGraphNotifType
 	}
 }
 
-#if WITH_RIGVMLEGACYEDITOR
-TSharedPtr<IDetailCustomization> FRigVMGraphDetailCustomization::MakeInstance(TSharedPtr<IBlueprintEditor> InEditor, const UClass* InExpectedBlueprintClass)
-#else
-TSharedPtr<IDetailCustomization> FRigVMGraphDetailCustomization::MakeInstance(TSharedPtr<FRigVMEditorBase> InEditor, const UClass* InExpectedBlueprintClass)
-#endif
+TSharedPtr<IDetailCustomization> FRigVMGraphDetailCustomization::MakeInstance(TSharedPtr<IBlueprintEditor> InBlueprintEditor, const UClass* InExpectedBlueprintClass)
 {
-	const TArray<UObject*>* Objects = (InEditor.IsValid() ? InEditor->GetObjectsCurrentlyBeingEdited() : nullptr);
+	const TArray<UObject*>* Objects = (InBlueprintEditor.IsValid() ? InBlueprintEditor->GetObjectsCurrentlyBeingEdited() : nullptr);
 	if (Objects && Objects->Num() == 1)
 	{
 		if (URigVMBlueprint* RigVMBlueprint = Cast<URigVMBlueprint>((*Objects)[0]))
 		{
 			if(RigVMBlueprint->GetClass() == InExpectedBlueprintClass)
 			{
-#if WITH_RIGVMLEGACYEDITOR
-				// TODO: need to change how the Kismet Inspector calls the customization with a ptr to IBlueprintEditor
-				TSharedPtr<FRigVMLegacyEditor> SharedPtr = StaticCastSharedPtr<FRigVMLegacyEditor>(InEditor);
-				return MakeShareable(new FRigVMGraphDetailCustomization(SharedPtr, RigVMBlueprint));
-#else
-				return MakeShareable(new FRigVMGraphDetailCustomization(InEditor, RigVMBlueprint));
-#endif
+				return MakeShareable(new FRigVMGraphDetailCustomization(StaticCastSharedPtr<FRigVMEditor>(InBlueprintEditor), RigVMBlueprint));
 			}
 		}
 	}
@@ -755,7 +742,7 @@ TSharedPtr<IDetailCustomization> FRigVMGraphDetailCustomization::MakeInstance(TS
 	return nullptr;
 }
 
-FRigVMGraphDetailCustomization::FRigVMGraphDetailCustomization(TSharedPtr<FRigVMEditorBase> RigVMigEditor, URigVMBlueprint* RigVMBlueprint)
+FRigVMGraphDetailCustomization::FRigVMGraphDetailCustomization(TSharedPtr<FRigVMEditor> RigVMigEditor, URigVMBlueprint* RigVMBlueprint)
 	: RigVMEditorPtr(RigVMigEditor)
 	, RigVMBlueprintPtr(RigVMBlueprint)
 	, RigVMGraphDetailCustomizationImpl(MakeShared<FRigVMGraphDetailCustomizationImpl>())
@@ -798,7 +785,7 @@ void FRigVMGraphDetailCustomizationImpl::CustomizeDetails(IDetailLayoutBuilder&
 	URigVMGraph* InModel,
 	URigVMController* InController,
 	IRigVMClientHost* InRigVMClientHost,
-	TWeakPtr<FRigVMEditorBase> InEditor)
+	TWeakPtr<FRigVMEditor> InEditor)
 {
 	WeakModel = InModel;
 	WeakController = InController;
@@ -1499,7 +1486,7 @@ void FRigVMGraphDetailCustomizationImpl::OnBrowseVariantRef(const FRigVMVariantR
 			
 				if(IAssetEditorInstance* Editor = GEditor->GetEditorSubsystem<UAssetEditorSubsystem>()->FindEditorForAsset(Blueprint, true))
 				{
-					if (FRigVMEditorBase* RigVMEditor = FRigVMEditorBase::GetFromAssetEditorInstance(Editor))
+					if(FRigVMEditor* RigVMEditor = static_cast<FRigVMEditor*>(Editor))
 					{
 						RigVMEditor->HandleJumpToHyperlink(LibraryNode);
 					}
@@ -2019,7 +2006,7 @@ void FRigVMWrappedNodeDetailCustomization::CustomizeDetails(IDetailLayoutBuilder
 			
 								if(IAssetEditorInstance* Editor = GEditor->GetEditorSubsystem<UAssetEditorSubsystem>()->FindEditorForAsset(Blueprint, true))
 								{
-									if (FRigVMEditorBase* RigVMEditor = FRigVMEditorBase::GetFromAssetEditorInstance(Editor))
+									if(FRigVMEditor* RigVMEditor = static_cast<FRigVMEditor*>(Editor))
 									{
 										RigVMEditor->HandleJumpToHyperlink(LibraryNode);
 										return FReply::Handled();
diff --git a/Engine/Plugins/Runtime/RigVM/Source/RigVMEditor/Private/Editor/RigVMGraphEditorSummoner.cpp b/Engine/Plugins/Runtime/RigVM/Source/RigVMEditor/Private/Editor/RigVMGraphEditorSummoner.cpp
deleted file mode 100644
diff --git a/Engine/Plugins/Runtime/RigVM/Source/RigVMEditor/Private/Editor/RigVMLegacyEditor.cpp b/Engine/Plugins/Runtime/RigVM/Source/RigVMEditor/Private/Editor/RigVMLegacyEditor.cpp
deleted file mode 100644
diff --git a/Engine/Plugins/Runtime/RigVM/Source/RigVMEditor/Private/Editor/RigVMLegacyEditorMode.cpp b/Engine/Plugins/Runtime/RigVM/Source/RigVMEditor/Private/Editor/RigVMLegacyEditorMode.cpp
deleted file mode 100644
diff --git a/Engine/Plugins/Runtime/RigVM/Source/RigVMEditor/Private/Editor/RigVMNewEditor.cpp b/Engine/Plugins/Runtime/RigVM/Source/RigVMEditor/Private/Editor/RigVMNewEditor.cpp
deleted file mode 100644
diff --git a/Engine/Plugins/Runtime/RigVM/Source/RigVMEditor/Private/Editor/RigVMNewEditorMode.cpp b/Engine/Plugins/Runtime/RigVM/Source/RigVMEditor/Private/Editor/RigVMNewEditorMode.cpp
deleted file mode 100644
diff --git a/Engine/Plugins/Runtime/RigVM/Source/RigVMEditor/Private/Editor/SRigVMDetailsInspector.cpp b/Engine/Plugins/Runtime/RigVM/Source/RigVMEditor/Private/Editor/SRigVMDetailsInspector.cpp
deleted file mode 100644
diff --git a/Engine/Plugins/Runtime/RigVM/Source/RigVMEditor/Private/Editor/SRigVMGraphTitleBar.cpp b/Engine/Plugins/Runtime/RigVM/Source/RigVMEditor/Private/Editor/SRigVMGraphTitleBar.cpp
deleted file mode 100644
diff --git a/Engine/Plugins/Runtime/RigVM/Source/RigVMEditor/Private/Editor/SRigVMGraphTitleBar.h b/Engine/Plugins/Runtime/RigVM/Source/RigVMEditor/Private/Editor/SRigVMGraphTitleBar.h
deleted file mode 100644
diff --git a/Engine/Plugins/Runtime/RigVM/Source/RigVMEditor/Private/RigVMEditorModule.cpp b/Engine/Plugins/Runtime/RigVM/Source/RigVMEditor/Private/RigVMEditorModule.cpp
--- a/Engine/Plugins/Runtime/RigVM/Source/RigVMEditor/Private/RigVMEditorModule.cpp
+++ b/Engine/Plugins/Runtime/RigVM/Source/RigVMEditor/Private/RigVMEditorModule.cpp
@@ -75,11 +75,9 @@ void FRigVMEditorModule::StartupModule()
 		ReconstructAllNodesDelegateHandle = FBlueprintEditorUtils::OnReconstructAllNodesEvent.AddStatic(&FRigVMBlueprintUtils::HandleReconstructAllNodes);
 		RefreshAllNodesDelegateHandle = FBlueprintEditorUtils::OnRefreshAllNodesEvent.AddStatic(&FRigVMBlueprintUtils::HandleRefreshAllNodes);
 
-#if WITH_RIGVMLEGACYEDITOR
 		// Register Blueprint editor variable customization
 		FBlueprintEditorModule& BlueprintEditorModule = FModuleManager::LoadModuleChecked<FBlueprintEditorModule>("Kismet");
 		BlueprintVariableCustomizationHandle = BlueprintEditorModule.RegisterVariableCustomization(FProperty::StaticClass(), FOnGetVariableCustomizationInstance::CreateStatic(&FRigVMVariableDetailCustomization::MakeInstance));
-#endif
 
 		IAssetRegistry* AssetRegistry = IAssetRegistry::Get();
 		if (ensure(AssetRegistry))
@@ -156,10 +154,8 @@ void FRigVMEditorModule::StartupModuleCommon()
 	const URigVMBlueprint* BlueprintCDO = GetRigVMBlueprintCDO();
 	const UClass* EdGraphSchemaClass = BlueprintCDO->GetRigVMEdGraphSchemaClass();
 	const URigVMEdGraphSchema* SchemaCDO = CastChecked<URigVMEdGraphSchema>(EdGraphSchemaClass->GetDefaultObject());
-#if WITH_RIGVMLEGACYEDITOR
 	FBlueprintEditorModule& BlueprintEditorModule = FModuleManager::LoadModuleChecked<FBlueprintEditorModule>("Kismet");
 	BlueprintEditorModule.RegisterGraphCustomization(SchemaCDO, FOnGetGraphCustomizationInstance::CreateStatic(&FRigVMGraphDetailCustomization::MakeInstance, BlueprintClass));
-#endif
 
 	// Register to fixup newly created BPs
 	FKismetEditorUtilities::RegisterOnBlueprintCreatedCallback(this, URigVMHost::StaticClass(), FKismetEditorUtilities::FOnBlueprintCreated::CreateRaw(this, &FRigVMEditorModule::HandleNewBlueprintCreated));
@@ -186,10 +182,8 @@ void FRigVMEditorModule::ShutdownModule()
 			FBlueprintEditorUtils::OnReconstructAllNodesEvent.Remove(ReconstructAllNodesDelegateHandle);
 
 			// Unregister Blueprint editor variable customization
-#if WITH_RIGVMLEGACYEDITOR
 			FBlueprintEditorModule& BlueprintEditorModule = FModuleManager::LoadModuleChecked<FBlueprintEditorModule>("Kismet");
 			BlueprintEditorModule.UnregisterVariableCustomization(FProperty::StaticClass(), BlueprintVariableCustomizationHandle);
-#endif
 		}
 	}
 
@@ -203,10 +197,8 @@ void FRigVMEditorModule::ShutdownModuleCommon()
 		const URigVMBlueprint* BlueprintCDO = GetRigVMBlueprintCDO();
 		const UClass* EdGraphSchemaClass = BlueprintCDO->GetRigVMEdGraphSchemaClass();
 		const URigVMEdGraphSchema* SchemaCDO = CastChecked<URigVMEdGraphSchema>(EdGraphSchemaClass->GetDefaultObject());
-#if WITH_RIGVMLEGACYEDITOR
 		FBlueprintEditorModule& BlueprintEditorModule = FModuleManager::LoadModuleChecked<FBlueprintEditorModule>("Kismet");
 		BlueprintEditorModule.UnregisterGraphCustomization(SchemaCDO);
-#endif
 	}
 
 	// Unregister to fixup newly created BPs
diff --git a/Engine/Plugins/Runtime/RigVM/Source/RigVMEditor/Private/Widgets/SRigVMBulkEditWidget.cpp b/Engine/Plugins/Runtime/RigVM/Source/RigVMEditor/Private/Widgets/SRigVMBulkEditWidget.cpp
--- a/Engine/Plugins/Runtime/RigVM/Source/RigVMEditor/Private/Widgets/SRigVMBulkEditWidget.cpp
+++ b/Engine/Plugins/Runtime/RigVM/Source/RigVMEditor/Private/Widgets/SRigVMBulkEditWidget.cpp
@@ -15,8 +15,6 @@
 #include "Misc/UObjectToken.h"
 #include "Dialogs/Dialogs.h"
 #include "Framework/Application/SlateApplication.h"
-#include "Editor/RigVMNewEditor.h"
-
 #define LOCTEXT_NAMESPACE "SRigVMBulkEditWidget"
 
 //////////////////////////////////////////////////////////////////////////
@@ -280,7 +278,7 @@ void SRigVMBulkEditWidget::Tick(const FGeometry& AllottedGeometry, const double
 														
 												if(IAssetEditorInstance* Editor = GEditor->GetEditorSubsystem<UAssetEditorSubsystem>()->FindEditorForAsset(Blueprint, true))
 												{
-													if (FRigVMEditorBase* RigVMEditor = FRigVMEditorBase::GetFromAssetEditorInstance(Editor))
+													if(FRigVMEditor* RigVMEditor = static_cast<FRigVMEditor*>(Editor))
 													{
 														RigVMEditor->HandleJumpToHyperlink(WeakSubject.Get());
 													}
diff --git a/Engine/Plugins/Runtime/RigVM/Source/RigVMEditor/Private/Widgets/SRigVMEditorGraphExplorer.cpp b/Engine/Plugins/Runtime/RigVM/Source/RigVMEditor/Private/Widgets/SRigVMEditorGraphExplorer.cpp
--- a/Engine/Plugins/Runtime/RigVM/Source/RigVMEditor/Private/Widgets/SRigVMEditorGraphExplorer.cpp
+++ b/Engine/Plugins/Runtime/RigVM/Source/RigVMEditor/Private/Widgets/SRigVMEditorGraphExplorer.cpp
@@ -2,7 +2,7 @@
 
 #include "Widgets/SRigVMEditorGraphExplorer.h"
 
-#include "Editor/RigVMNewEditor.h"
+#include "Editor/RigVMEditor.h"
 #include "EdGraph/RigVMEdGraphSchema.h"
 #include "Widgets/SRigVMEditorGraphExplorerTreeView.h"
 
@@ -14,7 +14,6 @@
 // Ick.
 #include "GraphActionNode.h"
 #include "GraphEditorActions.h"
-#include "RigVMBlueprint.h"
 #include "SPinTypeSelector.h"
 #include "Kismet2/BlueprintEditorUtils.h"
 #include "RigVMModel/Nodes/RigVMAggregateNode.h"
@@ -22,14 +21,6 @@
 #include "Widgets/Layout/SScaleBox.h"
 #include "HAL/PlatformApplicationMisc.h"
 #include "ScopedTransaction.h"
-#if WITH_RIGVMLEGACYEDITOR
-#include "SKismetInspector.h"
-#include "Editor/RigVMLegacyEditor.h"
-#else
-#include "Editor/SRigVMDetailsInspector.h"
-#endif
-#include "EdGraph/RigVMEdGraph.h"
-#include "WorkflowOrientedApp/WorkflowTabManager.h"
 
 #define LOCTEXT_NAMESPACE "RigVMGraphExplorer"
 
@@ -95,42 +86,10 @@ FReply FRigVMGraphExplorerDragDropOp::DroppedOnPanel(const TSharedRef<SWidget>&
 	{
 		if(URigVMBlueprint* Blueprint = TargetRigGraph->GetBlueprint())
 		{
-			//  Find the appropiate asset editor where the drop is happening
-			const TArray<IAssetEditorInstance*> AssetEditors = GEditor->GetEditorSubsystem<UAssetEditorSubsystem>()->FindEditorsForAsset(Blueprint);
-			IAssetEditorInstance* FocusedAssetEditor = nullptr;
-			for (IAssetEditorInstance* AssetEditor : AssetEditors)
+			if (GEditor->GetEditorSubsystem<UAssetEditorSubsystem>()->OpenEditorForAsset(Blueprint))
 			{
-				if (FRigVMEditorBase* RigVMEditor = FRigVMEditorBase::GetFromAssetEditorInstance(AssetEditor))
-				{
-					TWeakPtr<SGraphEditor> GraphEditor = RigVMEditor->GetFocusedGraphEditor();
-					TSharedPtr<SWidget> Widget = Panel->GetParentWidget();
-					while (Widget)
-					{
-						if (Widget == GraphEditor)
-						{
-							break;
-						}
-						Widget = Widget->GetParentWidget();
-					}
-					if (Widget)
-					{
-						FocusedAssetEditor = AssetEditor;
-						break;
-					}
-				}
-			}
-
-			if (!FocusedAssetEditor)
-			{
-				GEditor->GetEditorSubsystem<UAssetEditorSubsystem>()->OpenEditorForAsset(Blueprint);
-				FocusedAssetEditor = GEditor->GetEditorSubsystem<UAssetEditorSubsystem>()->FindEditorForAsset(Blueprint, /*bFocusIfOpen =*/true);
-			}
-			if (FocusedAssetEditor)
-			{
-				if (FRigVMEditorBase* RigVMEditor = FRigVMEditorBase::GetFromAssetEditorInstance(FocusedAssetEditor))
-				{
-					RigVMEditor->OnGraphNodeDropToPerform(SharedThis(this), TargetRigGraph, GraphPosition, ScreenPosition);
-				}
+				FRigVMEditor* Editor = static_cast<FRigVMEditor*>(GEditor->GetEditorSubsystem<UAssetEditorSubsystem>()->FindEditorForAsset(Blueprint, /*bFocusIfOpen =*/true));
+				Editor->OnGraphNodeDropToPerform(SharedThis(this), TargetRigGraph, GraphPosition, ScreenPosition);
 			}
 		}
 	}
@@ -196,12 +155,12 @@ void FRigVMEditorGraphExplorerCommands::RegisterCommands()
 }
 
 
-void SRigVMEditorGraphExplorer::Construct(const FArguments& InArgs, TWeakPtr<FRigVMEditorBase> InRigVMEditor)
+void SRigVMEditorGraphExplorer::Construct(const FArguments& InArgs, TWeakPtr<FRigVMEditor> InRigVMEditor)
 {
 	bNeedsRefresh = false;
 	RigVMEditor = InRigVMEditor;
 
-	TSharedPtr<FRigVMEditorBase> Editor = RigVMEditor.Pin();
+	TSharedPtr<FRigVMEditor> Editor = RigVMEditor.Pin();
 	if (Editor)
 	{
 		Editor->OnRefresh().AddRaw(this, &SRigVMEditorGraphExplorer::Refresh);
@@ -210,15 +169,12 @@ void SRigVMEditorGraphExplorer::Construct(const FArguments& InArgs, TWeakPtr<FRi
 	RegisterCommands();
 
 	CreateWidgets();
-
-	LastPinType.ResetToDefaults();
-	LastPinType.PinCategory = TEXT("bool");
 }
 
 
 SRigVMEditorGraphExplorer::~SRigVMEditorGraphExplorer()
 {
-	TSharedPtr<FRigVMEditorBase> Editor = RigVMEditor.Pin();
+	TSharedPtr<FRigVMEditor> Editor = RigVMEditor.Pin();
 	if (Editor)
 	{
 		Editor->OnRefresh().RemoveAll(this);
@@ -241,30 +197,9 @@ void SRigVMEditorGraphExplorer::Tick(const FGeometry& InAllottedGeometry, const
 	}
 }
 
-FName SRigVMEditorGraphExplorer::GetSelectedVariableName() const
-{
-	TArray<FRigVMExplorerElementKey> Selected = TreeView->GetSelectedKeys();
-	if (Selected.Num() != 1)
-	{
-		return NAME_None;
-	}
-
-	if (Selected[0].Type != ERigVMExplorerElementType::Variable || Selected[0].Type != ERigVMExplorerElementType::LocalVariable)
-	{
-		return NAME_None;
-	}
-
-	return *Selected[0].Name;
-}
-
-void SRigVMEditorGraphExplorer::ClearSelection()
-{
-	TreeView->ClearSelection();
-}
-
 void SRigVMEditorGraphExplorer::RegisterCommands()
 {
-	TSharedPtr<FRigVMEditorBase> Editor = RigVMEditor.Pin();
+	TSharedPtr<FRigVMEditor> Editor = RigVMEditor.Pin();
 
 	if (Editor)
 	{
@@ -518,7 +453,7 @@ TArray<const URigVMGraph*> SRigVMEditorGraphExplorer::GetRootGraphs() const
 {
 	TArray<const URigVMGraph*> Graphs;
 	
-	TSharedPtr<FRigVMEditorBase> Editor = RigVMEditor.Pin();
+	TSharedPtr<FRigVMEditor> Editor = RigVMEditor.Pin();
 	if (!Editor)
 	{
 		return Graphs;
@@ -537,7 +472,7 @@ TArray<const URigVMGraph*> SRigVMEditorGraphExplorer::GetChildrenGraphs(const FS
 {
 	TArray<const URigVMGraph*> Children;
 	
-	TSharedPtr<FRigVMEditorBase> Editor = RigVMEditor.Pin();
+	TSharedPtr<FRigVMEditor> Editor = RigVMEditor.Pin();
 	if (!Editor)
 	{
 		return Children;
@@ -569,7 +504,7 @@ TArray<URigVMNode*> SRigVMEditorGraphExplorer::GetEventNodesInGraph(const FStrin
 {
 	TArray<URigVMNode*> Events;
 	
-	TSharedPtr<FRigVMEditorBase> Editor = RigVMEditor.Pin();
+	TSharedPtr<FRigVMEditor> Editor = RigVMEditor.Pin();
 	if (!Editor)
 	{
 		return Events;
@@ -607,7 +542,7 @@ TArray<URigVMLibraryNode*> SRigVMEditorGraphExplorer::GetFunctions() const
 {
 	TArray<URigVMLibraryNode*> Functions;
 	
-	TSharedPtr<FRigVMEditorBase> Editor = RigVMEditor.Pin();
+	TSharedPtr<FRigVMEditor> Editor = RigVMEditor.Pin();
 	if (!Editor)
 	{
 		return Functions;
@@ -625,7 +560,7 @@ TArray<URigVMLibraryNode*> SRigVMEditorGraphExplorer::GetFunctions() const
 TArray<FRigVMGraphVariableDescription> SRigVMEditorGraphExplorer::GetVariables() const
 {
 	TArray<FRigVMGraphVariableDescription> Variables;
-	TSharedPtr<FRigVMEditorBase> Editor = RigVMEditor.Pin();
+	TSharedPtr<FRigVMEditor> Editor = RigVMEditor.Pin();
 	if (!Editor)
 	{
 		return Variables;
@@ -637,7 +572,7 @@ TArray<FRigVMGraphVariableDescription> SRigVMEditorGraphExplorer::GetVariables()
 TArray<FRigVMGraphVariableDescription> SRigVMEditorGraphExplorer::GetLocalVariables() const
 {
 	TArray<FRigVMGraphVariableDescription> Variables;
-	TSharedPtr<FRigVMEditorBase> Editor = RigVMEditor.Pin();
+	TSharedPtr<FRigVMEditor> Editor = RigVMEditor.Pin();
 	if (!Editor)
 	{
 		return Variables;
@@ -657,7 +592,7 @@ TArray<FRigVMGraphVariableDescription> SRigVMEditorGraphExplorer::GetLocalVariab
 FText SRigVMEditorGraphExplorer::GetGraphDisplayName(const FString& InGraphPath) const
 {
 	FText DisplayName;
-	TSharedPtr<FRigVMEditorBase> Editor = RigVMEditor.Pin();
+	TSharedPtr<FRigVMEditor> Editor = RigVMEditor.Pin();
 	if (!Editor)
 	{
 		return DisplayName;
@@ -680,7 +615,7 @@ FText SRigVMEditorGraphExplorer::GetGraphDisplayName(const FString& InGraphPath)
 FText SRigVMEditorGraphExplorer::GetEventDisplayName(const FString& InNodePath) const
 {
 	FText DisplayName;
-	TSharedPtr<FRigVMEditorBase> Editor = RigVMEditor.Pin();
+	TSharedPtr<FRigVMEditor> Editor = RigVMEditor.Pin();
 	if (!Editor)
 	{
 		return DisplayName;
@@ -698,7 +633,7 @@ FText SRigVMEditorGraphExplorer::GetEventDisplayName(const FString& InNodePath)
 const FSlateBrush* SRigVMEditorGraphExplorer::GetGraphIcon(const FString& InGraphPath) const
 {
 	const FSlateBrush* Brush = nullptr;
-	TSharedPtr<FRigVMEditorBase> Editor = RigVMEditor.Pin();
+	TSharedPtr<FRigVMEditor> Editor = RigVMEditor.Pin();
 	if (!Editor)
 	{
 		return Brush;
@@ -722,7 +657,7 @@ const FSlateBrush* SRigVMEditorGraphExplorer::GetGraphIcon(const FString& InGrap
 FText SRigVMEditorGraphExplorer::GetGraphTooltip(const FString& InGraphPath) const
 {
 	FText Tooltip;
-	TSharedPtr<FRigVMEditorBase> Editor = RigVMEditor.Pin();
+	TSharedPtr<FRigVMEditor> Editor = RigVMEditor.Pin();
 	if (!Editor)
 	{
 		return Tooltip;
@@ -744,7 +679,7 @@ FText SRigVMEditorGraphExplorer::GetGraphTooltip(const FString& InGraphPath) con
 
 void SRigVMEditorGraphExplorer::OnGraphClicked(const FString& InGraphPath)
 {
-	TSharedPtr<FRigVMEditorBase> Editor = RigVMEditor.Pin();
+	TSharedPtr<FRigVMEditor> Editor = RigVMEditor.Pin();
 	if (!Editor)
 	{
 		return;
@@ -760,42 +695,26 @@ void SRigVMEditorGraphExplorer::OnGraphClicked(const FString& InGraphPath)
 			const UEdGraphSchema* Schema = EdGraph->GetSchema();
 			check(Schema != nullptr);
 			Schema->GetGraphDisplayInformation(*EdGraph, DisplayInfo);
-#if WITH_RIGVMLEGACYEDITOR
-			if (TSharedPtr<SKismetInspector> CurrentInspector = Editor->GetInspector())
-			{
-				CurrentInspector->ShowDetailsForSingleObject(EdGraph, SKismetInspector::FShowDetailsOptions(DisplayInfo.PlainName));
-			}
-#else
-			if (TSharedPtr<SRigVMDetailsInspector> CurrentInspector = Editor->GetInspector())
-			{
-				CurrentInspector->ShowDetailsForSingleObject(EdGraph, SRigVMDetailsInspector::FShowDetailsOptions(DisplayInfo.PlainName));
-			}
-#endif
+			TSharedRef<SKismetInspector> CurrentInspector = Editor->GetInspector();
+			CurrentInspector->ShowDetailsForSingleObject(EdGraph, SKismetInspector::FShowDetailsOptions(DisplayInfo.PlainName));
 		}
 	}
 }
 
 void SRigVMEditorGraphExplorer::OnEventClicked(const FString& InEventPath)
 {
-	TSharedPtr<FRigVMEditorBase> Editor = RigVMEditor.Pin();
+	TSharedPtr<FRigVMEditor> Editor = RigVMEditor.Pin();
 	if (!Editor)
 	{
 		return;
 	}
-#if WITH_RIGVMLEGACYEDITOR
-	TSharedPtr<SKismetInspector> CurrentInspector = Editor->GetInspector();
-#else
-	TSharedPtr<SRigVMDetailsInspector> CurrentInspector = Editor->GetInspector();
-#endif
-	if (CurrentInspector.IsValid())
-	{
-		CurrentInspector->ShowDetailsForObjects(TArray<UObject*>());
-	}
+	TSharedRef<SKismetInspector> CurrentInspector = Editor->GetInspector();
+	CurrentInspector->ShowDetailsForObjects(TArray<UObject*>());
 }
 
 void SRigVMEditorGraphExplorer::OnFunctionClicked(const FString& InFunctionPath)
 {
-	TSharedPtr<FRigVMEditorBase> Editor = RigVMEditor.Pin();
+	TSharedPtr<FRigVMEditor> Editor = RigVMEditor.Pin();
 	if (!Editor)
 	{
 		return;
@@ -813,17 +732,8 @@ void SRigVMEditorGraphExplorer::OnFunctionClicked(const FString& InFunctionPath)
 				const UEdGraphSchema* Schema = EdGraph->GetSchema();
 				check(Schema != nullptr);
 				Schema->GetGraphDisplayInformation(*EdGraph, DisplayInfo);
-#if WITH_RIGVMLEGACYEDITOR
-				if (TSharedPtr<SKismetInspector> CurrentInspector = Editor->GetInspector())
-				{
-					CurrentInspector->ShowDetailsForSingleObject(EdGraph, SKismetInspector::FShowDetailsOptions(DisplayInfo.PlainName));
-				}
-#else
-				if (TSharedPtr<SRigVMDetailsInspector> CurrentInspector = Editor->GetInspector())
-				{
-					CurrentInspector->ShowDetailsForSingleObject(EdGraph, SRigVMDetailsInspector::FShowDetailsOptions(DisplayInfo.PlainName));
-				}
-#endif
+				TSharedRef<SKismetInspector> CurrentInspector = Editor->GetInspector();
+				CurrentInspector->ShowDetailsForSingleObject(EdGraph, SKismetInspector::FShowDetailsOptions(DisplayInfo.PlainName));
 			}
 		}
 	}
@@ -831,7 +741,7 @@ void SRigVMEditorGraphExplorer::OnFunctionClicked(const FString& InFunctionPath)
 
 void SRigVMEditorGraphExplorer::OnVariableClicked(const FRigVMExplorerElementKey& InVariable)
 {
-	TSharedPtr<FRigVMEditorBase> Editor = RigVMEditor.Pin();
+	TSharedPtr<FRigVMEditor> Editor = RigVMEditor.Pin();
 	if (!Editor)
 	{
 		return;
@@ -841,21 +751,13 @@ void SRigVMEditorGraphExplorer::OnVariableClicked(const FRigVMExplorerElementKey
 	if (InVariable.Type == ERigVMExplorerElementType::Variable)
 	{
 		FRigVMClient* RigVMClient = Blueprint->GetRigVMClient();
-#if WITH_RIGVMLEGACYEDITOR
-		TSharedPtr<SKismetInspector> CurrentInspector = Editor->GetInspector();
+		TSharedRef<SKismetInspector> CurrentInspector = Editor->GetInspector();
 		SKismetInspector::FShowDetailsOptions Options(FText::FromString(InVariable.Name));
-#else
-		TSharedPtr<SRigVMDetailsInspector> CurrentInspector = Editor->GetInspector();
-		SRigVMDetailsInspector::FShowDetailsOptions Options(FText::FromString(InVariable.Name));
-#endif
 		Options.bForceRefresh = true;
 
 		FProperty* Prop = Blueprint->SkeletonGeneratedClass->FindPropertyByName(*InVariable.Name);
 		UPropertyWrapper* PropWrap = (Prop ? Prop->GetUPropertyWrapper() : nullptr);
-		if (CurrentInspector.IsValid())
-		{
-			CurrentInspector->ShowDetailsForSingleObject(PropWrap, Options);
-		}
+		CurrentInspector->ShowDetailsForSingleObject(PropWrap, Options);
 	}
 	else if (InVariable.Type == ERigVMExplorerElementType::LocalVariable)
 	{
@@ -867,7 +769,7 @@ void SRigVMEditorGraphExplorer::OnVariableClicked(const FRigVMExplorerElementKey
 
 void SRigVMEditorGraphExplorer::OnGraphDoubleClicked(const FString& InGraphPath)
 {
-	TSharedPtr<FRigVMEditorBase> Editor = RigVMEditor.Pin();
+	TSharedPtr<FRigVMEditor> Editor = RigVMEditor.Pin();
 	if (!Editor)
 	{
 		return;
@@ -886,7 +788,7 @@ void SRigVMEditorGraphExplorer::OnGraphDoubleClicked(const FString& InGraphPath)
 
 void SRigVMEditorGraphExplorer::OnEventDoubleClicked(const FString& InEventPath)
 {
-	TSharedPtr<FRigVMEditorBase> Editor = RigVMEditor.Pin();
+	TSharedPtr<FRigVMEditor> Editor = RigVMEditor.Pin();
 	if (!Editor)
 	{
 		return;
@@ -903,7 +805,7 @@ void SRigVMEditorGraphExplorer::OnEventDoubleClicked(const FString& InEventPath)
 
 void SRigVMEditorGraphExplorer::OnFunctionDoubleClicked(const FString& InFunctionPath)
 {
-	TSharedPtr<FRigVMEditorBase> Editor = RigVMEditor.Pin();
+	TSharedPtr<FRigVMEditor> Editor = RigVMEditor.Pin();
 	if (!Editor)
 	{
 		return;
@@ -925,7 +827,7 @@ void SRigVMEditorGraphExplorer::OnFunctionDoubleClicked(const FString& InFunctio
 
 bool SRigVMEditorGraphExplorer::OnSetFunctionCategory(const FString& InFunctionPath, const FString& InCategory)
 {
-	TSharedPtr<FRigVMEditorBase> Editor = RigVMEditor.Pin();
+	TSharedPtr<FRigVMEditor> Editor = RigVMEditor.Pin();
 	if (!Editor.IsValid())
 	{
 		return false;
@@ -949,7 +851,7 @@ bool SRigVMEditorGraphExplorer::OnSetFunctionCategory(const FString& InFunctionP
 
 FString SRigVMEditorGraphExplorer::OnGetFunctionCategory(const FString& InFunctionPath) const
 {
-	TSharedPtr<FRigVMEditorBase> Editor = RigVMEditor.Pin();
+	TSharedPtr<FRigVMEditor> Editor = RigVMEditor.Pin();
 	if (!Editor.IsValid())
 	{
 		return FString();
@@ -968,7 +870,7 @@ FString SRigVMEditorGraphExplorer::OnGetFunctionCategory(const FString& InFuncti
 
 bool SRigVMEditorGraphExplorer::OnSetVariableCategory(const FString& InVariable, const FString& InCategory)
 {
-	TSharedPtr<FRigVMEditorBase> Editor = RigVMEditor.Pin();
+	TSharedPtr<FRigVMEditor> Editor = RigVMEditor.Pin();
 	if (!Editor.IsValid())
 	{
 		return false;
@@ -981,7 +883,7 @@ bool SRigVMEditorGraphExplorer::OnSetVariableCategory(const FString& InVariable,
 
 FString SRigVMEditorGraphExplorer::OnGetVariableCategory(const FString& InVariable) const
 {
-	TSharedPtr<FRigVMEditorBase> Editor = RigVMEditor.Pin();
+	TSharedPtr<FRigVMEditor> Editor = RigVMEditor.Pin();
 	if (!Editor.IsValid())
 	{
 		return FString();
@@ -992,7 +894,7 @@ FString SRigVMEditorGraphExplorer::OnGetVariableCategory(const FString& InVariab
 
 FEdGraphPinType SRigVMEditorGraphExplorer::OnGetVariablePinType(const FRigVMExplorerElementKey& InVariable)
 {
-	TSharedPtr<FRigVMEditorBase> Editor = RigVMEditor.Pin();
+	TSharedPtr<FRigVMEditor> Editor = RigVMEditor.Pin();
 	if (!Editor.IsValid())
 	{
 		return FEdGraphPinType();
@@ -1028,7 +930,7 @@ FEdGraphPinType SRigVMEditorGraphExplorer::OnGetVariablePinType(const FRigVMExpl
 
 bool SRigVMEditorGraphExplorer::OnSetVariablePinType(const FRigVMExplorerElementKey& InVariable, const FEdGraphPinType& InType)
 {
-	TSharedPtr<FRigVMEditorBase> Editor = RigVMEditor.Pin();
+	TSharedPtr<FRigVMEditor> Editor = RigVMEditor.Pin();
 	if (!Editor.IsValid())
 	{
 		return false;
@@ -1043,7 +945,6 @@ bool SRigVMEditorGraphExplorer::OnSetVariablePinType(const FRigVMExplorerElement
 			if (Variable.Name == InVariable.Name)
 			{
 				FBlueprintEditorUtils::ChangeMemberVariableType(Blueprint, *InVariable.Name, InType);
-				SetLastPinTypeUsed(InType);
 				return true;
 			}
 		}
@@ -1058,8 +959,6 @@ bool SRigVMEditorGraphExplorer::OnSetVariablePinType(const FRigVMExplorerElement
 				UObject* NewCPPTypeObject = nullptr;
 				RigVMTypeUtils::CPPTypeFromPinType(InType, NewCPPType, &NewCPPTypeObject);
 				Controller->SetLocalVariableType(*InVariable.Name, NewCPPType, NewCPPTypeObject, true, true);
-				SetLastPinTypeUsed(InType);
-				return true;
 			}
 		}
 	}
@@ -1069,7 +968,7 @@ bool SRigVMEditorGraphExplorer::OnSetVariablePinType(const FRigVMExplorerElement
 
 bool SRigVMEditorGraphExplorer::OnIsVariablePublic(const FString& InVariableName) const
 {
-	TSharedPtr<FRigVMEditorBase> Editor = RigVMEditor.Pin();
+	TSharedPtr<FRigVMEditor> Editor = RigVMEditor.Pin();
 	if (!Editor.IsValid())
 	{
 		return false;
@@ -1094,7 +993,7 @@ bool SRigVMEditorGraphExplorer::OnIsVariablePublic(const FString& InVariableName
 
 bool SRigVMEditorGraphExplorer::OnToggleVariablePublic(const FString& InVariableName) const
 {
-	TSharedPtr<FRigVMEditorBase> Editor = RigVMEditor.Pin();
+	TSharedPtr<FRigVMEditor> Editor = RigVMEditor.Pin();
 	if (!Editor.IsValid())
 	{
 		return false;
@@ -1109,7 +1008,7 @@ bool SRigVMEditorGraphExplorer::OnToggleVariablePublic(const FString& InVariable
 
 bool SRigVMEditorGraphExplorer::IsFunctionFocused() const
 {
-	TSharedPtr<FRigVMEditorBase> Editor = RigVMEditor.Pin();
+	TSharedPtr<FRigVMEditor> Editor = RigVMEditor.Pin();
 	if (!Editor.IsValid())
 	{
 		return false;
@@ -1140,7 +1039,7 @@ TArray<TSharedPtr<IPinTypeSelectorFilter>> SRigVMEditorGraphExplorer::GetCustomP
 
 TSharedPtr<SWidget> SRigVMEditorGraphExplorer::OnContextMenuOpening()
 {
-	TSharedPtr<FRigVMEditorBase> Editor = RigVMEditor.Pin();
+	TSharedPtr<FRigVMEditor> Editor = RigVMEditor.Pin();
 	if (!Editor.IsValid())
 	{
 		return TSharedPtr<SWidget>();
@@ -1231,7 +1130,7 @@ FReply SRigVMEditorGraphExplorer::OnDragDetected(const FGeometry& MyGeometry, co
 
 void SRigVMEditorGraphExplorer::OnCopy()
 {
-	TSharedPtr<FRigVMEditorBase> Editor = RigVMEditor.Pin();
+	TSharedPtr<FRigVMEditor> Editor = RigVMEditor.Pin();
 	if (!Editor)
 	{
 		return;
@@ -1357,7 +1256,7 @@ void SRigVMEditorGraphExplorer::OnDuplicate()
 
 bool SRigVMEditorGraphExplorer::CanDuplicate() const
 {
-	TSharedPtr<FRigVMEditorBase> Editor = RigVMEditor.Pin();
+	TSharedPtr<FRigVMEditor> Editor = RigVMEditor.Pin();
 	if (!Editor)
 	{
 		return false;
@@ -1426,7 +1325,7 @@ bool SRigVMEditorGraphExplorer::CanPasteGeneric() const
 
 void SRigVMEditorGraphExplorer::OnPasteFunction()
 {
-	TSharedPtr<FRigVMEditorBase> Editor = RigVMEditor.Pin();
+	TSharedPtr<FRigVMEditor> Editor = RigVMEditor.Pin();
 	if (!Editor)
 	{
 		return;
@@ -1440,7 +1339,7 @@ void SRigVMEditorGraphExplorer::OnPasteFunction()
 
 bool SRigVMEditorGraphExplorer::CanPasteFunction() const
 {
-	TSharedPtr<FRigVMEditorBase> Editor = RigVMEditor.Pin();
+	TSharedPtr<FRigVMEditor> Editor = RigVMEditor.Pin();
 	if (!Editor)
 	{
 		return false;
@@ -1454,7 +1353,7 @@ bool SRigVMEditorGraphExplorer::CanPasteFunction() const
 
 void SRigVMEditorGraphExplorer::OnPasteVariable()
 {
-	TSharedPtr<FRigVMEditorBase> Editor = RigVMEditor.Pin();
+	TSharedPtr<FRigVMEditor> Editor = RigVMEditor.Pin();
 	if (!Editor)
 	{
 		return;
@@ -1490,7 +1389,7 @@ void SRigVMEditorGraphExplorer::OnPasteVariable()
 
 bool SRigVMEditorGraphExplorer::CanPasteVariable() const
 {
-	TSharedPtr<FRigVMEditorBase> Editor = RigVMEditor.Pin();
+	TSharedPtr<FRigVMEditor> Editor = RigVMEditor.Pin();
 	if (!Editor)
 	{
 		return false;
@@ -1513,7 +1412,7 @@ bool SRigVMEditorGraphExplorer::CanPasteVariable() const
 
 void SRigVMEditorGraphExplorer::OnPasteLocalVariable()
 {
-	TSharedPtr<FRigVMEditorBase> Editor = RigVMEditor.Pin();
+	TSharedPtr<FRigVMEditor> Editor = RigVMEditor.Pin();
 	if (!Editor)
 	{
 		return;
@@ -1560,7 +1459,7 @@ bool SRigVMEditorGraphExplorer::IsShowingEmptySections() const
 
 void SRigVMEditorGraphExplorer::OnOpenGraph(bool bOnNewTab)
 {
-	TSharedPtr<FRigVMEditorBase> Editor = RigVMEditor.Pin();
+	TSharedPtr<FRigVMEditor> Editor = RigVMEditor.Pin();
 
 	if (Editor)
 	{
@@ -1593,17 +1492,17 @@ bool SRigVMEditorGraphExplorer::CanOpenGraph() const
 
 void SRigVMEditorGraphExplorer::OnCreateGraph()
 {
-	TSharedPtr<FRigVMEditorBase> Editor = RigVMEditor.Pin();
+	TSharedPtr<FRigVMEditor> Editor = RigVMEditor.Pin();
 	if (Editor)
 	{
-		Editor->NewDocument_OnClicked(FRigVMEditorBase::CGT_NewEventGraph);
+		Editor->NewDocument_OnClicked(FBlueprintEditor::CGT_NewEventGraph);
 	}
 }
 
 
 bool SRigVMEditorGraphExplorer::CanCreateGraph() const
 {
-	TSharedPtr<FRigVMEditorBase> Editor = RigVMEditor.Pin();
+	TSharedPtr<FRigVMEditor> Editor = RigVMEditor.Pin();
 	if (Editor)
 	{
 		return Editor->InEditingMode();
@@ -1614,10 +1513,10 @@ bool SRigVMEditorGraphExplorer::CanCreateGraph() const
 
 void SRigVMEditorGraphExplorer::OnCreateFunction()
 {
-	TSharedPtr<FRigVMEditorBase> Editor = RigVMEditor.Pin();
+	TSharedPtr<FRigVMEditor> Editor = RigVMEditor.Pin();
 	if (Editor)
 	{
-		Editor->NewDocument_OnClicked(FRigVMEditorBase::CGT_NewFunctionGraph);
+		Editor->NewDocument_OnClicked(FBlueprintEditor::CGT_NewFunctionGraph);
 	}
 }
 
@@ -1628,7 +1527,7 @@ bool SRigVMEditorGraphExplorer::CanCreateFunction() const
 
 void SRigVMEditorGraphExplorer::OnCreateVariable()
 {
-	TSharedPtr<FRigVMEditorBase> Editor = RigVMEditor.Pin();
+	TSharedPtr<FRigVMEditor> Editor = RigVMEditor.Pin();
 	if (Editor)
 	{
 		Editor->OnAddNewVariable();
@@ -1642,7 +1541,7 @@ bool SRigVMEditorGraphExplorer::CanCreateVariable() const
 
 void SRigVMEditorGraphExplorer::OnCreateLocalVariable()
 {
-	TSharedPtr<FRigVMEditorBase> Editor = RigVMEditor.Pin();
+	TSharedPtr<FRigVMEditor> Editor = RigVMEditor.Pin();
 	if (!Editor)
 	{
 		return;
@@ -1654,7 +1553,7 @@ void SRigVMEditorGraphExplorer::OnCreateLocalVariable()
 
 bool SRigVMEditorGraphExplorer::CanCreateLocalVariable() const
 {
-	TSharedPtr<FRigVMEditorBase> Editor = RigVMEditor.Pin();
+	TSharedPtr<FRigVMEditor> Editor = RigVMEditor.Pin();
 	if (!Editor)
 	{
 		return false;
@@ -1665,7 +1564,7 @@ bool SRigVMEditorGraphExplorer::CanCreateLocalVariable() const
 
 void SRigVMEditorGraphExplorer::OnAddFunctionVariant()
 {
-	TSharedPtr<FRigVMEditorBase> Editor = RigVMEditor.Pin();
+	TSharedPtr<FRigVMEditor> Editor = RigVMEditor.Pin();
 
 	if (!Editor)
 	{
@@ -1692,7 +1591,7 @@ void SRigVMEditorGraphExplorer::OnAddFunctionVariant()
 
 bool SRigVMEditorGraphExplorer::CanAddFunctionVariant() const
 {
-	TSharedPtr<FRigVMEditorBase> Editor = RigVMEditor.Pin();
+	TSharedPtr<FRigVMEditor> Editor = RigVMEditor.Pin();
 
 	if (!Editor)
 	{
@@ -1710,7 +1609,7 @@ bool SRigVMEditorGraphExplorer::CanAddFunctionVariant() const
 
 void SRigVMEditorGraphExplorer::OnDeleteEntry()
 {
-	TSharedPtr<FRigVMEditorBase> Editor = RigVMEditor.Pin();
+	TSharedPtr<FRigVMEditor> Editor = RigVMEditor.Pin();
 
 	if (Editor)
 	{
@@ -1772,7 +1671,7 @@ void SRigVMEditorGraphExplorer::OnDeleteEntry()
 
 bool SRigVMEditorGraphExplorer::CanDeleteEntry() const
 {
-	TSharedPtr<FRigVMEditorBase> Editor = RigVMEditor.Pin();
+	TSharedPtr<FRigVMEditor> Editor = RigVMEditor.Pin();
 
 	if (!Editor)
 	{
@@ -1831,7 +1730,7 @@ bool SRigVMEditorGraphExplorer::CanRenameEntry() const
 bool SRigVMEditorGraphExplorer::OnRenameGraph(const FString& InOldPath, const FString& InNewPath)
 {
 	bool bResult = false;
-	TSharedPtr<FRigVMEditorBase> Editor = RigVMEditor.Pin();
+	TSharedPtr<FRigVMEditor> Editor = RigVMEditor.Pin();
 
 	if (Editor)
 	{
@@ -1849,7 +1748,7 @@ bool SRigVMEditorGraphExplorer::OnRenameGraph(const FString& InOldPath, const FS
 bool SRigVMEditorGraphExplorer::OnCanRenameGraph(const FString& InOldPath, const FString& InNewPath, FText& OutErrorMessage) const
 {
 	bool bResult = false;
-	TSharedPtr<FRigVMEditorBase> Editor = RigVMEditor.Pin();
+	TSharedPtr<FRigVMEditor> Editor = RigVMEditor.Pin();
 
 	if (Editor)
 	{
@@ -1880,7 +1779,7 @@ bool SRigVMEditorGraphExplorer::OnCanRenameGraph(const FString& InOldPath, const
 bool SRigVMEditorGraphExplorer::OnRenameFunction(const FString& InOldPath, const FString& InNewPath)
 {
 	bool bResult = false;
-	TSharedPtr<FRigVMEditorBase> Editor = RigVMEditor.Pin();
+	TSharedPtr<FRigVMEditor> Editor = RigVMEditor.Pin();
 
 	if (Editor)
 	{
@@ -1904,7 +1803,7 @@ bool SRigVMEditorGraphExplorer::OnRenameFunction(const FString& InOldPath, const
 bool SRigVMEditorGraphExplorer::OnCanRenameFunction(const FString& InOldPath, const FString& InNewPath, FText& OutErrorMessage) const
 {
 	bool bResult = false;
-	TSharedPtr<FRigVMEditorBase> Editor = RigVMEditor.Pin();
+	TSharedPtr<FRigVMEditor> Editor = RigVMEditor.Pin();
 
 	if (Editor)
 	{
@@ -1927,7 +1826,7 @@ bool SRigVMEditorGraphExplorer::OnCanRenameFunction(const FString& InOldPath, co
 
 bool SRigVMEditorGraphExplorer::OnRenameVariable(const FRigVMExplorerElementKey& InOldKey, const FString& InNewName)
 {
-	TSharedPtr<FRigVMEditorBase> Editor = RigVMEditor.Pin();
+	TSharedPtr<FRigVMEditor> Editor = RigVMEditor.Pin();
 	if (!Editor)
 	{
 		return false;
@@ -1994,7 +1893,7 @@ void SRigVMEditorGraphExplorer::OnFindReference(bool bSearchAllBlueprints)
 		return;
 	}
 
-	TSharedPtr<FRigVMEditorBase> Editor = RigVMEditor.Pin();
+	TSharedPtr<FRigVMEditor> Editor = RigVMEditor.Pin();
 	URigVMBlueprint* Blueprint = Editor->GetRigVMBlueprint();
 	const URigVMEdGraphSchema* Schema = GetDefault<URigVMEdGraphSchema>();
 
diff --git a/Engine/Plugins/Runtime/RigVM/Source/RigVMEditor/Private/Widgets/SRigVMEditorGraphExplorerTreeView.cpp b/Engine/Plugins/Runtime/RigVM/Source/RigVMEditor/Private/Widgets/SRigVMEditorGraphExplorerTreeView.cpp
--- a/Engine/Plugins/Runtime/RigVM/Source/RigVMEditor/Private/Widgets/SRigVMEditorGraphExplorerTreeView.cpp
+++ b/Engine/Plugins/Runtime/RigVM/Source/RigVMEditor/Private/Widgets/SRigVMEditorGraphExplorerTreeView.cpp
@@ -3,9 +3,7 @@
 #include "Widgets/SRigVMEditorGraphExplorerTreeView.h"
 
 #include "SPinTypeSelector.h"
-#include "EdGraph/RigVMEdGraphSchema.h"
-#include "Editor/RigVMNewEditor.h"
-#include "RigVMModel/Nodes/RigVMLibraryNode.h"
+#include "Editor/RigVMEditor.h"
 #include "Widgets/SRigVMEditorGraphExplorer.h"
 
 #include "Widgets/Text/SInlineEditableTextBlock.h"
diff --git a/Engine/Plugins/Runtime/RigVM/Source/RigVMEditor/Private/Widgets/SRigVMExecutionStackView.cpp b/Engine/Plugins/Runtime/RigVM/Source/RigVMEditor/Private/Widgets/SRigVMExecutionStackView.cpp
--- a/Engine/Plugins/Runtime/RigVM/Source/RigVMEditor/Private/Widgets/SRigVMExecutionStackView.cpp
+++ b/Engine/Plugins/Runtime/RigVM/Source/RigVMEditor/Private/Widgets/SRigVMExecutionStackView.cpp
@@ -7,7 +7,7 @@
 #include "Widgets/Layout/SSpacer.h"
 #include "Widgets/Input/SSearchBox.h"
 #include "Widgets/Input/SNumericEntryBox.h"
-#include "Editor/RigVMNewEditor.h"
+#include "Editor/RigVMEditor.h"
 #include "Editor/RigVMExecutionStackCommands.h"
 #include "RigVMHost.h"
 #include "RigVMBlueprintGeneratedClass.h"
@@ -247,22 +247,22 @@ FText SRigStackItem::GetDurationText() const
 
 SRigVMExecutionStackView::~SRigVMExecutionStackView()
 {
-	if (TSharedPtr<FRigVMEditorBase> Editor = WeakEditor.Pin())
+	if (RigVMEditor.IsValid())
 	{
-		if (OnModelModified.IsValid() && Editor->GetRigVMBlueprint())
+		if (OnModelModified.IsValid() && RigVMEditor.Pin()->GetRigVMBlueprint())
 		{
-			Editor->GetRigVMBlueprint()->OnModified().Remove(OnModelModified);
+			RigVMEditor.Pin()->GetRigVMBlueprint()->OnModified().Remove(OnModelModified);
 		}
 		if (OnHostInitializedHandle.IsValid())
 		{
-			if(Editor->GetRigVMHost())
+			if(RigVMEditor.Pin()->GetRigVMHost())
 			{
-				Editor->GetRigVMHost()->OnInitialized_AnyThread().Remove(OnHostInitializedHandle);
+				RigVMEditor.Pin()->GetRigVMHost()->OnInitialized_AnyThread().Remove(OnHostInitializedHandle);
 			}
 		}
 		if (OnPreviewHostUpdatedHandle.IsValid())
 		{
-			Editor->OnPreviewHostUpdated().Remove(OnPreviewHostUpdatedHandle);
+			RigVMEditor.Pin()->OnPreviewHostUpdated().Remove(OnPreviewHostUpdatedHandle);
 		}
 	}
 	if (RigVMBlueprint.IsValid() && OnVMCompiledHandle.IsValid())
@@ -271,10 +271,10 @@ SRigVMExecutionStackView::~SRigVMExecutionStackView()
 	}
 }
 
-void SRigVMExecutionStackView::Construct( const FArguments& InArgs, TSharedRef<FRigVMEditorBase> InRigVMEditor)
+void SRigVMExecutionStackView::Construct( const FArguments& InArgs, TSharedRef<FRigVMEditor> InRigVMEditor)
 {
-	WeakEditor = InRigVMEditor;
-	RigVMBlueprint = InRigVMEditor->GetRigVMBlueprint();
+	RigVMEditor = InRigVMEditor;
+	RigVMBlueprint = RigVMEditor.Pin()->GetRigVMBlueprint();
 	CommandList = MakeShared<FUICommandList>();
 	bSuspendModelNotifications = false;
 	bSuspendControllerSelection = false;
@@ -348,18 +348,18 @@ void SRigVMExecutionStackView::Construct( const FArguments& InArgs, TSharedRef<F
 		OnVMCompiledHandle = RigVMBlueprint->OnVMCompiled().AddSP(this, &SRigVMExecutionStackView::OnVMCompiled);
 		OnModelModified = RigVMBlueprint->OnModified().AddSP(this, &SRigVMExecutionStackView::HandleModifiedEvent);
 
-		if (TSharedPtr<FRigVMEditorBase> Editor = WeakEditor.Pin())
+		if (RigVMEditor.IsValid())
 		{
-			if (URigVMHost* RigVMHost = Editor->GetRigVMHost())
+			if (URigVMHost* RigVMHost = RigVMEditor.Pin()->GetRigVMHost())
 			{
 				OnVMCompiled(RigVMBlueprint.Get(),  RigVMHost->GetVM(),  RigVMHost->GetRigVMExtendedExecuteContext());
 			}
 		}
 	}
 
-	if (TSharedPtr<FRigVMEditorBase> Editor = WeakEditor.Pin())
+	if (RigVMEditor.IsValid())
 	{
-		OnPreviewHostUpdatedHandle = Editor->OnPreviewHostUpdated().AddSP(this, &SRigVMExecutionStackView::HandlePreviewHostUpdated);
+		OnPreviewHostUpdatedHandle = RigVMEditor.Pin()->OnPreviewHostUpdated().AddSP(this, &SRigVMExecutionStackView::HandlePreviewHostUpdated);
 	}
 }
 
@@ -385,11 +385,7 @@ void SRigVMExecutionStackView::OnSelectionChanged(TSharedPtr<FRigStackEntry> Sel
 			return;
 		}
 
-		URigVMHost* Host = nullptr;
-		if (TSharedPtr<FRigVMEditorBase> Editor = WeakEditor.Pin())
-		{
-			Host = Editor->GetRigVMHost();
-		}
+		URigVMHost* Host = RigVMEditor.Pin()->GetRigVMHost();
 		if (Host == nullptr || Host->GetVM() == nullptr)
 		{
 			return;
@@ -458,6 +454,8 @@ void SRigVMExecutionStackView::PopulateStackView(URigVM* InVM, FRigVMExtendedExe
 {
 	if (InVM && InVMContext)
 	{
+		URigVMBlueprint* Blueprint = RigVMEditor.Pin()->GetRigVMBlueprint();
+
 		const FRigVMInstructionArray Instructions = InVM->GetInstructions();
 		const FRigVMByteCode& ByteCode = InVM->GetByteCode();
 
@@ -740,9 +738,9 @@ void SRigVMExecutionStackView::RefreshTreeView(URigVM* InVM, FRigVMExtendedExecu
 	if (InVM && InVMContext)
 	{
 		// fill the children from the log
-		if (TSharedPtr<FRigVMEditorBase> Editor = WeakEditor.Pin())
+		if (RigVMEditor.IsValid())
 		{
-			URigVMHost* Host = Editor->GetRigVMHost();
+			URigVMHost* Host = RigVMEditor.Pin()->GetRigVMHost();
 			if(Host && Host->GetLog())
 			{
 				const TArray<FRigVMLog::FLogEntry>& LogEntries = Host->GetLog()->Entries;
@@ -830,33 +828,30 @@ void SRigVMExecutionStackView::HandleFocusOnSelectedGraphNode()
 	TArray<TSharedPtr<FRigStackEntry>> SelectedItems = TreeView->GetSelectedItems();
 	if (SelectedItems.Num() > 0)
 	{
-		if (TSharedPtr<FRigVMEditorBase> Editor = WeakEditor.Pin())
+		URigVMHost* Host = RigVMEditor.Pin()->GetRigVMHost();
+		if (Host == nullptr || Host->GetVM() == nullptr)
 		{
-			URigVMHost* Host = Editor->GetRigVMHost();
-			if (Host == nullptr || Host->GetVM() == nullptr)
-			{
-				return;
-			}
+			return;
+		}
 
-			const FRigVMByteCode& ByteCode = Host->GetVM()->GetByteCode();
-			UObject* Subject = ByteCode.GetSubjectForInstruction(SelectedItems[0]->InstructionIndex);
-			if (URigVMNode* SelectedNode = Cast<URigVMNode>(Subject))
+		const FRigVMByteCode& ByteCode = Host->GetVM()->GetByteCode();
+		UObject* Subject = ByteCode.GetSubjectForInstruction(SelectedItems[0]->InstructionIndex);
+		if (URigVMNode* SelectedNode = Cast<URigVMNode>(Subject))
+		{
+			URigVMGraph* GraphToFocus = SelectedNode->GetGraph();
+			if (GraphToFocus && GraphToFocus->GetTypedOuter<URigVMAggregateNode>())
 			{
-				URigVMGraph* GraphToFocus = SelectedNode->GetGraph();
-				if (GraphToFocus && GraphToFocus->GetTypedOuter<URigVMAggregateNode>())
-				{
-					if(URigVMGraph* ParentGraph = GraphToFocus->GetParentGraph())
-					{
-						GraphToFocus = ParentGraph;
-					} 
-				}
-				
-				if (UEdGraph* EdGraph = RigVMBlueprint->GetEdGraph(GraphToFocus))
+				if(URigVMGraph* ParentGraph = GraphToFocus->GetParentGraph())
 				{
-					Editor->OpenGraphAndBringToFront(EdGraph, true);
-					Editor->ZoomToSelection_Clicked();
-					Editor->HandleModifiedEvent(ERigVMGraphNotifType::NodeSelected, GraphToFocus, SelectedNode);
-				}
+					GraphToFocus = ParentGraph;
+				} 
+			}
+			
+			if (UEdGraph* EdGraph = RigVMBlueprint->GetEdGraph(GraphToFocus))
+			{
+				RigVMEditor.Pin()->OpenGraphAndBringToFront(EdGraph, true);
+				RigVMEditor.Pin()->ZoomToSelection_Clicked();
+				RigVMEditor.Pin()->HandleModifiedEvent(ERigVMGraphNotifType::NodeSelected, GraphToFocus, SelectedNode);
 			}
 		}
 	}
@@ -933,12 +928,8 @@ void SRigVMExecutionStackView::HandleSelectTargetInstructions()
 TArray<TSharedPtr<FRigStackEntry>> SRigVMExecutionStackView::GetTargetItems(const TArray<TSharedPtr<FRigStackEntry>>& InItems) const
 {
 	TArray<TSharedPtr<FRigStackEntry>> TargetItems;
-
-	URigVMHost* Host = nullptr;
-	if (TSharedPtr<FRigVMEditorBase> Editor = WeakEditor.Pin())
-	{
-		Host = Editor->GetRigVMHost();
-	}
+	
+	URigVMHost* Host = RigVMEditor.Pin()->GetRigVMHost();
 	if (Host == nullptr || Host->GetVM() == nullptr)
 	{
 		return TargetItems;
@@ -1046,14 +1037,11 @@ void SRigVMExecutionStackView::OnVMCompiled(UObject* InCompiledObject, URigVM* I
 {
 	RefreshTreeView(InCompiledVM, &InVMContext);
 
-	if (TSharedPtr<FRigVMEditorBase> Editor = WeakEditor.Pin())
+	if (RigVMEditor.IsValid() && !OnHostInitializedHandle.IsValid())
 	{
-		if (!OnHostInitializedHandle.IsValid())
+		if(URigVMHost* Host = RigVMEditor.Pin()->GetRigVMHost())
 		{
-			if(URigVMHost* Host = Editor->GetRigVMHost())
-			{
-				OnHostInitializedHandle = Host->OnInitialized_AnyThread().AddSP(this, &SRigVMExecutionStackView::HandleHostInitializedEvent);
-			}
+			OnHostInitializedHandle = Host->OnInitialized_AnyThread().AddSP(this, &SRigVMExecutionStackView::HandleHostInitializedEvent);
 		}
 	}
 }
@@ -1065,13 +1053,10 @@ void SRigVMExecutionStackView::HandleExecutionHalted(const int32 InHaltedAtInstr
 		return;
 	}
 
-	if (TSharedPtr<FRigVMEditorBase> Editor = WeakEditor.Pin())
+	if (InHaltedAtInstruction == INDEX_NONE && InEntryName == RigVMEditor.Pin()->GetRigVMHost()->GetEventQueue().Last())
 	{
-		if (InHaltedAtInstruction == INDEX_NONE && InEntryName == Editor->GetRigVMHost()->GetEventQueue().Last())
-		{
-			HaltedAtInstruction = InHaltedAtInstruction;
-			return;
-		}
+		HaltedAtInstruction = InHaltedAtInstruction;
+		return;
 	}
 	
 	if (InHaltedAtInstruction != INDEX_NONE && Operators.Num() > InHaltedAtInstruction)
@@ -1085,13 +1070,10 @@ void SRigVMExecutionStackView::HandleExecutionHalted(const int32 InHaltedAtInstr
 void SRigVMExecutionStackView::OnFilterTextChanged(const FText& SearchText)
 {
 	FilterText = SearchText;
-	if (TSharedPtr<FRigVMEditorBase> Editor = WeakEditor.Pin())
+	URigVMHost* RigVMHost = RigVMEditor.Pin()->GetRigVMHost();
+	if (RigVMHost != nullptr)
 	{
-		URigVMHost* RigVMHost = Editor->GetRigVMHost();
-		if (RigVMHost != nullptr)
-		{
-			RefreshTreeView(RigVMHost->GetVM(), &RigVMHost->GetRigVMExtendedExecuteContext());
-		}
+		RefreshTreeView(RigVMHost->GetVM(), &RigVMHost->GetRigVMExtendedExecuteContext());
 	}
 }
 
@@ -1102,11 +1084,7 @@ void SRigVMExecutionStackView::HandleModifiedEvent(ERigVMGraphNotifType InNotifT
 		return;
 	}
 
-	URigVMHost* Host = nullptr;
-	if (TSharedPtr<FRigVMEditorBase> Editor = WeakEditor.Pin())
-	{
-		Host = Editor->GetRigVMHost();
-	}
+	URigVMHost* Host = RigVMEditor.Pin()->GetRigVMHost();
 	if (Host == nullptr || Host->GetVM() == nullptr)
 	{
 		return;
@@ -1178,12 +1156,12 @@ void SRigVMExecutionStackView::HandleHostInitializedEvent(URigVMHost* InHost, co
 		}
 	}
 	
-	if (TSharedPtr<FRigVMEditorBase> Editor = WeakEditor.Pin())
+	if (RigVMEditor.IsValid())
 	{
 		InHost->OnInitialized_AnyThread().Remove(OnHostInitializedHandle);
 		OnHostInitializedHandle.Reset();
 
-		if (URigVMBlueprint* RigBlueprint = Editor->GetRigVMBlueprint())
+		if (URigVMBlueprint* RigBlueprint = RigVMEditor.Pin()->GetRigVMBlueprint())
 		{
 			TArray<URigVMGraph*> Models = RigBlueprint->GetAllModels();
 			for (URigVMGraph* Model : Models)
@@ -1200,7 +1178,7 @@ void SRigVMExecutionStackView::HandleHostInitializedEvent(URigVMHost* InHost, co
 	}
 }
 
-void SRigVMExecutionStackView::HandlePreviewHostUpdated(FRigVMEditorBase* InEditor)
+void SRigVMExecutionStackView::HandlePreviewHostUpdated(FRigVMEditor* InEditor)
 {
 	if(URigVMHost* RigVMHost = InEditor->GetRigVMHost())
 	{
@@ -1210,37 +1188,34 @@ void SRigVMExecutionStackView::HandlePreviewHostUpdated(FRigVMEditorBase* InEdit
 
 void SRigVMExecutionStackView::HandleItemMouseDoubleClick(TSharedPtr<FRigStackEntry> InItem)
 {
-	if (TSharedPtr<FRigVMEditorBase> Editor = WeakEditor.Pin())
+	if (!RigVMEditor.IsValid() || !RigVMBlueprint.IsValid())
 	{
-		if (!RigVMBlueprint.IsValid())
-		{
-			return;
-		}
+		return;
+	}
 	
-		URigVMHost* Host = Cast<URigVMHost>(RigVMBlueprint->GetObjectBeingDebugged());
-		if (!Host || !Host->GetVM())
-		{
-			return;
-		}
+	URigVMHost* Host = Cast<URigVMHost>(RigVMBlueprint->GetObjectBeingDebugged());
+	if (!Host || !Host->GetVM())
+	{
+		return;
+	}
 
-		const FRigVMByteCode& ByteCode = Host->GetVM()->GetByteCode();
-		if (URigVMNode* Subject = Cast<URigVMNode>(ByteCode.GetSubjectForInstruction(InItem->InstructionIndex)))
+	const FRigVMByteCode& ByteCode = Host->GetVM()->GetByteCode();
+	if (URigVMNode* Subject = Cast<URigVMNode>(ByteCode.GetSubjectForInstruction(InItem->InstructionIndex)))
+	{
+		URigVMGraph* GraphToFocus = Subject->GetGraph();
+		if (GraphToFocus && GraphToFocus->GetTypedOuter<URigVMAggregateNode>())
 		{
-			URigVMGraph* GraphToFocus = Subject->GetGraph();
-			if (GraphToFocus && GraphToFocus->GetTypedOuter<URigVMAggregateNode>())
+			if(URigVMGraph* ParentGraph = GraphToFocus->GetParentGraph())
 			{
-				if(URigVMGraph* ParentGraph = GraphToFocus->GetParentGraph())
-				{
-					GraphToFocus = ParentGraph;
-				} 
-			}
+				GraphToFocus = ParentGraph;
+			} 
+		}
 		
-			if(URigVMEdGraph* EdGraph = Cast<URigVMEdGraph>(RigVMBlueprint->GetEdGraph(GraphToFocus)))
+		if(URigVMEdGraph* EdGraph = Cast<URigVMEdGraph>(RigVMBlueprint->GetEdGraph(GraphToFocus)))
+		{
+			if(const UEdGraphNode* Node = EdGraph->FindNodeForModelNodeName(Subject->GetFName()))
 			{
-				if(const UEdGraphNode* Node = EdGraph->FindNodeForModelNodeName(Subject->GetFName()))
-				{
-					Editor->JumpToHyperlink(Node, false);
-				}
+				RigVMEditor.Pin()->JumpToHyperlink(Node, false);
 			}
 		}
 	}
diff --git a/Engine/Plugins/Runtime/RigVM/Source/RigVMEditor/Private/Widgets/SRigVMGraphBreakLinksWidget.cpp b/Engine/Plugins/Runtime/RigVM/Source/RigVMEditor/Private/Widgets/SRigVMGraphBreakLinksWidget.cpp
--- a/Engine/Plugins/Runtime/RigVM/Source/RigVMEditor/Private/Widgets/SRigVMGraphBreakLinksWidget.cpp
+++ b/Engine/Plugins/Runtime/RigVM/Source/RigVMEditor/Private/Widgets/SRigVMGraphBreakLinksWidget.cpp
@@ -11,12 +11,10 @@
 #include "Styling/AppStyle.h"
 #include "HAL/ConsoleManager.h"
 #include "RigVMModel/RigVMController.h"
-#include "Editor/RigVMNewEditor.h"
+#include "Editor/RigVMEditor.h"
 #include "ContentBrowserModule.h"
 #include "IContentBrowserSingleton.h"
 #include "Misc/ScopedSlowTask.h"
-#include "Widgets/Views/STableRow.h"
-#include "Widgets/Views/SListView.h"
 
 #define LOCTEXT_NAMESPACE "SRigVMGraphBreakLinksWidget"
 
diff --git a/Engine/Plugins/Runtime/RigVM/Source/RigVMEditor/Private/Widgets/SRigVMSwapAssetReferencesWidget.cpp b/Engine/Plugins/Runtime/RigVM/Source/RigVMEditor/Private/Widgets/SRigVMSwapAssetReferencesWidget.cpp
--- a/Engine/Plugins/Runtime/RigVM/Source/RigVMEditor/Private/Widgets/SRigVMSwapAssetReferencesWidget.cpp
+++ b/Engine/Plugins/Runtime/RigVM/Source/RigVMEditor/Private/Widgets/SRigVMSwapAssetReferencesWidget.cpp
@@ -8,8 +8,9 @@
 #include "AssetRegistry/AssetRegistryModule.h"
 #include "Styling/AppStyle.h"
 #include "Algo/Count.h"
+#include "BlueprintEditor.h"
 #include "RigVMEditorBlueprintLibrary.h"
-#include "Editor/RigVMNewEditor.h"
+#include "Editor/RigVMEditor.h"
 #include "Misc/UObjectToken.h"
 #include "RigVMModel/RigVMControllerActions.h"
 
@@ -441,7 +442,7 @@ FReply SRigVMSwapAssetReferencesWidget::OnNodeDoubleClicked(TSharedRef<FRigVMTre
 						GEditor->GetEditorSubsystem<UAssetEditorSubsystem>()->OpenEditorForAsset(TopLevelObject);
 						if(IAssetEditorInstance* Editor = GEditor->GetEditorSubsystem<UAssetEditorSubsystem>()->FindEditorForAsset(TopLevelObject, true))
 						{
-							if (FRigVMEditorBase* RigVMEditor = FRigVMEditorBase::GetFromAssetEditorInstance(Editor))
+							if(FRigVMEditor* RigVMEditor = static_cast<FRigVMEditor*>(Editor))
 							{
 								RigVMEditor->HandleJumpToHyperlink(ObjectReference);
 								return FReply::Handled();
diff --git a/Engine/Plugins/Runtime/RigVM/Source/RigVMEditor/Private/Widgets/SRigVMSwapFunctionsWidget.cpp b/Engine/Plugins/Runtime/RigVM/Source/RigVMEditor/Private/Widgets/SRigVMSwapFunctionsWidget.cpp
--- a/Engine/Plugins/Runtime/RigVM/Source/RigVMEditor/Private/Widgets/SRigVMSwapFunctionsWidget.cpp
+++ b/Engine/Plugins/Runtime/RigVM/Source/RigVMEditor/Private/Widgets/SRigVMSwapFunctionsWidget.cpp
@@ -8,7 +8,8 @@
 #include "AssetRegistry/AssetRegistryModule.h"
 #include "Styling/AppStyle.h"
 #include "Algo/Count.h"
-#include "Editor/RigVMNewEditor.h"
+#include "BlueprintEditor.h"
+#include "Editor/RigVMEditor.h"
 #include "Misc/UObjectToken.h"
 #include "RigVMModel/RigVMControllerActions.h"
 
@@ -424,7 +425,7 @@ bool FRigVMSwapFunctionTask::Execute(const TSharedRef<FRigVMTreePhase>& InPhase)
 				
 				if(IAssetEditorInstance* Editor = GEditor->GetEditorSubsystem<UAssetEditorSubsystem>()->FindEditorForAsset(Blueprint, true))
 				{
-					if (FRigVMEditorBase* RigVMEditor = FRigVMEditorBase::GetFromAssetEditorInstance(Editor))
+					if(FRigVMEditor* RigVMEditor = static_cast<FRigVMEditor*>(Editor))
 					{
 						RigVMEditor->HandleJumpToHyperlink(WeakRefNode.Get());
 					}
@@ -832,7 +833,7 @@ FReply SRigVMSwapFunctionsWidget::OnNodeDoubleClicked(TSharedRef<FRigVMTreeNode>
 						GEditor->GetEditorSubsystem<UAssetEditorSubsystem>()->OpenEditorForAsset(TopLevelObject);
 						if(IAssetEditorInstance* Editor = GEditor->GetEditorSubsystem<UAssetEditorSubsystem>()->FindEditorForAsset(TopLevelObject, true))
 						{
-							if (FRigVMEditorBase* RigVMEditor = FRigVMEditorBase::GetFromAssetEditorInstance(Editor))
+							if(FRigVMEditor* RigVMEditor = static_cast<FRigVMEditor*>(Editor))
 							{
 								RigVMEditor->HandleJumpToHyperlink(ObjectReference);
 								return FReply::Handled();
diff --git a/Engine/Plugins/Runtime/RigVM/Source/RigVMEditor/Public/EdGraph/NodeSpawners/RigVMEdGraphFunctionRefNodeSpawner.h b/Engine/Plugins/Runtime/RigVM/Source/RigVMEditor/Public/EdGraph/NodeSpawners/RigVMEdGraphFunctionRefNodeSpawner.h
--- a/Engine/Plugins/Runtime/RigVM/Source/RigVMEditor/Public/EdGraph/NodeSpawners/RigVMEdGraphFunctionRefNodeSpawner.h
+++ b/Engine/Plugins/Runtime/RigVM/Source/RigVMEditor/Public/EdGraph/NodeSpawners/RigVMEdGraphFunctionRefNodeSpawner.h
@@ -64,5 +64,5 @@ private:
 	static URigVMEdGraphNode* SpawnNode(UEdGraph* ParentGraph, UBlueprint* Blueprint, FRigVMGraphFunctionHeader& InFunction, FVector2D const Location);
 
 	friend class UEngineTestControlRig;
-	friend class FRigVMEditorBase;
+	friend class FRigVMEditor;
 };
diff --git a/Engine/Plugins/Runtime/RigVM/Source/RigVMEditor/Public/EdGraph/NodeSpawners/RigVMEdGraphTemplateNodeSpawner.h b/Engine/Plugins/Runtime/RigVM/Source/RigVMEditor/Public/EdGraph/NodeSpawners/RigVMEdGraphTemplateNodeSpawner.h
--- a/Engine/Plugins/Runtime/RigVM/Source/RigVMEditor/Public/EdGraph/NodeSpawners/RigVMEdGraphTemplateNodeSpawner.h
+++ b/Engine/Plugins/Runtime/RigVM/Source/RigVMEditor/Public/EdGraph/NodeSpawners/RigVMEdGraphTemplateNodeSpawner.h
@@ -49,5 +49,5 @@ private:
 	static URigVMEdGraphNode* SpawnNode(UEdGraph* ParentGraph, UBlueprint* Blueprint, const FRigVMTemplate* Template, FVector2D const Location);
 
 	friend class UEngineTestControlRig;
-	friend class FRigVMEditorBase;
+	friend class FRigVMEditor;
 };
diff --git a/Engine/Plugins/Runtime/RigVM/Source/RigVMEditor/Public/EdGraph/NodeSpawners/RigVMEdGraphUnitNodeSpawner.h b/Engine/Plugins/Runtime/RigVM/Source/RigVMEditor/Public/EdGraph/NodeSpawners/RigVMEdGraphUnitNodeSpawner.h
--- a/Engine/Plugins/Runtime/RigVM/Source/RigVMEditor/Public/EdGraph/NodeSpawners/RigVMEdGraphUnitNodeSpawner.h
+++ b/Engine/Plugins/Runtime/RigVM/Source/RigVMEditor/Public/EdGraph/NodeSpawners/RigVMEdGraphUnitNodeSpawner.h
@@ -51,5 +51,5 @@ private:
 	static URigVMEdGraphNode* SpawnNode(UEdGraph* ParentGraph, UBlueprint* Blueprint, UScriptStruct* StructTemplate, const FName& InMethodName, FVector2D const Location);
 
 	friend class UEngineTestControlRig;
-	friend class FRigVMEditorBase;
+	friend class FRigVMEditor;
 };
diff --git a/Engine/Plugins/Runtime/RigVM/Source/RigVMEditor/Public/Editor/RigVMActionMenuBuilder.h b/Engine/Plugins/Runtime/RigVM/Source/RigVMEditor/Public/Editor/RigVMActionMenuBuilder.h
deleted file mode 100644
diff --git a/Engine/Plugins/Runtime/RigVM/Source/RigVMEditor/Public/Editor/RigVMActionMenuItem.h b/Engine/Plugins/Runtime/RigVM/Source/RigVMEditor/Public/Editor/RigVMActionMenuItem.h
deleted file mode 100644
diff --git a/Engine/Plugins/Runtime/RigVM/Source/RigVMEditor/Public/Editor/RigVMActionMenuUtils.h b/Engine/Plugins/Runtime/RigVM/Source/RigVMEditor/Public/Editor/RigVMActionMenuUtils.h
deleted file mode 100644
diff --git a/Engine/Plugins/Runtime/RigVM/Source/RigVMEditor/Public/Editor/RigVMEditor.h b/Engine/Plugins/Runtime/RigVM/Source/RigVMEditor/Public/Editor/RigVMEditor.h
--- a/Engine/Plugins/Runtime/RigVM/Source/RigVMEditor/Public/Editor/RigVMEditor.h
+++ b/Engine/Plugins/Runtime/RigVM/Source/RigVMEditor/Public/Editor/RigVMEditor.h
@@ -3,54 +3,17 @@
 #pragma once
 
 #include "CoreMinimal.h"
-#include "RigVMEditorGraphExplorerTabSummoner.h"
-#include "RigVMNewEditorMode.h"
-#include "RigVMCore/RigVMExecuteContext.h"
-#include "RigVMCore/RigVMGraphFunctionDefinition.h"
-#include "RigVMModel/RigVMController.h"
-#include "WorkflowOrientedApp/WorkflowCentricApplication.h"
-#include "RigVMModel/RigVMNotifications.h"
-#include "Widgets/SRigVMEditorGraphExplorer.h"
-#include "WorkflowOrientedApp/WorkflowTabManager.h"
-#include "WorkflowOrientedApp/WorkflowUObjectDocuments.h"
-#include "GraphEditor.h"
-#include "Misc/NotifyHook.h"
-#include "TickableEditorObject.h"
-#include "EditorUndoClient.h"
-#include "RigVMEditorModule.h"
-#include "SNodePanel.h"
+#include "BlueprintEditor.h"
+#include "RigVMHost.h"
+#include "RigVMBlueprint.h"
 #include "Editor/RigVMDetailsViewWrapperObject.h"
+#include "AssetRegistry/AssetData.h"
 
-#define LOCTEXT_NAMESPACE "RigVMEditor"
-
-class FDocumentTracker;
+class FRigVMGraphExplorerDragDropOp;
 class SRigVMEditorGraphExplorer;
-class URigVM;
-class FTabInfo;
-class FRigVMEditorBase;
-class URigVMController;
-class URigVMHost;
-class FTransaction;
-class URigVMBlueprint;
-class SGraphEditor;
-
-DECLARE_MULTICAST_DELEGATE_TwoParams(FRigVMEditorClosed, const FRigVMEditorBase*, URigVMBlueprint*);
+class FRigVMEditor;
 
-
-/* Enums to use when grouping the blueprint members in the list panel. The order here will determine the order in the list */
-namespace RigVMNodeSectionID
-{
-	// Keep the values as they are defined in NodeSectionID, which is defined in BlueprintEditor.h
-	// TODO: Once there is no need for FRigVMLegacyEditor, improve the definition of this enum, including a uint8 definition
-	enum Type
-	{
-		NONE = 0,
-		GRAPH = 1,					// Graph
-		FUNCTION = 4,				// Functions
-		VARIABLE = 8,				// Variables
-		LOCAL_VARIABLE = 12			// Local variables
-	};
-};
+DECLARE_MULTICAST_DELEGATE_TwoParams(FRigVMEditorClosed, const FRigVMEditor*, URigVMBlueprint*);
 
 struct FRigVMEditorModes
 {
@@ -74,144 +37,12 @@ private:
 	FRigVMEditorModes() {}
 };
 
-struct FRigVMCustomDebugObject
-{
-public:
-	// Custom object to include, regardless of the current debugging World
-	UObject* Object;
-
-	// Override for the object name (if not empty)
-	FString NameOverride;
-
-public:
-	FRigVMCustomDebugObject()
-		: Object(nullptr)
-	{
-	}
-
-	FRigVMCustomDebugObject(UObject* InObject, const FString& InLabel)
-		: Object(InObject)
-		, NameOverride(InLabel)
-	{
-	}
-};
-
-class RIGVMEDITOR_API IRigVMEditor
-{
-public:
-	
-	virtual TSharedRef<IRigVMEditor> SharedRef() = 0;
-
-	virtual TSharedPtr<FAssetEditorToolkit> GetHostingApp() = 0;
-	virtual void InitAssetEditor(const EToolkitMode::Type Mode, const TSharedPtr<IToolkitHost>& InitToolkitHost, const FName AppIdentifier, const TSharedRef<FTabManager::FLayout>& StandaloneDefaultLayout, const bool bCreateDefaultStandaloneMenu, const bool bCreateDefaultToolbar, const TArray<UObject*>& ObjectsToEdit, const bool bInIsToolbarFocusable = false, const bool bInUseSmallToolbarIcons = false, const TOptional<EAssetOpenMethod>& InOpenMethod = TOptional<EAssetOpenMethod>()) = 0;
-	virtual void CreateEditorToolbar() = 0;
-	virtual void CommonInitialization(const TArray<UBlueprint*>& InitBlueprints, bool bShouldOpenInDefaultsMode) = 0;
-	virtual TSharedPtr<FDocumentTracker> GetDocumentManager() const = 0;
-	virtual void AddApplicationMode(FName ModeName, TSharedRef<FApplicationMode> Mode) = 0;
-	virtual void RegenerateMenusAndToolbars() = 0;
-	virtual void SetCurrentMode(FName NewMode) = 0;
-	virtual FEditorModeTools& GetToolkitEditorModeManager() const = 0;
-	virtual void PostLayoutBlueprintEditorInitialization() = 0;
-	virtual TSharedPtr<SGraphEditor> OpenGraphAndBringToFront(UEdGraph* Graph, bool bSetFocus = true) = 0;
-	virtual bool FindOpenTabsContainingDocument(const UObject* DocumentID, /*inout*/ TArray< TSharedPtr<SDockTab> >& Results) = 0;
-	virtual TSharedPtr<SDockTab> OpenDocument(const UObject* DocumentID, FDocumentTracker::EOpenDocumentCause Cause) = 0;
-	virtual void CloseDocumentTab(const UObject* DocumentID) = 0;
-	virtual TSharedPtr<FTabManager> GetTabManager() = 0;
-#if WITH_RIGVMLEGACYEDITOR
-	virtual TSharedPtr<class SKismetInspector> GetInspector() const = 0;
-#else
-	virtual TSharedPtr<class SRigVMDetailsInspector> GetInspector() const = 0;
-#endif
-	virtual TSharedPtr<FApplicationMode> CreateEditorMode() = 0;
-	virtual const FName GetEditorAppName() const = 0;
-	virtual const TArray< UObject* >& GetEditingObjects() const = 0;
-	virtual const TSharedRef<IToolkitHost> GetToolkitHost() const = 0;
-	virtual bool IsHosted() const = 0;
-	virtual void BringToolkitToFrontImpl() = 0;
-	virtual void InitToolMenuContext(FToolMenuContext& MenuContext) = 0;
-	virtual TSharedRef<FUICommandList> GetToolkitCommands() = 0;
-	virtual TWeakPtr<class SGraphEditor> GetFocusedGraphEditor() = 0;
-	virtual TWeakPtr<FDocumentTabFactory> GetGraphEditorTabFactory() const = 0;
-	virtual bool TransactionObjectAffectsBlueprint(UObject* InTransactedObject) = 0;
-	virtual FEdGraphPinType GetLastPinTypeUsed() = 0;
-	virtual void LogSimpleMessage(const FText& MessageText) = 0;
-	virtual void RenameNewlyAddedAction(FName InActionName) = 0;
-	virtual FGraphPanelSelectionSet GetSelectedNodes() const = 0;
-	virtual void SetUISelectionState(FName SelectionOwner) = 0;
-	virtual void AnalyticsTrackNodeEvent(UBlueprint* Blueprint, UEdGraphNode *GraphNode, bool bNodeDelete = false) const = 0;
-	virtual void JumpToHyperlink(const UObject* ObjectReference, bool bRequestRename = false) = 0;
-	virtual void PostUndo(bool bSuccess) = 0;
-	virtual void PostRedo(bool bSuccess) = 0;
-	virtual UEdGraphPin* GetCurrentlySelectedPin() const = 0;
-	virtual void CreateDefaultCommands() = 0;
-	virtual TSharedRef<SGraphEditor> CreateGraphEditorWidget(TSharedRef<class FTabInfo> InTabInfo, class UEdGraph* InGraph) = 0;
-	virtual void CompileImpl() = 0;
-	virtual void SaveAsset_Execute_Impl() = 0;
-	virtual void SaveAssetAs_Execute_Impl() = 0;
-	virtual bool IsGraphInCurrentBlueprint(const UEdGraph* InGraph) const = 0;
-	virtual bool IsEditableImpl(UEdGraph* InGraph) const = 0;
-	virtual UEdGraph* GetFocusedGraph() const = 0;
-	virtual void JumpToNode(const class UEdGraphNode* Node, bool bRequestRename = false) = 0;
-	virtual void JumpToPin(const class UEdGraphPin* Pin) = 0;
-	virtual void AddToolbarExtender(TSharedPtr<FExtender> Extender) = 0;
-	virtual void RemoveToolbarExtender(TSharedPtr<FExtender> Extender) = 0;
-	virtual void AddMenuExtender(TSharedPtr<FExtender> Extender) = 0;
-	virtual void RemoveMenuExtender(TSharedPtr<FExtender> Extender) = 0;
-	virtual TSharedPtr<class IMessageLogListing> GetCompilerResultsListing() = 0;
-	virtual void OnBlueprintChangedInnerImpl(UBlueprint* InBlueprint, bool bIsJustBeingCompiled = false) = 0;
-	virtual void RefreshEditorsImpl(ERefreshRigVMEditorReason::Type Reason = ERefreshRigVMEditorReason::UnknownReason) = 0;
-	virtual void SetupGraphEditorEventsImpl(UEdGraph* InGraph, SGraphEditor::FGraphEditorEvents& InEvents) = 0;
-	virtual FActionMenuContent OnCreateGraphActionMenu(UEdGraph* InGraph, const FVector2D& InNodePosition, const TArray<UEdGraphPin*>& InDraggedPins, bool bAutoExpand, SGraphEditor::FActionMenuClosed InOnMenuClosed) = 0;
-	virtual void AddReferencedObjectsImpl(FReferenceCollector& Collector) = 0;
-	virtual FGraphAppearanceInfo GetGraphAppearanceImpl(class UEdGraph* InGraph) const = 0;
-	virtual void NotifyPreChangeImpl( FProperty* PropertyAboutToChange ) = 0;
-	virtual void NotifyPostChangeImpl( const FPropertyChangedEvent& PropertyChangedEvent, FProperty* PropertyThatChanged) = 0;
-	virtual FName GetSelectedVariableName() = 0;
-	virtual bool IsNodeTitleVisible(const UEdGraphNode* Node, bool bRequestRename) = 0;
-	virtual void EditClassDefaults_Clicked() = 0;
-	virtual void EditGlobalOptions_Clicked() = 0;
-	virtual void TryInvokingDetailsTab(bool bFlash = true) = 0;
-	virtual FName GetGraphExplorerWidgetID() = 0;
-	virtual void RefreshInspector() = 0;
-	virtual void RefreshStandAloneDefaultsEditor() = 0;
-	virtual void GetPinTypeSelectorFilters(TArray<TSharedPtr<class IPinTypeSelectorFilter>>& OutFilters) const = 0;
-	virtual void OnAddNewVariable() = 0;
-	virtual void SummonSearchUI(bool bSetFindWithinBlueprint, FString NewSearchTerms = FString(), bool bSelectFirstResult = false) = 0;
-	virtual void ZoomToSelection_Clicked() = 0;
-	virtual void RestoreEditedObjectState() = 0;
-	virtual void SetupViewForBlueprintEditingMode() = 0;
-	virtual bool GetIsContextSensitive() = 0;
-	virtual void SetIsContextSensitive(const bool bIsContextSensitive) = 0;
-	virtual void RegisterToolbarTab(const TSharedRef<class FTabManager>& TabManager) = 0;
-	virtual const TArray< UObject* >* GetObjectsCurrentlyBeingEdited() const = 0;
-	virtual void AddSelectedDebugObjectWidget(FToolBarBuilder& ToolbarBuilder) = 0;
-};
-
-class RIGVMEDITOR_API FRigVMEditorBase : public IRigVMEditor
+class RIGVMEDITOR_API FRigVMEditor : public FBlueprintEditor
 {
-protected:
-	// Type of new document/graph being created by a menu item
-	enum ECreatedDocumentType
-	{
-		CGT_NewVariable,
-		CGT_NewFunctionGraph,
-		CGT_NewMacroGraph,
-		CGT_NewAnimationLayer,
-		CGT_NewEventGraph,
-		CGT_NewLocalVariable
-	};
-	
-	FRigVMEditorBase();
-	virtual ~FRigVMEditorBase(){}
-
-
-	void UnbindEditor();
-
 public:
-	static FRigVMEditorBase* GetFromAssetEditorInstance(IAssetEditorInstance* Instance);
 
-	DECLARE_EVENT(FRigVMEditorBase, FOnRefreshEvent);
-	FOnRefreshEvent OnRefresh() { return OnRefreshEvent; }
+	FRigVMEditor();
+	virtual ~FRigVMEditor();
 
 	/**
 	 * Edits the specified asset(s)
@@ -221,60 +52,56 @@ public:
 	 * @param	InRigVMBlueprint	The blueprint object to start editing.
 	 */
 	virtual void InitRigVMEditor(const EToolkitMode::Type Mode, const TSharedPtr< class IToolkitHost >& InitToolkitHost, class URigVMBlueprint* InRigVMBlueprint);
-	
-protected:
 
 	virtual void HandleAssetRequestedOpen(UObject* InObject);
 	virtual void HandleAssetRequestClose(UObject* InObject, EAssetEditorCloseReason InReason);
 	bool bRequestedReopen = false;
 
+	virtual const FName GetEditorAppName() const;
 	virtual const FName GetEditorModeName() const;
+	virtual TSharedPtr<FApplicationMode> CreateEditorMode();
 
 	// FBlueprintEditor interface
-public:
-	virtual UBlueprint* GetBlueprintObj() const;
-protected:
-	virtual TSubclassOf<UEdGraphSchema> GetDefaultSchemaClass() const;
-	virtual bool InEditingMode() const;
-	virtual void OnGraphEditorFocused(const TSharedRef<SGraphEditor>& InGraphEditor);
+	virtual UBlueprint* GetBlueprintObj() const override;
+	virtual TSubclassOf<UEdGraphSchema> GetDefaultSchemaClass() const override;
+	virtual bool InEditingMode() const override;
+	virtual void OnGraphEditorFocused(const TSharedRef<SGraphEditor>& InGraphEditor) override;
 
 	//  FTickableEditorObject Interface
-	virtual void Tick(float DeltaTime);
+	virtual void Tick(float DeltaTime) override;
 
 	// IToolkit Interface
-	virtual void BringToolkitToFront();
-	virtual FName GetToolkitFName() const;
-	virtual FName GetToolkitContextFName() const;
-	virtual FText GetBaseToolkitName() const;
-	virtual FText GetToolkitToolTipText() const;
-	virtual FString GetWorldCentricTabPrefix() const;
-	virtual FLinearColor GetWorldCentricTabColorScale() const;	
-	virtual void InitToolMenuContextImpl(FToolMenuContext& MenuContext);
+	virtual void BringToolkitToFront() override;
+	virtual FName GetToolkitFName() const override;
+	virtual FName GetToolkitContextFName() const override;
+	virtual FText GetBaseToolkitName() const override;
+	virtual FText GetToolkitToolTipText() const override;
+	virtual FString GetWorldCentricTabPrefix() const override;
+	virtual FLinearColor GetWorldCentricTabColorScale() const override;	
+	virtual void InitToolMenuContext(FToolMenuContext& MenuContext) override;
 
 	// BlueprintEditor interface
-	virtual bool TransactionObjectAffectsBlueprintImpl(UObject* InTransactedObject);
-	virtual bool CanAddNewLocalVariable() const;
-	virtual void OnAddNewLocalVariable();
-	virtual void OnPasteNewLocalVariable(const FBPVariableDescription& VariableDescription);
-public:
-	virtual void DeleteSelectedNodes();
-	virtual bool CanDeleteNodes() const;
-protected:
-	virtual void CopySelectedNodes();
-	virtual bool CanCopyNodes() const;
-	virtual void PasteNodes();
-	virtual bool CanPasteNodes() const;
-	virtual bool IsNativeParentClassCodeLinkEnabled() const { return false; }
-	virtual bool ReparentBlueprint_IsVisible() const { return false; }
-	virtual FReply OnSpawnGraphNodeByShortcut(FInputChord InChord, const FVector2D& InPosition, UEdGraph* InGraph);
-	virtual bool ShouldLoadBPLibrariesFromAssetRegistry() { return false; }
-	virtual bool JumpToHyperlinkImpl(const UObject* ObjectReference, bool bRequestRename = false);
+	virtual bool TransactionObjectAffectsBlueprint(UObject* InTransactedObject) override;
+	virtual bool CanAddNewLocalVariable() const override;
+	virtual void OnAddNewLocalVariable() override;
+	virtual void OnPasteNewLocalVariable(const FBPVariableDescription& VariableDescription) override;
+	virtual void DeleteSelectedNodes() override;
+	virtual bool CanDeleteNodes() const override;
+	virtual void CopySelectedNodes() override;
+	virtual bool CanCopyNodes() const override;
+	virtual void PasteNodes() override;
+	virtual bool CanPasteNodes() const override;
+	virtual bool IsNativeParentClassCodeLinkEnabled() const override { return false; }
+	virtual bool ReparentBlueprint_IsVisible() const override { return false; }
+	virtual FReply OnSpawnGraphNodeByShortcut(FInputChord InChord, const FVector2D& InPosition, UEdGraph* InGraph) override;
+	virtual bool ShouldLoadBPLibrariesFromAssetRegistry() override { return false; }
+	virtual void JumpToHyperlink(const UObject* ObjectReference, bool bRequestRename = false) override;
 	virtual bool ShouldOpenGraphByDefault() const { return true; }
-	virtual void AddNewFunctionVariant(const UEdGraph* InOriginalFunction);
+	virtual void AddNewFunctionVariant(const UEdGraph* InOriginalFunction) override;
 
 	// FEditorUndoClient Interface
-	virtual void PostUndoImpl(bool bSuccess);
-	virtual void PostRedoImpl(bool bSuccess);
+	virtual void PostUndo(bool bSuccess) override;
+	virtual void PostRedo(bool bSuccess) override;
 
 	virtual void PostTransaction(bool bSuccess, const FTransaction* Transaction, bool bIsRedo) {};
 
@@ -285,20 +112,17 @@ protected:
 	virtual bool CanStopWatchingPin() const;
 
 	// IToolkitHost Interface
-	virtual void OnToolkitHostingStarted(const TSharedRef<class IToolkit>& Toolkit);
-	virtual void OnToolkitHostingFinished(const TSharedRef<class IToolkit>& Toolkit);
+	virtual void OnToolkitHostingStarted(const TSharedRef<class IToolkit>& Toolkit) override;
+	virtual void OnToolkitHostingFinished(const TSharedRef<class IToolkit>& Toolkit) override;
 
 	//  FTickableEditorObject Interface
-	virtual TStatId GetStatId() const;
+	virtual TStatId GetStatId() const override;
 
-public:
 	// returns the blueprint being edited
 	URigVMBlueprint* GetRigVMBlueprint() const;
 
 	// returns the currently debugged / viewed host
 	URigVMHost* GetRigVMHost() const;
-	
-protected:
 
 	virtual UObject* GetOuterForHost() const;
 
@@ -308,17 +132,15 @@ protected:
 	// allows the editor to fill an empty graph
 	virtual void CreateEmptyGraphContent(URigVMController* InController) {}
 
-	DECLARE_EVENT_OneParam(FRigVMEditorBase, FPreviewHostUpdated, FRigVMEditorBase*);
+	DECLARE_EVENT_OneParam(FRigVMEditor, FPreviewHostUpdated, FRigVMEditor*);
 	FPreviewHostUpdated& OnPreviewHostUpdated() { return PreviewHostUpdated;  }
 
-public:
 	FRigVMEditorClosed& OnEditorClosed() { return RigVMEditorClosedDelegate; }
-protected:
 
 	/** Get the toolbox hosting widget */
 	TSharedRef<SBorder> GetToolbox() { return Toolbox.ToSharedRef(); }
 
-	virtual bool SelectLocalVariable(const UEdGraph* Graph, const FName& VariableName);
+	virtual bool SelectLocalVariable(const UEdGraph* Graph, const FName& VariableName) override;
 
 protected:
 
@@ -329,26 +151,29 @@ protected:
 	};
 
 	// FBlueprintEditor Interface
-	virtual void CreateDefaultCommandsImpl();
+	virtual void CreateDefaultCommands() override;
 	virtual void OnCreateGraphEditorCommands(TSharedPtr<FUICommandList> GraphEditorCommandsList);
-	virtual void Compile();
-	virtual void SaveAsset_Execute();
-	virtual void SaveAssetAs_Execute();
-	virtual bool IsInAScriptingMode() const { return true; }
-	virtual void NewDocument_OnClicked(ECreatedDocumentType GraphType);
-	bool IsSectionVisibleImpl(RigVMNodeSectionID::Type InSectionID) const;
-	virtual bool AreEventGraphsAllowed() const;
-	virtual bool AreMacrosAllowed() const;
-	virtual bool AreDelegatesAllowed() const;
-	bool NewDocument_IsVisibleForTypeImpl(ECreatedDocumentType GraphType) const;
-	virtual FGraphAppearanceInfo GetGraphAppearance(class UEdGraph* InGraph) const;
-	virtual bool IsCompilingEnabled() const;
-	virtual void OnActiveTabChanged( TSharedPtr<SDockTab> PreviouslyActive, TSharedPtr<SDockTab> NewlyActivated );
-	virtual void OnSelectedNodesChangedImpl(const TSet<class UObject*>& NewSelection);
-	virtual void OnBlueprintChangedImpl(UBlueprint* InBlueprint, bool bIsJustBeingCompiled);
-	virtual void RefreshEditors(ERefreshRigVMEditorReason::Type Reason = ERefreshRigVMEditorReason::UnknownReason);
-	virtual void SetupGraphEditorEvents(UEdGraph* InGraph, SGraphEditor::FGraphEditorEvents& InEvents);
-	virtual void FocusInspectorOnGraphSelection(const TSet<class UObject*>& NewSelection, bool bForceRefresh = false);
+	virtual TSharedRef<SGraphEditor> CreateGraphEditorWidget(TSharedRef<FTabInfo> InTabInfo, UEdGraph* InGraph) override;
+	virtual void Compile() override;
+	virtual void SaveAsset_Execute() override;
+	virtual void SaveAssetAs_Execute() override;
+	virtual bool IsInAScriptingMode() const override { return true; }
+	virtual void NewDocument_OnClicked(ECreatedDocumentType GraphType) override;
+	virtual bool IsSectionVisible(NodeSectionID::Type InSectionID) const override;
+	virtual bool AreEventGraphsAllowed() const override;
+	virtual bool AreMacrosAllowed() const override;
+	virtual bool AreDelegatesAllowed() const override;
+	virtual bool NewDocument_IsVisibleForType(ECreatedDocumentType GraphType) const;
+	virtual FGraphAppearanceInfo GetGraphAppearance(class UEdGraph* InGraph) const override;
+	virtual bool IsEditable(UEdGraph* InGraph) const override;
+	virtual bool IsCompilingEnabled() const override;
+	virtual FText GetGraphDecorationString(UEdGraph* InGraph) const override;
+	virtual void OnActiveTabChanged( TSharedPtr<SDockTab> PreviouslyActive, TSharedPtr<SDockTab> NewlyActivated ) override;
+	virtual void OnSelectedNodesChangedImpl(const TSet<class UObject*>& NewSelection) override;
+	virtual void OnBlueprintChangedImpl(UBlueprint* InBlueprint, bool bIsJustBeingCompiled) override;
+	virtual void RefreshEditors(ERefreshBlueprintEditorReason::Type Reason = ERefreshBlueprintEditorReason::UnknownReason) override;
+	virtual void SetupGraphEditorEvents(UEdGraph* InGraph, SGraphEditor::FGraphEditorEvents& InEvents) override;
+	virtual void FocusInspectorOnGraphSelection(const TSet<class UObject*>& NewSelection, bool bForceRefresh = false) override;
 
 	virtual void HandleModifiedEvent(ERigVMGraphNotifType InNotifType, URigVMGraph* InGraph, UObject* InSubject);
 	UE_DEPRECATED(5.4, "Please use HandleVMCompiledEvent with ExtendedExecuteContext param.")
@@ -359,10 +184,10 @@ protected:
 	void SetHaltedNode(URigVMNode* Node);
 	
 	// FNotifyHook Interface
-	virtual void NotifyPreChange(FProperty* PropertyAboutToChange);
-	virtual void NotifyPostChange(const FPropertyChangedEvent& PropertyChangedEvent, FProperty* PropertyThatChanged);
+	virtual void NotifyPreChange(FProperty* PropertyAboutToChange) override;
+	virtual void NotifyPostChange(const FPropertyChangedEvent& PropertyChangedEvent, FProperty* PropertyThatChanged) override;
 	/** delegate for changing property */
-	virtual void OnFinishedChangingProperties(const FPropertyChangedEvent& PropertyChangedEvent);
+	virtual void OnFinishedChangingProperties(const FPropertyChangedEvent& PropertyChangedEvent) override;
 	void OnPropertyChanged(UObject* InObject, FPropertyChangedEvent& InEvent);
 	virtual void OnWrappedPropertyChangedChainEvent(URigVMDetailsViewWrapperObject* InWrapperObject, const FString& InPropertyPath, FPropertyChangedChainEvent& InPropertyChangedChainEvent);
 	void OnRequestLocalizeFunctionDialog(FRigVMGraphFunctionIdentifier& InFunction, URigVMController* InTargetController, IRigVMGraphFunctionHost* InTargetFunctionHost, bool bForce);
@@ -372,8 +197,6 @@ protected:
 
 public:
 	void HandleJumpToHyperlink(const UObject* InSubject);
-	virtual FText GetGraphDecorationString(UEdGraph* InGraph) const;
-	virtual bool IsEditable(UEdGraph* InGraph) const;
 
 protected:
 	bool UpdateDefaultValueForVariable(FBPVariableDescription& InVariable, bool bUseCDO);
@@ -388,7 +211,7 @@ protected:
 	virtual void CacheNameLists();
 
 	// FGCObject Interface
-	virtual void AddReferencedObjects( FReferenceCollector& Collector );
+	virtual void AddReferencedObjects( FReferenceCollector& Collector ) override;
 
 	virtual void BindCommands();
 
@@ -413,7 +236,7 @@ protected:
 	/** Redo Action **/
 	void RedoAction();
 	
-	virtual void OnCreateComment();
+	virtual void OnCreateComment() override;
 
 	bool IsDetailsPanelRefreshSuspended() const { return bSuspendDetailsPanelRefresh; }
 	bool& GetSuspendDetailsPanelRefreshFlag() { return bSuspendDetailsPanelRefresh; }
@@ -468,10 +291,10 @@ protected:
 	static FSlateIcon GetExecutionModeIcon(const ERigVMEditorExecutionModeType InExecutionMode);
 	FSlateIcon GetExecutionModeIcon() const;
 	
-	virtual void GetCustomDebugObjects(TArray<FRigVMCustomDebugObject>& DebugList) const;
-	virtual bool OnlyShowCustomDebugObjects() const { return true; }
-	virtual void HandleSetObjectBeingDebugged(UObject* InObject);
-	virtual FString GetCustomDebugObjectLabel(UObject* ObjectBeingDebugged) const;
+	virtual void GetCustomDebugObjects(TArray<FCustomDebugObject>& DebugList) const override;
+	virtual bool OnlyShowCustomDebugObjects() const override { return true; }
+	virtual void HandleSetObjectBeingDebugged(UObject* InObject) override;
+	virtual FString GetCustomDebugObjectLabel(UObject* ObjectBeingDebugged) const override;
 
 	/** Handle hiding items in the graph */
 	void HandleHideItem();
@@ -486,7 +309,7 @@ protected:
 	virtual void OnGraphNodeClicked(URigVMEdGraphNode* InNode);
 	virtual void OnNodeDoubleClicked(URigVMBlueprint* InBlueprint, URigVMNode* InNode);
 	virtual void OnGraphImported(UEdGraph* InEdGraph);
-	virtual bool OnActionMatchesName(FEdGraphSchemaAction* InAction, const FName& InName) const;
+	virtual bool OnActionMatchesName(FEdGraphSchemaAction* InAction, const FName& InName) const override;
 	virtual void HandleShowCurrentStatement();
 	virtual void HandleBreakpointActionRequested(const ERigVMBreakpointAction BreakpointAction);
 	virtual bool IsHaltedAtBreakpoint() const;
@@ -512,19 +335,16 @@ private:
 
 	FPreviewHostUpdated PreviewHostUpdated;
 
+	FRigVMEditorClosed RigVMEditorClosedDelegate;
+
 	/** Toolbox hosting widget */
 	TSharedPtr<SBorder> Toolbox;
 
-protected:
 	TSharedPtr<SRigVMEditorGraphExplorer> GraphExplorerWidget;
 
-	FRigVMEditorClosed RigVMEditorClosedDelegate;
-
 public:
 	void SetGraphExplorerWidget(TSharedPtr<SRigVMEditorGraphExplorer> InWidget) { GraphExplorerWidget = InWidget; }
 
-	bool IsEditingSingleBlueprint() const;
-	
 private:
 	bool bAnyErrorsLeft;
 	TMap<FString, FString> KnownInstructionLimitWarnings;
@@ -561,99 +381,6 @@ protected:
 
 	FDelegateHandle PropertyChangedHandle;
 
-	FOnRefreshEvent OnRefreshEvent;
-
 	friend class SRigVMExecutionStackView;
 	friend class SRigVMEditorGraphExplorer;
 };
-
-struct FRigVMEditorZoomLevelsContainer : public FZoomLevelsContainer
-{
-	struct FRigVMEditorZoomLevelEntry
-	{
-	public:
-		FRigVMEditorZoomLevelEntry(float InZoomAmount, const FText& InDisplayText, EGraphRenderingLOD::Type InLOD)
-			: DisplayText(FText::Format(NSLOCTEXT("GraphEditor", "Zoom", "Zoom {0}"), InDisplayText))
-		, ZoomAmount(InZoomAmount)
-		, LOD(InLOD)
-		{
-		}
-
-	public:
-		FText DisplayText;
-		float ZoomAmount;
-		EGraphRenderingLOD::Type LOD;
-	};
-	
-	FRigVMEditorZoomLevelsContainer()
-	{
-		ZoomLevels.Reserve(22);
-		ZoomLevels.Add(FRigVMEditorZoomLevelEntry(0.025f, FText::FromString(TEXT("-14")), EGraphRenderingLOD::LowestDetail));
-		ZoomLevels.Add(FRigVMEditorZoomLevelEntry(0.070f, FText::FromString(TEXT("-13")), EGraphRenderingLOD::LowestDetail));
-		ZoomLevels.Add(FRigVMEditorZoomLevelEntry(0.100f, FText::FromString(TEXT("-12")), EGraphRenderingLOD::LowestDetail));
-		ZoomLevels.Add(FRigVMEditorZoomLevelEntry(0.125f, FText::FromString(TEXT("-11")), EGraphRenderingLOD::LowestDetail));
-		ZoomLevels.Add(FRigVMEditorZoomLevelEntry(0.150f, FText::FromString(TEXT("-10")), EGraphRenderingLOD::LowestDetail));
-		ZoomLevels.Add(FRigVMEditorZoomLevelEntry(0.175f, FText::FromString(TEXT("-9")), EGraphRenderingLOD::LowestDetail));
-		ZoomLevels.Add(FRigVMEditorZoomLevelEntry(0.200f, FText::FromString(TEXT("-8")), EGraphRenderingLOD::LowestDetail));
-		ZoomLevels.Add(FRigVMEditorZoomLevelEntry(0.225f, FText::FromString(TEXT("-7")), EGraphRenderingLOD::LowDetail));
-		ZoomLevels.Add(FRigVMEditorZoomLevelEntry(0.250f, FText::FromString(TEXT("-6")), EGraphRenderingLOD::LowDetail));
-		ZoomLevels.Add(FRigVMEditorZoomLevelEntry(0.375f, FText::FromString(TEXT("-5")), EGraphRenderingLOD::MediumDetail));
-		ZoomLevels.Add(FRigVMEditorZoomLevelEntry(0.500f, FText::FromString(TEXT("-4")), EGraphRenderingLOD::MediumDetail));
-		ZoomLevels.Add(FRigVMEditorZoomLevelEntry(0.675f, FText::FromString(TEXT("-3")), EGraphRenderingLOD::MediumDetail));
-		ZoomLevels.Add(FRigVMEditorZoomLevelEntry(0.750f, FText::FromString(TEXT("-2")), EGraphRenderingLOD::DefaultDetail));
-		ZoomLevels.Add(FRigVMEditorZoomLevelEntry(0.875f, FText::FromString(TEXT("-1")), EGraphRenderingLOD::DefaultDetail));
-		ZoomLevels.Add(FRigVMEditorZoomLevelEntry(1.000f, FText::FromString(TEXT("1:1")), EGraphRenderingLOD::DefaultDetail));
-		ZoomLevels.Add(FRigVMEditorZoomLevelEntry(1.250f, FText::FromString(TEXT("+1")), EGraphRenderingLOD::DefaultDetail));
-		ZoomLevels.Add(FRigVMEditorZoomLevelEntry(1.375f, FText::FromString(TEXT("+2")), EGraphRenderingLOD::DefaultDetail));
-		ZoomLevels.Add(FRigVMEditorZoomLevelEntry(1.500f, FText::FromString(TEXT("+3")), EGraphRenderingLOD::FullyZoomedIn));
-		ZoomLevels.Add(FRigVMEditorZoomLevelEntry(1.675f, FText::FromString(TEXT("+4")), EGraphRenderingLOD::FullyZoomedIn));
-		ZoomLevels.Add(FRigVMEditorZoomLevelEntry(1.750f, FText::FromString(TEXT("+5")), EGraphRenderingLOD::FullyZoomedIn));
-		ZoomLevels.Add(FRigVMEditorZoomLevelEntry(1.875f, FText::FromString(TEXT("+6")), EGraphRenderingLOD::FullyZoomedIn));
-		ZoomLevels.Add(FRigVMEditorZoomLevelEntry(2.000f, FText::FromString(TEXT("+7")), EGraphRenderingLOD::FullyZoomedIn));
-	}
-
-	float GetZoomAmount(int32 InZoomLevel) const override
-	{
-		checkSlow(ZoomLevels.IsValidIndex(InZoomLevel));
-		return ZoomLevels[InZoomLevel].ZoomAmount;
-	}
-
-	int32 GetNearestZoomLevel(float InZoomAmount) const override
-	{
-		for (int32 ZoomLevelIndex=0; ZoomLevelIndex < GetNumZoomLevels(); ++ZoomLevelIndex)
-		{
-			if (InZoomAmount <= GetZoomAmount(ZoomLevelIndex))
-			{
-				return ZoomLevelIndex;
-			}
-		}
-
-		return GetDefaultZoomLevel();
-	}
-	
-	FText GetZoomText(int32 InZoomLevel) const override
-	{
-		checkSlow(ZoomLevels.IsValidIndex(InZoomLevel));
-		return ZoomLevels[InZoomLevel].DisplayText;
-	}
-	
-	int32 GetNumZoomLevels() const override
-	{
-		return ZoomLevels.Num();
-	}
-	
-	int32 GetDefaultZoomLevel() const override
-	{
-		return 14;
-	}
-
-	EGraphRenderingLOD::Type GetLOD(int32 InZoomLevel) const override
-	{
-		checkSlow(ZoomLevels.IsValidIndex(InZoomLevel));
-		return ZoomLevels[InZoomLevel].LOD;
-	}
-
-	TArray<FRigVMEditorZoomLevelEntry> ZoomLevels;
-};
-
-#undef LOCTEXT_NAMESPACE
\ No newline at end of file
diff --git a/Engine/Plugins/Runtime/RigVM/Source/RigVMEditor/Public/Editor/RigVMEditorGraphExplorerTabSummoner.h b/Engine/Plugins/Runtime/RigVM/Source/RigVMEditor/Public/Editor/RigVMEditorGraphExplorerTabSummoner.h
--- a/Engine/Plugins/Runtime/RigVM/Source/RigVMEditor/Public/Editor/RigVMEditorGraphExplorerTabSummoner.h
+++ b/Engine/Plugins/Runtime/RigVM/Source/RigVMEditor/Public/Editor/RigVMEditorGraphExplorerTabSummoner.h
@@ -5,7 +5,7 @@
 #include "CoreMinimal.h"
 #include "WorkflowOrientedApp/WorkflowTabFactory.h"
 
-class FRigVMEditorBase;
+class FRigVMEditor;
 
 struct RIGVMEDITOR_API FRigVMEditorGraphExplorerTabSummoner : public FWorkflowTabFactory
 {
@@ -14,10 +14,10 @@ public:
 	static inline const FLazyName TabID = FLazyName(TEXT("RigVM Graph Explorer"));
 	
 public:
-	FRigVMEditorGraphExplorerTabSummoner(const TSharedRef<FRigVMEditorBase>& InRigVMEditor);
+	FRigVMEditorGraphExplorerTabSummoner(const TSharedRef<FRigVMEditor>& InRigVMEditor);
 	
 	virtual TSharedRef<SWidget> CreateTabBody(const FWorkflowTabSpawnInfo& Info) const override;
 	
 protected:
-	TWeakPtr<FRigVMEditorBase> RigVMEditor;
+	TWeakPtr<FRigVMEditor> RigVMEditor;
 };
diff --git a/Engine/Plugins/Runtime/RigVM/Source/RigVMEditor/Public/Editor/RigVMEditorMenuContext.h b/Engine/Plugins/Runtime/RigVM/Source/RigVMEditor/Public/Editor/RigVMEditorMenuContext.h
--- a/Engine/Plugins/Runtime/RigVM/Source/RigVMEditor/Public/Editor/RigVMEditorMenuContext.h
+++ b/Engine/Plugins/Runtime/RigVM/Source/RigVMEditor/Public/Editor/RigVMEditorMenuContext.h
@@ -9,7 +9,7 @@
 
 class URigVMHost;
 class URigVMBlueprint;
-class FRigVMEditorBase;
+class FRigVMEditor;
 class URigVMGraph;
 class URigVMNode;
 class URigVMPin;
@@ -57,7 +57,7 @@ public:
 	 * @param InRigVMEditor 	    The RigVM Editor hosting the menus
 	 * @param InGraphMenuContext 	Additional context for specific menus
 	*/
-	void Init(TWeakPtr<FRigVMEditorBase> InRigVMEditor, const FRigVMEditorGraphMenuContext& InGraphMenuContext = FRigVMEditorGraphMenuContext());
+	void Init(TWeakPtr<FRigVMEditor> InRigVMEditor, const FRigVMEditorGraphMenuContext& InGraphMenuContext = FRigVMEditorGraphMenuContext());
 
 	/** Get the rigvm blueprint that we are editing */
 	UFUNCTION(BlueprintCallable, Category = RigVMEditor)
@@ -75,10 +75,10 @@ public:
 	UFUNCTION(BlueprintCallable, Category = RigVMEditor)
 	FRigVMEditorGraphMenuContext GetGraphMenuContext();
 	
-	FRigVMEditorBase* GetRigVMEditor() const;
+	FRigVMEditor* GetRigVMEditor() const;
 
 private:
-	TWeakPtr<FRigVMEditorBase> WeakRigVMEditor;
+	TWeakPtr<FRigVMEditor> WeakRigVMEditor;
 
 	FRigVMEditorGraphMenuContext GraphMenuContext;
 };
\ No newline at end of file
diff --git a/Engine/Plugins/Runtime/RigVM/Source/RigVMEditor/Public/Editor/RigVMEditorMode.h b/Engine/Plugins/Runtime/RigVM/Source/RigVMEditor/Public/Editor/RigVMEditorMode.h
new file mode 100644
--- /dev/null
+++ b/Engine/Plugins/Runtime/RigVM/Source/RigVMEditor/Public/Editor/RigVMEditorMode.h
@@ -0,0 +1,27 @@
+// Copyright Epic Games, Inc. All Rights Reserved.
+
+#pragma once
+
+#include "CoreMinimal.h"
+#include "WorkflowOrientedApp/WorkflowTabManager.h"
+#include "Editor/RigVMEditor.h"
+#include "RigVMBlueprint.h"
+#include "BlueprintEditorModes.h"
+
+class URigVMBlueprint;
+
+class FRigVMEditorMode : public FBlueprintEditorApplicationMode
+{
+public:
+	FRigVMEditorMode(const TSharedRef<FRigVMEditor>& InRigVMEditor);
+
+	// FApplicationMode interface
+	virtual void RegisterTabFactories(TSharedPtr<FTabManager> InTabManager) override;
+
+protected:
+	// Set of spawnable tabs
+	FWorkflowAllowedTabSet TabFactories;
+
+private:
+	TWeakObjectPtr<URigVMBlueprint> RigVMBlueprintPtr;
+};
diff --git a/Engine/Plugins/Runtime/RigVM/Source/RigVMEditor/Public/Editor/RigVMExecutionStackTabSummoner.h b/Engine/Plugins/Runtime/RigVM/Source/RigVMEditor/Public/Editor/RigVMExecutionStackTabSummoner.h
--- a/Engine/Plugins/Runtime/RigVM/Source/RigVMEditor/Public/Editor/RigVMExecutionStackTabSummoner.h
+++ b/Engine/Plugins/Runtime/RigVM/Source/RigVMEditor/Public/Editor/RigVMExecutionStackTabSummoner.h
@@ -5,7 +5,7 @@
 #include "CoreMinimal.h"
 #include "WorkflowOrientedApp/WorkflowTabFactory.h"
 
-class FRigVMEditorBase;
+class FRigVMEditor;
 
 struct RIGVMEDITOR_API FRigVMExecutionStackTabSummoner : public FWorkflowTabFactory
 {
@@ -14,10 +14,10 @@ public:
 	static inline const FLazyName TabID = FLazyName(TEXT("Execution Stack"));
 	
 public:
-	FRigVMExecutionStackTabSummoner(const TSharedRef<FRigVMEditorBase>& InRigVMEditor);
+	FRigVMExecutionStackTabSummoner(const TSharedRef<FRigVMEditor>& InRigVMEditor);
 	
 	virtual TSharedRef<SWidget> CreateTabBody(const FWorkflowTabSpawnInfo& Info) const override;
 	
 protected:
-	TWeakPtr<FRigVMEditorBase> RigVMEditor;
+	TWeakPtr<FRigVMEditor> RigVMEditor;
 };
diff --git a/Engine/Plugins/Runtime/RigVM/Source/RigVMEditor/Public/Editor/RigVMGraphDetailCustomization.h b/Engine/Plugins/Runtime/RigVM/Source/RigVMEditor/Public/Editor/RigVMGraphDetailCustomization.h
--- a/Engine/Plugins/Runtime/RigVM/Source/RigVMEditor/Public/Editor/RigVMGraphDetailCustomization.h
+++ b/Engine/Plugins/Runtime/RigVM/Source/RigVMEditor/Public/Editor/RigVMGraphDetailCustomization.h
@@ -8,8 +8,7 @@
 #include "UObject/WeakObjectPtr.h"
 #include "EdGraph/RigVMEdGraph.h"
 #include "RigVMBlueprint.h"
-#include "Editor/RigVMNewEditor.h"
-#include "Editor/RigVMLegacyEditor.h"
+#include "Editor/RigVMEditor.h"
 #include "SGraphPin.h"
 #include "Widgets/SRigVMGraphNode.h"
 #include "Widgets/Colors/SColorBlock.h"
@@ -30,7 +29,7 @@ public:
 	FRigVMFunctionArgumentGroupLayout(
 		const TWeakObjectPtr<URigVMGraph>& InGraph,
 		const TWeakInterfacePtr<IRigVMClientHost>& InRigVMClientHost,
-		const TWeakPtr<FRigVMEditorBase>& InEditor,
+		const TWeakPtr<FRigVMEditor>& InEditor,
 		bool bInputs);
 	virtual ~FRigVMFunctionArgumentGroupLayout();
 
@@ -50,7 +49,7 @@ private:
 
 	TWeakObjectPtr<URigVMGraph> GraphPtr;
 	TWeakInterfacePtr<IRigVMClientHost> WeakRigVMClientHost;
-	TWeakPtr<FRigVMEditorBase> RigVMEditorPtr;
+	TWeakPtr<FRigVMEditor> RigVMEditorPtr;
 	bool bIsInputGroup;
 	FSimpleDelegate OnRebuildChildren;
 };
@@ -63,7 +62,7 @@ public:
 		const TWeakObjectPtr<URigVMPin>& InPin,
 		const TWeakObjectPtr<URigVMGraph>& InGraph, 
 		const TWeakInterfacePtr<IRigVMClientHost>& InRigVMClientHost,
-		const TWeakPtr<FRigVMEditorBase>& InEditor)
+		const TWeakPtr<FRigVMEditor>& InEditor)
 		: PinPtr(InPin)
 		, GraphPtr(InGraph)
 		, WeakRigVMClientHost(InRigVMClientHost)
@@ -117,7 +116,7 @@ private:
 	TWeakInterfacePtr<IRigVMClientHost> WeakRigVMClientHost;
 
 	/** The editor we are editing */
-	TWeakPtr<FRigVMEditorBase> RigVMEditorPtr;
+	TWeakPtr<FRigVMEditor> RigVMEditorPtr;
 
 	/** Holds a weak pointer to the argument name widget, used for error notifications */
 	TWeakPtr<SEditableTextBox> ArgumentNameWidget;
@@ -163,20 +162,16 @@ class RIGVMEDITOR_API FRigVMGraphDetailCustomization : public IDetailCustomizati
 {
 public:
 	/** Makes a new instance of this detail layout class for a specific detail view requesting it */
-#if WITH_RIGVMLEGACYEDITOR
 	static TSharedPtr<IDetailCustomization> MakeInstance(TSharedPtr<IBlueprintEditor> InBlueprintEditor, const UClass* InExpectedBlueprintClass);
-#else
-	static TSharedPtr<IDetailCustomization> MakeInstance(TSharedPtr<FRigVMEditorBase> InBlueprintEditor, const UClass* InExpectedBlueprintClass);
-#endif
 
-	FRigVMGraphDetailCustomization(TSharedPtr<FRigVMEditorBase> RigVMigEditor, URigVMBlueprint* RigVMBlueprint);
+	FRigVMGraphDetailCustomization(TSharedPtr<FRigVMEditor> RigVMigEditor, URigVMBlueprint* RigVMBlueprint);
 
 	// IDetailCustomization interface
 	virtual void CustomizeDetails(IDetailLayoutBuilder& DetailLayout) override;
 
 private:
-	/** The editor we are embedded in */
-	TWeakPtr<FRigVMEditorBase> RigVMEditorPtr;
+	/** The Blueprint editor we are embedded in */
+	TWeakPtr<FRigVMEditor> RigVMEditorPtr;
 
 	/** The graph we are editing */
 	TWeakObjectPtr<URigVMEdGraph> GraphPtr;
@@ -194,7 +189,7 @@ public:
 		URigVMGraph* Model,
 		URigVMController* Controller,
 		IRigVMClientHost* InRigVMClientHost,
-		TWeakPtr<FRigVMEditorBase> InEditor);
+		TWeakPtr<FRigVMEditor> InEditor);
 
 private:
 
@@ -259,7 +254,7 @@ private:
 	TWeakObjectPtr<URigVMController> WeakController;
 
 	/** The editor we are embedded in */
-	TWeakPtr<FRigVMEditorBase> RigVMEditorPtr;
+	TWeakPtr<FRigVMEditor> RigVMEditorPtr;
 
 	/** The asset host we are editing */
 	TWeakInterfacePtr<IRigVMClientHost> RigVMClientHost;
diff --git a/Engine/Plugins/Runtime/RigVM/Source/RigVMEditor/Public/Editor/RigVMGraphEditorSummoner.h b/Engine/Plugins/Runtime/RigVM/Source/RigVMEditor/Public/Editor/RigVMGraphEditorSummoner.h
deleted file mode 100644
diff --git a/Engine/Plugins/Runtime/RigVM/Source/RigVMEditor/Public/Editor/RigVMLegacyEditor.h b/Engine/Plugins/Runtime/RigVM/Source/RigVMEditor/Public/Editor/RigVMLegacyEditor.h
deleted file mode 100644
diff --git a/Engine/Plugins/Runtime/RigVM/Source/RigVMEditor/Public/Editor/RigVMLegacyEditorMode.h b/Engine/Plugins/Runtime/RigVM/Source/RigVMEditor/Public/Editor/RigVMLegacyEditorMode.h
deleted file mode 100644
diff --git a/Engine/Plugins/Runtime/RigVM/Source/RigVMEditor/Public/Editor/RigVMNewEditor.h b/Engine/Plugins/Runtime/RigVM/Source/RigVMEditor/Public/Editor/RigVMNewEditor.h
deleted file mode 100644
diff --git a/Engine/Plugins/Runtime/RigVM/Source/RigVMEditor/Public/Editor/RigVMNewEditorMode.h b/Engine/Plugins/Runtime/RigVM/Source/RigVMEditor/Public/Editor/RigVMNewEditorMode.h
deleted file mode 100644
diff --git a/Engine/Plugins/Runtime/RigVM/Source/RigVMEditor/Public/Editor/SRigVMActionMenu.cpp b/Engine/Plugins/Runtime/RigVM/Source/RigVMEditor/Public/Editor/SRigVMActionMenu.cpp
deleted file mode 100644
diff --git a/Engine/Plugins/Runtime/RigVM/Source/RigVMEditor/Public/Editor/SRigVMActionMenu.h b/Engine/Plugins/Runtime/RigVM/Source/RigVMEditor/Public/Editor/SRigVMActionMenu.h
deleted file mode 100644
diff --git a/Engine/Plugins/Runtime/RigVM/Source/RigVMEditor/Public/Editor/SRigVMDetailsInspector.h b/Engine/Plugins/Runtime/RigVM/Source/RigVMEditor/Public/Editor/SRigVMDetailsInspector.h
deleted file mode 100644
diff --git a/Engine/Plugins/Runtime/RigVM/Source/RigVMEditor/Public/RigVMEditorModule.h b/Engine/Plugins/Runtime/RigVM/Source/RigVMEditor/Public/RigVMEditorModule.h
--- a/Engine/Plugins/Runtime/RigVM/Source/RigVMEditor/Public/RigVMEditorModule.h
+++ b/Engine/Plugins/Runtime/RigVM/Source/RigVMEditor/Public/RigVMEditorModule.h
@@ -9,35 +9,14 @@
 #include "CoreMinimal.h"
 #include "Modules/ModuleInterface.h"
 #include "Logging/LogMacros.h"
+#include "Editor/RigVMEditor.h"
 #include "EdGraph/RigVMEdGraphPanelNodeFactory.h"
 #include "EdGraph/RigVMEdGraphPanelPinFactory.h"
 #include "Toolkits/AssetEditorToolkit.h"
 #include "Kismet2/StructureEditorUtils.h"
-#include "Modules/ModuleManager.h"
-
-class URigVMUserWorkflowOptions;
-class URigVMPin;
-class URigVMNode;
-class URigVMEdGraphNode;
-class IRigVMClientHost;
-class URigVMBlueprint;
-class URigVMEdGraphSchema;
-class FRigVMEditorBase;
-class FBlueprintActionDatabaseRegistrar;
 
 DECLARE_LOG_CATEGORY_EXTERN(LogRigVMEditor, Log, All);
 
-/** Describes the reason for Refreshing the editor */
-namespace ERefreshRigVMEditorReason
-{
-	enum Type
-	{
-		BlueprintCompiled,
-		PostUndo,
-		UnknownReason
-	};
-}
-
 // shallow interface declaration for use within RigVMDeveloper 
 class IRigVMEditorModule : public IModuleInterface, public IHasMenuExtensibility, public IHasToolBarExtensibility, public FStructureEditorUtils::INotifyOnStructChanged
 {
@@ -90,7 +69,7 @@ public:
 	virtual void PostChange(const UUserDefinedStruct* Changed, FStructureEditorUtils::EStructureEditorChangeInfo ChangedType) override;
 
 	/** Get all toolbar extenders */
-	DECLARE_DELEGATE_RetVal_TwoParams(TSharedRef<FExtender>, FRigVMEditorToolbarExtender, const TSharedRef<FUICommandList> /*InCommandList*/, TSharedRef<FRigVMEditorBase> /*InRigVMEditor*/);
+	DECLARE_DELEGATE_RetVal_TwoParams(TSharedRef<FExtender>, FRigVMEditorToolbarExtender, const TSharedRef<FUICommandList> /*InCommandList*/, TSharedRef<FRigVMEditor> /*InRigVMEditor*/);
 	virtual TArray<FRigVMEditorToolbarExtender>& GetAllRigVMEditorToolbarExtenders();
 
 	virtual void GetContextMenuActions(const URigVMEdGraphSchema* Schema, class UToolMenu* Menu, class UGraphNodeContextMenuContext* Context) const override;
diff --git a/Engine/Plugins/Runtime/RigVM/Source/RigVMEditor/Public/Widgets/SRigVMEditorGraphExplorer.h b/Engine/Plugins/Runtime/RigVM/Source/RigVMEditor/Public/Widgets/SRigVMEditorGraphExplorer.h
--- a/Engine/Plugins/Runtime/RigVM/Source/RigVMEditor/Public/Widgets/SRigVMEditorGraphExplorer.h
+++ b/Engine/Plugins/Runtime/RigVM/Source/RigVMEditor/Public/Widgets/SRigVMEditorGraphExplorer.h
@@ -11,12 +11,11 @@
 #include "RigVMModel/RigVMGraph.h"
 #include "SRigVMEditorGraphExplorerTreeView.h"
 
-class FRigVMEditorBase;
 class SRigVMEditorGraphExplorerTreeView;
 //
 class URigVMGraph;
 class URigVMEdGraph;
-class FRigVMEditorBase;
+class FRigVMEditor;
 class URigVMBlueprint;
 class SGraphActionMenu;
 struct FCreateWidgetForActionData;
@@ -80,7 +79,7 @@ public:
 	SLATE_BEGIN_ARGS(SRigVMEditorGraphExplorer) {}
 	SLATE_END_ARGS()
 
-	void Construct(const FArguments& InArgs, TWeakPtr<FRigVMEditorBase> InRigVMEditor);
+	void Construct(const FArguments& InArgs, TWeakPtr<FRigVMEditor> InRigVMEditor);
 	~SRigVMEditorGraphExplorer();
 
 	// Refresh the graph action menu.
@@ -89,17 +88,8 @@ public:
 	// SWidget overrides
 	void Tick(const FGeometry& InAllottedGeometry, const double InCurrentTime, const float InDeltaTime) override;
 
-	void SetLastPinTypeUsed(const FEdGraphPinType& InType) { LastPinType = InType; }
-	FEdGraphPinType GetLastPinTypeUsed() const { return LastPinType; }
-
-	FName GetSelectedVariableName() const;
-
-	void ClearSelection();
-
 private:
 	bool bNeedsRefresh;
-
-	FEdGraphPinType LastPinType;
 	
 	void RegisterCommands();
 	void CreateWidgets();
@@ -196,7 +186,7 @@ private:
 
 	void HandleSelectionChanged(TSharedPtr<FRigVMEditorGraphExplorerTreeElement> Selection, ESelectInfo::Type SelectInfo);
 
-	TWeakPtr<FRigVMEditorBase> RigVMEditor;
+	TWeakPtr<FRigVMEditor> RigVMEditor;
 
 	TSharedPtr<SRigVMEditorGraphExplorerTreeView> TreeView;
 
diff --git a/Engine/Plugins/Runtime/RigVM/Source/RigVMEditor/Public/Widgets/SRigVMEditorGraphExplorerTreeView.h b/Engine/Plugins/Runtime/RigVM/Source/RigVMEditor/Public/Widgets/SRigVMEditorGraphExplorerTreeView.h
--- a/Engine/Plugins/Runtime/RigVM/Source/RigVMEditor/Public/Widgets/SRigVMEditorGraphExplorerTreeView.h
+++ b/Engine/Plugins/Runtime/RigVM/Source/RigVMEditor/Public/Widgets/SRigVMEditorGraphExplorerTreeView.h
@@ -7,7 +7,7 @@
 #include "RigVMModel/RigVMGraph.h"
 #include "Widgets/Views/STreeView.h"
 
-class FRigVMEditorBase;
+class FRigVMEditor;
 class SSearchBox;
 class SRigVMEditorGraphExplorerTreeView;
 class SRigVMEditorGraphExplorerItem;
@@ -15,32 +15,28 @@ class FRigVMEditorGraphExplorerTreeElement;
 class URigVMGraph;
 class IPinTypeSelectorFilter;
 
-namespace ERigVMExplorerElementType
+enum ERigVMExplorerElementType
 {
-	enum Type
-	{
-		Invalid,
-		Section,
-		FunctionCategory,
-		VariableCategory,
-		Graph,
-		Event,
-		Function,
-		Variable,
-		LocalVariable
-	};
-}
+	Invalid,
+	Section,
+	FunctionCategory,
+	VariableCategory,
+	Graph,
+	Event,
+	Function,
+	Variable,
+	LocalVariable
+};
 
 struct FRigVMExplorerElementKey
 {
-
 	FRigVMExplorerElementKey()
 		: Type(ERigVMExplorerElementType::Invalid), Name(FString()){}
 
-	FRigVMExplorerElementKey(const ERigVMExplorerElementType::Type& InType, const FString& InName)
+	FRigVMExplorerElementKey(const ERigVMExplorerElementType& InType, const FString& InName)
 		: Type(InType), Name(InName){}
 	
-	ERigVMExplorerElementType::Type Type;
+	ERigVMExplorerElementType Type;
 	FString Name;
 
 	friend bool operator==(const FRigVMExplorerElementKey& A, const FRigVMExplorerElementKey& B)
diff --git a/Engine/Plugins/Runtime/RigVM/Source/RigVMEditor/Public/Widgets/SRigVMExecutionStackView.h b/Engine/Plugins/Runtime/RigVM/Source/RigVMEditor/Public/Widgets/SRigVMExecutionStackView.h
--- a/Engine/Plugins/Runtime/RigVM/Source/RigVMEditor/Public/Widgets/SRigVMExecutionStackView.h
+++ b/Engine/Plugins/Runtime/RigVM/Source/RigVMEditor/Public/Widgets/SRigVMExecutionStackView.h
@@ -10,7 +10,7 @@
 #include "RigVMCore/RigVM.h"
 #include "RigVMBlueprint.h"
 
-class FRigVMEditorBase;
+class FRigVMEditor;
 class SRigVMExecutionStackView;
 class FUICommandList;
 class SSearchBox;
@@ -75,7 +75,7 @@ public:
 
 	~SRigVMExecutionStackView();
 
-	void Construct(const FArguments& InArgs, TSharedRef<FRigVMEditorBase> InRigVMEditor);
+	void Construct(const FArguments& InArgs, TSharedRef<FRigVMEditor> InRigVMEditor);
 
 	/** Set Selection Changed */
 	void OnSelectionChanged(TSharedPtr<FRigStackEntry> Selection, ESelectInfo::Type SelectInfo);
@@ -126,7 +126,7 @@ private:
 	bool bSuspendControllerSelection;
 	void HandleModifiedEvent(ERigVMGraphNotifType InNotifType, URigVMGraph* InGraph, UObject* InSubject);
 	void HandleHostInitializedEvent(URigVMHost* InHost, const FName& InEventName);
-	void HandlePreviewHostUpdated(FRigVMEditorBase* InEditor);
+	void HandlePreviewHostUpdated(FRigVMEditor* InEditor);
 	void HandleItemMouseDoubleClick(TSharedPtr<FRigStackEntry> InItem);
 
 	/** Populate the execution stack with descriptive names for each instruction */
@@ -137,7 +137,7 @@ private:
 	/** Command list we bind to */
 	TSharedPtr<FUICommandList> CommandList;
 
-	TWeakPtr<FRigVMEditorBase> WeakEditor;
+	TWeakPtr<FRigVMEditor> RigVMEditor;
 	TWeakObjectPtr<URigVMBlueprint> RigVMBlueprint;
 
 	TArray<TSharedPtr<FRigStackEntry>> Operators;
diff --git a/Engine/Plugins/Runtime/RigVM/Source/RigVMEditor/RigVMEditor.Build.cs b/Engine/Plugins/Runtime/RigVM/Source/RigVMEditor/RigVMEditor.Build.cs
--- a/Engine/Plugins/Runtime/RigVM/Source/RigVMEditor/RigVMEditor.Build.cs
+++ b/Engine/Plugins/Runtime/RigVM/Source/RigVMEditor/RigVMEditor.Build.cs
@@ -21,6 +21,7 @@ public class RigVMEditor : ModuleRules
                 "SlateCore",
                 "InputCore",
                 "ToolWidgets",
+                "Kismet",
                 "KismetCompiler",
                 "KismetWidgets",
                 "AnimationWidgets",
@@ -34,12 +35,5 @@ public class RigVMEditor : ModuleRules
 				"ContentBrowser",
 			}
 		);
-        
-        bool bAddBlueprintEditorDependency = true;
-        PublicDefinitions.Add("WITH_RIGVMLEGACYEDITOR=" + (bAddBlueprintEditorDependency ? '1' : '0'));
-        if (bAddBlueprintEditorDependency)
-        {
-			PublicDependencyModuleNames.Add("Kismet");
-		}
 	}
 }
diff --git a/Engine/Source/Editor/KismetCompiler/Public/KismetCompiler.h b/Engine/Source/Editor/KismetCompiler/Public/KismetCompiler.h
--- a/Engine/Source/Editor/KismetCompiler/Public/KismetCompiler.h
+++ b/Engine/Source/Editor/KismetCompiler/Public/KismetCompiler.h
@@ -678,7 +678,6 @@ private:
 	 * we have decided to make the compilation manager a friend:
 	 */
 	friend struct FBlueprintCompilationManagerImpl;
-	friend struct FRigVMBlueprintCompilationManagerImpl;
 };
 
 //////////////////////////////////////////////////////////////////////////
diff --git a/Engine/Source/Editor/UnrealEd/Public/Kismet2/KismetReinstanceUtilities.h b/Engine/Source/Editor/UnrealEd/Public/Kismet2/KismetReinstanceUtilities.h
--- a/Engine/Source/Editor/UnrealEd/Public/Kismet2/KismetReinstanceUtilities.h
+++ b/Engine/Source/Editor/UnrealEd/Public/Kismet2/KismetReinstanceUtilities.h
@@ -107,7 +107,6 @@ class FBlueprintCompileReinstancer : public TSharedFromThis<FBlueprintCompileRei
 {
 protected:
 	friend struct FBlueprintCompilationManagerImpl;
-	friend struct FRigVMBlueprintCompilationManagerImpl;
 	friend struct FReinstancingJob;
 
 	static UNREALED_API TSet<TWeakObjectPtr<UBlueprint>> CompiledBlueprintsToSave;
diff --git a/Engine/Source/Runtime/Engine/Classes/Engine/TimelineTemplate.h b/Engine/Source/Runtime/Engine/Classes/Engine/TimelineTemplate.h
--- a/Engine/Source/Runtime/Engine/Classes/Engine/TimelineTemplate.h
+++ b/Engine/Source/Runtime/Engine/Classes/Engine/TimelineTemplate.h
@@ -342,7 +342,6 @@ private:
 
 	/** Grant external permission to the compilation manager for backwards compatibility purposes as PostLoad will not have occurred yet for compile on load. */
 	friend struct FBlueprintCompilationManagerImpl;
-	friend struct FRigVMBlueprintCompilationManagerImpl;
 
 	/** Grant external permission to Blueprint editor utility methods. For example, duplicating a Blueprint for conversion to C++ requires that cached names be updated. */
 	friend class FBlueprintEditorUtils;
