[AutoRTFM] Replace atomic race check with a transactionally-safe MTAccessDetector.

Previously, this code would bump an atomic to indicate that work was taking place, and then elsewhere it would check-fail if the atomic was non-zero.

This code is directly reachable from Verse, and in AutoRTFM we don't want other threads to be able to view side effects from transactional work at all. I've replaced this atomic with an MT access detector.

diff --git a/Engine/Plugins/Runtime/RigVM/Source/RigVM/Private/RigVMCore/RigVM.cpp b/Engine/Plugins/Runtime/RigVM/Source/RigVM/Private/RigVMCore/RigVM.cpp
--- a/Engine/Plugins/Runtime/RigVM/Source/RigVM/Private/RigVMCore/RigVM.cpp
+++ b/Engine/Plugins/Runtime/RigVM/Source/RigVM/Private/RigVMCore/RigVM.cpp
@@ -20,6 +20,7 @@
 #include "RigVMHost.h"
 #include "RigVMCore/RigVMTrait.h"
 #include "RigVMStringUtils.h"
+#include "Misc/MTAccessDetector.h"
 
 #include UE_INLINE_GENERATED_CPP_BY_NAME(RigVM)
 
@@ -284,7 +285,7 @@ void URigVM::Load(FArchive& Ar)
 void URigVM::CopyDataForSerialization(URigVM* InVM)
 {
 	check(InVM);
-	check(ActiveExecutions.load() == 0);
+	UE_MT_SCOPED_WRITE_ACCESS(AccessDetector);
 
 	Reset_Internal();
 
@@ -1535,24 +1536,23 @@ ERigVMExecuteResult URigVM::ExecuteVM(FRigVMExtendedExecuteContext& Context, con
 {
 	// if this the first entry being executed - get ready for execution
 	const bool bIsRootEntry = Context.EntriesBeingExecuted.IsEmpty();
-	ON_SCOPE_EXIT
-	{
-		if (bIsRootEntry)
-		{
-			ActiveExecutions--;
-		}
-	};
 
 	TGuardValue<FName> EntryNameGuard(Context.CurrentEntryName, InEntryName);
 	TGuardValue<bool> RootEntryGuard(Context.bCurrentlyRunningRootEntry, bIsRootEntry);
 
 	FRigVMByteCode& ByteCode = GetByteCode();
 
-	if(bIsRootEntry)
+#if ENABLE_MT_DETECTOR
+	FBaseScopedAccessDetector ScopedMTAccessDetector;
+#endif
+
+	if (bIsRootEntry)
 	{
 		Context.CurrentExecuteResult = ERigVMExecuteResult::Succeeded;
-		
-		ActiveExecutions++;
+
+#if ENABLE_MT_DETECTOR
+		ScopedMTAccessDetector = MakeScopedReaderAccessDetector(AccessDetector);
+#endif
 
 		if (Instructions.Num() == 0)
 		{
diff --git a/Engine/Plugins/Runtime/RigVM/Source/RigVM/Public/RigVMCore/RigVM.h b/Engine/Plugins/Runtime/RigVM/Source/RigVM/Public/RigVMCore/RigVM.h
--- a/Engine/Plugins/Runtime/RigVM/Source/RigVM/Public/RigVMCore/RigVM.h
+++ b/Engine/Plugins/Runtime/RigVM/Source/RigVM/Public/RigVMCore/RigVM.h
@@ -14,6 +14,7 @@
 #include "Math/Vector.h"
 #include "Math/Vector2D.h"
 #include "Misc/AssertionMacros.h"
+#include "Misc/MTTransactionallySafeAccessDetector.h"
 #include "RigVMByteCode.h"
 #include "RigVMCore/RigVMExternalVariable.h"
 #include "RigVMCore/RigVMFunction.h"
@@ -1012,14 +1013,9 @@ protected:
 	uint32 NumExecutions_DEPRECATED = 0;
 #endif
 
-	std::atomic<int32> ActiveExecutions;
+	UE_MT_DECLARE_TS_RW_ACCESS_DETECTOR(AccessDetector);
 
 public:
-	int32 GetActiveExecutions() const
-	{
-		return ActiveExecutions.load();
-	}
-
 	bool ValidateBytecode();
 	void RefreshArgumentNameCaches();
 
