Control Rig: Fixed export function from archive not creating the correct graph
#rnx

diff --git a/Engine/Plugins/Runtime/RigVM/Source/RigVMDeveloper/Private/RigVMModel/RigVMPin.cpp b/Engine/Plugins/Runtime/RigVM/Source/RigVMDeveloper/Private/RigVMModel/RigVMPin.cpp
--- a/Engine/Plugins/Runtime/RigVM/Source/RigVMDeveloper/Private/RigVMModel/RigVMPin.cpp
+++ b/Engine/Plugins/Runtime/RigVM/Source/RigVMDeveloper/Private/RigVMModel/RigVMPin.cpp
@@ -130,7 +130,7 @@ URigVMPin::URigVMPin()
 
 bool URigVMPin::NameEquals(const FString& InName, bool bFollowCoreRedirectors) const
 {
-	if(InName.Equals(GetName()))
+	if(InName.Equals(GetName(), ESearchCase::IgnoreCase))
 	{
 		return true;
 	}
