Control Rig: FIx user defined struct pins getting orphaned for no reason
#rnx

diff --git a/Engine/Plugins/Runtime/RigVM/Source/RigVMDeveloper/Private/RigVMModel/RigVMController.cpp b/Engine/Plugins/Runtime/RigVM/Source/RigVMDeveloper/Private/RigVMModel/RigVMController.cpp
--- a/Engine/Plugins/Runtime/RigVM/Source/RigVMDeveloper/Private/RigVMModel/RigVMController.cpp
+++ b/Engine/Plugins/Runtime/RigVM/Source/RigVMDeveloper/Private/RigVMModel/RigVMController.cpp
@@ -443,14 +443,32 @@ FRigVMPinInfo::FRigVMPinInfo(const URigVMPin* InPin, int32 InParentIndex, ERigVM
 	, bIsLazy(InPin->IsLazy() && InPin->IsRootPin())
 {
 	// this method describes the info as currently represented in the model.
+	
+	const URigVMPin* Pin = InPin;
+	while (Pin)
+	{
+		if (const UUserDefinedStruct* UDStruct = Cast<const UUserDefinedStruct>(Pin->GetCPPTypeObject()))
+		{
+			const FString SubPinPath = InPin->GetSubPinPath(Pin);
+			if (const FProperty* StructProperty = UDStruct->FindPropertyByName(*SubPinPath))
+			{
+				Name = *UDStruct->GetAuthoredNameForField(StructProperty);
+				break;
+			}
+		}
+		Pin = Pin->GetParentPin();
+	}
 
 	CorrectExecuteTypeIndex();
-	DefaultValue = InPin->GetDefaultValue();
-	
-	if(DefaultValue.IsEmpty() && (InPin->IsArray() || InPin->IsStruct()))
+	if (InPin)
 	{
-		static const FString EmptyBraces(TEXT("()"));
-		DefaultValue = EmptyBraces;
+		DefaultValue = InPin->GetDefaultValue();
+	
+		if(DefaultValue.IsEmpty() && (InPin->IsArray() || InPin->IsStruct()))
+		{
+			static const FString EmptyBraces(TEXT("()"));
+			DefaultValue = EmptyBraces;
+		}
 	}
 }
 
@@ -470,6 +488,11 @@ FRigVMPinInfo::FRigVMPinInfo(FProperty* InProperty, ERigVMPinDirection InDirecti
 {
 	// this method describes the info as needed based on the property structure
 
+	if (const UUserDefinedStruct* UDStruct = Cast<const UUserDefinedStruct>(Property->GetOwnerStruct()))
+	{
+		Name = *UDStruct->GetAuthoredNameForField(InProperty);
+	}
+
 	if (Direction == ERigVMPinDirection::Invalid)
 	{
 		Direction = FRigVMStruct::GetPinDirectionFromProperty(Property);
