Control Rig: Fix static analysis warning
#rnx

diff --git a/Engine/Plugins/Runtime/RigVM/Source/RigVMDeveloper/Private/RigVMModel/RigVMController.cpp b/Engine/Plugins/Runtime/RigVM/Source/RigVMDeveloper/Private/RigVMModel/RigVMController.cpp
--- a/Engine/Plugins/Runtime/RigVM/Source/RigVMDeveloper/Private/RigVMModel/RigVMController.cpp
+++ b/Engine/Plugins/Runtime/RigVM/Source/RigVMDeveloper/Private/RigVMModel/RigVMController.cpp
@@ -444,24 +444,24 @@ FRigVMPinInfo::FRigVMPinInfo(const URigVMPin* InPin, int32 InParentIndex, ERigVM
 {
 	// this method describes the info as currently represented in the model.
 	
-	const URigVMPin* Pin = InPin;
-	while (Pin)
+	if (InPin)
 	{
-		if (const UUserDefinedStruct* UDStruct = Cast<const UUserDefinedStruct>(Pin->GetCPPTypeObject()))
+		const URigVMPin* Pin = InPin;
+		while (Pin)
 		{
-			const FString SubPinPath = InPin->GetSubPinPath(Pin);
-			if (const FProperty* StructProperty = UDStruct->FindPropertyByName(*SubPinPath))
+			if (const UUserDefinedStruct* UDStruct = Cast<const UUserDefinedStruct>(Pin->GetCPPTypeObject()))
 			{
-				Name = *UDStruct->GetAuthoredNameForField(StructProperty);
-				break;
+				const FString SubPinPath = InPin->GetSubPinPath(Pin);
+				if (const FProperty* StructProperty = UDStruct->FindPropertyByName(*SubPinPath))
+				{
+					Name = *UDStruct->GetAuthoredNameForField(StructProperty);
+					break;
+				}
 			}
+			Pin = Pin->GetParentPin();
 		}
-		Pin = Pin->GetParentPin();
-	}
 
-	CorrectExecuteTypeIndex();
-	if (InPin)
-	{
+		CorrectExecuteTypeIndex();
 		DefaultValue = InPin->GetDefaultValue();
 	
 		if(DefaultValue.IsEmpty() && (InPin->IsArray() || InPin->IsStruct()))
