Control Rig: Fix crash when openning F_MED_ForwardLake_DireLab_ControlRig and its animBP.
#rnx

diff --git a/Engine/Plugins/Animation/ControlRig/Source/ControlRigEditor/Private/Editor/ControlRigEditor.cpp b/Engine/Plugins/Animation/ControlRig/Source/ControlRigEditor/Private/Editor/ControlRigEditor.cpp
--- a/Engine/Plugins/Animation/ControlRig/Source/ControlRigEditor/Private/Editor/ControlRigEditor.cpp
+++ b/Engine/Plugins/Animation/ControlRig/Source/ControlRigEditor/Private/Editor/ControlRigEditor.cpp
@@ -159,13 +159,22 @@ UControlRigBlueprint* GetControlRigBlueprintFromContext(const UUnrealEdViewportT
 			if (const TSharedPtr<FAnimationViewportClient>& AnimViewportClient =
 					StaticCastSharedPtr<FAnimationViewportClient>(Viewport->GetViewportClient()))
 			{
+#if WITH_RIGVMLEGACYEDITOR
 				if (AnimViewportClient->GetAssetEditorToolkit()->IsBlueprintEditor())
 				{
-					const TSharedRef<FControlRigEditor>& Editor =
-						StaticCastSharedRef<FControlRigEditor>(AnimViewportClient->GetAssetEditorToolkit());
+					const TSharedRef<FBlueprintEditor>& BlueprintEditor =
+						StaticCastSharedRef<FBlueprintEditor>(AnimViewportClient->GetAssetEditorToolkit());
 
-					return Cast<UControlRigBlueprint>(Editor->FRigVMEditorBase::GetBlueprintObj());
+					return Cast<UControlRigBlueprint>(BlueprintEditor->GetBlueprintObj());
 				}
+#else
+				// TODO sara-s Figure out another way to make sure this is a ControlRigEditor
+				// Possibly getting the EditingObjects from the AssetEditorToolkit and checking if its a ControlRigBlueprint
+				const TSharedRef<FControlRigEditor>& Editor =
+					StaticCastSharedRef<FControlRigEditor>(AnimViewportClient->GetAssetEditorToolkit());
+
+				return Cast<UControlRigBlueprint>(Editor->FRigVMEditorBase::GetBlueprintObj());
+#endif
 			}
 		}
 	}
diff --git a/Engine/Plugins/Runtime/RigVM/Source/RigVMEditor/Private/Editor/RigVMEditor.cpp b/Engine/Plugins/Runtime/RigVM/Source/RigVMEditor/Private/Editor/RigVMEditor.cpp
--- a/Engine/Plugins/Runtime/RigVM/Source/RigVMEditor/Private/Editor/RigVMEditor.cpp
+++ b/Engine/Plugins/Runtime/RigVM/Source/RigVMEditor/Private/Editor/RigVMEditor.cpp
@@ -470,7 +470,7 @@ const FName FRigVMEditorBase::GetEditorModeName() const
 
 UBlueprint* FRigVMEditorBase::GetBlueprintObj() const
 {
-	const TArray<UObject*>& EditingObjs = GetEditingObjects();
+	const TArray<UObject*>& EditingObjs = GetEditingBlueprints();
 	for (UObject* Obj : EditingObjs)
 	{
 		if (Obj->IsA<URigVMBlueprint>()) 
@@ -3142,7 +3142,7 @@ void FRigVMEditorBase::ExtendMenu()
 
 	// add extensible menu if exists
 	FRigVMEditorModule& RigVMEditorModule = FModuleManager::LoadModuleChecked<FRigVMEditorModule>("RigVMEditor");
-	AddMenuExtender(RigVMEditorModule.GetMenuExtensibilityManager()->GetAllExtenders(GetToolkitCommands(), GetEditingObjects()));
+	AddMenuExtender(RigVMEditorModule.GetMenuExtensibilityManager()->GetAllExtenders(GetToolkitCommands(), GetEditingBlueprints()));
 }
 
 void FRigVMEditorBase::ExtendToolbar()
@@ -3161,7 +3161,7 @@ void FRigVMEditorBase::ExtendToolbar()
 	AddToolbarExtender(ToolbarExtender);
 
 	FRigVMEditorModule& RigVMEditorModule = FModuleManager::LoadModuleChecked<FRigVMEditorModule>("RigVMEditor");
-	AddToolbarExtender(RigVMEditorModule.GetToolBarExtensibilityManager()->GetAllExtenders(GetToolkitCommands(), GetEditingObjects()));
+	AddToolbarExtender(RigVMEditorModule.GetToolBarExtensibilityManager()->GetAllExtenders(GetToolkitCommands(), GetEditingBlueprints()));
 
 	TArray<FRigVMEditorModule::FRigVMEditorToolbarExtender> ToolbarExtenderDelegates = RigVMEditorModule.GetAllRigVMEditorToolbarExtenders();
 
diff --git a/Engine/Plugins/Runtime/RigVM/Source/RigVMEditor/Public/Editor/RigVMEditor.h b/Engine/Plugins/Runtime/RigVM/Source/RigVMEditor/Public/Editor/RigVMEditor.h
--- a/Engine/Plugins/Runtime/RigVM/Source/RigVMEditor/Public/Editor/RigVMEditor.h
+++ b/Engine/Plugins/Runtime/RigVM/Source/RigVMEditor/Public/Editor/RigVMEditor.h
@@ -124,7 +124,7 @@ public:
 #endif
 	virtual TSharedPtr<FApplicationMode> CreateEditorMode() = 0;
 	virtual const FName GetEditorAppName() const = 0;
-	virtual const TArray< UObject* >& GetEditingObjects() const = 0;
+	virtual const TArray< UObject* >& GetEditingBlueprints() const = 0;
 	virtual const TSharedRef<IToolkitHost> GetToolkitHost() const = 0;
 	virtual bool IsHosted() const = 0;
 	virtual void BringToolkitToFrontImpl() = 0;
diff --git a/Engine/Plugins/Runtime/RigVM/Source/RigVMEditor/Public/Editor/RigVMLegacyEditor.h b/Engine/Plugins/Runtime/RigVM/Source/RigVMEditor/Public/Editor/RigVMLegacyEditor.h
--- a/Engine/Plugins/Runtime/RigVM/Source/RigVMEditor/Public/Editor/RigVMLegacyEditor.h
+++ b/Engine/Plugins/Runtime/RigVM/Source/RigVMEditor/Public/Editor/RigVMLegacyEditor.h
@@ -51,7 +51,7 @@ protected:
 #endif
 	virtual TSharedPtr<FApplicationMode> CreateEditorMode() override;
 	virtual const FName GetEditorAppName() const override;
-	virtual const TArray< UObject* >& GetEditingObjects() const override { return FBlueprintEditor::GetEditingObjects(); }
+	virtual const TArray< UObject* >& GetEditingBlueprints() const override { return FBlueprintEditor::GetEditingObjects(); }
 	virtual void OnGraphEditorFocused(const TSharedRef<SGraphEditor>& InGraphEditor) override;
 	virtual const TSharedRef<IToolkitHost> GetToolkitHost() const override { return FBlueprintEditor::GetToolkitHost(); }
 	virtual bool IsHosted() const override { return FBlueprintEditor::IsHosted(); }
diff --git a/Engine/Plugins/Runtime/RigVM/Source/RigVMEditor/Public/Editor/RigVMNewEditor.h b/Engine/Plugins/Runtime/RigVM/Source/RigVMEditor/Public/Editor/RigVMNewEditor.h
--- a/Engine/Plugins/Runtime/RigVM/Source/RigVMEditor/Public/Editor/RigVMNewEditor.h
+++ b/Engine/Plugins/Runtime/RigVM/Source/RigVMEditor/Public/Editor/RigVMNewEditor.h
@@ -46,7 +46,7 @@ protected:
 #endif
 	virtual TSharedPtr<FApplicationMode> CreateEditorMode() override;
 	virtual const FName GetEditorAppName() const override;
-	virtual const TArray< UObject* >& GetEditingObjects() const override { return FWorkflowCentricApplication::GetEditingObjects(); }
+	virtual const TArray< UObject* >& GetEditingBlueprints() const override { return FWorkflowCentricApplication::GetEditingObjects(); }
 	virtual const TSharedRef<IToolkitHost> GetToolkitHost() const override { return FWorkflowCentricApplication::GetToolkitHost(); }
 	virtual bool IsHosted() const override { return FWorkflowCentricApplication::IsHosted(); }
 	virtual void BringToolkitToFrontImpl() override { FWorkflowCentricApplication::BringToolkitToFront(); }
