RigVM: Avoid compilation with missing variables

#rnx

diff --git a/Engine/Plugins/Runtime/RigVM/Source/RigVMDeveloper/Private/RigVMCompiler/RigVMCompiler.cpp b/Engine/Plugins/Runtime/RigVM/Source/RigVMDeveloper/Private/RigVMCompiler/RigVMCompiler.cpp
--- a/Engine/Plugins/Runtime/RigVM/Source/RigVMDeveloper/Private/RigVMCompiler/RigVMCompiler.cpp
+++ b/Engine/Plugins/Runtime/RigVM/Source/RigVMDeveloper/Private/RigVMCompiler/RigVMCompiler.cpp
@@ -821,6 +821,13 @@ bool URigVMCompiler::Compile(const FRigVMCompileSettings& InSettings, TArray<URi
 					}
 				}
 
+				if(!bFoundVariable)
+				{
+					static const FString MissingVariableNodeMessage = TEXT("Variable Node @@ is using a missing variable. Consider recreating the node.");
+					Settings.ASTSettings.Report(EMessageSeverity::Error, ModelNode, MissingVariableNodeMessage);
+					bEncounteredGraphError = true;
+				}
+
 				if(VariableDescription.CPPTypeObject && !RigVMCore::SupportsUObjects())
 				{
 					if(VariableDescription.CPPTypeObject->IsA<UClass>())
