Control Rig: Fix subpins of UserDefinedStructs with wrong name. Patch during load.
#rnx

diff --git a/Engine/Plugins/Runtime/RigVM/Source/RigVM/Private/RigVMCore/RigVMPropertyPath.cpp b/Engine/Plugins/Runtime/RigVM/Source/RigVM/Private/RigVMCore/RigVMPropertyPath.cpp
--- a/Engine/Plugins/Runtime/RigVM/Source/RigVM/Private/RigVMCore/RigVMPropertyPath.cpp
+++ b/Engine/Plugins/Runtime/RigVM/Source/RigVM/Private/RigVMCore/RigVMPropertyPath.cpp
@@ -46,17 +46,21 @@ FRigVMPropertyPath::FRigVMPropertyPath(const FProperty* InProperty, const FStrin
 		{
 			Segment.Type = ERigVMPropertyPathSegmentType::StructMember;
 
-			if(const FProperty* MemberProperty = StructProperty->Struct->FindPropertyByName(*PathSegment))
+			const FProperty* MemberProperty = StructProperty->Struct->FindPropertyByName(*PathSegment);
+			if (!MemberProperty)
 			{
-				Segment.Name = MemberProperty->GetFName();
-				Segment.Index = MemberProperty->GetOffset_ForInternal();
-				Property = Segment.Property = MemberProperty;
+				MemberProperty = StructProperty->Struct->CustomFindProperty(*PathSegment);
 			}
-			else
+
+			if (!MemberProperty)
 			{
 				Segments.Empty();
 				return;
 			}
+
+			Segment.Name = MemberProperty->GetFName();
+			Segment.Index = MemberProperty->GetOffset_ForInternal();
+			Property = Segment.Property = MemberProperty;
 		}
 		else if(const FArrayProperty* ArrayProperty = CastField<FArrayProperty>(Property))
 		{
diff --git a/Engine/Plugins/Runtime/RigVM/Source/RigVMDeveloper/Private/RigVMCompiler/RigVMCompiler.cpp b/Engine/Plugins/Runtime/RigVM/Source/RigVMDeveloper/Private/RigVMCompiler/RigVMCompiler.cpp
--- a/Engine/Plugins/Runtime/RigVM/Source/RigVMDeveloper/Private/RigVMCompiler/RigVMCompiler.cpp
+++ b/Engine/Plugins/Runtime/RigVM/Source/RigVMDeveloper/Private/RigVMCompiler/RigVMCompiler.cpp
@@ -3133,87 +3133,97 @@ bool URigVMCompiler::TraverseAssign(const FRigVMAssignExprAST* InExpr, FRigVMCom
 
 	FRigVMOperand Source = WorkData.ExprToOperand.FindChecked(GetSourceVarExpr(SourceExpr));
 
-	if (!WorkData.bSetupMemory)
+	const FRigVMVarExprAST* TargetExpr = InExpr->GetFirstParentOfType(FRigVMVarExprAST::EType::Var)->To<FRigVMVarExprAST>();
+	TargetExpr = GetSourceVarExpr(TargetExpr);
+
+	// We cannot wait for the target var expr to be created, we need it now to setup the register offset
+	const FRigVMVarExprAST* VarExpr = GetSourceVarExpr(TargetExpr);
+	if (WorkData.bSetupMemory)
 	{
-		const FRigVMVarExprAST* TargetExpr = InExpr->GetFirstParentOfType(FRigVMVarExprAST::EType::Var)->To<FRigVMVarExprAST>();
-		TargetExpr = GetSourceVarExpr(TargetExpr);
-		
-		FRigVMOperand Target = WorkData.ExprToOperand.FindChecked(GetSourceVarExpr(TargetExpr));
-		if(Target == Source)
+		if (!WorkData.ExprToOperand.Contains(VarExpr))
 		{
-			return true;
+			FindOrAddRegister(VarExpr, WorkData);
 		}
-		
-		// if this is a copy - we should check if operands need offsets
-		if (InExpr->GetType() == FRigVMExprAST::EType::Copy)
+	}
+	
+	FRigVMOperand Target = WorkData.ExprToOperand.FindChecked(VarExpr);
+	if(Target == Source)
+	{
+		return true;
+	}
+	
+	// if this is a copy - we should check if operands need offsets
+	if (InExpr->GetType() == FRigVMExprAST::EType::Copy)
+	{
+		struct Local
 		{
-			struct Local
+			static void SetupRegisterOffset(URigVM* VM, const FRigVMASTLinkDescription& InLink, URigVMPin* Pin,
+				FRigVMOperand& Operand, const FRigVMVarExprAST* VarExpr, bool bSource, FRigVMCompilerWorkData& WorkData)
 			{
-				static void SetupRegisterOffset(URigVM* VM, const FRigVMASTLinkDescription& InLink, URigVMPin* Pin,
-					FRigVMOperand& Operand, const FRigVMVarExprAST* VarExpr, bool bSource, FRigVMCompilerWorkData& WorkData)
+				const bool bHasTargetSegmentPath = !bSource && !InLink.SegmentPath.IsEmpty();
+				
+				URigVMPin* RootPin = Pin->GetRootPin();
+				if (Pin == RootPin && !bHasTargetSegmentPath)
 				{
-					const bool bHasTargetSegmentPath = !bSource && !InLink.SegmentPath.IsEmpty();
-					
-					URigVMPin* RootPin = Pin->GetRootPin();
-					if (Pin == RootPin && !bHasTargetSegmentPath)
+					return;
+				}
+
+				if(Pin->IsProgrammaticPin())
+				{
+					return;
+				}
+
+				FString SegmentPath = Pin->GetSegmentPath(false);
+				if(bHasTargetSegmentPath)
+				{
+					if(SegmentPath.IsEmpty())
 					{
-						return;
+						SegmentPath = InLink.SegmentPath;
 					}
-
-					if(Pin->IsProgrammaticPin())
+					else
 					{
-						return;
+						SegmentPath = URigVMPin::JoinPinPath(SegmentPath, InLink.SegmentPath);
 					}
+				}
 
-					FString SegmentPath = Pin->GetSegmentPath(false);
-					if(bHasTargetSegmentPath)
+				// for fixed array pins we create a register for each array element
+				// thus we do not need to setup a registeroffset for the array element.
+				if (RootPin->IsFixedSizeArray())
+				{
+					if (Pin->GetParentPin() == RootPin)
 					{
-						if(SegmentPath.IsEmpty())
-						{
-							SegmentPath = InLink.SegmentPath;
-						}
-						else
-						{
-							SegmentPath = URigVMPin::JoinPinPath(SegmentPath, InLink.SegmentPath);
-						}
+						return;
 					}
 
-					// for fixed array pins we create a register for each array element
-					// thus we do not need to setup a registeroffset for the array element.
-					if (RootPin->IsFixedSizeArray())
+					// if the pin is a sub pin of a case of a fixed array
+					// we'll need to re-adjust the root pin to the case pin (for example: Values.0)
+					TArray<FString> SegmentPathPaths;
+					if(ensure(URigVMPin::SplitPinPath(SegmentPath, SegmentPathPaths)))
 					{
-						if (Pin->GetParentPin() == RootPin)
-						{
-							return;
-						}
-
-						// if the pin is a sub pin of a case of a fixed array
-						// we'll need to re-adjust the root pin to the case pin (for example: Values.0)
-						TArray<FString> SegmentPathPaths;
-						if(ensure(URigVMPin::SplitPinPath(SegmentPath, SegmentPathPaths)))
-						{
-							RootPin = RootPin->FindSubPin(SegmentPathPaths[0]);
+						RootPin = RootPin->FindSubPin(SegmentPathPaths[0]);
 
-							SegmentPathPaths.RemoveAt(0);
-							ensure(SegmentPathPaths.Num() > 0);
-							SegmentPath = URigVMPin::JoinPinPath(SegmentPathPaths);
-						}
-						else
-						{
-							return;
-						}
+						SegmentPathPaths.RemoveAt(0);
+						ensure(SegmentPathPaths.Num() > 0);
+						SegmentPath = URigVMPin::JoinPinPath(SegmentPathPaths);
+					}
+					else
+					{
+						return;
 					}
-
-					const int32 PropertyPathIndex = WorkData.FindOrAddPropertyPath(Operand, RootPin->GetCPPType(), SegmentPath);
-					Operand = FRigVMOperand(Operand.GetMemoryType(), Operand.GetRegisterIndex(), PropertyPathIndex);
 				}
-			};
 
-			const FRigVMASTLinkDescription& Link = InExpr->GetLink();
-			Local::SetupRegisterOffset(WorkData.VM, Link, InExpr->GetSourcePin(), Source, SourceExpr, true, WorkData);
-			Local::SetupRegisterOffset(WorkData.VM, Link, InExpr->GetTargetPin(), Target, TargetExpr, false, WorkData);
-		}
+				const int32 PropertyPathIndex = WorkData.FindOrAddPropertyPath(Operand, RootPin->GetCPPType(), SegmentPath);
+				Operand = FRigVMOperand(Operand.GetMemoryType(), Operand.GetRegisterIndex(), PropertyPathIndex);
+			}
+		};
 
+		const FRigVMASTLinkDescription& Link = InExpr->GetLink();
+		Local::SetupRegisterOffset(WorkData.VM, Link, InExpr->GetSourcePin(), Source, SourceExpr, true, WorkData);
+		Local::SetupRegisterOffset(WorkData.VM, Link, InExpr->GetTargetPin(), Target, TargetExpr, false, WorkData);
+	}
+	
+	if (!WorkData.bSetupMemory)
+	{
 		FRigVMCopyOp CopyOp = WorkData.VM->GetCopyOpForOperands(Source, Target);
 		if(CopyOp.IsValid())
 		{
diff --git a/Engine/Plugins/Runtime/RigVM/Source/RigVMDeveloper/Private/RigVMModel/Nodes/RigVMFunctionReferenceNode.cpp b/Engine/Plugins/Runtime/RigVM/Source/RigVMDeveloper/Private/RigVMModel/Nodes/RigVMFunctionReferenceNode.cpp
--- a/Engine/Plugins/Runtime/RigVM/Source/RigVMDeveloper/Private/RigVMModel/Nodes/RigVMFunctionReferenceNode.cpp
+++ b/Engine/Plugins/Runtime/RigVM/Source/RigVMDeveloper/Private/RigVMModel/Nodes/RigVMFunctionReferenceNode.cpp
@@ -327,8 +327,12 @@ TArray<FString> URigVMFunctionReferenceNode::GetPinCategories() const
 	const FRigVMNodeLayout* Layout = &ReferencedFunctionHeader.Layout;
 	if(IsReferencedNodeLoaded())
 	{
-		ReferenceLayout = LoadReferencedNode()->GetNodeLayout();
-		Layout = &ReferenceLayout;
+		URigVMLibraryNode* ReferencedNode = LoadReferencedNode();
+		if (ensure(ReferencedNode))
+		{
+			ReferenceLayout = LoadReferencedNode()->GetNodeLayout();
+			Layout = &ReferenceLayout;
+		}
 	}
 	TArray<FString> TransientPinCategories;
 	for(const FRigVMPinCategory& Category : Layout->Categories)
@@ -341,8 +345,12 @@ TArray<FString> URigVMFunctionReferenceNode::GetPinCategories() const
 FRigVMNodeLayout URigVMFunctionReferenceNode::GetNodeLayout(bool bIncludeEmptyCategories) const
 {
 	if(IsReferencedNodeLoaded())
-	{
-		return LoadReferencedNode()->GetNodeLayout(bIncludeEmptyCategories);
+	{	
+		URigVMLibraryNode* ReferencedNode = LoadReferencedNode();
+		if (ensure(ReferencedNode))
+		{ 
+			return ReferencedNode->GetNodeLayout(bIncludeEmptyCategories); 
+		}
 	}
 	return ReferencedFunctionHeader.Layout;
 }
diff --git a/Engine/Plugins/Runtime/RigVM/Source/RigVMDeveloper/Private/RigVMModel/RigVMClient.cpp b/Engine/Plugins/Runtime/RigVM/Source/RigVMDeveloper/Private/RigVMModel/RigVMClient.cpp
--- a/Engine/Plugins/Runtime/RigVM/Source/RigVMDeveloper/Private/RigVMModel/RigVMClient.cpp
+++ b/Engine/Plugins/Runtime/RigVM/Source/RigVMDeveloper/Private/RigVMModel/RigVMClient.cpp
@@ -1241,6 +1241,7 @@ FRigVMClientPatchResult FRigVMClient::PatchModelsOnLoad()
 			Result.Merge(Controller->PatchInvalidLinksOnWildcards());
 			Result.Merge(Controller->PatchExecutePins());
 			Result.Merge(Controller->PatchLazyPins());
+			Result.Merge(Controller->PatchUserDefinedStructPinNames());
 		
 			if (URigVMCollapseNode* CollapseNode = Model->GetTypedOuter<URigVMCollapseNode>())
 			{
diff --git a/Engine/Plugins/Runtime/RigVM/Source/RigVMDeveloper/Private/RigVMModel/RigVMController.cpp b/Engine/Plugins/Runtime/RigVM/Source/RigVMDeveloper/Private/RigVMModel/RigVMController.cpp
--- a/Engine/Plugins/Runtime/RigVM/Source/RigVMDeveloper/Private/RigVMModel/RigVMController.cpp
+++ b/Engine/Plugins/Runtime/RigVM/Source/RigVMDeveloper/Private/RigVMModel/RigVMController.cpp
@@ -445,6 +445,21 @@ FRigVMPinInfo::FRigVMPinInfo(const URigVMPin* InPin, int32 InParentIndex, ERigVM
 	// this method describes the info as currently represented in the model.
 	if (InPin)
 	{
+		const URigVMPin* Pin = InPin;
+		while (Pin)
+		{
+			if (const UUserDefinedStruct* UDStruct = Cast<const UUserDefinedStruct>(Pin->GetCPPTypeObject()))
+			{
+				const FString SubPinPath = InPin->GetSubPinPath(Pin);
+				if (const FProperty* StructProperty = UDStruct->FindPropertyByName(*SubPinPath))
+				{
+					Name = *UDStruct->GetAuthoredNameForField(StructProperty);
+					break;
+				}
+			}
+			Pin = Pin->GetParentPin();
+		}
+
 		CorrectExecuteTypeIndex();
 	
 		DefaultValue = InPin->GetDefaultValue();
@@ -473,6 +488,11 @@ FRigVMPinInfo::FRigVMPinInfo(FProperty* InProperty, ERigVMPinDirection InDirecti
 {
 	// this method describes the info as needed based on the property structure
 
+	if (const UUserDefinedStruct* UDStruct = Cast<const UUserDefinedStruct>(Property->GetOwnerStruct()))
+	{
+		Name = *UDStruct->GetAuthoredNameForField(InProperty);
+	}
+
 	if (Direction == ERigVMPinDirection::Invalid)
 	{
 		Direction = FRigVMStruct::GetPinDirectionFromProperty(Property);
@@ -23085,6 +23105,56 @@ FRigVMClientPatchResult URigVMController::PatchPinDefaultValues()
 
 }
 
+FRigVMClientPatchResult URigVMController::PatchUserDefinedStructPinNames()
+{
+	FRigVMClientPatchResult Result;
+	if (const URigVMGraph* Graph = GetGraph())
+	{
+		for (const URigVMNode* Node : Graph->GetNodes())
+		{
+			for (URigVMPin* Pin : Node->GetAllPinsRecursively())
+			{
+				if (UUserDefinedStruct* UDStruct = Cast<UUserDefinedStruct>(Pin->GetCPPTypeObject()))
+				{
+					for (URigVMPin* SubPin : Pin->GetSubPins())
+					{
+						FString SubPinPath = SubPin->GetSubPinPath(Pin);
+						if (const FProperty* StructSubProperty = UDStruct->FindPropertyByName(*SubPinPath))
+						{
+							FString NewNameSubPin = *UDStruct->GetAuthoredNameForField(StructSubProperty);
+							if (!NewNameSubPin.Equals(SubPin->GetName(), ESearchCase::IgnoreCase))
+							{
+								Result.bChangedContent = true;
+
+								// We are going to rename the pin, so we need to modify the link paths
+								// First we will find the link pins, then rename, then udpate the link paths from the link pins
+								for (URigVMLink* Link : SubPin->GetSourceLinks(true))
+								{
+									Link->UpdatePinPointers();
+								}
+								for (URigVMLink* Link : SubPin->GetTargetLinks(true))
+								{
+									Link->UpdatePinPointers();
+								}
+								SubPin->Rename(*NewNameSubPin, nullptr);
+								for (URigVMLink* Link : SubPin->GetSourceLinks(true))
+								{
+									Link->UpdatePinPaths();
+								}
+								for (URigVMLink* Link : SubPin->GetTargetLinks(true))
+								{
+									Link->UpdatePinPaths();
+								}
+							}
+						}
+					}
+				}
+			}
+		}
+	}
+	return Result;
+}
+
 ERigVMPinDefaultValueType URigVMController::GetDefaultValueType(const URigVMPin* InPin, const FString& InDefaultValue) const
 {
 	if(!CVarRigVMEnablePinDefaultTypes.GetValueOnAnyThread())
diff --git a/Engine/Plugins/Runtime/RigVM/Source/RigVMDeveloper/Public/RigVMModel/RigVMController.h b/Engine/Plugins/Runtime/RigVM/Source/RigVMDeveloper/Public/RigVMModel/RigVMController.h
--- a/Engine/Plugins/Runtime/RigVM/Source/RigVMDeveloper/Public/RigVMModel/RigVMController.h
+++ b/Engine/Plugins/Runtime/RigVM/Source/RigVMDeveloper/Public/RigVMModel/RigVMController.h
@@ -1498,6 +1498,7 @@ protected:
 	FRigVMClientPatchResult PatchExecutePins();
 	FRigVMClientPatchResult PatchLazyPins();
 	FRigVMClientPatchResult PatchPinDefaultValues();
+	FRigVMClientPatchResult PatchUserDefinedStructPinNames();
 
 	ERigVMPinDefaultValueType GetDefaultValueType(const URigVMPin* InPin, const FString& InDefaultValue) const;
 
