Control Rig: Allow redirectors from UserDefinedStructs to CPP structs, when loading local variables and function arguments
#rnx

diff --git a/Engine/Plugins/Runtime/RigVM/Source/RigVM/Public/RigVMTypeUtils.h b/Engine/Plugins/Runtime/RigVM/Source/RigVM/Public/RigVMTypeUtils.h
--- a/Engine/Plugins/Runtime/RigVM/Source/RigVM/Public/RigVMTypeUtils.h
+++ b/Engine/Plugins/Runtime/RigVM/Source/RigVM/Public/RigVMTypeUtils.h
@@ -465,6 +465,18 @@ namespace RigVMTypeUtils
 		return CPPType;
 	}
 
+	static bool FixCPPTypeAndObject(FString& InOutCPPType, TObjectPtr<UObject>& InOutCPPTypeObject, const FRigVMUserDefinedTypeResolver* InResolvalInfo = nullptr)
+	{
+		const FString OldCPPType = InOutCPPType;
+		UObject* OldCPPTypeObject = InOutCPPTypeObject;
+		InOutCPPType = PostProcessCPPType(InOutCPPType, InOutCPPTypeObject, InResolvalInfo);
+
+		// Type object might have changed through a redirect
+		InOutCPPTypeObject = ObjectFromCPPType(InOutCPPType);
+		
+		return !InOutCPPType.Equals(OldCPPType) || InOutCPPTypeObject != OldCPPTypeObject;
+	}
+
 	// helper function to retrieve an object from a path
 	static UObject* FindObjectFromCPPTypeObjectPath(const FString& InObjectPath)
 	{
diff --git a/Engine/Plugins/Runtime/RigVM/Source/RigVMDeveloper/Private/RigVMModel/RigVMClient.cpp b/Engine/Plugins/Runtime/RigVM/Source/RigVMDeveloper/Private/RigVMModel/RigVMClient.cpp
--- a/Engine/Plugins/Runtime/RigVM/Source/RigVMDeveloper/Private/RigVMModel/RigVMClient.cpp
+++ b/Engine/Plugins/Runtime/RigVM/Source/RigVMDeveloper/Private/RigVMModel/RigVMClient.cpp
@@ -1242,6 +1242,7 @@ FRigVMClientPatchResult FRigVMClient::PatchModelsOnLoad()
 			Result.Merge(Controller->PatchExecutePins());
 			Result.Merge(Controller->PatchLazyPins());
 			Result.Merge(Controller->PatchUserDefinedStructPinNames());
+			Result.Merge(Controller->PatchLocalVariableTypes());
 		
 			if (URigVMCollapseNode* CollapseNode = Model->GetTypedOuter<URigVMCollapseNode>())
 			{
diff --git a/Engine/Plugins/Runtime/RigVM/Source/RigVMDeveloper/Private/RigVMModel/RigVMController.cpp b/Engine/Plugins/Runtime/RigVM/Source/RigVMDeveloper/Private/RigVMModel/RigVMController.cpp
--- a/Engine/Plugins/Runtime/RigVM/Source/RigVMDeveloper/Private/RigVMModel/RigVMController.cpp
+++ b/Engine/Plugins/Runtime/RigVM/Source/RigVMDeveloper/Private/RigVMModel/RigVMController.cpp
@@ -21088,17 +21088,22 @@ void URigVMController::ValidatePin(URigVMPin* InPin)
 	}
 }
 
-void URigVMController::EnsureLocalVariableValidity()
+bool URigVMController::EnsureLocalVariableValidity()
 {
+	bool bChanged = false;
 	if (URigVMGraph* Graph = GetGraph())
 	{
 		for (FRigVMGraphVariableDescription& Variable : Graph->LocalVariables)
 		{
 			// CPPType can become invalid when the type object is defined by
 			// an asset that have changed name or asset path, user defined struct is one possibility
-			Variable.CPPType = RigVMTypeUtils::PostProcessCPPType(Variable.CPPType, Variable.CPPTypeObject);
+			if (RigVMTypeUtils::FixCPPTypeAndObject(Variable.CPPType, Variable.CPPTypeObject))
+			{
+				bChanged = true;
+			}
 		}
 	}
+	return bChanged;
 }
 
 FRigVMExternalVariable URigVMController::GetVariableByName(const FName& InExternalVariableName, const bool bIncludeInputArguments) const
@@ -23155,6 +23160,19 @@ FRigVMClientPatchResult URigVMController::PatchUserDefinedStructPinNames()
 	return Result;
 }
 
+FRigVMClientPatchResult URigVMController::PatchLocalVariableTypes()
+{
+	FRigVMClientPatchResult Result;
+	if (const URigVMGraph* Graph = GetGraph())
+	{
+		if (EnsureLocalVariableValidity())
+		{
+			Result.bChangedContent = true;
+		}
+	}
+	return Result;
+}
+
 ERigVMPinDefaultValueType URigVMController::GetDefaultValueType(const URigVMPin* InPin, const FString& InDefaultValue) const
 {
 	if(!CVarRigVMEnablePinDefaultTypes.GetValueOnAnyThread())
diff --git a/Engine/Plugins/Runtime/RigVM/Source/RigVMDeveloper/Private/RigVMModel/RigVMPin.cpp b/Engine/Plugins/Runtime/RigVM/Source/RigVMDeveloper/Private/RigVMModel/RigVMPin.cpp
--- a/Engine/Plugins/Runtime/RigVM/Source/RigVMDeveloper/Private/RigVMModel/RigVMPin.cpp
+++ b/Engine/Plugins/Runtime/RigVM/Source/RigVMDeveloper/Private/RigVMModel/RigVMPin.cpp
@@ -1665,7 +1665,7 @@ void URigVMPin::UpdateTypeInformationIfRequired() const
 		{
 			URigVMPin* MutableThis = (URigVMPin*)this;
 			MutableThis->CPPTypeObject = RigVMTypeUtils::FindObjectFromCPPTypeObjectPath(CPPTypeObjectPath.ToString());
-			MutableThis->CPPType = RigVMTypeUtils::PostProcessCPPType(CPPType, CPPTypeObject);
+			RigVMTypeUtils::FixCPPTypeAndObject(MutableThis->CPPType, MutableThis->CPPTypeObject);
 			if (!MutableThis->CPPType.IsEmpty())
 			{
 				MutableThis->LastKnownTypeIndex = FRigVMRegistry::Get().FindOrAddType(FRigVMTemplateArgumentType(*CPPType, CPPTypeObject));
diff --git a/Engine/Plugins/Runtime/RigVM/Source/RigVMDeveloper/Public/RigVMModel/RigVMController.h b/Engine/Plugins/Runtime/RigVM/Source/RigVMDeveloper/Public/RigVMModel/RigVMController.h
--- a/Engine/Plugins/Runtime/RigVM/Source/RigVMDeveloper/Public/RigVMModel/RigVMController.h
+++ b/Engine/Plugins/Runtime/RigVM/Source/RigVMDeveloper/Public/RigVMModel/RigVMController.h
@@ -1391,7 +1391,7 @@ private:
 
 	// recreate the CPP type strings for variables that reference a type object
 	// they can get out of sync when the variable references a user defined struct
-	void EnsureLocalVariableValidity();
+	bool EnsureLocalVariableValidity();
 	
 	FRigVMExternalVariable GetVariableByName(const FName& InExternalVariableName, const bool bIncludeInputArguments = false) const;
 	TArray<FRigVMExternalVariable> GetAllVariables(const bool bIncludeInputArguments = false) const;
@@ -1499,6 +1499,7 @@ protected:
 	FRigVMClientPatchResult PatchLazyPins();
 	FRigVMClientPatchResult PatchPinDefaultValues();
 	FRigVMClientPatchResult PatchUserDefinedStructPinNames();
+	FRigVMClientPatchResult PatchLocalVariableTypes();
 
 	ERigVMPinDefaultValueType GetDefaultValueType(const URigVMPin* InPin, const FString& InDefaultValue) const;
 
