Control Rig: Add CPPTypeObjectPath to FRigVMGraphVariableDescription
#rnx

diff --git a/Engine/Plugins/Runtime/RigVM/Source/RigVMDeveloper/Private/RigVMBlueprint.cpp b/Engine/Plugins/Runtime/RigVM/Source/RigVMDeveloper/Private/RigVMBlueprint.cpp
--- a/Engine/Plugins/Runtime/RigVM/Source/RigVMDeveloper/Private/RigVMBlueprint.cpp
+++ b/Engine/Plugins/Runtime/RigVM/Source/RigVMDeveloper/Private/RigVMBlueprint.cpp
@@ -2700,6 +2700,10 @@ TArray<FRigVMGraphVariableDescription> URigVMBlueprint::GetMemberVariables() con
 		RigVMTypeUtils::CPPTypeFromPinType(BPVariable.VarType, CPPType, &CPPTypeObject);
 		NewVariable.CPPType = CPPType;
 		NewVariable.CPPTypeObject = CPPTypeObject;
+		if (NewVariable.CPPTypeObject)
+		{
+			NewVariable.CPPTypeObjectPath = *NewVariable.CPPTypeObject->GetPathName();
+		}
 		Variables.Add(NewVariable);
 	}
 
diff --git a/Engine/Plugins/Runtime/RigVM/Source/RigVMDeveloper/Private/RigVMModel/Nodes/RigVMVariableNode.cpp b/Engine/Plugins/Runtime/RigVM/Source/RigVMDeveloper/Private/RigVMModel/Nodes/RigVMVariableNode.cpp
--- a/Engine/Plugins/Runtime/RigVM/Source/RigVMDeveloper/Private/RigVMModel/Nodes/RigVMVariableNode.cpp
+++ b/Engine/Plugins/Runtime/RigVM/Source/RigVMDeveloper/Private/RigVMModel/Nodes/RigVMVariableNode.cpp
@@ -119,6 +119,10 @@ FRigVMGraphVariableDescription URigVMVariableNode::GetVariableDescription() cons
 	Variable.Name = GetVariableName();
 	Variable.CPPType = GetCPPType();
 	Variable.CPPTypeObject = GetCPPTypeObject();
+	if(Variable.CPPTypeObject)
+	{
+		Variable.CPPTypeObjectPath = *Variable.CPPTypeObject->GetPathName();
+	}
 	Variable.DefaultValue = GetDefaultValue();
 	return Variable;
 }
diff --git a/Engine/Plugins/Runtime/RigVM/Source/RigVMDeveloper/Private/RigVMModel/RigVMController.cpp b/Engine/Plugins/Runtime/RigVM/Source/RigVMDeveloper/Private/RigVMModel/RigVMController.cpp
--- a/Engine/Plugins/Runtime/RigVM/Source/RigVMDeveloper/Private/RigVMModel/RigVMController.cpp
+++ b/Engine/Plugins/Runtime/RigVM/Source/RigVMDeveloper/Private/RigVMModel/RigVMController.cpp
@@ -14630,6 +14630,10 @@ FRigVMGraphVariableDescription URigVMController::AddLocalVariable(const FName& I
 	NewVariable.Name = VariableName;
 	NewVariable.CPPType = InCPPType;
 	NewVariable.CPPTypeObject = InCPPTypeObject;
+	if (InCPPTypeObject)
+	{
+		NewVariable.CPPTypeObjectPath = *InCPPTypeObject->GetPathName();
+	}
 	NewVariable.DefaultValue = InDefaultValue;
 
 	Graph->LocalVariables.Add(NewVariable);
@@ -14982,6 +14986,14 @@ bool URigVMController::SetLocalVariableType(const FName& InVariableName, const F
 	
 	LocalVariables[FoundIndex].CPPType = InCPPType;
 	LocalVariables[FoundIndex].CPPTypeObject = InCPPTypeObject;
+	if (InCPPTypeObject)
+	{
+		LocalVariables[FoundIndex].CPPTypeObjectPath = *InCPPTypeObject->GetPathName();
+	}
+	else
+	{
+		LocalVariables[FoundIndex].CPPTypeObjectPath = NAME_None;
+	}
 
 	// Set default value
 	if (UScriptStruct* ScriptStruct = Cast<UScriptStruct>(InCPPTypeObject))
@@ -21101,6 +21113,25 @@ bool URigVMController::EnsureLocalVariableValidity()
 			{
 				bChanged = true;
 			}
+			
+			if (Variable.CPPTypeObject)
+			{
+				Variable.CPPTypeObjectPath = *Variable.CPPTypeObject.GetPathName();
+			}
+			else if(!Variable.CPPTypeObjectPath.IsNone())
+			{
+				Variable.CPPTypeObject = RigVMTypeUtils::FindObjectFromCPPTypeObjectPath(Variable.CPPTypeObjectPath.ToString());
+				RigVMTypeUtils::FixCPPTypeAndObject(Variable.CPPType, Variable.CPPTypeObject);
+				if (Variable.CPPTypeObject)
+				{
+					bChanged = true;
+				}
+			}
+			
+			if (!Variable.CPPTypeObject)
+			{
+				Variable.CPPTypeObjectPath = NAME_None;
+			}
 		}
 	}
 	return bChanged;
diff --git a/Engine/Plugins/Runtime/RigVM/Source/RigVMDeveloper/Private/RigVMModel/RigVMGraph.cpp b/Engine/Plugins/Runtime/RigVM/Source/RigVMDeveloper/Private/RigVMModel/RigVMGraph.cpp
--- a/Engine/Plugins/Runtime/RigVM/Source/RigVMDeveloper/Private/RigVMModel/RigVMGraph.cpp
+++ b/Engine/Plugins/Runtime/RigVM/Source/RigVMDeveloper/Private/RigVMModel/RigVMGraph.cpp
@@ -384,6 +384,10 @@ TArray<FRigVMGraphVariableDescription> URigVMGraph::GetInputArguments() const
 			Description.Name = Pin->GetFName();
 			Description.CPPType = Pin->GetCPPType();
 			Description.CPPTypeObject = Pin->GetCPPTypeObject();
+			if (Description.CPPTypeObject)
+			{
+				Description.CPPTypeObjectPath = *Description.CPPTypeObject->GetPathName();
+			}
 			Inputs.Add(Description);
 		}
 	}
@@ -405,7 +409,11 @@ TArray<FRigVMGraphVariableDescription> URigVMGraph::GetOutputArguments() const
 			Description.Name = Pin->GetFName();
 			Description.CPPType = Pin->GetCPPType();
 			Description.CPPTypeObject = Pin->GetCPPTypeObject();
-			Outputs.Add(Description);			
+			if (Description.CPPTypeObject)
+			{
+				Description.CPPTypeObjectPath = *Description.CPPTypeObject->GetPathName();
+			}
+			Outputs.Add(Description);
 		}
 	}
 	return Outputs;
diff --git a/Engine/Plugins/Runtime/RigVM/Source/RigVMDeveloper/Public/RigVMModel/RigVMVariableDescription.h b/Engine/Plugins/Runtime/RigVM/Source/RigVMDeveloper/Public/RigVMModel/RigVMVariableDescription.h
--- a/Engine/Plugins/Runtime/RigVM/Source/RigVMDeveloper/Public/RigVMModel/RigVMVariableDescription.h
+++ b/Engine/Plugins/Runtime/RigVM/Source/RigVMDeveloper/Public/RigVMModel/RigVMVariableDescription.h
@@ -37,6 +37,9 @@ public:
 	// The Struct of the C++ data type of the variable (or nullptr)
 	UPROPERTY(EditAnywhere, BlueprintReadOnly, Category = RigVMGraphVariableDescription)
 	TObjectPtr<UObject> CPPTypeObject = nullptr;
+	
+	UPROPERTY()
+	FName CPPTypeObjectPath;
 
 	// The default value of the variable
 	UPROPERTY(EditAnywhere, BlueprintReadOnly, Category = RigVMGraphVariableDescription)
@@ -62,6 +65,14 @@ public:
 		UObject* Object = nullptr;
 		const bool bSuccess = RigVMTypeUtils::CPPTypeFromPinType(PinType, CPPType, &Object);
 		CPPTypeObject = Object;
+		if (CPPTypeObject)
+		{
+			CPPTypeObjectPath = *CPPTypeObject->GetPathName();
+		}
+		else
+		{
+			CPPTypeObjectPath = NAME_None;
+		}
 		return bSuccess;
 	}
 	
