Contorl Rig: Fix break content from user defined enums not finding the object from the cpp type
#rnx

diff --git a/Engine/Plugins/Runtime/RigVM/Source/RigVM/Public/RigVMTypeUtils.h b/Engine/Plugins/Runtime/RigVM/Source/RigVM/Public/RigVMTypeUtils.h
--- a/Engine/Plugins/Runtime/RigVM/Source/RigVM/Public/RigVMTypeUtils.h
+++ b/Engine/Plugins/Runtime/RigVM/Source/RigVM/Public/RigVMTypeUtils.h
@@ -472,7 +472,21 @@ namespace RigVMTypeUtils
 		InOutCPPType = PostProcessCPPType(InOutCPPType, InOutCPPTypeObject, InResolvalInfo);
 
 		// Type object might have changed through a redirect
-		InOutCPPTypeObject = ObjectFromCPPType(InOutCPPType);
+		bool bFindObjectAgain = true;
+
+		// If its a user defined enum, do not attempt to find the object from the cpp type
+		if (const UEnum* Enum = Cast<UEnum>(InOutCPPTypeObject))
+		{
+			if (Enum->CppType.IsEmpty())
+			{
+				bFindObjectAgain = false;
+			}
+		}
+
+		if (bFindObjectAgain)
+		{
+			InOutCPPTypeObject = ObjectFromCPPType(InOutCPPType);
+		}
 		
 		return !InOutCPPType.Equals(OldCPPType) || InOutCPPTypeObject != OldCPPTypeObject;
 	}
