RigVM: Introduce the notion of an empty execute context to be used on nodes to save memory

#rnx

diff --git a/Engine/Plugins/Runtime/RigVM/Source/RigVM/Private/RigVMCore/RigVMFunction.cpp b/Engine/Plugins/Runtime/RigVM/Source/RigVM/Private/RigVMCore/RigVMFunction.cpp
--- a/Engine/Plugins/Runtime/RigVM/Source/RigVM/Private/RigVMCore/RigVMFunction.cpp
+++ b/Engine/Plugins/Runtime/RigVM/Source/RigVM/Private/RigVMCore/RigVMFunction.cpp
@@ -146,38 +146,13 @@ const UScriptStruct* FRigVMFunction::GetExecuteContextStruct(bool bLockRegistry)
 	{
 		return Factory->GetExecuteContextStruct();
 	}
-	if(Struct)
+	if(Struct && Struct->IsChildOf(FRigVMStruct::StaticStruct()))
 	{
-#if WITH_EDITOR
-		FString ExecuteContextName;
-		if(Struct->GetStringMetaDataHierarchical(FRigVMStruct::ExecuteContextName, &ExecuteContextName))
-		{
-			const FRigVMTemplateArgumentType& Type = FRigVMRegistry_RWLock::Get().FindTypeFromCPPType(ExecuteContextName, bLockRegistry);
-			if(const UScriptStruct* ExecuteContextStruct = Cast<UScriptStruct>(Type.CPPTypeObject))
-			{
-				return ExecuteContextStruct;
-			}
-		}
-#endif
-
-		for (TFieldIterator<FProperty> It(Struct); It; ++It)
-		{
-			const FProperty* Property = *It;
-
-			if(const FArrayProperty* ArrayProperty = CastField<FArrayProperty>(Property))
-			{
-				Property = ArrayProperty->Inner;
-			}
-			if(const FStructProperty* StructProperty = CastField<FStructProperty>(Property))
-			{
-				if(StructProperty->Struct->IsChildOf(FRigVMExecuteContext::StaticStruct()))
-				{
-					return StructProperty->Struct;
-				}
-			}
-		}
+		FStructOnScope StructOnScope(Struct);
+		const FRigVMStruct* StructOnScopeMemory = reinterpret_cast<const FRigVMStruct*>(StructOnScope.GetStructMemory());
+		return StructOnScopeMemory->GetExecuteContextStruct();
+		
 	}
-
 	return FRigVMExecuteContext::StaticStruct();
 }
 
diff --git a/Engine/Plugins/Runtime/RigVM/Source/RigVM/Private/RigVMCore/RigVMGraphFunctionDefinition.cpp b/Engine/Plugins/Runtime/RigVM/Source/RigVM/Private/RigVMCore/RigVMGraphFunctionDefinition.cpp
--- a/Engine/Plugins/Runtime/RigVM/Source/RigVM/Private/RigVMCore/RigVMGraphFunctionDefinition.cpp
+++ b/Engine/Plugins/Runtime/RigVM/Source/RigVM/Private/RigVMCore/RigVMGraphFunctionDefinition.cpp
@@ -55,7 +55,7 @@ bool FRigVMGraphFunctionArgument::IsExecuteContext() const
 	{
 		if(const UScriptStruct* ScriptStruct = Cast<UScriptStruct>(CPPTypeObject.Get()))
 		{
-			if(ScriptStruct->IsChildOf(FRigVMExecuteContext::StaticStruct()))
+			if(ScriptStruct->IsChildOf(FRigVMEmptyExecuteContext::StaticStruct()))
 			{
 				return true;
 			}
@@ -217,7 +217,7 @@ bool FRigVMGraphFunctionHeader::IsMutable() const
 		{
 			if(UScriptStruct* Struct = Cast<UScriptStruct>(Arg.CPPTypeObject.Get()))
 			{
-				if(Struct->IsChildOf(FRigVMExecuteContext::StaticStruct()))
+				if(Struct->IsChildOf(FRigVMEmptyExecuteContext::StaticStruct()))
 				{
 					return true;
 				}
diff --git a/Engine/Plugins/Runtime/RigVM/Source/RigVM/Private/RigVMCore/RigVMRegistry.cpp b/Engine/Plugins/Runtime/RigVM/Source/RigVM/Private/RigVMCore/RigVMRegistry.cpp
--- a/Engine/Plugins/Runtime/RigVM/Source/RigVM/Private/RigVMCore/RigVMRegistry.cpp
+++ b/Engine/Plugins/Runtime/RigVM/Source/RigVM/Private/RigVMCore/RigVMRegistry.cpp
@@ -385,7 +385,7 @@ void FRigVMRegistry_NoLock::RefreshEngineTypes_NoLock()
 		UScriptStruct* ScriptStruct = *ScriptStructIt;
 		
 		// if this is a C++ type - skip it
-		if(ScriptStruct->IsA<UUserDefinedStruct>() || ScriptStruct->IsChildOf(FRigVMExecuteContext::StaticStruct()))
+		if(ScriptStruct->IsA<UUserDefinedStruct>() || ScriptStruct->IsChildOf(FRigVMEmptyExecuteContext::StaticStruct()))
 		{
 			// this check for example makes sure we don't add structs defined in verse
 			if(IsAllowedType_NoLock(ScriptStruct))
@@ -615,7 +615,7 @@ TRigVMTypeIndex FRigVMRegistry_NoLock::FindOrAddType_NoLock(const FRigVMTemplate
 		bool bIsExecute = false;
 		if(const UScriptStruct* ScriptStruct = Cast<UScriptStruct>(CPPTypeObject))
 		{
-			bIsExecute = ScriptStruct->IsChildOf(FRigVMExecuteContext::StaticStruct());
+			bIsExecute = ScriptStruct->IsChildOf(FRigVMEmptyExecuteContext::StaticStruct());
 		}
 
 		TArray<TRigVMTypeIndex> Indices;
@@ -751,7 +751,7 @@ TRigVMTypeIndex FRigVMRegistry_NoLock::FindOrAddType_NoLock(const FRigVMTemplate
 			}
 			else if(const UStruct* Struct = Cast<UStruct>(CPPTypeObject))
 			{
-				if(Struct->IsChildOf(FRigVMExecuteContext::StaticStruct()))
+				if(Struct->IsChildOf(FRigVMEmptyExecuteContext::StaticStruct()))
 				{
 					if(ArrayDimension == 0)
 					{
diff --git a/Engine/Plugins/Runtime/RigVM/Source/RigVM/Private/RigVMCore/RigVMStruct.cpp b/Engine/Plugins/Runtime/RigVM/Source/RigVM/Private/RigVMCore/RigVMStruct.cpp
--- a/Engine/Plugins/Runtime/RigVM/Source/RigVM/Private/RigVMCore/RigVMStruct.cpp
+++ b/Engine/Plugins/Runtime/RigVM/Source/RigVM/Private/RigVMCore/RigVMStruct.cpp
@@ -453,11 +453,12 @@ bool FRigVMStruct::CheckPinType(UScriptStruct* InStruct, const FName& PinName, c
 	{
 		if (Property->GetCPPType() != ExpectedType)
 		{
-			if(ExpectedType == FRigVMExecuteContext::StaticStruct()->GetStructCPPName())
+			if(ExpectedType == FRigVMExecuteContext::StaticStruct()->GetStructCPPName() ||
+				ExpectedType == FRigVMEmptyExecuteContext::StaticStruct()->GetStructCPPName())
 			{
 				if(const FStructProperty* StructProperty = CastField<FStructProperty>(Property))
 				{
-					if(StructProperty->Struct->IsChildOf(FRigVMExecuteContext::StaticStruct()))
+					if(StructProperty->Struct->IsChildOf(FRigVMEmptyExecuteContext::StaticStruct()))
 					{
 						return true;
 					}
diff --git a/Engine/Plugins/Runtime/RigVM/Source/RigVM/Public/RigVMCore/RigVMExecuteContext.h b/Engine/Plugins/Runtime/RigVM/Source/RigVM/Public/RigVMCore/RigVMExecuteContext.h
--- a/Engine/Plugins/Runtime/RigVM/Source/RigVM/Public/RigVMCore/RigVMExecuteContext.h
+++ b/Engine/Plugins/Runtime/RigVM/Source/RigVM/Public/RigVMCore/RigVMExecuteContext.h
@@ -250,15 +250,25 @@ struct FRigVMExecuteContextFixLayout
 	virtual ~FRigVMExecuteContextFixLayout() = default;
 };
 
+// This structure is used for semantically describing a pin
+// on a RigVMStruct - but without any content. So rather than
+// adding an actual FRigVMExecuteContext this can be added as
+// the type for a pin and it won't use any memory. 
+USTRUCT(BlueprintType)
+struct RIGVM_API FRigVMEmptyExecuteContext
+#if CPP
+	: public FRigVMExecuteContextFixLayout
+#endif
+{
+	GENERATED_BODY()
+};
+
 /**
  * The execute context is used for mutable nodes to
  * indicate execution order.
  */
 USTRUCT(BlueprintType, meta=(DisplayName="Execute Context"))
-struct FRigVMExecuteContext 
-#if CPP
-	: public FRigVMExecuteContextFixLayout
-#endif
+struct FRigVMExecuteContext : public FRigVMEmptyExecuteContext
 {
 	GENERATED_BODY()
 
diff --git a/Engine/Plugins/Runtime/RigVM/Source/RigVMDeveloper/Private/EdGraph/RigVMEdGraphNode.cpp b/Engine/Plugins/Runtime/RigVM/Source/RigVMDeveloper/Private/EdGraph/RigVMEdGraphNode.cpp
--- a/Engine/Plugins/Runtime/RigVM/Source/RigVMDeveloper/Private/EdGraph/RigVMEdGraphNode.cpp
+++ b/Engine/Plugins/Runtime/RigVM/Source/RigVMDeveloper/Private/EdGraph/RigVMEdGraphNode.cpp
@@ -1602,7 +1602,7 @@ void URigVMEdGraphNode::UpdatePinLists()
 					if (ModelPin->IsStruct())
 					{
 						const UScriptStruct* ScriptStruct = ModelPin->GetScriptStruct();
-						if (ScriptStruct && ScriptStruct->IsChildOf(FRigVMExecuteContext::StaticStruct()))
+						if (ScriptStruct && ScriptStruct->IsChildOf(FRigVMEmptyExecuteContext::StaticStruct()))
 						{
 							ExecutePins.Add(ModelPin);
 							continue;
diff --git a/Engine/Plugins/Runtime/RigVM/Source/RigVMDeveloper/Private/EdGraph/RigVMEdGraphSchema.cpp b/Engine/Plugins/Runtime/RigVM/Source/RigVMDeveloper/Private/EdGraph/RigVMEdGraphSchema.cpp
--- a/Engine/Plugins/Runtime/RigVM/Source/RigVMDeveloper/Private/EdGraph/RigVMEdGraphSchema.cpp
+++ b/Engine/Plugins/Runtime/RigVM/Source/RigVMDeveloper/Private/EdGraph/RigVMEdGraphSchema.cpp
@@ -900,7 +900,7 @@ FLinearColor URigVMEdGraphSchema::GetPinTypeColor(const FEdGraphPinType& PinType
 	{
 		if (const UStruct* Struct = Cast<UStruct>(PinType.PinSubCategoryObject))
 		{
-			if (Struct->IsChildOf(FRigVMExecuteContext::StaticStruct()))
+			if (Struct->IsChildOf(FRigVMEmptyExecuteContext::StaticStruct()))
 			{
 				return FLinearColor::White;
 			}
@@ -1104,7 +1104,7 @@ bool URigVMEdGraphSchema::SupportsPinType(TWeakPtr<const FEdGraphSchemaAction> S
 		{
 			if(SchemaAction.IsValid() && SchemaAction.Pin()->IsAVariable())
 			{
-				if(ScriptStruct->IsChildOf(FRigVMExecuteContext::StaticStruct()))
+				if(ScriptStruct->IsChildOf(FRigVMEmptyExecuteContext::StaticStruct()))
 				{
 					return false;
 				}
@@ -1139,7 +1139,7 @@ bool URigVMEdGraphSchema::SupportsPinTypeContainer(TWeakPtr<const FEdGraphSchema
 	// Do not allow containers for execute context type
 	if(const UScriptStruct* ExecuteContextScriptStruct = Cast<UScriptStruct>(PinType.PinSubCategoryObject))
 	{
-		if (ExecuteContextScriptStruct->IsChildOf(FRigVMExecuteContext::StaticStruct()))
+		if (ExecuteContextScriptStruct->IsChildOf(FRigVMEmptyExecuteContext::StaticStruct()))
 		{
 			return ContainerType == EPinContainerType::None;
 		}
@@ -1768,7 +1768,7 @@ UEdGraphPin* URigVMEdGraphSchema::DropPinOnNode(UEdGraphNode* InTargetNode, cons
 
 								if(const UScriptStruct* CPPTypeStruct = Cast<UScriptStruct>(ExternalVar.TypeObject))
 								{
-									if(CPPTypeStruct->IsChildOf(FRigVMExecuteContext::StaticStruct()))
+									if(CPPTypeStruct->IsChildOf(FRigVMEmptyExecuteContext::StaticStruct()))
 									{
 										SourcePinName = FRigVMStruct::ExecuteContextName;
 										PinDirection = ERigVMPinDirection::IO;										
diff --git a/Engine/Plugins/Runtime/RigVM/Source/RigVMDeveloper/Private/RigVMCompiler/RigVMCodeGenerator.cpp b/Engine/Plugins/Runtime/RigVM/Source/RigVMDeveloper/Private/RigVMCompiler/RigVMCodeGenerator.cpp
--- a/Engine/Plugins/Runtime/RigVM/Source/RigVMDeveloper/Private/RigVMCompiler/RigVMCodeGenerator.cpp
+++ b/Engine/Plugins/Runtime/RigVM/Source/RigVMDeveloper/Private/RigVMCompiler/RigVMCodeGenerator.cpp
@@ -1178,7 +1178,7 @@ void FRigVMCodeGenerator::ParseInclude(UStruct* InDependency, const FName& InMet
 
 	if (InDependency == URigVM::StaticClass() ||
 		InDependency == FRigVMByteCode::StaticStruct() || 
-		InDependency->IsChildOf(FRigVMExecuteContext::StaticStruct()) ||
+		InDependency->IsChildOf(FRigVMEmptyExecuteContext::StaticStruct()) ||
 		InDependency->IsChildOf(FRigVMBaseOp::StaticStruct()) ||
 		InDependency->IsChildOf(URigVMMemoryStorage::StaticClass()))
 	{
diff --git a/Engine/Plugins/Runtime/RigVM/Source/RigVMDeveloper/Private/RigVMCompiler/RigVMCompiler.cpp b/Engine/Plugins/Runtime/RigVM/Source/RigVMDeveloper/Private/RigVMCompiler/RigVMCompiler.cpp
--- a/Engine/Plugins/Runtime/RigVM/Source/RigVMDeveloper/Private/RigVMCompiler/RigVMCompiler.cpp
+++ b/Engine/Plugins/Runtime/RigVM/Source/RigVMDeveloper/Private/RigVMCompiler/RigVMCompiler.cpp
@@ -952,7 +952,7 @@ bool URigVMCompiler::Compile(const FRigVMCompileSettings& InSettings, TArray<URi
 							
 							if(const FStructProperty* StructProperty = CastField<FStructProperty>(Property))
 							{
-								if(StructProperty->Struct->IsChildOf(FRigVMExecuteContext::StaticStruct()))
+								if(StructProperty->Struct->IsChildOf(FRigVMEmptyExecuteContext::StaticStruct()))
 								{
 									if(!ValidExecuteContextStructs.Contains(StructProperty->Struct))
 									{
@@ -2032,7 +2032,7 @@ bool URigVMCompiler::TraverseCallExtern(const FRigVMCallExternExprAST* InExpr, F
 
 	auto CheckExecuteStruct = [this, &WorkData](URigVMNode* Subject, const UScriptStruct* ExecuteStruct) -> bool
 	{
-		if(ExecuteStruct->IsChildOf(FRigVMExecuteContext::StaticStruct()))
+		if(ExecuteStruct->IsChildOf(FRigVMEmptyExecuteContext::StaticStruct()))
 		{
 			// top level expected execute struct is provided by the graph
 			const UScriptStruct* SpecializedExecuteStruct = WorkData.ExecuteContextStruct;
@@ -2609,7 +2609,7 @@ bool URigVMCompiler::TraverseInlineFunction(const FRigVMInlineFunctionExprAST* I
 					{
 						if (const UScriptStruct* ScriptStruct = Cast<UScriptStruct>(Argument.CPPTypeObject.Get()))
 						{
-							if (ScriptStruct->IsChildOf(FRigVMExecuteContext::StaticStruct()))
+							if (ScriptStruct->IsChildOf(FRigVMEmptyExecuteContext::StaticStruct()))
 							{
 								continue;
 							}
diff --git a/Engine/Plugins/Runtime/RigVM/Source/RigVMDeveloper/Private/RigVMModel/RigVMController.cpp b/Engine/Plugins/Runtime/RigVM/Source/RigVMDeveloper/Private/RigVMModel/RigVMController.cpp
--- a/Engine/Plugins/Runtime/RigVM/Source/RigVMDeveloper/Private/RigVMModel/RigVMController.cpp
+++ b/Engine/Plugins/Runtime/RigVM/Source/RigVMDeveloper/Private/RigVMModel/RigVMController.cpp
@@ -12515,7 +12515,7 @@ FName URigVMController::AddExposedPin(const FName& InPinName, ERigVMPinDirection
 	{
 		if(const UScriptStruct* CPPTypeStruct = Cast<UScriptStruct>(CPPTypeObject))
 		{
-			if(CPPTypeStruct->IsChildOf(FRigVMExecuteContext::StaticStruct()))
+			if(CPPTypeStruct->IsChildOf(FRigVMEmptyExecuteContext::StaticStruct()))
 			{
 				for(const URigVMPin* ExistingPin : LibraryNode->Pins)
 				{
@@ -13003,7 +13003,7 @@ bool URigVMController::ChangeExposedPinType(const FName& InPinName, const FStrin
 		{
 			if(const UScriptStruct* CPPTypeStruct = Cast<UScriptStruct>(CPPTypeObject))
 			{
-				if(CPPTypeStruct->IsChildOf(FRigVMExecuteContext::StaticStruct()))
+				if(CPPTypeStruct->IsChildOf(FRigVMEmptyExecuteContext::StaticStruct()))
 				{
 					for(URigVMPin* ExistingPin : LibraryNode->Pins)
 					{
@@ -13033,7 +13033,7 @@ bool URigVMController::ChangeExposedPinType(const FName& InPinName, const FStrin
 	{
 		if(const UScriptStruct* CPPTypeStruct = Cast<UScriptStruct>(CPPTypeObject))
 		{
-			if(CPPTypeStruct->IsChildOf(FRigVMExecuteContext::StaticStruct()))
+			if(CPPTypeStruct->IsChildOf(FRigVMEmptyExecuteContext::StaticStruct()))
 			{
 				bIsExecute = true;
 			}
diff --git a/Engine/Plugins/Runtime/RigVM/Source/RigVMDeveloper/Private/RigVMModel/RigVMNode.cpp b/Engine/Plugins/Runtime/RigVM/Source/RigVMDeveloper/Private/RigVMModel/RigVMNode.cpp
--- a/Engine/Plugins/Runtime/RigVM/Source/RigVMDeveloper/Private/RigVMModel/RigVMNode.cpp
+++ b/Engine/Plugins/Runtime/RigVM/Source/RigVMDeveloper/Private/RigVMModel/RigVMNode.cpp
@@ -691,12 +691,9 @@ bool URigVMNode::IsMutable() const
 {
 	for (const URigVMPin* Pin : GetPins())
 	{
-		if(const UScriptStruct* ScriptStruct = Pin->GetScriptStruct())
+		if(Pin->IsExecuteContext())
 		{
-			if (ScriptStruct->IsChildOf(FRigVMExecuteContext::StaticStruct()))
-			{
-				return true;
-			}
+			return true;
 		}
 	}
 	return false;
diff --git a/Engine/Plugins/Runtime/RigVM/Source/RigVMDeveloper/Private/RigVMModel/RigVMPin.cpp b/Engine/Plugins/Runtime/RigVM/Source/RigVMDeveloper/Private/RigVMModel/RigVMPin.cpp
--- a/Engine/Plugins/Runtime/RigVM/Source/RigVMDeveloper/Private/RigVMModel/RigVMPin.cpp
+++ b/Engine/Plugins/Runtime/RigVM/Source/RigVMDeveloper/Private/RigVMModel/RigVMPin.cpp
@@ -646,7 +646,7 @@ bool URigVMPin::IsExecuteContext() const
 {
 	if (const UScriptStruct* ScriptStruct = GetScriptStruct())
 	{
-		if (ScriptStruct->IsChildOf(FRigVMExecuteContext::StaticStruct()))
+		if (ScriptStruct->IsChildOf(FRigVMEmptyExecuteContext::StaticStruct()))
 		{
 			return true;
 		}
diff --git a/Engine/Plugins/Runtime/RigVM/Source/RigVMDeveloper/Private/RigVMModel/RigVMSchema.cpp b/Engine/Plugins/Runtime/RigVM/Source/RigVMDeveloper/Private/RigVMModel/RigVMSchema.cpp
--- a/Engine/Plugins/Runtime/RigVM/Source/RigVMDeveloper/Private/RigVMModel/RigVMSchema.cpp
+++ b/Engine/Plugins/Runtime/RigVM/Source/RigVMDeveloper/Private/RigVMModel/RigVMSchema.cpp
@@ -237,7 +237,7 @@ bool URigVMSchema::ShouldUnfoldStruct(URigVMController* InController, const UStr
 	{
 		return false;
 	}
-	if(InStruct->IsChildOf(FRigVMExecuteContext::StaticStruct()))
+	if(InStruct->IsChildOf(FRigVMEmptyExecuteContext::StaticStruct()))
 	{
 		return false;
 	}
diff --git a/Engine/Plugins/Runtime/RigVM/Source/RigVMEditor/Private/EdGraph/NodeSpawners/RigVMEdGraphUnitNodeSpawner.cpp b/Engine/Plugins/Runtime/RigVM/Source/RigVMEditor/Private/EdGraph/NodeSpawners/RigVMEdGraphUnitNodeSpawner.cpp
--- a/Engine/Plugins/Runtime/RigVM/Source/RigVMEditor/Private/EdGraph/NodeSpawners/RigVMEdGraphUnitNodeSpawner.cpp
+++ b/Engine/Plugins/Runtime/RigVM/Source/RigVMEditor/Private/EdGraph/NodeSpawners/RigVMEdGraphUnitNodeSpawner.cpp
@@ -385,7 +385,7 @@ void URigVMEdGraphUnitNodeSpawner::HookupMutableNode(URigVMNode* InModelNode, UR
 	{
 		if (Pin->GetScriptStruct())
 		{
-			if (Pin->GetScriptStruct()->IsChildOf(FRigVMExecuteContext::StaticStruct()))
+			if (Pin->GetScriptStruct()->IsChildOf(FRigVMEmptyExecuteContext::StaticStruct()))
 			{
 				if (Pin->GetDirection() == ERigVMPinDirection::IO || Pin->GetDirection() == ERigVMPinDirection::Input)
 				{
@@ -428,21 +428,18 @@ void URigVMEdGraphUnitNodeSpawner::HookupMutableNode(URigVMNode* InModelNode, UR
 
 				for (URigVMPin* Pin : OtherModelNode->GetPins())
 				{
-					if (Pin->GetScriptStruct())
+					if (Pin->IsExecuteContext())
 					{
-						if (Pin->GetScriptStruct()->IsChildOf(FRigVMExecuteContext::StaticStruct()))
+						if (Pin->GetDirection() == ERigVMPinDirection::IO || Pin->GetDirection() == ERigVMPinDirection::Output)
 						{
-							if (Pin->GetDirection() == ERigVMPinDirection::IO || Pin->GetDirection() == ERigVMPinDirection::Output)
+							if (Pin->GetLinkedTargetPins().Num() == 0)
 							{
-								if (Pin->GetLinkedTargetPins().Num() == 0)
+								float Distance = (OtherModelNode->GetPosition() - InModelNode->GetPosition()).Size();
+								if (Distance < ClosestDistance)
 								{
-									float Distance = (OtherModelNode->GetPosition() - InModelNode->GetPosition()).Size();
-									if (Distance < ClosestDistance)
-									{
-										ClosestOtherModelNodeExecutePin = Pin;
-										ClosestDistance = Distance;
-										break;
-									}
+									ClosestOtherModelNodeExecutePin = Pin;
+									ClosestDistance = Distance;
+									break;
 								}
 							}
 						}
diff --git a/Engine/Plugins/Runtime/RigVM/Source/RigVMEditor/Private/EdGraph/RigVMEdGraphPanelPinFactory.cpp b/Engine/Plugins/Runtime/RigVM/Source/RigVMEditor/Private/EdGraph/RigVMEdGraphPanelPinFactory.cpp
--- a/Engine/Plugins/Runtime/RigVM/Source/RigVMEditor/Private/EdGraph/RigVMEdGraphPanelPinFactory.cpp
+++ b/Engine/Plugins/Runtime/RigVM/Source/RigVMEditor/Private/EdGraph/RigVMEdGraphPanelPinFactory.cpp
@@ -136,7 +136,7 @@ TSharedPtr<SGraphPin> FRigVMEdGraphPanelPinFactory::CreatePin_Internal(UEdGraphP
 			if (InPin->PinType.PinCategory == UEdGraphSchema_K2::PC_Struct)
 			{
 				const UStruct* Struct = Cast<UStruct>(InPin->PinType.PinSubCategoryObject);
-				if (Struct && Struct->IsChildOf(FRigVMExecuteContext::StaticStruct()))
+				if (Struct && Struct->IsChildOf(FRigVMEmptyExecuteContext::StaticStruct()))
 				{
 					return SNew(SGraphPinExec, InPin);
 				}
diff --git a/Engine/Plugins/Runtime/RigVM/Source/RigVMEditor/Private/Editor/RigVMDetailsViewWrapperObject.cpp b/Engine/Plugins/Runtime/RigVM/Source/RigVMEditor/Private/Editor/RigVMDetailsViewWrapperObject.cpp
--- a/Engine/Plugins/Runtime/RigVM/Source/RigVMEditor/Private/Editor/RigVMDetailsViewWrapperObject.cpp
+++ b/Engine/Plugins/Runtime/RigVM/Source/RigVMEditor/Private/Editor/RigVMDetailsViewWrapperObject.cpp
@@ -133,7 +133,7 @@ UClass* URigVMDetailsViewWrapperObject::GetClassForStruct(UScriptStruct* InStruc
 			bool bIsEditable = true;
 			if (FStructProperty* StructProperty = CastField<FStructProperty>(NewProperty))
 			{
-				if (StructProperty->Struct->IsChildOf(FRigVMExecuteContext::StaticStruct()))
+				if (StructProperty->Struct->IsChildOf(FRigVMEmptyExecuteContext::StaticStruct()))
 				{
 					bIsEditable = false;
 				}
diff --git a/Engine/Plugins/Runtime/RigVM/Source/RigVMEditor/Private/Editor/RigVMGraphDetailCustomization.cpp b/Engine/Plugins/Runtime/RigVM/Source/RigVMEditor/Private/Editor/RigVMGraphDetailCustomization.cpp
--- a/Engine/Plugins/Runtime/RigVM/Source/RigVMEditor/Private/Editor/RigVMGraphDetailCustomization.cpp
+++ b/Engine/Plugins/Runtime/RigVM/Source/RigVMEditor/Private/Editor/RigVMGraphDetailCustomization.cpp
@@ -171,7 +171,7 @@ public:
 		RigVMTypeUtils::CPPTypeFromPinType(InItem.Get()->GetPinType(false), CPPType, &CPPTypeObject);
 		if (UScriptStruct* ScriptStruct = Cast<UScriptStruct>(CPPTypeObject))
 		{
-			if (ScriptStruct->IsChildOf(FRigVMExecuteContext::StaticStruct()))
+			if (ScriptStruct->IsChildOf(FRigVMEmptyExecuteContext::StaticStruct()))
 			{
 				if (GraphPtr.IsValid())
 				{
diff --git a/Engine/Plugins/Runtime/RigVM/Source/RigVMEditor/Private/Widgets/SRigVMGraphChangePinType.cpp b/Engine/Plugins/Runtime/RigVM/Source/RigVMEditor/Private/Widgets/SRigVMGraphChangePinType.cpp
--- a/Engine/Plugins/Runtime/RigVM/Source/RigVMEditor/Private/Widgets/SRigVMGraphChangePinType.cpp
+++ b/Engine/Plugins/Runtime/RigVM/Source/RigVMEditor/Private/Widgets/SRigVMGraphChangePinType.cpp
@@ -240,7 +240,7 @@ void SRigVMGraphChangePinType::FillPinTypeMenu(FMenuBuilder& MenuBuilder)
 
 			if(UScriptStruct* ScriptStruct = Cast<UScriptStruct>(Type.CPPTypeObject))
 			{
-				if(ScriptStruct->IsChildOf(FRigVMExecuteContext::StaticStruct()))
+				if(ScriptStruct->IsChildOf(FRigVMEmptyExecuteContext::StaticStruct()))
 				{
 					continue;
 				}
