RigVM: Added RigVM Details Inspector
#rnx

diff --git a/Engine/Plugins/Animation/ControlRig/Source/ControlRigEditor/Private/Editor/ControlRigEditorMode.cpp b/Engine/Plugins/Animation/ControlRig/Source/ControlRigEditor/Private/Editor/ControlRigEditorMode.cpp
--- a/Engine/Plugins/Animation/ControlRig/Source/ControlRigEditor/Private/Editor/ControlRigEditorMode.cpp
+++ b/Engine/Plugins/Animation/ControlRig/Source/ControlRigEditor/Private/Editor/ControlRigEditorMode.cpp
@@ -89,7 +89,7 @@ FControlRigEditorMode::FControlRigEditorMode(const TSharedRef<FControlRigEditor>
 #if WITH_RIGVMLEGACYEDITOR
 							->AddTab(FBlueprintEditorTabs::MyBlueprintID, ETabState::OpenedTab)
 #endif
-							->AddTab(FRigVMEditorGraphExplorerTabSummoner::TabID, ETabState::ClosedTab)
+							->AddTab(FRigVMEditorGraphExplorerTabSummoner::TabID(), ETabState::ClosedTab)
 						)
 					)
 					->Split
diff --git a/Engine/Plugins/Runtime/RigVM/Source/RigVMEditor/Private/Editor/RigVMDetailsInspectorTabSummoner.cpp b/Engine/Plugins/Runtime/RigVM/Source/RigVMEditor/Private/Editor/RigVMDetailsInspectorTabSummoner.cpp
new file mode 100644
--- /dev/null
+++ b/Engine/Plugins/Runtime/RigVM/Source/RigVMEditor/Private/Editor/RigVMDetailsInspectorTabSummoner.cpp
@@ -0,0 +1,31 @@
+// Copyright Epic Games, Inc. All Rights Reserved.
+
+#include "Editor/RigVMDetailsInspectorTabSummoner.h"
+#include "Editor/SRigVMDetailsInspector.h"
+#include "Editor/RigVMNewEditor.h"
+
+#define LOCTEXT_NAMESPACE "RigVMDetailsInspectorTabSummoner"
+
+FRigVMDetailsInspectorTabSummoner::FRigVMDetailsInspectorTabSummoner(const TSharedRef<FRigVMNewEditor>& InRigVMEditor)
+	: FWorkflowTabFactory(TabID(), InRigVMEditor.Get().GetHostingApp())
+	, RigVMEditor(InRigVMEditor)
+{
+	TabLabel = LOCTEXT("RigVMDetailsInspectorTabLabel", "Details");
+	TabIcon = FSlateIcon(TEXT("RigVMEditorStyle"), "DetailsInspector.TabIcon");
+
+	bIsSingleton = true;
+
+	ViewMenuDescription = LOCTEXT("RigVMDetailsInspector_ViewMenu_Desc", "Details");
+	ViewMenuTooltip = LOCTEXT("RigVMDetailsInspector_ViewMenu_ToolTip", "Show the RigVM Details tab");
+}
+
+TSharedRef<SWidget> FRigVMDetailsInspectorTabSummoner::CreateTabBody(const FWorkflowTabSpawnInfo& Info) const
+{
+	TSharedRef<SRigVMDetailsInspector> Inspector = SNew(SRigVMDetailsInspector)
+	.Editor(RigVMEditor);
+	
+	RigVMEditor.Pin()->SetInspector(Inspector);
+	return Inspector;
+}
+
+#undef LOCTEXT_NAMESPACE 
diff --git a/Engine/Plugins/Runtime/RigVM/Source/RigVMEditor/Private/Editor/RigVMEditor.cpp b/Engine/Plugins/Runtime/RigVM/Source/RigVMEditor/Private/Editor/RigVMEditor.cpp
--- a/Engine/Plugins/Runtime/RigVM/Source/RigVMEditor/Private/Editor/RigVMEditor.cpp
+++ b/Engine/Plugins/Runtime/RigVM/Source/RigVMEditor/Private/Editor/RigVMEditor.cpp
@@ -38,11 +38,10 @@
 #include "Widgets/Docking/SDockTab.h"
 #include "Widgets/Layout/SScrollBox.h"
 #include "IStructureDetailsView.h"
+#include "Editor/SRigVMDetailsInspector.h"
 
 #if WITH_RIGVMLEGACYEDITOR
 #include "SKismetInspector.h"
-#else
-#include "Editor/SRigVMDetailsInspector.h"
 #endif
 
 #define LOCTEXT_NAMESPACE "RigVMEditor"
@@ -79,6 +78,8 @@ void FRigVMEditorBase::UnbindEditor()
 	
 	FEditorDelegates::EndPIE.RemoveAll(this);
 	FEditorDelegates::CancelPIE.RemoveAll(this);
+	GEditor->GetEditorSubsystem<UAssetEditorSubsystem>()->OnAssetEditorRequestedOpen().RemoveAll(this);
+	GEditor->GetEditorSubsystem<UAssetEditorSubsystem>()->OnAssetEditorRequestClose().RemoveAll(this);
 
 	if (RigVMBlueprint)
 	{
@@ -99,12 +100,32 @@ void FRigVMEditorBase::UnbindEditor()
 		RigVMBlueprint->OnRequestJumpToHyperlink().Unbind();
 		RigVMBlueprint->OnReportCompilerMessage().RemoveAll(this);
 
+		RigVMBlueprint->OnModified().RemoveAll(this);
+		RigVMBlueprint->OnVMCompiled().RemoveAll(this);
+		RigVMBlueprint->OnRequestInspectObject().RemoveAll(this);
+		RigVMBlueprint->OnRequestInspectMemoryStorage().RemoveAll(this);
+
+		for(UEdGraph* Graph : RigVMBlueprint->UbergraphPages)
+		{
+			if (URigVMEdGraph* RigGraph = Cast<URigVMEdGraph>(Graph))
+			{
+				RigGraph->OnGraphNodeClicked.RemoveAll(this);
+			}
+		}
+
 #if WITH_EDITOR
 		RigVMBlueprint->SetDebugMode(false);
 		RigVMBlueprint->ClearBreakpoints();
 		SetHaltedNode(nullptr);
 		RigVMBlueprint->OnGetFocusedGraph().Unbind();
 #endif
+
+		if (URigVMHost* RigVMHost = GetRigVMHost())
+		{
+			RigVMHost->OnInitialized_AnyThread().RemoveAll(this);
+			RigVMHost->OnExecuted_AnyThread().RemoveAll(this);
+			RigVMHost->GetDebugInfo().ExecutionHalted().RemoveAll(this);
+		}
 	}
 
 	if (bRequestedReopen)
@@ -372,29 +393,32 @@ void FRigVMEditorBase::InitRigVMEditor(const EToolkitMode::Type Mode, const TSha
 			Compile();
 		}
 	}
-	
-	FFunctionGraphTask::CreateAndDispatchWhenReady([this, InRigVMBlueprint]()
+
+	TWeakPtr<IRigVMEditor> ThisWeakPtr = SharedRef().ToWeakPtr();
+	FFunctionGraphTask::CreateAndDispatchWhenReady([ThisWeakPtr, InRigVMBlueprint]()
 	{
-		// no need to do anything if the the CR is not opened anymore
-		// (i.e. destructor has been called before that task actually got a chance to start)
-		if (!URigVMBlueprint::sCurrentlyOpenedRigVMBlueprints.Contains(InRigVMBlueprint))
+		if (TSharedPtr<IRigVMEditor> Editor = ThisWeakPtr.Pin())
 		{
-			return;		
-		}
+			// no need to do anything if the the CR is not opened anymore
+		   // (i.e. destructor has been called before that task actually got a chance to start)
+		   if (!URigVMBlueprint::sCurrentlyOpenedRigVMBlueprints.Contains(InRigVMBlueprint))
+		   {
+			   return;		
+		   }
 		
-		TSharedPtr<FTabManager> TabManager = GetTabManager();
-		if (!TabManager.IsValid())
-		{
-			return;
-		}
+		   TSharedPtr<FTabManager> TabManager = Editor->GetTabManager();
+		   if (!TabManager.IsValid())
+		   {
+			   return;
+		   }
 		
-		// Always show the myblueprint tab
-		static const FTabId MyBlueprintTabId(GetGraphExplorerWidgetID());
-		if (!TabManager->FindExistingLiveTab(MyBlueprintTabId).IsValid())
-		{
-			TabManager->TryInvokeTab(MyBlueprintTabId);
+		   // Always show the myblueprint tab
+		   static const FTabId MyBlueprintTabId(Editor->GetGraphExplorerWidgetID());
+		   if (!TabManager->FindExistingLiveTab(MyBlueprintTabId).IsValid())
+		   {
+			   TabManager->TryInvokeTab(MyBlueprintTabId);
+		   }
 		}
-		
 	}, TStatId(), nullptr, ENamedThreads::GameThread);
 
 	bRigVMEditorInitialized = true;
@@ -415,27 +439,49 @@ void FRigVMEditorBase::InitRigVMEditor(const EToolkitMode::Type Mode, const TSha
 		TBaseStructure<FEulerTransform>::Get(),
 	};
 
+	auto Task = [&]<typename T>(T Inspector, UScriptStruct* StructToCustomize)
+	{
+		Inspector->GetPropertyView()->RegisterInstancedCustomPropertyTypeLayout(StructToCustomize->GetFName(),
+				   FOnGetPropertyTypeCustomizationInstance::CreateLambda([=]()
+				   {
+					   return FRigVMGraphMathTypeDetailCustomization::MakeInstance();
+				   }));
+	};
+
+	TSharedRef<FAssetEditorToolkit> SharedApp = GetHostingApp().ToSharedRef();
 	for(UScriptStruct* StructToCustomize : StructsToCustomize)
 	{
 #if WITH_RIGVMLEGACYEDITOR
-		if (TSharedPtr<SKismetInspector> Inspector = GetInspector())
-#else
-		if (TSharedPtr<SRigVMDetailsInspector> Inspector = GetInspector())
+		if (SharedApp->IsBlueprintEditor())
+		{
+			if (TSharedPtr<SKismetInspector> Inspector = GetKismetInspector())
+			{
+				Task(Inspector, StructToCustomize);
+			}
+		}
 #endif
+		if (TSharedPtr<SRigVMDetailsInspector> Inspector = GetRigVMInspector())
 		{
-			Inspector->GetPropertyView()->RegisterInstancedCustomPropertyTypeLayout(StructToCustomize->GetFName(),
-			   FOnGetPropertyTypeCustomizationInstance::CreateLambda([=]()
-			   {
-				   return FRigVMGraphMathTypeDetailCustomization::MakeInstance();
-			   }));
+			Task(Inspector, StructToCustomize);
 		}
 	}
 
 #if WITH_RIGVMLEGACYEDITOR
-	if (TSharedPtr<SKismetInspector> Inspector = GetInspector())
-#else
-	if (TSharedPtr<SRigVMDetailsInspector> Inspector = GetInspector())
+	if (SharedApp->IsBlueprintEditor())
+	{
+		if (TSharedPtr<SKismetInspector> Inspector = GetKismetInspector())
+		{
+			Inspector->GetPropertyView()->RegisterInstancedCustomPropertyTypeLayout(UEnum::StaticClass()->GetFName(),
+			FOnGetPropertyTypeCustomizationInstance::CreateLambda([=]()
+			{
+				return FRigVMGraphEnumDetailCustomization::MakeInstance();
+			}));
+		}
+	}
+	
+	
 #endif
+	if (TSharedPtr<SRigVMDetailsInspector> Inspector = GetRigVMInspector())
 	{
 		Inspector->GetPropertyView()->RegisterInstancedCustomPropertyTypeLayout(UEnum::StaticClass()->GetFName(),
 		   FOnGetPropertyTypeCustomizationInstance::CreateLambda([=]()
@@ -743,6 +789,29 @@ bool FRigVMEditorBase::CanPasteNodes() const
 	return GetFocusedController()->CanImportNodesFromText(TextToImport);
 }
 
+void FRigVMEditorBase::CutSelectedNodes()
+{
+	CopySelectedNodes();
+	DeleteSelectedNodes();
+}
+
+bool FRigVMEditorBase::CanCutNodes() const
+{
+	return CanCopyNodes() && CanDeleteNodes();
+}
+
+void FRigVMEditorBase::DuplicateNodes()
+{
+	// Copy and paste current selection
+	CopySelectedNodes();
+	PasteNodes();
+}
+
+bool FRigVMEditorBase::CanDuplicateNodes() const
+{
+	return CanCopyNodes() && IsEditable(GetFocusedGraph());
+}
+
 FReply FRigVMEditorBase::OnSpawnGraphNodeByShortcut(FInputChord InChord, const FVector2D& InPosition,
                                                      UEdGraph* InGraph)
 {
@@ -1440,10 +1509,17 @@ void FRigVMEditorBase::OnActiveTabChanged( TSharedPtr<SDockTab> PreviouslyActive
 	{
 		TArray<UObject*> ObjArray;
 #if WITH_RIGVMLEGACYEDITOR
-		if (TSharedPtr<SKismetInspector> Inspector = GetInspector())
-#else
-		if (TSharedPtr<SRigVMDetailsInspector> Inspector = GetInspector())
+		TSharedRef<FAssetEditorToolkit> SharedApp = GetHostingApp().ToSharedRef();
+		if (SharedApp->IsBlueprintEditor())
+		{
+			if (TSharedPtr<SKismetInspector> Inspector = GetKismetInspector())
+			{
+				Inspector->ShowDetailsForObjects(ObjArray);
+			}
+		}
+		
 #endif
+		if (TSharedPtr<SRigVMDetailsInspector> Inspector = GetRigVMInspector())
 		{
 			Inspector->ShowDetailsForObjects(ObjArray);
 		}
@@ -2651,11 +2727,8 @@ TArray<TWeakObjectPtr<UObject>> FRigVMEditorBase::GetSelectedObjects() const
 	// if the inspector shows wrapped objects - look in that array instead.
 	// with recent weak object pointer changes on the property detail view
 	// we cannot rely on the GetSelectedObjects being valid after blueprint compilation.
-#if WITH_RIGVMLEGACYEDITOR
-	if (TSharedPtr<SKismetInspector> Inspector = GetInspector())
-#else
-	if (TSharedPtr<SRigVMDetailsInspector> Inspector = GetInspector())
-#endif
+
+	auto Task = [this]<typename T>(T Inspector)
 	{
 		if(WrapperObjects.Num() == Inspector->GetSelectedObjects().Num())
 		{
@@ -2671,6 +2744,21 @@ TArray<TWeakObjectPtr<UObject>> FRigVMEditorBase::GetSelectedObjects() const
 			return WeakWrapperObjects;
 		}
 		return Inspector->GetSelectedObjects();
+	};
+	
+#if WITH_RIGVMLEGACYEDITOR
+	TSharedRef<FAssetEditorToolkit> SharedApp = GetHostingApp().ToSharedRef();
+	if (SharedApp->IsBlueprintEditor())
+	{
+		if (TSharedPtr<SKismetInspector> Inspector = GetKismetInspector())
+		{
+			return Task(Inspector);
+		}
+	}
+#endif
+	if (TSharedPtr<SRigVMDetailsInspector> Inspector = GetRigVMInspector())
+	{
+		return Task(Inspector);
 	}
 	return TArray<TWeakObjectPtr<UObject>>();
 }
@@ -2826,17 +2914,28 @@ void FRigVMEditorBase::SetDetailObjects(const TArray<UObject*>& InObjects, bool
 		}
 	}
 
+	auto Task = [&]<typename T, typename O>(T Inspector, O Options)
+	{
+		Inspector->ShowDetailsForObjects(FilteredObjects, Options);
+	};
+
 #if WITH_RIGVMLEGACYEDITOR
-	SKismetInspector::FShowDetailsOptions Options;
-	Options.bForceRefresh = true;
-	if (TSharedPtr<SKismetInspector> Inspector = GetInspector())
-#else
-	SRigVMDetailsInspector::FShowDetailsOptions Options;
-	Options.bForceRefresh = true;
-	if (TSharedPtr<SRigVMDetailsInspector> Inspector = GetInspector())
+	TSharedRef<FAssetEditorToolkit> SharedApp = GetHostingApp().ToSharedRef();
+	if (SharedApp->IsBlueprintEditor())
+	{
+		if (TSharedPtr<SKismetInspector> Inspector = GetKismetInspector())
+		{
+			SKismetInspector::FShowDetailsOptions Options;
+			Options.bForceRefresh = true;
+			Task(Inspector, Options);
+		}
+	}
 #endif
+	if (TSharedPtr<SRigVMDetailsInspector> Inspector = GetRigVMInspector())
 	{
-		Inspector->ShowDetailsForObjects(FilteredObjects, Options);
+		SRigVMDetailsInspector::FShowDetailsOptions Options;
+		Options.bForceRefresh = true;
+		Task(Inspector, Options);
 	}
 }
 
@@ -3035,18 +3134,33 @@ void FRigVMEditorBase::ClearDetailObject(bool bChangeUISelectionState)
 	}
 
 	ClearDetailsViewWrapperObjects();
+
+	auto Task = [&]<typename T>(T Inspector)
+	{
+		if(Inspector.IsValid())
+		{
+			Inspector->GetPropertyView()->SetObjects(TArray<UObject*>(), true); // clear property view synchronously
+			Inspector->ShowDetailsForObjects(TArray<UObject*>());
+			Inspector->ShowSingleStruct(TSharedPtr<FStructOnScope>());
+		}
+	};
+
 #if WITH_RIGVMLEGACYEDITOR
-	TSharedPtr<SKismetInspector> Inspector = GetInspector();
-#else
-	TSharedPtr<SRigVMDetailsInspector> Inspector = GetInspector();
+	TSharedRef<FAssetEditorToolkit> SharedApp = GetHostingApp().ToSharedRef();
+	if (SharedApp->IsBlueprintEditor())
+	{
+		if (TSharedPtr<SKismetInspector> Inspector = GetKismetInspector())
+		{
+			Task(Inspector);
+		}
+	}
 #endif
-	if(Inspector.IsValid())
+	if (TSharedPtr<SRigVMDetailsInspector> Inspector = GetRigVMInspector())
 	{
-		Inspector->GetPropertyView()->SetObjects(TArray<UObject*>(), true); // clear property view synchronously
-		Inspector->ShowDetailsForObjects(TArray<UObject*>());
-		Inspector->ShowSingleStruct(TSharedPtr<FStructOnScope>());
+		Task(Inspector);
 	}
 
+
 	if (bChangeUISelectionState)
 	{
 #if WITH_RIGVMLEGACYEDITOR
diff --git a/Engine/Plugins/Runtime/RigVM/Source/RigVMEditor/Private/Editor/RigVMEditorGraphExplorerTabSummoner.cpp b/Engine/Plugins/Runtime/RigVM/Source/RigVMEditor/Private/Editor/RigVMEditorGraphExplorerTabSummoner.cpp
--- a/Engine/Plugins/Runtime/RigVM/Source/RigVMEditor/Private/Editor/RigVMEditorGraphExplorerTabSummoner.cpp
+++ b/Engine/Plugins/Runtime/RigVM/Source/RigVMEditor/Private/Editor/RigVMEditorGraphExplorerTabSummoner.cpp
@@ -7,7 +7,7 @@
 #define LOCTEXT_NAMESPACE "RigVMEditorGraphExplorerTabSummoner"
 
 FRigVMEditorGraphExplorerTabSummoner::FRigVMEditorGraphExplorerTabSummoner(const TSharedRef<FRigVMEditorBase>& InRigVMEditor)
-	: FWorkflowTabFactory(TabID, InRigVMEditor.Get().GetHostingApp())
+	: FWorkflowTabFactory(TabID(), InRigVMEditor.Get().GetHostingApp())
 	, RigVMEditor(InRigVMEditor)
 {
 	TabLabel = LOCTEXT("RigVMEditorGraphExplorerTabLabel", "Graph Explorer");
diff --git a/Engine/Plugins/Runtime/RigVM/Source/RigVMEditor/Private/Editor/RigVMGraphDetailCustomization.cpp b/Engine/Plugins/Runtime/RigVM/Source/RigVMEditor/Private/Editor/RigVMGraphDetailCustomization.cpp
--- a/Engine/Plugins/Runtime/RigVM/Source/RigVMEditor/Private/Editor/RigVMGraphDetailCustomization.cpp
+++ b/Engine/Plugins/Runtime/RigVM/Source/RigVMEditor/Private/Editor/RigVMGraphDetailCustomization.cpp
@@ -728,33 +728,49 @@ void FRigVMFunctionArgumentDefaultNode::HandleModifiedEvent(ERigVMGraphNotifType
 	}
 }
 
-#if WITH_RIGVMLEGACYEDITOR
-TSharedPtr<IDetailCustomization> FRigVMGraphDetailCustomization::MakeInstance(TSharedPtr<IBlueprintEditor> InEditor, const UClass* InExpectedBlueprintClass)
-#else
-TSharedPtr<IDetailCustomization> FRigVMGraphDetailCustomization::MakeInstance(TSharedPtr<FRigVMEditorBase> InEditor, const UClass* InExpectedBlueprintClass)
-#endif
+TSharedRef<IDetailCustomization> FRigVMGraphDetailCustomization::MakeInstance(TSharedPtr<FRigVMEditorBase> InEditor, const UClass* InExpectedBlueprintClass)
 {
 	const TArray<UObject*>* Objects = (InEditor.IsValid() ? InEditor->GetObjectsCurrentlyBeingEdited() : nullptr);
 	if (Objects && Objects->Num() == 1)
 	{
 		if (URigVMBlueprint* RigVMBlueprint = Cast<URigVMBlueprint>((*Objects)[0]))
 		{
-			if(RigVMBlueprint->GetClass() == InExpectedBlueprintClass)
+			if(RigVMBlueprint->GetClass() == InExpectedBlueprintClass || RigVMBlueprint->GetClass()->IsChildOf(InExpectedBlueprintClass))
 			{
-#if WITH_RIGVMLEGACYEDITOR
-				// TODO: need to change how the Kismet Inspector calls the customization with a ptr to IBlueprintEditor
-				TSharedPtr<FRigVMLegacyEditor> SharedPtr = StaticCastSharedPtr<FRigVMLegacyEditor>(InEditor);
-				return MakeShareable(new FRigVMGraphDetailCustomization(SharedPtr, RigVMBlueprint));
-#else
 				return MakeShareable(new FRigVMGraphDetailCustomization(InEditor, RigVMBlueprint));
-#endif
 			}
 		}
 	}
 
-	return nullptr;
+	return MakeShareable(new FRigVMGraphDetailCustomization((TSharedPtr<FRigVMEditorBase>)nullptr, nullptr));
 }
 
+#if WITH_RIGVMLEGACYEDITOR
+TSharedPtr<IDetailCustomization> FRigVMGraphDetailCustomization::MakeLegacyInstance(TSharedPtr<IBlueprintEditor> InBlueprintEditor, const UClass* InExpectedBlueprintClass)
+{
+	const TArray<UObject*>* Objects = (InBlueprintEditor.IsValid() ? InBlueprintEditor->GetObjectsCurrentlyBeingEdited() : nullptr);
+	if (Objects && Objects->Num() == 1)
+	{
+		if (URigVMBlueprint* RigVMBlueprint = Cast<URigVMBlueprint>((*Objects)[0]))
+		{
+			if(RigVMBlueprint->GetClass() == InExpectedBlueprintClass || RigVMBlueprint->GetClass()->IsChildOf(InExpectedBlueprintClass))
+			{
+				return MakeShareable(new FRigVMGraphDetailCustomization(InBlueprintEditor, RigVMBlueprint));
+			}
+		}
+	}
+
+	return MakeShareable(new FRigVMGraphDetailCustomization((TSharedPtr<FRigVMEditorBase>)nullptr, nullptr));
+}
+
+FRigVMGraphDetailCustomization::FRigVMGraphDetailCustomization(TSharedPtr<IBlueprintEditor> RigVMigEditor, URigVMBlueprint* RigVMBlueprint)
+	: RigVMBlueprintPtr(RigVMBlueprint)
+	, RigVMGraphDetailCustomizationImpl(MakeShared<FRigVMGraphDetailCustomizationImpl>())
+{
+	RigVMEditorPtr = StaticCastSharedPtr<FRigVMLegacyEditor>(RigVMigEditor);
+}
+#endif
+
 FRigVMGraphDetailCustomization::FRigVMGraphDetailCustomization(TSharedPtr<FRigVMEditorBase> RigVMigEditor, URigVMBlueprint* RigVMBlueprint)
 	: RigVMEditorPtr(RigVMigEditor)
 	, RigVMBlueprintPtr(RigVMBlueprint)
diff --git a/Engine/Plugins/Runtime/RigVM/Source/RigVMEditor/Private/Editor/RigVMLegacyEditorMode.cpp b/Engine/Plugins/Runtime/RigVM/Source/RigVMEditor/Private/Editor/RigVMLegacyEditorMode.cpp
--- a/Engine/Plugins/Runtime/RigVM/Source/RigVMEditor/Private/Editor/RigVMLegacyEditorMode.cpp
+++ b/Engine/Plugins/Runtime/RigVM/Source/RigVMEditor/Private/Editor/RigVMLegacyEditorMode.cpp
@@ -35,7 +35,7 @@ FRigVMLegacyEditorMode::FRigVMLegacyEditorMode(const TSharedRef<FRigVMLegacyEdit
 					->SetSizeCoefficient(0.5f)
 					->AddTab(FRigVMExecutionStackTabSummoner::TabID, ETabState::OpenedTab)
 					->AddTab(FBlueprintEditorTabs::MyBlueprintID, ETabState::OpenedTab)
-					->AddTab(FRigVMEditorGraphExplorerTabSummoner::TabID, ETabState::ClosedTab)
+					->AddTab(FRigVMEditorGraphExplorerTabSummoner::TabID(), ETabState::ClosedTab)
 				)
 				->Split
 				(
diff --git a/Engine/Plugins/Runtime/RigVM/Source/RigVMEditor/Private/Editor/RigVMNewEditor.cpp b/Engine/Plugins/Runtime/RigVM/Source/RigVMEditor/Private/Editor/RigVMNewEditor.cpp
--- a/Engine/Plugins/Runtime/RigVM/Source/RigVMEditor/Private/Editor/RigVMNewEditor.cpp
+++ b/Engine/Plugins/Runtime/RigVM/Source/RigVMEditor/Private/Editor/RigVMNewEditor.cpp
@@ -26,12 +26,13 @@
 #include "Stats/StatsHierarchical.h"
 #include "Logging/MessageLog.h"
 #include "GraphEditor.h"
+#include "GraphEditorActions.h"
+#include "RigVMBlueprint.h"
 #include "Editor/RigVMGraphEditorSummoner.h"
 #include "Widgets/Input/SEditableTextBox.h"
+#include "Editor/SRigVMDetailsInspector.h"
 #if WITH_RIGVMLEGACYEDITOR
 #include "SKismetInspector.h"
-#else
-#include "Editor/SRigVMDetailsInspector.h"
 #endif
 
 
@@ -951,358 +952,89 @@ TSharedRef<SGraphEditor> FRigVMNewEditor::CreateGraphEditorWidget(TSharedRef<FTa
 	{
 		GraphEditorCommands = MakeShareable( new FUICommandList );
 		{
-			// GraphEditorCommands->MapAction(FGraphEditorCommands::Get().PromoteToVariable,
-			// 	FExecuteAction::CreateSP( this, &FBlueprintEditor::OnPromoteToVariable, true ),
-			// 	FCanExecuteAction::CreateSP( this, &FBlueprintEditor::CanPromoteToVariable, true )
-			// 	);
-			//
-			// GraphEditorCommands->MapAction(FGraphEditorCommands::Get().PromoteToLocalVariable,
-			// 	FExecuteAction::CreateSP( this, &FBlueprintEditor::OnPromoteToVariable, false ),
-			// 	FCanExecuteAction::CreateSP( this, &FBlueprintEditor::CanPromoteToVariable, false )
-			// 	);
-			//
-			// GraphEditorCommands->MapAction(FGraphEditorCommands::Get().SplitStructPin,
-			// 	FExecuteAction::CreateSP( this, &FBlueprintEditor::OnSplitStructPin ),
-			// 	FCanExecuteAction::CreateSP( this, &FBlueprintEditor::CanSplitStructPin )
-			// 	);
-			//
-			// GraphEditorCommands->MapAction(FGraphEditorCommands::Get().RecombineStructPin,
-			// 	FExecuteAction::CreateSP( this, &FBlueprintEditor::OnRecombineStructPin ),
-			// 	FCanExecuteAction::CreateSP( this, &FBlueprintEditor::CanRecombineStructPin )
-			// 	);
-			//
-			// GraphEditorCommands->MapAction( FGraphEditorCommands::Get().AddExecutionPin,
-			// 	FExecuteAction::CreateSP( this, &FBlueprintEditor::OnAddExecutionPin ),
-			// 	FCanExecuteAction::CreateSP( this, &FBlueprintEditor::CanAddExecutionPin )
-			// 	);
-			//
-			// GraphEditorCommands->MapAction( FGraphEditorCommands::Get().InsertExecutionPinBefore,
-			// 	FExecuteAction::CreateSP( this, &FBlueprintEditor::OnInsertExecutionPinBefore ),
-			// 	FCanExecuteAction::CreateSP( this, &FBlueprintEditor::CanInsertExecutionPin )
-			// 	);
-			//
-			// GraphEditorCommands->MapAction(FGraphEditorCommands::Get().InsertExecutionPinAfter,
-			// 	FExecuteAction::CreateSP(this, &FBlueprintEditor::OnInsertExecutionPinAfter),
-			// 	FCanExecuteAction::CreateSP(this, &FBlueprintEditor::CanInsertExecutionPin)
-			// );
-			//
-			// GraphEditorCommands->MapAction( FGraphEditorCommands::Get().RemoveExecutionPin,
-			// 	FExecuteAction::CreateSP( this, &FBlueprintEditor::OnRemoveExecutionPin ),
-			// 	FCanExecuteAction::CreateSP( this, &FBlueprintEditor::CanRemoveExecutionPin )
-			// 	);
-			//
-			// GraphEditorCommands->MapAction(FGraphEditorCommands::Get().RemoveThisStructVarPin,
-			// 	FExecuteAction::CreateSP(this, &FBlueprintEditor::OnRemoveThisStructVarPin),
-			// 	FCanExecuteAction::CreateSP(this, &FBlueprintEditor::CanRemoveThisStructVarPin)
-			// 	);
-			//
-			// GraphEditorCommands->MapAction(FGraphEditorCommands::Get().RemoveOtherStructVarPins,
-			// 	FExecuteAction::CreateSP(this, &FBlueprintEditor::OnRemoveOtherStructVarPins),
-			// 	FCanExecuteAction::CreateSP(this, &FBlueprintEditor::CanRemoveOtherStructVarPins)
-			// 	);
-			//
-			// GraphEditorCommands->MapAction(FGraphEditorCommands::Get().RestoreAllStructVarPins,
-			// 	FExecuteAction::CreateSP(this, &FBlueprintEditor::OnRestoreAllStructVarPins),
-			// 	FCanExecuteAction::CreateSP(this, &FBlueprintEditor::CanRestoreAllStructVarPins)
-			// 	);
-			//
-			// GraphEditorCommands->MapAction(FGraphEditorCommands::Get().ResetPinToDefaultValue,
-			// 	FExecuteAction::CreateSP(this, &FBlueprintEditor::OnResetPinToDefaultValue),
-			// 	FCanExecuteAction::CreateSP(this, &FBlueprintEditor::CanResetPinToDefaultValue)
-			// );
-			//
-			// GraphEditorCommands->MapAction( FGraphEditorCommands::Get().AddOptionPin,
-			// 	FExecuteAction::CreateSP( this, &FBlueprintEditor::OnAddOptionPin ),
-			// 	FCanExecuteAction::CreateSP( this, &FBlueprintEditor::CanAddOptionPin )
-			// 	);
-			//
-			// GraphEditorCommands->MapAction( FGraphEditorCommands::Get().RemoveOptionPin,
-			// 	FExecuteAction::CreateSP( this, &FBlueprintEditor::OnRemoveOptionPin ),
-			// 	FCanExecuteAction::CreateSP( this, &FBlueprintEditor::CanRemoveOptionPin )
-			// 	);
-			//
-			// GraphEditorCommands->MapAction( FGraphEditorCommands::Get().ChangePinType,
-			// 	FExecuteAction::CreateSP( this, &FBlueprintEditor::OnChangePinType ),
-			// 	FCanExecuteAction::CreateSP( this, &FBlueprintEditor::CanChangePinType )
-			// 	);
-			//
-			// GraphEditorCommands->MapAction( FGraphEditorCommands::Get().AddParentNode,
-			// 	FExecuteAction::CreateSP( this, &FBlueprintEditor::OnAddParentNode ),
-			// 	FCanExecuteAction::CreateSP( this, &FBlueprintEditor::CanAddParentNode )
-			// 	);
-			//
-			// GraphEditorCommands->MapAction( FGraphEditorCommands::Get().CreateMatchingFunction,
-			// 	FExecuteAction::CreateSP( this, &FBlueprintEditor::OnCreateMatchingFunction ),
-			// 	FCanExecuteAction::CreateSP( this, &FBlueprintEditor::CanCreateMatchingFunction )
-			// 	);
-			//
-			// // Debug actions
-			// GraphEditorCommands->MapAction( FGraphEditorCommands::Get().AddBreakpoint,
-			// 	FExecuteAction::CreateSP( this, &FBlueprintEditor::OnAddBreakpoint ),
-			// 	FCanExecuteAction::CreateSP( this, &FBlueprintEditor::CanAddBreakpoint ),
-			// 	FIsActionChecked(),
-			// 	FIsActionButtonVisible::CreateSP( this, &FBlueprintEditor::CanAddBreakpoint )
-			// 	);
-			//
-			// GraphEditorCommands->MapAction( FGraphEditorCommands::Get().RemoveBreakpoint,
-			// 	FExecuteAction::CreateSP( this, &FBlueprintEditor::OnRemoveBreakpoint ),
-			// 	FCanExecuteAction::CreateSP( this, &FBlueprintEditor::CanRemoveBreakpoint ),
-			// 	FIsActionChecked(),
-			// 	FIsActionButtonVisible::CreateSP( this, &FBlueprintEditor::CanRemoveBreakpoint )
-			// 	);
-			//
-			// GraphEditorCommands->MapAction( FGraphEditorCommands::Get().EnableBreakpoint,
-			// 	FExecuteAction::CreateSP( this, &FBlueprintEditor::OnEnableBreakpoint ),
-			// 	FCanExecuteAction::CreateSP( this, &FBlueprintEditor::CanEnableBreakpoint ),
-			// 	FIsActionChecked(),
-			// 	FIsActionButtonVisible::CreateSP( this, &FBlueprintEditor::CanEnableBreakpoint )
-			// 	);
-			//
-			// GraphEditorCommands->MapAction( FGraphEditorCommands::Get().DisableBreakpoint,
-			// 	FExecuteAction::CreateSP( this, &FBlueprintEditor::OnDisableBreakpoint ),
-			// 	FCanExecuteAction::CreateSP( this, &FBlueprintEditor::CanDisableBreakpoint ),
-			// 	FIsActionChecked(),
-			// 	FIsActionButtonVisible::CreateSP( this, &FBlueprintEditor::CanDisableBreakpoint )
-			// 	);
-			//
-			// GraphEditorCommands->MapAction( FGraphEditorCommands::Get().ToggleBreakpoint,
-			// 	FExecuteAction::CreateSP( this, &FBlueprintEditor::OnToggleBreakpoint ),
-			// 	FCanExecuteAction::CreateSP( this, &FBlueprintEditor::CanToggleBreakpoint ),
-			// 	FIsActionChecked(),
-			// 	FIsActionButtonVisible::CreateSP( this, &FBlueprintEditor::CanToggleBreakpoint )
-			// 	);
-			//
-			// // Encapsulation commands
-			// GraphEditorCommands->MapAction( FGraphEditorCommands::Get().CollapseNodes,
-			// 	FExecuteAction::CreateSP( this, &FBlueprintEditor::OnCollapseNodes ),
-			// 	FCanExecuteAction::CreateSP( this, &FBlueprintEditor::CanCollapseNodes )
-			// 	);
-			//
-			// GraphEditorCommands->MapAction( FGraphEditorCommands::Get().CollapseSelectionToFunction,
-			// 	FExecuteAction::CreateSP( this, &FBlueprintEditor::OnCollapseSelectionToFunction ),
-			// 	FCanExecuteAction::CreateSP( this, &FBlueprintEditor::CanCollapseSelectionToFunction ),
-			// 	FIsActionChecked(),
-			// 	FIsActionButtonVisible::CreateSP( this, &FBlueprintEditor::NewDocument_IsVisibleForType, CGT_NewFunctionGraph )
-			// 	);
-			//
-			// GraphEditorCommands->MapAction( FGraphEditorCommands::Get().CollapseSelectionToMacro,
-			// 	FExecuteAction::CreateSP( this, &FBlueprintEditor::OnCollapseSelectionToMacro ),
-			// 	FCanExecuteAction::CreateSP( this, &FBlueprintEditor::CanCollapseSelectionToMacro ),
-			// 	FIsActionChecked(),
-			// 	FIsActionButtonVisible::CreateSP( this, &FBlueprintEditor::NewDocument_IsVisibleForType, CGT_NewMacroGraph )
-			// 	);
-			//
-			// GraphEditorCommands->MapAction(FGraphEditorCommands::Get().ConvertFunctionToEvent,
-			// 	FExecuteAction::CreateSP(this, &FBlueprintEditor::OnConvertFunctionToEvent),
-			// 	FCanExecuteAction::CreateSP(this, &FBlueprintEditor::CanConvertFunctionToEvent),
-			// 	FIsActionChecked()
-			// );
-			//
-			// GraphEditorCommands->MapAction(FGraphEditorCommands::Get().ConvertEventToFunction,
-			// 	FExecuteAction::CreateSP(this, &FBlueprintEditor::OnConvertEventToFunction),
-			// 	FCanExecuteAction::CreateSP(this, &FBlueprintEditor::CanConvertEventToFunction),
-			// 	FIsActionChecked()
-			// );
-			//
-			// GraphEditorCommands->MapAction( FGraphEditorCommands::Get().PromoteSelectionToFunction,
-			// 	FExecuteAction::CreateSP( this, &FBlueprintEditor::OnPromoteSelectionToFunction ),
-			// 	FCanExecuteAction::CreateSP( this, &FBlueprintEditor::CanPromoteSelectionToFunction ),
-			// 	FIsActionChecked(),
-			// 	FIsActionButtonVisible::CreateSP( this, &FBlueprintEditor::NewDocument_IsVisibleForType, CGT_NewFunctionGraph )
-			// 	);
-			//
-			// GraphEditorCommands->MapAction( FGraphEditorCommands::Get().PromoteSelectionToMacro,
-			// 	FExecuteAction::CreateSP( this, &FBlueprintEditor::OnPromoteSelectionToMacro ),
-			// 	FCanExecuteAction::CreateSP( this, &FBlueprintEditor::CanPromoteSelectionToMacro ),
-			// 	FIsActionChecked(),
-			// 	FIsActionButtonVisible::CreateSP( this, &FBlueprintEditor::NewDocument_IsVisibleForType, CGT_NewMacroGraph )
-			// 	);
-			//
-			// // Alignment Commands
-			// GraphEditorCommands->MapAction( FGraphEditorCommands::Get().AlignNodesTop,
-			// 	FExecuteAction::CreateSP( this, &FBlueprintEditor::OnAlignTop )
-			// 	);
-			//
-			// GraphEditorCommands->MapAction( FGraphEditorCommands::Get().AlignNodesMiddle,
-			// 	FExecuteAction::CreateSP( this, &FBlueprintEditor::OnAlignMiddle )
-			// 	);
-			//
-			// GraphEditorCommands->MapAction( FGraphEditorCommands::Get().AlignNodesBottom,
-			// 	FExecuteAction::CreateSP( this, &FBlueprintEditor::OnAlignBottom )
-			// 	);
-			//
-			// GraphEditorCommands->MapAction( FGraphEditorCommands::Get().AlignNodesLeft,
-			// 	FExecuteAction::CreateSP( this, &FBlueprintEditor::OnAlignLeft )
-			// 	);
-			//
-			// GraphEditorCommands->MapAction( FGraphEditorCommands::Get().AlignNodesCenter,
-			// 	FExecuteAction::CreateSP( this, &FBlueprintEditor::OnAlignCenter )
-			// 	);
-			//
-			// GraphEditorCommands->MapAction( FGraphEditorCommands::Get().AlignNodesRight,
-			// 	FExecuteAction::CreateSP( this, &FBlueprintEditor::OnAlignRight )
-			// 	);
-			//
-			// GraphEditorCommands->MapAction( FGraphEditorCommands::Get().StraightenConnections,
-			// 	FExecuteAction::CreateSP( this, &FBlueprintEditor::OnStraightenConnections )
-			// 	);
-			//
-			// // Distribution Commands
-			// GraphEditorCommands->MapAction( FGraphEditorCommands::Get().DistributeNodesHorizontally,
-			// 	FExecuteAction::CreateSP( this, &FBlueprintEditor::OnDistributeNodesH )
-			// 	);
-			//
-			// GraphEditorCommands->MapAction( FGraphEditorCommands::Get().DistributeNodesVertically,
-			// 	FExecuteAction::CreateSP( this, &FBlueprintEditor::OnDistributeNodesV )
-			// 	);
-			//
-			// GraphEditorCommands->MapAction( FGenericCommands::Get().Rename,
-			// 	FExecuteAction::CreateSP( this, &FBlueprintEditor::OnRenameNode ),
-			// 	FCanExecuteAction::CreateSP( this, &FBlueprintEditor::CanRenameNodes )
-			// 	);
-			//
-			// GraphEditorCommands->MapAction( FGraphEditorCommands::Get().ExpandNodes,
-			// 	FExecuteAction::CreateSP( this, &FBlueprintEditor::OnExpandNodes ),
-			// 	FCanExecuteAction::CreateSP( this, &FBlueprintEditor::CanExpandNodes ),
-			// 	FIsActionChecked(),
-			// 	FIsActionButtonVisible::CreateSP( this, &FBlueprintEditor::CanExpandNodes )
-			// 	);
-			//
-			// // Editing commands
-			// GraphEditorCommands->MapAction( FGenericCommands::Get().SelectAll,
-			// 	FExecuteAction::CreateSP( this, &FBlueprintEditor::SelectAllNodes ),
-			// 	FCanExecuteAction::CreateSP( this, &FBlueprintEditor::CanSelectAllNodes )
-			// 	);
-			//
+			
+			// Alignment Commands
+			GraphEditorCommands->MapAction( FGraphEditorCommands::Get().AlignNodesTop,
+				FExecuteAction::CreateSP( this, &FRigVMNewEditor::OnAlignTop )
+				);
+			
+			GraphEditorCommands->MapAction( FGraphEditorCommands::Get().AlignNodesMiddle,
+				FExecuteAction::CreateSP( this, &FRigVMNewEditor::OnAlignMiddle )
+				);
+			
+			GraphEditorCommands->MapAction( FGraphEditorCommands::Get().AlignNodesBottom,
+				FExecuteAction::CreateSP( this, &FRigVMNewEditor::OnAlignBottom )
+				);
+			
+			GraphEditorCommands->MapAction( FGraphEditorCommands::Get().AlignNodesLeft,
+				FExecuteAction::CreateSP( this, &FRigVMNewEditor::OnAlignLeft )
+				);
+			
+			GraphEditorCommands->MapAction( FGraphEditorCommands::Get().AlignNodesCenter,
+				FExecuteAction::CreateSP( this, &FRigVMNewEditor::OnAlignCenter )
+				);
+			
+			GraphEditorCommands->MapAction( FGraphEditorCommands::Get().AlignNodesRight,
+				FExecuteAction::CreateSP( this, &FRigVMNewEditor::OnAlignRight )
+				);
+			
+			GraphEditorCommands->MapAction( FGraphEditorCommands::Get().StraightenConnections,
+				FExecuteAction::CreateSP( this, &FRigVMNewEditor::OnStraightenConnections )
+				);
+			
+			// Distribution Commands
+			GraphEditorCommands->MapAction( FGraphEditorCommands::Get().DistributeNodesHorizontally,
+				FExecuteAction::CreateSP( this, &FRigVMNewEditor::OnDistributeNodesH )
+				);
+			
+			GraphEditorCommands->MapAction( FGraphEditorCommands::Get().DistributeNodesVertically,
+				FExecuteAction::CreateSP( this, &FRigVMNewEditor::OnDistributeNodesV )
+				);
+
+			// Editing commands
+			GraphEditorCommands->MapAction( FGenericCommands::Get().SelectAll,
+				FExecuteAction::CreateSP( this, &FRigVMNewEditor::SelectAllNodes ),
+				FCanExecuteAction::CreateSP( this, &FRigVMNewEditor::CanSelectAllNodes )
+				);
+			
 			GraphEditorCommands->MapAction( FGenericCommands::Get().Delete,
 				FExecuteAction::CreateSP( this, &FRigVMEditorBase::DeleteSelectedNodes ),
 				FCanExecuteAction::CreateSP( this, &FRigVMEditorBase::CanDeleteNodes )
 				);
-			//
-			// GraphEditorCommands->MapAction(FGraphEditorCommands::Get().DeleteAndReconnectNodes,
-			// 	FExecuteAction::CreateSP(this, &FBlueprintEditor::DeleteSelectedNodes),
-			// 	FCanExecuteAction::CreateSP(this, &FBlueprintEditor::CanDeleteNodes)
-			// );
-			//
-			// GraphEditorCommands->MapAction( FGenericCommands::Get().Copy,
-			// 	FExecuteAction::CreateSP( this, &FBlueprintEditor::CopySelectedNodes ),
-			// 	FCanExecuteAction::CreateSP( this, &FBlueprintEditor::CanCopyNodes )
-			// 	);
-			//
-			// GraphEditorCommands->MapAction( FGenericCommands::Get().Cut,
-			// 	FExecuteAction::CreateSP( this, &FBlueprintEditor::CutSelectedNodes ),
-			// 	FCanExecuteAction::CreateSP( this, &FBlueprintEditor::CanCutNodes )
-			// 	);
-			//
-			// GraphEditorCommands->MapAction( FGenericCommands::Get().Paste,
-			// 	FExecuteAction::CreateSP( this, &FBlueprintEditor::PasteGeneric ),
-			// 	FCanExecuteAction::CreateSP( this, &FBlueprintEditor::CanPasteGeneric )
-			// 	);
-			//
-			// GraphEditorCommands->MapAction( FGenericCommands::Get().Duplicate,
-			// 	FExecuteAction::CreateSP( this, &FBlueprintEditor::DuplicateNodes ),
-			// 	FCanExecuteAction::CreateSP( this, &FBlueprintEditor::CanDuplicateNodes )
-			// 	);
-			//
-			// GraphEditorCommands->MapAction( FGraphEditorCommands::Get().SelectReferenceInLevel,
-			// 	FExecuteAction::CreateSP( this, &FBlueprintEditor::OnSelectReferenceInLevel ),
-			// 	FCanExecuteAction::CreateSP( this, &FBlueprintEditor::CanSelectReferenceInLevel ),
-			// 	FIsActionChecked(),
-			// 	FIsActionButtonVisible::CreateSP( this, &FBlueprintEditor::CanSelectReferenceInLevel )
-			// 	);
-			//
-			// GraphEditorCommands->MapAction( FGraphEditorCommands::Get().AssignReferencedActor,
-			// 	FExecuteAction::CreateSP( this, &FBlueprintEditor::OnAssignReferencedActor ),
-			// 	FCanExecuteAction::CreateSP( this, &FBlueprintEditor::CanAssignReferencedActor ) );
-			//
-			// GraphEditorCommands->MapAction( FGraphEditorCommands::Get().StartWatchingPin,
-			// 	FExecuteAction::CreateSP( this, &FBlueprintEditor::OnStartWatchingPin ),
-			// 	FCanExecuteAction::CreateSP( this, &FBlueprintEditor::CanStartWatchingPin )
-			// 	);
-			//
-			// GraphEditorCommands->MapAction( FGraphEditorCommands::Get().StopWatchingPin,
-			// 	FExecuteAction::CreateSP( this, &FBlueprintEditor::OnStopWatchingPin ),
-			// 	FCanExecuteAction::CreateSP( this, &FBlueprintEditor::CanStopWatchingPin )
-			// 	);
-			//
-			// GraphEditorCommands->MapAction( FGraphEditorCommands::Get().CreateComment,
-			// 	FExecuteAction::CreateSP( this, &FBlueprintEditor::OnCreateComment )
-			// 	);
-			// 	
-			// GraphEditorCommands->MapAction( FGraphEditorCommands::Get().CreateCustomEvent,
-			// 	FExecuteAction::CreateSP( this, &FBlueprintEditor::OnCreateCustomEvent )
-			// 	);
-			//
-			// GraphEditorCommands->MapAction(FGraphEditorCommands::Get().ShowAllPins,
-			// 	FExecuteAction::CreateSP(this, &FBlueprintEditor::SetPinVisibility, SGraphEditor::Pin_Show)
-			// 	);
-			//
-			// GraphEditorCommands->MapAction(FGraphEditorCommands::Get().HideNoConnectionPins,
-			// 	FExecuteAction::CreateSP(this, &FBlueprintEditor::SetPinVisibility, SGraphEditor::Pin_HideNoConnection)
-			// 	);
-			//
-			// GraphEditorCommands->MapAction(FGraphEditorCommands::Get().HideNoConnectionNoDefaultPins,
-			// 	FExecuteAction::CreateSP(this, &FBlueprintEditor::SetPinVisibility, SGraphEditor::Pin_HideNoConnectionNoDefault)
-			// 	);
-			//
-			// GraphEditorCommands->MapAction( FGraphEditorCommands::Get().FindReferences,
-			// 	FExecuteAction::CreateSP( this, &FBlueprintEditor::OnFindReferences, /*bSearchAllBlueprints=*/false, EGetFindReferenceSearchStringFlags::Legacy),
-			// 	FCanExecuteAction::CreateSP( this, &FBlueprintEditor::CanFindReferences )
-			// 	);
-			//
-			// GraphEditorCommands->MapAction( FGraphEditorCommands::Get().FindReferencesByNameLocal,
-			// 	FExecuteAction::CreateSP( this, &FBlueprintEditor::OnFindReferences, /*bSearchAllBlueprints=*/false, EGetFindReferenceSearchStringFlags::None ),
-			// 	FCanExecuteAction::CreateSP( this, &FBlueprintEditor::CanFindReferences )
-			// 	);
-			//
-			// GraphEditorCommands->MapAction( FGraphEditorCommands::Get().FindReferencesByNameGlobal,
-			// 	FExecuteAction::CreateSP( this, &FBlueprintEditor::OnFindReferences, /*bSearchAllBlueprints=*/true, EGetFindReferenceSearchStringFlags::None ),
-			// 	FCanExecuteAction::CreateSP( this, &FBlueprintEditor::CanFindReferences )
-			// 	);
-			//
-			// GraphEditorCommands->MapAction( FGraphEditorCommands::Get().FindReferencesByClassMemberLocal,
-			// 	FExecuteAction::CreateSP( this, &FBlueprintEditor::OnFindReferences, /*bSearchAllBlueprints=*/false, EGetFindReferenceSearchStringFlags::UseSearchSyntax ),
-			// 	FCanExecuteAction::CreateSP( this, &FBlueprintEditor::CanFindReferences )
-			// 	);
-			//
-			// GraphEditorCommands->MapAction( FGraphEditorCommands::Get().FindReferencesByClassMemberGlobal,
-			// 	FExecuteAction::CreateSP( this, &FBlueprintEditor::OnFindReferences, /*bSearchAllBlueprints=*/true, EGetFindReferenceSearchStringFlags::UseSearchSyntax ),
-			// 	FCanExecuteAction::CreateSP( this, &FBlueprintEditor::CanFindReferences )
-			// 	);
-			//
-			// GraphEditorCommands->MapAction( FGraphEditorCommands::Get().GoToDefinition,
-			// 	FExecuteAction::CreateSP( this, &FBlueprintEditor::OnGoToDefinition ),
-			// 	FCanExecuteAction::CreateSP( this, &FBlueprintEditor::CanGoToDefinition )
-			// 	);
-			//
-			// GraphEditorCommands->MapAction(FGraphEditorCommands::Get().GoToDocumentation,
-			// 	FExecuteAction::CreateSP(this, &FBlueprintEditor::OnGoToDocumentation),
-			// 	FCanExecuteAction::CreateSP(this, &FBlueprintEditor::CanGoToDocumentation)
-			// 	);
-			//
-			// GraphEditorCommands->MapAction(FGraphEditorCommands::Get().EnableNodes,
-			// 	FExecuteAction(),
-			// 	FCanExecuteAction(),
-			// 	FGetActionCheckState::CreateSP(this, &FBlueprintEditor::GetEnabledCheckBoxStateForSelectedNodes)
-			// 	);
-			//
-			// GraphEditorCommands->MapAction(FGraphEditorCommands::Get().DisableNodes,
-			// 	FExecuteAction::CreateSP(this, &FBlueprintEditor::OnSetEnabledStateForSelectedNodes, ENodeEnabledState::Disabled),
-			// 	FCanExecuteAction(),
-			// 	FGetActionCheckState::CreateSP(this, &FBlueprintEditor::CheckEnabledStateForSelectedNodes, ENodeEnabledState::Disabled)
-			// 	);
-			//
-			// GraphEditorCommands->MapAction(FGraphEditorCommands::Get().EnableNodes_Always,
-			// 	FExecuteAction::CreateSP(this, &FBlueprintEditor::OnSetEnabledStateForSelectedNodes, ENodeEnabledState::Enabled),
-			// 	FCanExecuteAction(),
-			// 	FGetActionCheckState::CreateSP(this, &FBlueprintEditor::CheckEnabledStateForSelectedNodes, ENodeEnabledState::Enabled)
-			// 	);
-			//
-			// GraphEditorCommands->MapAction(FGraphEditorCommands::Get().EnableNodes_DevelopmentOnly,
-			// 	FExecuteAction::CreateSP(this, &FBlueprintEditor::OnSetEnabledStateForSelectedNodes, ENodeEnabledState::DevelopmentOnly),
-			// 	FCanExecuteAction(),
-			// 	FGetActionCheckState::CreateSP(this, &FBlueprintEditor::CheckEnabledStateForSelectedNodes, ENodeEnabledState::DevelopmentOnly)
-			// 	);
+			
+			GraphEditorCommands->MapAction( FGenericCommands::Get().Copy,
+				FExecuteAction::CreateSP( this, &FRigVMEditorBase::CopySelectedNodes ),
+				FCanExecuteAction::CreateSP( this, &FRigVMEditorBase::CanCopyNodes )
+				);
+			
+			GraphEditorCommands->MapAction( FGenericCommands::Get().Cut,
+				FExecuteAction::CreateSP( this, &FRigVMEditorBase::CutSelectedNodes ),
+				FCanExecuteAction::CreateSP( this, &FRigVMEditorBase::CanCutNodes )
+				);
+			
+			GraphEditorCommands->MapAction( FGenericCommands::Get().Paste,
+				FExecuteAction::CreateSP( this, &FRigVMEditorBase::PasteNodes ),
+				FCanExecuteAction::CreateSP( this, &FRigVMEditorBase::CanPasteNodes )
+				);
+			
+			GraphEditorCommands->MapAction( FGenericCommands::Get().Duplicate,
+				FExecuteAction::CreateSP( this, &FRigVMEditorBase::DuplicateNodes ),
+				FCanExecuteAction::CreateSP( this, &FRigVMEditorBase::CanDuplicateNodes )
+				);
+			
+			GraphEditorCommands->MapAction( FGraphEditorCommands::Get().StartWatchingPin,
+				FExecuteAction::CreateSP( this, &FRigVMEditorBase::OnStartWatchingPin ),
+				FCanExecuteAction::CreateSP( this, &FRigVMEditorBase::CanStartWatchingPin )
+				);
+			
+			GraphEditorCommands->MapAction( FGraphEditorCommands::Get().StopWatchingPin,
+				FExecuteAction::CreateSP( this, &FRigVMEditorBase::OnStopWatchingPin ),
+				FCanExecuteAction::CreateSP( this, &FRigVMEditorBase::CanStopWatchingPin )
+				);
+			
+			GraphEditorCommands->MapAction( FGraphEditorCommands::Get().CreateComment,
+				FExecuteAction::CreateSP( this, &FRigVMNewEditor::OnCreateComment )
+				);
 
 			OnCreateGraphEditorCommands(GraphEditorCommands);
 		}
@@ -1604,10 +1336,7 @@ void FRigVMNewEditor::RefreshEditorsImpl(ERefreshRigVMEditorReason::Type Reason)
 void FRigVMNewEditor::SetupGraphEditorEventsImpl(UEdGraph* InGraph, SGraphEditor::FGraphEditorEvents& InEvents)
 {
 	// TODO sara-s: Replicate what blueprint editor is doing
-	// InEvents.OnSelectionChanged = SGraphEditor::FOnSelectionChanged::CreateSP( this, &FRigVMNewEditor::OnSelectedNodesChanged );
-	// InEvents.OnDropActor = SGraphEditor::FOnDropActor::CreateSP( this, &FRigVMNewEditor::OnGraphEditorDropActor );
-	// InEvents.OnDropStreamingLevel = SGraphEditor::FOnDropStreamingLevel::CreateSP( this, &FRigVMNewEditor::OnGraphEditorDropStreamingLevel );
-	// InEvents.OnNodeDoubleClicked = FSingleNodeEvent::CreateSP(this, &FRigVMNewEditor::OnNodeDoubleClicked);
+	InEvents.OnSelectionChanged = SGraphEditor::FOnSelectionChanged::CreateSP( this, &FRigVMNewEditor::OnSelectedNodesChanged );
 	// InEvents.OnVerifyTextCommit = FOnNodeVerifyTextCommit::CreateSP(this, &FRigVMNewEditor::OnNodeVerifyTitleCommit);
 	// InEvents.OnTextCommitted = FOnNodeTextCommitted::CreateSP(this, &FRigVMNewEditor::OnNodeTitleCommitted);
 	// InEvents.OnSpawnNodeByShortcut = SGraphEditor::FOnSpawnNodeByShortcut::CreateSP(this, &FRigVMNewEditor::OnSpawnGraphNodeByShortcut, InGraph);
@@ -1615,7 +1344,8 @@ void FRigVMNewEditor::SetupGraphEditorEventsImpl(UEdGraph* InGraph, SGraphEditor
 	// InEvents.OnDisallowedPinConnection = SGraphEditor::FOnDisallowedPinConnection::CreateSP(this, &FRigVMNewEditor::OnDisallowedPinConnection);
 	// InEvents.OnDoubleClicked = SGraphEditor::FOnDoubleClicked::CreateSP(this, &FRigVMNewEditor::NavigateToParentGraphByDoubleClick);
 	//
-	// // Custom menu for K2 schemas
+
+	// Custom menu for K2 schemas
 	if(InGraph->Schema != nullptr && InGraph->Schema->IsChildOf(UEdGraphSchema_K2::StaticClass()))
 	{
 		InEvents.OnCreateActionMenu = SGraphEditor::FOnCreateActionMenu::CreateSP(this, &FRigVMNewEditor::OnCreateGraphActionMenu);
@@ -1915,11 +1645,7 @@ void FRigVMNewEditor::RefreshInspector()
 void FRigVMNewEditor::RefreshStandAloneDefaultsEditor()
 {
 	// Update the details panel
-#if WITH_RIGVMLEGACYEDITOR
-	SKismetInspector::FShowDetailsOptions Options(FText::GetEmpty(), true);
-#else
 	SRigVMDetailsInspector::FShowDetailsOptions Options(FText::GetEmpty(), true);
-#endif
 
 	TArray<UObject*> DefaultObjects;
 	for ( int32 i = 0; i < GetEditingObjects().Num(); ++i )
@@ -1963,12 +1689,8 @@ void FRigVMNewEditor::StartEditingDefaults(bool bAutoFocus, bool bForceRefresh)
 				// Update the details panel
 				FString Title;
 				DefaultObject->GetName(Title);
-#if WITH_RIGVMLEGACYEDITOR
-				SKismetInspector::FShowDetailsOptions Options(FText::FromString(Title), bForceRefresh);
-				Options.bShowComponents = false;
-#else
+
 				SRigVMDetailsInspector::FShowDetailsOptions Options(FText::FromString(Title), bForceRefresh);
-#endif
 
 				if (Inspector.IsValid())
 				{
@@ -2072,29 +1794,29 @@ void FRigVMNewEditor::CreateDefaultTabContents(const TArray<UBlueprint*> InBluep
 	// CompilerResults = MessageLogModule.CreateLogListingWidget( CompilerResultsListing.ToSharedRef() );
 	// FindResults = SNew(SFindInBlueprints, SharedThis(this));
 	
-	// this->Inspector = 
-	// 	SNew(SRigVMDetailsInspector)
-	// 	.Editor(SharedRef())
-	// 	. HideNameArea(true)
-	// 	. ViewIdentifier(FName("BlueprintInspector"))
-	// 	//. OnFinishedChangingProperties( FOnFinishedChangingProperties::FDelegate::CreateSP(this, &FBlueprintEditor::OnFinishedChangingProperties) ); // TODO sara-s
-	// ;
+	this->Inspector = 
+		SNew(SRigVMDetailsInspector)
+		.Editor(SharedThis(this))
+		. HideNameArea(true)
+		. ViewIdentifier(FName("BlueprintInspector"))
+		//. OnFinishedChangingProperties( FOnFinishedChangingProperties::FDelegate::CreateSP(this, &FBlueprintEditor::OnFinishedChangingProperties) ); // TODO sara-s
+	;
 
 	if ( InBlueprints.Num() > 0 )
 	{
 		// Don't show the object name in defaults mode.
 		const bool bHideNameArea = bWasOpenedInDefaultsMode;
 
-		// this->DefaultEditor = 
-		// 	SNew(SRigVMDetailsInspector)
-		// 	.Editor(SharedRef())
-		// 	. ViewIdentifier(FName("BlueprintDefaults"))
-		// 	. IsEnabled(!bIsInterface)
-		// 	//. ShowPublicViewControl(this, &FBlueprintEditor::ShouldShowPublicViewControl) // TODO sara-s
-		// 	. ShowTitleArea(false)
-		// 	. HideNameArea(bHideNameArea)
-		// 	//. OnFinishedChangingProperties( FOnFinishedChangingProperties::FDelegate::CreateSP( this, &FBlueprintEditor::OnFinishedChangingProperties ) ); // TODO sara-s
-		// ;
+		this->DefaultEditor = 
+			SNew(SRigVMDetailsInspector)
+			.Editor(SharedThis(this))
+			. ViewIdentifier(FName("BlueprintDefaults"))
+			. IsEnabled(!bIsInterface)
+			//. ShowPublicViewControl(this, &FBlueprintEditor::ShouldShowPublicViewControl) // TODO sara-s
+			. ShowTitleArea(false)
+			. HideNameArea(bHideNameArea)
+			//. OnFinishedChangingProperties( FOnFinishedChangingProperties::FDelegate::CreateSP( this, &FBlueprintEditor::OnFinishedChangingProperties ) ); // TODO sara-s
+		;
 	}
 
 	if (InBlueprint && 
@@ -2159,5 +1881,105 @@ const FSlateBrush* FRigVMNewEditor::GetGlyphForGraph(const UEdGraph* Graph, bool
 	return ReturnValue;
 }
 
+void FRigVMNewEditor::OnSelectedNodesChanged(const FGraphPanelSelectionSet& NewSelection)
+{
+	OnSelectedNodesChangedImpl(NewSelection);
+}
+
+void FRigVMNewEditor::OnAlignTop()
+{
+	TSharedPtr<SGraphEditor> FocusedGraphEd = FocusedGraphEdPtr.Pin();
+	if (FocusedGraphEd.IsValid())
+	{
+		FocusedGraphEd->OnAlignTop();
+	}
+}
+
+void FRigVMNewEditor::OnAlignMiddle()
+{
+	TSharedPtr<SGraphEditor> FocusedGraphEd = FocusedGraphEdPtr.Pin();
+	if (FocusedGraphEd.IsValid())
+	{
+		FocusedGraphEd->OnAlignMiddle();
+	}
+}
+
+void FRigVMNewEditor::OnAlignBottom()
+{
+	TSharedPtr<SGraphEditor> FocusedGraphEd = FocusedGraphEdPtr.Pin();
+	if (FocusedGraphEd.IsValid())
+	{
+		FocusedGraphEd->OnAlignBottom();
+	}
+}
+
+void FRigVMNewEditor::OnAlignLeft()
+{
+	TSharedPtr<SGraphEditor> FocusedGraphEd = FocusedGraphEdPtr.Pin();
+	if (FocusedGraphEd.IsValid())
+	{
+		FocusedGraphEd->OnAlignLeft();
+	}
+}
+
+void FRigVMNewEditor::OnAlignCenter()
+{
+	TSharedPtr<SGraphEditor> FocusedGraphEd = FocusedGraphEdPtr.Pin();
+	if (FocusedGraphEd.IsValid())
+	{
+		FocusedGraphEd->OnAlignCenter();
+	}
+}
+
+void FRigVMNewEditor::OnAlignRight()
+{
+	TSharedPtr<SGraphEditor> FocusedGraphEd = FocusedGraphEdPtr.Pin();
+	if (FocusedGraphEd.IsValid())
+	{
+		FocusedGraphEd->OnAlignRight();
+	}
+}
+
+void FRigVMNewEditor::OnStraightenConnections()
+{
+	TSharedPtr<SGraphEditor> FocusedGraphEd = FocusedGraphEdPtr.Pin();
+	if (FocusedGraphEd.IsValid())
+	{
+		FocusedGraphEd->OnStraightenConnections();
+	}
+}
+
+void FRigVMNewEditor::OnDistributeNodesH()
+{
+	TSharedPtr<SGraphEditor> FocusedGraphEd = FocusedGraphEdPtr.Pin();
+	if (FocusedGraphEd.IsValid())
+	{
+		FocusedGraphEd->OnDistributeNodesH();
+	}
+}
+
+void FRigVMNewEditor::OnDistributeNodesV()
+{
+	TSharedPtr<SGraphEditor> FocusedGraphEd = FocusedGraphEdPtr.Pin();
+	if (FocusedGraphEd.IsValid())
+	{
+		FocusedGraphEd->OnDistributeNodesV();
+	}
+}
+
+void FRigVMNewEditor::SelectAllNodes()
+{
+	TSharedPtr<SGraphEditor> FocusedGraphEd = FocusedGraphEdPtr.Pin();
+	if (FocusedGraphEd.IsValid())
+	{
+		FocusedGraphEd->SelectAllNodes();
+	}
+}
+
+bool FRigVMNewEditor::CanSelectAllNodes() const
+{
+	return true;
+}
+
 
 #undef LOCTEXT_NAMESPACE
diff --git a/Engine/Plugins/Runtime/RigVM/Source/RigVMEditor/Private/Editor/RigVMNewEditorMode.cpp b/Engine/Plugins/Runtime/RigVM/Source/RigVMEditor/Private/Editor/RigVMNewEditorMode.cpp
--- a/Engine/Plugins/Runtime/RigVM/Source/RigVMEditor/Private/Editor/RigVMNewEditorMode.cpp
+++ b/Engine/Plugins/Runtime/RigVM/Source/RigVMEditor/Private/Editor/RigVMNewEditorMode.cpp
@@ -1,13 +1,11 @@
 // Copyright Epic Games, Inc. All Rights Reserved.
 
-#include "Editor/RigVMNewEditorMode.h" 
-#include "BlueprintEditorTabs.h"
-#include "SBlueprintEditorToolbar.h"
+#include "Editor/RigVMNewEditorMode.h"
+
+#include "Editor/RigVMDetailsInspectorTabSummoner.h"
+#include "Editor/RigVMEditor.h"
 #include "Editor/RigVMExecutionStackTabSummoner.h"
-#include "Editor/RigVMEditorGraphExplorerTabSummoner.h"
-#include "ToolMenus.h"
 #include "Editor/RigVMNewEditor.h"
-#include "RigVMBlueprint.h"
 
 const FName FRigVMNewEditorApplicationModes::StandardRigVMEditorMode() { static FName Mode = TEXT("GraphName"); return Mode; };
 const FName FRigVMNewEditorApplicationModes::RigVMDefaultsMode() { static FName Mode = TEXT("DefaultsName"); return Mode; };
@@ -16,10 +14,10 @@ FRigVMNewEditorMode::FRigVMNewEditorMode(const TSharedRef<FRigVMNewEditor>& InRi
 	: FApplicationMode(FRigVMEditorModes::RigVMEditorMode, FRigVMEditorModes::GetLocalizedMode)
 {
 	Editor = InRigVMEditor;
-	RigVMBlueprintPtr = InRigVMEditor->GetRigVMBlueprint();
-
+	
 	TabFactories.RegisterFactory(MakeShared<FRigVMExecutionStackTabSummoner>(InRigVMEditor));
 	TabFactories.RegisterFactory(MakeShared<FRigVMEditorGraphExplorerTabSummoner>(InRigVMEditor));
+	TabFactories.RegisterFactory(MakeShared<FRigVMDetailsInspectorTabSummoner>(InRigVMEditor));
 
 	TabLayout = FTabManager::NewLayout("Standalone_RigVMEditMode_Layout_v1.5")
 		->AddArea
@@ -37,7 +35,7 @@ FRigVMNewEditorMode::FRigVMNewEditorMode(const TSharedRef<FRigVMNewEditor>& InRi
 					FTabManager::NewStack()
 					->SetSizeCoefficient(0.5f)
 					->AddTab(FRigVMExecutionStackTabSummoner::TabID, ETabState::OpenedTab)
-					->AddTab(FRigVMEditorGraphExplorerTabSummoner::TabID, ETabState::OpenedTab)
+					->AddTab(FRigVMEditorGraphExplorerTabSummoner::TabID(), ETabState::OpenedTab)
 				)
 				->Split
 				(
@@ -53,6 +51,22 @@ FRigVMNewEditorMode::FRigVMNewEditorMode(const TSharedRef<FRigVMNewEditor>& InRi
 						->AddTab("Document", ETabState::ClosedTab)
 					)
 				)
+				->Split
+				(
+					// Right side
+					FTabManager::NewSplitter()
+					->SetOrientation(Orient_Vertical)
+					->SetSizeCoefficient(0.2f)
+					->Split
+					(
+						// Right top
+						FTabManager::NewStack()
+						->SetHideTabWell(false)
+						->SetSizeCoefficient(1.f)
+						->AddTab(FRigVMDetailsInspectorTabSummoner::TabID(), ETabState::OpenedTab)
+						->SetForegroundTab(FRigVMDetailsInspectorTabSummoner::TabID())
+					)
+				)
 			)
 		);
 }
diff --git a/Engine/Plugins/Runtime/RigVM/Source/RigVMEditor/Private/Editor/RigVMVariableDetailCustomization.cpp b/Engine/Plugins/Runtime/RigVM/Source/RigVMEditor/Private/Editor/RigVMVariableDetailCustomization.cpp
--- a/Engine/Plugins/Runtime/RigVM/Source/RigVMEditor/Private/Editor/RigVMVariableDetailCustomization.cpp
+++ b/Engine/Plugins/Runtime/RigVM/Source/RigVMEditor/Private/Editor/RigVMVariableDetailCustomization.cpp
@@ -4,111 +4,1025 @@
 #include "Kismet2/BlueprintEditorUtils.h"
 #include "RigVMHost.h"
 #include "DetailLayoutBuilder.h"
-#include "BlueprintEditorModule.h"
 #include "DetailCategoryBuilder.h"
 #include "DetailWidgetRow.h"
+#include "IDocumentation.h"
 #include "Widgets/Text/STextBlock.h"
 #include "Widgets/Input/SCheckBox.h"
 #include "Engine/BlueprintGeneratedClass.h"
 #include "RigVMBlueprint.h"
+#include "Editor/RigVMEditor.h"
+#if WITH_RIGVMLEGACYEDITOR
+#include "Editor/RigVMLegacyEditor.h"
+#endif
 #include "RigVMModel/RigVMController.h"
+#include "Widgets/Input/SMultiLineEditableTextBox.h"
+#include "RigVMBlueprint.h"
+#include "SPinTypeSelector.h"
+#include "SupportedRangeTypes.h"
+#include "Editor/SRigVMDetailsInspector.h"
+#include "Widgets/Input/SEditableTextBox.h"
+#include "Widgets/Input/SComboButton.h"
+#include "Widgets/SToolTip.h"
+
 
 #define LOCTEXT_NAMESPACE "RigVMVariableDetailCustomization"
 
-TSharedPtr<IDetailCustomization> FRigVMVariableDetailCustomization::MakeInstance(TSharedPtr<IBlueprintEditor> InBlueprintEditor)
+TSharedRef<IDetailCustomization> FRigVMVariableDetailCustomization::MakeInstance(TSharedPtr<FRigVMEditorBase> InEditor)
 {
-	const TArray<UObject*>* Objects = (InBlueprintEditor.IsValid() ? InBlueprintEditor->GetObjectsCurrentlyBeingEdited() : nullptr);
+	const TArray<UObject*>* Objects = (InEditor.IsValid() ? InEditor->GetObjectsCurrentlyBeingEdited() : nullptr);
 	if (Objects && Objects->Num() == 1)
 	{
-		if (UBlueprint* Blueprint = Cast<UBlueprint>((*Objects)[0]))
+		if (URigVMBlueprint* Blueprint = Cast<URigVMBlueprint>((*Objects)[0]))
 		{
 			if (Blueprint->ParentClass && Blueprint->ParentClass->IsChildOf(URigVMHost::StaticClass()))
 			{
-				return MakeShareable(new FRigVMVariableDetailCustomization(InBlueprintEditor, Blueprint));
+				return MakeShareable(new FRigVMVariableDetailCustomization(InEditor, Blueprint));
 			}
 		}
 	}
 
-	return nullptr;
+	return MakeShareable(new FRigVMVariableDetailCustomization((TSharedPtr<FRigVMEditorBase>)nullptr, nullptr));
+}
+
+FRigVMVariableDetailCustomization::FRigVMVariableDetailCustomization(TSharedPtr<FRigVMEditorBase> InEditor, URigVMBlueprint* Blueprint)
+		: EditorPtr(InEditor)
+		, BlueprintPtr(Blueprint)
+{}
+
+#if WITH_RIGVMLEGACYEDITOR
+TSharedPtr<IDetailCustomization> FRigVMVariableDetailCustomization::MakeLegacyInstance(TSharedPtr<IBlueprintEditor> InEditor)
+{
+	TSharedPtr<FRigVMLegacyEditor> LegacyEditor = StaticCastSharedPtr<FRigVMLegacyEditor>(InEditor);
+	TSharedPtr<FRigVMEditorBase> Editor = StaticCastSharedPtr<FRigVMEditorBase>(LegacyEditor);
+	return MakeInstance(Editor);
 }
 
+FRigVMVariableDetailCustomization::FRigVMVariableDetailCustomization(TSharedPtr<IBlueprintEditor> RigVMigEditor, UBlueprint* Blueprint)
+{
+	EditorPtr = StaticCastSharedPtr<FRigVMLegacyEditor>(RigVMigEditor);
+	BlueprintPtr = Cast<URigVMBlueprint>(Blueprint);
+}
+#endif
+
 void FRigVMVariableDetailCustomization::CustomizeDetails(IDetailLayoutBuilder& DetailLayout)
 {
-	TArray<TWeakObjectPtr<UObject>> ObjectsBeingCustomized;
-	DetailLayout.GetObjectsBeingCustomized(ObjectsBeingCustomized);
-	if (ObjectsBeingCustomized.Num() > 0)
-	{
-		UPropertyWrapper* PropertyWrapper = Cast<UPropertyWrapper>(ObjectsBeingCustomized[0].Get());
-		TWeakFieldPtr<FProperty> PropertyBeingCustomized = PropertyWrapper ? PropertyWrapper->GetProperty() : nullptr;
-		if (PropertyBeingCustomized.IsValid())
-		{
-			const FText AnimationInputText = LOCTEXT("AnimationInput", "Animation Input");
-			const FText AnimationOutputText = LOCTEXT("AnimationOutput", "Animation Output");
-			const FText AnimationInputTooltipText = LOCTEXT("AnimationInputTooltip", "Whether this variable acts as an input to this animation controller.\nSelecting this allow it to be exposed as an input pin on Evaluation nodes.");
-			const FText AnimationOutputTooltipText = LOCTEXT("AnimationOutputTooltip", "Whether this variable acts as an output from this animation controller.\nSelecting this will add a pin to the Animation Output node.");
-			
-			DetailLayout.EditCategory("Variable")
-				.AddCustomRow(AnimationInputText)
-				.NameContent()
+	TSharedPtr<FRigVMEditorBase> Editor = EditorPtr.Pin();
+	if (!Editor.IsValid())
+	{
+		return;
+	}
+
+	if (Editor->GetRigVMInspector())
+	{
+		FName VariableName = GetVariableName();
+		CachedVariableProperty = BlueprintPtr->SkeletonGeneratedClass->FindPropertyByName(VariableName);
+
+		if(!CachedVariableProperty.IsValid())
+		{
+			return;
+		}
+
+		Editor->OnRefresh().AddSP(this, &FRigVMVariableDetailCustomization::OnPostEditorRefresh);
+
+		// Get an appropiate name validator
+		TSharedPtr<INameValidatorInterface> NameValidator = nullptr;
+		{
+			const UEdGraphSchema* Schema = nullptr;
+			if (BlueprintPtr.IsValid())
+			{
+				TArray<UEdGraph*> Graphs;
+				BlueprintPtr->GetAllGraphs(Graphs);
+				if (Graphs.Num() > 0)
+				{
+					Schema = Graphs[0]->GetSchema();
+				}
+			}
+
+			if (Schema)
+			{
+				NameValidator = Schema->GetNameValidator(BlueprintPtr.Get(), CachedVariableName, nullptr, FEdGraphSchemaAction_K2Var::StaticGetTypeId());
+			}
+		}
+
+		FProperty* VariableProperty = CachedVariableProperty.Get();
+
+		IDetailCategoryBuilder& Category = DetailLayout.EditCategory("Variable", LOCTEXT("VariableDetailsCategory", "Variable"));
+		const FSlateFontInfo DetailFontInfo = IDetailLayoutBuilder::GetDetailFont();
+		
+		TSharedPtr<SToolTip> VarNameTooltip = IDocumentation::Get()->CreateToolTip(LOCTEXT("VarNameTooltip", "The name of the variable."), NULL, "", TEXT("VariableName"));
+
+		Category.AddCustomRow( LOCTEXT("BlueprintVarActionDetails_VariableNameLabel", "Variable Name") )
+		.NameContent()
+		[
+			SNew(STextBlock)
+			.Text(LOCTEXT("BlueprintVarActionDetails_VariableNameLabel", "Variable Name"))
+			.ToolTip(VarNameTooltip)
+			.Font(IDetailLayoutBuilder::GetDetailFont())
+		]
+		.ValueContent()
+		.MaxDesiredWidth(250.0f)
+		[
+			SAssignNew(VarNameEditableTextBox, SEditableTextBox)
+			.Text(this, &FRigVMVariableDetailCustomization::OnGetVariableName)
+			.ToolTip(VarNameTooltip)
+			.OnTextCommitted(this, &FRigVMVariableDetailCustomization::OnVarNameCommitted)
+			.OnVerifyTextChanged_Lambda([this, NameValidator](const FText& InNewText, FText& OutErrorMessage) -> bool
+			{	
+				if (NameValidator.IsValid())
+				{
+					EValidatorResult ValidatorResult = NameValidator->IsValid(InNewText.ToString());
+					switch (ValidatorResult)
+					{
+						case EValidatorResult::Ok:
+						case EValidatorResult::ExistingName:
+							// These are fine, don't need to surface to the user, the rename can 'proceed' even if the name is the existing one
+							return true;
+							break;
+						default:
+							OutErrorMessage = INameValidatorInterface::GetErrorText(InNewText.ToString(), ValidatorResult);
+							return false;
+							break;
+					}
+				}
+				
+				return true;
+			})
+			.Font(IDetailLayoutBuilder::GetDetailFont())
+		];
+
+	
+		TSharedPtr<SToolTip> VarTypeTooltip = IDocumentation::Get()->CreateToolTip(LOCTEXT("VarTypeTooltip", "The type of the variable."), NULL, "", TEXT("VariableType"));
+
+		TArray<TSharedPtr<IPinTypeSelectorFilter>> CustomPinTypeFilters;
+		Editor->GetPinTypeSelectorFilters(CustomPinTypeFilters);
+		
+		const UEdGraphSchema* Schema = GetDefault<UEdGraphSchema_K2>();
+		if (Editor->GetFocusedGraph())
+		{
+			Schema = Editor->GetFocusedGraph()->GetSchema();
+		}
+		
+
+		Category.AddCustomRow(LOCTEXT("VariableTypeLabel", "Variable Type"))
+			.NameContent()
+			[
+				SNew(STextBlock)
+				.Text(LOCTEXT("VariableTypeLabel", "Variable Type"))
+				.ToolTip(VarTypeTooltip)
+				.Font(DetailFontInfo)
+			]
+			.ValueContent()
+			.MaxDesiredWidth(980.f)
+			[
+				SNew(SPinTypeSelector, FGetPinTypeTree::CreateUObject(GetDefault<UEdGraphSchema_K2>(), &UEdGraphSchema_K2::GetVariableTypeTree))
+				.TargetPinType(this, &FRigVMVariableDetailCustomization::OnGetVarType)
+				.OnPinTypeChanged(this, &FRigVMVariableDetailCustomization::OnVarTypeChanged)
+				.Schema(Schema)
+				.TypeTreeFilter(ETypeTreeFilter::None)
+				.Font(DetailFontInfo)
+				.ToolTip(VarTypeTooltip)
+				.CustomFilters(CustomPinTypeFilters)
+			]
+			.AddCustomContextMenuAction(FUIAction(
+				FExecuteAction::CreateRaw(this, &FRigVMVariableDetailCustomization::OnBrowseToVarType),
+				FCanExecuteAction::CreateRaw(this, &FRigVMVariableDetailCustomization::CanBrowseToVarType)
+				),
+				LOCTEXT("BrowseToType", "Browse to Type"),
+				LOCTEXT("BrowseToTypeToolTip", "Browse to this variable type in the Content Browser."),
+				FSlateIcon(FAppStyle::GetAppStyleSetName(), "Icons.BrowseContent")
+			);
+
+
+
+		TSharedPtr<SToolTip> ToolTipTooltip = IDocumentation::Get()->CreateToolTip(LOCTEXT("VarToolTipTooltip", "Extra information about this variable, shown when cursor is over it."), NULL, "", TEXT("Description"));
+
+		Category.AddCustomRow( LOCTEXT("IsVariableToolTipLabel", "Description") )
+		.NameContent()
+		[
+			SNew(STextBlock)
+			.Text( LOCTEXT("IsVariableToolTipLabel", "Description") )
+			.ToolTip(ToolTipTooltip)
+			.Font( DetailFontInfo )
+		]
+		.ValueContent()
+		.MinDesiredWidth(250.f)
+		.MaxDesiredWidth(250.f)
+		[
+			SNew(SMultiLineEditableTextBox)
+			.Text( this, &FRigVMVariableDetailCustomization::OnGetTooltipText )
+			.ToolTipText( this, &FRigVMVariableDetailCustomization::OnGetTooltipText )
+			.OnTextCommitted( this, &FRigVMVariableDetailCustomization::OnTooltipTextCommitted, CachedVariableName )
+			.Font( DetailFontInfo )
+			.ModiferKeyForNewLine(EModifierKey::Shift)
+		];
+
+
+		TSharedPtr<SToolTip> ExposeOnSpawnTooltip = IDocumentation::Get()->CreateToolTip(LOCTEXT("VariableExposeToSpawn_Tooltip", "Should this variable be exposed as a pin when spawning this Blueprint?"), NULL, "", TEXT("ExposeOnSpawn"));
+
+		Category.AddCustomRow( LOCTEXT("VariableExposeToSpawnLabel", "Expose on Spawn") )
+		.NameContent()
+		[
+			SNew(STextBlock)
+			.ToolTip(ExposeOnSpawnTooltip)
+			.Text( LOCTEXT("VariableExposeToSpawnLabel", "Expose on Spawn") )
+			.Font( DetailFontInfo )
+		]
+		.ValueContent()
+		[
+			SNew(SCheckBox)
+			.IsChecked( this, &FRigVMVariableDetailCustomization::OnGetExposedToSpawnCheckboxState )
+			.OnCheckStateChanged( this, &FRigVMVariableDetailCustomization::OnExposedToSpawnChanged )
+			.ToolTip(ExposeOnSpawnTooltip)
+		];
+
+
+
+		TSharedPtr<SToolTip> PrivateTooltip = IDocumentation::Get()->CreateToolTip(LOCTEXT("VariablePrivate_Tooltip", "Should this variable be private (derived blueprints cannot modify it)?"), NULL, "", TEXT("Private"));
+
+		Category.AddCustomRow(LOCTEXT("VariablePrivate", "Private"))
+		.NameContent()
+		[
+			SNew(STextBlock)
+			.ToolTip(PrivateTooltip)
+			.Text( LOCTEXT("VariablePrivate", "Private") )
+			.Font( DetailFontInfo )
+		]
+		.ValueContent()
+		[
+			SNew(SCheckBox)
+			.IsChecked( this, &FRigVMVariableDetailCustomization::OnGetPrivateCheckboxState )
+			.OnCheckStateChanged( this, &FRigVMVariableDetailCustomization::OnPrivateChanged )
+			.ToolTip(PrivateTooltip)
+		];
+
+		TSharedPtr<SToolTip> ExposeToCinematicsTooltip = IDocumentation::Get()->CreateToolTip(LOCTEXT("VariableExposeToCinematics_Tooltip", "Should this variable be exposed for Sequencer to modify?"), NULL, "", TEXT("ExposeToCinematics"));
+
+		Category.AddCustomRow( LOCTEXT("VariableExposeToCinematics", "Expose to Cinematics") )
+		.NameContent()
+		[
+			SNew(STextBlock)
+			.ToolTip(ExposeToCinematicsTooltip)
+			.Text( LOCTEXT("VariableExposeToCinematics", "Expose to Cinematics") )
+			.Font( DetailFontInfo )
+		]
+		.ValueContent()
+		[
+			SNew(SCheckBox)
+			.IsChecked( this, &FRigVMVariableDetailCustomization::OnGetExposedToCinematicsCheckboxState )
+			.OnCheckStateChanged( this, &FRigVMVariableDetailCustomization::OnExposedToCinematicsChanged )
+			.ToolTip(ExposeToCinematicsTooltip)
+		];
+
+		PopulateCategories();
+		TSharedPtr<SComboButton> NewComboButton;
+		TSharedPtr<SListView<TSharedPtr<FText>>> NewListView;
+
+		TSharedPtr<SToolTip> CategoryTooltip = IDocumentation::Get()->CreateToolTip(LOCTEXT("EditCategoryName_Tooltip", "The category of the variable; editing this will place the variable into another category or create a new one."), NULL, "", TEXT("Category"));
+
+		Category.AddCustomRow( LOCTEXT("CategoryLabel", "Category") )
+		.NameContent()
+		[
+			SNew(STextBlock)
+			.Text( LOCTEXT("CategoryLabel", "Category") )
+			.ToolTip(CategoryTooltip)
+			.Font( DetailFontInfo )
+		]
+		.ValueContent()
+		[
+			SAssignNew(NewComboButton, SComboButton)
+			.ContentPadding(FMargin(0.0f, 0.0f, 5.0f, 0.0f))
+			.ToolTip(CategoryTooltip)
+			.ButtonContent()
+			[
+				SNew(SBorder)
+				.BorderImage(FAppStyle::Get().GetBrush("NoBorder") )
+				.Padding(FMargin(0, 0, 5, 0))
 				[
-					SNew(STextBlock)
-					.Visibility(IsAnimationFlagEnabled(PropertyBeingCustomized) ? EVisibility::Visible : EVisibility::Hidden)
-					.Font(DetailLayout.GetDetailFont())
-					.Text(AnimationOutputText)
-					.ToolTipText(AnimationOutputTooltipText)
+					SNew(SEditableTextBox)
+						.Text(this, &FRigVMVariableDetailCustomization::OnGetCategoryText)
+						.OnTextCommitted(this, &FRigVMVariableDetailCustomization::OnCategoryTextCommitted, CachedVariableName )
+						.OnVerifyTextChanged_Lambda([&](const FText& InNewText, FText& OutErrorMessage) -> bool
+						{
+							if (InNewText.IsEmpty())
+							{
+								OutErrorMessage = LOCTEXT("CategoryEmpty", "Cannot add a category with an empty string.");
+								return false;
+							}
+							if (InNewText.EqualTo(FText::FromString(BlueprintPtr->GetName())))
+							{
+								OutErrorMessage = LOCTEXT("CategoryEqualsBlueprintName", "Cannot add a category with the same name as the blueprint.");
+								return false;
+							}
+							return true;
+						})
+						.ToolTip(CategoryTooltip)
+						.SelectAllTextWhenFocused(true)
+						.RevertTextOnEscape(true)
+						.Font( DetailFontInfo )
 				]
-				.ValueContent()
+			]
+			.MenuContent()
+			[
+				SNew(SVerticalBox)
+				+SVerticalBox::Slot()
+				.AutoHeight()
+				.MaxHeight(400.0f)
 				[
-					SNew(SCheckBox)
-					.Visibility(IsAnimationFlagEnabled(PropertyBeingCustomized) ? EVisibility::Visible : EVisibility::Hidden)
-					.IsChecked_Raw(this, &FRigVMVariableDetailCustomization::IsAnimationOutputChecked, PropertyBeingCustomized)
-					.OnCheckStateChanged_Raw(this, &FRigVMVariableDetailCustomization::HandleAnimationOutputCheckStateChanged, PropertyBeingCustomized)
-					.ToolTipText(AnimationOutputTooltipText)
-				];
+					SAssignNew(NewListView, SListView<TSharedPtr<FText>>)
+						.ListItemsSource(&CategorySource)
+						.OnGenerateRow(this, &FRigVMVariableDetailCustomization::MakeCategoryViewWidget)
+						.OnSelectionChanged(this, &FRigVMVariableDetailCustomization::OnCategorySelectionChanged)
+				]
+			]
+		];
+
+		CategoryComboButton = NewComboButton;
+		CategoryListView = NewListView;
+
+
+		TSharedPtr<SToolTip> SliderRangeTooltip = IDocumentation::Get()->CreateToolTip(LOCTEXT("SliderRange_Tooltip", "Allows setting the minimum and maximum values for the UI slider for this variable."), NULL, "", TEXT("SliderRange"));
+
+		FName UIMin = TEXT("UIMin");
+		FName UIMax = TEXT("UIMax");
+		Category.AddCustomRow( LOCTEXT("SliderRangeLabel", "Slider Range") )
+		.Visibility(TAttribute<EVisibility>(this, &FRigVMVariableDetailCustomization::RangeVisibility))
+		.NameContent()
+		[
+			SNew(STextBlock)
+			.Text( LOCTEXT("SliderRangeLabel", "Slider Range") )
+			.ToolTip(SliderRangeTooltip)
+			.Font( DetailFontInfo )
+		]
+		.ValueContent()
+		[
+			SNew(SHorizontalBox)
+			.ToolTip(SliderRangeTooltip)
+			+SHorizontalBox::Slot()
+			.FillWidth(1)
+			[
+				SNew(SEditableTextBox)
+				.Text(this, &FRigVMVariableDetailCustomization::OnGetMetaKeyValue, UIMin)
+				.OnTextCommitted(this, &FRigVMVariableDetailCustomization::OnMetaKeyValueChanged, UIMin)
+				.Font( DetailFontInfo )
+			]
+			+SHorizontalBox::Slot()
+			.AutoWidth()
+			[
+				SNew( STextBlock )
+				.Text( LOCTEXT("Min .. Max Separator", " .. ") )
+				.Font(DetailFontInfo)
+			]
+			+SHorizontalBox::Slot()
+			.FillWidth(1)
+			[
+				SNew(SEditableTextBox)
+				.Text(this, &FRigVMVariableDetailCustomization::OnGetMetaKeyValue, UIMax)
+				.OnTextCommitted(this, &FRigVMVariableDetailCustomization::OnMetaKeyValueChanged, UIMax)
+				.Font( DetailFontInfo )
+			]
+		];
+
+		TSharedPtr<SToolTip> ValueRangeTooltip = IDocumentation::Get()->CreateToolTip(LOCTEXT("ValueRangeLabel_Tooltip", "The range of values allowed by this variable. Values outside of this will be clamped to the range."), NULL, "", TEXT("ValueRange"));
 
-			DetailLayout.EditCategory("Variable")
-				.AddCustomRow(AnimationOutputText)
-				.NameContent()
+		FName ClampMin = TEXT("ClampMin");
+		FName ClampMax = TEXT("ClampMax");
+		Category.AddCustomRow(LOCTEXT("ValueRangeLabel", "Value Range"))
+		.Visibility(TAttribute<EVisibility>(this, &FRigVMVariableDetailCustomization::RangeVisibility))
+		.NameContent()
+		[
+			SNew(STextBlock)
+			.Text(LOCTEXT("ValueRangeLabel", "Value Range"))
+			.ToolTipText(LOCTEXT("ValueRangeLabel_Tooltip", "The range of values allowed by this variable. Values outside of this will be clamped to the range."))
+			.ToolTip(ValueRangeTooltip)
+			.Font(DetailFontInfo)
+		]
+		.ValueContent()
+		[
+			SNew(SHorizontalBox)
+			+ SHorizontalBox::Slot()
+			.FillWidth(1)
+			[
+				SNew(SEditableTextBox)
+				.Text(this, &FRigVMVariableDetailCustomization::OnGetMetaKeyValue, ClampMin)
+				.OnTextCommitted(this, &FRigVMVariableDetailCustomization::OnMetaKeyValueChanged, ClampMin)
+				.Font(DetailFontInfo)
+			]
+			+ SHorizontalBox::Slot()
+			.AutoWidth()
+			[
+				SNew(STextBlock)
+				.Text(LOCTEXT("Min .. Max Separator", " .. "))
+				.Font(DetailFontInfo)
+			]
+			+ SHorizontalBox::Slot()
+			.FillWidth(1)
+			[
+				SNew(SEditableTextBox)
+				.Text(this, &FRigVMVariableDetailCustomization::OnGetMetaKeyValue, ClampMax)
+				.OnTextCommitted(this, &FRigVMVariableDetailCustomization::OnMetaKeyValueChanged, ClampMax)
+				.Font(DetailFontInfo)
+			]
+		];
+
+		TStrongObjectPtr<URigVMBlueprint> Blueprint = BlueprintPtr.Pin();
+
+		// Add in default value editing for properties that can be edited, local properties cannot be edited
+		if (Blueprint.IsValid() && (Blueprint->GeneratedClass != nullptr))
+		{
+			bool bVariableRenamed = false;
+			if (VariableProperty != nullptr)
+			{
+				const int32 VarIndex = FBlueprintEditorUtils::FindNewVariableIndex(Blueprint.Get(), CachedVariableName);
+				if (VarIndex != INDEX_NONE)
+				{
+					const FGuid VarGuid = Blueprint->NewVariables[VarIndex].VarGuid;
+					if (UBlueprintGeneratedClass* AuthoritiveBPGC = Cast<UBlueprintGeneratedClass>(Blueprint->GeneratedClass))
+					{
+						if (const FName* OldName = AuthoritiveBPGC->PropertyGuids.FindKey(VarGuid))
+						{
+							bVariableRenamed = CachedVariableName != *OldName;
+						}
+					}
+				}
+				
+				const FProperty* OriginalProperty = nullptr;
+				OriginalProperty = FindFProperty<FProperty>(Blueprint->GeneratedClass, VariableProperty->GetFName());
+
+				if (OriginalProperty == nullptr || bVariableRenamed)
+				{
+					// Prevent editing the default value of a skeleton property
+					VariableProperty = nullptr;
+				}
+				else if (const FStructProperty* StructProperty = CastField<const FStructProperty>(OriginalProperty))
+				{
+					// Prevent editing the default value of a stale struct
+					const UUserDefinedStruct* BGStruct = Cast<const UUserDefinedStruct>(StructProperty->Struct);
+					if (BGStruct && (EUserDefinedStructureStatus::UDSS_UpToDate != BGStruct->Status))
+					{
+						VariableProperty = nullptr;
+					}
+				}
+			}
+
+			// Find the class containing the variable
+			UClass* VariableClass = (VariableProperty ? VariableProperty->GetTypedOwner<UClass>() : nullptr);
+
+			FText ErrorMessage;
+			IDetailCategoryBuilder& DefaultValueCategory = DetailLayout.EditCategory(TEXT("DefaultValueCategory"), LOCTEXT("DefaultValueCategoryHeading", "Default Value"));
+
+			if (VariableProperty == nullptr)
+			{
+				if (Blueprint->Status != BS_UpToDate)
+				{
+					ErrorMessage = LOCTEXT("VariableMissing_DirtyBlueprint", "Please compile the blueprint");
+				}
+				else
+				{
+					ErrorMessage = LOCTEXT("VariableMissing_CleanBlueprint", "Failed to find variable property");
+				}
+			}
+			// Show the error message if something went wrong
+			if (!ErrorMessage.IsEmpty())
+			{
+				DefaultValueCategory.AddCustomRow( ErrorMessage )
 				[
 					SNew(STextBlock)
-					.Visibility(IsAnimationFlagEnabled(PropertyBeingCustomized) ? EVisibility::Visible : EVisibility::Hidden)
-					.Font(DetailLayout.GetDetailFont())
-					.Text(AnimationInputText)
-					.ToolTipText(AnimationInputTooltipText)
-				]
-				.ValueContent()
-				[
-					SNew(SCheckBox)
-					.Visibility(IsAnimationFlagEnabled(PropertyBeingCustomized) ? EVisibility::Visible : EVisibility::Hidden)
-					.IsChecked_Raw(this, &FRigVMVariableDetailCustomization::IsAnimationInputChecked, PropertyBeingCustomized)
-					.OnCheckStateChanged_Raw(this, &FRigVMVariableDetailCustomization::HandleAnimationInputCheckStateChanged, PropertyBeingCustomized)
-					.ToolTipText(AnimationInputTooltipText)
+					.ToolTipText(ErrorMessage)
+					.Text(ErrorMessage)
+					.Font(DetailFontInfo)
 				];
+			}
+			else 
+			{
+				TSharedPtr<IDetailsView> DetailsView;
+				if (Editor.IsValid())
+				{
+					DetailsView = Editor->GetRigVMInspector()->GetPropertyView();
+				}
+
+				{
+					// Things are in order, show the property and allow it to be edited
+					TArray<UObject*> ObjectList;
+					ObjectList.Add(Blueprint->GeneratedClass->GetDefaultObject());
+					IDetailPropertyRow* Row = DefaultValueCategory.AddExternalObjectProperty(ObjectList, VariableProperty->GetFName());
+
+					// if (DetailsView.IsValid())
+					// {
+					// 	DetailsView->OnFinishedChangingProperties().AddSP(this, &FRigVMVariableDetailCustomization::OnFinishedChangingVariable);
+					// }
+				}
+			}
 		}
 	}
 }
 
-bool FRigVMVariableDetailCustomization::IsAnimationFlagEnabled(TWeakFieldPtr<FProperty> PropertyBeingCustomized) const
+void FRigVMVariableDetailCustomization::PopulateCategories()
 {
+	TSharedPtr<FRigVMEditorBase> Editor = EditorPtr.Pin();
+	TStrongObjectPtr<URigVMBlueprint> Blueprint = BlueprintPtr.Pin();
+	
+	if (!Blueprint || !Editor)
+	{
+		return;
+	}
+	
+	auto IsNewCategorySource = [&](const FText& NewCategory)
+	{
+		return !CategorySource.ContainsByPredicate([&NewCategory](const TSharedPtr<FText>& ExistingCategory)
+		{
+			return ExistingCategory->ToString().Equals(NewCategory.ToString(), ESearchCase::CaseSensitive);
+		});
+	};
+
+	CategorySource.Reset();
+	CategorySource.Add(MakeShared<FText>(UEdGraphSchema_K2::VR_DefaultCategory));
+
+
+	TArray<FRigVMGraphVariableDescription> Variables = Blueprint->GetMemberVariables();
+	
+	for (const FRigVMGraphVariableDescription& Variable : Variables)
+	{
+		FText Category = Variable.Category;
+		if (!Category.IsEmpty() && !Category.EqualTo(FText::FromString(BlueprintPtr.Pin()->GetName())))
+		{
+			if (IsNewCategorySource(Category))
+			{
+				CategorySource.Add(MakeShared<FText>(Category));
+			}
+		}
+	}
+
+	// Sort categories, but keep the default category listed first
+	CategorySource.Sort([](const TSharedPtr <FText> &LHS, const TSharedPtr <FText> &RHS)
+	{
+		if (LHS.IsValid() && RHS.IsValid())
+		{
+			return (LHS->EqualTo(UEdGraphSchema_K2::VR_DefaultCategory) || LHS->CompareToCaseIgnored(*RHS) <= 0);
+		}
+		return false;
+	});
+}
+
+void FRigVMVariableDetailCustomization::OnPostEditorRefresh()
+{
+	FName VariableName = GetVariableName();
+	CachedVariableProperty = BlueprintPtr->SkeletonGeneratedClass->FindPropertyByName(VariableName);
+}
+
+FName FRigVMVariableDetailCustomization::GetVariableName() const
+{
+	TSharedPtr<FRigVMEditorBase> Editor = EditorPtr.Pin();
+    if (!Editor.IsValid())
+    {
+    	return NAME_None;
+    }
+    
+    CachedVariableName = Editor->GetGraphExplorerWidget()->GetSelectedVariableName();
+	return CachedVariableName;
+}
+
+FText FRigVMVariableDetailCustomization::OnGetVariableName() const
+{
+	return FText::FromName(CachedVariableName);
+}
+
+void FRigVMVariableDetailCustomization::OnVarNameCommitted(const FText& InNewName, ETextCommit::Type InTextCommit)
+{
+	TSharedPtr<FRigVMEditorBase> Editor = EditorPtr.Pin();
+	if (!Editor)
+	{
+		return;
+	}
+
+	TStrongObjectPtr<URigVMBlueprint> Blueprint = BlueprintPtr.Pin();
+	if (!Blueprint)
+	{
+		return;
+	}
+
+	TSharedPtr<SRigVMEditorGraphExplorer> Explorer = Editor->GetGraphExplorerWidget();
+	if (!Explorer)
+	{
+		return;
+	}
+	
+	ERigVMExplorerElementType::Type SelectedType = Explorer->GetSelectedType();
+	if (SelectedType == ERigVMExplorerElementType::Variable)
+	{
+		Blueprint->Modify();
+		FBlueprintEditorUtils::RenameMemberVariable(Blueprint.Get(), CachedVariableName, *InNewName.ToString());
+	}
+}
+
+FEdGraphPinType FRigVMVariableDetailCustomization::OnGetVarType() const
+{
+	FEdGraphPinType Type;
+	TSharedPtr<FRigVMEditorBase> Editor = EditorPtr.Pin();
+	if (!Editor)
+	{
+		return Type;
+	}
+
+	TStrongObjectPtr<URigVMBlueprint> Blueprint = BlueprintPtr.Pin();
+	if (!Blueprint)
+	{
+		return Type;
+	}
+
+	TSharedPtr<SRigVMEditorGraphExplorer> Explorer = Editor->GetGraphExplorerWidget();
+	if (!Explorer)
+	{
+		return Type;
+	}
+
+	ERigVMExplorerElementType::Type SelectionVarType = Explorer->GetSelectedType();
+	if (SelectionVarType == ERigVMExplorerElementType::Variable)
+	{
+		for (FBPVariableDescription& Var : Blueprint->NewVariables)
+		{
+			if (Var.VarName == CachedVariableName)
+			{
+				return Var.VarType;
+			}
+		}
+		return Type;
+	}
+	return Type;
+}
+
+void FRigVMVariableDetailCustomization::OnVarTypeChanged(const FEdGraphPinType& NewPinType)
+{
+	TSharedPtr<FRigVMEditorBase> Editor = EditorPtr.Pin();
+	if (!Editor)
+	{
+		return;
+	}
+
+	TStrongObjectPtr<URigVMBlueprint> Blueprint = BlueprintPtr.Pin();
+	if (!Blueprint)
+	{
+		return;
+	}
+
+	TSharedPtr<SRigVMEditorGraphExplorer> Explorer = Editor->GetGraphExplorerWidget();
+	if (!Explorer)
+	{
+		return;
+	}
+	
+	ERigVMExplorerElementType::Type SelectedType = Explorer->GetSelectedType();
+	if (SelectedType == ERigVMExplorerElementType::Variable)
+	{
+		for (FRigVMGraphVariableDescription& Variable : Blueprint->GetMemberVariables())
+		{
+			if (Variable.Name == CachedVariableName)
+			{
+				FBlueprintEditorUtils::ChangeMemberVariableType(Blueprint.Get(), CachedVariableName, NewPinType);
+				Explorer->SetLastPinTypeUsed(NewPinType);
+				return;
+			}
+		}
+	}
+}
+
+void FRigVMVariableDetailCustomization::OnBrowseToVarType() const
+{
+	FEdGraphPinType PinType = OnGetVarType();
+	if (const UObject* Object = PinType.PinSubCategoryObject.Get())
+	{
+		if (Object->IsAsset())
+		{
+			FAssetData AssetData(Object, false);
+			if (AssetData.IsValid())
+			{
+				TArray<FAssetData> AssetDataList = { AssetData };
+				GEditor->SyncBrowserToObjects(AssetDataList);
+			}
+		}
+	}
+}
+
+bool FRigVMVariableDetailCustomization::CanBrowseToVarType() const
+{
+	FEdGraphPinType PinType = OnGetVarType();
+	if (const UObject* Object = PinType.PinSubCategoryObject.Get())
+	{
+		if (Object->IsAsset())
+		{
+			FAssetData AssetData(Object, false);
+			if (AssetData.IsValid())
+			{
+				return true;
+			}
+		}
+	}
+
 	return false;
 }
 
-ECheckBoxState FRigVMVariableDetailCustomization::IsAnimationOutputChecked(TWeakFieldPtr<FProperty> PropertyBeingCustomized) const
+FText FRigVMVariableDetailCustomization::OnGetTooltipText() const
+{
+	FText ToolTip;
+	TSharedPtr<FRigVMEditorBase> Editor = EditorPtr.Pin();
+	if (!Editor)
+	{
+		return ToolTip;
+	}
+
+	TStrongObjectPtr<URigVMBlueprint> Blueprint = BlueprintPtr.Pin();
+	if (!Blueprint)
+	{
+		return ToolTip;
+	}
+
+	TSharedPtr<SRigVMEditorGraphExplorer> Explorer = Editor->GetGraphExplorerWidget();
+	if (!Explorer)
+	{
+		return ToolTip;
+	}
+
+	ERigVMExplorerElementType::Type SelectedType = Explorer->GetSelectedType();
+
+	if (SelectedType == ERigVMExplorerElementType::Variable)
+	{
+		FString Result;
+		FBlueprintEditorUtils::GetBlueprintVariableMetaData(Blueprint.Get(), CachedVariableName, nullptr, TEXT("tooltip"), Result);
+		return FText::FromString(Result);
+	}
+
+	return ToolTip;
+}
+
+void FRigVMVariableDetailCustomization::OnTooltipTextCommitted(const FText& NewText, ETextCommit::Type InTextCommit, FName VarName)
+{
+	TSharedPtr<FRigVMEditorBase> Editor = EditorPtr.Pin();
+	if (!Editor)
+	{
+		return;
+	}
+
+	TStrongObjectPtr<URigVMBlueprint> Blueprint = BlueprintPtr.Pin();
+	if (!Blueprint)
+	{
+		return;
+	}
+
+	TSharedPtr<SRigVMEditorGraphExplorer> Explorer = Editor->GetGraphExplorerWidget();
+	if (!Explorer)
+	{
+		return;
+	}
+
+	ERigVMExplorerElementType::Type SelectedType = Explorer->GetSelectedType();
+
+	if (SelectedType == ERigVMExplorerElementType::Variable)
+	{
+		FBlueprintEditorUtils::SetBlueprintVariableMetaData(Blueprint.Get(), CachedVariableName, nullptr, TEXT("tooltip"), NewText.ToString() );
+	}
+}
+
+EVisibility FRigVMVariableDetailCustomization::IsToolTipVisible() const
+{
+	EVisibility Visible = EVisibility::Collapsed;
+	TSharedPtr<FRigVMEditorBase> Editor = EditorPtr.Pin();
+	if (!Editor)
+	{
+		return Visible;
+	}
+	TSharedPtr<SRigVMEditorGraphExplorer> Explorer = Editor->GetGraphExplorerWidget();
+	if (!Explorer)
+	{
+		return Visible;
+	}
+	ERigVMExplorerElementType::Type SelectedType = Explorer->GetSelectedType();
+	if (SelectedType == ERigVMExplorerElementType::Variable)
+	{
+		Visible = EVisibility::Visible;
+	}
+	return Visible;
+}
+
+FText FRigVMVariableDetailCustomization::OnGetCategoryText() const
+{
+	TSharedPtr<FRigVMEditorBase> Editor = EditorPtr.Pin();
+	if (!Editor)
+	{
+		return FText();
+	}
+
+	TStrongObjectPtr<URigVMBlueprint> Blueprint = BlueprintPtr.Pin();
+	if (!Blueprint)
+	{
+		return FText();
+	}
+
+	if (CachedVariableName != NAME_None)
+	{
+		FText Category = FBlueprintEditorUtils::GetBlueprintVariableCategory(Blueprint.Get(), CachedVariableName, nullptr);
+
+		// Older blueprints will have their name as the default category and whenever it is the same as the default category, display localized text
+		if ( Category.EqualTo(FText::FromString(Blueprint->GetName())) || Category.EqualTo(UEdGraphSchema_K2::VR_DefaultCategory) )
+		{
+			return UEdGraphSchema_K2::VR_DefaultCategory;
+		}
+		else
+		{
+			return Category;
+		}
+	}
+	return FText();
+}
+
+void FRigVMVariableDetailCustomization::OnCategoryTextCommitted(const FText& NewText, ETextCommit::Type InTextCommit, FName VarName)
+{
+	TStrongObjectPtr<URigVMBlueprint> Blueprint = BlueprintPtr.Pin();
+	if (!Blueprint)
+	{
+		return;
+	}
+	
+	if (InTextCommit == ETextCommit::OnEnter || InTextCommit == ETextCommit::OnUserMovedFocus)
+	{
+		FBlueprintEditorUtils::SetBlueprintVariableCategory(Blueprint.Get(), VarName, nullptr, NewText);
+		PopulateCategories();
+	}
+}
+
+TSharedRef<ITableRow> FRigVMVariableDetailCustomization::MakeCategoryViewWidget(TSharedPtr<FText> Item, const TSharedRef<STableViewBase>& OwnerTable)
 {
+	return SNew(STableRow<TSharedPtr<FString>>, OwnerTable)
+		[
+			SNew(STextBlock) .Text(*Item.Get())
+		];
+}
+
+void FRigVMVariableDetailCustomization::OnCategorySelectionChanged(TSharedPtr<FText> ProposedSelection, ESelectInfo::Type)
+{
+	TStrongObjectPtr<URigVMBlueprint> Blueprint = BlueprintPtr.Pin();
+	if (!Blueprint)
+	{
+		return;
+	}
+	
+	if (ProposedSelection.IsValid() && CachedVariableName != NAME_None)
+	{
+		FText NewCategory = *ProposedSelection.Get();
+
+		FBlueprintEditorUtils::SetBlueprintVariableCategory(Blueprint.Get(), CachedVariableName, nullptr, NewCategory );
+		CategoryListView.Pin()->ClearSelection();
+		CategoryComboButton.Pin()->SetIsOpen(false);
+	}
+}
+
+ECheckBoxState FRigVMVariableDetailCustomization::OnGetExposedToSpawnCheckboxState() const
+{
+	FProperty* Property = CachedVariableProperty.Get();
+	if (Property)
+	{
+		return (Property && Property->GetBoolMetaData(FBlueprintMetadata::MD_ExposeOnSpawn) != false) ? ECheckBoxState::Checked : ECheckBoxState::Unchecked;
+	}
 	return ECheckBoxState::Unchecked;
 }
 
-void FRigVMVariableDetailCustomization::HandleAnimationOutputCheckStateChanged(ECheckBoxState CheckBoxState, TWeakFieldPtr<FProperty> PropertyBeingCustomized)
+void FRigVMVariableDetailCustomization::OnExposedToSpawnChanged(ECheckBoxState InNewState)
 {
+	TStrongObjectPtr<URigVMBlueprint> Blueprint = BlueprintPtr.Pin();
+	if (!Blueprint)
+	{
+		return;
+	}
+	
+	if (CachedVariableName != NAME_None)
+	{
+		const bool bExposeOnSpawn = (InNewState == ECheckBoxState::Checked);
+		if(bExposeOnSpawn)
+		{
+			FBlueprintEditorUtils::SetBlueprintVariableMetaData(Blueprint.Get(), CachedVariableName, NULL, FBlueprintMetadata::MD_ExposeOnSpawn, TEXT("true"));
+		}
+		else
+		{
+			FBlueprintEditorUtils::RemoveBlueprintVariableMetaData(Blueprint.Get(), CachedVariableName, NULL, FBlueprintMetadata::MD_ExposeOnSpawn);
+		} 
+	}
 }
 
-ECheckBoxState FRigVMVariableDetailCustomization::IsAnimationInputChecked(TWeakFieldPtr<FProperty> PropertyBeingCustomized) const
+ECheckBoxState FRigVMVariableDetailCustomization::OnGetPrivateCheckboxState() const
 {
+	FProperty* Property = CachedVariableProperty.Get();
+	if (Property)
+	{
+		return (Property && Property->GetBoolMetaData(FBlueprintMetadata::MD_Private) != false) ? ECheckBoxState::Checked : ECheckBoxState::Unchecked;
+	}
 	return ECheckBoxState::Unchecked;
 }
 
-void FRigVMVariableDetailCustomization::HandleAnimationInputCheckStateChanged(ECheckBoxState CheckBoxState, TWeakFieldPtr<FProperty> PropertyBeingCustomized)
+void FRigVMVariableDetailCustomization::OnPrivateChanged(ECheckBoxState InNewState)
 {
+	TStrongObjectPtr<URigVMBlueprint> Blueprint = BlueprintPtr.Pin();
+	if (!Blueprint)
+	{
+		return;
+	}
+	
+	if (CachedVariableName != NAME_None)
+	{
+		const bool bExposeOnSpawn = (InNewState == ECheckBoxState::Checked);
+		if(bExposeOnSpawn)
+		{
+			FBlueprintEditorUtils::SetBlueprintVariableMetaData(Blueprint.Get(), CachedVariableName, NULL, FBlueprintMetadata::MD_Private, TEXT("true"));
+		}
+		else
+		{
+			FBlueprintEditorUtils::RemoveBlueprintVariableMetaData(Blueprint.Get(), CachedVariableName, NULL, FBlueprintMetadata::MD_Private);
+		}
+	}
+}
+
+ECheckBoxState FRigVMVariableDetailCustomization::OnGetExposedToCinematicsCheckboxState() const
+{
+	FProperty* Property = CachedVariableProperty.Get();
+	if (Property)
+	{
+		return Property && Property->HasAnyPropertyFlags(CPF_Interp) ? ECheckBoxState::Checked : ECheckBoxState::Unchecked;
+	}
+	return ECheckBoxState::Unchecked;
+}
+
+void FRigVMVariableDetailCustomization::OnExposedToCinematicsChanged(ECheckBoxState InNewState)
+{
+	TStrongObjectPtr<URigVMBlueprint> Blueprint = BlueprintPtr.Pin();
+	if (!Blueprint)
+	{
+		return;
+	}
+	
+	// Toggle the flag on the blueprint's version of the variable description, based on state
+	const bool bExposeToCinematics = (InNewState == ECheckBoxState::Checked);
+	
+	const FName VarName = CachedVariableName;
+	if (VarName != NAME_None)
+	{
+		FBlueprintEditorUtils::SetInterpFlag(Blueprint.Get(), VarName, bExposeToCinematics);
+	}
+}
+
+FText FRigVMVariableDetailCustomization::OnGetMetaKeyValue(FName Key) const
+{
+	TStrongObjectPtr<URigVMBlueprint> Blueprint = BlueprintPtr.Pin();
+	if (!Blueprint)
+	{
+		return FText();
+	}
+	
+	if (CachedVariableName != NAME_None)
+	{
+		FString Result;
+		FBlueprintEditorUtils::GetBlueprintVariableMetaData(Blueprint.Get(), CachedVariableName, nullptr, Key, /*out*/ Result);
+
+		return FText::FromString(Result);
+	}
+	return FText();
+}
+
+void FRigVMVariableDetailCustomization::OnMetaKeyValueChanged(const FText& NewMinValue, ETextCommit::Type CommitInfo, FName Key)
+{
+	TStrongObjectPtr<URigVMBlueprint> Blueprint = BlueprintPtr.Pin();
+	if (!Blueprint)
+	{
+		return;
+	}
+	
+	if (CachedVariableName != NAME_None)
+	{
+		if ((CommitInfo == ETextCommit::OnEnter) || (CommitInfo == ETextCommit::OnUserMovedFocus))
+		{
+			FBlueprintEditorUtils::SetBlueprintVariableMetaData(Blueprint.Get(), CachedVariableName, nullptr, Key, NewMinValue.ToString());
+		}
+	}
+}
+
+EVisibility FRigVMVariableDetailCustomization::RangeVisibility() const
+{
+	FProperty* VariableProperty = CachedVariableProperty.Get();
+	if (VariableProperty)
+	{
+		const bool bIsInteger = VariableProperty->IsA(FIntProperty::StaticClass());
+		const bool bIsNonEnumByte = (VariableProperty->IsA(FByteProperty::StaticClass()) && CastField<const FByteProperty>(VariableProperty)->Enum == nullptr);
+		const bool bIsReal = (VariableProperty->IsA(FFloatProperty::StaticClass()) || VariableProperty->IsA(FDoubleProperty::StaticClass()));
+
+		// If this is a struct property than we must check the name of the struct it points to, so we can check
+		// if it supports the editing of the UIMin/UIMax metadata
+		const FStructProperty* StructProp = CastField<FStructProperty>(VariableProperty);
+		const UStruct* InnerStruct = StructProp ? StructProp->Struct : nullptr;
+		const bool bIsSupportedStruct = InnerStruct ? RangeVisibilityUtils::StructsSupportingRangeVisibility.Contains(InnerStruct->GetFName()) : false;
+
+		if (bIsInteger || bIsNonEnumByte || bIsReal || bIsSupportedStruct)
+		{
+			return EVisibility::Visible;
+		}
+	}
+	return EVisibility::Collapsed;
 }
 
 #undef LOCTEXT_NAMESPACE
diff --git a/Engine/Plugins/Runtime/RigVM/Source/RigVMEditor/Private/Editor/SRigVMDetailsInspector.cpp b/Engine/Plugins/Runtime/RigVM/Source/RigVMEditor/Private/Editor/SRigVMDetailsInspector.cpp
--- a/Engine/Plugins/Runtime/RigVM/Source/RigVMEditor/Private/Editor/SRigVMDetailsInspector.cpp
+++ b/Engine/Plugins/Runtime/RigVM/Source/RigVMEditor/Private/Editor/SRigVMDetailsInspector.cpp
@@ -1,3 +1,653 @@
 // Copyright Epic Games, Inc. All Rights Reserved.
 
-#include "Editor/SRigVMDetailsInspector.h"
\ No newline at end of file
+#include "Editor/SRigVMDetailsInspector.h"
+
+#include "Editor/RigVMEditor.h"
+#include "Engine/SCS_Node.h"
+#include "Widgets/Text/SRichTextBlock.h"
+#include "Widgets/Input/SCheckBox.h"
+#include "PropertyEditorModule.h"
+#include "Editor/RigVMGraphDetailCustomization.h"
+#include "Editor/RigVMVariableDetailCustomization.h"
+#include "Framework/Application/SlateApplication.h"
+#include "K2Node.h"
+
+class FNotifyHook;
+class FStructOnScope;
+class IClassViewerFilter;
+class IDetailLayoutBuilder;
+class SDockTab;
+class SWidget;
+struct FGeometry;
+
+#define LOCTEXT_NAMESPACE "RigVMDetailsInspector"
+
+//////////////////////////////////////////////////////////////////////////
+// SRigVMDetailsInspector
+
+void SRigVMDetailsInspector::Tick( const FGeometry& AllottedGeometry, const double InCurrentTime, const float InDeltaTime )
+{
+	if(bRefreshOnTick)
+	{
+		// if struct is valid, update struct
+		if (StructToDisplay.IsValid())
+		{
+			UpdateFromSingleStruct(StructToDisplay);
+			StructToDisplay.Reset();
+		}
+		else
+		{
+			TArray<UObject*> SelectionInfo;
+			RefreshPropertyObjects.Remove(nullptr);
+			UpdateFromObjects(RefreshPropertyObjects, SelectionInfo, RefreshOptions);
+			RefreshPropertyObjects.Empty();
+		}
+
+		bRefreshOnTick = false;
+	}
+}
+
+TSharedRef<SWidget> SRigVMDetailsInspector::MakeContextualEditingWidget(TArray<UObject*>& SelectionInfo, const FShowDetailsOptions& Options)
+{
+	TSharedRef< SVerticalBox > ContextualEditingWidget = SNew( SVerticalBox );
+
+	if(bShowTitleArea)
+	{
+		if (SelectedObjects.Num() == 0)
+		{
+			// Warning about nothing being selected
+			ContextualEditingWidget->AddSlot()
+			.AutoHeight()
+			.HAlign( HAlign_Center )
+			.Padding( 2.0f, 14.0f, 2.0f, 2.0f )
+			[
+				SNew( STextBlock )
+				.Text( LOCTEXT("NoNodesSelected", "Select a node to edit details.") )
+			];
+		}
+		else
+		{
+			// Title of things being edited
+			ContextualEditingWidget->AddSlot()
+			.AutoHeight()
+			.Padding( 2.0f, 0.0f, 2.0f, 0.0f )
+			[
+				SNew(STextBlock)
+				.Text(this, &SRigVMDetailsInspector::GetContextualEditingWidgetTitle)
+			];
+		}
+	}
+
+	// Show the property editor
+	PropertyView->HideFilterArea(Options.bHideFilterArea);
+	PropertyView->SetObjects(SelectionInfo, Options.bForceRefresh);
+
+	if (SelectionInfo.Num())
+	{
+		ContextualEditingWidget->AddSlot()
+		.FillHeight( 0.9f )
+		.VAlign( VAlign_Top )
+		[
+			SNew( SBox )
+			.Visibility(this, &SRigVMDetailsInspector::GetPropertyViewVisibility)
+			[
+				PropertyView.ToSharedRef()
+			]
+		];
+
+		ContextualEditingWidget->AddSlot()
+		.AutoHeight()
+		.VAlign(VAlign_Top)
+		[
+			SNew(SCheckBox)
+			.ToolTipText(LOCTEXT("TogglePublicView", "Toggle Public View"))
+			.IsChecked(this, &SRigVMDetailsInspector::GetPublicViewCheckboxState)
+			.OnCheckStateChanged(this, &SRigVMDetailsInspector::SetPublicViewCheckboxState)
+			[
+				SNew(STextBlock)
+				.Text(LOCTEXT("PublicViewCheckboxLabel", "Public View"))
+			]
+			.Visibility_Lambda([bShowPublicView = this->bShowPublicView]()
+			{
+				return bShowPublicView.Get() ? EVisibility::Visible : EVisibility::Hidden;
+			})
+		];
+	}
+
+	return ContextualEditingWidget;
+}
+
+TSharedPtr<SDockTab> SRigVMDetailsInspector::GetOwnerTab() const
+{
+	return OwnerTab.Pin();
+}
+
+const TArray< TWeakObjectPtr<UObject> >& SRigVMDetailsInspector::GetSelectedObjects() const
+{
+	return SelectedObjects;
+}
+
+void SRigVMDetailsInspector::AddReferencedObjects(FReferenceCollector& Collector)
+{
+	Collector.AddReferencedObjects(RefreshPropertyObjects);
+}
+
+FString SRigVMDetailsInspector::GetReferencerName() const
+{
+	return TEXT("SRigVMDetailsInspector");
+}
+
+FText SRigVMDetailsInspector::GetContextualEditingWidgetTitle() const
+{
+	FText Title = PropertyViewTitle;
+	if (Title.IsEmpty())
+	{
+		if (SelectedObjects.Num() == 1 && SelectedObjects[0].IsValid())
+		{
+			UObject* Object = SelectedObjects[0].Get();
+
+			if (UEdGraphNode* Node = Cast<UEdGraphNode>(Object))
+			{
+				Title = Node->GetNodeTitle(ENodeTitleType::ListView);
+			}
+		}
+		else if (SelectedObjects.Num() > 1)
+		{
+			UClass* BaseClass = nullptr;
+
+			for (auto ObjectWkPtrIt = SelectedObjects.CreateConstIterator(); ObjectWkPtrIt; ++ObjectWkPtrIt)
+			{
+				TWeakObjectPtr<UObject> ObjectWkPtr = *ObjectWkPtrIt;
+				if (ObjectWkPtr.IsValid())
+				{
+					UObject* Object = ObjectWkPtr.Get();
+					UClass* ObjClass = Object->GetClass();
+
+					if (UEdGraphNode* Node = Cast<UEdGraphNode>(Object))
+					{
+						// Hide any specifics of node types; they're all ed graph nodes
+						ObjClass = UEdGraphNode::StaticClass();
+					}
+
+					// Keep track of the class of objects selected
+					if (BaseClass == nullptr)
+					{
+						BaseClass = ObjClass;
+						checkSlow(ObjClass);
+					}
+					while (!ObjClass->IsChildOf(BaseClass))
+					{
+						BaseClass = BaseClass->GetSuperClass();
+					}
+				}
+			}
+
+			if (BaseClass)
+			{
+				Title = FText::Format(LOCTEXT("MultipleObjectsSelectedFmt", "{0} {1} selected"), FText::AsNumber(SelectedObjects.Num()), FText::FromString(BaseClass->GetName() + TEXT("s")));
+			}
+		}
+	}
+	return Title;
+}
+
+void SRigVMDetailsInspector::Construct(const FArguments& InArgs)
+{
+	bShowInspectorPropertyView = true;
+	PublicViewState = ECheckBoxState::Unchecked;
+	bRefreshOnTick = false;
+
+	WeakEditor = InArgs._Editor;
+	bShowPublicView = InArgs._ShowPublicViewControl;
+	bShowTitleArea = InArgs._ShowTitleArea;
+	TSharedPtr<FRigVMEditorBase> Editor = WeakEditor.Pin();
+
+	// Create a property view
+	FPropertyEditorModule& EditModule = FModuleManager::Get().GetModuleChecked<FPropertyEditorModule>("PropertyEditor");
+
+	FNotifyHook* NotifyHook = nullptr;
+	if(InArgs._SetNotifyHook)
+	{
+		NotifyHook = Editor->GetNotifyHook();
+	}
+
+	FDetailsViewArgs DetailsViewArgs;
+	DetailsViewArgs.NameAreaSettings = InArgs._HideNameArea ? FDetailsViewArgs::HideNameArea : FDetailsViewArgs::ObjectsUseNameArea;
+	DetailsViewArgs.bHideSelectionTip = true;
+	DetailsViewArgs.NotifyHook = NotifyHook;
+	DetailsViewArgs.ViewIdentifier = InArgs._ViewIdentifier;
+	DetailsViewArgs.ExternalScrollbar = InArgs._ExternalScrollbar;
+	DetailsViewArgs.ScrollbarAlignment = InArgs._ScrollbarAlignment;
+	DetailsViewArgs.bShowSectionSelector = InArgs._ShowSectionSelector;
+
+	PropertyView = EditModule.CreateDetailView( DetailsViewArgs );
+
+	PropertyView->SetIsPropertyVisibleDelegate( FIsPropertyVisible::CreateSP(this, &SRigVMDetailsInspector::IsPropertyVisible) );
+	PropertyView->SetIsPropertyEditingEnabledDelegate(FIsPropertyEditingEnabled::CreateSP(this, &SRigVMDetailsInspector::IsPropertyEditingEnabled));
+
+	IsPropertyEditingEnabledDelegate = InArgs._IsPropertyEditingEnabledDelegate;
+	UserOnFinishedChangingProperties = InArgs._OnFinishedChangingProperties;
+
+	TWeakPtr<SRigVMEditorGraphExplorer> GraphExplorer = Editor.IsValid() ? Editor->GetGraphExplorerWidget() : InArgs._GraphExplorerWidget;
+	
+	if( GraphExplorer.IsValid() )
+	{
+		const UClass* BlueprintClass = URigVMBlueprint::StaticClass();
+		FOnGetDetailCustomizationInstance FunctionDetails = FOnGetDetailCustomizationInstance::CreateStatic(&FRigVMGraphDetailCustomization::MakeInstance, Editor, BlueprintClass);
+		PropertyView->RegisterInstancedCustomPropertyLayout(URigVMEdGraph::StaticClass(), FunctionDetails);
+
+		FOnGetDetailCustomizationInstance LayoutVariableDetails = FOnGetDetailCustomizationInstance::CreateStatic(&FRigVMVariableDetailCustomization::MakeInstance, Editor);
+		PropertyView->RegisterInstancedCustomPropertyLayout(UPropertyWrapper::StaticClass(), LayoutVariableDetails);
+	}
+
+	// Create the border that all of the content will get stuffed into
+	ChildSlot
+	[
+		SNew(SVerticalBox)
+		.AddMetaData<FTagMetaData>(FTagMetaData(TEXT("RigVMInspector")))
+		+ SVerticalBox::Slot()
+		.FillHeight(1.0f)
+		[
+			SAssignNew( ContextualEditingBorderWidget, SBorder )
+			.Padding(0.0f)
+			.BorderImage( FAppStyle::GetBrush("NoBorder") )
+		]
+		+SVerticalBox::Slot()
+		[
+			SNew(STextBlock).Text(LOCTEXT("RigVMDetailsInspectorTest", "Test"))
+		]
+	];
+
+	// Update based on the current (empty) selection set
+	TArray<UObject*> InitialSelectedObjects;
+	TArray<UObject*> SelectionInfo;
+	UpdateFromObjects(InitialSelectedObjects, SelectionInfo, SRigVMDetailsInspector::FShowDetailsOptions(FText::GetEmpty(), true));
+
+	// create struct to display
+	FStructureDetailsViewArgs StructureViewArgs;
+	StructureViewArgs.bShowObjects = true;
+	StructureViewArgs.bShowAssets = true;
+	StructureViewArgs.bShowClasses = true;
+	StructureViewArgs.bShowInterfaces = true;
+
+	FDetailsViewArgs ViewArgs;
+	ViewArgs.bAllowSearch = false;
+	ViewArgs.bHideSelectionTip = false;
+	ViewArgs.bShowObjectLabel = false;
+	ViewArgs.NotifyHook = NotifyHook;
+
+	StructureDetailsView = EditModule.CreateStructureDetailView(ViewArgs, StructureViewArgs, StructToDisplay, LOCTEXT("Struct", "Struct View"));
+	StructureDetailsView->GetDetailsView()->SetIsPropertyReadOnlyDelegate(FIsPropertyReadOnly::CreateSP(this, &SRigVMDetailsInspector::IsStructViewPropertyReadOnly));
+	StructureDetailsView->GetOnFinishedChangingPropertiesDelegate().Clear();
+	StructureDetailsView->GetOnFinishedChangingPropertiesDelegate().Add(UserOnFinishedChangingProperties);
+}
+
+/** Update the inspector window to show information on the supplied object */
+void SRigVMDetailsInspector::ShowDetailsForSingleObject(UObject* Object, const FShowDetailsOptions& Options)
+{
+	TArray<UObject*> PropertyObjects;
+
+	if (Object != nullptr)
+	{
+		PropertyObjects.Add(Object);
+	}
+
+	ShowDetailsForObjects(PropertyObjects, Options);
+}
+
+void SRigVMDetailsInspector::ShowDetailsForObjects(const TArray<UObject*>& PropertyObjects, const FShowDetailsOptions& Options)
+{
+	// Refresh is being deferred until the next tick, this prevents batch operations from bombarding the details view with calls to refresh
+	RefreshPropertyObjects = PropertyObjects;
+	RefreshOptions = Options;
+	bRefreshOnTick = true;
+}
+
+/** Update the inspector window to show information on the supplied object */
+void SRigVMDetailsInspector::ShowSingleStruct(TSharedPtr<FStructOnScope> InStructToDisplay)
+{
+	static bool bIsReentrant = false;
+	if (!bIsReentrant)
+	{
+		bIsReentrant = true;
+		// When the selection is changed, we may be potentially actively editing a property,
+		// if this occurs we need need to immediately clear keyboard focus
+		if (FSlateApplication::Get().HasFocusedDescendants(AsShared()))
+		{
+			FSlateApplication::Get().ClearKeyboardFocus(EFocusCause::Mouse);
+		}
+		bIsReentrant = false;
+	}
+
+	StructToDisplay = InStructToDisplay;
+	// we don't defer this becasue StructDetailViews contains TSharedPtr to this sturct, 
+	// not clearing until next tick causes crash
+	// so will update struct view here, but updating widget will happen in the tick
+	StructureDetailsView->SetStructureData(InStructToDisplay);
+	bRefreshOnTick = true;
+}
+
+void SRigVMDetailsInspector::AddPropertiesRecursive(FProperty* Property)
+{
+	if (Property != nullptr)
+	{
+		// Add this property
+		SelectedObjectProperties.Add(Property);
+
+		// If this is a struct or an array of structs, recursively add the child properties
+		FArrayProperty* ArrayProperty = CastField<FArrayProperty>(Property);
+		FStructProperty* StructProperty = CastField<FStructProperty>(Property);
+		if(	StructProperty != nullptr && StructProperty->Struct != nullptr)
+		{
+			for (TFieldIterator<FProperty> StructPropIt(StructProperty->Struct); StructPropIt; ++StructPropIt)
+			{
+				FProperty* InsideStructProperty = *StructPropIt;
+				AddPropertiesRecursive(InsideStructProperty);
+			}
+		}
+		else if( ArrayProperty && ArrayProperty->Inner->IsA<FStructProperty>() )
+		{
+			AddPropertiesRecursive(ArrayProperty->Inner);
+		}
+	}
+}
+
+void SRigVMDetailsInspector::UpdateFromSingleStruct(const TSharedPtr<FStructOnScope>& InStructToDisplay)
+{
+	if (StructureDetailsView.IsValid())
+	{
+		SelectedObjects.Empty();
+
+		// Update our context-sensitive editing widget
+		ContextualEditingBorderWidget->SetContent(StructureDetailsView->GetWidget().ToSharedRef());
+	}
+}
+
+void SRigVMDetailsInspector::UpdateFromObjects(const TArray<UObject*>& PropertyObjects, TArray<UObject*>& SelectionInfo, const FShowDetailsOptions& Options)
+{
+	TSharedPtr<FRigVMEditorBase> Editor = WeakEditor.Pin();
+	
+	if (!Options.bForceRefresh)
+	{
+		// Early out if the PropertyObjects and the SelectedObjects are the same
+		bool bEquivalentSets = (PropertyObjects.Num() == SelectedObjects.Num());
+		if (bEquivalentSets)
+		{
+			// Verify the elements of the sets are equivalent
+			for (int32 i = 0; i < PropertyObjects.Num(); i++)
+			{
+				if (PropertyObjects[i] != SelectedObjects[i].Get())
+				{
+					if (PropertyObjects[i] && !PropertyObjects[i]->IsValidLowLevel())
+					{
+						ensureMsgf(false, TEXT("Object in RigVMInspector is invalid, see TTP 281915"));
+						continue;
+					}
+
+					bEquivalentSets = false;
+					break;
+				}
+			}
+		}
+
+		if (bEquivalentSets)
+		{
+			return;
+		}
+	}
+
+	PropertyView->OnFinishedChangingProperties().Clear();
+	PropertyView->OnFinishedChangingProperties().Add(UserOnFinishedChangingProperties);
+	PropertyView->OnFinishedChangingProperties().AddSP(this, &SRigVMDetailsInspector::OnFinishedChangingProperties);
+
+	// Proceed to update
+	SelectedObjects.Empty();
+
+	for (auto ObjectIt = PropertyObjects.CreateConstIterator(); ObjectIt; ++ObjectIt)
+	{
+		if (UObject* Object = *ObjectIt)
+		{
+			if (!Object->IsValidLowLevel())
+			{
+				ensureMsgf(false, TEXT("Object in KismetInspector is invalid, see TTP 281915"));
+				continue;
+			}
+
+			SelectedObjects.Add(Object);
+
+			if (UK2Node* K2Node = Cast<UK2Node>(Object))
+			{
+
+				// See if we should edit properties of the node
+				if (K2Node->ShouldShowNodeProperties())
+				{
+					SelectionInfo.Add(Object);
+				}
+			}
+			else
+			{
+				// Editing any UObject*
+				SelectionInfo.AddUnique(Object);
+			}
+		}
+	}
+
+	// By default, no property filtering
+	SelectedObjectProperties.Empty();
+
+	PropertyViewTitle = Options.ForcedTitle;
+
+	// Update our context-sensitive editing widget
+	ContextualEditingBorderWidget->SetContent( MakeContextualEditingWidget(SelectionInfo, Options) );
+}
+
+bool SRigVMDetailsInspector::IsStructViewPropertyReadOnly(const struct FPropertyAndParent& PropertyAndParent) const
+{
+	const FProperty& Property = PropertyAndParent.Property;
+	if (Property.HasAnyPropertyFlags(CPF_EditConst))
+	{
+		return true;
+	}
+
+	return false;
+}
+
+bool SRigVMDetailsInspector::IsAnyParentOrContainerSelected(const FPropertyAndParent& PropertyAndParent) const
+{
+	for (const FProperty* CurrentProperty : PropertyAndParent.ParentProperties)
+	{
+		if (SelectedObjectProperties.Find(const_cast<FProperty*>(CurrentProperty)))
+		{
+			return true;
+		}
+
+		// the property might be the Inner property of an array (or Key/Value of a map), so check if the outer property is selected
+		const FProperty* CurrentOuter = CurrentProperty->GetOwner<FProperty>();
+		if (CurrentOuter != nullptr && SelectedObjectProperties.Find(const_cast<FProperty*>(CurrentOuter)))
+		{
+			return true;
+		}
+	}
+
+	return false;
+} 
+
+bool SRigVMDetailsInspector::IsPropertyVisible( const FPropertyAndParent& PropertyAndParent ) const
+{
+	const FProperty& Property = PropertyAndParent.Property;
+
+	// If we are in 'instance preview' - hide anything marked 'disabled edit on instance'
+	if ((ECheckBoxState::Checked == PublicViewState) && Property.HasAnyPropertyFlags(CPF_DisableEditOnInstance))
+	{
+		return false;
+	}
+
+	// Only hide EditInstanceOnly properties if we are editing a CDO/archetype
+	bool bIsEditingTemplate = true;
+	for (const TWeakObjectPtr<UObject>& SelectedObject : SelectedObjects)
+	{
+		UObject* Object = SelectedObject.Get();
+		if (!Object || !Object->IsTemplate())
+		{
+			bIsEditingTemplate = false;
+			break;
+		}
+	}
+
+	if (bIsEditingTemplate)
+	{
+		// check if the property (or any of its parent properties) was added by this blueprint
+		// this is necessary because of Instanced objects, which will have a different owning class yet are conceptually contained in this blueprint
+		bool bVariableAddedInCurrentBlueprint = false;
+		TSharedPtr<FRigVMEditorBase> Editor = WeakEditor.Pin();
+		const UBlueprint* Blueprint = Editor.IsValid() ? Editor->GetBlueprintObj() : nullptr;
+
+		auto WasAddedInThisBlueprint = [Blueprint](const FProperty* Property)
+		{
+			if (const UClass* OwningClass = Property->GetOwnerClass())
+			{
+				return Blueprint && OwningClass->ClassGeneratedBy.Get() == Blueprint;
+			}
+			return false;
+		};
+		
+		bVariableAddedInCurrentBlueprint |= WasAddedInThisBlueprint(&Property);
+
+		for (const FProperty* Parent : PropertyAndParent.ParentProperties)
+		{
+			bVariableAddedInCurrentBlueprint |= WasAddedInThisBlueprint(Parent);
+		}
+
+		// if this property wasn't added in this blueprint, we want to filter it out if it (or any of its parents) are marked EditInstanceOnly or private
+		if (!bVariableAddedInCurrentBlueprint)
+		{
+			if (Property.HasAnyPropertyFlags(CPF_DisableEditOnTemplate) || Property.GetBoolMetaData(FBlueprintMetadata::MD_Private))
+			{
+				return false;
+			}
+
+			for (const FProperty* Parent : PropertyAndParent.ParentProperties)
+			{
+				if (Property.HasAnyPropertyFlags(CPF_DisableEditOnTemplate) || Parent->GetBoolMetaData(FBlueprintMetadata::MD_Private))
+				{
+					return false;
+				}
+			}
+		}
+	}
+	
+	// figure out if this Blueprint variable is an Actor variable
+	const FArrayProperty* ArrayProperty = CastField<const FArrayProperty>(&Property);
+	const FSetProperty* SetProperty = CastField<const FSetProperty>(&Property);
+	const FMapProperty* MapProperty = CastField<const FMapProperty>(&Property);
+
+	// Filter down to selected properties only if set.
+	if (SelectedObjectProperties.Find(const_cast<FProperty*>(&Property)))
+	{
+		// If the current property is selected, it is visible.
+		return true;
+	}
+	else if (PropertyAndParent.ParentProperties.Num() > 0 && SelectedObjectProperties.Num() > 0)
+	{
+		if (IsAnyParentOrContainerSelected(PropertyAndParent))
+		{
+			return true;
+		}
+	}
+	else if (ArrayProperty || MapProperty || SetProperty)
+	{
+		// .Find won't work here because the items inside of the container properties are not FProperties
+		for (const TWeakFieldPtr<FProperty>& CurProp : SelectedObjectProperties)
+		{
+			if ((ArrayProperty && (ArrayProperty->PropertyFlags & CPF_Edit) && CurProp->GetFName() == ArrayProperty->GetFName()) ||
+				(MapProperty && (MapProperty->PropertyFlags & CPF_Edit) && CurProp->GetFName() == MapProperty->GetFName()) ||
+				(SetProperty && (SetProperty->PropertyFlags & CPF_Edit) && CurProp->GetFName() == SetProperty->GetFName()))
+			{
+				return true;
+			}
+		}
+	}
+
+	return SelectedObjectProperties.Num() == 0;
+}
+
+EVisibility SRigVMDetailsInspector::GetPropertyViewVisibility() const
+{
+	return bShowInspectorPropertyView? EVisibility::Visible : EVisibility::Collapsed;
+}
+
+bool SRigVMDetailsInspector::IsPropertyEditingEnabled() const
+{
+	bool bIsEditable = true;
+
+	if (TSharedPtr<FRigVMEditorBase> Editor = WeakEditor.Pin())
+	{
+		// This function is essentially for PIE use so if we are NOT doing PIE use the normal path
+		if (GEditor->GetPIEWorldContext() == nullptr)
+		{
+			bIsEditable = Editor->InEditingMode();
+		}
+	}
+
+	for (const TWeakObjectPtr<UObject>& SelectedObject : SelectedObjects)
+	{
+		if (UActorComponent* Component = Cast<UActorComponent>(SelectedObject.Get()))
+		{
+			if(!CastChecked<UActorComponent>(Component->GetArchetype())->IsEditableWhenInherited())
+			{
+				bIsEditable = false;
+				break;
+			}
+		}
+		else if(UEdGraphNode* EdGraphNode = Cast<UEdGraphNode>(SelectedObject.Get()))
+		{
+			if(UEdGraph* OuterGraph = EdGraphNode->GetGraph())
+			{
+				if(WeakEditor.IsValid() && !WeakEditor.Pin()->IsEditable(OuterGraph))
+				{
+					bIsEditable = false;
+					break;
+				}
+			}
+		}
+	}
+	
+	return bIsEditable && (!IsPropertyEditingEnabledDelegate.IsBound() || IsPropertyEditingEnabledDelegate.Execute());
+}
+
+ECheckBoxState SRigVMDetailsInspector::GetPublicViewCheckboxState() const
+{
+	return PublicViewState;
+}
+
+void SRigVMDetailsInspector::SetPublicViewCheckboxState( ECheckBoxState InIsChecked )
+{
+	PublicViewState = InIsChecked;
+
+	//reset the details view
+	TArray<UObject*> Objs;
+	for(auto It(SelectedObjects.CreateIterator());It;++It)
+	{
+		Objs.Add(It->Get());
+	}
+	SelectedObjects.Empty();
+	
+	if(Objs.Num() > 1)
+	{
+		ShowDetailsForObjects(Objs);
+	}
+	else if(Objs.Num() == 1)
+	{
+		ShowDetailsForSingleObject(Objs[0], FShowDetailsOptions(PropertyViewTitle));
+	}
+	
+	//WeakEditor.Pin()->StartEditingDefaults();
+}
+
+void SRigVMDetailsInspector::OnFinishedChangingProperties(const FPropertyChangedEvent& InPropertyChangedEvent)
+{
+}
+
+//////////////////////////////////////////////////////////////////////////
+
+#undef LOCTEXT_NAMESPACE
diff --git a/Engine/Plugins/Runtime/RigVM/Source/RigVMEditor/Private/RigVMEditorModule.cpp b/Engine/Plugins/Runtime/RigVM/Source/RigVMEditor/Private/RigVMEditorModule.cpp
--- a/Engine/Plugins/Runtime/RigVM/Source/RigVMEditor/Private/RigVMEditorModule.cpp
+++ b/Engine/Plugins/Runtime/RigVM/Source/RigVMEditor/Private/RigVMEditorModule.cpp
@@ -78,7 +78,7 @@ void FRigVMEditorModule::StartupModule()
 #if WITH_RIGVMLEGACYEDITOR
 		// Register Blueprint editor variable customization
 		FBlueprintEditorModule& BlueprintEditorModule = FModuleManager::LoadModuleChecked<FBlueprintEditorModule>("Kismet");
-		BlueprintVariableCustomizationHandle = BlueprintEditorModule.RegisterVariableCustomization(FProperty::StaticClass(), FOnGetVariableCustomizationInstance::CreateStatic(&FRigVMVariableDetailCustomization::MakeInstance));
+		BlueprintVariableCustomizationHandle = BlueprintEditorModule.RegisterVariableCustomization(FProperty::StaticClass(), FOnGetVariableCustomizationInstance::CreateStatic(&FRigVMVariableDetailCustomization::MakeLegacyInstance));
 #endif
 
 		IAssetRegistry* AssetRegistry = IAssetRegistry::Get();
@@ -158,7 +158,7 @@ void FRigVMEditorModule::StartupModuleCommon()
 	const URigVMEdGraphSchema* SchemaCDO = CastChecked<URigVMEdGraphSchema>(EdGraphSchemaClass->GetDefaultObject());
 #if WITH_RIGVMLEGACYEDITOR
 	FBlueprintEditorModule& BlueprintEditorModule = FModuleManager::LoadModuleChecked<FBlueprintEditorModule>("Kismet");
-	BlueprintEditorModule.RegisterGraphCustomization(SchemaCDO, FOnGetGraphCustomizationInstance::CreateStatic(&FRigVMGraphDetailCustomization::MakeInstance, BlueprintClass));
+	BlueprintEditorModule.RegisterGraphCustomization(SchemaCDO, FOnGetGraphCustomizationInstance::CreateStatic(&FRigVMGraphDetailCustomization::MakeLegacyInstance, BlueprintClass));
 #endif
 
 	// Register to fixup newly created BPs
diff --git a/Engine/Plugins/Runtime/RigVM/Source/RigVMEditor/Private/Widgets/SRigVMEditorGraphExplorer.cpp b/Engine/Plugins/Runtime/RigVM/Source/RigVMEditor/Private/Widgets/SRigVMEditorGraphExplorer.cpp
--- a/Engine/Plugins/Runtime/RigVM/Source/RigVMEditor/Private/Widgets/SRigVMEditorGraphExplorer.cpp
+++ b/Engine/Plugins/Runtime/RigVM/Source/RigVMEditor/Private/Widgets/SRigVMEditorGraphExplorer.cpp
@@ -29,6 +29,7 @@
 #include "Editor/SRigVMDetailsInspector.h"
 #endif
 #include "EdGraph/RigVMEdGraph.h"
+#include "Editor/SRigVMDetailsInspector.h"
 #include "WorkflowOrientedApp/WorkflowTabManager.h"
 
 #define LOCTEXT_NAMESPACE "RigVMGraphExplorer"
@@ -249,7 +250,7 @@ FName SRigVMEditorGraphExplorer::GetSelectedVariableName() const
 		return NAME_None;
 	}
 
-	if (Selected[0].Type != ERigVMExplorerElementType::Variable || Selected[0].Type != ERigVMExplorerElementType::LocalVariable)
+	if (Selected[0].Type != ERigVMExplorerElementType::Variable && Selected[0].Type != ERigVMExplorerElementType::LocalVariable)
 	{
 		return NAME_None;
 	}
@@ -257,6 +258,17 @@ FName SRigVMEditorGraphExplorer::GetSelectedVariableName() const
 	return *Selected[0].Name;
 }
 
+ERigVMExplorerElementType::Type SRigVMEditorGraphExplorer::GetSelectedType() const
+{
+	TArray<FRigVMExplorerElementKey> Selected = TreeView->GetSelectedKeys();
+	if (Selected.Num() != 1)
+	{
+		return ERigVMExplorerElementType::Invalid;
+	}
+
+	return Selected[0].Type;
+}
+
 void SRigVMEditorGraphExplorer::ClearSelection()
 {
 	TreeView->ClearSelection();
@@ -761,16 +773,15 @@ void SRigVMEditorGraphExplorer::OnGraphClicked(const FString& InGraphPath)
 			check(Schema != nullptr);
 			Schema->GetGraphDisplayInformation(*EdGraph, DisplayInfo);
 #if WITH_RIGVMLEGACYEDITOR
-			if (TSharedPtr<SKismetInspector> CurrentInspector = Editor->GetInspector())
+			if (TSharedPtr<SKismetInspector> KismetInspector = Editor->GetKismetInspector())
 			{
-				CurrentInspector->ShowDetailsForSingleObject(EdGraph, SKismetInspector::FShowDetailsOptions(DisplayInfo.PlainName));
+				KismetInspector->ShowDetailsForSingleObject(EdGraph, SKismetInspector::FShowDetailsOptions(DisplayInfo.PlainName));
 			}
-#else
-			if (TSharedPtr<SRigVMDetailsInspector> CurrentInspector = Editor->GetInspector())
+#endif
+			if(TSharedPtr<SRigVMDetailsInspector> RigVMInspector = Editor->GetRigVMInspector())
 			{
-				CurrentInspector->ShowDetailsForSingleObject(EdGraph, SRigVMDetailsInspector::FShowDetailsOptions(DisplayInfo.PlainName));
+				RigVMInspector->ShowDetailsForSingleObject(EdGraph, SRigVMDetailsInspector::FShowDetailsOptions(DisplayInfo.PlainName));
 			}
-#endif
 		}
 	}
 }
@@ -783,13 +794,20 @@ void SRigVMEditorGraphExplorer::OnEventClicked(const FString& InEventPath)
 		return;
 	}
 #if WITH_RIGVMLEGACYEDITOR
-	TSharedPtr<SKismetInspector> CurrentInspector = Editor->GetInspector();
-#else
-	TSharedPtr<SRigVMDetailsInspector> CurrentInspector = Editor->GetInspector();
+	if (TSharedPtr<SKismetInspector> KismetInspector = Editor->GetKismetInspector())
+	{
+		if (KismetInspector.IsValid())
+		{
+			KismetInspector->ShowDetailsForObjects(TArray<UObject*>());
+		}
+	}
 #endif
-	if (CurrentInspector.IsValid())
+	if (TSharedPtr<SRigVMDetailsInspector> RigVMInspector = Editor->GetRigVMInspector())
 	{
-		CurrentInspector->ShowDetailsForObjects(TArray<UObject*>());
+		if (RigVMInspector.IsValid())
+		{
+			RigVMInspector->ShowDetailsForObjects(TArray<UObject*>());
+		}
 	}
 }
 
@@ -814,16 +832,15 @@ void SRigVMEditorGraphExplorer::OnFunctionClicked(const FString& InFunctionPath)
 				check(Schema != nullptr);
 				Schema->GetGraphDisplayInformation(*EdGraph, DisplayInfo);
 #if WITH_RIGVMLEGACYEDITOR
-				if (TSharedPtr<SKismetInspector> CurrentInspector = Editor->GetInspector())
+				if (TSharedPtr<SKismetInspector> KismetInspector = Editor->GetKismetInspector())
 				{
-					CurrentInspector->ShowDetailsForSingleObject(EdGraph, SKismetInspector::FShowDetailsOptions(DisplayInfo.PlainName));
+					KismetInspector->ShowDetailsForSingleObject(EdGraph, SKismetInspector::FShowDetailsOptions(DisplayInfo.PlainName));
 				}
-#else
-				if (TSharedPtr<SRigVMDetailsInspector> CurrentInspector = Editor->GetInspector())
+#endif
+				if (TSharedPtr<SRigVMDetailsInspector> RigVMInspector = Editor->GetRigVMInspector())
 				{
-					CurrentInspector->ShowDetailsForSingleObject(EdGraph, SRigVMDetailsInspector::FShowDetailsOptions(DisplayInfo.PlainName));
+					RigVMInspector->ShowDetailsForSingleObject(EdGraph, SRigVMDetailsInspector::FShowDetailsOptions(DisplayInfo.PlainName));
 				}
-#endif
 			}
 		}
 	}
@@ -841,20 +858,27 @@ void SRigVMEditorGraphExplorer::OnVariableClicked(const FRigVMExplorerElementKey
 	if (InVariable.Type == ERigVMExplorerElementType::Variable)
 	{
 		FRigVMClient* RigVMClient = Blueprint->GetRigVMClient();
-#if WITH_RIGVMLEGACYEDITOR
-		TSharedPtr<SKismetInspector> CurrentInspector = Editor->GetInspector();
-		SKismetInspector::FShowDetailsOptions Options(FText::FromString(InVariable.Name));
-#else
-		TSharedPtr<SRigVMDetailsInspector> CurrentInspector = Editor->GetInspector();
-		SRigVMDetailsInspector::FShowDetailsOptions Options(FText::FromString(InVariable.Name));
-#endif
-		Options.bForceRefresh = true;
-
 		FProperty* Prop = Blueprint->SkeletonGeneratedClass->FindPropertyByName(*InVariable.Name);
 		UPropertyWrapper* PropWrap = (Prop ? Prop->GetUPropertyWrapper() : nullptr);
-		if (CurrentInspector.IsValid())
+#if WITH_RIGVMLEGACYEDITOR
+		if (TSharedPtr<SKismetInspector> KismetInspector = Editor->GetKismetInspector())
+		{
+			SKismetInspector::FShowDetailsOptions Options(FText::FromString(InVariable.Name));
+			Options.bForceRefresh = true;
+			if (KismetInspector.IsValid())
+			{
+				KismetInspector->ShowDetailsForSingleObject(PropWrap, Options);
+			}
+		}
+#endif
+		if (TSharedPtr<SRigVMDetailsInspector> RigVMInspector = Editor->GetRigVMInspector())
 		{
-			CurrentInspector->ShowDetailsForSingleObject(PropWrap, Options);
+			SRigVMDetailsInspector::FShowDetailsOptions Options(FText::FromString(InVariable.Name));
+			Options.bForceRefresh = true;
+			if (RigVMInspector.IsValid())
+			{
+				RigVMInspector->ShowDetailsForSingleObject(PropWrap, Options);
+			}
 		}
 	}
 	else if (InVariable.Type == ERigVMExplorerElementType::LocalVariable)
diff --git a/Engine/Plugins/Runtime/RigVM/Source/RigVMEditor/Public/Editor/RigVMDetailsInspectorTabSummoner.h b/Engine/Plugins/Runtime/RigVM/Source/RigVMEditor/Public/Editor/RigVMDetailsInspectorTabSummoner.h
new file mode 100644
--- /dev/null
+++ b/Engine/Plugins/Runtime/RigVM/Source/RigVMEditor/Public/Editor/RigVMDetailsInspectorTabSummoner.h
@@ -0,0 +1,23 @@
+// Copyright Epic Games, Inc. All Rights Reserved.
+
+#pragma once
+
+#include "CoreMinimal.h"
+#include "WorkflowOrientedApp/WorkflowTabFactory.h"
+
+class FRigVMNewEditor;
+
+struct RIGVMEDITOR_API FRigVMDetailsInspectorTabSummoner : public FWorkflowTabFactory
+{
+public:
+
+	static const FName TabID() { static FName ID = TEXT("RigVM Details"); return ID; }
+
+public:
+	FRigVMDetailsInspectorTabSummoner(const TSharedRef<FRigVMNewEditor>& InRigVMEditor);
+
+	virtual TSharedRef<SWidget> CreateTabBody(const FWorkflowTabSpawnInfo& Info) const override;
+
+protected:
+	TWeakPtr<FRigVMNewEditor> RigVMEditor;
+};
diff --git a/Engine/Plugins/Runtime/RigVM/Source/RigVMEditor/Public/Editor/RigVMEditor.h b/Engine/Plugins/Runtime/RigVM/Source/RigVMEditor/Public/Editor/RigVMEditor.h
--- a/Engine/Plugins/Runtime/RigVM/Source/RigVMEditor/Public/Editor/RigVMEditor.h
+++ b/Engine/Plugins/Runtime/RigVM/Source/RigVMEditor/Public/Editor/RigVMEditor.h
@@ -103,6 +103,7 @@ public:
 	virtual TSharedRef<IRigVMEditor> SharedRef() = 0;
 
 	virtual TSharedPtr<FAssetEditorToolkit> GetHostingApp() = 0;
+	virtual const TSharedPtr<FAssetEditorToolkit> GetHostingApp() const = 0;
 	virtual void InitAssetEditor(const EToolkitMode::Type Mode, const TSharedPtr<IToolkitHost>& InitToolkitHost, const FName AppIdentifier, const TSharedRef<FTabManager::FLayout>& StandaloneDefaultLayout, const bool bCreateDefaultStandaloneMenu, const bool bCreateDefaultToolbar, const TArray<UObject*>& ObjectsToEdit, const bool bInIsToolbarFocusable = false, const bool bInUseSmallToolbarIcons = false, const TOptional<EAssetOpenMethod>& InOpenMethod = TOptional<EAssetOpenMethod>()) = 0;
 	virtual void CreateEditorToolbar() = 0;
 	virtual void CommonInitialization(const TArray<UBlueprint*>& InitBlueprints, bool bShouldOpenInDefaultsMode) = 0;
@@ -118,10 +119,9 @@ public:
 	virtual void CloseDocumentTab(const UObject* DocumentID) = 0;
 	virtual TSharedPtr<FTabManager> GetTabManager() = 0;
 #if WITH_RIGVMLEGACYEDITOR
-	virtual TSharedPtr<class SKismetInspector> GetInspector() const = 0;
-#else
-	virtual TSharedPtr<class SRigVMDetailsInspector> GetInspector() const = 0;
+	virtual TSharedPtr<class SKismetInspector> GetKismetInspector() const = 0;
 #endif
+	virtual TSharedPtr<class SRigVMDetailsInspector> GetRigVMInspector() const = 0;
 	virtual TSharedPtr<FApplicationMode> CreateEditorMode() = 0;
 	virtual const FName GetEditorAppName() const = 0;
 	virtual const TArray< UObject* >& GetEditingBlueprints() const = 0;
@@ -185,6 +185,7 @@ public:
 	virtual void RegisterToolbarTab(const TSharedRef<class FTabManager>& TabManager) = 0;
 	virtual const TArray< UObject* >* GetObjectsCurrentlyBeingEdited() const = 0;
 	virtual void AddSelectedDebugObjectWidget(FToolBarBuilder& ToolbarBuilder) = 0;
+	virtual FNotifyHook* GetNotifyHook() = 0;
 };
 
 class RIGVMEDITOR_API FRigVMEditorBase : public IRigVMEditor
@@ -233,9 +234,9 @@ protected:
 	// FBlueprintEditor interface
 public:
 	UBlueprint* GetBlueprintObj() const;
+	bool InEditingMode() const;
 protected:
 	TSubclassOf<UEdGraphSchema> GetDefaultSchemaClass() const;
-	bool InEditingMode() const;
 	void OnGraphEditorFocused(const TSharedRef<SGraphEditor>& InGraphEditor);
 
 	//  FTickableEditorObject Interface
@@ -259,11 +260,21 @@ protected:
 public:
 	void DeleteSelectedNodes();
 	bool CanDeleteNodes() const;
-protected:
 	void CopySelectedNodes();
 	bool CanCopyNodes() const;
 	void PasteNodes();
 	bool CanPasteNodes() const;
+	void CutSelectedNodes();
+	bool CanCutNodes() const;
+	void DuplicateNodes();
+	bool CanDuplicateNodes() const;
+
+	void OnStartWatchingPin();
+	bool CanStartWatchingPin() const;
+	void OnStopWatchingPin();
+	bool CanStopWatchingPin() const;
+protected:
+
 	bool IsNativeParentClassCodeLinkEnabled() const { return false; }
 	bool ReparentBlueprint_IsVisible() const { return false; }
 	FReply OnSpawnGraphNodeByShortcut(FInputChord InChord, const FVector2D& InPosition, UEdGraph* InGraph);
@@ -278,10 +289,6 @@ protected:
 
 	virtual void PostTransaction(bool bSuccess, const FTransaction* Transaction, bool bIsRedo) {}
 
-	void OnStartWatchingPin();
-	bool CanStartWatchingPin() const;
-	void OnStopWatchingPin();
-	bool CanStopWatchingPin() const;
 
 	// IToolkitHost Interface
 	void OnToolkitHostingStarted(const TSharedRef<class IToolkit>& Toolkit);
@@ -348,6 +355,9 @@ protected:
 	void RefreshEditors(ERefreshRigVMEditorReason::Type Reason = ERefreshRigVMEditorReason::UnknownReason);
 	void SetupGraphEditorEvents(UEdGraph* InGraph, SGraphEditor::FGraphEditorEvents& InEvents);
 	void FocusInspectorOnGraphSelection(const TSet<class UObject*>& NewSelection, bool bForceRefresh = false);
+#if WITH_RIGVMLEGACYEDITOR
+	virtual TSharedPtr<SKismetInspector> GetKismetInspector() const override { return nullptr; }
+#endif
 
 	virtual void HandleModifiedEvent(ERigVMGraphNotifType InNotifType, URigVMGraph* InGraph, UObject* InSubject);
 	UE_DEPRECATED(5.4, "Please use HandleVMCompiledEvent with ExtendedExecuteContext param.")
@@ -521,6 +531,7 @@ protected:
 
 public:
 	void SetGraphExplorerWidget(TSharedPtr<SRigVMEditorGraphExplorer> InWidget) { GraphExplorerWidget = InWidget; }
+	TSharedPtr<SRigVMEditorGraphExplorer> GetGraphExplorerWidget() { return GraphExplorerWidget; }
 
 	bool IsEditingSingleBlueprint() const;
 	
diff --git a/Engine/Plugins/Runtime/RigVM/Source/RigVMEditor/Public/Editor/RigVMEditorGraphExplorerTabSummoner.h b/Engine/Plugins/Runtime/RigVM/Source/RigVMEditor/Public/Editor/RigVMEditorGraphExplorerTabSummoner.h
--- a/Engine/Plugins/Runtime/RigVM/Source/RigVMEditor/Public/Editor/RigVMEditorGraphExplorerTabSummoner.h
+++ b/Engine/Plugins/Runtime/RigVM/Source/RigVMEditor/Public/Editor/RigVMEditorGraphExplorerTabSummoner.h
@@ -11,7 +11,7 @@ struct RIGVMEDITOR_API FRigVMEditorGraphExplorerTabSummoner : public FWorkflowTa
 {
 public:
 
-	static inline const FLazyName TabID = FLazyName(TEXT("RigVM Graph Explorer"));
+	static const FName TabID() { static FName ID = TEXT("RigVM Graph Explorer"); return ID; }
 	
 public:
 	FRigVMEditorGraphExplorerTabSummoner(const TSharedRef<FRigVMEditorBase>& InRigVMEditor);
diff --git a/Engine/Plugins/Runtime/RigVM/Source/RigVMEditor/Public/Editor/RigVMGraphDetailCustomization.h b/Engine/Plugins/Runtime/RigVM/Source/RigVMEditor/Public/Editor/RigVMGraphDetailCustomization.h
--- a/Engine/Plugins/Runtime/RigVM/Source/RigVMEditor/Public/Editor/RigVMGraphDetailCustomization.h
+++ b/Engine/Plugins/Runtime/RigVM/Source/RigVMEditor/Public/Editor/RigVMGraphDetailCustomization.h
@@ -163,13 +163,13 @@ class RIGVMEDITOR_API FRigVMGraphDetailCustomization : public IDetailCustomizati
 {
 public:
 	/** Makes a new instance of this detail layout class for a specific detail view requesting it */
+	static TSharedRef<IDetailCustomization> MakeInstance(TSharedPtr<FRigVMEditorBase> InBlueprintEditor, const UClass* InExpectedBlueprintClass);
+	FRigVMGraphDetailCustomization(TSharedPtr<FRigVMEditorBase> RigVMigEditor, URigVMBlueprint* RigVMBlueprint);
 #if WITH_RIGVMLEGACYEDITOR
-	static TSharedPtr<IDetailCustomization> MakeInstance(TSharedPtr<IBlueprintEditor> InBlueprintEditor, const UClass* InExpectedBlueprintClass);
-#else
-	static TSharedPtr<IDetailCustomization> MakeInstance(TSharedPtr<FRigVMEditorBase> InBlueprintEditor, const UClass* InExpectedBlueprintClass);
+	static TSharedPtr<IDetailCustomization> MakeLegacyInstance(TSharedPtr<IBlueprintEditor> InBlueprintEditor, const UClass* InExpectedBlueprintClass);
+	FRigVMGraphDetailCustomization(TSharedPtr<IBlueprintEditor> RigVMigEditor, URigVMBlueprint* RigVMBlueprint);
 #endif
 
-	FRigVMGraphDetailCustomization(TSharedPtr<FRigVMEditorBase> RigVMigEditor, URigVMBlueprint* RigVMBlueprint);
 
 	// IDetailCustomization interface
 	virtual void CustomizeDetails(IDetailLayoutBuilder& DetailLayout) override;
diff --git a/Engine/Plugins/Runtime/RigVM/Source/RigVMEditor/Public/Editor/RigVMLegacyEditor.h b/Engine/Plugins/Runtime/RigVM/Source/RigVMEditor/Public/Editor/RigVMLegacyEditor.h
--- a/Engine/Plugins/Runtime/RigVM/Source/RigVMEditor/Public/Editor/RigVMLegacyEditor.h
+++ b/Engine/Plugins/Runtime/RigVM/Source/RigVMEditor/Public/Editor/RigVMLegacyEditor.h
@@ -30,6 +30,7 @@ protected:
 
 	// IRigVMEditor overrides
 	virtual TSharedPtr<FAssetEditorToolkit> GetHostingApp() override { return AsShared(); }
+	virtual const TSharedPtr<FAssetEditorToolkit> GetHostingApp() const override { return ConstCastSharedRef<FAssetEditorToolkit>(AsShared()); }
 	virtual void InitAssetEditor(const EToolkitMode::Type Mode, const TSharedPtr<IToolkitHost>& InitToolkitHost, const FName AppIdentifier, const TSharedRef<FTabManager::FLayout>& StandaloneDefaultLayout, const bool bCreateDefaultStandaloneMenu, const bool bCreateDefaultToolbar, const TArray<UObject*>& ObjectsToEdit, const bool bInIsToolbarFocusable, const bool bInUseSmallToolbarIcons, const TOptional<EAssetOpenMethod>& InOpenMethod) override;
 	virtual void CreateEditorToolbar() override;
 	virtual void CommonInitialization(const TArray<UBlueprint*>& InitBlueprints, bool bShouldOpenInDefaultsMode) override { FBlueprintEditor::CommonInitialization(InitBlueprints, bShouldOpenInDefaultsMode); }
@@ -45,10 +46,9 @@ protected:
 	virtual void CloseDocumentTab(const UObject* DocumentID) override { FBlueprintEditor::CloseDocumentTab(DocumentID); }
 	virtual TSharedPtr<FTabManager> GetTabManager() override { return FBlueprintEditor::GetTabManager(); }
 #if WITH_RIGVMLEGACYEDITOR
-	virtual TSharedPtr<SKismetInspector> GetInspector() const override { return Inspector.IsValid() ? Inspector.ToSharedRef().ToSharedPtr() : nullptr; }
-#else
-	virtual TSharedPtr<SRigVMDetailsInspector> GetInspector() const override { return Inspector.IsValid() ? Inspector.ToSharedRef().ToSharedPtr() : nullptr; }
+	virtual TSharedPtr<SKismetInspector> GetKismetInspector() const override { return Inspector.IsValid() ? Inspector.ToSharedRef().ToSharedPtr() : nullptr; }
 #endif
+	virtual TSharedPtr<SRigVMDetailsInspector> GetRigVMInspector() const override { return nullptr; }
 	virtual TSharedPtr<FApplicationMode> CreateEditorMode() override;
 	virtual const FName GetEditorAppName() const override;
 	virtual const TArray< UObject* >& GetEditingBlueprints() const override { return FBlueprintEditor::GetEditingObjects(); }
@@ -190,6 +190,8 @@ protected:
 	virtual void OnToolkitHostingStarted(const TSharedRef<class IToolkit>& Toolkit) override { FRigVMEditorBase::OnToolkitHostingStarted(Toolkit); }
 	virtual void OnToolkitHostingFinished(const TSharedRef<class IToolkit>& Toolkit) override { FRigVMEditorBase::OnToolkitHostingFinished(Toolkit); }
 
+	virtual FNotifyHook* GetNotifyHook() override { return this; }
+
 };
 
 #endif
\ No newline at end of file
diff --git a/Engine/Plugins/Runtime/RigVM/Source/RigVMEditor/Public/Editor/RigVMNewEditor.h b/Engine/Plugins/Runtime/RigVM/Source/RigVMEditor/Public/Editor/RigVMNewEditor.h
--- a/Engine/Plugins/Runtime/RigVM/Source/RigVMEditor/Public/Editor/RigVMNewEditor.h
+++ b/Engine/Plugins/Runtime/RigVM/Source/RigVMEditor/Public/Editor/RigVMNewEditor.h
@@ -5,6 +5,8 @@
 #include "RigVMEditor.h"
 #include "WorkflowOrientedApp/WorkflowCentricApplication.h"
 
+class SRigVMDetailsInspector;
+
 struct RIGVMEDITOR_API FRigVMNewEditorTabs
 {
 	// Tab identifiers
@@ -20,9 +22,10 @@ public:
 
 	virtual TSharedRef<IRigVMEditor> SharedRef() override { return StaticCastSharedRef<IRigVMEditor>(SharedThis(this)); }
 
+	virtual TSharedPtr<FAssetEditorToolkit> GetHostingApp() override { return AsShared(); }
+	virtual const TSharedPtr<FAssetEditorToolkit> GetHostingApp() const override { return ConstCastSharedRef<FAssetEditorToolkit>(AsShared()); }
 protected:
 	
-	virtual TSharedPtr<FAssetEditorToolkit> GetHostingApp() override { return AsShared(); }
 	virtual void InitAssetEditor(const EToolkitMode::Type Mode, const TSharedPtr<IToolkitHost>& InitToolkitHost, const FName AppIdentifier, const TSharedRef<FTabManager::FLayout>& StandaloneDefaultLayout, const bool bCreateDefaultStandaloneMenu, const bool bCreateDefaultToolbar, const TArray<UObject*>& ObjectsToEdit, const bool bInIsToolbarFocusable, const bool bInUseSmallToolbarIcons, const TOptional<EAssetOpenMethod>& InOpenMethod) override;
 	virtual void CreateEditorToolbar() override {}
 	virtual void CommonInitialization(const TArray<UBlueprint*>& InitBlueprints, bool bShouldOpenInDefaultsMode) override;
@@ -39,11 +42,10 @@ protected:
 	virtual TSharedPtr<SDockTab> OpenDocument(const UObject* DocumentID, FDocumentTracker::EOpenDocumentCause Cause) override;
 	virtual void CloseDocumentTab(const UObject* DocumentID) override;
 	virtual TSharedPtr<FTabManager> GetTabManager() override { return FWorkflowCentricApplication::GetTabManager(); }
-#if WITH_RIGVMLEGACYEDITOR
-	virtual TSharedPtr<SKismetInspector> GetInspector() const override { return Inspector.IsValid() ? Inspector.ToSharedRef().ToSharedPtr() : nullptr; }
-#else
-	virtual TSharedPtr<SRigVMDetailsInspector> GetInspector() const override { return Inspector.IsValid() ? Inspector.ToSharedRef().ToSharedPtr() : nullptr; }
-#endif
+	virtual TSharedPtr<SRigVMDetailsInspector> GetRigVMInspector() const override { return Inspector.IsValid() ? Inspector.ToSharedRef().ToSharedPtr() : nullptr; }
+public:
+	virtual void SetInspector(TSharedPtr<SRigVMDetailsInspector> InWidget) { Inspector = InWidget; };
+protected:
 	virtual TSharedPtr<FApplicationMode> CreateEditorMode() override;
 	virtual const FName GetEditorAppName() const override;
 	virtual const TArray< UObject* >& GetEditingBlueprints() const override { return FWorkflowCentricApplication::GetEditingObjects(); }
@@ -95,7 +97,7 @@ protected:
 	virtual void EditClassDefaults_Clicked() override;
 	virtual void EditGlobalOptions_Clicked() override;
 	virtual void TryInvokingDetailsTab(bool bFlash = true) override;
-	virtual FName GetGraphExplorerWidgetID() override { return FRigVMEditorGraphExplorerTabSummoner::TabID; }
+	virtual FName GetGraphExplorerWidgetID() override { return FRigVMEditorGraphExplorerTabSummoner::TabID(); }
 	virtual void RefreshInspector() override;
 	virtual void RefreshStandAloneDefaultsEditor() override;
 	virtual void Tick(float DeltaTime) override;
@@ -143,15 +145,28 @@ public:
 	static const FName SelectionState_ClassSettings();
 	static const FName SelectionState_ClassDefaults();
 
+	virtual FNotifyHook* GetNotifyHook() override { return this; }
+
+	void OnSelectedNodesChanged(const FGraphPanelSelectionSet& NewSelection);
+
+	void OnAlignTop();
+	void OnAlignMiddle();
+	void OnAlignBottom();
+	void OnAlignLeft();
+	void OnAlignCenter();
+	void OnAlignRight();
+	void OnStraightenConnections();
+	void OnDistributeNodesH();
+	void OnDistributeNodesV();
+	void SelectAllNodes();
+	bool CanSelectAllNodes() const;
+
 protected:
 	
 	TSharedPtr<FDocumentTracker> DocumentManager;
+
 	/** Node inspector widget */
-#if WITH_RIGVMLEGACYEDITOR
-	TSharedPtr<class SKismetInspector> Inspector;
-#else
 	TSharedPtr<class SRigVMDetailsInspector> Inspector;
-#endif
 
 	/** Currently focused graph editor */
 	TWeakPtr<class SGraphEditor> FocusedGraphEdPtr;
@@ -176,11 +191,7 @@ protected:
 	float InstructionsFadeCountdown;
 
 	/** defaults inspector widget */
-#if WITH_RIGVMLEGACYEDITOR
-	TSharedPtr<class SKismetInspector> DefaultEditor;
-#else
 	TSharedPtr<class SRigVMDetailsInspector> DefaultEditor;
-#endif
 	
 	/** True if the editor was opened in defaults mode */
 	bool bWasOpenedInDefaultsMode;
diff --git a/Engine/Plugins/Runtime/RigVM/Source/RigVMEditor/Public/Editor/RigVMNewEditorMode.h b/Engine/Plugins/Runtime/RigVM/Source/RigVMEditor/Public/Editor/RigVMNewEditorMode.h
--- a/Engine/Plugins/Runtime/RigVM/Source/RigVMEditor/Public/Editor/RigVMNewEditorMode.h
+++ b/Engine/Plugins/Runtime/RigVM/Source/RigVMEditor/Public/Editor/RigVMNewEditorMode.h
@@ -36,7 +36,4 @@ protected:
 	FWorkflowAllowedTabSet TabFactories;
 
 	TWeakPtr<FRigVMNewEditor> Editor;
-private:
-	
-	TWeakObjectPtr<URigVMBlueprint> RigVMBlueprintPtr;
 };
diff --git a/Engine/Plugins/Runtime/RigVM/Source/RigVMEditor/Public/Editor/RigVMVariableDetailCustomization.h b/Engine/Plugins/Runtime/RigVM/Source/RigVMEditor/Public/Editor/RigVMVariableDetailCustomization.h
--- a/Engine/Plugins/Runtime/RigVM/Source/RigVMEditor/Public/Editor/RigVMVariableDetailCustomization.h
+++ b/Engine/Plugins/Runtime/RigVM/Source/RigVMEditor/Public/Editor/RigVMVariableDetailCustomization.h
@@ -7,40 +7,103 @@
 #include "IDetailCustomization.h"
 #include "Styling/SlateTypes.h"
 #include "UObject/WeakFieldPtr.h"
+#include "Layout/Visibility.h"
+#include "Widgets/Views/ITableRow.h"
+#include "Widgets/Views/STableViewBase.h"
+#include "Widgets/Views/SListView.h"
 
+class URigVMBlueprint;
+class SRigVMEditorGraphExplorer;
 class IDetailLayoutBuilder;
 class IBlueprintEditor;
 class UBlueprint;
+class FRigVMEditorBase;
+class FRigVMVariableDetailCustomization;
+class UK2Node_EditablePinBase;
+class FStructOnScope;
+class SEditableTextBox;
+class SComboButton;
 
 class RIGVMEDITOR_API FRigVMVariableDetailCustomization : public IDetailCustomization
 {
 public:
 	/** Makes a new instance of this detail layout class for a specific detail view requesting it */
-	static TSharedPtr<IDetailCustomization> MakeInstance(TSharedPtr<IBlueprintEditor> InBlueprintEditor);
+	static TSharedRef<IDetailCustomization> MakeInstance(TSharedPtr<FRigVMEditorBase> InEditor);
+	FRigVMVariableDetailCustomization(TSharedPtr<FRigVMEditorBase> InEditor, URigVMBlueprint* Blueprint);
 
-	FRigVMVariableDetailCustomization(TSharedPtr<IBlueprintEditor> InBlueprintEditor, UBlueprint* Blueprint)
-		: BlueprintEditorPtr(InBlueprintEditor)
-		, BlueprintPtr(Blueprint)
-	{}
+#if WITH_RIGVMLEGACYEDITOR
+	static TSharedPtr<IDetailCustomization> MakeLegacyInstance(TSharedPtr<IBlueprintEditor> InBlueprintEditor);
+	FRigVMVariableDetailCustomization(TSharedPtr<IBlueprintEditor> RigVMigEditor, UBlueprint* Blueprint);
+#endif
 
 	// IDetailCustomization interface
 	virtual void CustomizeDetails(IDetailLayoutBuilder& DetailLayout) override;
 
 private:
-	bool IsAnimationFlagEnabled(TWeakFieldPtr<FProperty> PropertyBeingCustomized) const;
+	void PopulateCategories();
 
-	ECheckBoxState IsAnimationOutputChecked(TWeakFieldPtr<FProperty> PropertyBeingCustomized) const;
+private:
+	/** Accessors passed to parent */
+	FName GetVariableName() const;
+	FText OnGetVariableName() const;
+	void OnVarNameCommitted(const FText& InNewName, ETextCommit::Type InTextCommit);
+
+	// Callbacks for uproperty details customization
+	FEdGraphPinType OnGetVarType() const;
+	void OnVarTypeChanged(const FEdGraphPinType& NewPinType);
+
+	void OnBrowseToVarType() const;
+	bool CanBrowseToVarType() const;
+
+	
+	FText OnGetTooltipText() const;
+	void OnTooltipTextCommitted(const FText& NewText, ETextCommit::Type InTextCommit, FName VarName);
+	EVisibility IsToolTipVisible() const;
+	
+	FText OnGetCategoryText() const;
+	void OnCategoryTextCommitted(const FText& NewText, ETextCommit::Type InTextCommit, FName VarName);
+	TSharedRef< ITableRow > MakeCategoryViewWidget( TSharedPtr<FText> Item, const TSharedRef< STableViewBase >& OwnerTable );
+	void OnCategorySelectionChanged( TSharedPtr<FText> ProposedSelection, ESelectInfo::Type /*SelectInfo*/ );
+	
+
+	ECheckBoxState OnGetExposedToSpawnCheckboxState() const;
+	void OnExposedToSpawnChanged(ECheckBoxState InNewState);
 
-	void HandleAnimationOutputCheckStateChanged(ECheckBoxState CheckBoxState, TWeakFieldPtr<FProperty> PropertyBeingCustomized);
+	ECheckBoxState OnGetPrivateCheckboxState() const;
+	void OnPrivateChanged(ECheckBoxState InNewState);
+	
+	ECheckBoxState OnGetExposedToCinematicsCheckboxState() const;
+	void OnExposedToCinematicsChanged(ECheckBoxState InNewState);
 
-	ECheckBoxState IsAnimationInputChecked(TWeakFieldPtr<FProperty> PropertyBeingCustomized) const;
 
-	void HandleAnimationInputCheckStateChanged(ECheckBoxState CheckBoxState, TWeakFieldPtr<FProperty> PropertyBeingCustomized);
+	FText OnGetMetaKeyValue(FName Key) const;
+	void OnMetaKeyValueChanged(const FText& NewMinValue, ETextCommit::Type CommitInfo, FName Key);
+	EVisibility RangeVisibility() const;
+
+
+	
+	/** Refreshes cached data that changes after a Blueprint recompile */
+	void OnPostEditorRefresh();
 
 private:
 	/** The Blueprint editor we are embedded in */
-	TWeakPtr<IBlueprintEditor> BlueprintEditorPtr;
+	TWeakPtr<FRigVMEditorBase> EditorPtr;
 
 	/** The blueprint we are editing */
-	TWeakObjectPtr<UBlueprint> BlueprintPtr;
+	TWeakObjectPtr<URigVMBlueprint> BlueprintPtr;
+
+	/** The widget used when in variable name editing mode */ 
+	TSharedPtr<SEditableTextBox> VarNameEditableTextBox;
+	
+	/** A list of all category names to choose from */
+	TArray<TSharedPtr<FText>> CategorySource;
+	/** Widgets for the categories */
+	TWeakPtr<SComboButton> CategoryComboButton;
+	TWeakPtr<SListView<TSharedPtr<FText>>> CategoryListView;
+
+	/** Cached property for the variable we are affecting */
+	TWeakFieldPtr<FProperty> CachedVariableProperty;
+
+	/** Cached name for the variable we are affecting */
+	mutable FName CachedVariableName;
 };
diff --git a/Engine/Plugins/Runtime/RigVM/Source/RigVMEditor/Public/Editor/SRigVMDetailsInspector.h b/Engine/Plugins/Runtime/RigVM/Source/RigVMEditor/Public/Editor/SRigVMDetailsInspector.h
--- a/Engine/Plugins/Runtime/RigVM/Source/RigVMEditor/Public/Editor/SRigVMDetailsInspector.h
+++ b/Engine/Plugins/Runtime/RigVM/Source/RigVMEditor/Public/Editor/SRigVMDetailsInspector.h
@@ -2,11 +2,48 @@
 
 #pragma once
 
-#if !WITH_RIGVMLEGACYEDITOR
-#include "Editor/RigVMEditor.h"
-
+#include "Containers/Array.h"
+#include "Containers/Set.h"
+#include "CoreMinimal.h"
+#include "Delegates/Delegate.h"
+#include "Framework/Text/SlateHyperlinkRun.h"
+#include "IStructureDetailsView.h"
+#include "Internationalization/Text.h"
+#include "Layout/Visibility.h"
+#include "PropertyEditorDelegates.h"
+#include "SlateFwd.h"
+#include "Styling/SlateTypes.h"
+#include "Templates/SharedPointer.h"
+#include "UObject/NameTypes.h"
+#include "UObject/WeakFieldPtr.h"
+#include "UObject/WeakObjectPtr.h"
+#include "UObject/WeakObjectPtrTemplates.h"
+#include "Widgets/DeclarativeSyntaxSupport.h"
+#include "Widgets/SCompoundWidget.h"
+#include "Widgets/SWidget.h"
+
+class FRigVMEditorBase;
+class FProperty;
+class FStructOnScope;
+class IDetailsView;
+class SBorder;
+class SDockTab;
+class SRigVMEditorGraphExplorer;
+class SWidget;
+class SScrollBar;
+class UObject;
+struct FGeometry;
+struct FPropertyChangedEvent;
+
+typedef TSet<class UObject*> FInspectorSelectionSet;
+
+//////////////////////////////////////////////////////////////////////////
+// SRigVMDetailsInspector
+
+/** Widget that shows properties and tools related to the selected node(s) */
 class RIGVMEDITOR_API SRigVMDetailsInspector : public SCompoundWidget, public FGCObject
 {
+public:
 	SLATE_BEGIN_ARGS( SRigVMDetailsInspector )
 		: _ShowPublicViewControl(false)
 		, _HideNameArea(true)
@@ -15,75 +52,168 @@ class RIGVMEDITOR_API SRigVMDetailsInspector : public SCompoundWidget, public FG
 		, _ShowLocalVariables(false)
 		, _ScrollbarAlignment(HAlign_Right)
 		, _ShowSectionSelector(false)
-	{}
-
-	SLATE_ARGUMENT(TWeakPtr<FRigVMEditorBase>, Editor)
-	SLATE_ARGUMENT(TWeakPtr<SRigVMEditorGraphExplorer>, GraphExplorer)
-	SLATE_ATTRIBUTE( bool, ShowPublicViewControl )
-	SLATE_ARGUMENT( bool, HideNameArea )
-	SLATE_ARGUMENT( FIsPropertyEditingEnabled, IsPropertyEditingEnabledDelegate )
-	SLATE_ARGUMENT( FOnFinishedChangingProperties::FDelegate, OnFinishedChangingProperties )
-	SLATE_ARGUMENT( FName, ViewIdentifier)
-	SLATE_ARGUMENT( bool, SetNotifyHook)
-	SLATE_ARGUMENT( bool, ShowTitleArea)
-	SLATE_ARGUMENT( bool, ShowLocalVariables)
-	SLATE_ARGUMENT( TSharedPtr<SScrollBar>, ExternalScrollbar)
-	SLATE_ARGUMENT( EHorizontalAlignment, ScrollbarAlignment)
-	SLATE_ARGUMENT( bool, ShowSectionSelector)
-	SLATE_END_ARGS()
+		{}
 
-	void Construct(const FArguments& InArgs){}
+		SLATE_ARGUMENT(TWeakPtr<FRigVMEditorBase>, Editor)
+		SLATE_ARGUMENT(TWeakPtr<SRigVMEditorGraphExplorer>, GraphExplorerWidget)
+		SLATE_ATTRIBUTE( bool, ShowPublicViewControl )
+		SLATE_ARGUMENT( bool, HideNameArea )
+		SLATE_ARGUMENT( FIsPropertyEditingEnabled, IsPropertyEditingEnabledDelegate )
+		SLATE_ARGUMENT( FOnFinishedChangingProperties::FDelegate, OnFinishedChangingProperties )
+		SLATE_ARGUMENT( FName, ViewIdentifier)
+		SLATE_ARGUMENT( bool, SetNotifyHook)
+		SLATE_ARGUMENT( bool, ShowTitleArea)
+		SLATE_ARGUMENT( bool, ShowLocalVariables)
+		SLATE_ARGUMENT( TSharedPtr<SScrollBar>, ExternalScrollbar)
+		SLATE_ARGUMENT( EHorizontalAlignment, ScrollbarAlignment)
+		SLATE_ARGUMENT( bool, ShowSectionSelector)
+	SLATE_END_ARGS()
+	
+	void Construct(const FArguments& InArgs);
 
 	/** Options for ShowDetails */
 	struct FShowDetailsOptions
 	{
 		FText ForcedTitle;
 		bool bForceRefresh;
+		bool bShowComponents;
+		bool bHideFilterArea;
 
 		FShowDetailsOptions()
 			:ForcedTitle()
 			,bForceRefresh(false)
+			,bShowComponents(true)
+			,bHideFilterArea(false)
 		{}
-		
+
 		FShowDetailsOptions(const FText& InForcedTitle, bool bInForceRefresh = false)
 			:ForcedTitle(InForcedTitle)
 			,bForceRefresh(bInForceRefresh)
+			,bShowComponents(true)
+			,bHideFilterArea(false)
 		{}
 	};
 
+	// SWidget interface
+	virtual void Tick( const FGeometry& AllottedGeometry, const double InCurrentTime, const float InDeltaTime ) override;
+	// End of SWidget interface
 
- 	/** Update the inspector window to show information on the supplied object */
-	void ShowDetailsForSingleObject(UObject* Object, const FShowDetailsOptions& Options = FShowDetailsOptions()){}
+	/** Update the inspector window to show information on the supplied object */
+	void ShowDetailsForSingleObject(UObject* Object, const FShowDetailsOptions& Options = FShowDetailsOptions());
 
- 	/** Update the inspector window to show information on the supplied objects */
-	void ShowDetailsForObjects(const TArray<UObject*>& PropertyObjects, const FShowDetailsOptions& Options = FShowDetailsOptions()){}
+	/** Update the inspector window to show information on the supplied objects */
+	void ShowDetailsForObjects(const TArray<UObject*>& PropertyObjects, const FShowDetailsOptions& Options = FShowDetailsOptions());
+
+	/** Update the inspector window to show single struct. This invalidates ShowDetailsForObjects */
+	void ShowSingleStruct(TSharedPtr<FStructOnScope> InStructToDisplay);
+
+	/** Used to control visibility of a property in the property window */
+	bool IsPropertyVisible( const struct FPropertyAndParent& PropertyAndParent ) const;
 
- 	/** Update the inspector window to show single struct. This invalidates ShowDetailsForObjects */
-	void ShowSingleStruct(TSharedPtr<FStructOnScope> InStructToDisplay){}
 
 	TSharedPtr<class IDetailsView> GetPropertyView() const { return PropertyView; }
 
-	void SetOwnerTab(TSharedRef<SDockTab> Tab);
-	TSharedPtr<SDockTab> GetOwnerTab() const { return OwnerTab.Pin(); }
+	TSharedPtr<SDockTab> GetOwnerTab() const;
+
 
 	/** returns the list of selected objects */
-	const TArray< TWeakObjectPtr<UObject> >& GetSelectedObjects() const { return SelectedObjects; }
+	const TArray< TWeakObjectPtr<UObject> >& GetSelectedObjects() const;
 
-	virtual void AddReferencedObjects(FReferenceCollector& Collector) override { Collector.AddReferencedObjects(RefreshPropertyObjects); }
-	virtual FString GetReferencerName() const override { return TEXT("SRigVMDetailsInspector"); }
+	virtual void AddReferencedObjects(FReferenceCollector& Collector) override;
+	virtual FString GetReferencerName() const override;
 
 protected:
+	/** Update the inspector window to show information on the supplied objects */
+	void UpdateFromObjects(const TArray<UObject*>& PropertyObjects, TArray<UObject*>& SelectionInfo, const FShowDetailsOptions& Options);
+
+	/** Add this property and all its child properties to SelectedObjectProperties */
+	void AddPropertiesRecursive(FProperty* Property);
+
+	/** Pointer back to the kismet 2 tool that owns us */
+	TWeakPtr<FRigVMEditorBase> WeakEditor;
 
 	/** The tab that owns this details view. */
 	TWeakPtr<SDockTab> OwnerTab;
 
+	/** String used as the title above the property window */
+	FText PropertyViewTitle;
+
+	/** Should we currently show the property view */
+	bool bShowInspectorPropertyView;
+
+	/** State of CheckBox representing whether to show only the public variables*/
+	ECheckBoxState	PublicViewState;
+
 	/** Property viewing widget */
 	TSharedPtr<class IDetailsView> PropertyView;
 
 	/** Selected objects for this detail view */
 	TArray< TWeakObjectPtr<UObject> > SelectedObjects;
 
+	/** Border widget that wraps a dynamic context-sensitive widget for editing objects that the property window is displaying */
+	TSharedPtr<SBorder> ContextualEditingBorderWidget;
+
+	/** If true show the public view control */
+	TAttribute<bool> bShowPublicView;
+
+	/** If true show the kismet inspector title widget */
+	bool bShowTitleArea;
+
+	/** Set of object properties that should be visible */
+	TSet< TWeakFieldPtr<FProperty> > SelectedObjectProperties;
+	
+	/** User defined delegate for IsPropertyEditingEnabled: */
+	FIsPropertyEditingEnabled IsPropertyEditingEnabledDelegate;
+
+	/** User defined delegate for OnFinishedChangingProperties */
+	FOnFinishedChangingProperties::FDelegate UserOnFinishedChangingProperties;
+
+	/** When TRUE, the Kismet inspector needs to refresh the details view on Tick */
+	bool bRefreshOnTick;
+
 	/** Holds the property objects that need to be displayed by the inspector starting on the next tick */
 	TArray<TObjectPtr<UObject>> RefreshPropertyObjects;
+
+	/** Details options that are used by the inspector on the next refresh. */
+	FShowDetailsOptions RefreshOptions;
+
+	/** Struct to preview */
+	TSharedPtr<FStructOnScope> StructToDisplay;
+
+	/* Sturct Detail View */
+	TSharedPtr<class IStructureDetailsView> StructureDetailsView;
+
+	/** Update the inspector window to show information on the single struct */
+	void UpdateFromSingleStruct(const TSharedPtr<FStructOnScope>& InStructToDisplay);
+
+	/** Is struct view property read only */
+	bool IsStructViewPropertyReadOnly(const struct FPropertyAndParent& PropertyAndParent) const;
+
+protected:
+
+	/** Returns whether the property view be visible */
+	EVisibility GetPropertyViewVisibility() const;
+
+	/** Returns whether the properties in the view should be editable */
+	bool IsPropertyEditingEnabled() const;
+
+	/**
+	 * Generates a widget that is used to edit the specified object array contextually.  This widget
+	 * will be displayed along with a property view in the level editor
+	 */
+	TSharedRef<SWidget> MakeContextualEditingWidget(TArray<UObject*>& SelectionInfo, const FShowDetailsOptions& Options);
+
+	/**
+	 * Generates the text for the title in the contextual editing widget
+	 */
+	FText GetContextualEditingWidgetTitle() const;
+
+	ECheckBoxState GetPublicViewCheckboxState() const;
+	void SetPublicViewCheckboxState(ECheckBoxState InIsChecked);
+
+	bool IsAnyParentOrContainerSelected(const FPropertyAndParent& PropertyAndParent) const;
+
+	/** Callback invoked after a value change on the selected object(s) */
+	void OnFinishedChangingProperties(const FPropertyChangedEvent& InPropertyChangedEvent);
+
 };
-#endif
\ No newline at end of file
diff --git a/Engine/Plugins/Runtime/RigVM/Source/RigVMEditor/Public/Widgets/SRigVMEditorGraphExplorer.h b/Engine/Plugins/Runtime/RigVM/Source/RigVMEditor/Public/Widgets/SRigVMEditorGraphExplorer.h
--- a/Engine/Plugins/Runtime/RigVM/Source/RigVMEditor/Public/Widgets/SRigVMEditorGraphExplorer.h
+++ b/Engine/Plugins/Runtime/RigVM/Source/RigVMEditor/Public/Widgets/SRigVMEditorGraphExplorer.h
@@ -93,6 +93,7 @@ public:
 	FEdGraphPinType GetLastPinTypeUsed() const { return LastPinType; }
 
 	FName GetSelectedVariableName() const;
+	ERigVMExplorerElementType::Type GetSelectedType() const;
 
 	void ClearSelection();
 
