Control Rig: Fix crash when selecting variable in a non RigVM blueprint
#rnx

diff --git a/Engine/Plugins/Runtime/RigVM/Source/RigVMEditor/Private/Editor/RigVMVariableDetailCustomization.cpp b/Engine/Plugins/Runtime/RigVM/Source/RigVMEditor/Private/Editor/RigVMVariableDetailCustomization.cpp
--- a/Engine/Plugins/Runtime/RigVM/Source/RigVMEditor/Private/Editor/RigVMVariableDetailCustomization.cpp
+++ b/Engine/Plugins/Runtime/RigVM/Source/RigVMEditor/Private/Editor/RigVMVariableDetailCustomization.cpp
@@ -53,9 +53,19 @@ FRigVMVariableDetailCustomization::FRigVMVariableDetailCustomization(TSharedPtr<
 #if WITH_RIGVMLEGACYEDITOR
 TSharedPtr<IDetailCustomization> FRigVMVariableDetailCustomization::MakeLegacyInstance(TSharedPtr<IBlueprintEditor> InEditor)
 {
-	TSharedPtr<FRigVMLegacyEditor> LegacyEditor = StaticCastSharedPtr<FRigVMLegacyEditor>(InEditor);
-	TSharedPtr<FRigVMEditorBase> Editor = StaticCastSharedPtr<FRigVMEditorBase>(LegacyEditor);
-	return MakeInstance(Editor);
+	const TArray<UObject*>* Objects = (InEditor.IsValid() ? InEditor->GetObjectsCurrentlyBeingEdited() : nullptr);
+	if (Objects && Objects->Num() == 1)
+	{
+		if (URigVMBlueprint* Blueprint = Cast<URigVMBlueprint>((*Objects)[0]))
+		{
+			if (Blueprint->ParentClass && Blueprint->ParentClass->IsChildOf(URigVMHost::StaticClass()))
+			{
+				return MakeShareable(new FRigVMVariableDetailCustomization(InEditor, Blueprint));
+			}
+		}
+	}
+
+	return MakeShareable(new FRigVMVariableDetailCustomization((TSharedPtr<IBlueprintEditor>)nullptr, nullptr));
 }
 
 FRigVMVariableDetailCustomization::FRigVMVariableDetailCustomization(TSharedPtr<IBlueprintEditor> RigVMigEditor, UBlueprint* Blueprint)
