Control Rig: Fix graph details customization showing pins twice
#rnx

diff --git a/Engine/Plugins/Runtime/RigVM/Source/RigVMEditor/Private/Editor/RigVMGraphDetailCustomization.cpp b/Engine/Plugins/Runtime/RigVM/Source/RigVMEditor/Private/Editor/RigVMGraphDetailCustomization.cpp
--- a/Engine/Plugins/Runtime/RigVM/Source/RigVMEditor/Private/Editor/RigVMGraphDetailCustomization.cpp
+++ b/Engine/Plugins/Runtime/RigVM/Source/RigVMEditor/Private/Editor/RigVMGraphDetailCustomization.cpp
@@ -735,7 +735,7 @@ TSharedRef<IDetailCustomization> FRigVMGraphDetailCustomization::MakeInstance(TS
 	{
 		if (URigVMBlueprint* RigVMBlueprint = Cast<URigVMBlueprint>((*Objects)[0]))
 		{
-			if(RigVMBlueprint->GetClass() == InExpectedBlueprintClass || RigVMBlueprint->GetClass()->IsChildOf(InExpectedBlueprintClass))
+			if(RigVMBlueprint->GetClass() == InExpectedBlueprintClass)
 			{
 				return MakeShareable(new FRigVMGraphDetailCustomization(InEditor, RigVMBlueprint));
 			}
@@ -753,7 +753,7 @@ TSharedPtr<IDetailCustomization> FRigVMGraphDetailCustomization::MakeLegacyInsta
 	{
 		if (URigVMBlueprint* RigVMBlueprint = Cast<URigVMBlueprint>((*Objects)[0]))
 		{
-			if(RigVMBlueprint->GetClass() == InExpectedBlueprintClass || RigVMBlueprint->GetClass()->IsChildOf(InExpectedBlueprintClass))
+			if(RigVMBlueprint->GetClass() == InExpectedBlueprintClass)
 			{
 				return MakeShareable(new FRigVMGraphDetailCustomization(InBlueprintEditor, RigVMBlueprint));
 			}
