RigVM: Fix patch code in cases where nothing is connected to the return pin

#rnx

diff --git a/Engine/Plugins/Runtime/RigVM/Source/RigVMDeveloper/Private/RigVMModel/RigVMController.cpp b/Engine/Plugins/Runtime/RigVM/Source/RigVMDeveloper/Private/RigVMModel/RigVMController.cpp
--- a/Engine/Plugins/Runtime/RigVM/Source/RigVMDeveloper/Private/RigVMModel/RigVMController.cpp
+++ b/Engine/Plugins/Runtime/RigVM/Source/RigVMDeveloper/Private/RigVMModel/RigVMController.cpp
@@ -23001,7 +23001,8 @@ FRigVMClientPatchResult URigVMController::PatchFunctionsWithInvalidReturnPaths()
 		URigVMPin* EntryExecutePin = FindExecutePin(EntryNode);
 		URigVMPin* ReturnExecutePin = FindExecutePin(ReturnNode);
 
-		URigVMPin* ExecuteSourcePin = ReturnExecutePin->GetLinkedSourcePins()[0];
+		const TArray<URigVMPin*> LinkedSourcePins = ReturnExecutePin->GetLinkedSourcePins();
+		URigVMPin* ExecuteSourcePin = LinkedSourcePins.IsEmpty() ? nullptr : LinkedSourcePins[0];
 		while(ExecuteSourcePin && !ExecuteSourcePin->GetNode()->IsA<URigVMFunctionEntryNode>())
 		{
 			URigVMPin* ExecuteTargetPin = ExecuteSourcePin;
