RigVM: Use proper display names for execute pin

#rnx

diff --git a/Engine/Plugins/Runtime/RigVM/Source/RigVMDeveloper/Private/RigVMModel/Nodes/RigVMFunctionReferenceNode.cpp b/Engine/Plugins/Runtime/RigVM/Source/RigVMDeveloper/Private/RigVMModel/Nodes/RigVMFunctionReferenceNode.cpp
--- a/Engine/Plugins/Runtime/RigVM/Source/RigVMDeveloper/Private/RigVMModel/Nodes/RigVMFunctionReferenceNode.cpp
+++ b/Engine/Plugins/Runtime/RigVM/Source/RigVMDeveloper/Private/RigVMModel/Nodes/RigVMFunctionReferenceNode.cpp
@@ -22,9 +22,12 @@ FText URigVMFunctionReferenceNode::GetToolTipText() const
 	return ReferencedFunctionHeader.GetTooltip();
 }
 
-FName URigVMFunctionReferenceNode::GetDisplayNameForPin(const FString& InPinPath) const
+FName URigVMFunctionReferenceNode::GetDisplayNameForPin(const URigVMPin* InPin) const
 {
-	if(const FString* DisplayName = ReferencedFunctionHeader.Layout.FindDisplayName(InPinPath))
+	check(InPin);
+	
+	const FString PinPath = InPin->GetPinPath();
+	if(const FString* DisplayName = ReferencedFunctionHeader.Layout.FindDisplayName(PinPath))
 	{
 		if(!DisplayName->IsEmpty())
 		{
@@ -32,31 +35,28 @@ FName URigVMFunctionReferenceNode::GetDisplayNameForPin(const FString& InPinPath
 		}
 	}
 
-	if(const URigVMPin* Pin = FindPin(InPinPath))
+	if(InPin->IsRootPin())
 	{
-		if(Pin->IsRootPin())
+		const FRigVMGraphFunctionArgument* Argument = ReferencedFunctionHeader.Arguments.FindByPredicate([InPin](const FRigVMGraphFunctionArgument& Argument)
 		{
-			const FRigVMGraphFunctionArgument* Argument = ReferencedFunctionHeader.Arguments.FindByPredicate([Pin](const FRigVMGraphFunctionArgument& Argument)
-			{
-				return Argument.Name == Pin->GetFName();
-			});
+			return Argument.Name == InPin->GetFName();
+		});
 
-			if (Argument && !Argument->DisplayName.IsNone())
-			{
-				return Argument->DisplayName;
-			}
+		if (Argument && !Argument->DisplayName.IsNone())
+		{
+			return Argument->DisplayName;
 		}
 	}
 
 	if (IsReferencedNodeLoaded())
 	{
-		if(const URigVMPin* ReferencedPin = FindReferencedPin(InPinPath))
+		if(const URigVMPin* ReferencedPin = FindReferencedPin(InPin))
 		{
 			return ReferencedPin->GetDisplayName();
 		}
 	}
 
-	return Super::GetDisplayNameForPin(InPinPath);
+	return Super::GetDisplayNameForPin(InPin);
 }
 
 FString URigVMFunctionReferenceNode::GetCategoryForPin(const FString& InPinPath) const
diff --git a/Engine/Plugins/Runtime/RigVM/Source/RigVMDeveloper/Private/RigVMModel/Nodes/RigVMTemplateNode.cpp b/Engine/Plugins/Runtime/RigVM/Source/RigVMDeveloper/Private/RigVMModel/Nodes/RigVMTemplateNode.cpp
--- a/Engine/Plugins/Runtime/RigVM/Source/RigVMDeveloper/Private/RigVMModel/Nodes/RigVMTemplateNode.cpp
+++ b/Engine/Plugins/Runtime/RigVM/Source/RigVMDeveloper/Private/RigVMModel/Nodes/RigVMTemplateNode.cpp
@@ -667,14 +667,30 @@ FString URigVMTemplateNode::GetInitialDefaultValueForPin(const FName& InRootPinN
 	return DefaultValue;
 }
 
-FName URigVMTemplateNode::GetDisplayNameForPin(const FString& InPinPath) const
+FName URigVMTemplateNode::GetDisplayNameForPin(const URigVMPin* InPin) const
 {
 #if WITH_EDITOR
 	if(const FRigVMTemplate* Template = GetTemplate())
 	{
-		if(!InPinPath.Contains(TEXT(".")))
+		check(InPin);
+		
+		if(InPin->IsRootPin())
 		{
-			const FName RootPinName = *InPinPath;
+			if(InPin->IsExecuteContext())
+			{
+				if(InPin->GetDirection() == ERigVMPinDirection::IO)
+				{
+					return FRigVMStruct::ExecuteName;
+				}
+
+				const FRigVMDispatchContext DispatchContext;
+				if(Template->NumExecuteArguments(DispatchContext) == 1)
+				{
+					return FRigVMStruct::ExecuteName;
+				}
+			}
+		
+			const FName RootPinName = InPin->GetFName();
 			const TArray<int32> AllPermutations = GetResolvedPermutationIndices(false);
 
 			const FText DisplayNameText = Template->GetDisplayNameForArgument(RootPinName, AllPermutations);
@@ -699,7 +715,7 @@ FName URigVMTemplateNode::GetDisplayNameForPin(const FString& InPinPath) const
 		}
 	}
 #endif
-	return Super::GetDisplayNameForPin(InPinPath);
+	return Super::GetDisplayNameForPin(InPin);
 }
 
 TRigVMTypeIndex URigVMTemplateNode::TryReduceTypesToSingle(const TArray<TRigVMTypeIndex>& InTypes, const TRigVMTypeIndex PreferredType) const
diff --git a/Engine/Plugins/Runtime/RigVM/Source/RigVMDeveloper/Private/RigVMModel/Nodes/RigVMUnitNode.cpp b/Engine/Plugins/Runtime/RigVM/Source/RigVMDeveloper/Private/RigVMModel/Nodes/RigVMUnitNode.cpp
--- a/Engine/Plugins/Runtime/RigVM/Source/RigVMDeveloper/Private/RigVMModel/Nodes/RigVMUnitNode.cpp
+++ b/Engine/Plugins/Runtime/RigVM/Source/RigVMDeveloper/Private/RigVMModel/Nodes/RigVMUnitNode.cpp
@@ -323,13 +323,45 @@ FName URigVMUnitNode::GetNextAggregateName(const FName& InLastAggregatePinName)
 #endif
 }
 
-FName URigVMUnitNode::GetDisplayNameForPin(const FString& InPinPath) const
+FName URigVMUnitNode::GetDisplayNameForPin(const URigVMPin* InPin) const
 {
 	if(const UScriptStruct* ScriptStruct = GetScriptStruct())
 	{
-		return GetDisplayNameForStructMember(ScriptStruct, InPinPath);
+		check(InPin);
+		
+		if(InPin->IsExecuteContext())
+		{
+			if(InPin->GetDirection() == ERigVMPinDirection::IO)
+			{
+				return FRigVMStruct::ExecuteName;
+			}
+
+			int32 NumExecuteContextPins = 0;
+			for (TFieldIterator<FProperty> It(ScriptStruct); It; ++It)
+			{
+				if(const FStructProperty* StructProperty = CastField<FStructProperty>(*It))
+				{
+					if(StructProperty->Struct->IsChildOf(FRigVMExecutePin::StaticStruct()))
+					{
+						NumExecuteContextPins++;
+					}
+				}
+			}
+			if(NumExecuteContextPins == 1)
+			{
+				return FRigVMStruct::ExecuteName;
+			}
+		}
+		
+		const FName StructMemberDisplayName = GetDisplayNameForStructMember(ScriptStruct, InPin->GetPinPath());
+		if((StructMemberDisplayName == FRigVMStruct::ExecuteContextName) ||
+			(InPin->GetFName() == FRigVMStruct::ExecuteContextName && StructMemberDisplayName.IsNone()))
+		{
+			return FRigVMStruct::ExecuteName;
+		}
+		return StructMemberDisplayName;
 	}
-	return Super::GetDisplayNameForPin(InPinPath);
+	return Super::GetDisplayNameForPin(InPin);
 }
 
 UScriptStruct* URigVMUnitNode::GetScriptStruct() const
diff --git a/Engine/Plugins/Runtime/RigVM/Source/RigVMDeveloper/Private/RigVMModel/RigVMController.cpp b/Engine/Plugins/Runtime/RigVM/Source/RigVMDeveloper/Private/RigVMModel/RigVMController.cpp
--- a/Engine/Plugins/Runtime/RigVM/Source/RigVMDeveloper/Private/RigVMModel/RigVMController.cpp
+++ b/Engine/Plugins/Runtime/RigVM/Source/RigVMDeveloper/Private/RigVMModel/RigVMController.cpp
@@ -19413,7 +19413,7 @@ void URigVMController::ResolveTemplateNodeMetaData(URigVMTemplateNode* InNode, b
 
 	for(URigVMPin* Pin : InNode->GetPins())
 	{
-		const FName DisplayName = InNode->GetDisplayNameForPin(Pin->GetName());
+		const FName DisplayName = InNode->GetDisplayNameForPin(Pin);
 		if(Pin->DisplayName != DisplayName)
 		{
 			Pin->DisplayName = DisplayName;
diff --git a/Engine/Plugins/Runtime/RigVM/Source/RigVMDeveloper/Private/RigVMModel/RigVMNode.cpp b/Engine/Plugins/Runtime/RigVM/Source/RigVMDeveloper/Private/RigVMModel/RigVMNode.cpp
--- a/Engine/Plugins/Runtime/RigVM/Source/RigVMDeveloper/Private/RigVMModel/RigVMNode.cpp
+++ b/Engine/Plugins/Runtime/RigVM/Source/RigVMDeveloper/Private/RigVMModel/RigVMNode.cpp
@@ -11,6 +11,7 @@
 #include "RigVMCore/RigVMStruct.h"
 #include "RigVMUserWorkflowRegistry.h"
 #include "Logging/LogScopedVerbosityOverride.h"
+#include "Algo/Count.h"
 
 #include UE_INLINE_GENERATED_CPP_BY_NAME(RigVMNode)
 
@@ -380,17 +381,17 @@ FRigVMNodeLayout URigVMNode::GetNodeLayout(bool bIncludeEmptyCategories) const
 	// fill in all user provided display names and pin category indices
 	for(const URigVMPin* Pin : AllPins)
 	{
-		const FString SegmentPath = Pin->GetSegmentPath(true);
-		
 		if(!Pin->GetCategory().IsEmpty() && Pin->GetIndexInCategory() != INDEX_NONE)
 		{
+			const FString SegmentPath = Pin->GetSegmentPath(true);
 			Layout.PinIndexInCategory.Add(SegmentPath, Pin->GetIndexInCategory());
 		}
 
 		if(!Pin->DisplayName.IsNone())
 		{
-			if(!Pin->DisplayName.IsEqual(GetDisplayNameForPin(SegmentPath), ENameCase::CaseSensitive))
+			if(!Pin->DisplayName.IsEqual(GetDisplayNameForPin(Pin), ENameCase::CaseSensitive))
 			{
+				const FString SegmentPath = Pin->GetSegmentPath(true);
 				Layout.DisplayNames.Add(SegmentPath, Pin->DisplayName.ToString());
 			}
 		}
@@ -973,9 +974,31 @@ FName URigVMNode::GetDisplayNameForPin(const FString& InPinPath) const
 {
 	if(const URigVMPin* Pin = FindPin(InPinPath))
 	{
-		if(Pin->IsArrayElement())
+		return GetDisplayNameForPin(Pin);
+	}
+	return NAME_None;
+}
+
+FName URigVMNode::GetDisplayNameForPin(const URigVMPin* InPin) const
+{
+	check(InPin);
+	if(InPin->IsArrayElement())
+	{
+		return *FString::FromInt(InPin->GetPinIndex());
+	}
+	if(InPin->IsExecuteContext())
+	{
+		if(InPin->GetDirection() == ERigVMPinDirection::IO)
+		{
+			return FRigVMStruct::ExecuteName;
+		}
+		const int32 NumExecutePins = Algo::CountIf(GetPins(), [](const URigVMPin* Pin) -> bool
+		{
+			return Pin->IsExecuteContext();
+		});
+		if(NumExecutePins == 1)
 		{
-			return *FString::FromInt(Pin->GetPinIndex());
+			return FRigVMStruct::ExecuteName;
 		}
 	}
 	return NAME_None;
diff --git a/Engine/Plugins/Runtime/RigVM/Source/RigVMDeveloper/Private/RigVMModel/RigVMPin.cpp b/Engine/Plugins/Runtime/RigVM/Source/RigVMDeveloper/Private/RigVMModel/RigVMPin.cpp
--- a/Engine/Plugins/Runtime/RigVM/Source/RigVMDeveloper/Private/RigVMModel/RigVMPin.cpp
+++ b/Engine/Plugins/Runtime/RigVM/Source/RigVMDeveloper/Private/RigVMModel/RigVMPin.cpp
@@ -407,7 +407,7 @@ FName URigVMPin::GetDisplayName() const
 		
 		if(const URigVMNode* Node = GetNode())
 		{
-			const FName DisplayNameFromNode = Node->GetDisplayNameForPin(this->GetSegmentPath(true));
+			const FName DisplayNameFromNode = Node->GetDisplayNameForPin(this);
 			if(!DisplayNameFromNode.IsNone())
 			{
 				return DisplayNameFromNode;
diff --git a/Engine/Plugins/Runtime/RigVM/Source/RigVMDeveloper/Public/RigVMModel/Nodes/RigVMFunctionReferenceNode.h b/Engine/Plugins/Runtime/RigVM/Source/RigVMDeveloper/Public/RigVMModel/Nodes/RigVMFunctionReferenceNode.h
--- a/Engine/Plugins/Runtime/RigVM/Source/RigVMDeveloper/Public/RigVMModel/Nodes/RigVMFunctionReferenceNode.h
+++ b/Engine/Plugins/Runtime/RigVM/Source/RigVMDeveloper/Public/RigVMModel/Nodes/RigVMFunctionReferenceNode.h
@@ -24,7 +24,7 @@ public:
 	virtual FString GetNodeTitle() const override;
 	virtual FLinearColor GetNodeColor() const override;
 	virtual FText GetToolTipText() const override;
-	virtual FName GetDisplayNameForPin(const FString& InPinPath) const override;
+	virtual FName GetDisplayNameForPin(const URigVMPin* InPin) const override;
 	virtual FString GetCategoryForPin(const FString& InPinPath) const override;
 	virtual int32 GetIndexInCategoryForPin(const FString& InPinPath) const override;
 	virtual FText GetToolTipTextForPin(const URigVMPin* InPin) const override;
diff --git a/Engine/Plugins/Runtime/RigVM/Source/RigVMDeveloper/Public/RigVMModel/Nodes/RigVMTemplateNode.h b/Engine/Plugins/Runtime/RigVM/Source/RigVMDeveloper/Public/RigVMModel/Nodes/RigVMTemplateNode.h
--- a/Engine/Plugins/Runtime/RigVM/Source/RigVMDeveloper/Public/RigVMModel/Nodes/RigVMTemplateNode.h
+++ b/Engine/Plugins/Runtime/RigVM/Source/RigVMDeveloper/Public/RigVMModel/Nodes/RigVMTemplateNode.h
@@ -132,7 +132,7 @@ public:
 	FString GetInitialDefaultValueForPin(const FName& InRootPinName, const TArray<int32>& InPermutationIndices = TArray<int32>()) const;
 
 	// returns the display name for a pin
-	virtual FName GetDisplayNameForPin(const FString& InPinPath) const override;
+	virtual FName GetDisplayNameForPin(const URigVMPin* InPin) const override;
 
 	// Tries to reduce the input types to a single type, if all are compatible
 	// Will prioritize the InPreferredType if available
diff --git a/Engine/Plugins/Runtime/RigVM/Source/RigVMDeveloper/Public/RigVMModel/Nodes/RigVMUnitNode.h b/Engine/Plugins/Runtime/RigVM/Source/RigVMDeveloper/Public/RigVMModel/Nodes/RigVMUnitNode.h
--- a/Engine/Plugins/Runtime/RigVM/Source/RigVMDeveloper/Public/RigVMModel/Nodes/RigVMUnitNode.h
+++ b/Engine/Plugins/Runtime/RigVM/Source/RigVMDeveloper/Public/RigVMModel/Nodes/RigVMUnitNode.h
@@ -37,7 +37,7 @@ public:
 	virtual TArray<URigVMPin*> GetAggregateInputs() const override;
 	virtual TArray<URigVMPin*> GetAggregateOutputs() const override;
 	virtual FName GetNextAggregateName(const FName& InLastAggregatePinName) const override;
-	virtual FName GetDisplayNameForPin(const FString& InPinPath) const override;
+	virtual FName GetDisplayNameForPin(const URigVMPin* InPin) const override;
 
 	virtual bool IsOutDated() const override;
 	virtual FString GetDeprecatedMetadata() const override;
diff --git a/Engine/Plugins/Runtime/RigVM/Source/RigVMDeveloper/Public/RigVMModel/RigVMNode.h b/Engine/Plugins/Runtime/RigVM/Source/RigVMDeveloper/Public/RigVMModel/RigVMNode.h
--- a/Engine/Plugins/Runtime/RigVM/Source/RigVMDeveloper/Public/RigVMModel/RigVMNode.h
+++ b/Engine/Plugins/Runtime/RigVM/Source/RigVMDeveloper/Public/RigVMModel/RigVMNode.h
@@ -384,7 +384,8 @@ public:
 	}
 
 	// returns the display name for a pin
-	virtual FName GetDisplayNameForPin(const FString& InPinPath) const;
+	FName GetDisplayNameForPin(const FString& InPinPath) const;
+	virtual FName GetDisplayNameForPin(const URigVMPin* InPin) const;
 
 	// returns the display name for a struct member
 	static FName GetDisplayNameForStructMember(const UStruct* InStruct, const FString& InPath);
