RigVM: Filter node spawners using new ExecutePin struct

#rnx

diff --git a/Engine/Plugins/Runtime/RigVM/Source/RigVMDeveloper/Private/EdGraph/RigVMEdGraphSchema.cpp b/Engine/Plugins/Runtime/RigVM/Source/RigVMDeveloper/Private/EdGraph/RigVMEdGraphSchema.cpp
--- a/Engine/Plugins/Runtime/RigVM/Source/RigVMDeveloper/Private/EdGraph/RigVMEdGraphSchema.cpp
+++ b/Engine/Plugins/Runtime/RigVM/Source/RigVMDeveloper/Private/EdGraph/RigVMEdGraphSchema.cpp
@@ -2000,6 +2000,24 @@ bool URigVMEdGraphSchema::ArePinsCompatible(const UEdGraphPin* PinA, const UEdGr
 		}
 	}
 
+	if(PinA->PinType.ContainerType == EPinContainerType::None &&
+		PinB->PinType.ContainerType == EPinContainerType::None &&
+		PinA->PinType.PinCategory == UEdGraphSchema_K2::PC_Struct &&
+		PinB->PinType.PinCategory == UEdGraphSchema_K2::PC_Struct)
+	{
+		if(const UScriptStruct* ScriptStructA = Cast<UScriptStruct>(PinA->PinType.PinSubCategoryObject))
+		{
+			if(const UScriptStruct* ScriptStructB = Cast<UScriptStruct>(PinB->PinType.PinSubCategoryObject))
+			{
+				if(ScriptStructA->IsChildOf(FRigVMExecutePin::StaticStruct()) &&
+					ScriptStructB->IsChildOf(FRigVMExecutePin::StaticStruct()))
+				{
+					return true;
+				}
+			}
+		}
+	}
+
 	if(GetDefault<UEdGraphSchema_K2>()->ArePinsCompatible(PinA, PinB, CallingContext, bIgnoreArray))
 	{
 		return true;
