Making UClass::ClassDefaultObject private.

diff --git a/Engine/Plugins/Animation/IKRig/Source/IKRigEditor/Private/RigEditor/IKRigEditorController.cpp b/Engine/Plugins/Animation/IKRig/Source/IKRigEditor/Private/RigEditor/IKRigEditorController.cpp
--- a/Engine/Plugins/Animation/IKRig/Source/IKRigEditor/Private/RigEditor/IKRigEditorController.cpp
+++ b/Engine/Plugins/Animation/IKRig/Source/IKRigEditor/Private/RigEditor/IKRigEditorController.cpp
@@ -873,7 +873,7 @@ bool FIKRigEditorController::PromptToAddDefaultSolver() const
 			continue; // skip base class
 		}
 
-		//const UIKRigSolver* SolverCDO = CastChecked<UIKRigSolver>(Class->ClassDefaultObject);
+		//const UIKRigSolver* SolverCDO = CastChecked<UIKRigSolver>(Class->GetDefaultObject(false));
 		TSharedPtr<FIKRigSolverTypeAndName> SolverType = MakeShared<FIKRigSolverTypeAndName>();
 		SolverType->NiceName = FText::FromString(Struct->GetName());
 		SolverType->SolverType = Struct;
diff --git a/Engine/Plugins/Compositing/Composure/Source/ComposureLayersEditor/Private/ComposureDetailCustomizations.cpp b/Engine/Plugins/Compositing/Composure/Source/ComposureLayersEditor/Private/ComposureDetailCustomizations.cpp
--- a/Engine/Plugins/Compositing/Composure/Source/ComposureLayersEditor/Private/ComposureDetailCustomizations.cpp
+++ b/Engine/Plugins/Compositing/Composure/Source/ComposureLayersEditor/Private/ComposureDetailCustomizations.cpp
@@ -463,7 +463,7 @@ void FCompElementDetailsCustomization::OnCameraSourceSelected(const FString& Sel
 				if (Outer)
 				{
 					UClass* ObjClass = Outer->GetClass();
-					ACompositingElement* CDO = ObjClass ? Cast<ACompositingElement>(ObjClass->ClassDefaultObject): nullptr;
+					const ACompositingElement* CDO = ObjClass ? Cast<ACompositingElement>(ObjClass->GetDefaultObject(false)): nullptr;
 
 					if (CDO && CDO->CameraSource != ESceneCameraLinkType::Unused)
 					{
diff --git a/Engine/Plugins/Editor/ObjectMixer/ObjectMixer/Source/ObjectMixer/Private/Views/Widgets/SObjectMixerPlacementAssetMenuEntry.cpp b/Engine/Plugins/Editor/ObjectMixer/ObjectMixer/Source/ObjectMixer/Private/Views/Widgets/SObjectMixerPlacementAssetMenuEntry.cpp
--- a/Engine/Plugins/Editor/ObjectMixer/ObjectMixer/Source/ObjectMixer/Private/Views/Widgets/SObjectMixerPlacementAssetMenuEntry.cpp
+++ b/Engine/Plugins/Editor/ObjectMixer/ObjectMixer/Source/ObjectMixer/Private/Views/Widgets/SObjectMixerPlacementAssetMenuEntry.cpp
@@ -36,7 +36,7 @@ void SObjectMixerPlacementAssetMenuEntry::Construct(const FArguments& InArgs, co
 	}
 	else if (bIsActor)
 	{
-		DefaultActor = CastChecked<AActor>(CastChecked<UClass>(Item->AssetData.GetAsset())->ClassDefaultObject);
+		DefaultActor = CastChecked<AActor>(CastChecked<UClass>(Item->AssetData.GetAsset())->GetDefaultObject(false));
 	}
 
 	UClass* DocClass = nullptr;
diff --git a/Engine/Plugins/FX/Niagara/Source/NiagaraEditor/Private/NiagaraNodeCustomHlsl.cpp b/Engine/Plugins/FX/Niagara/Source/NiagaraEditor/Private/NiagaraNodeCustomHlsl.cpp
--- a/Engine/Plugins/FX/Niagara/Source/NiagaraEditor/Private/NiagaraNodeCustomHlsl.cpp
+++ b/Engine/Plugins/FX/Niagara/Source/NiagaraEditor/Private/NiagaraNodeCustomHlsl.cpp
@@ -289,7 +289,7 @@ bool UNiagaraNodeCustomHlsl::CallsImpureDataInterfaceFunctions() const
 		FNiagaraTypeDefinition NiagaraType = UEdGraphSchema_Niagara::PinToTypeDefinition(InputPin, ENiagaraStructConversion::Simulation);
 		if (NiagaraType.IsDataInterface())
 		{
-			if (UNiagaraDataInterface* DataInterfaceClass = CastChecked<UNiagaraDataInterface>(NiagaraType.GetClass()->ClassDefaultObject))
+			if (UNiagaraDataInterface* DataInterfaceClass = CastChecked<UNiagaraDataInterface>(NiagaraType.GetClass()->GetDefaultObject(false)))
 			{
 				TArray<FNiagaraFunctionSignature> FunctionSignatures;
 				DataInterfaceClass->GetFunctionSignatures(FunctionSignatures);
diff --git a/Engine/Plugins/MovieScene/MovieRenderPipeline/Source/MovieRenderPipelineCore/Private/Graph/MovieGraphConfig.cpp b/Engine/Plugins/MovieScene/MovieRenderPipeline/Source/MovieRenderPipelineCore/Private/Graph/MovieGraphConfig.cpp
--- a/Engine/Plugins/MovieScene/MovieRenderPipeline/Source/MovieRenderPipelineCore/Private/Graph/MovieGraphConfig.cpp
+++ b/Engine/Plugins/MovieScene/MovieRenderPipeline/Source/MovieRenderPipelineCore/Private/Graph/MovieGraphConfig.cpp
@@ -1446,7 +1446,7 @@ void UMovieGraphConfig::GetOutputDirectory(FString& OutOutputDirectory) const
 		if (OutOutputDirectory.IsEmpty())
 		{
 			// If we didn't find any overrides, use the CDO values
-			UMovieGraphGlobalOutputSettingNode* CDO = Cast<UMovieGraphGlobalOutputSettingNode>(UMovieGraphGlobalOutputSettingNode::StaticClass()->ClassDefaultObject);
+			UMovieGraphGlobalOutputSettingNode* CDO = Cast<UMovieGraphGlobalOutputSettingNode>(UMovieGraphGlobalOutputSettingNode::StaticClass()->GetDefaultObject(false));
 			check(CDO);
 			
 			OutOutputDirectory = CDO->OutputDirectory.Path;
diff --git a/Engine/Plugins/Mutable/Source/CustomizableObjectEditor/Private/MuCOE/Nodes/CustomizableObjectNodeMaterial.cpp b/Engine/Plugins/Mutable/Source/CustomizableObjectEditor/Private/MuCOE/Nodes/CustomizableObjectNodeMaterial.cpp
--- a/Engine/Plugins/Mutable/Source/CustomizableObjectEditor/Private/MuCOE/Nodes/CustomizableObjectNodeMaterial.cpp
+++ b/Engine/Plugins/Mutable/Source/CustomizableObjectEditor/Private/MuCOE/Nodes/CustomizableObjectNodeMaterial.cpp
@@ -103,7 +103,7 @@ bool UCustomizableObjectNodeMaterialRemapPinsByName::HasSavedPinData(const UCust
 
 bool UCustomizableObjectNodeMaterialPinDataImage::IsDefault() const
 {
-	const UCustomizableObjectNodeMaterialPinDataImage* Default = Cast<UCustomizableObjectNodeMaterialPinDataImage>(GetClass()->ClassDefaultObject);
+	const UCustomizableObjectNodeMaterialPinDataImage* Default = Cast<UCustomizableObjectNodeMaterialPinDataImage>(GetClass()->GetDefaultObject(false));
 
 	return PinMode == Default->PinMode &&
 		UVLayoutMode == Default->UVLayoutMode &&
diff --git a/Engine/Plugins/PCG/Source/PCG/Private/Helpers/PCGConversion.cpp b/Engine/Plugins/PCG/Source/PCG/Private/Helpers/PCGConversion.cpp
--- a/Engine/Plugins/PCG/Source/PCG/Private/Helpers/PCGConversion.cpp
+++ b/Engine/Plugins/PCG/Source/PCG/Private/Helpers/PCGConversion.cpp
@@ -18,7 +18,7 @@ namespace PCGConversion
 		{
 			if (IsValid(Class))
 			{
-				return CastChecked<UPCGSettings>(Class->ClassDefaultObject)->GetDefaultNodeTitle();
+				return CastChecked<UPCGSettings>(Class->GetDefaultObject(false))->GetDefaultNodeTitle();
 			}
 
 			return {};
diff --git a/Engine/Plugins/PCG/Source/PCGEditor/Private/PCGEditorCommands.cpp b/Engine/Plugins/PCG/Source/PCGEditor/Private/PCGEditorCommands.cpp
--- a/Engine/Plugins/PCG/Source/PCGEditor/Private/PCGEditorCommands.cpp
+++ b/Engine/Plugins/PCG/Source/PCGEditor/Private/PCGEditorCommands.cpp
@@ -109,7 +109,7 @@ void FPCGEditorSpawnNodeCommands::RegisterCommands()
 			const bool bOverrideLabel = FParse::Value(*CurrentEntry, TEXT("Label="), OutValue, /*bShouldStopOnSeparator=*/true);
 			SpawnCommandInfo->PreconfiguredInfo.Label = FText::FromString(OutValue);
 
-			const UPCGSettings* CDO = CastChecked<UPCGSettings>(FoundClass->ClassDefaultObject);
+			const UPCGSettings* CDO = CastChecked<UPCGSettings>(FoundClass->GetDefaultObject(false));
 			const FText CommandLabelText = bOverrideLabel ? SpawnCommandInfo->PreconfiguredInfo.Label : CDO->GetDefaultNodeTitle();
 			const FText Description = FText::Format(NSLOCTEXT("PCGEditor", "SpawnNodeDescription", "Hold down the bound keys and left click in the graph panel to spawn a {0} node."), CommandLabelText);
 
diff --git a/Engine/Plugins/PCG/Source/PCGEditor/Private/Widgets/SPCGEditorGraphDebugObjectTree.cpp b/Engine/Plugins/PCG/Source/PCGEditor/Private/Widgets/SPCGEditorGraphDebugObjectTree.cpp
--- a/Engine/Plugins/PCG/Source/PCGEditor/Private/Widgets/SPCGEditorGraphDebugObjectTree.cpp
+++ b/Engine/Plugins/PCG/Source/PCGEditor/Private/Widgets/SPCGEditorGraphDebugObjectTree.cpp
@@ -31,7 +31,7 @@ namespace PCGEditorGraphDebugObjectTree
 	{
 		OutBrush = Item->GetIcon();
 
-		const FLinearColor DefaultGraphColor = CastChecked<UAssetDefinition_PCGGraphInterface>(UAssetDefinition_PCGGraphInterface::StaticClass()->ClassDefaultObject)->GetAssetColor();
+		const FLinearColor DefaultGraphColor = CastChecked<UAssetDefinition_PCGGraphInterface>(UAssetDefinition_PCGGraphInterface::StaticClass()->GetDefaultObject(false))->GetAssetColor();
 
 		OutColorAndOpacity = Item->IsDebuggable() ? DefaultGraphColor : FLinearColor(1.0f, 1.0f, 1.0f, 0.5f);
 	}
diff --git a/Engine/Plugins/Runtime/CommonUI/Source/CommonUI/Private/CommonBorder.cpp b/Engine/Plugins/Runtime/CommonUI/Source/CommonUI/Private/CommonBorder.cpp
--- a/Engine/Plugins/Runtime/CommonUI/Source/CommonUI/Private/CommonBorder.cpp
+++ b/Engine/Plugins/Runtime/CommonUI/Source/CommonUI/Private/CommonBorder.cpp
@@ -174,7 +174,7 @@ const UCommonBorderStyle* UCommonBorder::GetStyleCDO() const
 {
 	if (Style)
 	{
-		if (const UCommonBorderStyle* BorderStyle = Cast<UCommonBorderStyle>(Style->ClassDefaultObject))
+		if (const UCommonBorderStyle* BorderStyle = Cast<UCommonBorderStyle>(Style->GetDefaultObject(false)))
 		{
 			return BorderStyle;
 		}
diff --git a/Engine/Plugins/Runtime/CommonUI/Source/CommonUI/Private/CommonButtonBase.cpp b/Engine/Plugins/Runtime/CommonUI/Source/CommonUI/Private/CommonButtonBase.cpp
--- a/Engine/Plugins/Runtime/CommonUI/Source/CommonUI/Private/CommonButtonBase.cpp
+++ b/Engine/Plugins/Runtime/CommonUI/Source/CommonUI/Private/CommonButtonBase.cpp
@@ -63,7 +63,7 @@ UCommonTextStyle* UCommonButtonStyle::GetNormalTextStyle() const
 {
 	if (NormalTextStyle)
 	{
-		if (UCommonTextStyle* TextStyle = Cast<UCommonTextStyle>(NormalTextStyle->ClassDefaultObject))
+		if (UCommonTextStyle* TextStyle = Cast<UCommonTextStyle>(NormalTextStyle->GetDefaultObject(false)))
 		{
 			return TextStyle;
 		}
@@ -75,7 +75,7 @@ UCommonTextStyle* UCommonButtonStyle::GetNormalHoveredTextStyle() const
 {
 	if (NormalHoveredTextStyle)
 	{
-		if (UCommonTextStyle* TextStyle = Cast<UCommonTextStyle>(NormalHoveredTextStyle->ClassDefaultObject))
+		if (UCommonTextStyle* TextStyle = Cast<UCommonTextStyle>(NormalHoveredTextStyle->GetDefaultObject(false)))
 		{
 			return TextStyle;
 		}
@@ -87,7 +87,7 @@ UCommonTextStyle* UCommonButtonStyle::GetSelectedTextStyle() const
 {
 	if (SelectedTextStyle)
 	{
-		if (UCommonTextStyle* TextStyle = Cast<UCommonTextStyle>(SelectedTextStyle->ClassDefaultObject))
+		if (UCommonTextStyle* TextStyle = Cast<UCommonTextStyle>(SelectedTextStyle->GetDefaultObject(false)))
 		{
 			return TextStyle;
 		}
@@ -99,7 +99,7 @@ UCommonTextStyle* UCommonButtonStyle::GetSelectedHoveredTextStyle() const
 {
 	if (SelectedHoveredTextStyle)
 	{
-		if (UCommonTextStyle* TextStyle = Cast<UCommonTextStyle>(SelectedHoveredTextStyle->ClassDefaultObject))
+		if (UCommonTextStyle* TextStyle = Cast<UCommonTextStyle>(SelectedHoveredTextStyle->GetDefaultObject(false)))
 		{
 			return TextStyle;
 		}
@@ -111,7 +111,7 @@ UCommonTextStyle* UCommonButtonStyle::GetDisabledTextStyle() const
 {
 	if (DisabledTextStyle)
 	{
-		if (UCommonTextStyle* TextStyle = Cast<UCommonTextStyle>(DisabledTextStyle->ClassDefaultObject))
+		if (UCommonTextStyle* TextStyle = Cast<UCommonTextStyle>(DisabledTextStyle->GetDefaultObject(false)))
 		{
 			return TextStyle;
 		}
@@ -1115,7 +1115,7 @@ const UCommonButtonStyle* UCommonButtonBase::GetStyleCDO() const
 {
 	if (Style)
 	{
-		if (const UCommonButtonStyle* CommonButtonStyle = Cast<UCommonButtonStyle>(Style->ClassDefaultObject))
+		if (const UCommonButtonStyle* CommonButtonStyle = Cast<UCommonButtonStyle>(Style->GetDefaultObject(false)))
 		{
 			return CommonButtonStyle;
 		}
diff --git a/Engine/Plugins/Runtime/CommonUI/Source/CommonUI/Private/CommonTextBlock.cpp b/Engine/Plugins/Runtime/CommonUI/Source/CommonUI/Private/CommonTextBlock.cpp
--- a/Engine/Plugins/Runtime/CommonUI/Source/CommonUI/Private/CommonTextBlock.cpp
+++ b/Engine/Plugins/Runtime/CommonUI/Source/CommonUI/Private/CommonTextBlock.cpp
@@ -130,7 +130,7 @@ void UCommonTextBlock::PostLoad()
 		UCommonUIEditorSettings& Settings = ICommonUIModule::GetEditorSettings();
 		Settings.ConditionalPostLoad();
 		TSubclassOf<UCommonTextStyle> DefaultStyle = Settings.GetTemplateTextStyle();
-		UCommonTextStyle* DefaultStyleCDO = DefaultStyle ? Cast<UCommonTextStyle>(DefaultStyle->ClassDefaultObject) : nullptr;
+		UCommonTextStyle* DefaultStyleCDO = DefaultStyle ? Cast<UCommonTextStyle>(DefaultStyle->GetDefaultObject(false)) : nullptr;
 		if (DefaultStyleCDO)
 		{
 			DefaultStyleCDO->ConditionalPostLoad();
@@ -615,12 +615,12 @@ TSharedRef<SWidget> UCommonTextBlock::RebuildWidget()
 
 const UCommonTextStyle* UCommonTextBlock::GetStyleCDO() const
 {
-	return Style ? Cast<UCommonTextStyle>(Style->ClassDefaultObject) : nullptr;
+	return Style ? Cast<UCommonTextStyle>(Style->GetDefaultObject(false)) : nullptr;
 }
 
 const UCommonTextScrollStyle* UCommonTextBlock::GetScrollStyleCDO() const
 {
-	return ScrollStyle ? Cast<UCommonTextScrollStyle>(ScrollStyle->ClassDefaultObject) : nullptr;
+	return ScrollStyle ? Cast<UCommonTextScrollStyle>(ScrollStyle->GetDefaultObject(false)) : nullptr;
 }
 
 void UCommonTextBlock::ApplyFontSizeMultiplier() const
diff --git a/Engine/Plugins/Runtime/GameplayAbilities/Source/GameplayAbilities/Private/GameplayCueManager.cpp b/Engine/Plugins/Runtime/GameplayAbilities/Source/GameplayAbilities/Private/GameplayCueManager.cpp
--- a/Engine/Plugins/Runtime/GameplayAbilities/Source/GameplayAbilities/Private/GameplayCueManager.cpp
+++ b/Engine/Plugins/Runtime/GameplayAbilities/Source/GameplayAbilities/Private/GameplayCueManager.cpp
@@ -1032,8 +1032,8 @@ void UGameplayCueManager::HandleAssetAdded(UObject *Object)
 	UBlueprint* Blueprint = Cast<UBlueprint>(Object);
 	if (Blueprint && Blueprint->GeneratedClass)
 	{
-		UGameplayCueNotify_Static* StaticCDO = Cast<UGameplayCueNotify_Static>(Blueprint->GeneratedClass->ClassDefaultObject);
-		AGameplayCueNotify_Actor* ActorCDO = Cast<AGameplayCueNotify_Actor>(Blueprint->GeneratedClass->ClassDefaultObject);
+		UGameplayCueNotify_Static* StaticCDO = Cast<UGameplayCueNotify_Static>(Blueprint->GeneratedClass->GetDefaultObject(false));
+		AGameplayCueNotify_Actor* ActorCDO = Cast<AGameplayCueNotify_Actor>(Blueprint->GeneratedClass->GetDefaultObject(false));
 		
 		if (StaticCDO || ActorCDO)
 		{
@@ -1073,8 +1073,8 @@ void UGameplayCueManager::HandleAssetDeleted(UObject *Object)
 	UBlueprint* Blueprint = Cast<UBlueprint>(Object);
 	if (Blueprint && Blueprint->GeneratedClass)
 	{
-		UGameplayCueNotify_Static* StaticCDO = Cast<UGameplayCueNotify_Static>(Blueprint->GeneratedClass->ClassDefaultObject);
-		AGameplayCueNotify_Actor* ActorCDO = Cast<AGameplayCueNotify_Actor>(Blueprint->GeneratedClass->ClassDefaultObject);
+		UGameplayCueNotify_Static* StaticCDO = Cast<UGameplayCueNotify_Static>(Blueprint->GeneratedClass->GetDefaultObject(false));
+		AGameplayCueNotify_Actor* ActorCDO = Cast<AGameplayCueNotify_Actor>(Blueprint->GeneratedClass->GetDefaultObject(false));
 		
 		if (StaticCDO || ActorCDO)
 		{
@@ -1106,8 +1106,8 @@ void UGameplayCueManager::HandleAssetRenamed(const FAssetData& Data, const FStri
 		UClass* DataClass = FindObject<UClass>(nullptr, *ParentClassName);
 		if (DataClass)
 		{
-			UGameplayCueNotify_Static* StaticCDO = Cast<UGameplayCueNotify_Static>(DataClass->ClassDefaultObject);
-			AGameplayCueNotify_Actor* ActorCDO = Cast<AGameplayCueNotify_Actor>(DataClass->ClassDefaultObject);
+			UGameplayCueNotify_Static* StaticCDO = Cast<UGameplayCueNotify_Static>(DataClass->GetDefaultObject(false));
+			AGameplayCueNotify_Actor* ActorCDO = Cast<AGameplayCueNotify_Actor>(DataClass->GetDefaultObject(false));
 			if (StaticCDO || ActorCDO)
 			{
 				VerifyNotifyAssetIsInValidPath(Data.PackagePath.ToString());
@@ -1549,7 +1549,7 @@ bool UGameplayCueManager::DoesPendingCueExecuteMatch(FGameplayCuePendingExecute&
 
 void UGameplayCueManager::CheckForPreallocation(UClass* GCClass)
 {
-	if (AGameplayCueNotify_Actor* InstancedCue = Cast<AGameplayCueNotify_Actor>(GCClass->ClassDefaultObject))
+	if (AGameplayCueNotify_Actor* InstancedCue = Cast<AGameplayCueNotify_Actor>(GCClass->GetDefaultObject(false)))
 	{
 		if (InstancedCue->NumPreallocatedInstances > 0 && GameplayCueClassesForPreallocation.Contains(GCClass) == false)
 		{
diff --git a/Engine/Plugins/Runtime/GameplayAbilities/Source/GameplayAbilities/Private/GameplayCueSet.cpp b/Engine/Plugins/Runtime/GameplayAbilities/Source/GameplayAbilities/Private/GameplayCueSet.cpp
--- a/Engine/Plugins/Runtime/GameplayAbilities/Source/GameplayAbilities/Private/GameplayCueSet.cpp
+++ b/Engine/Plugins/Runtime/GameplayAbilities/Source/GameplayAbilities/Private/GameplayCueSet.cpp
@@ -293,7 +293,7 @@ bool UGameplayCueSet::HandleGameplayCueNotify_Internal(AActor* TargetActor, int3
 		check(CueData.LoadedGameplayCueClass);
 
 		// Handle the Notify if we found something
-		if (UGameplayCueNotify_Static* NonInstancedCue = Cast<UGameplayCueNotify_Static>(CueData.LoadedGameplayCueClass->ClassDefaultObject))
+		if (UGameplayCueNotify_Static* NonInstancedCue = Cast<UGameplayCueNotify_Static>(CueData.LoadedGameplayCueClass->GetDefaultObject(false)))
 		{
 			if (NonInstancedCue->HandlesEvent(EventType))
 			{
@@ -310,7 +310,7 @@ bool UGameplayCueSet::HandleGameplayCueNotify_Internal(AActor* TargetActor, int3
 				HandleGameplayCueNotify_Internal(TargetActor, CueData.ParentDataIdx, EventType, Parameters);
 			}
 		}
-		else if (AGameplayCueNotify_Actor* InstancedCue = Cast<AGameplayCueNotify_Actor>(CueData.LoadedGameplayCueClass->ClassDefaultObject))
+		else if (AGameplayCueNotify_Actor* InstancedCue = Cast<AGameplayCueNotify_Actor>(CueData.LoadedGameplayCueClass->GetDefaultObject(false)))
 		{
 			bool bShouldDestroy = false;
 			if (EventType == EGameplayCueEvent::Executed && !Parameters.bGameplayEffectActive && InstancedCue->bAutoDestroyOnRemove)
diff --git a/Engine/Plugins/Runtime/GameplayAbilities/Source/GameplayAbilities/Private/GameplayEffect.cpp b/Engine/Plugins/Runtime/GameplayAbilities/Source/GameplayAbilities/Private/GameplayEffect.cpp
--- a/Engine/Plugins/Runtime/GameplayAbilities/Source/GameplayAbilities/Private/GameplayEffect.cpp
+++ b/Engine/Plugins/Runtime/GameplayAbilities/Source/GameplayAbilities/Private/GameplayEffect.cpp
@@ -1394,7 +1394,7 @@ void FGameplayEffectExecutionDefinition::GetAttributeCaptureDefinitions(OUT TArr
 
 	if (CalculationClass)
 	{
-		const UGameplayEffectExecutionCalculation* CalculationCDO = Cast<UGameplayEffectExecutionCalculation>(CalculationClass->ClassDefaultObject);
+		const UGameplayEffectExecutionCalculation* CalculationCDO = Cast<UGameplayEffectExecutionCalculation>(CalculationClass->GetDefaultObject(false));
 		check(CalculationCDO);
 
 		OutCaptureDefs.Append(CalculationCDO->GetAttributeCaptureDefinitions());
diff --git a/Engine/Plugins/Runtime/GameplayAbilities/Source/GameplayAbilities/Private/GameplayTagResponseTable.cpp b/Engine/Plugins/Runtime/GameplayAbilities/Source/GameplayAbilities/Private/GameplayTagResponseTable.cpp
--- a/Engine/Plugins/Runtime/GameplayAbilities/Source/GameplayAbilities/Private/GameplayTagResponseTable.cpp
+++ b/Engine/Plugins/Runtime/GameplayAbilities/Source/GameplayAbilities/Private/GameplayTagResponseTable.cpp
@@ -178,7 +178,7 @@ void UGameplayTagReponseTable::AddOrUpdate(UAbilitySystemComponent* ASC, const T
 		{
 			for (const TSubclassOf<UGameplayEffect>& ResponseGameplayEffect : ResponseGameplayEffects)
 			{
-				FActiveGameplayEffectHandle NewHandle = ASC->ApplyGameplayEffectToSelf(Cast<UGameplayEffect>(ResponseGameplayEffect->ClassDefaultObject), TotalCount, ASC->MakeEffectContext());
+				FActiveGameplayEffectHandle NewHandle = ASC->ApplyGameplayEffectToSelf(Cast<UGameplayEffect>(ResponseGameplayEffect->GetDefaultObject(false)), TotalCount, ASC->MakeEffectContext());
 				if (NewHandle.IsValid())
 				{
 					Handles.Add(NewHandle);
diff --git a/Engine/Plugins/Runtime/GameplayAbilities/Source/GameplayAbilitiesEditor/Private/SGameplayAttributeWidget.cpp b/Engine/Plugins/Runtime/GameplayAbilities/Source/GameplayAbilitiesEditor/Private/SGameplayAttributeWidget.cpp
--- a/Engine/Plugins/Runtime/GameplayAbilities/Source/GameplayAbilitiesEditor/Private/SGameplayAttributeWidget.cpp
+++ b/Engine/Plugins/Runtime/GameplayAbilities/Source/GameplayAbilitiesEditor/Private/SGameplayAttributeWidget.cpp
@@ -388,7 +388,7 @@ void SGameplayAttributeWidget::OnCopyAttribute(FProperty* AttributeProperty)
 
 	if (Attribute.IsValid())
 	{
-		if (const TObjectPtr<UObject> OwnerObject = Attribute.GetAttributeSetClass()->ClassDefaultObject)
+		if (const TObjectPtr<UObject> OwnerObject = Attribute.GetAttributeSetClass()->GetDefaultObject(false))
 		{
 			FString ExportedString;
 
diff --git a/Engine/Plugins/Runtime/ModelViewViewModel/Source/ModelViewViewModelBlueprint/Private/Extensions/MVVMViewBlueprintListViewBaseExtension.cpp b/Engine/Plugins/Runtime/ModelViewViewModel/Source/ModelViewViewModelBlueprint/Private/Extensions/MVVMViewBlueprintListViewBaseExtension.cpp
--- a/Engine/Plugins/Runtime/ModelViewViewModel/Source/ModelViewViewModelBlueprint/Private/Extensions/MVVMViewBlueprintListViewBaseExtension.cpp
+++ b/Engine/Plugins/Runtime/ModelViewViewModel/Source/ModelViewViewModelBlueprint/Private/Extensions/MVVMViewBlueprintListViewBaseExtension.cpp
@@ -98,7 +98,7 @@ void UMVVMBlueprintViewExtension_ListViewBase::Precompile(UE::MVVM::Compiler::IM
 
 				WidgetPathHandle = FieldPathResult.StealValue();
 
-				if (const UUserWidget* EntryUserWidget = Cast<UUserWidget>(EntryWidgetClass->ClassDefaultObject))
+				if (const UUserWidget* EntryUserWidget = Cast<UUserWidget>(EntryWidgetClass->GetDefaultObject(false)))
 				{
 					if (const UMVVMBlueprintView* EntryBPView = GetEntryWidgetBlueprintView(EntryUserWidget))
 					{
@@ -171,7 +171,7 @@ void UMVVMBlueprintViewExtension_ListViewBase::Compile(UE::MVVM::Compiler::IMVVM
 
 				// Check that we have a valid entry view and valid entry viewmodel
 				// Otherwise, no runtime extension will be created.
-				if (const UUserWidget* EntryUserWidget = Cast<UUserWidget>(EntryWidgetClass->ClassDefaultObject))
+				if (const UUserWidget* EntryUserWidget = Cast<UUserWidget>(EntryWidgetClass->GetDefaultObject(false)))
 				{
 					if (const UMVVMBlueprintView* EntryBPView = GetEntryWidgetBlueprintView(EntryUserWidget))
 					{
diff --git a/Engine/Plugins/Runtime/ModelViewViewModel/Source/ModelViewViewModelBlueprint/Private/Extensions/MVVMViewBlueprintPanelWidgetExtension.cpp b/Engine/Plugins/Runtime/ModelViewViewModel/Source/ModelViewViewModelBlueprint/Private/Extensions/MVVMViewBlueprintPanelWidgetExtension.cpp
--- a/Engine/Plugins/Runtime/ModelViewViewModel/Source/ModelViewViewModelBlueprint/Private/Extensions/MVVMViewBlueprintPanelWidgetExtension.cpp
+++ b/Engine/Plugins/Runtime/ModelViewViewModel/Source/ModelViewViewModelBlueprint/Private/Extensions/MVVMViewBlueprintPanelWidgetExtension.cpp
@@ -133,7 +133,7 @@ void UMVVMBlueprintViewExtension_PanelWidget::Precompile(UE::MVVM::Compiler::IMV
 				}
 				WidgetPathHandle = FieldPathResult.StealValue();
 
-				if (const UUserWidget* EntryUserWidget = Cast<UUserWidget>(EntryWidgetClass->ClassDefaultObject))
+				if (const UUserWidget* EntryUserWidget = Cast<UUserWidget>(EntryWidgetClass->GetDefaultObject(false)))
 				{
 					if (const UMVVMBlueprintView* EntryBPView = GetEntryWidgetBlueprintView(EntryUserWidget))
 					{
@@ -209,7 +209,7 @@ void UMVVMBlueprintViewExtension_PanelWidget::Compile(UE::MVVM::Compiler::IMVVMB
 
 				// Check that we have a valid entry view and valid entry viewmodel
 				// Otherwise, no runtime extension will be created.
-				if (const UUserWidget* EntryUserWidget = Cast<UUserWidget>(EntryWidgetClass->ClassDefaultObject))
+				if (const UUserWidget* EntryUserWidget = Cast<UUserWidget>(EntryWidgetClass->GetDefaultObject(false)))
 				{
 					if (const UMVVMBlueprintView* EntryBPView = GetEntryWidgetBlueprintView(EntryUserWidget))
 					{
diff --git a/Engine/Plugins/Runtime/ModelViewViewModel/Source/ModelViewViewModelEditor/Private/Customizations/MVVMListViewBaseExtensionCustomizationExtender.cpp b/Engine/Plugins/Runtime/ModelViewViewModel/Source/ModelViewViewModelEditor/Private/Customizations/MVVMListViewBaseExtensionCustomizationExtender.cpp
--- a/Engine/Plugins/Runtime/ModelViewViewModel/Source/ModelViewViewModelEditor/Private/Customizations/MVVMListViewBaseExtensionCustomizationExtender.cpp
+++ b/Engine/Plugins/Runtime/ModelViewViewModel/Source/ModelViewViewModelEditor/Private/Customizations/MVVMListViewBaseExtensionCustomizationExtender.cpp
@@ -244,7 +244,7 @@ void FMVVMListViewBaseExtensionCustomizationExtender::HandleEntryClassChanged(bo
 	// Update other values that depend on the entry class (only if the cached value actually changed)
 	if (bEntryClassChanged && EntryClass)
 	{
-		if (UUserWidget* EntryCDO = Cast<UUserWidget>(EntryClass->ClassDefaultObject))
+		if (const UUserWidget* EntryCDO = GetDefault<UUserWidget>(EntryClass))
 		{
 			EntryWidgetBlueprint = Cast<UWidgetBlueprint>(EntryCDO->GetClass()->ClassGeneratedBy);
 		}
@@ -265,7 +265,7 @@ FText FMVVMListViewBaseExtensionCustomizationExtender::OnGetSelectedViewModel()
 		{
 			if (UMVVMBlueprintViewExtension_ListViewBase* ListBaseViewExtension = GetListBaseViewExtension())
 			{
-				if (const UUserWidget* EntryUserWidget = Cast<UUserWidget>(EntryClass->ClassDefaultObject))
+				if (const UUserWidget* EntryUserWidget = GetDefault<UUserWidget>(EntryClass))
 				{
 					if (const UWidgetBlueprint* EntryBlueprint = Cast<UWidgetBlueprint>(EntryUserWidget->GetClass()->ClassGeneratedBy))
 					{
diff --git a/Engine/Plugins/Runtime/ModelViewViewModel/Source/ModelViewViewModelEditor/Private/Customizations/MVVMPanelWidgetExtensionCustomizationExtender.cpp b/Engine/Plugins/Runtime/ModelViewViewModel/Source/ModelViewViewModelEditor/Private/Customizations/MVVMPanelWidgetExtensionCustomizationExtender.cpp
--- a/Engine/Plugins/Runtime/ModelViewViewModel/Source/ModelViewViewModelEditor/Private/Customizations/MVVMPanelWidgetExtensionCustomizationExtender.cpp
+++ b/Engine/Plugins/Runtime/ModelViewViewModel/Source/ModelViewViewModelEditor/Private/Customizations/MVVMPanelWidgetExtensionCustomizationExtender.cpp
@@ -356,7 +356,7 @@ void FMVVMPanelWidgetExtensionCustomizationExtender::HandleEntryClassChanged(boo
 	// Update other values that depend on the entry class (only if the cached value actually changed)
 	if (bEntryClassChanged)
 	{
-		if (UUserWidget* EntryCDO = EntryClass ? Cast<UUserWidget>(EntryClass->ClassDefaultObject) : nullptr)
+		if (const UUserWidget* EntryCDO = EntryClass ? GetDefault<UUserWidget>(EntryClass) : nullptr)
 		{
 			EntryWidgetBlueprint = Cast<UWidgetBlueprint>(EntryCDO->GetClass()->ClassGeneratedBy);
 		}
@@ -391,7 +391,7 @@ FText FMVVMPanelWidgetExtensionCustomizationExtender::OnGetSelectedViewModel() c
 		{
 			if (UMVVMBlueprintViewExtension_PanelWidget* PanelWidgetExtension = GetPanelWidgetExtension())
 			{
-				if (const UUserWidget* EntryUserWidget = Cast<UUserWidget>(EntryClass->ClassDefaultObject))
+				if (const UUserWidget* EntryUserWidget = GetDefault<UUserWidget>(EntryClass))
 				{
 					if (const UWidgetBlueprint* EntryBlueprint = Cast<UWidgetBlueprint>(EntryUserWidget->GetClass()->ClassGeneratedBy))
 					{
diff --git a/Engine/Plugins/Runtime/NetworkPredictionExtras/Source/NetworkPredictionExtras/Private/MockRootMotionComponent.cpp b/Engine/Plugins/Runtime/NetworkPredictionExtras/Source/NetworkPredictionExtras/Private/MockRootMotionComponent.cpp
--- a/Engine/Plugins/Runtime/NetworkPredictionExtras/Source/NetworkPredictionExtras/Private/MockRootMotionComponent.cpp
+++ b/Engine/Plugins/Runtime/NetworkPredictionExtras/Source/NetworkPredictionExtras/Private/MockRootMotionComponent.cpp
@@ -279,7 +279,7 @@ UMockRootMotionSource* UMockRootMotionComponent::ResolveRootMotionSource(int32 I
 	{
 		// No parameters: this implies no instantiation/CDO 
 		// (are we sure? Should there be a virtual here? What if we have a no parameter but internal state having source?)
-		RootMotionSourceCache.Instance = Cast<UMockRootMotionSource>(ClassObj->ClassDefaultObject);
+		RootMotionSourceCache.Instance = Cast<UMockRootMotionSource>(ClassObj->GetDefaultObject(false));
 	}
 	else
 	{
diff --git a/Engine/Plugins/Runtime/RigVM/Source/RigVM/Private/RigVMCore/RigVMMemoryStorage.cpp b/Engine/Plugins/Runtime/RigVM/Source/RigVM/Private/RigVMCore/RigVMMemoryStorage.cpp
--- a/Engine/Plugins/Runtime/RigVM/Source/RigVM/Private/RigVMCore/RigVMMemoryStorage.cpp
+++ b/Engine/Plugins/Runtime/RigVM/Source/RigVM/Private/RigVMCore/RigVMMemoryStorage.cpp
@@ -535,11 +535,11 @@ URigVMMemoryStorageGeneratorClass* URigVMMemoryStorageGeneratorClass::CreateStor
 		OldClass->ClassFlags |= CLASS_NewerVersionExists;
 		RenameAndMarkGarbage(OldClass, DiscardedMemoryClassName);
  
-		if (OldClass->ClassDefaultObject)
+		if (OldClass->GetDefaultObject(false))
 		{
 			static const TCHAR* DiscardedMemoryCDOTemplate(TEXT("DiscardedMemoryCDOTemplate"));
 			static int32 DiscardedMemoryCDOIndex = 0;
-			RenameAndMarkGarbage(OldClass->ClassDefaultObject, FindUniqueName(DiscardedMemoryCDOTemplate, DiscardedMemoryCDOIndex));
+			RenameAndMarkGarbage(OldClass->GetDefaultObject(false), FindUniqueName(DiscardedMemoryCDOTemplate, DiscardedMemoryCDOIndex));
 		}
 	}
 
diff --git a/Engine/Plugins/Runtime/RigVM/Source/RigVMEditor/Private/Editor/Kismet/RigVMBlueprintCompilationManager.cpp b/Engine/Plugins/Runtime/RigVM/Source/RigVMEditor/Private/Editor/Kismet/RigVMBlueprintCompilationManager.cpp
--- a/Engine/Plugins/Runtime/RigVM/Source/RigVMEditor/Private/Editor/Kismet/RigVMBlueprintCompilationManager.cpp
+++ b/Engine/Plugins/Runtime/RigVM/Source/RigVMEditor/Private/Editor/Kismet/RigVMBlueprintCompilationManager.cpp
@@ -694,8 +694,8 @@ void FRigVMBlueprintCompilationManagerImpl::FlushCompilationQueueImpl(bool bSupp
 			UBlueprint* QueuedBP = QueuedJob.BPToCompile;
 
 			ensure(!QueuedBP->GeneratedClass ||
-				!QueuedBP->GeneratedClass->ClassDefaultObject ||
-				!(QueuedBP->GeneratedClass->ClassDefaultObject->HasAnyFlags(RF_NeedLoad)));
+				!QueuedBP->GeneratedClass->GetDefaultObject(false) ||
+				!(QueuedBP->GeneratedClass->GetDefaultObject(false)->HasAnyFlags(RF_NeedLoad)));
 			bool bDefaultComponentMustBeAdded = false;
 			bool bHasPendingUberGraphFrame = false;
 			UBlueprintGeneratedClass* BPGC = Cast<UBlueprintGeneratedClass>(QueuedBP->GeneratedClass);
@@ -1347,7 +1347,7 @@ void FRigVMBlueprintCompilationManagerImpl::FlushCompilationQueueImpl(bool bSupp
 
 				if(BP->GeneratedClass)
 				{
-					OldCDOs.Add(BP, BP->GeneratedClass->ClassDefaultObject);
+					OldCDOs.Add(BP, BP->GeneratedClass->GetDefaultObject(false));
 				}
 
 				EBlueprintCompileReinstancerFlags CompileReinstancerFlags =
@@ -1408,14 +1408,14 @@ void FRigVMBlueprintCompilationManagerImpl::FlushCompilationQueueImpl(bool bSupp
 				SCOPED_LOADTIMER_ASSET_TEXT(*BP->GetPathName());
 
 				ensure( BP->GeneratedClass == nullptr ||
-						BP->GeneratedClass->ClassDefaultObject == nullptr ||
-						BP->GeneratedClass->ClassDefaultObject->GetClass() != BP->GeneratedClass);
+						BP->GeneratedClass->GetDefaultObject(false) == nullptr ||
+						BP->GeneratedClass->GetDefaultObject(false)->GetClass() != BP->GeneratedClass);
 				// default value propagation occurs in ReinstaneBatch, CDO will be created via CompileFunctions call:
 				if(BP->ParentClass)
 				{
 					if(BP->GeneratedClass)
 					{
-						BP->GeneratedClass->ClassDefaultObject = nullptr;
+						BP->GeneratedClass->SetDefaultObject(nullptr);
 					}
 
 					// Reset the flag, so if the user tries to use PIE it will warn them if the BP did not compile
@@ -1433,7 +1433,7 @@ void FRigVMBlueprintCompilationManagerImpl::FlushCompilationQueueImpl(bool bSupp
 					{
 						ChildClass->Bind();
 						ChildClass->StaticLink();
-						ensure(ChildClass->ClassDefaultObject == nullptr);
+						ensure(ChildClass->GetDefaultObject(false) == nullptr);
 					}
 				}
 				else
@@ -1481,8 +1481,8 @@ void FRigVMBlueprintCompilationManagerImpl::FlushCompilationQueueImpl(bool bSupp
 			if(!CompilerData.ShouldCompileClassFunctions())
 			{
 				if( BPGC &&
-					(	BPGC->ClassDefaultObject == nullptr ||
-						BPGC->ClassDefaultObject->GetClass() != BPGC) )
+					(	BPGC->GetDefaultObject(false) == nullptr ||
+						BPGC->GetDefaultObject(false)->GetClass() != BPGC) )
 				{
 					if (CompilerData.Reinstancer.IsValid())
 					{
@@ -1492,7 +1492,7 @@ void FRigVMBlueprintCompilationManagerImpl::FlushCompilationQueueImpl(bool bSupp
 					BPGC->bLayoutChanging = false;
 					BPGC->Bind();
 					BPGC->StaticLink(true);
-					BPGC->ClassDefaultObject = nullptr;
+					BPGC->SetDefaultObject(nullptr);
 					BPGC->GetDefaultObject(true);
 				}
 			}
@@ -1509,8 +1509,8 @@ void FRigVMBlueprintCompilationManagerImpl::FlushCompilationQueueImpl(bool bSupp
 						CompilerData.Reinstancer->PropagateSparseClassDataToNewClass(BPGC);
 					}
 
-					if( BPGC->ClassDefaultObject &&
-						BPGC->ClassDefaultObject->GetClass() == BPGC)
+					if( BPGC->GetDefaultObject(false) &&
+						BPGC->GetDefaultObject(false)->GetClass() == BPGC)
 					{
 						// the CDO has been created early, it is possible that the reflection data was still
 						// being mutated by CompileClassLayout. Warn the user and and move the CDO aside:
@@ -1519,7 +1519,7 @@ void FRigVMBlueprintCompilationManagerImpl::FlushCompilationQueueImpl(bool bSupp
 							*BP->GetName()
 						);
 
-						BPGC->ClassDefaultObject->Rename(
+						BPGC->GetDefaultObject(false)->Rename(
 							nullptr,
 							// destination - this is the important part of this call. Moving the object
 							// out of the way so we can reuse its name:
@@ -1528,7 +1528,7 @@ void FRigVMBlueprintCompilationManagerImpl::FlushCompilationQueueImpl(bool bSupp
 							REN_DoNotDirty | REN_DontCreateRedirectors
 						);
 					}
-					BPGC->ClassDefaultObject = nullptr;
+					BPGC->SetDefaultObject(nullptr);
 
 					// class layout is ready, we can clear bLayoutChanging and CompileFunctions can create the CDO:
 					BPGC->bLayoutChanging = false;
@@ -1580,7 +1580,7 @@ void FRigVMBlueprintCompilationManagerImpl::FlushCompilationQueueImpl(bool bSupp
 
 			FKismetCompilerUtilities::UpdateDependentBlueprints(BP);
 
-			ensure(BPGC == nullptr || BPGC->ClassDefaultObject->GetClass() == BPGC);
+			ensure(BPGC == nullptr || BPGC->GetDefaultObject(false)->GetClass() == BPGC);
 		}
 	} // end TimeCompiling scope
 
@@ -1624,13 +1624,13 @@ void FRigVMBlueprintCompilationManagerImpl::FlushCompilationQueueImpl(bool bSupp
 
 				const UBlueprintGeneratedClass* GenClass = Cast<UBlueprintGeneratedClass>(CompilerData.BP->GeneratedClass);
 
-				if (GenClass && GenClass->ClassDefaultObject)
+				if (GenClass && GenClass->GetDefaultObject(false))
 				{
 					for (const FBPVariableDescription& NewInheritedVar : CompilerData.NewDefaultVariables)
 					{
 						if (const FProperty* MatchingProperty = GenClass->FindPropertyByName(NewInheritedVar.VarName))
 						{
-							FBlueprintEditorUtils::PropertyValueFromString(MatchingProperty, NewInheritedVar.DefaultValue, reinterpret_cast<uint8*>(GenClass->ClassDefaultObject.Get()));
+							FBlueprintEditorUtils::PropertyValueFromString(MatchingProperty, NewInheritedVar.DefaultValue, reinterpret_cast<uint8*>(GenClass->GetDefaultObject(false)));
 						}
 					}
 				}
@@ -2044,9 +2044,9 @@ void FRigVMBlueprintCompilationManagerImpl::GetDefaultValue(const UClass* ForCla
 		return;
 	}
 
-	if (ForClass->ClassDefaultObject)
+	if (ForClass->GetDefaultObject(false))
 	{
-		FBlueprintEditorUtils::PropertyValueToString(Property, (uint8*)ForClass->ClassDefaultObject.Get(), OutDefaultValueAsString);
+		FBlueprintEditorUtils::PropertyValueToString(Property, (uint8*)ForClass->GetDefaultObject(false), OutDefaultValueAsString);
 	}
 	else
 	{
@@ -2103,7 +2103,7 @@ void FRigVMBlueprintCompilationManagerImpl::SaveConstructionHashes(const TArray<
 		UBlueprint* BP = CompilerData.BP;
 		if(const UClass* OriginalClass = BP->GeneratedClass)
 		{
-			if( const UObject* CDO = OriginalClass->ClassDefaultObject )
+			if( const UObject* CDO = OriginalClass->GetDefaultObject(false) )
 			{
 				if( FOverridableManager::Get().IsEnabled(*CDO) )
 				{
@@ -2118,7 +2118,7 @@ void FRigVMBlueprintCompilationManagerImpl::SaveConstructionHashes(const TArray<
 		UBlueprint* BP = CompilerData.BP;
 		if(const UClass* OriginalClass = BP->GeneratedClass)
 		{
-			if(OriginalClass->ClassDefaultObject)
+			if(OriginalClass->GetDefaultObject(false))
 			{
 				// if there are no instances, no need to calculate the class hash:
 				FBlake3Hash ConstructionHash = FBlueprintCompileReinstancer::CalculateConstructionHashIfNecessary(OriginalClass);
@@ -2145,7 +2145,7 @@ void FRigVMBlueprintCompilationManagerImpl::ReparentHierarchies(const TMap<UClas
 		for (const TPair<UClass*, UClass*>& OldToNewClass : OldToNewClasses)
 		{
 			// classes with no CDO do not need to be reinstanced:
-			if (OldToNewClass.Key->ClassDefaultObject == nullptr)
+			if (OldToNewClass.Key->GetDefaultObject(false) == nullptr)
 			{
 				continue;
 			}
@@ -2160,7 +2160,7 @@ void FRigVMBlueprintCompilationManagerImpl::ReparentHierarchies(const TMap<UClas
 
 			for (UClass* DerivedClass : DerivedClasses)
 			{
-				if (DerivedClass->ClassDefaultObject == nullptr &&
+				if (DerivedClass->GetDefaultObject(false) == nullptr &&
 					DerivedClass->GetSparseClassData(EGetSparseClassDataMethod::ReturnIfNull) == nullptr)
 				{
 					// no CDO->no other instances->no need to reinstance..
@@ -2400,7 +2400,7 @@ void FRigVMBlueprintCompilationManagerImpl::ReinstanceBatch(TArray<FReinstancing
 			{
 				if(IsValid(ClassToReinstance))
 				{
-					if (ClassToReinstance->ClassDefaultObject)
+					if (ClassToReinstance->GetDefaultObject(false))
 					{
 						ClassesToReinstance.Add(ClassToReinstance);
 					}
@@ -2474,7 +2474,7 @@ void FRigVMBlueprintCompilationManagerImpl::ReinstanceBatch(TArray<FReinstancing
 	// make new hierarchy
 	for(UClass* Class : ClassesToReinstance)
 	{
-		UObject* OriginalCDO = Class->ClassDefaultObject;
+		UObject* OriginalCDO = Class->GetDefaultObject(false);
 		Reinstancers.Emplace(
 			FReinstancingJob (
 				TSharedPtr<FBlueprintCompileReinstancer>(
@@ -2554,7 +2554,7 @@ void FRigVMBlueprintCompilationManagerImpl::ReinstanceBatch(TArray<FReinstancing
 		{
 			AddDependentClass(OldClass->GetSuperClass());
 
-			if (const UObject* CDO = OldClass->ClassDefaultObject)
+			if (const UObject* CDO = OldClass->GetDefaultObject(false))
 			{
 				TArray<UObject*> ContainedOldObjects;
 				GetObjectsWithOuter(CDO, ContainedOldObjects);
@@ -2599,7 +2599,7 @@ void FRigVMBlueprintCompilationManagerImpl::ReinstanceBatch(TArray<FReinstancing
 
 		if (OldClass)
 		{
-			OldCDO = OldClass->ClassDefaultObject;
+			OldCDO = OldClass->GetDefaultObject(false);
 			if (OldCDO && ReinstancingJob.Reinstancer.IsValid())
 			{
 				// Object using overridable serialization need to use delta serialization for it do work appropriatly
@@ -2610,7 +2610,7 @@ void FRigVMBlueprintCompilationManagerImpl::ReinstanceBatch(TArray<FReinstancing
 
 				// We do not expect the CDO to be already created at this point.
 				// Who ever bother to create it before didn't look if the parent was ready or not.
-				if (NewClass->ClassDefaultObject != nullptr && bUsingOverrideSerialization)
+				if (NewClass->GetDefaultObject(false) != nullptr && bUsingOverrideSerialization)
 				{
 					// Override serialization does not currently want to use any CDO created before this
 					// point. I do not have a test case justifying this need, but renaming the cdo here
@@ -2619,8 +2619,8 @@ void FRigVMBlueprintCompilationManagerImpl::ReinstanceBatch(TArray<FReinstancing
 					// we will have to address these shortcomings (likely by ordering compilation itself,
 					// rather than reinstancing, more carefully). Discarding the CDO here is also wasteful
 					// - if we cannot use the objects, why create them at all?
-					NewClass->ClassDefaultObject->Rename(nullptr, GetTransientPackage(), REN_DoNotDirty | REN_DontCreateRedirectors | REN_NonTransactional);
-					NewClass->ClassDefaultObject = nullptr;
+					NewClass->GetDefaultObject(false)->Rename(nullptr, GetTransientPackage(), REN_DoNotDirty | REN_DontCreateRedirectors | REN_NonTransactional);
+					NewClass->SetDefaultObject(nullptr);
 				}
 				UObject* NewCDO = NewClass->GetDefaultObject(true);
 
@@ -2666,7 +2666,7 @@ void FRigVMBlueprintCompilationManagerImpl::ReinstanceBatch(TArray<FReinstancing
 			if (ReinstancingJob.Reinstancer.IsValid())
 			{
 				UBlueprint* CurrentBP = CastChecked<UBlueprint>(BPGClass->ClassGeneratedBy);
-				UObject* NewCDO = CurrentBP->GeneratedClass->ClassDefaultObject;
+				UObject* NewCDO = CurrentBP->GeneratedClass->GetDefaultObject(false);
 				FLinkerLoad* CurrentLinker = CurrentBP->GetLinker();
 				if (CurrentLinker)
 				{
@@ -2734,13 +2734,13 @@ void FRigVMBlueprintCompilationManagerImpl::ReinstanceBatch(TArray<FReinstancing
 
 			UClass* NewClass = ReinstancingJob.OldToNew.Value;
 			if (NewClass &&
-				OldClass->ClassDefaultObject &&
-				NewClass->ClassDefaultObject &&
-				OldClass->ClassDefaultObject != NewClass->ClassDefaultObject)
+				OldClass->GetDefaultObject(false) &&
+				NewClass->GetDefaultObject(false) &&
+				OldClass->GetDefaultObject(false) != NewClass->GetDefaultObject(false))
 			{
-				OldArchetypeToNewArchetype.Add(OldClass->ClassDefaultObject, NewClass->ClassDefaultObject);
+				OldArchetypeToNewArchetype.Add(OldClass->GetDefaultObject(false), NewClass->GetDefaultObject(false));
 				// also map old *default* subobjects to new default subobjects:
-				BuildDSOMap(OldClass->ClassDefaultObject, NewClass->ClassDefaultObject, InOutOldToNewClassMap, OldArchetypeToNewArchetype);
+				BuildDSOMap(OldClass->GetDefaultObject(false), NewClass->GetDefaultObject(false), InOutOldToNewClassMap, OldArchetypeToNewArchetype);
 			}
 
 			TArray<UObject*> ArchetypeObjects;
@@ -3081,7 +3081,7 @@ UClass* FRigVMBlueprintCompilationManagerImpl::FastGenerateSkeletonClass(UBluepr
 			CompilerContext.ValidateVariableNames();
 		}
 
-		CompilerContext.CleanAndSanitizeClass(Ret, MutableView(Ret->ClassDefaultObject));
+		CompilerContext.CleanAndSanitizeClass(Ret, MutableView(Ret->GetDefaultObject(false)));
 	}
 
 	Ret->ClassGeneratedBy = BP;
@@ -3617,7 +3617,7 @@ void FRigVMBlueprintCompilationManagerImpl::ConformToParentAndInterfaces(UBluepr
 void FRigVMBlueprintCompilationManagerImpl::RelinkSkeleton(UClass* SkeletonToRelink)
 {
 	// CDO needs to be moved aside already:
-	ensure(SkeletonToRelink->ClassDefaultObject == nullptr);
+	ensure(SkeletonToRelink->GetDefaultObject(false) == nullptr);
 	ensure(!SkeletonToRelink->GetSuperClass()->HasAnyClassFlags(CLASS_NewerVersionExists));
 
 	SkeletonToRelink->ClassConstructor = nullptr;
@@ -3673,7 +3673,7 @@ void FRigVMBlueprintCompilationManagerImpl::QueueOutOfDateDependencies( const TA
 		{
 			return
 				CompileRequest.BPToCompile->GeneratedClass != nullptr &&
-				CompileRequest.BPToCompile->GeneratedClass->ClassDefaultObject != nullptr &&
+				CompileRequest.BPToCompile->GeneratedClass->GetDefaultObject(false) != nullptr &&
 				(CompileRequest.CompileOptions & EBlueprintCompileOptions::RegenerateSkeletonOnly)
 				== EBlueprintCompileOptions::None;
 		},
diff --git a/Engine/Plugins/Runtime/nDisplay/Source/DisplayCluster/Private/Blueprints/DisplayClusterBlueprint.cpp b/Engine/Plugins/Runtime/nDisplay/Source/DisplayCluster/Private/Blueprints/DisplayClusterBlueprint.cpp
--- a/Engine/Plugins/Runtime/nDisplay/Source/DisplayCluster/Private/Blueprints/DisplayClusterBlueprint.cpp
+++ b/Engine/Plugins/Runtime/nDisplay/Source/DisplayCluster/Private/Blueprints/DisplayClusterBlueprint.cpp
@@ -486,7 +486,7 @@ UDisplayClusterConfigurationData* UDisplayClusterBlueprint::GetOrLoadConfig()
 {
 	if (GeneratedClass)
 	{
-		if (ADisplayClusterRootActor* CDO = Cast<ADisplayClusterRootActor>(GeneratedClass->ClassDefaultObject))
+		if (ADisplayClusterRootActor* CDO = Cast<ADisplayClusterRootActor>(GeneratedClass->GetDefaultObject(false)))
 		{
 			ConfigData = CDO->GetConfigData();
 		}
@@ -503,7 +503,7 @@ void UDisplayClusterBlueprint::SetConfigData(UDisplayClusterConfigurationData* I
 
 	if (GeneratedClass)
 	{
-		if (ADisplayClusterRootActor* CDO = Cast<ADisplayClusterRootActor>(GeneratedClass->ClassDefaultObject))
+		if (ADisplayClusterRootActor* CDO = Cast<ADisplayClusterRootActor>(GeneratedClass->GetDefaultObject(false)))
 		{
 			CDO->UpdateConfigDataInstance(InConfigData, bForceRecreate);
 			GetOrLoadConfig();
@@ -546,7 +546,7 @@ void UDisplayClusterBlueprint::PrepareConfigForExport()
 		return;
 	}
 		
-	ADisplayClusterRootActor* CDO = CastChecked<ADisplayClusterRootActor>(GeneratedClass->ClassDefaultObject);
+	ADisplayClusterRootActor* CDO = CastChecked<ADisplayClusterRootActor>(GeneratedClass->GetDefaultObject(false));
 
 	UDisplayClusterConfigurationData* Data = GetOrLoadConfig();
 	check(Data);
diff --git a/Engine/Plugins/Runtime/nDisplay/Source/DisplayCluster/Private/Game/EngineClasses/Scene/DisplayClusterICVFXCameraComponent.cpp b/Engine/Plugins/Runtime/nDisplay/Source/DisplayCluster/Private/Game/EngineClasses/Scene/DisplayClusterICVFXCameraComponent.cpp
--- a/Engine/Plugins/Runtime/nDisplay/Source/DisplayCluster/Private/Game/EngineClasses/Scene/DisplayClusterICVFXCameraComponent.cpp
+++ b/Engine/Plugins/Runtime/nDisplay/Source/DisplayCluster/Private/Game/EngineClasses/Scene/DisplayClusterICVFXCameraComponent.cpp
@@ -120,7 +120,7 @@ void UDisplayClusterICVFXCameraComponent::PostLoad()
 	const int32 CustomVersion = GetLinkerCustomVersion(FDisplayClusterICVFXCameraCustomVersion::GUID);
 	if (CustomVersion < FDisplayClusterICVFXCameraCustomVersion::UpdateChromakeyConfig)
 	{
-		const bool bHasCustomArchetype = GetArchetype() != StaticClass()->ClassDefaultObject;
+		const bool bHasCustomArchetype = GetArchetype() != StaticClass()->GetDefaultObject(false);
 		const int32 ArchetypeVersion = GetArchetype()->GetLinkerCustomVersion(FDisplayClusterICVFXCameraCustomVersion::GUID);
 
 		// UE-184291: If this camera component has a user-defined archetype and that archetype has been updated already, do not
diff --git a/Engine/Plugins/Runtime/nDisplay/Source/DisplayClusterConfigurator/Private/DisplayClusterConfiguratorBlueprintEditor.cpp b/Engine/Plugins/Runtime/nDisplay/Source/DisplayClusterConfigurator/Private/DisplayClusterConfiguratorBlueprintEditor.cpp
--- a/Engine/Plugins/Runtime/nDisplay/Source/DisplayClusterConfigurator/Private/DisplayClusterConfiguratorBlueprintEditor.cpp
+++ b/Engine/Plugins/Runtime/nDisplay/Source/DisplayClusterConfigurator/Private/DisplayClusterConfiguratorBlueprintEditor.cpp
@@ -517,7 +517,7 @@ ADisplayClusterRootActor* FDisplayClusterConfiguratorBlueprintEditor::GetDefault
 	{
 		if (Blueprint->GeneratedClass)
 		{
-			return Cast<ADisplayClusterRootActor>(Blueprint->GeneratedClass->ClassDefaultObject);
+			return Cast<ADisplayClusterRootActor>(Blueprint->GeneratedClass->GetDefaultObject(false));
 		}
 	}
 
diff --git a/Engine/Plugins/Runtime/nDisplay/Source/DisplayClusterTests/Private/DisplayClusterTestUtils.cpp b/Engine/Plugins/Runtime/nDisplay/Source/DisplayClusterTests/Private/DisplayClusterTestUtils.cpp
--- a/Engine/Plugins/Runtime/nDisplay/Source/DisplayClusterTests/Private/DisplayClusterTestUtils.cpp
+++ b/Engine/Plugins/Runtime/nDisplay/Source/DisplayClusterTests/Private/DisplayClusterTestUtils.cpp
@@ -164,10 +164,10 @@ namespace DisplayClusterTestUtils
 		Asset->RemoveFromRoot();
 
 		UBlueprint* BlueprintObject = Cast<UBlueprint>(Asset);
-		if (BlueprintObject && BlueprintObject->GeneratedClass && BlueprintObject->GeneratedClass->ClassDefaultObject)
+		if (BlueprintObject && BlueprintObject->GeneratedClass && BlueprintObject->GeneratedClass->GetDefaultObject(false))
 		{
 			TArray<UObject*> InstancesToDelete;
-			BlueprintObject->GeneratedClass->ClassDefaultObject->GetArchetypeInstances( InstancesToDelete );
+			BlueprintObject->GeneratedClass->GetDefaultObject(false)->GetArchetypeInstances( InstancesToDelete );
 			
 			for (TArray<UObject*>::TConstIterator InstanceItr( InstancesToDelete ); InstanceItr; ++InstanceItr)
 			{
diff --git a/Engine/Plugins/VirtualProduction/RemoteControl/Source/RemoteControlUI/Private/UI/RCPanelWidgetRegistry.cpp b/Engine/Plugins/VirtualProduction/RemoteControl/Source/RemoteControlUI/Private/UI/RCPanelWidgetRegistry.cpp
--- a/Engine/Plugins/VirtualProduction/RemoteControl/Source/RemoteControlUI/Private/UI/RCPanelWidgetRegistry.cpp
+++ b/Engine/Plugins/VirtualProduction/RemoteControl/Source/RemoteControlUI/Private/UI/RCPanelWidgetRegistry.cpp
@@ -281,7 +281,7 @@ bool FRCPanelWidgetRegistry::IsNDisplayObject(UObject* InObject, const FString&
 		if (Actor != nullptr)
 		{
 			UBlueprintGeneratedClass* BlueprintClass = Cast<UBlueprintGeneratedClass>(Actor->GetClass());
-			if (BlueprintClass && BlueprintClass->ClassDefaultObject)
+			if (BlueprintClass && BlueprintClass->GetDefaultObject(false))
 			{
 				static constexpr TCHAR ConfigurationDataName[] = TEXT("DisplayClusterConfigurationData");
 				const FString ObjectClassName = InObject->GetClass()->GetName();
diff --git a/Engine/Source/Developer/AssetTools/Private/AssetRenameManager.cpp b/Engine/Source/Developer/AssetTools/Private/AssetRenameManager.cpp
--- a/Engine/Source/Developer/AssetTools/Private/AssetRenameManager.cpp
+++ b/Engine/Source/Developer/AssetTools/Private/AssetRenameManager.cpp
@@ -723,7 +723,7 @@ void FAssetRenameManager::FindCDOReferences(const TArrayView<FAssetRenameDataWit
 	for (TObjectIterator<UClass> ClassDefaultObjectIt; ClassDefaultObjectIt; ++ClassDefaultObjectIt)
 	{
 		UClass* Cls = (*ClassDefaultObjectIt);
-		UObject* CDO = Cls->ClassDefaultObject;
+		UObject* CDO = Cls->GetDefaultObject(false);
 
 		if (!CDO || !CDO->HasAllFlags(RF_ClassDefaultObject) || !IsValidChecked(CDO) || Cls->ClassGeneratedBy != nullptr)
 		{
diff --git a/Engine/Source/Developer/Windows/LiveCoding/Private/LiveCodingModule.cpp b/Engine/Source/Developer/Windows/LiveCoding/Private/LiveCodingModule.cpp
--- a/Engine/Source/Developer/Windows/LiveCoding/Private/LiveCodingModule.cpp
+++ b/Engine/Source/Developer/Windows/LiveCoding/Private/LiveCodingModule.cpp
@@ -871,14 +871,14 @@ void FLiveCodingModule::AttemptSyncLivePatching()
 				// where the reinstanced crashes recreating the default object probably due to a mismatch between then
 				// new constructor being invoked and the blueprint data associated with the old class.  With LC, the
 				// old constructor has been replaced.
-				static UObject* DummyDefaultObject = UObject::StaticClass()->ClassDefaultObject;
+				static UObject* DummyDefaultObject = UObject::StaticClass()->GetDefaultObject(false);
 				for (TObjectIterator<UClass> It; It; ++It)
 				{
 					UClass* Class = *It;
 					if (Class->GetName().StartsWith(TEXT("LIVECODING_")) ||
 						Class->GetName().StartsWith(TEXT("REINST_")))
 					{
-						Class->ClassDefaultObject = DummyDefaultObject;
+						Class->SetDefaultObject(DummyDefaultObject);
 					}
 				}
 			}
diff --git a/Engine/Source/Editor/AnimGraph/Private/AnimBlueprintCompiler.cpp b/Engine/Source/Editor/AnimGraph/Private/AnimBlueprintCompiler.cpp
--- a/Engine/Source/Editor/AnimGraph/Private/AnimBlueprintCompiler.cpp
+++ b/Engine/Source/Editor/AnimGraph/Private/AnimBlueprintCompiler.cpp
@@ -1373,7 +1373,7 @@ void FAnimBlueprintCompilerContext::OnPostCDOCompiled(const UObject::FPostCDOCom
 	}
 
 	UAnimBlueprintGeneratedClass* NewAnimBlueprintClass = GetNewAnimBlueprintClass();
-	NewAnimBlueprintClass->OnPostLoadDefaults(NewAnimBlueprintClass->ClassDefaultObject);
+	NewAnimBlueprintClass->OnPostLoadDefaults(NewAnimBlueprintClass->GetDefaultObject(false));
 }
 
 void FAnimBlueprintCompilerContext::OnNewClassSet(UBlueprintGeneratedClass* ClassToUse)
diff --git a/Engine/Source/Editor/BlueprintGraph/Private/BlueprintActionDatabase.cpp b/Engine/Source/Editor/BlueprintGraph/Private/BlueprintActionDatabase.cpp
--- a/Engine/Source/Editor/BlueprintGraph/Private/BlueprintActionDatabase.cpp
+++ b/Engine/Source/Editor/BlueprintGraph/Private/BlueprintActionDatabase.cpp
@@ -1566,7 +1566,7 @@ void FBlueprintActionDatabase::RefreshClassActions(UClass* const Class)
 		}
 
 		FBlueprintActionDatabaseRegistrar Registrar(ActionRegistry, UnloadedActionRegistry, ActionPrimingQueue);
-		Cast<UBlueprint>(Class->ClassDefaultObject)->GetTypeActions(Registrar);
+		Cast<UBlueprint>(Class->GetDefaultObject(false))->GetTypeActions(Registrar);
 	}
 	else
 	{
diff --git a/Engine/Source/Editor/BlueprintGraph/Private/BlueprintActionFilter.cpp b/Engine/Source/Editor/BlueprintGraph/Private/BlueprintActionFilter.cpp
--- a/Engine/Source/Editor/BlueprintGraph/Private/BlueprintActionFilter.cpp
+++ b/Engine/Source/Editor/BlueprintGraph/Private/BlueprintActionFilter.cpp
@@ -1222,7 +1222,7 @@ static bool BlueprintActionFilterImpl::IsIncompatibleWithGraphType(FBlueprintAct
 	
 	if (UClass const* NodeClass = BlueprintAction.GetNodeClass())
 	{
-		if (UEdGraphNode const* NodeCDO = CastChecked<UEdGraphNode>(NodeClass->ClassDefaultObject))
+		if (UEdGraphNode const* NodeCDO = CastChecked<UEdGraphNode>(NodeClass->GetDefaultObject(false)))
 		{
 			for (UEdGraph* Graph : FilterContext.Graphs)
 			{
@@ -1355,7 +1355,7 @@ static bool BlueprintActionFilterImpl::IsSchemaIncompatible(FBlueprintActionFilt
 	FBlueprintActionContext const& FilterContext = Filter.Context;
 
 	UClass const* NodeClass = BlueprintAction.GetNodeClass();
-	UEdGraphNode const* NodeCDO = CastChecked<UEdGraphNode>(NodeClass->ClassDefaultObject);
+	UEdGraphNode const* NodeCDO = CastChecked<UEdGraphNode>(NodeClass->GetDefaultObject(false));
 	checkSlow(NodeCDO != nullptr);
 
 
diff --git a/Engine/Source/Editor/BlueprintGraph/Private/K2Node_AddComponent.cpp b/Engine/Source/Editor/BlueprintGraph/Private/K2Node_AddComponent.cpp
--- a/Engine/Source/Editor/BlueprintGraph/Private/K2Node_AddComponent.cpp
+++ b/Engine/Source/Editor/BlueprintGraph/Private/K2Node_AddComponent.cpp
@@ -119,7 +119,7 @@ void UK2Node_AddComponent::AllocatePinsForExposedVariables()
 	const UClass* ComponentClass = GetSpawnedType();
 	if (ComponentClass != nullptr)
 	{
-		const UObject* ClassDefaultObject = ComponentClass->ClassDefaultObject;
+		const UObject* ClassDefaultObject = ComponentClass->GetDefaultObject(false);
 
 		for (TFieldIterator<FProperty> PropertyIt(ComponentClass, EFieldIteratorFlags::IncludeSuper); PropertyIt; ++PropertyIt)
 		{
@@ -261,9 +261,9 @@ void UK2Node_AddComponent::ValidateNodeDuringCompilation(FCompilerResultsLog& Me
 					MessageLog.Error(*FText::Format(NSLOCTEXT("KismetCompiler", "AddSelfComponent_Error", "@@ cannot add a '{ChildActorClass}' component in the construction script (could cause infinite recursion)."), Args).ToString(), this);
 				}
 			}
-			else if (ChildActorClass != nullptr && ChildActorClass->ClassDefaultObject)
+			else if (ChildActorClass != nullptr && ChildActorClass->GetDefaultObject(false))
 			{
-				AActor const* ChildActor = Cast<AActor>(ChildActorClass->ClassDefaultObject);
+				AActor const* ChildActor = Cast<AActor>(ChildActorClass->GetDefaultObject(false));
 				check(ChildActor != nullptr);
 				USceneComponent* RootComponent = ChildActor->GetRootComponent();
 
diff --git a/Engine/Source/Editor/BlueprintGraph/Private/K2Node_SpawnActorFromClass.cpp b/Engine/Source/Editor/BlueprintGraph/Private/K2Node_SpawnActorFromClass.cpp
--- a/Engine/Source/Editor/BlueprintGraph/Private/K2Node_SpawnActorFromClass.cpp
+++ b/Engine/Source/Editor/BlueprintGraph/Private/K2Node_SpawnActorFromClass.cpp
@@ -393,7 +393,7 @@ void UK2Node_SpawnActorFromClass::FixupScaleMethodPin()
 	        const UEnum* const ScaleMethodEnum = StaticEnum<ESpawnActorScaleMethod>();
             if (const UClass* Class = Cast<UClass>(ClassPin->DefaultObject))
             {
-            	if (const AActor* ActorCDO = Cast<AActor>(Class->ClassDefaultObject))
+            	if (const AActor* ActorCDO = Cast<AActor>(Class->GetDefaultObject(false)))
                 {
                     if (const USceneComponent* Root = ActorCDO->GetRootComponent()) // native root component
                     {
diff --git a/Engine/Source/Editor/Blutility/Private/BlutilityModule.cpp b/Engine/Source/Editor/Blutility/Private/BlutilityModule.cpp
--- a/Engine/Source/Editor/Blutility/Private/BlutilityModule.cpp
+++ b/Engine/Source/Editor/Blutility/Private/BlutilityModule.cpp
@@ -238,7 +238,7 @@ public:
 			{
 				if (NativeParent->IsChildOf(AActor::StaticClass()))
 				{
-					if(const AActor* NativeCDO = Cast<AActor>(NativeParent->ClassDefaultObject))
+					if(const AActor* NativeCDO = Cast<AActor>(NativeParent->GetDefaultObject(false)))
 					{
 						return !NativeCDO->IsEditorOnlyLoadedInPIE();
 					}
diff --git a/Engine/Source/Editor/GameplayTasksEditor/Private/K2Node_LatentGameplayTaskCall.cpp b/Engine/Source/Editor/GameplayTasksEditor/Private/K2Node_LatentGameplayTaskCall.cpp
--- a/Engine/Source/Editor/GameplayTasksEditor/Private/K2Node_LatentGameplayTaskCall.cpp
+++ b/Engine/Source/Editor/GameplayTasksEditor/Private/K2Node_LatentGameplayTaskCall.cpp
@@ -458,12 +458,12 @@ bool UK2Node_LatentGameplayTaskCall::ConnectSpawnProperties(UClass* ClassToSpawn
 				}
 
 				// This is sloppy, we should be comparing to defaults much later in the compile process:
-				if (ClassToSpawn->ClassDefaultObject != nullptr)
+				if (ClassToSpawn->GetDefaultObject(false) != nullptr)
 				{
 					// We don't want to generate an assignment node unless the default value 
 					// differs from the value in the CDO:
 					FString DefaultValueAsString;
-					FBlueprintEditorUtils::PropertyValueToString(Property, (uint8*)ClassToSpawn->ClassDefaultObject.Get(), DefaultValueAsString, this);
+					FBlueprintEditorUtils::PropertyValueToString(Property, (uint8*)ClassToSpawn->GetDefaultObject(false), DefaultValueAsString, this);
 					if (DefaultValueAsString == SpawnVarPin->DefaultValue)
 					{
 						continue;
diff --git a/Engine/Source/Editor/GraphEditor/Private/BlueprintConnectionDrawingPolicy.cpp b/Engine/Source/Editor/GraphEditor/Private/BlueprintConnectionDrawingPolicy.cpp
--- a/Engine/Source/Editor/GraphEditor/Private/BlueprintConnectionDrawingPolicy.cpp
+++ b/Engine/Source/Editor/GraphEditor/Private/BlueprintConnectionDrawingPolicy.cpp
@@ -405,7 +405,7 @@ UObject* FKismetConnectionDrawingPolicy::GetObjectBeingDebugged(UBlueprint* Targ
 		if (UClass* GeneratedClass = TargetBP->GeneratedClass;
 			ActiveObject == nullptr && BPTYPE_FunctionLibrary == TargetBP->BlueprintType && GeneratedClass)
 		{
-			ActiveObject = GeneratedClass->ClassDefaultObject;
+			ActiveObject = GeneratedClass->GetDefaultObject(false);
 		}
 	}
 	return ActiveObject;
diff --git a/Engine/Source/Editor/Kismet/Private/BlueprintCompilationManager.cpp b/Engine/Source/Editor/Kismet/Private/BlueprintCompilationManager.cpp
--- a/Engine/Source/Editor/Kismet/Private/BlueprintCompilationManager.cpp
+++ b/Engine/Source/Editor/Kismet/Private/BlueprintCompilationManager.cpp
@@ -706,8 +706,8 @@ void FBlueprintCompilationManagerImpl::FlushCompilationQueueImpl(bool bSuppressB
 			UBlueprint* QueuedBP = QueuedJob.BPToCompile;
 
 			ensure(!QueuedBP->GeneratedClass ||
-				!QueuedBP->GeneratedClass->ClassDefaultObject ||
-				!(QueuedBP->GeneratedClass->ClassDefaultObject->HasAnyFlags(RF_NeedLoad)));
+				!QueuedBP->GeneratedClass->GetDefaultObject(false) ||
+				!(QueuedBP->GeneratedClass->GetDefaultObject(false)->HasAnyFlags(RF_NeedLoad)));
 			bool bDefaultComponentMustBeAdded = false;
 			bool bHasPendingUberGraphFrame = false;
 			UBlueprintGeneratedClass* BPGC = Cast<UBlueprintGeneratedClass>(QueuedBP->GeneratedClass);
@@ -1359,7 +1359,7 @@ void FBlueprintCompilationManagerImpl::FlushCompilationQueueImpl(bool bSuppressB
 
 				if(BP->GeneratedClass)
 				{
-					OldCDOs.Add(BP, BP->GeneratedClass->ClassDefaultObject);
+					OldCDOs.Add(BP, BP->GeneratedClass->GetDefaultObject(false));
 				}
 
 				EBlueprintCompileReinstancerFlags CompileReinstancerFlags =
@@ -1420,14 +1420,14 @@ void FBlueprintCompilationManagerImpl::FlushCompilationQueueImpl(bool bSuppressB
 				SCOPED_LOADTIMER_ASSET_TEXT(*BP->GetPathName());
 
 				ensure( BP->GeneratedClass == nullptr ||
-						BP->GeneratedClass->ClassDefaultObject == nullptr || 
-						BP->GeneratedClass->ClassDefaultObject->GetClass() != BP->GeneratedClass);
+						BP->GeneratedClass->GetDefaultObject(false) == nullptr ||
+						BP->GeneratedClass->GetDefaultObject(false)->GetClass() != BP->GeneratedClass);
 				// default value propagation occurs in ReinstaneBatch, CDO will be created via CompileFunctions call:
 				if(BP->ParentClass)
 				{
 					if(BP->GeneratedClass)
 					{
-						BP->GeneratedClass->ClassDefaultObject = nullptr;
+						BP->GeneratedClass->SetDefaultObject(nullptr);
 					}
 
 					// Reset the flag, so if the user tries to use PIE it will warn them if the BP did not compile
@@ -1445,7 +1445,7 @@ void FBlueprintCompilationManagerImpl::FlushCompilationQueueImpl(bool bSuppressB
 					{
 						ChildClass->Bind();
 						ChildClass->StaticLink();
-						ensure(ChildClass->ClassDefaultObject == nullptr);
+						ensure(ChildClass->GetDefaultObject(false) == nullptr);
 					}
 				}
 				else
@@ -1493,8 +1493,8 @@ void FBlueprintCompilationManagerImpl::FlushCompilationQueueImpl(bool bSuppressB
 			if(!CompilerData.ShouldCompileClassFunctions())
 			{
 				if( BPGC &&
-					(	BPGC->ClassDefaultObject == nullptr || 
-						BPGC->ClassDefaultObject->GetClass() != BPGC) )
+					(	BPGC->GetDefaultObject(false) == nullptr ||
+						BPGC->GetDefaultObject(false)->GetClass() != BPGC) )
 				{
 					if (CompilerData.Reinstancer.IsValid())
 					{
@@ -1504,7 +1504,7 @@ void FBlueprintCompilationManagerImpl::FlushCompilationQueueImpl(bool bSuppressB
 					BPGC->bLayoutChanging = false;
 					BPGC->Bind();
 					BPGC->StaticLink(true);
-					BPGC->ClassDefaultObject = nullptr;
+					BPGC->SetDefaultObject(nullptr);
 					BPGC->GetDefaultObject(true);
 				}
 			}
@@ -1521,8 +1521,8 @@ void FBlueprintCompilationManagerImpl::FlushCompilationQueueImpl(bool bSuppressB
 						CompilerData.Reinstancer->PropagateSparseClassDataToNewClass(BPGC);
 					}
 
-					if( BPGC->ClassDefaultObject && 
-						BPGC->ClassDefaultObject->GetClass() == BPGC)
+					if( BPGC->GetDefaultObject(false) &&
+						BPGC->GetDefaultObject(false)->GetClass() == BPGC)
 					{
 						// the CDO has been created early, it is possible that the reflection data was still
 						// being mutated by CompileClassLayout. Warn the user and and move the CDO aside:
@@ -1531,7 +1531,7 @@ void FBlueprintCompilationManagerImpl::FlushCompilationQueueImpl(bool bSuppressB
 							*BP->GetName()
 						);
 
-						BPGC->ClassDefaultObject->Rename(
+						BPGC->GetDefaultObject(false)->Rename(
 							nullptr,
 							// destination - this is the important part of this call. Moving the object 
 							// out of the way so we can reuse its name:
@@ -1540,7 +1540,7 @@ void FBlueprintCompilationManagerImpl::FlushCompilationQueueImpl(bool bSuppressB
 							REN_DoNotDirty | REN_DontCreateRedirectors
 						);
 					}
-					BPGC->ClassDefaultObject = nullptr;
+					BPGC->SetDefaultObject(nullptr);
 		
 					// class layout is ready, we can clear bLayoutChanging and CompileFunctions can create the CDO:
 					BPGC->bLayoutChanging = false;
@@ -1592,7 +1592,7 @@ void FBlueprintCompilationManagerImpl::FlushCompilationQueueImpl(bool bSuppressB
 			
 			FKismetCompilerUtilities::UpdateDependentBlueprints(BP);
 
-			ensure(BPGC == nullptr || BPGC->ClassDefaultObject->GetClass() == BPGC);
+			ensure(BPGC == nullptr || BPGC->GetDefaultObject(false)->GetClass() == BPGC);
 		}
 	} // end TimeCompiling scope
 
@@ -1646,13 +1646,13 @@ void FBlueprintCompilationManagerImpl::FlushCompilationQueueImpl(bool bSuppressB
 
 				const UBlueprintGeneratedClass* GenClass = Cast<UBlueprintGeneratedClass>(CompilerData.BP->GeneratedClass);
 
-				if (GenClass && GenClass->ClassDefaultObject)
+				if (GenClass && GenClass->GetDefaultObject(false))
 				{
 					for (const FBPVariableDescription& NewInheritedVar : CompilerData.NewDefaultVariables)
 					{
 						if (const FProperty* MatchingProperty = GenClass->FindPropertyByName(NewInheritedVar.VarName))
 						{
-							FBlueprintEditorUtils::PropertyValueFromString(MatchingProperty, NewInheritedVar.DefaultValue, reinterpret_cast<uint8*>(GenClass->ClassDefaultObject.Get()));
+							FBlueprintEditorUtils::PropertyValueFromString(MatchingProperty, NewInheritedVar.DefaultValue, reinterpret_cast<uint8*>(GenClass->GetDefaultObject(false)));
 						}
 					}
 				}
@@ -2095,9 +2095,9 @@ void FBlueprintCompilationManagerImpl::GetDefaultValue(const UClass* ForClass, c
 		return;
 	}
 
-	if (ForClass->ClassDefaultObject)
+	if (ForClass->GetDefaultObject(false))
 	{
-		FBlueprintEditorUtils::PropertyValueToString(Property, (uint8*)ForClass->ClassDefaultObject.Get(), OutDefaultValueAsString);
+		FBlueprintEditorUtils::PropertyValueToString(Property, (uint8*)ForClass->GetDefaultObject(false), OutDefaultValueAsString);
 	}
 	else
 	{
@@ -2154,7 +2154,7 @@ void FBlueprintCompilationManagerImpl::SaveConstructionHashes(const TArray<FComp
 		UBlueprint* BP = CompilerData.BP;
 		if(const UClass* OriginalClass = BP->GeneratedClass)
 		{
-			if( const UObject* CDO = OriginalClass->ClassDefaultObject )
+			if( const UObject* CDO = OriginalClass->GetDefaultObject(false) )
 			{
 				if( FOverridableManager::Get().IsEnabled(CDO) )
 				{
@@ -2169,7 +2169,7 @@ void FBlueprintCompilationManagerImpl::SaveConstructionHashes(const TArray<FComp
 		UBlueprint* BP = CompilerData.BP;
 		if(const UClass* OriginalClass = BP->GeneratedClass)
 		{
-			if(OriginalClass->ClassDefaultObject)
+			if(OriginalClass->GetDefaultObject(false))
 			{
 				// if there are no instances, no need to calculate the class hash:
 				FBlake3Hash ConstructionHash = FBlueprintCompileReinstancer::CalculateConstructionHashIfNecessary(OriginalClass);
@@ -2198,7 +2198,7 @@ void FBlueprintCompilationManagerImpl::ReparentHierarchies(const TMap<UClass*, U
 		for (const TPair<UClass*, UClass*>& OldToNewClass : OldToNewClasses)
 		{
 			// classes with no CDO do not need to be reinstanced:
-			if (OldToNewClass.Key->ClassDefaultObject == nullptr)
+			if (OldToNewClass.Key->GetDefaultObject(false) == nullptr)
 			{
 				continue;
 			}
@@ -2213,7 +2213,7 @@ void FBlueprintCompilationManagerImpl::ReparentHierarchies(const TMap<UClass*, U
 
 			for (UClass* DerivedClass : DerivedClasses)
 			{
-				if (DerivedClass->ClassDefaultObject == nullptr && 
+				if (DerivedClass->GetDefaultObject(false) == nullptr &&
 					DerivedClass->GetSparseClassData(EGetSparseClassDataMethod::ReturnIfNull) == nullptr)
 				{
 					// no CDO->no other instances->no need to reinstance..
@@ -2473,7 +2473,7 @@ void FBlueprintCompilationManagerImpl::ReinstanceBatch(TArray<FReinstancingJob>&
 			{
 				if(IsValid(ClassToReinstance))
 				{
-					if (ClassToReinstance->ClassDefaultObject)
+					if (ClassToReinstance->GetDefaultObject(false))
 					{
 						ClassesToReinstance.Add(ClassToReinstance);
 					}
@@ -2547,7 +2547,6 @@ void FBlueprintCompilationManagerImpl::ReinstanceBatch(TArray<FReinstancingJob>&
 	// make new hierarchy
 	for(UClass* Class : ClassesToReinstance)
 	{
-		UObject* OriginalCDO = Class->ClassDefaultObject;
 		Reinstancers.Emplace(
 			FReinstancingJob ( 
 				TSharedPtr<FBlueprintCompileReinstancer>( 
@@ -2627,7 +2626,7 @@ void FBlueprintCompilationManagerImpl::ReinstanceBatch(TArray<FReinstancingJob>&
 		{
 			AddDependentClass(OldClass->GetSuperClass());
 
-			if (const UObject* CDO = OldClass->ClassDefaultObject)
+			if (const UObject* CDO = OldClass->GetDefaultObject(false))
 			{
 				TArray<UObject*> ContainedOldObjects;
 				GetObjectsWithOuter(CDO, ContainedOldObjects);
@@ -2672,7 +2671,7 @@ void FBlueprintCompilationManagerImpl::ReinstanceBatch(TArray<FReinstancingJob>&
 
 		if (OldClass)
 		{
-			OldCDO = OldClass->ClassDefaultObject;
+			OldCDO = OldClass->GetDefaultObject(false);
 			if (OldCDO && ReinstancingJob.Reinstancer.IsValid())
 			{
 				// Object using overridable serialization need to use delta serialization for it do work appropriately
@@ -2683,7 +2682,7 @@ void FBlueprintCompilationManagerImpl::ReinstanceBatch(TArray<FReinstancingJob>&
 
 				// We do not expect the CDO to be already created at this point. 
 				// Who ever bother to create it before didn't look if the parent was ready or not.
-				if (NewClass->ClassDefaultObject != nullptr && bUsingOverrideSerialization)
+				if (NewClass->GetDefaultObject(false) != nullptr && bUsingOverrideSerialization)
 				{
 					// Override serialization does not currently want to use any CDO created before this
 					// point. I do not have a test case justifying this need, but renaming the cdo here 
@@ -2692,8 +2691,8 @@ void FBlueprintCompilationManagerImpl::ReinstanceBatch(TArray<FReinstancingJob>&
 					// we will have to address these shortcomings (likely by ordering compilation itself,
 					// rather than reinstancing, more carefully). Discarding the CDO here is also wasteful
 					// - if we cannot use the objects, why create them at all?
-					NewClass->ClassDefaultObject->Rename(nullptr, GetTransientPackage(), REN_DoNotDirty | REN_DontCreateRedirectors | REN_NonTransactional);
-					NewClass->ClassDefaultObject = nullptr;
+					NewClass->GetDefaultObject(false)->Rename(nullptr, GetTransientPackage(), REN_DoNotDirty | REN_DontCreateRedirectors | REN_NonTransactional);
+					NewClass->SetDefaultObject(nullptr);
 				}
 				UObject* NewCDO = NewClass->GetDefaultObject(true);
 
@@ -2754,7 +2753,7 @@ void FBlueprintCompilationManagerImpl::ReinstanceBatch(TArray<FReinstancingJob>&
 			if (ReinstancingJob.Reinstancer.IsValid())
 			{
 				UBlueprint* CurrentBP = CastChecked<UBlueprint>(BPGClass->ClassGeneratedBy);
-				UObject* NewCDO = CurrentBP->GeneratedClass->ClassDefaultObject;
+				UObject* NewCDO = CurrentBP->GeneratedClass->GetDefaultObject(false);
 				FLinkerLoad* CurrentLinker = CurrentBP->GetLinker();
 				if (CurrentLinker)
 				{
@@ -2822,13 +2821,13 @@ void FBlueprintCompilationManagerImpl::ReinstanceBatch(TArray<FReinstancingJob>&
 
 			UClass* NewClass = ReinstancingJob.OldToNew.Value;
 			if (NewClass && 
-				OldClass->ClassDefaultObject && 
-				NewClass->ClassDefaultObject &&
-				OldClass->ClassDefaultObject != NewClass->ClassDefaultObject)
+				OldClass->GetDefaultObject(false) &&
+				NewClass->GetDefaultObject(false) &&
+				OldClass->GetDefaultObject(false) != NewClass->GetDefaultObject(false))
 			{
-				OldArchetypeToNewArchetype.Add(OldClass->ClassDefaultObject, NewClass->ClassDefaultObject);
+				OldArchetypeToNewArchetype.Add(OldClass->GetDefaultObject(false), NewClass->GetDefaultObject(false));
 				// also map old *default* subobjects to new default subobjects:
-				BuildDSOMap(OldClass->ClassDefaultObject, NewClass->ClassDefaultObject, InOutOldToNewClassMap, OldArchetypeToNewArchetype);
+				BuildDSOMap(OldClass->GetDefaultObject(false), NewClass->GetDefaultObject(false), InOutOldToNewClassMap, OldArchetypeToNewArchetype);
 			}
 
 			TArray<UObject*> ArchetypeObjects;
@@ -3166,8 +3165,9 @@ UClass* FBlueprintCompilationManagerImpl::FastGenerateSkeletonClass(UBlueprint*
 		{
 			CompilerContext.ValidateVariableNames();
 		}
-
+PRAGMA_DISABLE_DEPRECATION_WARNINGS
 		CompilerContext.CleanAndSanitizeClass(Ret, MutableView(Ret->ClassDefaultObject));
+PRAGMA_ENABLE_DEPRECATION_WARNINGS
 	}
 	
 	Ret->ClassGeneratedBy = BP;
@@ -3703,7 +3703,7 @@ void FBlueprintCompilationManagerImpl::ConformToParentAndInterfaces(UBlueprint*
 void FBlueprintCompilationManagerImpl::RelinkSkeleton(UClass* SkeletonToRelink)
 {
 	// CDO needs to be moved aside already:
-	ensure(SkeletonToRelink->ClassDefaultObject == nullptr);
+	ensure(SkeletonToRelink->GetDefaultObject(false) == nullptr);
 	ensure(!SkeletonToRelink->GetSuperClass()->HasAnyClassFlags(CLASS_NewerVersionExists));
 
 	SkeletonToRelink->ClassConstructor = nullptr;
@@ -3759,7 +3759,7 @@ void FBlueprintCompilationManagerImpl::QueueOutOfDateDependencies( const TArray<
 		{
 			return
 				CompileRequest.BPToCompile->GeneratedClass != nullptr &&
-				CompileRequest.BPToCompile->GeneratedClass->ClassDefaultObject != nullptr &&
+				CompileRequest.BPToCompile->GeneratedClass->GetDefaultObject(false) != nullptr &&
 				(CompileRequest.CompileOptions & EBlueprintCompileOptions::RegenerateSkeletonOnly)
 				== EBlueprintCompileOptions::None;
 		},
diff --git a/Engine/Source/Editor/Kismet/Private/Blueprints/BlueprintDependencies.cpp b/Engine/Source/Editor/Kismet/Private/Blueprints/BlueprintDependencies.cpp
--- a/Engine/Source/Editor/Kismet/Private/Blueprints/BlueprintDependencies.cpp
+++ b/Engine/Source/Editor/Kismet/Private/Blueprints/BlueprintDependencies.cpp
@@ -212,7 +212,7 @@ static TArray<const UObject*> UE::Private::GetAllImportedObjects(TConstArrayView
 
 static TArray<const UObject*> UE::Private::GetAllImportedObjects( const UBlueprint* ForBP)
 {
-	return GetAllImportedObjects({ForBP, ForBP->GeneratedClass, ForBP->GeneratedClass->ClassDefaultObject});
+	return GetAllImportedObjects({ForBP, ForBP->GeneratedClass, ForBP->GeneratedClass->GetDefaultObject(false) });
 }
 
 static void UE::Private::AddStructDependencyImpl(const UScriptStruct* Struct, TSet<const UScriptStruct*>& ReferencedStructs)
@@ -387,7 +387,7 @@ static FString UE::Private::HashNativeClass( const UClass* NativeClass)
 	// always produced an empty object, which is not meaningful. JsonObjectGraph::Stringify
 	// will provide identity for the UClass:
 	FUtf8String StringifiedClass = UE::JsonObjectGraph::Stringify(
-		{NativeClass, NativeClass->ClassDefaultObject}, 
+		{NativeClass, NativeClass->GetDefaultObject(false)},
 		FJsonStringifyOptions(EJsonStringifyFlags::FilterEditorOnlyData));
 	FBlake3 Hash;
 	Hash.Update(StringifiedClass.GetCharArray().GetData(), StringifiedClass.GetCharArray().Num());
diff --git a/Engine/Source/Editor/Kismet/Private/Blueprints/JsonObjectGraphFunctionLibrary.cpp b/Engine/Source/Editor/Kismet/Private/Blueprints/JsonObjectGraphFunctionLibrary.cpp
--- a/Engine/Source/Editor/Kismet/Private/Blueprints/JsonObjectGraphFunctionLibrary.cpp
+++ b/Engine/Source/Editor/Kismet/Private/Blueprints/JsonObjectGraphFunctionLibrary.cpp
@@ -126,7 +126,7 @@ void UJsonObjectGraphFunctionLibrary::WriteBlueprintClassToTempFile(const UBluep
 		return;
 	}
 
-	if (!BPGC->ClassDefaultObject)
+	if (!BPGC->GetDefaultObject(false))
 	{
 		UE_LOG(LogBlueprint, Warning, TEXT("Attempted to serialize class with no CDO: %s"), *BPGC->GetPathName());
 		return;
@@ -134,7 +134,7 @@ void UJsonObjectGraphFunctionLibrary::WriteBlueprintClassToTempFile(const UBluep
 
 	// When writing a class we should always exclude editor only data:
 	Options.Flags |= EJsonStringifyFlags::FilterEditorOnlyData;
-	FUtf8String Result = UE::JsonObjectGraph::Stringify({ BPGC, BPGC->ClassDefaultObject }, Options);
+	FUtf8String Result = UE::JsonObjectGraph::Stringify({ BPGC, BPGC->GetDefaultObject(false) }, Options);
 	FString Filename = UE::Private::GetIntermediateAssetName(BP->GetPackage(), *Label);
 	TUniquePtr<FArchive> FileArchive(IFileManager::Get().CreateDebugFileWriter(*Filename));
 	FileArchive->Serialize((void*)Result.GetCharArray().GetData(), Result.GetCharArray().Num() - 1);
diff --git a/Engine/Source/Editor/KismetCompiler/Private/KismetCompiler.cpp b/Engine/Source/Editor/KismetCompiler/Private/KismetCompiler.cpp
--- a/Engine/Source/Editor/KismetCompiler/Private/KismetCompiler.cpp
+++ b/Engine/Source/Editor/KismetCompiler/Private/KismetCompiler.cpp
@@ -655,7 +655,7 @@ void FKismetCompilerContext::CleanAndSanitizeClass(UBlueprintGeneratedClass* Cla
 	TransientClass->ClassFlags |= CLASS_CompiledFromBlueprint|CLASS_NewerVersionExists;
 
 	SetNewClass( ClassToClean );
-	InOldCDO = ClassToClean->ClassDefaultObject; // we don't need to create the CDO at this point
+	InOldCDO = ClassToClean->GetDefaultObject(false); // we don't need to create the CDO at this point
 	
 	const ERenameFlags RenFlags = REN_DontCreateRedirectors | REN_NonTransactional | REN_DoNotDirty;
 
@@ -5263,7 +5263,7 @@ void FKismetCompilerContext::CompileFunctions(EInternalCompilerFlags InternalFla
 
 
 	check(NewClass->PropertiesSize >= UObject::StaticClass()->PropertiesSize);
-	check(NewClass->ClassDefaultObject != NULL);
+	check(NewClass->GetDefaultObject(false) != NULL);
 
 	PostCompileDiagnostics();
 
@@ -5439,23 +5439,23 @@ void FKismetCompilerContext::PostCDOCompiled(const UObject::FPostCDOCompiledCont
 	}
 
 	// Conform any sparse class data now that the CDO is ready
-	NewClass->ConformSparseClassData(NewClass->ClassDefaultObject);
+	NewClass->ConformSparseClassData(NewClass->GetDefaultObject(false));
 
 	// Notify the CDO that it has finished compiling
-	NewClass->ClassDefaultObject->PostCDOCompiled(Context);
+	NewClass->GetDefaultObject(false)->PostCDOCompiled(Context);
 	if (!Context.bIsSkeletonOnly)
 	{
 		PRAGMA_DISABLE_DEPRECATION_WARNINGS;
-		NewClass->ClassDefaultObject->PostCDOCompiled();
+		NewClass->GetDefaultObject(false)->PostCDOCompiled();
 		PRAGMA_ENABLE_DEPRECATION_WARNINGS;
 		
 		for (const TFunction<void(const UObject::FPostCDOCompiledContext&, UObject*)>& Step : PostCDOCompileSteps)
 		{
-			Step(Context, NewClass->ClassDefaultObject);
+			Step(Context, NewClass->GetDefaultObject(false));
 		}
 	}
 	
-	FCoreUObjectDelegates::OnObjectPostCDOCompiled.Broadcast(NewClass->ClassDefaultObject, Context);
+	FCoreUObjectDelegates::OnObjectPostCDOCompiled.Broadcast(NewClass->GetDefaultObject(false), Context);
 
 	// Allow children to customize PostCDOCompile:
 	OnPostCDOCompiled(Context);
diff --git a/Engine/Source/Editor/KismetCompiler/Private/KismetCompilerMisc.cpp b/Engine/Source/Editor/KismetCompiler/Private/KismetCompilerMisc.cpp
--- a/Engine/Source/Editor/KismetCompiler/Private/KismetCompilerMisc.cpp
+++ b/Engine/Source/Editor/KismetCompiler/Private/KismetCompilerMisc.cpp
@@ -189,8 +189,8 @@ static bool DoesTypeNotMatchProperty(UEdGraphPin* SourcePin, const FEdGraphPinTy
 			ObjectType->HasAnyFlags(RF_Transient) &&
 			ObjectType->HasAnyClassFlags(CLASS_NewerVersionExists) &&
 			!ObjectType->IsRooted() && 
-			ObjectType->ClassDefaultObject &&
-			!ObjectType->ClassDefaultObject->IsRooted())
+			ObjectType->GetDefaultObject(false) &&
+			!ObjectType->GetDefaultObject(false)->IsRooted())
 		{
 			MessageLog.Error(*LOCTEXT("InvalidClassForPin_Error", "Class for pin @@ is invalid! It has likely been deleted.").ToString(), SourcePin);
 		}
@@ -507,12 +507,12 @@ void FKismetCompilerUtilities::ConsignToOblivion(UClass* OldClass, bool bForceNo
 		TSharedPtr<FBlueprintCompileReinstancer> CTOResinstancer = FBlueprintCompileReinstancer::Create(OldClass);
 
 		UPackage* OwnerPackage = OldClass->GetPackage();
-		if( OldClass->ClassDefaultObject )
+		if (UObject* OldClassDefaultObject = OldClass->GetDefaultObject(false))
 		{
 			// rename to a temp name, move into transient package
-			OldClass->ClassDefaultObject->ClearFlags(RF_Public);
-			OldClass->ClassDefaultObject->SetFlags(RF_Transient);
-			OldClass->ClassDefaultObject->RemoveFromRoot(); // make sure no longer in root set
+			OldClassDefaultObject->ClearFlags(RF_Public);
+			OldClassDefaultObject->SetFlags(RF_Transient);
+			OldClassDefaultObject->RemoveFromRoot(); // make sure no longer in root set
 		}
 		
 		OldClass->SetMetaData(FBlueprintMetadata::MD_IsBlueprintBase, TEXT("false"));
@@ -816,9 +816,9 @@ UEdGraphPin* FKismetCompilerUtilities::GenerateAssignmentNodes(class FKismetComp
 
 					if (!FBlueprintCompilationManager::GetDefaultValue(ForClass, Property, DefaultValueAsString))
 					{
-						if (ForClass->ClassDefaultObject)
+						if (ForClass->GetDefaultObject(false))
 						{
-							FBlueprintEditorUtils::PropertyValueToString(Property, (uint8*)ForClass->ClassDefaultObject.Get(), DefaultValueAsString);
+							FBlueprintEditorUtils::PropertyValueToString(Property, (uint8*)ForClass->GetDefaultObject(false), DefaultValueAsString);
 						}
 					}
 
diff --git a/Engine/Source/Editor/KismetCompiler/Private/KismetCompilerModule.cpp b/Engine/Source/Editor/KismetCompiler/Private/KismetCompilerModule.cpp
--- a/Engine/Source/Editor/KismetCompiler/Private/KismetCompilerModule.cpp
+++ b/Engine/Source/Editor/KismetCompiler/Private/KismetCompilerModule.cpp
@@ -100,7 +100,7 @@ void FKismet2CompilerModule::RefreshVariables(UBlueprint* Blueprint)
 
 	if(MissingVariables.Num() != 0)
 	{
-		UObject* OldCDO = Blueprint->GeneratedClass->ClassDefaultObject;
+		UObject* OldCDO = Blueprint->GeneratedClass->GetDefaultObject(false);
 		auto Reinstancer = FBlueprintCompileReinstancer::Create(*Blueprint->GeneratedClass);
 		ensure(OldCDO->GetClass() != *Blueprint->GeneratedClass);
 		// move old cdo aside:
@@ -108,7 +108,7 @@ void FKismet2CompilerModule::RefreshVariables(UBlueprint* Blueprint)
 		{
 			OldCDO->Rename(NULL, GetTransientPackage(), REN_DontCreateRedirectors | REN_NonTransactional | REN_AllowPackageLinkerMismatch);
 		}
-		Blueprint->GeneratedClass->ClassDefaultObject = nullptr;
+		Blueprint->GeneratedClass->SetDefaultObject(nullptr);
 
 		// add missing properties:
 		for( int32 MissingVarIndex : MissingVariables)
@@ -137,7 +137,7 @@ void FKismet2CompilerModule::RefreshVariables(UBlueprint* Blueprint)
 		// cpfuo:
 		UEngine::FCopyPropertiesForUnrelatedObjectsParams Params;
 		Params.bNotifyObjectReplacement = true;
-		UEngine::CopyPropertiesForUnrelatedObjects(OldCDO, Blueprint->GeneratedClass->ClassDefaultObject, Params);
+		UEngine::CopyPropertiesForUnrelatedObjects(OldCDO, Blueprint->GeneratedClass->GetDefaultObject(false), Params);
 	}
 }
 
diff --git a/Engine/Source/Editor/LevelEditor/Private/LevelEditorCreateActorMenu.cpp b/Engine/Source/Editor/LevelEditor/Private/LevelEditorCreateActorMenu.cpp
--- a/Engine/Source/Editor/LevelEditor/Private/LevelEditorCreateActorMenu.cpp
+++ b/Engine/Source/Editor/LevelEditor/Private/LevelEditorCreateActorMenu.cpp
@@ -104,7 +104,7 @@ static void GetMenuEntryText(const FAssetData& Asset, const TArray<FActorFactory
 			UClass* MenuItemClass = Cast<UClass>(MenuItem.AssetData.GetAsset());
 			if (MenuItemClass && MenuItemClass->IsChildOf(AActor::StaticClass()))
 			{
-				AActor* DefaultActor = Cast<AActor>(Cast<UClass>(MenuItem.AssetData.GetAsset())->ClassDefaultObject);
+				AActor* DefaultActor = Cast<AActor>(Cast<UClass>(MenuItem.AssetData.GetAsset())->GetDefaultObject(false));
 				OutActorTypeDisplayName = DefaultActor->GetClass()->GetDisplayNameText();
 			}
 		}
diff --git a/Engine/Source/Editor/LevelEditor/Private/LevelEditorOutlinerSettings.cpp b/Engine/Source/Editor/LevelEditor/Private/LevelEditorOutlinerSettings.cpp
--- a/Engine/Source/Editor/LevelEditor/Private/LevelEditorOutlinerSettings.cpp
+++ b/Engine/Source/Editor/LevelEditor/Private/LevelEditorOutlinerSettings.cpp
@@ -225,14 +225,14 @@ void FLevelEditorOutlinerSettings::CreateDefaultFilters()
 			// Get the underlying class from the Actor belonging to this item
 			const bool bIsClass = Item->AssetData.GetClass() == UClass::StaticClass();
 			const bool bIsActor = bIsClass ? CastChecked<UClass>(Item->AssetData.GetAsset())->IsChildOf(AActor::StaticClass()) : false;
-			AActor* DefaultActor = nullptr;
+			const AActor* DefaultActor = nullptr;
 			if (Item->Factory != nullptr)
 			{
 				DefaultActor = Item->Factory->GetDefaultActor(Item->AssetData);
 			}
 			else if (bIsActor)
 			{
-				DefaultActor = CastChecked<AActor>(CastChecked<UClass>(Item->AssetData.GetAsset())->ClassDefaultObject);
+				DefaultActor = CastChecked<AActor>(CastChecked<UClass>(Item->AssetData.GetAsset())->GetDefaultObject(false));
 			
 			}
 			if (!DefaultActor)
diff --git a/Engine/Source/Editor/MovieSceneTools/Private/MovieSceneToolHelpers.cpp b/Engine/Source/Editor/MovieSceneTools/Private/MovieSceneToolHelpers.cpp
--- a/Engine/Source/Editor/MovieSceneTools/Private/MovieSceneToolHelpers.cpp
+++ b/Engine/Source/Editor/MovieSceneTools/Private/MovieSceneToolHelpers.cpp
@@ -5868,7 +5868,7 @@ void MovieSceneToolHelpers::ImportAnimSequenceTransforms(const TSharedRef<ISeque
 
 				if(ObjectClass)
 				{
-					const ACharacter* Character = Cast<const ACharacter>(ObjectClass->ClassDefaultObject);
+					const ACharacter* Character = Cast<const ACharacter>(ObjectClass->GetDefaultObject(false));
 					if(Character)
 					{
 						const USkeletalMeshComponent* SkeletalMeshComponent = Character->GetMesh();
diff --git a/Engine/Source/Editor/MovieSceneTools/Private/TrackEditors/CameraShakeSourceShakeTrackEditor.cpp b/Engine/Source/Editor/MovieSceneTools/Private/TrackEditors/CameraShakeSourceShakeTrackEditor.cpp
--- a/Engine/Source/Editor/MovieSceneTools/Private/TrackEditors/CameraShakeSourceShakeTrackEditor.cpp
+++ b/Engine/Source/Editor/MovieSceneTools/Private/TrackEditors/CameraShakeSourceShakeTrackEditor.cpp
@@ -581,7 +581,7 @@ bool FCameraShakeSourceShakeTrackEditor::OnShouldFilterCameraShake(const FAssetD
 	if (SelectedObject && SelectedObject->GeneratedClass && SelectedObject->GeneratedClass->IsChildOf(UCameraShakeBase::StaticClass()))
 	{
 		TSubclassOf<UCameraShakeBase> CameraShakeClass = *(SelectedObject->GeneratedClass);
-		if (const UCameraShakeBase* CameraShakeCDO = Cast<UCameraShakeBase>(CameraShakeClass->ClassDefaultObject))
+		if (const UCameraShakeBase* CameraShakeCDO = Cast<UCameraShakeBase>(CameraShakeClass->GetDefaultObject(false)))
 		{
 			return CameraShakeCDO->bSingleInstance;
 		}
diff --git a/Engine/Source/Editor/Persona/Private/PersonaUtils.cpp b/Engine/Source/Editor/Persona/Private/PersonaUtils.cpp
--- a/Engine/Source/Editor/Persona/Private/PersonaUtils.cpp
+++ b/Engine/Source/Editor/Persona/Private/PersonaUtils.cpp
@@ -208,7 +208,7 @@ static TSharedRef<SWidget> MakeNewNotifyPicker(UAnimSequenceBase* Sequence, cons
 		{
 			const bool bChildOfObjectClass = InClass->IsChildOf(NotifyTypeClass::StaticClass());
 			const bool bMatchesFlags = !InClass->HasAnyClassFlags(CLASS_Hidden | CLASS_HideDropDown | CLASS_Deprecated | CLASS_Abstract);
-			return bChildOfObjectClass && bMatchesFlags && CastChecked<NotifyTypeClass>(InClass->ClassDefaultObject)->CanBePlaced(Sequence);
+			return bChildOfObjectClass && bMatchesFlags && CastChecked<NotifyTypeClass>(InClass->GetDefaultObject(false))->CanBePlaced(Sequence);
 		}
 
 		virtual bool IsUnloadedClassAllowed(const FClassViewerInitializationOptions& InInitOptions, const TSharedRef< const IUnloadedBlueprintData > InUnloadedClassData, TSharedRef< FClassViewerFilterFuncs > InFilterFuncs) override
@@ -220,7 +220,7 @@ static TSharedRef<SWidget> MakeNewNotifyPicker(UAnimSequenceBase* Sequence, cons
 			{
 				if (const UClass* NativeBaseClass = InUnloadedClassData->GetNativeParent())
 				{
-					bValidToPlace = CastChecked<NotifyTypeClass>(NativeBaseClass->ClassDefaultObject)->CanBePlaced(Sequence);
+					bValidToPlace = CastChecked<NotifyTypeClass>(NativeBaseClass->GetDefaultObject(false))->CanBePlaced(Sequence);
 				}
 			}
 
diff --git a/Engine/Source/Editor/PlacementMode/Private/SPlacementModeTools.cpp b/Engine/Source/Editor/PlacementMode/Private/SPlacementModeTools.cpp
--- a/Engine/Source/Editor/PlacementMode/Private/SPlacementModeTools.cpp
+++ b/Engine/Source/Editor/PlacementMode/Private/SPlacementModeTools.cpp
@@ -173,7 +173,7 @@ void SPlacementAssetEntry::Construct(const FArguments& InArgs, const TSharedPtr<
 	}
 	else if (bIsActor)
 	{
-		DefaultActor = CastChecked<AActor>(CastChecked<UClass>(Item->AssetData.GetAsset())->ClassDefaultObject);
+		DefaultActor = CastChecked<AActor>(CastChecked<UClass>(Item->AssetData.GetAsset())->GetDefaultObject(false));
 	}
 
 	TSharedPtr<IToolTip> AssetEntryToolTip;
@@ -415,7 +415,7 @@ void SPlacementAssetMenuEntry::Construct(const FArguments& InArgs, const TShared
 	}
 	else if (bIsActor)
 	{
-		DefaultActor = CastChecked<AActor>(CastChecked<UClass>(Item->AssetData.GetAsset())->ClassDefaultObject);
+		DefaultActor = CastChecked<AActor>(CastChecked<UClass>(Item->AssetData.GetAsset())->GetDefaultObject(false));
 	}
 
 	UClass* DocClass = nullptr;
diff --git a/Engine/Source/Editor/PropertyEditor/Private/PropertyCustomizationHelpers.cpp b/Engine/Source/Editor/PropertyEditor/Private/PropertyCustomizationHelpers.cpp
--- a/Engine/Source/Editor/PropertyEditor/Private/PropertyCustomizationHelpers.cpp
+++ b/Engine/Source/Editor/PropertyEditor/Private/PropertyCustomizationHelpers.cpp
@@ -418,7 +418,7 @@ namespace PropertyCustomizationHelpers
 							FEditorScriptExecutionGuard ScriptGuard;
 							UEditorEngine* EditorEngine = Cast<UEditorEngine>(GEngine);
 							UObject* WorldContextObject = EditorEngine->GetEditorWorldContext().World();
-							TStrongObjectPtr<UObject> CDO(Function->GetOwnerClass()->ClassDefaultObject);
+							TStrongObjectPtr<UObject> CDO(Function->GetOwnerClass()->GetDefaultObject(false));
 							CDO->ProcessEvent(Function, &WorldContextObject);
 						}
 						else
diff --git a/Engine/Source/Editor/PropertyEditor/Private/PropertyNode.cpp b/Engine/Source/Editor/PropertyEditor/Private/PropertyNode.cpp
--- a/Engine/Source/Editor/PropertyEditor/Private/PropertyNode.cpp
+++ b/Engine/Source/Editor/PropertyEditor/Private/PropertyNode.cpp
@@ -4011,10 +4011,10 @@ void FPropertyNode::PropagatePropertyChange( UObject* ModifiedObject, const TCHA
 			GetDerivedClasses(Object->GetClass(), Children);
 			for (UClass* ChildClass : Children)
 			{
-				if (ChildClass->ClassDefaultObject &&
+				if (ChildClass->GetDefaultObject(false) &&
 					!ChildClass->GetPackage()->HasAnyFlags(RF_Transient))
 				{
-					ArchetypeInstances.Add(ChildClass->ClassDefaultObject);
+					ArchetypeInstances.Add(ChildClass->GetDefaultObject(false));
 				}
 			}
 		}
diff --git a/Engine/Source/Editor/SubobjectDataInterface/Private/SubobjectDataBlueprintFunctionLibrary.cpp b/Engine/Source/Editor/SubobjectDataInterface/Private/SubobjectDataBlueprintFunctionLibrary.cpp
--- a/Engine/Source/Editor/SubobjectDataInterface/Private/SubobjectDataBlueprintFunctionLibrary.cpp
+++ b/Engine/Source/Editor/SubobjectDataInterface/Private/SubobjectDataBlueprintFunctionLibrary.cpp
@@ -46,7 +46,7 @@ const UObject* USubobjectDataBlueprintFunctionLibrary::GetAssociatedObject(const
 	const UObject* Object = Data.GetObject();
 	if(const UBlueprint* BP = Cast<UBlueprint>(Object))
 	{
-		return BP->GeneratedClass ? BP->GeneratedClass->ClassDefaultObject : nullptr;
+		return BP->GeneratedClass ? BP->GeneratedClass->GetDefaultObject(false) : nullptr;
 	}
 	else
 	{
diff --git a/Engine/Source/Editor/UMGEditor/Private/WidgetBlueprintCompiler.cpp b/Engine/Source/Editor/UMGEditor/Private/WidgetBlueprintCompiler.cpp
--- a/Engine/Source/Editor/UMGEditor/Private/WidgetBlueprintCompiler.cpp
+++ b/Engine/Source/Editor/UMGEditor/Private/WidgetBlueprintCompiler.cpp
@@ -498,7 +498,7 @@ void FWidgetBlueprintCompilerContext::SaveSubObjectsFromCleanAndSanitizeClass(FS
 	// otherwise they all get renamed, and it causes early loading errors.
 	SubObjectsToSave.AddObject(WidgetBP->WidgetTree);
 
-	if (UUserWidget* ClassDefaultWidgetToClean = Cast<UUserWidget>(ClassToClean->ClassDefaultObject))
+	if (UUserWidget* ClassDefaultWidgetToClean = Cast<UUserWidget>(ClassToClean->GetDefaultObject(false)))
 	{
 		// We need preserve any named slots that have been slotted into the CDO.  This can happen when someone subclasses
 		// from a widget with named slots.  Those named slots are exposed to the child classes widget tree as
diff --git a/Engine/Source/Editor/UnrealEd/Private/DiffUtils.cpp b/Engine/Source/Editor/UnrealEd/Private/DiffUtils.cpp
--- a/Engine/Source/Editor/UnrealEd/Private/DiffUtils.cpp
+++ b/Engine/Source/Editor/UnrealEd/Private/DiffUtils.cpp
@@ -648,7 +648,7 @@ const UObject* DiffUtils::GetCDO(const UBlueprint* ForBlueprint)
 		return NULL;
 	}
 
-	return ForBlueprint->GeneratedClass->ClassDefaultObject;
+	return ForBlueprint->GeneratedClass->GetDefaultObject(false);
 }
 
 void DiffUtils::CompareUnrelatedStructs(const UStruct* StructA, const void* A, const UStruct* StructB, const void* B,
diff --git a/Engine/Source/Editor/UnrealEd/Private/Kismet2/BlueprintEditorUtils.cpp b/Engine/Source/Editor/UnrealEd/Private/Kismet2/BlueprintEditorUtils.cpp
--- a/Engine/Source/Editor/UnrealEd/Private/Kismet2/BlueprintEditorUtils.cpp
+++ b/Engine/Source/Editor/UnrealEd/Private/Kismet2/BlueprintEditorUtils.cpp
@@ -1154,7 +1154,7 @@ struct FRegenerationHelper
 				UBlueprint::ForceLoadMembers(InterfaceBlueprint);
 				if (InterfaceBlueprint->HasAnyFlags(RF_BeingRegenerated))
 				{
-					InterfaceBlueprint->RegenerateClass(InterfaceClass, InterfaceClass->ClassDefaultObject);
+					InterfaceBlueprint->RegenerateClass(InterfaceClass, InterfaceClass->GetDefaultObject(false));
 				}
 			}
 		}
@@ -4561,7 +4561,7 @@ void FBlueprintEditorUtils::ValidateEditorOnlyNodes(const UK2Node* Node, FCompil
 
 	const UBlueprint* BP = Node->GetBlueprint();
 	const UClass* NodeClass = Node->GetClass();
-	const UPackage* NodeCDOPackage = NodeClass->ClassDefaultObject ? NodeClass->ClassDefaultObject->GetOutermost() : nullptr;
+	const UPackage* NodeCDOPackage = NodeClass->GetDefaultObject(false) ? NodeClass->GetDefaultObject(false)->GetOutermost() : nullptr;
 	
 	if(NodeCDOPackage && BP)
 	{
diff --git a/Engine/Source/Editor/UnrealEd/Private/Kismet2/Kismet2.cpp b/Engine/Source/Editor/UnrealEd/Private/Kismet2/Kismet2.cpp
--- a/Engine/Source/Editor/UnrealEd/Private/Kismet2/Kismet2.cpp
+++ b/Engine/Source/Editor/UnrealEd/Private/Kismet2/Kismet2.cpp
@@ -982,7 +982,7 @@ static UObject* ConformComponentsUtils::FindNativeArchetype(const UObject* Nativ
 
 	if (!NativeSubobject)
 	{
-		return Component->GetClass()->ClassDefaultObject;
+		return Component->GetClass()->GetDefaultObject(false);
 	}
 	else if (!Component->IsA(NativeSubobject->GetClass()))
 	{
@@ -1017,7 +1017,7 @@ void FKismetEditorUtilities::ConformBlueprintFlagsAndComponents(UBlueprint* Blue
 		GenClass->ClassFlags |= (ParentClass->ClassFlags & CLASS_ScriptInherit);
 		GenClass->ClassConfigName = ParentClass->ClassConfigName;
 		GenClass->ClassWithin = ParentClass->ClassWithin;
-		if (UObject* GenCDO = GenClass->ClassDefaultObject)
+		if (UObject* GenCDO = GenClass->GetDefaultObject(false))
 		{
 			ConformComponentsUtils::ConformRemovedNativeComponents(GenCDO);
 			GenCDO->InstanceSubobjectTemplates();
diff --git a/Engine/Source/Editor/UnrealEd/Private/Kismet2/KismetDebugUtilities.cpp b/Engine/Source/Editor/UnrealEd/Private/Kismet2/KismetDebugUtilities.cpp
--- a/Engine/Source/Editor/UnrealEd/Private/Kismet2/KismetDebugUtilities.cpp
+++ b/Engine/Source/Editor/UnrealEd/Private/Kismet2/KismetDebugUtilities.cpp
@@ -258,7 +258,7 @@ void FKismetDebugUtilities::OnScriptException(const UObject* ActiveObject, const
 		if (UClass* GeneratedClass = BlueprintObj->GeneratedClass; 
 			ObjectBeingDebugged == nullptr && BPTYPE_FunctionLibrary == BlueprintObj->BlueprintType && GeneratedClass)
 		{
-			ObjectBeingDebugged = GeneratedClass->ClassDefaultObject;
+			ObjectBeingDebugged = GeneratedClass->GetDefaultObject(false);
 		}
 
 		auto IsAPreviewOrInactiveObject = [](const UObject* InObject)
diff --git a/Engine/Source/Editor/UnrealEd/Private/Kismet2/KismetReinstanceUtilities.cpp b/Engine/Source/Editor/UnrealEd/Private/Kismet2/KismetReinstanceUtilities.cpp
--- a/Engine/Source/Editor/UnrealEd/Private/Kismet2/KismetReinstanceUtilities.cpp
+++ b/Engine/Source/Editor/UnrealEd/Private/Kismet2/KismetReinstanceUtilities.cpp
@@ -496,7 +496,7 @@ FBlueprintCompileReinstancer::FBlueprintCompileReinstancer(UClass* InClassToRein
 	, bIsRootReinstancer(false)
 	, bAllowResaveAtTheEndIfRequested(false)
 {
-	if( InClassToReinstance != nullptr && InClassToReinstance->ClassDefaultObject )
+	if( InClassToReinstance != nullptr && InClassToReinstance->GetDefaultObject(false) )
 	{
 		bool bAutoInferSaveOnCompile = !!(Flags & EBlueprintCompileReinstancerFlags::AutoInferSaveOnCompile);
 		bool bIsBytecodeOnly = !!(Flags & EBlueprintCompileReinstancerFlags::BytecodeOnly);
@@ -526,11 +526,11 @@ FBlueprintCompileReinstancer::FBlueprintCompileReinstancer(UClass* InClassToRein
 
 		if(!bAvoidCDODuplication)
 		{
-			ensure( ClassToReinstance->ClassDefaultObject->GetClass() == DuplicatedClass );
-			ClassToReinstance->ClassDefaultObject->Rename(nullptr, GetTransientPackage(), REN_DoNotDirty | REN_DontCreateRedirectors | REN_AllowPackageLinkerMismatch);
+			ensure( ClassToReinstance->GetDefaultObject(false)->GetClass() == DuplicatedClass );
+			ClassToReinstance->GetDefaultObject(false)->Rename(nullptr, GetTransientPackage(), REN_DoNotDirty | REN_DontCreateRedirectors | REN_AllowPackageLinkerMismatch);
 		}
 		
-		// Note that we can't clear ClassToReinstance->ClassDefaultObject even though
+		// Note that we can't clear ClassToReinstance->GetDefaultObject(false) even though
 		// we have moved it aside CleanAndSanitizeClass will want to grab the old CDO 
 		// so it can propagate values to the new one note that until that happens we are 
 		// in an extraordinary state: this class has a CDO of a different type
@@ -2006,7 +2006,7 @@ void FBlueprintCompileReinstancer::MoveDependentSkelToReinst(UClass* const Owner
 
 		GIsDuplicatingClassForReinstancing = true;
 		// Create a REINST version of the given class
-		UObject* OldCDO = CurClass->ClassDefaultObject;
+		UObject* OldCDO = CurClass->GetDefaultObject(false);
 		const FName ReinstanceName = MakeUniqueObjectName(GetTransientPackage(), CurClass->GetClass(), *(FString(TEXT("REINST_")) + *CurClass->GetName()));
 
 		if (!IsValid(CurClass) || CurClass->HasAnyClassFlags(CLASS_NewerVersionExists))
@@ -2048,9 +2048,9 @@ void FBlueprintCompileReinstancer::MoveDependentSkelToReinst(UClass* const Owner
 		// Actually move the old CDO reference out of the way
 		if (OldCDO)
 		{
-			CurClass->ClassDefaultObject = nullptr;
+			CurClass->SetDefaultObject(nullptr);
 			OldCDO->Rename(nullptr, ReinstClass->GetOuter(), REN_DoNotDirty | REN_DontCreateRedirectors | REN_NonTransactional);
-			ReinstClass->ClassDefaultObject = OldCDO;
+			ReinstClass->SetDefaultObject(OldCDO);
 			OldCDO->SetClass(ReinstClass);
 		}
 
@@ -2087,7 +2087,7 @@ UClass* FBlueprintCompileReinstancer::MoveCDOToNewClass(UClass* OwnerClass, cons
 		}
 	}
 
-	UObject* OldCDO = OwnerClass->ClassDefaultObject;
+	UObject* OldCDO = OwnerClass->GetDefaultObject(false);
 	const FName ReinstanceName = MakeUniqueObjectName(GetTransientPackage(), OwnerClass->GetClass(), *(FString(TEXT("REINST_")) + *OwnerClass->GetName()));
 
 	checkf(IsValid(OwnerClass), TEXT("%s is invalid - will not duplicate successfully"), *(OwnerClass->GetName()));
@@ -2151,9 +2151,9 @@ UClass* FBlueprintCompileReinstancer::MoveCDOToNewClass(UClass* OwnerClass, cons
 		// @todo: #dano, rename bAvoidCDODuplication because it's really a flag to move the CDO aside not 'prevent duplication':
 		if(bAvoidCDODuplication)
 		{
-			OwnerClass->ClassDefaultObject = nullptr;
+			OwnerClass->SetDefaultObject(nullptr);
 			OldCDO->Rename(nullptr, CopyOfOwnerClass->GetOuter(), REN_DoNotDirty | REN_DontCreateRedirectors | REN_AllowPackageLinkerMismatch | REN_NonTransactional);
-			CopyOfOwnerClass->ClassDefaultObject = OldCDO;
+			CopyOfOwnerClass->SetDefaultObject(OldCDO);
 		}
 		OldCDO->SetClass(CopyOfOwnerClass);
 	}
@@ -3431,7 +3431,7 @@ FBlake3Hash FBlueprintCompileReinstancer::CalculateConstructionHashIfNecessary(c
 
 FBlake3Hash FBlueprintCompileReinstancer::CalculateConstructionHash(const UClass* ForClass)
 {
-	check(ForClass && ForClass->ClassDefaultObject);
+	check(ForClass && ForClass->GetDefaultObject(false));
 
 	// must capture all template data - just grab the subobjects of the class and exclude the ufunctions.
 	// UFunctions should be excluded as on load of an editor asset they will not include bytecode, but
@@ -3445,7 +3445,7 @@ FBlake3Hash FBlueprintCompileReinstancer::CalculateConstructionHash(const UClass
 		return Obj->IsA<UFunction>();
 	}));
 	Objects.Sort([](const UObject& A, const UObject& B) { return A.GetFName().LexicalLess(B.GetFName()); });
-	Objects.Add(ForClass->ClassDefaultObject);
+	Objects.Add(ForClass->GetDefaultObject(false));
 	
 	FUtf8String TemplateData = UE::JsonObjectGraph::Stringify(
 		Objects,
diff --git a/Engine/Source/Editor/UnrealEd/Private/Kismet2/ReloadUtilities.cpp b/Engine/Source/Editor/UnrealEd/Private/Kismet2/ReloadUtilities.cpp
--- a/Engine/Source/Editor/UnrealEd/Private/Kismet2/ReloadUtilities.cpp
+++ b/Engine/Source/Editor/UnrealEd/Private/Kismet2/ReloadUtilities.cpp
@@ -679,7 +679,7 @@ namespace UE::Reload::Private
 					REN_DoNotDirty | REN_DontCreateRedirectors | REN_NonTransactional | REN_SkipGeneratedClasses);
 
 				// Clear the class default object so it gets recreated.
-				Pair.Value.OldClass->ClassDefaultObject = nullptr;
+				Pair.Value.OldClass->SetDefaultObject(nullptr);
 			}
 		}
 
diff --git a/Engine/Source/Editor/UnrealEd/Private/ObjectTools.cpp b/Engine/Source/Editor/UnrealEd/Private/ObjectTools.cpp
--- a/Engine/Source/Editor/UnrealEd/Private/ObjectTools.cpp
+++ b/Engine/Source/Editor/UnrealEd/Private/ObjectTools.cpp
@@ -1392,7 +1392,7 @@ namespace ObjectTools
 					if (BPObjectToUpdate)
 					{
 						FArchiveReplaceObjectAndStructPropertyRef<UObject> ReplaceInBPClassObject_Ar(BPObjectToUpdate->GeneratedClass, ReplacementMap, EArchiveReplaceObjectFlags::IncludeClassGeneratedByRef);
-						FArchiveReplaceObjectAndStructPropertyRef<UObject> ReplaceInBPClassDefaultObject_Ar(BPObjectToUpdate->GeneratedClass->ClassDefaultObject, ReplacementMap, EArchiveReplaceObjectFlags::IncludeClassGeneratedByRef);
+						FArchiveReplaceObjectAndStructPropertyRef<UObject> ReplaceInBPClassDefaultObject_Ar(BPObjectToUpdate->GeneratedClass->GetDefaultObject(false), ReplacementMap, EArchiveReplaceObjectFlags::IncludeClassGeneratedByRef);
 					}
 
 					UDataTable* DataTableObjectToUpdate = Cast<UDataTable>(CurObject);
@@ -3630,10 +3630,10 @@ namespace ObjectTools
 
 			// If the object about to be deleted is a Blueprint asset, make sure that any instances of the Blueprint class get deleted as well
 			UBlueprint* BlueprintObject = Cast<UBlueprint>(CurrentObject);
-			if ( BlueprintObject && BlueprintObject->GeneratedClass && BlueprintObject->GeneratedClass->ClassDefaultObject )
+			if ( BlueprintObject && BlueprintObject->GeneratedClass && BlueprintObject->GeneratedClass->GetDefaultObject(false) )
 			{
 				TArray<UObject*> InstancesToDelete;
-				BlueprintObject->GeneratedClass->ClassDefaultObject->GetArchetypeInstances( InstancesToDelete );
+				BlueprintObject->GeneratedClass->GetDefaultObject(false)->GetArchetypeInstances( InstancesToDelete );
 
 				for ( TArray<UObject*>::TConstIterator InstanceItr( InstancesToDelete ); InstanceItr; ++InstanceItr )
 				{
@@ -5430,7 +5430,7 @@ namespace ThumbnailTools
 			{
 				if (Blueprint->GeneratedClass != nullptr)
 				{
-					ObjectToRender = Blueprint->GeneratedClass->ClassDefaultObject;
+					ObjectToRender = Blueprint->GeneratedClass->GetDefaultObject(false);
 				}
 			}
 		}
diff --git a/Engine/Source/Editor/UnrealEd/Private/ThumbnailManager.cpp b/Engine/Source/Editor/UnrealEd/Private/ThumbnailManager.cpp
--- a/Engine/Source/Editor/UnrealEd/Private/ThumbnailManager.cpp
+++ b/Engine/Source/Editor/UnrealEd/Private/ThumbnailManager.cpp
@@ -183,7 +183,7 @@ FThumbnailRenderingInfo* UThumbnailManager::GetRenderingInfo(UObject* Object)
 		{
 			if (Blueprint->GeneratedClass != nullptr)
 			{
-				UObject* ClassCDO = Blueprint->GeneratedClass->ClassDefaultObject;
+				UObject* ClassCDO = Blueprint->GeneratedClass->GetDefaultObject(false);
 				RenderInfo = FindThumbnailRendererForClass(ClassCDO, Blueprint->GeneratedClass);
 				if (RenderInfo != nullptr)
 				{
diff --git a/Engine/Source/Runtime/CoreUObject/Private/Blueprint/BlueprintSupport.cpp b/Engine/Source/Runtime/CoreUObject/Private/Blueprint/BlueprintSupport.cpp
--- a/Engine/Source/Runtime/CoreUObject/Private/Blueprint/BlueprintSupport.cpp
+++ b/Engine/Source/Runtime/CoreUObject/Private/Blueprint/BlueprintSupport.cpp
@@ -809,8 +809,8 @@ bool FLinkerLoad::RegenerateBlueprintClass(UClass* LoadClass, UObject* ClassDefa
 	// manually... do this before we Preload() any children members so that if 
 	// one of those preloads subsequently ends up back here for this class, 
 	// then the ExportObject is carried along and used in the eventual RegenerateClass() call
-	check(!bAlreadyRegenerating || (LoadClass->ClassDefaultObject == ClassDefaultObject));
-	LoadClass->ClassDefaultObject = ClassDefaultObject;
+	check(!bAlreadyRegenerating || (LoadClass->GetDefaultObject(false) == ClassDefaultObject));
+	LoadClass->SetDefaultObject(ClassDefaultObject);
 
 	// Finish loading the class here, so we have all the appropriate data to copy over to the new CDO
 	TArray<UObject*> AllChildMembers;
@@ -865,7 +865,7 @@ bool FLinkerLoad::RegenerateBlueprintClass(UClass* LoadClass, UObject* ClassDefa
 				FPreloadMembersHelper::PreloadMembers(SuperClassSourceObject);
 				// recurse into this function for this parent class; 
 				// 'ClassDefaultObject' should be the class's original ExportObject
-				RegenerateBlueprintClass(SuperClass, SuperClass->ClassDefaultObject);
+				RegenerateBlueprintClass(SuperClass, SuperClass->GetDefaultObject(false));
 			}
 		}
 
@@ -990,7 +990,7 @@ static UObject* FindExistingImportObject(const int32 Index, const TArray<FObject
 				// emulate/keep the contract that that block provides, we'll only 
 				// lookup the object if its class is loaded, and has a CDO (this
 				// is just to mitigate risk from this change)
-				if (!ImportClass->HasAnyFlags(RF_NeedLoad) && ImportClass->ClassDefaultObject)
+				if (!ImportClass->HasAnyFlags(RF_NeedLoad) && ImportClass->GetDefaultObject(false))
 				{
 					FoundObject = StaticFindObjectFast(ImportClass, FindOuter, Import.ObjectName);
 				}				
@@ -1407,9 +1407,10 @@ bool FLinkerLoad::DeferExportCreation(const int32 Index, UObject* Outer)
 	const bool bIsPlaceholderClass = (AsPlaceholderClass != nullptr);
 
 	FLinkerLoad* ClassLinker = LoadClass->GetLinker();
+	const UObject* ClassDefaultObject = LoadClass->GetDefaultObject(false);
 	if (!bIsPlaceholderClass
 		&& ((ClassLinker == nullptr) || !ClassLinker->IsBlueprintFinalizationPending())
-		&& (!LoadClass->ClassDefaultObject || LoadClass->ClassDefaultObject->HasAnyFlags(RF_LoadCompleted) || !LoadClass->ClassDefaultObject->HasAnyFlags(RF_WasLoaded)))
+		&& (!ClassDefaultObject || ClassDefaultObject->HasAnyFlags(RF_LoadCompleted) || !ClassDefaultObject->HasAnyFlags(RF_WasLoaded)))
 	{
 		return false;
 	}
@@ -1455,8 +1456,8 @@ bool FLinkerLoad::DeferExportCreation(const int32 Index, UObject* Outer)
 
 	if (bIsCDOExport)
 	{
-		DEFERRED_DEPENDENCY_CHECK(LoadClass->ClassDefaultObject == nullptr);
-		LoadClass->ClassDefaultObject = Placeholder;
+		DEFERRED_DEPENDENCY_CHECK(LoadClass->GetDefaultObject(false) == nullptr);
+		LoadClass->SetDefaultObject(Placeholder);
 	}
 #endif // USE_CIRCULAR_DEPENDENCY_LOAD_DEFERRING
 
@@ -1961,7 +1962,7 @@ void FLinkerLoad::FinalizeBlueprint(UClass* LoadClass)
 		if ((SuperLinker != nullptr) && SuperLinker->IsBlueprintFinalizationPending())
 		{
 			DEFERRED_DEPENDENCY_CHECK(SuperLinker->DeferredCDOIndex != INDEX_NONE || SuperLinker->bForceBlueprintFinalization);
-			UObject* SuperCDO = SuperLinker->DeferredCDOIndex != INDEX_NONE ? ToRawPtr(SuperLinker->ExportMap[SuperLinker->DeferredCDOIndex].Object) : ToRawPtr(SuperClass->ClassDefaultObject);
+			UObject* SuperCDO = SuperLinker->DeferredCDOIndex != INDEX_NONE ? ToRawPtr(SuperLinker->ExportMap[SuperLinker->DeferredCDOIndex].Object) : SuperClass->GetDefaultObject(false);
 			// we MUST have the super fully serialized before we can finalize  
 			// this (class and CDO); if the SuperCDO is already in the midst of 
 			// serializing somewhere up the stack (and a cyclic dependency has  
@@ -2048,7 +2049,7 @@ void FLinkerLoad::FinalizeBlueprint(UClass* LoadClass)
 	if (IsBlueprintFinalizationPending())
 	{
 		int32 DeferredCDOIndexCopy = DeferredCDOIndex;
-		UObject* CDO = DeferredCDOIndex != INDEX_NONE ? ToRawPtr(ExportMap[DeferredCDOIndexCopy].Object) : ToRawPtr(LoadClass->ClassDefaultObject);
+		UObject* CDO = DeferredCDOIndex != INDEX_NONE ? ToRawPtr(ExportMap[DeferredCDOIndexCopy].Object) : LoadClass->GetDefaultObject(false);
 		// clear this so IsBlueprintFinalizationPending() doesn't report true:
 		FLinkerLoad::bForceBlueprintFinalization = false;
 		// clear this because we're processing this CDO now:
@@ -2096,16 +2097,16 @@ void FLinkerLoad::FinalizeBlueprint(UClass* LoadClass)
 		if (!LoadClass->bCooked)
 #endif
 		{
-			UObject* OldCDO = LoadClass->ClassDefaultObject;
+			UObject* OldCDO = LoadClass->GetDefaultObject(false);
 			if (RegenerateBlueprintClass(LoadClass, CDO))
 			{
 				// emulate class CDO serialization (RegenerateBlueprintClass() could 
 				// have a side-effect where it overwrites the class's CDO; so we 
 				// want to make sure that we don't overwrite that new CDO with a 
 				// stale one)
-				if (OldCDO == LoadClass->ClassDefaultObject)
+				if (OldCDO == LoadClass->GetDefaultObject(false))
 				{
-					LoadClass->ClassDefaultObject = CDO;
+					LoadClass->SetDefaultObject(CDO);
 				}
 			}
 		}
@@ -2129,7 +2130,7 @@ void FLinkerLoad::ResolveDeferredExports(UClass* LoadClass)
 	{
 		if (ULinkerPlaceholderExportObject* PlaceholderExport = Cast<ULinkerPlaceholderExportObject>(ExportMap[DeferredCDOIndex].Object))
 		{
-			LoadClass->ClassDefaultObject = nullptr;
+			LoadClass->SetDefaultObject(nullptr);
 			PlaceholderExport->SetLinker(nullptr, INDEX_NONE);
 			ExportMap[DeferredCDOIndex].ResetObject();
 			UObject* ExportObj = CreateExport(DeferredCDOIndex);
@@ -2138,11 +2139,11 @@ void FLinkerLoad::ResolveDeferredExports(UClass* LoadClass)
 			ResolvedDeferredSubobjects(PlaceholderExport);
 			PlaceholderExport->MarkAsGarbage();
 
-			DEFERRED_DEPENDENCY_CHECK(LoadClass->ClassDefaultObject == ExportObj);
+			DEFERRED_DEPENDENCY_CHECK(LoadClass->GetDefaultObject(false) == ExportObj);
 		}
 	}
 
-	UObject* BlueprintCDO = DeferredCDOIndex != INDEX_NONE ? ToRawPtr(ExportMap[DeferredCDOIndex].Object) : ToRawPtr(LoadClass->ClassDefaultObject);
+	UObject* BlueprintCDO = DeferredCDOIndex != INDEX_NONE ? ToRawPtr(ExportMap[DeferredCDOIndex].Object) : LoadClass->GetDefaultObject(false);
 	DEFERRED_DEPENDENCY_CHECK(BlueprintCDO != nullptr);
 	
 	TArray<int32> DeferredTemplateObjects;
@@ -2847,10 +2848,10 @@ bool FDeferredCdoInitializationTracker::DeferPreload(UObject* Object)
 		{
 			UClass* OwningClass = Cast<UClass>(Object->GetOuter());
 
-			DEFERRED_DEPENDENCY_CHECK(OwningClass && OwningClass->ClassDefaultObject);
+			DEFERRED_DEPENDENCY_CHECK(OwningClass && OwningClass->GetDefaultObject(false));
 			if (OwningClass)
 			{
-				bDeferPostload = ShouldDeferSubObjectPreload(OwningClass->ClassDefaultObject);
+				bDeferPostload = ShouldDeferSubObjectPreload(OwningClass->GetDefaultObject(false));
 			}
 		}
 	}
@@ -2982,7 +2983,7 @@ FObjectInitializer* FDeferredObjInitializationHelper::DeferObjectInitializerIfNe
 				// component's CDO is used, which is probably wrong)
 				// 
 				// So if the super CDO isn't ready, we need to defer this sub-object
-				const UObject* SuperCDO = SuperClass->ClassDefaultObject;
+				const UObject* SuperCDO = SuperClass->GetDefaultObject(false);
 				FDeferredCdoInitializationTracker& CdoInitDeferalSys = FDeferredCdoInitializationTracker::Get();
 				if (!IsSuperCdoReadyToBeCopied(CdoInitDeferalSys, OwnerClass, SuperCDO))
 				{
diff --git a/Engine/Source/Runtime/CoreUObject/Private/UObject/Class.cpp b/Engine/Source/Runtime/CoreUObject/Private/UObject/Class.cpp
--- a/Engine/Source/Runtime/CoreUObject/Private/UObject/Class.cpp
+++ b/Engine/Source/Runtime/CoreUObject/Private/UObject/Class.cpp
@@ -4784,6 +4784,7 @@ void UClass::GetDefaultObjectSubobjects(TArray<UObject*>& OutDefaultSubobjects)
  */
 void UClass::AddReferencedObjects(UObject* InThis, FReferenceCollector& Collector)
 {
+	PRAGMA_DISABLE_DEPRECATION_WARNINGS
 	UClass* This = CastChecked<UClass>(InThis);
 	for (FImplementedInterface& Inter : This->Interfaces)
 	{
@@ -4820,6 +4821,7 @@ void UClass::AddReferencedObjects(UObject* InThis, FReferenceCollector& Collecto
 	}
 	
 	Super::AddReferencedObjects( This, Collector );
+	PRAGMA_ENABLE_DEPRECATION_WARNINGS
 }
 
 /**
@@ -4879,7 +4881,9 @@ public:
 #if WITH_EDITORONLY_DATA
 		Target->ClassGeneratedBy = GeneratedBy;
 #endif
+PRAGMA_DISABLE_DEPRECATION_WARNINGS
 		Target->ClassDefaultObject = DefaultObject;
+PRAGMA_ENABLE_DEPRECATION_WARNINGS
 		Target->ClassFlags |= Flags;
 		Target->ClassCastFlags |= CastFlags;
 		Target->ClassConstructor = Constructor;
@@ -4906,6 +4910,7 @@ FRestoreForUObjectOverwrite* UClass::GetRestoreForUObjectOverwrite()
 **/
 UObject* UClass::CreateDefaultObject()
 {
+PRAGMA_DISABLE_DEPRECATION_WARNINGS
 	if ( ClassDefaultObject == NULL )
 	{
 		ensureMsgf(!bLayoutChanging, TEXT("Class named %s creating its CDO while changing its layout"), *GetName());
@@ -4986,6 +4991,7 @@ UObject* UClass::CreateDefaultObject()
 		}
 	}
 	return ClassDefaultObject;
+PRAGMA_ENABLE_DEPRECATION_WARNINGS
 }
 
 class FFeedbackContextImportDefaults final : public FFeedbackContext
@@ -5053,11 +5059,13 @@ bool UClass::Rename( const TCHAR* InName, UObject* NewOuter, ERenameFlags Flags
 {
 	bool bSuccess = Super::Rename( InName, NewOuter, Flags );
 
+PRAGMA_DISABLE_DEPRECATION_WARNINGS
 	// If we have a default object, rename that to the same package as the class, and rename so it still matches the class name (Default__ClassName)
 	if(bSuccess && (ClassDefaultObject != NULL))
 	{
 		ClassDefaultObject->Rename(*GetDefaultObjectName().ToString(), NewOuter, Flags);
 	}
+PRAGMA_ENABLE_DEPRECATION_WARNINGS
 
 	// Now actually rename the class
 	return bSuccess;
@@ -5066,12 +5074,13 @@ bool UClass::Rename( const TCHAR* InName, UObject* NewOuter, ERenameFlags Flags
 void UClass::TagSubobjects(EObjectFlags NewFlags)
 {
 	Super::TagSubobjects(NewFlags);
-
+PRAGMA_DISABLE_DEPRECATION_WARNINGS
 	if (ClassDefaultObject && !ClassDefaultObject->HasAnyFlags(GARBAGE_COLLECTION_KEEPFLAGS) && !ClassDefaultObject->IsRooted())
 	{
 		ClassDefaultObject->SetFlags(NewFlags);
 		ClassDefaultObject->TagSubobjects(NewFlags);
 	}
+PRAGMA_ENABLE_DEPRECATION_WARNINGS
 }
 
 /**
@@ -5195,7 +5204,9 @@ void UClass::FinishDestroy()
 	NetFields.Empty();
 	ClassReps.Empty();
 
+PRAGMA_DISABLE_DEPRECATION_WARNINGS
 	ClassDefaultObject = nullptr;
+PRAGMA_ENABLE_DEPRECATION_WARNINGS
 
 	CleanupSparseClassData();
 	SparseClassDataStruct = nullptr;
@@ -5639,7 +5650,9 @@ void UClass::SetSuperStruct(UStruct* NewSuperStruct)
 
 bool UClass::IsStructTrashed() const
 {
+PRAGMA_DISABLE_DEPRECATION_WARNINGS
 	return Children == nullptr && ChildProperties == nullptr && ClassDefaultObject == nullptr;
+PRAGMA_ENABLE_DEPRECATION_WARNINGS
 }
 
 #if WITH_EDITORONLY_DATA
@@ -5782,6 +5795,8 @@ void UClass::SerializePropertyValueFlags(EPropertyValueFlags Flags, void* Data,
 
 void UClass::Serialize( FArchive& Ar )
 {
+	PRAGMA_DISABLE_DEPRECATION_WARNINGS
+
 	if ( Ar.IsLoading() || (Ar.IsModifyingWeakAndStrongReferences() && !Ar.IsSaving()) )
 	{
 		// Rehash since SuperStruct will be serialized in UStruct::Serialize
@@ -6053,6 +6068,7 @@ void UClass::Serialize( FArchive& Ar )
 			}
 		}
 	}
+	PRAGMA_ENABLE_DEPRECATION_WARNINGS
 }
 
 bool UClass::ImplementsInterface( const class UClass* SomeInterface ) const
@@ -6218,7 +6234,9 @@ void UClass::PurgeClass(bool bRecompilingOnLoad)
 	}
 #endif
 
+PRAGMA_DISABLE_DEPRECATION_WARNINGS
 	ClassDefaultObject = nullptr;
+PRAGMA_ENABLE_DEPRECATION_WARNINGS
 
 	Interfaces.Empty();
 	NativeFunctionLookupTable.Empty();
@@ -6353,7 +6371,9 @@ UClass::UClass(const FObjectInitializer& ObjectInitializer)
 ,	ClassGeneratedBy(nullptr)
 ,	PropertiesPendingDestruction(nullptr)
 #endif
+PRAGMA_DISABLE_DEPRECATION_WARNINGS
 ,	ClassDefaultObject(nullptr)
+PRAGMA_ENABLE_DEPRECATION_WARNINGS
 ,	SparseClassData(nullptr)
 ,	SparseClassDataStruct(nullptr)
 {
@@ -6378,7 +6398,9 @@ UClass::UClass(const FObjectInitializer& ObjectInitializer, UClass* InBaseClass)
 ,	ClassGeneratedBy(nullptr)
 ,	PropertiesPendingDestruction(nullptr)
 #endif
+PRAGMA_DISABLE_DEPRECATION_WARNINGS
 ,	ClassDefaultObject(nullptr)
+PRAGMA_ENABLE_DEPRECATION_WARNINGS
 ,	SparseClassData(nullptr)
 ,	SparseClassDataStruct(nullptr)
 {
@@ -6439,7 +6461,9 @@ UClass::UClass
 #endif
 ,	ClassConfigName			()
 ,	NetFields				()
+PRAGMA_DISABLE_DEPRECATION_WARNINGS
 ,	ClassDefaultObject		( nullptr )
+PRAGMA_ENABLE_DEPRECATION_WARNINGS
 ,	SparseClassData			( nullptr )
 ,	SparseClassDataStruct	( nullptr )
 {
@@ -7233,11 +7257,12 @@ bool UClass::CanCreateInstanceDataObject() const
 IMPLEMENT_CORE_INTRINSIC_CLASS(UClass, UStruct,
 	{
 		Class->CppClassStaticFunctions = UOBJECT_CPPCLASS_STATICFUNCTIONS_FORCLASS(UClass);
-
+PRAGMA_DISABLE_DEPRECATION_WARNINGS
 		UE::GC::DeclareIntrinsicMembers(Class, {	UE_GC_MEMBER(UClass, ClassDefaultObject),
 													UE_GC_MEMBER(UClass, ClassWithin), 
 													UE_GC_MEMBER(UClass, NetFields),
 													UCLASS_EDITOR_GC_MEMBERS} );
+PRAGMA_ENABLE_DEPRECATION_WARNINGS
 	}
 );
 
diff --git a/Engine/Source/Runtime/CoreUObject/Private/UObject/InstanceDataObjectUtils.cpp b/Engine/Source/Runtime/CoreUObject/Private/UObject/InstanceDataObjectUtils.cpp
--- a/Engine/Source/Runtime/CoreUObject/Private/UObject/InstanceDataObjectUtils.cpp
+++ b/Engine/Source/Runtime/CoreUObject/Private/UObject/InstanceDataObjectUtils.cpp
@@ -422,7 +422,7 @@ namespace UE
 		{
 			if (const UObject* BlueprintGeneratedObject = GetBlueprintGeneratedObject(InObject))
 			{
-				const bool bIsArchetype = BlueprintGeneratedObject->GetClass()->ClassDefaultObject == BlueprintGeneratedObject;
+				const bool bIsArchetype = BlueprintGeneratedObject->GetClass()->GetDefaultObject(false) == BlueprintGeneratedObject;
 				if (!bEnableIDOsForBlueprintArchetypes && bIsArchetype)
 				{
 					return false;
diff --git a/Engine/Source/Runtime/CoreUObject/Private/UObject/LinkerLoad.cpp b/Engine/Source/Runtime/CoreUObject/Private/UObject/LinkerLoad.cpp
--- a/Engine/Source/Runtime/CoreUObject/Private/UObject/LinkerLoad.cpp
+++ b/Engine/Source/Runtime/CoreUObject/Private/UObject/LinkerLoad.cpp
@@ -5234,7 +5234,7 @@ UObject* FLinkerLoad::CreateExport( int32 Index )
 					&& Export.SuperIndex.IsImport())
 				{
 					const UClass* AsClass = dynamic_cast<UClass*>(SuperStruct);
-					if (AsClass && !AsClass->ClassDefaultObject)
+					if (AsClass && !AsClass->GetDefaultObject(false))
 					{
 						check(!GEventDrivenLoaderEnabled || !EVENT_DRIVEN_ASYNC_LOAD_ACTIVE_AT_RUNTIME);
 						SuperStruct->SetFlags(RF_NeedLoad);
diff --git a/Engine/Source/Runtime/CoreUObject/Private/UObject/ScriptCore.cpp b/Engine/Source/Runtime/CoreUObject/Private/UObject/ScriptCore.cpp
--- a/Engine/Source/Runtime/CoreUObject/Private/UObject/ScriptCore.cpp
+++ b/Engine/Source/Runtime/CoreUObject/Private/UObject/ScriptCore.cpp
@@ -1063,7 +1063,7 @@ DEFINE_FUNCTION(UObject::execCallMathFunction)
 	checkSlow(Function->FunctionFlags & FUNC_Native);
 	// ProcessContext is the arbiter of net callspace, so we can't call net functions using this instruction:
 	checkSlow(!Function->HasAnyFunctionFlags(FUNC_NetFuncFlags|FUNC_BlueprintAuthorityOnly|FUNC_BlueprintCosmetic|FUNC_NetRequest|FUNC_NetResponse));
-	UObject* NewContext = Function->GetOuterUClassUnchecked()->ClassDefaultObject;
+	UObject* NewContext = Function->GetOuterUClassUnchecked()->GetDefaultObject(false);
 	checkSlow(NewContext);
 	{
 #if PER_FUNCTION_SCRIPT_STATS
diff --git a/Engine/Source/Runtime/CoreUObject/Private/UObject/UObjectBase.cpp b/Engine/Source/Runtime/CoreUObject/Private/UObject/UObjectBase.cpp
--- a/Engine/Source/Runtime/CoreUObject/Private/UObject/UObjectBase.cpp
+++ b/Engine/Source/Runtime/CoreUObject/Private/UObject/UObjectBase.cpp
@@ -422,7 +422,7 @@ bool UObjectBase::IsValidLowLevelFast(bool bRecursive /*= true*/) const
     return false;
 	}
 	// These should all be non-NULL (except CDO-alignment check which should be 0)
-	if (Class == nullptr || Class->ClassDefaultObject == nullptr || ((UPTRINT)ObjectPtr_Private::Friend::NoAccessTrackingGet(Class->ClassDefaultObject) & AlignmentCheck) != 0)
+	if (Class == nullptr || !Class->CheckDefaultObjectIsValidLowLevelFast(AlignmentCheck))
 	{
 #if WITH_LIVE_CODING
 		// When live coding is re-instancing blueprint generated classes, we have to clear out the default object so it can get 
diff --git a/Engine/Source/Runtime/CoreUObject/Public/UObject/Class.h b/Engine/Source/Runtime/CoreUObject/Public/UObject/Class.h
--- a/Engine/Source/Runtime/CoreUObject/Public/UObject/Class.h
+++ b/Engine/Source/Runtime/CoreUObject/Public/UObject/Class.h
@@ -2957,6 +2957,7 @@ public:
 	friend class FRestoreClassInfo;
 	friend class FBlueprintEditorUtils;
 	friend class FBlueprintCompileReinstancer;
+	friend struct FBlueprintCompilationManagerImpl;
 
 	typedef void		(*ClassConstructorType)				(const FObjectInitializer&);
 	typedef UObject*	(*ClassVTableHelperCtorCallerType)	(FVTableHelper& Helper);
@@ -3068,6 +3069,7 @@ public:
 #endif
 
 	/** The class default object; used for delta serialization and object initialization */
+	UE_DEPRECATED(5.6, "ClassDefaultObject will be made private in the next release. For general purpose access to the CDO consider using the global GetDefault<>() or GetMutableDefault<>() functions.")
 	TObjectPtr<UObject> ClassDefaultObject;
 
 protected:
@@ -3460,7 +3462,9 @@ public:
 	/** Returns amount of memory used by default object */
 	int32 GetDefaultsCount()
 	{
+		PRAGMA_DISABLE_DEPRECATION_WARNINGS
 		return ClassDefaultObject != nullptr ? GetPropertiesSize() : 0;
+		PRAGMA_ENABLE_DEPRECATION_WARNINGS
 	}
 
 	/**
@@ -3470,12 +3474,37 @@ public:
 	 */
 	UObject* GetDefaultObject(bool bCreateIfNeeded = true) const
 	{
+		PRAGMA_DISABLE_DEPRECATION_WARNINGS
 		if (ClassDefaultObject == nullptr && bCreateIfNeeded)
 		{
 			InternalCreateDefaultObjectWrapper();
 		}
 
 		return ClassDefaultObject;
+		PRAGMA_ENABLE_DEPRECATION_WARNINGS
+	}
+
+	/**
+	 * Sets the default object of the class
+	 * @param	InClassDefaultObject new class default object or null
+	 */
+	void SetDefaultObject(UObject* InClassDefaultObject)
+	{
+		PRAGMA_DISABLE_DEPRECATION_WARNINGS
+		ClassDefaultObject = InClassDefaultObject;
+		PRAGMA_ENABLE_DEPRECATION_WARNINGS
+	}
+
+	/**
+	 * Checks if the class default object is valid (not null and its address is a properly aligned pointer)
+	 * @param	AlignmentCheck address alignment to check
+	 * @return	true if the CDO is valid
+	 */
+	FORCEINLINE bool CheckDefaultObjectIsValidLowLevelFast(const int32 AlignmentCheck) const
+	{
+		PRAGMA_DISABLE_DEPRECATION_WARNINGS
+		return ClassDefaultObject != nullptr && ((UPTRINT)ObjectPtr_Private::Friend::NoAccessTrackingGet(ClassDefaultObject) & AlignmentCheck) == 0;
+		PRAGMA_ENABLE_DEPRECATION_WARNINGS
 	}
 
 	/**
diff --git a/Engine/Source/Runtime/Engine/Private/ActorConstruction.cpp b/Engine/Source/Runtime/Engine/Private/ActorConstruction.cpp
--- a/Engine/Source/Runtime/Engine/Private/ActorConstruction.cpp
+++ b/Engine/Source/Runtime/Engine/Private/ActorConstruction.cpp
@@ -188,7 +188,7 @@ void AActor::DestroyConstructedComponents()
 				if (Component->CreationMethod == EComponentCreationMethod::Native && Component->HasAnyFlags(RF_DefaultSubObject))
 				{
 					UObject* ComponentArchetype = Component->GetArchetype();
-					if (ComponentArchetype == ComponentArchetype->GetClass()->ClassDefaultObject)
+					if (ComponentArchetype == ComponentArchetype->GetClass()->GetDefaultObject(false))
 					{
 						bDestroyComponent = true;
 					}
diff --git a/Engine/Source/Runtime/Engine/Private/Blueprint.cpp b/Engine/Source/Runtime/Engine/Private/Blueprint.cpp
--- a/Engine/Source/Runtime/Engine/Private/Blueprint.cpp
+++ b/Engine/Source/Runtime/Engine/Private/Blueprint.cpp
@@ -60,7 +60,7 @@ void UBlueprint::ConformNativeComponents()
 {
 	if (UClass* const BlueprintClass = GeneratedClass)
 	{
-		if (AActor* BlueprintCDO = Cast<AActor>(BlueprintClass->ClassDefaultObject))
+		if (AActor* BlueprintCDO = Cast<AActor>(BlueprintClass->GetDefaultObject(false)))
 		{
 			TInlineComponentArray<UActorComponent*> OldNativeComponents;
 			// collect the native components that this blueprint was serialized out 
@@ -68,7 +68,7 @@ void UBlueprint::ConformNativeComponents()
 			BlueprintCDO->GetComponents(OldNativeComponents);
 
 			UClass* const NativeSuperClass = FBlueprintEditorUtils::FindFirstNativeClass(BlueprintClass);
-			AActor* NativeCDO = CastChecked<AActor>(NativeSuperClass->ClassDefaultObject);
+			AActor* NativeCDO = CastChecked<AActor>(NativeSuperClass->GetDefaultObject(false));
 			// collect the more up to date native components (directly from the 
 			// native super-class)
 			TInlineComponentArray<UActorComponent*> NewNativeComponents;
@@ -637,10 +637,10 @@ UClass* UBlueprint::RegenerateClass(UClass* ClassToRegenerate, UObject* Previous
 
 	UClass* GeneratedClassResolved = GeneratedClass;
 
-	if (ensure(GeneratedClassResolved->ClassDefaultObject ))
+	if (ensure(GeneratedClassResolved->GetDefaultObject(false)))
 	{
 		UBlueprint::ForceLoadMembers(GeneratedClassResolved);
-		UBlueprint::ForceLoadMembers(GeneratedClassResolved->ClassDefaultObject);
+		UBlueprint::ForceLoadMembers(GeneratedClassResolved->GetDefaultObject(false));
 	}
 	UBlueprint::ForceLoadMembers(this);
 
@@ -1165,7 +1165,7 @@ FPrimaryAssetId UBlueprint::GetPrimaryAssetId() const
 {
 	// Forward to our Class, which will forward to CDO if needed
 	// We use Generated instead of Skeleton because the CDO data is more accurate on Generated
-	if (GeneratedClass && GeneratedClass->ClassDefaultObject)
+	if (GeneratedClass && GeneratedClass->GetDefaultObject(false))
 	{
 		return GeneratedClass->GetPrimaryAssetId();
 	}
diff --git a/Engine/Source/Runtime/Engine/Private/BlueprintGeneratedClass.cpp b/Engine/Source/Runtime/Engine/Private/BlueprintGeneratedClass.cpp
--- a/Engine/Source/Runtime/Engine/Private/BlueprintGeneratedClass.cpp
+++ b/Engine/Source/Runtime/Engine/Private/BlueprintGeneratedClass.cpp
@@ -436,7 +436,7 @@ void UBlueprintGeneratedClass::PostLoad()
 			return;
 		}
 
-		UObject* ClassCDO = ClassDefaultObject;
+		UObject* ClassCDO = GetDefaultObject(false);
 
 		// Go through the CDO of the class, and make sure we don't have any legacy components that aren't instanced hanging on.
 		struct FCheckIfComponentChildHelper
@@ -513,10 +513,10 @@ void UBlueprintGeneratedClass::InitializeFieldNotifies()
 {
 	//Initialize the interface with the computed FieldNotifies
 	FieldNotifiesStartBitNumber = 0;
-	if (FieldNotifies.Num() && ImplementsInterface(UNotifyFieldValueChanged::StaticClass()) && ensure(ClassDefaultObject))
+	if (FieldNotifies.Num() && ImplementsInterface(UNotifyFieldValueChanged::StaticClass()) && ensure(GetDefaultObject(false)))
 	{
 		int32 NumberOfField = 0;
-		TScriptInterface<INotifyFieldValueChanged>(ClassDefaultObject)->GetFieldNotificationDescriptor().ForEachField(this, [&NumberOfField](::UE::FieldNotification::FFieldId FielId)
+		TScriptInterface<INotifyFieldValueChanged>(GetDefaultObject(false))->GetFieldNotificationDescriptor().ForEachField(this, [&NumberOfField](::UE::FieldNotification::FFieldId FielId)
 			{
 				++NumberOfField;
 				return true;
@@ -599,7 +599,7 @@ void UBlueprintGeneratedClass::GetAssetRegistryTags(FAssetRegistryTagsContext Co
 	 */
 	if (GIsSavingPackage && !IsRunningCookCommandlet())
 	{
-		if (AActor* BlueprintCDO = Cast<AActor>(ClassDefaultObject))
+		if (AActor* BlueprintCDO = Cast<AActor>(GetDefaultObject(false)))
 		{
 			if (UPackage* BlueprintCDOPackage = BlueprintCDO->GetPackage())
 			{
@@ -650,7 +650,7 @@ void UBlueprintGeneratedClass::ThreadedPostLoadAssetRegistryTagsOverride(FPostLo
 FPrimaryAssetId UBlueprintGeneratedClass::GetPrimaryAssetId() const
 {
 	FPrimaryAssetId AssetId;
-	if (!ClassDefaultObject)
+	if (!GetDefaultObject(false))
 	{
 		// All UBlueprintGeneratedClass objects should have a pointer to their generated ClassDefaultObject, except for the
 		// CDO itself of the UBlueprintGeneratedClass class.
@@ -658,7 +658,7 @@ FPrimaryAssetId UBlueprintGeneratedClass::GetPrimaryAssetId() const
 		return AssetId;
 	}
 
-	AssetId = ClassDefaultObject->GetPrimaryAssetId();
+	AssetId = GetDefaultObject(false)->GetPrimaryAssetId();
 
 	/*
 	if (!AssetId.IsValid())
@@ -726,7 +726,7 @@ void UBlueprintGeneratedClass::SerializeDefaultObject(UObject* Object, FStructur
 
 #if WITH_EDITOR
 	using namespace UE::Runtime::Engine::Private;
-	if (UnderlyingArchive.IsSaving() && !UnderlyingArchive.IsObjectReferenceCollector() && Object == ClassDefaultObject)
+	if (UnderlyingArchive.IsSaving() && !UnderlyingArchive.IsObjectReferenceCollector() && Object == GetDefaultObject(false))
 	{
 		// Validate/fix up default subobjects prior to saving the CDO.
 		FBlueprintGeneratedClassUtils::ValidateDefaultSubobjects(Object);
@@ -735,7 +735,7 @@ void UBlueprintGeneratedClass::SerializeDefaultObject(UObject* Object, FStructur
 
 	Super::SerializeDefaultObject(Object, Slot);
 
-	if (UnderlyingArchive.IsLoading() && !UnderlyingArchive.IsObjectReferenceCollector() && Object == ClassDefaultObject)
+	if (UnderlyingArchive.IsLoading() && !UnderlyingArchive.IsObjectReferenceCollector() && Object == GetDefaultObject(false))
 	{
 #if WITH_EDITOR
 		// Validate/fix up default subobjects after serializing the CDO (e.g. fix up any unexpected NULL refs).
@@ -887,14 +887,14 @@ void UBlueprintGeneratedClass::SerializeDefaultObject(UObject* Object, FStructur
 		if (!GetPackage()->HasAnyPackageFlags(PKG_Cooked))
 		{
 			PrepareToConformSparseClassData(GetSparseClassDataArchetypeStruct());
-			if(ClassDefaultObject)
+			if(UObject* CDO = GetDefaultObject(false))
 			{
 				// conform immediately, so that child types can correct delta serialize their
 				// sparse data - we can't check RF_LoadCompleted because it is going to be set by the 
 				// preload call that is causing SerializeDefaultObject to run, but Super::SerializeDefaultObject
 				// has run...
-				ConformSparseClassData(ClassDefaultObject);
-				ClassDefaultObject->MoveDataToSparseClassDataStruct();
+				ConformSparseClassData(CDO);
+				CDO->MoveDataToSparseClassDataStruct();
 			}
 		}
 #endif
@@ -907,7 +907,7 @@ void UBlueprintGeneratedClass::PostLoadDefaultObject(UObject* Object)
 
 	Super::PostLoadDefaultObject(Object);
 
-	if (Object == ClassDefaultObject)
+	if (Object == GetDefaultObject(false))
 	{
 		// Rebuild the custom property list used in post-construct initialization logic. Note that PostLoad() may have altered some serialized properties.
 		UpdateCustomPropertyListForPostConstruction();
@@ -915,7 +915,7 @@ void UBlueprintGeneratedClass::PostLoadDefaultObject(UObject* Object)
 		// Restore any property values from config file
 		if (HasAnyClassFlags(CLASS_Config))
 		{
-			ClassDefaultObject->LoadConfig();
+			GetDefaultObject(false)->LoadConfig();
 		}
 	}
 
@@ -1205,11 +1205,11 @@ void UBlueprintGeneratedClass::UpdateCustomPropertyListForPostConstruction()
 	// Find the first native antecedent. All non-native decendant properties are attached to the PostConstructLink chain (see UStruct::Link), so we only need to worry about properties owned by native super classes here.
 	if (UClass* SuperClass = FBlueprintGeneratedClassUtils::FindFirstNativeClassInHierarchy(GetSuperClass()))
 	{
-		check(ClassDefaultObject != nullptr);
+		check(GetDefaultObject(false) != nullptr);
 
 		// Recursively gather native class-owned property values that differ from defaults.
 		FCustomPropertyListNode* PropertyList = nullptr;
-		BuildCustomPropertyListForPostConstruction(PropertyList, SuperClass, (uint8*)ClassDefaultObject.Get(), (uint8*)SuperClass->GetDefaultObject(false));
+		BuildCustomPropertyListForPostConstruction(PropertyList, SuperClass, (uint8*)GetDefaultObject(false), (uint8*)SuperClass->GetDefaultObject(false));
 	}
 
 	bCustomPropertyListForPostConstructionInitialized = true;
@@ -2106,7 +2106,7 @@ bool UBlueprintGeneratedClass::NeedsLoadForServer() const
 		{
 			return false;
 		}
-		if (ensure(ClassDefaultObject) && !ClassDefaultObject->NeedsLoadForServer())
+		if (ensure(GetDefaultObject(false)) && !GetDefaultObject(false)->NeedsLoadForServer())
 		{
 			return false;
 		}
@@ -2127,7 +2127,7 @@ bool UBlueprintGeneratedClass::NeedsLoadForClient() const
 		{
 			return false;
 		}
-		if (ensure(ClassDefaultObject) && !ClassDefaultObject->NeedsLoadForClient())
+		if (ensure(GetDefaultObject(false)) && !GetDefaultObject(false)->NeedsLoadForClient())
 		{
 			return false;
 		}
@@ -2151,10 +2151,10 @@ UClass* UBlueprintGeneratedClass::RegenerateClass(UClass* ClassToRegenerate, UOb
 {
 	if (HasAnyFlags(RF_BeingRegenerated))
 	{
-		if (ensure(ClassDefaultObject))
+		if (ensure(GetDefaultObject(false)))
 		{
 			UBlueprint::ForceLoadMembers(this);
-			UBlueprint::ForceLoadMembers(ClassDefaultObject);
+			UBlueprint::ForceLoadMembers(GetDefaultObject(false));
 		}
 
 		if (SimpleConstructionScript)
diff --git a/Engine/Source/Runtime/Engine/Private/Tests/Internationalization/TextPropertyTest.cpp b/Engine/Source/Runtime/Engine/Private/Tests/Internationalization/TextPropertyTest.cpp
--- a/Engine/Source/Runtime/Engine/Private/Tests/Internationalization/TextPropertyTest.cpp
+++ b/Engine/Source/Runtime/Engine/Private/Tests/Internationalization/TextPropertyTest.cpp
@@ -24,7 +24,7 @@ bool FTextPropertyTest::RunTest (const FString& Parameters)
 	UClass* const TextPropertyTestObjectClass = UTextPropertyTestObject::StaticClass();
 	FTextProperty* const DefaultedTextProperty = FindFProperty<FTextProperty>(TextPropertyTestObjectClass, "DefaultedText");
 	FTextProperty* const UndefaultedTextProperty = FindFProperty<FTextProperty>(TextPropertyTestObjectClass, "UndefaultedText");
-	UTextPropertyTestObject* const TextPropertyTestCDO = Cast<UTextPropertyTestObject>( TextPropertyTestObjectClass->ClassDefaultObject );
+	const UTextPropertyTestObject* const TextPropertyTestCDO = Cast<UTextPropertyTestObject>( TextPropertyTestObjectClass->GetDefaultObject(false) );
 
 	{
 		UTextPropertyTestObject* NewUObject = NewObject<UTextPropertyTestObject>(GetTransientPackage(), TextPropertyTestObjectClass);
diff --git a/Engine/Source/Runtime/Engine/Private/TransactionCommon.cpp b/Engine/Source/Runtime/Engine/Private/TransactionCommon.cpp
--- a/Engine/Source/Runtime/Engine/Private/TransactionCommon.cpp
+++ b/Engine/Source/Runtime/Engine/Private/TransactionCommon.cpp
@@ -21,12 +21,12 @@ FPersistentObjectRef::FPersistentObjectRef(UObject* InObject)
 
 			const bool bIsCDO = Obj->HasAllFlags(RF_ClassDefaultObject);
 			const UObject* CDO = bIsCDO ? Obj : nullptr;
-			const bool bIsClassCDO = (CDO != nullptr) ? (CDO->GetClass()->ClassDefaultObject == CDO) : false;
+			const bool bIsClassCDO = (CDO != nullptr) ? (CDO->GetClass()->GetDefaultObject(false) == CDO) : false;
 			if (!bIsClassCDO && CDO)
 			{
 				// Likely a trashed CDO, try to recover
 				// Only known cause of this is ambiguous use of DSOs
-				CDO = CDO->GetClass()->ClassDefaultObject;
+				CDO = CDO->GetClass()->GetDefaultObject(false);
 			}
 			const UActorComponent* AsComponent = Cast<UActorComponent>(Obj);
 			const bool bIsDSO = Obj->HasAnyFlags(RF_DefaultSubObject);
diff --git a/Engine/Source/Runtime/Experimental/JsonObjectGraph/Private/JsonStringifyImpl.cpp b/Engine/Source/Runtime/Experimental/JsonObjectGraph/Private/JsonStringifyImpl.cpp
--- a/Engine/Source/Runtime/Experimental/JsonObjectGraph/Private/JsonStringifyImpl.cpp
+++ b/Engine/Source/Runtime/Experimental/JsonObjectGraph/Private/JsonStringifyImpl.cpp
@@ -79,7 +79,7 @@ static void FilterEditorOnlyObjects(
 	#if WITH_EDITOR
 	Objects.SetNum(Algo::StableRemoveIf(Objects, [](const UObject* Obj) {
 		const UClass* NativeClass = GetFirstNativeClass(Obj->GetClass());
-		return IsEditorOnlyObject(NativeClass) || !(NativeClass->ClassDefaultObject->NeedsLoadForClient() || NativeClass->ClassDefaultObject->NeedsLoadForServer());
+		return IsEditorOnlyObject(NativeClass) || !(NativeClass->GetDefaultObject(false)->NeedsLoadForClient() || NativeClass->GetDefaultObject(false)->NeedsLoadForServer());
 		}));
 	#endif
 }
diff --git a/Engine/Source/Runtime/JsonUtilities/Private/JsonObjectConverter.cpp b/Engine/Source/Runtime/JsonUtilities/Private/JsonObjectConverter.cpp
--- a/Engine/Source/Runtime/JsonUtilities/Private/JsonObjectConverter.cpp
+++ b/Engine/Source/Runtime/JsonUtilities/Private/JsonObjectConverter.cpp
@@ -863,7 +863,7 @@ namespace
 				}
 
 				UObject* createdObj = StaticAllocateObject(PropertyClass, Outer, NAME_None, EObjectFlags::RF_NoFlags, EInternalObjectFlags::None, false);
-				(*PropertyClass->ClassConstructor)(FObjectInitializer(createdObj, PropertyClass->ClassDefaultObject, EObjectInitializerOptions::None));
+				(*PropertyClass->ClassConstructor)(FObjectInitializer(createdObj, PropertyClass->GetDefaultObject(false), EObjectInitializerOptions::None));
 
 				ObjectProperty->SetObjectPropertyValue(OutValue, createdObj);
 
diff --git a/Engine/Source/Runtime/LevelSequence/Private/LevelSequence.cpp b/Engine/Source/Runtime/LevelSequence/Private/LevelSequence.cpp
--- a/Engine/Source/Runtime/LevelSequence/Private/LevelSequence.cpp
+++ b/Engine/Source/Runtime/LevelSequence/Private/LevelSequence.cpp
@@ -281,11 +281,11 @@ void PurgeLegacyBlueprints(UObject* InObject, UPackage* Package)
 			UObject* OldOuter   = BP->GeneratedClass->GetOuter();
 			UClass*  SuperClass = BP->GeneratedClass->GetSuperClass();
 
-			if( BP->GeneratedClass->ClassDefaultObject )
+			if (UObject* ClassDefaultObject = BP->GeneratedClass->GetDefaultObject(false))
 			{
-				BP->GeneratedClass->ClassDefaultObject->ClearFlags(RF_Public);
-				BP->GeneratedClass->ClassDefaultObject->SetFlags(RF_Transient);
-				BP->GeneratedClass->ClassDefaultObject->RemoveFromRoot();
+				ClassDefaultObject->ClearFlags(RF_Public);
+				ClassDefaultObject->SetFlags(RF_Transient);
+				ClassDefaultObject->RemoveFromRoot();
 			}
 
 			BP->GeneratedClass->ClearFlags(RF_Public);
