RigVM: Avoid removing execute pins during template resolval

#rnx

diff --git a/Engine/Plugins/Runtime/RigVM/Source/RigVMDeveloper/Private/RigVMModel/RigVMController.cpp b/Engine/Plugins/Runtime/RigVM/Source/RigVMDeveloper/Private/RigVMModel/RigVMController.cpp
--- a/Engine/Plugins/Runtime/RigVM/Source/RigVMDeveloper/Private/RigVMModel/RigVMController.cpp
+++ b/Engine/Plugins/Runtime/RigVM/Source/RigVMDeveloper/Private/RigVMModel/RigVMController.cpp
@@ -19600,6 +19600,13 @@ bool URigVMController::FullyResolveTemplateNode(URigVMTemplateNode* InNode, int3
 		{
 			for(URigVMPin* Pin : InNode->GetPins())
 			{
+				// don't remove execute contexts, they are not part of the template's
+				// wildcard pins - so they can't change between resolvals.
+				if(Pin->IsExecuteContext())
+				{
+					continue;
+				}
+				
 				if(ResolvedFunction->Struct->FindPropertyByName(Pin->GetFName()) == nullptr)
 				{
 					PinsToRemove.Add(Pin);
