Control Rig: Fix crash when reopening the new rigvm example editor
#rnx

diff --git a/Engine/Plugins/Runtime/RigVM/Source/RigVMEditor/Private/Editor/RigVMEditorGraphExplorerTabSummoner.cpp b/Engine/Plugins/Runtime/RigVM/Source/RigVMEditor/Private/Editor/RigVMEditorGraphExplorerTabSummoner.cpp
--- a/Engine/Plugins/Runtime/RigVM/Source/RigVMEditor/Private/Editor/RigVMEditorGraphExplorerTabSummoner.cpp
+++ b/Engine/Plugins/Runtime/RigVM/Source/RigVMEditor/Private/Editor/RigVMEditorGraphExplorerTabSummoner.cpp
@@ -21,9 +21,13 @@ FRigVMEditorGraphExplorerTabSummoner::FRigVMEditorGraphExplorerTabSummoner(const
 
 TSharedRef<SWidget> FRigVMEditorGraphExplorerTabSummoner::CreateTabBody(const FWorkflowTabSpawnInfo& Info) const
 {
-	TSharedRef<SRigVMEditorGraphExplorer> Explorer = SNew(SRigVMEditorGraphExplorer, RigVMEditor.Pin().ToSharedRef());
-	RigVMEditor.Pin()->SetGraphExplorerWidget(Explorer);
-	return Explorer;
+	if (TSharedPtr<FRigVMEditorBase> Editor = RigVMEditor.Pin())
+	{
+		TSharedRef<SRigVMEditorGraphExplorer> Explorer = SNew(SRigVMEditorGraphExplorer, Editor);
+		Editor->SetGraphExplorerWidget(Explorer);
+		return Explorer;
+	}
+	return SNullWidget::NullWidget;
 }
 
 #undef LOCTEXT_NAMESPACE 
diff --git a/Engine/Plugins/Runtime/RigVM/Source/RigVMEditor/Private/Editor/SRigVMDetailsInspector.cpp b/Engine/Plugins/Runtime/RigVM/Source/RigVMEditor/Private/Editor/SRigVMDetailsInspector.cpp
--- a/Engine/Plugins/Runtime/RigVM/Source/RigVMEditor/Private/Editor/SRigVMDetailsInspector.cpp
+++ b/Engine/Plugins/Runtime/RigVM/Source/RigVMEditor/Private/Editor/SRigVMDetailsInspector.cpp
@@ -127,6 +127,12 @@ const TArray< TWeakObjectPtr<UObject> >& SRigVMDetailsInspector::GetSelectedObje
 	return SelectedObjects;
 }
 
+void SRigVMDetailsInspector::OnEditorClose(const FRigVMEditorBase* RigVMEditorBase, URigVMBlueprint* RigVMBlueprint)
+{
+	PropertyView->UnregisterInstancedCustomPropertyLayout(URigVMEdGraph::StaticClass());
+	PropertyView->UnregisterInstancedCustomPropertyLayout(UPropertyWrapper::StaticClass());
+}
+
 void SRigVMDetailsInspector::AddReferencedObjects(FReferenceCollector& Collector)
 {
 	Collector.AddReferencedObjects(RefreshPropertyObjects);
@@ -228,17 +234,14 @@ void SRigVMDetailsInspector::Construct(const FArguments& InArgs)
 	IsPropertyEditingEnabledDelegate = InArgs._IsPropertyEditingEnabledDelegate;
 	UserOnFinishedChangingProperties = InArgs._OnFinishedChangingProperties;
 
-	TWeakPtr<SRigVMEditorGraphExplorer> GraphExplorer = Editor.IsValid() ? Editor->GetGraphExplorerWidget() : InArgs._GraphExplorerWidget;
-	
-	if( GraphExplorer.IsValid() )
-	{
-		const UClass* BlueprintClass = URigVMBlueprint::StaticClass();
-		FOnGetDetailCustomizationInstance FunctionDetails = FOnGetDetailCustomizationInstance::CreateStatic(&FRigVMGraphDetailCustomization::MakeInstance, Editor, BlueprintClass);
-		PropertyView->RegisterInstancedCustomPropertyLayout(URigVMEdGraph::StaticClass(), FunctionDetails);
+	const UClass* BlueprintClass = Editor->GetRigVMBlueprint()->GetClass();
+	FOnGetDetailCustomizationInstance FunctionDetails = FOnGetDetailCustomizationInstance::CreateStatic(&FRigVMGraphDetailCustomization::MakeInstance, Editor, BlueprintClass);
+	PropertyView->RegisterInstancedCustomPropertyLayout(URigVMEdGraph::StaticClass(), FunctionDetails);
 
-		FOnGetDetailCustomizationInstance LayoutVariableDetails = FOnGetDetailCustomizationInstance::CreateStatic(&FRigVMVariableDetailCustomization::MakeInstance, Editor);
-		PropertyView->RegisterInstancedCustomPropertyLayout(UPropertyWrapper::StaticClass(), LayoutVariableDetails);
-	}
+	FOnGetDetailCustomizationInstance LayoutVariableDetails = FOnGetDetailCustomizationInstance::CreateStatic(&FRigVMVariableDetailCustomization::MakeInstance, Editor);
+	PropertyView->RegisterInstancedCustomPropertyLayout(UPropertyWrapper::StaticClass(), LayoutVariableDetails);
+
+	Editor->OnEditorClosed().AddSP(this, &SRigVMDetailsInspector::OnEditorClose);
 
 	// Create the border that all of the content will get stuffed into
 	ChildSlot
@@ -252,10 +255,6 @@ void SRigVMDetailsInspector::Construct(const FArguments& InArgs)
 			.Padding(0.0f)
 			.BorderImage( FAppStyle::GetBrush("NoBorder") )
 		]
-		+SVerticalBox::Slot()
-		[
-			SNew(STextBlock).Text(LOCTEXT("RigVMDetailsInspectorTest", "Test"))
-		]
 	];
 
 	// Update based on the current (empty) selection set
diff --git a/Engine/Plugins/Runtime/RigVM/Source/RigVMEditor/Public/Editor/SRigVMDetailsInspector.h b/Engine/Plugins/Runtime/RigVM/Source/RigVMEditor/Public/Editor/SRigVMDetailsInspector.h
--- a/Engine/Plugins/Runtime/RigVM/Source/RigVMEditor/Public/Editor/SRigVMDetailsInspector.h
+++ b/Engine/Plugins/Runtime/RigVM/Source/RigVMEditor/Public/Editor/SRigVMDetailsInspector.h
@@ -11,6 +11,7 @@
 #include "Internationalization/Text.h"
 #include "Layout/Visibility.h"
 #include "PropertyEditorDelegates.h"
+#include "RigVMBlueprint.h"
 #include "SlateFwd.h"
 #include "Styling/SlateTypes.h"
 #include "Templates/SharedPointer.h"
@@ -119,6 +120,8 @@ public:
 	/** returns the list of selected objects */
 	const TArray< TWeakObjectPtr<UObject> >& GetSelectedObjects() const;
 
+	void OnEditorClose(const FRigVMEditorBase* RigVMEditorBase, URigVMBlueprint* RigVMBlueprint);
+
 	virtual void AddReferencedObjects(FReferenceCollector& Collector) override;
 	virtual FString GetReferencerName() const override;
 
