Avoid Cotnrol Rig mark the asset as dirty with node selection (but keeps selection undo redo)

diff --git a/Engine/Plugins/Runtime/RigVM/Source/RigVMDeveloper/Private/RigVMModel/RigVMController.cpp b/Engine/Plugins/Runtime/RigVM/Source/RigVMDeveloper/Private/RigVMModel/RigVMController.cpp
--- a/Engine/Plugins/Runtime/RigVM/Source/RigVMDeveloper/Private/RigVMModel/RigVMController.cpp
+++ b/Engine/Plugins/Runtime/RigVM/Source/RigVMDeveloper/Private/RigVMModel/RigVMController.cpp
@@ -7406,7 +7406,7 @@ bool URigVMController::SetNodeSelection(const TArray<FName>& InNodeNames, bool b
 					Action.SetTitle(TEXT("Selected multiple nodes."));
 				}
 			}
-			GetActionStack()->AddAction(Action);
+			GetActionStack()->AddAction(Action, false, false);
 		}
 	}
 
diff --git a/Engine/Plugins/Runtime/RigVM/Source/RigVMDeveloper/Public/RigVMModel/RigVMControllerActions.h b/Engine/Plugins/Runtime/RigVM/Source/RigVMDeveloper/Public/RigVMModel/RigVMControllerActions.h
--- a/Engine/Plugins/Runtime/RigVM/Source/RigVMDeveloper/Public/RigVMModel/RigVMControllerActions.h
+++ b/Engine/Plugins/Runtime/RigVM/Source/RigVMDeveloper/Public/RigVMModel/RigVMControllerActions.h
@@ -317,7 +317,7 @@ public:
 	// a potential merge of this action with the previous action to
 	// compact the stack.
 	template<class ActionType>
-	void AddAction(const ActionType& InAction, bool bPerformMerge = false)
+	void AddAction(const ActionType& InAction, bool bPerformMerge = false, bool bMarkDirty = true)
 	{
 		if(IsDisabled())
 		{
@@ -382,7 +382,7 @@ public:
 #if WITH_EDITOR
 				const FScopedTransaction Transaction(FText::FromString(InAction.GetTitle()));
 				SetFlags(RF_Transactional);
-				Modify();
+				Modify(bMarkDirty);
 #endif
 
 				ActionIndex = ActionIndex + 1;
