Control Rig: Remove AutoCompile button in toolbar. Also hide/unhide debugging buttons based on whether we are halted at a breakpoint
#rnx

diff --git a/Engine/Plugins/Runtime/RigVM/Source/RigVMEditor/Private/Editor/RigVMEditor.cpp b/Engine/Plugins/Runtime/RigVM/Source/RigVMEditor/Private/Editor/RigVMEditor.cpp
--- a/Engine/Plugins/Runtime/RigVM/Source/RigVMEditor/Private/Editor/RigVMEditor.cpp
+++ b/Engine/Plugins/Runtime/RigVM/Source/RigVMEditor/Private/Editor/RigVMEditor.cpp
@@ -1387,27 +1387,37 @@ void FRigVMEditorBase::BindCommands()
 	GetToolkitCommands()->MapAction(
 		FRigVMEditorCommands::Get().ResumeExecution,
 		FExecuteAction::CreateSP(ThisRef, &FRigVMEditorBase::HandleBreakpointActionRequested, ERigVMBreakpointAction::Resume),			
-		FIsActionChecked::CreateSP(ThisRef, &FRigVMEditorBase::IsHaltedAtBreakpoint));
+		FIsActionChecked::CreateSP(ThisRef, &FRigVMEditorBase::IsHaltedAtBreakpoint),
+		FIsActionChecked(),
+		FIsActionButtonVisible::CreateSP(ThisRef, &FRigVMEditorBase::IsHaltedAtBreakpoint));
 
 	GetToolkitCommands()->MapAction(
 		FRigVMEditorCommands::Get().ShowCurrentStatement,
 		FExecuteAction::CreateSP(ThisRef, &FRigVMEditorBase::HandleShowCurrentStatement),
-		FIsActionChecked::CreateSP(ThisRef, &FRigVMEditorBase::IsHaltedAtBreakpoint));
+		FIsActionChecked::CreateSP(ThisRef, &FRigVMEditorBase::IsHaltedAtBreakpoint),
+		FIsActionChecked(),
+		FIsActionButtonVisible::CreateSP(ThisRef, &FRigVMEditorBase::IsHaltedAtBreakpoint));
 
 	GetToolkitCommands()->MapAction(
 		FRigVMEditorCommands::Get().StepOver,
 		FExecuteAction::CreateSP(ThisRef, &FRigVMEditorBase::HandleBreakpointActionRequested, ERigVMBreakpointAction::StepOver),
-		FIsActionChecked::CreateSP(ThisRef, &FRigVMEditorBase::IsHaltedAtBreakpoint));
+		FIsActionChecked::CreateSP(ThisRef, &FRigVMEditorBase::IsHaltedAtBreakpoint),
+		FIsActionChecked(),
+		FIsActionButtonVisible::CreateSP(ThisRef, &FRigVMEditorBase::IsHaltedAtBreakpoint));
 
 	GetToolkitCommands()->MapAction(
 		FRigVMEditorCommands::Get().StepInto,
 		FExecuteAction::CreateSP(ThisRef, &FRigVMEditorBase::HandleBreakpointActionRequested, ERigVMBreakpointAction::StepInto),
-		FIsActionChecked::CreateSP(ThisRef, &FRigVMEditorBase::IsHaltedAtBreakpoint));
+		FIsActionChecked::CreateSP(ThisRef, &FRigVMEditorBase::IsHaltedAtBreakpoint),
+		FIsActionChecked(),
+		FIsActionButtonVisible::CreateSP(ThisRef, &FRigVMEditorBase::IsHaltedAtBreakpoint));
 
 	GetToolkitCommands()->MapAction(
 		FRigVMEditorCommands::Get().StepOut,
 		FExecuteAction::CreateSP(ThisRef, &FRigVMEditorBase::HandleBreakpointActionRequested, ERigVMBreakpointAction::StepOut),
-		FIsActionChecked::CreateSP(ThisRef, &FRigVMEditorBase::IsHaltedAtBreakpoint));
+		FIsActionChecked::CreateSP(ThisRef, &FRigVMEditorBase::IsHaltedAtBreakpoint),
+		FIsActionChecked(),
+		FIsActionButtonVisible::CreateSP(ThisRef, &FRigVMEditorBase::IsHaltedAtBreakpoint));
 
 	GetToolkitCommands()->MapAction(
 		FRigVMEditorCommands::Get().FrameSelection,
@@ -3316,9 +3326,6 @@ void FRigVMEditorBase::FillToolbar(FToolBarBuilder& ToolbarBuilder, bool bEndSec
 			FSlateIcon(FAppStyle::GetAppStyleSetName(), "LevelEditor.Recompile"),
 			true);
 
-		ToolbarBuilder.AddToolBarButton(FRigVMEditorCommands::Get().AutoCompileGraph,
-			NAME_None, TAttribute<FText>(), TAttribute<FText>(), FSlateIcon(FRigVMEditorStyle::Get().GetStyleSetName(), "RigVM.AutoCompileGraph"));
-
 		AddSelectedDebugObjectWidget(ToolbarBuilder);
 
 		ToolbarBuilder.AddSeparator();
