Control Rig: Add Compile toolbar button in new rigvm editor
#rnx

diff --git a/Engine/Plugins/Runtime/RigVM/Source/RigVMEditor/Private/Editor/Kismet/RigVMBlueprintCompilationManager.cpp b/Engine/Plugins/Runtime/RigVM/Source/RigVMEditor/Private/Editor/Kismet/RigVMBlueprintCompilationManager.cpp
--- a/Engine/Plugins/Runtime/RigVM/Source/RigVMEditor/Private/Editor/Kismet/RigVMBlueprintCompilationManager.cpp
+++ b/Engine/Plugins/Runtime/RigVM/Source/RigVMEditor/Private/Editor/Kismet/RigVMBlueprintCompilationManager.cpp
@@ -2105,7 +2105,7 @@ void FRigVMBlueprintCompilationManagerImpl::SaveConstructionHashes(const TArray<
 		{
 			if( const UObject* CDO = OriginalClass->GetDefaultObject(false) )
 			{
-				if( FOverridableManager::Get().IsEnabled(*CDO) )
+				if( FOverridableManager::Get().IsEnabled(CDO) )
 				{
 					return;
 				}
@@ -2604,7 +2604,7 @@ void FRigVMBlueprintCompilationManagerImpl::ReinstanceBatch(TArray<FReinstancing
 			{
 				// Object using overridable serialization need to use delta serialization for it do work appropriatly
 				// Eventually we should do this to all BP classes CDO no matter what
-				const bool bUsingOverrideSerialization = FOverridableManager::Get().IsEnabled(*OldCDO);
+				const bool bUsingOverrideSerialization = FOverridableManager::Get().IsEnabled(OldCDO);
 				const bool bUseDeltaSerialization = bUsingOverrideSerialization ? true : (ReinstancingJob.Reinstancer.IsValid() ? ReinstancingJob.Reinstancer->bUseDeltaSerializationToCopyProperties : false);
 				UClass* NewClass = ReinstancingJob.OldToNew.Value;
 
@@ -2755,7 +2755,7 @@ void FRigVMBlueprintCompilationManagerImpl::ReinstanceBatch(TArray<FReinstancing
 				{
 					bool bIsArchetype =
 						Obj->HasAnyFlags(RF_ArchetypeObject|RF_InheritableComponentTemplate)
-						|| (FOverridableManager::Get().IsEnabled(*Obj) && Obj->GetOutermostObject()->HasAnyFlags(RF_ClassDefaultObject))
+						|| (FOverridableManager::Get().IsEnabled(Obj) && Obj->GetOutermostObject()->HasAnyFlags(RF_ClassDefaultObject))
 						|| Obj->GetTypedOuter<UBlueprintGeneratedClass>()
 						|| Obj->GetTypedOuter<UBlueprint>();
 					// remove if this is not an archetype or its already in the transient package, note
@@ -2805,7 +2805,7 @@ void FRigVMBlueprintCompilationManagerImpl::ReinstanceBatch(TArray<FReinstancing
 							}
 
 							// this handles nested subobjects:
-							if (FOverridableManager::Get().IsEnabled(*Archetype))
+							if (FOverridableManager::Get().IsEnabled(Archetype))
 							{
 								ForEachObjectWithOuter(Iter, [Archetype, &ArchetypeReferencers](UObject* SubObject)
 								{
diff --git a/Engine/Plugins/Runtime/RigVM/Source/RigVMEditor/Private/Editor/RigVMEditor.cpp b/Engine/Plugins/Runtime/RigVM/Source/RigVMEditor/Private/Editor/RigVMEditor.cpp
--- a/Engine/Plugins/Runtime/RigVM/Source/RigVMEditor/Private/Editor/RigVMEditor.cpp
+++ b/Engine/Plugins/Runtime/RigVM/Source/RigVMEditor/Private/Editor/RigVMEditor.cpp
@@ -3309,6 +3309,8 @@ void FRigVMEditorBase::FillToolbar(FToolBarBuilder& ToolbarBuilder, bool bEndSec
 {
 	ToolbarBuilder.BeginSection("Toolbar");
 	{
+		AddCompileWidget(ToolbarBuilder);
+		
 		ToolbarBuilder.AddToolBarButton(
 			FRigVMEditorCommands::Get().ToggleEventQueue,
 			NAME_None, 
diff --git a/Engine/Plugins/Runtime/RigVM/Source/RigVMEditor/Private/Editor/RigVMEditorCommands.cpp b/Engine/Plugins/Runtime/RigVM/Source/RigVMEditor/Private/Editor/RigVMEditorCommands.cpp
--- a/Engine/Plugins/Runtime/RigVM/Source/RigVMEditor/Private/Editor/RigVMEditorCommands.cpp
+++ b/Engine/Plugins/Runtime/RigVM/Source/RigVMEditor/Private/Editor/RigVMEditorCommands.cpp
@@ -6,6 +6,12 @@
 
 void FRigVMEditorCommands::RegisterCommands()
 {
+	UI_COMMAND(Compile, "Compile", "Compile the blueprint", EUserInterfaceActionType::Button, FInputChord());
+	UI_COMMAND(SaveOnCompile_Never, "Never", "Sets the save-on-compile option to 'Never', meaning that your Blueprints will not be saved when they are compiled", EUserInterfaceActionType::RadioButton, FInputChord());
+	UI_COMMAND(SaveOnCompile_SuccessOnly, "On Success Only", "Sets the save-on-compile option to 'Success Only', meaning that your Blueprints will be saved whenever they are successfully compiled", EUserInterfaceActionType::RadioButton, FInputChord());
+	UI_COMMAND(SaveOnCompile_Always, "Always", "Sets the save-on-compile option to 'Always', meaning that your Blueprints will be saved whenever they are compiled (even if there were errors)", EUserInterfaceActionType::RadioButton, FInputChord());
+	UI_COMMAND(JumpToErrorNode, "Jump to Error Node", "When enabled, then the Blueprint will snap focus to nodes producing an error during compilation", EUserInterfaceActionType::ToggleButton, FInputChord());
+
 	UI_COMMAND(DeleteItem, "Delete", "Deletes the selected items and removes their nodes from the graph.", EUserInterfaceActionType::Button, FInputChord(EKeys::Delete));
 	UI_COMMAND(ExecuteGraph, "Execute", "Execute the rig graph if On.", EUserInterfaceActionType::ToggleButton, FInputChord());
 	UI_COMMAND(AutoCompileGraph, "Auto Compile", "Auto-compile the rig graph if On.", EUserInterfaceActionType::ToggleButton, FInputChord());
diff --git a/Engine/Plugins/Runtime/RigVM/Source/RigVMEditor/Private/Editor/RigVMEditorCommands.h b/Engine/Plugins/Runtime/RigVM/Source/RigVMEditor/Private/Editor/RigVMEditorCommands.h
--- a/Engine/Plugins/Runtime/RigVM/Source/RigVMEditor/Private/Editor/RigVMEditorCommands.h
+++ b/Engine/Plugins/Runtime/RigVM/Source/RigVMEditor/Private/Editor/RigVMEditorCommands.h
@@ -17,6 +17,13 @@ public:
 		FRigVMEditorStyle::Get().GetStyleSetName() // Icon Style Set
 	)
 	{}
+
+	/** Compile the blueprint */
+	TSharedPtr<FUICommandInfo> Compile;
+	TSharedPtr<FUICommandInfo> SaveOnCompile_Never;
+	TSharedPtr<FUICommandInfo> SaveOnCompile_SuccessOnly;
+	TSharedPtr<FUICommandInfo> SaveOnCompile_Always;
+	TSharedPtr<FUICommandInfo> JumpToErrorNode;
 	
 	/** Deletes the selected items and removes their nodes from the graph. */
 	TSharedPtr< FUICommandInfo > DeleteItem;
diff --git a/Engine/Plugins/Runtime/RigVM/Source/RigVMEditor/Private/Editor/RigVMLegacyEditor.cpp b/Engine/Plugins/Runtime/RigVM/Source/RigVMEditor/Private/Editor/RigVMLegacyEditor.cpp
--- a/Engine/Plugins/Runtime/RigVM/Source/RigVMEditor/Private/Editor/RigVMLegacyEditor.cpp
+++ b/Engine/Plugins/Runtime/RigVM/Source/RigVMEditor/Private/Editor/RigVMLegacyEditor.cpp
@@ -53,6 +53,14 @@ void FRigVMLegacyEditor::OnGraphEditorFocused(const TSharedRef<SGraphEditor>& In
 	FRigVMEditorBase::OnGraphEditorFocused(InGraphEditor);
 }
 
+void FRigVMLegacyEditor::AddCompileWidget(FToolBarBuilder& ToolbarBuilder)
+{
+	if (UToolMenu* ToolMenu = RegisterModeToolbarIfUnregistered(FRigVMEditorModes::RigVMEditorMode))
+	{
+		GetToolbarBuilder()->AddCompileToolbar(ToolMenu);
+	}
+}
+
 void FRigVMLegacyEditor::AddSelectedDebugObjectWidget(FToolBarBuilder& ToolbarBuilder)
 {
 	ToolbarBuilder.AddWidget(SNew(SBlueprintEditorSelectedDebugObjectWidget, SharedThis(this)));
diff --git a/Engine/Plugins/Runtime/RigVM/Source/RigVMEditor/Private/Editor/RigVMNewEditor.cpp b/Engine/Plugins/Runtime/RigVM/Source/RigVMEditor/Private/Editor/RigVMNewEditor.cpp
--- a/Engine/Plugins/Runtime/RigVM/Source/RigVMEditor/Private/Editor/RigVMNewEditor.cpp
+++ b/Engine/Plugins/Runtime/RigVM/Source/RigVMEditor/Private/Editor/RigVMNewEditor.cpp
@@ -3,6 +3,7 @@
 #include "Editor/RigVMNewEditor.h"
 
 #include "BlueprintActionDatabase.h"
+#include "BlueprintEditorSettings.h"
 #include "BlueprintEditorTabs.h"
 #include "IMessageLogListing.h"
 #include "K2Node_Composite.h"
@@ -28,9 +29,12 @@
 #include "GraphEditor.h"
 #include "GraphEditorActions.h"
 #include "RigVMBlueprint.h"
+#include "RigVMEditorCommands.h"
 #include "Editor/RigVMGraphEditorSummoner.h"
 #include "Widgets/Input/SEditableTextBox.h"
 #include "Editor/SRigVMDetailsInspector.h"
+#include "Engine/LevelScriptBlueprint.h"
+#include "Kismet2/KismetEditorUtilities.h"
 #if WITH_RIGVMLEGACYEDITOR
 #include "SKismetInspector.h"
 #endif
@@ -511,6 +515,112 @@ void FRigVMNewEditor::OnGraphEditorBackgrounded(const TSharedRef<SGraphEditor>&
 	}
 }
 
+FText FRigVMNewEditor::GetCompileStatusTooltip() const
+{
+	// Copied from FBlueprintEditorToolbar::GetStatusTooltip
+	
+	UBlueprint* BlueprintObj = GetBlueprintObj();
+	EBlueprintStatus Status = BlueprintObj->Status;
+
+	switch (Status)
+	{
+	case BS_Dirty:
+		return LOCTEXT("Dirty_Status", "Dirty; needs to be recompiled");
+	case BS_Error:
+		return LOCTEXT("CompileError_Status", "There was an error during compilation, see the log for details");
+	case BS_UpToDate:
+		return LOCTEXT("GoodToGo_Status", "Good to go");
+	case BS_UpToDateWithWarnings:
+		return LOCTEXT("GoodToGoWarning_Status", "There was a warning during compilation, see the log for details");
+	default:
+		return LOCTEXT("Recompile_Status", "Unknown status; should recompile");
+	}
+}
+
+FSlateIcon FRigVMNewEditor::GetCompileStatusImage() const
+{
+	// Copied from FBlueprintEditorToolbar::GetStatusImage
+	
+	UBlueprint* BlueprintObj = GetBlueprintObj();
+	EBlueprintStatus Status = BlueprintObj->Status;
+
+	// For macro types, always show as up-to-date, since we don't compile them
+	if (BlueprintObj->BlueprintType == BPTYPE_MacroLibrary)
+	{
+		Status = BS_UpToDate;
+	}
+
+	
+	static const FName CompileStatusBackground("Blueprint.CompileStatus.Background");
+	static const FName CompileStatusUnknown("Blueprint.CompileStatus.Overlay.Unknown");
+	static const FName CompileStatusError("Blueprint.CompileStatus.Overlay.Error");
+	static const FName CompileStatusGood("Blueprint.CompileStatus.Overlay.Good");
+	static const FName CompileStatusWarning("Blueprint.CompileStatus.Overlay.Warning");
+
+	switch (Status)
+	{
+	case BS_Error:
+		return FSlateIcon(FAppStyle::GetAppStyleSetName(), CompileStatusBackground, NAME_None, CompileStatusError);
+	case BS_UpToDate:
+		return FSlateIcon(FAppStyle::GetAppStyleSetName(), CompileStatusBackground, NAME_None, CompileStatusGood);
+	case BS_UpToDateWithWarnings:
+		return FSlateIcon(FAppStyle::GetAppStyleSetName(), CompileStatusBackground, NAME_None, CompileStatusWarning);
+	default:
+		return FSlateIcon(FAppStyle::GetAppStyleSetName(), CompileStatusBackground, NAME_None, CompileStatusUnknown);
+	}
+}
+
+void FRigVMNewEditor::MakeSaveOnCompileSubMenu(FMenuBuilder& InMenuBuilder)
+{
+	InMenuBuilder.BeginSection(TEXT("Section"));
+	const FRigVMEditorCommands& Commands = FRigVMEditorCommands::Get();
+	InMenuBuilder.AddMenuEntry(Commands.SaveOnCompile_Never);
+	InMenuBuilder.AddMenuEntry(Commands.SaveOnCompile_SuccessOnly);
+	InMenuBuilder.AddMenuEntry(Commands.SaveOnCompile_Always);
+}
+
+TSharedRef<SWidget> FRigVMNewEditor::GenerateCompileOptionsMenu()
+{
+	FMenuBuilder MenuBuilder(true, GetToolkitCommands());
+	MenuBuilder.BeginSection(TEXT("Section"));
+	const FRigVMEditorCommands& Commands = FRigVMEditorCommands::Get();
+
+	// @TODO: disable the menu and change up the tooltip when all sub items are disabled
+	MenuBuilder.AddSubMenu(
+		LOCTEXT("SaveOnCompileSubMenu", "Save on Compile"),
+		LOCTEXT("SaveOnCompileSubMenu_ToolTip", "Determines how the Blueprint is saved whenever you compile it."),
+		FNewMenuDelegate::CreateSP(this, &FRigVMNewEditor::MakeSaveOnCompileSubMenu)
+	);
+
+	MenuBuilder.AddMenuEntry(Commands.JumpToErrorNode);
+	MenuBuilder.AddMenuEntry(Commands.AutoCompileGraph);
+	MenuBuilder.EndSection();
+
+	return MenuBuilder.MakeWidget();
+}
+
+void FRigVMNewEditor::AddCompileWidget(FToolBarBuilder& ToolbarBuilder)
+{
+	ToolbarBuilder.AddToolBarButton(FRigVMEditorCommands::Get().Compile,
+		NAME_None, TAttribute<FText>(),
+		TAttribute<FText>::CreateLambda([this]()
+		{
+			return GetCompileStatusTooltip();
+		}),
+		TAttribute<FSlateIcon>::CreateLambda([this]()
+		{
+			return GetCompileStatusImage();
+		}),
+		"CompileBlueprint");
+
+	ToolbarBuilder.AddComboButton(
+			FUIAction(),
+				FOnGetContent::CreateSP(this, &FRigVMNewEditor::GenerateCompileOptionsMenu),
+				TAttribute<FText>(),
+		LOCTEXT("BlueprintCompileOptions_ToolbarTooltip", "Options to customize how Blueprints compile"),
+		TAttribute<FSlateIcon>(), true);
+}
+
 void FRigVMNewEditor::RegisterTabSpawners(const TSharedRef<FTabManager>& InTabManager)
 {
 	//@TODO: Can't we do this sooner?
@@ -930,12 +1040,151 @@ UEdGraphPin* FRigVMNewEditor::GetCurrentlySelectedPin() const
 	return nullptr;
 }
 
+void FRigVMNewEditor::SetSaveOnCompileSetting(ESaveOnCompile NewSetting)
+{
+	// Copied from BlueprintEditorImpl::SetSaveOnCompileSetting
+	
+	UBlueprintEditorSettings* Settings = GetMutableDefault<UBlueprintEditorSettings>();
+	Settings->SaveOnCompile = NewSetting;
+	Settings->SaveConfig();
+}
+
+bool FRigVMNewEditor::IsSaveOnCompileEnabled() const
+{
+	// Copied from BlueprintEditorImpl::IsSaveOnCompileEnabled
+	
+	UBlueprint* Blueprint = GetBlueprintObj();
+	bool const bIsLevelScript = (Cast<ULevelScriptBlueprint>(Blueprint) != nullptr);
+
+	return !bIsLevelScript;
+}
+
+bool FRigVMNewEditor::IsSaveOnCompileOptionSet(TWeakPtr<FRigVMNewEditor> Editor, ESaveOnCompile Option)
+{
+	// Copied from BlueprintEditorImpl::IsSaveOnCompileOptionSet
+	
+	const UBlueprintEditorSettings* Settings = GetDefault<UBlueprintEditorSettings>();
+
+	ESaveOnCompile CurrentSetting = Settings->SaveOnCompile;
+	if (!Editor.IsValid() || !Editor.Pin()->IsSaveOnCompileEnabled())
+	{
+		// if save-on-compile is disabled for the blueprint, then we want to 
+		// show "Never" as being selected
+		// 
+		// @TODO: a tooltip explaining why would be nice too
+		CurrentSetting = SoC_Never;
+	}
+
+	return (CurrentSetting == Option);
+}
+
+void FRigVMNewEditor::ToggleJumpToErrorNodeSetting()
+{
+	// Copied from BlueprintEditorImpl::ToggleJumpToErrorNodeSetting
+	
+	UBlueprintEditorSettings* Settings = GetMutableDefault<UBlueprintEditorSettings>();
+	Settings->bJumpToNodeErrors = !Settings->bJumpToNodeErrors;
+	Settings->SaveConfig();
+}
+
+bool FRigVMNewEditor::IsJumpToErrorNodeOptionSet()
+{
+	// Copied from BlueprintEditorImpl::IsJumpToErrorNodeOptionSet
+	
+	const UBlueprintEditorSettings* Settings = GetDefault<UBlueprintEditorSettings>();
+	return Settings->bJumpToNodeErrors;
+}
+
+UEdGraphNode* FRigVMNewEditor::FindNodeWithError(UBlueprint* Blueprint, EMessageSeverity::Type Severity/* = EMessageSeverity::Error*/)
+{
+	// Copied from BlueprintEditorImpl::FindNodeWithError
+	
+	TArray<UEdGraph*> Graphs;
+	Blueprint->GetAllGraphs(Graphs);
+
+	UEdGraphNode* ChoiceNode = nullptr;
+	for (UEdGraph* Graph : Graphs)
+	{
+		for (UEdGraphNode* Node : Graph->Nodes)
+		{
+			if (Node && Node->bHasCompilerMessage && !Node->ErrorMsg.IsEmpty() && (Node->ErrorType <= Severity))
+			{
+				if ((ChoiceNode == nullptr) || (ChoiceNode->ErrorType > Node->ErrorType))
+				{
+					ChoiceNode = Node;
+					if (ChoiceNode->ErrorType == 0)
+					{
+						break;
+					}
+				}
+			}
+		}
+	}
+	return ChoiceNode;
+}
+
+UEdGraphNode* FRigVMNewEditor::FindNodeWithError(FCompilerResultsLog const& ErrorLog, EMessageSeverity::Type Severity/* = EMessageSeverity::Error*/)
+{
+	// Copied from BlueprintEditorImpl::FindNodeWithError
+	
+	UEdGraphNode* ChoiceNode = nullptr;
+	for (TWeakObjectPtr<UEdGraphNode> NodePtr : ErrorLog.AnnotatedNodes)
+	{
+		UEdGraphNode* Node = NodePtr.Get();
+		if ((Node != nullptr) && (Node->ErrorType <= Severity))
+		{
+			if ((ChoiceNode == nullptr) || (Node->ErrorType < ChoiceNode->ErrorType))
+			{
+				ChoiceNode = Node;
+				if (ChoiceNode->ErrorType == 0)
+				{
+					break;
+				}
+			}
+		}
+	}
+
+	return ChoiceNode;
+}
+
 void FRigVMNewEditor::CreateDefaultCommands()
 {
-	// TODO sara-s: Replicate what blueprint editor is doing
 	if (GetBlueprintObj())
 	{
+		// TODO sara-s: Replicate what blueprint editor is doing
 		//FBlueprintEditor::CreateDefaultCommands();
+
+		GetToolkitCommands()->MapAction(
+			FRigVMEditorCommands::Get().Compile,
+			FExecuteAction::CreateSP(this, &FRigVMEditorBase::Compile),
+			FCanExecuteAction::CreateSP(this, &FRigVMEditorBase::IsCompilingEnabled));
+
+		TWeakPtr<FRigVMNewEditor> WeakThisPtr = SharedThis(this);
+		ToolkitCommands->MapAction(
+			FRigVMEditorCommands::Get().SaveOnCompile_Never,
+			FExecuteAction::CreateSP(this, &FRigVMNewEditor::SetSaveOnCompileSetting, (ESaveOnCompile)SoC_Never),
+			FCanExecuteAction::CreateSP(this, &FRigVMNewEditor::IsSaveOnCompileEnabled),
+			FIsActionChecked::CreateSP(this, &FRigVMNewEditor::IsSaveOnCompileOptionSet, WeakThisPtr, (ESaveOnCompile)SoC_Never)
+		);
+		ToolkitCommands->MapAction(
+			FRigVMEditorCommands::Get().SaveOnCompile_SuccessOnly,
+			FExecuteAction::CreateSP(this, &FRigVMNewEditor::SetSaveOnCompileSetting, (ESaveOnCompile)SoC_SuccessOnly),
+			FCanExecuteAction::CreateSP(this, &FRigVMNewEditor::IsSaveOnCompileEnabled),
+			FIsActionChecked::CreateSP(this, &FRigVMNewEditor::IsSaveOnCompileOptionSet, WeakThisPtr, (ESaveOnCompile)SoC_SuccessOnly)
+		);
+		ToolkitCommands->MapAction(
+			FRigVMEditorCommands::Get().SaveOnCompile_Always,
+			FExecuteAction::CreateSP(this, &FRigVMNewEditor::SetSaveOnCompileSetting, (ESaveOnCompile)SoC_Always),
+			FCanExecuteAction::CreateSP(this, &FRigVMNewEditor::IsSaveOnCompileEnabled),
+			FIsActionChecked::CreateSP(this, &FRigVMNewEditor::IsSaveOnCompileOptionSet, WeakThisPtr, (ESaveOnCompile)SoC_Always)
+		);
+
+		ToolkitCommands->MapAction(
+			FRigVMEditorCommands::Get().JumpToErrorNode,
+			FExecuteAction::CreateSP(this, &FRigVMNewEditor::ToggleJumpToErrorNodeSetting),
+			FCanExecuteAction(),
+			FIsActionChecked::CreateSP(this, &FRigVMNewEditor::IsJumpToErrorNodeOptionSet)
+		);
 	}
 
 	FRigVMEditorBase::CreateDefaultCommandsImpl();
@@ -1107,27 +1356,18 @@ void FRigVMNewEditor::CompileImpl()
 		LogResults.BeginEvent(TEXT("Compile"));
 		
 		// TODO: sara-s remove once blueprint backend is replaced
-		// {
-		// 	EBlueprintCompileOptions CompileOptions = EBlueprintCompileOptions::None;
-		// 	// if( bSaveIntermediateBuildProducts )
-		// 	// {
-		// 	// 	CompileOptions |= EBlueprintCompileOptions::SaveIntermediateProducts;
-		// 	// }
-		//
-		// 	// if (bIsReparentingBlueprint)
-		// 	// {
-		// 	// 	CompileOptions |= (EBlueprintCompileOptions::UseDeltaSerializationDuringReinstancing | EBlueprintCompileOptions::SkipNewVariableDefaultsDetection);
-		// 	// }
-		//
-		// 	// If compilation is enabled during PIE/simulation, references to the CDO might be held by a script variable.
-		// 	// Thus, we set the flag to direct the compiler to allow those references to be replaced during reinstancing.
-		// 	if (GEditor->PlayWorld != nullptr)
-		// 	{
-		// 		CompileOptions |= EBlueprintCompileOptions::IncludeCDOInReferenceReplacement;
-		// 	}
-		//
-		// 	FKismetEditorUtilities::CompileBlueprint(BlueprintObj, CompileOptions, &LogResults);
-		// }
+		{
+			EBlueprintCompileOptions CompileOptions = EBlueprintCompileOptions::None;
+
+			// If compilation is enabled during PIE/simulation, references to the CDO might be held by a script variable.
+			// Thus, we set the flag to direct the compiler to allow those references to be replaced during reinstancing.
+			if (GEditor->PlayWorld != nullptr)
+			{
+				CompileOptions |= EBlueprintCompileOptions::IncludeCDOInReferenceReplacement;
+			}
+		
+			FKismetEditorUtilities::CompileBlueprint(BlueprintObj, CompileOptions, &LogResults);
+		}
 
 		LogResults.EndEvent();
 
diff --git a/Engine/Plugins/Runtime/RigVM/Source/RigVMEditor/Public/Editor/RigVMEditor.h b/Engine/Plugins/Runtime/RigVM/Source/RigVMEditor/Public/Editor/RigVMEditor.h
--- a/Engine/Plugins/Runtime/RigVM/Source/RigVMEditor/Public/Editor/RigVMEditor.h
+++ b/Engine/Plugins/Runtime/RigVM/Source/RigVMEditor/Public/Editor/RigVMEditor.h
@@ -184,6 +184,7 @@ public:
 	virtual void SetIsContextSensitive(const bool bIsContextSensitive) = 0;
 	virtual void RegisterToolbarTab(const TSharedRef<class FTabManager>& TabManager) = 0;
 	virtual const TArray< UObject* >* GetObjectsCurrentlyBeingEdited() const = 0;
+	virtual void AddCompileWidget(FToolBarBuilder& ToolbarBuilder) = 0;
 	virtual void AddSelectedDebugObjectWidget(FToolBarBuilder& ToolbarBuilder) = 0;
 	virtual FNotifyHook* GetNotifyHook() = 0;
 };
@@ -337,7 +338,10 @@ protected:
 	// FBlueprintEditor Interface
 	void CreateDefaultCommandsImpl();
 	void OnCreateGraphEditorCommands(TSharedPtr<FUICommandList> GraphEditorCommandsList);
+public:
 	void Compile();
+	bool IsCompilingEnabled() const;
+protected:
 	void SaveAsset_Execute();
 	void SaveAssetAs_Execute();
 	bool IsInAScriptingMode() const { return true; }
@@ -348,7 +352,6 @@ protected:
 	bool AreDelegatesAllowed() const;
 	bool NewDocument_IsVisibleForTypeImpl(ECreatedDocumentType GraphType) const;
 	FGraphAppearanceInfo GetGraphAppearance(class UEdGraph* InGraph) const;
-	bool IsCompilingEnabled() const;
 	void OnActiveTabChanged( TSharedPtr<SDockTab> PreviouslyActive, TSharedPtr<SDockTab> NewlyActivated );
 	void OnSelectedNodesChangedImpl(const TSet<class UObject*>& NewSelection);
 	void OnBlueprintChangedImpl(UBlueprint* InBlueprint, bool bIsJustBeingCompiled);
diff --git a/Engine/Plugins/Runtime/RigVM/Source/RigVMEditor/Public/Editor/RigVMLegacyEditor.h b/Engine/Plugins/Runtime/RigVM/Source/RigVMEditor/Public/Editor/RigVMLegacyEditor.h
--- a/Engine/Plugins/Runtime/RigVM/Source/RigVMEditor/Public/Editor/RigVMLegacyEditor.h
+++ b/Engine/Plugins/Runtime/RigVM/Source/RigVMEditor/Public/Editor/RigVMLegacyEditor.h
@@ -114,6 +114,7 @@ protected:
 	virtual void SetIsContextSensitive(const bool bIsContextSensitive) override { FBlueprintEditor::GetIsContextSensitive() = bIsContextSensitive; }
 	virtual void RegisterToolbarTab(const TSharedRef<FTabManager>& InTabManager) override { FBlueprintEditor::RegisterToolbarTab(InTabManager); }
 	virtual const TArray<UObject*>* GetObjectsCurrentlyBeingEdited() const override { return FBlueprintEditor::GetObjectsCurrentlyBeingEdited(); }
+	virtual void AddCompileWidget(FToolBarBuilder& ToolbarBuilder) override;
 	virtual void AddSelectedDebugObjectWidget(FToolBarBuilder& ToolbarBuilder) override;
 	virtual FString GetCustomDebugObjectLabel(UObject* ObjectBeingDebugged) const override { return FRigVMEditorBase::GetCustomDebugObjectLabel(ObjectBeingDebugged); }
 public:
diff --git a/Engine/Plugins/Runtime/RigVM/Source/RigVMEditor/Public/Editor/RigVMNewEditor.h b/Engine/Plugins/Runtime/RigVM/Source/RigVMEditor/Public/Editor/RigVMNewEditor.h
--- a/Engine/Plugins/Runtime/RigVM/Source/RigVMEditor/Public/Editor/RigVMNewEditor.h
+++ b/Engine/Plugins/Runtime/RigVM/Source/RigVMEditor/Public/Editor/RigVMNewEditor.h
@@ -2,6 +2,7 @@
 
 #pragma once
 
+#include "BlueprintEditorSettings.h"
 #include "RigVMEditor.h"
 #include "WorkflowOrientedApp/WorkflowCentricApplication.h"
 
@@ -114,6 +115,7 @@ public:
 	virtual void SetIsContextSensitive(const bool bIsContextSensitive) override { bIsActionMenuContextSensitive = bIsContextSensitive; }
 	virtual void RegisterToolbarTab(const TSharedRef<FTabManager>& InTabManager) override { FAssetEditorToolkit::RegisterTabSpawners(InTabManager); }
 	virtual const TArray<UObject*>* GetObjectsCurrentlyBeingEdited() const override { return FAssetEditorToolkit::GetObjectsCurrentlyBeingEdited(); }
+	virtual void AddCompileWidget(FToolBarBuilder& ToolbarBuilder) override;
 	virtual void AddSelectedDebugObjectWidget(FToolBarBuilder& ToolbarBuilder) override {}
 	virtual void Compile() { FRigVMEditorBase::Compile(); }
 	virtual void OnCreateGraphEditorCommands(TSharedPtr<FUICommandList> GraphEditorCommandsList) { FRigVMEditorBase::OnCreateGraphEditorCommands(GraphEditorCommandsList); }
@@ -137,6 +139,18 @@ public:
 	virtual void ClearSelectionStateFor(FName SelectionOwner);
 	void CreateDefaultTabContents(const TArray<UBlueprint*> InBlueprints);
 
+	TSharedRef<SWidget> GenerateCompileOptionsMenu();
+	void MakeSaveOnCompileSubMenu(FMenuBuilder& InMenu);
+	void SetSaveOnCompileSetting(ESaveOnCompile NewSetting);
+	bool IsSaveOnCompileEnabled() const;
+	bool IsSaveOnCompileOptionSet(TWeakPtr<FRigVMNewEditor> Editor, ESaveOnCompile Option);
+	void ToggleJumpToErrorNodeSetting();
+	bool IsJumpToErrorNodeOptionSet();
+	UEdGraphNode* FindNodeWithError(UBlueprint* Blueprint, EMessageSeverity::Type Severity);
+	UEdGraphNode* FindNodeWithError(FCompilerResultsLog const& ErrorLog, EMessageSeverity::Type Severity);
+	FText GetCompileStatusTooltip() const;
+	FSlateIcon GetCompileStatusImage() const;
+
 public:
 	static const FSlateBrush* GetGlyphForGraph(const UEdGraph* Graph, bool bInLargeIcon);
 
