Fix for RigVM removing spaces from default pin values

diff --git a/Engine/Plugins/Runtime/RigVM/Source/RigVM/Private/RigVMStringUtils.cpp b/Engine/Plugins/Runtime/RigVM/Source/RigVM/Private/RigVMStringUtils.cpp
--- a/Engine/Plugins/Runtime/RigVM/Source/RigVM/Private/RigVMStringUtils.cpp
+++ b/Engine/Plugins/Runtime/RigVM/Source/RigVM/Private/RigVMStringUtils.cpp
@@ -210,10 +210,20 @@ TArray<FString> RigVMStringUtils::SplitDefaultValue(const FString& InDefaultValu
 	int32 BraceCount = 0;
 	int32 QuoteCount = 0;
 
+	const int32 ContentLen = Content.Len();
+	FString CurrentPart;
+	CurrentPart.Reserve(ContentLen);
+
 	int32 LastPartStartIndex = 0;
-	for (int32 CharIndex = 0; CharIndex < Content.Len(); CharIndex++)
+	for (int32 CharIndex = 0; CharIndex < ContentLen; CharIndex++)
 	{
-		TCHAR Char = Content[CharIndex];
+		const TCHAR Char = Content[CharIndex];
+
+		if (Char == TCHAR(' ') && QuoteCount == 0) // ignore any white space that is not between quotes (i.e. not a pin/variable name)
+		{
+			continue;
+		}
+
 		if (QuoteCount > 0)
 		{
 			if (Char == TCHAR('"'))
@@ -243,10 +253,13 @@ TArray<FString> RigVMStringUtils::SplitDefaultValue(const FString& InDefaultValu
 		}
 		else if (Char == TCHAR(',') && BraceCount == 0 && QuoteCount == 0)
 		{
-			// ignore whitespaces
-			Parts.Add(Content.Mid(LastPartStartIndex, CharIndex - LastPartStartIndex).Replace(TEXT(" "), TEXT("")));
+			Parts.Add(CurrentPart);
 			LastPartStartIndex = CharIndex + 1;
+			CurrentPart.Reset();
+			continue;
 		}
+
+		CurrentPart.AppendChar(Char);
 	}
 
 	if (!Content.IsEmpty())
