Control Rig: Fix user defined struct not finding field names because CustomFindProperty does not work in game the same way it does in editor
#rnx

diff --git a/Engine/Plugins/Runtime/RigVM/Source/RigVM/Private/RigVMCore/RigVMPropertyPath.cpp b/Engine/Plugins/Runtime/RigVM/Source/RigVM/Private/RigVMCore/RigVMPropertyPath.cpp
--- a/Engine/Plugins/Runtime/RigVM/Source/RigVM/Private/RigVMCore/RigVMPropertyPath.cpp
+++ b/Engine/Plugins/Runtime/RigVM/Source/RigVM/Private/RigVMCore/RigVMPropertyPath.cpp
@@ -46,21 +46,17 @@ FRigVMPropertyPath::FRigVMPropertyPath(const FProperty* InProperty, const FStrin
 		{
 			Segment.Type = ERigVMPropertyPathSegmentType::StructMember;
 
-			const FProperty* MemberProperty = StructProperty->Struct->FindPropertyByName(*PathSegment);
-			if (!MemberProperty)
+			if(const FProperty* MemberProperty = StructProperty->Struct->FindPropertyByName(*PathSegment))
 			{
-				MemberProperty = StructProperty->Struct->CustomFindProperty(*PathSegment);
+				Segment.Name = MemberProperty->GetFName();
+				Segment.Index = MemberProperty->GetOffset_ForInternal();
+				Property = Segment.Property = MemberProperty;
 			}
-
-			if (!MemberProperty)
+			else
 			{
 				Segments.Empty();
 				return;
 			}
-
-			Segment.Name = MemberProperty->GetFName();
-			Segment.Index = MemberProperty->GetOffset_ForInternal();
-			Property = Segment.Property = MemberProperty;
 		}
 		else if(const FArrayProperty* ArrayProperty = CastField<FArrayProperty>(Property))
 		{
diff --git a/Engine/Plugins/Runtime/RigVM/Source/RigVMDeveloper/Private/RigVMModel/RigVMController.cpp b/Engine/Plugins/Runtime/RigVM/Source/RigVMDeveloper/Private/RigVMModel/RigVMController.cpp
--- a/Engine/Plugins/Runtime/RigVM/Source/RigVMDeveloper/Private/RigVMModel/RigVMController.cpp
+++ b/Engine/Plugins/Runtime/RigVM/Source/RigVMDeveloper/Private/RigVMModel/RigVMController.cpp
@@ -453,7 +453,7 @@ FRigVMPinInfo::FRigVMPinInfo(const URigVMPin* InPin, int32 InParentIndex, ERigVM
 				const FString SubPinPath = InPin->GetSubPinPath(Pin);
 				if (const FProperty* StructProperty = UDStruct->FindPropertyByName(*SubPinPath))
 				{
-					Name = *UDStruct->GetAuthoredNameForField(StructProperty);
+					DisplayName = *UDStruct->GetAuthoredNameForField(StructProperty);
 					break;
 				}
 			}
@@ -490,7 +490,7 @@ FRigVMPinInfo::FRigVMPinInfo(FProperty* InProperty, ERigVMPinDirection InDirecti
 
 	if (const UUserDefinedStruct* UDStruct = Cast<const UUserDefinedStruct>(Property->GetOwnerStruct()))
 	{
-		Name = *UDStruct->GetAuthoredNameForField(InProperty);
+		DisplayName = *UDStruct->GetAuthoredNameForField(InProperty);
 	}
 
 	if (Direction == ERigVMPinDirection::Invalid)
@@ -16951,10 +16951,6 @@ void URigVMController::AddPinsForStruct(UStruct* InStruct, URigVMNode* InNode, U
 			}
 			
 			FName PropertyName = It->GetFName();
-			if (UUserDefinedStruct* UDStruct = Cast<UUserDefinedStruct>(StructToVisit))
-			{
-				PropertyName = *UDStruct->GetAuthoredNameForField(*It);
-			}
 
 			URigVMPin* Pin = NewObject<URigVMPin>(InParentPin == nullptr ? Cast<UObject>(InNode) : Cast<UObject>(InParentPin), PropertyName);
 			ConfigurePinFromProperty(*It, Pin, InPinDirection);
@@ -23163,9 +23159,15 @@ FRigVMClientPatchResult URigVMController::PatchUserDefinedStructPinNames()
 					for (URigVMPin* SubPin : Pin->GetSubPins())
 					{
 						FString SubPinPath = SubPin->GetSubPinPath(Pin);
-						if (const FProperty* StructSubProperty = UDStruct->FindPropertyByName(*SubPinPath))
+						const FProperty* StructSubProperty = UDStruct->FindPropertyByName(*SubPinPath);
+						if (!StructSubProperty)
+						{
+							// Patch UserDefinedStruct subpins that are named after their display name
+							StructSubProperty = UDStruct->CustomFindProperty(*SubPinPath);
+						}
+						if (StructSubProperty)
 						{
-							FString NewNameSubPin = *UDStruct->GetAuthoredNameForField(StructSubProperty);
+							FString NewNameSubPin = *StructSubProperty->GetName();
 							if (!NewNameSubPin.Equals(SubPin->GetName(), ESearchCase::IgnoreCase))
 							{
 								Result.bChangedContent = true;
@@ -23190,6 +23192,13 @@ FRigVMClientPatchResult URigVMController::PatchUserDefinedStructPinNames()
 									Link->UpdatePinPaths();
 								}
 							}
+
+							FName DisplayName = *StructSubProperty->GetAuthoredName();
+							if (SubPin->GetDisplayName() != DisplayName)
+							{
+								Result.bChangedContent = true;
+								SubPin->SetDisplayName(DisplayName);
+							}
 						}
 					}
 				}
