RigVM: Fix crash when copy/pasting node
#rnx

diff --git a/Engine/Plugins/Runtime/RigVM/Source/RigVMDeveloper/Private/RigVMModel/RigVMControllerActions.cpp b/Engine/Plugins/Runtime/RigVM/Source/RigVMDeveloper/Private/RigVMModel/RigVMControllerActions.cpp
--- a/Engine/Plugins/Runtime/RigVM/Source/RigVMDeveloper/Private/RigVMModel/RigVMControllerActions.cpp
+++ b/Engine/Plugins/Runtime/RigVM/Source/RigVMDeveloper/Private/RigVMModel/RigVMControllerActions.cpp
@@ -91,14 +91,16 @@ bool URigVMActionStack::CloseUndoBracket(URigVMController* InController)
 	{
 		return true;
 	}
-	ensure(BracketActions.Num() > 0);
-	if(BracketActions.Last()->IsEmpty())
+	if (ensure(BracketActions.Num() > 0))
 	{
-		return CancelUndoBracket(InController);
+		if(BracketActions.Last()->IsEmpty())
+		{
+			return CancelUndoBracket(InController);
+		}
+		FRigVMBaseAction* Action = BracketActions.Pop();
+		EndAction(*Action);
+		delete(Action);
 	}
-	FRigVMBaseAction* Action = BracketActions.Pop();
-	EndAction(*Action);
-	delete(Action);
 	return true;
 }
 
diff --git a/Engine/Plugins/Runtime/RigVM/Source/RigVMEditor/Private/Editor/RigVMEditorTools.cpp b/Engine/Plugins/Runtime/RigVM/Source/RigVMEditor/Private/Editor/RigVMEditorTools.cpp
--- a/Engine/Plugins/Runtime/RigVM/Source/RigVMEditor/Private/Editor/RigVMEditorTools.cpp
+++ b/Engine/Plugins/Runtime/RigVM/Source/RigVMEditor/Private/Editor/RigVMEditorTools.cpp
@@ -28,7 +28,7 @@ RIGVMEDITOR_API bool PasteNodes(const FVector2D& PasteLocation
 		InFocusedController->OpenUndoBracket(TEXT("Paste Nodes."));
 	}
 
-	TArray<FName> NodeNamesCreated = ImportNodesFromText(PasteLocation, TextToImport, InFocusedController, InFocusedModel, InLocalFunctionLibrary, InGraphFunctionHost, bPrintPythonCommands);
+	TArray<FName> NodeNamesCreated = ImportNodesFromText(PasteLocation, TextToImport, InFocusedController, InFocusedModel, InLocalFunctionLibrary, InGraphFunctionHost, true, bPrintPythonCommands);
 
 	const bool bPastePerformed = NodeNamesCreated.Num() > 0;
 	if (bPastePerformed)
