RigVM: Allowed event nodes to vary their event name based around the node's defaults

Creating a node when setting defaults now sets the defaults when the node is first created, rather than afterwards. This allows the call to GetEventName to vary based on the defaults provided to AddUnitNodeWithDefaults.
Basically all this does is make AddUnitNode call AddUnitNodeWithDefaults, which takes on the old AddUnitNode implemewntation with the slight differences that it uses the passed-in struct defaults (if valid) rather than always creating an instance inline.

Intended for UE5

diff --git a/Engine/Plugins/Runtime/RigVM/Source/RigVMDeveloper/Private/RigVMModel/RigVMController.cpp b/Engine/Plugins/Runtime/RigVM/Source/RigVMDeveloper/Private/RigVMModel/RigVMController.cpp
--- a/Engine/Plugins/Runtime/RigVM/Source/RigVMDeveloper/Private/RigVMModel/RigVMController.cpp
+++ b/Engine/Plugins/Runtime/RigVM/Source/RigVMDeveloper/Private/RigVMModel/RigVMController.cpp
@@ -1474,5 +1474,86 @@
 
 URigVMUnitNode* URigVMController::AddUnitNode(UScriptStruct* InScriptStruct, TSubclassOf<URigVMUnitNode> InUnitNodeClass, const FName& InMethodName, const FVector2D& InPosition, const FString& InNodeName, bool bSetupUndoRedo, bool bPrintPythonCommand)
 {
+	return AddUnitNodeWithDefaults(InScriptStruct, InUnitNodeClass, FRigStructScope(), InMethodName, InPosition, InNodeName, bSetupUndoRedo, bPrintPythonCommand);
+}
+
+const UClass* URigVMController::GetNodeClassForTemplate(FRigVMTemplate* InTemplate) const
+{
+	// determine what kind of node we need to create
+	if(InTemplate->UsesDispatch())
+	{
+		return URigVMDispatchNode::StaticClass();
+	}
+	if(const FRigVMFunction* FirstFunction = InTemplate->GetOrCreatePermutation(0))
+	{
+		const UScriptStruct* PotentialUnitStruct = FirstFunction->Struct;
+		if(PotentialUnitStruct && PotentialUnitStruct->IsChildOf(FRigVMStruct::StaticStruct()))
+		{
+			return URigVMUnitNode::StaticClass();
+		}
+	}
+	return nullptr;
+}
+
+URigVMUnitNode* URigVMController::AddUnitNodeFromStructPath(const FString& InScriptStructPath, const FName& InMethodName, const FVector2D& InPosition, const FString& InNodeName, bool bSetupUndoRedo, bool bPrintPythonCommand)
+{
 	if(!IsValidGraph())
 	{
+		return nullptr;
+	}
+
+	if (!bIsTransacting && !IsGraphEditable())
+	{
+		return nullptr;
+	}
+
+	UScriptStruct* ScriptStruct = RigVMTypeUtils::FindObjectFromCPPTypeObjectPath<UScriptStruct>(InScriptStructPath);
+	if (ScriptStruct == nullptr)
+	{
+		ReportErrorf(TEXT("Cannot find struct for path '%s'."), *InScriptStructPath);
+		return nullptr;
+	}
+
+	return AddUnitNode(ScriptStruct, InMethodName, InPosition, InNodeName, bSetupUndoRedo, bPrintPythonCommand);
+}
+
+URigVMUnitNode* URigVMController::AddUnitNodeWithDefaults(UScriptStruct* InScriptStruct, const FString& InDefaults,
+	const FName& InMethodName, const FVector2D& InPosition, const FString& InNodeName, bool bSetupUndoRedo,
+	bool bPrintPythonCommand)
+{
+	if(InScriptStruct == nullptr)
+	{
+		return nullptr;
+	}
+
+	FStructOnScope StructOnScope;
+
+	if(!InDefaults.IsEmpty())
+	{
+		StructOnScope = FStructOnScope(InScriptStruct);
+		
+		FRigVMPinDefaultValueImportErrorContext ErrorPipe;
+		InScriptStruct->ImportText(*InDefaults, StructOnScope.GetStructMemory(), nullptr, PPF_None, &ErrorPipe, FString());
+
+		if(ErrorPipe.NumErrors > 0)
+		{
+			return nullptr;
+		}
+	}
+
+	return AddUnitNodeWithDefaults(InScriptStruct, StructOnScope, InMethodName, InPosition, InNodeName, bSetupUndoRedo, bPrintPythonCommand);
+}
+
+URigVMUnitNode* URigVMController::AddUnitNodeWithDefaults(UScriptStruct* InScriptStruct, const FRigStructScope& InDefaults,
+														  const FName& InMethodName, const FVector2D& InPosition, const FString& InNodeName, bool bSetupUndoRedo,
+														  bool bPrintPythonCommand)
+{
+	return AddUnitNodeWithDefaults(InScriptStruct, URigVMUnitNode::StaticClass(), InDefaults, InMethodName, InPosition, InNodeName, bSetupUndoRedo, bPrintPythonCommand);
+}
+
+URigVMUnitNode* URigVMController::AddUnitNodeWithDefaults(UScriptStruct* InScriptStruct, TSubclassOf<URigVMUnitNode> InUnitNodeClass, const FRigStructScope& InDefaults,
+                                                          const FName& InMethodName, const FVector2D& InPosition, const FString& InNodeName, bool bSetupUndoRedo,
+                                                          bool bPrintPythonCommand)
+{
+	if(!IsValidGraph())
+	{
@@ -1589,8 +1670,18 @@
 		return TemplateNode;
 	}
 
+	// Use the defaults if they are passed
+	FRigVMStruct* StructMemory = nullptr;
 	FStructOnScope StructOnScope(InScriptStruct);
-	FRigVMStruct* StructMemory = (FRigVMStruct*)StructOnScope.GetStructMemory();
+	if (InDefaults.IsValid())
+	{
+		StructMemory = (FRigVMStruct*)InDefaults.GetMemory();
+	}
+	else
+	{
+		StructMemory = (FRigVMStruct*)StructOnScope.GetStructMemory();
+	}
+	
 	const bool bIsEventNode = (!StructMemory->GetEventName().IsNone());
 	if (bIsEventNode)
 	{
@@ -1637,7 +1728,14 @@
 	}
 
 	FString ExportedDefaultValue;
-	CreateDefaultValueForStructIfRequired(InScriptStruct, ExportedDefaultValue);
+	if (InDefaults.IsValid())
 	{
+		InDefaults.GetScriptStruct()->ExportText(ExportedDefaultValue, InDefaults.GetMemory(), nullptr, nullptr, PPF_None, nullptr, true);
+	}
+	else
+	{
+		CreateDefaultValueForStructIfRequired(InScriptStruct, ExportedDefaultValue);
+	}
+	{
 		FRigVMDefaultValueTypeGuard _(this, ERigVMPinDefaultValueType::Unset, true);
 		TGuardValue<bool> SuspendNotifications(bSuspendNotifications, true);
@@ -1692,111 +1790,6 @@
 	return Node;
 }
 
-const UClass* URigVMController::GetNodeClassForTemplate(FRigVMTemplate* InTemplate) const
-{
-	// determine what kind of node we need to create
-	if(InTemplate->UsesDispatch())
-	{
-		return URigVMDispatchNode::StaticClass();
-	}
-	if(const FRigVMFunction* FirstFunction = InTemplate->GetOrCreatePermutation(0))
-	{
-		const UScriptStruct* PotentialUnitStruct = FirstFunction->Struct;
-		if(PotentialUnitStruct && PotentialUnitStruct->IsChildOf(FRigVMStruct::StaticStruct()))
-		{
-			return URigVMUnitNode::StaticClass();
-		}
-	}
-	return nullptr;
-}
-
-URigVMUnitNode* URigVMController::AddUnitNodeFromStructPath(const FString& InScriptStructPath, const FName& InMethodName, const FVector2D& InPosition, const FString& InNodeName, bool bSetupUndoRedo, bool bPrintPythonCommand)
-{
-	if(!IsValidGraph())
-	{
-		return nullptr;
-	}
-
-	if (!bIsTransacting && !IsGraphEditable())
-	{
-		return nullptr;
-	}
-
-	UScriptStruct* ScriptStruct = RigVMTypeUtils::FindObjectFromCPPTypeObjectPath<UScriptStruct>(InScriptStructPath);
-	if (ScriptStruct == nullptr)
-	{
-		ReportErrorf(TEXT("Cannot find struct for path '%s'."), *InScriptStructPath);
-		return nullptr;
-	}
-
-	return AddUnitNode(ScriptStruct, InMethodName, InPosition, InNodeName, bSetupUndoRedo, bPrintPythonCommand);
-}
-
-URigVMUnitNode* URigVMController::AddUnitNodeWithDefaults(UScriptStruct* InScriptStruct, const FString& InDefaults,
-	const FName& InMethodName, const FVector2D& InPosition, const FString& InNodeName, bool bSetupUndoRedo,
-	bool bPrintPythonCommand)
-{
-	if(InScriptStruct == nullptr)
-	{
-		return nullptr;
-	}
-
-	FStructOnScope StructOnScope;
-
-	if(!InDefaults.IsEmpty())
-	{
-		StructOnScope = FStructOnScope(InScriptStruct);
-		
-		FRigVMPinDefaultValueImportErrorContext ErrorPipe;
-		InScriptStruct->ImportText(*InDefaults, StructOnScope.GetStructMemory(), nullptr, PPF_None, &ErrorPipe, FString());
-
-		if(ErrorPipe.NumErrors > 0)
-		{
-			return nullptr;
-		}
-	}
-
-	return AddUnitNodeWithDefaults(InScriptStruct, StructOnScope, InMethodName, InPosition, InNodeName, bSetupUndoRedo, bPrintPythonCommand);
-}
-
-URigVMUnitNode* URigVMController::AddUnitNodeWithDefaults(UScriptStruct* InScriptStruct, const FRigStructScope& InDefaults,
-                                                          const FName& InMethodName, const FVector2D& InPosition, const FString& InNodeName, bool bSetupUndoRedo,
-                                                          bool bPrintPythonCommand)
-{
-	if(InScriptStruct == nullptr)
-	{
-		return nullptr;
-	}
-
-	const bool bSetPinDefaults = InDefaults.IsValid() && (InDefaults.GetScriptStruct() == InScriptStruct); 
-	if(bSetPinDefaults)
-	{
-		static constexpr TCHAR AddUnitNodeTitle[] = TEXT("Add Unit Node");
-		OpenUndoBracket(AddUnitNodeTitle);
-	}
-
-	URigVMUnitNode* Node = AddUnitNode(InScriptStruct, InMethodName, InPosition, InNodeName, bSetupUndoRedo, bPrintPythonCommand);
-	if(Node == nullptr)
-	{
-		if(bSetPinDefaults)
-		{
-			CancelUndoBracket();
-		}
-		return nullptr;
-	}
-
-	if(bSetPinDefaults)
-	{
-		if(!SetUnitNodeDefaults(Node, InDefaults))
-		{
-			CancelUndoBracket();
-		}
-	}
-
-	CloseUndoBracket();
-	return Node;
-}
-
 bool URigVMController::SetUnitNodeDefaults(URigVMUnitNode* InNode, const FString& InDefaults, bool bSetupUndoRedo,
 	bool bPrintPythonCommand)
 {
diff --git a/Engine/Plugins/Runtime/RigVM/Source/RigVMDeveloper/Public/RigVMModel/RigVMController.h b/Engine/Plugins/Runtime/RigVM/Source/RigVMDeveloper/Public/RigVMModel/RigVMController.h
--- a/Engine/Plugins/Runtime/RigVM/Source/RigVMDeveloper/Public/RigVMModel/RigVMController.h
+++ b/Engine/Plugins/Runtime/RigVM/Source/RigVMDeveloper/Public/RigVMModel/RigVMController.h
@@ -373,6 +373,7 @@ public:
 	// UnitNode represent a RIGVM_METHOD declaration on a USTRUCT.
 	// This causes a NodeAdded modified event.
 	URigVMUnitNode* AddUnitNodeWithDefaults(UScriptStruct* InScriptStruct, const FRigStructScope& InDefaults, const FName& InMethodName = TEXT("Execute"), const FVector2D& InPosition = FVector2D::ZeroVector, const FString& InNodeName = TEXT(""), bool bSetupUndoRedo = true, bool bPrintPythonCommand = false);
+	URigVMUnitNode* AddUnitNodeWithDefaults(UScriptStruct* InScriptStruct, TSubclassOf<URigVMUnitNode> InUnitNodeClass, const FRigStructScope& InDefaults, const FName& InMethodName = TEXT("Execute"), const FVector2D& InPosition = FVector2D::ZeroVector, const FString& InNodeName = TEXT(""), bool bSetupUndoRedo = true, bool bPrintPythonCommand = false);
 
 	// Adds a Function / Struct Node to the edited Graph.
 	// UnitNode represent a RIGVM_METHOD declaration on a USTRUCT.
