UE: Remove recent warning about finding duplicate event nodes as there seem to be many of them in existing content

diff --git a/Engine/Plugins/Runtime/RigVM/Source/RigVMDeveloper/Private/RigVMCompiler/RigVMCompiler.cpp b/Engine/Plugins/Runtime/RigVM/Source/RigVMDeveloper/Private/RigVMCompiler/RigVMCompiler.cpp
--- a/Engine/Plugins/Runtime/RigVM/Source/RigVMDeveloper/Private/RigVMCompiler/RigVMCompiler.cpp
+++ b/Engine/Plugins/Runtime/RigVM/Source/RigVMDeveloper/Private/RigVMCompiler/RigVMCompiler.cpp
@@ -651,8 +651,7 @@ bool URigVMCompiler::Compile(const FRigVMCompileSettings& InSettings, TArray<URi
 			// that is, make sure their cpp type and type object agree with each other
 			VisitedGraphController->EnsureLocalVariableValidity();
 		}
-
-		TSet<FName> FoundEvents;
+		
 		for(URigVMNode* ModelNode : VisitedGraph->GetNodes())
 		{
 			URigVMController* VisitedGraphController = InController->GetControllerForGraph(VisitedGraph);
@@ -673,17 +672,6 @@ bool URigVMCompiler::Compile(const FRigVMCompileSettings& InSettings, TArray<URi
 				}
 			}
 
-			if (ModelNode->IsEvent())
-			{
-				if (FoundEvents.Contains(ModelNode->GetEventName()))
-				{
-					static const FString LinkedMessage = TEXT("Duplicate event node @@ found.");
-					Settings.ASTSettings.Report(EMessageSeverity::Warning, ModelNode, LinkedMessage);
-				}
-
-				FoundEvents.Add(ModelNode->GetEventName());
-			}
-
 			if(ModelNode->IsA<UDEPRECATED_RigVMBranchNode>())
 			{
 				static const FString LinkedMessage = TEXT("Node @@ is a deprecated branch node. Cannot compile.");
