AnimNext: Fix crash during automated test

Nodes that are not yet fully resolved can return nullptr from GetScriptStruct. This change guards against that.

diff --git a/Engine/Plugins/Runtime/RigVM/Source/RigVMDeveloper/Private/RigVMModel/RigVMClient.cpp b/Engine/Plugins/Runtime/RigVM/Source/RigVMDeveloper/Private/RigVMModel/RigVMClient.cpp
--- a/Engine/Plugins/Runtime/RigVM/Source/RigVMDeveloper/Private/RigVMModel/RigVMClient.cpp
+++ b/Engine/Plugins/Runtime/RigVM/Source/RigVMDeveloper/Private/RigVMModel/RigVMClient.cpp
@@ -760,7 +760,8 @@ TArray<FName> FRigVMClient::GetEntryNames(UScriptStruct* InUnitScriptStructFilte
 			{
 				if(const URigVMUnitNode* UnitNode = Cast<URigVMUnitNode>(Node))
 				{
-					if(!UnitNode->GetScriptStruct()->IsChildOf(InUnitScriptStructFilter))
+					UScriptStruct* ScriptStruct = UnitNode->GetScriptStruct();
+					if(ScriptStruct == nullptr || !ScriptStruct->IsChildOf(InUnitScriptStructFilter))
 					{
 						continue;
 					}
