RigVM: Fix defaults not being set for template nodes

Additional fix for automated test.
CL 40703363 did not take into account the early-out path for template nodes. This re-implements the code path to set the defaults from a separate call to SetUnitNodeDefaults as there is no AddTemplateNodeWithDefaults.

diff --git a/Engine/Plugins/Runtime/RigVM/Source/RigVMDeveloper/Private/RigVMModel/RigVMController.cpp b/Engine/Plugins/Runtime/RigVM/Source/RigVMDeveloper/Private/RigVMModel/RigVMController.cpp
--- a/Engine/Plugins/Runtime/RigVM/Source/RigVMDeveloper/Private/RigVMModel/RigVMController.cpp
+++ b/Engine/Plugins/Runtime/RigVM/Source/RigVMDeveloper/Private/RigVMModel/RigVMController.cpp
@@ -1661,7 +1661,16 @@ URigVMUnitNode* URigVMController::AddUnitNodeWithDefaults(UScriptStruct* InScrip
 				StructInstance->OnUnitNodeCreated(UnitNodeCreatedContext);
 			}
 		}
-		
+
+		if(InDefaults.IsValid() && (InDefaults.GetScriptStruct() == InScriptStruct))
+		{
+			if (!SetUnitNodeDefaults(TemplateNode, InDefaults, bSetupUndoRedo, bPrintPythonCommand))
+			{
+				CancelUndoBracket();
+				return nullptr;
+			}
+		}
+
 		if(bSetupUndoRedo)
 		{
 			CloseUndoBracket();
