RigVM: Cleanup IRigVMEditor and FRigVMEditorBase accessibility
#rnx

diff --git a/Engine/Plugins/Animation/ControlRig/Source/ControlRigEditor/Private/Editor/RigValidationTabSummoner.cpp b/Engine/Plugins/Animation/ControlRig/Source/ControlRigEditor/Private/Editor/RigValidationTabSummoner.cpp
--- a/Engine/Plugins/Animation/ControlRig/Source/ControlRigEditor/Private/Editor/RigValidationTabSummoner.cpp
+++ b/Engine/Plugins/Animation/ControlRig/Source/ControlRigEditor/Private/Editor/RigValidationTabSummoner.cpp
@@ -32,7 +32,7 @@ TSharedRef<SWidget> FRigValidationTabSummoner::CreateTabBody(const FWorkflowTabS
 	
 	if (TSharedPtr<FRigVMEditorBase> Editor = EditorPtr)
 	{
-		UControlRigBlueprint* RigBlueprint = Cast<UControlRigBlueprint>(Editor->FRigVMEditorBase::GetBlueprintObj());
+		UControlRigBlueprint* RigBlueprint = Cast<UControlRigBlueprint>(Editor->GetRigVMBlueprint());
 		check(RigBlueprint);
 	
 		UControlRigValidator* Validator = RigBlueprint->Validator;
diff --git a/Engine/Plugins/Runtime/RigVM/Source/RigVMEditor/Private/Editor/RigVMEditor.cpp b/Engine/Plugins/Runtime/RigVM/Source/RigVMEditor/Private/Editor/RigVMEditor.cpp
--- a/Engine/Plugins/Runtime/RigVM/Source/RigVMEditor/Private/Editor/RigVMEditor.cpp
+++ b/Engine/Plugins/Runtime/RigVM/Source/RigVMEditor/Private/Editor/RigVMEditor.cpp
@@ -558,7 +558,7 @@ void FRigVMEditorBase::Tick(float DeltaTime)
 		if (LastDebuggedHost != GetCustomDebugObjectLabel(Blueprint->GetObjectBeingDebugged()))
 		{
 			TArray<FRigVMCustomDebugObject> DebugList;
-			GetCustomDebugObjects(DebugList);
+			GetDebugObjects(DebugList);
 
 			for (const FRigVMCustomDebugObject& DebugObject : DebugList)
 			{
@@ -1111,7 +1111,7 @@ void FRigVMEditorBase::Compile()
 		}
 
 		TArray<FRigVMCustomDebugObject> DebugList;
-		GetCustomDebugObjects(DebugList);
+		GetDebugObjects(DebugList);
 
 		for (const FRigVMCustomDebugObject& DebugObject : DebugList)
 		{
@@ -1281,7 +1281,7 @@ void FRigVMEditorBase::OnBlueprintChangedImpl(UBlueprint* InBlueprint, bool bIsJ
 			if (!LastDebuggedHost.IsEmpty())
 			{
 				TArray<FRigVMCustomDebugObject> DebugList;
-				GetCustomDebugObjects(DebugList);
+				GetDebugObjects(DebugList);
 
 				for (const FRigVMCustomDebugObject& DebugObject : DebugList)
 				{
@@ -1296,7 +1296,7 @@ void FRigVMEditorBase::OnBlueprintChangedImpl(UBlueprint* InBlueprint, bool bIsJ
 		}
 	}}
 
-void FRigVMEditorBase::RefreshEditors(ERefreshRigVMEditorReason::Type Reason)
+void FRigVMEditorBase::ForceEditorRefresh(ERefreshRigVMEditorReason::Type Reason)
 {
 	if(Reason == ERefreshRigVMEditorReason::UnknownReason)
 	{
@@ -1546,7 +1546,7 @@ void FRigVMEditorBase::RedoAction()
 	GEditor->RedoTransaction();
 }
 
-void FRigVMEditorBase::NewDocument_OnClicked(ECreatedDocumentType GraphType)
+void FRigVMEditorBase::OnNewDocumentClicked(ECreatedDocumentType GraphType)
 {
 	if (GraphType == FRigVMNewEditor::CGT_NewFunctionGraph)
 	{
@@ -2177,7 +2177,7 @@ void FRigVMEditorBase::OnWrappedPropertyChangedChainEvent(URigVMDetailsViewWrapp
 					}
 				}
 
-				RefreshEditors();
+				ForceEditorRefresh();
 				GetRigVMBlueprint()->RequestAutoVMRecompilation();
 			}
 			else if (PropertyPath == TEXT("CPPType") || PropertyPath == TEXT("CPPTypeObject"))
@@ -4004,7 +4004,7 @@ FSlateIcon FRigVMEditorBase::GetExecutionModeIcon() const
 	return GetExecutionModeIcon(ExecutionMode);
 }
 
-void FRigVMEditorBase::GetCustomDebugObjects(TArray<FRigVMCustomDebugObject>& DebugList) const
+void FRigVMEditorBase::GetDebugObjects(TArray<FRigVMCustomDebugObject>& DebugList) const
 {
 	URigVMBlueprint* RigVMBlueprint = GetRigVMBlueprint();
 	if (RigVMBlueprint == nullptr)
diff --git a/Engine/Plugins/Runtime/RigVM/Source/RigVMEditor/Private/Editor/RigVMEditorGraphExplorerTabSummoner.cpp b/Engine/Plugins/Runtime/RigVM/Source/RigVMEditor/Private/Editor/RigVMEditorGraphExplorerTabSummoner.cpp
--- a/Engine/Plugins/Runtime/RigVM/Source/RigVMEditor/Private/Editor/RigVMEditorGraphExplorerTabSummoner.cpp
+++ b/Engine/Plugins/Runtime/RigVM/Source/RigVMEditor/Private/Editor/RigVMEditorGraphExplorerTabSummoner.cpp
@@ -6,7 +6,7 @@
 
 #define LOCTEXT_NAMESPACE "RigVMEditorGraphExplorerTabSummoner"
 
-FRigVMEditorGraphExplorerTabSummoner::FRigVMEditorGraphExplorerTabSummoner(const TSharedRef<FRigVMEditorBase>& InRigVMEditor)
+FRigVMEditorGraphExplorerTabSummoner::FRigVMEditorGraphExplorerTabSummoner(const TSharedRef<IRigVMEditor>& InRigVMEditor)
 	: FWorkflowTabFactory(TabID(), InRigVMEditor.Get().GetHostingApp())
 	, RigVMEditor(InRigVMEditor)
 {
@@ -21,7 +21,7 @@ FRigVMEditorGraphExplorerTabSummoner::FRigVMEditorGraphExplorerTabSummoner(const
 
 TSharedRef<SWidget> FRigVMEditorGraphExplorerTabSummoner::CreateTabBody(const FWorkflowTabSpawnInfo& Info) const
 {
-	if (TSharedPtr<FRigVMEditorBase> Editor = RigVMEditor.Pin())
+	if (TSharedPtr<IRigVMEditor> Editor = RigVMEditor.Pin())
 	{
 		TSharedRef<SRigVMEditorGraphExplorer> Explorer = SNew(SRigVMEditorGraphExplorer, Editor);
 		Editor->SetGraphExplorerWidget(Explorer);
diff --git a/Engine/Plugins/Runtime/RigVM/Source/RigVMEditor/Private/Editor/RigVMEditorMenuContext.cpp b/Engine/Plugins/Runtime/RigVM/Source/RigVMEditor/Private/Editor/RigVMEditorMenuContext.cpp
--- a/Engine/Plugins/Runtime/RigVM/Source/RigVMEditor/Private/Editor/RigVMEditorMenuContext.cpp
+++ b/Engine/Plugins/Runtime/RigVM/Source/RigVMEditor/Private/Editor/RigVMEditorMenuContext.cpp
@@ -8,7 +8,7 @@
 
 #include UE_INLINE_GENERATED_CPP_BY_NAME(RigVMEditorMenuContext)
 
-void URigVMEditorMenuContext::Init(TWeakPtr<FRigVMEditorBase> InRigVMEditor, const FRigVMEditorGraphMenuContext& InGraphMenuContext)
+void URigVMEditorMenuContext::Init(TWeakPtr<IRigVMEditor> InRigVMEditor, const FRigVMEditorGraphMenuContext& InGraphMenuContext)
 {
 	WeakRigVMEditor = InRigVMEditor;
 	GraphMenuContext = InGraphMenuContext;
@@ -16,7 +16,7 @@ void URigVMEditorMenuContext::Init(TWeakPtr<FRigVMEditorBase> InRigVMEditor, con
 
 URigVMBlueprint* URigVMEditorMenuContext::GetRigVMBlueprint() const
 {
-	if (const TSharedPtr<FRigVMEditorBase> Editor = WeakRigVMEditor.Pin())
+	if (const TSharedPtr<IRigVMEditor> Editor = WeakRigVMEditor.Pin())
 	{
 		return Editor->GetRigVMBlueprint();
 	}
@@ -46,9 +46,9 @@ FRigVMEditorGraphMenuContext URigVMEditorMenuContext::GetGraphMenuContext()
 	return GraphMenuContext;
 }
 
-FRigVMEditorBase* URigVMEditorMenuContext::GetRigVMEditor() const
+IRigVMEditor* URigVMEditorMenuContext::GetRigVMEditor() const
 {
-	if (const TSharedPtr<FRigVMEditorBase> Editor = WeakRigVMEditor.Pin())
+	if (const TSharedPtr<IRigVMEditor> Editor = WeakRigVMEditor.Pin())
 	{
 		return Editor.Get();
 	}
diff --git a/Engine/Plugins/Runtime/RigVM/Source/RigVMEditor/Private/Editor/RigVMExecutionStackTabSummoner.cpp b/Engine/Plugins/Runtime/RigVM/Source/RigVMEditor/Private/Editor/RigVMExecutionStackTabSummoner.cpp
--- a/Engine/Plugins/Runtime/RigVM/Source/RigVMEditor/Private/Editor/RigVMExecutionStackTabSummoner.cpp
+++ b/Engine/Plugins/Runtime/RigVM/Source/RigVMEditor/Private/Editor/RigVMExecutionStackTabSummoner.cpp
@@ -6,7 +6,7 @@
 
 #define LOCTEXT_NAMESPACE "RigVMExecutionStackTabSummoner"
 
-FRigVMExecutionStackTabSummoner::FRigVMExecutionStackTabSummoner(const TSharedRef<FRigVMEditorBase>& InRigVMEditor)
+FRigVMExecutionStackTabSummoner::FRigVMExecutionStackTabSummoner(const TSharedRef<IRigVMEditor>& InRigVMEditor)
 	: FWorkflowTabFactory(TabID, InRigVMEditor->GetHostingApp())
 	, RigVMEditor(InRigVMEditor)
 {
diff --git a/Engine/Plugins/Runtime/RigVM/Source/RigVMEditor/Private/Editor/RigVMGraphDetailCustomization.cpp b/Engine/Plugins/Runtime/RigVM/Source/RigVMEditor/Private/Editor/RigVMGraphDetailCustomization.cpp
--- a/Engine/Plugins/Runtime/RigVM/Source/RigVMEditor/Private/Editor/RigVMGraphDetailCustomization.cpp
+++ b/Engine/Plugins/Runtime/RigVM/Source/RigVMEditor/Private/Editor/RigVMGraphDetailCustomization.cpp
@@ -46,7 +46,7 @@ static const FText RigVMGraphDetailCustomizationMultipleValues = LOCTEXT("Multip
 FRigVMFunctionArgumentGroupLayout::FRigVMFunctionArgumentGroupLayout(
 	const TWeakObjectPtr<URigVMGraph>& InGraph,
 	const TWeakInterfacePtr<IRigVMClientHost>& InRigVMClientHost,
-	const TWeakPtr<FRigVMEditorBase>& InEditor,
+	const TWeakPtr<IRigVMEditor>& InEditor,
 	bool bInputs)
 	: GraphPtr(InGraph)
 	, WeakRigVMClientHost(InRigVMClientHost)
@@ -153,7 +153,7 @@ void FRigVMFunctionArgumentGroupLayout::HandleModifiedEvent(ERigVMGraphNotifType
 class FRigVMFunctionArgumentPinTypeSelectorFilter : public IPinTypeSelectorFilter
 {
 public:
-	FRigVMFunctionArgumentPinTypeSelectorFilter(const TWeakPtr<FRigVMEditorBase>& InRigVMEditor, const TWeakObjectPtr<URigVMGraph>& InGraph)
+	FRigVMFunctionArgumentPinTypeSelectorFilter(const TWeakPtr<IRigVMEditor>& InRigVMEditor, const TWeakObjectPtr<URigVMGraph>& InGraph)
 		: RigVMEditorPtr(InRigVMEditor), GraphPtr(InGraph)
 	{
 	}
@@ -208,7 +208,7 @@ public:
 
 private:
 
-	TWeakPtr<FRigVMEditorBase> RigVMEditorPtr;
+	TWeakPtr<IRigVMEditor> RigVMEditorPtr;
 	
 	TWeakObjectPtr<URigVMGraph> GraphPtr;
 };
@@ -728,7 +728,7 @@ void FRigVMFunctionArgumentDefaultNode::HandleModifiedEvent(ERigVMGraphNotifType
 	}
 }
 
-TSharedRef<IDetailCustomization> FRigVMGraphDetailCustomization::MakeInstance(TSharedPtr<FRigVMEditorBase> InEditor, const UClass* InExpectedBlueprintClass)
+TSharedRef<IDetailCustomization> FRigVMGraphDetailCustomization::MakeInstance(TSharedPtr<IRigVMEditor> InEditor, const UClass* InExpectedBlueprintClass)
 {
 	const TArray<UObject*>* Objects = (InEditor.IsValid() ? InEditor->GetObjectsCurrentlyBeingEdited() : nullptr);
 	if (Objects && Objects->Num() == 1)
@@ -742,7 +742,7 @@ TSharedRef<IDetailCustomization> FRigVMGraphDetailCustomization::MakeInstance(TS
 		}
 	}
 
-	return MakeShareable(new FRigVMGraphDetailCustomization((TSharedPtr<FRigVMEditorBase>)nullptr, nullptr));
+	return MakeShareable(new FRigVMGraphDetailCustomization((TSharedPtr<IRigVMEditor>)nullptr, nullptr));
 }
 
 #if WITH_RIGVMLEGACYEDITOR
@@ -760,7 +760,7 @@ TSharedPtr<IDetailCustomization> FRigVMGraphDetailCustomization::MakeLegacyInsta
 		}
 	}
 
-	return MakeShareable(new FRigVMGraphDetailCustomization((TSharedPtr<FRigVMEditorBase>)nullptr, nullptr));
+	return MakeShareable(new FRigVMGraphDetailCustomization((TSharedPtr<IRigVMEditor>)nullptr, nullptr));
 }
 
 FRigVMGraphDetailCustomization::FRigVMGraphDetailCustomization(TSharedPtr<IBlueprintEditor> RigVMigEditor, URigVMBlueprint* RigVMBlueprint)
@@ -771,7 +771,7 @@ FRigVMGraphDetailCustomization::FRigVMGraphDetailCustomization(TSharedPtr<IBluep
 }
 #endif
 
-FRigVMGraphDetailCustomization::FRigVMGraphDetailCustomization(TSharedPtr<FRigVMEditorBase> RigVMigEditor, URigVMBlueprint* RigVMBlueprint)
+FRigVMGraphDetailCustomization::FRigVMGraphDetailCustomization(TSharedPtr<IRigVMEditor> RigVMigEditor, URigVMBlueprint* RigVMBlueprint)
 	: RigVMEditorPtr(RigVMigEditor)
 	, RigVMBlueprintPtr(RigVMBlueprint)
 	, RigVMGraphDetailCustomizationImpl(MakeShared<FRigVMGraphDetailCustomizationImpl>())
@@ -814,7 +814,7 @@ void FRigVMGraphDetailCustomizationImpl::CustomizeDetails(IDetailLayoutBuilder&
 	URigVMGraph* InModel,
 	URigVMController* InController,
 	IRigVMClientHost* InRigVMClientHost,
-	TWeakPtr<FRigVMEditorBase> InEditor)
+	TWeakPtr<IRigVMEditor> InEditor)
 {
 	WeakModel = InModel;
 	WeakController = InController;
diff --git a/Engine/Plugins/Runtime/RigVM/Source/RigVMEditor/Private/Editor/RigVMLegacyEditor.cpp b/Engine/Plugins/Runtime/RigVM/Source/RigVMEditor/Private/Editor/RigVMLegacyEditor.cpp
--- a/Engine/Plugins/Runtime/RigVM/Source/RigVMEditor/Private/Editor/RigVMLegacyEditor.cpp
+++ b/Engine/Plugins/Runtime/RigVM/Source/RigVMEditor/Private/Editor/RigVMLegacyEditor.cpp
@@ -84,7 +84,7 @@ void FRigVMLegacyEditor::GetCustomDebugObjects(TArray<FCustomDebugObject>& Debug
 		RigVMList.Add(RigVMObj);
 	}
 
-	FRigVMEditorBase::GetCustomDebugObjects(RigVMList);
+	GetDebugObjects(RigVMList);
 
 	DebugList.Reset();
 	DebugList.Reserve(RigVMList.Num());
diff --git a/Engine/Plugins/Runtime/RigVM/Source/RigVMEditor/Private/Editor/RigVMLegacyEditorMode.cpp b/Engine/Plugins/Runtime/RigVM/Source/RigVMEditor/Private/Editor/RigVMLegacyEditorMode.cpp
--- a/Engine/Plugins/Runtime/RigVM/Source/RigVMEditor/Private/Editor/RigVMLegacyEditorMode.cpp
+++ b/Engine/Plugins/Runtime/RigVM/Source/RigVMEditor/Private/Editor/RigVMLegacyEditorMode.cpp
@@ -13,7 +13,7 @@
 FRigVMLegacyEditorMode::FRigVMLegacyEditorMode(const TSharedRef<FRigVMLegacyEditor>& InRigVMEditor)
 	: FBlueprintEditorApplicationMode(InRigVMEditor, FRigVMEditorModes::RigVMEditorMode, FRigVMEditorModes::GetLocalizedMode, false, false)
 {
-	RigVMBlueprintPtr = CastChecked<URigVMBlueprint>(InRigVMEditor->FRigVMEditorBase::GetBlueprintObj());
+	RigVMBlueprintPtr = CastChecked<URigVMBlueprint>(InRigVMEditor->GetRigVMBlueprint());
 
 	TabFactories.RegisterFactory(MakeShared<FRigVMExecutionStackTabSummoner>(InRigVMEditor));
 	TabFactories.RegisterFactory(MakeShared<FRigVMEditorGraphExplorerTabSummoner>(InRigVMEditor));
diff --git a/Engine/Plugins/Runtime/RigVM/Source/RigVMEditor/Private/Editor/RigVMNewEditor.cpp b/Engine/Plugins/Runtime/RigVM/Source/RigVMEditor/Private/Editor/RigVMNewEditor.cpp
--- a/Engine/Plugins/Runtime/RigVM/Source/RigVMEditor/Private/Editor/RigVMNewEditor.cpp
+++ b/Engine/Plugins/Runtime/RigVM/Source/RigVMEditor/Private/Editor/RigVMNewEditor.cpp
@@ -297,7 +297,7 @@ void FRigVMNewEditor::OnBlueprintChanged(UBlueprint* InBlueprint)
 		// Refresh the graphs
 		//TODO sara-s
 		ERefreshRigVMEditorReason::Type Reason = ERefreshRigVMEditorReason::UnknownReason;// = bIsJustBeingCompiled ? ERefreshBlueprintEditorReason::BlueprintCompiled : ERefreshBlueprintEditorReason::UnknownReason;
-		RefreshEditors(Reason);
+		ForceEditorRefresh(Reason);
 
 		// In case objects were deleted, which should close the tab
 		if (GetCurrentMode() == FRigVMNewEditorApplicationModes::StandardRigVMEditorMode())
@@ -335,7 +335,7 @@ void FRigVMNewEditor::PostLayoutBlueprintEditorInitialization()
 	if (UBlueprint* Blueprint = GetBlueprintObj())
 	{
 		// Refresh the graphs
-		RefreshEditors();
+		ForceEditorRefresh();
 		
 		// EnsureBlueprintIsUpToDate may have updated the blueprint so show notifications to user.
 		if (bBlueprintModifiedOnOpen)
@@ -1472,7 +1472,7 @@ void FRigVMNewEditor::OnBlueprintChangedInnerImpl(UBlueprint* InBlueprint, bool
 
 		// Refresh the graphs
 		ERefreshRigVMEditorReason::Type Reason = bIsJustBeingCompiled ? ERefreshRigVMEditorReason::BlueprintCompiled : ERefreshRigVMEditorReason::UnknownReason;
-		RefreshEditors(Reason);
+		ForceEditorRefresh(Reason);
 
 		// In case objects were deleted, which should close the tab
 		if (GetCurrentMode() == FRigVMNewEditorApplicationModes::StandardRigVMEditorMode())
diff --git a/Engine/Plugins/Runtime/RigVM/Source/RigVMEditor/Private/Editor/RigVMVariableDetailCustomization.cpp b/Engine/Plugins/Runtime/RigVM/Source/RigVMEditor/Private/Editor/RigVMVariableDetailCustomization.cpp
--- a/Engine/Plugins/Runtime/RigVM/Source/RigVMEditor/Private/Editor/RigVMVariableDetailCustomization.cpp
+++ b/Engine/Plugins/Runtime/RigVM/Source/RigVMEditor/Private/Editor/RigVMVariableDetailCustomization.cpp
@@ -28,7 +28,7 @@
 
 #define LOCTEXT_NAMESPACE "RigVMVariableDetailCustomization"
 
-TSharedRef<IDetailCustomization> FRigVMVariableDetailCustomization::MakeInstance(TSharedPtr<FRigVMEditorBase> InEditor)
+TSharedRef<IDetailCustomization> FRigVMVariableDetailCustomization::MakeInstance(TSharedPtr<IRigVMEditor> InEditor)
 {
 	const TArray<UObject*>* Objects = (InEditor.IsValid() ? InEditor->GetObjectsCurrentlyBeingEdited() : nullptr);
 	if (Objects && Objects->Num() == 1)
@@ -42,10 +42,10 @@ TSharedRef<IDetailCustomization> FRigVMVariableDetailCustomization::MakeInstance
 		}
 	}
 
-	return MakeShareable(new FRigVMVariableDetailCustomization((TSharedPtr<FRigVMEditorBase>)nullptr, nullptr));
+	return MakeShareable(new FRigVMVariableDetailCustomization((TSharedPtr<IRigVMEditor>)nullptr, nullptr));
 }
 
-FRigVMVariableDetailCustomization::FRigVMVariableDetailCustomization(TSharedPtr<FRigVMEditorBase> InEditor, URigVMBlueprint* Blueprint)
+FRigVMVariableDetailCustomization::FRigVMVariableDetailCustomization(TSharedPtr<IRigVMEditor> InEditor, URigVMBlueprint* Blueprint)
 		: EditorPtr(InEditor)
 		, BlueprintPtr(Blueprint)
 {}
@@ -77,7 +77,7 @@ FRigVMVariableDetailCustomization::FRigVMVariableDetailCustomization(TSharedPtr<
 
 void FRigVMVariableDetailCustomization::CustomizeDetails(IDetailLayoutBuilder& DetailLayout)
 {
-	TSharedPtr<FRigVMEditorBase> Editor = EditorPtr.Pin();
+	TSharedPtr<IRigVMEditor> Editor = EditorPtr.Pin();
 	if (!Editor.IsValid())
 	{
 		return;
@@ -527,7 +527,7 @@ void FRigVMVariableDetailCustomization::CustomizeDetails(IDetailLayoutBuilder& D
 
 void FRigVMVariableDetailCustomization::PopulateCategories()
 {
-	TSharedPtr<FRigVMEditorBase> Editor = EditorPtr.Pin();
+	TSharedPtr<IRigVMEditor> Editor = EditorPtr.Pin();
 	TStrongObjectPtr<URigVMBlueprint> Blueprint = BlueprintPtr.Pin();
 	
 	if (!Blueprint || !Editor)
@@ -580,7 +580,7 @@ void FRigVMVariableDetailCustomization::OnPostEditorRefresh()
 
 FName FRigVMVariableDetailCustomization::GetVariableName() const
 {
-	TSharedPtr<FRigVMEditorBase> Editor = EditorPtr.Pin();
+	TSharedPtr<IRigVMEditor> Editor = EditorPtr.Pin();
     if (!Editor.IsValid())
     {
     	return NAME_None;
@@ -597,7 +597,7 @@ FText FRigVMVariableDetailCustomization::OnGetVariableName() const
 
 void FRigVMVariableDetailCustomization::OnVarNameCommitted(const FText& InNewName, ETextCommit::Type InTextCommit)
 {
-	TSharedPtr<FRigVMEditorBase> Editor = EditorPtr.Pin();
+	TSharedPtr<IRigVMEditor> Editor = EditorPtr.Pin();
 	if (!Editor)
 	{
 		return;
@@ -626,7 +626,7 @@ void FRigVMVariableDetailCustomization::OnVarNameCommitted(const FText& InNewNam
 FEdGraphPinType FRigVMVariableDetailCustomization::OnGetVarType() const
 {
 	FEdGraphPinType Type;
-	TSharedPtr<FRigVMEditorBase> Editor = EditorPtr.Pin();
+	TSharedPtr<IRigVMEditor> Editor = EditorPtr.Pin();
 	if (!Editor)
 	{
 		return Type;
@@ -661,7 +661,7 @@ FEdGraphPinType FRigVMVariableDetailCustomization::OnGetVarType() const
 
 void FRigVMVariableDetailCustomization::OnVarTypeChanged(const FEdGraphPinType& NewPinType)
 {
-	TSharedPtr<FRigVMEditorBase> Editor = EditorPtr.Pin();
+	TSharedPtr<IRigVMEditor> Editor = EditorPtr.Pin();
 	if (!Editor)
 	{
 		return;
@@ -732,7 +732,7 @@ bool FRigVMVariableDetailCustomization::CanBrowseToVarType() const
 FText FRigVMVariableDetailCustomization::OnGetTooltipText() const
 {
 	FText ToolTip;
-	TSharedPtr<FRigVMEditorBase> Editor = EditorPtr.Pin();
+	TSharedPtr<IRigVMEditor> Editor = EditorPtr.Pin();
 	if (!Editor)
 	{
 		return ToolTip;
@@ -764,7 +764,7 @@ FText FRigVMVariableDetailCustomization::OnGetTooltipText() const
 
 void FRigVMVariableDetailCustomization::OnTooltipTextCommitted(const FText& NewText, ETextCommit::Type InTextCommit, FName VarName)
 {
-	TSharedPtr<FRigVMEditorBase> Editor = EditorPtr.Pin();
+	TSharedPtr<IRigVMEditor> Editor = EditorPtr.Pin();
 	if (!Editor)
 	{
 		return;
@@ -793,7 +793,7 @@ void FRigVMVariableDetailCustomization::OnTooltipTextCommitted(const FText& NewT
 EVisibility FRigVMVariableDetailCustomization::IsToolTipVisible() const
 {
 	EVisibility Visible = EVisibility::Collapsed;
-	TSharedPtr<FRigVMEditorBase> Editor = EditorPtr.Pin();
+	TSharedPtr<IRigVMEditor> Editor = EditorPtr.Pin();
 	if (!Editor)
 	{
 		return Visible;
@@ -813,7 +813,7 @@ EVisibility FRigVMVariableDetailCustomization::IsToolTipVisible() const
 
 FText FRigVMVariableDetailCustomization::OnGetCategoryText() const
 {
-	TSharedPtr<FRigVMEditorBase> Editor = EditorPtr.Pin();
+	TSharedPtr<IRigVMEditor> Editor = EditorPtr.Pin();
 	if (!Editor)
 	{
 		return FText();
diff --git a/Engine/Plugins/Runtime/RigVM/Source/RigVMEditor/Private/Editor/SRigVMDetailsInspector.cpp b/Engine/Plugins/Runtime/RigVM/Source/RigVMEditor/Private/Editor/SRigVMDetailsInspector.cpp
--- a/Engine/Plugins/Runtime/RigVM/Source/RigVMEditor/Private/Editor/SRigVMDetailsInspector.cpp
+++ b/Engine/Plugins/Runtime/RigVM/Source/RigVMEditor/Private/Editor/SRigVMDetailsInspector.cpp
@@ -127,7 +127,7 @@ const TArray< TWeakObjectPtr<UObject> >& SRigVMDetailsInspector::GetSelectedObje
 	return SelectedObjects;
 }
 
-void SRigVMDetailsInspector::OnEditorClose(const FRigVMEditorBase* RigVMEditorBase, URigVMBlueprint* RigVMBlueprint)
+void SRigVMDetailsInspector::OnEditorClose(const IRigVMEditor* RigVMEditorBase, URigVMBlueprint* RigVMBlueprint)
 {
 	PropertyView->UnregisterInstancedCustomPropertyLayout(URigVMEdGraph::StaticClass());
 	PropertyView->UnregisterInstancedCustomPropertyLayout(UPropertyWrapper::StaticClass());
@@ -206,7 +206,7 @@ void SRigVMDetailsInspector::Construct(const FArguments& InArgs)
 	WeakEditor = InArgs._Editor;
 	bShowPublicView = InArgs._ShowPublicViewControl;
 	bShowTitleArea = InArgs._ShowTitleArea;
-	TSharedPtr<FRigVMEditorBase> Editor = WeakEditor.Pin();
+	TSharedPtr<IRigVMEditor> Editor = WeakEditor.Pin();
 
 	// Create a property view
 	FPropertyEditorModule& EditModule = FModuleManager::Get().GetModuleChecked<FPropertyEditorModule>("PropertyEditor");
@@ -364,7 +364,7 @@ void SRigVMDetailsInspector::UpdateFromSingleStruct(const TSharedPtr<FStructOnSc
 
 void SRigVMDetailsInspector::UpdateFromObjects(const TArray<UObject*>& PropertyObjects, TArray<UObject*>& SelectionInfo, const FShowDetailsOptions& Options)
 {
-	TSharedPtr<FRigVMEditorBase> Editor = WeakEditor.Pin();
+	TSharedPtr<IRigVMEditor> Editor = WeakEditor.Pin();
 	
 	if (!Options.bForceRefresh)
 	{
@@ -498,8 +498,8 @@ bool SRigVMDetailsInspector::IsPropertyVisible( const FPropertyAndParent& Proper
 		// check if the property (or any of its parent properties) was added by this blueprint
 		// this is necessary because of Instanced objects, which will have a different owning class yet are conceptually contained in this blueprint
 		bool bVariableAddedInCurrentBlueprint = false;
-		TSharedPtr<FRigVMEditorBase> Editor = WeakEditor.Pin();
-		const UBlueprint* Blueprint = Editor.IsValid() ? Editor->GetBlueprintObj() : nullptr;
+		TSharedPtr<IRigVMEditor> Editor = WeakEditor.Pin();
+		const UBlueprint* Blueprint = Editor.IsValid() ? Editor->GetRigVMBlueprint() : nullptr;
 
 		auto WasAddedInThisBlueprint = [Blueprint](const FProperty* Property)
 		{
@@ -579,7 +579,7 @@ bool SRigVMDetailsInspector::IsPropertyEditingEnabled() const
 {
 	bool bIsEditable = true;
 
-	if (TSharedPtr<FRigVMEditorBase> Editor = WeakEditor.Pin())
+	if (TSharedPtr<IRigVMEditor> Editor = WeakEditor.Pin())
 	{
 		// This function is essentially for PIE use so if we are NOT doing PIE use the normal path
 		if (GEditor->GetPIEWorldContext() == nullptr)
diff --git a/Engine/Plugins/Runtime/RigVM/Source/RigVMEditor/Private/Widgets/SRigVMEditorGraphExplorer.cpp b/Engine/Plugins/Runtime/RigVM/Source/RigVMEditor/Private/Widgets/SRigVMEditorGraphExplorer.cpp
--- a/Engine/Plugins/Runtime/RigVM/Source/RigVMEditor/Private/Widgets/SRigVMEditorGraphExplorer.cpp
+++ b/Engine/Plugins/Runtime/RigVM/Source/RigVMEditor/Private/Widgets/SRigVMEditorGraphExplorer.cpp
@@ -197,12 +197,12 @@ void FRigVMEditorGraphExplorerCommands::RegisterCommands()
 }
 
 
-void SRigVMEditorGraphExplorer::Construct(const FArguments& InArgs, TWeakPtr<FRigVMEditorBase> InRigVMEditor)
+void SRigVMEditorGraphExplorer::Construct(const FArguments& InArgs, TWeakPtr<IRigVMEditor> InRigVMEditor)
 {
 	bNeedsRefresh = false;
 	RigVMEditor = InRigVMEditor;
 
-	TSharedPtr<FRigVMEditorBase> Editor = RigVMEditor.Pin();
+	TSharedPtr<IRigVMEditor> Editor = RigVMEditor.Pin();
 	if (Editor)
 	{
 		Editor->OnRefresh().AddRaw(this, &SRigVMEditorGraphExplorer::Refresh);
@@ -219,7 +219,7 @@ void SRigVMEditorGraphExplorer::Construct(const FArguments& InArgs, TWeakPtr<FRi
 
 SRigVMEditorGraphExplorer::~SRigVMEditorGraphExplorer()
 {
-	TSharedPtr<FRigVMEditorBase> Editor = RigVMEditor.Pin();
+	TSharedPtr<IRigVMEditor> Editor = RigVMEditor.Pin();
 	if (Editor)
 	{
 		Editor->OnRefresh().RemoveAll(this);
@@ -276,7 +276,7 @@ void SRigVMEditorGraphExplorer::ClearSelection()
 
 void SRigVMEditorGraphExplorer::RegisterCommands()
 {
-	TSharedPtr<FRigVMEditorBase> Editor = RigVMEditor.Pin();
+	TSharedPtr<IRigVMEditor> Editor = RigVMEditor.Pin();
 
 	if (Editor)
 	{
@@ -530,7 +530,7 @@ TArray<const URigVMGraph*> SRigVMEditorGraphExplorer::GetRootGraphs() const
 {
 	TArray<const URigVMGraph*> Graphs;
 	
-	TSharedPtr<FRigVMEditorBase> Editor = RigVMEditor.Pin();
+	TSharedPtr<IRigVMEditor> Editor = RigVMEditor.Pin();
 	if (!Editor)
 	{
 		return Graphs;
@@ -549,7 +549,7 @@ TArray<const URigVMGraph*> SRigVMEditorGraphExplorer::GetChildrenGraphs(const FS
 {
 	TArray<const URigVMGraph*> Children;
 	
-	TSharedPtr<FRigVMEditorBase> Editor = RigVMEditor.Pin();
+	TSharedPtr<IRigVMEditor> Editor = RigVMEditor.Pin();
 	if (!Editor)
 	{
 		return Children;
@@ -581,7 +581,7 @@ TArray<URigVMNode*> SRigVMEditorGraphExplorer::GetEventNodesInGraph(const FStrin
 {
 	TArray<URigVMNode*> Events;
 	
-	TSharedPtr<FRigVMEditorBase> Editor = RigVMEditor.Pin();
+	TSharedPtr<IRigVMEditor> Editor = RigVMEditor.Pin();
 	if (!Editor)
 	{
 		return Events;
@@ -619,7 +619,7 @@ TArray<URigVMLibraryNode*> SRigVMEditorGraphExplorer::GetFunctions() const
 {
 	TArray<URigVMLibraryNode*> Functions;
 	
-	TSharedPtr<FRigVMEditorBase> Editor = RigVMEditor.Pin();
+	TSharedPtr<IRigVMEditor> Editor = RigVMEditor.Pin();
 	if (!Editor)
 	{
 		return Functions;
@@ -637,7 +637,7 @@ TArray<URigVMLibraryNode*> SRigVMEditorGraphExplorer::GetFunctions() const
 TArray<FRigVMGraphVariableDescription> SRigVMEditorGraphExplorer::GetVariables() const
 {
 	TArray<FRigVMGraphVariableDescription> Variables;
-	TSharedPtr<FRigVMEditorBase> Editor = RigVMEditor.Pin();
+	TSharedPtr<IRigVMEditor> Editor = RigVMEditor.Pin();
 	if (!Editor)
 	{
 		return Variables;
@@ -649,7 +649,7 @@ TArray<FRigVMGraphVariableDescription> SRigVMEditorGraphExplorer::GetVariables()
 TArray<FRigVMGraphVariableDescription> SRigVMEditorGraphExplorer::GetLocalVariables() const
 {
 	TArray<FRigVMGraphVariableDescription> Variables;
-	TSharedPtr<FRigVMEditorBase> Editor = RigVMEditor.Pin();
+	TSharedPtr<IRigVMEditor> Editor = RigVMEditor.Pin();
 	if (!Editor)
 	{
 		return Variables;
@@ -669,7 +669,7 @@ TArray<FRigVMGraphVariableDescription> SRigVMEditorGraphExplorer::GetLocalVariab
 FText SRigVMEditorGraphExplorer::GetGraphDisplayName(const FString& InGraphPath) const
 {
 	FText DisplayName;
-	TSharedPtr<FRigVMEditorBase> Editor = RigVMEditor.Pin();
+	TSharedPtr<IRigVMEditor> Editor = RigVMEditor.Pin();
 	if (!Editor)
 	{
 		return DisplayName;
@@ -692,7 +692,7 @@ FText SRigVMEditorGraphExplorer::GetGraphDisplayName(const FString& InGraphPath)
 FText SRigVMEditorGraphExplorer::GetEventDisplayName(const FString& InNodePath) const
 {
 	FText DisplayName;
-	TSharedPtr<FRigVMEditorBase> Editor = RigVMEditor.Pin();
+	TSharedPtr<IRigVMEditor> Editor = RigVMEditor.Pin();
 	if (!Editor)
 	{
 		return DisplayName;
@@ -710,7 +710,7 @@ FText SRigVMEditorGraphExplorer::GetEventDisplayName(const FString& InNodePath)
 const FSlateBrush* SRigVMEditorGraphExplorer::GetGraphIcon(const FString& InGraphPath) const
 {
 	const FSlateBrush* Brush = nullptr;
-	TSharedPtr<FRigVMEditorBase> Editor = RigVMEditor.Pin();
+	TSharedPtr<IRigVMEditor> Editor = RigVMEditor.Pin();
 	if (!Editor)
 	{
 		return Brush;
@@ -734,7 +734,7 @@ const FSlateBrush* SRigVMEditorGraphExplorer::GetGraphIcon(const FString& InGrap
 FText SRigVMEditorGraphExplorer::GetGraphTooltip(const FString& InGraphPath) const
 {
 	FText Tooltip;
-	TSharedPtr<FRigVMEditorBase> Editor = RigVMEditor.Pin();
+	TSharedPtr<IRigVMEditor> Editor = RigVMEditor.Pin();
 	if (!Editor)
 	{
 		return Tooltip;
@@ -756,7 +756,7 @@ FText SRigVMEditorGraphExplorer::GetGraphTooltip(const FString& InGraphPath) con
 
 void SRigVMEditorGraphExplorer::OnGraphClicked(const FString& InGraphPath)
 {
-	TSharedPtr<FRigVMEditorBase> Editor = RigVMEditor.Pin();
+	TSharedPtr<IRigVMEditor> Editor = RigVMEditor.Pin();
 	if (!Editor)
 	{
 		return;
@@ -788,7 +788,7 @@ void SRigVMEditorGraphExplorer::OnGraphClicked(const FString& InGraphPath)
 
 void SRigVMEditorGraphExplorer::OnEventClicked(const FString& InEventPath)
 {
-	TSharedPtr<FRigVMEditorBase> Editor = RigVMEditor.Pin();
+	TSharedPtr<IRigVMEditor> Editor = RigVMEditor.Pin();
 	if (!Editor)
 	{
 		return;
@@ -813,7 +813,7 @@ void SRigVMEditorGraphExplorer::OnEventClicked(const FString& InEventPath)
 
 void SRigVMEditorGraphExplorer::OnFunctionClicked(const FString& InFunctionPath)
 {
-	TSharedPtr<FRigVMEditorBase> Editor = RigVMEditor.Pin();
+	TSharedPtr<IRigVMEditor> Editor = RigVMEditor.Pin();
 	if (!Editor)
 	{
 		return;
@@ -848,7 +848,7 @@ void SRigVMEditorGraphExplorer::OnFunctionClicked(const FString& InFunctionPath)
 
 void SRigVMEditorGraphExplorer::OnVariableClicked(const FRigVMExplorerElementKey& InVariable)
 {
-	TSharedPtr<FRigVMEditorBase> Editor = RigVMEditor.Pin();
+	TSharedPtr<IRigVMEditor> Editor = RigVMEditor.Pin();
 	if (!Editor)
 	{
 		return;
@@ -891,7 +891,7 @@ void SRigVMEditorGraphExplorer::OnVariableClicked(const FRigVMExplorerElementKey
 
 void SRigVMEditorGraphExplorer::OnGraphDoubleClicked(const FString& InGraphPath)
 {
-	TSharedPtr<FRigVMEditorBase> Editor = RigVMEditor.Pin();
+	TSharedPtr<IRigVMEditor> Editor = RigVMEditor.Pin();
 	if (!Editor)
 	{
 		return;
@@ -910,7 +910,7 @@ void SRigVMEditorGraphExplorer::OnGraphDoubleClicked(const FString& InGraphPath)
 
 void SRigVMEditorGraphExplorer::OnEventDoubleClicked(const FString& InEventPath)
 {
-	TSharedPtr<FRigVMEditorBase> Editor = RigVMEditor.Pin();
+	TSharedPtr<IRigVMEditor> Editor = RigVMEditor.Pin();
 	if (!Editor)
 	{
 		return;
@@ -927,7 +927,7 @@ void SRigVMEditorGraphExplorer::OnEventDoubleClicked(const FString& InEventPath)
 
 void SRigVMEditorGraphExplorer::OnFunctionDoubleClicked(const FString& InFunctionPath)
 {
-	TSharedPtr<FRigVMEditorBase> Editor = RigVMEditor.Pin();
+	TSharedPtr<IRigVMEditor> Editor = RigVMEditor.Pin();
 	if (!Editor)
 	{
 		return;
@@ -949,7 +949,7 @@ void SRigVMEditorGraphExplorer::OnFunctionDoubleClicked(const FString& InFunctio
 
 bool SRigVMEditorGraphExplorer::OnSetFunctionCategory(const FString& InFunctionPath, const FString& InCategory)
 {
-	TSharedPtr<FRigVMEditorBase> Editor = RigVMEditor.Pin();
+	TSharedPtr<IRigVMEditor> Editor = RigVMEditor.Pin();
 	if (!Editor.IsValid())
 	{
 		return false;
@@ -973,7 +973,7 @@ bool SRigVMEditorGraphExplorer::OnSetFunctionCategory(const FString& InFunctionP
 
 FString SRigVMEditorGraphExplorer::OnGetFunctionCategory(const FString& InFunctionPath) const
 {
-	TSharedPtr<FRigVMEditorBase> Editor = RigVMEditor.Pin();
+	TSharedPtr<IRigVMEditor> Editor = RigVMEditor.Pin();
 	if (!Editor.IsValid())
 	{
 		return FString();
@@ -992,7 +992,7 @@ FString SRigVMEditorGraphExplorer::OnGetFunctionCategory(const FString& InFuncti
 
 bool SRigVMEditorGraphExplorer::OnSetVariableCategory(const FString& InVariable, const FString& InCategory)
 {
-	TSharedPtr<FRigVMEditorBase> Editor = RigVMEditor.Pin();
+	TSharedPtr<IRigVMEditor> Editor = RigVMEditor.Pin();
 	if (!Editor.IsValid())
 	{
 		return false;
@@ -1005,7 +1005,7 @@ bool SRigVMEditorGraphExplorer::OnSetVariableCategory(const FString& InVariable,
 
 FString SRigVMEditorGraphExplorer::OnGetVariableCategory(const FString& InVariable) const
 {
-	TSharedPtr<FRigVMEditorBase> Editor = RigVMEditor.Pin();
+	TSharedPtr<IRigVMEditor> Editor = RigVMEditor.Pin();
 	if (!Editor.IsValid())
 	{
 		return FString();
@@ -1016,7 +1016,7 @@ FString SRigVMEditorGraphExplorer::OnGetVariableCategory(const FString& InVariab
 
 FEdGraphPinType SRigVMEditorGraphExplorer::OnGetVariablePinType(const FRigVMExplorerElementKey& InVariable)
 {
-	TSharedPtr<FRigVMEditorBase> Editor = RigVMEditor.Pin();
+	TSharedPtr<IRigVMEditor> Editor = RigVMEditor.Pin();
 	if (!Editor.IsValid())
 	{
 		return FEdGraphPinType();
@@ -1052,7 +1052,7 @@ FEdGraphPinType SRigVMEditorGraphExplorer::OnGetVariablePinType(const FRigVMExpl
 
 bool SRigVMEditorGraphExplorer::OnSetVariablePinType(const FRigVMExplorerElementKey& InVariable, const FEdGraphPinType& InType)
 {
-	TSharedPtr<FRigVMEditorBase> Editor = RigVMEditor.Pin();
+	TSharedPtr<IRigVMEditor> Editor = RigVMEditor.Pin();
 	if (!Editor.IsValid())
 	{
 		return false;
@@ -1093,7 +1093,7 @@ bool SRigVMEditorGraphExplorer::OnSetVariablePinType(const FRigVMExplorerElement
 
 bool SRigVMEditorGraphExplorer::OnIsVariablePublic(const FString& InVariableName) const
 {
-	TSharedPtr<FRigVMEditorBase> Editor = RigVMEditor.Pin();
+	TSharedPtr<IRigVMEditor> Editor = RigVMEditor.Pin();
 	if (!Editor.IsValid())
 	{
 		return false;
@@ -1118,7 +1118,7 @@ bool SRigVMEditorGraphExplorer::OnIsVariablePublic(const FString& InVariableName
 
 bool SRigVMEditorGraphExplorer::OnToggleVariablePublic(const FString& InVariableName) const
 {
-	TSharedPtr<FRigVMEditorBase> Editor = RigVMEditor.Pin();
+	TSharedPtr<IRigVMEditor> Editor = RigVMEditor.Pin();
 	if (!Editor.IsValid())
 	{
 		return false;
@@ -1133,7 +1133,7 @@ bool SRigVMEditorGraphExplorer::OnToggleVariablePublic(const FString& InVariable
 
 bool SRigVMEditorGraphExplorer::IsFunctionFocused() const
 {
-	TSharedPtr<FRigVMEditorBase> Editor = RigVMEditor.Pin();
+	TSharedPtr<IRigVMEditor> Editor = RigVMEditor.Pin();
 	if (!Editor.IsValid())
 	{
 		return false;
@@ -1164,7 +1164,7 @@ TArray<TSharedPtr<IPinTypeSelectorFilter>> SRigVMEditorGraphExplorer::GetCustomP
 
 TSharedPtr<SWidget> SRigVMEditorGraphExplorer::OnContextMenuOpening()
 {
-	TSharedPtr<FRigVMEditorBase> Editor = RigVMEditor.Pin();
+	TSharedPtr<IRigVMEditor> Editor = RigVMEditor.Pin();
 	if (!Editor.IsValid())
 	{
 		return TSharedPtr<SWidget>();
@@ -1254,7 +1254,7 @@ FReply SRigVMEditorGraphExplorer::OnDragDetected(const FGeometry& MyGeometry, co
 
 void SRigVMEditorGraphExplorer::OnCopy()
 {
-	TSharedPtr<FRigVMEditorBase> Editor = RigVMEditor.Pin();
+	TSharedPtr<IRigVMEditor> Editor = RigVMEditor.Pin();
 	if (!Editor)
 	{
 		return;
@@ -1380,7 +1380,7 @@ void SRigVMEditorGraphExplorer::OnDuplicate()
 
 bool SRigVMEditorGraphExplorer::CanDuplicate() const
 {
-	TSharedPtr<FRigVMEditorBase> Editor = RigVMEditor.Pin();
+	TSharedPtr<IRigVMEditor> Editor = RigVMEditor.Pin();
 	if (!Editor)
 	{
 		return false;
@@ -1449,7 +1449,7 @@ bool SRigVMEditorGraphExplorer::CanPasteGeneric() const
 
 void SRigVMEditorGraphExplorer::OnPasteFunction()
 {
-	TSharedPtr<FRigVMEditorBase> Editor = RigVMEditor.Pin();
+	TSharedPtr<IRigVMEditor> Editor = RigVMEditor.Pin();
 	if (!Editor)
 	{
 		return;
@@ -1463,7 +1463,7 @@ void SRigVMEditorGraphExplorer::OnPasteFunction()
 
 bool SRigVMEditorGraphExplorer::CanPasteFunction() const
 {
-	TSharedPtr<FRigVMEditorBase> Editor = RigVMEditor.Pin();
+	TSharedPtr<IRigVMEditor> Editor = RigVMEditor.Pin();
 	if (!Editor)
 	{
 		return false;
@@ -1477,7 +1477,7 @@ bool SRigVMEditorGraphExplorer::CanPasteFunction() const
 
 void SRigVMEditorGraphExplorer::OnPasteVariable()
 {
-	TSharedPtr<FRigVMEditorBase> Editor = RigVMEditor.Pin();
+	TSharedPtr<IRigVMEditor> Editor = RigVMEditor.Pin();
 	if (!Editor)
 	{
 		return;
@@ -1513,7 +1513,7 @@ void SRigVMEditorGraphExplorer::OnPasteVariable()
 
 bool SRigVMEditorGraphExplorer::CanPasteVariable() const
 {
-	TSharedPtr<FRigVMEditorBase> Editor = RigVMEditor.Pin();
+	TSharedPtr<IRigVMEditor> Editor = RigVMEditor.Pin();
 	if (!Editor)
 	{
 		return false;
@@ -1536,7 +1536,7 @@ bool SRigVMEditorGraphExplorer::CanPasteVariable() const
 
 void SRigVMEditorGraphExplorer::OnPasteLocalVariable()
 {
-	TSharedPtr<FRigVMEditorBase> Editor = RigVMEditor.Pin();
+	TSharedPtr<IRigVMEditor> Editor = RigVMEditor.Pin();
 	if (!Editor)
 	{
 		return;
@@ -1583,7 +1583,7 @@ bool SRigVMEditorGraphExplorer::IsShowingEmptySections() const
 
 void SRigVMEditorGraphExplorer::OnOpenGraph(bool bOnNewTab)
 {
-	TSharedPtr<FRigVMEditorBase> Editor = RigVMEditor.Pin();
+	TSharedPtr<IRigVMEditor> Editor = RigVMEditor.Pin();
 
 	if (Editor)
 	{
@@ -1616,17 +1616,17 @@ bool SRigVMEditorGraphExplorer::CanOpenGraph() const
 
 void SRigVMEditorGraphExplorer::OnCreateGraph()
 {
-	TSharedPtr<FRigVMEditorBase> Editor = RigVMEditor.Pin();
+	TSharedPtr<IRigVMEditor> Editor = RigVMEditor.Pin();
 	if (Editor)
 	{
-		Editor->NewDocument_OnClicked(FRigVMEditorBase::CGT_NewEventGraph);
+		Editor->OnNewDocumentClicked(FRigVMEditorBase::CGT_NewEventGraph);
 	}
 }
 
 
 bool SRigVMEditorGraphExplorer::CanCreateGraph() const
 {
-	TSharedPtr<FRigVMEditorBase> Editor = RigVMEditor.Pin();
+	TSharedPtr<IRigVMEditor> Editor = RigVMEditor.Pin();
 	if (Editor)
 	{
 		return Editor->InEditingMode();
@@ -1637,10 +1637,10 @@ bool SRigVMEditorGraphExplorer::CanCreateGraph() const
 
 void SRigVMEditorGraphExplorer::OnCreateFunction()
 {
-	TSharedPtr<FRigVMEditorBase> Editor = RigVMEditor.Pin();
+	TSharedPtr<IRigVMEditor> Editor = RigVMEditor.Pin();
 	if (Editor)
 	{
-		Editor->NewDocument_OnClicked(FRigVMEditorBase::CGT_NewFunctionGraph);
+		Editor->OnNewDocumentClicked(FRigVMEditorBase::CGT_NewFunctionGraph);
 	}
 }
 
@@ -1651,7 +1651,7 @@ bool SRigVMEditorGraphExplorer::CanCreateFunction() const
 
 void SRigVMEditorGraphExplorer::OnCreateVariable()
 {
-	TSharedPtr<FRigVMEditorBase> Editor = RigVMEditor.Pin();
+	TSharedPtr<IRigVMEditor> Editor = RigVMEditor.Pin();
 	if (Editor)
 	{
 		Editor->OnAddNewVariable();
@@ -1665,7 +1665,7 @@ bool SRigVMEditorGraphExplorer::CanCreateVariable() const
 
 void SRigVMEditorGraphExplorer::OnCreateLocalVariable()
 {
-	TSharedPtr<FRigVMEditorBase> Editor = RigVMEditor.Pin();
+	TSharedPtr<IRigVMEditor> Editor = RigVMEditor.Pin();
 	if (!Editor)
 	{
 		return;
@@ -1677,7 +1677,7 @@ void SRigVMEditorGraphExplorer::OnCreateLocalVariable()
 
 bool SRigVMEditorGraphExplorer::CanCreateLocalVariable() const
 {
-	TSharedPtr<FRigVMEditorBase> Editor = RigVMEditor.Pin();
+	TSharedPtr<IRigVMEditor> Editor = RigVMEditor.Pin();
 	if (!Editor)
 	{
 		return false;
@@ -1688,7 +1688,7 @@ bool SRigVMEditorGraphExplorer::CanCreateLocalVariable() const
 
 void SRigVMEditorGraphExplorer::OnAddFunctionVariant()
 {
-	TSharedPtr<FRigVMEditorBase> Editor = RigVMEditor.Pin();
+	TSharedPtr<IRigVMEditor> Editor = RigVMEditor.Pin();
 
 	if (!Editor)
 	{
@@ -1715,7 +1715,7 @@ void SRigVMEditorGraphExplorer::OnAddFunctionVariant()
 
 bool SRigVMEditorGraphExplorer::CanAddFunctionVariant() const
 {
-	TSharedPtr<FRigVMEditorBase> Editor = RigVMEditor.Pin();
+	TSharedPtr<IRigVMEditor> Editor = RigVMEditor.Pin();
 
 	if (!Editor)
 	{
@@ -1733,7 +1733,7 @@ bool SRigVMEditorGraphExplorer::CanAddFunctionVariant() const
 
 void SRigVMEditorGraphExplorer::OnDeleteEntry()
 {
-	TSharedPtr<FRigVMEditorBase> Editor = RigVMEditor.Pin();
+	TSharedPtr<IRigVMEditor> Editor = RigVMEditor.Pin();
 
 	if (Editor)
 	{
@@ -1795,7 +1795,7 @@ void SRigVMEditorGraphExplorer::OnDeleteEntry()
 
 bool SRigVMEditorGraphExplorer::CanDeleteEntry() const
 {
-	TSharedPtr<FRigVMEditorBase> Editor = RigVMEditor.Pin();
+	TSharedPtr<IRigVMEditor> Editor = RigVMEditor.Pin();
 
 	if (!Editor)
 	{
@@ -1852,7 +1852,7 @@ bool SRigVMEditorGraphExplorer::CanRenameEntry() const
 bool SRigVMEditorGraphExplorer::OnRenameGraph(const FString& InOldPath, const FString& InNewPath)
 {
 	bool bResult = false;
-	TSharedPtr<FRigVMEditorBase> Editor = RigVMEditor.Pin();
+	TSharedPtr<IRigVMEditor> Editor = RigVMEditor.Pin();
 
 	if (Editor)
 	{
@@ -1870,7 +1870,7 @@ bool SRigVMEditorGraphExplorer::OnRenameGraph(const FString& InOldPath, const FS
 bool SRigVMEditorGraphExplorer::OnCanRenameGraph(const FString& InOldPath, const FString& InNewPath, FText& OutErrorMessage) const
 {
 	bool bResult = false;
-	TSharedPtr<FRigVMEditorBase> Editor = RigVMEditor.Pin();
+	TSharedPtr<IRigVMEditor> Editor = RigVMEditor.Pin();
 
 	if (Editor)
 	{
@@ -1901,7 +1901,7 @@ bool SRigVMEditorGraphExplorer::OnCanRenameGraph(const FString& InOldPath, const
 bool SRigVMEditorGraphExplorer::OnRenameFunction(const FString& InOldPath, const FString& InNewPath)
 {
 	bool bResult = false;
-	TSharedPtr<FRigVMEditorBase> Editor = RigVMEditor.Pin();
+	TSharedPtr<IRigVMEditor> Editor = RigVMEditor.Pin();
 
 	if (Editor)
 	{
@@ -1925,7 +1925,7 @@ bool SRigVMEditorGraphExplorer::OnRenameFunction(const FString& InOldPath, const
 bool SRigVMEditorGraphExplorer::OnCanRenameFunction(const FString& InOldPath, const FString& InNewPath, FText& OutErrorMessage) const
 {
 	bool bResult = false;
-	TSharedPtr<FRigVMEditorBase> Editor = RigVMEditor.Pin();
+	TSharedPtr<IRigVMEditor> Editor = RigVMEditor.Pin();
 
 	if (Editor)
 	{
@@ -1948,7 +1948,7 @@ bool SRigVMEditorGraphExplorer::OnCanRenameFunction(const FString& InOldPath, co
 
 bool SRigVMEditorGraphExplorer::OnRenameVariable(const FRigVMExplorerElementKey& InOldKey, const FString& InNewName)
 {
-	TSharedPtr<FRigVMEditorBase> Editor = RigVMEditor.Pin();
+	TSharedPtr<IRigVMEditor> Editor = RigVMEditor.Pin();
 	if (!Editor)
 	{
 		return false;
@@ -2015,7 +2015,7 @@ void SRigVMEditorGraphExplorer::OnFindReference(bool bSearchAllBlueprints)
 		return;
 	}
 
-	TSharedPtr<FRigVMEditorBase> Editor = RigVMEditor.Pin();
+	TSharedPtr<IRigVMEditor> Editor = RigVMEditor.Pin();
 	URigVMBlueprint* Blueprint = Editor->GetRigVMBlueprint();
 	const URigVMEdGraphSchema* Schema = GetDefault<URigVMEdGraphSchema>();
 
diff --git a/Engine/Plugins/Runtime/RigVM/Source/RigVMEditor/Private/Widgets/SRigVMExecutionStackView.cpp b/Engine/Plugins/Runtime/RigVM/Source/RigVMEditor/Private/Widgets/SRigVMExecutionStackView.cpp
--- a/Engine/Plugins/Runtime/RigVM/Source/RigVMEditor/Private/Widgets/SRigVMExecutionStackView.cpp
+++ b/Engine/Plugins/Runtime/RigVM/Source/RigVMEditor/Private/Widgets/SRigVMExecutionStackView.cpp
@@ -247,7 +247,7 @@ FText SRigStackItem::GetDurationText() const
 
 SRigVMExecutionStackView::~SRigVMExecutionStackView()
 {
-	if (TSharedPtr<FRigVMEditorBase> Editor = WeakEditor.Pin())
+	if (TSharedPtr<IRigVMEditor> Editor = WeakEditor.Pin())
 	{
 		if (OnModelModified.IsValid() && Editor->GetRigVMBlueprint())
 		{
@@ -271,7 +271,7 @@ SRigVMExecutionStackView::~SRigVMExecutionStackView()
 	}
 }
 
-void SRigVMExecutionStackView::Construct( const FArguments& InArgs, TSharedRef<FRigVMEditorBase> InRigVMEditor)
+void SRigVMExecutionStackView::Construct( const FArguments& InArgs, TSharedRef<IRigVMEditor> InRigVMEditor)
 {
 	WeakEditor = InRigVMEditor;
 	RigVMBlueprint = InRigVMEditor->GetRigVMBlueprint();
@@ -348,7 +348,7 @@ void SRigVMExecutionStackView::Construct( const FArguments& InArgs, TSharedRef<F
 		OnVMCompiledHandle = RigVMBlueprint->OnVMCompiled().AddSP(this, &SRigVMExecutionStackView::OnVMCompiled);
 		OnModelModified = RigVMBlueprint->OnModified().AddSP(this, &SRigVMExecutionStackView::HandleModifiedEvent);
 
-		if (TSharedPtr<FRigVMEditorBase> Editor = WeakEditor.Pin())
+		if (TSharedPtr<IRigVMEditor> Editor = WeakEditor.Pin())
 		{
 			if (URigVMHost* RigVMHost = Editor->GetRigVMHost())
 			{
@@ -357,7 +357,7 @@ void SRigVMExecutionStackView::Construct( const FArguments& InArgs, TSharedRef<F
 		}
 	}
 
-	if (TSharedPtr<FRigVMEditorBase> Editor = WeakEditor.Pin())
+	if (TSharedPtr<IRigVMEditor> Editor = WeakEditor.Pin())
 	{
 		OnPreviewHostUpdatedHandle = Editor->OnPreviewHostUpdated().AddSP(this, &SRigVMExecutionStackView::HandlePreviewHostUpdated);
 	}
@@ -386,7 +386,7 @@ void SRigVMExecutionStackView::OnSelectionChanged(TSharedPtr<FRigStackEntry> Sel
 		}
 
 		URigVMHost* Host = nullptr;
-		if (TSharedPtr<FRigVMEditorBase> Editor = WeakEditor.Pin())
+		if (TSharedPtr<IRigVMEditor> Editor = WeakEditor.Pin())
 		{
 			Host = Editor->GetRigVMHost();
 		}
@@ -740,7 +740,7 @@ void SRigVMExecutionStackView::RefreshTreeView(URigVM* InVM, FRigVMExtendedExecu
 	if (InVM && InVMContext)
 	{
 		// fill the children from the log
-		if (TSharedPtr<FRigVMEditorBase> Editor = WeakEditor.Pin())
+		if (TSharedPtr<IRigVMEditor> Editor = WeakEditor.Pin())
 		{
 			URigVMHost* Host = Editor->GetRigVMHost();
 			if(Host && Host->GetLog())
@@ -830,7 +830,7 @@ void SRigVMExecutionStackView::HandleFocusOnSelectedGraphNode()
 	TArray<TSharedPtr<FRigStackEntry>> SelectedItems = TreeView->GetSelectedItems();
 	if (SelectedItems.Num() > 0)
 	{
-		if (TSharedPtr<FRigVMEditorBase> Editor = WeakEditor.Pin())
+		if (TSharedPtr<IRigVMEditor> Editor = WeakEditor.Pin())
 		{
 			URigVMHost* Host = Editor->GetRigVMHost();
 			if (Host == nullptr || Host->GetVM() == nullptr)
@@ -935,7 +935,7 @@ TArray<TSharedPtr<FRigStackEntry>> SRigVMExecutionStackView::GetTargetItems(cons
 	TArray<TSharedPtr<FRigStackEntry>> TargetItems;
 
 	URigVMHost* Host = nullptr;
-	if (TSharedPtr<FRigVMEditorBase> Editor = WeakEditor.Pin())
+	if (TSharedPtr<IRigVMEditor> Editor = WeakEditor.Pin())
 	{
 		Host = Editor->GetRigVMHost();
 	}
@@ -1046,7 +1046,7 @@ void SRigVMExecutionStackView::OnVMCompiled(UObject* InCompiledObject, URigVM* I
 {
 	RefreshTreeView(InCompiledVM, &InVMContext);
 
-	if (TSharedPtr<FRigVMEditorBase> Editor = WeakEditor.Pin())
+	if (TSharedPtr<IRigVMEditor> Editor = WeakEditor.Pin())
 	{
 		if (!OnHostInitializedHandle.IsValid())
 		{
@@ -1065,7 +1065,7 @@ void SRigVMExecutionStackView::HandleExecutionHalted(const int32 InHaltedAtInstr
 		return;
 	}
 
-	if (TSharedPtr<FRigVMEditorBase> Editor = WeakEditor.Pin())
+	if (TSharedPtr<IRigVMEditor> Editor = WeakEditor.Pin())
 	{
 		if (InHaltedAtInstruction == INDEX_NONE && InEntryName == Editor->GetRigVMHost()->GetEventQueue().Last())
 		{
@@ -1085,7 +1085,7 @@ void SRigVMExecutionStackView::HandleExecutionHalted(const int32 InHaltedAtInstr
 void SRigVMExecutionStackView::OnFilterTextChanged(const FText& SearchText)
 {
 	FilterText = SearchText;
-	if (TSharedPtr<FRigVMEditorBase> Editor = WeakEditor.Pin())
+	if (TSharedPtr<IRigVMEditor> Editor = WeakEditor.Pin())
 	{
 		URigVMHost* RigVMHost = Editor->GetRigVMHost();
 		if (RigVMHost != nullptr)
@@ -1103,7 +1103,7 @@ void SRigVMExecutionStackView::HandleModifiedEvent(ERigVMGraphNotifType InNotifT
 	}
 
 	URigVMHost* Host = nullptr;
-	if (TSharedPtr<FRigVMEditorBase> Editor = WeakEditor.Pin())
+	if (TSharedPtr<IRigVMEditor> Editor = WeakEditor.Pin())
 	{
 		Host = Editor->GetRigVMHost();
 	}
@@ -1178,7 +1178,7 @@ void SRigVMExecutionStackView::HandleHostInitializedEvent(URigVMHost* InHost, co
 		}
 	}
 	
-	if (TSharedPtr<FRigVMEditorBase> Editor = WeakEditor.Pin())
+	if (TSharedPtr<IRigVMEditor> Editor = WeakEditor.Pin())
 	{
 		InHost->OnInitialized_AnyThread().Remove(OnHostInitializedHandle);
 		OnHostInitializedHandle.Reset();
@@ -1200,7 +1200,7 @@ void SRigVMExecutionStackView::HandleHostInitializedEvent(URigVMHost* InHost, co
 	}
 }
 
-void SRigVMExecutionStackView::HandlePreviewHostUpdated(FRigVMEditorBase* InEditor)
+void SRigVMExecutionStackView::HandlePreviewHostUpdated(IRigVMEditor* InEditor)
 {
 	if(URigVMHost* RigVMHost = InEditor->GetRigVMHost())
 	{
@@ -1210,7 +1210,7 @@ void SRigVMExecutionStackView::HandlePreviewHostUpdated(FRigVMEditorBase* InEdit
 
 void SRigVMExecutionStackView::HandleItemMouseDoubleClick(TSharedPtr<FRigStackEntry> InItem)
 {
-	if (TSharedPtr<FRigVMEditorBase> Editor = WeakEditor.Pin())
+	if (TSharedPtr<IRigVMEditor> Editor = WeakEditor.Pin())
 	{
 		if (!RigVMBlueprint.IsValid())
 		{
diff --git a/Engine/Plugins/Runtime/RigVM/Source/RigVMEditor/Public/Editor/RigVMEditor.h b/Engine/Plugins/Runtime/RigVM/Source/RigVMEditor/Public/Editor/RigVMEditor.h
--- a/Engine/Plugins/Runtime/RigVM/Source/RigVMEditor/Public/Editor/RigVMEditor.h
+++ b/Engine/Plugins/Runtime/RigVM/Source/RigVMEditor/Public/Editor/RigVMEditor.h
@@ -34,7 +34,7 @@
 class URigVMBlueprint;
 class SGraphEditor;
 
-DECLARE_MULTICAST_DELEGATE_TwoParams(FRigVMEditorClosed, const FRigVMEditorBase*, URigVMBlueprint*);
+DECLARE_MULTICAST_DELEGATE_TwoParams(FRigVMEditorClosed, const IRigVMEditor*, URigVMBlueprint*);
 
 
 /* Enums to use when grouping the blueprint members in the list panel. The order here will determine the order in the list */
@@ -99,11 +99,83 @@
 class RIGVMEDITOR_API IRigVMEditor
 {
 public:
+
+	virtual TSharedPtr<FAssetEditorToolkit> GetHostingApp() = 0;
+	virtual URigVMBlueprint* GetRigVMBlueprint() const = 0;
+	virtual const TArray< UObject* >* GetObjectsCurrentlyBeingEdited() const = 0;
+	virtual URigVMHost* GetRigVMHost() const = 0;
+
+	virtual TSharedPtr<FTabManager> GetTabManager() = 0;
+	virtual FName GetGraphExplorerWidgetID() = 0;
+	virtual TSharedPtr<class SRigVMDetailsInspector> GetRigVMInspector() const = 0;
+	virtual TSharedPtr<SRigVMEditorGraphExplorer> GetGraphExplorerWidget() = 0;
+#if WITH_RIGVMLEGACYEDITOR
+	virtual TSharedPtr<class SKismetInspector> GetKismetInspector() const = 0;
+#endif
 	
+	virtual bool GetIsContextSensitive() = 0;
+	virtual void SetIsContextSensitive(const bool bIsContextSensitive) = 0;
+
+	virtual void SetGraphExplorerWidget(TSharedPtr<SRigVMEditorGraphExplorer> InWidget) = 0;
+	
+	virtual void GetPinTypeSelectorFilters(TArray<TSharedPtr<class IPinTypeSelectorFilter>>& OutFilters) const = 0;
+
+	DECLARE_EVENT(IRigVMEditor, FOnRefreshEvent);
+	virtual FOnRefreshEvent OnRefresh() = 0;
+	virtual void ForceEditorRefresh(ERefreshRigVMEditorReason::Type Reason = ERefreshRigVMEditorReason::UnknownReason) = 0;
+
+	DECLARE_EVENT_OneParam(IRigVMEditor, FPreviewHostUpdated, IRigVMEditor*);
+	virtual FPreviewHostUpdated& OnPreviewHostUpdated()  = 0;
+
+	virtual FRigVMEditorClosed& OnEditorClosed() = 0;
+	virtual UEdGraph* GetFocusedGraph() const = 0;
+	virtual URigVMGraph* GetFocusedModel() const = 0;
+	virtual FNotifyHook* GetNotifyHook() = 0;
+	virtual TWeakPtr<class SGraphEditor> GetFocusedGraphEditor() = 0;
+
+	virtual bool InEditingMode() const = 0;
+	virtual bool IsEditable(UEdGraph* InGraph) const = 0;
+	
+	virtual void OnGraphNodeDropToPerform(TSharedPtr<FDragDropOperation> InDragDropOp, UEdGraph* InGraph, const FVector2D& InNodePosition, const FVector2D& InScreenPosition) = 0;
+	virtual TSharedRef<FUICommandList> GetToolkitCommands() = 0;
+
+	virtual bool SelectLocalVariable(const UEdGraph* Graph, const FName& VariableName) = 0;
+	virtual void JumpToHyperlink(const UObject* ObjectReference, bool bRequestRename = false) = 0;
+	virtual void OnAddNewLocalVariable() = 0;
+	virtual bool CanAddNewLocalVariable() const = 0;
+	virtual void OnAddNewVariable() = 0;
+	virtual void OnPasteNewLocalVariable(const FBPVariableDescription& VariableDescription) = 0;
+	virtual void AddNewFunctionVariant(const UEdGraph* InOriginalFunction) = 0;
+	virtual TSharedPtr<SDockTab> OpenDocument(const UObject* DocumentID, FDocumentTracker::EOpenDocumentCause Cause) = 0;
+	
+	// Type of new document/graph being created by a menu item
+	enum ECreatedDocumentType
+	{
+		CGT_NewVariable,
+		CGT_NewFunctionGraph,
+		CGT_NewMacroGraph,
+		CGT_NewAnimationLayer,
+		CGT_NewEventGraph,
+		CGT_NewLocalVariable
+	};
+	virtual void OnNewDocumentClicked(ECreatedDocumentType GraphType) = 0;
+	virtual void SummonSearchUI(bool bSetFindWithinBlueprint, FString NewSearchTerms = FString(), bool bSelectFirstResult = false) = 0;
+	
+	virtual void GetDebugObjects(TArray<FRigVMCustomDebugObject>& DebugList) const = 0;
+	virtual bool OnlyShowCustomDebugObjects() const = 0;
+	virtual FString GetCustomDebugObjectLabel(UObject* ObjectBeingDebugged) const = 0;
+
+	virtual TSharedPtr<SGraphEditor> OpenGraphAndBringToFront(UEdGraph* Graph, bool bSetFocus = true) = 0;
+	virtual void ZoomToSelection_Clicked() = 0;
+
+	virtual void HandleModifiedEvent(ERigVMGraphNotifType InNotifType, URigVMGraph* InGraph, UObject* InSubject) = 0;
+	
+protected:
+	
 	virtual TSharedRef<IRigVMEditor> SharedRef() = 0;
 	virtual TSharedRef<const IRigVMEditor> SharedRef() const = 0;
 
-	virtual TSharedPtr<FAssetEditorToolkit> GetHostingApp() = 0;
+	
 	virtual const TSharedPtr<FAssetEditorToolkit> GetHostingApp() const = 0;
 	virtual void InitAssetEditor(const EToolkitMode::Type Mode, const TSharedPtr<IToolkitHost>& InitToolkitHost, const FName AppIdentifier, const TSharedRef<FTabManager::FLayout>& StandaloneDefaultLayout, const bool bCreateDefaultStandaloneMenu, const bool bCreateDefaultToolbar, const TArray<UObject*>& ObjectsToEdit, const bool bInIsToolbarFocusable = false, const bool bInUseSmallToolbarIcons = false, const TOptional<EAssetOpenMethod>& InOpenMethod = TOptional<EAssetOpenMethod>()) = 0;
 	virtual void CreateEditorToolbar() = 0;
@@ -114,15 +186,8 @@
 	virtual void SetCurrentMode(FName NewMode) = 0;
 	virtual FEditorModeTools& GetToolkitEditorModeManager() const = 0;
 	virtual void PostLayoutBlueprintEditorInitialization() = 0;
-	virtual TSharedPtr<SGraphEditor> OpenGraphAndBringToFront(UEdGraph* Graph, bool bSetFocus = true) = 0;
 	virtual bool FindOpenTabsContainingDocument(const UObject* DocumentID, /*inout*/ TArray< TSharedPtr<SDockTab> >& Results) = 0;
-	virtual TSharedPtr<SDockTab> OpenDocument(const UObject* DocumentID, FDocumentTracker::EOpenDocumentCause Cause) = 0;
 	virtual void CloseDocumentTab(const UObject* DocumentID) = 0;
-	virtual TSharedPtr<FTabManager> GetTabManager() = 0;
-#if WITH_RIGVMLEGACYEDITOR
-	virtual TSharedPtr<class SKismetInspector> GetKismetInspector() const = 0;
-#endif
-	virtual TSharedPtr<class SRigVMDetailsInspector> GetRigVMInspector() const = 0;
 	virtual TSharedPtr<FApplicationMode> CreateEditorMode() = 0;
 	virtual const FName GetEditorAppName() const = 0;
 	virtual const TArray< UObject* >& GetEditingBlueprints() const = 0;
@@ -130,8 +195,6 @@
 	virtual bool IsHosted() const = 0;
 	virtual void BringToolkitToFrontImpl() = 0;
 	virtual void InitToolMenuContext(FToolMenuContext& MenuContext) = 0;
-	virtual TSharedRef<FUICommandList> GetToolkitCommands() = 0;
-	virtual TWeakPtr<class SGraphEditor> GetFocusedGraphEditor() = 0;
 	virtual TWeakPtr<FDocumentTabFactory> GetGraphEditorTabFactory() const = 0;
 	virtual bool TransactionObjectAffectsBlueprint(UObject* InTransactedObject) = 0;
 	virtual FEdGraphPinType GetLastPinTypeUsed() = 0;
@@ -140,7 +203,6 @@
 	virtual FGraphPanelSelectionSet GetSelectedNodes() const = 0;
 	virtual void SetUISelectionState(FName SelectionOwner) = 0;
 	virtual void AnalyticsTrackNodeEvent(UBlueprint* Blueprint, UEdGraphNode *GraphNode, bool bNodeDelete = false) const = 0;
-	virtual void JumpToHyperlink(const UObject* ObjectReference, bool bRequestRename = false) = 0;
 	virtual void PostUndo(bool bSuccess) = 0;
 	virtual void PostRedo(bool bSuccess) = 0;
 	virtual UEdGraphPin* GetCurrentlySelectedPin() const = 0;
@@ -151,7 +213,6 @@
 	virtual void SaveAssetAs_Execute_Impl() = 0;
 	virtual bool IsGraphInCurrentBlueprint(const UEdGraph* InGraph) const = 0;
 	virtual bool IsEditableImpl(UEdGraph* InGraph) const = 0;
-	virtual UEdGraph* GetFocusedGraph() const = 0;
 	virtual void JumpToNode(const class UEdGraphNode* Node, bool bRequestRename = false) = 0;
 	virtual void JumpToPin(const class UEdGraphPin* Pin) = 0;
 	virtual void AddToolbarExtender(TSharedPtr<FExtender> Extender) = 0;
@@ -172,50 +233,19 @@
 	virtual void EditClassDefaults_Clicked() = 0;
 	virtual void EditGlobalOptions_Clicked() = 0;
 	virtual void TryInvokingDetailsTab(bool bFlash = true) = 0;
-	virtual FName GetGraphExplorerWidgetID() = 0;
 	virtual void RefreshInspector() = 0;
 	virtual void RefreshStandAloneDefaultsEditor() = 0;
-	virtual void GetPinTypeSelectorFilters(TArray<TSharedPtr<class IPinTypeSelectorFilter>>& OutFilters) const = 0;
-	virtual void OnAddNewVariable() = 0;
-	virtual void SummonSearchUI(bool bSetFindWithinBlueprint, FString NewSearchTerms = FString(), bool bSelectFirstResult = false) = 0;
-	virtual void ZoomToSelection_Clicked() = 0;
 	virtual void RestoreEditedObjectState() = 0;
 	virtual void SetupViewForBlueprintEditingMode() = 0;
-	virtual bool GetIsContextSensitive() = 0;
-	virtual void SetIsContextSensitive(const bool bIsContextSensitive) = 0;
 	virtual void RegisterToolbarTab(const TSharedRef<class FTabManager>& TabManager) = 0;
-	virtual const TArray< UObject* >* GetObjectsCurrentlyBeingEdited() const = 0;
 	virtual void AddCompileWidget(FToolBarBuilder& ToolbarBuilder) = 0;
 	virtual void AddSelectedDebugObjectWidget(FToolBarBuilder& ToolbarBuilder) = 0;
-	virtual FNotifyHook* GetNotifyHook() = 0;
 };
 
 class RIGVMEDITOR_API FRigVMEditorBase : public IRigVMEditor
 {
-protected:
-	// Type of new document/graph being created by a menu item
-	enum ECreatedDocumentType
-	{
-		CGT_NewVariable,
-		CGT_NewFunctionGraph,
-		CGT_NewMacroGraph,
-		CGT_NewAnimationLayer,
-		CGT_NewEventGraph,
-		CGT_NewLocalVariable
-	};
-	
-	FRigVMEditorBase();
-	virtual ~FRigVMEditorBase(){}
-
-
-	void UnbindEditor();
-
 public:
-	static FRigVMEditorBase* GetFromAssetEditorInstance(IAssetEditorInstance* Instance);
-
-	DECLARE_EVENT(FRigVMEditorBase, FOnRefreshEvent);
-	FOnRefreshEvent OnRefresh() { return OnRefreshEvent; }
-
+	
 	/**
 	 * Edits the specified asset(s)
 	 *
@@ -224,8 +254,46 @@
 	 * @param	InRigVMBlueprint	The blueprint object to start editing.
 	 */
 	virtual void InitRigVMEditor(const EToolkitMode::Type Mode, const TSharedPtr< class IToolkitHost >& InitToolkitHost, class URigVMBlueprint* InRigVMBlueprint);
+
+	static FRigVMEditorBase* GetFromAssetEditorInstance(IAssetEditorInstance* Instance);
+	// returns the blueprint being edited
+	virtual URigVMBlueprint* GetRigVMBlueprint() const override;
 	
+	void HandleJumpToHyperlink(const UObject* InSubject);
+
+	void Compile();
+	bool IsCompilingEnabled() const;
+
+	void DeleteSelectedNodes();
+	bool CanDeleteNodes() const;
+	void CopySelectedNodes();
+	bool CanCopyNodes() const;
+	void PasteNodes();
+	bool CanPasteNodes() const;
+	void CutSelectedNodes();
+	bool CanCutNodes() const;
+	void DuplicateNodes();
+	bool CanDuplicateNodes() const;
+	
+	void OnStartWatchingPin();
+	bool CanStartWatchingPin() const;
+	void OnStopWatchingPin();
+	bool CanStopWatchingPin() const;
+
+	virtual FOnRefreshEvent OnRefresh() override { return OnRefreshEvent; }
+
+	FText GetGraphDecorationString(UEdGraph* InGraph) const;
+	virtual bool IsEditable(UEdGraph* InGraph) const override;
+
+	virtual void OnGraphNodeDropToPerform(TSharedPtr<FDragDropOperation> InDragDropOp, UEdGraph* InGraph, const FVector2D& InNodePosition, const FVector2D& InScreenPosition) override;
+	
 protected:
+	
+	FRigVMEditorBase();
+	virtual ~FRigVMEditorBase(){}
 
+	void UnbindEditor();
+	
+
 	void HandleAssetRequestedOpen(UObject* InObject);
 	void HandleAssetRequestClose(UObject* InObject, EAssetEditorCloseReason InReason);
@@ -234,10 +302,8 @@
 	virtual const FName GetEditorModeName() const;
 
 	// FBlueprintEditor interface
-public:
 	UBlueprint* GetBlueprintObj() const;
-	bool InEditingMode() const;
-protected:
+	virtual bool InEditingMode() const override;
 	TSubclassOf<UEdGraphSchema> GetDefaultSchemaClass() const;
 	void OnGraphEditorFocused(const TSharedRef<SGraphEditor>& InGraphEditor);
 
@@ -256,26 +322,10 @@
 
 	// BlueprintEditor interface
 	bool TransactionObjectAffectsBlueprintImpl(UObject* InTransactedObject);
-	bool CanAddNewLocalVariable() const;
-	void OnAddNewLocalVariable();
-	void OnPasteNewLocalVariable(const FBPVariableDescription& VariableDescription);
-public:
-	void DeleteSelectedNodes();
-	bool CanDeleteNodes() const;
-	void CopySelectedNodes();
-	bool CanCopyNodes() const;
-	void PasteNodes();
-	bool CanPasteNodes() const;
-	void CutSelectedNodes();
-	bool CanCutNodes() const;
-	void DuplicateNodes();
-	bool CanDuplicateNodes() const;
+	virtual bool CanAddNewLocalVariable() const override;
+	virtual void OnAddNewLocalVariable() override;
+	virtual void OnPasteNewLocalVariable(const FBPVariableDescription& VariableDescription) override;
 
-	void OnStartWatchingPin();
-	bool CanStartWatchingPin() const;
-	void OnStopWatchingPin();
-	bool CanStopWatchingPin() const;
-protected:
 
 	bool IsNativeParentClassCodeLinkEnabled() const { return false; }
 	bool ReparentBlueprint_IsVisible() const { return false; }
@@ -283,7 +333,7 @@
 	bool ShouldLoadBPLibrariesFromAssetRegistry() { return false; }
 	bool JumpToHyperlinkImpl(const UObject* ObjectReference, bool bRequestRename = false);
 	bool ShouldOpenGraphByDefault() const { return true; }
-	void AddNewFunctionVariant(const UEdGraph* InOriginalFunction);
+	virtual void AddNewFunctionVariant(const UEdGraph* InOriginalFunction) override;
 
 	// FEditorUndoClient Interface
 	void PostUndoImpl(bool bSuccess);
@@ -299,14 +349,8 @@
 	//  FTickableEditorObject Interface
 	TStatId GetStatId() const;
 
-public:
-	// returns the blueprint being edited
-	URigVMBlueprint* GetRigVMBlueprint() const;
-
 	// returns the currently debugged / viewed host
-	URigVMHost* GetRigVMHost() const;
-	
-protected:
+	virtual URigVMHost* GetRigVMHost() const override;
 
 	virtual UObject* GetOuterForHost() const;
 
@@ -316,19 +360,16 @@
 	// allows the editor to fill an empty graph
 	virtual void CreateEmptyGraphContent(URigVMController* InController) {}
 
-	DECLARE_EVENT_OneParam(FRigVMEditorBase, FPreviewHostUpdated, FRigVMEditorBase*);
-	FPreviewHostUpdated& OnPreviewHostUpdated() { return PreviewHostUpdated;  }
+	virtual FPreviewHostUpdated& OnPreviewHostUpdated() override { return PreviewHostUpdated;  }
 
-public:
-	FRigVMEditorClosed& OnEditorClosed() { return RigVMEditorClosedDelegate; }
-protected:
+	virtual FRigVMEditorClosed& OnEditorClosed() override { return RigVMEditorClosedDelegate; }
 
+
 	/** Get the toolbox hosting widget */
 	TSharedRef<SBorder> GetToolbox() { return Toolbox.ToSharedRef(); }
 
-	bool SelectLocalVariable(const UEdGraph* Graph, const FName& VariableName);
+	virtual bool SelectLocalVariable(const UEdGraph* Graph, const FName& VariableName) override;
 
-protected:
 
 	enum ERigVMEditorExecutionModeType
 	{
@@ -339,14 +380,11 @@
 	// FBlueprintEditor Interface
 	void CreateDefaultCommandsImpl();
 	void OnCreateGraphEditorCommands(TSharedPtr<FUICommandList> GraphEditorCommandsList);
-public:
-	void Compile();
-	bool IsCompilingEnabled() const;
-protected:
+
 	void SaveAsset_Execute();
 	void SaveAssetAs_Execute();
 	bool IsInAScriptingMode() const { return true; }
-	void NewDocument_OnClicked(ECreatedDocumentType GraphType);
+	virtual void OnNewDocumentClicked(ECreatedDocumentType GraphType) override;
 	bool IsSectionVisibleImpl(RigVMNodeSectionID::Type InSectionID) const;
 	bool AreEventGraphsAllowed() const;
 	bool AreMacrosAllowed() const;
@@ -356,14 +394,14 @@
 	void OnActiveTabChanged( TSharedPtr<SDockTab> PreviouslyActive, TSharedPtr<SDockTab> NewlyActivated );
 	void OnSelectedNodesChangedImpl(const TSet<class UObject*>& NewSelection);
 	void OnBlueprintChangedImpl(UBlueprint* InBlueprint, bool bIsJustBeingCompiled);
-	void RefreshEditors(ERefreshRigVMEditorReason::Type Reason = ERefreshRigVMEditorReason::UnknownReason);
+	virtual void ForceEditorRefresh(ERefreshRigVMEditorReason::Type Reason = ERefreshRigVMEditorReason::UnknownReason) override;
 	void SetupGraphEditorEvents(UEdGraph* InGraph, SGraphEditor::FGraphEditorEvents& InEvents);
 	void FocusInspectorOnGraphSelection(const TSet<class UObject*>& NewSelection, bool bForceRefresh = false);
 #if WITH_RIGVMLEGACYEDITOR
 	virtual TSharedPtr<SKismetInspector> GetKismetInspector() const override { return nullptr; }
 #endif
 
-	virtual void HandleModifiedEvent(ERigVMGraphNotifType InNotifType, URigVMGraph* InGraph, UObject* InSubject);
+	virtual void HandleModifiedEvent(ERigVMGraphNotifType InNotifType, URigVMGraph* InGraph, UObject* InSubject) override;
 	UE_DEPRECATED(5.4, "Please use HandleVMCompiledEvent with ExtendedExecuteContext param.")
 	virtual void HandleVMCompiledEvent(UObject* InCompiledObject, URigVM* InVM) {}
 	virtual void HandleVMCompiledEvent(UObject* InCompiledObject, URigVM* InVM, FRigVMExtendedExecuteContext& InContext);
@@ -383,12 +421,6 @@
 	bool OnRequestBreakLinksDialog(TArray<URigVMLink*> InLinks);
 	TRigVMTypeIndex OnRequestPinTypeSelectionDialog(const TArray<TRigVMTypeIndex>& InTypes);
 
-public:
-	void HandleJumpToHyperlink(const UObject* InSubject);
-	FText GetGraphDecorationString(UEdGraph* InGraph) const;
-	bool IsEditable(UEdGraph* InGraph) const;
-
-protected:
 	bool UpdateDefaultValueForVariable(FBPVariableDescription& InVariable, bool bUseCDO);
 
 	URigVMController* ActiveController;
@@ -447,7 +479,7 @@
 
 	void SetHost(URigVMHost* InHost);
 
-	URigVMGraph* GetFocusedModel() const;
+	virtual URigVMGraph* GetFocusedModel() const override;
 	URigVMController* GetFocusedController() const;
 	TSharedPtr<SGraphEditor> GetGraphEditor(UEdGraph* InEdGraph) const;
 
@@ -459,8 +491,6 @@
 	
 	/** Fill the toolbar with content */
 	virtual void FillToolbar(FToolBarBuilder& ToolbarBuilder, bool bEndSection = true);
-
-public:
 	
 	virtual TArray<FName> GetDefaultEventQueue() const;
 	TArray<FName> GetEventQueue() const;
@@ -471,20 +501,16 @@
 	virtual FSlateIcon GetEventQueueIcon(const TArray<FName>& InEventQueue) const;
 	FSlateIcon GetEventQueueIcon() const;
 
-	virtual void OnGraphNodeDropToPerform(TSharedPtr<FDragDropOperation> InDragDropOp, UEdGraph* InGraph, const FVector2D& InNodePosition, const FVector2D& InScreenPosition);
-
-protected:
-
 	void SetExecutionMode(const ERigVMEditorExecutionModeType InExecutionMode);
 	int32 GetExecutionModeComboValue() const;
 	FText GetExecutionModeLabel() const;
 	static FSlateIcon GetExecutionModeIcon(const ERigVMEditorExecutionModeType InExecutionMode);
 	FSlateIcon GetExecutionModeIcon() const;
 	
-	void GetCustomDebugObjects(TArray<FRigVMCustomDebugObject>& DebugList) const;
-	bool OnlyShowCustomDebugObjects() const { return true; }
+	virtual void GetDebugObjects(TArray<FRigVMCustomDebugObject>& DebugList) const override;
+	virtual bool OnlyShowCustomDebugObjects() const override { return true; }
 	void HandleSetObjectBeingDebugged(UObject* InObject);
-	FString GetCustomDebugObjectLabel(UObject* ObjectBeingDebugged) const;
+	virtual FString GetCustomDebugObjectLabel(UObject* ObjectBeingDebugged) const override;
 
 	/** Handle hiding items in the graph */
 	void HandleHideItem();
@@ -516,8 +542,6 @@
 	/** Returns true if PIE is currently running */
 	static bool IsPIERunning();
 
-private:
-
 	void OnPIEStopped(bool bSimulation);
 
 	/** Our currently running rig vm instance */
@@ -528,18 +552,16 @@
 	/** Toolbox hosting widget */
 	TSharedPtr<SBorder> Toolbox;
 
-protected:
 	TSharedPtr<SRigVMEditorGraphExplorer> GraphExplorerWidget;
 
 	FRigVMEditorClosed RigVMEditorClosedDelegate;
 
-public:
-	void SetGraphExplorerWidget(TSharedPtr<SRigVMEditorGraphExplorer> InWidget) { GraphExplorerWidget = InWidget; }
-	TSharedPtr<SRigVMEditorGraphExplorer> GetGraphExplorerWidget() { return GraphExplorerWidget; }
+	virtual void SetGraphExplorerWidget(TSharedPtr<SRigVMEditorGraphExplorer> InWidget) override { GraphExplorerWidget = InWidget; }
+	virtual TSharedPtr<SRigVMEditorGraphExplorer> GetGraphExplorerWidget() override { return GraphExplorerWidget; }
 
 	bool IsEditingSingleBlueprint() const;
 	
-private:
+protected:
 	bool bAnyErrorsLeft;
 	TMap<FString, FString> KnownInstructionLimitWarnings;
 	URigVMNode* HaltedAtNode;
@@ -549,10 +571,8 @@
 	bool bAllowBulkEdits;
 	bool bIsSettingObjectBeingDebugged;
 
-protected:
 	bool bRigVMEditorInitialized;
 
-private:
 	/** Are we currently compiling through the user interface */
 	bool bIsCompilingThroughUI;
 
@@ -563,8 +583,6 @@
 	/** The log to use for errors resulting from the init phase of the units */
 	FRigVMLog RigVMLog;
 	
-
-protected:
 	TArray<FName> LastEventQueue;
 
 	/** The extender to pass to the level editor to extend it's window menu */
diff --git a/Engine/Plugins/Runtime/RigVM/Source/RigVMEditor/Public/Editor/RigVMEditorGraphExplorerTabSummoner.h b/Engine/Plugins/Runtime/RigVM/Source/RigVMEditor/Public/Editor/RigVMEditorGraphExplorerTabSummoner.h
--- a/Engine/Plugins/Runtime/RigVM/Source/RigVMEditor/Public/Editor/RigVMEditorGraphExplorerTabSummoner.h
+++ b/Engine/Plugins/Runtime/RigVM/Source/RigVMEditor/Public/Editor/RigVMEditorGraphExplorerTabSummoner.h
@@ -5,7 +5,7 @@
 #include "CoreMinimal.h"
 #include "WorkflowOrientedApp/WorkflowTabFactory.h"
 
-class FRigVMEditorBase;
+class IRigVMEditor;
 
 struct RIGVMEDITOR_API FRigVMEditorGraphExplorerTabSummoner : public FWorkflowTabFactory
 {
@@ -14,10 +14,10 @@ public:
 	static const FName TabID() { static FName ID = TEXT("RigVM Graph Explorer"); return ID; }
 	
 public:
-	FRigVMEditorGraphExplorerTabSummoner(const TSharedRef<FRigVMEditorBase>& InRigVMEditor);
+	FRigVMEditorGraphExplorerTabSummoner(const TSharedRef<IRigVMEditor>& InRigVMEditor);
 	
 	virtual TSharedRef<SWidget> CreateTabBody(const FWorkflowTabSpawnInfo& Info) const override;
 	
 protected:
-	TWeakPtr<FRigVMEditorBase> RigVMEditor;
+	TWeakPtr<IRigVMEditor> RigVMEditor;
 };
diff --git a/Engine/Plugins/Runtime/RigVM/Source/RigVMEditor/Public/Editor/RigVMEditorMenuContext.h b/Engine/Plugins/Runtime/RigVM/Source/RigVMEditor/Public/Editor/RigVMEditorMenuContext.h
--- a/Engine/Plugins/Runtime/RigVM/Source/RigVMEditor/Public/Editor/RigVMEditorMenuContext.h
+++ b/Engine/Plugins/Runtime/RigVM/Source/RigVMEditor/Public/Editor/RigVMEditorMenuContext.h
@@ -9,7 +9,7 @@
 
 class URigVMHost;
 class URigVMBlueprint;
-class FRigVMEditorBase;
+class IRigVMEditor;
 class URigVMGraph;
 class URigVMNode;
 class URigVMPin;
@@ -57,7 +57,7 @@ public:
 	 * @param InRigVMEditor 	    The RigVM Editor hosting the menus
 	 * @param InGraphMenuContext 	Additional context for specific menus
 	*/
-	void Init(TWeakPtr<FRigVMEditorBase> InRigVMEditor, const FRigVMEditorGraphMenuContext& InGraphMenuContext = FRigVMEditorGraphMenuContext());
+	void Init(TWeakPtr<IRigVMEditor> InRigVMEditor, const FRigVMEditorGraphMenuContext& InGraphMenuContext = FRigVMEditorGraphMenuContext());
 
 	/** Get the rigvm blueprint that we are editing */
 	UFUNCTION(BlueprintCallable, Category = RigVMEditor)
@@ -75,10 +75,10 @@ public:
 	UFUNCTION(BlueprintCallable, Category = RigVMEditor)
 	FRigVMEditorGraphMenuContext GetGraphMenuContext();
 	
-	FRigVMEditorBase* GetRigVMEditor() const;
+	IRigVMEditor* GetRigVMEditor() const;
 
 private:
-	TWeakPtr<FRigVMEditorBase> WeakRigVMEditor;
+	TWeakPtr<IRigVMEditor> WeakRigVMEditor;
 
 	FRigVMEditorGraphMenuContext GraphMenuContext;
 };
\ No newline at end of file
diff --git a/Engine/Plugins/Runtime/RigVM/Source/RigVMEditor/Public/Editor/RigVMExecutionStackTabSummoner.h b/Engine/Plugins/Runtime/RigVM/Source/RigVMEditor/Public/Editor/RigVMExecutionStackTabSummoner.h
--- a/Engine/Plugins/Runtime/RigVM/Source/RigVMEditor/Public/Editor/RigVMExecutionStackTabSummoner.h
+++ b/Engine/Plugins/Runtime/RigVM/Source/RigVMEditor/Public/Editor/RigVMExecutionStackTabSummoner.h
@@ -5,7 +5,7 @@
 #include "CoreMinimal.h"
 #include "WorkflowOrientedApp/WorkflowTabFactory.h"
 
-class FRigVMEditorBase;
+class IRigVMEditor;
 
 struct RIGVMEDITOR_API FRigVMExecutionStackTabSummoner : public FWorkflowTabFactory
 {
@@ -14,10 +14,10 @@ public:
 	static inline const FLazyName TabID = FLazyName(TEXT("Execution Stack"));
 	
 public:
-	FRigVMExecutionStackTabSummoner(const TSharedRef<FRigVMEditorBase>& InRigVMEditor);
+	FRigVMExecutionStackTabSummoner(const TSharedRef<IRigVMEditor>& InRigVMEditor);
 	
 	virtual TSharedRef<SWidget> CreateTabBody(const FWorkflowTabSpawnInfo& Info) const override;
 	
 protected:
-	TWeakPtr<FRigVMEditorBase> RigVMEditor;
+	TWeakPtr<IRigVMEditor> RigVMEditor;
 };
diff --git a/Engine/Plugins/Runtime/RigVM/Source/RigVMEditor/Public/Editor/RigVMGraphDetailCustomization.h b/Engine/Plugins/Runtime/RigVM/Source/RigVMEditor/Public/Editor/RigVMGraphDetailCustomization.h
--- a/Engine/Plugins/Runtime/RigVM/Source/RigVMEditor/Public/Editor/RigVMGraphDetailCustomization.h
+++ b/Engine/Plugins/Runtime/RigVM/Source/RigVMEditor/Public/Editor/RigVMGraphDetailCustomization.h
@@ -30,7 +30,7 @@ public:
 	FRigVMFunctionArgumentGroupLayout(
 		const TWeakObjectPtr<URigVMGraph>& InGraph,
 		const TWeakInterfacePtr<IRigVMClientHost>& InRigVMClientHost,
-		const TWeakPtr<FRigVMEditorBase>& InEditor,
+		const TWeakPtr<IRigVMEditor>& InEditor,
 		bool bInputs);
 	virtual ~FRigVMFunctionArgumentGroupLayout();
 
@@ -50,7 +50,7 @@ private:
 
 	TWeakObjectPtr<URigVMGraph> GraphPtr;
 	TWeakInterfacePtr<IRigVMClientHost> WeakRigVMClientHost;
-	TWeakPtr<FRigVMEditorBase> RigVMEditorPtr;
+	TWeakPtr<IRigVMEditor> RigVMEditorPtr;
 	bool bIsInputGroup;
 	FSimpleDelegate OnRebuildChildren;
 };
@@ -63,7 +63,7 @@ public:
 		const TWeakObjectPtr<URigVMPin>& InPin,
 		const TWeakObjectPtr<URigVMGraph>& InGraph, 
 		const TWeakInterfacePtr<IRigVMClientHost>& InRigVMClientHost,
-		const TWeakPtr<FRigVMEditorBase>& InEditor)
+		const TWeakPtr<IRigVMEditor>& InEditor)
 		: PinPtr(InPin)
 		, GraphPtr(InGraph)
 		, WeakRigVMClientHost(InRigVMClientHost)
@@ -117,7 +117,7 @@ private:
 	TWeakInterfacePtr<IRigVMClientHost> WeakRigVMClientHost;
 
 	/** The editor we are editing */
-	TWeakPtr<FRigVMEditorBase> RigVMEditorPtr;
+	TWeakPtr<IRigVMEditor> RigVMEditorPtr;
 
 	/** Holds a weak pointer to the argument name widget, used for error notifications */
 	TWeakPtr<SEditableTextBox> ArgumentNameWidget;
@@ -163,8 +163,8 @@ class RIGVMEDITOR_API FRigVMGraphDetailCustomization : public IDetailCustomizati
 {
 public:
 	/** Makes a new instance of this detail layout class for a specific detail view requesting it */
-	static TSharedRef<IDetailCustomization> MakeInstance(TSharedPtr<FRigVMEditorBase> InBlueprintEditor, const UClass* InExpectedBlueprintClass);
-	FRigVMGraphDetailCustomization(TSharedPtr<FRigVMEditorBase> RigVMigEditor, URigVMBlueprint* RigVMBlueprint);
+	static TSharedRef<IDetailCustomization> MakeInstance(TSharedPtr<IRigVMEditor> InBlueprintEditor, const UClass* InExpectedBlueprintClass);
+	FRigVMGraphDetailCustomization(TSharedPtr<IRigVMEditor> RigVMigEditor, URigVMBlueprint* RigVMBlueprint);
 #if WITH_RIGVMLEGACYEDITOR
 	static TSharedPtr<IDetailCustomization> MakeLegacyInstance(TSharedPtr<IBlueprintEditor> InBlueprintEditor, const UClass* InExpectedBlueprintClass);
 	FRigVMGraphDetailCustomization(TSharedPtr<IBlueprintEditor> RigVMigEditor, URigVMBlueprint* RigVMBlueprint);
@@ -176,7 +176,7 @@ public:
 
 private:
 	/** The editor we are embedded in */
-	TWeakPtr<FRigVMEditorBase> RigVMEditorPtr;
+	TWeakPtr<IRigVMEditor> RigVMEditorPtr;
 
 	/** The graph we are editing */
 	TWeakObjectPtr<URigVMEdGraph> GraphPtr;
@@ -194,7 +194,7 @@ public:
 		URigVMGraph* Model,
 		URigVMController* Controller,
 		IRigVMClientHost* InRigVMClientHost,
-		TWeakPtr<FRigVMEditorBase> InEditor);
+		TWeakPtr<IRigVMEditor> InEditor);
 
 private:
 
@@ -259,7 +259,7 @@ private:
 	TWeakObjectPtr<URigVMController> WeakController;
 
 	/** The editor we are embedded in */
-	TWeakPtr<FRigVMEditorBase> RigVMEditorPtr;
+	TWeakPtr<IRigVMEditor> RigVMEditorPtr;
 
 	/** The asset host we are editing */
 	TWeakInterfacePtr<IRigVMClientHost> RigVMClientHost;
diff --git a/Engine/Plugins/Runtime/RigVM/Source/RigVMEditor/Public/Editor/RigVMLegacyEditor.h b/Engine/Plugins/Runtime/RigVM/Source/RigVMEditor/Public/Editor/RigVMLegacyEditor.h
--- a/Engine/Plugins/Runtime/RigVM/Source/RigVMEditor/Public/Editor/RigVMLegacyEditor.h
+++ b/Engine/Plugins/Runtime/RigVM/Source/RigVMEditor/Public/Editor/RigVMLegacyEditor.h
@@ -154,7 +154,7 @@ protected:
 	virtual void SaveAsset_Execute() override { return FRigVMEditorBase::SaveAsset_Execute(); }
 	virtual void SaveAssetAs_Execute() override { return FRigVMEditorBase::SaveAssetAs_Execute(); }
 	virtual bool IsInAScriptingMode() const override { return FRigVMEditorBase::IsInAScriptingMode(); }
-	virtual void NewDocument_OnClicked(FBlueprintEditor::ECreatedDocumentType GraphType) override { return FRigVMEditorBase::NewDocument_OnClicked((FRigVMEditorBase::ECreatedDocumentType)GraphType); }
+	virtual void NewDocument_OnClicked(FBlueprintEditor::ECreatedDocumentType GraphType) override { return FRigVMEditorBase::OnNewDocumentClicked((FRigVMEditorBase::ECreatedDocumentType)GraphType); }
 	virtual bool AreEventGraphsAllowed() const override { return FRigVMEditorBase::AreEventGraphsAllowed(); }
 	virtual bool AreMacrosAllowed() const override { return FRigVMEditorBase::AreMacrosAllowed(); }
 	virtual bool AreDelegatesAllowed() const override { return FRigVMEditorBase::AreDelegatesAllowed(); }
@@ -165,7 +165,7 @@ protected:
 	virtual void OnActiveTabChanged( TSharedPtr<SDockTab> PreviouslyActive, TSharedPtr<SDockTab> NewlyActivated ) override { return FRigVMEditorBase::OnActiveTabChanged(PreviouslyActive, NewlyActivated); }
 	virtual void OnSelectedNodesChangedImpl(const TSet<class UObject*>& NewSelection) override { FRigVMEditorBase::OnSelectedNodesChangedImpl(NewSelection); }
 	virtual void OnBlueprintChangedImpl(UBlueprint* InBlueprint, bool bIsJustBeingCompiled) override { return FRigVMEditorBase::OnBlueprintChangedImpl(InBlueprint, bIsJustBeingCompiled); }
-	virtual void RefreshEditors(ERefreshBlueprintEditorReason::Type Reason = ERefreshBlueprintEditorReason::UnknownReason) override { return FRigVMEditorBase::RefreshEditors((ERefreshRigVMEditorReason::Type)Reason); }
+	virtual void RefreshEditors(ERefreshBlueprintEditorReason::Type Reason = ERefreshBlueprintEditorReason::UnknownReason) override { return ForceEditorRefresh((ERefreshRigVMEditorReason::Type)Reason); }
 	virtual void FocusInspectorOnGraphSelection(const TSet<class UObject*>& NewSelection, bool bForceRefresh = false) override { return FRigVMEditorBase::FocusInspectorOnGraphSelection(NewSelection, bForceRefresh); }
 	virtual void NotifyPreChange(FProperty* PropertyAboutToChange) override { return FRigVMEditorBase::NotifyPreChange(PropertyAboutToChange); }
 	virtual void NotifyPostChange(const FPropertyChangedEvent& PropertyChangedEvent, FProperty* PropertyThatChanged) override { return FRigVMEditorBase::NotifyPostChange(PropertyChangedEvent, PropertyThatChanged); }
diff --git a/Engine/Plugins/Runtime/RigVM/Source/RigVMEditor/Public/Editor/RigVMVariableDetailCustomization.h b/Engine/Plugins/Runtime/RigVM/Source/RigVMEditor/Public/Editor/RigVMVariableDetailCustomization.h
--- a/Engine/Plugins/Runtime/RigVM/Source/RigVMEditor/Public/Editor/RigVMVariableDetailCustomization.h
+++ b/Engine/Plugins/Runtime/RigVM/Source/RigVMEditor/Public/Editor/RigVMVariableDetailCustomization.h
@@ -17,7 +17,7 @@ class SRigVMEditorGraphExplorer;
 class IDetailLayoutBuilder;
 class IBlueprintEditor;
 class UBlueprint;
-class FRigVMEditorBase;
+class IRigVMEditor;
 class FRigVMVariableDetailCustomization;
 class UK2Node_EditablePinBase;
 class FStructOnScope;
@@ -28,8 +28,8 @@ class RIGVMEDITOR_API FRigVMVariableDetailCustomization : public IDetailCustomiz
 {
 public:
 	/** Makes a new instance of this detail layout class for a specific detail view requesting it */
-	static TSharedRef<IDetailCustomization> MakeInstance(TSharedPtr<FRigVMEditorBase> InEditor);
-	FRigVMVariableDetailCustomization(TSharedPtr<FRigVMEditorBase> InEditor, URigVMBlueprint* Blueprint);
+	static TSharedRef<IDetailCustomization> MakeInstance(TSharedPtr<IRigVMEditor> InEditor);
+	FRigVMVariableDetailCustomization(TSharedPtr<IRigVMEditor> InEditor, URigVMBlueprint* Blueprint);
 
 #if WITH_RIGVMLEGACYEDITOR
 	static TSharedPtr<IDetailCustomization> MakeLegacyInstance(TSharedPtr<IBlueprintEditor> InBlueprintEditor);
@@ -87,7 +87,7 @@ private:
 
 private:
 	/** The Blueprint editor we are embedded in */
-	TWeakPtr<FRigVMEditorBase> EditorPtr;
+	TWeakPtr<IRigVMEditor> EditorPtr;
 
 	/** The blueprint we are editing */
 	TWeakObjectPtr<URigVMBlueprint> BlueprintPtr;
diff --git a/Engine/Plugins/Runtime/RigVM/Source/RigVMEditor/Public/Editor/SRigVMActionMenu.cpp b/Engine/Plugins/Runtime/RigVM/Source/RigVMEditor/Public/Editor/SRigVMActionMenu.cpp
--- a/Engine/Plugins/Runtime/RigVM/Source/RigVMEditor/Public/Editor/SRigVMActionMenu.cpp
+++ b/Engine/Plugins/Runtime/RigVM/Source/RigVMEditor/Public/Editor/SRigVMActionMenu.cpp
@@ -3,6 +3,7 @@
 
 #include "SRigVMActionMenu.h"
 
+#include "Editor/RigVMEditor.h"
 #include "IDocumentation.h"
 #include "RigVMNewEditor.h"
 #include "SPinTypeSelector.h"
@@ -56,7 +57,7 @@ struct FRigVMAction_PromoteVariable : public FEdGraphSchemaAction
 	// End of FEdGraphSchemaAction interface
 
 	/* Pointer to the blueprint editor containing the blueprint in which we will promote the variable. */
-	TWeakPtr<class FRigVMEditorBase> MyBlueprintEditor;
+	TWeakPtr<class IRigVMEditor> MyBlueprintEditor;
 
 	/* TRUE if promoting to member variable, FALSE if promoting to local variable */
 	bool bToMemberVariable;
@@ -72,7 +73,7 @@ SRigVMActionMenu::~SRigVMActionMenu()
 	
 }
 
-void SRigVMActionMenu::Construct( const FArguments& InArgs, TSharedPtr<FRigVMEditorBase> InEditor )
+void SRigVMActionMenu::Construct( const FArguments& InArgs, TSharedPtr<IRigVMEditor> InEditor )
 {
 	bActionExecuted = false;
 
@@ -254,7 +255,7 @@ ECheckBoxState SRigVMActionMenu::ContextToggleIsChecked() const
 TSharedRef<FGraphActionListBuilderBase> SRigVMActionMenu::OnGetActionList()
 {
 	check(EditorPtr.IsValid());
-	TSharedPtr<FRigVMEditorBase> BlueprintEditor = EditorPtr.Pin();
+	TSharedPtr<IRigVMEditor> BlueprintEditor = EditorPtr.Pin();
 	bool const bIsContextSensitive = BlueprintEditor->GetIsContextSensitive();
 
 	uint32 ContextTargetMask = 0;
@@ -284,7 +285,7 @@ TSharedRef<FGraphActionListBuilderBase> SRigVMActionMenu::OnGetActionList()
 void SRigVMActionMenu::ConstructActionContext(FBlueprintActionContext& ContextDescOut)
 {
 	check(EditorPtr.IsValid());
-	TSharedPtr<FRigVMEditorBase> BlueprintEditor = EditorPtr.Pin();
+	TSharedPtr<IRigVMEditor> BlueprintEditor = EditorPtr.Pin();
 	bool const bIsContextSensitive = BlueprintEditor->GetIsContextSensitive();
 
 	// we still want context from the graph (even if the user has unchecked
@@ -293,7 +294,7 @@ void SRigVMActionMenu::ConstructActionContext(FBlueprintActionContext& ContextDe
 	// with a valid node, then fix it up in filtering
 	ContextDescOut.Graphs.Add(GraphObj);
 
-	UBlueprint* Blueprint = BlueprintEditor->GetBlueprintObj();
+	UBlueprint* Blueprint = BlueprintEditor->GetRigVMBlueprint();
 	const bool bBlueprintIsValid = IsValid(Blueprint) && Blueprint->GeneratedClass && (Blueprint->GeneratedClass->ClassGeneratedBy == Blueprint);
 	if (!ensure(bBlueprintIsValid))  // to track UE-11597 and UE-11595
 	{
diff --git a/Engine/Plugins/Runtime/RigVM/Source/RigVMEditor/Public/Editor/SRigVMActionMenu.h b/Engine/Plugins/Runtime/RigVM/Source/RigVMEditor/Public/Editor/SRigVMActionMenu.h
--- a/Engine/Plugins/Runtime/RigVM/Source/RigVMEditor/Public/Editor/SRigVMActionMenu.h
+++ b/Engine/Plugins/Runtime/RigVM/Source/RigVMEditor/Public/Editor/SRigVMActionMenu.h
@@ -7,7 +7,7 @@
 
 
 class SGraphActionMenu;
-class FRigVMEditorBase;
+class IRigVMEditor;
 
 
 /*******************************************************************************
@@ -28,7 +28,7 @@ public:
 		SLATE_ARGUMENT( TArray<UEdGraphPin*>, DraggedFromPins )
 	SLATE_END_ARGS()
 
-	void Construct( const FArguments& InArgs, TSharedPtr<FRigVMEditorBase> InEditor );
+	void Construct( const FArguments& InArgs, TSharedPtr<IRigVMEditor> InEditor );
 
 	~SRigVMActionMenu();
 
@@ -64,7 +64,7 @@ private:
 	FVector2D NewNodePosition;
 
 	TSharedPtr<SGraphActionMenu> GraphActionMenu;
-	TWeakPtr<FRigVMEditorBase> EditorPtr;
+	TWeakPtr<IRigVMEditor> EditorPtr;
 	TSharedPtr<FRigVMActionMenuBuilder> ContextMenuBuilder;
 
 	bool bActionExecuted = false;
diff --git a/Engine/Plugins/Runtime/RigVM/Source/RigVMEditor/Public/Editor/SRigVMDetailsInspector.h b/Engine/Plugins/Runtime/RigVM/Source/RigVMEditor/Public/Editor/SRigVMDetailsInspector.h
--- a/Engine/Plugins/Runtime/RigVM/Source/RigVMEditor/Public/Editor/SRigVMDetailsInspector.h
+++ b/Engine/Plugins/Runtime/RigVM/Source/RigVMEditor/Public/Editor/SRigVMDetailsInspector.h
@@ -23,7 +23,7 @@
 #include "Widgets/SCompoundWidget.h"
 #include "Widgets/SWidget.h"
 
-class FRigVMEditorBase;
+class IRigVMEditor;
 class FProperty;
 class FStructOnScope;
 class IDetailsView;
@@ -55,7 +55,7 @@ public:
 		, _ShowSectionSelector(false)
 		{}
 
-		SLATE_ARGUMENT(TWeakPtr<FRigVMEditorBase>, Editor)
+		SLATE_ARGUMENT(TWeakPtr<IRigVMEditor>, Editor)
 		SLATE_ARGUMENT(TWeakPtr<SRigVMEditorGraphExplorer>, GraphExplorerWidget)
 		SLATE_ATTRIBUTE( bool, ShowPublicViewControl )
 		SLATE_ARGUMENT( bool, HideNameArea )
@@ -120,7 +120,7 @@ public:
 	/** returns the list of selected objects */
 	const TArray< TWeakObjectPtr<UObject> >& GetSelectedObjects() const;
 
-	void OnEditorClose(const FRigVMEditorBase* RigVMEditorBase, URigVMBlueprint* RigVMBlueprint);
+	void OnEditorClose(const IRigVMEditor* RigVMEditorBase, URigVMBlueprint* RigVMBlueprint);
 
 	virtual void AddReferencedObjects(FReferenceCollector& Collector) override;
 	virtual FString GetReferencerName() const override;
@@ -133,7 +133,7 @@ protected:
 	void AddPropertiesRecursive(FProperty* Property);
 
 	/** Pointer back to the kismet 2 tool that owns us */
-	TWeakPtr<FRigVMEditorBase> WeakEditor;
+	TWeakPtr<IRigVMEditor> WeakEditor;
 
 	/** The tab that owns this details view. */
 	TWeakPtr<SDockTab> OwnerTab;
diff --git a/Engine/Plugins/Runtime/RigVM/Source/RigVMEditor/Public/Widgets/SRigVMEditorGraphExplorer.h b/Engine/Plugins/Runtime/RigVM/Source/RigVMEditor/Public/Widgets/SRigVMEditorGraphExplorer.h
--- a/Engine/Plugins/Runtime/RigVM/Source/RigVMEditor/Public/Widgets/SRigVMEditorGraphExplorer.h
+++ b/Engine/Plugins/Runtime/RigVM/Source/RigVMEditor/Public/Widgets/SRigVMEditorGraphExplorer.h
@@ -11,12 +11,11 @@
 #include "RigVMModel/RigVMGraph.h"
 #include "SRigVMEditorGraphExplorerTreeView.h"
 
-class FRigVMEditorBase;
+class IRigVMEditor;
 class SRigVMEditorGraphExplorerTreeView;
 //
 class URigVMGraph;
 class URigVMEdGraph;
-class FRigVMEditorBase;
 class URigVMBlueprint;
 class SGraphActionMenu;
 struct FCreateWidgetForActionData;
@@ -80,7 +79,7 @@ public:
 	SLATE_BEGIN_ARGS(SRigVMEditorGraphExplorer) {}
 	SLATE_END_ARGS()
 
-	void Construct(const FArguments& InArgs, TWeakPtr<FRigVMEditorBase> InRigVMEditor);
+	void Construct(const FArguments& InArgs, TWeakPtr<IRigVMEditor> InRigVMEditor);
 	~SRigVMEditorGraphExplorer();
 
 	// Refresh the graph action menu.
@@ -197,7 +196,7 @@ private:
 
 	void HandleSelectionChanged(TSharedPtr<FRigVMEditorGraphExplorerTreeElement> Selection, ESelectInfo::Type SelectInfo);
 
-	TWeakPtr<FRigVMEditorBase> RigVMEditor;
+	TWeakPtr<IRigVMEditor> RigVMEditor;
 
 	TSharedPtr<SRigVMEditorGraphExplorerTreeView> TreeView;
 
diff --git a/Engine/Plugins/Runtime/RigVM/Source/RigVMEditor/Public/Widgets/SRigVMExecutionStackView.h b/Engine/Plugins/Runtime/RigVM/Source/RigVMEditor/Public/Widgets/SRigVMExecutionStackView.h
--- a/Engine/Plugins/Runtime/RigVM/Source/RigVMEditor/Public/Widgets/SRigVMExecutionStackView.h
+++ b/Engine/Plugins/Runtime/RigVM/Source/RigVMEditor/Public/Widgets/SRigVMExecutionStackView.h
@@ -10,7 +10,7 @@
 #include "RigVMCore/RigVM.h"
 #include "RigVMBlueprint.h"
 
-class FRigVMEditorBase;
+class IRigVMEditor;
 class SRigVMExecutionStackView;
 class FUICommandList;
 class SSearchBox;
@@ -75,7 +75,7 @@ public:
 
 	~SRigVMExecutionStackView();
 
-	void Construct(const FArguments& InArgs, TSharedRef<FRigVMEditorBase> InRigVMEditor);
+	void Construct(const FArguments& InArgs, TSharedRef<IRigVMEditor> InRigVMEditor);
 
 	/** Set Selection Changed */
 	void OnSelectionChanged(TSharedPtr<FRigStackEntry> Selection, ESelectInfo::Type SelectInfo);
@@ -126,7 +126,7 @@ private:
 	bool bSuspendControllerSelection;
 	void HandleModifiedEvent(ERigVMGraphNotifType InNotifType, URigVMGraph* InGraph, UObject* InSubject);
 	void HandleHostInitializedEvent(URigVMHost* InHost, const FName& InEventName);
-	void HandlePreviewHostUpdated(FRigVMEditorBase* InEditor);
+	void HandlePreviewHostUpdated(IRigVMEditor* InEditor);
 	void HandleItemMouseDoubleClick(TSharedPtr<FRigStackEntry> InItem);
 
 	/** Populate the execution stack with descriptive names for each instruction */
@@ -137,7 +137,7 @@ private:
 	/** Command list we bind to */
 	TSharedPtr<FUICommandList> CommandList;
 
-	TWeakPtr<FRigVMEditorBase> WeakEditor;
+	TWeakPtr<IRigVMEditor> WeakEditor;
 	TWeakObjectPtr<URigVMBlueprint> RigVMBlueprint;
 
 	TArray<TSharedPtr<FRigStackEntry>> Operators;
