Control Rig: Add Class Settings and Class Default to the new editor
#rnx

diff --git a/Engine/Plugins/Runtime/RigVM/Source/RigVMEditor/Private/Editor/RigVMDetailsInspectorTabSummoner.cpp b/Engine/Plugins/Runtime/RigVM/Source/RigVMEditor/Private/Editor/RigVMDetailsInspectorTabSummoner.cpp
--- a/Engine/Plugins/Runtime/RigVM/Source/RigVMEditor/Private/Editor/RigVMDetailsInspectorTabSummoner.cpp
+++ b/Engine/Plugins/Runtime/RigVM/Source/RigVMEditor/Private/Editor/RigVMDetailsInspectorTabSummoner.cpp
@@ -22,10 +22,23 @@ FRigVMDetailsInspectorTabSummoner::FRigVMDetailsInspectorTabSummoner(const TShar
 TSharedRef<SWidget> FRigVMDetailsInspectorTabSummoner::CreateTabBody(const FWorkflowTabSpawnInfo& Info) const
 {
 	TSharedRef<SRigVMDetailsInspector> Inspector = SNew(SRigVMDetailsInspector)
-	.Editor(RigVMEditor);
+		.Editor(RigVMEditor)
+		.OnFinishedChangingProperties(FOnFinishedChangingProperties::FDelegate::CreateSP(RigVMEditor.Pin().Get(), &FRigVMNewEditor::OnFinishedChangingProperties));
 	
 	RigVMEditor.Pin()->SetInspector(Inspector);
 	return Inspector;
 }
 
+TSharedRef<SDockTab> FRigVMDetailsInspectorTabSummoner::SpawnTab(const FWorkflowTabSpawnInfo& Info) const
+{
+	TSharedRef<SDockTab> Tab = FWorkflowTabFactory::SpawnTab(Info);
+
+	TSharedPtr<FRigVMNewEditor> BlueprintEditorPtr = StaticCastSharedPtr<FRigVMNewEditor>(HostingApp.Pin());
+	BlueprintEditorPtr->GetRigVMInspector()->SetOwnerTab(Tab);
+
+	BlueprintEditorPtr->GetRigVMInspector()->GetPropertyView()->SetHostTabManager(Info.TabManager);
+
+	return Tab;
+}
+
 #undef LOCTEXT_NAMESPACE 
diff --git a/Engine/Plugins/Runtime/RigVM/Source/RigVMEditor/Private/Editor/RigVMEditor.cpp b/Engine/Plugins/Runtime/RigVM/Source/RigVMEditor/Private/Editor/RigVMEditor.cpp
--- a/Engine/Plugins/Runtime/RigVM/Source/RigVMEditor/Private/Editor/RigVMEditor.cpp
+++ b/Engine/Plugins/Runtime/RigVM/Source/RigVMEditor/Private/Editor/RigVMEditor.cpp
@@ -3377,6 +3377,8 @@ void FRigVMEditorBase::FillToolbar(FToolBarBuilder& ToolbarBuilder, bool bEndSec
 			LOCTEXT("BulkEdit_ToolTip", "Perform changes across many nodes / assets"),
 			FSlateIcon(FAppStyle::GetAppStyleSetName(), "LevelEditor.Recompile"),
 			false);
+
+		AddSettingsAndDefaultWidget(ToolbarBuilder);
 	}
 
 	if(bEndSection)
diff --git a/Engine/Plugins/Runtime/RigVM/Source/RigVMEditor/Private/Editor/RigVMEditorCommands.cpp b/Engine/Plugins/Runtime/RigVM/Source/RigVMEditor/Private/Editor/RigVMEditorCommands.cpp
--- a/Engine/Plugins/Runtime/RigVM/Source/RigVMEditor/Private/Editor/RigVMEditorCommands.cpp
+++ b/Engine/Plugins/Runtime/RigVM/Source/RigVMEditor/Private/Editor/RigVMEditorCommands.cpp
@@ -12,6 +12,9 @@ void FRigVMEditorCommands::RegisterCommands()
 	UI_COMMAND(SaveOnCompile_Always, "Always", "Sets the save-on-compile option to 'Always', meaning that your Blueprints will be saved whenever they are compiled (even if there were errors)", EUserInterfaceActionType::RadioButton, FInputChord());
 	UI_COMMAND(JumpToErrorNode, "Jump to Error Node", "When enabled, then the Blueprint will snap focus to nodes producing an error during compilation", EUserInterfaceActionType::ToggleButton, FInputChord());
 
+	UI_COMMAND(EditGlobalOptions, "Settings", "Edit Class Settings", EUserInterfaceActionType::ToggleButton, FInputChord());
+	UI_COMMAND(EditClassDefaults, "Defaults", "Edit the initial values of your class.", EUserInterfaceActionType::ToggleButton, FInputChord());
+
 	UI_COMMAND(DeleteItem, "Delete", "Deletes the selected items and removes their nodes from the graph.", EUserInterfaceActionType::Button, FInputChord(EKeys::Delete));
 	UI_COMMAND(ExecuteGraph, "Execute", "Execute the rig graph if On.", EUserInterfaceActionType::ToggleButton, FInputChord());
 	UI_COMMAND(AutoCompileGraph, "Auto Compile", "Auto-compile the rig graph if On.", EUserInterfaceActionType::ToggleButton, FInputChord());
diff --git a/Engine/Plugins/Runtime/RigVM/Source/RigVMEditor/Private/Editor/RigVMEditorCommands.h b/Engine/Plugins/Runtime/RigVM/Source/RigVMEditor/Private/Editor/RigVMEditorCommands.h
--- a/Engine/Plugins/Runtime/RigVM/Source/RigVMEditor/Private/Editor/RigVMEditorCommands.h
+++ b/Engine/Plugins/Runtime/RigVM/Source/RigVMEditor/Private/Editor/RigVMEditorCommands.h
@@ -24,6 +24,10 @@ public:
 	TSharedPtr<FUICommandInfo> SaveOnCompile_SuccessOnly;
 	TSharedPtr<FUICommandInfo> SaveOnCompile_Always;
 	TSharedPtr<FUICommandInfo> JumpToErrorNode;
+
+	/** Edit global options */
+	TSharedPtr<FUICommandInfo> EditGlobalOptions;
+	TSharedPtr<FUICommandInfo> EditClassDefaults;
 	
 	/** Deletes the selected items and removes their nodes from the graph. */
 	TSharedPtr< FUICommandInfo > DeleteItem;
diff --git a/Engine/Plugins/Runtime/RigVM/Source/RigVMEditor/Private/Editor/RigVMNewEditor.cpp b/Engine/Plugins/Runtime/RigVM/Source/RigVMEditor/Private/Editor/RigVMNewEditor.cpp
--- a/Engine/Plugins/Runtime/RigVM/Source/RigVMEditor/Private/Editor/RigVMNewEditor.cpp
+++ b/Engine/Plugins/Runtime/RigVM/Source/RigVMEditor/Private/Editor/RigVMNewEditor.cpp
@@ -30,6 +30,7 @@
 #include "GraphEditorActions.h"
 #include "RigVMBlueprint.h"
 #include "RigVMEditorCommands.h"
+#include "Editor/RigVMDetailsInspectorTabSummoner.h"
 #include "Editor/RigVMGraphEditorSummoner.h"
 #include "Widgets/Input/SEditableTextBox.h"
 #include "Editor/SRigVMDetailsInspector.h"
@@ -42,7 +43,6 @@
 
 #define LOCTEXT_NAMESPACE "RigVMNewEditor"
 
-const FName FRigVMNewEditorTabs::DetailsID() {			static FName ID = TEXT( "Inspector" ); return ID; };
 const FName FRigVMNewEditorTabs::CompilerResultsID() {	static FName ID = TEXT( "CompilerResults" ); return ID; };
 
 const FName FRigVMNewEditor::SelectionState_GraphExplorer() {	static FName State = (TEXT("GraphExplorer")); return State; };
@@ -621,6 +621,12 @@ void FRigVMNewEditor::AddCompileWidget(FToolBarBuilder& ToolbarBuilder)
 		TAttribute<FSlateIcon>(), true);
 }
 
+void FRigVMNewEditor::AddSettingsAndDefaultWidget(FToolBarBuilder& ToolbarBuilder)
+{
+	ToolbarBuilder.AddToolBarButton(FRigVMEditorCommands::Get().EditGlobalOptions);
+	ToolbarBuilder.AddToolBarButton(FRigVMEditorCommands::Get().EditClassDefaults);
+}
+
 void FRigVMNewEditor::RegisterTabSpawners(const TSharedRef<FTabManager>& InTabManager)
 {
 	//@TODO: Can't we do this sooner?
@@ -1185,6 +1191,18 @@ void FRigVMNewEditor::CreateDefaultCommands()
 			FCanExecuteAction(),
 			FIsActionChecked::CreateSP(this, &FRigVMNewEditor::IsJumpToErrorNodeOptionSet)
 		);
+
+		ToolkitCommands->MapAction(
+			FRigVMEditorCommands::Get().EditGlobalOptions,
+			FExecuteAction::CreateSP(this, &FRigVMNewEditor::EditGlobalOptions_Clicked),
+			FCanExecuteAction(),
+			FIsActionChecked::CreateSP(this, &FRigVMNewEditor::IsDetailsPanelEditingGlobalOptions));
+
+		ToolkitCommands->MapAction(
+			FRigVMEditorCommands::Get().EditClassDefaults,
+			FExecuteAction::CreateSP(this, &FRigVMNewEditor::EditClassDefaults_Clicked),
+			FCanExecuteAction(),
+			FIsActionChecked::CreateSP(this, &FRigVMNewEditor::IsDetailsPanelEditingClassDefaults));
 	}
 
 	FRigVMEditorBase::CreateDefaultCommandsImpl();
@@ -1842,9 +1860,19 @@ void FRigVMNewEditor::EditGlobalOptions_Clicked()
 	}
 }
 
+bool FRigVMNewEditor::IsDetailsPanelEditingGlobalOptions() const
+{
+	return CurrentUISelection == FRigVMNewEditor::SelectionState_ClassSettings();
+}
+
+bool FRigVMNewEditor::IsDetailsPanelEditingClassDefaults() const
+{
+	return CurrentUISelection == FRigVMNewEditor::SelectionState_ClassDefaults();
+}
+
 void FRigVMNewEditor::TryInvokingDetailsTab(bool bFlash)
 {
-	if ( TabManager->HasTabSpawner(FRigVMNewEditorTabs::DetailsID()) )
+	if ( TabManager->HasTabSpawner(FRigVMDetailsInspectorTabSummoner::TabID()) )
 	{
 		TSharedPtr<SDockTab> BlueprintTab = FGlobalTabmanager::Get()->GetMajorTabForTabManager(TabManager.ToSharedRef());
 
@@ -1852,21 +1880,17 @@ void FRigVMNewEditor::TryInvokingDetailsTab(bool bFlash)
 		// being interacted with.  So we make sure the window it's in is focused and the tab is in the foreground.
 		if ( BlueprintTab.IsValid() && BlueprintTab->IsForeground() )
 		{
-			TSharedPtr<SWindow> ParentWindow = BlueprintTab->GetParentWindow();
-			if ( ParentWindow.IsValid() && ParentWindow->HasFocusedDescendants() )
+			if ( !Inspector.IsValid() || !Inspector->GetOwnerTab().IsValid() || Inspector->GetOwnerTab()->GetDockArea().IsValid() )
 			{
-				if ( !Inspector.IsValid() || !Inspector->GetOwnerTab().IsValid() || Inspector->GetOwnerTab()->GetDockArea().IsValid() )
-				{
-					// Show the details panel if it doesn't exist.
-					TabManager->TryInvokeTab(FRigVMNewEditorTabs::DetailsID());
+				// Show the details panel if it doesn't exist.
+				TabManager->TryInvokeTab(FRigVMDetailsInspectorTabSummoner::TabID());
 
-					if ( bFlash && Inspector.IsValid())
+				if ( bFlash && Inspector.IsValid())
+				{
+					TSharedPtr<SDockTab> OwnerTab = Inspector->GetOwnerTab();
+					if ( OwnerTab.IsValid() )
 					{
-						TSharedPtr<SDockTab> OwnerTab = Inspector->GetOwnerTab();
-						if ( OwnerTab.IsValid() )
-						{
-							OwnerTab->FlashTab();
-						}
+						OwnerTab->FlashTab();
 					}
 				}
 			}
@@ -2039,7 +2063,7 @@ void FRigVMNewEditor::CreateDefaultTabContents(const TArray<UBlueprint*> InBluep
 		.Editor(SharedThis(this))
 		. HideNameArea(true)
 		. ViewIdentifier(FName("BlueprintInspector"))
-		//. OnFinishedChangingProperties( FOnFinishedChangingProperties::FDelegate::CreateSP(this, &FBlueprintEditor::OnFinishedChangingProperties) ); // TODO sara-s
+		. OnFinishedChangingProperties( FOnFinishedChangingProperties::FDelegate::CreateSP(this, &FRigVMNewEditor::OnFinishedChangingProperties) ); 
 	;
 
 	if ( InBlueprints.Num() > 0 )
@@ -2055,7 +2079,7 @@ void FRigVMNewEditor::CreateDefaultTabContents(const TArray<UBlueprint*> InBluep
 			//. ShowPublicViewControl(this, &FBlueprintEditor::ShouldShowPublicViewControl) // TODO sara-s
 			. ShowTitleArea(false)
 			. HideNameArea(bHideNameArea)
-			//. OnFinishedChangingProperties( FOnFinishedChangingProperties::FDelegate::CreateSP( this, &FBlueprintEditor::OnFinishedChangingProperties ) ); // TODO sara-s
+			. OnFinishedChangingProperties( FOnFinishedChangingProperties::FDelegate::CreateSP( this, &FRigVMNewEditor::OnFinishedChangingProperties ) ); 
 		;
 	}
 
diff --git a/Engine/Plugins/Runtime/RigVM/Source/RigVMEditor/Private/Editor/SRigVMDetailsInspector.cpp b/Engine/Plugins/Runtime/RigVM/Source/RigVMEditor/Private/Editor/SRigVMDetailsInspector.cpp
--- a/Engine/Plugins/Runtime/RigVM/Source/RigVMEditor/Private/Editor/SRigVMDetailsInspector.cpp
+++ b/Engine/Plugins/Runtime/RigVM/Source/RigVMEditor/Private/Editor/SRigVMDetailsInspector.cpp
@@ -122,6 +122,11 @@ TSharedPtr<SDockTab> SRigVMDetailsInspector::GetOwnerTab() const
 	return OwnerTab.Pin();
 }
 
+void SRigVMDetailsInspector::SetOwnerTab(TSharedRef<SDockTab> Tab)
+{
+	OwnerTab = Tab;
+}
+
 const TArray< TWeakObjectPtr<UObject> >& SRigVMDetailsInspector::GetSelectedObjects() const
 {
 	return SelectedObjects;
diff --git a/Engine/Plugins/Runtime/RigVM/Source/RigVMEditor/Public/Editor/RigVMDetailsInspectorTabSummoner.h b/Engine/Plugins/Runtime/RigVM/Source/RigVMEditor/Public/Editor/RigVMDetailsInspectorTabSummoner.h
--- a/Engine/Plugins/Runtime/RigVM/Source/RigVMEditor/Public/Editor/RigVMDetailsInspectorTabSummoner.h
+++ b/Engine/Plugins/Runtime/RigVM/Source/RigVMEditor/Public/Editor/RigVMDetailsInspectorTabSummoner.h
@@ -17,6 +17,7 @@ public:
 	FRigVMDetailsInspectorTabSummoner(const TSharedRef<FRigVMNewEditor>& InRigVMEditor);
 
 	virtual TSharedRef<SWidget> CreateTabBody(const FWorkflowTabSpawnInfo& Info) const override;
+	virtual TSharedRef<SDockTab> SpawnTab(const FWorkflowTabSpawnInfo& Info) const override;
 
 protected:
 	TWeakPtr<FRigVMNewEditor> RigVMEditor;
diff --git a/Engine/Plugins/Runtime/RigVM/Source/RigVMEditor/Public/Editor/RigVMEditor.h b/Engine/Plugins/Runtime/RigVM/Source/RigVMEditor/Public/Editor/RigVMEditor.h
--- a/Engine/Plugins/Runtime/RigVM/Source/RigVMEditor/Public/Editor/RigVMEditor.h
+++ b/Engine/Plugins/Runtime/RigVM/Source/RigVMEditor/Public/Editor/RigVMEditor.h
@@ -33,6 +33,7 @@ class URigVMHost;
 class FTransaction;
 class URigVMBlueprint;
 class SGraphEditor;
+class FPreviewScene;
 
 DECLARE_MULTICAST_DELEGATE_TwoParams(FRigVMEditorClosed, const IRigVMEditor*, URigVMBlueprint*);
 
@@ -169,6 +170,7 @@ public:
 	virtual void ZoomToSelection_Clicked() = 0;
 
 	virtual void HandleModifiedEvent(ERigVMGraphNotifType InNotifType, URigVMGraph* InGraph, UObject* InSubject) = 0;
+	virtual void OnFinishedChangingProperties(const FPropertyChangedEvent& PropertyChangedEvent) = 0;
 	
 protected:
 	
@@ -239,6 +241,7 @@ protected:
 	virtual void SetupViewForBlueprintEditingMode() = 0;
 	virtual void RegisterToolbarTab(const TSharedRef<class FTabManager>& TabManager) = 0;
 	virtual void AddCompileWidget(FToolBarBuilder& ToolbarBuilder) = 0;
+	virtual void AddSettingsAndDefaultWidget(FToolBarBuilder& ToolbarBuilder) = 0;
 	virtual void AddSelectedDebugObjectWidget(FToolBarBuilder& ToolbarBuilder) = 0;
 };
 
@@ -413,7 +416,7 @@ protected:
 	void NotifyPreChange(FProperty* PropertyAboutToChange);
 	void NotifyPostChange(const FPropertyChangedEvent& PropertyChangedEvent, FProperty* PropertyThatChanged);
 	/** delegate for changing property */
-	void OnFinishedChangingProperties(const FPropertyChangedEvent& PropertyChangedEvent);
+	virtual void OnFinishedChangingProperties(const FPropertyChangedEvent& PropertyChangedEvent) override;
 	void OnPropertyChanged(UObject* InObject, FPropertyChangedEvent& InEvent);
 	virtual void OnWrappedPropertyChangedChainEvent(URigVMDetailsViewWrapperObject* InWrapperObject, const FString& InPropertyPath, FPropertyChangedChainEvent& InPropertyChangedChainEvent);
 	void OnRequestLocalizeFunctionDialog(FRigVMGraphFunctionIdentifier& InFunction, URigVMController* InTargetController, IRigVMGraphFunctionHost* InTargetFunctionHost, bool bForce);
diff --git a/Engine/Plugins/Runtime/RigVM/Source/RigVMEditor/Public/Editor/RigVMLegacyEditor.h b/Engine/Plugins/Runtime/RigVM/Source/RigVMEditor/Public/Editor/RigVMLegacyEditor.h
--- a/Engine/Plugins/Runtime/RigVM/Source/RigVMEditor/Public/Editor/RigVMLegacyEditor.h
+++ b/Engine/Plugins/Runtime/RigVM/Source/RigVMEditor/Public/Editor/RigVMLegacyEditor.h
@@ -116,6 +116,7 @@ protected:
 	virtual void RegisterToolbarTab(const TSharedRef<FTabManager>& InTabManager) override { FBlueprintEditor::RegisterToolbarTab(InTabManager); }
 	virtual const TArray<UObject*>* GetObjectsCurrentlyBeingEdited() const override { return FBlueprintEditor::GetObjectsCurrentlyBeingEdited(); }
 	virtual void AddCompileWidget(FToolBarBuilder& ToolbarBuilder) override;
+	virtual void AddSettingsAndDefaultWidget(FToolBarBuilder& ToolbarBuilder) override {}
 	virtual void AddSelectedDebugObjectWidget(FToolBarBuilder& ToolbarBuilder) override;
 	virtual FString GetCustomDebugObjectLabel(UObject* ObjectBeingDebugged) const override { return FRigVMEditorBase::GetCustomDebugObjectLabel(ObjectBeingDebugged); }
 public:
diff --git a/Engine/Plugins/Runtime/RigVM/Source/RigVMEditor/Public/Editor/RigVMNewEditor.h b/Engine/Plugins/Runtime/RigVM/Source/RigVMEditor/Public/Editor/RigVMNewEditor.h
--- a/Engine/Plugins/Runtime/RigVM/Source/RigVMEditor/Public/Editor/RigVMNewEditor.h
+++ b/Engine/Plugins/Runtime/RigVM/Source/RigVMEditor/Public/Editor/RigVMNewEditor.h
@@ -11,7 +11,6 @@ class SRigVMDetailsInspector;
 struct RIGVMEDITOR_API FRigVMNewEditorTabs
 {
 	// Tab identifiers
-	static const FName DetailsID();
 	static const FName CompilerResultsID();
 };
 
@@ -44,8 +43,8 @@ protected:
 	virtual TSharedPtr<SDockTab> OpenDocument(const UObject* DocumentID, FDocumentTracker::EOpenDocumentCause Cause) override;
 	virtual void CloseDocumentTab(const UObject* DocumentID) override;
 	virtual TSharedPtr<FTabManager> GetTabManager() override { return FWorkflowCentricApplication::GetTabManager(); }
-	virtual TSharedPtr<SRigVMDetailsInspector> GetRigVMInspector() const override { return Inspector.IsValid() ? Inspector.ToSharedRef().ToSharedPtr() : nullptr; }
 public:
+	virtual TSharedPtr<SRigVMDetailsInspector> GetRigVMInspector() const override { return Inspector.IsValid() ? Inspector.ToSharedRef().ToSharedPtr() : nullptr; }
 	virtual void SetInspector(TSharedPtr<SRigVMDetailsInspector> InWidget) { Inspector = InWidget; };
 protected:
 	virtual TSharedPtr<FApplicationMode> CreateEditorMode() override;
@@ -98,6 +97,8 @@ protected:
 	virtual bool IsNodeTitleVisible(const UEdGraphNode* Node, bool bRequestRename) override;
 	virtual void EditClassDefaults_Clicked() override;
 	virtual void EditGlobalOptions_Clicked() override;
+	bool IsDetailsPanelEditingGlobalOptions() const;
+	bool IsDetailsPanelEditingClassDefaults() const;
 	virtual void TryInvokingDetailsTab(bool bFlash = true) override;
 	virtual FName GetGraphExplorerWidgetID() override { return FRigVMEditorGraphExplorerTabSummoner::TabID(); }
 	virtual void RefreshInspector() override;
@@ -117,6 +118,7 @@ public:
 	virtual void RegisterToolbarTab(const TSharedRef<FTabManager>& InTabManager) override { FAssetEditorToolkit::RegisterTabSpawners(InTabManager); }
 	virtual const TArray<UObject*>* GetObjectsCurrentlyBeingEdited() const override { return FAssetEditorToolkit::GetObjectsCurrentlyBeingEdited(); }
 	virtual void AddCompileWidget(FToolBarBuilder& ToolbarBuilder) override;
+	virtual void AddSettingsAndDefaultWidget(FToolBarBuilder& ToolbarBuilder) override;
 	virtual void AddSelectedDebugObjectWidget(FToolBarBuilder& ToolbarBuilder) override {}
 	virtual void Compile() { FRigVMEditorBase::Compile(); }
 	virtual void OnCreateGraphEditorCommands(TSharedPtr<FUICommandList> GraphEditorCommandsList) { FRigVMEditorBase::OnCreateGraphEditorCommands(GraphEditorCommandsList); }
diff --git a/Engine/Plugins/Runtime/RigVM/Source/RigVMEditor/Public/Editor/SRigVMDetailsInspector.h b/Engine/Plugins/Runtime/RigVM/Source/RigVMEditor/Public/Editor/SRigVMDetailsInspector.h
--- a/Engine/Plugins/Runtime/RigVM/Source/RigVMEditor/Public/Editor/SRigVMDetailsInspector.h
+++ b/Engine/Plugins/Runtime/RigVM/Source/RigVMEditor/Public/Editor/SRigVMDetailsInspector.h
@@ -115,6 +115,7 @@ public:
 	TSharedPtr<class IDetailsView> GetPropertyView() const { return PropertyView; }
 
 	TSharedPtr<SDockTab> GetOwnerTab() const;
+	void SetOwnerTab(TSharedRef<SDockTab> Tab);
 
 
 	/** returns the list of selected objects */
