RigVM: Add debug object dropdown to new editor
#rnx

diff --git a/Engine/Plugins/Runtime/RigVM/Source/RigVMEditor/Private/Editor/RigVMNewEditor.cpp b/Engine/Plugins/Runtime/RigVM/Source/RigVMEditor/Private/Editor/RigVMNewEditor.cpp
--- a/Engine/Plugins/Runtime/RigVM/Source/RigVMEditor/Private/Editor/RigVMNewEditor.cpp
+++ b/Engine/Plugins/Runtime/RigVM/Source/RigVMEditor/Private/Editor/RigVMNewEditor.cpp
@@ -36,9 +36,7 @@
 #include "Editor/SRigVMDetailsInspector.h"
 #include "Engine/LevelScriptBlueprint.h"
 #include "Kismet2/KismetEditorUtilities.h"
-#if WITH_RIGVMLEGACYEDITOR
-#include "SKismetInspector.h"
-#endif
+#include "Widgets/SRigVMEditorSelectedDebugObjectWidget.h"
 
 
 #define LOCTEXT_NAMESPACE "RigVMNewEditor"
@@ -627,6 +625,11 @@ void FRigVMNewEditor::AddSettingsAndDefaultWidget(FToolBarBuilder& ToolbarBuilde
 	ToolbarBuilder.AddToolBarButton(FRigVMEditorCommands::Get().EditClassDefaults);
 }
 
+void FRigVMNewEditor::AddSelectedDebugObjectWidget(FToolBarBuilder& ToolbarBuilder)
+{
+	ToolbarBuilder.AddWidget(SNew(SRigVMEditorSelectedDebugObjectWidget, SharedThis(this)));
+}
+
 void FRigVMNewEditor::RegisterTabSpawners(const TSharedRef<FTabManager>& InTabManager)
 {
 	//@TODO: Can't we do this sooner?
diff --git a/Engine/Plugins/Runtime/RigVM/Source/RigVMEditor/Private/Widgets/SRigVMEditorSelectedDebugObjectWidget.cpp b/Engine/Plugins/Runtime/RigVM/Source/RigVMEditor/Private/Widgets/SRigVMEditorSelectedDebugObjectWidget.cpp
new file mode 100644
--- /dev/null
+++ b/Engine/Plugins/Runtime/RigVM/Source/RigVMEditor/Private/Widgets/SRigVMEditorSelectedDebugObjectWidget.cpp
@@ -0,0 +1,786 @@
+// Copyright Epic Games, Inc. All Rights Reserved.
+
+#include "Widgets/SRigVMEditorSelectedDebugObjectWidget.h"
+
+#include "UObject/UObjectIterator.h"
+#include "PreviewScene.h"
+#include "IDocumentation.h"
+#include "PropertyCustomizationHelpers.h"
+#include "RigVMBlueprint.h"
+#include "SLevelOfDetailBranchNode.h"
+#include "UnrealEdGlobals.h"
+#include "Components/Widget.h"
+#include "Editor/UnrealEdEngine.h"
+#include "Widgets/Input/STextComboBox.h"
+#include "Widgets/SToolTip.h"
+#include "Components/Widget.h"
+
+#define LOCTEXT_NAMESPACE "RigVMToolbar"
+
+//////////////////////////////////////////////////////////////////////////
+// SRigVMEditorSelectedDebugObjectWidget
+
+void SRigVMEditorSelectedDebugObjectWidget::Construct(const FArguments& InArgs, TSharedPtr<IRigVMEditor> InEditor)
+{
+	Editor = InEditor;
+
+	GenerateDebugWorldNames(false);
+	GenerateDebugObjectInstances(false);
+
+	LastObjectObserved = nullptr;
+
+	DebugWorldsComboBox = SNew(STextComboBox)
+		.ToolTip(IDocumentation::Get()->CreateToolTip(
+			LOCTEXT("RigVMDebugWorldTooltip", "Select a world to debug, will filter what to debug if no specific object selected"),
+			nullptr, FString(),
+			TEXT("DebugWorld")))
+		.OptionsSource(&DebugWorldNames)
+		.InitiallySelectedItem(GetDebugWorldName())
+		.Visibility(this, &SRigVMEditorSelectedDebugObjectWidget::IsDebugWorldComboVisible)
+		.OnComboBoxOpening(this, &SRigVMEditorSelectedDebugObjectWidget::GenerateDebugWorldNames, true)
+		.OnSelectionChanged(this, &SRigVMEditorSelectedDebugObjectWidget::DebugWorldSelectionChanged)
+		.ContentPadding(FMargin(0.f, 4.f));
+
+	DebugObjectsComboBox = SNew(SComboBox<TSharedPtr<FRigVMDebugObjectInstance>>)
+		.ToolTip(IDocumentation::Get()->CreateToolTip(
+			LOCTEXT("RigVMDebugObjectTooltip", "Select an object to debug, if set to none will debug any object"),
+			nullptr,
+			FString(),
+			TEXT("DebugObject")))
+		.OptionsSource(&DebugObjects)
+		.InitiallySelectedItem(GetDebugObjectInstance())
+		.OnComboBoxOpening(this, &SRigVMEditorSelectedDebugObjectWidget::GenerateDebugObjectInstances, true)
+		.OnSelectionChanged(this, &SRigVMEditorSelectedDebugObjectWidget::DebugObjectSelectionChanged)
+		.OnGenerateWidget(this, &SRigVMEditorSelectedDebugObjectWidget::CreateDebugObjectItemWidget)
+		.ContentPadding(FMargin(0.f, 4.f))
+		.AddMetaData<FTagMetaData>(TEXT("SelectDebugObjectCobmo"))
+		[
+			SNew(STextBlock)
+			.Text(this, &SRigVMEditorSelectedDebugObjectWidget::GetSelectedDebugObjectTextLabel)
+		];
+
+	const TSharedRef<SWidget> BrowseButton = PropertyCustomizationHelpers::MakeBrowseButton(
+		FSimpleDelegate::CreateSP(this, &SRigVMEditorSelectedDebugObjectWidget::SelectedDebugObject_OnClicked),
+		LOCTEXT("DebugSelectActor", "Select and frame the debug actor in the Level Editor."),
+		TAttribute<bool>::Create(TAttribute<bool>::FGetter::CreateSP(this, &SRigVMEditorSelectedDebugObjectWidget::IsDebugObjectSelected))
+	);
+
+
+	TSharedRef<SWidget> DebugObjectSelectionWidget =
+		SNew(SHorizontalBox)
+		+ SHorizontalBox::Slot()
+		.AutoWidth()
+		.Padding(FMargin(8.0f, 0.0f, 0.0f, 0.0f))
+		[
+			DebugObjectsComboBox.ToSharedRef()
+		]
+		+ SHorizontalBox::Slot()
+		.AutoWidth()
+		.HAlign(HAlign_Right)
+		.VAlign(VAlign_Center)
+		.Padding(4.0f)
+		[
+			BrowseButton
+		];
+
+	ChildSlot
+	[
+		SNew(SHorizontalBox)
+			+ SHorizontalBox::Slot()
+			.VAlign(VAlign_Center)
+			.Padding(0.0f)
+			.AutoWidth()
+			[
+				DebugWorldsComboBox.ToSharedRef()
+			]
+
+			+ SHorizontalBox::Slot()
+			.VAlign(VAlign_Center)
+			.Padding(0.0f)
+			.AutoWidth()
+			[
+				DebugObjectSelectionWidget
+			]
+	];
+}
+
+void SRigVMEditorSelectedDebugObjectWidget::Tick( const FGeometry& AllottedGeometry, const double InCurrentTime, const float InDeltaTime )
+{
+	if (UBlueprint* Blueprint = GetBlueprintObj())
+	{
+		if (UObject* Object = Blueprint->GetObjectBeingDebugged())
+		{
+			if (Object != LastObjectObserved.Get())
+			{
+				// bRestoreSelection attempts to restore the selection by name, 
+				// this ensures that if the last object we had selected was 
+				// regenerated (spawning a new object), then we select that  
+				// again, even if it is technically a different object
+				GenerateDebugObjectInstances(/*bRestoreSelection =*/true);
+
+				LastObjectObserved = Object;
+			}
+		}
+		else
+		{
+			LastObjectObserved = nullptr;
+
+			// If the object name is a name (rather than the 'No debug selected' string then regenerate the names (which will reset the combo box) as the object is invalid.
+			TSharedPtr<FRigVMDebugObjectInstance> CurrentSelection = DebugObjectsComboBox->GetSelectedItem();
+			if (CurrentSelection.IsValid() && CurrentSelection->IsEditorObject())
+			{
+				GenerateDebugObjectInstances(false);
+			}
+		}
+	}
+}
+
+const FString& SRigVMEditorSelectedDebugObjectWidget::GetNoDebugString() const
+{
+	return NSLOCTEXT("BlueprintEditor", "DebugObjectNothingSelected", "No debug object selected").ToString();
+}
+
+const FString& SRigVMEditorSelectedDebugObjectWidget::GetDebugAllWorldsString() const
+{
+	return NSLOCTEXT("BlueprintEditor", "DebugWorldNothingSelected", "All Worlds").ToString();
+}
+
+void SRigVMEditorSelectedDebugObjectWidget::OnRefresh()
+{
+	if (GetBlueprintObj())
+	{
+		GenerateDebugWorldNames(false);
+		GenerateDebugObjectInstances(false);
+
+		if (DebugObjectsComboBox.IsValid())
+		{
+			DebugWorldsComboBox->SetSelectedItem(GetDebugWorldName());
+			DebugObjectsComboBox->SetSelectedItem(GetDebugObjectInstance());
+		}
+	}
+}
+
+void SRigVMEditorSelectedDebugObjectWidget::GenerateDebugWorldNames(bool bRestoreSelection)
+{
+	DebugWorldNames.Empty();
+	DebugWorlds.Empty();
+
+	DebugWorlds.Add(nullptr);
+	DebugWorldNames.Add(MakeShareable(new FString(GetDebugAllWorldsString())));
+
+	for (TObjectIterator<UWorld> It; It; ++It)
+	{
+		UWorld *TestWorld = *It;
+
+		// Include only PIE and worlds that own the persistent level (i.e. non-streaming levels).
+		const bool bIsValidDebugWorld = (TestWorld != nullptr)
+			&& TestWorld->WorldType == EWorldType::PIE
+			&& TestWorld->PersistentLevel != nullptr
+			&& TestWorld->PersistentLevel->OwningWorld == TestWorld;
+
+		if (!bIsValidDebugWorld)
+		{
+			continue;
+		}
+
+		ENetMode NetMode = TestWorld->GetNetMode();
+
+		FString WorldName;
+
+		switch (NetMode)
+		{
+		case NM_Standalone:
+			WorldName = NSLOCTEXT("BlueprintEditor", "DebugWorldStandalone", "Standalone").ToString();
+			break;
+
+		case NM_ListenServer:
+			WorldName = NSLOCTEXT("BlueprintEditor", "DebugWorldListenServer", "Listen Server").ToString();
+			break;
+
+		case NM_DedicatedServer:
+			WorldName = NSLOCTEXT("BlueprintEditor", "DebugWorldDedicatedServer", "Dedicated Server").ToString();
+			break;
+
+		case NM_Client:
+			if (FWorldContext* PieContext = GEngine->GetWorldContextFromWorld(TestWorld))
+			{
+				WorldName = FString::Printf(TEXT("%s %d"), *NSLOCTEXT("BlueprintEditor", "DebugWorldClient", "Client").ToString(), PieContext->PIEInstance - 1);
+			}
+			break;
+		};
+
+		if (!WorldName.IsEmpty())
+		{
+			if (FWorldContext* PieContext = GEngine->GetWorldContextFromWorld(TestWorld))
+			{
+				if (!PieContext->CustomDescription.IsEmpty())
+				{
+					WorldName += TEXT(" ") + PieContext->CustomDescription;
+				}
+			}
+
+			// DebugWorlds & DebugWorldNames need to be the same size (we expect
+			// an index in one to correspond to the other) - DebugWorldNames is
+			// what populates the dropdown, so it is the authority (if there's 
+			// no name to present, they can't select from DebugWorlds)
+			DebugWorlds.Add(TestWorld);
+			DebugWorldNames.Add( MakeShareable(new FString(WorldName)) );
+		}
+	}
+
+	if (DebugWorldsComboBox.IsValid())
+	{
+		// Attempt to restore the old selection
+		if (bRestoreSelection)
+		{
+			TSharedPtr<FString> CurrentDebugWorld = GetDebugWorldName();
+			if (CurrentDebugWorld.IsValid())
+			{
+				DebugWorldsComboBox->SetSelectedItem(CurrentDebugWorld);
+			}
+		}
+
+		// Finally ensure we have a valid selection
+		TSharedPtr<FString> CurrentSelection = DebugWorldsComboBox->GetSelectedItem();
+		if (DebugWorldNames.Find(CurrentSelection) == INDEX_NONE)
+		{
+			if (DebugWorldNames.Num() > 0)
+			{
+				DebugWorldsComboBox->SetSelectedItem(DebugWorldNames[0]);
+			}
+			else
+			{
+				DebugWorldsComboBox->ClearSelection();
+			}
+		}
+
+		DebugWorldsComboBox->RefreshOptions();
+	}
+}
+
+void SRigVMEditorSelectedDebugObjectWidget::GenerateDebugObjectInstances(bool bRestoreSelection)
+{
+	// Cache the current selection as we may need to restore it
+	TSharedPtr<FRigVMDebugObjectInstance> LastSelection = GetDebugObjectInstance();
+
+	// Empty the lists of actors and regenerate them
+	DebugObjects.Empty();
+	DebugObjects.Add(MakeShareable(new FRigVMDebugObjectInstance(nullptr, GetNoDebugString())));
+
+	// Grab custom objects that should always be visible, regardless of the world
+	TArray<FRigVMCustomDebugObject> CustomDebugObjects;
+	Editor.Pin()->GetDebugObjects(/*inout*/ CustomDebugObjects);
+
+	for (const FRigVMCustomDebugObject& Entry : CustomDebugObjects)
+	{
+		AddDebugObject(Entry.Object, Entry.NameOverride);
+	}
+
+	// Check for a specific debug world. If DebugWorld=nullptr we take that as "any PIE world"
+	UWorld* DebugWorld = nullptr;
+	if (DebugWorldsComboBox.IsValid())
+	{
+		TSharedPtr<FString> CurrentWorldSelection = DebugWorldsComboBox->GetSelectedItem();
+		int32 SelectedIndex = INDEX_NONE;
+		for (int32 WorldIdx = 0; WorldIdx < DebugWorldNames.Num(); ++WorldIdx)
+		{
+			if (DebugWorldNames[WorldIdx].IsValid() && CurrentWorldSelection.IsValid()
+				&& (*DebugWorldNames[WorldIdx] == *CurrentWorldSelection))
+			{
+				SelectedIndex = WorldIdx;
+				break;
+			}
+		}
+		if (SelectedIndex > 0 && DebugWorldNames.IsValidIndex(SelectedIndex))
+		{
+			DebugWorld = DebugWorlds[SelectedIndex].Get();
+		}
+	}
+
+	UWorld* PreviewWorld = Editor.Pin()->GetPreviewScene()->GetWorld();
+
+	if (!Editor.Pin()->OnlyShowCustomDebugObjects())
+	{
+		UClass* BlueprintClass = GetBlueprintObj()->GeneratedClass;
+
+		if (BlueprintClass)
+		{
+			// Experimental new path for debug object discovery
+			TArray<UObject*> BlueprintInstances;
+			GetObjectsOfClass(BlueprintClass, BlueprintInstances, true);
+
+			for (auto It = BlueprintInstances.CreateIterator(); It; ++It)
+			{
+				UObject* TestObject = *It;
+				// Skip Blueprint preview objects (don't allow them to be selected for debugging)
+				if (PreviewWorld != nullptr && TestObject->IsIn(PreviewWorld))
+				{
+					continue;
+				}
+
+				// check outer chain for pending kill objects
+				bool bPendingKill = false;
+				UObject* ObjOuter = TestObject;
+				do
+				{
+					bPendingKill = !IsValid(ObjOuter);
+					ObjOuter = ObjOuter->GetOuter();
+				} while (!bPendingKill && ObjOuter != nullptr);
+
+				if (!TestObject->HasAnyFlags(RF_ClassDefaultObject) && !bPendingKill)
+				{
+					ObjOuter = TestObject;
+					UWorld *ObjWorld = nullptr;
+					static bool bUseNewWorldCode = false;
+					do		// Run through at least once in case the TestObject is a UGameInstance
+					{
+						UGameInstance *ObjGameInstance = Cast<UGameInstance>(ObjOuter);
+
+						ObjOuter = ObjOuter->GetOuter();
+						ObjWorld = ObjGameInstance ? ObjGameInstance->GetWorld() : Cast<UWorld>(ObjOuter);
+					} while (ObjWorld == nullptr && ObjOuter != nullptr);
+
+					if (ObjWorld)
+					{
+						// Make check on owning level (not streaming level)
+						if (ObjWorld->PersistentLevel && ObjWorld->PersistentLevel->OwningWorld)
+						{
+							ObjWorld = ObjWorld->PersistentLevel->OwningWorld;
+						}
+
+						// We have a specific debug world and the object isn't in it
+						if (DebugWorld && ObjWorld != DebugWorld)
+						{
+							continue;
+						}
+
+						if ((ObjWorld->WorldType == EWorldType::Editor) && (GUnrealEd->GetPIEViewport() == nullptr))
+						{
+							AddDebugObject(TestObject);
+						}
+						else if (ObjWorld->WorldType == EWorldType::PIE)
+						{
+							AddDebugObject(TestObject);
+						}
+					}
+				}
+			}
+		}
+		else
+		{
+			for (TObjectIterator<UObject> It; It; ++It)
+			{
+				UObject* TestObject = *It;
+
+				// Skip Blueprint preview objects (don't allow them to be selected for debugging)
+				if (PreviewWorld != nullptr && TestObject->IsIn(PreviewWorld))
+				{
+					continue;
+				}
+
+				const bool bPassesFlags = !TestObject->HasAnyFlags(RF_ClassDefaultObject) && IsValid(TestObject);
+				const bool bGeneratedByAnyBlueprint = TestObject->GetClass()->ClassGeneratedBy != nullptr;
+				const bool bGeneratedByThisBlueprint = bGeneratedByAnyBlueprint && GetBlueprintObj()->GeneratedClass && TestObject->IsA(GetBlueprintObj()->GeneratedClass);
+
+				if (bPassesFlags && bGeneratedByThisBlueprint)
+				{
+					UObject *ObjOuter = TestObject;
+					UWorld *ObjWorld = nullptr;
+					do		// Run through at least once in case the TestObject is a UGameInstance
+					{
+						UGameInstance *ObjGameInstance = Cast<UGameInstance>(ObjOuter);
+
+						ObjOuter = ObjOuter->GetOuter();
+						ObjWorld = ObjGameInstance ? ObjGameInstance->GetWorld() : Cast<UWorld>(ObjOuter);
+					} while (ObjWorld == nullptr && ObjOuter != nullptr);
+
+					if (ObjWorld)
+					{
+						// Make check on owning level (not streaming level)
+						if (ObjWorld->PersistentLevel && ObjWorld->PersistentLevel->OwningWorld)
+						{
+							ObjWorld = ObjWorld->PersistentLevel->OwningWorld;
+						}
+
+						// We have a specific debug world and the object isn't in it
+						if (DebugWorld && ObjWorld != DebugWorld)
+						{
+							continue;
+						}
+
+						if ((ObjWorld->WorldType == EWorldType::Editor) && (GUnrealEd->GetPIEViewport() == nullptr))
+						{
+							AddDebugObject(TestObject);
+						}
+						else if (ObjWorld->WorldType == EWorldType::PIE)
+						{
+							AddDebugObject(TestObject);
+						}
+					}
+				}
+			}
+		}
+	}
+
+	if (DebugObjectsComboBox.IsValid())
+	{
+		if (bRestoreSelection)
+		{
+			TSharedPtr<FRigVMDebugObjectInstance> NewSelection = GetDebugObjectInstance();
+			if (NewSelection.IsValid() && !NewSelection->IsEmptyObject())
+			{
+				// If our new selection matches the actual debug object, set it
+				DebugObjectsComboBox->SetSelectedItem(NewSelection);
+			}
+			else if (LastSelection.IsValid() && !LastSelection->IsEditorObject() && !LastSelection->IsEmptyObject())
+			{
+				// Re-add the desired runtime object if needed, even though it is currently null
+				DebugObjects.Add(LastSelection);
+				DebugObjectsComboBox->SetSelectedItem(LastSelection);
+			}
+		}
+
+		// Finally ensure we have a valid selection, this will set to all objects as a backup
+		TSharedPtr<FRigVMDebugObjectInstance> CurrentSelection = DebugObjectsComboBox->GetSelectedItem();
+		if (DebugObjects.Find(CurrentSelection) == INDEX_NONE)
+		{
+			if (DebugObjects.Num() > 0)
+			{
+				DebugObjectsComboBox->SetSelectedItem(DebugObjects[0]);
+			}
+			else
+			{
+				DebugObjectsComboBox->ClearSelection();
+			}
+		}
+
+		DebugObjectsComboBox->RefreshOptions();
+	}
+}
+
+TSharedPtr<FRigVMDebugObjectInstance> SRigVMEditorSelectedDebugObjectWidget::GetDebugObjectInstance() const
+{
+	check(GetBlueprintObj());
+	const FString& PathToDebug = GetBlueprintObj()->GetObjectPathToDebug();
+	if (!PathToDebug.IsEmpty())
+	{
+		for (int32 ObjectIndex = 0; ObjectIndex < DebugObjects.Num(); ++ObjectIndex)
+		{
+			if (DebugObjects[ObjectIndex].IsValid() && PathToDebug.Equals(DebugObjects[ObjectIndex]->ObjectPath))
+			{
+				return DebugObjects[ObjectIndex];
+			}
+		}
+	}
+
+	if (DebugObjects.Num() > 0)
+	{
+		return DebugObjects[0];
+	}
+	
+	return nullptr;
+}
+
+TSharedPtr<FString> SRigVMEditorSelectedDebugObjectWidget::GetDebugWorldName() const
+{
+	check(GetBlueprintObj());
+	if (ensure(DebugWorlds.Num() == DebugWorldNames.Num()))
+	{
+		UWorld* DebugWorld = GetBlueprintObj()->GetWorldBeingDebugged();
+		if (DebugWorld != nullptr)
+		{
+			for (int32 WorldIndex = 0; WorldIndex < DebugWorlds.Num(); ++WorldIndex)
+			{
+				if (DebugWorlds[WorldIndex].IsValid() && (DebugWorlds[WorldIndex].Get() == DebugWorld))
+				{
+					return DebugWorldNames[WorldIndex];
+				}
+			}
+		}
+	}
+
+
+	if (DebugWorldNames.Num() > 0)
+	{
+		return DebugWorldNames[0];
+	}
+
+	return nullptr;
+}
+
+void SRigVMEditorSelectedDebugObjectWidget::DebugWorldSelectionChanged(TSharedPtr<FString> NewSelection, ESelectInfo::Type SelectInfo)
+{
+	if (NewSelection != GetDebugWorldName())
+	{
+		check(DebugWorlds.Num() == DebugWorldNames.Num());
+		for (int32 WorldIdx = 0; WorldIdx < DebugWorldNames.Num(); ++WorldIdx)
+		{
+			if (DebugWorldNames[WorldIdx] == NewSelection)
+			{
+				GetBlueprintObj()->SetWorldBeingDebugged(DebugWorlds[WorldIdx].Get());
+
+				GetBlueprintObj()->SetObjectBeingDebugged(nullptr);
+				LastObjectObserved.Reset();
+
+				GenerateDebugObjectInstances(false);
+				break;
+			}
+		}
+	}
+}
+
+void SRigVMEditorSelectedDebugObjectWidget::DebugObjectSelectionChanged(TSharedPtr<FRigVMDebugObjectInstance> NewSelection, ESelectInfo::Type SelectInfo)
+{
+	if (NewSelection != GetDebugObjectInstance() && NewSelection.IsValid())
+	{
+		UObject* DebugObj = NewSelection->ObjectPtr.Get();
+		GetBlueprintObj()->SetObjectBeingDebugged(DebugObj);
+
+		if (TSharedPtr<IRigVMEditor> SharedBlueprintEditor = Editor.Pin())
+		{
+			SharedBlueprintEditor->ForceEditorRefresh();
+		}
+		
+		LastObjectObserved = DebugObj;
+	}
+}
+
+bool SRigVMEditorSelectedDebugObjectWidget::IsDebugObjectSelected() const
+{
+	check(GetBlueprintObj());
+	if (UObject* DebugObj = GetBlueprintObj()->GetObjectBeingDebugged())
+	{
+		if (AActor* Actor = Cast<AActor>(DebugObj))
+		{
+			return true;
+		}
+	}
+	return false;
+}
+
+void SRigVMEditorSelectedDebugObjectWidget::SelectedDebugObject_OnClicked()
+{
+	if (UObject* DebugObj = GetBlueprintObj()->GetObjectBeingDebugged())
+	{
+		if (AActor* Actor = Cast<AActor>(DebugObj))
+		{
+			GEditor->SelectNone(false, true, false);
+			GEditor->SelectActor(Actor, true, true, true);
+			GUnrealEd->Exec(Actor->GetWorld(), TEXT("CAMERA ALIGN ACTIVEVIEWPORTONLY"));
+		}
+	}
+}
+
+EVisibility SRigVMEditorSelectedDebugObjectWidget::IsDebugWorldComboVisible() const
+{
+	if (GEditor->PlayWorld != nullptr)
+	{
+		int32 LocalWorldCount = 0;
+		for (const FWorldContext& Context : GEngine->GetWorldContexts())
+		{
+			if (Context.WorldType == EWorldType::PIE && Context.World() != nullptr)
+			{
+				++LocalWorldCount;
+			}
+		}
+
+		if (LocalWorldCount > 1)
+		{
+			return EVisibility::Visible;
+		}
+	}
+
+	return EVisibility::Collapsed;
+}
+
+FString SRigVMEditorSelectedDebugObjectWidget::MakeDebugObjectLabel(UObject* TestObject, bool bAddContextIfSelectedInEditor, bool bAddSpawnedContext) const
+{
+	FString CustomLabelFromEditor = Editor.Pin()->GetCustomDebugObjectLabel(TestObject);
+	if (!CustomLabelFromEditor.IsEmpty())
+	{
+		return CustomLabelFromEditor;
+	}
+
+	auto GetActorLabelStringLambda = [](AActor* InActor, bool bIncludeNetModeSuffix, bool bIncludeSelectedSuffix, bool bIncludeSpawnedContext)
+	{
+		FString Label = InActor->GetActorLabel();
+
+		FString Context;
+
+		if (bIncludeNetModeSuffix)
+		{
+			switch (InActor->GetNetMode())
+			{
+			case ENetMode::NM_Client:
+			{
+				Context = NSLOCTEXT("BlueprintEditor", "DebugWorldClient", "Client").ToString();
+
+				FWorldContext* WorldContext = GEngine->GetWorldContextFromWorld(InActor->GetWorld());
+				if (WorldContext != nullptr && WorldContext->PIEInstance > 1)
+				{
+					Context += TEXT(" ");
+					Context += FText::AsNumber(WorldContext->PIEInstance - 1).ToString();
+				}
+			}
+			break;
+
+			case ENetMode::NM_ListenServer:
+			case ENetMode::NM_DedicatedServer:
+				Context = NSLOCTEXT("BlueprintEditor", "DebugWorldServer", "Server").ToString();
+				break;
+			}
+		}
+
+		if (bIncludeSpawnedContext)
+		{
+			if (!Context.IsEmpty())
+			{
+				Context += TEXT(", ");
+			}
+
+			Context += NSLOCTEXT("BlueprintEditor", "DebugObjectSpawned", "spawned").ToString();
+		}
+
+		if (bIncludeSelectedSuffix && InActor->IsSelected())
+		{
+			if (!Context.IsEmpty())
+			{
+				Context += TEXT(", ");
+			}
+
+			Context += NSLOCTEXT("BlueprintEditor", "DebugObjectSelected", "selected").ToString();
+		}
+
+		if (!Context.IsEmpty())
+		{
+			Label = FString::Printf(TEXT("%s (%s)"), *Label, *Context);
+		}
+
+		return Label;
+	};
+
+	// Include net mode suffix when "All worlds" is selected.
+	const bool bIncludeNetModeSuffix = *GetDebugWorldName() == GetDebugAllWorldsString();
+
+	FString Label;
+	if (AActor* Actor = Cast<AActor>(TestObject))
+	{
+		Label = GetActorLabelStringLambda(Actor, bIncludeNetModeSuffix, bAddContextIfSelectedInEditor, bAddSpawnedContext);
+	}
+	else
+	{
+		if (AActor* ParentActor = TestObject->GetTypedOuter<AActor>())
+		{
+			// We don't need the full path because it's in the tooltip
+			const FString RelativePath = TestObject->GetName();
+			Label = FString::Printf(TEXT("%s in %s"), *RelativePath, *GetActorLabelStringLambda(ParentActor, bIncludeNetModeSuffix, bAddContextIfSelectedInEditor, bAddSpawnedContext));
+		}
+		else
+		{
+			Label = TestObject->GetName();
+		}
+	}
+
+	return Label;
+}
+
+void SRigVMEditorSelectedDebugObjectWidget::FillDebugObjectInstance(TSharedPtr<FRigVMDebugObjectInstance> Instance)
+{
+	check(Instance.IsValid());
+	FRigVMDebugObjectInstance& Ref = *Instance.Get();
+
+	if (Ref.ObjectPtr.IsValid())
+	{
+		Ref.ObjectPath = Ref.ObjectPtr->GetPathName();
+
+		// Compute non-PIE path
+		FString OriginalPath = UWorld::RemovePIEPrefix(Ref.ObjectPath);
+
+		// Look for original object
+		UObject* OriginalObject = FindObjectSafe<UObject>(nullptr, *OriginalPath);
+
+		if (OriginalObject)
+		{
+			Ref.EditorObjectPath = OriginalPath;
+		}
+		else
+		{
+			// No editor path, was dynamically spawned
+			Ref.EditorObjectPath = FString();
+		}
+	}
+	else
+	{
+		Ref.ObjectPath = Ref.EditorObjectPath = FString();
+	}
+}
+
+void SRigVMEditorSelectedDebugObjectWidget::AddDebugObject(UObject* TestObject, const FString& TestObjectName)
+{
+	TSharedPtr<FRigVMDebugObjectInstance> NewInstance = MakeShareable(new FRigVMDebugObjectInstance(TestObject, TestObjectName));
+	FillDebugObjectInstance(NewInstance);
+
+	if (TestObjectName.IsEmpty())
+	{
+		NewInstance->ObjectLabel = MakeDebugObjectLabel(TestObject, true, NewInstance->IsSpawnedObject());
+	}
+
+	if (UWidget* DebugWidget = Cast<UWidget>(TestObject))
+	{
+		if (!DebugWidget->IsConstructed())
+		{
+			NewInstance->ObjectLabel += " (No Slate Widget)";
+		}
+	}
+	
+	DebugObjects.Add(NewInstance);
+}
+
+UBlueprint* SRigVMEditorSelectedDebugObjectWidget::GetBlueprintObj() const
+{
+	return Editor.Pin()->GetRigVMBlueprint();
+}
+
+TSharedRef<SWidget> SRigVMEditorSelectedDebugObjectWidget::CreateDebugObjectItemWidget(TSharedPtr<FRigVMDebugObjectInstance> InItem)
+{
+	FString ItemString;
+	FString ItemTooltip;
+
+	if (InItem.IsValid())
+	{
+		ItemString = InItem->ObjectLabel;
+		ItemTooltip = InItem->ObjectPath;
+	}
+
+	return SNew(STextBlock)
+		.Text(FText::FromString(*ItemString))
+		.ToolTipText(FText::FromString(*ItemTooltip));
+}
+
+FText SRigVMEditorSelectedDebugObjectWidget::GetSelectedDebugObjectTextLabel() const
+{
+	FString Label;
+
+	TSharedPtr<FRigVMDebugObjectInstance> DebugInstance = GetDebugObjectInstance();
+	if (DebugInstance.IsValid())
+	{
+		Label = DebugInstance->ObjectLabel;
+
+		UBlueprint* Blueprint = GetBlueprintObj();
+		if (Blueprint != nullptr)
+		{
+			UObject* DebugObj = Blueprint->GetObjectBeingDebugged();
+			if (DebugObj != nullptr)
+			{
+				// Exclude the editor selection suffix for the combo button's label.
+				Label = MakeDebugObjectLabel(DebugObj, false, DebugInstance->IsSpawnedObject());
+			}
+		}
+	}
+
+	return FText::FromString(Label);
+}
+
+//////////////////////////////////////////////////////////////////////////
+
+#undef LOCTEXT_NAMESPACE
diff --git a/Engine/Plugins/Runtime/RigVM/Source/RigVMEditor/Public/Editor/RigVMEditor.h b/Engine/Plugins/Runtime/RigVM/Source/RigVMEditor/Public/Editor/RigVMEditor.h
--- a/Engine/Plugins/Runtime/RigVM/Source/RigVMEditor/Public/Editor/RigVMEditor.h
+++ b/Engine/Plugins/Runtime/RigVM/Source/RigVMEditor/Public/Editor/RigVMEditor.h
@@ -171,6 +171,7 @@ public:
 
 	virtual void HandleModifiedEvent(ERigVMGraphNotifType InNotifType, URigVMGraph* InGraph, UObject* InSubject) = 0;
 	virtual void OnFinishedChangingProperties(const FPropertyChangedEvent& PropertyChangedEvent) = 0;
+	virtual FPreviewScene* GetPreviewScene() = 0;
 	
 protected:
 	
diff --git a/Engine/Plugins/Runtime/RigVM/Source/RigVMEditor/Public/Editor/RigVMLegacyEditor.h b/Engine/Plugins/Runtime/RigVM/Source/RigVMEditor/Public/Editor/RigVMLegacyEditor.h
--- a/Engine/Plugins/Runtime/RigVM/Source/RigVMEditor/Public/Editor/RigVMLegacyEditor.h
+++ b/Engine/Plugins/Runtime/RigVM/Source/RigVMEditor/Public/Editor/RigVMLegacyEditor.h
@@ -123,6 +123,7 @@ public:
 	virtual void SummonSearchUI(bool bSetFindWithinBlueprint, FString NewSearchTerms, bool bSelectFirstResult) override { FBlueprintEditor::SummonSearchUI(bSetFindWithinBlueprint, NewSearchTerms, bSelectFirstResult); }
 	virtual TSharedRef<FUICommandList> GetToolkitCommands() override { return ToolkitCommands; }
 	virtual void OnClose() override;
+	virtual FPreviewScene* GetPreviewScene() override { return &PreviewScene; }
 
 	
 protected:
diff --git a/Engine/Plugins/Runtime/RigVM/Source/RigVMEditor/Public/Editor/RigVMNewEditor.h b/Engine/Plugins/Runtime/RigVM/Source/RigVMEditor/Public/Editor/RigVMNewEditor.h
--- a/Engine/Plugins/Runtime/RigVM/Source/RigVMEditor/Public/Editor/RigVMNewEditor.h
+++ b/Engine/Plugins/Runtime/RigVM/Source/RigVMEditor/Public/Editor/RigVMNewEditor.h
@@ -119,13 +119,14 @@ public:
 	virtual const TArray<UObject*>* GetObjectsCurrentlyBeingEdited() const override { return FAssetEditorToolkit::GetObjectsCurrentlyBeingEdited(); }
 	virtual void AddCompileWidget(FToolBarBuilder& ToolbarBuilder) override;
 	virtual void AddSettingsAndDefaultWidget(FToolBarBuilder& ToolbarBuilder) override;
-	virtual void AddSelectedDebugObjectWidget(FToolBarBuilder& ToolbarBuilder) override {}
+	virtual void AddSelectedDebugObjectWidget(FToolBarBuilder& ToolbarBuilder) override;
 	virtual void Compile() { FRigVMEditorBase::Compile(); }
 	virtual void OnCreateGraphEditorCommands(TSharedPtr<FUICommandList> GraphEditorCommandsList) { FRigVMEditorBase::OnCreateGraphEditorCommands(GraphEditorCommandsList); }
 	virtual bool ShouldOpenGraphByDefault() const { return FRigVMEditorBase::ShouldOpenGraphByDefault(); }
 	virtual void OnFinishedChangingProperties(const FPropertyChangedEvent& PropertyChangedEvent) { FRigVMEditorBase::OnFinishedChangingProperties(PropertyChangedEvent); }
 	virtual void HandleSetObjectBeingDebugged(UObject* InObject) { FRigVMEditorBase::HandleSetObjectBeingDebugged(InObject); }
 	virtual FReply OnSpawnGraphNodeByShortcut(FInputChord InChord, const FVector2D& InPosition, UEdGraph* InGraph) { return FRigVMEditorBase::OnSpawnGraphNodeByShortcut(InChord, InPosition, InGraph); }
+	virtual FPreviewScene* GetPreviewScene() override { return nullptr; }
 
 	//~ Begin IToolkit Interface
 	virtual void RegisterTabSpawners(const TSharedRef<class FTabManager>& InTabManager) override;
diff --git a/Engine/Plugins/Runtime/RigVM/Source/RigVMEditor/Public/Widgets/SRigVMEditorSelectedDebugObjectWidget.h b/Engine/Plugins/Runtime/RigVM/Source/RigVMEditor/Public/Widgets/SRigVMEditorSelectedDebugObjectWidget.h
new file mode 100644
--- /dev/null
+++ b/Engine/Plugins/Runtime/RigVM/Source/RigVMEditor/Public/Widgets/SRigVMEditorSelectedDebugObjectWidget.h
@@ -0,0 +1,127 @@
+// Copyright Epic Games, Inc. All Rights Reserved.
+
+#pragma once
+
+#include "Editor/RigVMEditor.h"
+#include "Widgets/Input/SComboBox.h"
+
+struct FRigVMDebugObjectInstance
+{
+	/** Actual object to debug, can be null */
+	TWeakObjectPtr<UObject> ObjectPtr;
+
+	/** Friendly label for object to debug */
+	FString ObjectLabel;
+
+	/** Raw object path of spawned PIE object, this is not a SoftObjectPath because we don't want it to get fixed up */
+	FString ObjectPath;
+
+	/** Object path to object in the editor, will only be set for static objects */
+	FString EditorObjectPath;
+
+	FRigVMDebugObjectInstance(TWeakObjectPtr<UObject> InPtr, const FString& InLabel)
+		: ObjectPtr(InPtr)
+		, ObjectLabel(InLabel)
+	{
+	}
+
+	/** Returns true if this is the special entry for no specific object */
+	bool IsEmptyObject() const
+	{
+		return ObjectPath.IsEmpty();
+	}
+
+	/** If this has no editor path, it was spawned */
+	bool IsSpawnedObject() const
+	{
+		return !ObjectPath.IsEmpty() && EditorObjectPath.IsEmpty();
+	}
+
+	/** If editor and object path are the same length because there's no prefix, this is the editor object */
+	bool IsEditorObject() const
+	{
+		return !ObjectPath.IsEmpty() && ObjectPath.Len() == EditorObjectPath.Len();
+	}
+};
+
+class RIGVMEDITOR_API SRigVMEditorSelectedDebugObjectWidget : public SCompoundWidget
+{
+public:
+	SLATE_BEGIN_ARGS(SRigVMEditorSelectedDebugObjectWidget){}
+	SLATE_END_ARGS()
+
+	void Construct(const FArguments& InArgs, TSharedPtr<IRigVMEditor> InEditor);
+
+	// SWidget interface
+	virtual void Tick( const FGeometry& AllottedGeometry, const double InCurrentTime, const float InDeltaTime ) override;
+	// End of SWidget interface
+
+	/** Adds an object to the list of debug choices */
+	void AddDebugObject(UObject* TestObject, const FString& TestObjectName = FString());
+
+private:
+	UBlueprint* GetBlueprintObj() const;
+
+	/** Creates a list of all debug objects **/
+	void GenerateDebugObjectInstances(bool bRestoreSelection);
+
+	/** Generate list of active PIE worlds to debug **/
+	void GenerateDebugWorldNames(bool bRestoreSelection);
+
+	/** Refresh the widget. **/
+	void OnRefresh();
+
+	/** Returns the entry for the current debug actor */
+	TSharedPtr<FRigVMDebugObjectInstance> GetDebugObjectInstance() const;
+
+	/** Returns the name of the current debug actor */
+	TSharedPtr<FString> GetDebugWorldName() const;
+
+	/** Handles the selection changed event for the debug actor combo box */
+	void DebugObjectSelectionChanged(TSharedPtr<FRigVMDebugObjectInstance> NewSelection, ESelectInfo::Type SelectInfo);
+
+	void DebugWorldSelectionChanged(TSharedPtr<FString> NewSelection, ESelectInfo::Type SelectInfo);
+
+	/** Called when user clicks button to select the current object being debugged */
+	void SelectedDebugObject_OnClicked();
+
+	/** Returns true if a debug actor is currently selected */
+	bool IsDebugObjectSelected() const;
+
+	EVisibility IsDebugWorldComboVisible() const;
+
+	/* Returns the string to indicate no debug object is selected */
+	const FString& GetNoDebugString() const;
+
+	const FString& GetDebugAllWorldsString() const;
+
+	/** Helper method to construct a debug object label string */
+	FString MakeDebugObjectLabel(UObject* TestObject, bool bAddContextIfSelectedInEditor, bool bAddSpawnedContext) const;
+
+	/** Fills in data for a specific instance */
+	void FillDebugObjectInstance(TSharedPtr<FRigVMDebugObjectInstance> Instance);
+
+	/** Called to create a widget for each debug object item */
+	TSharedRef<SWidget> CreateDebugObjectItemWidget(TSharedPtr<FRigVMDebugObjectInstance> InItem);
+
+	/** Returns the combo button label to use for the currently-selected debug object item */
+	FText GetSelectedDebugObjectTextLabel() const;
+
+private:
+	/** Pointer back to the rigvm editor tool that owns us */
+	TWeakPtr<IRigVMEditor> Editor;
+
+	/** Lists of actors of a given blueprint type and their names */
+	TArray<TSharedPtr<FRigVMDebugObjectInstance>> DebugObjects;
+
+	/** PIE worlds that we can debug */
+	TArray< TWeakObjectPtr<UWorld> > DebugWorlds;
+	TArray< TSharedPtr<FString> > DebugWorldNames;
+
+	/** Widget containing the names of all possible debug actors. This is a "generic" SComboBox rather than an STextComboBox so that we can customize the label on the combo button widget. */
+	TSharedPtr<SComboBox<TSharedPtr<FRigVMDebugObjectInstance>>> DebugObjectsComboBox;
+
+	TSharedPtr<STextComboBox> DebugWorldsComboBox;
+
+	TWeakObjectPtr<UObject> LastObjectObserved;
+};
diff --git a/Engine/Plugins/Runtime/RigVM/Source/RigVMEditor/RigVMEditor.Build.cs b/Engine/Plugins/Runtime/RigVM/Source/RigVMEditor/RigVMEditor.Build.cs
--- a/Engine/Plugins/Runtime/RigVM/Source/RigVMEditor/RigVMEditor.Build.cs
+++ b/Engine/Plugins/Runtime/RigVM/Source/RigVMEditor/RigVMEditor.Build.cs
@@ -19,6 +19,7 @@ public class RigVMEditor : ModuleRules
                 "UnrealEd",
                 "Slate",
                 "SlateCore",
+                "UMG",
                 "InputCore",
                 "ToolWidgets",
                 "KismetCompiler",
