RigVM: Add compiler results tab
#rnx

diff --git a/Engine/Plugins/Animation/ControlRig/Source/ControlRigEditor/Private/Editor/ControlRigEditorMode.cpp b/Engine/Plugins/Animation/ControlRig/Source/ControlRigEditor/Private/Editor/ControlRigEditorMode.cpp
--- a/Engine/Plugins/Animation/ControlRig/Source/ControlRigEditor/Private/Editor/ControlRigEditorMode.cpp
+++ b/Engine/Plugins/Animation/ControlRig/Source/ControlRigEditor/Private/Editor/ControlRigEditorMode.cpp
@@ -17,6 +17,7 @@
 #include "Editor/RigValidationTabSummoner.h"
 #include "Editor/RigAnimAttributeTabSummoner.h"
 #include "ToolMenus.h"
+#include "Editor/RigVMCompilerResultsTabSummoner.h"
 #include "Editor/RigVMDetailsInspectorTabSummoner.h"
 
 #if WITH_RIGVMLEGACYEDITOR
@@ -334,6 +335,7 @@ FControlRigEditorMode::FControlRigEditorMode(const TSharedRef<FControlRigEditor>
 							// Middle bottom - compiler results & find
 							FTabManager::NewStack()
 							->SetSizeCoefficient(0.2f)
+							->AddTab(FRigVMCompilerResultsTabSummoner::TabID(), ETabState::OpenedTab)
 						)
 					)
 					->Split
diff --git a/Engine/Plugins/Runtime/RigVM/Source/RigVMEditor/Private/Editor/RigVMCompilerResultsTabSummoner.cpp b/Engine/Plugins/Runtime/RigVM/Source/RigVMEditor/Private/Editor/RigVMCompilerResultsTabSummoner.cpp
new file mode 100644
--- /dev/null
+++ b/Engine/Plugins/Runtime/RigVM/Source/RigVMEditor/Private/Editor/RigVMCompilerResultsTabSummoner.cpp
@@ -0,0 +1,33 @@
+// Copyright Epic Games, Inc. All Rights Reserved.
+
+#include "Editor/RigVMCompilerResultsTabSummoner.h"
+#include "Editor/RigVMNewEditor.h"
+
+class SWidget;
+struct FSlateBrush;
+
+#define LOCTEXT_NAMESPACE "RigVMEditor"
+
+FRigVMCompilerResultsTabSummoner::FRigVMCompilerResultsTabSummoner(TSharedPtr<class FAssetEditorToolkit> InHostingApp)
+	: FWorkflowTabFactory(TabID(), InHostingApp)
+{
+	TabLabel = LOCTEXT("CompilerResultsTabTitle", "Compiler Results");
+	TabIcon = FSlateIcon(FAppStyle::GetAppStyleSetName(), "Kismet.Tabs.CompilerResults");
+
+	bIsSingleton = true;
+
+	ViewMenuDescription = LOCTEXT("CompilerResultsView", "Compiler Results");
+	ViewMenuTooltip = LOCTEXT("CompilerResultsView_ToolTip", "Show compiler results of all functions and variables");
+}
+
+
+TSharedRef<SWidget> FRigVMCompilerResultsTabSummoner::CreateTabBody(const FWorkflowTabSpawnInfo& Info) const
+{
+	TSharedPtr<FRigVMNewEditor> EditorPtr = StaticCastSharedPtr<FRigVMNewEditor>(HostingApp.Pin());
+
+	return EditorPtr->GetCompilerResults();
+}
+
+
+
+#undef LOCTEXT_NAMESPACE
diff --git a/Engine/Plugins/Runtime/RigVM/Source/RigVMEditor/Private/Editor/RigVMNewEditor.cpp b/Engine/Plugins/Runtime/RigVM/Source/RigVMEditor/Private/Editor/RigVMNewEditor.cpp
--- a/Engine/Plugins/Runtime/RigVM/Source/RigVMEditor/Private/Editor/RigVMNewEditor.cpp
+++ b/Engine/Plugins/Runtime/RigVM/Source/RigVMEditor/Private/Editor/RigVMNewEditor.cpp
@@ -28,21 +28,22 @@
 #include "Logging/MessageLog.h"
 #include "GraphEditor.h"
 #include "GraphEditorActions.h"
+#include "MessageLogModule.h"
 #include "RigVMBlueprint.h"
 #include "RigVMEditorCommands.h"
+#include "Editor/RigVMCompilerResultsTabSummoner.h"
 #include "Editor/RigVMDetailsInspectorTabSummoner.h"
 #include "Editor/RigVMGraphEditorSummoner.h"
 #include "Widgets/Input/SEditableTextBox.h"
 #include "Editor/SRigVMDetailsInspector.h"
 #include "Engine/LevelScriptBlueprint.h"
 #include "Kismet2/KismetEditorUtilities.h"
+#include "Misc/UObjectToken.h"
 #include "Widgets/SRigVMEditorSelectedDebugObjectWidget.h"
 
 
 #define LOCTEXT_NAMESPACE "RigVMNewEditor"
 
-const FName FRigVMNewEditorTabs::CompilerResultsID() {	static FName ID = TEXT( "CompilerResults" ); return ID; };
-
 const FName FRigVMNewEditor::SelectionState_GraphExplorer() {	static FName State = (TEXT("GraphExplorer")); return State; };
 const FName FRigVMNewEditor::SelectionState_Graph() {			static FName State = (TEXT("Graph")); return State; };
 const FName FRigVMNewEditor::SelectionState_ClassSettings() {	static FName State = (TEXT("ClassSettings")); return State; };
@@ -342,7 +343,7 @@ void FRigVMNewEditor::PostLayoutBlueprintEditorInitialization()
 
 			if (FocusedGraphEdPtr.IsValid())
 			{
-				FNotificationInfo Info( NSLOCTEXT("Kismet", "Blueprint Modified", "Blueprint requires updating. Please resave.") );
+				FNotificationInfo Info( NSLOCTEXT("RigVM", "Blueprint Modified", "Blueprint requires updating. Please resave.") );
 				Info.Image = FAppStyle::GetBrush(TEXT("Icons.Info"));
 				Info.bFireAndForget = true;
 				Info.bUseSuccessFailIcons = false;
@@ -361,32 +362,19 @@ void FRigVMNewEditor::PostLayoutBlueprintEditorInitialization()
 		}
 
 		// Determine if the current "mode" supports invoking the Compiler Results tab.
-#if WITH_RIGVMLEGACYEDITOR
-		const bool bCanInvokeCompilerResultsTab = TabManager->HasTabSpawner(FBlueprintEditorTabs::CompilerResultsID);
-#else
-		const bool bCanInvokeCompilerResultsTab = TabManager->HasTabSpawner(FRigVMNewEditorTabs::CompilerResultsID());
-#endif
+		const bool bCanInvokeCompilerResultsTab = TabManager->HasTabSpawner(FRigVMCompilerResultsTabSummoner::TabID());
 
 		// If we have a warning/error, open output log if the current mode allows us to invoke it.
-		const bool bIsBlueprintInWarningOrErrorState = !Blueprint->IsUpToDate() || (Blueprint->Status == BS_UpToDateWithWarnings);
-		if (bIsBlueprintInWarningOrErrorState && bCanInvokeCompilerResultsTab)
+		if (bCanInvokeCompilerResultsTab)
 		{
-#if WITH_RIGVMLEGACYEDITOR
-			TabManager->TryInvokeTab(FBlueprintEditorTabs::CompilerResultsID);
-#else
-			TabManager->TryInvokeTab(FRigVMNewEditorTabs::CompilerResultsID());
-#endif
+			TabManager->TryInvokeTab(FRigVMCompilerResultsTabSummoner::TabID());
 		}
 		else
 		{
 			// Toolkit modes that don't include this tab may have been incorrectly saved with layout information for restoring it
 			// as an "unrecognized" tab, due to having previously invoked it above without checking to see if the layout can open
 			// it first. To correct this, we check if the tab was restored from a saved layout here, and close it if not supported.
-#if WITH_RIGVMLEGACYEDITOR
-			TSharedPtr<SDockTab> TabPtr = TabManager->FindExistingLiveTab(FBlueprintEditorTabs::CompilerResultsID);
-#else
-			TSharedPtr<SDockTab> TabPtr = TabManager->FindExistingLiveTab(FRigVMNewEditorTabs::CompilerResultsID());
-#endif
+			TSharedPtr<SDockTab> TabPtr = TabManager->FindExistingLiveTab(FRigVMCompilerResultsTabSummoner::TabID());
 			if (TabPtr.IsValid() && !bCanInvokeCompilerResultsTab)
 			{
 				TabPtr->RequestCloseTab();
@@ -1213,8 +1201,6 @@ void FRigVMNewEditor::CreateDefaultCommands()
 
 TSharedRef<SGraphEditor> FRigVMNewEditor::CreateGraphEditorWidget(TSharedRef<FTabInfo> InTabInfo, UEdGraph* InGraph)
 {
-	// TODO sara-s Populate the graph editor commands 
-
 	check((InGraph != nullptr) && IsEditingSingleBlueprint());
 
 	// No need to regenerate the commands.
@@ -1396,7 +1382,7 @@ void FRigVMNewEditor::CompileImpl()
 		// CachedNumErrors = LogResults.NumErrors;
 
 		const bool bForceMessageDisplay = (LogResults.NumWarnings > 0 || LogResults.NumErrors > 0) && !BlueprintObj->bIsRegeneratingOnLoad;
-		//DumpMessagesToCompilerLog(LogResults.Messages, bForceMessageDisplay);
+		DumpMessagesToCompilerLog(LogResults.Messages, bForceMessageDisplay);
 
 		if (BlueprintObj->UpgradeNotesLog.IsValid())
 		{
@@ -2032,6 +2018,37 @@ void FRigVMNewEditor::ClearSelectionStateFor(FName SelectionOwner)
 	}
 }
 
+void FRigVMNewEditor::OnLogTokenClicked(const TSharedRef<IMessageToken>& MessageToken)
+{
+	if (MessageToken->GetType() == EMessageToken::EdGraph)
+	{
+		const TSharedRef<FEdGraphToken> EdGraphToken = StaticCastSharedRef<FEdGraphToken>(MessageToken);
+		const UEdGraphPin* PinBeingReferenced = EdGraphToken->GetPin();
+		const UObject* ObjectBeingReferenced = EdGraphToken->GetGraphObject();
+		if (PinBeingReferenced)
+		{
+			JumpToPin(PinBeingReferenced);
+		}
+		else if(ObjectBeingReferenced)
+		{
+			JumpToHyperlink(ObjectBeingReferenced, false);
+		}
+	}
+}
+
+void FRigVMNewEditor::DumpMessagesToCompilerLog(const TArray<TSharedRef<class FTokenizedMessage>>& Messages, bool bForceMessageDisplay)
+{
+	CompilerResultsListing->ClearMessages();
+
+	// Note we dont mirror to the output log here as the compiler already does that
+	CompilerResultsListing->AddMessages(Messages, false);
+	
+	if (bForceMessageDisplay)
+	{
+		TabManager->TryInvokeTab(FBlueprintEditorTabs::CompilerResultsID);
+	}
+}
+
 void FRigVMNewEditor::CreateDefaultTabContents(const TArray<UBlueprint*> InBlueprints)
 {
 	UBlueprint* InBlueprint = InBlueprints.Num() == 1 ? InBlueprints[0] : nullptr;
@@ -2054,11 +2071,10 @@ void FRigVMNewEditor::CreateDefaultTabContents(const TArray<UBlueprint*> InBluep
 	}
 	
 	CompilerResultsListing = FCompilerResultsLog::GetBlueprintMessageLog(InBlueprint);
-	//CompilerResultsListing->OnMessageTokenClicked().AddSP(this, &FBlueprintEditor::OnLogTokenClicked); // TODO sara-s
+	CompilerResultsListing->OnMessageTokenClicked().AddSP(this, &FRigVMNewEditor::OnLogTokenClicked); 
 
-	// TODO sara-s
-	// FMessageLogModule& MessageLogModule = FModuleManager::LoadModuleChecked<FMessageLogModule>("MessageLog");
-	// CompilerResults = MessageLogModule.CreateLogListingWidget( CompilerResultsListing.ToSharedRef() );
+	FMessageLogModule& MessageLogModule = FModuleManager::LoadModuleChecked<FMessageLogModule>("MessageLog");
+	CompilerResults = MessageLogModule.CreateLogListingWidget( CompilerResultsListing.ToSharedRef() );
 	// FindResults = SNew(SFindInBlueprints, SharedThis(this));
 	
 	this->Inspector = 
diff --git a/Engine/Plugins/Runtime/RigVM/Source/RigVMEditor/Private/Editor/RigVMNewEditorMode.cpp b/Engine/Plugins/Runtime/RigVM/Source/RigVMEditor/Private/Editor/RigVMNewEditorMode.cpp
--- a/Engine/Plugins/Runtime/RigVM/Source/RigVMEditor/Private/Editor/RigVMNewEditorMode.cpp
+++ b/Engine/Plugins/Runtime/RigVM/Source/RigVMEditor/Private/Editor/RigVMNewEditorMode.cpp
@@ -3,6 +3,7 @@
 #include "Editor/RigVMNewEditorMode.h"
 
 #include "Editor/RigVMDetailsInspectorTabSummoner.h"
+#include "Editor/RigVMCompilerResultsTabSummoner.h"
 #include "Editor/RigVMEditor.h"
 #include "Editor/RigVMExecutionStackTabSummoner.h"
 #include "Editor/RigVMNewEditor.h"
@@ -18,6 +19,7 @@ FRigVMNewEditorMode::FRigVMNewEditorMode(const TSharedRef<FRigVMNewEditor>& InRi
 	TabFactories.RegisterFactory(MakeShared<FRigVMExecutionStackTabSummoner>(InRigVMEditor));
 	TabFactories.RegisterFactory(MakeShared<FRigVMEditorGraphExplorerTabSummoner>(InRigVMEditor));
 	TabFactories.RegisterFactory(MakeShared<FRigVMDetailsInspectorTabSummoner>(InRigVMEditor));
+	TabFactories.RegisterFactory(MakeShared<FRigVMCompilerResultsTabSummoner>(InRigVMEditor));
 
 	TabLayout = FTabManager::NewLayout("Standalone_RigVMEditMode_Layout_v1.5")
 		->AddArea
@@ -45,11 +47,18 @@ FRigVMNewEditorMode::FRigVMNewEditorMode(const TSharedRef<FRigVMNewEditor>& InRi
 					->SetSizeCoefficient(0.6f)
 					->Split
 					(
-						// Middle top - document edit area
+						// Middle top
 						FTabManager::NewStack()
 						->SetSizeCoefficient(0.8f)
 						->AddTab("Document", ETabState::ClosedTab)
 					)
+					->Split
+						(
+							// Middle bottom - compiler results & find
+							FTabManager::NewStack()
+							->SetSizeCoefficient(0.2f)
+							->AddTab(FRigVMCompilerResultsTabSummoner::TabID(), ETabState::OpenedTab)
+						)
 				)
 				->Split
 				(
diff --git a/Engine/Plugins/Runtime/RigVM/Source/RigVMEditor/Private/Widgets/SRigVMEditorSelectedDebugObjectWidget.cpp b/Engine/Plugins/Runtime/RigVM/Source/RigVMEditor/Private/Widgets/SRigVMEditorSelectedDebugObjectWidget.cpp
--- a/Engine/Plugins/Runtime/RigVM/Source/RigVMEditor/Private/Widgets/SRigVMEditorSelectedDebugObjectWidget.cpp
+++ b/Engine/Plugins/Runtime/RigVM/Source/RigVMEditor/Private/Widgets/SRigVMEditorSelectedDebugObjectWidget.cpp
@@ -260,6 +260,12 @@ void SRigVMEditorSelectedDebugObjectWidget::GenerateDebugWorldNames(bool bRestor
 
 void SRigVMEditorSelectedDebugObjectWidget::GenerateDebugObjectInstances(bool bRestoreSelection)
 {
+	TSharedPtr<IRigVMEditor> EditorPtr = Editor.Pin();
+	if (!EditorPtr)
+	{
+		return;
+	}
+	
 	// Cache the current selection as we may need to restore it
 	TSharedPtr<FRigVMDebugObjectInstance> LastSelection = GetDebugObjectInstance();
 
@@ -269,7 +275,7 @@ void SRigVMEditorSelectedDebugObjectWidget::GenerateDebugObjectInstances(bool bR
 
 	// Grab custom objects that should always be visible, regardless of the world
 	TArray<FRigVMCustomDebugObject> CustomDebugObjects;
-	Editor.Pin()->GetDebugObjects(/*inout*/ CustomDebugObjects);
+	EditorPtr->GetDebugObjects(/*inout*/ CustomDebugObjects);
 
 	for (const FRigVMCustomDebugObject& Entry : CustomDebugObjects)
 	{
@@ -297,9 +303,13 @@ void SRigVMEditorSelectedDebugObjectWidget::GenerateDebugObjectInstances(bool bR
 		}
 	}
 
-	UWorld* PreviewWorld = Editor.Pin()->GetPreviewScene()->GetWorld();
+	UWorld* PreviewWorld = nullptr;
+	if (FPreviewScene* PreviewScene = EditorPtr->GetPreviewScene())
+	{
+		PreviewWorld = PreviewScene->GetWorld();
+	}
 
-	if (!Editor.Pin()->OnlyShowCustomDebugObjects())
+	if (!EditorPtr->OnlyShowCustomDebugObjects())
 	{
 		UClass* BlueprintClass = GetBlueprintObj()->GeneratedClass;
 
diff --git a/Engine/Plugins/Runtime/RigVM/Source/RigVMEditor/Public/Editor/RigVMCompilerResultsTabSummoner.h b/Engine/Plugins/Runtime/RigVM/Source/RigVMEditor/Public/Editor/RigVMCompilerResultsTabSummoner.h
new file mode 100644
--- /dev/null
+++ b/Engine/Plugins/Runtime/RigVM/Source/RigVMEditor/Public/Editor/RigVMCompilerResultsTabSummoner.h
@@ -0,0 +1,32 @@
+// Copyright Epic Games, Inc. All Rights Reserved.
+
+#pragma once
+
+#include "CoreMinimal.h"
+#include "Widgets/SWidget.h"
+#include "Toolkits/AssetEditorToolkit.h"
+#include "WorkflowOrientedApp/WorkflowTabFactory.h"
+
+#define LOCTEXT_NAMESPACE "RigVMEditor"
+
+//////////////////////////////////////////////////////////////////////////
+// FCompilerResultsSummoner
+
+struct RIGVMEDITOR_API FRigVMCompilerResultsTabSummoner : public FWorkflowTabFactory
+{
+
+public:
+
+	static const FName TabID() { static FName ID = TEXT("RigVM Compiler Results"); return ID; }
+	
+	FRigVMCompilerResultsTabSummoner(TSharedPtr<class FAssetEditorToolkit> InHostingApp);
+
+	virtual TSharedRef<SWidget> CreateTabBody(const FWorkflowTabSpawnInfo& Info) const override;
+
+	virtual FText GetTabToolTipText(const FWorkflowTabSpawnInfo& Info) const override
+	{
+		return LOCTEXT("CompilerResultsTooltip", "The compiler results tab shows any errors or warnings generated when compiling this Blueprint.");
+	}
+};
+
+#undef LOCTEXT_NAMESPACE
\ No newline at end of file
diff --git a/Engine/Plugins/Runtime/RigVM/Source/RigVMEditor/Public/Editor/RigVMNewEditor.h b/Engine/Plugins/Runtime/RigVM/Source/RigVMEditor/Public/Editor/RigVMNewEditor.h
--- a/Engine/Plugins/Runtime/RigVM/Source/RigVMEditor/Public/Editor/RigVMNewEditor.h
+++ b/Engine/Plugins/Runtime/RigVM/Source/RigVMEditor/Public/Editor/RigVMNewEditor.h
@@ -46,7 +46,10 @@ protected:
 public:
 	virtual TSharedPtr<SRigVMDetailsInspector> GetRigVMInspector() const override { return Inspector.IsValid() ? Inspector.ToSharedRef().ToSharedPtr() : nullptr; }
 	virtual void SetInspector(TSharedPtr<SRigVMDetailsInspector> InWidget) { Inspector = InWidget; };
+	
+	TSharedRef<SWidget> GetCompilerResults() const { return CompilerResults.ToSharedRef(); }
 protected:
+	virtual TSharedPtr<IMessageLogListing> GetCompilerResultsListing() override { return CompilerResultsListing; }
 	virtual TSharedPtr<FApplicationMode> CreateEditorMode() override;
 	virtual const FName GetEditorAppName() const override;
 	virtual const TArray< UObject* >& GetEditingBlueprints() const override { return FWorkflowCentricApplication::GetEditingObjects(); }
@@ -82,7 +85,6 @@ protected:
 	virtual void RemoveToolbarExtender(TSharedPtr<FExtender> Extender) override { FWorkflowCentricApplication::RemoveToolbarExtender(Extender); };
 	virtual void AddMenuExtender(TSharedPtr<FExtender> Extender) override { FWorkflowCentricApplication::AddMenuExtender(Extender); }
 	virtual void RemoveMenuExtender(TSharedPtr<FExtender> Extender) override { FWorkflowCentricApplication::RemoveMenuExtender(Extender); }
-	virtual TSharedPtr<IMessageLogListing> GetCompilerResultsListing() override { return CompilerResultsListing; }
 	virtual void OnBlueprintChangedInnerImpl(UBlueprint* InBlueprint, bool bIsJustBeingCompiled) override;
 	virtual void RefreshEditorsImpl(ERefreshRigVMEditorReason::Type Reason) override;
 	virtual void SetupGraphEditorEventsImpl(UEdGraph* InGraph, SGraphEditor::FGraphEditorEvents& InEvents) override;
@@ -141,6 +143,12 @@ public:
 
 	float GetInstructionTextOpacity(UEdGraph* InGraph) const;
 	virtual void ClearSelectionStateFor(FName SelectionOwner);
+protected:
+	void OnLogTokenClicked(const TSharedRef<IMessageToken>& MessageToken);
+
+	/** Dumps messages to the compiler log, with an option to force it to display/come to front */
+	void DumpMessagesToCompilerLog(const TArray<TSharedRef<class FTokenizedMessage>>& Messages, bool bForceMessageDisplay);
+public:
 	void CreateDefaultTabContents(const TArray<UBlueprint*> InBlueprints);
 
 	TSharedRef<SWidget> GenerateCompileOptionsMenu();
@@ -199,7 +207,7 @@ protected:
 	TSharedPtr<FUICommandList> GraphEditorCommands;
 
 	/** Compiler results log, with the log listing that it reflects */
-	//TSharedPtr<class SWidget> CompilerResults;
+	TSharedPtr<class SWidget> CompilerResults;
 	TSharedPtr<class IMessageLogListing> CompilerResultsListing;
 
 	/** When set, flags which graph has a action menu currently open (if null, no graphs do). */
