Avoid infinite loop/recursion in URigVMPin::GetParentScriptStruct, which could happen when loading assets that have invalid but handleable data types in them.

diff --git a/Engine/Plugins/Runtime/RigVM/Source/RigVMDeveloper/Private/RigVMModel/RigVMPin.cpp b/Engine/Plugins/Runtime/RigVM/Source/RigVMDeveloper/Private/RigVMModel/RigVMPin.cpp
--- a/Engine/Plugins/Runtime/RigVM/Source/RigVMDeveloper/Private/RigVMModel/RigVMPin.cpp
+++ b/Engine/Plugins/Runtime/RigVM/Source/RigVMDeveloper/Private/RigVMModel/RigVMPin.cpp
@@ -1703,7 +1703,7 @@ UScriptStruct* URigVMPin::GetParentScriptStruct(const URigVMUnitNode* FallbackNo
 		}
 		else
 		{
-			return GetParentScriptStruct(FallbackNode);
+			return nullptr;
 		}
 	}
 
