AnimNext/RigVM: Add opt-in warnings around deprecated nodes and duplicate events

Both new warnings are disabled by default, but AnimNext enables the duplicate event warning as this can indicate unexpected behavior.

diff --git a/Engine/Plugins/Experimental/AnimNext/Source/AnimNextUncookedOnly/Private/AnimNextRigVMAssetEditorData.cpp b/Engine/Plugins/Experimental/AnimNext/Source/AnimNextUncookedOnly/Private/AnimNextRigVMAssetEditorData.cpp
--- a/Engine/Plugins/Experimental/AnimNext/Source/AnimNextUncookedOnly/Private/AnimNextRigVMAssetEditorData.cpp
+++ b/Engine/Plugins/Experimental/AnimNext/Source/AnimNextUncookedOnly/Private/AnimNextRigVMAssetEditorData.cpp
@@ -788,6 +788,7 @@ void UAnimNextRigVMAssetEditorData::RecompileVM()
 	VMCompileSettings.SetExecuteContextStruct(FAnimNextExecuteContext::StaticStruct());
 	FRigVMCompileSettings Settings = (bCompileInDebugMode) ? FRigVMCompileSettings::Fast(VMCompileSettings.GetExecuteContextStruct()) : VMCompileSettings;
 	Settings.SurpressInfoMessages = false;
+	Settings.bWarnAboutDuplicateEvents = true;
 	Settings.ASTSettings.ReportDelegate.BindUObject(this, &UAnimNextRigVMAssetEditorData::HandleReportFromCompiler);
 
 	Asset->VMRuntimeSettings = VMRuntimeSettings;
diff --git a/Engine/Plugins/Runtime/RigVM/Source/RigVMDeveloper/Private/RigVMCompiler/RigVMCompiler.cpp b/Engine/Plugins/Runtime/RigVM/Source/RigVMDeveloper/Private/RigVMCompiler/RigVMCompiler.cpp
--- a/Engine/Plugins/Runtime/RigVM/Source/RigVMDeveloper/Private/RigVMCompiler/RigVMCompiler.cpp
+++ b/Engine/Plugins/Runtime/RigVM/Source/RigVMDeveloper/Private/RigVMCompiler/RigVMCompiler.cpp
@@ -67,6 +67,18 @@ public:
 	}
 };
 
+FAutoConsoleVariable CVarRigVMWarnAboutDeprecatedNodes(
+	TEXT("RigVM.Compiler.WarnAboutDeprecatedNodes"),
+	false,
+	TEXT("Enable output of warnings when compiling nodes that are deprecated.")
+);
+
+FAutoConsoleVariable CVarRigVMWarnAboutDuplicateEvents(
+	TEXT("RigVM.Compiler.WarnAboutDuplicateEvents"),
+	false,
+	TEXT("Enable output of warnings when compiling event nodes that are duplicated.")
+);
+
 FRigVMCompileSettings::FRigVMCompileSettings()
 	: SurpressInfoMessages(true)
 	, SurpressWarnings(false)
@@ -76,6 +88,8 @@ FRigVMCompileSettings::FRigVMCompileSettings()
 	, ASTSettings(FRigVMParserASTSettings::Optimized())
 	, SetupNodeInstructionIndex(true)
 	, ASTErrorsAsNotifications(false)
+	, bWarnAboutDeprecatedNodes(false)
+	, bWarnAboutDuplicateEvents(false)
 {
 }
 
@@ -641,6 +655,7 @@ bool URigVMCompiler::Compile(const FRigVMCompileSettings& InSettings, TArray<URi
 
 	const FRigVMRegistry& Registry = FRigVMRegistry::Get();
 
+	TSet<FName> FoundEvents;
 	for(int32 GraphIndex=0; GraphIndex<VisitedGraphs.Num(); GraphIndex++)
 	{
 		URigVMGraph* VisitedGraph = VisitedGraphs[GraphIndex];
@@ -672,6 +687,17 @@ bool URigVMCompiler::Compile(const FRigVMCompileSettings& InSettings, TArray<URi
 				}
 			}
 
+			if((Settings.bWarnAboutDuplicateEvents || CVarRigVMWarnAboutDuplicateEvents->GetBool()) && ModelNode->IsEvent())
+			{
+				if (FoundEvents.Contains(ModelNode->GetEventName()))
+				{
+					static const FString LinkedMessage = TEXT("Duplicate event node @@ found.");
+					Settings.ASTSettings.Report(EMessageSeverity::Warning, ModelNode, LinkedMessage);
+				}
+
+				FoundEvents.Add(ModelNode->GetEventName());
+			}
+
 			if(ModelNode->IsA<UDEPRECATED_RigVMBranchNode>())
 			{
 				static const FString LinkedMessage = TEXT("Node @@ is a deprecated branch node. Cannot compile.");
@@ -893,6 +919,12 @@ bool URigVMCompiler::Compile(const FRigVMCompileSettings& InSettings, TArray<URi
 						}
 					}
 				}
+
+				if((Settings.bWarnAboutDeprecatedNodes || CVarRigVMWarnAboutDeprecatedNodes->GetBool()) && UnitNode->IsOutDated())
+				{
+					static const FString LinkedMessage = TEXT("Node @@ is outdated.");
+					Settings.ASTSettings.Report(EMessageSeverity::Warning, ModelNode, LinkedMessage);
+				}
 			}
 
 			auto ReportIncompatibleExecuteContextString = [&] (const FString InExecuteContextName)
diff --git a/Engine/Plugins/Runtime/RigVM/Source/RigVMDeveloper/Public/RigVMCompiler/RigVMCompiler.h b/Engine/Plugins/Runtime/RigVM/Source/RigVMDeveloper/Public/RigVMCompiler/RigVMCompiler.h
--- a/Engine/Plugins/Runtime/RigVM/Source/RigVMDeveloper/Public/RigVMCompiler/RigVMCompiler.h
+++ b/Engine/Plugins/Runtime/RigVM/Source/RigVMDeveloper/Public/RigVMCompiler/RigVMCompiler.h
@@ -57,10 +57,16 @@ public:
 
 	UPROPERTY()
 	bool SetupNodeInstructionIndex;
-	
+
 	UPROPERTY()
 	bool ASTErrorsAsNotifications;
 
+	UPROPERTY(EditAnywhere, AdvancedDisplay, BlueprintReadWrite, Category = FRigVMCompileSettings)
+	bool bWarnAboutDeprecatedNodes;
+
+	UPROPERTY(EditAnywhere, AdvancedDisplay, BlueprintReadWrite, Category = FRigVMCompileSettings)
+	bool bWarnAboutDuplicateEvents;
+
 	static FRigVMCompileSettings Fast(UScriptStruct* InExecuteContextStruct = nullptr)
 	{
 		FRigVMCompileSettings Settings(InExecuteContextStruct);
