RigVM: Add find references window in new editor
#rnx

diff --git a/Engine/Plugins/Animation/ControlRig/Source/ControlRigEditor/Private/Editor/ControlRigEditor.cpp b/Engine/Plugins/Animation/ControlRig/Source/ControlRigEditor/Private/Editor/ControlRigEditor.cpp
--- a/Engine/Plugins/Animation/ControlRig/Source/ControlRigEditor/Private/Editor/ControlRigEditor.cpp
+++ b/Engine/Plugins/Animation/ControlRig/Source/ControlRigEditor/Private/Editor/ControlRigEditor.cpp
@@ -175,7 +175,7 @@ UControlRigBlueprint* GetControlRigBlueprintFromContext(const UUnrealEdViewportT
 				const TSharedRef<FControlRigEditor>& Editor =
 					StaticCastSharedRef<FControlRigEditor>(AnimViewportClient->GetAssetEditorToolkit());
 
-				return Cast<UControlRigBlueprint>(Editor->FRigVMEditorBase::GetBlueprintObj());
+				return Cast<UControlRigBlueprint>(Editor->GetRigVMBlueprint());
 #endif
 			}
 		}
diff --git a/Engine/Plugins/Animation/ControlRig/Source/ControlRigEditor/Private/Editor/ControlRigEditorMode.cpp b/Engine/Plugins/Animation/ControlRig/Source/ControlRigEditor/Private/Editor/ControlRigEditorMode.cpp
--- a/Engine/Plugins/Animation/ControlRig/Source/ControlRigEditor/Private/Editor/ControlRigEditorMode.cpp
+++ b/Engine/Plugins/Animation/ControlRig/Source/ControlRigEditor/Private/Editor/ControlRigEditorMode.cpp
@@ -19,6 +19,7 @@
 #include "ToolMenus.h"
 #include "Editor/RigVMCompilerResultsTabSummoner.h"
 #include "Editor/RigVMDetailsInspectorTabSummoner.h"
+#include "Editor/RigVMFindReferencesTabSummoner.h"
 
 #if WITH_RIGVMLEGACYEDITOR
 FControlRigLegacyEditorMode::FControlRigLegacyEditorMode(const TSharedRef<FControlRigLegacyEditor>& InControlRigEditor, bool bCreateDefaultLayout)
@@ -336,6 +337,7 @@ FControlRigEditorMode::FControlRigEditorMode(const TSharedRef<FControlRigEditor>
 							FTabManager::NewStack()
 							->SetSizeCoefficient(0.2f)
 							->AddTab(FRigVMCompilerResultsTabSummoner::TabID(), ETabState::OpenedTab)
+							->AddTab(FRigVMFindReferencesTabSummoner::TabID(), ETabState::OpenedTab)
 						)
 					)
 					->Split
diff --git a/Engine/Plugins/Runtime/RigVM/Source/RigVMDeveloper/Private/EdGraph/RigVMEdGraph.cpp b/Engine/Plugins/Runtime/RigVM/Source/RigVMDeveloper/Private/EdGraph/RigVMEdGraph.cpp
--- a/Engine/Plugins/Runtime/RigVM/Source/RigVMDeveloper/Private/EdGraph/RigVMEdGraph.cpp
+++ b/Engine/Plugins/Runtime/RigVM/Source/RigVMDeveloper/Private/EdGraph/RigVMEdGraph.cpp
@@ -13,6 +13,7 @@
 #include "RigVMModel/Nodes/RigVMFunctionReturnNode.h"
 #include "RigVMModel/Nodes/RigVMRerouteNode.h"
 #include "Stats/StatsHierarchical.h"
+#include "Blueprint/BlueprintExtension.h"
 
 #include UE_INLINE_GENERATED_CPP_BY_NAME(RigVMEdGraph)
 
@@ -964,6 +965,30 @@ UEdGraphNode* URigVMEdGraph::FindNodeForModelNodeName(const FName& InModelNodeNa
 	return nullptr;
 }
 
+void URigVMEdGraph::AddLocalVariableSearchMetaDataInfo(const FName InVariableName, TArray<UBlueprintExtension::FSearchTagDataPair>& OutTaggedMetaData) const
+{
+	if (URigVMGraph* RigVMGraph = GetModel())
+	{
+		const TArray<FRigVMGraphVariableDescription> LocalVariables = RigVMGraph->GetLocalVariables(false);
+		for (const FRigVMGraphVariableDescription& LocalVariable : LocalVariables)
+		{
+			if (LocalVariable.Name.IsEqual(InVariableName))
+			{
+				OutTaggedMetaData.Emplace(FRigVMSearchTags::FiB_Name, FText::FromName(InVariableName));
+				
+				const FEdGraphPinType Type = LocalVariable.ToPinType();
+				OutTaggedMetaData.Emplace(FRigVMSearchTags::FiB_PinCategory, FText::FromName(Type.PinCategory));
+				OutTaggedMetaData.Emplace(FRigVMSearchTags::FiB_PinSubCategory, FText::FromName(Type.PinSubCategory));
+				if (Type.PinSubCategoryObject.IsValid())
+				{
+					OutTaggedMetaData.Emplace(FRigVMSearchTags::FiB_ObjectClass, FText::FromString(Type.PinSubCategoryObject->GetPathName()));
+				}
+				OutTaggedMetaData.Emplace(FRigVMSearchTags::FiB_IsArray, FText::Format(LOCTEXT("RigVMNodePinIsArray", "{0}"), Type.IsArray() ? 1 : 0));
+			}
+		}
+	}
+}
+
 URigVMBlueprint* URigVMEdGraph::GetBlueprint() const
 {
 	if (URigVMEdGraph* OuterGraph = Cast<URigVMEdGraph>(GetOuter()))
diff --git a/Engine/Plugins/Runtime/RigVM/Source/RigVMDeveloper/Private/EdGraph/RigVMEdGraphNode.cpp b/Engine/Plugins/Runtime/RigVM/Source/RigVMDeveloper/Private/EdGraph/RigVMEdGraphNode.cpp
--- a/Engine/Plugins/Runtime/RigVM/Source/RigVMDeveloper/Private/EdGraph/RigVMEdGraphNode.cpp
+++ b/Engine/Plugins/Runtime/RigVM/Source/RigVMDeveloper/Private/EdGraph/RigVMEdGraphNode.cpp
@@ -2174,11 +2174,17 @@ void URigVMEdGraphNode::AddPinSearchMetaDataInfo(const UEdGraphPin* Pin, TArray<
 	Super::AddPinSearchMetaDataInfo(Pin, OutTaggedMetaData);
 	if(const URigVMPin* ModelPin = FindModelPinFromGraphPin(Pin))
 	{
-		OutTaggedMetaData.Emplace(FText::FromString(TEXT("Data Type")), FText::FromString(ModelPin->GetCPPType()));
+		OutTaggedMetaData.Emplace(FRigVMSearchTags::FiB_PinCategory, FText::FromName(Pin->PinType.PinCategory));
+		OutTaggedMetaData.Emplace(FRigVMSearchTags::FiB_PinSubCategory, FText::FromName(Pin->PinType.PinSubCategory));
+		if (Pin->PinType.PinSubCategoryObject.IsValid())
+		{
+			OutTaggedMetaData.Emplace(FRigVMSearchTags::FiB_ObjectClass, FText::FromString(Pin->PinType.PinSubCategoryObject->GetPathName()));
+		}
+		OutTaggedMetaData.Emplace(FRigVMSearchTags::FiB_IsArray, FText::Format(LOCTEXT("RigVMNodePinIsArray", "{0}"), Pin->PinType.IsArray() ? 1 : 0));
 		
 		if (ModelPin->IsBoundToVariable())
 		{
-			OutTaggedMetaData.Add(FSearchTagDataPair(FText::FromString(TEXT("Binding")), FText::FromString(ModelPin->GetBoundVariablePath())));
+			OutTaggedMetaData.Add(FSearchTagDataPair(FRigVMSearchTags::FiB_PinBinding, FText::FromString(ModelPin->GetBoundVariablePath())));
 		}
 	}
 }
diff --git a/Engine/Plugins/Runtime/RigVM/Source/RigVMDeveloper/Private/EdGraph/RigVMEdGraphSchema.cpp b/Engine/Plugins/Runtime/RigVM/Source/RigVMDeveloper/Private/EdGraph/RigVMEdGraphSchema.cpp
--- a/Engine/Plugins/Runtime/RigVM/Source/RigVMDeveloper/Private/EdGraph/RigVMEdGraphSchema.cpp
+++ b/Engine/Plugins/Runtime/RigVM/Source/RigVMDeveloper/Private/EdGraph/RigVMEdGraphSchema.cpp
@@ -44,6 +44,18 @@
 
 #define LOCTEXT_NAMESPACE "CRigVMGraphSchema"
 
+const FText FRigVMSearchTags::FiB_Name = LOCTEXT("Name", "Name");
+const FText FRigVMSearchTags::FiB_ClassName = LOCTEXT("ClassName", "ClassName");
+const FText FRigVMSearchTags::FiB_NodeGuid = LOCTEXT("NodeGuid", "NodeGuid");
+const FText FRigVMSearchTags::FiB_PinCategory = LOCTEXT("PinCategory", "PinCategory");
+const FText FRigVMSearchTags::FiB_PinSubCategory = LOCTEXT("SubCategory", "SubCategory");
+const FText FRigVMSearchTags::FiB_PinBinding = LOCTEXT("Binding", "Binding");
+const FText FRigVMSearchTags::FiB_ObjectClass = LOCTEXT("ObjectClass", "ObjectClass");
+const FText FRigVMSearchTags::FiB_IsArray = LOCTEXT("IsArray", "IsArray");
+const FText FRigVMSearchTags::FiB_Glyph = LOCTEXT("Glyph", "Glyph");
+const FText FRigVMSearchTags::FiB_GlyphStyleSet = LOCTEXT("GlyphStyleSet", "GlyphStyleSet");
+const FText FRigVMSearchTags::FiB_GlyphColor = LOCTEXT("GlyphColor", "GlyphColor");
+
 FRigVMLocalVariableNameValidator::FRigVMLocalVariableNameValidator(const UBlueprint* Blueprint, const URigVMGraph* Graph, FName InExistingName)
 	: FStringSetNameValidator(InExistingName.ToString())
 {
diff --git a/Engine/Plugins/Runtime/RigVM/Source/RigVMDeveloper/Private/RigVMBlueprint.cpp b/Engine/Plugins/Runtime/RigVM/Source/RigVMDeveloper/Private/RigVMBlueprint.cpp
--- a/Engine/Plugins/Runtime/RigVM/Source/RigVMDeveloper/Private/RigVMBlueprint.cpp
+++ b/Engine/Plugins/Runtime/RigVM/Source/RigVMDeveloper/Private/RigVMBlueprint.cpp
@@ -2220,6 +2220,28 @@ TArray<FRigVMExternalDependency> URigVMBlueprint::GetExternalDependenciesForCate
 	return Dependencies;
 }
 
+#if WITH_EDITOR
+void URigVMBlueprint::AddVariableSearchMetaDataInfo(const FName InVariableName, TArray<UBlueprintExtension::FSearchTagDataPair>& OutTaggedMetaData) const
+{
+	for (const FBPVariableDescription& Variable : NewVariables)
+	{
+		if (Variable.VarName.IsEqual(InVariableName))
+		{
+			OutTaggedMetaData.Emplace(FRigVMSearchTags::FiB_Name, FText::FromName(InVariableName));
+				
+			const FEdGraphPinType& Type = Variable.VarType;
+			OutTaggedMetaData.Emplace(FRigVMSearchTags::FiB_PinCategory, FText::FromName(Type.PinCategory));
+			OutTaggedMetaData.Emplace(FRigVMSearchTags::FiB_PinSubCategory, FText::FromName(Type.PinSubCategory));
+			if (Type.PinSubCategoryObject.IsValid())
+			{
+				OutTaggedMetaData.Emplace(FRigVMSearchTags::FiB_ObjectClass, FText::FromString(Type.PinSubCategoryObject->GetPathName()));
+			}
+			OutTaggedMetaData.Emplace(FRigVMSearchTags::FiB_IsArray, FText::Format(LOCTEXT("RigVMNodePinIsArray", "{0}"), Type.IsArray() ? 1 : 0));
+		}
+	}
+}
+#endif
+
 URigVMGraph* URigVMBlueprint::GetTemplateModel(bool bIsFunctionLibrary)
 {
 #if WITH_EDITORONLY_DATA
diff --git a/Engine/Plugins/Runtime/RigVM/Source/RigVMDeveloper/Public/EdGraph/RigVMEdGraph.h b/Engine/Plugins/Runtime/RigVM/Source/RigVMDeveloper/Public/EdGraph/RigVMEdGraph.h
--- a/Engine/Plugins/Runtime/RigVM/Source/RigVMDeveloper/Public/EdGraph/RigVMEdGraph.h
+++ b/Engine/Plugins/Runtime/RigVM/Source/RigVMDeveloper/Public/EdGraph/RigVMEdGraph.h
@@ -178,6 +178,8 @@ public:
 
 	virtual const TArray<TSharedPtr<FRigVMStringWithTag>>* GetNameListForWidget(const FString& InWidgetName) const { return nullptr; }
 
+	void AddLocalVariableSearchMetaDataInfo(const FName InVariableName, TArray<UBlueprintExtension::FSearchTagDataPair>& OutTaggedMetaData) const;
+
 	UPROPERTY()
 	FString ModelNodePath;
 
diff --git a/Engine/Plugins/Runtime/RigVM/Source/RigVMDeveloper/Public/EdGraph/RigVMEdGraphSchema.h b/Engine/Plugins/Runtime/RigVM/Source/RigVMDeveloper/Public/EdGraph/RigVMEdGraphSchema.h
--- a/Engine/Plugins/Runtime/RigVM/Source/RigVMDeveloper/Public/EdGraph/RigVMEdGraphSchema.h
+++ b/Engine/Plugins/Runtime/RigVM/Source/RigVMDeveloper/Public/EdGraph/RigVMEdGraphSchema.h
@@ -16,6 +16,41 @@ class URigVMBlueprint;
 class URigVMEdGraph;
 class URigVMEdGraphNode;
 
+/**
+ *Const values for Find-in-Blueprints to tag searchable data
+ */
+struct RIGVMDEVELOPER_API FRigVMSearchTags
+{
+	// /** Begin const values for Find-in-Blueprint */
+	
+	/** Name tag */
+	static const FText FiB_Name;
+	/** Class Name tag */
+	static const FText FiB_ClassName;
+	/** NodeGuid tag */
+	static const FText FiB_NodeGuid;
+	
+	// /** Pin type tags */
+	/** Pin Category tag */
+	static const FText FiB_PinCategory;
+	/** Pin Sub-Category tag */
+	static const FText FiB_PinSubCategory;
+	/** Pin Binding tag */
+	static const FText FiB_PinBinding;
+	/** Pin object class tag */
+	static const FText FiB_ObjectClass;
+	/** Pin IsArray tag */
+	static const FText FiB_IsArray;
+	/** Glyph icon tag */
+	static const FText FiB_Glyph;
+	/** Style set the glyph belongs to */
+	static const FText FiB_GlyphStyleSet;
+	/** Glyph icon color tag */
+	static const FText FiB_GlyphColor;
+	
+	// /** End const values for Find-in-Blueprint */
+};
+
 /** Extra operations that can be performed on pin connection */
 enum class ECanCreateConnectionResponse_Extended
 {
diff --git a/Engine/Plugins/Runtime/RigVM/Source/RigVMDeveloper/Public/RigVMBlueprint.h b/Engine/Plugins/Runtime/RigVM/Source/RigVMDeveloper/Public/RigVMBlueprint.h
--- a/Engine/Plugins/Runtime/RigVM/Source/RigVMDeveloper/Public/RigVMBlueprint.h
+++ b/Engine/Plugins/Runtime/RigVM/Source/RigVMDeveloper/Public/RigVMBlueprint.h
@@ -428,6 +428,8 @@ public:
 	FOnRigVMRequestInspectMemoryStorage& OnRequestInspectMemoryStorage() { return OnRequestInspectMemoryStorageEvent; }
 	void RequestInspectMemoryStorage(const TArray<FRigVMMemoryStorageStruct*>& InMemoryStorageStructs) { OnRequestInspectMemoryStorageEvent.Broadcast(InMemoryStorageStructs); }
 
+	void AddVariableSearchMetaDataInfo(const FName InVariableName, TArray<UBlueprintExtension::FSearchTagDataPair>& OutTaggedMetaData) const;
+
 #endif	// #if WITH_EDITOR
 
 	
diff --git a/Engine/Plugins/Runtime/RigVM/Source/RigVMEditor/Private/Editor/Kismet/RigVMFindInBlueprints.cpp b/Engine/Plugins/Runtime/RigVM/Source/RigVMEditor/Private/Editor/Kismet/RigVMFindInBlueprints.cpp
--- a/Engine/Plugins/Runtime/RigVM/Source/RigVMEditor/Private/Editor/Kismet/RigVMFindInBlueprints.cpp
+++ b/Engine/Plugins/Runtime/RigVM/Source/RigVMEditor/Private/Editor/Kismet/RigVMFindInBlueprints.cpp
@@ -92,6 +92,7 @@
 #include "Widgets/Text/SMultiLineEditableText.h"
 #include "Widgets/Text/STextBlock.h"
 #include "Widgets/Views/STableRow.h"
+#include "RigVMBlueprint.h"
 
 class ITableRow;
 class SWidget;
@@ -1235,7 +1236,7 @@ void SRigVMFindInBlueprints::MakeSearchQuery(FString InSearchString, bool bInIsF
 
 			UBlueprint* Blueprint =
 #if WITH_EDITOR
-				EditorPtr.Pin()->GetBlueprintObj();
+				EditorPtr.Pin()->GetRigVMBlueprint();
 #else
 				nullptr;
 #endif
diff --git a/Engine/Plugins/Runtime/RigVM/Source/RigVMEditor/Private/Editor/RigVMEditor.cpp b/Engine/Plugins/Runtime/RigVM/Source/RigVMEditor/Private/Editor/RigVMEditor.cpp
--- a/Engine/Plugins/Runtime/RigVM/Source/RigVMEditor/Private/Editor/RigVMEditor.cpp
+++ b/Engine/Plugins/Runtime/RigVM/Source/RigVMEditor/Private/Editor/RigVMEditor.cpp
@@ -38,6 +38,7 @@
 #include "Widgets/Docking/SDockTab.h"
 #include "Widgets/Layout/SScrollBox.h"
 #include "IStructureDetailsView.h"
+#include "Editor/RigVMDetailsInspectorTabSummoner.h"
 #include "Editor/SRigVMDetailsInspector.h"
 
 #if WITH_RIGVMLEGACYEDITOR
@@ -2812,9 +2813,9 @@ void FRigVMEditorBase::SetDetailObjects(const TArray<UObject*>& InObjects, bool
 
 			GetTabManager()->InsertNewDocumentTab(
 #if WITH_RIGVMLEGACYEDITOR
-				FBlueprintEditorTabs::DetailsID,
+				GetRigVMInspector() ? FRigVMDetailsInspectorTabSummoner::TabID() : FBlueprintEditorTabs::DetailsID,
 #else
-				FRigVMNewEditorTabs::DetailsID(),
+				FRigVMDetailsInspectorTabSummoner::TabID(),
 #endif
 				TabId,
 				FTabManager::FLastMajorOrNomadTab(TEXT("RigVMMemoryDetails")),
@@ -2987,9 +2988,9 @@ void FRigVMEditorBase::SetMemoryStorageDetails(const TArray<FRigVMMemoryStorageS
 
 			GetTabManager()->InsertNewDocumentTab(
 #if WITH_RIGVMLEGACYEDITOR
-				FBlueprintEditorTabs::DetailsID,
+				GetRigVMInspector() ? FRigVMDetailsInspectorTabSummoner::TabID() : FBlueprintEditorTabs::DetailsID,
 #else
-				FRigVMNewEditorTabs::DetailsID(),
+				FRigVMDetailsInspectorTabSummoner::TabID(),
 #endif
 				TabId,
 				FTabManager::FLastMajorOrNomadTab(TEXT("RigVMMemoryDetails")),
diff --git a/Engine/Plugins/Runtime/RigVM/Source/RigVMEditor/Private/Editor/RigVMFindReferences.cpp b/Engine/Plugins/Runtime/RigVM/Source/RigVMEditor/Private/Editor/RigVMFindReferences.cpp
new file mode 100644
--- /dev/null
+++ b/Engine/Plugins/Runtime/RigVM/Source/RigVMEditor/Private/Editor/RigVMFindReferences.cpp
@@ -0,0 +1,946 @@
+// Copyright Epic Games, Inc. All Rights Reserved.
+
+#include "Editor/RigVMFindReferences.h"
+
+#include "K2Node_Variable.h"
+#include "RigVMBlueprint.h"
+#include "EdGraph/RigVMEdGraphSchema.h"
+#include "Editor/RigVMNewEditor.h"
+#include "Framework/Commands/GenericCommands.h"
+#include "Kismet2/BlueprintEditorUtils.h"
+#include "Toolkits/ToolkitManager.h"
+#include "Widgets/Input/SSearchBox.h"
+#include "HAL/PlatformApplicationMisc.h"
+#include "Layout/WidgetPath.h"
+#include "Framework/Application/SlateApplication.h"
+
+#define LOCTEXT_NAMESPACE "RigVMFindResults"
+
+
+FString RigVMFindReferencesHelpers::GetPinTypeAsString(const FEdGraphPinType& InPinType)
+{
+	FString Result = InPinType.PinCategory.ToString();
+	if(UObject* SubCategoryObject = InPinType.PinSubCategoryObject.Get())
+	{
+		Result += FString(" '") + SubCategoryObject->GetName() + "'";
+	}
+	else
+	{
+		Result += FString(" '") + InPinType.PinSubCategory.ToString() + "'";
+	}
+
+	return Result;
+}
+
+bool RigVMFindReferencesHelpers::ParsePinType(FText InKey, FText InValue, FEdGraphPinType& InOutPinType)
+{
+	bool bParsed = true;
+
+	if(InKey.CompareTo(FRigVMSearchTags::FiB_PinCategory) == 0)
+	{
+		InOutPinType.PinCategory = *InValue.ToString();
+	}
+	else if(InKey.CompareTo(FRigVMSearchTags::FiB_PinSubCategory) == 0)
+	{
+		InOutPinType.PinSubCategory = *InValue.ToString();
+	}
+	else if(InKey.CompareTo(FRigVMSearchTags::FiB_ObjectClass) == 0)
+	{
+		InOutPinType.PinSubCategory = *InValue.ToString();
+	}
+	else if(InKey.CompareTo(FRigVMSearchTags::FiB_IsArray) == 0)
+	{
+		InOutPinType.ContainerType = (InValue.ToString().ToBool() ? EPinContainerType::Array : EPinContainerType::None);
+	}
+	else
+	{
+		bParsed = false;
+	}
+
+	return bParsed;
+}
+
+void RigVMFindReferencesHelpers::ExpandAllChildren(FRigVMSearchResult InTreeNode, TSharedPtr<STreeView<FRigVMSearchResult>> InTreeView)
+{
+	if (InTreeNode->Children.Num())
+	{
+		InTreeView->SetItemExpansion(InTreeNode, true);
+		for (int32 i = 0; i < InTreeNode->Children.Num(); i++)
+		{
+			ExpandAllChildren(InTreeNode->Children[i], InTreeView);
+		}
+	}
+}
+
+//////////////////////////////////////////////////////////////////////////
+// FBlueprintSearchResult
+
+FRigVMFindResult::FRigVMFindResult(TWeakObjectPtr<URigVMBlueprint> InBlueprint)
+	: WeakBlueprint(InBlueprint)
+{
+}
+
+FRigVMFindResult::FRigVMFindResult(TWeakObjectPtr<URigVMBlueprint> InBlueprint, const FText& InDisplayText )
+	: WeakBlueprint(InBlueprint), DisplayText(InDisplayText)
+{
+}
+
+FReply FRigVMFindResult::OnClick()
+{
+	// If there is a parent, handle it using the parent's functionality
+	if(Parent.IsValid())
+	{
+		return Parent.Pin()->OnClick();
+	}
+	
+	return FReply::Handled();
+}
+
+UObject* FRigVMFindResult::GetObject(URigVMBlueprint* InBlueprint) const
+{
+	return GetBlueprint();
+}
+
+FText FRigVMFindResult::GetCategory() const
+{
+	return FText::GetEmpty();
+}
+
+TSharedRef<SWidget> FRigVMFindResult::CreateIcon() const
+{
+	const FSlateBrush* Brush = NULL;
+
+	return 	SNew(SImage)
+			.Image(Brush)
+			.ColorAndOpacity(FStyleColors::Foreground)
+			.ToolTipText( GetCategory() );
+}
+
+URigVMBlueprint* FRigVMFindResult::GetBlueprint() const
+{
+	if (TStrongObjectPtr<URigVMBlueprint> Blueprint = WeakBlueprint.Pin())
+	{
+		return Blueprint.Get();
+	}
+	return nullptr;
+}
+
+FText FRigVMFindResult::GetDisplayString() const
+{
+	return DisplayText;
+}
+
+//////////////////////////////////////////////////////////
+// FRigVMFindReferencesGraphNode
+
+FRigVMFindReferencesGraphNode::FRigVMFindReferencesGraphNode(TWeakObjectPtr<URigVMBlueprint> InBlueprint)
+	: FRigVMFindResult(InBlueprint), Glyph(FAppStyle::GetAppStyleSetName(), "")
+	, Class(nullptr)
+{
+}
+
+FReply FRigVMFindReferencesGraphNode::OnClick()
+{
+	UBlueprint* Blueprint = GetBlueprint();
+	if(Blueprint)
+	{
+		UEdGraphNode* OutNode = NULL;
+		if(	UEdGraphNode* GraphNode = FBlueprintEditorUtils::GetNodeByGUID(Blueprint, NodeGuid) )
+		{
+			TSharedPtr< IToolkit > FoundAssetEditor = FToolkitManager::Get().FindEditorForAsset(Blueprint);
+			if (!FoundAssetEditor)
+			{
+				GEditor->GetEditorSubsystem<UAssetEditorSubsystem>()->OpenEditorForAsset(Blueprint);
+				FoundAssetEditor = FToolkitManager::Get().FindEditorForAsset(Blueprint);
+			}
+
+			// If we found an Editor
+			if (FoundAssetEditor.IsValid())
+			{
+				TSharedPtr<FRigVMNewEditor> Editor = StaticCastSharedPtr<FRigVMNewEditor>(FoundAssetEditor);
+				Editor->FocusWindow();
+				Editor->JumpToHyperlink(GraphNode, false);
+			}
+			return FReply::Handled();
+		}
+	}
+
+	return FRigVMFindResult::OnClick();
+}
+
+TSharedRef<SWidget> FRigVMFindReferencesGraphNode::CreateIcon() const
+{
+	return 	SNew(SImage)
+		.Image(Glyph.GetOptionalIcon())
+		.ColorAndOpacity(GlyphColor)
+		.ToolTipText( GetCategory() );
+}
+
+void FRigVMFindReferencesGraphNode::ParseSearchInfo(FText InKey, FText InValue)
+{
+	if(InKey.CompareTo(FRigVMSearchTags::FiB_NodeGuid) == 0)
+	{
+		FString NodeGUIDAsString = InValue.ToString();
+		FGuid::Parse(NodeGUIDAsString, NodeGuid);
+	}
+
+	if(InKey.CompareTo(FRigVMSearchTags::FiB_ClassName) == 0)
+	{
+		ClassName = InValue.ToString();
+	}
+	else if(InKey.CompareTo(FRigVMSearchTags::FiB_Name) == 0)
+	{
+		DisplayText = InValue;
+	}
+	else if(InKey.CompareTo(FRigVMSearchTags::FiB_Glyph) == 0)
+	{
+		Glyph = FSlateIcon(Glyph.GetStyleSetName(), *InValue.ToString());
+	}
+	else if(InKey.CompareTo(FRigVMSearchTags::FiB_GlyphStyleSet) == 0)
+	{
+		Glyph = FSlateIcon(*InValue.ToString(), Glyph.GetStyleName());
+	}
+	else if(InKey.CompareTo(FRigVMSearchTags::FiB_GlyphColor) == 0)
+	{
+		GlyphColor.InitFromString(InValue.ToString());
+	}
+}
+
+void FRigVMFindReferencesGraphNode::FinalizeSearchData()
+{
+	if(!ClassName.IsEmpty())
+	{
+		// Check the node subclasses and look for one with the same short name
+		TArray<UClass*> NodeClasses;
+		GetDerivedClasses(UEdGraphNode::StaticClass(), NodeClasses, /*bRecursive=*/true);
+
+		for (UClass* FoundClass : NodeClasses)
+		{
+			if (FoundClass->GetName() == ClassName)
+			{
+				Class = FoundClass;
+				break;
+			}
+		}
+
+		ClassName.Empty();
+	}
+}
+
+UObject* FRigVMFindReferencesGraphNode::GetObject(URigVMBlueprint* InBlueprint) const
+{
+	return FBlueprintEditorUtils::GetNodeByGUID(InBlueprint, NodeGuid);
+}
+
+FText FRigVMFindReferencesGraphNode::GetCategory() const
+{
+	return LOCTEXT("NodeCategory", "Node");
+}
+
+//////////////////////////////////////////////////////////
+// FRigVMFindReferencesPin
+
+FRigVMFindReferencesPin::FRigVMFindReferencesPin(TWeakObjectPtr<URigVMBlueprint> InBlueprint, FString InSchemaName)
+	: FRigVMFindResult(InBlueprint)
+	, SchemaName(InSchemaName)
+	, IconColor(FSlateColor::UseForeground())
+{
+}
+
+TSharedRef<SWidget> FRigVMFindReferencesPin::CreateIcon() const
+{
+	const FSlateBrush* Brush = nullptr;
+
+	if( PinType.IsArray() )
+	{
+		Brush = FAppStyle::Get().GetBrush( TEXT("GraphEditor.ArrayPinIcon") );
+	}
+	else
+	{
+		Brush = FAppStyle::Get().GetBrush( TEXT("GraphEditor.PinIcon") );
+	}
+
+	return SNew(SImage)
+		.Image(Brush)
+		.ColorAndOpacity(IconColor)
+		.ToolTipText(FText::FromString(RigVMFindReferencesHelpers::GetPinTypeAsString(PinType)));
+}
+
+void FRigVMFindReferencesPin::ParseSearchInfo(FText InKey, FText InValue)
+{
+	if(InKey.CompareTo(FRigVMSearchTags::FiB_Name) == 0)
+	{
+		DisplayText = InValue;
+	}
+	else
+	{
+		RigVMFindReferencesHelpers::ParsePinType(InKey, InValue, PinType);
+	}
+}
+
+FText FRigVMFindReferencesPin::GetCategory() const
+{
+	return LOCTEXT("PinCategory", "Pin");
+}
+
+void FRigVMFindReferencesPin::FinalizeSearchData()
+{
+	if(!PinType.PinSubCategory.IsNone())
+	{
+		// This can either be a full path to an object, or a short name specific to the category
+		if (FPackageName::IsShortPackageName(PinType.PinSubCategory))
+		{
+			// This could also be an old class name without the full path, but it's fine to ignore in that case
+		}
+		else
+		{
+			PinType.PinSubCategoryObject = FindObject<UObject>(UObject::StaticClass(), *PinType.PinSubCategory.ToString());
+			if (PinType.PinSubCategoryObject.IsValid())
+			{
+				PinType.PinSubCategory = NAME_None;
+			}
+		}
+	}
+
+	if(!SchemaName.IsEmpty())
+	{
+		// Get all subclasses of schema and find the one with a matching short name
+		TArray<UClass*> SchemaClasses;
+		GetDerivedClasses(UEdGraphSchema::StaticClass(), SchemaClasses, /*bRecursive=*/true);
+
+		for (UClass* FoundClass : SchemaClasses)
+		{
+			if (FoundClass->GetName() == SchemaName)
+			{
+				UEdGraphSchema* Schema = FoundClass->GetDefaultObject<UEdGraphSchema>();
+				IconColor = Schema->GetPinTypeColor(PinType);
+				break;
+			}
+		}
+
+		SchemaName.Empty();
+	}
+}
+
+//////////////////////////////////////////////////////////
+// FRigVMFindReferencesVariable
+
+FRigVMFindReferencesVariable::FRigVMFindReferencesVariable(TWeakObjectPtr<URigVMBlueprint> InBlueprint)
+	: FRigVMFindResult(InBlueprint)
+{
+}
+
+TSharedRef<SWidget> FRigVMFindReferencesVariable::CreateIcon() const
+{
+	FLinearColor IconColor = FStyleColors::Foreground.GetSpecifiedColor();
+	const FSlateBrush* Brush = UK2Node_Variable::GetVarIconFromPinType(PinType, IconColor).GetOptionalIcon();
+	IconColor = UEdGraphSchema_K2::StaticClass()->GetDefaultObject<UEdGraphSchema_K2>()->GetPinTypeColor(PinType);
+
+	return 	SNew(SImage)
+		.Image(Brush)
+		.ColorAndOpacity(IconColor)
+		.ToolTipText( FText::FromString(RigVMFindReferencesHelpers::GetPinTypeAsString(PinType)) );
+}
+
+void FRigVMFindReferencesVariable::ParseSearchInfo(FText InKey, FText InValue)
+{
+	if(InKey.CompareTo(FRigVMSearchTags::FiB_Name) == 0)
+	{
+		DisplayText = InValue;
+	}
+	else
+	{
+		RigVMFindReferencesHelpers::ParsePinType(InKey, InValue, PinType);
+	}
+}
+
+FText FRigVMFindReferencesVariable::GetCategory() const
+{
+	return LOCTEXT("Variable", "Variable");
+}
+
+void FRigVMFindReferencesVariable::FinalizeSearchData()
+{
+	if (!PinType.PinSubCategory.IsNone())
+	{
+		// This can either be a full path to an object, or a short name specific to the category
+		if (FPackageName::IsShortPackageName(PinType.PinSubCategory))
+		{
+			// This could also be an old class name without the full path, but it's fine to ignore in that case
+		}
+		else
+		{
+			PinType.PinSubCategoryObject = FindObject<UObject>(UObject::StaticClass(), *PinType.PinSubCategory.ToString());
+			if (PinType.PinSubCategoryObject.IsValid())
+			{
+				PinType.PinSubCategory = NAME_None;
+			}
+		}
+	}
+}
+
+//////////////////////////////////////////////////////////
+// FRigVMFindReferencesGraph
+
+FRigVMFindReferencesGraph::FRigVMFindReferencesGraph(TWeakObjectPtr<URigVMBlueprint> InBlueprint, EGraphType InGraphType)
+	: FRigVMFindResult(InBlueprint), GraphType(InGraphType)
+{
+}
+
+FReply FRigVMFindReferencesGraph::OnClick()
+{
+	UBlueprint* Blueprint = GetBlueprint();
+	if(Blueprint)
+	{
+		GEditor->GetEditorSubsystem<UAssetEditorSubsystem>()->OpenEditorForAsset(Blueprint);
+		TSharedPtr< IToolkit > FoundAssetEditor = FToolkitManager::Get().FindEditorForAsset(Blueprint);
+
+		// If we found an Editor
+		if (FoundAssetEditor.IsValid())
+		{
+			TSharedPtr<FRigVMNewEditor> Editor = StaticCastSharedPtr<FRigVMNewEditor>(FoundAssetEditor);
+
+			TArray<UEdGraph*> BlueprintGraphs;
+			Blueprint->GetAllGraphs(BlueprintGraphs);
+
+			for(UEdGraph* Graph : BlueprintGraphs)
+			{
+				FGraphDisplayInfo DisplayInfo;
+				Graph->GetSchema()->GetGraphDisplayInformation(*Graph, DisplayInfo);
+			
+				if(DisplayInfo.PlainName.EqualTo(DisplayText))
+				{
+					Editor->FocusWindow();
+					Editor->JumpToHyperlink(Graph, false);
+					break;
+				}
+			}
+		}
+	}
+	else
+	{
+		return FRigVMFindResult::OnClick();
+	}
+	return FReply::Handled();
+}
+
+TSharedRef<SWidget> FRigVMFindReferencesGraph::CreateIcon() const
+{
+	const FSlateBrush* Brush = NULL;
+	switch (GraphType)
+	{
+		case GT_Function:
+			{
+				Brush = FAppStyle::GetBrush(TEXT("GraphEditor.Function_16x"));
+				break;
+			}
+		case GT_Ubergraph:
+			{
+				Brush = FAppStyle::GetBrush( TEXT("GraphEditor.EventGraph_16x") );
+				break;
+			}
+		default:
+			{
+				break;
+			}
+	}
+	
+	return 	SNew(SImage)
+		.Image(Brush)
+		.ToolTipText( GetCategory() );
+}
+
+void FRigVMFindReferencesGraph::ParseSearchInfo(FText InKey, FText InValue)
+{
+	if(InKey.CompareTo(FRigVMSearchTags::FiB_Name) == 0)
+	{
+		DisplayText = InValue;
+	}
+}
+
+FText FRigVMFindReferencesGraph::GetCategory() const
+{
+	if(GraphType == GT_Function)
+	{
+		return LOCTEXT("FunctionGraphCategory", "Function");
+	}
+	return LOCTEXT("GraphCategory", "Graph");
+}
+
+//////////////////////////////////////////////////////////////////////////
+// SBlueprintSearch
+
+void SRigVMFindReferences::Construct( const FArguments& InArgs, TSharedPtr<FRigVMEditorBase> InEditor)
+{
+	EditorPtr = InEditor;
+	
+	RegisterCommands();
+	
+	this->ChildSlot
+	[
+		SNew(SBorder)
+		.BorderImage(FAppStyle::Get().GetBrush("Brushes.Panel"))
+		[
+			SNew(SVerticalBox)
+			+SVerticalBox::Slot()
+			.AutoHeight()
+			.Padding(8.f, 5.f, 8.f, 5.f)
+			[
+				SNew(SHorizontalBox)
+				+SHorizontalBox::Slot()
+				.FillWidth(1)
+				[
+					SAssignNew(SearchTextField, SSearchBox)
+					.HintText(LOCTEXT("BlueprintSearchHint", "Enter function or event name to find references..."))
+					.OnTextChanged(this, &SRigVMFindReferences::OnSearchTextChanged)
+					.OnTextCommitted(this, &SRigVMFindReferences::OnSearchTextCommitted)
+					.DelayChangeNotificationsWhileTyping(false)
+				]
+			]
+			+SVerticalBox::Slot()
+			.FillHeight(1.0f)
+			[
+				SNew(SBorder)
+				.BorderImage(FAppStyle::Get().GetBrush("Brushes.Recessed"))
+				.Padding(FMargin(8.f, 8.f, 4.f, 0.f))
+				[
+					SAssignNew(TreeView, SRigVMTreeViewType)
+						.TreeItemsSource( &ItemsFound )
+						.OnGenerateRow( this, &SRigVMFindReferences::OnGenerateRow )
+						.OnGetChildren( this, &SRigVMFindReferences::OnGetChildren )
+						.OnMouseButtonDoubleClick(this,&SRigVMFindReferences::OnTreeSelectionDoubleClicked)
+						.SelectionMode( ESelectionMode::Multi )
+						.OnContextMenuOpening(this, &SRigVMFindReferences::OnContextMenuOpening)
+				]
+			]
+		]
+	];
+}
+
+void SRigVMFindReferences::RegisterCommands()
+{
+	CommandList = MakeShareable(new FUICommandList());
+
+	CommandList->MapAction( FGenericCommands::Get().Copy,
+		FExecuteAction::CreateSP(this, &SRigVMFindReferences::OnCopyAction));
+
+	CommandList->MapAction( FGenericCommands::Get().SelectAll,
+		FExecuteAction::CreateSP(this, &SRigVMFindReferences::OnSelectAllAction) );
+}
+
+void SRigVMFindReferences::FocusForUse(bool bSetFindWithinBlueprint, FString NewSearchTerms, bool bSelectFirstResult)
+{
+	// NOTE: Careful, GeneratePathToWidget can be reentrant in that it can call visibility delegates and such
+	FWidgetPath FilterTextBoxWidgetPath;
+	FSlateApplication::Get().GeneratePathToWidgetUnchecked( SearchTextField.ToSharedRef(), FilterTextBoxWidgetPath );
+
+	// Set keyboard focus directly
+	FSlateApplication::Get().SetKeyboardFocus( FilterTextBoxWidgetPath, EFocusCause::SetDirectly );
+
+	if (!NewSearchTerms.IsEmpty())
+	{
+		SearchTextField->SetText(FText::FromString(NewSearchTerms));
+		FindReferences(NewSearchTerms);
+
+		// Select the first result
+		if (bSelectFirstResult && ItemsFound.Num())
+		{
+			TSharedPtr<FRigVMFindResult> ItemToFocusOn = ItemsFound[0];
+
+			// We want the first childmost item to select, as that is the item that is most-likely to be what was searched for (parents being graphs).
+			// Will fail back upward as neccessary to focus on a focusable item
+			while(ItemToFocusOn->Children.Num())
+			{
+				ItemToFocusOn = ItemToFocusOn->Children[0];
+			}
+			TreeView->SetSelection(ItemToFocusOn);
+			ItemToFocusOn->OnClick();
+		}
+	}
+}
+
+void SRigVMFindReferences::FindReferences(const FString& SearchTerms)
+{
+	ItemsFound.Empty();
+	ElementHashToResult.Empty();
+
+	HighlightText = FText::FromString(SearchTerms);
+
+	TSharedPtr<FRigVMEditorBase> Editor = EditorPtr.Pin();
+	if (!Editor.IsValid())
+	{
+		return;
+	}
+
+	URigVMBlueprint* Blueprint = Editor->GetRigVMBlueprint();
+	if (!Blueprint)
+	{
+		return;
+	}
+	
+	const URigVMEdGraphSchema* Schema = GetDefault<URigVMEdGraphSchema>();
+
+	auto FindOrAddGraphResult = [&](const UEdGraph* EdGraph)
+	{
+		FRigVMSearchResult GraphResult;
+		if (FRigVMSearchResult* Found = ElementHashToResult.Find(GetTypeHash(EdGraph)))
+		{
+			GraphResult = *Found;
+		}
+		else
+		{
+			GraphResult = ItemsFound.Add_GetRef(FRigVMSearchResult(new FRigVMFindReferencesGraph(Blueprint, Schema->GetGraphType(EdGraph))));
+			ElementHashToResult.Add(GetTypeHash(EdGraph), GraphResult);
+		}
+
+		FGraphDisplayInfo DisplayInfo;
+		EdGraph->GetSchema()->GetGraphDisplayInformation(*EdGraph, DisplayInfo);
+
+		GraphResult->ParseSearchInfo(FRigVMSearchTags::FiB_Name, DisplayInfo.PlainName);
+		return GraphResult;
+	};
+
+	auto FindOrAddNodeResult = [&](const UEdGraphNode* EdNode, const TArray<struct FSearchTagDataPair>& NodeMetaData)
+	{
+		FRigVMSearchResult NodeResult;
+		if (FRigVMSearchResult* Found = ElementHashToResult.Find(GetTypeHash(EdNode)))
+		{
+			NodeResult = *Found;
+		}
+		else
+		{
+			FRigVMSearchResult GraphResult = FindOrAddGraphResult(EdNode->GetGraph());
+			NodeResult = FRigVMSearchResult(new FRigVMFindReferencesGraphNode(Blueprint));
+			ElementHashToResult.Add(GetTypeHash(EdNode), NodeResult);
+			GraphResult->Children.Add(NodeResult);
+			NodeResult->Parent = GraphResult;
+
+			for (const FSearchTagDataPair& MetadataTag : NodeMetaData)
+			{
+				NodeResult->ParseSearchInfo(MetadataTag.Key, MetadataTag.Value);
+			}
+		}
+		return NodeResult;
+	};
+
+	auto FindOrAddPinResult = [&](const UEdGraphPin* EdPin, const TArray<struct FSearchTagDataPair>& PinMetaData, const TArray<struct FSearchTagDataPair>& NodeMetaData)
+	{
+		FRigVMSearchResult PinResult;
+		if (FRigVMSearchResult* Found = ElementHashToResult.Find(GetTypeHash(EdPin)))
+		{
+			PinResult = *Found;
+		}
+		else
+		{
+			UEdGraphNode* EdNode = EdPin->GetOwningNode();
+			
+			const FString SchemaClassName = Blueprint->GetRigVMEdGraphSchemaClass()->GetName();
+			FRigVMSearchResult NodeResult = FindOrAddNodeResult(EdPin->GetOwningNode(), NodeMetaData);
+			PinResult = FRigVMSearchResult(new FRigVMFindReferencesPin(Blueprint, SchemaClassName));
+			ElementHashToResult.Add(GetTypeHash(EdPin), PinResult);
+			NodeResult->Children.Add(PinResult);
+			PinResult->Parent = NodeResult;
+
+			for (const FSearchTagDataPair& MetadataTag : PinMetaData)
+			{
+				PinResult->ParseSearchInfo(MetadataTag.Key, MetadataTag.Value);
+			}
+		}
+		return PinResult;
+	};
+
+	auto FindOrAddVariableResult = [&](const FName& VariableName, const UEdGraph* Graph, const TArray<UBlueprintExtension::FSearchTagDataPair>& VariableMetaData)
+	{
+		FRigVMSearchResult VariableResult;
+
+		const uint32 Hash = HashCombine(GetTypeHash(VariableName), GetTypeHash(Graph));
+		if (FRigVMSearchResult* Found = ElementHashToResult.Find(Hash))
+		{
+			VariableResult = *Found;
+		}
+		else
+		{
+			FRigVMSearchResult GraphResult;
+			if (Graph)
+			{
+				GraphResult = FindOrAddGraphResult(Graph);
+			}
+
+			VariableResult = FRigVMSearchResult(new FRigVMFindReferencesVariable(Blueprint));
+			ElementHashToResult.Add(Hash, VariableResult);
+			if (GraphResult)
+			{
+				GraphResult->Children.Add(VariableResult);
+				VariableResult->Parent = GraphResult;
+			}
+			else
+			{
+				ItemsFound.Add(VariableResult);
+			}
+
+			for (const UBlueprintExtension::FSearchTagDataPair& MetadataTag : VariableMetaData)
+			{
+				VariableResult->ParseSearchInfo(MetadataTag.Key, MetadataTag.Value);
+			}
+		}
+		return VariableResult;
+	};
+
+	TArray<URigVMGraph*> Graphs = Blueprint->GetRigVMClient()->GetAllModels(true, true);
+	for (const URigVMGraph* Graph : Graphs)
+	{
+		if (URigVMEdGraph* EdGraph = Cast<URigVMEdGraph>(Blueprint->GetEdGraph(Graph)))
+		{
+			for (TObjectPtr<UEdGraphNode> EdNode : EdGraph->Nodes)
+			{
+				TArray<struct FSearchTagDataPair> NodeMetadata;
+				EdNode->AddSearchMetaDataInfo(NodeMetadata);
+
+				const FSearchTagDataPair* MatchedPair = nullptr;
+				for (const FSearchTagDataPair& MetadataTag : NodeMetadata)
+				{
+					if (MetadataTag.Value.ToString().Contains(SearchTerms))
+					{
+						MatchedPair = &MetadataTag;
+						break;
+					}
+				}
+				if (MatchedPair)
+				{
+					FRigVMSearchResult Result = FindOrAddNodeResult(EdNode, NodeMetadata);
+
+					const FText DisplayText = FText::Format(LOCTEXT("RigVMFindReferencesValues", "{0}: {1}"), MatchedPair->Key, MatchedPair->Value);
+					FRigVMSearchResult DetailResult = FRigVMSearchResult(new FRigVMFindResult(Blueprint, DisplayText));
+
+					Result->Children.Add(DetailResult);
+					DetailResult->Parent = Result;
+				}
+
+				for (UEdGraphPin* Pin : EdNode->GetAllPins())
+				{
+					TArray<struct FSearchTagDataPair> PinMetaData;
+					EdNode->AddPinSearchMetaDataInfo(Pin, PinMetaData);
+
+					const FSearchTagDataPair* MatchedPinPair = nullptr;
+					for (const FSearchTagDataPair& MetadataTag : PinMetaData)
+					{
+						if (MetadataTag.Value.ToString().Contains(SearchTerms))
+						{
+							MatchedPinPair = &MetadataTag;
+							break;
+						}
+					}
+
+					if (MatchedPinPair)
+					{
+						FRigVMSearchResult Result = FindOrAddPinResult(Pin, PinMetaData, NodeMetadata);
+
+						const FText DisplayText = FText::Format(LOCTEXT("RigVMFindReferencesValues", "{0}: {1}"), MatchedPinPair->Key, MatchedPinPair->Value);
+						FRigVMSearchResult DetailResult = FRigVMSearchResult(new FRigVMFindResult(Blueprint, DisplayText));
+
+						Result->Children.Add(DetailResult);
+						DetailResult->Parent = Result;
+					}
+				}
+			}
+
+			TArray<FRigVMGraphVariableDescription> LocalVariables = Graph->GetLocalVariables();
+			for (const FRigVMGraphVariableDescription& Variable : LocalVariables)
+			{
+				TArray<UBlueprintExtension::FSearchTagDataPair> LocalVariableMetadata;
+				EdGraph->AddLocalVariableSearchMetaDataInfo(Variable.Name, LocalVariableMetadata);
+
+				const UBlueprintExtension::FSearchTagDataPair* MatchedVariablePair = nullptr;
+				for (const UBlueprintExtension::FSearchTagDataPair& MetadataTag : LocalVariableMetadata)
+				{
+					if (MetadataTag.Value.ToString().Contains(SearchTerms))
+					{
+						MatchedVariablePair = &MetadataTag;
+						break;
+					}
+				}
+
+				if (MatchedVariablePair)
+				{
+					FRigVMSearchResult Result = FindOrAddVariableResult(Variable.Name, EdGraph, LocalVariableMetadata);
+
+					const FText DisplayText = FText::Format(LOCTEXT("RigVMFindReferencesValues", "{0}: {1}"), MatchedVariablePair->Key, MatchedVariablePair->Value);
+					FRigVMSearchResult DetailResult = FRigVMSearchResult(new FRigVMFindResult(Blueprint, DisplayText));
+
+					Result->Children.Add(DetailResult);
+					DetailResult->Parent = Result;
+				}
+			}
+		}
+	}
+
+	TArray<FRigVMGraphVariableDescription> BPVariables = Blueprint->GetMemberVariables();
+	for (const FRigVMGraphVariableDescription& Variable : BPVariables)
+	{
+		TArray<UBlueprintExtension::FSearchTagDataPair> BPVariableMetadata;
+		Blueprint->AddVariableSearchMetaDataInfo(Variable.Name, BPVariableMetadata);
+		
+
+		const UBlueprintExtension::FSearchTagDataPair* MatchedVariablePair = nullptr;
+		for (const UBlueprintExtension::FSearchTagDataPair& MetadataTag : BPVariableMetadata)
+		{
+			if (MetadataTag.Value.ToString().Contains(SearchTerms))
+			{
+				MatchedVariablePair = &MetadataTag;
+				break;
+			}
+		}
+
+		if (MatchedVariablePair)
+		{
+			FRigVMSearchResult Result = FindOrAddVariableResult(Variable.Name, nullptr, BPVariableMetadata);
+
+			const FText DisplayText = FText::Format(LOCTEXT("RigVMFindReferencesValues", "{0}: {1}"), MatchedVariablePair->Key, MatchedVariablePair->Value);
+			FRigVMSearchResult DetailResult = FRigVMSearchResult(new FRigVMFindResult(Blueprint, DisplayText));
+
+			Result->Children.Add(DetailResult);
+			DetailResult->Parent = Result;
+		}
+	}
+
+	TreeView->RequestTreeRefresh();
+	for (FRigVMSearchResult Item : ItemsFound)
+	{
+		RigVMFindReferencesHelpers::ExpandAllChildren(Item, TreeView);
+	}
+}
+
+void SRigVMFindReferences::OnSearchTextChanged( const FText& Text)
+{
+	SearchValue = Text.ToString();
+}
+
+void SRigVMFindReferences::OnSearchTextCommitted( const FText& Text, ETextCommit::Type CommitType )
+{
+	if (CommitType == ETextCommit::OnEnter)
+	{
+		FindReferences(SearchValue);
+	}
+}
+
+TSharedRef<ITableRow> SRigVMFindReferences::OnGenerateRow( FRigVMSearchResult InItem, const TSharedRef<STableViewBase>& OwnerTable )
+{
+	// Finalize the search data, this does some non-thread safe actions that could not be done on the separate thread.
+	InItem->FinalizeSearchData();
+
+	FFormatNamedArguments Args;
+	Args.Add(TEXT("Category"), InItem->GetCategory());
+	Args.Add(TEXT("DisplayTitle"), InItem->DisplayText);
+
+	FText Tooltip = FText::Format(LOCTEXT("BlueprintResultSearchToolTip", "{Category} : {DisplayTitle}"), Args);
+
+	return SNew( STableRow< TSharedPtr<FRigVMFindResult> >, OwnerTable )
+		.Style( &FAppStyle::Get().GetWidgetStyle<FTableRowStyle>("ShowParentsTableView.Row") )
+		[
+			SNew(SHorizontalBox)
+			+SHorizontalBox::Slot()
+			.VAlign(VAlign_Center)
+			.AutoWidth()
+			[
+				InItem->CreateIcon()
+			]
+			+SHorizontalBox::Slot()
+			.AutoWidth()
+			.VAlign(VAlign_Center)
+			.Padding(2.f)
+			[
+				SNew(STextBlock)
+					.Text(InItem.Get(), &FRigVMFindResult::GetDisplayString)
+					.HighlightText(HighlightText)
+					.ToolTipText(Tooltip)
+			]
+		];
+}
+
+void SRigVMFindReferences::OnGetChildren( FRigVMSearchResult InItem, TArray< FRigVMSearchResult >& OutChildren )
+{
+	OutChildren += InItem->Children;
+}
+
+void SRigVMFindReferences::OnTreeSelectionDoubleClicked( FRigVMSearchResult Item )
+{
+	if(Item.IsValid())
+	{
+		Item->OnClick();
+	}
+}
+
+TSharedPtr<SWidget> SRigVMFindReferences::OnContextMenuOpening()
+{
+	const bool bShouldCloseWindowAfterMenuSelection = true;
+	FMenuBuilder MenuBuilder( bShouldCloseWindowAfterMenuSelection, CommandList);
+
+	MenuBuilder.BeginSection("BasicOperations");
+	{
+		MenuBuilder.AddMenuEntry(FGenericCommands::Get().SelectAll);
+		MenuBuilder.AddMenuEntry(FGenericCommands::Get().Copy);
+	}
+	MenuBuilder.EndSection();
+
+	return MenuBuilder.MakeWidget();
+}
+
+void SRigVMFindReferences::SelectAllItemsHelper(FRigVMSearchResult InItemToSelect)
+{
+	// Iterates over all children and recursively selects all items in the results
+	TreeView->SetItemSelection(InItemToSelect, true);
+
+	for( const auto& Child : InItemToSelect->Children )
+	{
+		SelectAllItemsHelper(Child);
+	}
+}
+
+void SRigVMFindReferences::OnSelectAllAction()
+{
+	for( const auto& Item : ItemsFound )
+	{
+		SelectAllItemsHelper(Item);
+	}
+}
+
+void SRigVMFindReferences::OnCopyAction()
+{
+	TArray< TSharedPtr<FRigVMFindResult> > SelectedItems = TreeView->GetSelectedItems();
+
+	FString SelectedText;
+
+	for(const TSharedPtr<FRigVMFindResult>& SelectedItem : SelectedItems)
+	{
+		// Add indents for each layer into the tree the item is
+		for(auto ParentItem = SelectedItem->Parent; ParentItem.IsValid(); ParentItem = ParentItem.Pin()->Parent)
+		{
+			SelectedText += TEXT("\t");
+		}
+
+		// Add the display string
+		SelectedText += SelectedItem->GetDisplayString().ToString();
+
+		// Line terminator so the next item will be on a new line
+		SelectedText += LINE_TERMINATOR;
+	}
+
+	// Copy text to clipboard
+	FPlatformApplicationMisc::ClipboardCopy( *SelectedText );
+}
+
+FReply SRigVMFindReferences::OnKeyDown(const FGeometry& MyGeometry, const FKeyEvent& InKeyEvent)
+{
+	// BlueprintEditor's IToolkit code will handle shortcuts itself - but we can just use
+	// simple slate handlers when we're standalone:
+	if(!EditorPtr.IsValid() && CommandList.IsValid())
+	{
+		if( CommandList->ProcessCommandBindings(InKeyEvent))
+		{
+			return FReply::Handled();
+		}
+	}
+	return SCompoundWidget::OnKeyDown(MyGeometry, InKeyEvent);
+}
+
+#undef LOCTEXT_NAMESPACE
\ No newline at end of file
diff --git a/Engine/Plugins/Runtime/RigVM/Source/RigVMEditor/Private/Editor/RigVMFindReferencesTabSummoner.cpp b/Engine/Plugins/Runtime/RigVM/Source/RigVMEditor/Private/Editor/RigVMFindReferencesTabSummoner.cpp
new file mode 100644
--- /dev/null
+++ b/Engine/Plugins/Runtime/RigVM/Source/RigVMEditor/Private/Editor/RigVMFindReferencesTabSummoner.cpp
@@ -0,0 +1,34 @@
+// Copyright Epic Games, Inc. All Rights Reserved.
+
+#include "Editor/RigVMFindReferencesTabSummoner.h"
+#include "Editor/RigVMNewEditor.h"
+#include "Editor/RigVMFindReferences.h"
+
+class SWidget;
+struct FSlateBrush;
+
+#define LOCTEXT_NAMESPACE "RigVMEditor"
+
+FRigVMFindReferencesTabSummoner::FRigVMFindReferencesTabSummoner(TSharedPtr<class FAssetEditorToolkit> InHostingApp)
+	: FWorkflowTabFactory(TabID(), InHostingApp)
+{
+	TabLabel = LOCTEXT("FindResultsTabTitle", "Find Results");
+	TabIcon = FSlateIcon(FAppStyle::GetAppStyleSetName(), "Kismet.Tabs.FindResults");
+
+	bIsSingleton = true;
+
+	ViewMenuDescription = LOCTEXT("FindResultsView", "Find Results");
+	ViewMenuTooltip = LOCTEXT("FindResultsView_ToolTip", "Show find results for searching in this blueprint");
+}
+
+
+TSharedRef<SWidget> FRigVMFindReferencesTabSummoner::CreateTabBody(const FWorkflowTabSpawnInfo& Info) const
+{
+	TSharedPtr<FRigVMNewEditor> EditorPtr = StaticCastSharedPtr<FRigVMNewEditor>(HostingApp.Pin());
+
+	return EditorPtr->GetFindResults();
+}
+
+
+
+#undef LOCTEXT_NAMESPACE
diff --git a/Engine/Plugins/Runtime/RigVM/Source/RigVMEditor/Private/Editor/RigVMNewEditor.cpp b/Engine/Plugins/Runtime/RigVM/Source/RigVMEditor/Private/Editor/RigVMNewEditor.cpp
--- a/Engine/Plugins/Runtime/RigVM/Source/RigVMEditor/Private/Editor/RigVMNewEditor.cpp
+++ b/Engine/Plugins/Runtime/RigVM/Source/RigVMEditor/Private/Editor/RigVMNewEditor.cpp
@@ -31,8 +31,10 @@
 #include "MessageLogModule.h"
 #include "RigVMBlueprint.h"
 #include "RigVMEditorCommands.h"
+#include "Editor/RigVMFindReferences.h"
 #include "Editor/RigVMCompilerResultsTabSummoner.h"
 #include "Editor/RigVMDetailsInspectorTabSummoner.h"
+#include "Editor/RigVMFindReferencesTabSummoner.h"
 #include "Editor/RigVMGraphEditorSummoner.h"
 #include "Widgets/Input/SEditableTextBox.h"
 #include "Editor/SRigVMDetailsInspector.h"
@@ -659,23 +661,15 @@ void FRigVMNewEditor::OnAddNewVariable()
 
 void FRigVMNewEditor::SummonSearchUI(bool bSetFindWithinBlueprint, FString NewSearchTerms, bool bSelectFirstResult)
 {
-	// TSharedPtr<SFindInBlueprints> FindResultsToUse = nullptr;
+	TSharedPtr<SRigVMFindReferences> FindResultsToUse = nullptr;
 
-	// TODO: sara-s
-	// if (bSetFindWithinBlueprint)
-	// {
-	// 	FindResultsToUse = FindResults;
-	// 	TabManager->TryInvokeTab(FRigVMNewEditorTabs::FindResultsID);
-	// }
-	// else
-	// {
-	// 	FindResultsToUse = FFindInBlueprintSearchManager::Get().GetGlobalFindResults();
-	// }
-	//
-	// if (FindResultsToUse.IsValid())
-	// {
-	// 	FindResultsToUse->FocusForUse(bSetFindWithinBlueprint, NewSearchTerms, bSelectFirstResult);
-	// }
+	FindResultsToUse = FindResults;
+	TabManager->TryInvokeTab(FRigVMFindReferencesTabSummoner::TabID());
+	
+	if (FindResultsToUse.IsValid())
+	{
+		FindResultsToUse->FocusForUse(bSetFindWithinBlueprint, NewSearchTerms, bSelectFirstResult);
+	}
 }
 
 void FRigVMNewEditor::ZoomToSelection_Clicked()
@@ -2045,7 +2039,7 @@ void FRigVMNewEditor::DumpMessagesToCompilerLog(const TArray<TSharedRef<class FT
 	
 	if (bForceMessageDisplay)
 	{
-		TabManager->TryInvokeTab(FBlueprintEditorTabs::CompilerResultsID);
+		TabManager->TryInvokeTab(FRigVMCompilerResultsTabSummoner::TabID());
 	}
 }
 
@@ -2075,7 +2069,7 @@ void FRigVMNewEditor::CreateDefaultTabContents(const TArray<UBlueprint*> InBluep
 
 	FMessageLogModule& MessageLogModule = FModuleManager::LoadModuleChecked<FMessageLogModule>("MessageLog");
 	CompilerResults = MessageLogModule.CreateLogListingWidget( CompilerResultsListing.ToSharedRef() );
-	// FindResults = SNew(SFindInBlueprints, SharedThis(this));
+	FindResults = SNew(SRigVMFindReferences, SharedThis(this));
 	
 	this->Inspector = 
 		SNew(SRigVMDetailsInspector)
diff --git a/Engine/Plugins/Runtime/RigVM/Source/RigVMEditor/Private/Editor/RigVMNewEditorMode.cpp b/Engine/Plugins/Runtime/RigVM/Source/RigVMEditor/Private/Editor/RigVMNewEditorMode.cpp
--- a/Engine/Plugins/Runtime/RigVM/Source/RigVMEditor/Private/Editor/RigVMNewEditorMode.cpp
+++ b/Engine/Plugins/Runtime/RigVM/Source/RigVMEditor/Private/Editor/RigVMNewEditorMode.cpp
@@ -6,6 +6,7 @@
 #include "Editor/RigVMCompilerResultsTabSummoner.h"
 #include "Editor/RigVMEditor.h"
 #include "Editor/RigVMExecutionStackTabSummoner.h"
+#include "Editor/RigVMFindReferencesTabSummoner.h"
 #include "Editor/RigVMNewEditor.h"
 
 const FName FRigVMNewEditorApplicationModes::StandardRigVMEditorMode() { static FName Mode = TEXT("GraphName"); return Mode; };
@@ -20,6 +21,7 @@ FRigVMNewEditorMode::FRigVMNewEditorMode(const TSharedRef<FRigVMNewEditor>& InRi
 	TabFactories.RegisterFactory(MakeShared<FRigVMEditorGraphExplorerTabSummoner>(InRigVMEditor));
 	TabFactories.RegisterFactory(MakeShared<FRigVMDetailsInspectorTabSummoner>(InRigVMEditor));
 	TabFactories.RegisterFactory(MakeShared<FRigVMCompilerResultsTabSummoner>(InRigVMEditor));
+	TabFactories.RegisterFactory(MakeShared<FRigVMFindReferencesTabSummoner>(InRigVMEditor));
 
 	TabLayout = FTabManager::NewLayout("Standalone_RigVMEditMode_Layout_v1.5")
 		->AddArea
@@ -58,6 +60,7 @@ FRigVMNewEditorMode::FRigVMNewEditorMode(const TSharedRef<FRigVMNewEditor>& InRi
 							FTabManager::NewStack()
 							->SetSizeCoefficient(0.2f)
 							->AddTab(FRigVMCompilerResultsTabSummoner::TabID(), ETabState::OpenedTab)
+							->AddTab(FRigVMFindReferencesTabSummoner::TabID(), ETabState::OpenedTab)
 						)
 				)
 				->Split
diff --git a/Engine/Plugins/Runtime/RigVM/Source/RigVMEditor/Private/Widgets/SRigVMEditorGraphExplorer.cpp b/Engine/Plugins/Runtime/RigVM/Source/RigVMEditor/Private/Widgets/SRigVMEditorGraphExplorer.cpp
--- a/Engine/Plugins/Runtime/RigVM/Source/RigVMEditor/Private/Widgets/SRigVMEditorGraphExplorer.cpp
+++ b/Engine/Plugins/Runtime/RigVM/Source/RigVMEditor/Private/Widgets/SRigVMEditorGraphExplorer.cpp
@@ -2019,7 +2019,6 @@ void SRigVMEditorGraphExplorer::OnFindReference(bool bSearchAllBlueprints)
 	URigVMBlueprint* Blueprint = Editor->GetRigVMBlueprint();
 	const URigVMEdGraphSchema* Schema = GetDefault<URigVMEdGraphSchema>();
 
-	bool bUseQuotes = true;
 	TArray<FRigVMExplorerElementKey> Selection = TreeView->GetSelectedKeys();
 	if (Selection.Num() != 1)
 	{
@@ -2049,11 +2048,6 @@ void SRigVMEditorGraphExplorer::OnFindReference(bool bSearchAllBlueprints)
 	
 	if (!SearchTerm.IsEmpty())
 	{
-		if (bUseQuotes)
-		{
-			SearchTerm = FString::Printf(TEXT("\"%s\""), *SearchTerm);
-		}
-		
 		const bool bSetFindWithinBlueprint = !bSearchAllBlueprints;
 		Editor->SummonSearchUI(bSetFindWithinBlueprint, SearchTerm);
 	}
diff --git a/Engine/Plugins/Runtime/RigVM/Source/RigVMEditor/Public/Editor/RigVMFindReferences.h b/Engine/Plugins/Runtime/RigVM/Source/RigVMEditor/Public/Editor/RigVMFindReferences.h
new file mode 100644
--- /dev/null
+++ b/Engine/Plugins/Runtime/RigVM/Source/RigVMEditor/Public/Editor/RigVMFindReferences.h
@@ -0,0 +1,277 @@
+// Copyright Epic Games, Inc. All Rights Reserved.
+
+#pragma once
+
+#include "CoreMinimal.h"
+#include "Input/Reply.h"
+#include "Widgets/Views/STreeView.h"
+#include "EdGraph/EdGraphPin.h"
+#include "Textures/SlateIcon.h"
+#include "EdGraph/EdGraphSchema.h"
+
+class SWidget;
+class SDockTab;
+class UBlueprint;
+class FUICommandList;
+class URigVMBlueprint;
+
+
+class RIGVMEDITOR_API FRigVMFindResult : public TSharedFromThis< FRigVMFindResult >
+{
+public:
+	FRigVMFindResult() = default;
+	virtual ~FRigVMFindResult() = default;
+
+	/* Create a root */
+	explicit FRigVMFindResult(TWeakObjectPtr<URigVMBlueprint> InBlueprint);
+	explicit FRigVMFindResult(TWeakObjectPtr<URigVMBlueprint> InBlueprint, const FText& InDisplayText);
+
+	/* Called when user clicks on the search item */
+	virtual FReply OnClick();
+
+	/* Get Category for this search result */
+	virtual FText GetCategory() const;
+
+	/* Create an icon to represent the result */
+	virtual TSharedRef<SWidget>	CreateIcon() const;
+
+	/** Finalizes any content for the search data that was unsafe to do on a separate thread */
+	virtual void FinalizeSearchData() {};
+
+	/** gets the blueprint housing all these search results */
+	URigVMBlueprint* GetBlueprint() const;
+
+	/**
+	* Parses search info for specific data important for displaying the search result in an easy to understand format
+	*
+	* @param	InKey			This is the tag for the data, describing what it is so special handling can occur if needed
+	* @param	InValue			Compared against search query to see if it passes the filter, sometimes data is rejected because it is deemed unsearchable
+	*/
+	virtual void ParseSearchInfo(FText InKey, FText InValue) {};
+
+	/** Returns the Object represented by this search information give the Blueprint it can be found in */
+	virtual UObject* GetObject(URigVMBlueprint* InBlueprint) const;
+
+	/** Returns the display string for the row */
+	FText GetDisplayString() const;
+
+public:
+	/*Any children listed under this category */
+	TArray< TSharedPtr<FRigVMFindResult> > Children;
+
+	/*If it exists it is the blueprint*/
+	TWeakPtr<FRigVMFindResult> Parent;
+
+	/*If it exists it is the blueprint*/
+	TWeakObjectPtr<URigVMBlueprint> WeakBlueprint;
+
+	/*The display text for this item */
+	FText DisplayText;
+};
+
+typedef TSharedPtr<FRigVMFindResult> FRigVMSearchResult;
+typedef STreeView<FRigVMSearchResult>  SRigVMTreeViewType;
+
+/** Some utility functions to help with Find-in-Blueprint functionality */
+namespace RigVMFindReferencesHelpers
+{
+	/**
+	 * Retrieves the pin type as a string value
+	 *
+	 * @param InPinType		The pin type to look at
+	 *
+	 * @return				The pin type as a string in format [category]'[sub-category object]'
+	 */
+	FString GetPinTypeAsString(const FEdGraphPinType& InPinType);
+
+	/**
+	 * Parses a pin type from passed in key names and values
+	 *
+	 * @param InKey					The key name for what the data should be translated as
+	 * @param InValue				Value to be be translated
+	 * @param InOutPinType			Modifies the PinType based on the passed parameters, building it up over multiple calls
+	 * @return						TRUE when the parsing is successful
+	 */
+	bool ParsePinType(FText InKey, FText InValue, FEdGraphPinType& InOutPinType);
+
+	/**
+	* Iterates through all the given tree node's children and tells the tree view to expand them
+	*/
+	void ExpandAllChildren(FRigVMSearchResult InTreeNode, TSharedPtr<STreeView<FRigVMSearchResult>> InTreeView);
+}
+
+/** Graph nodes use this class to store their data */
+class FRigVMFindReferencesGraphNode : public FRigVMFindResult
+{
+public:
+	FRigVMFindReferencesGraphNode(TWeakObjectPtr<URigVMBlueprint> InBlueprint);
+	virtual ~FRigVMFindReferencesGraphNode() {}
+
+	/** FRigVMFindResult Interface */
+	virtual FReply OnClick() override;
+	virtual TSharedRef<SWidget>	CreateIcon() const override;
+	virtual void ParseSearchInfo(FText InKey, FText InValue) override;
+	virtual void FinalizeSearchData() override;
+	virtual UObject* GetObject(URigVMBlueprint* InBlueprint) const override;
+	virtual FText GetCategory() const override;
+	/** End FRigVMFindResult Interface */
+
+private:
+	/** The Node Guid to find when jumping to the node */
+	FGuid NodeGuid;
+
+	/** The glyph brush for this node */
+	FSlateIcon Glyph;
+
+	/** The glyph color for this node */
+	FLinearColor GlyphColor;
+
+	/*The class this item refers to */
+	UClass* Class;
+
+	/*The class name this item refers to */
+	FString ClassName;
+};
+
+/** Pins use this class to store their data */
+class FRigVMFindReferencesPin : public FRigVMFindResult
+{
+public:
+	FRigVMFindReferencesPin(TWeakObjectPtr<URigVMBlueprint> InBlueprint, FString InSchemaName);
+	virtual ~FRigVMFindReferencesPin() {}
+
+	/** FRigVMFindResult Interface */
+	virtual TSharedRef<SWidget>	CreateIcon() const override;
+	virtual void ParseSearchInfo(FText InKey, FText InValue) override;
+	virtual FText GetCategory() const override;
+	virtual void FinalizeSearchData() override;
+	/** End FRigVMFindResult Interface */
+
+private:
+	/** The name of the schema this pin exists under */
+	FString SchemaName;
+
+	/** The pin that this search result refers to */
+	FEdGraphPinType PinType;
+
+	/** Pin's icon color */
+	FSlateColor IconColor;
+};
+
+/** Property data is stored here */
+class FRigVMFindReferencesVariable : public FRigVMFindResult
+{
+public:
+	FRigVMFindReferencesVariable(TWeakObjectPtr<URigVMBlueprint> InBlueprint);
+	virtual ~FRigVMFindReferencesVariable() {}
+
+	/** FRigVMFindResult Interface */
+	virtual TSharedRef<SWidget>	CreateIcon() const override;
+	virtual void ParseSearchInfo(FText InKey, FText InValue) override;
+	virtual FText GetCategory() const override;
+	virtual void FinalizeSearchData() override;
+	/** End FRigVMFindResult Interface */
+
+private:
+	/** The pin that this search result refers to */
+	FEdGraphPinType PinType;
+
+	/** The default value of a property as a string */
+	FString DefaultValue;
+};
+
+/** Graphs, such as functions and macros, are stored here */
+class FRigVMFindReferencesGraph : public FRigVMFindResult
+{
+public:
+	FRigVMFindReferencesGraph(TWeakObjectPtr<URigVMBlueprint> InBlueprint, EGraphType InGraphType);
+	virtual ~FRigVMFindReferencesGraph() {}
+
+	/** FRigVMFindResult Interface */
+	virtual FReply OnClick() override;
+	virtual TSharedRef<SWidget>	CreateIcon() const override;
+	virtual void ParseSearchInfo(FText InKey, FText InValue) override;
+	virtual FText GetCategory() const override;
+	/** End FRigVMFindResult Interface */
+
+private:
+	/** The type of graph this represents */
+	EGraphType GraphType;
+};
+
+/*Widget for searching for (functions/events) across all blueprints or just a single blueprint */
+class RIGVMEDITOR_API SRigVMFindReferences: public SCompoundWidget
+{
+public:
+	SLATE_BEGIN_ARGS( SRigVMFindReferences )
+	{}
+	SLATE_END_ARGS()
+
+	void Construct(const FArguments& InArgs, TSharedPtr<class FRigVMEditorBase> InBlueprintEditor = nullptr);
+
+	/** Focuses this widget's search box, and changes the mode as well, and optionally the search terms */
+	void FocusForUse(bool bSetFindWithinBlueprint, FString NewSearchTerms = FString(), bool bSelectFirstResult = false);
+
+	/** SWidget overrides */
+	virtual FReply OnKeyDown(const FGeometry& MyGeometry, const FKeyEvent& InKeyEvent) override;
+
+private:
+	
+	/** The main function that will find references and build the tree */
+	void FindReferences(const FString& SearchTerms);
+
+	/** Register any Find-in-Blueprint commands */
+	void RegisterCommands();
+
+	/*Called when user changes the text they are searching for */
+	void OnSearchTextChanged(const FText& Text);
+
+	/*Called when user changes commits text to the search box */
+	void OnSearchTextCommitted(const FText& Text, ETextCommit::Type CommitType);
+
+	/* Get the children of a row */
+	void OnGetChildren( FRigVMSearchResult InItem, TArray< FRigVMSearchResult >& OutChildren );
+
+	/* Called when user double clicks on a new result */
+	void OnTreeSelectionDoubleClicked( FRigVMSearchResult Item );
+
+	/* Called when a new row is being generated */
+	TSharedRef<ITableRow> OnGenerateRow(FRigVMSearchResult InItem, const TSharedRef<STableViewBase>& OwnerTable);
+
+	/** Callback to build the context menu when right clicking in the tree */
+	TSharedPtr<SWidget> OnContextMenuOpening();
+
+	/** Helper function to select all items */
+	void SelectAllItemsHelper(FRigVMSearchResult InItemToSelect);
+
+	/** Callback when user attempts to select all items in the search results */
+	void OnSelectAllAction();
+
+	/** Callback when user attempts to copy their selection in the Find-in-Blueprints */
+	void OnCopyAction();
+
+private:
+	/** Pointer back to the blueprint editor that owns us */
+	TWeakPtr<class FRigVMEditorBase> EditorPtr;
+	
+	/* The tree view displays the results */
+	TSharedPtr<SRigVMTreeViewType> TreeView;
+
+	/** The search text box */
+	TSharedPtr<class SSearchBox> SearchTextField;
+	
+	/* This buffer stores the currently displayed results */
+	TArray<FRigVMSearchResult> ItemsFound;
+
+	/* Map relationship between element hash and its result  */ 
+	TMap<uint32, FRigVMSearchResult> ElementHashToResult;
+
+	/* The string to highlight in the results */
+	FText HighlightText;
+
+	/* The string to search for */
+	FString	SearchValue;
+
+	/** Commands handled by this widget */
+	TSharedPtr< FUICommandList > CommandList;
+};
\ No newline at end of file
diff --git a/Engine/Plugins/Runtime/RigVM/Source/RigVMEditor/Public/Editor/RigVMFindReferencesTabSummoner.h b/Engine/Plugins/Runtime/RigVM/Source/RigVMEditor/Public/Editor/RigVMFindReferencesTabSummoner.h
new file mode 100644
--- /dev/null
+++ b/Engine/Plugins/Runtime/RigVM/Source/RigVMEditor/Public/Editor/RigVMFindReferencesTabSummoner.h
@@ -0,0 +1,32 @@
+// Copyright Epic Games, Inc. All Rights Reserved.
+
+#pragma once
+
+#include "CoreMinimal.h"
+#include "Widgets/SWidget.h"
+#include "Toolkits/AssetEditorToolkit.h"
+#include "WorkflowOrientedApp/WorkflowTabFactory.h"
+
+#define LOCTEXT_NAMESPACE "RigVMEditor"
+
+//////////////////////////////////////////////////////////////////////////
+// FCompilerResultsSummoner
+
+struct RIGVMEDITOR_API FRigVMFindReferencesTabSummoner : public FWorkflowTabFactory
+{
+
+public:
+
+	static const FName TabID() { static FName ID = TEXT("RigVM Find References"); return ID; }
+	
+	FRigVMFindReferencesTabSummoner(TSharedPtr<class FAssetEditorToolkit> InHostingApp);
+
+	virtual TSharedRef<SWidget> CreateTabBody(const FWorkflowTabSpawnInfo& Info) const override;
+
+	virtual FText GetTabToolTipText(const FWorkflowTabSpawnInfo& Info) const override
+	{
+		return LOCTEXT("FindResultsTooltip", "The find results tab shows results of searching this Blueprint (or all Blueprints).");
+	}
+};
+
+#undef LOCTEXT_NAMESPACE
\ No newline at end of file
diff --git a/Engine/Plugins/Runtime/RigVM/Source/RigVMEditor/Public/Editor/RigVMNewEditor.h b/Engine/Plugins/Runtime/RigVM/Source/RigVMEditor/Public/Editor/RigVMNewEditor.h
--- a/Engine/Plugins/Runtime/RigVM/Source/RigVMEditor/Public/Editor/RigVMNewEditor.h
+++ b/Engine/Plugins/Runtime/RigVM/Source/RigVMEditor/Public/Editor/RigVMNewEditor.h
@@ -6,6 +6,7 @@
 #include "RigVMEditor.h"
 #include "WorkflowOrientedApp/WorkflowCentricApplication.h"
 
+class SRigVMFindReferences;
 class SRigVMDetailsInspector;
 
 struct RIGVMEDITOR_API FRigVMNewEditorTabs
@@ -48,6 +49,8 @@ public:
 	virtual void SetInspector(TSharedPtr<SRigVMDetailsInspector> InWidget) { Inspector = InWidget; };
 	
 	TSharedRef<SWidget> GetCompilerResults() const { return CompilerResults.ToSharedRef(); }
+	TSharedRef<SRigVMFindReferences> GetFindResults() const { return FindResults.ToSharedRef(); }
+	virtual void JumpToHyperlink(const UObject* ObjectReference, bool bRequestRename) override;
 protected:
 	virtual TSharedPtr<IMessageLogListing> GetCompilerResultsListing() override { return CompilerResultsListing; }
 	virtual TSharedPtr<FApplicationMode> CreateEditorMode() override;
@@ -67,7 +70,6 @@ protected:
 	virtual FGraphPanelSelectionSet GetSelectedNodes() const override;
 	virtual void SetUISelectionState(FName SelectionOwner) override;
 	virtual void AnalyticsTrackNodeEvent(UBlueprint* Blueprint, UEdGraphNode* GraphNode, bool bNodeDelete) const override;
-	virtual void JumpToHyperlink(const UObject* ObjectReference, bool bRequestRename) override;
 	virtual void PostUndo(bool bSuccess) override;
 	virtual void PostRedo(bool bSuccess) override;
 	virtual UEdGraphPin* GetCurrentlySelectedPin() const override;
@@ -210,6 +212,9 @@ protected:
 	TSharedPtr<class SWidget> CompilerResults;
 	TSharedPtr<class IMessageLogListing> CompilerResultsListing;
 
+	/** Find results log as well as the search filter */
+	TSharedPtr<class SRigVMFindReferences> FindResults;
+
 	/** When set, flags which graph has a action menu currently open (if null, no graphs do). */
 	UEdGraph* HasOpenActionMenu;
 	
diff --git a/Engine/Plugins/Runtime/RigVM/Source/RigVMEditor/RigVMEditor.Build.cs b/Engine/Plugins/Runtime/RigVM/Source/RigVMEditor/RigVMEditor.Build.cs
--- a/Engine/Plugins/Runtime/RigVM/Source/RigVMEditor/RigVMEditor.Build.cs
+++ b/Engine/Plugins/Runtime/RigVM/Source/RigVMEditor/RigVMEditor.Build.cs
@@ -33,6 +33,7 @@ public class RigVMEditor : ModuleRules
                 "MessageLog",
 				"StructUtilsEditor",
 				"ContentBrowser",
+				"EditorFramework"
 			}
 		);
         
