Refactored RigVM math customizations to use propertyhandles, rather than raw acess to struct memory, in order to fix a memory stomp in StateTrees

diff --git a/Engine/Plugins/Runtime/RigVM/Source/RigVMEditor/Private/Editor/RigVMGraphDetailCustomization.cpp b/Engine/Plugins/Runtime/RigVM/Source/RigVMEditor/Private/Editor/RigVMGraphDetailCustomization.cpp
--- a/Engine/Plugins/Runtime/RigVM/Source/RigVMEditor/Private/Editor/RigVMGraphDetailCustomization.cpp
+++ b/Engine/Plugins/Runtime/RigVM/Source/RigVMEditor/Private/Editor/RigVMGraphDetailCustomization.cpp
@@ -41,6 +41,17 @@
 
 #define LOCTEXT_NAMESPACE "RigVMGraphDetailCustomization"
 
+// --- FTransform and EulerTransform customization data start ---
+static const FName NAME_Translation(TEXT("Translation"));
+static const FName NAME_Scale3D(TEXT("Scale3D"));
+static const FName NAME_Scale(TEXT("Scale"));
+static const FName TransformComponentNames[] = {NAME_Translation, NAME_Rotation, NAME_Scale3D};
+static const FName EulerTransformComponentNames[] = { NAME_Location, NAME_Rotation, NAME_Scale };
+static_assert(ESlateTransformComponent::Location == 0); // make sure the names array is in the same order than the components enum
+static_assert(ESlateTransformComponent::Rotation == 1);
+static_assert(ESlateTransformComponent::Scale == 2);
+// --- FTransform and EulerTransform customization data end ---
+
 static const FText RigVMGraphDetailCustomizationMultipleValues = LOCTEXT("MultipleValues", "Multiple Values");
 
 FRigVMFunctionArgumentGroupLayout::FRigVMFunctionArgumentGroupLayout(
@@ -3103,7 +3114,7 @@ void FRigVMGraphMathTypeDetailCustomization::MakeRotationHeaderRow(TSharedRef<cl
 }
 
 template <typename TransformType>
-void FRigVMGraphMathTypeDetailCustomization::ConfigureTransformWidgetArgs(TSharedRef<IPropertyHandle> InPropertyHandle, typename SAdvancedTransformInputBox<TransformType>::FArguments& WidgetArgs)
+void FRigVMGraphMathTypeDetailCustomization::ConfigureTransformWidgetArgs(TSharedRef<IPropertyHandle> InPropertyHandle, typename SAdvancedTransformInputBox<TransformType>::FArguments& WidgetArgs, TConstArrayView<FName> ComponentNames)
 {
 	FEditPropertyChain PropertyChain;
 	TArray<int32> PropertyArrayIndices;
@@ -3120,30 +3131,164 @@ void FRigVMGraphMathTypeDetailCustomization::ConfigureTransformWidgetArgs(TShare
 
 	static TransformType Identity = TransformType::Identity;
 
-	uint8* ContainerMemory = nullptr;
+	void* ContainerMemory = nullptr;
 	TSharedPtr<FStructOnScope> DefaultStruct = nullptr;
-	if (!ObjectsBeingCustomized.IsEmpty())
-	{
-		ContainerMemory = (uint8*) ObjectsBeingCustomized[0]->GetClass()->GetDefaultObject();
-	}
-	else if(!StructsBeingCustomized.IsEmpty())
+
+	TArray<TSharedPtr<FStructOnScope>> StructsBeingCustomized;
+	InPropertyHandle->GetOuterStructs(StructsBeingCustomized);
+
+	if (!StructsBeingCustomized.IsEmpty())
 	{
 		DefaultStruct = MakeShareable(new FStructOnScope(StructsBeingCustomized[0]->GetStruct()));
 		ContainerMemory = DefaultStruct->GetStructMemory();
 	}
+	else
+	{
+		TArray<UObject*> ObjectsBeingCustomized;
+		InPropertyHandle->GetOuterObjects(ObjectsBeingCustomized);
+		if (!ObjectsBeingCustomized.IsEmpty())
+		{
+			ContainerMemory = (uint8*)ObjectsBeingCustomized[0]->GetClass()->GetDefaultObject();
+		}
+	}
+
 	if (!ContainerMemory)
 	{
 		return;
 	}
-	TransformType DefaultValue = ContainerMemoryBlockToValueRef<TransformType>(ContainerMemory, Identity, PropertyChain, PropertyArrayIndices);
+	TransformType DefaultValue = ContainerMemoryBlockToValueRef<TransformType>((uint8*)ContainerMemory, Identity, PropertyChain, PropertyArrayIndices);
 	
-	WidgetArgs.DiffersFromDefault_Lambda([this, InPropertyHandle, DefaultValue](ESlateTransformComponent::Type InTransformComponent) -> bool
+	TSharedPtr<IPropertyHandle> TranslationHandle = InPropertyHandle->GetChildHandle(ComponentNames[ESlateTransformComponent::Location]);
+	TSharedPtr<IPropertyHandle> RotationHandle = InPropertyHandle->GetChildHandle(ComponentNames[ESlateTransformComponent::Rotation]);
+	TSharedPtr<IPropertyHandle> ScaleHandle = InPropertyHandle->GetChildHandle(ComponentNames[ESlateTransformComponent::Scale]);
+
+	if (!TranslationHandle.IsValid() || !RotationHandle.IsValid() || !ScaleHandle.IsValid())
 	{
-		if (ObjectsBeingCustomized.IsEmpty())
+		return;
+	}
+
+	const auto GetTranslation = [TranslationHandle]() -> FVector
 		{
-			return false;
-		}
-		
+			FVector Translation = FVector::ZeroVector;
+			if (TranslationHandle && TranslationHandle->IsValidHandle())
+			{
+				TranslationHandle->GetValue(Translation);
+			}
+			return Translation;
+		};
+	const auto GetRotation = [RotationHandle]() -> FQuat
+		{
+			FQuat Rotation = FQuat::Identity;
+			if (RotationHandle && RotationHandle->IsValidHandle())
+			{
+				RotationHandle->GetValue(Rotation);
+			}
+			return Rotation;
+		};
+	const auto GetScale3D = [ScaleHandle]() -> FVector
+		{
+			FVector Scale = FVector::ZeroVector;
+			if (ScaleHandle && ScaleHandle->IsValidHandle())
+			{
+				ScaleHandle->GetValue(Scale);
+			}
+			return Scale;
+		};
+
+	const auto SetTranslation = [TranslationHandle](const FVector& InValue)
+		{
+			if (TranslationHandle && TranslationHandle->IsValidHandle())
+			{
+				TranslationHandle->SetValue(InValue);
+			}
+		};
+	const auto SetRotation = [RotationHandle](const FQuat& InValue)
+		{
+			if (RotationHandle && RotationHandle->IsValidHandle())
+			{
+				RotationHandle->SetValue(InValue);
+			}
+		};
+	const auto SetScale3D = [ScaleHandle](const FVector& InValue)
+		{
+			if (ScaleHandle && ScaleHandle->IsValidHandle())
+			{
+				ScaleHandle->SetValue(InValue);
+			}
+		};
+
+	const auto GetTransformComponent = [GetTranslation, GetRotation, GetScale3D](ESlateTransformComponent::Type InTransformComponent) -> TransformType
+		{
+			TransformType Transform = TransformType::Identity;
+			switch (InTransformComponent)
+			{
+				case ESlateTransformComponent::Location:
+				{
+					const FVector Translation = GetTranslation();
+					Transform.SetLocation(Translation);
+					break;
+				}
+				case ESlateTransformComponent::Rotation:
+				{
+					const FQuat Rotation = GetRotation();
+					Transform.SetRotation(Rotation);
+					break;
+				}
+				case ESlateTransformComponent::Scale:
+				{
+					const FVector Scale = GetScale3D();
+					Transform.SetScale3D(Scale);
+					break;
+				}
+				case ESlateTransformComponent::Max: // It means all components
+				{
+					Transform.SetLocation(GetTranslation());
+					Transform.SetRotation(GetRotation());
+					Transform.SetScale3D(GetScale3D());
+					break;
+				}
+				default:
+				{
+					check(false);
+				}
+			}
+			return Transform;
+		};
+	const auto SetTransformComponent = [InPropertyHandle, SetTranslation, SetRotation, SetScale3D](ESlateTransformComponent::Type InTransformComponent, const TransformType& InValue)
+		{
+			switch (InTransformComponent)
+			{
+				case ESlateTransformComponent::Location:
+				{
+					SetTranslation(InValue.GetLocation());
+					break;
+				}
+				case ESlateTransformComponent::Rotation:
+				{
+					SetRotation(InValue.GetRotation());
+					break;
+				}
+				case ESlateTransformComponent::Scale:
+				{
+					SetScale3D(InValue.GetScale3D());
+					break;
+				}
+				case ESlateTransformComponent::Max: // It means all components
+				{
+					SetTranslation(InValue.GetLocation());
+					SetRotation(InValue.GetRotation());
+					SetScale3D(InValue.GetScale3D());
+					break;
+				}
+				default:
+				{
+					check(false);
+				}
+			}
+		};
+
+	WidgetArgs.DiffersFromDefault_Lambda([this, InPropertyHandle, DefaultValue, GetTransformComponent](ESlateTransformComponent::Type InTransformComponent) -> bool
+	{
 		FEditPropertyChain PropertyChain;
 		TArray<int32> PropertyArrayIndices;
 		bool bEnabled;
@@ -3151,50 +3296,47 @@ void FRigVMGraphMathTypeDetailCustomization::ConfigureTransformWidgetArgs(TShare
 		{
 			return false;
 		}
-		
-		for(const TWeakObjectPtr<UObject>& Object : ObjectsBeingCustomized)
+
+		const TransformType Transform = GetTransformComponent(InTransformComponent);
+		switch(InTransformComponent)
 		{
-			if(Object.Get() && InPropertyHandle->IsValidHandle())
+			case ESlateTransformComponent::Location:
 			{
-				const TransformType& Transform = ContainerMemoryBlockToValueRef<TransformType>((uint8*)Object.Get(), Identity, PropertyChain, PropertyArrayIndices);
-
-				switch(InTransformComponent)
+				if (!Transform.GetLocation().Equals(DefaultValue.GetLocation()))
 				{
-					case ESlateTransformComponent::Location:
-					{
-						if(!Transform.GetLocation().Equals(DefaultValue.GetLocation()))
-						{
-							return true;
-						}
-						break;
-					}
-					case ESlateTransformComponent::Rotation:
-					{
-						if(!Transform.Rotator().Equals(DefaultValue.Rotator()))
-						{
-							return true;
-						}
-						break;
-					}
-					case ESlateTransformComponent::Scale:
-					{
-						if(!Transform.GetScale3D().Equals(DefaultValue.GetScale3D()))
-						{
-							return true;
-						}
-						break;
-					}
-					default:
-					{
-						break;
-					}
+					return true;
 				}
+				break;
+			}
+			case ESlateTransformComponent::Rotation:
+			{
+				if (!Transform.Rotator().Equals(DefaultValue.Rotator()))
+				{
+					return true;
+				}
+				break;
+			}
+			case ESlateTransformComponent::Scale:
+			{
+				if (!Transform.GetScale3D().Equals(DefaultValue.GetScale3D()))
+				{
+					return true;
+				}
+				break;
+			}
+			default:
+			{
+				if (!Transform.Equals(DefaultValue))
+				{
+					return true;
+				}
+				break;
 			}
 		}
 		return false;
 	});
 
-	WidgetArgs.OnGetNumericValue_Lambda([this, InPropertyHandle](
+	WidgetArgs.OnGetNumericValue_Lambda([this, InPropertyHandle, GetTransformComponent](
 		ESlateTransformComponent::Type InTransformComponent,
 		ESlateRotationRepresentation::Type InRotationRepresentation,
 		ESlateTransformSubComponent::Type InTransformSubComponent) -> TOptional<FReal>
@@ -3208,41 +3350,22 @@ void FRigVMGraphMathTypeDetailCustomization::ConfigureTransformWidgetArgs(TShare
 			return Result;
 		}
 
-		const TArray<uint8*> MemoryBlocks = GetMemoryBeingCustomized();
-		for(uint8* MemoryBlock: MemoryBlocks)
-		{
-			if(MemoryBlock)
-			{
-				const TransformType& Transform = ContainerMemoryBlockToValueRef<TransformType>(MemoryBlock, Identity, PropertyChain, PropertyArrayIndices);
-				
-				TOptional<FReal> Value = SAdvancedTransformInputBox<TransformType>::GetNumericValueFromTransform(
-					Transform,
-					InTransformComponent,
-					InRotationRepresentation,
-					InTransformSubComponent
-					);
+		const TransformType Transform = GetTransformComponent(InTransformComponent);
+		const TOptional<FReal> Value = SAdvancedTransformInputBox<TransformType>::GetNumericValueFromTransform(
+			Transform,
+			InTransformComponent,
+			InRotationRepresentation,
+			InTransformSubComponent);
 
-				if(Value.IsSet())
-				{
-					if(Result.IsSet())
-					{
-						if(!FMath::IsNearlyEqual(Result.GetValue(), Value.GetValue()))
-						{
-							return TOptional<FReal>();
-						}
-					}
-					else
-					{
-						Result = Value;
-					}
-				}
-			}
+		if(Value.IsSet())
+		{
+			Result = Value;
 		}
 		return Result;
 	});
 	
 
-	auto OnNumericValueChanged = [this, InPropertyHandle](
+	auto OnNumericValueChanged = [this, InPropertyHandle, GetTransformComponent, SetTransformComponent](
 		ESlateTransformComponent::Type InTransformComponent, 
 		ESlateRotationRepresentation::Type InRotationRepresentation, 
 		ESlateTransformSubComponent::Type InSubComponent,
@@ -3259,41 +3382,25 @@ void FRigVMGraphMathTypeDetailCustomization::ConfigureTransformWidgetArgs(TShare
 		}
 		
 		URigVMController* Controller = nullptr;
-		if(BlueprintBeingCustomized && GraphBeingCustomized)
+		if (TStrongObjectPtr<URigVMBlueprint> Blueprint = BlueprintBeingCustomized.Pin(); GraphBeingCustomized.IsValid())
 		{
-			Controller = BlueprintBeingCustomized->GetController(GraphBeingCustomized);
+			Controller = BlueprintBeingCustomized->GetController(GraphBeingCustomized.Get());
 			if(bIsCommit)
 			{
 				Controller->OpenUndoBracket(FString::Printf(TEXT("Set %s"), *InPropertyHandle->GetProperty()->GetName()));
 			}
 		}
 
-		const EPropertyChangeType::Type ChangeType = bIsCommit ? EPropertyChangeType::ValueSet : EPropertyChangeType::Interactive;
-
-		const TArray<uint8*> AllMemoryBeingCustomized = GetMemoryBeingCustomized();
-		for (int32 Index = 0; Index < AllMemoryBeingCustomized.Num(); Index++)
-		{
-			const uint8* Memory = AllMemoryBeingCustomized[Index];
-			if (Memory != nullptr && InPropertyHandle->IsValidHandle())
-			{
-				TransformType& Transform = ContainerMemoryBlockToValueRef<TransformType>((uint8*)Memory, Identity, PropertyChain, PropertyArrayIndices);
-				TransformType PreviousTransform = Transform;
-
-				SAdvancedTransformInputBox<TransformType>::ApplyNumericValueChange(
-					Transform,
-					InValue,
-					InTransformComponent,
-					InRotationRepresentation,
-					InSubComponent);
-
-				if (!PreviousTransform.Equals(Transform))
-				{
-					InPropertyHandle->NotifyPostChange(ChangeType);
-				}
-			}
-		}
+		TransformType Transform = GetTransformComponent(InTransformComponent);
+		SAdvancedTransformInputBox<TransformType>::ApplyNumericValueChange(
+			Transform,
+			InValue,
+			InTransformComponent,
+			InRotationRepresentation,
+			InSubComponent);
+		SetTransformComponent(InTransformComponent, Transform);
 
-		if(Controller && bIsCommit)
+		if (Controller && bIsCommit)
 		{
 			Controller->CloseUndoBracket();
 		}
@@ -3318,13 +3425,8 @@ void FRigVMGraphMathTypeDetailCustomization::ConfigureTransformWidgetArgs(TShare
 		OnNumericValueChanged(InTransformComponent, InRotationRepresentation, InSubComponent, InValue, true, InCommitType);
 	});
 
-	WidgetArgs.OnResetToDefault_Lambda([this, DefaultValue, InPropertyHandle](ESlateTransformComponent::Type InTransformComponent)
+	WidgetArgs.OnResetToDefault_Lambda([this, DefaultValue, InPropertyHandle, SetTransformComponent](ESlateTransformComponent::Type InTransformComponent)
 	{
-		if (ObjectsBeingCustomized.IsEmpty() && StructsBeingCustomized.IsEmpty())
-		{
-			return;
-		}
-		
 		FEditPropertyChain PropertyChain;
 		TArray<int32> PropertyArrayIndices;
 		bool bEnabled;
@@ -3334,66 +3436,55 @@ void FRigVMGraphMathTypeDetailCustomization::ConfigureTransformWidgetArgs(TShare
 		}
 		
 		URigVMController* Controller = nullptr;
-		if(BlueprintBeingCustomized && GraphBeingCustomized)
+		if (TStrongObjectPtr<URigVMBlueprint> Blueprint = BlueprintBeingCustomized.Pin(); GraphBeingCustomized.IsValid())
 		{
-			Controller = BlueprintBeingCustomized->GetController(GraphBeingCustomized);
+			Controller = BlueprintBeingCustomized->GetController(GraphBeingCustomized.Get());
 			if(Controller)
 			{
 				Controller->OpenUndoBracket(FString::Printf(TEXT("Reset %s to Default"), *InPropertyHandle->GetProperty()->GetName()));
 			}
 		}
 
-		const EPropertyChangeType::Type ChangeType = EPropertyChangeType::ValueSet;
-
-		const TArray<uint8*> AllMemoryBeingCustomized = GetMemoryBeingCustomized();
-		for (int32 Index = 0; Index < AllMemoryBeingCustomized.Num(); Index++)
+		TransformType Transform;
+		switch(InTransformComponent)
 		{
-			const uint8* Memory = AllMemoryBeingCustomized[Index];
-			if (Memory != nullptr && InPropertyHandle->IsValidHandle())
+			case ESlateTransformComponent::Location:
 			{
-				TransformType& Transform = ContainerMemoryBlockToValueRef<TransformType>((uint8*)Memory, Identity, PropertyChain, PropertyArrayIndices);
-				TransformType PreviousTransform = Transform;
-
-				switch(InTransformComponent)
-				{
-					case ESlateTransformComponent::Location:
-					{
-						Transform.SetLocation(DefaultValue.GetLocation());
-						break;
-					}
-					case ESlateTransformComponent::Rotation:
-					{
-						Transform.SetRotation(DefaultValue.GetRotation());
-						break;
-					}
-					case ESlateTransformComponent::Scale:
-					{
-						Transform.SetScale3D(DefaultValue.GetScale3D());
-						break;
-					}
-					case ESlateTransformComponent::Max:
-					default:
-					{
-						Transform.SetLocation(DefaultValue.GetLocation());
-						break;
-					}
-				}
-				
-				if(!PreviousTransform.Equals(Transform))
-				{
-					InPropertyHandle->NotifyPostChange(ChangeType);
-				}
+				Transform.SetLocation(DefaultValue.GetLocation());
+				break;
+			}
+			case ESlateTransformComponent::Rotation:
+			{
+				Transform.SetRotation(DefaultValue.GetRotation());
+				break;
+			}
+			case ESlateTransformComponent::Scale:
+			{
+				Transform.SetScale3D(DefaultValue.GetScale3D());
+				break;
+			}
+			case ESlateTransformComponent::Max:
+			{
+				Transform.SetLocation(DefaultValue.GetLocation());
+				Transform.SetRotation(DefaultValue.GetRotation());
+				Transform.SetScale3D(DefaultValue.GetScale3D());
+				break;
+			}
+			default:
+			{
+				check(false);
 			}
 		}
-
+		SetTransformComponent(InTransformComponent, Transform);
+				
 		if(Controller)
 		{
 			Controller->CloseUndoBracket();
 		}
 	});
 
-	WidgetArgs.OnCopyToClipboard_Lambda([this, InPropertyHandle](
-		ESlateTransformComponent::Type InComponent
+	WidgetArgs.OnCopyToClipboard_Lambda([this, InPropertyHandle, GetTransformComponent](
+		ESlateTransformComponent::Type InTransformComponent
 		)
 	{
 		TOptional<FReal> Result;
@@ -3405,52 +3496,44 @@ void FRigVMGraphMathTypeDetailCustomization::ConfigureTransformWidgetArgs(TShare
 			return;
 		}
 
-		const TArray<uint8*> MemoryBlocks = GetMemoryBeingCustomized();
-		for(uint8* MemoryBlock: MemoryBlocks)
+		const TransformType Transform = GetTransformComponent(InTransformComponent);
+
+		FString Content;
+		switch(InTransformComponent)
 		{
-			if(MemoryBlock)
+			case ESlateTransformComponent::Location:
 			{
-				const TransformType& Transform = ContainerMemoryBlockToValueRef<TransformType>(MemoryBlock, Identity, PropertyChain, PropertyArrayIndices);
-				FString Content;
-				switch(InComponent)
-				{
-					case ESlateTransformComponent::Location:
-					{
-						const FVector Data = Transform.GetLocation();
-						TBaseStructure<FVector>::Get()->ExportText(Content, &Data, &Data, nullptr, PPF_None, nullptr);
-						break;
-					}
-					case ESlateTransformComponent::Rotation:
-					{
-						const FRotator Data = Transform.Rotator();
-						TBaseStructure<FRotator>::Get()->ExportText(Content, &Data, &Data, nullptr, PPF_None, nullptr);
-						break;
-					}
-					case ESlateTransformComponent::Scale:
-					{
-						const FVector Data = Transform.GetScale3D();
-						TBaseStructure<FVector>::Get()->ExportText(Content, &Data, &Data, nullptr, PPF_None, nullptr);
-						break;
-					}
-					case ESlateTransformComponent::Max:
-					default:
-					{
-						TBaseStructure<TransformType>::Get()->ExportText(Content, &Transform, &Transform, nullptr, PPF_None, nullptr);
-						break;
-					}
-				}
-
-				if(!Content.IsEmpty())
-				{
-					FPlatformApplicationMisc::ClipboardCopy(*Content);
-				}
+				const FVector Data = Transform.GetLocation();
+				TBaseStructure<FVector>::Get()->ExportText(Content, &Data, &Data, nullptr, PPF_None, nullptr);
+				break;
+			}
+			case ESlateTransformComponent::Rotation:
+			{
+				const FRotator Data = Transform.Rotator();
+				TBaseStructure<FRotator>::Get()->ExportText(Content, &Data, &Data, nullptr, PPF_None, nullptr);
+				break;
+			}
+			case ESlateTransformComponent::Scale:
+			{
+				const FVector Data = Transform.GetScale3D();
+				TBaseStructure<FVector>::Get()->ExportText(Content, &Data, &Data, nullptr, PPF_None, nullptr);
+				break;
+			}
+			case ESlateTransformComponent::Max:
+			default:
+			{
+				TBaseStructure<TransformType>::Get()->ExportText(Content, &Transform, &Transform, nullptr, PPF_None, nullptr);
+				break;
 			}
 		}
+
+		if(!Content.IsEmpty())
+		{
+			FPlatformApplicationMisc::ClipboardCopy(*Content);
+		}
 	});
 
-	WidgetArgs.OnPasteFromClipboard_Lambda([this, InPropertyHandle, OnNumericValueChanged](
-		ESlateTransformComponent::Type InComponent
-		)
+	WidgetArgs.OnPasteFromClipboard_Lambda([this, InPropertyHandle, SetTransformComponent](ESlateTransformComponent::Type InTransformComponent)
 	{
 		FString Content;
 		FPlatformApplicationMisc::ClipboardPaste(Content);
@@ -3460,11 +3543,6 @@ void FRigVMGraphMathTypeDetailCustomization::ConfigureTransformWidgetArgs(TShare
 			return;
 		}
 
-		if (ObjectsBeingCustomized.IsEmpty() && StructsBeingCustomized.IsEmpty())
-		{
-			return;
-		}
-
 		TOptional<FReal> Result;
 		FEditPropertyChain PropertyChain;
 		TArray<int32> PropertyArrayIndices;
@@ -3475,86 +3553,70 @@ void FRigVMGraphMathTypeDetailCustomization::ConfigureTransformWidgetArgs(TShare
 		}
 
 		URigVMController* Controller = nullptr;
-		if(BlueprintBeingCustomized && GraphBeingCustomized)
+		if (TStrongObjectPtr<URigVMBlueprint> Blueprint = BlueprintBeingCustomized.Pin(); GraphBeingCustomized.IsValid())
 		{
-			Controller = BlueprintBeingCustomized->GetController(GraphBeingCustomized);
+			Controller = BlueprintBeingCustomized->GetController(GraphBeingCustomized.Get());
 			Controller->OpenUndoBracket(FString::Printf(TEXT("Set %s"), *InPropertyHandle->GetProperty()->GetName()));
 		}
-		
-		const EPropertyChangeType::Type ChangeType = EPropertyChangeType::ValueSet;
 
-		const TArray<uint8*> AllMemoryBeingCustomized = GetMemoryBeingCustomized();
-		for (int32 Index = 0; Index < AllMemoryBeingCustomized.Num(); Index++)
+		// Apply the new value
 		{
-			const uint8* Memory = AllMemoryBeingCustomized[Index];
-			if (Memory != nullptr && InPropertyHandle->IsValidHandle())
+			class FRigPasteTransformWidgetErrorPipe : public FOutputDevice
 			{
-				TransformType& Transform = ContainerMemoryBlockToValueRef<TransformType>((uint8*)Memory, Identity, PropertyChain, PropertyArrayIndices);
-				const TransformType PreviousTransform = Transform;
-
-				// Apply the new value
-				{
-					class FRigPasteTransformWidgetErrorPipe : public FOutputDevice
-					{
-					public:
+			public:
 				
-						int32 NumErrors;
+				int32 NumErrors;
 				
-						FRigPasteTransformWidgetErrorPipe()
-							: FOutputDevice()
-							, NumErrors(0)
-						{
-						}
+				FRigPasteTransformWidgetErrorPipe()
+					: FOutputDevice()
+					, NumErrors(0)
+				{
+				}
 				
-						virtual void Serialize(const TCHAR* V, ELogVerbosity::Type Verbosity, const class FName& Category) override
-						{
-							UE_LOG(LogRigVM, Error, TEXT("Error Pasting to Widget: %s"), V);
-							NumErrors++;
-						}
-					};
+				virtual void Serialize(const TCHAR* V, ELogVerbosity::Type Verbosity, const class FName& Category) override
+				{
+					UE_LOG(LogRigVM, Error, TEXT("Error Pasting to Widget: %s"), V);
+					NumErrors++;
+				}
+			};
 				
-					FRigPasteTransformWidgetErrorPipe ErrorPipe;
-					
-					switch(InComponent)
-					{
-						case ESlateTransformComponent::Location:
-						{
-							FVector Data = Transform.GetLocation();
-							TBaseStructure<FVector>::Get()->ImportText(*Content, &Data, nullptr, PPF_None, &ErrorPipe, TBaseStructure<FVector>::Get()->GetName(), true);
-							Transform.SetLocation(Data);
-							break;
-						}
-						case ESlateTransformComponent::Rotation:
-						{
-							FRotator Data = Transform.Rotator();
-							TBaseStructure<FRotator>::Get()->ImportText(*Content, &Data, nullptr, PPF_None, &ErrorPipe, TBaseStructure<FRotator>::Get()->GetName(), true);
-							FQuat Quat = Data.Quaternion();
-							Transform.SetRotation(Quat);
+			FRigPasteTransformWidgetErrorPipe ErrorPipe;
+
+			TransformType Transform;
+			switch(InTransformComponent)
+			{
+				case ESlateTransformComponent::Location:
+				{
+					FVector Data = FVector::ZeroVector;
+					TBaseStructure<FVector>::Get()->ImportText(*Content, &Data, nullptr, PPF_None, &ErrorPipe, TBaseStructure<FVector>::Get()->GetName(), true);
+					Transform.SetLocation(Data);
+					break;
+				}
+				case ESlateTransformComponent::Rotation:
+				{
+					FRotator Data = FRotator::ZeroRotator;
+					TBaseStructure<FRotator>::Get()->ImportText(*Content, &Data, nullptr, PPF_None, &ErrorPipe, TBaseStructure<FRotator>::Get()->GetName(), true);
+					FQuat Quat = Data.Quaternion();
+					Transform.SetRotation(Quat);
 							
-							break;
-						}
-						case ESlateTransformComponent::Scale:
-						{
-							FVector Data = Transform.GetScale3D();
-							TBaseStructure<FVector>::Get()->ImportText(*Content, &Data, nullptr, PPF_None, &ErrorPipe, TBaseStructure<FVector>::Get()->GetName(), true);
-							Transform.SetScale3D(Data);
-							break;
-						}
-						case ESlateTransformComponent::Max:
-						default:
-						{
-							TBaseStructure<TransformType>::Get()->ImportText(*Content, &Transform, nullptr, PPF_None, &ErrorPipe, TBaseStructure<TransformType>::Get()->GetName(), true);
-							break;
-						}
-					}
-				
-					if(ErrorPipe.NumErrors == 0 && !PreviousTransform.Equals(Transform))
-					{
-						InPropertyHandle->NotifyPostChange(ChangeType);
-					}
+					break;
+				}
+				case ESlateTransformComponent::Scale:
+				{
+					FVector Data = FVector::OneVector;
+					TBaseStructure<FVector>::Get()->ImportText(*Content, &Data, nullptr, PPF_None, &ErrorPipe, TBaseStructure<FVector>::Get()->GetName(), true);
+					Transform.SetScale3D(Data);
+					break;
+				}
+				case ESlateTransformComponent::Max:
+				default:
+				{
+					TBaseStructure<TransformType>::Get()->ImportText(*Content, &Transform, nullptr, PPF_None, &ErrorPipe, TBaseStructure<TransformType>::Get()->GetName(), true);
+					break;
 				}
 			}
-		}
+			SetTransformComponent(InTransformComponent, Transform);
+		}		
 		
 		if (Controller)
 		{
@@ -3564,31 +3626,27 @@ void FRigVMGraphMathTypeDetailCustomization::ConfigureTransformWidgetArgs(TShare
 }
 
 template<typename TransformType>
-void FRigVMGraphMathTypeDetailCustomization::MakeTransformHeaderRow(TSharedRef<class IPropertyHandle> InPropertyHandle, FDetailWidgetRow& HeaderRow, IPropertyTypeCustomizationUtils& StructCustomizationUtils)
+void FRigVMGraphMathTypeDetailCustomization::MakeTransformHeaderRow(TSharedRef<class IPropertyHandle> InPropertyHandle, FDetailWidgetRow& HeaderRow, IPropertyTypeCustomizationUtils& StructCustomizationUtils, TConstArrayView<FName> ComponentNames)
 {
 	typename SAdvancedTransformInputBox<TransformType>::FArguments WidgetArgs;
-	ConfigureTransformWidgetArgs<TransformType>(InPropertyHandle, WidgetArgs);
+	ConfigureTransformWidgetArgs<TransformType>(InPropertyHandle, WidgetArgs, ComponentNames);
 
 	SAdvancedTransformInputBox<TransformType>::ConfigureHeader(HeaderRow, InPropertyHandle->GetPropertyDisplayName(), InPropertyHandle->GetToolTipText(), WidgetArgs);
 	SAdvancedTransformInputBox<TransformType>::ConfigureComponentWidgetRow(HeaderRow, ESlateTransformComponent::Max, WidgetArgs);
 }
 	
 template<typename TransformType>
-void FRigVMGraphMathTypeDetailCustomization::MakeTransformChildren(TSharedRef<class IPropertyHandle> InPropertyHandle, IDetailChildrenBuilder& StructBuilder, IPropertyTypeCustomizationUtils& StructCustomizationUtils)
+void FRigVMGraphMathTypeDetailCustomization::MakeTransformChildren(TSharedRef<class IPropertyHandle> InPropertyHandle, IDetailChildrenBuilder& StructBuilder, IPropertyTypeCustomizationUtils& StructCustomizationUtils, TConstArrayView<FName> ComponentNames)
 {
-	FDetailWidgetRow* LocationRow = nullptr;
-	FDetailWidgetRow* RotationRow = nullptr;
-	FDetailWidgetRow* ScaleRow = nullptr;
-
 	typename SAdvancedTransformInputBox<TransformType>::FArguments WidgetArgs;
-	ConfigureTransformWidgetArgs<TransformType>(InPropertyHandle, WidgetArgs);
+	ConfigureTransformWidgetArgs<TransformType>(InPropertyHandle, WidgetArgs, ComponentNames);
 
 	const FString StandaloneWidgetMetadata = InPropertyHandle->GetMetaData(TEXT("SStandaloneCustomizedValueWidget"));
 	const bool IsUsingStandaloneWidget = StandaloneWidgetMetadata.Equals(TEXT("True"));
 	
-	LocationRow = &StructBuilder.AddCustomRow(LOCTEXT("TransformLocation", "Location"));
-	RotationRow = &StructBuilder.AddCustomRow(LOCTEXT("TransformRotation", "Rotation"));
-	ScaleRow = &StructBuilder.AddCustomRow(LOCTEXT("TransformScale", "Scale"));
+	FDetailWidgetRow* LocationRow = &StructBuilder.AddCustomRow(LOCTEXT("TransformLocation", "Location"));
+	FDetailWidgetRow* RotationRow = &StructBuilder.AddCustomRow(LOCTEXT("TransformRotation", "Rotation"));
+	FDetailWidgetRow* ScaleRow = &StructBuilder.AddCustomRow(LOCTEXT("TransformScale", "Scale"));
 	
 	SAdvancedTransformInputBox<TransformType>::ConfigureComponentWidgetRow(*LocationRow, ESlateTransformComponent::Location, WidgetArgs);
 	SAdvancedTransformInputBox<TransformType>::ConfigureComponentWidgetRow(*RotationRow, ESlateTransformComponent::Rotation, WidgetArgs);
@@ -3596,9 +3654,8 @@ void FRigVMGraphMathTypeDetailCustomization::MakeTransformChildren(TSharedRef<cl
 }
 
 FRigVMGraphMathTypeDetailCustomization::FRigVMGraphMathTypeDetailCustomization()
-: ScriptStruct(nullptr)
-, BlueprintBeingCustomized(nullptr)
-, GraphBeingCustomized(nullptr)
+	: BlueprintBeingCustomized(nullptr)
+	, GraphBeingCustomized(nullptr)
 {
 }
 
@@ -3608,13 +3665,8 @@ void FRigVMGraphMathTypeDetailCustomization::CustomizeHeader(TSharedRef<IPropert
 	TArray<UObject*> Objects;
 	InPropertyHandle->GetOuterObjects(Objects);
 
-	StructsBeingCustomized.Reset();
-	InPropertyHandle->GetOuterStructs(StructsBeingCustomized);
-
 	for (UObject* Object : Objects)
 	{
-		ObjectsBeingCustomized.Add(Object);
-
 		if(BlueprintBeingCustomized == nullptr)
 		{
 			BlueprintBeingCustomized = Object->GetTypedOuter<URigVMBlueprint>();
@@ -3626,9 +3678,9 @@ void FRigVMGraphMathTypeDetailCustomization::CustomizeHeader(TSharedRef<IPropert
 		}
 	}
 
-	FProperty* Property = InPropertyHandle->GetProperty();
+	const FProperty* Property = InPropertyHandle->GetProperty();
 	const FStructProperty* StructProperty = CastField<FStructProperty>(Property);
-	ScriptStruct = StructProperty->Struct;
+	const UScriptStruct* ScriptStruct = StructProperty->Struct;
 
 	if(ScriptStruct == TBaseStructure<FVector>::Get())
 	{
@@ -3652,11 +3704,11 @@ void FRigVMGraphMathTypeDetailCustomization::CustomizeHeader(TSharedRef<IPropert
 	}
 	else if(ScriptStruct == TBaseStructure<FTransform>::Get())
 	{
-		MakeTransformHeaderRow<FTransform>(InPropertyHandle, HeaderRow, StructCustomizationUtils);
+		MakeTransformHeaderRow<FTransform>(InPropertyHandle, HeaderRow, StructCustomizationUtils, TransformComponentNames);
 	}
 	else if(ScriptStruct == TBaseStructure<FEulerTransform>::Get())
 	{
-		MakeTransformHeaderRow<FEulerTransform>(InPropertyHandle, HeaderRow, StructCustomizationUtils);
+		MakeTransformHeaderRow<FEulerTransform>(InPropertyHandle, HeaderRow, StructCustomizationUtils, EulerTransformComponentNames);
 	}
 }
 
@@ -3673,8 +3725,6 @@ void FRigVMGraphMathTypeDetailCustomization::CustomizeChildren(TSharedRef<IPrope
 
 	for (UObject* Object : Objects)
 	{
-		ObjectsBeingCustomized.Add(Object);
-
 		if(BlueprintBeingCustomized == nullptr)
 		{
 			BlueprintBeingCustomized = Object->GetTypedOuter<URigVMBlueprint>();
@@ -3686,20 +3736,17 @@ void FRigVMGraphMathTypeDetailCustomization::CustomizeChildren(TSharedRef<IPrope
 		}
 	}
 
-	StructsBeingCustomized.Reset();
-	InPropertyHandle->GetOuterStructs(StructsBeingCustomized);
-
-	FProperty* Property = InPropertyHandle->GetProperty();
+	const FProperty* Property = InPropertyHandle->GetProperty();
 	const FStructProperty* StructProperty = CastField<FStructProperty>(Property);
-	ScriptStruct = StructProperty->Struct;
+	const UScriptStruct* ScriptStruct = StructProperty->Struct;
 
 	if(ScriptStruct == TBaseStructure<FTransform>::Get())
 	{
-		MakeTransformChildren<FTransform>(InPropertyHandle, StructBuilder, StructCustomizationUtils);
+		MakeTransformChildren<FTransform>(InPropertyHandle, StructBuilder, StructCustomizationUtils, TransformComponentNames);
 	}
 	else if(ScriptStruct == TBaseStructure<FEulerTransform>::Get())
 	{
-		MakeTransformChildren<FEulerTransform>(InPropertyHandle, StructBuilder, StructCustomizationUtils);
+		MakeTransformChildren<FEulerTransform>(InPropertyHandle, StructBuilder, StructCustomizationUtils, EulerTransformComponentNames);
 	}
 }
 
diff --git a/Engine/Plugins/Runtime/RigVM/Source/RigVMEditor/Public/Editor/RigVMGraphDetailCustomization.h b/Engine/Plugins/Runtime/RigVM/Source/RigVMEditor/Public/Editor/RigVMGraphDetailCustomization.h
--- a/Engine/Plugins/Runtime/RigVM/Source/RigVMEditor/Public/Editor/RigVMGraphDetailCustomization.h
+++ b/Engine/Plugins/Runtime/RigVM/Source/RigVMEditor/Public/Editor/RigVMGraphDetailCustomization.h
@@ -448,30 +448,6 @@ public:
 
 protected:
 
-	TArray<uint8*> GetMemoryBeingCustomized()
-	{
-		TArray<uint8*> MemoryPtr;
-		MemoryPtr.Reserve(ObjectsBeingCustomized.Num() + StructsBeingCustomized.Num());
-
-		for(const TWeakObjectPtr<UObject>& Object : ObjectsBeingCustomized)
-		{
-			if(Object.IsValid())
-			{
-				MemoryPtr.Add((uint8*)Object.Get());
-			}
-		}
-
-		for(const TSharedPtr<FStructOnScope>& StructPtr: StructsBeingCustomized)
-		{
-			if(StructPtr.IsValid())
-			{
-				MemoryPtr.Add(StructPtr->GetStructMemory());
-			}
-		}
-
-		return MemoryPtr;
-	}
-	
 	bool GetPropertyChain(TSharedRef<class IPropertyHandle> InPropertyHandle, FEditPropertyChain& OutPropertyChain, TArray<int32> &OutPropertyArrayIndices, bool& bOutEnabled)
 	{
 		if (!InPropertyHandle->IsValidHandle())
@@ -483,10 +459,7 @@ protected:
 		OutPropertyArrayIndices.Reset();
 		bOutEnabled = false;
 
-		const bool bHasObject = !ObjectsBeingCustomized.IsEmpty() && ObjectsBeingCustomized[0].Get();
-		const bool bHasStruct = !StructsBeingCustomized.IsEmpty() && StructsBeingCustomized[0].Get();
-		
-		if (bHasStruct || bHasObject)
+		if (InPropertyHandle->GetNumPerObjectValues() > 0)
 		{
 			TSharedPtr<class IPropertyHandle> ChainHandle = InPropertyHandle;
 			while (ChainHandle.IsValid() && ChainHandle->GetProperty() != nullptr)
@@ -569,27 +542,18 @@ protected:
 			return Result;
 		}
 
-		const TArray<uint8*> MemoryBlocks = GetMemoryBeingCustomized();
-		for(uint8* MemoryBlock: MemoryBlocks)
+		if (TSharedPtr<IPropertyHandle> ChildHandle = InPropertyHandle->GetChildHandle(InComponent))
 		{
-			if(MemoryBlock)
+			if (ChildHandle->IsValidHandle())
 			{
-				static VectorType ZeroVector = VectorType();
-				const VectorType& Vector = ContainerMemoryBlockToValueRef<VectorType>(MemoryBlock, ZeroVector, PropertyChain, PropertyArrayIndices);
-				NumericType Component = Vector[InComponent];
-				if(Result.IsSet())
-				{
-					if(!FMath::IsNearlyEqual(Result.GetValue(), Component))
-					{
-						return TOptional<NumericType>();
-					}
-				}
-				else
+				NumericType Value;
+				if (ChildHandle->GetValue(Value) == FPropertyAccess::Success) // note that this will fail if multiple values
 				{
-					Result = Component;
+					Result = Value;
 				}
 			}
 		}
+
 		return Result;
 	};
 
@@ -597,46 +561,21 @@ protected:
 	template<typename VectorType, typename NumericType>
 	void OnVectorComponentChanged(TSharedRef<class IPropertyHandle> InPropertyHandle, int32 InComponent, NumericType InValue, bool bIsCommit, ETextCommit::Type InCommitType = ETextCommit::Default)
 	{
-		if (ObjectsBeingCustomized.IsEmpty() && StructsBeingCustomized.IsEmpty())
-		{
-			return;
-		}
-
-		FEditPropertyChain PropertyChain;
-		TArray<int32> PropertyArrayIndices;
-		bool bEnabled;
-		if (!GetPropertyChain(InPropertyHandle, PropertyChain, PropertyArrayIndices, bEnabled))
-		{
-			return;
-		}
-
 		URigVMController* Controller = nullptr;
-		if (BlueprintBeingCustomized && GraphBeingCustomized)
+		if (TStrongObjectPtr<URigVMBlueprint> Blueprint = BlueprintBeingCustomized.Pin(); GraphBeingCustomized.IsValid())
 		{
-			Controller = BlueprintBeingCustomized->GetController(GraphBeingCustomized);
+			Controller = BlueprintBeingCustomized->GetController(GraphBeingCustomized.Get());
 			if (bIsCommit)
 			{
 				Controller->OpenUndoBracket(FString::Printf(TEXT("Set %s"), *InPropertyHandle->GetProperty()->GetName()));
 			}
 		}
 
-		const EPropertyChangeType::Type ChangeType = bIsCommit ? EPropertyChangeType::ValueSet : EPropertyChangeType::Interactive;
-
-		const TArray<uint8*> AllMemoryBeingCustomized = GetMemoryBeingCustomized();
-		for (int32 Index = 0; Index < AllMemoryBeingCustomized.Num(); Index++)
+		if (TSharedPtr<IPropertyHandle> ChildHandle = InPropertyHandle->GetChildHandle(InComponent))
 		{
-			const uint8* Memory = AllMemoryBeingCustomized[Index];
-			if (Memory != nullptr && InPropertyHandle->IsValidHandle())
+			if (ChildHandle->IsValidHandle())
 			{
-				static VectorType ZeroVector = VectorType();
-				VectorType& Vector = ContainerMemoryBlockToValueRef<VectorType>((uint8*)Memory, ZeroVector, PropertyChain, PropertyArrayIndices);
-				VectorType PreviousVector = Vector;
-				Vector[InComponent] = InValue;
-
-				if (!PreviousVector.Equals(Vector))
-				{
-					InPropertyHandle->NotifyPostChange(ChangeType);
-				}
+				ChildHandle->SetValue(InValue);
 			}
 		}
 
@@ -657,13 +596,13 @@ protected:
 	void MakeRotationHeaderRow(TSharedRef<class IPropertyHandle> InPropertyHandle, FDetailWidgetRow& HeaderRow, IPropertyTypeCustomizationUtils& StructCustomizationUtils);
 
 	template<typename TransformType>
-	void ConfigureTransformWidgetArgs(TSharedRef<class IPropertyHandle> InPropertyHandle, typename SAdvancedTransformInputBox<TransformType>::FArguments& WidgetArgs);
+	void ConfigureTransformWidgetArgs(TSharedRef<class IPropertyHandle> InPropertyHandle, typename SAdvancedTransformInputBox<TransformType>::FArguments& WidgetArgs, TConstArrayView<FName> ComponentNames);
 
 	template<typename TransformType>
-	void MakeTransformHeaderRow(TSharedRef<class IPropertyHandle> InPropertyHandle, FDetailWidgetRow& HeaderRow, IPropertyTypeCustomizationUtils& StructCustomizationUtils);
+	void MakeTransformHeaderRow(TSharedRef<class IPropertyHandle> InPropertyHandle, FDetailWidgetRow& HeaderRow, IPropertyTypeCustomizationUtils& StructCustomizationUtils, TConstArrayView<FName> ComponentNames);
 
 	template<typename TransformType>
-	void MakeTransformChildren(TSharedRef<class IPropertyHandle> InPropertyHandle, IDetailChildrenBuilder& StructBuilder, IPropertyTypeCustomizationUtils& StructCustomizationUtils);
+	void MakeTransformChildren(TSharedRef<class IPropertyHandle> InPropertyHandle, IDetailChildrenBuilder& StructBuilder, IPropertyTypeCustomizationUtils& StructCustomizationUtils, TConstArrayView<FName> ComponentNames);
 
 	// returns the rotation for rotator or quaternions (or empty optional for varying values)
 	template<typename RotationType>
@@ -678,26 +617,15 @@ protected:
 			return Result;
 		}
 
-		const TArray<uint8*> MemoryBlocks = GetMemoryBeingCustomized();
-		for(uint8* MemoryBlock: MemoryBlocks)
+		if (InPropertyHandle->IsValidHandle())
 		{
-			if(MemoryBlock)
+			RotationType Value;
+			if (InPropertyHandle->GetValue(Value) == FPropertyAccess::Success) // note that this will fail if multiple values
 			{
-				static RotationType ZeroRotation = RotationType();
-				const RotationType& Rotation = ContainerMemoryBlockToValueRef<RotationType>(MemoryBlock, ZeroRotation, PropertyChain, PropertyArrayIndices);
-				if(Result.IsSet())
-				{
-					if(!Rotation.Equals(Result.GetValue()))
-					{
-						return TOptional<RotationType>();
-					}
-				}
-				else
-				{
-					Result = Rotation;
-				}
+				Result = Value;
 			}
 		}
+
 		return Result;
 	};
 
@@ -705,10 +633,6 @@ protected:
 	template<typename RotationType>
 	void OnRotationChanged(TSharedRef<class IPropertyHandle> InPropertyHandle, RotationType InValue, bool bIsCommit, ETextCommit::Type InCommitType = ETextCommit::Default)
 	{
-		if (ObjectsBeingCustomized.IsEmpty() && StructsBeingCustomized.IsEmpty())
-		{
-			return;
-		}
 		
 		FEditPropertyChain PropertyChain;
         TArray<int32> PropertyArrayIndices;
@@ -719,33 +643,18 @@ protected:
         }
 	
 		URigVMController* Controller = nullptr;
-		if(BlueprintBeingCustomized && GraphBeingCustomized)
+		if (TStrongObjectPtr<URigVMBlueprint> Blueprint = BlueprintBeingCustomized.Pin(); GraphBeingCustomized.IsValid())
 		{
-			Controller = BlueprintBeingCustomized->GetController(GraphBeingCustomized);
+			Controller = BlueprintBeingCustomized->GetController(GraphBeingCustomized.Get());
 			if(bIsCommit)
 			{
 				Controller->OpenUndoBracket(FString::Printf(TEXT("Set %s"), *InPropertyHandle->GetProperty()->GetName()));
 			}
 		}
 
-		const EPropertyChangeType::Type ChangeType = bIsCommit ? EPropertyChangeType::ValueSet : EPropertyChangeType::Interactive;
-
-		const TArray<uint8*> AllMemoryBeingCustomized = GetMemoryBeingCustomized();
-		for (int32 Index = 0; Index < AllMemoryBeingCustomized.Num(); Index++)
+		if (InPropertyHandle->IsValidHandle())
 		{
-			const uint8* Memory = AllMemoryBeingCustomized[Index];
-			if (Memory != nullptr && InPropertyHandle->IsValidHandle())
-			{
-				static RotationType ZeroRotation = RotationType();
-				RotationType& Rotation = ContainerMemoryBlockToValueRef<RotationType>((uint8*)Memory, ZeroRotation, PropertyChain, PropertyArrayIndices);
-				RotationType PreviousRotation = Rotation;
-				Rotation = InValue;
-
-				if (!PreviousRotation.Equals(Rotation))
-				{
-					InPropertyHandle->NotifyPostChange(ChangeType);
-				}
-			}
+			InPropertyHandle->SetValue(InValue);
 		}
 
 		if(Controller && bIsCommit)
@@ -758,11 +667,8 @@ protected:
 	template<typename RotationType>
 	void ExtendRotationArgs(TSharedRef<class IPropertyHandle> InPropertyHandle, void* ArgumentsPtr) {}
 		
-	UScriptStruct* ScriptStruct;
-	URigVMBlueprint* BlueprintBeingCustomized;
-	URigVMGraph* GraphBeingCustomized;
-	TArray<TWeakObjectPtr<UObject>> ObjectsBeingCustomized;
-	TArray<TSharedPtr<FStructOnScope>> StructsBeingCustomized;
+	TWeakObjectPtr<URigVMBlueprint> BlueprintBeingCustomized;
+	TWeakObjectPtr<URigVMGraph> GraphBeingCustomized;
 };
 
 template<>
diff --git a/Engine/Source/Editor/PropertyEditor/Private/PropertyHandleImpl.cpp b/Engine/Source/Editor/PropertyEditor/Private/PropertyHandleImpl.cpp
--- a/Engine/Source/Editor/PropertyEditor/Private/PropertyHandleImpl.cpp
+++ b/Engine/Source/Editor/PropertyEditor/Private/PropertyHandleImpl.cpp
@@ -5205,6 +5205,34 @@ FPropertyAccess::Result FPropertyHandleRotator::SetValue( const FRotator& NewVal
 	}
 }
 
+FPropertyAccess::Result FPropertyHandleRotator::GetValue(FQuat& OutValue) const
+{
+	if (!RollValue || !PitchValue || !YawValue)
+	{
+		return FPropertyAccess::Fail;
+	}
+
+	FRotator RotatorProxy;
+	FPropertyAccess::Result Res = GetValue(RotatorProxy);
+	if (Res == FPropertyAccess::Success)
+	{
+		OutValue = RotatorProxy.Quaternion();
+	}
+
+	return Res;
+}
+
+FPropertyAccess::Result FPropertyHandleRotator::SetValue(const FQuat& InValue, EPropertyValueSetFlags::Type Flags)
+{
+	if (!RollValue || !PitchValue || !YawValue)
+	{
+		return FPropertyAccess::Fail;
+	}
+
+	const FRotator RotatorProxy(InValue);
+	return SetValue(RotatorProxy, Flags);
+}
+
 FPropertyAccess::Result FPropertyHandleRotator::SetRoll( double InRoll, EPropertyValueSetFlags::Type Flags )
 {
 	if (!RollValue)
diff --git a/Engine/Source/Editor/PropertyEditor/Private/PropertyHandleImpl.h b/Engine/Source/Editor/PropertyEditor/Private/PropertyHandleImpl.h
--- a/Engine/Source/Editor/PropertyEditor/Private/PropertyHandleImpl.h
+++ b/Engine/Source/Editor/PropertyEditor/Private/PropertyHandleImpl.h
@@ -692,6 +692,8 @@ public:
 	static bool Supports( TSharedRef<FPropertyNode> PropertyNode );
 	virtual FPropertyAccess::Result GetValue( FRotator& OutValue ) const override;
 	virtual FPropertyAccess::Result SetValue( const FRotator& InValue, EPropertyValueSetFlags::Type Flags = EPropertyValueSetFlags::DefaultFlags ) override;
+	virtual FPropertyAccess::Result GetValue(FQuat& OutValue) const override;
+	virtual FPropertyAccess::Result SetValue(const FQuat& InValue, EPropertyValueSetFlags::Type Flags = EPropertyValueSetFlags::DefaultFlags) override;
 
 	virtual FPropertyAccess::Result SetRoll( double InValue, EPropertyValueSetFlags::Type Flags = EPropertyValueSetFlags::DefaultFlags );
 	virtual FPropertyAccess::Result SetPitch( double InValue, EPropertyValueSetFlags::Type Flags = EPropertyValueSetFlags::DefaultFlags );
