Avoid runtime FName creation, where possible, at RigVMPropertyUtils
Removed unused strings at RigVMExternalVariable header

diff --git a/Engine/Plugins/Runtime/RigVM/Source/RigVM/Private/RigVMPropertyUtils.cpp b/Engine/Plugins/Runtime/RigVM/Source/RigVM/Private/RigVMPropertyUtils.cpp
--- a/Engine/Plugins/Runtime/RigVM/Source/RigVM/Private/RigVMPropertyUtils.cpp
+++ b/Engine/Plugins/Runtime/RigVM/Source/RigVM/Private/RigVMPropertyUtils.cpp
@@ -8,56 +8,71 @@
 #include "Misc/AssertionMacros.h"
 #include "UObject/TextProperty.h"
 
+namespace
+{
+	static const FLazyName TypeName_Bool(TEXT("bool"));
+	static const FLazyName TypeName_Int32(TEXT("int32"));
+	static const FLazyName TypeName_Int64(TEXT("int64"));
+	static const FLazyName TypeName_UInt32(TEXT("uint32"));
+	static const FLazyName TypeName_UInt64(TEXT("uint64"));
+	static const FLazyName TypeName_Float(TEXT("float"));
+	static const FLazyName TypeName_Double(TEXT("double"));
+	static const FLazyName TypeName_FString(TEXT("FString"));
+	static const FLazyName TypeName_FName(TEXT("FName"));
+	static const FLazyName TypeName_FText(TEXT("FText"));
+	static const FLazyName TypeName_UInt8(TEXT("uint8"));
+}
+
 void RigVMPropertyUtils::GetTypeFromProperty(const FProperty* InProperty, FName& OutTypeName, UObject*& OutTypeObject)
 {
 	if (CastField<FBoolProperty>(InProperty))
 	{
-		OutTypeName = TEXT("bool");
+		OutTypeName = TypeName_Bool;
 		OutTypeObject = nullptr;
 	}
 	else if (CastField<FIntProperty>(InProperty))
 	{
-		OutTypeName = TEXT("int32");
+		OutTypeName = TypeName_Int32;
 		OutTypeObject = nullptr;
 	}
 	else if (CastField<FInt64Property>(InProperty))
 	{
-		OutTypeName = TEXT("int64");
+		OutTypeName = TypeName_Int64;
 		OutTypeObject = nullptr;
 	}
 	else if (CastField<FUInt32Property>(InProperty))
 	{
-		OutTypeName = TEXT("uint32");
+		OutTypeName = TypeName_UInt32;
 		OutTypeObject = nullptr;
 	}
 	else if (CastField<FUInt64Property>(InProperty))
 	{
-		OutTypeName = TEXT("uint64");
+		OutTypeName = TypeName_UInt64;
 		OutTypeObject = nullptr;
 	}
 	else if (CastField<FFloatProperty>(InProperty))
 	{
-		OutTypeName = TEXT("float");
+		OutTypeName = TypeName_Float;
 		OutTypeObject = nullptr;
 	}
 	else if (CastField<FDoubleProperty>(InProperty))
 	{
-		OutTypeName = TEXT("double");
+		OutTypeName = TypeName_Double;
 		OutTypeObject = nullptr;
 	}
 	else if (CastField<FStrProperty>(InProperty))
 	{
-		OutTypeName = TEXT("FString");
+		OutTypeName = TypeName_FString;
 		OutTypeObject = nullptr;
 	}
 	else if (CastField<FNameProperty>(InProperty))
 	{
-		OutTypeName = TEXT("FName");
+		OutTypeName = TypeName_FName;
 		OutTypeObject = nullptr;
 	}
 	else if (CastField<FTextProperty>(InProperty))
 	{
-		OutTypeName = TEXT("FText");
+		OutTypeName = TypeName_FText;
 		OutTypeObject = nullptr;
 	}
 	else if (const FEnumProperty* EnumProperty = CastField<FEnumProperty>(InProperty))
@@ -74,7 +89,7 @@ void RigVMPropertyUtils::GetTypeFromProperty(const FProperty* InProperty, FName&
 		}
 		else
 		{
-			OutTypeName = TEXT("uint8");
+			OutTypeName = TypeName_UInt8;
 			OutTypeObject = nullptr;
 		}
 	}
diff --git a/Engine/Plugins/Runtime/RigVM/Source/RigVM/Public/RigVMCore/RigVMExternalVariable.h b/Engine/Plugins/Runtime/RigVM/Source/RigVM/Public/RigVMCore/RigVMExternalVariable.h
--- a/Engine/Plugins/Runtime/RigVM/Source/RigVM/Public/RigVMCore/RigVMExternalVariable.h
+++ b/Engine/Plugins/Runtime/RigVM/Source/RigVM/Public/RigVMCore/RigVMExternalVariable.h
@@ -110,12 +110,9 @@ struct RIGVM_API FRigVMExternalVariable : public FRigVMExternalVariableDef
 			ExternalVariable.Memory = (uint8*)Property->ContainerPtrToValuePtr<uint8>(InContainer);
 		}
 
-		FString TypePrefix, TypeSuffix;
 		if (const FArrayProperty* ArrayProperty = CastField<FArrayProperty>(Property))
 		{
 			ExternalVariable.bIsArray = true;
-			TypePrefix = TEXT("TArray<");
-			TypeSuffix = TEXT(">");
 			Property = ArrayProperty->Inner;
 		}
 
