RigVM Model Optimizations - Introduce front end relevant caches

#rnx

diff --git a/Engine/Plugins/Runtime/RigVM/Source/RigVMDeveloper/Private/RigVMModel/Nodes/RigVMUnitNode.cpp b/Engine/Plugins/Runtime/RigVM/Source/RigVMDeveloper/Private/RigVMModel/Nodes/RigVMUnitNode.cpp
--- a/Engine/Plugins/Runtime/RigVM/Source/RigVMDeveloper/Private/RigVMModel/Nodes/RigVMUnitNode.cpp
+++ b/Engine/Plugins/Runtime/RigVM/Source/RigVMDeveloper/Private/RigVMModel/Nodes/RigVMUnitNode.cpp
@@ -78,13 +78,21 @@ bool URigVMUnitNode::IsDefinedAsVarying() const
 
 FName URigVMUnitNode::GetEventName() const
 {
+	if(!CachedEventName.IsBound())
+	{
+		CachedEventName.Bind(this);
+	}
+	if(CachedEventName.IsValid())
+	{
+		return CachedEventName.GetValue();
+	}
 	TSharedPtr<FStructOnScope> StructOnScope = ConstructStructInstance(false);
 	if (StructOnScope.IsValid())
 	{
 		const FRigVMStruct* StructMemory = (FRigVMStruct*)StructOnScope->GetStructMemory();
-		return StructMemory->GetEventName();
+		CachedEventName = StructMemory->GetEventName();
 	}
-	return Super::GetEventName();
+	return CachedEventName.Get(Super::GetEventName());
 }
 
 bool URigVMUnitNode::CanOnlyExistOnce() const
diff --git a/Engine/Plugins/Runtime/RigVM/Source/RigVMDeveloper/Private/RigVMModel/RigVMController.cpp b/Engine/Plugins/Runtime/RigVM/Source/RigVMDeveloper/Private/RigVMModel/RigVMController.cpp
--- a/Engine/Plugins/Runtime/RigVM/Source/RigVMDeveloper/Private/RigVMModel/RigVMController.cpp
+++ b/Engine/Plugins/Runtime/RigVM/Source/RigVMDeveloper/Private/RigVMModel/RigVMController.cpp
@@ -3740,6 +3740,7 @@ URigVMInjectionInfo* URigVMController::InjectNodeIntoPin(URigVMPin* InPin, bool
 		InjectionInfo->OutputPin = OutputPin;
 	
 		InPin->InjectionInfos.Add(InjectionInfo);
+		InPin->IncrementVersion();
 
 		Notify(ERigVMGraphNotifType::NodeAdded, NodeToInject);
 	}
@@ -3881,6 +3882,7 @@ URigVMNode* URigVMController::EjectNodeFromPin(URigVMPin* InPin, bool bSetupUndo
 	RenameObject(NodeToEject, nullptr, Graph);
 	SetNodePosition(NodeToEject, Position, false);
 	InPin->InjectionInfos.Remove(Injection);
+	InPin->IncrementVersion();
 	DestroyObject(Injection);
 
 	Notify(ERigVMGraphNotifType::NodeAdded, NodeToEject);
@@ -7225,6 +7227,7 @@ bool URigVMController::RenameNode(URigVMNode* InNode, const FName& InNewName, bo
 		return false;
 	}
 
+	InNode->IncrementVersion();
 	Notify(ERigVMGraphNotifType::FunctionRenamed, InNode);
 	Notify(ERigVMGraphNotifType::NodeRenamed, InNode);
 
@@ -8322,6 +8325,8 @@ bool URigVMController::SetPinDisplayName(URigVMPin* InPin, const FString& InDisp
 	}
 
 	InPin->DisplayName = InDisplayName.IsEmpty() ? FName(NAME_None) : FName(*InDisplayName);
+	InPin->IncrementVersion();
+	
 	Notify(ERigVMGraphNotifType::PinDisplayNameChanged, InPin);
 
 	if(const URigVMLibraryNode* LibraryNode = Cast<URigVMLibraryNode>(InPin->GetNode()))
@@ -8677,6 +8682,7 @@ bool URigVMController::SetPinCategory(URigVMPin* InPin, const FString& InCategor
 
 	InPin->UserDefinedCategory = NewCategory;
 	InPin->IndexInCategory = IndexInCategory;
+	InPin->IncrementVersion();
 	Notify(ERigVMGraphNotifType::PinCategoryChanged, InPin);
 
 	if(const URigVMLibraryNode* LibraryNode = Cast<URigVMLibraryNode>(InPin->GetNode()))
@@ -9401,6 +9407,7 @@ bool URigVMController::SetPinIndexInCategory(URigVMPin* InPin, int32 InIndexInCa
 	else
 	{
 		InPin->IndexInCategory = InIndexInCategory;
+		InPin->IncrementVersion();
 		Notify(ERigVMGraphNotifType::PinCategoryChanged, InPin);
 	}
 
@@ -9974,12 +9981,14 @@ bool URigVMController::SetPinDefaultValue(URigVMPin* InPin, const FString& InDef
 		if(InPin->DefaultValueType != NewDefaultValueType)
 		{
 			InPin->DefaultValueType = NewDefaultValueType;
+			InPin->IncrementVersion();
 			bSetPinDefaultValueSucceeded = true;
 		}
 		// no need to send notifications if not changing the value
 		if(InPin->DefaultValue != ClampedDefaultValue)
 		{
 			InPin->DefaultValue = ClampedDefaultValue;
+			InPin->IncrementVersion();
 			
 			// root trait pin store their default value in a separate property bag so that
 			// things like soft object ptr can be used and tracked in a uproperty
@@ -10928,6 +10937,7 @@ URigVMPin* URigVMController::InsertArrayPin(URigVMPin* ArrayPin, int32 InIndex,
 	ConfigurePinFromPin(Pin, ArrayPin);
 	Pin->CPPType = ArrayPin->GetArrayElementCppType();
 	ArrayPin->SubPins.Insert(Pin, InIndex);
+	ArrayPin->IncrementVersion();
 
 	if (Pin->IsStruct())
 	{
@@ -11109,12 +11119,14 @@ bool URigVMController::RemovePin(URigVMPin* InPinToRemove, bool bSetupUndoRedo,
 	if (URigVMPin* ParentPin = InPinToRemove->GetParentPin())
 	{
 		ParentPin->SubPins.Remove(InPinToRemove);
+		ParentPin->IncrementVersion();
 	}
 	
 	else if(URigVMNode* Node = InPinToRemove->GetNode())
 	{
 		Node->Pins.Remove(InPinToRemove);
 		Node->OrphanedPins.Remove(InPinToRemove);
+		Node->IncrementVersion();
 	}
 
 	TArray<URigVMPin*> SubPins = InPinToRemove->GetSubPins();
@@ -12861,6 +12873,7 @@ bool URigVMController::RenameExposedPin(const FName& InOldPinName, const FName&
 			}
 			
 			InPin->SetDisplayName(InNewName);
+			InPin->IncrementVersion();
 			
 			PinController->Notify(ERigVMGraphNotifType::PinRenamed, InPin);
 
@@ -18515,6 +18528,7 @@ void URigVMController::RepopulatePinsOnNode(const FRigVMRegistry& Registry, cons
 				check(Pin->CPPTypeObject == nullptr);
 				Pin->CPPTypeObjectPath = NAME_None;
 				Pin->LastKnownTypeIndex = NewPinInfo->TypeIndex;
+				Pin->IncrementVersion();
 
 				Notify(ERigVMGraphNotifType::PinTypeChanged, Pin);
 			}
@@ -18978,6 +18992,7 @@ void URigVMController::ApplyPinState(URigVMPin* InPin, const FPinState& InPinSta
 					if(URigVMPin* OuterPin = Cast<URigVMPin>(Injection->GetOuter()))
 					{
 						OuterPin->InjectionInfos.Remove(Injection);
+						OuterPin->IncrementVersion();
 					}
 					RenameObject(FormerlyInjectedNode, nullptr, InPin->GetGraph());
 					DestroyObject(Injection);
@@ -19013,6 +19028,7 @@ void URigVMController::ApplyPinState(URigVMPin* InPin, const FPinState& InPinSta
 	}
 
 	SetPinExpansion(InPin, InPinState.bIsExpanded, bSetupUndoRedo);
+	InPin->IncrementVersion();
 }
 
 void URigVMController::ApplyPinStates(URigVMNode* InNode, const TMap<FString, URigVMController::FPinState>& InPinStates, const TMap<FString, FString>& InRedirectedPinPaths, bool bSetupUndoRedo)
@@ -19039,6 +19055,7 @@ void URigVMController::ApplyPinStates(URigVMNode* InNode, const TMap<FString, UR
 				if(URigVMPin* OuterPin = Cast<URigVMPin>(InjectionInfo->GetOuter()))
 				{
 					OuterPin->InjectionInfos.Remove(InjectionInfo);
+					OuterPin->IncrementVersion();
 				}
 				RenameObject(InjectionInfo->Node, nullptr, InNode->GetGraph());
 				DestroyObject(InjectionInfo);
@@ -19668,6 +19685,8 @@ bool URigVMController::FullyResolveTemplateNode(URigVMTemplateNode* InNode, int3
 		RemovePin(PinToRemove, false);
 	}
 
+	InNode->IncrementVersion();
+
 	// add missing pins
 	if(ResolvedFunction)
 	{
@@ -20654,6 +20673,7 @@ bool URigVMController::ChangePinType(URigVMPin* InPin, TRigVMTypeIndex InTypeInd
 		TemplateNode->InvalidateCache();
 	}
 
+	InPin->IncrementVersion();
 	Notify(ERigVMGraphNotifType::PinTypeChanged, InPin);
 	Notify(ERigVMGraphNotifType::PinDefaultValueChanged, InPin);
 
@@ -20954,6 +20974,7 @@ bool URigVMController::MakeExecutePin(URigVMPin* InOutPin)
 			InOutPin->LastKnownTypeIndex = RigVMTypeUtils::TypeIndex::Execute;
 		}
 		InOutPin->LastKnownCPPType = InOutPin->CPPType;
+		InOutPin->IncrementVersion();
 		return true;
 	}
 	return false;
@@ -21028,6 +21049,8 @@ void URigVMController::AddNodePin(URigVMNode* InNode, URigVMPin* InPin)
 			InPin->DefaultValueType = GetDefaultValueType(InPin, InPin->DefaultValue);
 		}
 	}
+
+	InNode->IncrementVersion();
 }
 
 void URigVMController::AddSubPin(URigVMPin* InParentPin, URigVMPin* InPin)
@@ -21035,6 +21058,7 @@ void URigVMController::AddSubPin(URigVMPin* InParentPin, URigVMPin* InPin)
 	ValidatePin(InPin);
 	checkf(!InParentPin->SubPins.Contains(InPin), TEXT("Parent pin %s already contains subpin %s"), *InParentPin->GetPathName(), *InPin->GetName());
 	InParentPin->SubPins.Add(InPin);
+	InParentPin->IncrementVersion();
 }
 
 bool URigVMController::EnsurePinValidity(URigVMPin* InPin, bool bRecursive)
@@ -21064,10 +21088,12 @@ bool URigVMController::EnsurePinValidity(URigVMPin* InPin, bool bRecursive)
 				return false;
 			}
 			InPin->CPPType = CPPType;
+			InPin->IncrementVersion();
 		}
 		else
 		{
 			InPin->CPPType = RigVMTypeUtils::PostProcessCPPType(InPin->CPPType, InPin->GetCPPTypeObject());
+			InPin->IncrementVersion();
 		}
 	}
 
diff --git a/Engine/Plugins/Runtime/RigVM/Source/RigVMDeveloper/Private/RigVMModel/RigVMNode.cpp b/Engine/Plugins/Runtime/RigVM/Source/RigVMDeveloper/Private/RigVMModel/RigVMNode.cpp
--- a/Engine/Plugins/Runtime/RigVM/Source/RigVMDeveloper/Private/RigVMModel/RigVMNode.cpp
+++ b/Engine/Plugins/Runtime/RigVM/Source/RigVMDeveloper/Private/RigVMModel/RigVMNode.cpp
@@ -26,11 +26,11 @@ URigVMNode::URigVMNode()
 , NodeColor(FLinearColor::White)
 , bHasBreakpoint(false)
 , bHaltedAtThisNode(false)
+, NodeVersion(0)
 #if WITH_EDITOR
 , ProfilingHash(0)
 #endif
 {
-
 }
 
 URigVMNode::~URigVMNode()
@@ -679,6 +679,11 @@ void URigVMNode::UpdateTraitRootPinNames()
 	TraitRootPinNames = NewTraitRootPinNames;
 }
 
+void URigVMNode::IncrementVersion()
+{
+	NodeVersion++;
+}
+
 bool URigVMNode::IsSelected() const
 {
 	URigVMGraph* Graph = GetGraph();
@@ -1189,6 +1194,7 @@ void URigVMNode::GetLinkedNodesRecursive(URigVMPin* InPin, bool bLookForSources,
 void URigVMNode::InvalidateCache()
 {
 	CachedOriginalPinDefaultValues.Reset();
+	IncrementVersion();
 }
 
 const TArray<int32>& URigVMNode::GetInstructionsForVM(const FRigVMExtendedExecuteContext& Context, URigVM* InVM, const FRigVMASTProxy& InProxy) const
diff --git a/Engine/Plugins/Runtime/RigVM/Source/RigVMDeveloper/Private/RigVMModel/RigVMPin.cpp b/Engine/Plugins/Runtime/RigVM/Source/RigVMDeveloper/Private/RigVMModel/RigVMPin.cpp
--- a/Engine/Plugins/Runtime/RigVM/Source/RigVMDeveloper/Private/RigVMModel/RigVMPin.cpp
+++ b/Engine/Plugins/Runtime/RigVM/Source/RigVMDeveloper/Private/RigVMModel/RigVMPin.cpp
@@ -118,14 +118,18 @@ URigVMPin::URigVMPin()
 	, CustomWidgetName(NAME_None)
 	, IndexInCategory(INDEX_NONE)
 	, BoundVariablePath_DEPRECATED()
+	, PinVersion(0)
+	, CombinedPinVersion(0)
 	, LastKnownTypeIndex(INDEX_NONE)
+	, CachedIsStringType(this)
+	, CachedDefaultValue(this)
+	, CachedCPPTypeObjectHash(this)
+	, CachedShowInDetailsPanelOnly(this)
+	, CachedPinPath(this)
+	, CachedPinPathWithNodePath(this)
+	, CachedPinCategory(this)
+	, CachedDisplayName(this)
 {
-#if UE_BUILD_DEBUG
-	if (GetNode())
-	{
-		CachedPinPath = GetPinPath();
-	}
-#endif
 }
 
 bool URigVMPin::NameEquals(const FString& InName, bool bFollowCoreRedirectors) const
@@ -183,28 +187,25 @@ bool URigVMPin::NameEquals(const FString& InName, bool bFollowCoreRedirectors) c
 
 FString URigVMPin::GetPinPath(bool bUseNodePath) const
 {
-	FString PinPath;
-
-	URigVMPin* ParentPin = GetParentPin();
-	if (ParentPin)
+	TRigVMModelCachedValue<URigVMPin, FString>& Cache = bUseNodePath ? CachedPinPathWithNodePath : CachedPinPath;
+	if(Cache.IsValid())
 	{
-		PinPath = JoinPinPath(ParentPin->GetPinPath(bUseNodePath), GetName());
+		return Cache.GetValue();
 	}
-	else
+	
+	Cache.ResetCachedValue();
+	
+	if (const URigVMPin* ParentPin = GetParentPin())
 	{
-		URigVMNode* Node = GetNode();
-		if (Node != nullptr)
-		{
-			PinPath = JoinPinPath(Node->GetNodePath(bUseNodePath), GetName());
-		}
+		Cache = JoinPinPath(ParentPin->GetPinPath(bUseNodePath), GetName());
 	}
-#if UE_BUILD_DEBUG
-	if (!bUseNodePath)
+	else if (const URigVMNode* Node = GetNode())
 	{
-		CachedPinPath = PinPath;
+		Cache = JoinPinPath(Node->GetNodePath(bUseNodePath), GetName());
 	}
-#endif
-	return PinPath;
+
+	static const FString EmptyPinPath;
+	return Cache.Get(EmptyPinPath);
 }
 
 FString URigVMPin::GetSubPinPath(const URigVMPin* InParentPin, bool bIncludeParentPinName) const
@@ -230,14 +231,24 @@ FString URigVMPin::GetCategory() const
 {
 	if (UserDefinedCategory.IsEmpty())
 	{
+		if(CachedPinCategory.IsValid())
+		{
+			return CachedPinCategory.GetValue();
+		}
+		
+		CachedPinCategory.ResetCachedValue();
+		
 		if(const URigVMNode* Node = GetNode())
 		{
 			const FString CategoryFromNode = Node->GetCategoryForPin(this->GetSegmentPath(true));
 			if(!CategoryFromNode.IsEmpty())
 			{
-				return CategoryFromNode;
+				CachedPinCategory = CategoryFromNode;
 			}
 		}
+
+		static const FString EmptyCategory;
+		return CachedPinCategory.Get(EmptyCategory);
 	}
 	return UserDefinedCategory;
 }
@@ -398,11 +409,17 @@ void URigVMPin::GetExposedPinChainImpl(TArray<const URigVMPin*>& OutExposedPins,
 
 FName URigVMPin::GetDisplayName() const
 {
+	if(CachedDisplayName.IsValid())
+	{
+		return CachedDisplayName.GetValue();
+	}
+
 	if (DisplayName == NAME_None)
 	{
 		if(IsArrayElement())
 		{
-			return *FString::FromInt(GetPinIndex());
+			CachedDisplayName = *FString::FromInt(GetPinIndex());
+			return CachedDisplayName.GetValue();
 		}
 		
 		if(const URigVMNode* Node = GetNode())
@@ -410,10 +427,12 @@ FName URigVMPin::GetDisplayName() const
 			const FName DisplayNameFromNode = Node->GetDisplayNameForPin(this);
 			if(!DisplayNameFromNode.IsNone())
 			{
-				return DisplayNameFromNode;
+				CachedDisplayName = DisplayNameFromNode;
+				return CachedDisplayName.GetValue();
 			}
 		}
-		return GetFName();
+		CachedDisplayName = GetFName();
+		return CachedDisplayName.GetValue();
 	}
 
 	if (InjectionInfos.Num() > 0)
@@ -432,7 +451,8 @@ FName URigVMPin::GetDisplayName() const
 			}
 		}
 
-		return *ProcessedDisplayName;
+		CachedDisplayName = *ProcessedDisplayName;
+		return CachedDisplayName.GetValue();
 	}
 
 	return DisplayName;
@@ -588,6 +608,26 @@ void URigVMPin::SetDisplayName(const FName& InDisplayName)
 	}
 }
 
+void URigVMPin::IncrementVersion(bool bAffectParentPin, bool bAffectSubPins)
+{
+	PinVersion++;
+
+	if(bAffectParentPin)
+	{
+		if(URigVMPin* ParentPin = GetParentPin())
+		{
+			ParentPin->IncrementVersion(true, false);
+		}
+	}
+	if(bAffectSubPins)
+	{
+		for(URigVMPin* SubPin : SubPins)
+		{
+			SubPin->IncrementVersion(false, true);
+		}
+	}
+}
+
 int32 URigVMPin::GetArraySize() const
 {
 	return SubPins.Num();
@@ -638,8 +678,12 @@ TRigVMTypeIndex URigVMPin::GetTypeIndex() const
 
 bool URigVMPin::IsStringType() const
 {
-	const FString ResolvedType = GetCPPType();
-	return ResolvedType.Equals(TEXT("FString")) || ResolvedType.Equals(TEXT("FName"));
+	if(!CachedIsStringType.IsValid())
+	{
+		const FString ResolvedType = GetCPPType();
+		CachedIsStringType = ResolvedType.Equals(TEXT("FString")) || ResolvedType.Equals(TEXT("FName"));
+	}
+	return CachedIsStringType.GetValue();
 }
 
 bool URigVMPin::IsExecuteContext() const
@@ -807,13 +851,19 @@ FString URigVMPin::GetDefaultValue(const URigVMPin::FPinOverride& InOverride, bo
 		return OverrideValuePtr->DefaultValue;
 	}
 
+	if(CachedDefaultValue.IsValid())
+	{
+		return CachedDefaultValue.GetValue();
+	}
+
 	if (IsArray())
 	{
 		FRigVMRegistry& Registry = FRigVMRegistry::Get();
 		const TRigVMTypeIndex ArrayType = GetTypeIndex();
 		if (!Registry.IsArrayType(ArrayType))
 		{
-			return TEXT("()");
+			CachedDefaultValue = TEXT("()"); 
+			return CachedDefaultValue.GetValue();
 		}
 		if (SubPins.Num() > 0)
 		{
@@ -823,7 +873,8 @@ FString URigVMPin::GetDefaultValue(const URigVMPin::FPinOverride& InOverride, bo
 			{
 				if (SubPin->GetTypeIndex() != ElementType)
 				{
-					return TEXT("()");
+					CachedDefaultValue = TEXT("()"); 
+					return CachedDefaultValue.GetValue();
 				}
 				FString ElementDefaultValue = SubPin->GetDefaultValue(InOverride, bAdaptValueForPinType);
 				if (SubPin->IsStringType())
@@ -835,9 +886,11 @@ FString URigVMPin::GetDefaultValue(const URigVMPin::FPinOverride& InOverride, bo
 			}
 			if (ElementDefaultValues.Num() == 0)
 			{
-				return TEXT("()");
+				CachedDefaultValue = TEXT("()"); 
+				return CachedDefaultValue.GetValue();
 			}
-			return FString::Printf(TEXT("(%s)"), *FString::Join(ElementDefaultValues, TEXT(",")));
+			CachedDefaultValue = FString::Printf(TEXT("(%s)"), *FString::Join(ElementDefaultValues, TEXT(",")));
+			return CachedDefaultValue.GetValue();
 		}
 
 		return DefaultValue.IsEmpty() ? TEXT("()") : DefaultValue;
@@ -886,7 +939,8 @@ FString URigVMPin::GetDefaultValue(const URigVMPin::FPinOverride& InOverride, bo
 					URigVMController::OverrideDefaultValueMember(SubPin->GetName(), MemberDefaultValue, FinalDefaultValue);
 				}
 
-				return !FinalDefaultValue.IsEmpty() ? FinalDefaultValue : EmptyStructDefaultValue;
+				CachedDefaultValue = !FinalDefaultValue.IsEmpty() ? FinalDefaultValue : EmptyStructDefaultValue;
+				return CachedDefaultValue.GetValue();
 			}
 		}
 
@@ -907,7 +961,8 @@ FString URigVMPin::GetDefaultValue(const URigVMPin::FPinOverride& InOverride, bo
 					const FQuat Quat = FQuat::MakeFromRotator(Rotator);
 					FString AdaptedDefaultValue;	
 					TBaseStructure<FQuat>::Get()->ExportText(AdaptedDefaultValue, &Quat, &Quat, nullptr, PPF_None, nullptr);
-					return AdaptedDefaultValue;
+					CachedDefaultValue = AdaptedDefaultValue;
+					return CachedDefaultValue.GetValue();
 				}
 			}
 		}
@@ -924,7 +979,8 @@ FString URigVMPin::GetDefaultValue(const URigVMPin::FPinOverride& InOverride, bo
 		// This default value is used during URigVMCompiler::FindOrAddRegister(...)
 		// Thus in this block, we have to return something like 0.0 instead of empty string
 
-		return URigVMController::GetPinInitialDefaultValue(this);
+		CachedDefaultValue = URigVMController::GetPinInitialDefaultValue(this);
+		return CachedDefaultValue.GetValue();
 	}
 
 	return DefaultValue;
@@ -1519,6 +1575,12 @@ bool URigVMPin::CanBeBoundToVariable(const FRigVMExternalVariable& InExternalVar
 
 bool URigVMPin::ShowInDetailsPanelOnly() const
 {
+	if(CachedShowInDetailsPanelOnly.IsValid())
+	{
+		return CachedShowInDetailsPanelOnly.GetValue();
+	}
+	CachedShowInDetailsPanelOnly.ResetCachedValue();
+
 #if WITH_EDITOR
 	if (GetParentPin() == nullptr)
 	{
@@ -1530,7 +1592,7 @@ bool URigVMPin::ShowInDetailsPanelOnly() const
 				{
 					if (Property->HasMetaData(FRigVMStruct::DetailsOnlyMetaName))
 					{
-						return true;
+						CachedShowInDetailsPanelOnly = true;
 					}
 				}
 			}
@@ -1539,12 +1601,14 @@ bool URigVMPin::ShowInDetailsPanelOnly() const
 		{
 			if(const FRigVMTemplate* Template = TemplateNode->GetTemplate())
 			{
-				return !Template->GetArgumentMetaData(GetFName(), FRigVMStruct::DetailsOnlyMetaName).IsEmpty();
+				CachedShowInDetailsPanelOnly = !Template->GetArgumentMetaData(GetFName(), FRigVMStruct::DetailsOnlyMetaName).IsEmpty();
 			}
 		}
 	}
 #endif
-	return false;
+
+	static constexpr bool bShowOnlyInDetailsPanel = false;
+	return CachedShowInDetailsPanelOnly.Get(bShowOnlyInDetailsPanel); 
 }
 
 // Returns nullptr external variable matching this description
@@ -1657,8 +1721,34 @@ UScriptStruct* URigVMPin::GetTraitScriptStruct() const
 	return nullptr;
 }
 
+const uint32& URigVMPin::GetNodeCachedValueVersion() const
+{
+	if(const URigVMNode* Node = GetNode())
+	{
+		return Node->GetCachedValueVersion();
+	}
+	static constexpr uint32 InvalidVersion = 0;
+	return InvalidVersion;
+}
+
+const uint32& URigVMPin::GetCachedValueVersion() const
+{
+	CombinedPinVersion = HashCombine(GetNodeCachedValueVersion(), PinVersion);
+	return CombinedPinVersion;
+}
+
 void URigVMPin::UpdateTypeInformationIfRequired() const
 {
+	const uint32 CPPTypeObjectHash = GetTypeHash(CPPTypeObject);
+	if(CachedCPPTypeObjectHash.IsValid())
+	{
+		if(CachedCPPTypeObjectHash == CPPTypeObjectHash)
+		{
+			return;
+		}
+	}
+	CachedCPPTypeObjectHash = CPPTypeObjectHash;
+	
 	if (CPPTypeObject == nullptr)
 	{
 		if (CPPTypeObjectPath != NAME_None)
diff --git a/Engine/Plugins/Runtime/RigVM/Source/RigVMDeveloper/Public/RigVMModel/Nodes/RigVMUnitNode.h b/Engine/Plugins/Runtime/RigVM/Source/RigVMDeveloper/Public/RigVMModel/Nodes/RigVMUnitNode.h
--- a/Engine/Plugins/Runtime/RigVM/Source/RigVMDeveloper/Public/RigVMModel/Nodes/RigVMUnitNode.h
+++ b/Engine/Plugins/Runtime/RigVM/Source/RigVMDeveloper/Public/RigVMModel/Nodes/RigVMUnitNode.h
@@ -5,6 +5,7 @@
 #include "RigVMModel/Nodes/RigVMTemplateNode.h"
 #include "RigVMCore/RigVMStruct.h"
 #include "UObject/StructOnScope.h"
+#include "RigVMModel/RigVMModelCachedValue.h"
 #include "RigVMUnitNode.generated.h"
 
 class URigVMHost;
@@ -118,6 +119,8 @@ private:
 	FName MethodName_DEPRECATED;
 #endif
 
+	mutable TRigVMModelCachedValue<URigVMUnitNode, FName> CachedEventName;
+
 	friend class URigVMController;
 };
 
diff --git a/Engine/Plugins/Runtime/RigVM/Source/RigVMDeveloper/Public/RigVMModel/RigVMModelCachedValue.h b/Engine/Plugins/Runtime/RigVM/Source/RigVMDeveloper/Public/RigVMModel/RigVMModelCachedValue.h
new file mode 100644
--- /dev/null
+++ b/Engine/Plugins/Runtime/RigVM/Source/RigVMDeveloper/Public/RigVMModel/RigVMModelCachedValue.h
@@ -0,0 +1,144 @@
+// Copyright Epic Games, Inc. All Rights Reserved.
+
+#pragma once
+
+#include "CoreMinimal.h"
+
+/**
+ * A cached value which relies on a version providing object.
+ * The versioned object needs to provide uint32 GetCachedValueVersion()
+ * Values within the cache are considered valid if they are a) set and
+ * b) the version of the cache matches the version of the versioned object.
+ */
+template<typename VersionObjectType, typename ValueType>
+class TRigVMModelCachedValue
+{
+public:
+
+	// default empty constructor
+	TRigVMModelCachedValue()
+		: LastObjectVersion(UINT_MAX)
+	{
+	}
+
+	// constructor for an empty cache bound to a versioned object
+	TRigVMModelCachedValue(const VersionObjectType* InObject)
+		: VersionedObjectWeakPtr(InObject)
+		, LastObjectVersion(UINT_MAX)
+	{
+	}
+
+	// constructor for a valid cache bound to a versioned object
+	TRigVMModelCachedValue(VersionObjectType* InObject, const ValueType& InValue)
+		: VersionedObjectWeakPtr(InObject)
+		, LastObjectVersion(UINT_MAX)
+	{
+		Set(InObject ? InObject->GetCachedValueVersion() : 0, InValue);
+	}
+
+	// returns true if the cache is bound, the value has been set and the version is current
+	bool IsValid() const
+	{
+		if(Value.IsSet())
+		{
+			if(const VersionObjectType* VersionedObject = VersionedObjectWeakPtr.Get())
+			{
+				return VersionedObject->GetCachedValueVersion() == LastObjectVersion;
+			}
+		}
+		return false;
+	}
+
+	// returns true if the cache is bound to a versioned object
+	bool IsBound() const
+	{
+		return VersionedObjectWeakPtr.IsValid();
+	}
+
+	// binds the cache to a new versioned object
+	void Bind(const VersionObjectType* InVersionedObject)
+	{
+		if(VersionedObjectWeakPtr.IsValid())
+		{
+			if(VersionedObjectWeakPtr.Get() == InVersionedObject)
+			{
+				return;
+			}
+		}
+		VersionedObjectWeakPtr = InVersionedObject;
+		ResetCachedValue();
+	}
+
+	// unbinds the cache from a versioned object
+	void Unbind()
+	{
+		Bind(nullptr);
+	}
+
+	// returns the cached value if it is valid - or otherwise returns the default
+	const ValueType& Get(const ValueType& InDefault) const
+	{
+		if(IsValid())
+		{
+			return GetValue();
+		}
+		return InDefault;
+	}
+
+	// returns the cached value
+	const ValueType& GetValue() const
+	{
+		check(IsValid());
+		return Value.GetValue();
+	}
+
+	// sets the cached value and updates the last version given the bound object
+	void Set(const ValueType& InValue)
+	{
+		const VersionObjectType* VersionedObject = VersionedObjectWeakPtr.Get();
+		check(VersionedObject);
+		
+		Value = InValue;
+		LastObjectVersion = VersionedObject->GetCachedValueVersion(); 
+	}
+
+	// resets the cached value
+	void ResetCachedValue()
+	{
+		Value.Reset();
+		LastObjectVersion = UINT_MAX;
+	}
+
+	// sets the value using an assignment operator
+	TRigVMModelCachedValue& operator =(const ValueType& InValueType)
+	{
+		Set(InValueType);
+		return *this;
+	}
+
+	// compares the cached value using a comparison operator
+	bool operator ==(const ValueType& InValueType) const
+	{
+		if(IsValid())
+		{
+			return GetValue() == InValueType;
+		}
+		return false;
+	}
+
+	// compares the cached value usinga  comparison operator
+	bool operator !=(const ValueType& InValueType) const
+	{
+		if(IsValid())
+		{
+			return GetValue() != InValueType;
+		}
+		return true;
+	}
+
+private:;
+
+	TWeakObjectPtr<const VersionObjectType> VersionedObjectWeakPtr;
+	TOptional<ValueType> Value;
+	uint32 LastObjectVersion;
+};
\ No newline at end of file
diff --git a/Engine/Plugins/Runtime/RigVM/Source/RigVMDeveloper/Public/RigVMModel/RigVMNode.h b/Engine/Plugins/Runtime/RigVM/Source/RigVMDeveloper/Public/RigVMModel/RigVMNode.h
--- a/Engine/Plugins/Runtime/RigVM/Source/RigVMDeveloper/Public/RigVMModel/RigVMNode.h
+++ b/Engine/Plugins/Runtime/RigVM/Source/RigVMDeveloper/Public/RigVMModel/RigVMNode.h
@@ -419,6 +419,15 @@ protected:
 
 	void UpdateTraitRootPinNames();
 
+public:
+	const uint32& GetCachedValueVersion() const
+	{
+		return NodeVersion;
+	}
+
+protected:
+	void IncrementVersion();
+
 	UPROPERTY()
 	FString NodeTitle;
 
@@ -465,6 +474,8 @@ protected:
 	TMap<FString, bool> PinCategoryExpansion;
 
 	FString LastAffectedPinCategory;
+
+	uint32 NodeVersion;
 	
 private:
 
diff --git a/Engine/Plugins/Runtime/RigVM/Source/RigVMDeveloper/Public/RigVMModel/RigVMPin.h b/Engine/Plugins/Runtime/RigVM/Source/RigVMDeveloper/Public/RigVMModel/RigVMPin.h
--- a/Engine/Plugins/Runtime/RigVM/Source/RigVMDeveloper/Public/RigVMModel/RigVMPin.h
+++ b/Engine/Plugins/Runtime/RigVM/Source/RigVMDeveloper/Public/RigVMModel/RigVMPin.h
@@ -11,6 +11,7 @@
 #include "RigVMCompiler/RigVMASTProxy.h"
 #include "UObject/StructOnScope.h"
 #include <RigVMCore/RigVMTrait.h>
+#include "RigVMModelCachedValue.h"
 #include "RigVMPin.generated.h"
 
 class URigVMGraph;
@@ -527,27 +528,21 @@ public:
 	// Returns the struct of the trait backing up this pin
 	UScriptStruct* GetTraitScriptStruct() const;
 
+	const uint32& GetNodeCachedValueVersion() const;
+	const uint32& GetCachedValueVersion() const;
+
 private:
 
 	void UpdateTypeInformationIfRequired() const;
 	void SetNameFromIndex();
 	void SetDisplayName(const FName& InDisplayName);
+	void IncrementVersion(bool bAffectParentPin = true, bool bAffectSubPins = true);
 
 	void GetExposedPinChainImpl(TArray<const URigVMPin*>& OutExposedPins, TArray<const URigVMPin*>& VisitedPins) const;
 
 	UPROPERTY()
 	FName DisplayName;
 
-#if UE_BUILD_DEBUG
-	
-	// A cache for the pin path for debugging purposes.
-	// When looking at a debug symbol of a pin it is difficult
-	// To grasp where the pin is stored etc.
-	// With this member you can see the full pin path during a debugging session
-	mutable FString CachedPinPath;
-
-#endif
-
 	// if new members are added to the pin in the future 
 	// it is important to search for all existing usages of all members
 	// to make sure things are copied/initialized properly
@@ -612,8 +607,19 @@ private:
 	FString BoundVariablePath_DEPRECATED;
 #endif
 
+	uint32 PinVersion;
+	mutable uint32 CombinedPinVersion;
+	
 	mutable FString LastKnownCPPType;
 	mutable TRigVMTypeIndex LastKnownTypeIndex;
+	mutable TRigVMModelCachedValue<URigVMPin, bool> CachedIsStringType;
+	mutable TRigVMModelCachedValue<URigVMPin, FString> CachedDefaultValue;
+	mutable TRigVMModelCachedValue<URigVMPin, uint32> CachedCPPTypeObjectHash;
+	mutable TRigVMModelCachedValue<URigVMPin, bool> CachedShowInDetailsPanelOnly;
+	mutable TRigVMModelCachedValue<URigVMPin, FString> CachedPinPath;
+	mutable TRigVMModelCachedValue<URigVMPin, FString> CachedPinPathWithNodePath;
+	mutable TRigVMModelCachedValue<URigVMPin, FString> CachedPinCategory;
+	mutable TRigVMModelCachedValue<URigVMPin, FName> CachedDisplayName;
 
 	static const inline TCHAR* OrphanPinPrefix = TEXT("Orphan::");
 
