RigVM: Remove mention of experimental from Variants workflows

#rnx

diff --git a/Engine/Plugins/Runtime/RigVM/Source/RigVMEditor/Private/RigVMEditorModule.cpp b/Engine/Plugins/Runtime/RigVM/Source/RigVMEditor/Private/RigVMEditorModule.cpp
--- a/Engine/Plugins/Runtime/RigVM/Source/RigVMEditor/Private/RigVMEditorModule.cpp
+++ b/Engine/Plugins/Runtime/RigVM/Source/RigVMEditor/Private/RigVMEditorModule.cpp
@@ -116,7 +116,7 @@ void FRigVMEditorModule::StartupModule()
 						if(CVarRigVMEnableVariants.GetValueOnAnyThread())
 						{
 							FSoftObjectPath SoftObjectPath = SelectedAssetData.GetSoftObjectPath();
-							InSection.AddMenuEntry("CreateVariant", LOCTEXT("CreateVariant", "Create variant (Experimental)"), LOCTEXT("CreateVariant_ToolTip", "Create a variant for this asset"), FSlateIcon(FRigVMEditorStyle::Get().GetStyleSetName(), "RigVM", "RigVM.Unit"), FExecuteAction::CreateLambda([SoftObjectPath]()
+							InSection.AddMenuEntry("CreateVariant", LOCTEXT("CreateVariant", "Create variant"), LOCTEXT("CreateVariant_ToolTip", "Create a variant for this asset"), FSlateIcon(FRigVMEditorStyle::Get().GetStyleSetName(), "RigVM", "RigVM.Unit"), FExecuteAction::CreateLambda([SoftObjectPath]()
 								{
 									// Perform the load from within our lambda since this can be expensive, and should not be done speculatively
 									UObject* SelectedObject = SoftObjectPath.TryLoad();
diff --git a/Engine/Plugins/Runtime/RigVM/Source/RigVMEditor/Private/Widgets/SRigVMEditorGraphExplorer.cpp b/Engine/Plugins/Runtime/RigVM/Source/RigVMEditor/Private/Widgets/SRigVMEditorGraphExplorer.cpp
--- a/Engine/Plugins/Runtime/RigVM/Source/RigVMEditor/Private/Widgets/SRigVMEditorGraphExplorer.cpp
+++ b/Engine/Plugins/Runtime/RigVM/Source/RigVMEditor/Private/Widgets/SRigVMEditorGraphExplorer.cpp
@@ -190,7 +190,7 @@ void FRigVMEditorGraphExplorerCommands::RegisterCommands()
 	UI_COMMAND(CreateFunction, "Add New Function", "Create a new function and show it in the editor.", EUserInterfaceActionType::Button, FInputChord());
 	UI_COMMAND(CreateVariable, "Add New Variable", "Create a new member variable.", EUserInterfaceActionType::Button, FInputChord());
 	UI_COMMAND(CreateLocalVariable, "Add New Local Variable", "Create a new local variable to the function.", EUserInterfaceActionType::Button, FInputChord());
-	UI_COMMAND(AddFunctionVariant, "Add Variant (Experimental)", "Creates a new variant of a function.", EUserInterfaceActionType::Button, FInputChord());
+	UI_COMMAND(AddFunctionVariant, "Add Variant", "Creates a new variant of a function.", EUserInterfaceActionType::Button, FInputChord());
 	UI_COMMAND(PasteFunction, "Paste Function", "Pastes the function.", EUserInterfaceActionType::Button, FInputChord());
 	UI_COMMAND(PasteVariable, "Paste Variable", "Pastes the variable.", EUserInterfaceActionType::Button, FInputChord());
 	UI_COMMAND(PasteLocalVariable, "Paste Local Variable", "Pastes the local variable.", EUserInterfaceActionType::Button, FInputChord());
diff --git a/Engine/Plugins/Runtime/RigVM/Source/RigVMEditor/Private/Widgets/SRigVMVariantWidget.cpp b/Engine/Plugins/Runtime/RigVM/Source/RigVMEditor/Private/Widgets/SRigVMVariantWidget.cpp
--- a/Engine/Plugins/Runtime/RigVM/Source/RigVMEditor/Private/Widgets/SRigVMVariantWidget.cpp
+++ b/Engine/Plugins/Runtime/RigVM/Source/RigVMEditor/Private/Widgets/SRigVMVariantWidget.cpp
@@ -169,15 +169,6 @@ void SRigVMVariantWidget::Construct(
 		SNew(SVerticalBox)
 		+ SVerticalBox::Slot()
 		.AutoHeight()
-		.Padding(0, 8, 0, 8)
-		.HAlign(HAlign_Left)
-		[
-			SNew(STextBlock)
-			.Font(IDetailLayoutBuilder::GetDetailFont())
-			.Text(LOCTEXT("Experimental", "Note: This is experimental"))
-		]
-		+ SVerticalBox::Slot()
-		.AutoHeight()
 		.Padding(0, 8, 0, 0)
 		.HAlign(HAlign_Left)
 		[
diff --git a/Engine/Source/Editor/Kismet/Private/SMyBlueprint.cpp b/Engine/Source/Editor/Kismet/Private/SMyBlueprint.cpp
--- a/Engine/Source/Editor/Kismet/Private/SMyBlueprint.cpp
+++ b/Engine/Source/Editor/Kismet/Private/SMyBlueprint.cpp
@@ -2549,7 +2549,7 @@ TSharedPtr<SWidget> SMyBlueprint::OnContextMenuOpening()
 				if (Schema->AllowsFunctionVariants())
 				{
 					MenuBuilder.AddMenuEntry(
-						LOCTEXT("MyBlueprint_Add_Func_Variant", "Add Variant (Experimental)"), FText(), FSlateIcon(),
+						LOCTEXT("MyBlueprint_Add_Func_Variant", "Add Variant"), FText(), FSlateIcon(),
 						FExecuteAction::CreateLambda([BlueprintEditor, Graph]()
 						{
 							BlueprintEditor->AddNewFunctionVariant(Graph->EdGraph);
