RigVM: Enable "set engine content" filter by default

#rnx

diff --git a/Engine/Plugins/Runtime/RigVM/Source/RigVMEditor/Private/Widgets/SRigVMSwapAssetReferencesWidget.cpp b/Engine/Plugins/Runtime/RigVM/Source/RigVMEditor/Private/Widgets/SRigVMSwapAssetReferencesWidget.cpp
--- a/Engine/Plugins/Runtime/RigVM/Source/RigVMEditor/Private/Widgets/SRigVMSwapAssetReferencesWidget.cpp
+++ b/Engine/Plugins/Runtime/RigVM/Source/RigVMEditor/Private/Widgets/SRigVMSwapAssetReferencesWidget.cpp
@@ -118,6 +118,9 @@ void SRigVMSwapAssetReferencesWidget::Construct(const FArguments& InArgs)
 	};
 	static const TSharedRef<FRigVMTreeFilter> DefaultPathFilter = FRigVMTreePathFilter::Create();
 
+	// enable show engine content by default (this filter has to be inverted)
+	DefaultFilters[0]->SetEnabled(false);
+
 	OnGetReferences = InArgs._OnGetReferences;
 	OnSwapReference = InArgs._OnSwapReference;
 	SourceAssetFilters = InArgs._SourceAssetFilters;
diff --git a/Engine/Plugins/Runtime/RigVM/Source/RigVMEditor/Private/Widgets/SRigVMSwapFunctionsWidget.cpp b/Engine/Plugins/Runtime/RigVM/Source/RigVMEditor/Private/Widgets/SRigVMSwapFunctionsWidget.cpp
--- a/Engine/Plugins/Runtime/RigVM/Source/RigVMEditor/Private/Widgets/SRigVMSwapFunctionsWidget.cpp
+++ b/Engine/Plugins/Runtime/RigVM/Source/RigVMEditor/Private/Widgets/SRigVMSwapFunctionsWidget.cpp
@@ -481,6 +481,9 @@ void SRigVMSwapFunctionsWidget::Construct(const FArguments& InArgs)
 	};
 	static const TSharedRef<FRigVMTreeFilter> DefaultPathFilter = FRigVMTreePathFilter::Create();
 
+	// enable show engine content by default (this filter has to be inverted)
+	DefaultFilters[0]->SetEnabled(false);
+
 	PickTargetContext = FRigVMSwapFunctionContext::Create();
 	PickFunctionRefsContext = FRigVMSwapFunctionContext::Create();
 	SourcePreviewEnvironment = MakeShareable(new FRigVMMinimalEnvironment());
