RigVM: Pause the OnObjectRenamed multicast delegate during RIgVMController object destruction

#rnx

diff --git a/Engine/Plugins/Runtime/RigVM/Source/RigVMDeveloper/Private/RigVMModel/RigVMController.cpp b/Engine/Plugins/Runtime/RigVM/Source/RigVMDeveloper/Private/RigVMModel/RigVMController.cpp
--- a/Engine/Plugins/Runtime/RigVM/Source/RigVMDeveloper/Private/RigVMModel/RigVMController.cpp
+++ b/Engine/Plugins/Runtime/RigVM/Source/RigVMDeveloper/Private/RigVMModel/RigVMController.cpp
@@ -20914,6 +20914,11 @@ bool URigVMController::RenameObject(UObject* InObjectToRename, const TCHAR* InNe
 
 void URigVMController::DestroyObject(UObject* InObjectToDestroy) const
 {
+	// disable the on object renamed callback
+	// to avoid flushing the compilation queue during graph changes. 
+	static FCoreUObjectDelegates::FOnObjectRenamed EmptyOnObjectRenamed;
+	const TGuardValue<FCoreUObjectDelegates::FOnObjectRenamed> _(FCoreUObjectDelegates::OnObjectRenamed, EmptyOnObjectRenamed);
+	
 	if(const URigVMCollapseNode* CollapseNode = Cast<URigVMCollapseNode>(InObjectToDestroy))
 	{
 		if(IRigVMClientHost* ClientHost = GetImplementingOuter<IRigVMClientHost>())
@@ -23210,7 +23215,7 @@ FRigVMClientPatchResult URigVMController::PatchUserDefinedStructPinNames()
 								{
 									Link->UpdatePinPointers();
 								}
-								SubPin->Rename(*NewNameSubPin, nullptr);
+								RenameObject(SubPin, *NewNameSubPin);
 								for (URigVMLink* Link : SubPin->GetSourceLinks(true))
 								{
 									Link->UpdatePinPaths();
