RigVM: Fix User Defined Events unit tests

#rnx

diff --git a/Engine/Plugins/Runtime/RigVM/Source/RigVMDeveloper/Private/RigVMModel/RigVMController.cpp b/Engine/Plugins/Runtime/RigVM/Source/RigVMDeveloper/Private/RigVMModel/RigVMController.cpp
--- a/Engine/Plugins/Runtime/RigVM/Source/RigVMDeveloper/Private/RigVMModel/RigVMController.cpp
+++ b/Engine/Plugins/Runtime/RigVM/Source/RigVMDeveloper/Private/RigVMModel/RigVMController.cpp
@@ -1739,7 +1739,7 @@ URigVMUnitNode* URigVMController::AddUnitNodeWithDefaults(UScriptStruct* InScrip
 	FString ExportedDefaultValue;
 	if (InDefaults.IsValid())
 	{
-		InDefaults.GetScriptStruct()->ExportText(ExportedDefaultValue, InDefaults.GetMemory(), nullptr, nullptr, PPF_None, nullptr, true);
+		InDefaults.GetScriptStruct()->ExportText(ExportedDefaultValue, InDefaults.GetMemory(), nullptr, nullptr, (PPF_ExternalEditor | PPF_IncludeTransient), nullptr, true);
 	}
 	else
 	{
@@ -1781,6 +1781,8 @@ URigVMUnitNode* URigVMController::AddUnitNodeWithDefaults(UScriptStruct* InScrip
 		}
 	}
 
+	Node->IncrementVersion();
+
 	if (bSetupUndoRedo)
 	{
 		GetActionStack()->EndAction(Action);
@@ -17270,6 +17272,8 @@ void URigVMController::AddPinsForTemplate(const FRigVMTemplate* InTemplate, cons
 	}
 
 	AddExecutePins(ERigVMPinDirection::Output);
+
+	InNode->IncrementVersion();
 }
 
 void URigVMController::ConfigurePinFromProperty(FProperty* InProperty, URigVMPin* InOutPin, ERigVMPinDirection InPinDirection) const
@@ -19308,7 +19312,7 @@ void URigVMController::CreateDefaultValueForStructIfRequired(UScriptStruct* InSt
 		// before ExportText() because ExportText() appends to it.
 		InOutDefaultValue.Reset();
 
-		InStruct->ExportText(InOutDefaultValue, TempBuffer.GetData(), TempBuffer.GetData(), nullptr, PPF_None, nullptr);
+		InStruct->ExportText(InOutDefaultValue, TempBuffer.GetData(), TempBuffer.GetData(), nullptr, (PPF_ExternalEditor | PPF_IncludeTransient), nullptr);
 		InStruct->DestroyStruct(TempBuffer.GetData());
 	}
 }
diff --git a/Engine/Plugins/Runtime/RigVM/Source/RigVMDeveloper/Private/RigVMModel/RigVMPin.cpp b/Engine/Plugins/Runtime/RigVM/Source/RigVMDeveloper/Private/RigVMModel/RigVMPin.cpp
--- a/Engine/Plugins/Runtime/RigVM/Source/RigVMDeveloper/Private/RigVMModel/RigVMPin.cpp
+++ b/Engine/Plugins/Runtime/RigVM/Source/RigVMDeveloper/Private/RigVMModel/RigVMPin.cpp
@@ -626,6 +626,19 @@ void URigVMPin::IncrementVersion(bool bAffectParentPin, bool bAffectSubPins)
 			SubPin->IncrementVersion(false, true);
 		}
 	}
+
+	if(IsRootPin())
+	{
+		// pin changes may affect the event name on a node
+		// so we need to invalidate the node as well.
+		if(URigVMNode* Node = GetNode())
+		{
+			if(Node->IsEvent())
+			{
+				Node->IncrementVersion();
+			}
+		}
+	}
 }
 
 int32 URigVMPin::GetArraySize() const
