Fixed instanced structs crash in RigVM graph details, due to the customization interface not being propagated from the header to the children when creating a property value widget with customization. This causes the details property row to create a new instance of the custom type for the children, , which is different from the one used by the header.

diff --git a/Engine/Plugins/Runtime/RigVM/Source/RigVMEditor/Private/Editor/RigVMGraphDetailCustomization.cpp b/Engine/Plugins/Runtime/RigVM/Source/RigVMEditor/Private/Editor/RigVMGraphDetailCustomization.cpp
--- a/Engine/Plugins/Runtime/RigVM/Source/RigVMEditor/Private/Editor/RigVMGraphDetailCustomization.cpp
+++ b/Engine/Plugins/Runtime/RigVM/Source/RigVMEditor/Private/Editor/RigVMGraphDetailCustomization.cpp
@@ -2481,19 +2481,20 @@ void FRigVMWrappedNodeDetailCustomization::CustomizeDetails(IDetailLayoutBuilder
         	}
         }
 
-		TSharedPtr<SWidget> ValueWidget = PropertyHandle->CreatePropertyValueWidgetWithCustomization(DetailLayout.GetDetailsViewSharedPtr().Get());
+		TSharedPtr<SWidget> NameWidget;
+		TSharedPtr<SWidget> ValueWidget;
+		Row->GetDefaultWidgets(NameWidget, ValueWidget, /*bAddWidgetDecoration*/true);
 
-		constexpr bool bShowChildren = true;
-		Row->CustomWidget(bShowChildren)
-		.NameContent()
-		[
-			LabelWidget
-		]
-		.ValueContent()
-		[
-			ValueWidget ? ValueWidget.ToSharedRef() : SNullWidget::NullWidget
-		]
-		.OverrideResetToDefault(ResetToDefault);
+		Row->CustomWidget(/*bShowChildren*/true)
+			.NameContent()
+			[
+				LabelWidget
+			]
+			.ValueContent()
+			[
+				ValueWidget.ToSharedRef()
+			]
+			.OverrideResetToDefault(ResetToDefault);
 	}
 
 	// now loop over all handles and determine expansion states of the corresponding pins
