RigVM: Additional unit test to identify bottleneck when refreshing engine types on demand

#rnx

diff --git a/Engine/Plugins/Runtime/RigVM/Source/RigVM/Private/RigVMCore/RigVMRegistry.cpp b/Engine/Plugins/Runtime/RigVM/Source/RigVM/Private/RigVMCore/RigVMRegistry.cpp
--- a/Engine/Plugins/Runtime/RigVM/Source/RigVM/Private/RigVMCore/RigVMRegistry.cpp
+++ b/Engine/Plugins/Runtime/RigVM/Source/RigVM/Private/RigVMCore/RigVMRegistry.cpp
@@ -1113,7 +1113,6 @@ TRigVMTypeIndex FRigVMRegistry_NoLock::GetTypeIndexFromCPPType_NoLock(const FStr
 
 		Result = Types.IndexOfByPredicate(Predicate);
 
-#if !WITH_EDITOR
 		// in game / non-editor it's possible that a user defined struct or enum 
 		// has not been registered. thus we'll call RefreshEngineTypes to bring
 		// things up to date here. 
@@ -1121,10 +1120,9 @@ TRigVMTypeIndex FRigVMRegistry_NoLock::GetTypeIndexFromCPPType_NoLock(const FStr
 		{
 			// we may need to update the types again to registry potentially
 			// missing predicate types 
-			FRigVMRegistry_NoLock::GetForWrite().RefreshEngineTypes_NoLock();
+			const_cast<FRigVMRegistry_NoLock*>(this)->RefreshEngineTypes_NoLock();
 			Result = Types.IndexOfByPredicate(Predicate);
 		}
-#endif
 
 		// If not found, try to find a redirect
 		if (Result == INDEX_NONE)
