Control Rig: Make sure the menu items on the right click list of the graph contains icons
#rnx

diff --git a/Engine/Plugins/Runtime/RigVM/Source/RigVMEditor/Private/Editor/RigVMActionMenuItem.cpp b/Engine/Plugins/Runtime/RigVM/Source/RigVMEditor/Private/Editor/RigVMActionMenuItem.cpp
--- a/Engine/Plugins/Runtime/RigVM/Source/RigVMEditor/Private/Editor/RigVMActionMenuItem.cpp
+++ b/Engine/Plugins/Runtime/RigVM/Source/RigVMEditor/Private/Editor/RigVMActionMenuItem.cpp
@@ -346,4 +346,10 @@ UBlueprintNodeSpawner const* FRigVMActionMenuItem::GetRawAction() const
 	return Action;
 }
 
+FSlateBrush const* FRigVMActionMenuItem::GetMenuIcon(FSlateColor& ColorOut)
+{
+	ColorOut = IconTint;
+	return IconBrush;
+}
+
 #undef LOCTEXT_NAMESPACE
diff --git a/Engine/Plugins/Runtime/RigVM/Source/RigVMEditor/Private/Editor/RigVMEditor.cpp b/Engine/Plugins/Runtime/RigVM/Source/RigVMEditor/Private/Editor/RigVMEditor.cpp
--- a/Engine/Plugins/Runtime/RigVM/Source/RigVMEditor/Private/Editor/RigVMEditor.cpp
+++ b/Engine/Plugins/Runtime/RigVM/Source/RigVMEditor/Private/Editor/RigVMEditor.cpp
@@ -3807,6 +3807,31 @@ void FRigVMEditorBase::OnGraphNodeDropToPerform(TSharedPtr<FDragDropOperation> I
 	}
 }
 
+FSlateBrush const* FRigVMEditorBase::GetVarIconAndColorFromProperty(const FProperty* Property, FSlateColor& IconColorOut, FSlateBrush const*& SecondaryBrushOut, FSlateColor& SecondaryColorOut)
+{
+	SecondaryBrushOut = nullptr;
+	if (Property != nullptr)
+	{
+		const UEdGraphSchema_K2* K2Schema = GetDefault<UEdGraphSchema_K2>();
+
+		FEdGraphPinType PinType;
+		if (K2Schema->ConvertPropertyToPinType(Property, PinType)) // use schema to get the color
+		{
+			return GetVarIconAndColorFromPinType(PinType, IconColorOut, SecondaryBrushOut, SecondaryColorOut);
+		}
+	}
+	return FAppStyle::GetBrush(TEXT("Kismet.AllClasses.VariableIcon"));
+}
+
+FSlateBrush const* FRigVMEditorBase::GetVarIconAndColorFromPinType(const FEdGraphPinType& PinType, FSlateColor& IconColorOut, FSlateBrush const*& SecondaryBrushOut, FSlateColor& SecondaryColorOut)
+{
+	const UEdGraphSchema_K2* K2Schema = GetDefault<UEdGraphSchema_K2>();
+	IconColorOut = K2Schema->GetPinTypeColor(PinType);
+	SecondaryBrushOut = FBlueprintEditorUtils::GetSecondaryIconFromPin(PinType);
+	SecondaryColorOut = K2Schema->GetSecondaryPinTypeColor(PinType);
+	return FBlueprintEditorUtils::GetIconFromPin(PinType);
+}
+
 void FRigVMEditorBase::UpdateGraphCompilerErrors()
 {
 	DECLARE_SCOPE_HIERARCHICAL_COUNTER_FUNC()
diff --git a/Engine/Plugins/Runtime/RigVM/Source/RigVMEditor/Public/Editor/RigVMActionMenuItem.h b/Engine/Plugins/Runtime/RigVM/Source/RigVMEditor/Public/Editor/RigVMActionMenuItem.h
--- a/Engine/Plugins/Runtime/RigVM/Source/RigVMEditor/Public/Editor/RigVMActionMenuItem.h
+++ b/Engine/Plugins/Runtime/RigVM/Source/RigVMEditor/Public/Editor/RigVMActionMenuItem.h
@@ -54,6 +54,16 @@ public:
 	/** @return */
 	UBlueprintNodeSpawner const* GetRawAction() const;
 
+	/**
+	 * Retrieves the icon brush for this menu entry (to be displayed alongside
+	 * in the menu).
+	 *
+	 * @param  ColorOut	An output parameter that's filled in with the color to tint the brush with.
+	 * @return An slate brush to be used for this menu item in the action menu.
+	 */
+	FSlateBrush const* GetMenuIcon(FSlateColor& ColorOut);
+
+
 
 private:
 	/** Specialized node-spawner, that comprises the action portion of this menu entry. */
diff --git a/Engine/Plugins/Runtime/RigVM/Source/RigVMEditor/Public/Editor/RigVMEditor.h b/Engine/Plugins/Runtime/RigVM/Source/RigVMEditor/Public/Editor/RigVMEditor.h
--- a/Engine/Plugins/Runtime/RigVM/Source/RigVMEditor/Public/Editor/RigVMEditor.h
+++ b/Engine/Plugins/Runtime/RigVM/Source/RigVMEditor/Public/Editor/RigVMEditor.h
@@ -290,6 +290,28 @@ public:
 	virtual bool IsEditable(UEdGraph* InGraph) const override;
 
 	virtual void OnGraphNodeDropToPerform(TSharedPtr<FDragDropOperation> InDragDropOp, UEdGraph* InGraph, const FVector2D& InNodePosition, const FVector2D& InScreenPosition) override;
+
+	/**
+	 * Util for finding a glyph and color for a variable.
+	 *
+	 * @param Property       The variable's property
+	 * @param IconColorOut      The resulting color for the glyph
+	 * @param SecondaryBrushOut The resulting secondary glyph brush (used for Map types)
+	 * @param SecondaryColorOut The resulting secondary color for the glyph (used for Map types)
+	 * @return					The resulting glyph brush
+	 */
+	static FSlateBrush const* GetVarIconAndColorFromProperty(const FProperty* Property, FSlateColor& IconColorOut, FSlateBrush const*& SecondaryBrushOut, FSlateColor& SecondaryColorOut);
+
+	/**
+	* Util for finding a glyph and color for a variable.
+	*
+	* @param PinType       The variable's pin type
+	* @param IconColorOut      The resulting color for the glyph
+	* @param SecondaryBrushOut The resulting secondary glyph brush (used for Map types)
+	* @param SecondaryColorOut The resulting secondary color for the glyph (used for Map types)
+	* @return					The resulting glyph brush
+	*/
+	static FSlateBrush const* GetVarIconAndColorFromPinType(const FEdGraphPinType& PinType, FSlateColor& IconColorOut, FSlateBrush const*& SecondaryBrushOut, FSlateColor& SecondaryColorOut);
 	
 protected:
 	
diff --git a/Engine/Plugins/Runtime/RigVM/Source/RigVMEditor/Public/Editor/SRigVMActionMenu.cpp b/Engine/Plugins/Runtime/RigVM/Source/RigVMEditor/Public/Editor/SRigVMActionMenu.cpp
--- a/Engine/Plugins/Runtime/RigVM/Source/RigVMEditor/Public/Editor/SRigVMActionMenu.cpp
+++ b/Engine/Plugins/Runtime/RigVM/Source/RigVMEditor/Public/Editor/SRigVMActionMenu.cpp
@@ -12,6 +12,9 @@
 #include "Widgets/Text/SInlineEditableTextBlock.h"
 #include "Editor/RigVMActionMenuUtils.h"
 #include "BlueprintEditorSettings.h"
+#include "BlueprintNodeSpawner.h"
+#include "K2Node_Variable.h"
+#include "RigVMActionMenuItem.h"
 #include "RigVMBlueprint.h"
 #include "Widgets/Input/SCheckBox.h"
 #include "Framework/Application/SlateApplication.h"
@@ -330,7 +333,28 @@ TSharedRef<SWidget> SRigVMActionMenu::OnCreateWidgetForAction(FCreateWidgetForAc
 	FSlateColor        SecondaryIconColor   = FSlateColor::UseForeground();
 	FText			   IconToolTip = GraphAction->GetTooltipDescription();
 	FString			   IconDocLink, IconDocExcerpt;
-	//GetPaletteItemIcon(GraphAction, Blueprint, IconBrush, IconColor, IconToolTip, IconDocLink, IconDocExcerpt, SecondaryBrush, SecondaryIconColor);
+	
+	// Get Palette Item Icon
+	{
+		// Default to tooltip based on action supplied
+		IconToolTip = GraphAction->GetTooltipDescription().IsEmpty() ? GraphAction->GetMenuDescription() : GraphAction->GetTooltipDescription();
+
+		if (GraphAction->GetTypeId() == FRigVMActionMenuItem::StaticGetTypeId())
+		{
+			FRigVMActionMenuItem* NodeSpawnerAction = (FRigVMActionMenuItem*)GraphAction.Get();
+			IconBrush = NodeSpawnerAction->GetMenuIcon(IconColor);
+
+			TSubclassOf<UEdGraphNode> VarNodeClass = NodeSpawnerAction->GetRawAction()->NodeClass;
+
+			// if the node is a variable getter or setter, use the variable icon instead, because maps need two brushes
+			if (*VarNodeClass && VarNodeClass->IsChildOf(UK2Node_Variable::StaticClass()))
+			{
+				const UK2Node_Variable* TemplateNode = Cast<UK2Node_Variable>(NodeSpawnerAction->GetRawAction()->GetTemplateNode());
+				FProperty* Property = TemplateNode->GetPropertyForVariable();
+				IconBrush = FRigVMEditorBase::GetVarIconAndColorFromProperty(Property, IconColor, SecondaryBrush, SecondaryIconColor);
+			}
+		}
+	}
 	TSharedRef<SWidget> IconWidget = SPinTypeSelector::ConstructPinTypeImage(
 		IconBrush, 
 		IconColor, 
