Fixed Control Rtig type registry performing a full engine type scan when an user defined struct or enum is missing from startup type scan. Now it will try to find the specific type and register it to get the index.

diff --git a/Engine/Plugins/Runtime/RigVM/Source/RigVM/Private/RigVMCore/RigVMRegistry.cpp b/Engine/Plugins/Runtime/RigVM/Source/RigVM/Private/RigVMCore/RigVMRegistry.cpp
--- a/Engine/Plugins/Runtime/RigVM/Source/RigVM/Private/RigVMCore/RigVMRegistry.cpp
+++ b/Engine/Plugins/Runtime/RigVM/Source/RigVM/Private/RigVMCore/RigVMRegistry.cpp
@@ -1114,14 +1114,50 @@ TRigVMTypeIndex FRigVMRegistry_NoLock::GetTypeIndexFromCPPType_NoLock(const FStr
 		Result = Types.IndexOfByPredicate(Predicate);
 
 		// in game / non-editor it's possible that a user defined struct or enum 
-		// has not been registered. thus we'll call RefreshEngineTypes to bring
-		// things up to date here. 
+		// has not been registered. thus we'll try to find it and if not,
+		// we will call RefreshEngineTypes to bring things up to date here.
 		if(Result == INDEX_NONE)
 		{
-			// we may need to update the types again to registry potentially
-			// missing predicate types 
-			const_cast<FRigVMRegistry_NoLock*>(this)->RefreshEngineTypes_NoLock();
-			Result = Types.IndexOfByPredicate(Predicate);
+			TGuardValue<bool> EnableGuardRefresh(const_cast<FRigVMRegistry_NoLock*>(this)->bIsRefreshingEngineTypes, true); // avoid a full redo of template permutations
+
+			for (TObjectIterator<UUserDefinedStruct> ScriptStructIt; ScriptStructIt; ++ScriptStructIt)
+			{
+				UUserDefinedStruct* ScriptStruct = *ScriptStructIt;
+				FRigVMTemplateArgumentType ArgumentType(ScriptStruct);
+				if (ArgumentType.CPPType == CPPTypeName)
+				{
+					// this check for example makes sure we don't add structs defined in verse
+					if (IsAllowedType_NoLock(ScriptStruct))
+					{
+						Result = const_cast<FRigVMRegistry_NoLock*>(this)->FindOrAddType_NoLock(ArgumentType, false);
+						break;
+					}
+				}
+			}
+			if (Result == INDEX_NONE) // if we can not find a struct, lets try an enum
+			{
+				for (TObjectIterator<UUserDefinedEnum> EnumIt; EnumIt; ++EnumIt)
+				{
+					UUserDefinedEnum* Enum = *EnumIt;
+					FRigVMTemplateArgumentType ArgumentType(Enum);
+					if (ArgumentType.CPPType == CPPTypeName)
+					{
+						// this check for example makes sure we don't add enums defined in verse
+						if (IsAllowedType_NoLock(Enum))
+						{
+							Result = const_cast<FRigVMRegistry_NoLock*>(this)->FindOrAddType_NoLock(ArgumentType, false);
+							break;
+						}
+					}
+				}
+			}
+			if (Result == INDEX_NONE) // else a full scan
+			{
+				// we may need to update the types again to registry potentially
+				// missing predicate types 
+				const_cast<FRigVMRegistry_NoLock*>(this)->RefreshEngineTypes_NoLock();
+				Result = Types.IndexOfByPredicate(Predicate);
+			}
 		}
 
 		// If not found, try to find a redirect
