RigVM: Refactor public context struct to use FTopLevelAssetPath instead of FString

#rnx

diff --git a/Engine/Plugins/Runtime/RigVM/Source/RigVM/Private/RigVMCore/RigVM.cpp b/Engine/Plugins/Runtime/RigVM/Source/RigVM/Private/RigVMCore/RigVM.cpp
--- a/Engine/Plugins/Runtime/RigVM/Source/RigVM/Private/RigVMCore/RigVM.cpp
+++ b/Engine/Plugins/Runtime/RigVM/Source/RigVM/Private/RigVMCore/RigVM.cpp
@@ -1561,12 +1561,12 @@ ERigVMExecuteResult URigVM::ExecuteVM(FRigVMExtendedExecuteContext& Context, con
 
 		Context.CurrentVMMemory = GetInstanceMemory(Context);
 
-		if (ByteCode.HasPublicContextPathName())
+		if (ByteCode.HasPublicContextAssetPath())
 		{
-			const FString ContextPublicDataPathName = Context.GetContextPublicDataStruct() != nullptr ? Context.GetContextPublicDataStruct()->GetPathName() : FString();
-			if (ByteCode.GetPublicContextPathName().IsEmpty() || ByteCode.GetPublicContextPathName() != ContextPublicDataPathName)
+			const FTopLevelAssetPath ContextPublicDataAssetPath(Context.GetContextPublicDataStruct());
+			if (ByteCode.GetPublicContextAssetPath().IsNull() || ByteCode.GetPublicContextAssetPath() != ContextPublicDataAssetPath)
 			{
-				UE_LOG(LogRigVM, Error, TEXT("Context PublicData [%s] does not match ByteCode Public Data [%s]. Likely a corrupt VM. Exiting."), *ContextPublicDataPathName, *ByteCode.GetPublicContextPathName());
+				UE_LOG(LogRigVM, Error, TEXT("Context PublicData [%s] does not match ByteCode Public Data [%s]. Likely a corrupt VM. Exiting."), *ContextPublicDataAssetPath.ToString(), *ByteCode.GetPublicContextAssetPath().ToString());
 				return Context.CurrentExecuteResult = ERigVMExecuteResult::Failed;
 			}
 		}
diff --git a/Engine/Plugins/Runtime/RigVM/Source/RigVM/Private/RigVMCore/RigVMByteCode.cpp b/Engine/Plugins/Runtime/RigVM/Source/RigVM/Private/RigVMCore/RigVMByteCode.cpp
--- a/Engine/Plugins/Runtime/RigVM/Source/RigVM/Private/RigVMCore/RigVMByteCode.cpp
+++ b/Engine/Plugins/Runtime/RigVM/Source/RigVM/Private/RigVMCore/RigVMByteCode.cpp
@@ -472,8 +472,8 @@ void FRigVMByteCode::Save(FArchive& Ar)
 	Ar << TempBranchInfos;
 	UE_RIGVM_ARCHIVETRACE_ENTRY(Ar, TEXT("BranchInfos"));
 
-	Ar << PublicContextPathName;
-	UE_RIGVM_ARCHIVETRACE_ENTRY(Ar, TEXT("PublicContextPathName"));
+	Ar << PublicContextAssetPath;
+	UE_RIGVM_ARCHIVETRACE_ENTRY(Ar, TEXT("PublicContextAssetPath"));
 }
 
 void FRigVMByteCode::Load(FArchive& Ar)
@@ -681,10 +681,18 @@ void FRigVMByteCode::Load(FArchive& Ar)
 		BranchInfoLookup.Reset();
 	}
 
-	if (Ar.CustomVer(FRigVMObjectVersion::GUID) >= FRigVMObjectVersion::VMBytecodeStorePublicContextPath)
+	if (Ar.CustomVer(FRigVMObjectVersion::GUID) >= FRigVMObjectVersion::VMBytecodeStorePublicContextPathAsTopLevelAssetPath)
 	{
+		Ar << PublicContextAssetPath;
+		bHasPublicContextAssetPath = true;
+	}
+	else if (Ar.CustomVer(FRigVMObjectVersion::GUID) >= FRigVMObjectVersion::VMBytecodeStorePublicContextPath)
+	{
+		FString PublicContextPathName; 
 		Ar << PublicContextPathName;
-		bHasPublicContextPathName = true;
+
+		PublicContextAssetPath = FTopLevelAssetPath(PublicContextPathName);
+		bHasPublicContextAssetPath = true;
 	}
 }
 
@@ -696,7 +704,7 @@ void FRigVMByteCode::Reset()
 	Entries.Reset();
 	BranchInfos.Reset();
 	BranchInfoLookup.Reset();
-	PublicContextPathName.Reset();
+	PublicContextAssetPath.Reset();
 
 #if WITH_EDITORONLY_DATA
 	SubjectPerInstruction.Reset();
@@ -717,7 +725,7 @@ void FRigVMByteCode::Empty()
 	bByteCodeIsAligned = false;
 	NumInstructions = 0;
 	Entries.Empty();
-	PublicContextPathName.Empty();
+	PublicContextAssetPath.Reset();
 
 #if WITH_EDITORONLY_DATA
 	SubjectPerInstruction.Empty();
diff --git a/Engine/Plugins/Runtime/RigVM/Source/RigVM/Public/RigVMCore/RigVMByteCode.h b/Engine/Plugins/Runtime/RigVM/Source/RigVM/Public/RigVMCore/RigVMByteCode.h
--- a/Engine/Plugins/Runtime/RigVM/Source/RigVM/Public/RigVMCore/RigVMByteCode.h
+++ b/Engine/Plugins/Runtime/RigVM/Source/RigVM/Public/RigVMCore/RigVMByteCode.h
@@ -1353,20 +1353,20 @@ public:
 
 	FString DumpToText() const;
 
-	bool HasPublicContextPathName() const
+	bool HasPublicContextAssetPath() const
 	{
-		return bHasPublicContextPathName;
+		return bHasPublicContextAssetPath;
 	}
 
-	const FString& GetPublicContextPathName() const
+	const FTopLevelAssetPath& GetPublicContextAssetPath() const
 	{
-		return PublicContextPathName;
+		return PublicContextAssetPath;
 	}
 
-	void SetPublicContextPathName(const FString& InPublicContextPathName)
+	void SetPublicContextAssetPath(const FTopLevelAssetPath& InPublicContextAssetPath)
 	{
-		PublicContextPathName = InPublicContextPathName;
-		bHasPublicContextPathName = true;
+		PublicContextAssetPath = InPublicContextAssetPath;
+		bHasPublicContextAssetPath = true;
 	}
 
 #if WITH_EDITOR
@@ -1538,7 +1538,7 @@ private:
 	TArray<FRigVMPredicateBranch> PredicateBranches;
 
 	UPROPERTY()
-	FString PublicContextPathName;
+	FTopLevelAssetPath PublicContextAssetPath;
 
 	const FRigVMBranchInfo* GetBranchInfo(const FRigVMBranchInfoKey& InBranchInfoKey) const;
 	mutable TMap<FRigVMBranchInfoKey, const FRigVMBranchInfo*> BranchInfoLookup;
@@ -1546,7 +1546,7 @@ private:
 	// if this is set to true the stored bytecode is aligned / padded
 	bool bByteCodeIsAligned;
 	// If the serialization has loaded a PublicContextPathName, so we check on new versions and skip check on older
-	bool bHasPublicContextPathName = false;
+	bool bHasPublicContextAssetPath = false;
 
 	static TArray<int32> EmptyInstructionIndices;
 
diff --git a/Engine/Plugins/Runtime/RigVM/Source/RigVM/Public/RigVMObjectVersion.h b/Engine/Plugins/Runtime/RigVM/Source/RigVM/Public/RigVMObjectVersion.h
--- a/Engine/Plugins/Runtime/RigVM/Source/RigVM/Public/RigVMObjectVersion.h
+++ b/Engine/Plugins/Runtime/RigVM/Source/RigVM/Public/RigVMObjectVersion.h
@@ -60,6 +60,9 @@ struct RIGVM_API FRigVMObjectVersion
 		// Storing function graph collapse node content as part of the header
 		RigVMSaveSerializedGraphInGraphFunctionDataAsByteArray,
 
+		// VM Bytecode Stores the Public Context Path as a FTopLevelAssetPath 
+     	VMBytecodeStorePublicContextPathAsTopLevelAssetPath,
+
 		// -----<new versions can be added above this line>-------------------------------------------------
 		VersionPlusOne,
 		LatestVersion = VersionPlusOne - 1
diff --git a/Engine/Plugins/Runtime/RigVM/Source/RigVMDeveloper/Private/RigVMCompiler/RigVMCompiler.cpp b/Engine/Plugins/Runtime/RigVM/Source/RigVMDeveloper/Private/RigVMCompiler/RigVMCompiler.cpp
--- a/Engine/Plugins/Runtime/RigVM/Source/RigVMDeveloper/Private/RigVMCompiler/RigVMCompiler.cpp
+++ b/Engine/Plugins/Runtime/RigVM/Source/RigVMDeveloper/Private/RigVMCompiler/RigVMCompiler.cpp
@@ -1555,7 +1555,7 @@ bool URigVMCompiler::Compile(const FRigVMCompileSettings& InSettings, TArray<URi
 		CompileTimer.Stop();
 
 		WorkData.VM->SetVMHash(WorkData.VM->ComputeVMHash());
-		WorkData.VM->GetByteCode().SetPublicContextPathName((WorkData.Context&& WorkData.Context->GetContextPublicDataStruct() != nullptr) ? WorkData.Context->GetContextPublicDataStruct()->GetPathName() : FString());
+		WorkData.VM->GetByteCode().SetPublicContextAssetPath(FTopLevelAssetPath(WorkData.Context ? WorkData.Context->GetContextPublicDataStruct() : nullptr));
 	}
 
 	return true;
