Control Rig: Add node comment customization
#rnx

diff --git a/Engine/Plugins/Runtime/RigVM/Source/RigVMEditor/Private/Editor/RigVMCommentNodeDetailCustomization.cpp b/Engine/Plugins/Runtime/RigVM/Source/RigVMEditor/Private/Editor/RigVMCommentNodeDetailCustomization.cpp
new file mode 100644
--- /dev/null
+++ b/Engine/Plugins/Runtime/RigVM/Source/RigVMEditor/Private/Editor/RigVMCommentNodeDetailCustomization.cpp
@@ -0,0 +1,317 @@
+﻿// Copyright Epic Games, Inc. All Rights Reserved.
+
+#include "Editor/RigVMCommentNodeDetailCustomization.h"
+#include "DetailWidgetRow.h"
+#include "DetailLayoutBuilder.h"
+#include "IDetailChildrenBuilder.h"
+#include "Widgets/Input/SCheckBox.h"
+#include "Widgets/Input/STextComboBox.h"
+#include "RigVMBlueprint.h"
+#include "RigVMCore/RigVM.h"
+#include "Widgets/Colors/SColorBlock.h"
+#include "Widgets/Colors/SColorPicker.h"
+#include "UnrealEngine.h"
+#include "Widgets/Input/SNumericEntryBox.h"
+#include "Widgets/Input/SEditableTextBox.h"
+
+#define LOCTEXT_NAMESPACE "CommentNodeDetails"
+
+void FRigVMCommentNodeDetailCustomization::CustomizeDetails(IDetailLayoutBuilder& DetailBuilder)
+{
+	ObjectsBeingCustomized.Reset();
+
+	TArray<TWeakObjectPtr<UObject>> DetailObjects;
+	DetailBuilder.GetObjectsBeingCustomized(DetailObjects);
+	for(TWeakObjectPtr<UObject> DetailObject : DetailObjects)
+	{
+		URigVMCommentNode* Node = CastChecked<URigVMCommentNode>(DetailObject.Get());
+		ObjectsBeingCustomized.Add(Node);
+	}
+
+	if(ObjectsBeingCustomized[0].IsValid())
+	{
+		BlueprintBeingCustomized = ObjectsBeingCustomized[0]->GetTypedOuter<URigVMBlueprint>();
+	}
+
+	IDetailCategoryBuilder& Category = DetailBuilder.EditCategory(TEXT("Comment Node"));
+	
+	const UEdGraphSchema* Schema = GetDefault<URigVMEdGraphSchema>();
+
+	const FSlateFontInfo DetailFontInfo = IDetailLayoutBuilder::GetDetailFont();
+	Category.AddCustomRow( LOCTEXT("CommentNodeText", "Comment Text") )
+	.NameContent()
+	[
+		SNew(STextBlock)
+		.Text(LOCTEXT("CommentNodeText", "Comment Text"))
+		.Font(IDetailLayoutBuilder::GetDetailFont())
+	]
+	.ValueContent()
+	.MaxDesiredWidth(250.0f)
+	[
+		SNew(SEditableTextBox)
+		.Font(IDetailLayoutBuilder::GetDetailFont())
+		.Text(this, &FRigVMCommentNodeDetailCustomization::GetText)
+		.OnTextCommitted(this, &FRigVMCommentNodeDetailCustomization::SetText)
+	];
+
+	Category.AddCustomRow( LOCTEXT("CommentNodeColor", "Comment Color") )
+	.NameContent()
+	[
+		SNew(STextBlock)
+		.Text(LOCTEXT("CommentNodeColor", "Comment Color"))
+		.Font(IDetailLayoutBuilder::GetDetailFont())
+	]
+	.ValueContent()
+	.MaxDesiredWidth(250.0f)
+	[
+		SNew(SColorBlock)
+		.Color(this, &FRigVMCommentNodeDetailCustomization::GetColor)
+		.OnMouseButtonDown(this, &FRigVMCommentNodeDetailCustomization::OnChooseColor)
+		.AlphaBackgroundBrush(FAppStyle::Get().GetBrush("ColorPicker.RoundedAlphaBackground"))
+		.ShowBackgroundForAlpha(true)
+		.AlphaDisplayMode(EColorBlockAlphaDisplayMode::Ignore)
+		.Size(FVector2D(70.0f, 20.0f))
+		.CornerRadius(FVector4(4.0f,4.0f,4.0f,4.0f))
+	];
+
+	Category.AddCustomRow( LOCTEXT("CommentNodeShowBubble", "Show Bubble") )
+	.NameContent()
+	[
+		SNew(STextBlock)
+		.Text(LOCTEXT("CommentNodeShowBubble", "Show Bubble"))
+		.Font(IDetailLayoutBuilder::GetDetailFont())
+	]
+	.ValueContent()
+	.MaxDesiredWidth(250.0f)
+	[
+		SNew(SCheckBox)
+		.IsChecked(this, &FRigVMCommentNodeDetailCustomization::IsShowingBubbleEnabled)
+		.OnCheckStateChanged(this, &FRigVMCommentNodeDetailCustomization::OnShowingBubbleStateChanged)
+	];
+
+	Category.AddCustomRow( LOCTEXT("CommentNodeColorBubble", "Color Bubble") )
+	.NameContent()
+	[
+		SNew(STextBlock)
+		.Text(LOCTEXT("CommentNodeColorBubble", "Color Bubble"))
+		.Font(IDetailLayoutBuilder::GetDetailFont())
+	]
+	.ValueContent()
+	.MaxDesiredWidth(250.0f)
+	[
+		SNew(SCheckBox)
+		.IsChecked(this, &FRigVMCommentNodeDetailCustomization::IsColorBubbleEnabled)
+		.OnCheckStateChanged(this, &FRigVMCommentNodeDetailCustomization::OnColorBubbleStateChanged)
+	];
+
+	Category.AddCustomRow( LOCTEXT("CommentNodeFontSize", "Font Size") )
+	.NameContent()
+	[
+		SNew(STextBlock)
+		.Text(LOCTEXT("CommentNodeFontSize", "Font Size"))
+		.Font(IDetailLayoutBuilder::GetDetailFont())
+	]
+	.ValueContent()
+	.MaxDesiredWidth(250.0f)
+	[
+		SNew(SNumericEntryBox<int32>)
+		.Value(this, &FRigVMCommentNodeDetailCustomization::GetFontSize)
+		.OnValueCommitted(this, &FRigVMCommentNodeDetailCustomization::OnFontSizeChanged)
+	];
+}
+
+void FRigVMCommentNodeDetailCustomization::GetValuesFromNode(TWeakObjectPtr<URigVMCommentNode> WeakNode)
+{
+	if (TStrongObjectPtr<URigVMCommentNode> Node = WeakNode.Pin())
+	{
+		CommentText = Node->GetCommentText();
+		bShowingBubble = Node->GetCommentBubbleVisible();
+		bBubbleColorEnabled = Node->GetCommentColorBubble();
+		FontSize = Node->GetCommentFontSize();
+	}
+}
+
+void FRigVMCommentNodeDetailCustomization::SetValues(TWeakObjectPtr<URigVMCommentNode> WeakNode)
+{
+	if (BlueprintBeingCustomized)
+	{
+		if (TStrongObjectPtr<URigVMCommentNode> Node = WeakNode.Pin())
+		{
+			URigVMController* Controller = BlueprintBeingCustomized->GetController(Node->GetGraph());
+			Controller->SetCommentText(Node.Get(), CommentText, FontSize, bShowingBubble, bBubbleColorEnabled, true, true);
+		}
+	}
+}
+
+FText FRigVMCommentNodeDetailCustomization::GetText() const
+{
+	if (!ObjectsBeingCustomized.IsEmpty())
+	{
+		FString Value = ObjectsBeingCustomized[0]->GetCommentText();
+		for (int32 i=1; i<ObjectsBeingCustomized.Num(); ++i)
+		{
+			if (ObjectsBeingCustomized[i].IsValid())
+			{
+				if (!Value.Equals(ObjectsBeingCustomized[i]->GetCommentText()))
+				{
+					Value = TEXT("Multiple Values");
+					break;
+				}
+			}
+		}
+		return FText::FromString(Value);
+	}
+	return FText::GetEmpty();
+}
+
+void FRigVMCommentNodeDetailCustomization::SetText(const FText& InNewText, ETextCommit::Type InCommitType)
+{
+	for (TWeakObjectPtr<URigVMCommentNode> WeakNode : ObjectsBeingCustomized)
+	{
+		GetValuesFromNode(WeakNode);
+		CommentText = InNewText.ToString();
+		SetValues(WeakNode);
+	}
+}
+
+FLinearColor FRigVMCommentNodeDetailCustomization::GetColor() const
+{
+	if (!ObjectsBeingCustomized.IsEmpty())
+	{
+		FLinearColor Value = ObjectsBeingCustomized[0]->GetNodeColor();
+		for (int32 i=1; i<ObjectsBeingCustomized.Num(); ++i)
+		{
+			if (ObjectsBeingCustomized[i].IsValid())
+			{
+				if (!Value.Equals(ObjectsBeingCustomized[i]->GetNodeColor()))
+				{
+					Value = FLinearColor::Black;
+					break;
+				}
+			}
+		}
+		return Value;
+	}
+	return FLinearColor::Black;
+}
+
+FReply FRigVMCommentNodeDetailCustomization::OnChooseColor(const FGeometry& MyGeometry, const FPointerEvent& MouseEvent)
+{
+	FColorPickerArgs PickerArgs;
+	PickerArgs.bUseAlpha = false;
+	PickerArgs.DisplayGamma = TAttribute<float>::Create(TAttribute<float>::FGetter::CreateUObject(GEngine, &UEngine::GetDisplayGamma));
+	PickerArgs.InitialColor = GetColor();
+	PickerArgs.OnColorCommitted = FOnLinearColorValueChanged::CreateSP(this, &FRigVMCommentNodeDetailCustomization::OnColorPicked);
+	OpenColorPicker(PickerArgs);
+	return FReply::Handled();
+}
+
+void FRigVMCommentNodeDetailCustomization::OnColorPicked(FLinearColor LinearColor)
+{
+	if (BlueprintBeingCustomized)
+	{
+		for (TWeakObjectPtr<URigVMCommentNode> WeakNode : ObjectsBeingCustomized)
+		{
+			if (TStrongObjectPtr<URigVMCommentNode> Node = WeakNode.Pin())
+			{
+				URigVMController* Controller = BlueprintBeingCustomized->GetController(Node->GetGraph());
+				Controller->SetNodeColor(Node.Get(), LinearColor, true, true);
+				
+			}
+		}
+	}
+}
+
+ECheckBoxState FRigVMCommentNodeDetailCustomization::IsShowingBubbleEnabled() const
+{
+	if (!ObjectsBeingCustomized.IsEmpty())
+	{
+		bool Value = ObjectsBeingCustomized[0]->GetCommentBubbleVisible();
+		for (int32 i=1; i<ObjectsBeingCustomized.Num(); ++i)
+		{
+			if (ObjectsBeingCustomized[i].IsValid())
+			{
+				if (Value != (ObjectsBeingCustomized[i]->GetCommentBubbleVisible()))
+				{
+					Value = false;
+					break;
+				}
+			}
+		}
+		return Value ? ECheckBoxState::Checked : ECheckBoxState::Unchecked;
+	}
+	return ECheckBoxState::Unchecked;
+}
+
+void FRigVMCommentNodeDetailCustomization::OnShowingBubbleStateChanged(ECheckBoxState InValue)
+{
+	for (TWeakObjectPtr<URigVMCommentNode> WeakNode : ObjectsBeingCustomized)
+	{
+		GetValuesFromNode(WeakNode);
+		bShowingBubble = InValue == ECheckBoxState::Checked ? true : false;
+		SetValues(WeakNode);
+	}
+}
+
+ECheckBoxState FRigVMCommentNodeDetailCustomization::IsColorBubbleEnabled() const
+{
+	if (!ObjectsBeingCustomized.IsEmpty())
+	{
+		bool Value = ObjectsBeingCustomized[0]->GetCommentColorBubble();
+		for (int32 i=1; i<ObjectsBeingCustomized.Num(); ++i)
+		{
+			if (ObjectsBeingCustomized[i].IsValid())
+			{
+				if (Value != (ObjectsBeingCustomized[i]->GetCommentColorBubble()))
+				{
+					Value = false;
+					break;
+				}
+			}
+		}
+		return Value ? ECheckBoxState::Checked : ECheckBoxState::Unchecked;
+	}
+	return ECheckBoxState::Unchecked;
+}
+
+void FRigVMCommentNodeDetailCustomization::OnColorBubbleStateChanged(ECheckBoxState InValue)
+{
+	for (TWeakObjectPtr<URigVMCommentNode> WeakNode : ObjectsBeingCustomized)
+	{
+		GetValuesFromNode(WeakNode);
+		bBubbleColorEnabled = InValue == ECheckBoxState::Checked ? true : false;
+		SetValues(WeakNode);
+	}
+}
+
+TOptional<int> FRigVMCommentNodeDetailCustomization::GetFontSize() const
+{
+	if (!ObjectsBeingCustomized.IsEmpty())
+	{
+		TOptional<int32> Value = ObjectsBeingCustomized[0]->GetCommentFontSize();
+		for (int32 i=1; i<ObjectsBeingCustomized.Num(); ++i)
+		{
+			if (ObjectsBeingCustomized[i].IsValid())
+			{
+				if (Value != (ObjectsBeingCustomized[i]->GetCommentFontSize()))
+				{
+					return TOptional<int32>();
+				}
+			}
+		}
+		return Value;
+	}
+	return TOptional<int32>();
+}
+
+void FRigVMCommentNodeDetailCustomization::OnFontSizeChanged(int32 InValue, ETextCommit::Type Arg)
+{
+	for (TWeakObjectPtr<URigVMCommentNode> WeakNode : ObjectsBeingCustomized)
+	{
+		GetValuesFromNode(WeakNode);
+		FontSize = InValue;
+		SetValues(WeakNode);
+	}
+}
+
+#undef LOCTEXT_NAMESPACE
diff --git a/Engine/Plugins/Runtime/RigVM/Source/RigVMEditor/Private/Editor/RigVMEditor.cpp b/Engine/Plugins/Runtime/RigVM/Source/RigVMEditor/Private/Editor/RigVMEditor.cpp
--- a/Engine/Plugins/Runtime/RigVM/Source/RigVMEditor/Private/Editor/RigVMEditor.cpp
+++ b/Engine/Plugins/Runtime/RigVM/Source/RigVMEditor/Private/Editor/RigVMEditor.cpp
@@ -2888,14 +2888,23 @@ void FRigVMEditorBase::SetDetailObjects(const TArray<UObject*>& InObjects, bool
 		}
 		else if (URigVMCommentNode* CommentNode = Cast<URigVMCommentNode>(InObject))
 		{
-			if (URigVMEdGraph* EdGraph = Cast<URigVMEdGraph>(GetRigVMBlueprint()->GetEdGraph(CastChecked<URigVMNode>(InObject)->GetGraph())))
+			TSharedRef<FAssetEditorToolkit> SharedApp = GetHostingApp().ToSharedRef();
+			if (SharedApp->IsBlueprintEditor())
 			{
-				if(UEdGraphNode* EdGraphNode = EdGraph->FindNodeForModelNodeName(CommentNode->GetFName()))
+				if (URigVMEdGraph* EdGraph = Cast<URigVMEdGraph>(GetRigVMBlueprint()->GetEdGraph(CastChecked<URigVMNode>(InObject)->GetGraph())))
 				{
-					FilteredObjects.AddUnique(EdGraphNode);
-					continue;
+					if(UEdGraphNode* EdGraphNode = EdGraph->FindNodeForModelNodeName(CommentNode->GetFName()))
+					{
+						FilteredObjects.AddUnique(EdGraphNode);
+						continue;
+					}
 				}
 			}
+			else
+			{
+				FilteredObjects.AddUnique(CommentNode);
+				continue;
+			}
 		}
 
 		if (URigVMNode* ModelNode = Cast<URigVMNode>(InObject))
diff --git a/Engine/Plugins/Runtime/RigVM/Source/RigVMEditor/Private/Editor/SRigVMDetailsInspector.cpp b/Engine/Plugins/Runtime/RigVM/Source/RigVMEditor/Private/Editor/SRigVMDetailsInspector.cpp
--- a/Engine/Plugins/Runtime/RigVM/Source/RigVMEditor/Private/Editor/SRigVMDetailsInspector.cpp
+++ b/Engine/Plugins/Runtime/RigVM/Source/RigVMEditor/Private/Editor/SRigVMDetailsInspector.cpp
@@ -9,6 +9,7 @@
 #include "PropertyEditorModule.h"
 #include "Editor/RigVMGraphDetailCustomization.h"
 #include "Editor/RigVMVariableDetailCustomization.h"
+#include "Editor/RigVMCommentNodeDetailCustomization.h"
 #include "Framework/Application/SlateApplication.h"
 #include "K2Node.h"
 
@@ -136,6 +137,7 @@ void SRigVMDetailsInspector::OnEditorClose(const IRigVMEditor* RigVMEditorBase,
 {
 	PropertyView->UnregisterInstancedCustomPropertyLayout(URigVMEdGraph::StaticClass());
 	PropertyView->UnregisterInstancedCustomPropertyLayout(UPropertyWrapper::StaticClass());
+	PropertyView->UnregisterInstancedCustomPropertyLayout(URigVMCommentNode::StaticClass());
 }
 
 void SRigVMDetailsInspector::AddReferencedObjects(FReferenceCollector& Collector)
@@ -246,6 +248,9 @@ void SRigVMDetailsInspector::Construct(const FArguments& InArgs)
 	FOnGetDetailCustomizationInstance LayoutVariableDetails = FOnGetDetailCustomizationInstance::CreateStatic(&FRigVMVariableDetailCustomization::MakeInstance, Editor);
 	PropertyView->RegisterInstancedCustomPropertyLayout(UPropertyWrapper::StaticClass(), LayoutVariableDetails);
 
+	FOnGetDetailCustomizationInstance CommentNodeDetails = FOnGetDetailCustomizationInstance::CreateStatic(&FRigVMCommentNodeDetailCustomization::MakeInstance);
+	PropertyView->RegisterInstancedCustomPropertyLayout(URigVMCommentNode::StaticClass(), CommentNodeDetails);
+
 	Editor->OnEditorClosed().AddSP(this, &SRigVMDetailsInspector::OnEditorClose);
 
 	// Create the border that all of the content will get stuffed into
diff --git a/Engine/Plugins/Runtime/RigVM/Source/RigVMEditor/Public/Editor/RigVMCommentNodeDetailCustomization.h b/Engine/Plugins/Runtime/RigVM/Source/RigVMEditor/Public/Editor/RigVMCommentNodeDetailCustomization.h
new file mode 100644
--- /dev/null
+++ b/Engine/Plugins/Runtime/RigVM/Source/RigVMEditor/Public/Editor/RigVMCommentNodeDetailCustomization.h
@@ -0,0 +1,65 @@
+﻿// Copyright Epic Games, Inc. All Rights Reserved.
+
+#pragma once
+
+#include "CoreMinimal.h"
+#include "IDetailCustomization.h"
+#include "IPropertyTypeCustomization.h"
+#include "RigVMHost.h"
+#include "RigVMBlueprint.h"
+#include "Styling/SlateTypes.h"
+#include "IPropertyUtilities.h"
+
+class IRigVMEditor;
+class IPropertyHandle;
+
+class RIGVMEDITOR_API FRigVMCommentNodeDetailCustomization : public IDetailCustomization
+{
+	FRigVMCommentNodeDetailCustomization()
+	: BlueprintBeingCustomized(nullptr)
+	{}
+
+	
+public:
+
+	// Makes a new instance of this detail layout class for a specific detail view requesting it
+	static TSharedRef<IDetailCustomization> MakeInstance()
+	{
+		return MakeShareable(new FRigVMCommentNodeDetailCustomization);
+	}
+
+	/** IDetailCustomization interface */
+	virtual void CustomizeDetails(IDetailLayoutBuilder& DetailBuilder) override;
+
+private:
+
+	URigVMBlueprint* BlueprintBeingCustomized;
+	TArray<TWeakObjectPtr<URigVMCommentNode>> ObjectsBeingCustomized;
+
+	void GetValuesFromNode(TWeakObjectPtr<URigVMCommentNode> Node);
+	void SetValues(TWeakObjectPtr<URigVMCommentNode> Node);
+	
+	FText GetText() const;
+	void SetText(const FText& InNewText, ETextCommit::Type InCommitType);
+
+	FLinearColor GetColor() const;
+	FReply OnChooseColor(const FGeometry& MyGeometry, const FPointerEvent& MouseEvent);
+	void OnColorPicked(FLinearColor LinearColor);
+
+	ECheckBoxState IsShowingBubbleEnabled() const;
+	void OnShowingBubbleStateChanged(ECheckBoxState InValue);
+
+	ECheckBoxState IsColorBubbleEnabled() const;
+	void OnColorBubbleStateChanged(ECheckBoxState InValue);
+
+	TOptional<int> GetFontSize() const;
+	void OnFontSizeChanged(int32 InValue, ETextCommit::Type Arg);
+	
+	FString CommentText;
+	bool bShowingBubble;
+	bool bBubbleColorEnabled;
+	int32 FontSize;
+};
+
+
+
