RigVM UI Changes to support Overrides for pins on nodes

#rnx

diff --git a/Engine/Plugins/Runtime/RigVM/Source/RigVMDeveloper/Private/RigVMModel/RigVMController.cpp b/Engine/Plugins/Runtime/RigVM/Source/RigVMDeveloper/Private/RigVMModel/RigVMController.cpp
--- a/Engine/Plugins/Runtime/RigVM/Source/RigVMDeveloper/Private/RigVMModel/RigVMController.cpp
+++ b/Engine/Plugins/Runtime/RigVM/Source/RigVMDeveloper/Private/RigVMModel/RigVMController.cpp
@@ -10063,6 +10063,551 @@ bool URigVMController::ResetPinDefaultValue(const FString& InPinPath, bool bSetu
 	return bSuccess;
 }
 
+bool URigVMController::ResetDefaultValueForPins(const TArray<FString>& InPinPaths, bool bSetupUndo, bool bPrintPythonCommand)
+{
+	if (!IsValidGraph() || InPinPaths.IsEmpty())
+	{
+		return false;
+	}
+
+	if (!bIsTransacting && !IsGraphEditable())
+	{
+		return false;
+	}
+
+	URigVMGraph* Graph = GetGraph();
+	check(Graph);
+
+	OpenUndoBracket(TEXT("Reset Default Value On Pins"));
+
+	int32 NumPinsAffected = 0;
+	for(const FString& PinPath : InPinPaths)
+	{
+		if(ResetPinDefaultValue(PinPath, bSetupUndo, false))
+		{
+			NumPinsAffected++;
+		}
+	}
+
+	if(NumPinsAffected == 0)
+	{
+		CancelUndoBracket();
+		return false;
+	}
+
+	CloseUndoBracket();
+	
+	if (bPrintPythonCommand)
+	{
+		const FString GraphName = GetSchema()->GetSanitizedGraphName(GetGraph()->GetGraphName());
+
+		TArray<FString> SanitizedPinPaths;
+		SanitizedPinPaths.Reserve(InPinPaths.Num());
+		for(const FString& PinPath : InPinPaths)
+		{
+			SanitizedPinPaths.Add(FString::Printf(TEXT("'%s'"), *GetSchema()->GetSanitizedPinPath(PinPath)));
+		}
+
+		RigVMPythonUtils::Print(GetSchema()->GetGraphOuterName(GetGraph()),
+			FString::Printf(TEXT("blueprint.get_controller_by_name('%s').reset_default_value_for_pins([%s])"),
+			*GraphName,
+			*FString::Join(SanitizedPinPaths, TEXT(","))));
+	}
+
+	return true;
+}
+
+bool URigVMController::ResetDefaultValueForAllPinsOnNode(const FName& InNodeName, bool bSetupUndo, bool bPrintPythonCommand)
+{
+	if (!IsValidGraph())
+	{
+		return false;
+	}
+
+	if (!bIsTransacting && !IsGraphEditable())
+	{
+		return false;
+	}
+
+	URigVMGraph* Graph = GetGraph();
+	check(Graph);
+
+	URigVMNode* Node = Graph->FindNodeByName(InNodeName);
+	if (Node == nullptr)
+	{
+		ReportErrorf(TEXT("Cannot find node '%s'."), *InNodeName.ToString());
+		return false;
+	}
+
+	TArray<FString> PinPaths;
+	for(const URigVMPin* Pin : Node->GetPins())
+	{
+		if(Pin->CanProvideDefaultValue() && !Pin->HasOriginalDefaultValue())
+		{
+			PinPaths.Add(Pin->GetPinPath());
+		}
+	}
+
+	if(!ResetDefaultValueForPins(PinPaths, bSetupUndo, false))
+	{
+		return false;
+	}
+	
+	if (bPrintPythonCommand)
+	{
+		const FString GraphName = GetSchema()->GetSanitizedGraphName(GetGraph()->GetGraphName());
+
+		RigVMPythonUtils::Print(GetSchema()->GetGraphOuterName(GetGraph()),
+			FString::Printf(TEXT("blueprint.get_controller_by_name('%s').reset_default_value_for_all_pins_on_node('%s')"),
+			*GraphName,
+			*GetSchema()->GetSanitizedNodeName(InNodeName.ToString())));
+	}
+
+	return true;
+}
+
+bool URigVMController::ResetDefaultValueForAllPinsOnNodes(const TArray<FName>& InNodeNames, bool bSetupUndo, bool bPrintPythonCommand)
+{
+	if (!IsValidGraph())
+	{
+		return false;
+	}
+
+	if (!bIsTransacting && !IsGraphEditable())
+	{
+		return false;
+	}
+
+	OpenUndoBracket(TEXT("Reset Default Value For All Pins on Nodes"));
+
+	int32 NumNodesAffected = 0;
+	for(const FName& NodeName : InNodeNames)
+	{
+		if(ResetDefaultValueForAllPinsOnNode(NodeName, bSetupUndo, false))
+		{
+			NumNodesAffected++;
+		}
+	}
+
+	if(NumNodesAffected == 0)
+	{
+		CancelUndoBracket();
+		return false;
+	}
+
+	CloseUndoBracket();
+	
+	if (bPrintPythonCommand)
+	{
+		const FString GraphName = GetSchema()->GetSanitizedGraphName(GetGraph()->GetGraphName());
+
+		TArray<FString> SanitizedNodeNames;
+		SanitizedNodeNames.Reserve(InNodeNames.Num());
+		for(const FName& NodeName : InNodeNames)
+		{
+			SanitizedNodeNames.Add(FString::Printf(TEXT("'%s'"), *GetSchema()->GetSanitizedNodeName(NodeName.ToString())));
+		}
+		
+		RigVMPythonUtils::Print(GetSchema()->GetGraphOuterName(GetGraph()),
+			FString::Printf(TEXT("blueprint.get_controller_by_name('%s').reset_default_value_for_all_pins_on_nodes([%s])"),
+			*GraphName,
+			*FString::Join(SanitizedNodeNames, TEXT(","))));
+	}
+
+	return true;
+}
+
+bool URigVMController::AddOverrideToPin(const FString& InPinPath, bool bSetupUndo, bool bPrintPythonCommand)
+{
+	if (!IsValidGraph())
+	{
+		return false;
+	}
+
+	if (!bIsTransacting && !IsGraphEditable())
+	{
+		return false;
+	}
+
+	URigVMGraph* Graph = GetGraph();
+	check(Graph);
+
+	URigVMPin* Pin = Graph->FindPin(InPinPath);
+	if (Pin == nullptr)
+	{
+		ReportErrorf(TEXT("Cannot find pin '%s'."), *InPinPath);
+		return false;
+	}
+
+	FRigVMDefaultValueTypeGuard _(this, ERigVMPinDefaultValueType::Override, true);
+	if(!SetPinDefaultValue(Pin, Pin->GetDefaultValue(), false, bSetupUndo, false, false))
+	{
+		return false;
+	}
+	
+	if (bPrintPythonCommand)
+	{
+		const FString GraphName = GetSchema()->GetSanitizedGraphName(GetGraph()->GetGraphName());
+
+		RigVMPythonUtils::Print(GetSchema()->GetGraphOuterName(GetGraph()),
+			FString::Printf(TEXT("blueprint.get_controller_by_name('%s').add_override_to_pin('%s')"),
+			*GraphName,
+			*GetSchema()->GetSanitizedPinPath(InPinPath)));
+	}
+
+	return true;
+}
+
+bool URigVMController::AddOverrideToPins(const TArray<FString>& InPinPaths, bool bSetupUndo, bool bPrintPythonCommand)
+{
+	if (!IsValidGraph() || InPinPaths.IsEmpty())
+	{
+		return false;
+	}
+
+	if (!bIsTransacting && !IsGraphEditable())
+	{
+		return false;
+	}
+
+	URigVMGraph* Graph = GetGraph();
+	check(Graph);
+
+	OpenUndoBracket(TEXT("Add Override To Pins"));
+
+	int32 NumPinsAffected = 0;
+	for(const FString& PinPath : InPinPaths)
+	{
+		if(AddOverrideToPin(PinPath, bSetupUndo, false))
+		{
+			NumPinsAffected++;
+		}
+	}
+
+	if(NumPinsAffected == 0)
+	{
+		CancelUndoBracket();
+		return false;
+	}
+
+	CloseUndoBracket();
+	
+	if (bPrintPythonCommand)
+	{
+		const FString GraphName = GetSchema()->GetSanitizedGraphName(GetGraph()->GetGraphName());
+
+		TArray<FString> SanitizedPinPaths;
+		SanitizedPinPaths.Reserve(InPinPaths.Num());
+		for(const FString& PinPath : InPinPaths)
+		{
+			SanitizedPinPaths.Add(FString::Printf(TEXT("'%s'"), *GetSchema()->GetSanitizedPinPath(PinPath)));
+		}
+
+		RigVMPythonUtils::Print(GetSchema()->GetGraphOuterName(GetGraph()),
+			FString::Printf(TEXT("blueprint.get_controller_by_name('%s').add_override_to_pins([%s])"),
+			*GraphName,
+			*FString::Join(SanitizedPinPaths, TEXT(","))));
+	}
+
+	return true;
+}
+
+bool URigVMController::AddOverrideToAllPinsOnNode(const FName& InNodeName, bool bSetupUndo, bool bPrintPythonCommand)
+{
+	if (!IsValidGraph())
+	{
+		return false;
+	}
+
+	if (!bIsTransacting && !IsGraphEditable())
+	{
+		return false;
+	}
+
+	URigVMGraph* Graph = GetGraph();
+	check(Graph);
+
+	URigVMNode* Node = Graph->FindNodeByName(InNodeName);
+	if (Node == nullptr)
+	{
+		ReportErrorf(TEXT("Cannot find node '%s'."), *InNodeName.ToString());
+		return false;
+	}
+
+	TArray<FString> PinPaths;
+	for(const URigVMPin* Pin : Node->GetPins())
+	{
+		if(Pin->CanProvideDefaultValue() && Pin->GetDefaultValueType() != ERigVMPinDefaultValueType::Override)
+		{
+			PinPaths.Add(Pin->GetPinPath());
+		}
+	}
+
+	if(!AddOverrideToPins(PinPaths, bSetupUndo, false))
+	{
+		return false;
+	}
+	
+	if (bPrintPythonCommand)
+	{
+		const FString GraphName = GetSchema()->GetSanitizedGraphName(GetGraph()->GetGraphName());
+
+		RigVMPythonUtils::Print(GetSchema()->GetGraphOuterName(GetGraph()),
+			FString::Printf(TEXT("blueprint.get_controller_by_name('%s').add_override_to_all_pins_on_node('%s')"),
+			*GraphName,
+			*GetSchema()->GetSanitizedNodeName(InNodeName.ToString())));
+	}
+
+	return true;
+}
+
+bool URigVMController::AddOverrideToAllPinsOnNodes(const TArray<FName>& InNodeNames, bool bSetupUndo, bool bPrintPythonCommand)
+{
+	if (!IsValidGraph())
+	{
+		return false;
+	}
+
+	if (!bIsTransacting && !IsGraphEditable())
+	{
+		return false;
+	}
+
+	OpenUndoBracket(TEXT("Add Override To All Pins On Nodes"));
+
+	int32 NumNodesAffected = 0;
+	for(const FName& NodeName : InNodeNames)
+	{
+		if(AddOverrideToAllPinsOnNode(NodeName, bSetupUndo, false))
+		{
+			NumNodesAffected++;
+		}
+	}
+
+	if(NumNodesAffected == 0)
+	{
+		CancelUndoBracket();
+		return false;
+	}
+
+	CloseUndoBracket();
+	
+	if (bPrintPythonCommand)
+	{
+		const FString GraphName = GetSchema()->GetSanitizedGraphName(GetGraph()->GetGraphName());
+
+		TArray<FString> SanitizedNodeNames;
+		SanitizedNodeNames.Reserve(InNodeNames.Num());
+		for(const FName& NodeName : InNodeNames)
+		{
+			SanitizedNodeNames.Add(FString::Printf(TEXT("'%s'"), *GetSchema()->GetSanitizedNodeName(NodeName.ToString())));
+		}
+		
+		RigVMPythonUtils::Print(GetSchema()->GetGraphOuterName(GetGraph()),
+			FString::Printf(TEXT("blueprint.get_controller_by_name('%s').add_override_to_all_pins_on_nodes([%s])"),
+			*GraphName,
+			*FString::Join(SanitizedNodeNames, TEXT(","))));
+	}
+
+	return true;
+}
+
+bool URigVMController::ClearOverrideOnPin(const FString& InPinPath, bool bSetupUndo, bool bPrintPythonCommand)
+{
+	if (!IsValidGraph())
+	{
+		return false;
+	}
+
+	if (!bIsTransacting && !IsGraphEditable())
+	{
+		return false;
+	}
+
+	URigVMGraph* Graph = GetGraph();
+	check(Graph);
+
+	URigVMPin* Pin = Graph->FindPin(InPinPath);
+	if (Pin == nullptr)
+	{
+		ReportErrorf(TEXT("Cannot find pin '%s'."), *InPinPath);
+		return false;
+	}
+
+	FRigVMDefaultValueTypeGuard _(this, ERigVMPinDefaultValueType::Unset, true);
+	if(!ResetPinDefaultValue(Pin->GetPinPath(), bSetupUndo, false))
+	{
+		return false;
+	}
+	
+	if (bPrintPythonCommand)
+	{
+		const FString GraphName = GetSchema()->GetSanitizedGraphName(GetGraph()->GetGraphName());
+
+		RigVMPythonUtils::Print(GetSchema()->GetGraphOuterName(GetGraph()),
+			FString::Printf(TEXT("blueprint.get_controller_by_name('%s').clear_override_on_pin('%s')"),
+			*GraphName,
+			*GetSchema()->GetSanitizedPinPath(InPinPath)));
+	}
+
+	return true;
+}
+
+bool URigVMController::ClearOverrideOnPins(const TArray<FString>& InPinPaths, bool bSetupUndo, bool bPrintPythonCommand)
+{
+	if (!IsValidGraph() || InPinPaths.IsEmpty())
+	{
+		return false;
+	}
+
+	if (!bIsTransacting && !IsGraphEditable())
+	{
+		return false;
+	}
+
+	URigVMGraph* Graph = GetGraph();
+	check(Graph);
+
+	OpenUndoBracket(TEXT("Clear Override On Pins"));
+
+	int32 NumPinsAffected = 0;
+	for(const FString& PinPath : InPinPaths)
+	{
+		if(ClearOverrideOnPin(PinPath, bSetupUndo, false))
+		{
+			NumPinsAffected++;
+		}
+	}
+
+	if(NumPinsAffected == 0)
+	{
+		CancelUndoBracket();
+		return false;
+	}
+
+	CloseUndoBracket();
+	
+	if (bPrintPythonCommand)
+	{
+		const FString GraphName = GetSchema()->GetSanitizedGraphName(GetGraph()->GetGraphName());
+
+		TArray<FString> SanitizedPinPaths;
+		SanitizedPinPaths.Reserve(InPinPaths.Num());
+		for(const FString& PinPath : InPinPaths)
+		{
+			SanitizedPinPaths.Add(FString::Printf(TEXT("'%s'"), *GetSchema()->GetSanitizedPinPath(PinPath)));
+		}
+
+		RigVMPythonUtils::Print(GetSchema()->GetGraphOuterName(GetGraph()),
+			FString::Printf(TEXT("blueprint.get_controller_by_name('%s').clear_override_on_pins([%s])"),
+			*GraphName,
+			*FString::Join(SanitizedPinPaths, TEXT(","))));
+	}
+
+	return true;
+}
+
+bool URigVMController::ClearOverrideOnAllPinsOnNode(const FName& InNodeName, bool bSetupUndo, bool bPrintPythonCommand)
+{
+	if (!IsValidGraph())
+	{
+		return false;
+	}
+
+	if (!bIsTransacting && !IsGraphEditable())
+	{
+		return false;
+	}
+
+	URigVMGraph* Graph = GetGraph();
+	check(Graph);
+
+	URigVMNode* Node = Graph->FindNodeByName(InNodeName);
+	if (Node == nullptr)
+	{
+		ReportErrorf(TEXT("Cannot find node '%s'."), *InNodeName.ToString());
+		return false;
+	}
+
+	TArray<FString> PinPaths;
+	TArray<URigVMPin*> Pins = Node->GetAllPinsRecursively();
+	for(const URigVMPin* Pin : Pins)
+	{
+		if(Pin->CanProvideDefaultValue() && Pin->HasDefaultValueOverride())
+		{
+			PinPaths.Add(Pin->GetPinPath());
+		}
+	}
+
+	if(!ClearOverrideOnPins(PinPaths, bSetupUndo, false))
+	{
+		return false;
+	}
+	
+	if (bPrintPythonCommand)
+	{
+		const FString GraphName = GetSchema()->GetSanitizedGraphName(GetGraph()->GetGraphName());
+
+		RigVMPythonUtils::Print(GetSchema()->GetGraphOuterName(GetGraph()),
+			FString::Printf(TEXT("blueprint.get_controller_by_name('%s').clear_override_on_all_pins_on_node('%s')"),
+			*GraphName,
+			*GetSchema()->GetSanitizedNodeName(InNodeName.ToString())));
+	}
+
+	return true;
+}
+
+bool URigVMController::ClearOverrideOnAllPinsOnNodes(const TArray<FName>& InNodeNames, bool bSetupUndo, bool bPrintPythonCommand)
+{
+	if (!IsValidGraph())
+	{
+		return false;
+	}
+
+	if (!bIsTransacting && !IsGraphEditable())
+	{
+		return false;
+	}
+
+	OpenUndoBracket(TEXT("Add Override To All Pins On Nodes"));
+
+	int32 NumNodesAffected = 0;
+	for(const FName& NodeName : InNodeNames)
+	{
+		if(ClearOverrideOnAllPinsOnNode(NodeName, bSetupUndo, false))
+		{
+			NumNodesAffected++;
+		}
+	}
+
+	if(NumNodesAffected == 0)
+	{
+		CancelUndoBracket();
+		return false;
+	}
+
+	CloseUndoBracket();
+	
+	if (bPrintPythonCommand)
+	{
+		const FString GraphName = GetSchema()->GetSanitizedGraphName(GetGraph()->GetGraphName());
+
+		TArray<FString> SanitizedNodeNames;
+		SanitizedNodeNames.Reserve(InNodeNames.Num());
+		for(const FName& NodeName : InNodeNames)
+		{
+			SanitizedNodeNames.Add(FString::Printf(TEXT("'%s'"), *GetSchema()->GetSanitizedNodeName(NodeName.ToString())));
+		}
+		
+		RigVMPythonUtils::Print(GetSchema()->GetGraphOuterName(GetGraph()),
+			FString::Printf(TEXT("blueprint.get_controller_by_name('%s').clear_override_on_all_pins_on_nodes([%s])"),
+			*GraphName,
+			*FString::Join(SanitizedNodeNames, TEXT(","))));
+	}
+
+	return true;
+}
+
 bool URigVMController::ResetPinDefaultValue(URigVMPin* InPin, bool bSetupUndoRedo)
 {
 	if (!bIsTransacting && !IsGraphEditable())
@@ -23120,7 +23665,7 @@ FRigVMClientPatchResult URigVMController::PatchLazyPins()
 FRigVMClientPatchResult URigVMController::PatchPinDefaultValues()
 {
 	FRigVMClientPatchResult Result;
-	if(!CVarRigVMEnablePinDefaultTypes.GetValueOnAnyThread())
+	if(!CVarRigVMEnablePinOverrides.GetValueOnAnyThread())
 	{
 		return Result;
 	}
@@ -23149,7 +23694,7 @@ FRigVMClientPatchResult URigVMController::PatchPinDefaultValues()
 							}
 						}
 					}
-					else if(!Pin->GetRootPin()->HasUserProvidedDefaultValue())
+					else if(!Pin->GetRootPin()->HasDefaultValueOverride())
 					{
 						if(Pin->GetDefaultValueType() == ERigVMPinDefaultValueType::Unset)
 						{
@@ -23260,7 +23805,7 @@ FRigVMClientPatchResult URigVMController::PatchLocalVariableTypes()
 
 ERigVMPinDefaultValueType URigVMController::GetDefaultValueType(const URigVMPin* InPin, const FString& InDefaultValue) const
 {
-	if(!CVarRigVMEnablePinDefaultTypes.GetValueOnAnyThread())
+	if(!CVarRigVMEnablePinOverrides.GetValueOnAnyThread())
 	{
 		return ERigVMPinDefaultValueType::AutoDetect;
 	}
diff --git a/Engine/Plugins/Runtime/RigVM/Source/RigVMDeveloper/Private/RigVMModel/RigVMNode.cpp b/Engine/Plugins/Runtime/RigVM/Source/RigVMDeveloper/Private/RigVMModel/RigVMNode.cpp
--- a/Engine/Plugins/Runtime/RigVM/Source/RigVMDeveloper/Private/RigVMModel/RigVMNode.cpp
+++ b/Engine/Plugins/Runtime/RigVM/Source/RigVMDeveloper/Private/RigVMModel/RigVMNode.cpp
@@ -489,6 +489,35 @@ FString URigVMNode::GetOriginalPinDefaultValue(const URigVMPin* InPin) const
 	return OriginalDefaultValue;
 }
 
+ERigVMNodeDefaultValueOverrideState::Type URigVMNode::GetPinDefaultValueOverrideState() const
+{
+	int32 NumPinsWithDefaultValue = 0;
+	int32 NumPinsWithDefaultValueOverride = 0;
+	if(CVarRigVMEnablePinOverrides.GetValueOnAnyThread())
+	{
+		for(const URigVMPin* Pin : Pins)
+		{
+			if(Pin->CanProvideDefaultValue())
+			{
+				NumPinsWithDefaultValue++;
+				if(Pin->HasDefaultValueOverride())
+				{
+					NumPinsWithDefaultValueOverride++;
+				}
+			}
+		}
+	}
+	if(NumPinsWithDefaultValueOverride == 0)
+	{
+		return ERigVMNodeDefaultValueOverrideState::None;
+	}
+	if(NumPinsWithDefaultValueOverride < NumPinsWithDefaultValue)
+	{
+		return ERigVMNodeDefaultValueOverrideState::SomePins;
+	}
+	return ERigVMNodeDefaultValueOverrideState::AllPins;
+}
+
 URigVMPin* URigVMNode::FindPin(const FString& InPinPath) const
 {
 	FString Left, Right;
diff --git a/Engine/Plugins/Runtime/RigVM/Source/RigVMDeveloper/Private/RigVMModel/RigVMPin.cpp b/Engine/Plugins/Runtime/RigVM/Source/RigVMDeveloper/Private/RigVMModel/RigVMPin.cpp
--- a/Engine/Plugins/Runtime/RigVM/Source/RigVMDeveloper/Private/RigVMModel/RigVMPin.cpp
+++ b/Engine/Plugins/Runtime/RigVM/Source/RigVMDeveloper/Private/RigVMModel/RigVMPin.cpp
@@ -28,7 +28,7 @@
 
 #include UE_INLINE_GENERATED_CPP_BY_NAME(RigVMPin)
 
-TAutoConsoleVariable<bool> CVarRigVMEnablePinDefaultTypes(TEXT("RigVM.EnablePinDefaultTypes"), false, TEXT("enables the use of pin default types"));
+TAutoConsoleVariable<bool> CVarRigVMEnablePinOverrides(TEXT("RigVM.EnablePinOverrides"), false, TEXT("enables the use of overrides for pin default values"));
 
 #if WITH_EDITOR
 #include "UObject/CoreRedirects.h"
@@ -129,6 +129,8 @@ URigVMPin::URigVMPin()
 	, CachedPinPathWithNodePath(this)
 	, CachedPinCategory(this)
 	, CachedDisplayName(this)
+	, CachedDefaultValueOverride(this)
+	, CachedHasOriginalDefaultValue(this)
 {
 }
 
@@ -852,6 +854,46 @@ FString URigVMPin::GetOriginalDefaultValue() const
 	return FString();
 }
 
+bool URigVMPin::HasOriginalDefaultValue() const
+{
+	if(CachedHasOriginalDefaultValue.IsValid())
+	{
+		return CachedHasOriginalDefaultValue.GetValue();
+	}
+
+	if(!CanProvideDefaultValue())
+	{
+		CachedHasOriginalDefaultValue = false;
+		return false;
+	}
+
+	CachedHasOriginalDefaultValue = true;
+
+	if(SubPins.IsEmpty())
+	{
+		const FString CurrentDefaultValue = GetDefaultValue();
+		FString OriginalDefaultValue = GetOriginalDefaultValue();
+		URigVMController::PostProcessDefaultValue(this, OriginalDefaultValue);
+		if(CurrentDefaultValue != OriginalDefaultValue)
+		{
+			CachedHasOriginalDefaultValue = false;
+		}
+	}
+	else
+	{
+		for(const URigVMPin* SubPin : SubPins)
+		{
+			if(!SubPin->HasOriginalDefaultValue())
+			{
+				CachedHasOriginalDefaultValue = false;
+				break;
+			}
+		}
+	}
+
+	return CachedHasOriginalDefaultValue.GetValue();
+}
+
 FString URigVMPin::GetDefaultValue() const
 {
 	return GetDefaultValue(EmptyPinOverride, true);
@@ -1191,30 +1233,47 @@ bool URigVMPin::IsValidDefaultValue(const FString& InDefaultValue) const
 
 bool URigVMPin::HasUserProvidedDefaultValue() const
 {
-	if(!CVarRigVMEnablePinDefaultTypes.GetValueOnAnyThread())
+	return HasDefaultValueOverride();
+}
+
+bool URigVMPin::HasDefaultValueOverride() const
+{
+	if(!CVarRigVMEnablePinOverrides.GetValueOnAnyThread())
 	{
 		return false;
 	}
-	
+
 	if(!CanProvideDefaultValue())
 	{
 		return false;
 	}
-	
+
+	if(CachedDefaultValueOverride.IsValid())
+	{
+		return CachedDefaultValueOverride.GetValue();
+	}
+
 	if(DefaultValueType == ERigVMPinDefaultValueType::Override)
 	{
-		return true;
+		CachedDefaultValueOverride = true;
+		return CachedDefaultValueOverride.GetValue();
 	}
 
 	for(const TObjectPtr<URigVMPin>& SubPin : SubPins)
 	{
-		if(SubPin->HasUserProvidedDefaultValue())
+		if(SubPin->HasDefaultValueOverride())
 		{
-			return true;
+			CachedDefaultValueOverride = true;
+			return CachedDefaultValueOverride.GetValue();
 		}
 	}
 
-	return false;
+	CachedDefaultValueOverride = false;
+	if(!HasOriginalDefaultValue())
+	{
+		CachedDefaultValueOverride = true;
+	}
+	return CachedDefaultValueOverride.GetValue();
 }
 
 bool URigVMPin::CanProvideDefaultValue() const
diff --git a/Engine/Plugins/Runtime/RigVM/Source/RigVMDeveloper/Public/RigVMModel/RigVMController.h b/Engine/Plugins/Runtime/RigVM/Source/RigVMDeveloper/Public/RigVMModel/RigVMController.h
--- a/Engine/Plugins/Runtime/RigVM/Source/RigVMDeveloper/Public/RigVMModel/RigVMController.h
+++ b/Engine/Plugins/Runtime/RigVM/Source/RigVMDeveloper/Public/RigVMModel/RigVMController.h
@@ -897,6 +897,50 @@ public:
 	UFUNCTION(BlueprintCallable, Category = RigVMController)
 	bool ResetPinDefaultValue(const FString& InPinPath, bool bSetupUndoRedo = true, bool bPrintPythonCommand = false);
 
+	// Resets the default value of a list of pin given the pinpaths.
+	UFUNCTION(BlueprintCallable, Category = RigVMController)
+	bool ResetDefaultValueForPins(const TArray<FString>& InPinPaths, bool bSetupUndo = true, bool bPrintPythonCommand = false);
+
+	// Resets the default value of all pins of a given node
+	UFUNCTION(BlueprintCallable, Category = RigVMController)
+	bool ResetDefaultValueForAllPinsOnNode(const FName& InNodeName, bool bSetupUndo = true, bool bPrintPythonCommand = false);
+
+	// Resets the default value of all pins of a list of nodes
+	UFUNCTION(BlueprintCallable, Category = RigVMController)
+	bool ResetDefaultValueForAllPinsOnNodes(const TArray<FName>& InNodeNames, bool bSetupUndo = true, bool bPrintPythonCommand = false);
+
+	// Adds an override to the given pin path
+	UFUNCTION(BlueprintCallable, Category = RigVMController)
+	bool AddOverrideToPin(const FString& InPinPath, bool bSetupUndo = true, bool bPrintPythonCommand = false);
+
+	// Adds an override to a given list of pin paths
+	UFUNCTION(BlueprintCallable, Category = RigVMController)
+	bool AddOverrideToPins(const TArray<FString>& InPinPaths, bool bSetupUndo = true, bool bPrintPythonCommand = false);
+
+	// Adds an override to all pins on a node
+	UFUNCTION(BlueprintCallable, Category = RigVMController)
+	bool AddOverrideToAllPinsOnNode(const FName& InNodeName, bool bSetupUndo = true, bool bPrintPythonCommand = false);
+
+	// Adds an override to all pins on a list of nodes
+	UFUNCTION(BlueprintCallable, Category = RigVMController)
+	bool AddOverrideToAllPinsOnNodes(const TArray<FName>& InNodeNames, bool bSetupUndo = true, bool bPrintPythonCommand = false);
+
+	// Clears an override on a given pin path
+	UFUNCTION(BlueprintCallable, Category = RigVMController)
+	bool ClearOverrideOnPin(const FString& InPinPath, bool bSetupUndo = true, bool bPrintPythonCommand = false);
+
+	// Clears the overrides on a given list of pin paths
+	UFUNCTION(BlueprintCallable, Category = RigVMController)
+	bool ClearOverrideOnPins(const TArray<FString>& InPinPaths, bool bSetupUndo = true, bool bPrintPythonCommand = false);
+
+	// Clears the overrides for all pins on a node
+	UFUNCTION(BlueprintCallable, Category = RigVMController)
+	bool ClearOverrideOnAllPinsOnNode(const FName& InNodeName, bool bSetupUndo = true, bool bPrintPythonCommand = false);
+
+	// Clears the overrides for all pins of a list of nodes
+	UFUNCTION(BlueprintCallable, Category = RigVMController)
+	bool ClearOverrideOnAllPinsOnNodes(const TArray<FName>& InNodeNames, bool bSetupUndo = true, bool bPrintPythonCommand = false);
+
 	UFUNCTION(BlueprintCallable, Category = RigVMController)
 	FString AddAggregatePin(const FString& InNodeName, const FString& InPinName, const FString& InDefaultValue = TEXT(""), bool bSetupUndoRedo = true, bool bPrintPythonCommand = false);
 
diff --git a/Engine/Plugins/Runtime/RigVM/Source/RigVMDeveloper/Public/RigVMModel/RigVMNode.h b/Engine/Plugins/Runtime/RigVM/Source/RigVMDeveloper/Public/RigVMModel/RigVMNode.h
--- a/Engine/Plugins/Runtime/RigVM/Source/RigVMDeveloper/Public/RigVMModel/RigVMNode.h
+++ b/Engine/Plugins/Runtime/RigVM/Source/RigVMDeveloper/Public/RigVMModel/RigVMNode.h
@@ -13,6 +13,16 @@
 
 class URigVMGraph;
 
+namespace ERigVMNodeDefaultValueOverrideState
+{
+	enum Type : int
+	{
+		None, // No override.
+		SomePins, // Some pins have been overriden
+		AllPins, // All pins contains an override 
+	};
+}
+
 /**
  * The Node represents a single statement within a Graph. 
  * Nodes can represent values such as Variables / Parameters,
@@ -110,6 +120,11 @@ public:
 	// Returns the default value for a given pin
 	FString GetOriginalPinDefaultValue(const URigVMPin* InPin) const;
 
+	// Returns 0 if there's no override on this node,
+	// 1 if some pins have an override and
+	// 2 if all pins are overriden
+	ERigVMNodeDefaultValueOverrideState::Type GetPinDefaultValueOverrideState() const;
+
 	// Returns a Pin given it's partial pin path below
 	// this node (for example: "Color.R")
 	UFUNCTION(BlueprintCallable, Category = RigVMNode)
diff --git a/Engine/Plugins/Runtime/RigVM/Source/RigVMDeveloper/Public/RigVMModel/RigVMPin.h b/Engine/Plugins/Runtime/RigVM/Source/RigVMDeveloper/Public/RigVMModel/RigVMPin.h
--- a/Engine/Plugins/Runtime/RigVM/Source/RigVMDeveloper/Public/RigVMModel/RigVMPin.h
+++ b/Engine/Plugins/Runtime/RigVM/Source/RigVMDeveloper/Public/RigVMModel/RigVMPin.h
@@ -21,7 +21,7 @@ class URigVMPin;
 class URigVMLink;
 class URigVMVariableNode;
 
-extern RIGVMDEVELOPER_API TAutoConsoleVariable<bool> CVarRigVMEnablePinDefaultTypes;
+extern RIGVMDEVELOPER_API TAutoConsoleVariable<bool> CVarRigVMEnablePinOverrides;
 
 /**
  * The Injected Info is used for injecting a node on a pin.
@@ -309,6 +309,9 @@ public:
 
 	UFUNCTION(BlueprintCallable, Category = RigVMPin)
 	FString GetOriginalDefaultValue() const;
+
+	UFUNCTION(BlueprintCallable, Category = RigVMPin)
+	bool HasOriginalDefaultValue() const;
 	
 	// Returns the default value of the Pin as a string.
 	// Note that this value is computed based on the Pin's
@@ -329,9 +332,14 @@ public:
 	bool IsValidDefaultValue(const FString& InDefaultValue) const;
 
 	// Returns true if the default value was ever changed by the user
+	UE_DEPRECATED(5.6, "Please use URigVMPin::HasDefaultValueOverride()")
 	UFUNCTION(BlueprintPure, Category = RigVMPin)
 	bool HasUserProvidedDefaultValue() const;
 
+	// Returns true if the default value was ever changed by the user
+    UFUNCTION(BlueprintPure, Category = RigVMPin)
+	bool HasDefaultValueOverride() const;
+
 	// Returns true if the pin can / may provide a default value 
 	UFUNCTION(BlueprintPure, Category = RigVMPin)
 	ERigVMPinDefaultValueType GetDefaultValueType() const { return DefaultValueType; }
@@ -620,6 +628,8 @@ private:
 	mutable TRigVMModelCachedValue<URigVMPin, FString> CachedPinPathWithNodePath;
 	mutable TRigVMModelCachedValue<URigVMPin, FString> CachedPinCategory;
 	mutable TRigVMModelCachedValue<URigVMPin, FName> CachedDisplayName;
+	mutable TRigVMModelCachedValue<URigVMPin, bool> CachedDefaultValueOverride;
+	mutable TRigVMModelCachedValue<URigVMPin, bool> CachedHasOriginalDefaultValue;
 
 	static const inline TCHAR* OrphanPinPrefix = TEXT("Orphan::");
 
diff --git a/Engine/Plugins/Runtime/RigVM/Source/RigVMEditor/Private/Editor/RigVMEditor.cpp b/Engine/Plugins/Runtime/RigVM/Source/RigVMEditor/Private/Editor/RigVMEditor.cpp
--- a/Engine/Plugins/Runtime/RigVM/Source/RigVMEditor/Private/Editor/RigVMEditor.cpp
+++ b/Engine/Plugins/Runtime/RigVM/Source/RigVMEditor/Private/Editor/RigVMEditor.cpp
@@ -38,8 +38,12 @@
 #include "Widgets/Docking/SDockTab.h"
 #include "Widgets/Layout/SScrollBox.h"
 #include "IStructureDetailsView.h"
+#include "PropertyPath.h"
 #include "Editor/RigVMDetailsInspectorTabSummoner.h"
 #include "Editor/SRigVMDetailsInspector.h"
+#include "Overrides/OverrideStatusDetailsObjectFilter.h"
+#include "Overrides/OverrideStatusDetailsWidgetBuilder.h"
+#include "Overrides/OverrideStatusDetailsWidgetBuilder.h"
 
 #if WITH_RIGVMLEGACYEDITOR
 #include "SKismetInspector.h"
@@ -54,6 +58,7 @@ FRigVMEditorBase::FRigVMEditorBase()
 	, HaltedAtNode(nullptr)
 	, LastDebuggedHost()
 	, bSuspendDetailsPanelRefresh(false)
+	, bDetailsPanelRequiresClear(false)
 	, bAllowBulkEdits(false)
 	, bIsSettingObjectBeingDebugged(false)
 	, bRigVMEditorInitialized(false)
@@ -2838,7 +2843,10 @@ void FRigVMEditorBase::SetDetailObjects(const TArray<UObject*>& InObjects, bool
 		}
 	}
 
-	ClearDetailObject(bChangeUISelectionState);
+	if(bDetailsPanelRequiresClear)
+	{
+		ClearDetailObject(bChangeUISelectionState);
+	}
 
 	if (InObjects.Num() == 1)
 	{
@@ -2864,7 +2872,7 @@ void FRigVMEditorBase::SetDetailObjects(const TArray<UObject*>& InObjects, bool
 			ModelNodes.Add(ModelNode);
 		}
 	}
-	
+
 	for(UObject* InObject : InObjects)
 	{
 		if (URigVMLibraryNode* LibraryNode = Cast<URigVMLibraryNode>(InObject))
@@ -2874,6 +2882,7 @@ void FRigVMEditorBase::SetDetailObjects(const TArray<UObject*>& InObjects, bool
 				if (UEdGraph* EdGraph = GetRigVMBlueprint()->GetEdGraph(LibraryNode->GetContainedGraph()))
 				{
 					FilteredObjects.AddUnique(EdGraph);
+					ModelNodes.Remove(LibraryNode);
 					continue;
 				}
 			}
@@ -2883,6 +2892,7 @@ void FRigVMEditorBase::SetDetailObjects(const TArray<UObject*>& InObjects, bool
 			if (UEdGraph* EdGraph = GetRigVMBlueprint()->GetEdGraph(CastChecked<URigVMNode>(InObject)->GetGraph()))
 			{
 				FilteredObjects.AddUnique(EdGraph);
+				ModelNodes.Remove(Cast<URigVMNode>(InObject));
 				continue;
 			}
 		}
@@ -2896,6 +2906,7 @@ void FRigVMEditorBase::SetDetailObjects(const TArray<UObject*>& InObjects, bool
 					if(UEdGraphNode* EdGraphNode = EdGraph->FindNodeForModelNodeName(CommentNode->GetFName()))
 					{
 						FilteredObjects.AddUnique(EdGraphNode);
+						ModelNodes.Remove(CommentNode);
 						continue;
 					}
 				}
@@ -2903,6 +2914,7 @@ void FRigVMEditorBase::SetDetailObjects(const TArray<UObject*>& InObjects, bool
 			else
 			{
 				FilteredObjects.AddUnique(CommentNode);
+				ModelNodes.Remove(CommentNode);
 				continue;
 			}
 		}
@@ -2939,6 +2951,244 @@ void FRigVMEditorBase::SetDetailObjects(const TArray<UObject*>& InObjects, bool
 		}
 	}
 
+	if(!ModelNodes.IsEmpty() && CVarRigVMEnablePinOverrides.GetValueOnAnyThread())
+	{
+		TSharedPtr<FOverrideStatusDetailsViewObjectFilter> ObjectFilter = FOverrideStatusDetailsViewObjectFilter::Create();
+
+		ObjectFilter->OnCanCreateWidget().BindLambda([](const FOverrideStatusSubject& InSubject) -> bool
+		{
+			return InSubject.Contains<URigVMDetailsViewWrapperObject>([](const URigVMDetailsViewWrapperObject* InWrapper)
+			{
+				return Cast<URigVMNode>(InWrapper->GetSubject()) != nullptr;
+			});
+		});
+
+		auto PropertyPathToPinPath = [](const FOverrideStatusSubject& InSubject) -> FString
+		{
+			FString PinPath = InSubject.GetPropertyPathString(TEXT("."));
+			PinPath.ReplaceCharInline(TEXT('['), TEXT('.'));
+			PinPath.ReplaceInline(TEXT("]"), TEXT(""));
+			return PinPath;
+		};
+
+		ObjectFilter->OnGetStatus().BindLambda([PropertyPathToPinPath](const FOverrideStatusSubject& InSubject)
+		{
+			const FString PinPath = PropertyPathToPinPath(InSubject);
+			return InSubject.GetStatus<URigVMDetailsViewWrapperObject>(
+				[PinPath](const URigVMDetailsViewWrapperObject* InWrapper) -> TOptional<EOverrideWidgetStatus::Type>
+				{
+					if(const URigVMNode* Node = Cast<URigVMNode>(InWrapper->GetSubject()))
+					{
+						if(PinPath.IsEmpty())
+						{
+							if(Node->GetPinDefaultValueOverrideState() != ERigVMNodeDefaultValueOverrideState::None)
+							{
+								return EOverrideWidgetStatus::ChangedInside;
+							}
+						}
+						else
+						{
+							if(const URigVMPin* Pin = Node->FindPin(PinPath))
+							{
+								if(Pin->GetDefaultValueType() == ERigVMPinDefaultValueType::Override)
+								{
+									return EOverrideWidgetStatus::ChangedHere;
+								}
+								if(!Pin->GetSubPins().IsEmpty())
+								{
+									if(Pin->HasDefaultValueOverride())
+									{
+										return EOverrideWidgetStatus::ChangedInside;
+									}
+								}
+
+								// if this is an old asset and the default value type has not yet
+								// been determined - we need to fall back on comparing to the default.
+								if(!Pin->HasOriginalDefaultValue())
+								{
+									return EOverrideWidgetStatus::ChangedHere;
+								}
+							}
+						}
+						return EOverrideWidgetStatus::None;
+					}
+					return {};
+				}
+			).Get(EOverrideWidgetStatus::Mixed);
+		});
+
+		ObjectFilter->OnAddOverride().BindLambda(
+			[this, PropertyPathToPinPath](const FOverrideStatusSubject& InSubject)
+			{
+				const FString PinPath = PropertyPathToPinPath(InSubject);
+				TArray<FString> PinPaths;
+				TArray<FName> NodeNames;
+
+				InSubject.ForEach<URigVMDetailsViewWrapperObject>(
+					[PinPath, &PinPaths, &NodeNames](const URigVMDetailsViewWrapperObject* InWrapper)
+					{
+						if(const URigVMNode* Node = Cast<URigVMNode>(InWrapper->GetSubject()))
+						{
+							if(PinPath.IsEmpty())
+							{
+								NodeNames.Add(Node->GetFName());
+							}
+							else
+							{
+								if(const URigVMPin* Pin = Node->FindPin(PinPath))
+								{
+									PinPaths.AddUnique(Pin->GetPinPath());
+								}
+							}
+						}
+					}
+				);
+
+				if(!PinPaths.IsEmpty())
+				{
+					if(GetFocusedController()->AddOverrideToPins(PinPaths))
+					{
+						return FReply::Handled();
+					}
+				}
+				else if(!NodeNames.IsEmpty())
+				{
+					if(GetFocusedController()->AddOverrideToAllPinsOnNodes(NodeNames))
+					{
+						return FReply::Handled();
+					}
+				}
+
+				return FReply::Unhandled();;
+			});
+
+		ObjectFilter->OnClearOverride().BindLambda(
+			[this, PropertyPathToPinPath](const FOverrideStatusSubject& InSubject)
+			{
+				const FString PinPath = PropertyPathToPinPath(InSubject);
+				TArray<FString> PinPaths;
+				TArray<FName> NodeNames;
+
+				InSubject.ForEach<URigVMDetailsViewWrapperObject>(
+					[PinPath, &PinPaths, &NodeNames](const URigVMDetailsViewWrapperObject* InWrapper)
+					{
+						if(const URigVMNode* Node = Cast<URigVMNode>(InWrapper->GetSubject()))
+						{
+							if(PinPath.IsEmpty())
+							{
+								NodeNames.Add(Node->GetFName());
+							}
+							else
+							{
+								if(const URigVMPin* Pin = Node->FindPin(PinPath))
+								{
+									PinPaths.AddUnique(Pin->GetPinPath());
+								}
+							}
+						}
+					}
+				);
+				
+				if(!PinPaths.IsEmpty())
+				{
+					if(GetFocusedController()->ClearOverrideOnPins(PinPaths))
+					{
+						return FReply::Handled();
+					}
+				}
+				else if(!NodeNames.IsEmpty())
+				{
+					if(GetFocusedController()->ClearOverrideOnAllPinsOnNodes(NodeNames))
+					{
+						return FReply::Handled();
+					}
+				}
+
+				return FReply::Unhandled();
+			});
+
+		ObjectFilter->OnResetToDefault().BindLambda(
+			[this, PropertyPathToPinPath](const FOverrideStatusSubject& InSubject)
+			{
+				const FString PinPath = PropertyPathToPinPath(InSubject);
+				TArray<FString> PinPaths;
+				TArray<FName> NodeNames;
+
+				InSubject.ForEach<URigVMDetailsViewWrapperObject>(
+					[PinPath, &PinPaths, &NodeNames](const URigVMDetailsViewWrapperObject* InWrapper)
+					{
+						if(const URigVMNode* Node = Cast<URigVMNode>(InWrapper->GetSubject()))
+						{
+							if(PinPath.IsEmpty())
+							{
+								NodeNames.Add(Node->GetFName());
+							}
+							else
+							{
+								if(const URigVMPin* Pin = Node->FindPin(PinPath))
+								{
+									PinPaths.AddUnique(Pin->GetPinPath());
+								}
+							}
+						}
+					}
+				);
+
+				if(!PinPaths.IsEmpty())
+				{					
+					FRigVMDefaultValueTypeGuard _(GetFocusedController(), ERigVMPinDefaultValueType::KeepValueType, true);
+					if(GetFocusedController()->ResetDefaultValueForPins(PinPaths))
+					{
+						return FReply::Handled();
+					}
+				}
+				else if(!NodeNames.IsEmpty())
+				{
+					FRigVMDefaultValueTypeGuard _(GetFocusedController(), ERigVMPinDefaultValueType::KeepValueType, true);
+					if(GetFocusedController()->ResetDefaultValueForAllPinsOnNodes(NodeNames))
+					{
+						return FReply::Handled();
+					}
+				}
+				return FReply::Unhandled();
+			});
+
+		ObjectFilter->OnValueDiffersFromDefault().BindLambda(
+			[this, PropertyPathToPinPath](const FOverrideStatusSubject& InSubject)
+			{
+				const FString PinPath = PropertyPathToPinPath(InSubject);
+				return InSubject.Contains<URigVMDetailsViewWrapperObject>([PinPath](const URigVMDetailsViewWrapperObject* InWrapper)
+				{
+					if(const URigVMNode* Node = Cast<URigVMNode>(InWrapper->GetSubject()))
+					{
+						if(PinPath.IsEmpty())
+						{
+							for(const URigVMPin* Pin : Node->GetPins())
+							{
+								if(Pin->CanProvideDefaultValue())
+								{
+									if(!Pin->HasOriginalDefaultValue())
+									{
+										return true;
+									}
+								}
+							}
+						}
+						else
+						{
+							if(const URigVMPin* Pin = Node->FindPin(PinPath))
+							{
+								return !Pin->HasOriginalDefaultValue();
+							}
+						}
+					}
+					return false;
+				});
+			});
+
+		SetDetailObjectFilter(ObjectFilter);
+	}
+
 	auto Task = [&]<typename T, typename O>(T Inspector, O Options)
 	{
 		Inspector->ShowDetailsForObjects(FilteredObjects, Options);
@@ -2962,6 +3212,26 @@ void FRigVMEditorBase::SetDetailObjects(const TArray<UObject*>& InObjects, bool
 		Options.bForceRefresh = true;
 		Task(Inspector, Options);
 	}
+	
+	bDetailsPanelRequiresClear = true;
+}
+
+void FRigVMEditorBase::SetDetailObjectFilter(TSharedPtr<FDetailsViewObjectFilter> InObjectFilter)
+{
+#if WITH_RIGVMLEGACYEDITOR
+	TSharedRef<FAssetEditorToolkit> SharedApp = GetHostingApp().ToSharedRef();
+	if (SharedApp->IsBlueprintEditor())
+	{
+		if (TSharedPtr<SKismetInspector> Inspector = GetKismetInspector())
+		{
+			Inspector->GetPropertyView()->SetObjectFilter(InObjectFilter);
+		}
+	}
+#endif
+	if (TSharedPtr<SRigVMDetailsInspector> Inspector = GetRigVMInspector())
+	{
+		Inspector->GetPropertyView()->SetObjectFilter(InObjectFilter);
+	}
 }
 
 void FRigVMEditorBase::SetMemoryStorageDetails(const TArray<FRigVMMemoryStorageStruct*>& InStructs)
@@ -3024,7 +3294,10 @@ void FRigVMEditorBase::SetDetailViewForGraph(URigVMGraph* InGraph)
 		return;
 	}
 
-	ClearDetailObject();
+	if(bDetailsPanelRequiresClear)
+	{
+		ClearDetailObject();
+	}
 
 	TArray<UObject*> SelectedNodes;
 	TArray<FName> SelectedNodeNames = InGraph->GetSelectNodes();
@@ -3158,6 +3431,7 @@ void FRigVMEditorBase::ClearDetailObject(bool bChangeUISelectionState)
 		return;
 	}
 
+	SetDetailObjectFilter(nullptr);
 	ClearDetailsViewWrapperObjects();
 
 	auto Task = [&]<typename T>(T Inspector)
@@ -3194,6 +3468,8 @@ void FRigVMEditorBase::ClearDetailObject(bool bChangeUISelectionState)
 		SetUISelectionState(FRigVMNewEditor::SelectionState_Graph());
 #endif
 	}
+
+	bDetailsPanelRequiresClear = false;
 }
 
 void FRigVMEditorBase::ClearDetailsViewWrapperObjects()
diff --git a/Engine/Plugins/Runtime/RigVM/Source/RigVMEditor/Private/Editor/RigVMGraphDetailCustomization.cpp b/Engine/Plugins/Runtime/RigVM/Source/RigVMEditor/Private/Editor/RigVMGraphDetailCustomization.cpp
--- a/Engine/Plugins/Runtime/RigVM/Source/RigVMEditor/Private/Editor/RigVMGraphDetailCustomization.cpp
+++ b/Engine/Plugins/Runtime/RigVM/Source/RigVMEditor/Private/Editor/RigVMGraphDetailCustomization.cpp
@@ -2275,9 +2275,9 @@ void FRigVMWrappedNodeDetailCustomization::CustomizeDetails(IDetailLayoutBuilder
 
 		PropertiesAddedToLayout.Add(PropertyHandle);
 
-		TAttribute<bool> HasUserProvidedDefaultValue = TAttribute<bool>::CreateLambda([this, PinPath, PropertyHandle]() -> bool
+		TAttribute<bool> HasDefaultValueOverride = TAttribute<bool>::CreateLambda([this, PinPath, PropertyHandle]() -> bool
 		{
-			if(CVarRigVMEnablePinDefaultTypes.GetValueOnAnyThread())
+			if(CVarRigVMEnablePinOverrides.GetValueOnAnyThread())
 			{
 				for(const TWeakObjectPtr<URigVMNode>& Node : NodesBeingCustomized)
 				{
@@ -2285,7 +2285,7 @@ void FRigVMWrappedNodeDetailCustomization::CustomizeDetails(IDetailLayoutBuilder
 					{
 						if(const URigVMPin* Pin = Node->FindPin(PinPath))
 						{
-							if(Pin->HasUserProvidedDefaultValue())
+							if(Pin->HasDefaultValueOverride())
 							{
 								return true;
 							}
@@ -2297,7 +2297,7 @@ void FRigVMWrappedNodeDetailCustomization::CustomizeDetails(IDetailLayoutBuilder
 		});
 
 		FResetToDefaultOverride ResetToDefault = FResetToDefaultOverride::Create(
-			HasUserProvidedDefaultValue,
+			HasDefaultValueOverride,
 			FSimpleDelegate::CreateLambda([this, PropertyHandle, PinPath]()
 			{
 				FScopedTransaction Transaction(LOCTEXT("ResetValueToDefault", "Reset Value To Default"));
@@ -2317,6 +2317,12 @@ void FRigVMWrappedNodeDetailCustomization::CustomizeDetails(IDetailLayoutBuilder
 			})
 		);
 
+		// disable the reset to default arrow
+		if(!CVarRigVMEnablePinOverrides.GetValueOnAnyThread())
+		{
+			ResetToDefault = FResetToDefaultOverride::Create(true, FSimpleDelegate(), true);
+		}
+
 		static const FSlateFontInfo NameFont = FAppStyle::GetFontStyle( TEXT("PropertyWindow.NormalFont") );
 
 		FText LabelOverride;
@@ -2335,11 +2341,11 @@ void FRigVMWrappedNodeDetailCustomization::CustomizeDetails(IDetailLayoutBuilder
 		[
 			SNew(SBorder)
 			.HAlign(HAlign_Left)
-			.BorderImage_Lambda([HasUserProvidedDefaultValue]() -> const FSlateBrush*
+			.BorderImage_Lambda([HasDefaultValueOverride]() -> const FSlateBrush*
 			{
-				if(CVarRigVMEnablePinDefaultTypes.GetValueOnAnyThread())
+				if(CVarRigVMEnablePinOverrides.GetValueOnAnyThread())
 				{
-					if(HasUserProvidedDefaultValue.Get())
+					if(HasDefaultValueOverride.Get())
 					{
 						static const FSlateBrush* BorderBrush = FAppStyle::Get().GetBrush("FloatingBorder");
 						return BorderBrush;
@@ -2347,11 +2353,11 @@ void FRigVMWrappedNodeDetailCustomization::CustomizeDetails(IDetailLayoutBuilder
 				}
 				return nullptr;
 			})
-			.BorderBackgroundColor_Lambda([HasUserProvidedDefaultValue]() -> FSlateColor
+			.BorderBackgroundColor_Lambda([HasDefaultValueOverride]() -> FSlateColor
 			{
-				if(CVarRigVMEnablePinDefaultTypes.GetValueOnAnyThread())
+				if(CVarRigVMEnablePinOverrides.GetValueOnAnyThread())
 				{
-					if(HasUserProvidedDefaultValue.Get())
+					if(HasDefaultValueOverride.Get())
 					{
 						return FSlateColor(FLinearColor::Red);
 					}
@@ -2366,6 +2372,132 @@ void FRigVMWrappedNodeDetailCustomization::CustomizeDetails(IDetailLayoutBuilder
 		];
 		*/
 
+		/*
+		auto GetOverrideStatus = [this, PinPath, PropertyHandle]() -> EOverrideWidgetStatus::Type
+		{
+			if(CVarRigVMEnablePinOverrides.GetValueOnAnyThread())
+			{
+				for(const TWeakObjectPtr<URigVMNode>& Node : NodesBeingCustomized)
+				{
+					if(Node.IsValid())
+					{
+						if(const URigVMPin* Pin = Node->FindPin(PinPath))
+						{
+							if(Pin->GetDefaultValueType() == ERigVMPinDefaultValueType::Override)
+							{
+								return PropertyHandle->DiffersFromDefault() ? EOverrideWidgetStatus::ChangedHere : EOverrideWidgetStatus::ChangedToDefault; 
+							}
+							if(Pin->HasDefaultValueOverride())
+							{
+								return EOverrideWidgetStatus::ChangedInside;
+							}
+							if(Pin->GetDefaultValueType() == ERigVMPinDefaultValueType::Unset)
+							{
+								return EOverrideWidgetStatus::None;
+							}
+						}
+					}	
+				}
+			}
+			return PropertyHandle->DiffersFromDefault() ? EOverrideWidgetStatus::ChangedHere : EOverrideWidgetStatus::None;
+		};
+
+		TSharedRef<SWidget> OverrideWidget = SNullWidget::NullWidget;
+		
+		if(CVarRigVMEnablePinOverrides.GetValueOnAnyThread())
+		{
+			OverrideWidget = SNew(SOverrideStatusWidget)
+			.Status_Lambda(GetOverrideStatus)
+			.MenuContent_Lambda([this, PinPath, PropertyHandle, GetOverrideStatus]
+			{
+				FMenuBuilder MenuBuilder(true, nullptr);
+
+				const EOverrideWidgetStatus::Type Status = GetOverrideStatus();
+				switch(Status)
+				{
+					case EOverrideWidgetStatus::None:
+					case EOverrideWidgetStatus::Inherited:
+					{
+						MenuBuilder.AddMenuEntry(
+							LOCTEXT("OverrideValueWithCurrent", "Set Override"),
+							LOCTEXT("OverrideValueWithCurrentTooltip", "Overrides value while keeping the value (locks the value in place)."),
+							FSlateIcon(),
+							FUIAction(
+								FExecuteAction::CreateLambda([this, PinPath]()
+								{
+									const URigVMGraph* Graph = NodesBeingCustomized[0]->GetGraph();
+									URigVMController* Controller = BlueprintBeingCustomized->GetController(Graph);
+
+									Controller->OpenUndoBracket(TEXT("Set Override"));
+										
+									for(const TWeakObjectPtr<URigVMNode> Node : NodesBeingCustomized)
+									{
+										if(const URigVMPin* Pin = Node->FindPin(PinPath))
+										{
+											FRigVMDefaultValueTypeGuard _(Controller, ERigVMPinDefaultValueType::Override);
+											if(Pin->CanProvideDefaultValue())
+											{
+												FString DefaultValue = Pin->GetDefaultValue();
+												if(DefaultValue.IsEmpty())
+												{
+													DefaultValue = Pin->GetOriginalDefaultValue();
+												}
+												if(!DefaultValue.IsEmpty())
+												{
+													Controller->SetPinDefaultValue(Pin->GetPinPath(), DefaultValue);
+												}
+											}
+										}
+									}
+
+									Controller->CloseUndoBracket();
+								})
+							)
+						);
+						break;
+					}
+					case EOverrideWidgetStatus::ChangedHere:
+					case EOverrideWidgetStatus::ChangedToDefault:
+					case EOverrideWidgetStatus::ChangedInside:
+					{
+						MenuBuilder.AddMenuEntry(
+							LOCTEXT("RemoveOverride", "Remove Override"),
+							LOCTEXT("RemoveOverrideTooltip", "Removes the override and restores the inherited value."),
+							FSlateIcon(),
+							FUIAction(
+								FExecuteAction::CreateLambda([this, PinPath]()
+								{
+									const URigVMGraph* Graph = NodesBeingCustomized[0]->GetGraph();
+									URigVMController* Controller = BlueprintBeingCustomized->GetController(Graph);
+
+									Controller->OpenUndoBracket(TEXT("Remove Override"));
+										
+									for(const TWeakObjectPtr<URigVMNode> Node : NodesBeingCustomized)
+									{
+										if(const URigVMPin* Pin = Node->FindPin(PinPath))
+										{
+											Controller->ResetPinDefaultValue(Pin->GetPinPath(), true);
+										}
+									}
+
+									Controller->CloseUndoBracket();
+								})
+							)
+						);
+						break;
+					}
+					case EOverrideWidgetStatus::Undetermined:
+					case EOverrideWidgetStatus::Uninitialized:
+					default:
+					{
+						break;
+					}
+				}
+				return MenuBuilder.MakeWidget();
+			});
+		}
+		*/
+		
 		IDetailPropertyRow* Row = nullptr;
 		if(NodeWithCategories)
 		{
@@ -2466,7 +2598,7 @@ void FRigVMWrappedNodeDetailCustomization::CustomizeDetails(IDetailLayoutBuilder
 						]
 					]
         			.OverrideResetToDefault(ResetToDefault);
-        			
+
         			NameListWidgets.Add(Property->GetFName(), NameListWidget);
        				continue;
         		}
@@ -2477,6 +2609,7 @@ void FRigVMWrappedNodeDetailCustomization::CustomizeDetails(IDetailLayoutBuilder
 					LabelWidget
 				]
         		.OverrideResetToDefault(ResetToDefault);
+
         		continue;
         	}
         }
@@ -3642,12 +3775,13 @@ void FRigVMGraphMathTypeDetailCustomization::MakeTransformChildren(TSharedRef<cl
 	typename SAdvancedTransformInputBox<TransformType>::FArguments WidgetArgs;
 	ConfigureTransformWidgetArgs<TransformType>(InPropertyHandle, WidgetArgs, ComponentNames);
 
-	const FString StandaloneWidgetMetadata = InPropertyHandle->GetMetaData(TEXT("SStandaloneCustomizedValueWidget"));
-	const bool IsUsingStandaloneWidget = StandaloneWidgetMetadata.Equals(TEXT("True"));
+	const FName TranslationName = GetTranslationPropertyName<TransformType>();
+	const FName RotationName = GetRotationPropertyName<TransformType>();
+	const FName ScaleName = GetScalePropertyName<TransformType>();
 	
-	FDetailWidgetRow* LocationRow = &StructBuilder.AddCustomRow(LOCTEXT("TransformLocation", "Location"));
-	FDetailWidgetRow* RotationRow = &StructBuilder.AddCustomRow(LOCTEXT("TransformRotation", "Rotation"));
-	FDetailWidgetRow* ScaleRow = &StructBuilder.AddCustomRow(LOCTEXT("TransformScale", "Scale"));
+	FDetailWidgetRow* LocationRow = &StructBuilder.AddProperty(InPropertyHandle->GetChildHandle(TranslationName).ToSharedRef()).CustomWidget();
+	FDetailWidgetRow* RotationRow = &StructBuilder.AddProperty(InPropertyHandle->GetChildHandle(RotationName).ToSharedRef()).CustomWidget();
+	FDetailWidgetRow* ScaleRow = &StructBuilder.AddProperty(InPropertyHandle->GetChildHandle(ScaleName).ToSharedRef()).CustomWidget();
 	
 	SAdvancedTransformInputBox<TransformType>::ConfigureComponentWidgetRow(*LocationRow, ESlateTransformComponent::Location, WidgetArgs);
 	SAdvancedTransformInputBox<TransformType>::ConfigureComponentWidgetRow(*RotationRow, ESlateTransformComponent::Rotation, WidgetArgs);
diff --git a/Engine/Plugins/Runtime/RigVM/Source/RigVMEditor/Private/RigVMEditorModule.cpp b/Engine/Plugins/Runtime/RigVM/Source/RigVMEditor/Private/RigVMEditorModule.cpp
--- a/Engine/Plugins/Runtime/RigVM/Source/RigVMEditor/Private/RigVMEditorModule.cpp
+++ b/Engine/Plugins/Runtime/RigVM/Source/RigVMEditor/Private/RigVMEditorModule.cpp
@@ -752,7 +752,7 @@ void FRigVMEditorModule::GetNodeEventsContextMenuActions(IRigVMClientHost* RigVM
 
 void FRigVMEditorModule::GetNodeDefaultValueContextMenuActions(IRigVMClientHost* RigVMClientHost, const URigVMEdGraphNode* EdGraphNode, URigVMNode* ModelNode, UToolMenu* Menu) const
 {
-	if(!CVarRigVMEnablePinDefaultTypes.GetValueOnAnyThread())
+	if(!CVarRigVMEnablePinOverrides.GetValueOnAnyThread())
 	{
 		return;
 	}
diff --git a/Engine/Plugins/Runtime/RigVM/Source/RigVMEditor/Private/Widgets/SRigVMGraphNode.cpp b/Engine/Plugins/Runtime/RigVM/Source/RigVMEditor/Private/Widgets/SRigVMGraphNode.cpp
--- a/Engine/Plugins/Runtime/RigVM/Source/RigVMEditor/Private/Widgets/SRigVMGraphNode.cpp
+++ b/Engine/Plugins/Runtime/RigVM/Source/RigVMEditor/Private/Widgets/SRigVMGraphNode.cpp
@@ -42,6 +42,7 @@
 #include "Widgets/SRigVMGraphPinCategory.h"
 #include "Widgets/SRigVMVariantTagWidget.h"
 #include "Widgets/SToolTip.h"
+#include "Overrides/OverrideStatusWidgetMenuBuilder.h"
 
 #if WITH_EDITOR
 #include "Editor.h"
@@ -97,6 +98,11 @@ void SRigVMGraphNode::Construct( const FArguments& InArgs )
 		.Image(ImageBrush)
 		.Visibility(EVisibility::Visible);
 
+	OverrideStatusWidget =
+		SNew(SOverrideStatusWidget)
+		.Visibility(EVisibility::Visible)
+		.MenuContent(this, &SRigVMGraphNode::OnOverrideWidgetMenu);
+
 	static const FSlateColorBrush WhiteBrush(FLinearColor::White);
 	
 	SAssignNew(InstructionCountTextBlockWidget, STextBlock)
@@ -734,24 +740,9 @@ TOptional<FSlateColor> SRigVMGraphNode::GetHighlightColor(const SGraphPin* InGra
 	{
 		return FSlateColor(EStyleColor::Hover);
 	}
-	if(HasUserProvidedValue(InGraphPin))
-	{
-		return FSlateColor(FLinearColor::Red);
-	}
 	return SGraphNode::GetHighlightColor(InGraphPin);
 }
 
-TOptional<const FSlateBrush*> SRigVMGraphNode::GetPinBorder(const SGraphPin* InGraphPin) const
-{
-	if(HasUserProvidedValue(InGraphPin) || IsCategoryPin(InGraphPin))
-	{
-		static const FName NAME_Pin_BackgroundHovered("Graph.Pin.BackgroundHovered");
-		static const FSlateBrush* CachedImg_Pin_BackgroundHovered = FAppStyle::GetBrush( NAME_Pin_BackgroundHovered );
-		return CachedImg_Pin_BackgroundHovered;
-	}
-	return SGraphNode::GetPinBorder(InGraphPin);
-}
-
 FSlateColor SRigVMGraphNode::GetVariableLabelTextColor(
 	TWeakObjectPtr<URigVMFunctionReferenceNode> FunctionReferenceNode, FName InVariableName) const
 {
@@ -987,6 +978,23 @@ TArray<FOverlayWidgetInfo> SRigVMGraphNode::GetOverlayWidgets(bool bSelected, co
 
 	if (ModelNode.IsValid())
 	{
+		float HorizontalOffset = 0;
+		
+		if(CVarRigVMEnablePinOverrides.GetValueOnAnyThread())
+		{
+			EOverrideWidgetStatus::Type Status = GetOverrideStatus();
+			OverrideStatusWidget->SetStatus(Status);
+			
+			const FVector2f ImageSize = OverrideStatusWidget->GetDesiredSize();
+
+			FOverlayWidgetInfo Info;
+			Info.OverlayOffset = FVector2f(WidgetSize.X - ImageSize.X - 6.f - HorizontalOffset, 6.f);
+			HorizontalOffset += ImageSize.X + 6.f;
+			Info.Widget = OverrideStatusWidget;
+
+			Widgets.Add(Info);
+		}
+
 		bool bSetColor = false;
 		FLinearColor Color = FLinearColor::Black;
 		int32 PreviousNumWidgets = Widgets.Num();
@@ -1028,7 +1036,8 @@ TArray<FOverlayWidgetInfo> SRigVMGraphNode::GetOverlayWidgets(bool bSelected, co
 								   const FVector2f ImageSize = VisualDebugIndicatorWidget->GetDesiredSize();
 
 								   FOverlayWidgetInfo Info;
-								   Info.OverlayOffset = FVector2f(WidgetSize.X - ImageSize.X - 6.f, 6.f);
+								   Info.OverlayOffset = FVector2f(WidgetSize.X - ImageSize.X - 6.f - HorizontalOffset, 6.f);
+							   	   HorizontalOffset += ImageSize.X + 6.f;
 								   Info.Widget = VisualDebugIndicatorWidget;
 
 								   Widgets.Add(Info);
@@ -1414,6 +1423,81 @@ FText SRigVMGraphNode::GetInstructionDurationText() const
 	return FText();
 }
 
+TSharedRef<SWidget> SRigVMGraphNode::OnOverrideWidgetMenu() const
+{
+	if(const URigVMNode* Node = ModelNode.Get())
+	{
+		TSharedPtr<FOverrideStatusWidgetMenuBuilder> MenuBuilder = MakeShared<FOverrideStatusWidgetMenuBuilder>(Node, nullptr /* property path */);
+		MenuBuilder->GetStatusAttribute().Set(GetOverrideStatus());
+		MenuBuilder->OnAddOverride().BindLambda([this](const FOverrideStatusSubject&) -> FReply
+		{
+			if(const URigVMEdGraphNode* EdGraphNode = Cast<URigVMEdGraphNode>(GraphNode))
+			{
+				if(URigVMController* Controller = EdGraphNode->GetController())
+				{
+					if(Controller->AddOverrideToAllPinsOnNode(GraphNode->GetFName()))
+					{
+						return FReply::Handled();
+					}
+				}
+			}
+			return FReply::Unhandled();
+		});
+		MenuBuilder->OnClearOverride().BindLambda([this](const FOverrideStatusSubject&) -> FReply
+		{
+			if(const URigVMEdGraphNode* EdGraphNode = Cast<URigVMEdGraphNode>(GraphNode))
+			{
+				if(URigVMController* Controller = EdGraphNode->GetController())
+				{
+					if(Controller->ClearOverrideOnAllPinsOnNode(GraphNode->GetFName()))
+					{
+						return FReply::Handled();
+					}
+				}
+			}
+			return FReply::Unhandled();
+		});
+		MenuBuilder->OnResetToDefault().BindLambda([this](const FOverrideStatusSubject&) -> FReply
+		{
+			if(const URigVMEdGraphNode* EdGraphNode = Cast<URigVMEdGraphNode>(GraphNode))
+			{
+				if(URigVMController* Controller = EdGraphNode->GetController())
+				{
+					FRigVMDefaultValueTypeGuard _(Controller, ERigVMPinDefaultValueType::KeepValueType, true);
+					if(Controller->ResetDefaultValueForAllPinsOnNode(GraphNode->GetFName()))
+					{
+						return FReply::Handled();
+					}
+				}
+			}
+			return FReply::Unhandled();
+		});
+		MenuBuilder->OnValueDiffersFromDefault().BindLambda([this](const FOverrideStatusSubject&)
+		{
+			if(const URigVMNode* Node = Cast<URigVMNode>(ModelNode.Get()))
+			{
+				for(const URigVMPin* Pin : Node->GetPins())
+				{
+					if(Pin->CanProvideDefaultValue())
+					{
+						if(!Pin->HasOriginalDefaultValue())
+						{
+							return true;
+						}
+					}
+				}
+			}
+			return false;
+		});
+
+		if(const TSharedPtr<SWidget> MenuWidget = MenuBuilder->GenerateWidget())
+		{
+			return MenuWidget.ToSharedRef();
+		}
+	}
+	return SNullWidget::NullWidget;
+}
+
 int32 SRigVMGraphNode::GetNodeTopologyVersion() const
 {
 	if(URigVMEdGraphNode* RigVMEdGraphNode = Cast<URigVMEdGraphNode>(GraphNode))
@@ -2499,24 +2583,18 @@ void SRigVMGraphNode::UpdatePinTreeView()
 	];
 }
 
-bool SRigVMGraphNode::HasUserProvidedValue(const SGraphPin* InGraphPin) const
+ERigVMNodeDefaultValueOverrideState::Type SRigVMGraphNode::GetPinDefaultValueOverrideState() const
 {
-	if(!CVarRigVMEnablePinDefaultTypes.GetValueOnAnyThread())
+	if(!CVarRigVMEnablePinOverrides.GetValueOnAnyThread())
 	{
-		return false;
+		return ERigVMNodeDefaultValueOverrideState::None;
 	}
 
-	if (const URigVMEdGraphNode* RigVMEdGraphNode = Cast<URigVMEdGraphNode>(GraphNode))
+	if (ModelNode.IsValid())
 	{
-		if (const UEdGraphPin* Pin = InGraphPin->GetPinObj())
-		{
-			if(const URigVMPin* ModelPin = RigVMEdGraphNode->FindModelPinFromGraphPin(Pin))
-			{
-				return ModelPin->HasUserProvidedDefaultValue();
-			}
-		}
+		return ModelNode.Get()->GetPinDefaultValueOverrideState();
 	}
-	return false;
+	return ERigVMNodeDefaultValueOverrideState::None;
 }
 
 bool SRigVMGraphNode::IsCategoryPin(const SGraphPin* InGraphPin) const
@@ -2531,4 +2609,19 @@ bool SRigVMGraphNode::IsCategoryPin(const SGraphPin* InGraphPin) const
 	return false;
 }
 
+EOverrideWidgetStatus::Type SRigVMGraphNode::GetOverrideStatus() const
+{
+	const ERigVMNodeDefaultValueOverrideState::Type OverrideState = GetPinDefaultValueOverrideState();
+	if(OverrideState == ERigVMNodeDefaultValueOverrideState::None)
+	{
+		return EOverrideWidgetStatus::None;
+	}
+	if(OverrideState == ERigVMNodeDefaultValueOverrideState::SomePins)
+	{
+		return EOverrideWidgetStatus::ChangedInside;
+	}
+	
+	return EOverrideWidgetStatus::ChangedHere;
+}
+
 #undef LOCTEXT_NAMESPACE
diff --git a/Engine/Plugins/Runtime/RigVM/Source/RigVMEditor/Public/Editor/RigVMEditor.h b/Engine/Plugins/Runtime/RigVM/Source/RigVMEditor/Public/Editor/RigVMEditor.h
--- a/Engine/Plugins/Runtime/RigVM/Source/RigVMEditor/Public/Editor/RigVMEditor.h
+++ b/Engine/Plugins/Runtime/RigVM/Source/RigVMEditor/Public/Editor/RigVMEditor.h
@@ -491,6 +491,7 @@ protected:
 	TArray<TWeakObjectPtr<UObject>> GetSelectedObjects() const;
 	virtual void SetDetailObjects(const TArray<UObject*>& InObjects);
 	virtual void SetDetailObjects(const TArray<UObject*>& InObjects, bool bChangeUISelectionState);
+	virtual void SetDetailObjectFilter(TSharedPtr<FDetailsViewObjectFilter> InObjectFilter);
 	void SetMemoryStorageDetails(const TArray<FRigVMMemoryStorageStruct*>& InStructs);
 	void SetDetailViewForGraph(URigVMGraph* InGraph);
 	void SetDetailViewForFocusedGraph();
@@ -594,6 +595,7 @@ protected:
 	FString LastDebuggedHost;
 
 	bool bSuspendDetailsPanelRefresh;
+	bool bDetailsPanelRequiresClear;
 	bool bAllowBulkEdits;
 	bool bIsSettingObjectBeingDebugged;
 
diff --git a/Engine/Plugins/Runtime/RigVM/Source/RigVMEditor/Public/Editor/RigVMGraphDetailCustomization.h b/Engine/Plugins/Runtime/RigVM/Source/RigVMEditor/Public/Editor/RigVMGraphDetailCustomization.h
--- a/Engine/Plugins/Runtime/RigVM/Source/RigVMEditor/Public/Editor/RigVMGraphDetailCustomization.h
+++ b/Engine/Plugins/Runtime/RigVM/Source/RigVMEditor/Public/Editor/RigVMGraphDetailCustomization.h
@@ -666,7 +666,25 @@ protected:
 	// specializations for FRotator and FQuat at the end of this file
 	template<typename RotationType>
 	void ExtendRotationArgs(TSharedRef<class IPropertyHandle> InPropertyHandle, void* ArgumentsPtr) {}
-		
+
+	template<typename TransformType>
+	static FName GetTranslationPropertyName()
+	{
+		return TEXT("Translation");
+	}
+
+	template<typename TransformType>
+	static FName GetRotationPropertyName()
+	{
+		return TEXT("Rotation");
+	}
+
+	template<typename TransformType>
+	static FName GetScalePropertyName()
+	{
+		return TEXT("Scale3D");
+	}
+
 	TWeakObjectPtr<URigVMBlueprint> BlueprintBeingCustomized;
 	TWeakObjectPtr<URigVMGraph> GraphBeingCustomized;
 };
@@ -782,3 +800,15 @@ inline void FRigVMGraphMathTypeDetailCustomization::ExtendRotationArgs<FRotator>
 		OnRotationChanged<RotationType>(InPropertyHandle, InValue, true, InCommitType);
 	});
 }
+
+template<>
+inline FName FRigVMGraphMathTypeDetailCustomization::GetTranslationPropertyName<FEulerTransform>()
+{
+	return TEXT("Location");
+}
+
+template<>
+inline FName FRigVMGraphMathTypeDetailCustomization::GetScalePropertyName<FEulerTransform>()
+{
+	return TEXT("Scale");
+}
diff --git a/Engine/Plugins/Runtime/RigVM/Source/RigVMEditor/Public/Widgets/SRigVMGraphNode.h b/Engine/Plugins/Runtime/RigVM/Source/RigVMEditor/Public/Widgets/SRigVMGraphNode.h
--- a/Engine/Plugins/Runtime/RigVM/Source/RigVMEditor/Public/Widgets/SRigVMGraphNode.h
+++ b/Engine/Plugins/Runtime/RigVM/Source/RigVMEditor/Public/Widgets/SRigVMGraphNode.h
@@ -9,6 +9,7 @@
 #include "Widgets/Text/STextBlock.h"
 #include "RigVMModel/RigVMPin.h"
 #include "RigVMBlueprint.h"
+#include "Overrides/SOverrideStatusWidget.h"
 
 class URigVMEdGraphNode;
 class STableViewBase;
@@ -76,7 +77,6 @@ protected:
 	FText GetPinLabel(TWeakPtr<SGraphPin> GraphPin) const;
 
 	virtual TOptional<FSlateColor> GetHighlightColor(const SGraphPin* InGraphPin) const override;
-	virtual TOptional<const FSlateBrush*> GetPinBorder(const SGraphPin* InGraphPin) const override;
 	FSlateColor GetVariableLabelTextColor(TWeakObjectPtr<URigVMFunctionReferenceNode> FunctionReferenceNode, FName InVariableName) const;
 	FText GetVariableLabelTooltipText(TWeakObjectPtr<URigVMBlueprint> InBlueprint, FName InVariableName) const;
 
@@ -89,6 +89,8 @@ protected:
 	FText GetInstructionCountText() const;
 	FText GetInstructionDurationText() const;
 
+	TSharedRef<SWidget> OnOverrideWidgetMenu() const;
+
 protected:
 
 	int32 GetNodeTopologyVersion() const;
@@ -97,8 +99,9 @@ protected:
 	FReply OnExpanderArrowClicked(int32 InPinInfoIndex);
 	void HandleModifiedEvent(ERigVMGraphNotifType InNotifType, URigVMGraph* InGraph, UObject* InSubject);
 	virtual void UpdatePinTreeView();
-	bool HasUserProvidedValue(const SGraphPin* InGraphPin) const;
+	ERigVMNodeDefaultValueOverrideState::Type GetPinDefaultValueOverrideState() const;
 	bool IsCategoryPin(const SGraphPin* InGraphPin) const;
+	EOverrideWidgetStatus::Type GetOverrideStatus() const;
 
 	/** Cached widget title area */
 	TSharedPtr<SOverlay> TitleAreaWidget;
@@ -108,6 +111,7 @@ protected:
 	TSharedPtr<SImage> VisualDebugIndicatorWidget;
 	TSharedPtr<STextBlock> InstructionCountTextBlockWidget;
 	TSharedPtr<STextBlock> InstructionDurationTextBlockWidget;
+	TSharedPtr<SOverrideStatusWidget> OverrideStatusWidget; 
 
 	static const FSlateBrush* CachedImg_CR_Pin_Connected;
 	static const FSlateBrush* CachedImg_CR_Pin_Disconnected;
diff --git a/Engine/Plugins/Runtime/RigVM/Source/RigVMEditor/RigVMEditor.Build.cs b/Engine/Plugins/Runtime/RigVM/Source/RigVMEditor/RigVMEditor.Build.cs
--- a/Engine/Plugins/Runtime/RigVM/Source/RigVMEditor/RigVMEditor.Build.cs
+++ b/Engine/Plugins/Runtime/RigVM/Source/RigVMEditor/RigVMEditor.Build.cs
@@ -31,11 +31,13 @@ public class RigVMEditor : ModuleRules
                 "PropertyEditor",
                 "ToolMenus",
                 "MessageLog",
-				"StructUtilsEditor",
-				"ContentBrowser",
-				"ContentBrowserData",
-				"EditorFramework"
-			}
+                "StructUtilsEditor",
+                "ContentBrowser",
+                "ContentBrowserData",
+                "EditorFramework",
+                "AnimationEditorWidgets",
+                "WidgetRegistration",
+            }
 		);
         
         bool bAddBlueprintEditorDependency = true;
