RigVM: Performance improvement for registry type propagation

#rnx

diff --git a/Engine/Plugins/Animation/ControlRig/Source/ControlRig/Private/Units/Animation/RigUnit_AnimAttribute.cpp b/Engine/Plugins/Animation/ControlRig/Source/ControlRig/Private/Units/Animation/RigUnit_AnimAttribute.cpp
--- a/Engine/Plugins/Animation/ControlRig/Source/ControlRig/Private/Units/Animation/RigUnit_AnimAttribute.cpp
+++ b/Engine/Plugins/Animation/ControlRig/Source/ControlRig/Private/Units/Animation/RigUnit_AnimAttribute.cpp
@@ -386,7 +386,7 @@ FRigVMTemplateTypeMap FRigDispatch_GetAnimAttribute::OnNewArgumentType(const FNa
 	// similar pattern to FRigVMDispatch_ArrayGetAtIndex::OnNewArgumentType
 	// we used to need the if check below to avoid double registration for both default and value arguments, but since CL 29742624
 	// only the primary argument is used to add new permutations and thus this check becomes optional.
-	// primary argument is defined as the first argument that uses category, see FRigVMTemplate::UpdateArgumentTypes(), and in this case
+	// primary argument is defined as the first argument that uses category, see FRigVMTemplate::UpdateAllArgumentTypesSlow(), and in this case
 	// default is the first argument instead of the value
 	
 	if (ensure(InArgumentName == DefaultArgName))
diff --git a/Engine/Plugins/Runtime/RigVM/Source/RigVM/Private/RigVMCore/RigVMRegistry.cpp b/Engine/Plugins/Runtime/RigVM/Source/RigVM/Private/RigVMCore/RigVMRegistry.cpp
--- a/Engine/Plugins/Runtime/RigVM/Source/RigVM/Private/RigVMCore/RigVMRegistry.cpp
+++ b/Engine/Plugins/Runtime/RigVM/Source/RigVM/Private/RigVMCore/RigVMRegistry.cpp
@@ -23,7 +23,7 @@ static FDelayedAutoRegisterHelper GRigVMRegistrySingletonHelper(EDelayedRegister
 static TAutoConsoleVariable<bool> CVarRigVMUpdateDispatchFactoriesGreedily(TEXT("RigVM.UpdateDispatchFactoriesGreedily"), true, TEXT("Set this to false to avoid loading dispatch factories during engine init / plugin mount."));
 
 FRigVMRegistry_NoLock::FRigVMRegistry_NoLock() :
-	bIsRefreshingEngineTypes(false),
+	bAvoidTypePropagation(false),
 	bEverRefreshedEngineTypes(false),
 	bEverRefreshedDispatchFactoriesAfterEngineInit(false)
 {
@@ -373,7 +373,7 @@ void FRigVMRegistry_NoLock::RefreshEngineTypes_NoLock()
 {
 	TRACE_CPUPROFILER_EVENT_SCOPE(FRigVMRegistry::RefreshEngineTypes);
 
-	TGuardValue<bool> EnableGuardRefresh(bIsRefreshingEngineTypes, true);
+	TGuardValue<bool> AvoidTypePropagationGuard(bAvoidTypePropagation, true);
 
 	const int32 NumTypesBefore = Types.Num(); 
 	
@@ -446,7 +446,7 @@ void FRigVMRegistry_NoLock::RefreshEngineTypes_NoLock()
 				if(!TemplateProcessed[TemplateIndex])
 				{
 					FRigVMTemplate& Template = Templates[TemplateIndex];
-					(void)Template.UpdateArgumentTypes();
+					(void)Template.UpdateAllArgumentTypesSlow();
 					TemplateProcessed[TemplateIndex] = true;
 				}
 			}
@@ -867,7 +867,7 @@ void FRigVMRegistry_NoLock::RegisterTypeInCategory_NoLock(FRigVMTemplateArgument
 
 void FRigVMRegistry_NoLock::PropagateTypeAddedToCategory_NoLock(const FRigVMTemplateArgument::ETypeCategory InCategory, const TRigVMTypeIndex InTypeIndex)
 {
-	if(bIsRefreshingEngineTypes)
+	if(bAvoidTypePropagation)
 	{
 		return;
 	}
@@ -880,7 +880,7 @@ void FRigVMRegistry_NoLock::PropagateTypeAddedToCategory_NoLock(const FRigVMTemp
 		for(const int32 TemplateIndex : TemplatesToUseType)
 		{
 			FRigVMTemplate& Template = Templates[TemplateIndex];
-			(void)Template.UpdateArgumentTypes();
+			(void)Template.HandlePropagatedArgumentType(InTypeIndex);
 		}
 	}
 }
@@ -2037,7 +2037,7 @@ const FRigVMTemplate* FRigVMRegistry_NoLock::AddTemplateFromArguments_NoLock(con
 #if WITH_EDITOR
 	for(FRigVMTemplateArgument& Argument : Template.Arguments)
 	{
-		Argument.UpdateTypeToPermutations();
+		Argument.UpdateTypeToPermutationsSlow();
 	}
 #endif
 
diff --git a/Engine/Plugins/Runtime/RigVM/Source/RigVM/Private/RigVMCore/RigVMTemplate.cpp b/Engine/Plugins/Runtime/RigVM/Source/RigVM/Private/RigVMCore/RigVMTemplate.cpp
--- a/Engine/Plugins/Runtime/RigVM/Source/RigVM/Private/RigVMCore/RigVMTemplate.cpp
+++ b/Engine/Plugins/Runtime/RigVM/Source/RigVM/Private/RigVMCore/RigVMTemplate.cpp
@@ -90,7 +90,7 @@ FRigVMTemplateArgument::FRigVMTemplateArgument(FProperty* InProperty, FRigVMRegi
 	TypeIndices.Add(TypeIndex);
 	EnsureValidExecuteType_NoLock(InRegistry);
 #if WITH_EDITOR
-	UpdateTypeToPermutations();
+	UpdateTypeToPermutationsSlow();
 #endif
 }
 
@@ -117,7 +117,7 @@ FRigVMTemplateArgument::FRigVMTemplateArgument(const FName& InName, ERigVMPinDir
 {
 	EnsureValidExecuteType_NoLock(FRigVMRegistry_RWLock::Get());
 #if WITH_EDITOR
-	UpdateTypeToPermutations();
+	UpdateTypeToPermutationsSlow();
 #endif
 }
 
@@ -129,7 +129,7 @@ FRigVMTemplateArgument::FRigVMTemplateArgument(const FName& InName, ERigVMPinDir
 	check(TypeIndices.Num() > 0);
 	EnsureValidExecuteType_NoLock(FRigVMRegistry_RWLock::Get());
 #if WITH_EDITOR
-	UpdateTypeToPermutations();
+	UpdateTypeToPermutationsSlow();
 #endif
 }
 
@@ -194,7 +194,7 @@ FRigVMTemplateArgument::FRigVMTemplateArgument(const FName& InName, ERigVMPinDir
 		}
 
 #if WITH_EDITOR
-		UpdateTypeToPermutations();
+		UpdateTypeToPermutationsSlow();
 #endif
 	}
 }
@@ -209,7 +209,7 @@ void FRigVMTemplateArgument::EnsureValidExecuteType_NoLock(FRigVMRegistry_NoLock
 
 #if WITH_EDITOR
 
-void FRigVMTemplateArgument::UpdateTypeToPermutations()
+void FRigVMTemplateArgument::UpdateTypeToPermutationsSlow()
 {
 	TypeToPermutations.Reset();
 	TypeToPermutations.Reserve(GetNumTypes_NoLock());
@@ -2189,7 +2189,7 @@ FString FRigVMTemplate::GetKeywords() const
 
 #endif
 
-bool FRigVMTemplate::UpdateArgumentTypes()
+bool FRigVMTemplate::UpdateAllArgumentTypesSlow()
 {
 	const int32 PrimaryArgumentIndex = Arguments.IndexOfByPredicate([](const FRigVMTemplateArgument& Argument)
 	{
@@ -2232,90 +2232,147 @@ bool FRigVMTemplate::UpdateArgumentTypes()
 
 	PrimaryArgument.ForEachType([&](const TRigVMTypeIndex PrimaryTypeIndex)
 	{
-		TypesArray.Reset();
-		if(Factory)
+		if(!UpdateArgumentTypes_Impl(PrimaryArgument, PrimaryTypeIndex, Registry, Factory, TypesArray))
 		{
-			Factory->GetPermutationsFromArgumentType(PrimaryArgument.Name, PrimaryTypeIndex, TypesArray);
+			bResult = false;
 		}
-		else if(OnNewArgumentType().IsBound())
+		return true;
+	});
+
+#if WITH_EDITOR
+	for(int32 ArgumentIndex = 0; ArgumentIndex < Arguments.Num(); ArgumentIndex++)
+	{
+		FRigVMTemplateArgument& Argument = Arguments[ArgumentIndex];
+		if(Argument.bUseCategories)
 		{
-			FRigVMTemplateTypeMap Types = OnNewArgumentType().Execute(PrimaryArgument.Name, PrimaryTypeIndex);
-			TypesArray.Add(Types);
+			Argument.UpdateTypeToPermutationsSlow();
 		}
+	}
+#endif
+
+	return bResult;
+}
+
+bool FRigVMTemplate::UpdateArgumentTypes_Impl(
+	const FRigVMTemplateArgument& InPrimaryArgument,
+	const TRigVMTypeIndex InPrimaryTypeIndex,
+	const FRigVMRegistry_NoLock& InRegistry, 
+	const FRigVMDispatchFactory* InFactory,
+	TArray<FRigVMTemplateTypeMap, TInlineAllocator<1>>& InOutTypesArray)
+{
+	InOutTypesArray.Reset();
+	if(InFactory)
+	{
+		InFactory->GetPermutationsFromArgumentType(InPrimaryArgument.Name, InPrimaryTypeIndex, InOutTypesArray);
+	}
+	else if(OnNewArgumentType().IsBound())
+	{
+		FRigVMTemplateTypeMap Types = OnNewArgumentType().Execute(InPrimaryArgument.Name, InPrimaryTypeIndex);
+		InOutTypesArray.Add(Types);
+	}
 
-		if (!TypesArray.IsEmpty())
+	if (!InOutTypesArray.IsEmpty())
+	{
+		for (FRigVMTemplateTypeMap& Types : InOutTypesArray)
 		{
-			for (FRigVMTemplateTypeMap& Types : TypesArray)
+			if(Types.Num() == Arguments.Num())
 			{
-				if(Types.Num() == Arguments.Num())
+				for (TPair<FName, TRigVMTypeIndex>& ArgumentAndType : Types)
 				{
-					for (TPair<FName, TRigVMTypeIndex>& ArgumentAndType : Types)
-					{
-						// similar to FRigVMTemplateArgument::EnsureValidExecuteType
-						Registry.ConvertExecuteContextToBaseType_NoLock(ArgumentAndType.Value);
-					}
+					// similar to FRigVMTemplateArgument::EnsureValidExecuteType
+					InRegistry.ConvertExecuteContextToBaseType_NoLock(ArgumentAndType.Value);
+				}
 
-					// Find if these types were already registered
-					if (ContainsPermutation_NoLock(Types))
+				// Find if these types were already registered
+				if (ContainsPermutation_NoLock(Types))
+				{
+					return true;
+				}
+
+				uint32 TypeHash=0;
+				for(int32 ArgumentIndex = 0; ArgumentIndex < Arguments.Num(); ArgumentIndex++)
+				{
+					FRigVMTemplateArgument& Argument = Arguments[ArgumentIndex];
+
+					const TRigVMTypeIndex* TypeIndexPtr = Types.Find(Argument.Name);
+					if(TypeIndexPtr == nullptr)
 					{
-						return true;
+						return false;
 					}
-
-					uint32 TypeHash=0;
-					for(int32 ArgumentIndex = 0; ArgumentIndex < Arguments.Num(); ArgumentIndex++)
+					if(*TypeIndexPtr == INDEX_NONE)
 					{
-						FRigVMTemplateArgument& Argument = Arguments[ArgumentIndex];
-
-						const TRigVMTypeIndex* TypeIndexPtr = Types.Find(Argument.Name);
-						if(TypeIndexPtr == nullptr)
-						{
-							bResult = false;
-							return true;
-						}
-						if(*TypeIndexPtr == INDEX_NONE)
-						{
-							bResult = false;
-							return true;
-						}
+						return false;
+					}
 
-						const TRigVMTypeIndex& TypeIndex = *TypeIndexPtr;
-						TypeHash = HashCombine(TypeHash, GetTypeHash(TypeIndex));
+					const TRigVMTypeIndex& TypeIndex = *TypeIndexPtr;
+					TypeHash = HashCombine(TypeHash, GetTypeHash(TypeIndex));
 
-						if(Argument.bUseCategories || Argument.IsSingleton_NoLock())
-						{
-							continue;
-						}
-					
-						Argument.TypeIndices.Add(TypeIndex);
 #if WITH_EDITOR
-						Argument.TypeToPermutations.FindOrAdd(TypeIndex).Add(Permutations.Num());
+					Argument.TypeToPermutations.FindOrAdd(TypeIndex).Add(Permutations.Num());
 #endif
+					if(Argument.bUseCategories || Argument.IsSingleton_NoLock())
+					{
+						continue;
 					}
-
-					Permutations.Add(INDEX_NONE);
-					TypesHashToPermutation.Add(TypeHash, Permutations.Num()-1);
-				}
-				else
-				{
-					bResult = false;
+					Argument.TypeIndices.Add(TypeIndex);
 				}
+
+				Permutations.Add(INDEX_NONE);
+				TypesHashToPermutation.Add(TypeHash, Permutations.Num()-1);
+			}
+			else
+			{
+				return false;
 			}
 		}
-		else
-		{
-			bResult = false;
-		}
-		return true;
+	}
+	else
+	{
+		return false;
+	}
+	return true;
+}
+
+bool FRigVMTemplate::HandlePropagatedArgumentType(const TRigVMTypeIndex InTypeIndex)
+{
+	const int32 PrimaryArgumentIndex = Arguments.IndexOfByPredicate([](const FRigVMTemplateArgument& Argument)
+	{
+		return Argument.bUseCategories;
 	});
 
+	// this template may not be affected at all by this
+	if(PrimaryArgumentIndex == INDEX_NONE)
+	{
+		return true;
+	}
+
+	InvalidateHash();
+
+	const FRigVMRegistry_NoLock& Registry = FRigVMRegistry_NoLock::GetForRead();
+	const FRigVMTemplateArgument& PrimaryArgument = Arguments[PrimaryArgumentIndex];
+	TArray<FRigVMTemplateTypeMap, TInlineAllocator<1>> TypesArray;
+
+	const FRigVMDispatchFactory* Factory = nullptr;
+	if(Delegates.GetDispatchFactoryDelegate.IsBound())
+	{
+		Factory = Delegates.GetDispatchFactoryDelegate.Execute();
+		ensure(Factory);
+	}
+
+	const bool bResult = UpdateArgumentTypes_Impl(PrimaryArgument, InTypeIndex, Registry, Factory, TypesArray);
 #if WITH_EDITOR
-	for(int32 ArgumentIndex = 0; ArgumentIndex < Arguments.Num(); ArgumentIndex++)
+	if(bResult)
 	{
-		FRigVMTemplateArgument& Argument = Arguments[ArgumentIndex];
-		Argument.UpdateTypeToPermutations();
+		for(int32 ArgumentIndex = 0; ArgumentIndex < Arguments.Num(); ArgumentIndex++)
+		{
+			FRigVMTemplateArgument& Argument = Arguments[ArgumentIndex];
+			if(Argument.bUseCategories)
+			{
+				Argument.UpdateTypeToPermutationsSlow();
+			}
+		}
 	}
 #endif
-
 	return bResult;
 }
 
diff --git a/Engine/Plugins/Runtime/RigVM/Source/RigVM/Public/RigVMCore/RigVMRegistry.h b/Engine/Plugins/Runtime/RigVM/Source/RigVM/Public/RigVMCore/RigVMRegistry.h
--- a/Engine/Plugins/Runtime/RigVM/Source/RigVM/Public/RigVMCore/RigVMRegistry.h
+++ b/Engine/Plugins/Runtime/RigVM/Source/RigVM/Public/RigVMCore/RigVMRegistry.h
@@ -413,10 +413,9 @@ protected:
 	// All allowed structs
 	TSet<TObjectPtr<const UScriptStruct>> AllowedStructs;
 	
-
-
 	// If this is true the registry is currently refreshing all types
-	bool bIsRefreshingEngineTypes;
+	// and we want to avoid propagating types to the templates
+	bool bAvoidTypePropagation;
 
 	// This is true if the engine has ever refreshed the engine types
 	bool bEverRefreshedEngineTypes;
diff --git a/Engine/Plugins/Runtime/RigVM/Source/RigVM/Public/RigVMCore/RigVMTemplate.h b/Engine/Plugins/Runtime/RigVM/Source/RigVM/Public/RigVMCore/RigVMTemplate.h
--- a/Engine/Plugins/Runtime/RigVM/Source/RigVM/Public/RigVMCore/RigVMTemplate.h
+++ b/Engine/Plugins/Runtime/RigVM/Source/RigVM/Public/RigVMCore/RigVMTemplate.h
@@ -315,7 +315,7 @@ protected:
 	void EnsureValidExecuteType_NoLock(FRigVMRegistry_NoLock& InRegistry);
 
 #if WITH_EDITOR
-	void UpdateTypeToPermutations();
+	void UpdateTypeToPermutationsSlow();
 #endif
 
 	friend struct FRigVMTemplate;
@@ -534,7 +534,10 @@ public:
 
 	// Updates the template's argument types. This only affects templates which have category based
 	// arguments and will resolve the other arguments to the expected types.
-	bool UpdateArgumentTypes();
+	bool UpdateAllArgumentTypesSlow();
+
+	// Updates the template's argument types by adding a single new type
+	bool HandlePropagatedArgumentType(const TRigVMTypeIndex InTypeIndex);
 
 	// Invalidates template permutations whenever a type such as a user defined struct is removed
 	void HandleTypeRemoval(TRigVMTypeIndex InTypeIndex);
@@ -586,6 +589,13 @@ private:
 
 	uint32 ComputeTypeHash() const;
 
+	bool UpdateArgumentTypes_Impl(
+		const FRigVMTemplateArgument& InPrimaryArgument,
+		const TRigVMTypeIndex InPrimaryTypeIndex,
+		const FRigVMRegistry_NoLock& InRegistry, 
+		const FRigVMDispatchFactory* InFactory,
+		TArray<FRigVMTemplateTypeMap, TInlineAllocator<1>>& InOutTypesArray);
+
 	int32 Index;
 	FName Notation;
 	TArray<FRigVMTemplateArgument> Arguments;
