RigVM Developer: Fix truncation warnings
#rnx

diff --git a/Engine/Plugins/Runtime/RigVM/Source/RigVM/Private/RigVMCore/RigVMByteCode.cpp b/Engine/Plugins/Runtime/RigVM/Source/RigVM/Private/RigVMCore/RigVMByteCode.cpp
--- a/Engine/Plugins/Runtime/RigVM/Source/RigVM/Private/RigVMCore/RigVMByteCode.cpp
+++ b/Engine/Plugins/Runtime/RigVM/Source/RigVM/Private/RigVMCore/RigVMByteCode.cpp
@@ -116,18 +116,21 @@ void FRigVMCopyOp::Serialize(FArchive& Ar)
 	{
 		if (Ar.CustomVer(FUE5MainStreamObjectVersion::GUID) < FUE5MainStreamObjectVersion::RigVMCopyOpStoreNumBytes)
 		{
-			NumBytes = 0;
 			RegisterType = ERigVMRegisterType::Invalid;
 		}
 		else
 		{
-			Ar << NumBytes;
+			if (Ar.CustomVer(FRigVMObjectVersion::GUID) < FRigVMObjectVersion::ByteCodeCleanup)
+			{
+				uint16 NumBytes = 0;
+				Ar << NumBytes;
+			}
+			
 			Ar << RegisterType;
 		}
 	}
 	else
 	{
-		Ar << NumBytes;
 		Ar << RegisterType;
 	}
 }
@@ -1184,23 +1187,26 @@ int32 FRigVMByteCode::AddNotEqualsOp(const FRigVMOperand& InA, const FRigVMOpera
 	return ByteIndex;
 }
 
-int32 FRigVMByteCode::AddJumpOp(ERigVMOpCode InOpCode, uint16 InInstructionIndex)
+int32 FRigVMByteCode::AddJumpOp(ERigVMOpCode InOpCode, int32 InInstructionIndex)
 {
 	FRigVMJumpOp Op(InOpCode, InInstructionIndex);
 	return AddOp(Op);
 }
 
-int32 FRigVMByteCode::AddJumpIfOp(ERigVMOpCode InOpCode, uint16 InInstructionIndex, const FRigVMOperand& InConditionArg, bool bJumpWhenConditionIs)
+int32 FRigVMByteCode::AddJumpIfOp(ERigVMOpCode InOpCode, int32 InInstructionIndex, const FRigVMOperand& InConditionArg, bool bJumpWhenConditionIs)
 {
 	FRigVMJumpIfOp Op(InOpCode, InConditionArg, InInstructionIndex, bJumpWhenConditionIs);
 	return AddOp(Op);
 }
 
-int32 FRigVMByteCode::AddExecuteOp(uint16 InFunctionIndex, const FRigVMOperandArray& InOperands, const int32& StartPredicateIndex, const int32 PredicateCount)
+int32 FRigVMByteCode::AddExecuteOp(int32 InFunctionIndex, const FRigVMOperandArray& InOperands, const int32 InStartPredicateIndex, const int32 InPredicateCount)
 {
-	FRigVMExecuteOp Op(InFunctionIndex, IntCastChecked<uint16>(InOperands.Num()));
-	Op.FirstPredicateIndex = static_cast<uint16>(StartPredicateIndex);
-	Op.PredicateCount = static_cast<uint16>(PredicateCount);
+	FRigVMExecuteOp Op(IntCastChecked<uint16>(InFunctionIndex), IntCastChecked<uint16>(InOperands.Num()));
+	if (InPredicateCount)
+	{
+		Op.FirstPredicateIndex = IntCastChecked<uint16>(InStartPredicateIndex);
+		Op.PredicateCount = IntCastChecked<uint16>(InPredicateCount);
+	}
 	int32 OpByteIndex = AddOp(Op);
 
 	const int32 OperandsByteIndex = ByteCode.AddZeroed(sizeof(FRigVMOperand) * InOperands.Num());
@@ -1361,7 +1367,7 @@ FString FRigVMByteCode::DumpToText() const
 					{
 						break;
 					}
-					BranchInfoTexts.Add(FString::Printf(TEXT("%s (%u)"), *BranchInfo.Label.ToString(), BranchInfo.FirstInstruction));
+					BranchInfoTexts.Add(FString::Printf(TEXT("%s (%d)"), *BranchInfo.Label.ToString(), BranchInfo.FirstInstruction));
 				}
 				Line += RigVMStringUtils::JoinStrings(BranchInfoTexts, TEXT(", "));
 				break;
@@ -1454,8 +1460,8 @@ int32 FRigVMByteCode::AddBranchInfo(const FName& InBranchLabel, int32 InInstruct
 	BranchInfo.Label = InBranchLabel;
 	BranchInfo.InstructionIndex = InInstructionIndex;
 	BranchInfo.ArgumentIndex = InArgumentIndex;
-	BranchInfo.FirstInstruction = static_cast<uint16>(InFirstBranchInstruction);
-	BranchInfo.LastInstruction = static_cast<uint16>(InLastBranchInstruction);
+	BranchInfo.FirstInstruction = InFirstBranchInstruction;
+	BranchInfo.LastInstruction = InLastBranchInstruction;
 	return AddBranchInfo(BranchInfo);
 }
 
diff --git a/Engine/Plugins/Runtime/RigVM/Source/RigVM/Private/RigVMCore/RigVMMemoryStorage.cpp b/Engine/Plugins/Runtime/RigVM/Source/RigVM/Private/RigVMCore/RigVMMemoryStorage.cpp
--- a/Engine/Plugins/Runtime/RigVM/Source/RigVM/Private/RigVMCore/RigVMMemoryStorage.cpp
+++ b/Engine/Plugins/Runtime/RigVM/Source/RigVM/Private/RigVMCore/RigVMMemoryStorage.cpp
@@ -2,6 +2,7 @@
 
 #include "RigVMCore/RigVMMemoryStorage.h"
 #include "RigVMModule.h"
+#include "RigVMObjectVersion.h"
 #include "RigVMTypeUtils.h"
 #include "UObject/Interface.h"
 #include "AssetRegistry/AssetData.h"
@@ -69,8 +70,21 @@ void FRigVMBranchInfo::Serialize(FArchive& Ar)
 	
 	Ar << InstructionIndex;
 	Ar << ArgumentIndex;
-	Ar << FirstInstruction;
-	Ar << LastInstruction;
+
+	if (Ar.IsLoading() && Ar.CustomVer(FRigVMObjectVersion::GUID) < FRigVMObjectVersion::ByteCodeCleanup)
+	{
+		uint16 FirstInstructionTruncated = 0, LastInstructionTruncated = 0;
+		Ar << FirstInstructionTruncated;
+		Ar << LastInstructionTruncated;
+
+		FirstInstruction = FirstInstructionTruncated;
+		LastInstruction = LastInstructionTruncated;
+	}
+	else
+	{
+		Ar << FirstInstruction;
+		Ar << LastInstruction;
+	}
 }
 
 ERigVMExecuteResult FRigVMLazyBranch::Execute(FRigVMExtendedExecuteContext& Context)
diff --git a/Engine/Plugins/Runtime/RigVM/Source/RigVM/Public/RigVMCore/RigVMByteCode.h b/Engine/Plugins/Runtime/RigVM/Source/RigVM/Public/RigVMCore/RigVMByteCode.h
--- a/Engine/Plugins/Runtime/RigVM/Source/RigVM/Public/RigVMCore/RigVMByteCode.h
+++ b/Engine/Plugins/Runtime/RigVM/Source/RigVM/Public/RigVMCore/RigVMByteCode.h
@@ -245,7 +245,7 @@ struct RIGVM_API FRigVMBaseOp
 		return GetTypeHash(Op.OpCode);
 	}
 
-	void ZeroPaddedMemoryIfNeeded(FRigVMBaseOp* InMemory) const
+	static void ZeroPaddedMemoryIfNeeded(FRigVMBaseOp* InMemory)
 	{
 		RigVM::ZeroPaddedMemory(&InMemory->OpCode, reinterpret_cast<uint8*>(InMemory) + sizeof(FRigVMBaseOp));
 	}
@@ -301,7 +301,7 @@ struct RIGVM_API FRigVMExecuteOp : public FRigVMBaseOp
 		return Ar;
 	}
 	
-	void ZeroPaddedMemoryIfNeeded(FRigVMBaseOp* InMemory) const
+	static void ZeroPaddedMemoryIfNeeded(FRigVMBaseOp* InMemory)
 	{
 		FRigVMExecuteOp* This = reinterpret_cast<FRigVMExecuteOp*>(InMemory);
 		RigVM::ZeroPaddedMemory(&This->OpCode, &This->FunctionIndex);
@@ -324,19 +324,19 @@ struct RIGVM_API FRigVMUnaryOp : public FRigVMBaseOp
 		: FRigVMBaseOp(InOpCode)
 		, Arg(InArg)
 	{
-		ensure(
-			uint8(InOpCode) == uint8(ERigVMOpCode::Zero) ||
-			uint8(InOpCode) == uint8(ERigVMOpCode::BoolFalse) ||
-			uint8(InOpCode) == uint8(ERigVMOpCode::BoolTrue) ||
-			uint8(InOpCode) == uint8(ERigVMOpCode::Increment) ||
-			uint8(InOpCode) == uint8(ERigVMOpCode::Decrement) ||
-			uint8(InOpCode) == uint8(ERigVMOpCode::JumpAbsoluteIf) ||
-			uint8(InOpCode) == uint8(ERigVMOpCode::JumpForwardIf) ||
-			uint8(InOpCode) == uint8(ERigVMOpCode::JumpBackwardIf) ||
-			uint8(InOpCode) == uint8(ERigVMOpCode::ChangeType) ||
-			uint8(InOpCode) == uint8(ERigVMOpCode::JumpToBranch) ||
-			uint8(InOpCode) == uint8(ERigVMOpCode::RunInstructions) ||
-			uint8(InOpCode) == uint8(ERigVMOpCode::SetupTraits)
+		check(
+			InOpCode == ERigVMOpCode::Zero ||
+			InOpCode == ERigVMOpCode::BoolFalse ||
+			InOpCode == ERigVMOpCode::BoolTrue ||
+			InOpCode == ERigVMOpCode::Increment ||
+			InOpCode == ERigVMOpCode::Decrement ||
+			InOpCode == ERigVMOpCode::JumpAbsoluteIf ||
+			InOpCode == ERigVMOpCode::JumpForwardIf ||
+			InOpCode == ERigVMOpCode::JumpBackwardIf ||
+			InOpCode == ERigVMOpCode::ChangeType ||
+			InOpCode == ERigVMOpCode::JumpToBranch ||
+			InOpCode == ERigVMOpCode::RunInstructions ||
+			InOpCode == ERigVMOpCode::SetupTraits
 		);
 	}
 
@@ -344,7 +344,7 @@ struct RIGVM_API FRigVMUnaryOp : public FRigVMBaseOp
 
 	friend uint32 GetTypeHash(const FRigVMUnaryOp& Op)
 	{
-		return HashCombine(GetTypeHash((const FRigVMBaseOp&)Op), GetTypeHash(Op.Arg));
+		return HashCombine(GetTypeHash(static_cast<const FRigVMBaseOp&>(Op)), GetTypeHash(Op.Arg));
 	}
 
 	void Serialize(FArchive& Ar);
@@ -354,7 +354,7 @@ struct RIGVM_API FRigVMUnaryOp : public FRigVMBaseOp
 		return Ar;
 	}
 
-	void ZeroPaddedMemoryIfNeeded(FRigVMBaseOp* InMemory) const
+	static void ZeroPaddedMemoryIfNeeded(FRigVMBaseOp* InMemory)
 	{
 		FRigVMUnaryOp* This = reinterpret_cast<FRigVMUnaryOp*>(InMemory);
 		RigVM::ZeroPaddedMemory(&This->OpCode, &This->Arg);
@@ -387,7 +387,7 @@ struct RIGVM_API FRigVMBinaryOp : public FRigVMBaseOp
 
 	friend uint32 GetTypeHash(const FRigVMBinaryOp& Op)
 	{
-		uint32 Hash = GetTypeHash((const FRigVMBaseOp&)Op);
+		uint32 Hash = GetTypeHash(static_cast<const FRigVMBaseOp&>(Op));
 		Hash = HashCombine(Hash, GetTypeHash(Op.ArgA));
 		Hash = HashCombine(Hash, GetTypeHash(Op.ArgB));
 		return Hash;
@@ -400,7 +400,7 @@ struct RIGVM_API FRigVMBinaryOp : public FRigVMBaseOp
 		return Ar;
 	}
 	
-	void ZeroPaddedMemoryIfNeeded(FRigVMBaseOp* InMemory) const
+	static void ZeroPaddedMemoryIfNeeded(FRigVMBaseOp* InMemory)
 	{
 		FRigVMBinaryOp* This = reinterpret_cast<FRigVMBinaryOp*>(InMemory);
 		RigVM::ZeroPaddedMemory(&This->OpCode, &This->ArgA);
@@ -437,7 +437,7 @@ struct RIGVM_API FRigVMTernaryOp : public FRigVMBaseOp
 
 	friend uint32 GetTypeHash(const FRigVMTernaryOp& Op)
 	{
-		uint32 Hash = GetTypeHash((const FRigVMBaseOp&)Op);
+		uint32 Hash = GetTypeHash(static_cast<const FRigVMBaseOp&>(Op));
 		Hash = HashCombine(Hash, GetTypeHash(Op.ArgA));
 		Hash = HashCombine(Hash, GetTypeHash(Op.ArgB));
 		Hash = HashCombine(Hash, GetTypeHash(Op.ArgC));
@@ -451,7 +451,7 @@ struct RIGVM_API FRigVMTernaryOp : public FRigVMBaseOp
 		return Ar;
 	}
 
-	void ZeroPaddedMemoryIfNeeded(FRigVMBaseOp* InMemory) const
+	static void ZeroPaddedMemoryIfNeeded(FRigVMBaseOp* InMemory)
 	{
 		FRigVMTernaryOp* This = reinterpret_cast<FRigVMTernaryOp*>(InMemory);
 		RigVM::ZeroPaddedMemory(&This->OpCode, &This->ArgA);
@@ -483,9 +483,7 @@ struct RIGVM_API FRigVMQuaternaryOp : public FRigVMBaseOp
 		, ArgC(InArgC)
 		, ArgD(InArgD)
 	{
-		ensure(
-			uint8(InOpCode) == uint8(ERigVMOpCode::ArrayFind)
-		);
+		check(InOpCode == ERigVMOpCode::ArrayFind);
 	}
 
 	FRigVMOperand ArgA;
@@ -495,7 +493,7 @@ struct RIGVM_API FRigVMQuaternaryOp : public FRigVMBaseOp
 
 	friend uint32 GetTypeHash(const FRigVMQuaternaryOp& Op)
 	{
-		uint32 Hash = GetTypeHash((const FRigVMBaseOp&)Op);
+		uint32 Hash = GetTypeHash(static_cast<const FRigVMBaseOp&>(Op));
 		Hash = HashCombine(Hash, GetTypeHash(Op.ArgA));
 		Hash = HashCombine(Hash, GetTypeHash(Op.ArgB));
 		Hash = HashCombine(Hash, GetTypeHash(Op.ArgC));
@@ -510,7 +508,7 @@ struct RIGVM_API FRigVMQuaternaryOp : public FRigVMBaseOp
 		return Ar;
 	}
 	
-	void ZeroPaddedMemoryIfNeeded(FRigVMBaseOp* InMemory) const
+	static void ZeroPaddedMemoryIfNeeded(FRigVMBaseOp* InMemory)
 	{
 		FRigVMQuaternaryOp* This = reinterpret_cast<FRigVMQuaternaryOp*>(InMemory);
 		RigVM::ZeroPaddedMemory(&This->OpCode, &This->ArgA);
@@ -555,7 +553,7 @@ struct RIGVM_API FRigVMQuinaryOp : public FRigVMBaseOp
 
 	friend uint32 GetTypeHash(const FRigVMQuinaryOp& Op)
 	{
-		uint32 Hash = GetTypeHash((const FRigVMBaseOp&)Op);
+		uint32 Hash = GetTypeHash(static_cast<const FRigVMBaseOp&>(Op));
 		Hash = HashCombine(Hash, GetTypeHash(Op.ArgA));
 		Hash = HashCombine(Hash, GetTypeHash(Op.ArgB));
 		Hash = HashCombine(Hash, GetTypeHash(Op.ArgC));
@@ -571,7 +569,7 @@ struct RIGVM_API FRigVMQuinaryOp : public FRigVMBaseOp
 		return Ar;
 	}
 	
-	void ZeroPaddedMemoryIfNeeded(FRigVMBaseOp* InMemory) const
+	static void ZeroPaddedMemoryIfNeeded(FRigVMBaseOp* InMemory)
 	{
 		FRigVMQuinaryOp* This = reinterpret_cast<FRigVMQuinaryOp*>(InMemory);
 		RigVM::ZeroPaddedMemory(&This->OpCode, &This->ArgA);
@@ -620,7 +618,7 @@ struct RIGVM_API FRigVMSenaryOp : public FRigVMBaseOp
 
 	friend uint32 GetTypeHash(const FRigVMSenaryOp& Op)
 	{
-		uint32 Hash = GetTypeHash((const FRigVMBaseOp&)Op);
+		uint32 Hash = GetTypeHash(static_cast<const FRigVMBaseOp&>(Op));
 		Hash = HashCombine(Hash, GetTypeHash(Op.ArgA));
 		Hash = HashCombine(Hash, GetTypeHash(Op.ArgB));
 		Hash = HashCombine(Hash, GetTypeHash(Op.ArgC));
@@ -637,7 +635,7 @@ struct RIGVM_API FRigVMSenaryOp : public FRigVMBaseOp
 		return Ar;
 	}
 	
-	void ZeroPaddedMemoryIfNeeded(FRigVMBaseOp* InMemory) const
+	static void ZeroPaddedMemoryIfNeeded(FRigVMBaseOp* InMemory)
 	{
 		FRigVMSenaryOp* This = reinterpret_cast<FRigVMSenaryOp*>(InMemory);
 		RigVM::ZeroPaddedMemory(&This->OpCode, &This->ArgA);
@@ -671,7 +669,6 @@ public:
 	: FRigVMBaseOp(ERigVMOpCode::Copy)
 	, Source()
 	, Target()
-	, NumBytes(0)
 	, RegisterType(ERigVMRegisterType::Invalid)
 	, CopyType(ERigVMCopyType::Default)
 	{
@@ -684,7 +681,6 @@ public:
 		: FRigVMBaseOp(ERigVMOpCode::Copy)
 		, Source(InSource)
 		, Target(InTarget)
-		, NumBytes(0)
 		, RegisterType(ERigVMRegisterType::Invalid)
 		, CopyType(ERigVMCopyType::Default)
 	{
@@ -704,14 +700,13 @@ public:
 
 	friend uint32 GetTypeHash(const FRigVMCopyOp& Op)
 	{
-		uint32 Hash = GetTypeHash((const FRigVMBaseOp&)Op);
+		uint32 Hash = GetTypeHash(static_cast<const FRigVMBaseOp&>(Op));
 		Hash = HashCombine(Hash, GetTypeHash(Op.Source));
 		Hash = HashCombine(Hash, GetTypeHash(Op.Target));
 		return Hash;
 	}
 
 private:
-	uint16 NumBytes;
 	ERigVMRegisterType RegisterType;
 	ERigVMCopyType CopyType;
 
@@ -723,7 +718,7 @@ public:
 		return Ar;
 	}
 
-	void ZeroPaddedMemoryIfNeeded(FRigVMBaseOp* InMemory) const
+	static void ZeroPaddedMemoryIfNeeded(FRigVMBaseOp* InMemory)
 	{
 		FRigVMCopyOp* This = reinterpret_cast<FRigVMCopyOp*>(InMemory);
 		RigVM::ZeroPaddedMemory(&This->OpCode, &This->Source);
@@ -759,9 +754,9 @@ struct RIGVM_API FRigVMComparisonOp : public FRigVMBaseOp
 		, B(InB)
 		, Result(InResult)
 	{
-		ensure(
-			uint8(InOpCode) == uint8(ERigVMOpCode::Equals) ||
-			uint8(InOpCode) == uint8(ERigVMOpCode::NotEquals)
+		check(
+			InOpCode == ERigVMOpCode::Equals ||
+			InOpCode == ERigVMOpCode::NotEquals
 			);
 	}
 
@@ -771,7 +766,7 @@ struct RIGVM_API FRigVMComparisonOp : public FRigVMBaseOp
 
 	friend uint32 GetTypeHash(const FRigVMComparisonOp& Op)
 	{
-		uint32 Hash = GetTypeHash((const FRigVMBaseOp&)Op);
+		uint32 Hash = GetTypeHash(static_cast<const FRigVMBaseOp&>(Op));
 		Hash = HashCombine(Hash, GetTypeHash(Op.A));
 		Hash = HashCombine(Hash, GetTypeHash(Op.B));
 		Hash = HashCombine(Hash, GetTypeHash(Op.Result));
@@ -785,7 +780,7 @@ struct RIGVM_API FRigVMComparisonOp : public FRigVMBaseOp
 		return Ar;
 	}
 
-	void ZeroPaddedMemoryIfNeeded(FRigVMBaseOp* InMemory) const
+	static void ZeroPaddedMemoryIfNeeded(FRigVMBaseOp* InMemory)
 	{
 		FRigVMComparisonOp* This = reinterpret_cast<FRigVMComparisonOp*>(InMemory);
 		RigVM::ZeroPaddedMemory(&This->OpCode, &This->A);
@@ -813,15 +808,15 @@ struct RIGVM_API FRigVMJumpOp : public FRigVMBaseOp
 	: FRigVMBaseOp(InOpCode)
 	, InstructionIndex(InInstructionIndex)
 	{
-		ensure(uint8(InOpCode) >= uint8(ERigVMOpCode::JumpAbsolute));
-		ensure(uint8(InOpCode) <= uint8(ERigVMOpCode::JumpBackward));
+		check(InOpCode >= ERigVMOpCode::JumpAbsolute);
+		check(InOpCode <= ERigVMOpCode::JumpBackward);
 	}
 
 	int32 InstructionIndex;
 
 	friend uint32 GetTypeHash(const FRigVMJumpOp& Op)
 	{
-		uint32 Hash = GetTypeHash((const FRigVMBaseOp&)Op);
+		uint32 Hash = GetTypeHash(static_cast<const FRigVMBaseOp&>(Op));
 		Hash = HashCombine(Hash, GetTypeHash(Op.InstructionIndex));
 		return Hash;
 	}
@@ -833,7 +828,7 @@ struct RIGVM_API FRigVMJumpOp : public FRigVMBaseOp
 		return Ar;
 	}
 	
-	void ZeroPaddedMemoryIfNeeded(FRigVMBaseOp* InMemory) const
+	static void ZeroPaddedMemoryIfNeeded(FRigVMBaseOp* InMemory)
 	{
 		FRigVMJumpOp* This = reinterpret_cast<FRigVMJumpOp*>(InMemory);
 		RigVM::ZeroPaddedMemory(&This->OpCode, &This->InstructionIndex);
@@ -849,8 +844,7 @@ struct RIGVM_API FRigVMJumpIfOp : public FRigVMUnaryOp
 	GENERATED_USTRUCT_BODY()
 
 	FRigVMJumpIfOp()
-		: FRigVMUnaryOp()
-		, InstructionIndex(INDEX_NONE)
+		: InstructionIndex(INDEX_NONE)
 		, Condition(true)
 	{
 	}
@@ -860,8 +854,8 @@ struct RIGVM_API FRigVMJumpIfOp : public FRigVMUnaryOp
 		, InstructionIndex(InInstructionIndex)
 		, Condition(InCondition)
 	{
-		ensure(uint8(InOpCode) >= uint8(ERigVMOpCode::JumpAbsoluteIf));
-		ensure(uint8(InOpCode) <= uint8(ERigVMOpCode::JumpBackwardIf));
+		ensure(InOpCode >= ERigVMOpCode::JumpAbsoluteIf);
+		ensure(InOpCode <= ERigVMOpCode::JumpBackwardIf);
 	}
 
 	int32 InstructionIndex;
@@ -869,7 +863,7 @@ struct RIGVM_API FRigVMJumpIfOp : public FRigVMUnaryOp
 
 	friend uint32 GetTypeHash(const FRigVMJumpIfOp& Op)
 	{
-		uint32 Hash = GetTypeHash((const FRigVMUnaryOp&)Op);
+		uint32 Hash = GetTypeHash(static_cast<const FRigVMUnaryOp&>(Op));
 		Hash = HashCombine(Hash, GetTypeHash(Op.InstructionIndex));
 		Hash = HashCombine(Hash, GetTypeHash(Op.Condition));
 		return Hash;
@@ -882,7 +876,7 @@ struct RIGVM_API FRigVMJumpIfOp : public FRigVMUnaryOp
 		return Ar;
 	}
 	
-	void ZeroPaddedMemoryIfNeeded(FRigVMBaseOp* InMemory) const
+	static void ZeroPaddedMemoryIfNeeded(FRigVMBaseOp* InMemory)
 	{
 		FRigVMJumpIfOp* This = reinterpret_cast<FRigVMJumpIfOp*>(InMemory);
 		RigVM::ZeroPaddedMemory(&This->OpCode, &This->Arg);
@@ -925,7 +919,7 @@ struct RIGVM_API FRigVMInvokeEntryOp : public FRigVMBaseOp
 
 	friend uint32 GetTypeHash(const FRigVMInvokeEntryOp& Op)
 	{
-		return HashCombine(GetTypeHash((const FRigVMBaseOp&)Op), GetTypeHash(Op.EntryName.ToString()));
+		return HashCombine(GetTypeHash(static_cast<const FRigVMBaseOp&>(Op)), GetTypeHash(Op.EntryName.ToString()));
 	}
 
 	void Serialize(FArchive& Ar);
@@ -935,7 +929,7 @@ struct RIGVM_API FRigVMInvokeEntryOp : public FRigVMBaseOp
 		return Ar;
 	}
 
-	void ZeroPaddedMemoryIfNeeded(FRigVMBaseOp* InMemory) const
+	static void ZeroPaddedMemoryIfNeeded(FRigVMBaseOp* InMemory)
 	{
 		FRigVMInvokeEntryOp* This = reinterpret_cast<FRigVMInvokeEntryOp*>(InMemory);
 		RigVM::ZeroPaddedMemory(&This->OpCode, &This->EntryName);
@@ -964,7 +958,7 @@ struct RIGVM_API FRigVMJumpToBranchOp : public FRigVMUnaryOp
 
 	friend uint32 GetTypeHash(const FRigVMJumpToBranchOp& Op)
 	{
-		uint32 Hash = GetTypeHash((const FRigVMBaseOp&)Op);
+		uint32 Hash = GetTypeHash(static_cast<const FRigVMBaseOp&>(Op));
 		Hash = HashCombine(Hash, GetTypeHash(Op.FirstBranchInfoIndex));
 		return Hash;
 	}
@@ -1002,7 +996,7 @@ struct RIGVM_API FRigVMRunInstructionsOp : public FRigVMUnaryOp
 
 	friend uint32 GetTypeHash(const FRigVMRunInstructionsOp& Op)
 	{
-		uint32 Hash = GetTypeHash((const FRigVMBaseOp&)Op);
+		uint32 Hash = GetTypeHash(static_cast<const FRigVMBaseOp&>(Op));
 		Hash = HashCombine(Hash, GetTypeHash(Op.StartInstruction));
 		Hash = HashCombine(Hash, GetTypeHash(Op.EndInstruction));
 		return Hash;
@@ -1179,7 +1173,7 @@ public:
 	int32 FindEntryIndex(const FName& InEntryName) const;
 
 	// adds an execute operator given its function index operands
-	int32 AddExecuteOp(uint16 InFunctionIndex, const FRigVMOperandArray& InOperands, const int32& StartPredicateIndex, const int32 PredicateCount);
+	int32 AddExecuteOp(int32 InFunctionIndex, const FRigVMOperandArray& InOperands, const int32 InStartPredicateIndex, const int32 InPredicateCount);
 	
 	// adds an execute operator given its function index operands
 	int32 InlineFunction(const FRigVMByteCode* FunctionByteCode, const FRigVMOperandArray& InOperands);
@@ -1212,10 +1206,10 @@ public:
 	int32 AddNotEqualsOp(const FRigVMOperand& InA, const FRigVMOperand& InB, const FRigVMOperand& InResult);
 
 	// adds an absolute, forward or backward jump operator
-	int32 AddJumpOp(ERigVMOpCode InOpCode, uint16 InInstructionIndex);
+	int32 AddJumpOp(ERigVMOpCode InOpCode, int32 InInstructionIndex);
 
 	// adds an absolute, forward or backward jump operator based on a condition argument
-	int32 AddJumpIfOp(ERigVMOpCode InOpCode, uint16 InInstructionIndex, const FRigVMOperand& InConditionArg, bool bJumpWhenConditionIs = false);
+	int32 AddJumpIfOp(ERigVMOpCode InOpCode, int32 InInstructionIndex, const FRigVMOperand& InConditionArg, bool bJumpWhenConditionIs = false);
 
 	// adds an exit operator to exit the execution loop
 	int32 AddExitOp();
@@ -1327,10 +1321,9 @@ public:
 	bool IsOperandShared(const FRigVMOperand& InOperand) const { return GetInstructionsForOperand(InOperand).Num() > 1; }
 
 	// returns the raw data of the byte code
-	const TArrayView<const uint8> GetByteCode() const
+	TConstArrayView<uint8> GetByteCode() const
 	{
-		const uint8* Data = ByteCode.GetData();
-		return TArrayView<const uint8>((uint8*)Data, ByteCode.Num());
+		return {ByteCode.GetData(), ByteCode.Num()};
 	}
 
 	// returns the statistics information
diff --git a/Engine/Plugins/Runtime/RigVM/Source/RigVM/Public/RigVMCore/RigVMMemoryStorage.h b/Engine/Plugins/Runtime/RigVM/Source/RigVM/Public/RigVMCore/RigVMMemoryStorage.h
--- a/Engine/Plugins/Runtime/RigVM/Source/RigVM/Public/RigVMCore/RigVMMemoryStorage.h
+++ b/Engine/Plugins/Runtime/RigVM/Source/RigVM/Public/RigVMCore/RigVMMemoryStorage.h
@@ -125,10 +125,10 @@ struct RIGVM_API FRigVMBranchInfo
 	int32 ArgumentIndex;
 
 	UPROPERTY()
-	uint16 FirstInstruction;
+	int32 FirstInstruction;
 
 	UPROPERTY()
-	uint16 LastInstruction;
+	int32 LastInstruction;
 };
 
 
diff --git a/Engine/Plugins/Runtime/RigVM/Source/RigVM/Public/RigVMObjectVersion.h b/Engine/Plugins/Runtime/RigVM/Source/RigVM/Public/RigVMObjectVersion.h
--- a/Engine/Plugins/Runtime/RigVM/Source/RigVM/Public/RigVMObjectVersion.h
+++ b/Engine/Plugins/Runtime/RigVM/Source/RigVM/Public/RigVMObjectVersion.h
@@ -63,6 +63,10 @@ struct RIGVM_API FRigVMObjectVersion
 		// VM Bytecode Stores the Public Context Path as a FTopLevelAssetPath 
      	VMBytecodeStorePublicContextPathAsTopLevelAssetPath,
 
+		// Serialized instruction offsets are now int32 rather than uint16, NumBytes has been removed
+		// from RigVMCopyOp
+		ByteCodeCleanup,
+
 		// -----<new versions can be added above this line>-------------------------------------------------
 		VersionPlusOne,
 		LatestVersion = VersionPlusOne - 1
diff --git a/Engine/Plugins/Runtime/RigVM/Source/RigVMDeveloper/Private/EdGraph/RigVMEdGraph.cpp b/Engine/Plugins/Runtime/RigVM/Source/RigVMDeveloper/Private/EdGraph/RigVMEdGraph.cpp
--- a/Engine/Plugins/Runtime/RigVM/Source/RigVMDeveloper/Private/EdGraph/RigVMEdGraph.cpp
+++ b/Engine/Plugins/Runtime/RigVM/Source/RigVMDeveloper/Private/EdGraph/RigVMEdGraph.cpp
@@ -295,10 +295,10 @@ bool URigVMEdGraph::HandleModifiedEvent_Internal(ERigVMGraphNotifType InNotifTyp
 					NewNode->PostPlacedNewNode();
 					NewNode->AllocateDefaultPins();
 
-					NewNode->NodePosX = ModelNode->GetPosition().X;
-					NewNode->NodePosY = ModelNode->GetPosition().Y;
-					NewNode->NodeWidth = ModelNode->GetSize().X;
-					NewNode->NodeHeight = ModelNode->GetSize().Y;
+					NewNode->NodePosX = FMath::RoundToInt32(ModelNode->GetPosition().X);
+					NewNode->NodePosY = FMath::RoundToInt32(ModelNode->GetPosition().Y);
+					NewNode->NodeWidth = FMath::RoundToInt32(ModelNode->GetSize().X);
+					NewNode->NodeHeight = FMath::RoundToInt32(ModelNode->GetSize().Y);
 					NewNode->CommentColor = ModelNode->GetNodeColor();
 					NewNode->NodeComment = CommentModelNode->GetCommentText();
 					NewNode->FontSize = CommentModelNode->GetCommentFontSize();
@@ -322,8 +322,8 @@ bool URigVMEdGraph::HandleModifiedEvent_Internal(ERigVMGraphNotifType InNotifTyp
 					NewNode->AllocateDefaultPins();
 					NewNode->PostReconstructNode();
 
-					NewNode->NodePosX = ModelNode->GetPosition().X;
-					NewNode->NodePosY = ModelNode->GetPosition().Y;
+					NewNode->NodePosX = FMath::RoundToInt32(ModelNode->GetPosition().X);
+					NewNode->NodePosY = FMath::RoundToInt32(ModelNode->GetPosition().Y);
 					if (ModelNode->IsA<URigVMRerouteNode>())
 					{
 						if (URigVMPin* ModelPin = ModelNode->FindPin("Value"))
diff --git a/Engine/Plugins/Runtime/RigVM/Source/RigVMDeveloper/Private/RigVMBlueprint.cpp b/Engine/Plugins/Runtime/RigVM/Source/RigVMDeveloper/Private/RigVMBlueprint.cpp
--- a/Engine/Plugins/Runtime/RigVM/Source/RigVMDeveloper/Private/RigVMBlueprint.cpp
+++ b/Engine/Plugins/Runtime/RigVM/Source/RigVMDeveloper/Private/RigVMBlueprint.cpp
@@ -1703,7 +1703,7 @@ bool URigVMBlueprint::AddBreakpointToHost(URigVMNode* InBreakpointNode)
 					if (!AddedCallpaths.Contains(BreakpointCallPath))
 					{
 						AddedCallpaths.Add(BreakpointCallPath);
-						CDO->AddBreakpoint(i, InBreakpointNode, BreakpointProxy.GetCallstack().Num());
+						CDO->AddBreakpoint(i, InBreakpointNode, static_cast<uint16>(BreakpointProxy.GetCallstack().Num()));
 					}
 				}
 			}
@@ -2847,7 +2847,7 @@ bool URigVMBlueprint::ChangeMemberVariableType(const FName& InName, const FStrin
 	{
 		Variable.TypeName = *RigVMTypeUtils::CPPTypeFromEnum(Enum);
 		Variable.TypeObject = Enum;
-		Variable.Size = Enum->GetResourceSizeBytes(EResourceSizeMode::EstimatedTotal);
+		Variable.Size = static_cast<int32>(Enum->GetResourceSizeBytes(EResourceSizeMode::EstimatedTotal));
 	}
 
 	FEdGraphPinType PinType = RigVMTypeUtils::PinTypeFromExternalVariable(Variable);
@@ -3178,10 +3178,10 @@ void URigVMBlueprint::HandleModifiedEvent(ERigVMGraphNotifType InNotifType, URig
 						if(URigVMEdGraph* EdGraph = Cast<URigVMEdGraph>(GetEdGraph(RigVMNode->GetGraph())))
 						{
 							// decide if this graph should be renameable
-							const int32 NumberOfEvents = Algo::CountIf(RigVMNode->GetGraph()->GetNodes(), [](const URigVMNode* NodeToCount) -> bool
+							const int32 NumberOfEvents = static_cast<int32>(Algo::CountIf(RigVMNode->GetGraph()->GetNodes(), [](const URigVMNode* NodeToCount) -> bool
 							{
 								return NodeToCount->IsEvent() && NodeToCount->CanOnlyExistOnce();
-							});
+							}));
 							EdGraph->bAllowRenaming = NumberOfEvents != 1;
 						}
 					}
diff --git a/Engine/Plugins/Runtime/RigVM/Source/RigVMDeveloper/Private/RigVMCompiler/RigVMCodeGenerator.cpp b/Engine/Plugins/Runtime/RigVM/Source/RigVMDeveloper/Private/RigVMCompiler/RigVMCodeGenerator.cpp
--- a/Engine/Plugins/Runtime/RigVM/Source/RigVMDeveloper/Private/RigVMCompiler/RigVMCodeGenerator.cpp
+++ b/Engine/Plugins/Runtime/RigVM/Source/RigVMDeveloper/Private/RigVMCompiler/RigVMCodeGenerator.cpp
@@ -548,14 +548,14 @@ FString FRigVMCodeGenerator::DumpInstructions(const FRigVMExtendedExecuteContext
 				continue;
 			}
 			
-			if((int32)BranchInfo.FirstInstruction < Group.First ||
-				(int32)BranchInfo.LastInstruction > Group.Last)
+			if(BranchInfo.FirstInstruction < Group.First ||
+				BranchInfo.LastInstruction > Group.Last)
 			{
 				continue;
 			}
 
 			InstructionIndices.RemoveAll([BranchInfo](int32 Index) {
-				return FMath::IsWithinInclusive<uint16>((uint16)Index, BranchInfo.FirstInstruction, BranchInfo.LastInstruction);
+				return FMath::IsWithinInclusive<int32>(Index, BranchInfo.FirstInstruction, BranchInfo.LastInstruction);
 			});
 
 			// find the operand this lazy eval lambda belongs to
@@ -569,7 +569,7 @@ FString FRigVMCodeGenerator::DumpInstructions(const FRigVMExtendedExecuteContext
 
 			// dump the instructions for the lambda wrapped with the lambda definition
 			Lines.Add(Format(RigVM_LazyEvalLambdaDefine, BranchInfo.Index, BranchInfo.Label.ToString(), *OperandCPPType, *OperandName, *PropertyName));
-			Lines.Add(DumpInstructions(Context, TEXT("\t\t"), (int32)BranchInfo.FirstInstruction, (int32)BranchInfo.LastInstruction, Group, false));
+			Lines.Add(DumpInstructions(Context, TEXT("\t\t"), BranchInfo.FirstInstruction, BranchInfo.LastInstruction, Group, false));
 			Lines.Add(RigVM_LazyEvalLambdaReturn);
 			Lines.Emplace();
 
@@ -861,7 +861,7 @@ FString FRigVMCodeGenerator::DumpInstructions(const FRigVMExtendedExecuteContext
 					{
 						break;
 					}
-					Lines.Add(Prefix + Format(RigVM_JumpToBranchFormat, *GetOperandName(Context, Op.Arg, false), Branches[BranchIndex].Label.ToString(), (int32)Branch.FirstInstruction));
+					Lines.Add(Prefix + Format(RigVM_JumpToBranchFormat, *GetOperandName(Context, Op.Arg, false), Branches[BranchIndex].Label.ToString(), Branch.FirstInstruction));
 				}
 				break;
 			}
@@ -1563,7 +1563,7 @@ void FRigVMCodeGenerator::ParseInstructionGroups(const FRigVMExtendedExecuteCont
 						// to a jump to branch instruction
 						if(Branches[BranchIndex].IsOutputBranch())
 						{
-							Group.RequiredLabels.AddUnique((int32)Branches[BranchIndex].FirstInstruction);
+							Group.RequiredLabels.AddUnique(Branches[BranchIndex].FirstInstruction);
 						}
 					}
 					break;
diff --git a/Engine/Plugins/Runtime/RigVM/Source/RigVMDeveloper/Private/RigVMCompiler/RigVMCompiler.cpp b/Engine/Plugins/Runtime/RigVM/Source/RigVMDeveloper/Private/RigVMCompiler/RigVMCompiler.cpp
--- a/Engine/Plugins/Runtime/RigVM/Source/RigVMDeveloper/Private/RigVMCompiler/RigVMCompiler.cpp
+++ b/Engine/Plugins/Runtime/RigVM/Source/RigVMDeveloper/Private/RigVMCompiler/RigVMCompiler.cpp
@@ -1355,7 +1355,7 @@ bool URigVMCompiler::Compile(const FRigVMCompileSettings& InSettings, TArray<URi
 	// traverse all other blocks - this has to be an index based loop
 	if(!WorkData.LazyBlocksToProcess.IsEmpty())
 	{
-		const uint64 JumpToEndOfBlocksExternByte = WorkData.VM->GetByteCode().AddJumpOp(ERigVMOpCode::JumpForward, INDEX_NONE);
+		const int32 JumpToEndOfBlocksExternByte = WorkData.VM->GetByteCode().AddJumpOp(ERigVMOpCode::JumpForward, INDEX_NONE);
 		const int32 JumpToEndOfBlocksExternInstruction = WorkData.VM->GetByteCode().GetNumInstructions() - 1;
 
 		FRigVMByteCode& ByteCode = WorkData.VM->GetByteCode();
@@ -1389,7 +1389,7 @@ bool URigVMCompiler::Compile(const FRigVMCompileSettings& InSettings, TArray<URi
 				WorkData.LazyBlocks.FindChecked(WorkData.LazyBlocksToProcess[LazyBlockHashIndex]);
 
 			// update all run instructions ops in the bytecode
-			for(uint64 RunInstructionsByteCode : BlockInfo->RunInstructionsToUpdate)
+			for(int32 RunInstructionsByteCode : BlockInfo->RunInstructionsToUpdate)
 			{
 				FRigVMRunInstructionsOp& RunInstructionsOp = ByteCode.GetOpAt<FRigVMRunInstructionsOp>(RunInstructionsByteCode);
 				RunInstructionsOp.StartInstruction = BlockInfo->StartInstruction;
@@ -2317,7 +2317,7 @@ bool URigVMCompiler::TraverseCallExtern(const FRigVMCallExternExprAST* InExpr, F
 		if(!LazyChildExprs.IsEmpty())
 		{
 			// set up an operator to skip the lazy branches 
-			const uint64 JumpToCallExternByte = WorkData.VM->GetByteCode().AddJumpOp(ERigVMOpCode::JumpForward, INDEX_NONE);
+			const int32 JumpToCallExternByte = WorkData.VM->GetByteCode().AddJumpOp(ERigVMOpCode::JumpForward, INDEX_NONE);
 			const int32 JumpToCallExternInstruction = WorkData.VM->GetByteCode().GetNumInstructions() - 1;
 
 			// traverse the lazy children 
@@ -2861,8 +2861,8 @@ bool URigVMCompiler::TraverseInlineFunction(const FRigVMInlineFunctionExprAST* I
 				BranchIndexStart = BranchInfoIndex;
 			}
 			ByteCode.BranchInfos[BranchInfoIndex].InstructionIndex += InstructionIndexStart;
-			ByteCode.BranchInfos[BranchInfoIndex].FirstInstruction += (uint16)InstructionIndexStart;
-			ByteCode.BranchInfos[BranchInfoIndex].LastInstruction += (uint16)InstructionIndexStart;
+			ByteCode.BranchInfos[BranchInfoIndex].FirstInstruction += InstructionIndexStart;
+			ByteCode.BranchInfos[BranchInfoIndex].LastInstruction += InstructionIndexStart;
 		}
 
 		TMap<FRigVMOperand, FRigVMOperand> FunctionOperandToNewOperand;
@@ -2878,7 +2878,7 @@ bool URigVMCompiler::TraverseInlineFunction(const FRigVMInlineFunctionExprAST* I
 			}
 			const FRigVMInstruction& Instruction = Instructions[i];
 			const FRigVMOperandArray OperandArray = ByteCode.GetOperandsForOp(Instruction);
-			uint64 OperandsIndex = ByteCode.GetFirstOperandByteIndex(Instruction);
+			int32 OperandsIndex = ByteCode.GetFirstOperandByteIndex(Instruction);
 			for (int32 j=0; j<OperandArray.Num(); ++j)
 			{
 				FRigVMOperand* Operand = (FRigVMOperand*)(ByteCode.ByteCode.GetData() + OperandsIndex + j*sizeof(FRigVMOperand));
@@ -2906,7 +2906,7 @@ bool URigVMCompiler::TraverseInlineFunction(const FRigVMInlineFunctionExprAST* I
 				{
 					const FRigVMOperand& InterfaceOperand = InterfaceOperands[FunctionInterfaceParameterIndex];
 					Operand->MemoryType = InterfaceOperand.GetMemoryType();
-					Operand->RegisterIndex = InterfaceOperand.GetRegisterIndex();
+					Operand->RegisterIndex = static_cast<uint16>(InterfaceOperand.GetRegisterIndex());
 				}
 				else
 				{
@@ -2945,7 +2945,7 @@ bool URigVMCompiler::TraverseInlineFunction(const FRigVMInlineFunctionExprAST* I
 							break;
 						}
 					}
-					Operand->RegisterOffset = WorkData.FindOrAddPropertyPath(*Operand, Description.HeadCPPType, Description.SegmentPath);
+					Operand->RegisterOffset = IntCastChecked<uint16>(WorkData.FindOrAddPropertyPath(*Operand, Description.HeadCPPType, Description.SegmentPath));
 				}
 
 				// Store correspondence between Function operand and New operands
@@ -2970,7 +2970,7 @@ bool URigVMCompiler::TraverseInlineFunction(const FRigVMInlineFunctionExprAST* I
 			{
 				FRigVMExecuteOp& Op = ByteCode.GetOpAt<FRigVMExecuteOp>(Instruction);
 				const int32 FunctionIndex = WorkData.VM->AddRigVMFunction(FunctionCompilationData->FunctionNames[Op.FunctionIndex].ToString());
-				Op.FunctionIndex = FunctionIndex;
+				Op.FunctionIndex = IntCastChecked<uint16>(FunctionIndex);
 			}
 			else if (Instruction.OpCode == ERigVMOpCode::JumpToBranch)
 			{
@@ -3900,7 +3900,7 @@ TArray<URigVMPin*> URigVMCompiler::GetLinkedPins(URigVMPin* InPin, bool bInputs,
 	return LinkedPins;
 }
 
-uint16 URigVMCompiler::GetElementSizeFromCPPType(const FString& InCPPType, UScriptStruct* InScriptStruct)
+int32 URigVMCompiler::GetElementSizeFromCPPType(const FString& InCPPType, UScriptStruct* InScriptStruct)
 {
 	if (InScriptStruct == nullptr)
 	{
@@ -3914,7 +3914,7 @@ uint16 URigVMCompiler::GetElementSizeFromCPPType(const FString& InCPPType, UScri
 	{
 		return sizeof(bool);
 	}
-	else if (InCPPType == TEXT("int32"))
+	if (InCPPType == TEXT("int32"))
 	{
 		return sizeof(int32);
 	}
diff --git a/Engine/Plugins/Runtime/RigVM/Source/RigVMDeveloper/Private/RigVMModel/RigVMController.cpp b/Engine/Plugins/Runtime/RigVM/Source/RigVMDeveloper/Private/RigVMModel/RigVMController.cpp
--- a/Engine/Plugins/Runtime/RigVM/Source/RigVMDeveloper/Private/RigVMModel/RigVMController.cpp
+++ b/Engine/Plugins/Runtime/RigVM/Source/RigVMDeveloper/Private/RigVMModel/RigVMController.cpp
@@ -19775,7 +19775,7 @@ void URigVMController::SendUserFacingNotification(const FString& InMessage, floa
 
 	if(InDuration < SMALL_NUMBER)
 	{
-		InDuration = FMath::Clamp(0.1f * InMessage.Len(), 5.0f, 20.0f);
+		InDuration = FMath::Clamp(0.1f * static_cast<float>(InMessage.Len()), 5.0f, 20.0f);
 	}
 	
 	FNotificationInfo Info(FText::FromString(InMessage));
@@ -21154,7 +21154,7 @@ bool URigVMController::ChangePinType(URigVMPin* InPin, TRigVMTypeIndex InTypeInd
 	// compute the number of remaining wildcard pins
 	auto WildCardPinCountPredicate = [](const URigVMPin* Pin) { return Pin->IsWildCard(); };
 	TArray<URigVMPin*> AllPins = InPin->GetNode()->GetAllPinsRecursively();
-	int32 RemainingWildCardPins = Algo::CountIf(AllPins, WildCardPinCountPredicate);
+	int32 RemainingWildCardPins = static_cast<int32>(Algo::CountIf(AllPins, WildCardPinCountPredicate));
 	const bool bPinWasWildCard = InPin->IsWildCard();
 
 	const FPinState PreviousPinState = GetPinState(InPin);
@@ -21234,7 +21234,7 @@ bool URigVMController::ChangePinType(URigVMPin* InPin, TRigVMTypeIndex InTypeInd
 		if(InPin->GetNode()->IsA<URigVMTemplateNode>())
 		{
 			AllPins = InPin->GetNode()->GetAllPinsRecursively();
-			RemainingWildCardPins = Algo::CountIf(AllPins, WildCardPinCountPredicate);
+			RemainingWildCardPins = static_cast<int32>(Algo::CountIf(AllPins, WildCardPinCountPredicate));
 		}
 
 		// if this is the first time that there are no wild card pins left
diff --git a/Engine/Plugins/Runtime/RigVM/Source/RigVMDeveloper/Private/RigVMModel/RigVMNode.cpp b/Engine/Plugins/Runtime/RigVM/Source/RigVMDeveloper/Private/RigVMModel/RigVMNode.cpp
--- a/Engine/Plugins/Runtime/RigVM/Source/RigVMDeveloper/Private/RigVMModel/RigVMNode.cpp
+++ b/Engine/Plugins/Runtime/RigVM/Source/RigVMDeveloper/Private/RigVMModel/RigVMNode.cpp
@@ -1026,10 +1026,10 @@ FName URigVMNode::GetDisplayNameForPin(const URigVMPin* InPin) const
 		{
 			return FRigVMStruct::ExecuteName;
 		}
-		const int32 NumExecutePins = Algo::CountIf(GetPins(), [](const URigVMPin* Pin) -> bool
+		const int32 NumExecutePins = static_cast<int32>(Algo::CountIf(GetPins(), [](const URigVMPin* Pin) -> bool
 		{
 			return Pin->IsExecuteContext();
-		});
+		}));
 		if(NumExecutePins == 1)
 		{
 			return FRigVMStruct::ExecuteName;
diff --git a/Engine/Plugins/Runtime/RigVM/Source/RigVMDeveloper/Private/RigVMModel/RigVMSchema.cpp b/Engine/Plugins/Runtime/RigVM/Source/RigVMDeveloper/Private/RigVMModel/RigVMSchema.cpp
--- a/Engine/Plugins/Runtime/RigVM/Source/RigVMDeveloper/Private/RigVMModel/RigVMSchema.cpp
+++ b/Engine/Plugins/Runtime/RigVM/Source/RigVMDeveloper/Private/RigVMModel/RigVMSchema.cpp
@@ -480,7 +480,7 @@ bool URigVMSchema::CanRecolorNode(URigVMController* InController, const URigVMNo
 {
 	RIGVMSCHEMA_DEFAULT_FUNCTION_BODY
 
-	if (InNode->NodeColor.Equals(InNewColor, 0.001))
+	if (InNode->NodeColor.Equals(InNewColor, 0.001f))
 	{
 		return false;
 	}
diff --git a/Engine/Plugins/Runtime/RigVM/Source/RigVMDeveloper/Public/RigVMCompiler/RigVMCompiler.h b/Engine/Plugins/Runtime/RigVM/Source/RigVMDeveloper/Public/RigVMCompiler/RigVMCompiler.h
--- a/Engine/Plugins/Runtime/RigVM/Source/RigVMDeveloper/Public/RigVMCompiler/RigVMCompiler.h
+++ b/Engine/Plugins/Runtime/RigVM/Source/RigVMDeveloper/Public/RigVMCompiler/RigVMCompiler.h
@@ -221,7 +221,7 @@ public:
 		int32 StartInstruction;
 		int32 EndInstruction;
 		TArray<const FRigVMExprAST*> Expressions;
-		TArray<uint64> RunInstructionsToUpdate;
+		TArray<int32> RunInstructionsToUpdate;
 		bool bProcessed;
 	};
 
@@ -347,7 +347,7 @@ private:
 	TSet<const URigVMLibraryNode*> FunctionCompilationStack;
 
 	TArray<URigVMPin*> GetLinkedPins(URigVMPin* InPin, bool bInputs = true, bool bOutputs = true, bool bRecursive = true);
-	uint16 GetElementSizeFromCPPType(const FString& InCPPType, UScriptStruct* InScriptStruct);
+	int32 GetElementSizeFromCPPType(const FString& InCPPType, UScriptStruct* InScriptStruct);
 
 	static FString GetPinHashImpl(const URigVMPin* InPin, const FRigVMVarExprAST* InVarExpr, bool bIsDebugValue = false, const URigVMLibraryNode* FunctionCompiling = nullptr, const FRigVMASTProxy& InPinProxy = FRigVMASTProxy());
 
diff --git a/Engine/Plugins/Runtime/RigVM/Source/RigVMDeveloper/RigVMDeveloper.Build.cs b/Engine/Plugins/Runtime/RigVM/Source/RigVMDeveloper/RigVMDeveloper.Build.cs
--- a/Engine/Plugins/Runtime/RigVM/Source/RigVMDeveloper/RigVMDeveloper.Build.cs
+++ b/Engine/Plugins/Runtime/RigVM/Source/RigVMDeveloper/RigVMDeveloper.Build.cs
@@ -6,6 +6,8 @@ public class RigVMDeveloper : ModuleRules
 {
     public RigVMDeveloper(ReadOnlyTargetRules Target) : base(Target)
     {
+	    CppCompileWarningSettings.UnsafeTypeCastWarningLevel = WarningLevel.Error;
+
         PublicDependencyModuleNames.AddRange(
             new string[] {
                 "Core",
@@ -26,7 +28,6 @@ public class RigVMDeveloper : ModuleRules
                     "UnrealEd",
 					"Slate",
 					"SlateCore",
-					
 					"MessageLog",
 					"BlueprintGraph",
 					"GraphEditor",
