RigVM: Improve pin display name for struct member

#rnx

diff --git a/Engine/Plugins/Runtime/RigVM/Source/RigVMDeveloper/Private/RigVMModel/Nodes/RigVMUnitNode.cpp b/Engine/Plugins/Runtime/RigVM/Source/RigVMDeveloper/Private/RigVMModel/Nodes/RigVMUnitNode.cpp
--- a/Engine/Plugins/Runtime/RigVM/Source/RigVMDeveloper/Private/RigVMModel/Nodes/RigVMUnitNode.cpp
+++ b/Engine/Plugins/Runtime/RigVM/Source/RigVMDeveloper/Private/RigVMModel/Nodes/RigVMUnitNode.cpp
@@ -361,7 +361,7 @@ FName URigVMUnitNode::GetDisplayNameForPin(const URigVMPin* InPin) const
 			}
 		}
 		
-		const FName StructMemberDisplayName = GetDisplayNameForStructMember(ScriptStruct, InPin->GetPinPath());
+		const FName StructMemberDisplayName = GetDisplayNameForStructMember(ScriptStruct, InPin->GetSegmentPath(true));
 		if((StructMemberDisplayName == FRigVMStruct::ExecuteContextName) ||
 			(InPin->GetFName() == FRigVMStruct::ExecuteContextName && StructMemberDisplayName.IsNone()))
 		{
diff --git a/Engine/Plugins/Runtime/RigVM/Source/RigVMDeveloper/Private/RigVMModel/RigVMNode.cpp b/Engine/Plugins/Runtime/RigVM/Source/RigVMDeveloper/Private/RigVMModel/RigVMNode.cpp
--- a/Engine/Plugins/Runtime/RigVM/Source/RigVMDeveloper/Private/RigVMModel/RigVMNode.cpp
+++ b/Engine/Plugins/Runtime/RigVM/Source/RigVMDeveloper/Private/RigVMModel/RigVMNode.cpp
@@ -1035,6 +1035,30 @@ FName URigVMNode::GetDisplayNameForPin(const URigVMPin* InPin) const
 			return FRigVMStruct::ExecuteName;
 		}
 	}
+	return GetDisplayNameForStructMember(InPin);
+}
+
+FName URigVMNode::GetDisplayNameForStructMember(const URigVMPin* InPin)
+{
+#if WITH_EDITORONLY_DATA
+	if(InPin)
+	{
+		if(const URigVMPin* ParentPin = InPin->GetParentPin())
+		{
+			if(const UStruct* Struct = ParentPin->GetScriptStruct())
+			{
+				if(const FProperty* Property = Struct->FindPropertyByName(InPin->GetFName()))
+				{
+					const FText DisplayNameText = Property->GetDisplayNameText();
+					if(!DisplayNameText.IsEmpty())
+					{
+						return *DisplayNameText.ToString();
+					}
+				}
+			}
+		}
+	}
+#endif
 	return NAME_None;
 }
 
diff --git a/Engine/Plugins/Runtime/RigVM/Source/RigVMDeveloper/Private/RigVMModel/RigVMPin.cpp b/Engine/Plugins/Runtime/RigVM/Source/RigVMDeveloper/Private/RigVMModel/RigVMPin.cpp
--- a/Engine/Plugins/Runtime/RigVM/Source/RigVMDeveloper/Private/RigVMModel/RigVMPin.cpp
+++ b/Engine/Plugins/Runtime/RigVM/Source/RigVMDeveloper/Private/RigVMModel/RigVMPin.cpp
@@ -433,6 +433,14 @@ FName URigVMPin::GetDisplayName() const
 				return CachedDisplayName.GetValue();
 			}
 		}
+
+		const FName StructMemberDisplayName = URigVMNode::GetDisplayNameForStructMember(this);
+		if(!StructMemberDisplayName.IsNone())
+		{
+			CachedDisplayName = StructMemberDisplayName;
+			return CachedDisplayName.GetValue();
+		}
+		
 		CachedDisplayName = GetFName();
 		return CachedDisplayName.GetValue();
 	}
diff --git a/Engine/Plugins/Runtime/RigVM/Source/RigVMDeveloper/Public/RigVMModel/RigVMNode.h b/Engine/Plugins/Runtime/RigVM/Source/RigVMDeveloper/Public/RigVMModel/RigVMNode.h
--- a/Engine/Plugins/Runtime/RigVM/Source/RigVMDeveloper/Public/RigVMModel/RigVMNode.h
+++ b/Engine/Plugins/Runtime/RigVM/Source/RigVMDeveloper/Public/RigVMModel/RigVMNode.h
@@ -403,6 +403,7 @@ public:
 	virtual FName GetDisplayNameForPin(const URigVMPin* InPin) const;
 
 	// returns the display name for a struct member
+	static FName GetDisplayNameForStructMember(const URigVMPin* InPin);
 	static FName GetDisplayNameForStructMember(const UStruct* InStruct, const FString& InPath);
 
 	// returns the display name within a property
