RigVM Editor: Fix truncation warnings
#rnx

diff --git a/Engine/Plugins/Runtime/RigVM/Source/RigVMDeveloper/Private/EdGraph/RigVMEdGraphSchema.cpp b/Engine/Plugins/Runtime/RigVM/Source/RigVMDeveloper/Private/EdGraph/RigVMEdGraphSchema.cpp
--- a/Engine/Plugins/Runtime/RigVM/Source/RigVMDeveloper/Private/EdGraph/RigVMEdGraphSchema.cpp
+++ b/Engine/Plugins/Runtime/RigVM/Source/RigVMDeveloper/Private/EdGraph/RigVMEdGraphSchema.cpp
@@ -1359,7 +1359,7 @@ bool URigVMEdGraphSchema::IsStructEditable(UStruct* InStruct) const
 
 void URigVMEdGraphSchema::SetNodePosition(UEdGraphNode* Node, const FVector2f& Position) const
 {
-	return SetNodePosition(Node, FDeprecateSlateVector2D(Position), true);
+	return SetNodePosition(Node, FVector2D(Position), true);
 }
 
 void URigVMEdGraphSchema::SetNodePosition(UEdGraphNode* Node, const FVector2D& Position, bool bSetupUndo) const
diff --git a/Engine/Plugins/Runtime/RigVM/Source/RigVMEditor/Private/EdGraph/NodeSpawners/RigVMEdGraphUnitNodeSpawner.cpp b/Engine/Plugins/Runtime/RigVM/Source/RigVMEditor/Private/EdGraph/NodeSpawners/RigVMEdGraphUnitNodeSpawner.cpp
--- a/Engine/Plugins/Runtime/RigVM/Source/RigVMEditor/Private/EdGraph/NodeSpawners/RigVMEdGraphUnitNodeSpawner.cpp
+++ b/Engine/Plugins/Runtime/RigVM/Source/RigVMEditor/Private/EdGraph/NodeSpawners/RigVMEdGraphUnitNodeSpawner.cpp
@@ -400,7 +400,6 @@ void URigVMEdGraphUnitNodeSpawner::HookupMutableNode(URigVMNode* InModelNode, UR
 	if (ModelNodeExecutePin)
 	{
 		URigVMPin* ClosestOtherModelNodeExecutePin = nullptr;
-		float ClosestDistance = FLT_MAX;
 
 		const URigVMEdGraphSchema* Schema = GetDefault<URigVMEdGraphSchema>();
 		if (Schema->LastPinForCompatibleCheck)
@@ -419,6 +418,8 @@ void URigVMEdGraphUnitNodeSpawner::HookupMutableNode(URigVMNode* InModelNode, UR
 
 		if (ClosestOtherModelNodeExecutePin == nullptr)
 		{
+			double ClosestDistance = DBL_MAX;
+			
 			for (URigVMNode* OtherModelNode : Controller->GetGraph()->GetNodes())
 			{
 				if (OtherModelNode == InModelNode)
@@ -434,7 +435,7 @@ void URigVMEdGraphUnitNodeSpawner::HookupMutableNode(URigVMNode* InModelNode, UR
 						{
 							if (Pin->GetLinkedTargetPins().Num() == 0)
 							{
-								float Distance = (OtherModelNode->GetPosition() - InModelNode->GetPosition()).Size();
+								double Distance = (OtherModelNode->GetPosition() - InModelNode->GetPosition()).Size();
 								if (Distance < ClosestDistance)
 								{
 									ClosestOtherModelNodeExecutePin = Pin;
diff --git a/Engine/Plugins/Runtime/RigVM/Source/RigVMEditor/Private/EdGraph/RigVMEdGraphConnectionDrawingPolicy.cpp b/Engine/Plugins/Runtime/RigVM/Source/RigVMEditor/Private/EdGraph/RigVMEdGraphConnectionDrawingPolicy.cpp
--- a/Engine/Plugins/Runtime/RigVM/Source/RigVMEditor/Private/EdGraph/RigVMEdGraphConnectionDrawingPolicy.cpp
+++ b/Engine/Plugins/Runtime/RigVM/Source/RigVMEditor/Private/EdGraph/RigVMEdGraphConnectionDrawingPolicy.cpp
@@ -228,7 +228,7 @@ bool FRigVMEdGraphConnectionDrawingPolicy::GetAverageConnectedPositionForPin(UEd
 
 	if (ResultCount > 0)
 	{
-		OutPos = Result * (1.0f / ResultCount);
+		OutPos = Result * (1.0 / ResultCount);
 		return true;
 	}
 	else
diff --git a/Engine/Plugins/Runtime/RigVM/Source/RigVMEditor/Private/Editor/RigVMActionMenuBuilder.cpp b/Engine/Plugins/Runtime/RigVM/Source/RigVMEditor/Private/Editor/RigVMActionMenuBuilder.cpp
--- a/Engine/Plugins/Runtime/RigVM/Source/RigVMEditor/Private/Editor/RigVMActionMenuBuilder.cpp
+++ b/Engine/Plugins/Runtime/RigVM/Source/RigVMEditor/Private/Editor/RigVMActionMenuBuilder.cpp
@@ -383,7 +383,7 @@ bool FRigVMActionMenuBuilder::ProcessPendingActions()
 
 	bool bProcessedActions = false;
 	const double StartTime = FPlatformTime::Seconds();
-	const float MaxTimeThresholdSeconds = GetDefault<UBlueprintEditorSettings>()->ContextMenuTimeSlicingThresholdMs / 1000.0f;
+	const float MaxTimeThresholdSeconds = static_cast<float>(GetDefault<UBlueprintEditorSettings>()->ContextMenuTimeSlicingThresholdMs) / 1000.0f;
 
 	FBlueprintActionInfo* CurrentAction = MenuItemListAddHelper->GetNextAction();
 	while (CurrentAction)
diff --git a/Engine/Plugins/Runtime/RigVM/Source/RigVMEditor/Private/Editor/RigVMActionMenuItem.cpp b/Engine/Plugins/Runtime/RigVM/Source/RigVMEditor/Private/Editor/RigVMActionMenuItem.cpp
--- a/Engine/Plugins/Runtime/RigVM/Source/RigVMEditor/Private/Editor/RigVMActionMenuItem.cpp
+++ b/Engine/Plugins/Runtime/RigVM/Source/RigVMEditor/Private/Editor/RigVMActionMenuItem.cpp
@@ -239,9 +239,9 @@ UEdGraphNode* FRigVMActionMenuItem::PerformAction(UEdGraph* ParentGraph, UEdGrap
 		{
 			UEdGraphNode* FromNode = FromPin->GetOwningNode();
 			check(FromNode != nullptr);
-			const double FromNodeX = FromNode->NodePosX;
+			const float FromNodeX = static_cast<float>(FromNode->NodePosX);
 
-			static const double MinNodeDistance = 60.0; // min distance between spawned nodes (to keep them from overlapping)
+			constexpr float MinNodeDistance = 60.0f; // min distance between spawned nodes (to keep them from overlapping)
 			if (MinNodeDistance > FMath::Abs(FromNodeX - Location.X))
 			{
 				ModifiedLocation.X = FromNodeX - MinNodeDistance;
diff --git a/Engine/Plugins/Runtime/RigVM/Source/RigVMEditor/Private/Editor/RigVMEditor.cpp b/Engine/Plugins/Runtime/RigVM/Source/RigVMEditor/Private/Editor/RigVMEditor.cpp
--- a/Engine/Plugins/Runtime/RigVM/Source/RigVMEditor/Private/Editor/RigVMEditor.cpp
+++ b/Engine/Plugins/Runtime/RigVM/Source/RigVMEditor/Private/Editor/RigVMEditor.cpp
@@ -2445,7 +2445,7 @@ TRigVMTypeIndex FRigVMEditorBase::OnRequestPinTypeSelectionDialog(const TArray<T
 		}))
 		{
 			TypeNames.AddUnique(MakeShared<FName>(ArgumentType.CPPType));
-			TypeNameToIndex.Add(ArgumentType.CPPType, i);
+			TypeNameToIndex.Add(ArgumentType.CPPType, static_cast<uint8>(i));
 		}
 	}
 	TSharedPtr< SWindow > Window = SNew(SWindow)
diff --git a/Engine/Plugins/Runtime/RigVM/Source/RigVMEditor/Private/Editor/RigVMGraphDetailCustomization.cpp b/Engine/Plugins/Runtime/RigVM/Source/RigVMEditor/Private/Editor/RigVMGraphDetailCustomization.cpp
--- a/Engine/Plugins/Runtime/RigVM/Source/RigVMEditor/Private/Editor/RigVMGraphDetailCustomization.cpp
+++ b/Engine/Plugins/Runtime/RigVM/Source/RigVMEditor/Private/Editor/RigVMGraphDetailCustomization.cpp
@@ -296,7 +296,7 @@ void FRigVMFunctionArgumentLayout::GenerateHeaderRowContent(FDetailWidgetRow& No
 			[
 				SNew(SButton)
 				.ButtonStyle(FAppStyle::Get(), TEXT("SimpleButton"))
-				.ContentPadding(0)
+				.ContentPadding(0.f)
 				.IsEnabled_Raw(this, &FRigVMFunctionArgumentLayout::CanArgumentBeMoved, true)
 				.OnClicked(this, &FRigVMFunctionArgumentLayout::OnArgMoveUp)
 				.ToolTipText(LOCTEXT("FunctionArgDetailsArgMoveUpTooltip", "Move this parameter up in the list."))
@@ -312,7 +312,7 @@ void FRigVMFunctionArgumentLayout::GenerateHeaderRowContent(FDetailWidgetRow& No
 			[
 				SNew(SButton)
 				.ButtonStyle(FAppStyle::Get(), TEXT("SimpleButton"))
-				.ContentPadding(0)
+				.ContentPadding(0.f)
 				.IsEnabled_Raw(this, &FRigVMFunctionArgumentLayout::CanArgumentBeMoved, false)
 				.OnClicked(this, &FRigVMFunctionArgumentLayout::OnArgMoveDown)
 				.ToolTipText(LOCTEXT("FunctionArgDetailsArgMoveDownTooltip", "Move this parameter down in the list."))
@@ -1014,7 +1014,7 @@ void FRigVMGraphDetailCustomizationImpl::CustomizeDetails(IDetailLayoutBuilder&
         .ValueContent()
         [
             SNew(SComboButton)
-            .ContentPadding(0)
+            .ContentPadding(0.f)
             .ButtonContent()
             [
                 SNew(STextBlock)
@@ -2034,7 +2034,7 @@ void FRigVMWrappedNodeDetailCustomization::CustomizeDetails(IDetailLayoutBuilder
 				[
 					SNew(SButton)
 					.ButtonStyle(FAppStyle::Get(), TEXT("SimpleButton"))
-					.ContentPadding(0)
+					.ContentPadding(0.f)
 					.Text(FText::FromString(Identifier.GetLibraryNodePath()))
 					.OnClicked_Lambda([Header]() -> FReply
 					{
@@ -3060,8 +3060,8 @@ void FRigVMGraphEnumDetailCustomization::CustomizeHeader(TSharedRef<IPropertyHan
 			.HAlign(HAlign_Left)
 			[
 				SNew(SBox)
-					.MinDesiredWidth(150)
-					.MaxDesiredWidth(400)
+					.MinDesiredWidth(150.f)
+					.MaxDesiredWidth(400.f)
 					[
 						SNew(SRigVMEnumPicker)
 						.IsEnabled(true)
diff --git a/Engine/Plugins/Runtime/RigVM/Source/RigVMEditor/Private/RigVMEditorBlueprintLibrary.cpp b/Engine/Plugins/Runtime/RigVM/Source/RigVMEditor/Private/RigVMEditorBlueprintLibrary.cpp
--- a/Engine/Plugins/Runtime/RigVM/Source/RigVMEditor/Private/RigVMEditorBlueprintLibrary.cpp
+++ b/Engine/Plugins/Runtime/RigVM/Source/RigVMEditor/Private/RigVMEditorBlueprintLibrary.cpp
@@ -138,7 +138,7 @@ TArray<URigVMBlueprint*> URigVMEditorBlueprintLibrary::LoadAssetsWithAssetDataAn
 	TArray<URigVMBlueprint*> LoadedAssets;
 	{
 		const FString Title = FString::Printf(TEXT("Load all %s assets..."), *InClass->GetName()); 
-		FScopedSlowTask LoadAssetsTask(NumAssets, FText::FromString(Title));
+		FScopedSlowTask LoadAssetsTask(static_cast<float>(NumAssets), FText::FromString(Title));
 		LoadAssetsTask.MakeDialog(true);
 
 		for(int32 Index = 0; Index < NumAssets; Index++)
@@ -265,7 +265,7 @@ TArray<URigVMBlueprint*> URigVMEditorBlueprintLibrary::LoadAssetsWithAssetDataAn
 			{
 				NumNodes += Model->GetNodes().Num();
 			}
-			FScopedSlowTask FilterNodesTask(NumNodes, LOCTEXT("FilteringNodes", "Filtering Nodes"));
+			FScopedSlowTask FilterNodesTask(static_cast<float>(NumNodes), LOCTEXT("FilteringNodes", "Filtering Nodes"));
 			for(const URigVMGraph* Model : Models)
 			{
 				for(const URigVMNode* Node : Model->GetNodes())
diff --git a/Engine/Plugins/Runtime/RigVM/Source/RigVMEditor/Private/RigVMEditorModule.cpp b/Engine/Plugins/Runtime/RigVM/Source/RigVMEditor/Private/RigVMEditorModule.cpp
--- a/Engine/Plugins/Runtime/RigVM/Source/RigVMEditor/Private/RigVMEditorModule.cpp
+++ b/Engine/Plugins/Runtime/RigVM/Source/RigVMEditor/Private/RigVMEditorModule.cpp
@@ -1630,7 +1630,7 @@ void FRigVMEditorModule::GetPinVariableContextMenuActions(IRigVMClientHost* RigV
 
 			if (ModelPin->IsBoundToVariable())
 			{
-				FVector2D NodePosition = FVector2D(EdGraphNode->NodePosX - 200.f, EdGraphNode->NodePosY);
+				FVector2D NodePosition = FVector2D(EdGraphNode->NodePosX - 200.0, EdGraphNode->NodePosY);
 
 				FToolMenuSection& VariablesSection = Menu->AddSection("RigVMEditorContextMenuVariables", LOCTEXT("Variables", "Variables"));
 				VariablesSection.AddMenuEntry(
@@ -1645,7 +1645,7 @@ void FRigVMEditorModule::GetPinVariableContextMenuActions(IRigVMClientHost* RigV
 			}
 			else
 			{
-				FVector2D NodePosition = FVector2D(EdGraphNode->NodePosX - 200.f, EdGraphNode->NodePosY);
+				FVector2D NodePosition = FVector2D(EdGraphNode->NodePosX - 200.0, EdGraphNode->NodePosY);
 
 				FToolMenuSection& VariablesSection = Menu->AddSection("RigVMEditorContextMenuVariables", LOCTEXT("Variables", "Variables"));
 				VariablesSection.AddMenuEntry(
diff --git a/Engine/Plugins/Runtime/RigVM/Source/RigVMEditor/Private/Widgets/SRigVMEditorGraphExplorer.cpp b/Engine/Plugins/Runtime/RigVM/Source/RigVMEditor/Private/Widgets/SRigVMEditorGraphExplorer.cpp
--- a/Engine/Plugins/Runtime/RigVM/Source/RigVMEditor/Private/Widgets/SRigVMEditorGraphExplorer.cpp
+++ b/Engine/Plugins/Runtime/RigVM/Source/RigVMEditor/Private/Widgets/SRigVMEditorGraphExplorer.cpp
@@ -475,7 +475,7 @@ void SRigVMEditorGraphExplorer::CreateWidgets()
 	                    .ComboButtonStyle(FAppStyle::Get(), "ToolbarComboButton")
 	                    .ForegroundColor(FSlateColor::UseForeground())
 	                    .HasDownArrow(false)
-	                    .ContentPadding(0)
+	                    .ContentPadding(0.f)
 	                    .AddMetaData<FTagMetaData>(FTagMetaData(TEXT("ViewOptions")))
 	                    .MenuContent()
 	                    [
diff --git a/Engine/Plugins/Runtime/RigVM/Source/RigVMEditor/Private/Widgets/SRigVMExecutionStackView.cpp b/Engine/Plugins/Runtime/RigVM/Source/RigVMEditor/Private/Widgets/SRigVMExecutionStackView.cpp
--- a/Engine/Plugins/Runtime/RigVM/Source/RigVMEditor/Private/Widgets/SRigVMExecutionStackView.cpp
+++ b/Engine/Plugins/Runtime/RigVM/Source/RigVMEditor/Private/Widgets/SRigVMExecutionStackView.cpp
@@ -912,7 +912,7 @@ void SRigVMExecutionStackView::HandleGoToInstruction()
 	if(OptionsDialog->ShowModal() == 0)
 	{
 		TreeView->SetSelection(Operators[Index]);
-		TreeView->SetScrollOffset(FMath::Max(Index-5, 0));
+		TreeView->SetScrollOffset(static_cast<float>(FMath::Max(Index-5, 0)));
 	}
 }
 
@@ -1078,7 +1078,7 @@ void SRigVMExecutionStackView::HandleExecutionHalted(const int32 InHaltedAtInstr
 	{
 		HaltedAtInstruction = InHaltedAtInstruction;
 		TreeView->SetSelection(Operators[InHaltedAtInstruction]);
-		TreeView->SetScrollOffset(FMath::Max(InHaltedAtInstruction-5, 0));
+		TreeView->SetScrollOffset(static_cast<float>(FMath::Max(InHaltedAtInstruction-5, 0)));
 	}	
 }
 
diff --git a/Engine/Plugins/Runtime/RigVM/Source/RigVMEditor/Private/Widgets/SRigVMGraphNode.cpp b/Engine/Plugins/Runtime/RigVM/Source/RigVMEditor/Private/Widgets/SRigVMGraphNode.cpp
--- a/Engine/Plugins/Runtime/RigVM/Source/RigVMEditor/Private/Widgets/SRigVMGraphNode.cpp
+++ b/Engine/Plugins/Runtime/RigVM/Source/RigVMEditor/Private/Widgets/SRigVMGraphNode.cpp
@@ -1064,10 +1064,10 @@ TArray<FOverlayWidgetInfo> SRigVMGraphNode::GetOverlayWidgets(bool bSelected, co
 						const int32 Count = ModelNode->GetInstructionVisitedCount(DebuggedHost->GetRigVMExtendedExecuteContext(), DebuggedHost->GetVM(), FRigVMASTProxy());
 						if((Count > Blueprint->RigGraphDisplaySettings.NodeRunLowerBound) || bShowInstructionIndex)
 						{
-							const int32 VOffset = bSelected ? -2 : 2;
+							const float VOffset = bSelected ? -2.f : 2.f;
 							const FVector2f TextSize = InstructionCountTextBlockWidget->GetDesiredSize();
 							FOverlayWidgetInfo Info;
-							Info.OverlayOffset = FVector2f(WidgetSize.X - TextSize.X - 8.f, VOffset - TextSize.Y);
+							Info.OverlayOffset = FVector2f(WidgetSize.X - TextSize.X - 8.0f, VOffset - TextSize.Y);
 							Info.Widget = InstructionCountTextBlockWidget;
 							Widgets.Add(Info);
 						}
@@ -1078,7 +1078,7 @@ TArray<FOverlayWidgetInfo> SRigVMGraphNode::GetOverlayWidgets(bool bSelected, co
 						const double MicroSeconds = ModelNode->GetInstructionMicroSeconds(DebuggedHost->GetRigVMExtendedExecuteContext(), DebuggedHost->GetVM(), FRigVMASTProxy());
 						if(MicroSeconds >= 0.0)
 						{
-							const int32 VOffset = bSelected ? -2 : 2;
+							const float VOffset = bSelected ? -2.f : 2.f;
 							const FVector2f TextSize = InstructionDurationTextBlockWidget->GetDesiredSize();
 							FOverlayWidgetInfo Info;
 							Info.OverlayOffset = FVector2f(8.f, VOffset - TextSize.Y);
@@ -2192,7 +2192,7 @@ void SRigVMGraphNode::UpdatePinTreeView()
 					[
 						SNew(SButton)
 						.ButtonStyle(FAppStyle::Get(), TEXT("SimpleButton"))
-						.ContentPadding(0)
+						.ContentPadding(0.f)
 						.VAlign(VAlign_Center)
 						.HAlign(HAlign_Center)
 						.ClickMethod( EButtonClickMethod::MouseDown )
diff --git a/Engine/Plugins/Runtime/RigVM/Source/RigVMEditor/Private/Widgets/SRigVMGraphNodeComment.cpp b/Engine/Plugins/Runtime/RigVM/Source/RigVMEditor/Private/Widgets/SRigVMGraphNodeComment.cpp
--- a/Engine/Plugins/Runtime/RigVM/Source/RigVMEditor/Private/Widgets/SRigVMGraphNodeComment.cpp
+++ b/Engine/Plugins/Runtime/RigVM/Source/RigVMEditor/Private/Widgets/SRigVMGraphNodeComment.cpp
@@ -168,7 +168,7 @@ void SRigVMGraphNodeComment::MoveTo(const FVector2f& NewPosition, FNodeSet& Node
 					// Don't drag note content if either of the shift keys are down.
 					FModifierKeysState KeysState = FSlateApplication::Get().GetModifierKeys();
 
-					if((NodeToMove == GraphNode) || (IsNodeUnderComment(FVector2D(CommentOldPosition), FVector2D(OldPosition)) && !KeysState.IsShiftDown()))
+					if((NodeToMove == GraphNode) || (IsNodeUnderComment(CommentOldPosition, OldPosition) && !KeysState.IsShiftDown()))
 					{ 
 						const FVector2f AccumulatedPosition = OldPosition + Delta; 
 						RigSchema->SetNodePosition(NodeToMove, FVector2D(AccumulatedPosition), false);
@@ -206,23 +206,23 @@ void SRigVMGraphNodeComment::OnCommentTextCommitted(const FText& NewComment, ETe
 
 bool SRigVMGraphNodeComment::IsNodeUnderComment(UEdGraphNode_Comment* InCommentNode, const TSharedRef<SGraphNode> InNodeWidget) const
 {
-	return IsNodeUnderComment(FVector2D(InNodeWidget->GetPosition2f()));
+	return IsNodeUnderComment(InNodeWidget->GetPosition2f());
 }
 
 
 bool SRigVMGraphNodeComment::IsNodeUnderComment(UEdGraphNode* InNode) const
 {
-	return IsNodeUnderComment(FVector2D(InNode->NodePosX, InNode->NodePosY));
+	return IsNodeUnderComment(FVector2f(static_cast<float>(InNode->NodePosX), static_cast<float>(InNode->NodePosY)));
 }
 
-bool SRigVMGraphNodeComment::IsNodeUnderComment(const FVector2D& InNodePosition) const
+bool SRigVMGraphNodeComment::IsNodeUnderComment(const FVector2f& InNodePosition) const
 {
-	return IsNodeUnderComment(FVector2D(GetPosition2f()), InNodePosition);
+	return IsNodeUnderComment(GetPosition2f(), InNodePosition);
 }
 
-bool SRigVMGraphNodeComment::IsNodeUnderComment(const FVector2D& InCommentPosition, const FVector2D& InNodePosition) const
+bool SRigVMGraphNodeComment::IsNodeUnderComment(const FVector2f& InCommentPosition, const FVector2f& InNodePosition) const
 {
-	const FVector2D NodeSize = GetDesiredSize();
+	const FVector2f NodeSize = GetDesiredSize();
 	const FSlateRect CommentRect(InCommentPosition.X, InCommentPosition.Y, InCommentPosition.X + NodeSize.X, InCommentPosition.Y + NodeSize.Y);
 
 	return CommentRect.ContainsPoint(InNodePosition);
diff --git a/Engine/Plugins/Runtime/RigVM/Source/RigVMEditor/Private/Widgets/SRigVMGraphPinCurveFloat.cpp b/Engine/Plugins/Runtime/RigVM/Source/RigVMEditor/Private/Widgets/SRigVMGraphPinCurveFloat.cpp
--- a/Engine/Plugins/Runtime/RigVM/Source/RigVMEditor/Private/Widgets/SRigVMGraphPinCurveFloat.cpp
+++ b/Engine/Plugins/Runtime/RigVM/Source/RigVMEditor/Private/Widgets/SRigVMGraphPinCurveFloat.cpp
@@ -25,10 +25,10 @@ TSharedRef<SWidget>	SRigVMGraphPinCurveFloat::GetDefaultValueWidget()
 
 	// 360 is the minimum width required to display keys' values 
 	TSharedRef<SWidget> Widget = SNew(SBox)
-		.MinDesiredWidth(360)
-		.MaxDesiredWidth(400)
-		.MinDesiredHeight(175)
-		.MaxDesiredHeight(300)
+		.MinDesiredWidth(360.f)
+		.MaxDesiredWidth(400.f)
+		.MinDesiredHeight(175.f)
+		.MaxDesiredHeight(300.f)
 		[
 			SAssignNew(CurveEditor, SCurveEditor)
 			.Visibility(this, &SGraphPin::GetDefaultValueVisibility)
diff --git a/Engine/Plugins/Runtime/RigVM/Source/RigVMEditor/Private/Widgets/SRigVMGraphPinEnumPicker.cpp b/Engine/Plugins/Runtime/RigVM/Source/RigVMEditor/Private/Widgets/SRigVMGraphPinEnumPicker.cpp
--- a/Engine/Plugins/Runtime/RigVM/Source/RigVMEditor/Private/Widgets/SRigVMGraphPinEnumPicker.cpp
+++ b/Engine/Plugins/Runtime/RigVM/Source/RigVMEditor/Private/Widgets/SRigVMGraphPinEnumPicker.cpp
@@ -22,8 +22,8 @@ void SRigVMEnumPicker::Construct(const FArguments& InArgs)
 	ChildSlot
 	[
 		SNew(SBox)
-		.MinDesiredWidth(150)
-		.MaxDesiredWidth(400)
+		.MinDesiredWidth(150.f)
+		.MaxDesiredWidth(400.f)
 		[
 			SNew(SSearchableComboBox)
 			.OptionsSource(&EnumOptions)
diff --git a/Engine/Plugins/Runtime/RigVM/Source/RigVMEditor/Private/Widgets/SRigVMGraphPinNameList.cpp b/Engine/Plugins/Runtime/RigVM/Source/RigVMEditor/Private/Widgets/SRigVMGraphPinNameList.cpp
--- a/Engine/Plugins/Runtime/RigVM/Source/RigVMEditor/Private/Widgets/SRigVMGraphPinNameList.cpp
+++ b/Engine/Plugins/Runtime/RigVM/Source/RigVMEditor/Private/Widgets/SRigVMGraphPinNameList.cpp
@@ -55,7 +55,7 @@ TSharedRef<SWidget>	SRigVMGraphPinNameList::GetDefaultValueWidget()
 	}
 
 	return SNew(SBox)
-		.MinDesiredWidth(150)
+		.MinDesiredWidth(150.f)
 		[
 			SNew(SHorizontalBox)
 			+SHorizontalBox::Slot()
@@ -108,7 +108,7 @@ TSharedRef<SWidget>	SRigVMGraphPinNameList::GetDefaultValueWidget()
 				.ButtonStyle( FAppStyle::Get(), "NoBorder" )
 				.ButtonColorAndOpacity( this, &SRigVMGraphPinNameList::OnGetWidgetBackground )
 				.OnClicked(this, &SRigVMGraphPinNameList::HandleBrowseClicked)
-				.ContentPadding(0)
+				.ContentPadding(0.f)
 				.ToolTipText(NSLOCTEXT("SRigVMGraphPinNameList", "ObjectGraphPin_Browse_Tooltip", "Browse"))
 				.Visibility(OnGetNameFromSelection.IsBound() ? EVisibility::Visible : EVisibility::Collapsed)
 				[
diff --git a/Engine/Plugins/Runtime/RigVM/Source/RigVMEditor/Private/Widgets/SRigVMGraphPinQuat.cpp b/Engine/Plugins/Runtime/RigVM/Source/RigVMEditor/Private/Widgets/SRigVMGraphPinQuat.cpp
--- a/Engine/Plugins/Runtime/RigVM/Source/RigVMEditor/Private/Widgets/SRigVMGraphPinQuat.cpp
+++ b/Engine/Plugins/Runtime/RigVM/Source/RigVMEditor/Private/Widgets/SRigVMGraphPinQuat.cpp
@@ -18,8 +18,8 @@ void SRigVMGraphPinQuat::Construct(const FArguments& InArgs, UEdGraphPin* InGrap
 TSharedRef<SWidget>	SRigVMGraphPinQuat::GetDefaultValueWidget()
 {
 	TSharedRef<SWidget> Widget = SNew(SBox)
-		.MinDesiredWidth(150)
-		.MaxDesiredWidth(400)
+		.MinDesiredWidth(150.f)
+		.MaxDesiredWidth(400.f)
 		[
 			SNew(SRotatorInputBox)
 			.DisplayToggle(false)
@@ -97,13 +97,13 @@ TOptional<float> SRigVMGraphPinQuat::GetRotatorComponent(int32 InComponent) cons
 	{
 		if(InComponent == 0)
 		{
-			return Rotator.GetValue().Pitch;
+			return static_cast<float>(Rotator.GetValue().Pitch);
 		}
 		if(InComponent == 1)
 		{
-			return Rotator.GetValue().Yaw;
+			return static_cast<float>(Rotator.GetValue().Yaw);
 		}
-		return Rotator.GetValue().Roll;
+		return static_cast<float>(Rotator.GetValue().Roll);
 	}
 	return TOptional<float>();
 }
diff --git a/Engine/Plugins/Runtime/RigVM/Source/RigVMEditor/Private/Widgets/SRigVMGraphPinUserDataNameSpace.cpp b/Engine/Plugins/Runtime/RigVM/Source/RigVMEditor/Private/Widgets/SRigVMGraphPinUserDataNameSpace.cpp
--- a/Engine/Plugins/Runtime/RigVM/Source/RigVMEditor/Private/Widgets/SRigVMGraphPinUserDataNameSpace.cpp
+++ b/Engine/Plugins/Runtime/RigVM/Source/RigVMEditor/Private/Widgets/SRigVMGraphPinUserDataNameSpace.cpp
@@ -26,8 +26,8 @@ TSharedRef<SWidget>	SRigVMGraphPinUserDataNameSpace::GetDefaultValueWidget()
 	}
 
 	return SNew(SBox)
-		.MinDesiredWidth(150)
-		.MaxDesiredWidth(400)
+		.MinDesiredWidth(150.f)
+		.MaxDesiredWidth(400.f)
 		[
 			SAssignNew(NameComboBox, SRigVMGraphPinEditableNameValueWidget)
 				.Visibility(this, &SGraphPin::GetDefaultValueVisibility)
diff --git a/Engine/Plugins/Runtime/RigVM/Source/RigVMEditor/Private/Widgets/SRigVMGraphPinVariableName.cpp b/Engine/Plugins/Runtime/RigVM/Source/RigVMEditor/Private/Widgets/SRigVMGraphPinVariableName.cpp
--- a/Engine/Plugins/Runtime/RigVM/Source/RigVMEditor/Private/Widgets/SRigVMGraphPinVariableName.cpp
+++ b/Engine/Plugins/Runtime/RigVM/Source/RigVMEditor/Private/Widgets/SRigVMGraphPinVariableName.cpp
@@ -32,8 +32,8 @@ TSharedRef<SWidget>	SRigVMGraphPinVariableName::GetDefaultValueWidget()
 	}
 
 	return SNew(SBox)
-		.MinDesiredWidth(150)
-		.MaxDesiredWidth(400)
+		.MinDesiredWidth(150.f)
+		.MaxDesiredWidth(400.f)
 		[
 			SAssignNew(NameComboBox, SRigVMGraphPinEditableNameValueWidget)
 				.Visibility(this, &SGraphPin::GetDefaultValueVisibility)
diff --git a/Engine/Plugins/Runtime/RigVM/Source/RigVMEditor/Private/Widgets/SRigVMNodeLayoutWidget.cpp b/Engine/Plugins/Runtime/RigVM/Source/RigVMEditor/Private/Widgets/SRigVMNodeLayoutWidget.cpp
--- a/Engine/Plugins/Runtime/RigVM/Source/RigVMEditor/Private/Widgets/SRigVMNodeLayoutWidget.cpp
+++ b/Engine/Plugins/Runtime/RigVM/Source/RigVMEditor/Private/Widgets/SRigVMNodeLayoutWidget.cpp
@@ -390,7 +390,7 @@ void SRigVMNodeLayoutWidget::SRigVMNowLayoutRow::Construct(const FArguments& InA
 	[
 		HorizontalBox
 	];
-	SuperArguments.Padding(0);
+	SuperArguments.Padding(0.f);
 	SuperArguments.OnCanAcceptDrop(InArgs._OnCanAcceptDrop);
 	SuperArguments.OnAcceptDrop(InArgs._OnAcceptDrop);
 	SuperArguments.OnPaintDropIndicator(InArgs._OnPaintDropIndicator);
diff --git a/Engine/Plugins/Runtime/RigVM/Source/RigVMEditor/Private/Widgets/SRigVMVariableMappingWidget.cpp b/Engine/Plugins/Runtime/RigVM/Source/RigVMEditor/Private/Widgets/SRigVMVariableMappingWidget.cpp
--- a/Engine/Plugins/Runtime/RigVM/Source/RigVMEditor/Private/Widgets/SRigVMVariableMappingWidget.cpp
+++ b/Engine/Plugins/Runtime/RigVM/Source/RigVMEditor/Private/Widgets/SRigVMVariableMappingWidget.cpp
@@ -101,7 +101,7 @@ void SRigVMVariableMappingTreeRow::Construct(const FArguments& InArgs, const TSh
 					.OnSelectionChanged(this, &SRigVMVariableMappingTreeRow::OnVarOptionSourceChanged)
 					.OnComboBoxOpening(this, &SRigVMVariableMappingTreeRow::OnVarOptionComboOpening)
 					.IsEnabled(this, &SRigVMVariableMappingTreeRow::IsVarOptionEnabled)
-					.ContentPadding(2)
+					.ContentPadding(2.f)
 					.Content()
 					[
 						SNew(STextBlock)
diff --git a/Engine/Plugins/Runtime/RigVM/Source/RigVMEditor/Private/Widgets/SRigVMVariantTagWidget.cpp b/Engine/Plugins/Runtime/RigVM/Source/RigVMEditor/Private/Widgets/SRigVMVariantTagWidget.cpp
--- a/Engine/Plugins/Runtime/RigVM/Source/RigVMEditor/Private/Widgets/SRigVMVariantTagWidget.cpp
+++ b/Engine/Plugins/Runtime/RigVM/Source/RigVMEditor/Private/Widgets/SRigVMVariantTagWidget.cpp
@@ -36,7 +36,7 @@ void SRigVMVariantCapsule::Construct(
 	SButton::FArguments ButtonArgs;
 	ButtonArgs
 	.ButtonStyle( &FRigVMEditorStyle::Get().GetWidgetStyle< FButtonStyle >( "TagButton" ) )
-	.ContentPadding(0)
+	.ContentPadding(0.f)
 	.HAlign(HAlign_Fill)
 	.OnClicked(InArgs._OnClicked)
 	.Content()
diff --git a/Engine/Plugins/Runtime/RigVM/Source/RigVMEditor/Private/Widgets/SRigVMVariantWidget.cpp b/Engine/Plugins/Runtime/RigVM/Source/RigVMEditor/Private/Widgets/SRigVMVariantWidget.cpp
--- a/Engine/Plugins/Runtime/RigVM/Source/RigVMEditor/Private/Widgets/SRigVMVariantWidget.cpp
+++ b/Engine/Plugins/Runtime/RigVM/Source/RigVMEditor/Private/Widgets/SRigVMVariantWidget.cpp
@@ -303,7 +303,7 @@ void SRigVMVariantWidget::SRigVMVariantRefTreeRow::Construct(const FArguments& I
 	[
 		InArgs._Content.ToSharedRef()
 	];
-	SuperArguments.Padding(0);
+	SuperArguments.Padding(0.f);
 	
 	STableRow< TSharedPtr<FVariantTreeRowInfo> >::Construct(SuperArguments, OwnerTableView);
 }
@@ -375,8 +375,8 @@ TSharedPtr<SWidget> SRigVMVariantWidget::CreateDefaultVariantRefRow(const FRigVM
 				[
 					SNew(SBox)
 					.ToolTip(TooltipWithTags)
-					.WidthOverride(32)
-					.HeightOverride(32)
+					.WidthOverride(32.f)
+					.HeightOverride(32.f)
 					[
 						AssetThumbnail->MakeThumbnailWidget(ThumbnailConfig)
 					]
diff --git a/Engine/Plugins/Runtime/RigVM/Source/RigVMEditor/Public/Editor/SRigVMActionMenu.cpp b/Engine/Plugins/Runtime/RigVM/Source/RigVMEditor/Public/Editor/SRigVMActionMenu.cpp
--- a/Engine/Plugins/Runtime/RigVM/Source/RigVMEditor/Public/Editor/SRigVMActionMenu.cpp
+++ b/Engine/Plugins/Runtime/RigVM/Source/RigVMEditor/Public/Editor/SRigVMActionMenu.cpp
@@ -178,7 +178,7 @@ void SRigVMActionMenu::Construct( const FArguments& InArgs, TSharedPtr<IRigVMEdi
 				.AutoHeight()
 				[
 					SNew(SBox)
-					.HeightOverride(2)
+					.HeightOverride(2.f)
 					.Visibility_Lambda([this]()
 					{
 						return ContextMenuBuilder.IsValid() && ContextMenuBuilder->GetNumPendingActions() > 0 ? EVisibility::SelfHitTestInvisible : EVisibility::Collapsed;
diff --git a/Engine/Plugins/Runtime/RigVM/Source/RigVMEditor/Public/Widgets/SRigVMGraphNodeComment.h b/Engine/Plugins/Runtime/RigVM/Source/RigVMEditor/Public/Widgets/SRigVMGraphNodeComment.h
--- a/Engine/Plugins/Runtime/RigVM/Source/RigVMEditor/Public/Widgets/SRigVMGraphNodeComment.h
+++ b/Engine/Plugins/Runtime/RigVM/Source/RigVMEditor/Public/Widgets/SRigVMGraphNodeComment.h
@@ -22,8 +22,8 @@ protected:
 
 	virtual bool IsNodeUnderComment(UEdGraphNode_Comment* InCommentNode, const TSharedRef<SGraphNode> InNodeWidget) const override;
 	bool IsNodeUnderComment(UEdGraphNode* InNode) const;
-	bool IsNodeUnderComment(const FVector2D& InNodePosition) const;
-	bool IsNodeUnderComment(const FVector2D& InCommentPosition, const FVector2D& InNodePosition) const;
+	bool IsNodeUnderComment(const FVector2f& InNodePosition) const;
+	bool IsNodeUnderComment(const FVector2f& InCommentPosition, const FVector2f& InNodePosition) const;
 
 private:
 
diff --git a/Engine/Plugins/Runtime/RigVM/Source/RigVMEditor/Public/Widgets/SRigVMNodePreviewWidget.h b/Engine/Plugins/Runtime/RigVM/Source/RigVMEditor/Public/Widgets/SRigVMNodePreviewWidget.h
--- a/Engine/Plugins/Runtime/RigVM/Source/RigVMEditor/Public/Widgets/SRigVMNodePreviewWidget.h
+++ b/Engine/Plugins/Runtime/RigVM/Source/RigVMEditor/Public/Widgets/SRigVMNodePreviewWidget.h
@@ -11,7 +11,7 @@ class RIGVMEDITOR_API SRigVMNodePreviewWidget : public SCompoundWidget
 public:
 	
 	SLATE_BEGIN_ARGS(SRigVMNodePreviewWidget)
-		: _Padding(16)
+		: _Padding(16.f)
 	{}
 	SLATE_ATTRIBUTE(FMargin, Padding)
 	SLATE_ARGUMENT(TSharedPtr<FRigVMMinimalEnvironment>, Environment)
diff --git a/Engine/Plugins/Runtime/RigVM/Source/RigVMEditor/Public/Widgets/SRigVMVariantTagWidget.h b/Engine/Plugins/Runtime/RigVM/Source/RigVMEditor/Public/Widgets/SRigVMVariantTagWidget.h
--- a/Engine/Plugins/Runtime/RigVM/Source/RigVMEditor/Public/Widgets/SRigVMVariantTagWidget.h
+++ b/Engine/Plugins/Runtime/RigVM/Source/RigVMEditor/Public/Widgets/SRigVMVariantTagWidget.h
@@ -16,7 +16,7 @@ class SRigVMVariantCapsule : public SButton
 public:
 	
 	SLATE_BEGIN_ARGS(SRigVMVariantCapsule)
-		: _Padding(0)
+		: _Padding(0.f)
 		, _EnableContextMenu(false)
 		, _MinDesiredLabelWidth(0.f)
 	{
diff --git a/Engine/Plugins/Runtime/RigVM/Source/RigVMEditor/RigVMEditor.Build.cs b/Engine/Plugins/Runtime/RigVM/Source/RigVMEditor/RigVMEditor.Build.cs
--- a/Engine/Plugins/Runtime/RigVM/Source/RigVMEditor/RigVMEditor.Build.cs
+++ b/Engine/Plugins/Runtime/RigVM/Source/RigVMEditor/RigVMEditor.Build.cs
@@ -6,6 +6,8 @@ public class RigVMEditor : ModuleRules
 {
     public RigVMEditor(ReadOnlyTargetRules Target) : base(Target)
     {
+	    CppCompileWarningSettings.UnsafeTypeCastWarningLevel = WarningLevel.Error;
+
         PublicDependencyModuleNames.AddRange(
             new string[] {
                 "Core",
