RigVM Editor: Build fix

#rnx

diff --git a/Engine/Plugins/Runtime/RigVM/Source/RigVMEditor/Private/EdGraph/RigVMEdGraphConnectionDrawingPolicy.cpp b/Engine/Plugins/Runtime/RigVM/Source/RigVMEditor/Private/EdGraph/RigVMEdGraphConnectionDrawingPolicy.cpp
--- a/Engine/Plugins/Runtime/RigVM/Source/RigVMEditor/Private/EdGraph/RigVMEdGraphConnectionDrawingPolicy.cpp
+++ b/Engine/Plugins/Runtime/RigVM/Source/RigVMEditor/Private/EdGraph/RigVMEdGraphConnectionDrawingPolicy.cpp
@@ -228,7 +228,7 @@ bool FRigVMEdGraphConnectionDrawingPolicy::GetAverageConnectedPositionForPin(UEd
 
 	if (ResultCount > 0)
 	{
-		OutPos = Result * (1.0 / ResultCount);
+		OutPos = Result * (1.0f / static_cast<float>(ResultCount));
 		return true;
 	}
 	else
