Disable UE_RIGVM_DEBUG_TYPEINDEX in Test configuration as well

diff --git a/Engine/Plugins/Runtime/RigVM/Source/RigVM/Public/RigVMDefines.h b/Engine/Plugins/Runtime/RigVM/Source/RigVM/Public/RigVMDefines.h
--- a/Engine/Plugins/Runtime/RigVM/Source/RigVM/Public/RigVMDefines.h
+++ b/Engine/Plugins/Runtime/RigVM/Source/RigVM/Public/RigVMDefines.h
@@ -26,12 +26,10 @@
 	#define UE_RIGVM_DEBUG_TYPEINDEX 1
 #endif
 
-// in shipping always disable debugging the type index
-#ifdef UE_BUILD_SHIPPING
-	#if UE_BUILD_SHIPPING
-		#undef UE_RIGVM_DEBUG_TYPEINDEX
-		#define UE_RIGVM_DEBUG_TYPEINDEX 0
-	#endif
+// in shipping and test always disable debugging the type index
+#if defined(UE_BUILD_SHIPPING) && UE_BUILD_SHIPPING || defined(UE_BUILD_TEST) && UE_BUILD_TEST
+	#undef UE_RIGVM_DEBUG_TYPEINDEX
+	#define UE_RIGVM_DEBUG_TYPEINDEX 0
 #endif
 
 #ifndef UE_RIGVM_DEBUG_EXECUTION
