RigVM: Remove mention of experimental for tags

#rnx

diff --git a/Engine/Plugins/Runtime/RigVM/Source/RigVM/Public/RigVMSettings.h b/Engine/Plugins/Runtime/RigVM/Source/RigVM/Public/RigVMSettings.h
--- a/Engine/Plugins/Runtime/RigVM/Source/RigVM/Public/RigVMSettings.h
+++ b/Engine/Plugins/Runtime/RigVM/Source/RigVM/Public/RigVMSettings.h
@@ -47,7 +47,7 @@ class RIGVM_API URigVMProjectSettings : public UDeveloperSettings
 
 public:
 
-	UPROPERTY(EditAnywhere, config, Category = Experimental)
+	UPROPERTY(EditAnywhere, config, Category = Variants)
 	TArray<FRigVMTag> VariantTags;
 
 	UFUNCTION(BlueprintPure, Category= Variants)
