RigVM template: try to get more info when failing to resolve

#rnx

diff --git a/Engine/Plugins/Runtime/RigVM/Source/RigVM/Private/RigVMCore/RigVMTemplate.cpp b/Engine/Plugins/Runtime/RigVM/Source/RigVM/Private/RigVMCore/RigVMTemplate.cpp
--- a/Engine/Plugins/Runtime/RigVM/Source/RigVM/Private/RigVMCore/RigVMTemplate.cpp
+++ b/Engine/Plugins/Runtime/RigVM/Source/RigVM/Private/RigVMCore/RigVMTemplate.cpp
@@ -27,7 +27,15 @@ FRigVMTemplateArgumentType::FRigVMTemplateArgumentType(const FName& InCPPType, U
 #if WITH_EDITOR
 	if (CPPType.IsNone())
 	{
-		UE_LOG(LogRigVM, Warning, TEXT("FRigVMTemplateArgumentType(): Input CPPType '%s' could not be resolved."), *InCPPTypeString);
+		if (const UClass* ObjectClass = InCPPTypeObject ? InCPPTypeObject->GetClass() : nullptr)
+		{
+			UE_LOG(LogRigVM, Warning, TEXT("FRigVMTemplateArgumentType(): Input CPPType '%s' (Input Object '%s') could not be resolved."),
+				*InCPPTypeString, *ObjectClass->GetName());
+		}
+		else
+		{
+			UE_LOG(LogRigVM, Warning, TEXT("FRigVMTemplateArgumentType(): Input CPPType '%s' could not be resolved."), *InCPPTypeString);
+		}
 	}
 #endif
 }
