RigVM: Performance improvement for VM CurrentVMMemory

#rnx

diff --git a/Engine/Plugins/Runtime/RigVM/Source/RigVM/Private/RigVMCore/RigVM.cpp b/Engine/Plugins/Runtime/RigVM/Source/RigVM/Private/RigVMCore/RigVM.cpp
--- a/Engine/Plugins/Runtime/RigVM/Source/RigVM/Private/RigVMCore/RigVM.cpp
+++ b/Engine/Plugins/Runtime/RigVM/Source/RigVM/Private/RigVMCore/RigVM.cpp
@@ -1559,7 +1559,7 @@ ERigVMExecuteResult URigVM::ExecuteVM(FRigVMExtendedExecuteContext& Context, con
 			return Context.CurrentExecuteResult = ERigVMExecuteResult::Failed;
 		}
 
-		Context.CurrentVMMemory = GetInstanceMemory(Context);
+		Context.UpdateInstanceMemory(GetLiteralMemory());
 
 		if (ByteCode.HasPublicContextAssetPath())
 		{
diff --git a/Engine/Plugins/Runtime/RigVM/Source/RigVM/Public/RigVMCore/RigVM.h b/Engine/Plugins/Runtime/RigVM/Source/RigVM/Public/RigVMCore/RigVM.h
--- a/Engine/Plugins/Runtime/RigVM/Source/RigVM/Public/RigVMCore/RigVM.h
+++ b/Engine/Plugins/Runtime/RigVM/Source/RigVM/Public/RigVMCore/RigVM.h
@@ -316,16 +316,6 @@ public:
 	UE_DEPRECATED(5.4, "Please, use GetInstanceMemoryArray with a Context param")
 	TArray<URigVMMemoryStorage*> GetLocalMemoryArray() { return TArray<URigVMMemoryStorage*>(); }
 
-	// Returns the VM Literal memory and context instanced memory.
-	TArray<FRigVMMemoryStorageStruct*> GetInstanceMemory(FRigVMExtendedExecuteContext& Context)
-	{
-		TArray<FRigVMMemoryStorageStruct*> LocalMemory;
-		LocalMemory.Add(GetWorkMemory(Context));
-		LocalMemory.Add(GetLiteralMemory());
-		LocalMemory.Add(GetDebugMemory(Context));
-		return LocalMemory;
-	}
-
 	// The default non mutable reference literal memory (generated by the compiler and copied to the instances as initial state)
 	const FRigVMMemoryStorageStruct& GetDefaultLiteralMemory() const
 	{
diff --git a/Engine/Plugins/Runtime/RigVM/Source/RigVM/Public/RigVMCore/RigVMExecuteContext.h b/Engine/Plugins/Runtime/RigVM/Source/RigVM/Public/RigVMCore/RigVMExecuteContext.h
--- a/Engine/Plugins/Runtime/RigVM/Source/RigVM/Public/RigVMCore/RigVMExecuteContext.h
+++ b/Engine/Plugins/Runtime/RigVM/Source/RigVM/Public/RigVMCore/RigVMExecuteContext.h
@@ -791,6 +791,17 @@ struct RIGVM_API FRigVMExtendedExecuteContext
 		CachedMemoryHandles.Reset();
 		LazyBranchExecuteState.Reset();
 	}
+	
+	void UpdateInstanceMemory(FRigVMMemoryStorageStruct* InLiteralMemory)
+	{
+		if(CurrentVMMemory.IsEmpty())
+		{
+			CurrentVMMemory.SetNumUninitialized(3);
+		}
+		CurrentVMMemory[0] = &WorkMemoryStorage;
+		CurrentVMMemory[1] = InLiteralMemory;
+		CurrentVMMemory[2] = &DebugMemoryStorage;
+	}
 
 	uint32 GetNumExecutions() const
 	{
@@ -836,7 +847,7 @@ struct RIGVM_API FRigVMExtendedExecuteContext
 	FName CurrentEntryName = NAME_None;
 	bool bCurrentlyRunningRootEntry = false;
 
-	TArray<FRigVMMemoryStorageStruct*> CurrentVMMemory;
+	TArray<FRigVMMemoryStorageStruct*, TInlineAllocator<3>> CurrentVMMemory;
 
 #if WITH_EDITORONLY_DATA
 	// changes to the layout of cached memory array should be reflected in GetContainerIndex()
