[RigVM]
* Moved dll storage to methods/variables instead of on types. This is needed to get merged modules working

diff --git a/Engine/Plugins/Runtime/RigVM/Source/RigVM/Public/RigVMCore/RigVMDebugDrawSettings.h b/Engine/Plugins/Runtime/RigVM/Source/RigVM/Public/RigVMCore/RigVMDebugDrawSettings.h
--- a/Engine/Plugins/Runtime/RigVM/Source/RigVM/Public/RigVMCore/RigVMDebugDrawSettings.h
+++ b/Engine/Plugins/Runtime/RigVM/Source/RigVM/Public/RigVMCore/RigVMDebugDrawSettings.h
@@ -6,7 +6,7 @@
 #include "RigVMDebugDrawSettings.generated.h"
 
 USTRUCT()
-struct RIGVM_API FRigVMDebugDrawSettings
+struct FRigVMDebugDrawSettings
 {
 	GENERATED_BODY()
 
diff --git a/Engine/Plugins/Runtime/RigVM/Source/RigVM/Public/RigVMCore/RigVMDrawContainer.h b/Engine/Plugins/Runtime/RigVM/Source/RigVM/Public/RigVMCore/RigVMDrawContainer.h
--- a/Engine/Plugins/Runtime/RigVM/Source/RigVM/Public/RigVMCore/RigVMDrawContainer.h
+++ b/Engine/Plugins/Runtime/RigVM/Source/RigVM/Public/RigVMCore/RigVMDrawContainer.h
@@ -7,13 +7,13 @@
 #include "RigVMDrawContainer.generated.h"
 
 USTRUCT()
-struct RIGVM_API FRigVMDrawContainer
+struct FRigVMDrawContainer
 {
 	GENERATED_BODY()
 	virtual ~FRigVMDrawContainer() {}
 
 	int32 Num() const { return Instructions.Num(); }
-	int32 GetIndex(const FName& InName) const;
+	RIGVM_API int32 GetIndex(const FName& InName) const;
 	const FRigVMDrawInstruction& operator[](int32 InIndex) const { return Instructions[InIndex]; }
 	FRigVMDrawInstruction& operator[](int32 InIndex) { return Instructions[InIndex]; }
 	const FRigVMDrawInstruction& operator[](const FName& InName) const { return Instructions[GetIndex(InName)]; }
@@ -26,7 +26,7 @@ struct RIGVM_API FRigVMDrawContainer
 
 	uint32 GetAllocatedSize() const { return static_cast<uint32>(Instructions.GetAllocatedSize()); }
 
-	void Reset();
+	RIGVM_API void Reset();
 
 	UPROPERTY(EditAnywhere, Category = "DrawContainer")
 	TArray<FRigVMDrawInstruction> Instructions;
diff --git a/Engine/Plugins/Runtime/RigVM/Source/RigVM/Public/RigVMCore/RigVMDrawInstruction.h b/Engine/Plugins/Runtime/RigVM/Source/RigVM/Public/RigVMCore/RigVMDrawInstruction.h
--- a/Engine/Plugins/Runtime/RigVM/Source/RigVM/Public/RigVMCore/RigVMDrawInstruction.h
+++ b/Engine/Plugins/Runtime/RigVM/Source/RigVM/Public/RigVMCore/RigVMDrawInstruction.h
@@ -22,7 +22,7 @@ namespace ERigVMDrawSettings
 }
 
 USTRUCT()
-struct RIGVM_API FRigVMDrawInstruction
+struct FRigVMDrawInstruction
 {
 	GENERATED_BODY()
 
diff --git a/Engine/Plugins/Runtime/RigVM/Source/RigVM/Public/RigVMCore/RigVMExternalVariable.h b/Engine/Plugins/Runtime/RigVM/Source/RigVM/Public/RigVMCore/RigVMExternalVariable.h
--- a/Engine/Plugins/Runtime/RigVM/Source/RigVM/Public/RigVMCore/RigVMExternalVariable.h
+++ b/Engine/Plugins/Runtime/RigVM/Source/RigVM/Public/RigVMCore/RigVMExternalVariable.h
@@ -20,7 +20,7 @@
  * memory into the VM and back out
  */
 USTRUCT(BlueprintType)
-struct RIGVM_API FRigVMExternalVariableDef
+struct FRigVMExternalVariableDef
 {
 	GENERATED_BODY()
 
@@ -66,7 +66,7 @@ struct RIGVM_API FRigVMExternalVariableDef
 };
 
 USTRUCT()
-struct RIGVM_API FRigVMExternalVariable : public FRigVMExternalVariableDef
+struct FRigVMExternalVariable : public FRigVMExternalVariableDef
 {
 	GENERATED_BODY()
 	
@@ -466,7 +466,7 @@ struct RIGVM_API FRigVMExternalVariable : public FRigVMExternalVariableDef
 			&& (bAllowNullPtr || Memory != nullptr);
 	}
 
-	TRigVMTypeIndex GetTypeIndex() const;
+	RIGVM_API TRigVMTypeIndex GetTypeIndex() const;
 
 
 	template<typename VarType>
diff --git a/Engine/Plugins/Runtime/RigVM/Source/RigVM/Public/RigVMCore/RigVMFunction.h b/Engine/Plugins/Runtime/RigVM/Source/RigVM/Public/RigVMCore/RigVMFunction.h
--- a/Engine/Plugins/Runtime/RigVM/Source/RigVM/Public/RigVMCore/RigVMFunction.h
+++ b/Engine/Plugins/Runtime/RigVM/Source/RigVM/Public/RigVMCore/RigVMFunction.h
@@ -60,7 +60,7 @@ enum class ERigVMFunctionArgumentDirection : uint8
 /**
  * The FRigVMFunctionArgument describes an argument necessary for the C++ invocation of the RIGVM_METHOD backed function
  */
-struct RIGVM_API FRigVMFunctionArgument
+struct FRigVMFunctionArgument
 {
 	const TCHAR* Name;
 	const TCHAR* Type;
@@ -99,7 +99,7 @@ typedef TArray<FRigVMFunctionArgument> FRigVMFunctionArgumentArray;
 /**
  * The FRigVMExecuteArgument describes an execute argument.
  */
-struct RIGVM_API FRigVMExecuteArgument
+struct FRigVMExecuteArgument
 {
 	FName Name;
 	ERigVMPinDirection Direction;
@@ -124,7 +124,7 @@ struct RIGVM_API FRigVMExecuteArgument
  * The FRigVMFunction is used to represent a function pointer generated by UHT
  * for a given name. The name might be something like "FMyStruct::MyVirtualMethod"
  */
-struct RIGVM_API FRigVMFunction
+struct FRigVMFunction
 {
 	FString Name;
 	UScriptStruct* Struct;
@@ -164,15 +164,15 @@ struct RIGVM_API FRigVMFunction
 	}
 
 	bool IsValid() const { return !Name.IsEmpty() && FunctionPtr != nullptr; }
-	FString GetName() const;
-	FName GetMethodName() const;
-	FString GetModuleName() const;
-	FString GetModuleRelativeHeaderPath() const;
+	RIGVM_API FString GetName() const;
+	RIGVM_API FName GetMethodName() const;
+	RIGVM_API FString GetModuleName() const;
+	RIGVM_API FString GetModuleRelativeHeaderPath() const;
 	const TArray<FRigVMFunctionArgument>& GetArguments() const { return Arguments; }
-	const TArray<TRigVMTypeIndex>& GetArgumentTypeIndices() const;
-	const FRigVMTemplate* GetTemplate() const;
-	const FRigVMTemplate* GetTemplate_NoLock() const;
-	const UScriptStruct* GetExecuteContextStruct(bool bLockRegistry = true) const;
-	bool SupportsExecuteContextStruct(const UScriptStruct* InExecuteContextStruct) const;
-	const FName& GetArgumentNameForOperandIndex(int32 InOperandIndex, int32 InTotalOperands) const;
+	RIGVM_API const TArray<TRigVMTypeIndex>& GetArgumentTypeIndices() const;
+	RIGVM_API const FRigVMTemplate* GetTemplate() const;
+	RIGVM_API const FRigVMTemplate* GetTemplate_NoLock() const;
+	RIGVM_API const UScriptStruct* GetExecuteContextStruct(bool bLockRegistry = true) const;
+	RIGVM_API bool SupportsExecuteContextStruct(const UScriptStruct* InExecuteContextStruct) const;
+	RIGVM_API const FName& GetArgumentNameForOperandIndex(int32 InOperandIndex, int32 InTotalOperands) const;
 };
diff --git a/Engine/Plugins/Runtime/RigVM/Source/RigVM/Public/RigVMCore/RigVMGraphFunctionDefinition.h b/Engine/Plugins/Runtime/RigVM/Source/RigVM/Public/RigVMCore/RigVMGraphFunctionDefinition.h
--- a/Engine/Plugins/Runtime/RigVM/Source/RigVM/Public/RigVMCore/RigVMGraphFunctionDefinition.h
+++ b/Engine/Plugins/Runtime/RigVM/Source/RigVM/Public/RigVMCore/RigVMGraphFunctionDefinition.h
@@ -500,7 +500,7 @@ protected:
 };
 
 USTRUCT(BlueprintType)
-struct RIGVM_API FRigVMGraphFunctionData
+struct FRigVMGraphFunctionData
 {
 	GENERATED_BODY()
 
@@ -521,7 +521,7 @@ struct RIGVM_API FRigVMGraphFunctionData
 	UPROPERTY()
 	FRigVMObjectArchive CollapseNodeArchive;
 
-	bool IsMutable() const;
+	RIGVM_API bool IsMutable() const;
 
 	bool operator==(const FRigVMGraphFunctionData& Other) const
 	{
@@ -532,21 +532,21 @@ struct RIGVM_API FRigVMGraphFunctionData
 
 	RIGVM_API friend FArchive& operator<<(FArchive& Ar, FRigVMGraphFunctionData& Data);
 
-	static FRigVMGraphFunctionData* FindFunctionData(const FSoftObjectPath& InHostObjectPath, const FName& InFunctionName, bool* bOutIsPublic = nullptr, FString* OutErrorMessage = nullptr);	
+	static RIGVM_API FRigVMGraphFunctionData* FindFunctionData(const FSoftObjectPath& InHostObjectPath, const FName& InFunctionName, bool* bOutIsPublic = nullptr, FString* OutErrorMessage = nullptr);	
 
-	static FRigVMGraphFunctionData* FindFunctionData(const FRigVMGraphFunctionIdentifier& InIdentifier, bool* bOutIsPublic = nullptr, FString* OutErrorMessage = nullptr);	
+	static RIGVM_API FRigVMGraphFunctionData* FindFunctionData(const FRigVMGraphFunctionIdentifier& InIdentifier, bool* bOutIsPublic = nullptr, FString* OutErrorMessage = nullptr);	
 
-	static FString GetArgumentNameFromPinHash(const FString& InPinHash);
+	static RIGVM_API FString GetArgumentNameFromPinHash(const FString& InPinHash);
 	
-	FRigVMOperand GetOperandForArgument(const FName& InArgumentName) const;
+	RIGVM_API FRigVMOperand GetOperandForArgument(const FName& InArgumentName) const;
 
-	bool IsAnyOperandSharedAcrossArguments() const;
+	RIGVM_API bool IsAnyOperandSharedAcrossArguments() const;
 
-	bool PatchSharedArgumentOperandsIfRequired();
+	RIGVM_API bool PatchSharedArgumentOperandsIfRequired();
 
 	static const inline TCHAR* EntryString = TEXT("Entry");
 	static const inline TCHAR* ReturnString = TEXT("Return");
-	static TFunction<IRigVMGraphFunctionHost*(UObject*)> GetFunctionHostFromObjectFunc;
+	static RIGVM_API TFunction<IRigVMGraphFunctionHost*(UObject*)> GetFunctionHostFromObjectFunc;
 
 	friend class URigVMBuildData;
 };
diff --git a/Engine/Plugins/Runtime/RigVM/Source/RigVM/Public/RigVMCore/RigVMMemoryCommon.h b/Engine/Plugins/Runtime/RigVM/Source/RigVM/Public/RigVMCore/RigVMMemoryCommon.h
--- a/Engine/Plugins/Runtime/RigVM/Source/RigVM/Public/RigVMCore/RigVMMemoryCommon.h
+++ b/Engine/Plugins/Runtime/RigVM/Source/RigVM/Public/RigVMCore/RigVMMemoryCommon.h
@@ -64,7 +64,7 @@ namespace RigVM
 // A runtime cache for determining if a set of instruction has to
 // run for this execution of the VM
 USTRUCT(BlueprintType)
-struct RIGVM_API FRigVMInstructionSetExecuteState
+struct FRigVMInstructionSetExecuteState
 {
 	GENERATED_BODY()
 
@@ -87,7 +87,7 @@ struct RIGVM_API FRigVMInstructionSetExecuteState
  * index of the external variable within the running VM.
  */
 USTRUCT()
-struct RIGVM_API FRigVMOperand
+struct FRigVMOperand
 {
 	GENERATED_BODY()
 
@@ -152,7 +152,7 @@ public:
   		return HashCombine(HashCombine((uint32)Operand.MemoryType, (uint32)Operand.RegisterIndex), (uint32)Operand.RegisterOffset);
   	}
 	
-	void Serialize(FArchive& Ar);
+	RIGVM_API void Serialize(FArchive& Ar);
 	friend FArchive& operator<<(FArchive& Ar, FRigVMOperand& P)
 	{
 		P.Serialize(Ar);
@@ -188,7 +188,7 @@ typedef TArrayView<const FRigVMOperand> FRigVMOperandArray;
 /**
  * Helper class to catch default value import errors on properties
  */
-class RIGVM_API FRigVMMemoryStorageImportErrorContext : public FOutputDevice
+class FRigVMMemoryStorageImportErrorContext : public FOutputDevice
 {
 public:
 
@@ -202,5 +202,5 @@ public:
 	{
 	}
 
-	void Serialize(const TCHAR* V, ELogVerbosity::Type Verbosity, const class FName& Category) override;
+	RIGVM_API void Serialize(const TCHAR* V, ELogVerbosity::Type Verbosity, const class FName& Category) override;
 };
diff --git a/Engine/Plugins/Runtime/RigVM/Source/RigVM/Public/RigVMCore/RigVMNameCache.h b/Engine/Plugins/Runtime/RigVM/Source/RigVM/Public/RigVMCore/RigVMNameCache.h
--- a/Engine/Plugins/Runtime/RigVM/Source/RigVM/Public/RigVMCore/RigVMNameCache.h
+++ b/Engine/Plugins/Runtime/RigVM/Source/RigVM/Public/RigVMCore/RigVMNameCache.h
@@ -22,7 +22,7 @@ enum class ERigVMNameOp : uint8
     ContainsNoCase
 };
 
-struct RIGVM_API FRigVMNameOp
+struct FRigVMNameOp
 {
 public:
 
@@ -38,15 +38,15 @@ public:
 		return HashCombine(InName.GetComparisonIndex().ToUnstableInt(), uint32(InName.GetNumber()));
 	}
 	
-	static FRigVMNameOp Concat(const FName& InA, const FName& InB);
-	static FRigVMNameOp Left(const FName& InA, const uint32 InCount);
-	static FRigVMNameOp Right(const FName& InA, const uint32 InCount);
-	static FRigVMNameOp LeftChop(const FName& InA, const uint32 InCount);
-	static FRigVMNameOp RightChop(const FName& InA, const uint32 InCount);
-	static FRigVMNameOp Replace(const FName& InA, const FName& InB, const FName& InC, const ESearchCase::Type InSearchCase);
-    static FRigVMNameOp EndsWith(const FName& InA, const FName& InB, const ESearchCase::Type InSearchCase);
-    static FRigVMNameOp StartsWith(const FName& InA, const FName& InB, const ESearchCase::Type InSearchCase);
-    static FRigVMNameOp Contains(const FName& InA, const FName& InB, const ESearchCase::Type InSearchCase);
+	static RIGVM_API FRigVMNameOp Concat(const FName& InA, const FName& InB);
+	static RIGVM_API FRigVMNameOp Left(const FName& InA, const uint32 InCount);
+	static RIGVM_API FRigVMNameOp Right(const FName& InA, const uint32 InCount);
+	static RIGVM_API FRigVMNameOp LeftChop(const FName& InA, const uint32 InCount);
+	static RIGVM_API FRigVMNameOp RightChop(const FName& InA, const uint32 InCount);
+	static RIGVM_API FRigVMNameOp Replace(const FName& InA, const FName& InB, const FName& InC, const ESearchCase::Type InSearchCase);
+    static RIGVM_API FRigVMNameOp EndsWith(const FName& InA, const FName& InB, const ESearchCase::Type InSearchCase);
+    static RIGVM_API FRigVMNameOp StartsWith(const FName& InA, const FName& InB, const ESearchCase::Type InSearchCase);
+    static RIGVM_API FRigVMNameOp Contains(const FName& InA, const FName& InB, const ESearchCase::Type InSearchCase);
 
 	friend uint32 GetTypeHash(const FRigVMNameOp& Op)
 	{
@@ -106,7 +106,7 @@ private:
 	ERigVMNameOp Type;
 };
 
-class RIGVM_API FRigVMNameCache
+class FRigVMNameCache
 {
 public:
 
@@ -116,22 +116,22 @@ public:
 		BoolCache.Reset();
 	}
 
-	FName Concat(const FName& InA, const FName& InB);
-	FName Left(const FName& InA, const uint32 InCount);
-	FName Right(const FName& InA, const uint32 InCount);
-	FName LeftChop(const FName& InA, const uint32 InCount);
-	FName RightChop(const FName& InA, const uint32 InCount);
-	FName Replace(const FName& InA, const FName& InB, const FName& InC, const ESearchCase::Type InSearchCase);
-	bool EndsWith(const FName& InA, const FName& InB, const ESearchCase::Type InSearchCase);
-	bool StartsWith(const FName& InA, const FName& InB, const ESearchCase::Type InSearchCase);
-	bool Contains(const FName& InA, const FName& InB, const ESearchCase::Type InSearchCase);
+	RIGVM_API FName Concat(const FName& InA, const FName& InB);
+	RIGVM_API FName Left(const FName& InA, const uint32 InCount);
+	RIGVM_API FName Right(const FName& InA, const uint32 InCount);
+	RIGVM_API FName LeftChop(const FName& InA, const uint32 InCount);
+	RIGVM_API FName RightChop(const FName& InA, const uint32 InCount);
+	RIGVM_API FName Replace(const FName& InA, const FName& InB, const FName& InC, const ESearchCase::Type InSearchCase);
+	RIGVM_API bool EndsWith(const FName& InA, const FName& InB, const ESearchCase::Type InSearchCase);
+	RIGVM_API bool StartsWith(const FName& InA, const FName& InB, const ESearchCase::Type InSearchCase);
+	RIGVM_API bool Contains(const FName& InA, const FName& InB, const ESearchCase::Type InSearchCase);
 
 	const TMap<FRigVMNameOp, FName>& GetNameCache() const { return NameCache; }
 	const TMap<FRigVMNameOp, bool>&  GetBoolCache() const { return BoolCache; }
-	TArray<FRigVMNameOp> GetNameOps() const;
-	TArray<FName> GetNameValues() const;
-	TArray<FRigVMNameOp> GetBoolOps() const;
-	TArray<bool> GetBoolValues() const;
+	RIGVM_API TArray<FRigVMNameOp> GetNameOps() const;
+	RIGVM_API TArray<FName> GetNameValues() const;
+	RIGVM_API TArray<FRigVMNameOp> GetBoolOps() const;
+	RIGVM_API TArray<bool> GetBoolValues() const;
 
 private:
 
diff --git a/Engine/Plugins/Runtime/RigVM/Source/RigVM/Public/RigVMCore/RigVMObjectArchive.h b/Engine/Plugins/Runtime/RigVM/Source/RigVM/Public/RigVMCore/RigVMObjectArchive.h
--- a/Engine/Plugins/Runtime/RigVM/Source/RigVM/Public/RigVMCore/RigVMObjectArchive.h
+++ b/Engine/Plugins/Runtime/RigVM/Source/RigVM/Public/RigVMCore/RigVMObjectArchive.h
@@ -5,8 +5,10 @@
 #include "Serialization/ArchiveUObject.h"
 #include "RigVMObjectArchive.generated.h"
 
+#define UE_API RIGVM_API
+
 USTRUCT()
-struct RIGVM_API FRigVMObjectArchive
+struct FRigVMObjectArchive
 {
 public:
 	GENERATED_BODY()
@@ -18,11 +20,11 @@ public:
 	{
 	}
 	
-	void Reset();
-	void Empty();
-	bool IsEmpty() const;
-	void Compress();
-	void Decompress();
+	UE_API void Reset();
+	UE_API void Empty();
+	UE_API bool IsEmpty() const;
+	UE_API void Compress();
+	UE_API void Decompress();
 	bool IsCompressed() const
 	{
 		return bIsCompressed;
@@ -48,18 +50,18 @@ private:
 	friend class FRigVMObjectArchiveReader;
 };
 
-class RIGVM_API FRigVMObjectArchiveWriter : public FArchiveUObject
+class FRigVMObjectArchiveWriter : public FArchiveUObject
 {
 public:
-	FRigVMObjectArchiveWriter( FRigVMObjectArchive& InOutArchive, const UObject* InRoot );
-	virtual void Serialize( void* V, int64 Length ) override;
-	virtual int64 Tell() override;
-	virtual int64 TotalSize() override;
-	virtual void Seek(int64 InPos) override;
+	UE_API FRigVMObjectArchiveWriter( FRigVMObjectArchive& InOutArchive, const UObject* InRoot );
+	UE_API virtual void Serialize( void* V, int64 Length ) override;
+	UE_API virtual int64 Tell() override;
+	UE_API virtual int64 TotalSize() override;
+	UE_API virtual void Seek(int64 InPos) override;
 	using FArchiveUObject::operator<<; // For visibility of the overloads we don't override
-	virtual FArchive& operator<<(UObject*& Obj) override;
-	virtual FArchive& operator<<(FName& Value) override;
-	virtual FArchive& operator<<(FText& Value) override;
+	UE_API virtual FArchive& operator<<(UObject*& Obj) override;
+	UE_API virtual FArchive& operator<<(FName& Value) override;
+	UE_API virtual FArchive& operator<<(FText& Value) override;
 
 protected:
 	
@@ -77,16 +79,16 @@ protected:
 	static inline constexpr uint8 StoringNameAsOffset = 1;
 };
 
-class RIGVM_API FRigVMObjectArchiveReader : public FRigVMObjectArchiveWriter
+class FRigVMObjectArchiveReader : public FRigVMObjectArchiveWriter
 {
 public:
 	
-	FRigVMObjectArchiveReader( FRigVMObjectArchive& InOutArchive, const UObject* InRoot );
-	virtual void Serialize( void* V, int64 Length ) override;
+	UE_API FRigVMObjectArchiveReader( FRigVMObjectArchive& InOutArchive, const UObject* InRoot );
+	UE_API virtual void Serialize( void* V, int64 Length ) override;
 	using FArchiveUObject::operator<<; // For visibility of the overloads we don't override
-	virtual FArchive& operator<<(UObject*& Obj) override;
-	virtual FArchive& operator<<(FName& Value) override;
-	virtual FArchive& operator<<(FText& Value) override;
+	UE_API virtual FArchive& operator<<(UObject*& Obj) override;
+	UE_API virtual FArchive& operator<<(FName& Value) override;
+	UE_API virtual FArchive& operator<<(FText& Value) override;
 
 	struct FObjectHeader
 	{
@@ -108,8 +110,8 @@ public:
 		FName OuterPathName;
 	};
 
-	FObjectHeader ReadObjectHeader(int64 InPosition = INDEX_NONE);
-	FObjectHeader GetRootObjectHeader();
+	UE_API FObjectHeader ReadObjectHeader(int64 InPosition = INDEX_NONE);
+	UE_API FObjectHeader GetRootObjectHeader();
 
 protected:
 
@@ -117,4 +119,6 @@ protected:
 	TMap<FName, UObject*> ReadObjects;
 	TArray<UObject*> DeserializedObjects;
 	TMap<int64,FName> OffsetToName;
-};
\ No newline at end of file
+};
+
+#undef UE_API
\ No newline at end of file
diff --git a/Engine/Plugins/Runtime/RigVM/Source/RigVM/Public/RigVMCore/RigVMRegistry.h b/Engine/Plugins/Runtime/RigVM/Source/RigVM/Public/RigVMCore/RigVMRegistry.h
--- a/Engine/Plugins/Runtime/RigVM/Source/RigVM/Public/RigVMCore/RigVMRegistry.h
+++ b/Engine/Plugins/Runtime/RigVM/Source/RigVM/Public/RigVMCore/RigVMRegistry.h
@@ -437,13 +437,13 @@ protected:
  * 
  * Inheriting from FGCObject to ensure that all type objects cannot be GCed
  */
-struct RIGVM_API FRigVMRegistry_RWLock final : public FRigVMRegistry_NoLock
+struct FRigVMRegistry_RWLock final : public FRigVMRegistry_NoLock
 {
 	typedef FRigVMRegistry_NoLock Super;
 	
 protected:
 
-	class RIGVM_API FConditionalScopeLock
+	class FConditionalScopeLock
 	{
 	public:
 		UE_NODISCARD_CTOR explicit FConditionalScopeLock(const FRigVMRegistry_RWLock& InRegistry, ELockType InLockType, bool bInLockEnabled = true)
@@ -527,7 +527,7 @@ protected:
 
 public:
 
-	class RIGVM_API FConditionalReadScopeLock : public FConditionalScopeLock 
+	class FConditionalReadScopeLock : public FConditionalScopeLock 
 	{
 	public:
 		UE_NODISCARD_CTOR explicit FConditionalReadScopeLock(const FRigVMRegistry_RWLock& InRegistry, bool bInLockEnabled = true)
@@ -541,7 +541,7 @@ public:
 		}
 	};
 
-	class RIGVM_API FConditionalWriteScopeLock : public FConditionalScopeLock
+	class FConditionalWriteScopeLock : public FConditionalScopeLock
 	{
 	public:
 		UE_NODISCARD_CTOR explicit FConditionalWriteScopeLock(const FRigVMRegistry_RWLock& InRegistry, bool bInLockEnabled = true)
@@ -556,7 +556,7 @@ public:
 	};
 
 	// Returns the singleton registry
-	static FRigVMRegistry_RWLock& Get();
+	static RIGVM_API FRigVMRegistry_RWLock& Get();
 
 	static FRigVMRegistry_NoLock& Get(ELockType InLockType) = delete;
 	static const FRigVMRegistry_NoLock& GetForRead() = delete;
@@ -630,17 +630,17 @@ public:
 	}
 	
 	// Update the registry when old types are removed
-    void OnAssetRemoved(const FAssetData& InAssetData);
+    RIGVM_API void OnAssetRemoved(const FAssetData& InAssetData);
 
 	// May add factories and unit functions declared in the plugin 
-	void OnPluginLoaded(IPlugin& InPlugin);
+	RIGVM_API void OnPluginLoaded(IPlugin& InPlugin);
 
 	// Removes all types associated with a plugin that's being unloaded. 
-	void OnPluginUnloaded(IPlugin& InPlugin);
+	RIGVM_API void OnPluginUnloaded(IPlugin& InPlugin);
 
 	// Update the registry when new types are added to the attribute system so that they can be selected
 	// on Attribute Nodes
-	void OnAnimationAttributeTypesChanged(const UScriptStruct* InStruct, bool bIsAdded);
+	RIGVM_API void OnAnimationAttributeTypesChanged(const UScriptStruct* InStruct, bool bIsAdded);
 
 	// Notifies other system that types have been added/removed, and template permutations have been updated
 	FOnRigVMRegistryChanged& OnRigVMRegistryChanged() { return OnRigVMRegistryChangedDelegate; }
diff --git a/Engine/Plugins/Runtime/RigVM/Source/RigVM/Public/RigVMCore/RigVMTemplate.h b/Engine/Plugins/Runtime/RigVM/Source/RigVM/Public/RigVMCore/RigVMTemplate.h
--- a/Engine/Plugins/Runtime/RigVM/Source/RigVM/Public/RigVMCore/RigVMTemplate.h
+++ b/Engine/Plugins/Runtime/RigVM/Source/RigVM/Public/RigVMCore/RigVMTemplate.h
@@ -37,14 +37,14 @@ typedef TMap<FName, TRigVMTypeIndex> FRigVMTemplateTypeMap;
 DECLARE_DELEGATE_RetVal_TwoParams(FRigVMTemplateTypeMap, FRigVMTemplate_NewArgumentTypeDelegate, const FName& /* InArgumentName */, TRigVMTypeIndex /* InTypeIndexToAdd */);
 DECLARE_DELEGATE_RetVal(FRigVMDispatchFactory*, FRigVMTemplate_GetDispatchFactoryDelegate);
 
-struct RIGVM_API FRigVMTemplateDelegates
+struct FRigVMTemplateDelegates
 {
 	FRigVMTemplate_NewArgumentTypeDelegate NewArgumentTypeDelegate;
 	FRigVMTemplate_GetDispatchFactoryDelegate GetDispatchFactoryDelegate;
 };
 
 USTRUCT()
-struct RIGVM_API FRigVMTemplateArgumentType
+struct FRigVMTemplateArgumentType
 {
 	GENERATED_BODY()
 
@@ -62,7 +62,7 @@ struct RIGVM_API FRigVMTemplateArgumentType
 		CPPTypeObject = RigVMTypeUtils::GetWildCardCPPTypeObject();
 	}
 
-	FRigVMTemplateArgumentType(const FName& InCPPType, UObject* InCPPTypeObject = nullptr);
+	RIGVM_API FRigVMTemplateArgumentType(const FName& InCPPType, UObject* InCPPTypeObject = nullptr);
 	
 	FRigVMTemplateArgumentType(UClass* InClass, RigVMTypeUtils::EClassArgType InClassArgType = RigVMTypeUtils::EClassArgType::AsObject)
 	: CPPType(*RigVMTypeUtils::CPPTypeFromObject(InClass, InClassArgType))
@@ -102,7 +102,7 @@ struct RIGVM_API FRigVMTemplateArgumentType
 		return CPPType != InOther.CPPType;
 	}
 
-	RIGVM_API friend FORCEINLINE uint32 GetTypeHash(const FRigVMTemplateArgumentType& InType)
+	friend FORCEINLINE uint32 GetTypeHash(const FRigVMTemplateArgumentType& InType)
 	{
 		return GetTypeHash(InType.CPPType.ToString());
 	}
@@ -154,7 +154,7 @@ struct RIGVM_API FRigVMTemplateArgumentType
  * The template argument represents a single parameter
  * in a function call and all of its possible types
  */
-struct RIGVM_API FRigVMTemplateArgument
+struct FRigVMTemplateArgument
 {
 	DECLARE_DELEGATE_RetVal_OneParam(bool, FTypeFilter, const TRigVMTypeIndex&);
 	
@@ -192,12 +192,12 @@ struct RIGVM_API FRigVMTemplateArgument
 	};
 
 	// default constructor
-	FRigVMTemplateArgument();
+	RIGVM_API FRigVMTemplateArgument();
 
-	FRigVMTemplateArgument(const FName& InName, ERigVMPinDirection InDirection);
-	FRigVMTemplateArgument(const FName& InName, ERigVMPinDirection InDirection, TRigVMTypeIndex InType);
-	FRigVMTemplateArgument(const FName& InName, ERigVMPinDirection InDirection, const TArray<TRigVMTypeIndex>& InTypeIndices);
-	FRigVMTemplateArgument(const FName& InName, ERigVMPinDirection InDirection, const TArray<ETypeCategory>& InTypeCategories, TFunction<bool(const TRigVMTypeIndex&)> InFilterType = nullptr);
+	RIGVM_API FRigVMTemplateArgument(const FName& InName, ERigVMPinDirection InDirection);
+	RIGVM_API FRigVMTemplateArgument(const FName& InName, ERigVMPinDirection InDirection, TRigVMTypeIndex InType);
+	RIGVM_API FRigVMTemplateArgument(const FName& InName, ERigVMPinDirection InDirection, const TArray<TRigVMTypeIndex>& InTypeIndices);
+	RIGVM_API FRigVMTemplateArgument(const FName& InName, ERigVMPinDirection InDirection, const TArray<ETypeCategory>& InTypeCategories, TFunction<bool(const TRigVMTypeIndex&)> InFilterType = nullptr);
 
 	// returns the name of the argument
 	const FName& GetName() const { return Name; }
@@ -207,24 +207,24 @@ struct RIGVM_API FRigVMTemplateArgument
 
 #if WITH_EDITOR
 	// returns true if this argument supports a given type across a set of permutations
-	bool SupportsTypeIndex(TRigVMTypeIndex InTypeIndex, TRigVMTypeIndex* OutTypeIndex = nullptr, const bool bLockRegistry = true) const;
-	bool SupportsTypeIndex_NoLock(TRigVMTypeIndex InTypeIndex, TRigVMTypeIndex* OutTypeIndex = nullptr) const;
+	RIGVM_API bool SupportsTypeIndex(TRigVMTypeIndex InTypeIndex, TRigVMTypeIndex* OutTypeIndex = nullptr, const bool bLockRegistry = true) const;
+	RIGVM_API bool SupportsTypeIndex_NoLock(TRigVMTypeIndex InTypeIndex, TRigVMTypeIndex* OutTypeIndex = nullptr) const;
 #endif
 
 	// returns the flat list of types (including duplicates) of this argument
-	void GetAllTypes(TArray<TRigVMTypeIndex>& OutTypes, const bool bLockRegistry = true) const;
-	void GetAllTypes_NoLock(TArray<TRigVMTypeIndex>& OutTypes) const;
-
-	TRigVMTypeIndex GetTypeIndex(const int32 InIndex, const bool bLockRegistry = true) const;
-	TRigVMTypeIndex GetTypeIndex_NoLock(const int32 InIndex) const;
-	TOptional<TRigVMTypeIndex> TryToGetTypeIndex(const int32 InIndex, const bool bLockRegistry = true) const;
-	TOptional<TRigVMTypeIndex> TryToGetTypeIndex_NoLock(const int32 InIndex) const;
-	int32 GetNumTypes() const;
-	int32 GetNumTypes_NoLock() const;
-	void AddTypeIndex(const TRigVMTypeIndex InTypeIndex);
-	void RemoveType(const int32 InIndex);
-	void ForEachType(TFunction<bool(const TRigVMTypeIndex InType)>&& InCallback) const;
-	int32 FindTypeIndex(const TRigVMTypeIndex InTypeIndex) const;
+	RIGVM_API void GetAllTypes(TArray<TRigVMTypeIndex>& OutTypes, const bool bLockRegistry = true) const;
+	RIGVM_API void GetAllTypes_NoLock(TArray<TRigVMTypeIndex>& OutTypes) const;
+
+	RIGVM_API TRigVMTypeIndex GetTypeIndex(const int32 InIndex, const bool bLockRegistry = true) const;
+	RIGVM_API TRigVMTypeIndex GetTypeIndex_NoLock(const int32 InIndex) const;
+	RIGVM_API TOptional<TRigVMTypeIndex> TryToGetTypeIndex(const int32 InIndex, const bool bLockRegistry = true) const;
+	RIGVM_API TOptional<TRigVMTypeIndex> TryToGetTypeIndex_NoLock(const int32 InIndex) const;
+	RIGVM_API int32 GetNumTypes() const;
+	RIGVM_API int32 GetNumTypes_NoLock() const;
+	RIGVM_API void AddTypeIndex(const TRigVMTypeIndex InTypeIndex);
+	RIGVM_API void RemoveType(const int32 InIndex);
+	RIGVM_API void ForEachType(TFunction<bool(const TRigVMTypeIndex InType)>&& InCallback) const;
+	RIGVM_API int32 FindTypeIndex(const TRigVMTypeIndex InTypeIndex) const;
 
 	template <typename Predicate>
 	int32 IndexOfByPredicate(Predicate Pred) const
@@ -260,30 +260,30 @@ struct RIGVM_API FRigVMTemplateArgument
 	}
 
 	// returns an array of all of the supported types
-	TArray<TRigVMTypeIndex> GetSupportedTypeIndices(const TArray<int32>& InPermutationIndices = TArray<int32>()) const;
+	RIGVM_API TArray<TRigVMTypeIndex> GetSupportedTypeIndices(const TArray<int32>& InPermutationIndices = TArray<int32>()) const;
 
 	// returns an array of all supported types as strings. this is used for automated testing only.
-	TArray<FString> GetSupportedTypeStrings(const TArray<int32>& InPermutationIndices = TArray<int32>()) const;
+	RIGVM_API TArray<FString> GetSupportedTypeStrings(const TArray<int32>& InPermutationIndices = TArray<int32>()) const;
 	
 	// returns true if an argument is singleton (same type for all variants)
-	bool IsSingleton(const TArray<int32>& InPermutationIndices = TArray<int32>(), const bool bLockRegistry = true) const;
-	bool IsSingleton_NoLock(const TArray<int32>& InPermutationIndices = TArray<int32>()) const;
+	RIGVM_API bool IsSingleton(const TArray<int32>& InPermutationIndices = TArray<int32>(), const bool bLockRegistry = true) const;
+	RIGVM_API bool IsSingleton_NoLock(const TArray<int32>& InPermutationIndices = TArray<int32>()) const;
 
 	// returns true if this argument is an execute
-	bool IsExecute() const;
-	bool IsExecute_NoLock(const FRigVMRegistry_NoLock& InRegistry) const;
+	RIGVM_API bool IsExecute() const;
+	RIGVM_API bool IsExecute_NoLock(const FRigVMRegistry_NoLock& InRegistry) const;
 	
 	// returns true if the argument uses an array container
-	EArrayType GetArrayType(const bool bLockRegistry = true) const;
-	EArrayType GetArrayType_NoLock() const;
+	RIGVM_API EArrayType GetArrayType(const bool bLockRegistry = true) const;
+	RIGVM_API EArrayType GetArrayType_NoLock() const;
 
 	RIGVM_API friend uint32 GetTypeHash_NoLock(const FRigVMTemplateArgument& InArgument);
 
 #if WITH_EDITOR
 	// Get the map of types to permutation indices
-	const TArray<int32>& GetPermutations(const TRigVMTypeIndex InType, bool bLockRegistry = true) const;
-	const TArray<int32>& GetPermutations_NoLock(const TRigVMTypeIndex InType) const;
-	void InvalidatePermutations(const TRigVMTypeIndex InType);
+	RIGVM_API const TArray<int32>& GetPermutations(const TRigVMTypeIndex InType, bool bLockRegistry = true) const;
+	RIGVM_API const TArray<int32>& GetPermutations_NoLock(const TRigVMTypeIndex InType) const;
+	RIGVM_API void InvalidatePermutations(const TRigVMTypeIndex InType);
 #endif
 
 protected:
@@ -303,19 +303,19 @@ protected:
 	mutable TOptional<EArrayType> CachedArrayType;
 
 	// constructor from a property. this forces the type to be created
-	FRigVMTemplateArgument(FProperty* InProperty, FRigVMRegistry_NoLock& InRegistry);
+	RIGVM_API FRigVMTemplateArgument(FProperty* InProperty, FRigVMRegistry_NoLock& InRegistry);
 
 	// static make function using a lock to create it
-	static FRigVMTemplateArgument Make(FProperty* InProperty);
+	static RIGVM_API FRigVMTemplateArgument Make(FProperty* InProperty);
 
 	// static make function using a lock to create it
-	static FRigVMTemplateArgument Make_NoLock(FProperty* InProperty);
-	static FRigVMTemplateArgument Make_NoLock(FProperty* InProperty, FRigVMRegistry_NoLock& InRegistry);
+	static RIGVM_API FRigVMTemplateArgument Make_NoLock(FProperty* InProperty);
+	static RIGVM_API FRigVMTemplateArgument Make_NoLock(FProperty* InProperty, FRigVMRegistry_NoLock& InRegistry);
 
-	void EnsureValidExecuteType_NoLock(FRigVMRegistry_NoLock& InRegistry);
+	RIGVM_API void EnsureValidExecuteType_NoLock(FRigVMRegistry_NoLock& InRegistry);
 
 #if WITH_EDITOR
-	void UpdateTypeToPermutationsSlow();
+	RIGVM_API void UpdateTypeToPermutationsSlow();
 #endif
 
 	friend struct FRigVMTemplate;
@@ -363,7 +363,7 @@ private:
  * FRigVMTemplateArgumentInfo 
  */
 
-struct RIGVM_API FRigVMTemplateArgumentInfo
+struct FRigVMTemplateArgumentInfo
 {
 	using ArgumentCallback = TFunction<FRigVMTemplateArgument(const FName /*InName*/, ERigVMPinDirection /*InDirection*/)>;
 	using TypeFilterCallback = TFunction<bool(const TRigVMTypeIndex& InNewType)>;
@@ -372,19 +372,19 @@ struct RIGVM_API FRigVMTemplateArgumentInfo
 	ERigVMPinDirection Direction = ERigVMPinDirection::Invalid;
 	TFunction<FRigVMTemplateArgument(const FName /*InName*/, ERigVMPinDirection /*InDirection*/)> FactoryCallback = [](const FName, ERigVMPinDirection) { return FRigVMTemplateArgument(); };
 	
-	FRigVMTemplateArgumentInfo(const FName InName, ERigVMPinDirection InDirection, TRigVMTypeIndex InTypeIndex);
-	FRigVMTemplateArgumentInfo(const FName InName, ERigVMPinDirection InDirection, const TArray<TRigVMTypeIndex>& InTypeIndices);
-	FRigVMTemplateArgumentInfo(const FName InName, ERigVMPinDirection InDirection, const TArray<FRigVMTemplateArgument::ETypeCategory>& InTypeCategories, TypeFilterCallback InTypeFilter = nullptr);
-	FRigVMTemplateArgumentInfo(const FName InName, ERigVMPinDirection InDirection);
-	FRigVMTemplateArgumentInfo(const FName InName, ERigVMPinDirection InDirection, ArgumentCallback&& InCallback);
+	RIGVM_API FRigVMTemplateArgumentInfo(const FName InName, ERigVMPinDirection InDirection, TRigVMTypeIndex InTypeIndex);
+	RIGVM_API FRigVMTemplateArgumentInfo(const FName InName, ERigVMPinDirection InDirection, const TArray<TRigVMTypeIndex>& InTypeIndices);
+	RIGVM_API FRigVMTemplateArgumentInfo(const FName InName, ERigVMPinDirection InDirection, const TArray<FRigVMTemplateArgument::ETypeCategory>& InTypeCategories, TypeFilterCallback InTypeFilter = nullptr);
+	RIGVM_API FRigVMTemplateArgumentInfo(const FName InName, ERigVMPinDirection InDirection);
+	RIGVM_API FRigVMTemplateArgumentInfo(const FName InName, ERigVMPinDirection InDirection, ArgumentCallback&& InCallback);
 
-	FRigVMTemplateArgument GetArgument() const;
+	RIGVM_API FRigVMTemplateArgument GetArgument() const;
 	
-	static FName ComputeTemplateNotation(
+	static RIGVM_API FName ComputeTemplateNotation(
 		const FName InTemplateName,
 		const TArray<FRigVMTemplateArgumentInfo>& InInfos);
 	
-	static TArray<TRigVMTypeIndex> GetTypesFromCategories(
+	static RIGVM_API TArray<TRigVMTypeIndex> GetTypesFromCategories(
 		const TArray<FRigVMTemplateArgument::ETypeCategory>& InTypeCategories,
 		const FRigVMTemplateArgument::FTypeFilter& InTypeFilter = {});	
 };
@@ -395,7 +395,7 @@ struct RIGVM_API FRigVMTemplateArgumentInfo
  * to build polymorphic nodes (RigVMTemplateNode) that can
  * take on any of the permutations supported by the template.
  */
-struct RIGVM_API FRigVMTemplate
+struct FRigVMTemplate
 {
 public:
 
@@ -403,151 +403,151 @@ public:
 	typedef TPair<FName, TRigVMTypeIndex> FTypePair;
 
 	// Default constructor
-	FRigVMTemplate();
+	RIGVM_API FRigVMTemplate();
 
 	// returns true if this is a valid template
-	bool IsValid() const;
+	RIGVM_API bool IsValid() const;
 
 	// Returns the notation of this template
-	const FName& GetNotation() const;
+	RIGVM_API const FName& GetNotation() const;
 
 	// Returns the name of the template
-	FName GetName() const;
+	RIGVM_API FName GetName() const;
 
 	// Returns the name to use for a node
-	FName GetNodeName() const;
+	RIGVM_API FName GetNodeName() const;
 
 	// returns true if this template can merge another one
-	bool Merge(const FRigVMTemplate& InOther);
+	RIGVM_API bool Merge(const FRigVMTemplate& InOther);
 
 	// returns the number of args of this template
-	int32 NumArguments() const { return Arguments.Num(); }
+	RIGVM_API int32 NumArguments() const { return Arguments.Num(); }
 
 	// returns an argument for a given index
 	const FRigVMTemplateArgument* GetArgument(int32 InIndex) const { return &Arguments[InIndex]; }
 
 	// returns an argument given a name (or nullptr)
-	const FRigVMTemplateArgument* FindArgument(const FName& InArgumentName) const;
+	RIGVM_API const FRigVMTemplateArgument* FindArgument(const FName& InArgumentName) const;
 
 	// returns the number of args of this template
-	int32 NumExecuteArguments(const FRigVMDispatchContext& InContext) const;
+	RIGVM_API int32 NumExecuteArguments(const FRigVMDispatchContext& InContext) const;
 
 	// returns an argument for a given index
-	const FRigVMExecuteArgument* GetExecuteArgument(int32 InIndex, const FRigVMDispatchContext& InContext) const;
+	RIGVM_API const FRigVMExecuteArgument* GetExecuteArgument(int32 InIndex, const FRigVMDispatchContext& InContext) const;
 
 	// returns an argument given a name (or nullptr)
-	const FRigVMExecuteArgument* FindExecuteArgument(const FName& InArgumentName, const FRigVMDispatchContext& InContext) const;
+	RIGVM_API const FRigVMExecuteArgument* FindExecuteArgument(const FName& InArgumentName, const FRigVMDispatchContext& InContext) const;
 
 	// returns the top level execute context struct this template uses
-	const UScriptStruct* GetExecuteContextStruct(bool bLockRegistry = true) const;
+	RIGVM_API const UScriptStruct* GetExecuteContextStruct(bool bLockRegistry = true) const;
 
 	// returns true if this template supports a given execute context struct
-	bool SupportsExecuteContextStruct(const UScriptStruct* InExecuteContextStruct) const;
+	RIGVM_API bool SupportsExecuteContextStruct(const UScriptStruct* InExecuteContextStruct) const;
 
 #if WITH_EDITOR
 	// returns true if a given arg supports a type
-	bool ArgumentSupportsTypeIndex(const FName& InArgumentName, TRigVMTypeIndex InTypeIndex, TRigVMTypeIndex* OutTypeIndex = nullptr) const;
+	RIGVM_API bool ArgumentSupportsTypeIndex(const FName& InArgumentName, TRigVMTypeIndex InTypeIndex, TRigVMTypeIndex* OutTypeIndex = nullptr) const;
 #endif
 
 	// returns the number of permutations supported by this template
 	int32 NumPermutations() const { return Permutations.Num(); }
 
 	// returns the first / primary permutation of the template
-	const FRigVMFunction* GetPrimaryPermutation(bool bLockRegistry = true) const;
+	RIGVM_API const FRigVMFunction* GetPrimaryPermutation(bool bLockRegistry = true) const;
 
 	// returns a permutation given an index
-	const FRigVMFunction* GetPermutation(int32 InIndex, bool bLockRegistry = true) const;
+	RIGVM_API const FRigVMFunction* GetPermutation(int32 InIndex, bool bLockRegistry = true) const;
 
 	// returns a permutation given an index and creates it using the backing factory if needed
-	const FRigVMFunction* GetOrCreatePermutation(int32 InIndex, bool bLockRegistry = true);
+	RIGVM_API const FRigVMFunction* GetOrCreatePermutation(int32 InIndex, bool bLockRegistry = true);
 
 	// returns true if a given function is a permutation of this template
-	bool ContainsPermutation(const FRigVMFunction* InPermutation, bool bLockRegistry = true) const;
-	bool ContainsPermutation_NoLock(const FRigVMFunction* InPermutation) const;
+	RIGVM_API bool ContainsPermutation(const FRigVMFunction* InPermutation, bool bLockRegistry = true) const;
+	RIGVM_API bool ContainsPermutation_NoLock(const FRigVMFunction* InPermutation) const;
 
 	// returns the index of the permutation within the template of a given function (or INDEX_NONE)
-	int32 FindPermutation(const FRigVMFunction* InPermutation) const;
+	RIGVM_API int32 FindPermutation(const FRigVMFunction* InPermutation) const;
 
 	// returns the index of the permutation within the template of a given set of types
-	int32 FindPermutation(const FTypeMap& InTypes, bool bLockRegistry = true) const;
+	RIGVM_API int32 FindPermutation(const FTypeMap& InTypes, bool bLockRegistry = true) const;
 
 	// returns true if the template was able to resolve to single permutation
-	bool FullyResolve(FTypeMap& InOutTypes, int32& OutPermutationIndex, bool bLockRegistry = true) const;
+	RIGVM_API bool FullyResolve(FTypeMap& InOutTypes, int32& OutPermutationIndex, bool bLockRegistry = true) const;
 
 	// returns true if the template was able to resolve to at least one permutation
-	bool Resolve(FTypeMap& InOutTypes, TArray<int32> & OutPermutationIndices, bool bAllowFloatingPointCasts, bool bLockRegistry = true) const;
+	RIGVM_API bool Resolve(FTypeMap& InOutTypes, TArray<int32> & OutPermutationIndices, bool bAllowFloatingPointCasts, bool bLockRegistry = true) const;
 
 	// Will return the hash of the input type map, if it is a valid type map. Otherwise, will return 0.
 	// It is only a valid type map if it includes all arguments, and non of the types is a wildcard.
-	uint32 GetTypesHashFromTypes(const FTypeMap& InTypes) const;
+	RIGVM_API uint32 GetTypesHashFromTypes(const FTypeMap& InTypes) const;
 
 	// returns true if the template was able to resolve to at least one permutation
-	bool ContainsPermutation(const FTypeMap& InTypes, bool bLockRegistry = true) const;
-	bool ContainsPermutation_NoLock(const FTypeMap& InTypes) const;
+	RIGVM_API bool ContainsPermutation(const FTypeMap& InTypes, bool bLockRegistry = true) const;
+	RIGVM_API bool ContainsPermutation_NoLock(const FTypeMap& InTypes) const;
 
 	// returns true if the template can resolve an argument to a new type
-	bool ResolveArgument(const FName& InArgumentName, const TRigVMTypeIndex InTypeIndex, FTypeMap& InOutTypes, bool bLockRegistry = true) const;
+	RIGVM_API bool ResolveArgument(const FName& InArgumentName, const TRigVMTypeIndex InTypeIndex, FTypeMap& InOutTypes, bool bLockRegistry = true) const;
 
 	// returns the types for a specific permutation
-	FRigVMTemplateTypeMap GetTypesForPermutation(const int32 InPermutationIndex, const bool bLockRegistry = true) const;
-	FRigVMTemplateTypeMap GetTypesForPermutation_NoLock(const int32 InPermutationIndex) const;
+	RIGVM_API FRigVMTemplateTypeMap GetTypesForPermutation(const int32 InPermutationIndex, const bool bLockRegistry = true) const;
+	RIGVM_API FRigVMTemplateTypeMap GetTypesForPermutation_NoLock(const int32 InPermutationIndex) const;
 
 	// returns true if a given argument is valid for a template
-	static bool IsValidArgumentForTemplate(const ERigVMPinDirection InDirection);
+	static RIGVM_API bool IsValidArgumentForTemplate(const ERigVMPinDirection InDirection);
 
 	// returns the prefix for an argument in the notation
-	static const FString& GetDirectionPrefix(const ERigVMPinDirection InDirection);
+	static RIGVM_API const FString& GetDirectionPrefix(const ERigVMPinDirection InDirection);
 
 	// returns the notation of an argument
-	static FString GetArgumentNotation(const FName InName, const ERigVMPinDirection InDirection);
+	static RIGVM_API FString GetArgumentNotation(const FName InName, const ERigVMPinDirection InDirection);
 
 	// recomputes the notation from its arguments
-	void ComputeNotationFromArguments(const FString& InTemplateName);
+	RIGVM_API void ComputeNotationFromArguments(const FString& InTemplateName);
 
 	// returns an array of structs in the inheritance order of a given struct
-	static TArray<UStruct*> GetSuperStructs(UStruct* InStruct, bool bIncludeLeaf = true);
+	static RIGVM_API TArray<UStruct*> GetSuperStructs(UStruct* InStruct, bool bIncludeLeaf = true);
 
 	// converts the types provided by a string (like "A:float,B:int32") into a type map
-	FTypeMap GetArgumentTypesFromString(const FString& InTypeString, const FRigVMUserDefinedTypeResolver* InTypeResolver = nullptr) const;
+	RIGVM_API FTypeMap GetArgumentTypesFromString(const FString& InTypeString, const FRigVMUserDefinedTypeResolver* InTypeResolver = nullptr) const;
 
 	// converts the types provided to a string (like "A:float,B:int32")
-	static FString GetStringFromArgumentTypes(const FTypeMap& InTypes, bool bLockRegistry = true); 
+	static RIGVM_API FString GetStringFromArgumentTypes(const FTypeMap& InTypes, bool bLockRegistry = true); 
 
 #if WITH_EDITOR
 
 	// Returns the color based on the permutation's metadata
-	FLinearColor GetColor(const TArray<int32>& InPermutationIndices = TArray<int32>()) const;
+	RIGVM_API FLinearColor GetColor(const TArray<int32>& InPermutationIndices = TArray<int32>()) const;
 
 	// Returns the tooltip based on the permutation's metadata
-	FText GetTooltipText(const TArray<int32>& InPermutationIndices = TArray<int32>()) const;
+	RIGVM_API FText GetTooltipText(const TArray<int32>& InPermutationIndices = TArray<int32>()) const;
 
 	// Returns the display name text for an argument 
-	FText GetDisplayNameForArgument(const FName& InArgumentName, const TArray<int32>& InPermutationIndices = TArray<int32>()) const;
+	RIGVM_API FText GetDisplayNameForArgument(const FName& InArgumentName, const TArray<int32>& InPermutationIndices = TArray<int32>()) const;
 
 	// Returns meta data on the property of the permutations 
-	FString GetArgumentMetaData(const FName& InArgumentName, const FName& InMetaDataKey, const TArray<int32>& InPermutationIndices = TArray<int32>()) const;
+	RIGVM_API FString GetArgumentMetaData(const FName& InArgumentName, const FName& InMetaDataKey, const TArray<int32>& InPermutationIndices = TArray<int32>()) const;
 
-	FString GetCategory() const;
-	FString GetKeywords() const;
+	RIGVM_API FString GetCategory() const;
+	RIGVM_API FString GetKeywords() const;
 
 #endif
 
 	// Updates the template's argument types. This only affects templates which have category based
 	// arguments and will resolve the other arguments to the expected types.
-	bool UpdateAllArgumentTypesSlow();
+	RIGVM_API bool UpdateAllArgumentTypesSlow();
 
 	// Updates the template's argument types by adding a single new type
-	bool HandlePropagatedArgumentType(const TRigVMTypeIndex InTypeIndex);
+	RIGVM_API bool HandlePropagatedArgumentType(const TRigVMTypeIndex InTypeIndex);
 
 	// Invalidates template permutations whenever a type such as a user defined struct is removed
-	void HandleTypeRemoval(TRigVMTypeIndex InTypeIndex);
+	RIGVM_API void HandleTypeRemoval(TRigVMTypeIndex InTypeIndex);
 	
 	// Returns the delegate to be able to react to type changes dynamically
 	// This delegate is deprecated
 	FRigVMTemplate_NewArgumentTypeDelegate& OnNewArgumentType() { return Delegates.NewArgumentTypeDelegate; }
 
 	// Returns the factory this template was created by
-	const FRigVMDispatchFactory* GetDispatchFactory(const bool bLockRegistry = true) const;
+	RIGVM_API const FRigVMDispatchFactory* GetDispatchFactory(const bool bLockRegistry = true) const;
 
 	// Returns the factory this template was created by
 	const FRigVMDispatchFactory* GetDispatchFactory_NoLock() const
@@ -561,10 +561,10 @@ public:
 		return Delegates.GetDispatchFactoryDelegate.IsBound();
 	}
 
-	void RecomputeTypesHashToPermutations();
-	void RecomputeTypesHashToPermutations(const TArray<FRigVMTypeCacheScope_NoLock>& InTypeCaches);
+	RIGVM_API void RecomputeTypesHashToPermutations();
+	RIGVM_API void RecomputeTypesHashToPermutations(const TArray<FRigVMTypeCacheScope_NoLock>& InTypeCaches);
 
-	void UpdateTypesHashToPermutation(const int32 InPermutation);
+	RIGVM_API void UpdateTypesHashToPermutation(const int32 InPermutation);
 
 	RIGVM_API friend uint32 GetTypeHash(const FRigVMTemplate& InTemplate);
 	RIGVM_API friend uint32 GetTypeHash_NoLock(const FRigVMTemplate& InTemplate);
@@ -580,7 +580,7 @@ private:
 	static FLinearColor GetColorFromMetadata(FString InMetadata);
 
 	void InvalidateHash() { Hash = UINT32_MAX; }
-	const TArray<FRigVMExecuteArgument>& GetExecuteArguments(const FRigVMDispatchContext& InContext) const;
+	RIGVM_API const TArray<FRigVMExecuteArgument>& GetExecuteArguments(const FRigVMDispatchContext& InContext) const;
 
 	const FRigVMFunction* GetPermutation_NoLock(int32 InIndex) const;
 	const FRigVMFunction* GetOrCreatePermutation_NoLock(int32 InIndex);
diff --git a/Engine/Plugins/Runtime/RigVM/Source/RigVM/Public/RigVMCore/RigVMTrait.h b/Engine/Plugins/Runtime/RigVM/Source/RigVM/Public/RigVMCore/RigVMTrait.h
--- a/Engine/Plugins/Runtime/RigVM/Source/RigVM/Public/RigVMCore/RigVMTrait.h
+++ b/Engine/Plugins/Runtime/RigVM/Source/RigVM/Public/RigVMCore/RigVMTrait.h
@@ -18,7 +18,7 @@ class URigVMNode;
  * The base class for all RigVM traits.
  */
 USTRUCT()
-struct RIGVM_API FRigVMTrait : public FRigVMStruct
+struct FRigVMTrait : public FRigVMStruct
 {
 	GENERATED_BODY()
 
diff --git a/Engine/Plugins/Runtime/RigVM/Source/RigVM/Public/RigVMCore/RigVMTraitScope.h b/Engine/Plugins/Runtime/RigVM/Source/RigVM/Public/RigVMCore/RigVMTraitScope.h
--- a/Engine/Plugins/Runtime/RigVM/Source/RigVM/Public/RigVMCore/RigVMTraitScope.h
+++ b/Engine/Plugins/Runtime/RigVM/Source/RigVM/Public/RigVMCore/RigVMTraitScope.h
@@ -9,7 +9,7 @@ struct FRigVMTrait;
 class UScriptStruct;
 struct FRigVMMemoryHandle;
 
-class RIGVM_API FRigVMTraitScope
+class FRigVMTraitScope
 {
 public:
 	FRigVMTraitScope()
diff --git a/Engine/Plugins/Runtime/RigVM/Source/RigVM/Public/RigVMCore/RigVMTypeIndex.h b/Engine/Plugins/Runtime/RigVM/Source/RigVM/Public/RigVMCore/RigVMTypeIndex.h
--- a/Engine/Plugins/Runtime/RigVM/Source/RigVM/Public/RigVMCore/RigVMTypeIndex.h
+++ b/Engine/Plugins/Runtime/RigVM/Source/RigVM/Public/RigVMCore/RigVMTypeIndex.h
@@ -6,7 +6,7 @@
 #include "RigVMDefines.h"
 
 #if UE_RIGVM_DEBUG_TYPEINDEX
-class RIGVM_API TRigVMTypeIndex
+class TRigVMTypeIndex
 {
 public:
 
diff --git a/Engine/Plugins/Runtime/RigVM/Source/RigVM/Public/RigVMFunctions/Animation/RigVMFunction_AnimBase.h b/Engine/Plugins/Runtime/RigVM/Source/RigVM/Public/RigVMFunctions/Animation/RigVMFunction_AnimBase.h
--- a/Engine/Plugins/Runtime/RigVM/Source/RigVM/Public/RigVMFunctions/Animation/RigVMFunction_AnimBase.h
+++ b/Engine/Plugins/Runtime/RigVM/Source/RigVM/Public/RigVMFunctions/Animation/RigVMFunction_AnimBase.h
@@ -6,7 +6,7 @@
 #include "RigVMFunction_AnimBase.generated.h"
 
 USTRUCT(meta=(Abstract, Category = "Animation", NodeColor = "0.05 0.05 0.25"))
-struct RIGVM_API FRigVMFunction_AnimBase : public FRigVMStruct
+struct FRigVMFunction_AnimBase : public FRigVMStruct
 {
 	GENERATED_BODY()
 };
diff --git a/Engine/Plugins/Runtime/RigVM/Source/RigVM/Public/RigVMFunctions/Animation/RigVMFunction_AnimEasing.h b/Engine/Plugins/Runtime/RigVM/Source/RigVM/Public/RigVMFunctions/Animation/RigVMFunction_AnimEasing.h
--- a/Engine/Plugins/Runtime/RigVM/Source/RigVM/Public/RigVMFunctions/Animation/RigVMFunction_AnimEasing.h
+++ b/Engine/Plugins/Runtime/RigVM/Source/RigVM/Public/RigVMFunctions/Animation/RigVMFunction_AnimEasing.h
@@ -10,7 +10,7 @@
  * A constant value of an easing type
  */
 USTRUCT(meta = (DisplayName = "EaseType", Keywords = "Constant"))
-struct RIGVM_API FRigVMFunction_AnimEasingType : public FRigVMFunction_AnimBase
+struct FRigVMFunction_AnimEasingType : public FRigVMFunction_AnimBase
 {
 	GENERATED_BODY()
 
@@ -20,7 +20,7 @@ struct RIGVM_API FRigVMFunction_AnimEasingType : public FRigVMFunction_AnimBase
 	}
 
 	RIGVM_METHOD()
-	virtual void Execute() override;
+	RIGVM_API virtual void Execute() override;
 
 	UPROPERTY(meta = (Input, Output))
 	ERigVMAnimEasingType Type;
@@ -30,7 +30,7 @@ struct RIGVM_API FRigVMFunction_AnimEasingType : public FRigVMFunction_AnimBase
  * Returns the eased version of the input value
  */
 USTRUCT(meta=(DisplayName="Ease", Keywords="Easing,Profile,Smooth,Cubic"))
-struct RIGVM_API FRigVMFunction_AnimEasing : public FRigVMFunction_AnimBase
+struct FRigVMFunction_AnimEasing : public FRigVMFunction_AnimBase
 {
 	GENERATED_BODY()
 	
@@ -43,7 +43,7 @@ struct RIGVM_API FRigVMFunction_AnimEasing : public FRigVMFunction_AnimBase
 	}
 
 	RIGVM_METHOD()
-	virtual void Execute() override;
+	RIGVM_API virtual void Execute() override;
 
 	UPROPERTY(meta=(Input))
 	float Value;
diff --git a/Engine/Plugins/Runtime/RigVM/Source/RigVM/Public/RigVMFunctions/Animation/RigVMFunction_AnimRichCurve.h b/Engine/Plugins/Runtime/RigVM/Source/RigVM/Public/RigVMFunctions/Animation/RigVMFunction_AnimRichCurve.h
--- a/Engine/Plugins/Runtime/RigVM/Source/RigVM/Public/RigVMFunctions/Animation/RigVMFunction_AnimRichCurve.h
+++ b/Engine/Plugins/Runtime/RigVM/Source/RigVM/Public/RigVMFunctions/Animation/RigVMFunction_AnimRichCurve.h
@@ -10,7 +10,7 @@
  * Provides a constant curve to be used for multiple curve evaluations
  */
 USTRUCT(meta=(DisplayName="Curve", Keywords="Curve,Profile"))
-struct RIGVM_API FRigVMFunction_AnimRichCurve : public FRigVMFunction_AnimBase
+struct FRigVMFunction_AnimRichCurve : public FRigVMFunction_AnimBase
 {
 	GENERATED_BODY()
 	
@@ -22,7 +22,7 @@ struct RIGVM_API FRigVMFunction_AnimRichCurve : public FRigVMFunction_AnimBase
 	}
 
 	RIGVM_METHOD()
-	virtual void Execute() override;
+	RIGVM_API virtual void Execute() override;
 
 	UPROPERTY(meta=(Input, Output, Constant))
 	FRuntimeFloatCurve Curve;
diff --git a/Engine/Plugins/Runtime/RigVM/Source/RigVM/Public/RigVMFunctions/Animation/RigVMFunction_GetDeltaTime.h b/Engine/Plugins/Runtime/RigVM/Source/RigVM/Public/RigVMFunctions/Animation/RigVMFunction_GetDeltaTime.h
--- a/Engine/Plugins/Runtime/RigVM/Source/RigVM/Public/RigVMFunctions/Animation/RigVMFunction_GetDeltaTime.h
+++ b/Engine/Plugins/Runtime/RigVM/Source/RigVM/Public/RigVMFunctions/Animation/RigVMFunction_GetDeltaTime.h
@@ -9,7 +9,7 @@
  * Returns the time gone by from the previous evaluation
  */
 USTRUCT(meta=(DisplayName="Delta Time", Varying))
-struct RIGVM_API FRigVMFunction_GetDeltaTime : public FRigVMFunction_AnimBase
+struct FRigVMFunction_GetDeltaTime : public FRigVMFunction_AnimBase
 {
 	GENERATED_BODY()
 	
@@ -19,7 +19,7 @@ struct RIGVM_API FRigVMFunction_GetDeltaTime : public FRigVMFunction_AnimBase
 	}
 
 	RIGVM_METHOD()
-	virtual void Execute() override;
+	RIGVM_API virtual void Execute() override;
 
 	UPROPERTY(meta=(Output))
 	float Result;
diff --git a/Engine/Plugins/Runtime/RigVM/Source/RigVM/Public/RigVMFunctions/Animation/RigVMFunction_GetWorldTime.h b/Engine/Plugins/Runtime/RigVM/Source/RigVM/Public/RigVMFunctions/Animation/RigVMFunction_GetWorldTime.h
--- a/Engine/Plugins/Runtime/RigVM/Source/RigVM/Public/RigVMFunctions/Animation/RigVMFunction_GetWorldTime.h
+++ b/Engine/Plugins/Runtime/RigVM/Source/RigVM/Public/RigVMFunctions/Animation/RigVMFunction_GetWorldTime.h
@@ -9,7 +9,7 @@
  * Returns the current time (year, month, day, hour, minute)
  */
 USTRUCT(meta = (DisplayName = "Now", Keywords = "Time,Clock", Varying))
-struct RIGVM_API FRigVMFunction_GetWorldTime : public FRigVMFunction_AnimBase
+struct FRigVMFunction_GetWorldTime : public FRigVMFunction_AnimBase
 {
 	GENERATED_BODY()
 
@@ -19,7 +19,7 @@ struct RIGVM_API FRigVMFunction_GetWorldTime : public FRigVMFunction_AnimBase
 	}
 
 	RIGVM_METHOD()
-	virtual void Execute() override;
+	RIGVM_API virtual void Execute() override;
 
 	UPROPERTY(meta = (Output))
 	float Year;
diff --git a/Engine/Plugins/Runtime/RigVM/Source/RigVM/Public/RigVMFunctions/Animation/RigVMFunction_TimeConversion.h b/Engine/Plugins/Runtime/RigVM/Source/RigVM/Public/RigVMFunctions/Animation/RigVMFunction_TimeConversion.h
--- a/Engine/Plugins/Runtime/RigVM/Source/RigVM/Public/RigVMFunctions/Animation/RigVMFunction_TimeConversion.h
+++ b/Engine/Plugins/Runtime/RigVM/Source/RigVM/Public/RigVMFunctions/Animation/RigVMFunction_TimeConversion.h
@@ -9,7 +9,7 @@
  * Converts frames to seconds based on the current frame rate
  */
 USTRUCT(meta=(DisplayName="Frames to Seconds", Varying))
-struct RIGVM_API FRigVMFunction_FramesToSeconds : public FRigVMFunction_AnimBase
+struct FRigVMFunction_FramesToSeconds : public FRigVMFunction_AnimBase
 {
 	GENERATED_BODY()
 	
@@ -19,7 +19,7 @@ struct RIGVM_API FRigVMFunction_FramesToSeconds : public FRigVMFunction_AnimBase
 	}
 
 	RIGVM_METHOD()
-	virtual void Execute() override;
+	RIGVM_API virtual void Execute() override;
 
 	UPROPERTY(meta=(Input))
 	float Frames;
@@ -32,7 +32,7 @@ struct RIGVM_API FRigVMFunction_FramesToSeconds : public FRigVMFunction_AnimBase
  * Converts seconds to frames based on the current frame rate
  */
 USTRUCT(meta=(DisplayName="Seconds to Frames", Varying))
-struct RIGVM_API FRigVMFunction_SecondsToFrames : public FRigVMFunction_AnimBase
+struct  FRigVMFunction_SecondsToFrames : public FRigVMFunction_AnimBase
 {
 	GENERATED_BODY()
 	
@@ -42,7 +42,7 @@ struct RIGVM_API FRigVMFunction_SecondsToFrames : public FRigVMFunction_AnimBase
 	}
 
 	RIGVM_METHOD()
-	virtual void Execute() override;
+	RIGVM_API virtual void Execute() override;
 
 	UPROPERTY(meta=(Input))
 	float Seconds;
diff --git a/Engine/Plugins/Runtime/RigVM/Source/RigVM/Public/RigVMFunctions/Debug/RigVMFunction_DebugBase.h b/Engine/Plugins/Runtime/RigVM/Source/RigVM/Public/RigVMFunctions/Debug/RigVMFunction_DebugBase.h
--- a/Engine/Plugins/Runtime/RigVM/Source/RigVM/Public/RigVMFunctions/Debug/RigVMFunction_DebugBase.h
+++ b/Engine/Plugins/Runtime/RigVM/Source/RigVM/Public/RigVMFunctions/Debug/RigVMFunction_DebugBase.h
@@ -7,7 +7,7 @@
 #include "RigVMFunction_DebugBase.generated.h"
 
 USTRUCT(meta=(Abstract, Category="Debug", NodeColor = "0.83077 0.846873 0.049707"))
-struct RIGVM_API FRigVMFunction_DebugBase : public FRigVMStruct
+struct FRigVMFunction_DebugBase : public FRigVMStruct
 {
 	GENERATED_BODY()
 
@@ -16,7 +16,7 @@ struct RIGVM_API FRigVMFunction_DebugBase : public FRigVMStruct
 };
 
 USTRUCT(meta=(Abstract, Category="Debug", NodeColor = "0.83077 0.846873 0.049707"))
-struct RIGVM_API FRigVMFunction_DebugBaseMutable : public FRigVMStructMutable
+struct FRigVMFunction_DebugBaseMutable : public FRigVMStructMutable
 {
 	GENERATED_BODY()
 
diff --git a/Engine/Plugins/Runtime/RigVM/Source/RigVM/Public/RigVMFunctions/Debug/RigVMFunction_DebugLine.h b/Engine/Plugins/Runtime/RigVM/Source/RigVM/Public/RigVMFunctions/Debug/RigVMFunction_DebugLine.h
--- a/Engine/Plugins/Runtime/RigVM/Source/RigVM/Public/RigVMFunctions/Debug/RigVMFunction_DebugLine.h
+++ b/Engine/Plugins/Runtime/RigVM/Source/RigVM/Public/RigVMFunctions/Debug/RigVMFunction_DebugLine.h
@@ -5,11 +5,13 @@
 #include "RigVMFunction_DebugBase.h"
 #include "RigVMFunction_DebugLine.generated.h"
 
+#define UE_API RIGVM_API
+
 /**
  * Draws a line in the viewport given a start and end vector
  */
 USTRUCT(meta=(DisplayName="Draw Line"))
-struct RIGVM_API FRigVMFunction_DebugLineNoSpace : public FRigVMFunction_DebugBaseMutable
+struct FRigVMFunction_DebugLineNoSpace : public FRigVMFunction_DebugBaseMutable
 {
 	GENERATED_BODY()
 
@@ -23,7 +25,7 @@ struct RIGVM_API FRigVMFunction_DebugLineNoSpace : public FRigVMFunction_DebugBa
 	}
 
 	RIGVM_METHOD()
-	virtual void Execute() override;
+	UE_API virtual void Execute() override;
 
 	UPROPERTY(meta = (Input))
 	FVector A;
@@ -42,4 +44,6 @@ struct RIGVM_API FRigVMFunction_DebugLineNoSpace : public FRigVMFunction_DebugBa
 
 	UPROPERTY(meta = (Input))
 	bool bEnabled;
-};
\ No newline at end of file
+};
+
+#undef UE_API
\ No newline at end of file
diff --git a/Engine/Plugins/Runtime/RigVM/Source/RigVM/Public/RigVMFunctions/Debug/RigVMFunction_DebugLineStrip.h b/Engine/Plugins/Runtime/RigVM/Source/RigVM/Public/RigVMFunctions/Debug/RigVMFunction_DebugLineStrip.h
--- a/Engine/Plugins/Runtime/RigVM/Source/RigVM/Public/RigVMFunctions/Debug/RigVMFunction_DebugLineStrip.h
+++ b/Engine/Plugins/Runtime/RigVM/Source/RigVM/Public/RigVMFunctions/Debug/RigVMFunction_DebugLineStrip.h
@@ -9,7 +9,7 @@
  * Draws a line strip in the viewport given any number of points
  */
 USTRUCT(meta=(DisplayName="Draw Line Strip"))
-struct RIGVM_API FRigVMFunction_DebugLineStripNoSpace : public FRigVMFunction_DebugBaseMutable
+struct FRigVMFunction_DebugLineStripNoSpace : public FRigVMFunction_DebugBaseMutable
 {
 	GENERATED_BODY()
 
@@ -22,7 +22,7 @@ struct RIGVM_API FRigVMFunction_DebugLineStripNoSpace : public FRigVMFunction_De
 	}
 
 	RIGVM_METHOD()
-	virtual void Execute() override;
+	RIGVM_API virtual void Execute() override;
 
 	UPROPERTY(meta = (Input))
 	TArray<FVector> Points;
diff --git a/Engine/Plugins/Runtime/RigVM/Source/RigVM/Public/RigVMFunctions/Debug/RigVMFunction_DebugPoint.h b/Engine/Plugins/Runtime/RigVM/Source/RigVM/Public/RigVMFunctions/Debug/RigVMFunction_DebugPoint.h
--- a/Engine/Plugins/Runtime/RigVM/Source/RigVM/Public/RigVMFunctions/Debug/RigVMFunction_DebugPoint.h
+++ b/Engine/Plugins/Runtime/RigVM/Source/RigVM/Public/RigVMFunctions/Debug/RigVMFunction_DebugPoint.h
@@ -19,7 +19,7 @@ enum class ERigUnitDebugPointMode : uint8
 };
 
 USTRUCT(meta=(DisplayName="Draw Point In Place", Keywords="Draw Vector", Deprecated = "4.25.0"))
-struct RIGVM_API FRigVMFunction_DebugPoint : public FRigVMFunction_DebugBase
+struct FRigVMFunction_DebugPoint : public FRigVMFunction_DebugBase
 {
 	GENERATED_BODY()
 
@@ -35,7 +35,7 @@ struct RIGVM_API FRigVMFunction_DebugPoint : public FRigVMFunction_DebugBase
 	}
 
 	RIGVM_METHOD()
-	virtual void Execute() override;
+	RIGVM_API virtual void Execute() override;
 
 	UPROPERTY(meta = (Input, Output))
 	FVector Vector;
@@ -66,7 +66,7 @@ struct RIGVM_API FRigVMFunction_DebugPoint : public FRigVMFunction_DebugBase
 };
 
 USTRUCT(meta=(DisplayName="Draw Point", Keywords="Draw Vector", Deprecated = "4.25.0"))
-struct RIGVM_API FRigVMFunction_DebugPointMutable : public FRigVMFunction_DebugBaseMutable
+struct FRigVMFunction_DebugPointMutable : public FRigVMFunction_DebugBaseMutable
 {
 	GENERATED_BODY()
 
@@ -82,7 +82,7 @@ struct RIGVM_API FRigVMFunction_DebugPointMutable : public FRigVMFunction_DebugB
 	}
 
 	RIGVM_METHOD()
-	virtual void Execute() override;
+	RIGVM_API virtual void Execute() override;
 
 	UPROPERTY(meta = (Input))
 	FVector Vector;
diff --git a/Engine/Plugins/Runtime/RigVM/Source/RigVM/Public/RigVMFunctions/Debug/RigVMFunction_DebugPrimitives.h b/Engine/Plugins/Runtime/RigVM/Source/RigVM/Public/RigVMFunctions/Debug/RigVMFunction_DebugPrimitives.h
--- a/Engine/Plugins/Runtime/RigVM/Source/RigVM/Public/RigVMFunctions/Debug/RigVMFunction_DebugPrimitives.h
+++ b/Engine/Plugins/Runtime/RigVM/Source/RigVM/Public/RigVMFunctions/Debug/RigVMFunction_DebugPrimitives.h
@@ -5,8 +5,10 @@
 #include "RigVMFunction_DebugBase.h"
 #include "RigVMFunction_DebugPrimitives.generated.h"
 
+#define UE_API RIGVM_API
+
 USTRUCT(meta=(DisplayName="Draw Rectangle", Keywords="Draw Square", Deprecated = "4.25"))
-struct RIGVM_API FRigVMFunction_DebugRectangle : public FRigVMFunction_DebugBaseMutable
+struct FRigVMFunction_DebugRectangle : public FRigVMFunction_DebugBaseMutable
 {
 	GENERATED_BODY()
 
@@ -21,7 +23,7 @@ struct RIGVM_API FRigVMFunction_DebugRectangle : public FRigVMFunction_DebugBase
 	}
 
 	RIGVM_METHOD()
-	virtual void Execute() override;
+	UE_API virtual void Execute() override;
 
 	UPROPERTY(meta = (Input))
 	FTransform Transform;
@@ -52,7 +54,7 @@ struct RIGVM_API FRigVMFunction_DebugRectangle : public FRigVMFunction_DebugBase
  * Draws a rectangle in the viewport given a transform
  */
 USTRUCT(meta=(DisplayName="Draw Rectangle", Keywords="Draw Square"))
-struct RIGVM_API FRigVMFunction_DebugRectangleNoSpace : public FRigVMFunction_DebugBaseMutable
+struct FRigVMFunction_DebugRectangleNoSpace : public FRigVMFunction_DebugBaseMutable
 {
 	GENERATED_BODY()
 
@@ -67,7 +69,7 @@ struct RIGVM_API FRigVMFunction_DebugRectangleNoSpace : public FRigVMFunction_De
 	}
 
 	RIGVM_METHOD()
-	virtual void Execute() override;
+	UE_API virtual void Execute() override;
 
 	UPROPERTY(meta = (Input))
 	FTransform Transform;
@@ -90,7 +92,7 @@ struct RIGVM_API FRigVMFunction_DebugRectangleNoSpace : public FRigVMFunction_De
 
 
 USTRUCT(meta=(DisplayName="Draw Arc", Keywords="Draw Ellipse, Draw Circle", Deprecated = "4.25"))
-struct RIGVM_API FRigVMFunction_DebugArc : public FRigVMFunction_DebugBaseMutable
+struct FRigVMFunction_DebugArc : public FRigVMFunction_DebugBaseMutable
 {
 	GENERATED_BODY()
 
@@ -107,7 +109,7 @@ struct RIGVM_API FRigVMFunction_DebugArc : public FRigVMFunction_DebugBaseMutabl
 	}
 
 	RIGVM_METHOD()
-	virtual void Execute() override;
+	UE_API virtual void Execute() override;
 
 	UPROPERTY(meta = (Input))
 	FTransform Transform;
@@ -147,7 +149,7 @@ struct RIGVM_API FRigVMFunction_DebugArc : public FRigVMFunction_DebugBaseMutabl
  * Draws an arc in the viewport, can take in different min and max degrees
  */
 USTRUCT(meta=(DisplayName="Draw Arc", Keywords="Draw Ellipse, Draw Circle"))
-struct RIGVM_API FRigVMFunction_DebugArcNoSpace : public FRigVMFunction_DebugBaseMutable
+struct FRigVMFunction_DebugArcNoSpace : public FRigVMFunction_DebugBaseMutable
 {
 	GENERATED_BODY()
 
@@ -164,7 +166,7 @@ struct RIGVM_API FRigVMFunction_DebugArcNoSpace : public FRigVMFunction_DebugBas
 	}
 
 	RIGVM_METHOD()
-	virtual void Execute() override;
+	UE_API virtual void Execute() override;
 
 	UPROPERTY(meta = (Input))
 	FTransform Transform;
@@ -198,7 +200,7 @@ struct RIGVM_API FRigVMFunction_DebugArcNoSpace : public FRigVMFunction_DebugBas
  * Draws a box in the viewport
  */
 USTRUCT(meta=(DisplayName="Draw Box", Keywords="BoundingBox,Bbox"))
-struct RIGVM_API FRigVMFunction_DebugBoxNoSpace : public FRigVMFunction_DebugBaseMutable
+struct FRigVMFunction_DebugBoxNoSpace : public FRigVMFunction_DebugBaseMutable
 {
 	GENERATED_BODY()
 
@@ -213,7 +215,7 @@ struct RIGVM_API FRigVMFunction_DebugBoxNoSpace : public FRigVMFunction_DebugBas
 	}
 
 	RIGVM_METHOD()
-	virtual void Execute() override;
+	UE_API virtual void Execute() override;
 
 	UPROPERTY(meta = (Input))
 	FBox Box;
@@ -229,4 +231,6 @@ struct RIGVM_API FRigVMFunction_DebugBoxNoSpace : public FRigVMFunction_DebugBas
 
 	UPROPERTY(meta = (Input))
 	bool bEnabled;
-};
\ No newline at end of file
+};
+
+#undef UE_API
\ No newline at end of file
diff --git a/Engine/Plugins/Runtime/RigVM/Source/RigVM/Public/RigVMFunctions/Debug/RigVMFunction_DebugTransform.h b/Engine/Plugins/Runtime/RigVM/Source/RigVM/Public/RigVMFunctions/Debug/RigVMFunction_DebugTransform.h
--- a/Engine/Plugins/Runtime/RigVM/Source/RigVM/Public/RigVMFunctions/Debug/RigVMFunction_DebugTransform.h
+++ b/Engine/Plugins/Runtime/RigVM/Source/RigVM/Public/RigVMFunctions/Debug/RigVMFunction_DebugTransform.h
@@ -25,7 +25,7 @@ enum class ERigUnitDebugTransformMode : uint8
  * Given a transform, will draw a point, axis, or a box in the viewport
  */
 USTRUCT(meta=(DisplayName="Draw Transform"))
-struct RIGVM_API FRigVMFunction_DebugTransformMutableNoSpace : public FRigVMFunction_DebugBaseMutable
+struct FRigVMFunction_DebugTransformMutableNoSpace : public FRigVMFunction_DebugBaseMutable
 {
 	GENERATED_BODY()
 
@@ -40,7 +40,7 @@ struct RIGVM_API FRigVMFunction_DebugTransformMutableNoSpace : public FRigVMFunc
 	}
 
 	RIGVM_METHOD()
-	virtual void Execute() override;
+	RIGVM_API virtual void Execute() override;
 
 	UPROPERTY(meta = (Input))
 	FTransform Transform;
@@ -65,7 +65,7 @@ struct RIGVM_API FRigVMFunction_DebugTransformMutableNoSpace : public FRigVMFunc
 };
 
 USTRUCT()
-struct RIGVM_API FRigVMFunction_DebugTransformArrayMutable_WorkData
+struct FRigVMFunction_DebugTransformArrayMutable_WorkData
 {
 	GENERATED_BODY()
 		
@@ -77,7 +77,7 @@ struct RIGVM_API FRigVMFunction_DebugTransformArrayMutable_WorkData
 * Given a transform array, will draw a point, axis, or a box in the viewport
 */
 USTRUCT(meta=(DisplayName="Draw Transform Array"))
-struct RIGVM_API FRigVMFunction_DebugTransformArrayMutableNoSpace : public FRigVMFunction_DebugBaseMutable
+struct FRigVMFunction_DebugTransformArrayMutableNoSpace : public FRigVMFunction_DebugBaseMutable
 {
 	GENERATED_BODY()
 
@@ -92,7 +92,7 @@ struct RIGVM_API FRigVMFunction_DebugTransformArrayMutableNoSpace : public FRigV
 	}
 
 	RIGVM_METHOD()
-	virtual void Execute() override;
+	RIGVM_API virtual void Execute() override;
 
 	UPROPERTY(meta = (Input))
 	TArray<FTransform> Transforms;
diff --git a/Engine/Plugins/Runtime/RigVM/Source/RigVM/Public/RigVMFunctions/Debug/RigVMFunction_VisualDebug.h b/Engine/Plugins/Runtime/RigVM/Source/RigVM/Public/RigVMFunctions/Debug/RigVMFunction_VisualDebug.h
--- a/Engine/Plugins/Runtime/RigVM/Source/RigVM/Public/RigVMFunctions/Debug/RigVMFunction_VisualDebug.h
+++ b/Engine/Plugins/Runtime/RigVM/Source/RigVM/Public/RigVMFunctions/Debug/RigVMFunction_VisualDebug.h
@@ -19,7 +19,7 @@ enum class ERigUnitVisualDebugPointMode : uint8
 };
 
 USTRUCT(meta=(DisplayName = "Visual Debug Vector", Keywords = "Draw,Point", Deprecated = "4.25", Varying))
-struct RIGVM_API FRigVMFunction_VisualDebugVector : public FRigVMFunction_DebugBase
+struct FRigVMFunction_VisualDebugVector : public FRigVMFunction_DebugBase
 {
 	GENERATED_BODY()
 
@@ -35,7 +35,7 @@ struct RIGVM_API FRigVMFunction_VisualDebugVector : public FRigVMFunction_DebugB
 	}
 
 	RIGVM_METHOD()
-	virtual void Execute() override;
+	RIGVM_API virtual void Execute() override;
 
 	UPROPERTY(meta = (Input, Output))
 	FVector Value;
@@ -66,7 +66,7 @@ struct RIGVM_API FRigVMFunction_VisualDebugVector : public FRigVMFunction_DebugB
  * Debug draw parameters for a Point or Vector given a vector
  */
 USTRUCT(meta=(DisplayName = "Visual Debug Vector", TemplateName="VisualDebug", Keywords = "Draw,Point", Varying))
-struct RIGVM_API FRigVMFunction_VisualDebugVectorNoSpace : public FRigVMFunction_DebugBase
+struct FRigVMFunction_VisualDebugVectorNoSpace : public FRigVMFunction_DebugBase
 {
 	GENERATED_BODY()
 
@@ -81,7 +81,7 @@ struct RIGVM_API FRigVMFunction_VisualDebugVectorNoSpace : public FRigVMFunction
 	}
 
 	RIGVM_METHOD()
-	virtual void Execute() override;
+	RIGVM_API virtual void Execute() override;
 
 	UPROPERTY(meta = (Input, Output))
 	FVector Value;
@@ -103,7 +103,7 @@ struct RIGVM_API FRigVMFunction_VisualDebugVectorNoSpace : public FRigVMFunction
 };
 
 USTRUCT(meta = (DisplayName = "Visual Debug Quat", Keywords = "Draw,Rotation", Deprecated = "4.25", Varying))
-struct RIGVM_API FRigVMFunction_VisualDebugQuat : public FRigVMFunction_DebugBase
+struct FRigVMFunction_VisualDebugQuat : public FRigVMFunction_DebugBase
 {
 	GENERATED_BODY()
 
@@ -117,7 +117,7 @@ struct RIGVM_API FRigVMFunction_VisualDebugQuat : public FRigVMFunction_DebugBas
 	}
 
 	RIGVM_METHOD()
-	virtual void Execute() override;
+	RIGVM_API virtual void Execute() override;
 
 	UPROPERTY(meta = (Input, Output))
 	FQuat Value;
@@ -142,7 +142,7 @@ struct RIGVM_API FRigVMFunction_VisualDebugQuat : public FRigVMFunction_DebugBas
  * Debug draw parameters for an Axis given a quaternion
  */
 USTRUCT(meta = (DisplayName = "Visual Debug Quat", TemplateName="VisualDebug", Keywords = "Draw,Rotation", Varying))
-struct RIGVM_API FRigVMFunction_VisualDebugQuatNoSpace : public FRigVMFunction_DebugBase
+struct FRigVMFunction_VisualDebugQuatNoSpace : public FRigVMFunction_DebugBase
 {
 	GENERATED_BODY()
 
@@ -155,7 +155,7 @@ struct RIGVM_API FRigVMFunction_VisualDebugQuatNoSpace : public FRigVMFunction_D
 	}
 
 	RIGVM_METHOD()
-	virtual void Execute() override;
+	RIGVM_API virtual void Execute() override;
 
 	UPROPERTY(meta = (Input, Output))
 	FQuat Value;
@@ -171,7 +171,7 @@ struct RIGVM_API FRigVMFunction_VisualDebugQuatNoSpace : public FRigVMFunction_D
 };
 
 USTRUCT(meta=(DisplayName="Visual Debug Transform", Keywords = "Draw,Axes", Deprecated = "4.25", Varying))
-struct RIGVM_API FRigVMFunction_VisualDebugTransform : public FRigVMFunction_DebugBase
+struct FRigVMFunction_VisualDebugTransform : public FRigVMFunction_DebugBase
 {
 	GENERATED_BODY()
 
@@ -185,7 +185,7 @@ struct RIGVM_API FRigVMFunction_VisualDebugTransform : public FRigVMFunction_Deb
 	}
 
 	RIGVM_METHOD()
-	virtual void Execute() override;
+	RIGVM_API virtual void Execute() override;
 
 	UPROPERTY(meta = (Input, Output))
 	FTransform Value;
@@ -210,7 +210,7 @@ struct RIGVM_API FRigVMFunction_VisualDebugTransform : public FRigVMFunction_Deb
  * Debug draw parameters for an Axis given a transform
  */
 USTRUCT(meta=(DisplayName="Visual Debug Transform", TemplateName="VisualDebug", Keywords = "Draw,Axes", Varying))
-struct RIGVM_API FRigVMFunction_VisualDebugTransformNoSpace : public FRigVMFunction_DebugBase
+struct FRigVMFunction_VisualDebugTransformNoSpace : public FRigVMFunction_DebugBase
 {
 	GENERATED_BODY()
 
@@ -223,7 +223,7 @@ struct RIGVM_API FRigVMFunction_VisualDebugTransformNoSpace : public FRigVMFunct
 	}
 
 	RIGVM_METHOD()
-	virtual void Execute() override;
+	RIGVM_API virtual void Execute() override;
 
 	UPROPERTY(meta = (Input, Output))
 	FTransform Value;
diff --git a/Engine/Plugins/Runtime/RigVM/Source/RigVM/Public/RigVMFunctions/Execution/RigVMFunction_Context.h b/Engine/Plugins/Runtime/RigVM/Source/RigVM/Public/RigVMFunctions/Execution/RigVMFunction_Context.h
--- a/Engine/Plugins/Runtime/RigVM/Source/RigVM/Public/RigVMFunctions/Execution/RigVMFunction_Context.h
+++ b/Engine/Plugins/Runtime/RigVM/Source/RigVM/Public/RigVMFunctions/Execution/RigVMFunction_Context.h
@@ -9,7 +9,7 @@
  * Returns true if a graph is run with its asset editor open. This is editor only - in shipping it always returns false.
  */
 USTRUCT(meta=(DisplayName="Is Asset Editor Open", Category="Execution", TitleColor="1 0 0", NodeColor="1 1 1", Keywords="Debug,Open,Inspect"))
-struct RIGVM_API FRigVMFunction_IsHostBeingDebugged : public FRigVMStruct
+struct FRigVMFunction_IsHostBeingDebugged : public FRigVMStruct
 {
 	GENERATED_BODY()
 
@@ -18,7 +18,7 @@ struct RIGVM_API FRigVMFunction_IsHostBeingDebugged : public FRigVMStruct
 	{}
 
 	RIGVM_METHOD()
-	virtual void Execute() override;
+	RIGVM_API virtual void Execute() override;
 
 	// True if the graph is currently open in the asset editor
 	UPROPERTY(EditAnywhere, Transient, Category = "Execution", meta = (Output))
diff --git a/Engine/Plugins/Runtime/RigVM/Source/RigVM/Public/RigVMFunctions/Execution/RigVMFunction_ForLoop.h b/Engine/Plugins/Runtime/RigVM/Source/RigVM/Public/RigVMFunctions/Execution/RigVMFunction_ForLoop.h
--- a/Engine/Plugins/Runtime/RigVM/Source/RigVM/Public/RigVMFunctions/Execution/RigVMFunction_ForLoop.h
+++ b/Engine/Plugins/Runtime/RigVM/Source/RigVM/Public/RigVMFunctions/Execution/RigVMFunction_ForLoop.h
@@ -9,7 +9,7 @@
  * Given a count, execute iteratively until the count is up
  */
 USTRUCT(meta=(DisplayName="For Loop", Category="Execution", TitleColor="1 0 0", NodeColor="1 1 1", Keywords="Iterate", Icon="EditorStyle|GraphEditor.Macro.Loop_16x"))
-struct RIGVM_API FRigVMFunction_ForLoopCount : public FRigVMStructMutable
+struct FRigVMFunction_ForLoopCount : public FRigVMStructMutable
 {
 	GENERATED_BODY()
 
@@ -31,7 +31,7 @@ struct RIGVM_API FRigVMFunction_ForLoopCount : public FRigVMStructMutable
 	virtual int32 GetNumSlices() const override { return Count; }
 
 	RIGVM_METHOD()
-	virtual void Execute() override;
+	RIGVM_API virtual void Execute() override;
 
 	UPROPERTY(meta = (Singleton))
 	FName BlockToRun;
diff --git a/Engine/Plugins/Runtime/RigVM/Source/RigVM/Public/RigVMFunctions/Execution/RigVMFunction_Sequence.h b/Engine/Plugins/Runtime/RigVM/Source/RigVM/Public/RigVMFunctions/Execution/RigVMFunction_Sequence.h
--- a/Engine/Plugins/Runtime/RigVM/Source/RigVM/Public/RigVMFunctions/Execution/RigVMFunction_Sequence.h
+++ b/Engine/Plugins/Runtime/RigVM/Source/RigVM/Public/RigVMFunctions/Execution/RigVMFunction_Sequence.h
@@ -9,12 +9,12 @@
  * Allows for a single execution pulse to trigger a series of events in order.
  */
 USTRUCT(meta=(DisplayName="Sequence", Category="Execution", TitleColor="1 0 0", NodeColor="1 1 1", Icon="EditorStyle|GraphEditor.Sequence_16x"))
-struct RIGVM_API FRigVMFunction_Sequence : public FRigVMStruct
+struct FRigVMFunction_Sequence : public FRigVMStruct
 {
 	GENERATED_BODY()
 
 	RIGVM_METHOD()
-	virtual void Execute() override;
+	RIGVM_API virtual void Execute() override;
 
 	// The execution input
 	UPROPERTY(EditAnywhere, Transient, DisplayName = "Execute", Category = "SequenceExecution", meta = (Input, Aggregate))
diff --git a/Engine/Plugins/Runtime/RigVM/Source/RigVM/Public/RigVMFunctions/Math/RigVMFunction_MathBase.h b/Engine/Plugins/Runtime/RigVM/Source/RigVM/Public/RigVMFunctions/Math/RigVMFunction_MathBase.h
--- a/Engine/Plugins/Runtime/RigVM/Source/RigVM/Public/RigVMFunctions/Math/RigVMFunction_MathBase.h
+++ b/Engine/Plugins/Runtime/RigVM/Source/RigVM/Public/RigVMFunctions/Math/RigVMFunction_MathBase.h
@@ -7,7 +7,7 @@
 #include "RigVMFunction_MathBase.generated.h"
 
 USTRUCT(meta=(Abstract, NodeColor = "0.05 0.25 0.05"))
-struct RIGVM_API FRigVMFunction_MathBase : public FRigVMStruct
+struct FRigVMFunction_MathBase : public FRigVMStruct
 {
 	GENERATED_BODY()
 
@@ -15,7 +15,7 @@ struct RIGVM_API FRigVMFunction_MathBase : public FRigVMStruct
 };
 
 USTRUCT(meta=(Abstract, NodeColor = "0.05 0.25 0.05"))
-struct RIGVM_API FRigVMFunction_MathMutableBase : public FRigVMStructMutable
+struct FRigVMFunction_MathMutableBase : public FRigVMStructMutable
 {
 	GENERATED_BODY()
 
diff --git a/Engine/Plugins/Runtime/RigVM/Source/RigVM/Public/RigVMFunctions/Math/RigVMFunction_MathBool.h b/Engine/Plugins/Runtime/RigVM/Source/RigVM/Public/RigVMFunctions/Math/RigVMFunction_MathBool.h
--- a/Engine/Plugins/Runtime/RigVM/Source/RigVM/Public/RigVMFunctions/Math/RigVMFunction_MathBool.h
+++ b/Engine/Plugins/Runtime/RigVM/Source/RigVM/Public/RigVMFunctions/Math/RigVMFunction_MathBool.h
@@ -6,13 +6,13 @@
 #include "RigVMFunction_MathBool.generated.h"
 
 USTRUCT(meta=(Abstract, Category="Math|Boolean"))
-struct RIGVM_API FRigVMFunction_MathBoolBase : public FRigVMFunction_MathBase
+struct FRigVMFunction_MathBoolBase : public FRigVMFunction_MathBase
 {
 	GENERATED_BODY()
 };
 
 USTRUCT(meta=(Abstract))
-struct RIGVM_API FRigVMFunction_MathBoolConstant : public FRigVMFunction_MathBoolBase
+struct FRigVMFunction_MathBoolConstant : public FRigVMFunction_MathBoolBase
 {
 	GENERATED_BODY()
 
@@ -26,7 +26,7 @@ struct RIGVM_API FRigVMFunction_MathBoolConstant : public FRigVMFunction_MathBoo
 };
 
 USTRUCT(meta=(Abstract))
-struct RIGVM_API FRigVMFunction_MathBoolUnaryOp : public FRigVMFunction_MathBoolBase
+struct FRigVMFunction_MathBoolUnaryOp : public FRigVMFunction_MathBoolBase
 {
 	GENERATED_BODY()
 
@@ -43,7 +43,7 @@ struct RIGVM_API FRigVMFunction_MathBoolUnaryOp : public FRigVMFunction_MathBool
 };
 
 USTRUCT(meta=(Abstract))
-struct RIGVM_API FRigVMFunction_MathBoolBinaryOp : public FRigVMFunction_MathBoolBase
+struct FRigVMFunction_MathBoolBinaryOp : public FRigVMFunction_MathBoolBase
 {
 	GENERATED_BODY()
 
@@ -63,7 +63,7 @@ struct RIGVM_API FRigVMFunction_MathBoolBinaryOp : public FRigVMFunction_MathBoo
 };
 
 USTRUCT(meta=(Abstract))
-struct RIGVM_API FRigVMFunction_MathBoolBinaryAggregateOp : public FRigVMFunction_MathBoolBase
+struct FRigVMFunction_MathBoolBinaryAggregateOp : public FRigVMFunction_MathBoolBase
 {
 	GENERATED_BODY()
 
@@ -86,7 +86,7 @@ struct RIGVM_API FRigVMFunction_MathBoolBinaryAggregateOp : public FRigVMFunctio
  * A boolean constant
  */
 USTRUCT(meta=(DisplayName="Bool", Keywords="Make,Construct,Constant"))
-struct RIGVM_API FRigVMFunction_MathBoolMake : public FRigVMFunction_MathBoolBase
+struct FRigVMFunction_MathBoolMake : public FRigVMFunction_MathBoolBase
 {
 	GENERATED_BODY()
 
@@ -96,7 +96,7 @@ struct RIGVM_API FRigVMFunction_MathBoolMake : public FRigVMFunction_MathBoolBas
 	}
 	
 	RIGVM_METHOD()
-	virtual void Execute() override;
+	RIGVM_API virtual void Execute() override;
 
 	UPROPERTY(meta=(Input, Output))
 	bool Value;
@@ -106,7 +106,7 @@ struct RIGVM_API FRigVMFunction_MathBoolMake : public FRigVMFunction_MathBoolBas
  * Returns true
  */
 USTRUCT(meta=(DisplayName="True", Keywords="Yes"))
-struct RIGVM_API FRigVMFunction_MathBoolConstTrue : public FRigVMFunction_MathBoolConstant
+struct FRigVMFunction_MathBoolConstTrue : public FRigVMFunction_MathBoolConstant
 {
 	GENERATED_BODY()
 	
@@ -116,14 +116,14 @@ struct RIGVM_API FRigVMFunction_MathBoolConstTrue : public FRigVMFunction_MathBo
 	}
 
 	RIGVM_METHOD()
-	virtual void Execute() override;
+	RIGVM_API virtual void Execute() override;
 };
 
 /**
  * Returns false
  */
 USTRUCT(meta=(DisplayName="False", Keywords="No"))
-struct RIGVM_API FRigVMFunction_MathBoolConstFalse : public FRigVMFunction_MathBoolConstant
+struct FRigVMFunction_MathBoolConstFalse : public FRigVMFunction_MathBoolConstant
 {
 	GENERATED_BODY()
 	
@@ -133,43 +133,43 @@ struct RIGVM_API FRigVMFunction_MathBoolConstFalse : public FRigVMFunction_MathB
 	}
 
 	RIGVM_METHOD()
-	virtual void Execute() override;
+	RIGVM_API virtual void Execute() override;
 };
 
 /**
  * Returns true if the condition is false
  */
 USTRUCT(meta=(DisplayName="Not", TemplateName="Not", Keywords="!"))
-struct RIGVM_API FRigVMFunction_MathBoolNot : public FRigVMFunction_MathBoolUnaryOp
+struct FRigVMFunction_MathBoolNot : public FRigVMFunction_MathBoolUnaryOp
 {
 	GENERATED_BODY()
 
 	RIGVM_METHOD()
-	virtual void Execute() override;
+	RIGVM_API virtual void Execute() override;
 };
 
 /**
  * Returns true if both conditions are true
  */
 USTRUCT(meta=(DisplayName="And", TemplateName="And", Keywords="&&"))
-struct RIGVM_API FRigVMFunction_MathBoolAnd : public FRigVMFunction_MathBoolBinaryAggregateOp
+struct FRigVMFunction_MathBoolAnd : public FRigVMFunction_MathBoolBinaryAggregateOp
 {
 	GENERATED_BODY()
 
 	RIGVM_METHOD()
-	virtual void Execute() override;
+	RIGVM_API virtual void Execute() override;
 };
 
 /**
  * Returns true if both conditions are false
  */
 USTRUCT(meta=(DisplayName="Nand", Deprecated="5.1"))
-struct RIGVM_API FRigVMFunction_MathBoolNand : public FRigVMFunction_MathBoolBinaryOp
+struct FRigVMFunction_MathBoolNand : public FRigVMFunction_MathBoolBinaryOp
 {
 	GENERATED_BODY()
 
 	RIGVM_METHOD()
-	virtual void Execute() override;
+	RIGVM_API virtual void Execute() override;
 
 	RIGVM_METHOD()
 	virtual FRigVMStructUpgradeInfo GetUpgradeInfo() const override;
@@ -179,35 +179,35 @@ struct RIGVM_API FRigVMFunction_MathBoolNand : public FRigVMFunction_MathBoolBin
  * Returns false if both conditions are true
  */
 USTRUCT(meta=(DisplayName="Nand", TemplateName="Nand"))
-struct RIGVM_API FRigVMFunction_MathBoolNand2 : public FRigVMFunction_MathBoolBinaryOp
+struct FRigVMFunction_MathBoolNand2 : public FRigVMFunction_MathBoolBinaryOp
 {
 	GENERATED_BODY()
 
 	RIGVM_METHOD()
-	virtual void Execute() override;
+	RIGVM_API virtual void Execute() override;
 };
 
 /**
  * Returns true if one of the conditions is true
  */
 USTRUCT(meta=(DisplayName="Or", TemplateName="Or", Keywords="||"))
-struct RIGVM_API FRigVMFunction_MathBoolOr : public FRigVMFunction_MathBoolBinaryAggregateOp
+struct FRigVMFunction_MathBoolOr : public FRigVMFunction_MathBoolBinaryAggregateOp
 {
 	GENERATED_BODY()
 
 	RIGVM_METHOD()
-	virtual void Execute() override;
+	RIGVM_API virtual void Execute() override;
 };
 
 /**
  * Returns true if the value A equals B
  */
 USTRUCT(meta=(DisplayName="Equals", TemplateName="Equals", Keywords="Same,==", Deprecated="5.1"))
-struct RIGVM_API FRigVMFunction_MathBoolEquals : public FRigVMFunction_MathBoolBase
+struct FRigVMFunction_MathBoolEquals : public FRigVMFunction_MathBoolBase
 {
 	GENERATED_BODY()
 	RIGVM_METHOD()
-	virtual void Execute() override;
+	RIGVM_API virtual void Execute() override;
 
 	FRigVMFunction_MathBoolEquals()
 	{
@@ -232,11 +232,11 @@ struct RIGVM_API FRigVMFunction_MathBoolEquals : public FRigVMFunction_MathBoolB
  * Returns true if the value A does not equal B
  */
 USTRUCT(meta=(DisplayName="Not Equals", TemplateName="NotEquals", Keywords="Different,!=,Xor", Deprecated="5.1"))
-struct RIGVM_API FRigVMFunction_MathBoolNotEquals : public FRigVMFunction_MathBoolBase
+struct FRigVMFunction_MathBoolNotEquals : public FRigVMFunction_MathBoolBase
 {
 	GENERATED_BODY()
 	RIGVM_METHOD()
-	virtual void Execute() override;
+	RIGVM_API virtual void Execute() override;
 
 	FRigVMFunction_MathBoolNotEquals()
 	{
@@ -261,11 +261,11 @@ struct RIGVM_API FRigVMFunction_MathBoolNotEquals : public FRigVMFunction_MathBo
 * Returns true if the value has changed from the last run
 */
 USTRUCT(meta=(DisplayName="Toggled", TemplateName="Toggled", Keywords="Changed,Different"))
-struct RIGVM_API FRigVMFunction_MathBoolToggled : public FRigVMFunction_MathBoolBase
+struct FRigVMFunction_MathBoolToggled : public FRigVMFunction_MathBoolBase
 {
 	GENERATED_BODY()
 	RIGVM_METHOD()
-	virtual void Execute() override;
+	RIGVM_API virtual void Execute() override;
 
 	FRigVMFunction_MathBoolToggled()
 	{
@@ -289,11 +289,11 @@ struct RIGVM_API FRigVMFunction_MathBoolToggled : public FRigVMFunction_MathBool
  * Returns true and false as a sequence.
  */
 USTRUCT(meta=(DisplayName="FlipFlop", Keywords="Toggle,Changed,Different", Varying))
-struct RIGVM_API FRigVMFunction_MathBoolFlipFlop : public FRigVMFunction_MathBoolBase
+struct FRigVMFunction_MathBoolFlipFlop : public FRigVMFunction_MathBoolBase
 {
 	GENERATED_BODY()
 	RIGVM_METHOD()
-	virtual void Execute() override;
+	RIGVM_API virtual void Execute() override;
 
 	FRigVMFunction_MathBoolFlipFlop()
 	{
@@ -327,11 +327,11 @@ struct RIGVM_API FRigVMFunction_MathBoolFlipFlop : public FRigVMFunction_MathBoo
  * Returns true once the first time this node is hit
  */
 USTRUCT(meta=(DisplayName="Once", Keywords="FlipFlop,Toggle,Changed,Different", Varying))
-struct RIGVM_API FRigVMFunction_MathBoolOnce : public FRigVMFunction_MathBoolBase
+struct FRigVMFunction_MathBoolOnce : public FRigVMFunction_MathBoolBase
 {
 	GENERATED_BODY()
 	RIGVM_METHOD()
-	virtual void Execute() override;
+	RIGVM_API virtual void Execute() override;
 
 	FRigVMFunction_MathBoolOnce()
 	{
@@ -361,11 +361,11 @@ struct RIGVM_API FRigVMFunction_MathBoolOnce : public FRigVMFunction_MathBoolBas
  * Turns the given bool into a float value
  */
 USTRUCT(meta=(DisplayName="To Float", TemplateName="Cast", ExecuteContext="FRigVMExecuteContext"))
-struct RIGVM_API FRigVMFunction_MathBoolToFloat : public FRigVMFunction_MathBoolBase
+struct FRigVMFunction_MathBoolToFloat : public FRigVMFunction_MathBoolBase
 {
 	GENERATED_BODY()
 	RIGVM_METHOD()
-	virtual void Execute() override;
+	RIGVM_API virtual void Execute() override;
 
 	FRigVMFunction_MathBoolToFloat()
 	{
@@ -384,11 +384,11 @@ struct RIGVM_API FRigVMFunction_MathBoolToFloat : public FRigVMFunction_MathBool
  * Turns the given bool into an integer value
  */
 USTRUCT(meta=(DisplayName="To Integer", TemplateName="Cast", ExecuteContext="FRigVMExecuteContext"))
-struct RIGVM_API FRigVMFunction_MathBoolToInteger : public FRigVMFunction_MathBoolBase
+struct FRigVMFunction_MathBoolToInteger : public FRigVMFunction_MathBoolBase
 {
 	GENERATED_BODY()
 	RIGVM_METHOD()
-	virtual void Execute() override;
+	RIGVM_API virtual void Execute() override;
 
 	FRigVMFunction_MathBoolToInteger()
 	{
diff --git a/Engine/Plugins/Runtime/RigVM/Source/RigVM/Public/RigVMFunctions/Math/RigVMFunction_MathBox.h b/Engine/Plugins/Runtime/RigVM/Source/RigVM/Public/RigVMFunctions/Math/RigVMFunction_MathBox.h
--- a/Engine/Plugins/Runtime/RigVM/Source/RigVM/Public/RigVMFunctions/Math/RigVMFunction_MathBox.h
+++ b/Engine/Plugins/Runtime/RigVM/Source/RigVM/Public/RigVMFunctions/Math/RigVMFunction_MathBox.h
@@ -7,7 +7,7 @@
 #include "RigVMFunction_MathBox.generated.h"
 
 USTRUCT(meta=(Abstract, Category="Math|Box", MenuDescSuffix="(Box)"))
-struct RIGVM_API FRigVMFunction_MathBoxBase : public FRigVMFunction_MathBase
+struct FRigVMFunction_MathBoxBase : public FRigVMFunction_MathBase
 {
 	GENERATED_BODY()
 };
@@ -16,7 +16,7 @@ struct RIGVM_API FRigVMFunction_MathBoxBase : public FRigVMFunction_MathBase
  * Returns bounding box of the given array of positions
  */
 USTRUCT(meta = (DisplayName = "Box from Array", TemplateName = "ArrayBounds,CreateBox,CreateBoundingBox,NewBox,NewBoundingBoxMakeBox,MakeBoundingBox,BoundingBox,Bbox,Bounds"))
-struct RIGVM_API FRigVMFunction_MathBoxFromArray : public FRigVMFunction_MathBoxBase
+struct FRigVMFunction_MathBoxFromArray : public FRigVMFunction_MathBoxBase
 {
 	GENERATED_BODY()
 
@@ -28,7 +28,7 @@ struct RIGVM_API FRigVMFunction_MathBoxFromArray : public FRigVMFunction_MathBox
 
 	/** Execute logic for this rig unit */
 	RIGVM_METHOD()
-	virtual void Execute() override;
+	RIGVM_API virtual void Execute() override;
 
 	UPROPERTY(meta = (Input))
 	TArray<FVector> Array;
@@ -53,7 +53,7 @@ struct RIGVM_API FRigVMFunction_MathBoxFromArray : public FRigVMFunction_MathBox
  * Returns true if the box has any content / is valid
  */
 USTRUCT(meta = (DisplayName = "Is Box Valid", TemplateName = "IsValid,HasVolume,ContainsPoints,Bounds,BoundingBox,Bbox"))
-struct RIGVM_API FRigVMFunction_MathBoxIsValid : public FRigVMFunction_MathBoxBase
+struct FRigVMFunction_MathBoxIsValid : public FRigVMFunction_MathBoxBase
 {
 	GENERATED_BODY()
 
@@ -65,7 +65,7 @@ struct RIGVM_API FRigVMFunction_MathBoxIsValid : public FRigVMFunction_MathBoxBa
 
 	/** Execute logic for this rig unit */
 	RIGVM_METHOD()
-	virtual void Execute() override;
+	RIGVM_API virtual void Execute() override;
 
 	UPROPERTY(meta = (Input))
 	FBox Box;
@@ -78,7 +78,7 @@ struct RIGVM_API FRigVMFunction_MathBoxIsValid : public FRigVMFunction_MathBoxBa
  * Returns the center of a bounding box
  */
 USTRUCT(meta = (DisplayName = "Get Box Center", TemplateName = "Middle,Origin,Bounds,BoundingBox,Bbox"))
-struct RIGVM_API FRigVMFunction_MathBoxGetCenter : public FRigVMFunction_MathBoxBase
+struct FRigVMFunction_MathBoxGetCenter : public FRigVMFunction_MathBoxBase
 {
 	GENERATED_BODY()
 
@@ -90,7 +90,7 @@ struct RIGVM_API FRigVMFunction_MathBoxGetCenter : public FRigVMFunction_MathBox
 
 	/** Execute logic for this rig unit */
 	RIGVM_METHOD()
-	virtual void Execute() override;
+	RIGVM_API virtual void Execute() override;
 
 	UPROPERTY(meta = (Input))
 	FBox Box;
@@ -103,7 +103,7 @@ struct RIGVM_API FRigVMFunction_MathBoxGetCenter : public FRigVMFunction_MathBox
  * Returns the size of a bounding box
  */
 USTRUCT(meta = (DisplayName = "Get Box Size", TemplateName = "Middle,Origin,Bounds,BoundingBox,Bbox"))
-struct RIGVM_API FRigVMFunction_MathBoxGetSize : public FRigVMFunction_MathBoxBase
+struct FRigVMFunction_MathBoxGetSize : public FRigVMFunction_MathBoxBase
 {
 	GENERATED_BODY()
 
@@ -115,7 +115,7 @@ struct RIGVM_API FRigVMFunction_MathBoxGetSize : public FRigVMFunction_MathBoxBa
 
 	/** Execute logic for this rig unit */
 	RIGVM_METHOD()
-	virtual void Execute() override;
+	RIGVM_API virtual void Execute() override;
 
 	UPROPERTY(meta = (Input))
 	FBox Box;
@@ -133,7 +133,7 @@ struct RIGVM_API FRigVMFunction_MathBoxGetSize : public FRigVMFunction_MathBoxBa
  * Move the box by a certain amount 
  */
 USTRUCT(meta = (DisplayName = "Shift Box", TemplateName = "Bbox,Translate,Move,BoundingBox"))
-struct RIGVM_API FRigVMFunction_MathBoxShift : public FRigVMFunction_MathBoxBase
+struct FRigVMFunction_MathBoxShift : public FRigVMFunction_MathBoxBase
 {
 	GENERATED_BODY()
 
@@ -145,7 +145,7 @@ struct RIGVM_API FRigVMFunction_MathBoxShift : public FRigVMFunction_MathBoxBase
 
 	/** Execute logic for this rig unit */
 	RIGVM_METHOD()
-	virtual void Execute() override;
+	RIGVM_API virtual void Execute() override;
 
 	UPROPERTY(meta = (Input))
 	FBox Box;
@@ -162,7 +162,7 @@ struct RIGVM_API FRigVMFunction_MathBoxShift : public FRigVMFunction_MathBoxBase
  * Moves the center of the box to a new location
  */
 USTRUCT(meta = (DisplayName = "Move Box To", TemplateName = "Bbox,Translate,Move,BoundingBox"))
-struct RIGVM_API FRigVMFunction_MathBoxMoveTo : public FRigVMFunction_MathBoxBase
+struct FRigVMFunction_MathBoxMoveTo : public FRigVMFunction_MathBoxBase
 {
 	GENERATED_BODY()
 
@@ -174,7 +174,7 @@ struct RIGVM_API FRigVMFunction_MathBoxMoveTo : public FRigVMFunction_MathBoxBas
 
 	/** Execute logic for this rig unit */
 	RIGVM_METHOD()
-	virtual void Execute() override;
+	RIGVM_API virtual void Execute() override;
 
 	UPROPERTY(meta = (Input))
 	FBox Box;
@@ -191,7 +191,7 @@ struct RIGVM_API FRigVMFunction_MathBoxMoveTo : public FRigVMFunction_MathBoxBas
  * Expands the size of the box by a given amount
  */
 USTRUCT(meta = (DisplayName = "Expand Box", TemplateName = "Bbox,Scale,Grow,Shrink,BoundingBox"))
-struct RIGVM_API FRigVMFunction_MathBoxExpand : public FRigVMFunction_MathBoxBase
+struct FRigVMFunction_MathBoxExpand : public FRigVMFunction_MathBoxBase
 {
 	GENERATED_BODY()
 
@@ -203,7 +203,7 @@ struct RIGVM_API FRigVMFunction_MathBoxExpand : public FRigVMFunction_MathBoxBas
 
 	/** Execute logic for this rig unit */
 	RIGVM_METHOD()
-	virtual void Execute() override;
+	RIGVM_API virtual void Execute() override;
 
 	UPROPERTY(meta = (Input))
 	FBox Box;
@@ -220,7 +220,7 @@ struct RIGVM_API FRigVMFunction_MathBoxExpand : public FRigVMFunction_MathBoxBas
  * Transforms the box by a given transform
  */
 USTRUCT(meta = (DisplayName = "Transform Box", TemplateName = "Bbox,BoundingBox"))
-struct RIGVM_API FRigVMFunction_MathBoxTransform : public FRigVMFunction_MathBoxBase
+struct FRigVMFunction_MathBoxTransform : public FRigVMFunction_MathBoxBase
 {
 	GENERATED_BODY()
 
@@ -232,7 +232,7 @@ struct RIGVM_API FRigVMFunction_MathBoxTransform : public FRigVMFunction_MathBox
 
 	/** Execute logic for this rig unit */
 	RIGVM_METHOD()
-	virtual void Execute() override;
+	RIGVM_API virtual void Execute() override;
 
 	UPROPERTY(meta = (Input))
 	FBox Box;
@@ -248,7 +248,7 @@ struct RIGVM_API FRigVMFunction_MathBoxTransform : public FRigVMFunction_MathBox
  * Returns the distance to a given box
  */
 USTRUCT(meta = (DisplayName = "Get Distance to Box", TemplateName = "Bbox,Closest,BoundingBox"))
-struct RIGVM_API FRigVMFunction_MathBoxGetDistance : public FRigVMFunction_MathBoxBase
+struct FRigVMFunction_MathBoxGetDistance : public FRigVMFunction_MathBoxBase
 {
 	GENERATED_BODY()
 
@@ -263,7 +263,7 @@ struct RIGVM_API FRigVMFunction_MathBoxGetDistance : public FRigVMFunction_MathB
 
 	/** Execute logic for this rig unit */
 	RIGVM_METHOD()
-	virtual void Execute() override;
+	RIGVM_API virtual void Execute() override;
 
 	UPROPERTY(meta = (Input))
 	FBox Box;
@@ -286,7 +286,7 @@ struct RIGVM_API FRigVMFunction_MathBoxGetDistance : public FRigVMFunction_MathB
  * Returns true if a point is inside a given box
  */
 USTRUCT(meta = (DisplayName = "Is Inside Box", TemplateName = "Bbox,Contains,Encompasses,BoundingBox"))
-struct RIGVM_API FRigVMFunction_MathBoxIsInside : public FRigVMFunction_MathBoxBase
+struct FRigVMFunction_MathBoxIsInside : public FRigVMFunction_MathBoxBase
 {
 	GENERATED_BODY()
 
@@ -299,7 +299,7 @@ struct RIGVM_API FRigVMFunction_MathBoxIsInside : public FRigVMFunction_MathBoxB
 
 	/** Execute logic for this rig unit */
 	RIGVM_METHOD()
-	virtual void Execute() override;
+	RIGVM_API virtual void Execute() override;
 
 	UPROPERTY(meta = (Input))
 	FBox Box;
@@ -316,7 +316,7 @@ struct RIGVM_API FRigVMFunction_MathBoxIsInside : public FRigVMFunction_MathBoxB
  * Returns the volume of a given box
  */
 USTRUCT(meta = (DisplayName = "Get Box Volume", TemplateName = "Bbox,BoundingBox"))
-struct RIGVM_API FRigVMFunction_MathBoxGetVolume : public FRigVMFunction_MathBoxBase
+struct FRigVMFunction_MathBoxGetVolume : public FRigVMFunction_MathBoxBase
 {
 	GENERATED_BODY()
 
@@ -328,7 +328,7 @@ struct RIGVM_API FRigVMFunction_MathBoxGetVolume : public FRigVMFunction_MathBox
 
 	/** Execute logic for this rig unit */
 	RIGVM_METHOD()
-	virtual void Execute() override;
+	RIGVM_API virtual void Execute() override;
 
 	UPROPERTY(meta = (Input))
 	FBox Box;
diff --git a/Engine/Plugins/Runtime/RigVM/Source/RigVM/Public/RigVMFunctions/Math/RigVMFunction_MathColor.h b/Engine/Plugins/Runtime/RigVM/Source/RigVM/Public/RigVMFunctions/Math/RigVMFunction_MathColor.h
--- a/Engine/Plugins/Runtime/RigVM/Source/RigVM/Public/RigVMFunctions/Math/RigVMFunction_MathColor.h
+++ b/Engine/Plugins/Runtime/RigVM/Source/RigVM/Public/RigVMFunctions/Math/RigVMFunction_MathColor.h
@@ -6,13 +6,13 @@
 #include "RigVMFunction_MathColor.generated.h"
 
 USTRUCT(meta=(Abstract, Category="Math|Color", MenuDescSuffix="(Color)"))
-struct RIGVM_API FRigVMFunction_MathColorBase : public FRigVMFunction_MathBase
+struct FRigVMFunction_MathColorBase : public FRigVMFunction_MathBase
 {
 	GENERATED_BODY()
 };
 
 USTRUCT(meta=(Abstract))
-struct RIGVM_API FRigVMFunction_MathColorBinaryOp : public FRigVMFunction_MathColorBase
+struct FRigVMFunction_MathColorBinaryOp : public FRigVMFunction_MathColorBase
 {
 	GENERATED_BODY()
 
@@ -32,7 +32,7 @@ struct RIGVM_API FRigVMFunction_MathColorBinaryOp : public FRigVMFunction_MathCo
 };
 
 USTRUCT(meta=(Abstract))
-struct RIGVM_API FRigVMFunction_MathColorBinaryAggregateOp : public FRigVMFunction_MathColorBase
+struct FRigVMFunction_MathColorBinaryAggregateOp : public FRigVMFunction_MathColorBase
 {
 	GENERATED_BODY()
 
@@ -55,11 +55,11 @@ struct RIGVM_API FRigVMFunction_MathColorBinaryAggregateOp : public FRigVMFuncti
  * Makes a color from its components
  */
 USTRUCT(meta=(DisplayName="Make Color", Keywords="Make,Construct,Constant"))
-struct RIGVM_API FRigVMFunction_MathColorMake : public FRigVMFunction_MathColorBase
+struct FRigVMFunction_MathColorMake : public FRigVMFunction_MathColorBase
 {
 	GENERATED_BODY()
 	RIGVM_METHOD()
-	virtual void Execute() override;
+	RIGVM_API virtual void Execute() override;
 
 	FRigVMFunction_MathColorMake()
 	{
@@ -88,11 +88,11 @@ struct RIGVM_API FRigVMFunction_MathColorMake : public FRigVMFunction_MathColorB
  * Makes a vector from a single float
  */
 USTRUCT(meta=(DisplayName="From Float", TemplateName="Cast", ExecuteContext="FRigVMExecuteContext", Keywords="Make,Construct"))
-struct RIGVM_API FRigVMFunction_MathColorFromFloat : public FRigVMFunction_MathColorBase
+struct FRigVMFunction_MathColorFromFloat : public FRigVMFunction_MathColorBase
 {
 	GENERATED_BODY()
 	RIGVM_METHOD()
-	virtual void Execute() override;
+	RIGVM_API virtual void Execute() override;
 
 	FRigVMFunction_MathColorFromFloat()
 	{
@@ -111,11 +111,11 @@ struct RIGVM_API FRigVMFunction_MathColorFromFloat : public FRigVMFunction_MathC
  * Makes a vector from a single double
  */
 USTRUCT(meta=(DisplayName="From Double", TemplateName="Cast", ExecuteContext="FRigVMExecuteContext", Keywords="Make,Construct"))
-struct RIGVM_API FRigVMFunction_MathColorFromDouble : public FRigVMFunction_MathColorBase
+struct FRigVMFunction_MathColorFromDouble : public FRigVMFunction_MathColorBase
 {
 	GENERATED_BODY()
 	RIGVM_METHOD()
-	virtual void Execute() override;
+	RIGVM_API virtual void Execute() override;
 
 	FRigVMFunction_MathColorFromDouble()
 	{
@@ -134,31 +134,31 @@ struct RIGVM_API FRigVMFunction_MathColorFromDouble : public FRigVMFunction_Math
  * Returns the sum of the two values
  */
 USTRUCT(meta=(DisplayName="Add", TemplateName="Add", Keywords="Sum,+"))
-struct RIGVM_API FRigVMFunction_MathColorAdd : public FRigVMFunction_MathColorBinaryAggregateOp
+struct FRigVMFunction_MathColorAdd : public FRigVMFunction_MathColorBinaryAggregateOp
 {
 	GENERATED_BODY()
 
 	RIGVM_METHOD()
-	virtual void Execute() override;
+	RIGVM_API virtual void Execute() override;
 };
 
 /**
  * Returns the difference of the two values
  */
 USTRUCT(meta=(DisplayName="Subtract", TemplateName="Subtract", Keywords="-"))
-struct RIGVM_API FRigVMFunction_MathColorSub : public FRigVMFunction_MathColorBinaryOp
+struct FRigVMFunction_MathColorSub : public FRigVMFunction_MathColorBinaryOp
 {
 	GENERATED_BODY()
 
 	RIGVM_METHOD()
-	virtual void Execute() override;
+	RIGVM_API virtual void Execute() override;
 };
 
 /**
  * Returns the product of the two values
  */
 USTRUCT(meta=(DisplayName="Multiply", TemplateName="Multiply", Keywords="Product,*"))
-struct RIGVM_API FRigVMFunction_MathColorMul : public FRigVMFunction_MathColorBinaryAggregateOp
+struct FRigVMFunction_MathColorMul : public FRigVMFunction_MathColorBinaryAggregateOp
 {
 	GENERATED_BODY()
 
@@ -168,14 +168,14 @@ struct RIGVM_API FRigVMFunction_MathColorMul : public FRigVMFunction_MathColorBi
 	}
 
 	RIGVM_METHOD()
-	virtual void Execute() override;
+	RIGVM_API virtual void Execute() override;
 };
 
 /**
  * Linearly interpolates between A and B using the ratio T
  */
 USTRUCT(meta=(DisplayName="Interpolate", TemplateName="Interpolate", Keywords="Lerp,Mix,Blend"))
-struct RIGVM_API FRigVMFunction_MathColorLerp : public FRigVMFunction_MathColorBase
+struct FRigVMFunction_MathColorLerp : public FRigVMFunction_MathColorBase
 {
 	GENERATED_BODY()
 
@@ -187,7 +187,7 @@ struct RIGVM_API FRigVMFunction_MathColorLerp : public FRigVMFunction_MathColorB
 	}
 
 	RIGVM_METHOD()
-	virtual void Execute() override;
+	RIGVM_API virtual void Execute() override;
 
 	UPROPERTY(meta=(Input))
 	FLinearColor A;
