[RigVM]

diff --git a/Engine/Plugins/Runtime/RigVM/Source/RigVM/Public/RigVMFunctions/Math/RigVMFunction_MathDouble.h b/Engine/Plugins/Runtime/RigVM/Source/RigVM/Public/RigVMFunctions/Math/RigVMFunction_MathDouble.h
--- a/Engine/Plugins/Runtime/RigVM/Source/RigVM/Public/RigVMFunctions/Math/RigVMFunction_MathDouble.h
+++ b/Engine/Plugins/Runtime/RigVM/Source/RigVM/Public/RigVMFunctions/Math/RigVMFunction_MathDouble.h
@@ -6,13 +6,13 @@
 #include "RigVMFunction_MathDouble.generated.h"
 
 USTRUCT(meta=(Abstract, Category="Math|Double", MenuDescSuffix="(Double)"))
-struct RIGVM_API FRigVMFunction_MathDoubleBase : public FRigVMFunction_MathBase
+struct FRigVMFunction_MathDoubleBase : public FRigVMFunction_MathBase
 {
 	GENERATED_BODY()
 };
 
 USTRUCT(meta=(Abstract))
-struct RIGVM_API FRigVMFunction_MathDoubleConstant : public FRigVMFunction_MathDoubleBase
+struct FRigVMFunction_MathDoubleConstant : public FRigVMFunction_MathDoubleBase
 {
 	GENERATED_BODY()
 
@@ -26,7 +26,7 @@ struct RIGVM_API FRigVMFunction_MathDoubleConstant : public FRigVMFunction_MathD
 };
 
 USTRUCT(meta=(Abstract))
-struct RIGVM_API FRigVMFunction_MathDoubleUnaryOp : public FRigVMFunction_MathDoubleBase
+struct FRigVMFunction_MathDoubleUnaryOp : public FRigVMFunction_MathDoubleBase
 {
 	GENERATED_BODY()
 
@@ -43,7 +43,7 @@ struct RIGVM_API FRigVMFunction_MathDoubleUnaryOp : public FRigVMFunction_MathDo
 };
 
 USTRUCT(meta=(Abstract))
-struct RIGVM_API FRigVMFunction_MathDoubleBinaryOp : public FRigVMFunction_MathDoubleBase
+struct FRigVMFunction_MathDoubleBinaryOp : public FRigVMFunction_MathDoubleBase
 {
 	GENERATED_BODY()
 
@@ -64,7 +64,7 @@ struct RIGVM_API FRigVMFunction_MathDoubleBinaryOp : public FRigVMFunction_MathD
 
 
 USTRUCT(meta=(Abstract))
-struct RIGVM_API FRigVMFunction_MathDoubleBinaryAggregateOp : public FRigVMFunction_MathDoubleBase
+struct FRigVMFunction_MathDoubleBinaryAggregateOp : public FRigVMFunction_MathDoubleBase
 {
 	GENERATED_BODY()
 
@@ -87,7 +87,7 @@ struct RIGVM_API FRigVMFunction_MathDoubleBinaryAggregateOp : public FRigVMFunct
  * A double constant
  */
 USTRUCT(meta=(DisplayName="Double", Keywords="Make,Construct,Constant"))
-struct RIGVM_API FRigVMFunction_MathDoubleMake : public FRigVMFunction_MathDoubleBase
+struct FRigVMFunction_MathDoubleMake : public FRigVMFunction_MathDoubleBase
 {
 	GENERATED_BODY()
 
@@ -97,7 +97,7 @@ struct RIGVM_API FRigVMFunction_MathDoubleMake : public FRigVMFunction_MathDoubl
 	}
 	
 	RIGVM_METHOD()
-	virtual void Execute() override;
+	RIGVM_API virtual void Execute() override;
 
 	UPROPERTY(meta=(Input, Output))
 	double Value;
@@ -107,79 +107,79 @@ struct RIGVM_API FRigVMFunction_MathDoubleMake : public FRigVMFunction_MathDoubl
  * Returns PI
  */
 USTRUCT(meta=(DisplayName="Pi", TemplateName="Pi"))
-struct RIGVM_API FRigVMFunction_MathDoubleConstPi : public FRigVMFunction_MathDoubleConstant
+struct FRigVMFunction_MathDoubleConstPi : public FRigVMFunction_MathDoubleConstant
 {
 	GENERATED_BODY()
 	
 	RIGVM_METHOD()
-	virtual void Execute() override;
+	RIGVM_API virtual void Execute() override;
 };
 
 /**
  * Returns PI * 0.5
  */
 USTRUCT(meta=(DisplayName="Half Pi", TemplateName="HalfPi"))
-struct RIGVM_API FRigVMFunction_MathDoubleConstHalfPi : public FRigVMFunction_MathDoubleConstant
+struct FRigVMFunction_MathDoubleConstHalfPi : public FRigVMFunction_MathDoubleConstant
 {
 	GENERATED_BODY()
 
 	RIGVM_METHOD()
-	virtual void Execute() override;
+	RIGVM_API virtual void Execute() override;
 };
 
 /**
  * Returns PI * 2.0
  */
 USTRUCT(meta=(DisplayName="Two Pi", TemplateName="TwoPi", Keywords="Tau"))
-struct RIGVM_API FRigVMFunction_MathDoubleConstTwoPi : public FRigVMFunction_MathDoubleConstant
+struct FRigVMFunction_MathDoubleConstTwoPi : public FRigVMFunction_MathDoubleConstant
 {
 	GENERATED_BODY()
 
 	RIGVM_METHOD()
-	virtual void Execute() override;
+	RIGVM_API virtual void Execute() override;
 };
 
 /**
  * Returns E
  */
 USTRUCT(meta=(DisplayName="E", TemplateName="E", Keywords="Euler"))
-struct RIGVM_API FRigVMFunction_MathDoubleConstE : public FRigVMFunction_MathDoubleConstant
+struct FRigVMFunction_MathDoubleConstE : public FRigVMFunction_MathDoubleConstant
 {
 	GENERATED_BODY()
 	
 	RIGVM_METHOD()
-	virtual void Execute() override;
+	RIGVM_API virtual void Execute() override;
 };
 
 /**
  * Returns the sum of the two values
  */
 USTRUCT(meta=(DisplayName="Add", TemplateName="Add", Keywords="Sum,+"))
-struct RIGVM_API FRigVMFunction_MathDoubleAdd : public FRigVMFunction_MathDoubleBinaryAggregateOp
+struct FRigVMFunction_MathDoubleAdd : public FRigVMFunction_MathDoubleBinaryAggregateOp
 {
 	GENERATED_BODY()
 
 	RIGVM_METHOD()
-	virtual void Execute() override;
+	RIGVM_API virtual void Execute() override;
 };
 
 /**
  * Returns the difference of the two values
  */
 USTRUCT(meta=(DisplayName="Subtract", TemplateName="Subtract", Keywords="-"))
-struct RIGVM_API FRigVMFunction_MathDoubleSub : public FRigVMFunction_MathDoubleBinaryOp
+struct FRigVMFunction_MathDoubleSub : public FRigVMFunction_MathDoubleBinaryOp
 {
 	GENERATED_BODY()
 
 	RIGVM_METHOD()
-	virtual void Execute() override;
+	RIGVM_API virtual void Execute() override;
 };
 
 /**
  * Returns the product of the two values
  */
 USTRUCT(meta=(DisplayName="Multiply", TemplateName="Multiply", Keywords="Product,*"))
-struct RIGVM_API FRigVMFunction_MathDoubleMul : public FRigVMFunction_MathDoubleBinaryAggregateOp
+struct FRigVMFunction_MathDoubleMul : public FRigVMFunction_MathDoubleBinaryAggregateOp
 {
 	GENERATED_BODY()
 
@@ -190,14 +190,14 @@ struct RIGVM_API FRigVMFunction_MathDoubleMul : public FRigVMFunction_MathDouble
 	}
 
 	RIGVM_METHOD()
-	virtual void Execute() override;
+	RIGVM_API virtual void Execute() override;
 };
 
 /**
  * Returns the division of the two values
  */
 USTRUCT(meta=(DisplayName="Divide", TemplateName="Divide", Keywords="Division,Divisor,/"))
-struct RIGVM_API FRigVMFunction_MathDoubleDiv : public FRigVMFunction_MathDoubleBinaryOp
+struct FRigVMFunction_MathDoubleDiv : public FRigVMFunction_MathDoubleBinaryOp
 {
 	GENERATED_BODY()
 
@@ -207,14 +207,14 @@ struct RIGVM_API FRigVMFunction_MathDoubleDiv : public FRigVMFunction_MathDouble
 	}
 
 	RIGVM_METHOD()
-	virtual void Execute() override;
+	RIGVM_API virtual void Execute() override;
 };
 
 /**
  * Returns the modulo of the two values
  */
 USTRUCT(meta=(DisplayName="Modulo", TemplateName="Modulo", Keywords="%,fmod"))
-struct RIGVM_API FRigVMFunction_MathDoubleMod : public FRigVMFunction_MathDoubleBinaryOp
+struct FRigVMFunction_MathDoubleMod : public FRigVMFunction_MathDoubleBinaryOp
 {
 	GENERATED_BODY()
 
@@ -225,38 +225,38 @@ struct RIGVM_API FRigVMFunction_MathDoubleMod : public FRigVMFunction_MathDouble
 	}
 
 	RIGVM_METHOD()
-	virtual void Execute() override;
+	RIGVM_API virtual void Execute() override;
 };
 
 /**
  * Returns the smaller of the two values
  */
 USTRUCT(meta=(DisplayName="Minimum", TemplateName="Minimum"))
-struct RIGVM_API FRigVMFunction_MathDoubleMin : public FRigVMFunction_MathDoubleBinaryAggregateOp
+struct FRigVMFunction_MathDoubleMin : public FRigVMFunction_MathDoubleBinaryAggregateOp
 {
 	GENERATED_BODY()
 
 	RIGVM_METHOD()
-	virtual void Execute() override;
+	RIGVM_API virtual void Execute() override;
 };
 
 /**
  * Returns the larger of the two values
  */
 USTRUCT(meta=(DisplayName="Maximum", TemplateName="Maximum"))
-struct RIGVM_API FRigVMFunction_MathDoubleMax : public FRigVMFunction_MathDoubleBinaryAggregateOp
+struct FRigVMFunction_MathDoubleMax : public FRigVMFunction_MathDoubleBinaryAggregateOp
 {
 	GENERATED_BODY()
 
 	RIGVM_METHOD()
-	virtual void Execute() override;
+	RIGVM_API virtual void Execute() override;
 };
 
 /**
  * Returns the value of A raised to the power of B.
  */
 USTRUCT(meta=(DisplayName="Power", TemplateName="Power"))
-struct RIGVM_API FRigVMFunction_MathDoublePow : public FRigVMFunction_MathDoubleBinaryOp
+struct FRigVMFunction_MathDoublePow : public FRigVMFunction_MathDoubleBinaryOp
 {
 	GENERATED_BODY()
 
@@ -267,50 +267,50 @@ struct RIGVM_API FRigVMFunction_MathDoublePow : public FRigVMFunction_MathDouble
 	}
 
 	RIGVM_METHOD()
-	virtual void Execute() override;
+	RIGVM_API virtual void Execute() override;
 };
 
 /**
  * Returns the square root of the given value
  */
 USTRUCT(meta=(DisplayName="Sqrt", TemplateName="Sqrt", Keywords="Root,Square"))
-struct RIGVM_API FRigVMFunction_MathDoubleSqrt : public FRigVMFunction_MathDoubleUnaryOp
+struct FRigVMFunction_MathDoubleSqrt : public FRigVMFunction_MathDoubleUnaryOp
 {
 	GENERATED_BODY()
 
 	RIGVM_METHOD()
-	virtual void Execute() override;
+	RIGVM_API virtual void Execute() override;
 };
 
 /**
  * Returns the negative value
  */
 USTRUCT(meta=(DisplayName="Negate", TemplateName="Negate", Keywords="-,Abs"))
-struct RIGVM_API FRigVMFunction_MathDoubleNegate : public FRigVMFunction_MathDoubleUnaryOp
+struct FRigVMFunction_MathDoubleNegate : public FRigVMFunction_MathDoubleUnaryOp
 {
 	GENERATED_BODY()
 
 	RIGVM_METHOD()
-	virtual void Execute() override;
+	RIGVM_API virtual void Execute() override;
 };
 
 /**
  * Returns the absolute (positive) value
  */
 USTRUCT(meta=(DisplayName="Absolute", TemplateName="Absolute", Keywords="Abs,Neg"))
-struct RIGVM_API FRigVMFunction_MathDoubleAbs : public FRigVMFunction_MathDoubleUnaryOp
+struct FRigVMFunction_MathDoubleAbs : public FRigVMFunction_MathDoubleUnaryOp
 {
 	GENERATED_BODY()
 
 	RIGVM_METHOD()
-	virtual void Execute() override;
+	RIGVM_API virtual void Execute() override;
 };
 
 /**
  * Returns the closest lower full number (integer) of the value
  */
 USTRUCT(meta=(DisplayName="Floor", TemplateName="Floor", Keywords="Round"))
-struct RIGVM_API FRigVMFunction_MathDoubleFloor : public FRigVMFunction_MathDoubleBase
+struct FRigVMFunction_MathDoubleFloor : public FRigVMFunction_MathDoubleBase
 {
 	GENERATED_BODY()
 
@@ -321,7 +321,7 @@ struct RIGVM_API FRigVMFunction_MathDoubleFloor : public FRigVMFunction_MathDoub
 	}
 
 	RIGVM_METHOD()
-	virtual void Execute() override;
+	RIGVM_API virtual void Execute() override;
 
 	UPROPERTY(meta=(Input))
 	double Value;
@@ -337,7 +337,7 @@ struct RIGVM_API FRigVMFunction_MathDoubleFloor : public FRigVMFunction_MathDoub
  * Returns the closest higher full number (integer) of the value
  */
 USTRUCT(meta=(DisplayName="Ceiling", TemplateName="Ceiling", Keywords="Round"))
-struct RIGVM_API FRigVMFunction_MathDoubleCeil : public FRigVMFunction_MathDoubleBase
+struct FRigVMFunction_MathDoubleCeil : public FRigVMFunction_MathDoubleBase
 {
 	GENERATED_BODY()
 
@@ -348,7 +348,7 @@ struct RIGVM_API FRigVMFunction_MathDoubleCeil : public FRigVMFunction_MathDoubl
 	}
 
 	RIGVM_METHOD()
-	virtual void Execute() override;
+	RIGVM_API virtual void Execute() override;
 
 	UPROPERTY(meta=(Input))
 	double Value;
@@ -364,7 +364,7 @@ struct RIGVM_API FRigVMFunction_MathDoubleCeil : public FRigVMFunction_MathDoubl
  * Returns the closest higher full number (integer) of the value
  */
 USTRUCT(meta=(DisplayName="Round", TemplateName="Round"))
-struct RIGVM_API FRigVMFunction_MathDoubleRound : public FRigVMFunction_MathDoubleBase
+struct FRigVMFunction_MathDoubleRound : public FRigVMFunction_MathDoubleBase
 {
 	GENERATED_BODY()
 
@@ -375,7 +375,7 @@ struct RIGVM_API FRigVMFunction_MathDoubleRound : public FRigVMFunction_MathDoub
 	}
 
 	RIGVM_METHOD()
-	virtual void Execute() override;
+	RIGVM_API virtual void Execute() override;
 
 	UPROPERTY(meta=(Input))
 	double Value;
@@ -391,7 +391,7 @@ struct RIGVM_API FRigVMFunction_MathDoubleRound : public FRigVMFunction_MathDoub
  * Returns the double cast to an int (this uses floor)
  */
 USTRUCT(meta=(DisplayName="To Int", TemplateName="Cast", ExecuteContext="FRigVMExecuteContext"))
-struct RIGVM_API FRigVMFunction_MathDoubleToInt : public FRigVMFunction_MathDoubleBase
+struct FRigVMFunction_MathDoubleToInt : public FRigVMFunction_MathDoubleBase
 {
 	GENERATED_BODY()
 
@@ -402,7 +402,7 @@ struct RIGVM_API FRigVMFunction_MathDoubleToInt : public FRigVMFunction_MathDoub
 	}
 
 	RIGVM_METHOD()
-	virtual void Execute() override;
+	RIGVM_API virtual void Execute() override;
 
 	UPROPERTY(meta=(Input))
 	double Value;
@@ -415,7 +415,7 @@ struct RIGVM_API FRigVMFunction_MathDoubleToInt : public FRigVMFunction_MathDoub
  * Returns the double cast to a float
  */
 USTRUCT(meta=(DisplayName="To Float", TemplateName="Cast", ExecuteContext="FRigVMExecuteContext"))
-struct RIGVM_API FRigVMFunction_MathDoubleToFloat : public FRigVMFunction_MathDoubleBase
+struct FRigVMFunction_MathDoubleToFloat : public FRigVMFunction_MathDoubleBase
 {
 	GENERATED_BODY()
 
@@ -426,7 +426,7 @@ struct RIGVM_API FRigVMFunction_MathDoubleToFloat : public FRigVMFunction_MathDo
 	}
 
 	RIGVM_METHOD()
-	virtual void Execute() override;
+	RIGVM_API virtual void Execute() override;
 
 	UPROPERTY(meta=(Input))
 	double Value;
@@ -439,19 +439,19 @@ struct RIGVM_API FRigVMFunction_MathDoubleToFloat : public FRigVMFunction_MathDo
  * Returns the sign of the value (+1 for >= 0.0, -1 for < 0.0)
  */
 USTRUCT(meta=(DisplayName="Sign", TemplateName="Sign"))
-struct RIGVM_API FRigVMFunction_MathDoubleSign : public FRigVMFunction_MathDoubleUnaryOp
+struct FRigVMFunction_MathDoubleSign : public FRigVMFunction_MathDoubleUnaryOp
 {
 	GENERATED_BODY()
 
 	RIGVM_METHOD()
-	virtual void Execute() override;
+	RIGVM_API virtual void Execute() override;
 };
 
 /**
  * Clamps the given value within the range provided by minimum and maximum
  */
 USTRUCT(meta=(DisplayName="Clamp", TemplateName="Clamp", Keywords="Range,Remap"))
-struct RIGVM_API FRigVMFunction_MathDoubleClamp : public FRigVMFunction_MathDoubleBase
+struct FRigVMFunction_MathDoubleClamp : public FRigVMFunction_MathDoubleBase
 {
 	GENERATED_BODY()
 
@@ -461,7 +461,7 @@ struct RIGVM_API FRigVMFunction_MathDoubleClamp : public FRigVMFunction_MathDoub
 	}
 
 	RIGVM_METHOD()
-	virtual void Execute() override;
+	RIGVM_API virtual void Execute() override;
 
 	UPROPERTY(meta=(Input))
 	double Value;
@@ -480,7 +480,7 @@ struct RIGVM_API FRigVMFunction_MathDoubleClamp : public FRigVMFunction_MathDoub
  * Linearly interpolates between A and B using the ratio T
  */
 USTRUCT(meta=(DisplayName="Interpolate", TemplateName="Interpolate", Keywords="Lerp,Mix,Blend"))
-struct RIGVM_API FRigVMFunction_MathDoubleLerp : public FRigVMFunction_MathDoubleBase
+struct FRigVMFunction_MathDoubleLerp : public FRigVMFunction_MathDoubleBase
 {
 	GENERATED_BODY()
 
@@ -491,7 +491,7 @@ struct RIGVM_API FRigVMFunction_MathDoubleLerp : public FRigVMFunction_MathDoubl
 	}
 
 	RIGVM_METHOD()
-	virtual void Execute() override;
+	RIGVM_API virtual void Execute() override;
 
 	UPROPERTY(meta=(Input))
 	double A;
@@ -510,7 +510,7 @@ struct RIGVM_API FRigVMFunction_MathDoubleLerp : public FRigVMFunction_MathDoubl
  * Remaps the given value from a source range to a target range.
  */
 USTRUCT(meta=(DisplayName="Remap", TemplateName="Remap", Keywords="Rescale,Scale"))
-struct RIGVM_API FRigVMFunction_MathDoubleRemap : public FRigVMFunction_MathDoubleBase
+struct FRigVMFunction_MathDoubleRemap : public FRigVMFunction_MathDoubleBase
 {
 	GENERATED_BODY()
 
@@ -522,7 +522,7 @@ struct RIGVM_API FRigVMFunction_MathDoubleRemap : public FRigVMFunction_MathDoub
 	}
 
 	RIGVM_METHOD()
-	virtual void Execute() override;
+	RIGVM_API virtual void Execute() override;
 
 	UPROPERTY(meta=(Input))
 	double Value;
@@ -551,11 +551,11 @@ struct RIGVM_API FRigVMFunction_MathDoubleRemap : public FRigVMFunction_MathDoub
  * Returns true if the value A equals B
  */
 USTRUCT(meta=(DisplayName="Equals", TemplateName="Equals", Keywords="Same,==", Deprecated="5.1"))
-struct RIGVM_API FRigVMFunction_MathDoubleEquals : public FRigVMFunction_MathDoubleBase
+struct FRigVMFunction_MathDoubleEquals : public FRigVMFunction_MathDoubleBase
 {
 	GENERATED_BODY()
 	RIGVM_METHOD()
-	virtual void Execute() override;
+	RIGVM_API virtual void Execute() override;
 
 	FRigVMFunction_MathDoubleEquals()
 	{
@@ -580,11 +580,11 @@ struct RIGVM_API FRigVMFunction_MathDoubleEquals : public FRigVMFunction_MathDou
  * Returns true if the value A does not equal B
  */
 USTRUCT(meta=(DisplayName="Not Equals", TemplateName="NotEquals", Keywords="Different,!=", Deprecated="5.1"))
-struct RIGVM_API FRigVMFunction_MathDoubleNotEquals : public FRigVMFunction_MathDoubleBase
+struct FRigVMFunction_MathDoubleNotEquals : public FRigVMFunction_MathDoubleBase
 {
 	GENERATED_BODY()
 	RIGVM_METHOD()
-	virtual void Execute() override;
+	RIGVM_API virtual void Execute() override;
 
 	FRigVMFunction_MathDoubleNotEquals()
 	{
@@ -609,7 +609,7 @@ struct RIGVM_API FRigVMFunction_MathDoubleNotEquals : public FRigVMFunction_Math
  * Returns true if the value A is greater than B
  */
 USTRUCT(meta=(DisplayName="Greater", TemplateName="Greater", Keywords="Larger,Bigger,>"))
-struct RIGVM_API FRigVMFunction_MathDoubleGreater : public FRigVMFunction_MathDoubleBase
+struct FRigVMFunction_MathDoubleGreater : public FRigVMFunction_MathDoubleBase
 {
 	GENERATED_BODY()
 
@@ -620,7 +620,7 @@ struct RIGVM_API FRigVMFunction_MathDoubleGreater : public FRigVMFunction_MathDo
 	}
 
 	RIGVM_METHOD()
-	virtual void Execute() override;
+	RIGVM_API virtual void Execute() override;
 
 	UPROPERTY(meta=(Input))
 	double A;
@@ -636,7 +636,7 @@ struct RIGVM_API FRigVMFunction_MathDoubleGreater : public FRigVMFunction_MathDo
  * Returns true if the value A is less than B
  */
 USTRUCT(meta=(DisplayName="Less", TemplateName="Less", Keywords="Smaller,<"))
-struct RIGVM_API FRigVMFunction_MathDoubleLess : public FRigVMFunction_MathDoubleBase
+struct FRigVMFunction_MathDoubleLess : public FRigVMFunction_MathDoubleBase
 {
 	GENERATED_BODY()
 	
@@ -647,7 +647,7 @@ struct RIGVM_API FRigVMFunction_MathDoubleLess : public FRigVMFunction_MathDoubl
 	}
 
 	RIGVM_METHOD()
-	virtual void Execute() override;
+	RIGVM_API virtual void Execute() override;
 
 	UPROPERTY(meta=(Input))
 	double A;
@@ -663,7 +663,7 @@ struct RIGVM_API FRigVMFunction_MathDoubleLess : public FRigVMFunction_MathDoubl
  * Returns true if the value A is greater than or equal to B
  */
 USTRUCT(meta=(DisplayName="Greater Equal", TemplateName="GreaterEqual", Keywords="Larger,Bigger,>="))
-struct RIGVM_API FRigVMFunction_MathDoubleGreaterEqual : public FRigVMFunction_MathDoubleBase
+struct FRigVMFunction_MathDoubleGreaterEqual : public FRigVMFunction_MathDoubleBase
 {
 	GENERATED_BODY()
 
@@ -674,7 +674,7 @@ struct RIGVM_API FRigVMFunction_MathDoubleGreaterEqual : public FRigVMFunction_M
 	}
 
 	RIGVM_METHOD()
-	virtual void Execute() override;
+	RIGVM_API virtual void Execute() override;
 
 	UPROPERTY(meta=(Input))
 	double A;
@@ -690,7 +690,7 @@ struct RIGVM_API FRigVMFunction_MathDoubleGreaterEqual : public FRigVMFunction_M
  * Returns true if the value A is less than or equal to B
  */
 USTRUCT(meta=(DisplayName="Less Equal", TemplateName="LessEqual", Keywords="Smaller,<="))
-struct RIGVM_API FRigVMFunction_MathDoubleLessEqual : public FRigVMFunction_MathDoubleBase
+struct FRigVMFunction_MathDoubleLessEqual : public FRigVMFunction_MathDoubleBase
 {
 	GENERATED_BODY()
 
@@ -701,7 +701,7 @@ struct RIGVM_API FRigVMFunction_MathDoubleLessEqual : public FRigVMFunction_Math
 	}
 
 	RIGVM_METHOD()
-	virtual void Execute() override;
+	RIGVM_API virtual void Execute() override;
 
 	UPROPERTY(meta=(Input))
 	double A;
@@ -717,7 +717,7 @@ struct RIGVM_API FRigVMFunction_MathDoubleLessEqual : public FRigVMFunction_Math
  * Returns true if the value is nearly zero
  */
 USTRUCT(meta=(DisplayName="Is Nearly Zero", TemplateName="IsNearlyZero", Keywords="AlmostZero,0"))
-struct RIGVM_API FRigVMFunction_MathDoubleIsNearlyZero : public FRigVMFunction_MathDoubleBase
+struct FRigVMFunction_MathDoubleIsNearlyZero : public FRigVMFunction_MathDoubleBase
 {
 	GENERATED_BODY()
 	
@@ -728,7 +728,7 @@ struct RIGVM_API FRigVMFunction_MathDoubleIsNearlyZero : public FRigVMFunction_M
 	}
 	
 	RIGVM_METHOD()
-	virtual void Execute() override;
+	RIGVM_API virtual void Execute() override;
 
 	UPROPERTY(meta=(Input))
 	double Value;
@@ -744,7 +744,7 @@ struct RIGVM_API FRigVMFunction_MathDoubleIsNearlyZero : public FRigVMFunction_M
  * Returns true if the value A is almost equal to B
  */
 USTRUCT(meta=(DisplayName="Is Nearly Equal", TemplateName="IsNearlyEqual", Keywords="AlmostEqual,=="))
-struct RIGVM_API FRigVMFunction_MathDoubleIsNearlyEqual : public FRigVMFunction_MathDoubleBase
+struct FRigVMFunction_MathDoubleIsNearlyEqual : public FRigVMFunction_MathDoubleBase
 {
 	GENERATED_BODY()
 
@@ -755,7 +755,7 @@ struct RIGVM_API FRigVMFunction_MathDoubleIsNearlyEqual : public FRigVMFunction_
 	}
 
 	RIGVM_METHOD()
-	virtual void Execute() override;
+	RIGVM_API virtual void Execute() override;
 
 	UPROPERTY(meta=(Input))
 	double A;
@@ -774,115 +774,115 @@ struct RIGVM_API FRigVMFunction_MathDoubleIsNearlyEqual : public FRigVMFunction_
  * Returns the degrees of a given value in radians
  */
 USTRUCT(meta=(DisplayName="Degrees", TemplateName="Degrees"))
-struct RIGVM_API FRigVMFunction_MathDoubleDeg : public FRigVMFunction_MathDoubleUnaryOp
+struct FRigVMFunction_MathDoubleDeg : public FRigVMFunction_MathDoubleUnaryOp
 {
 	GENERATED_BODY()
 
 	RIGVM_METHOD()
-	virtual void Execute() override;
+	RIGVM_API virtual void Execute() override;
 };
 
 /**
  * Returns the radians of a given value in degrees
  */
 USTRUCT(meta=(DisplayName="Radians", TemplateName="Radians"))
-struct RIGVM_API FRigVMFunction_MathDoubleRad : public FRigVMFunction_MathDoubleUnaryOp
+struct FRigVMFunction_MathDoubleRad : public FRigVMFunction_MathDoubleUnaryOp
 {
 	GENERATED_BODY()
 
 	RIGVM_METHOD()
-	virtual void Execute() override;
+	RIGVM_API virtual void Execute() override;
 };
 
 /**
  * Returns the sinus value of the given value (in radians)
  */
 USTRUCT(meta=(DisplayName="Sin", TemplateName="Sin"))
-struct RIGVM_API FRigVMFunction_MathDoubleSin : public FRigVMFunction_MathDoubleUnaryOp
+struct FRigVMFunction_MathDoubleSin : public FRigVMFunction_MathDoubleUnaryOp
 {
 	GENERATED_BODY()
 
 	RIGVM_METHOD()
-	virtual void Execute() override;
+	RIGVM_API virtual void Execute() override;
 };
 
 /**
  * Returns the cosinus value of the given value (in radians)
  */
 USTRUCT(meta=(DisplayName="Cos", TemplateName="Cos"))
-struct RIGVM_API FRigVMFunction_MathDoubleCos : public FRigVMFunction_MathDoubleUnaryOp
+struct FRigVMFunction_MathDoubleCos : public FRigVMFunction_MathDoubleUnaryOp
 {
 	GENERATED_BODY()
 
 	RIGVM_METHOD()
-	virtual void Execute() override;
+	RIGVM_API virtual void Execute() override;
 };
 
 /**
  * Returns the tangens value of the given value (in radians)
  */
 USTRUCT(meta=(DisplayName="Tan", TemplateName="Tan"))
-struct RIGVM_API FRigVMFunction_MathDoubleTan : public FRigVMFunction_MathDoubleUnaryOp
+struct FRigVMFunction_MathDoubleTan : public FRigVMFunction_MathDoubleUnaryOp
 {
 	GENERATED_BODY()
 
 	RIGVM_METHOD()
-	virtual void Execute() override;
+	RIGVM_API virtual void Execute() override;
 };
 
 /**
  * Returns the inverse sinus value (in radians) of the given value
  */
 USTRUCT(meta=(DisplayName="Asin", TemplateName="Asin", Keywords="Arcsin"))
-struct RIGVM_API FRigVMFunction_MathDoubleAsin : public FRigVMFunction_MathDoubleUnaryOp
+struct FRigVMFunction_MathDoubleAsin : public FRigVMFunction_MathDoubleUnaryOp
 {
 	GENERATED_BODY()
 
 	RIGVM_METHOD()
-	virtual void Execute() override;
+	RIGVM_API virtual void Execute() override;
 };
 
 /**
  * Returns the inverse cosinus value (in radians) of the given value
  */
 USTRUCT(meta=(DisplayName="Acos", TemplateName="Acos", Keywords="Arccos"))
-struct RIGVM_API FRigVMFunction_MathDoubleAcos : public FRigVMFunction_MathDoubleUnaryOp
+struct FRigVMFunction_MathDoubleAcos : public FRigVMFunction_MathDoubleUnaryOp
 {
 	GENERATED_BODY()
 
 	RIGVM_METHOD()
-	virtual void Execute() override;
+	RIGVM_API virtual void Execute() override;
 };
 
 /**
  * Returns the inverse tangens value (in radians) of the given value
  */
 USTRUCT(meta=(DisplayName="Atan", TemplateName="Atan", Keywords="Arctan"))
-struct RIGVM_API FRigVMFunction_MathDoubleAtan : public FRigVMFunction_MathDoubleUnaryOp
+struct FRigVMFunction_MathDoubleAtan : public FRigVMFunction_MathDoubleUnaryOp
 {
 	GENERATED_BODY()
 
 	RIGVM_METHOD()
-	virtual void Execute() override;
+	RIGVM_API virtual void Execute() override;
 };
 
 /**
  * Returns the arctangent of the specified A and B coordinates.
  */
 USTRUCT(meta=(DisplayName="Atan2", TemplateName="Atan2", Keywords="Arctan"))
-struct RIGVM_API FRigVMFunction_MathDoubleAtan2 : public FRigVMFunction_MathDoubleBinaryOp
+struct FRigVMFunction_MathDoubleAtan2 : public FRigVMFunction_MathDoubleBinaryOp
 {
 	GENERATED_BODY()
 
 	RIGVM_METHOD()
-	virtual void Execute() override;
+	RIGVM_API virtual void Execute() override;
 };
 
 /**
  * Computes the angles alpha, beta and gamma (in radians) between the three sides A, B and C
  */
 USTRUCT(meta = (DisplayName = "Law Of Cosine", TemplateName="LawOfCosine"))
-struct RIGVM_API FRigVMFunction_MathDoubleLawOfCosine : public FRigVMFunction_MathDoubleBase
+struct FRigVMFunction_MathDoubleLawOfCosine : public FRigVMFunction_MathDoubleBase
 {
 	GENERATED_BODY()
 
@@ -893,7 +893,7 @@ struct RIGVM_API FRigVMFunction_MathDoubleLawOfCosine : public FRigVMFunction_Ma
 	}
 
 	RIGVM_METHOD()
-	virtual void Execute() override;
+	RIGVM_API virtual void Execute() override;
 
 	UPROPERTY(meta = (Input))
 	double A;
@@ -921,19 +921,19 @@ struct RIGVM_API FRigVMFunction_MathDoubleLawOfCosine : public FRigVMFunction_Ma
  * Computes the base-e exponential of the given value 
  */
 USTRUCT(meta = (DisplayName = "Exponential", TemplateName="Exponential"))
-struct RIGVM_API FRigVMFunction_MathDoubleExponential : public FRigVMFunction_MathDoubleUnaryOp
+struct FRigVMFunction_MathDoubleExponential : public FRigVMFunction_MathDoubleUnaryOp
 {
 	GENERATED_BODY()
 
 	RIGVM_METHOD()
-	virtual void Execute() override;
+	RIGVM_API virtual void Execute() override;
 };
 
 /**
  * Returns the sum of the given array
  */
 USTRUCT(meta = (DisplayName = "Array Sum", TemplateName = "ArraySum"))
-struct RIGVM_API FRigVMFunction_MathDoubleArraySum : public FRigVMFunction_MathDoubleBase
+struct FRigVMFunction_MathDoubleArraySum : public FRigVMFunction_MathDoubleBase
 {
 	GENERATED_BODY()
 
@@ -944,7 +944,7 @@ struct RIGVM_API FRigVMFunction_MathDoubleArraySum : public FRigVMFunction_MathD
 
 	/** Execute logic for this rig unit */
 	RIGVM_METHOD()
-	virtual void Execute() override;
+	RIGVM_API virtual void Execute() override;
 
 	UPROPERTY(meta = (Input))
 	TArray<double> Array;
@@ -957,7 +957,7 @@ struct RIGVM_API FRigVMFunction_MathDoubleArraySum : public FRigVMFunction_MathD
  * Returns the average of the given array
  */
 USTRUCT(meta = (DisplayName = "Array Average", TemplateName = "ArrayAverage"))
-struct RIGVM_API FRigVMFunction_MathDoubleArrayAverage : public FRigVMFunction_MathDoubleBase
+struct FRigVMFunction_MathDoubleArrayAverage : public FRigVMFunction_MathDoubleBase
 {
 	GENERATED_BODY()
 
@@ -968,7 +968,7 @@ struct RIGVM_API FRigVMFunction_MathDoubleArrayAverage : public FRigVMFunction_M
 
 	/** Execute logic for this rig unit */
 	RIGVM_METHOD()
-	virtual void Execute() override;
+	RIGVM_API virtual void Execute() override;
 
 	UPROPERTY(meta = (Input))
 	TArray<double> Array;
diff --git a/Engine/Plugins/Runtime/RigVM/Source/RigVM/Public/RigVMFunctions/Math/RigVMFunction_MathFloat.h b/Engine/Plugins/Runtime/RigVM/Source/RigVM/Public/RigVMFunctions/Math/RigVMFunction_MathFloat.h
--- a/Engine/Plugins/Runtime/RigVM/Source/RigVM/Public/RigVMFunctions/Math/RigVMFunction_MathFloat.h
+++ b/Engine/Plugins/Runtime/RigVM/Source/RigVM/Public/RigVMFunctions/Math/RigVMFunction_MathFloat.h
@@ -5,14 +5,16 @@
 #include "RigVMFunction_MathBase.h"
 #include "RigVMFunction_MathFloat.generated.h"
 
+#define UE_API RIGVM_API
+
 USTRUCT(meta=(Abstract, Category="Math|Float", MenuDescSuffix="(Float)"))
-struct RIGVM_API FRigVMFunction_MathFloatBase : public FRigVMFunction_MathBase
+struct FRigVMFunction_MathFloatBase : public FRigVMFunction_MathBase
 {
 	GENERATED_BODY()
 };
 
 USTRUCT(meta=(Abstract))
-struct RIGVM_API FRigVMFunction_MathFloatConstant : public FRigVMFunction_MathFloatBase
+struct FRigVMFunction_MathFloatConstant : public FRigVMFunction_MathFloatBase
 {
 	GENERATED_BODY()
 
@@ -26,7 +28,7 @@ struct RIGVM_API FRigVMFunction_MathFloatConstant : public FRigVMFunction_MathFl
 };
 
 USTRUCT(meta=(Abstract))
-struct RIGVM_API FRigVMFunction_MathFloatUnaryOp : public FRigVMFunction_MathFloatBase
+struct FRigVMFunction_MathFloatUnaryOp : public FRigVMFunction_MathFloatBase
 {
 	GENERATED_BODY()
 
@@ -43,7 +45,7 @@ struct RIGVM_API FRigVMFunction_MathFloatUnaryOp : public FRigVMFunction_MathFlo
 };
 
 USTRUCT(meta=(Abstract))
-struct RIGVM_API FRigVMFunction_MathFloatBinaryOp : public FRigVMFunction_MathFloatBase
+struct FRigVMFunction_MathFloatBinaryOp : public FRigVMFunction_MathFloatBase
 {
 	GENERATED_BODY()
 
@@ -63,7 +65,7 @@ struct RIGVM_API FRigVMFunction_MathFloatBinaryOp : public FRigVMFunction_MathFl
 };
 
 USTRUCT(meta=(Abstract))
-struct RIGVM_API FRigVMFunction_MathFloatBinaryAggregateOp : public FRigVMFunction_MathFloatBase
+struct FRigVMFunction_MathFloatBinaryAggregateOp : public FRigVMFunction_MathFloatBase
 {
 	GENERATED_BODY()
 
@@ -86,7 +88,7 @@ struct RIGVM_API FRigVMFunction_MathFloatBinaryAggregateOp : public FRigVMFuncti
  * A float constant
  */
 USTRUCT(meta=(DisplayName="Float", Keywords="Make,Construct,Constant"))
-struct RIGVM_API FRigVMFunction_MathFloatMake : public FRigVMFunction_MathFloatBase
+struct FRigVMFunction_MathFloatMake : public FRigVMFunction_MathFloatBase
 {
 	GENERATED_BODY()
 
@@ -96,7 +98,7 @@ struct RIGVM_API FRigVMFunction_MathFloatMake : public FRigVMFunction_MathFloatB
 	}
 	
 	RIGVM_METHOD()
-	virtual void Execute() override;
+	UE_API virtual void Execute() override;
 
 	UPROPERTY(meta=(Input, Output))
 	float Value;
@@ -106,79 +108,79 @@ struct RIGVM_API FRigVMFunction_MathFloatMake : public FRigVMFunction_MathFloatB
  * Returns PI
  */
 USTRUCT(meta=(DisplayName="Pi", TemplateName="Pi"))
-struct RIGVM_API FRigVMFunction_MathFloatConstPi : public FRigVMFunction_MathFloatConstant
+struct FRigVMFunction_MathFloatConstPi : public FRigVMFunction_MathFloatConstant
 {
 	GENERATED_BODY()
 	
 	RIGVM_METHOD()
-	virtual void Execute() override;
+	UE_API virtual void Execute() override;
 };
 
 /**
  * Returns PI * 0.5
  */
 USTRUCT(meta=(DisplayName="Half Pi", TemplateName="HalfPi"))
-struct RIGVM_API FRigVMFunction_MathFloatConstHalfPi : public FRigVMFunction_MathFloatConstant
+struct FRigVMFunction_MathFloatConstHalfPi : public FRigVMFunction_MathFloatConstant
 {
 	GENERATED_BODY()
 
 	RIGVM_METHOD()
-	virtual void Execute() override;
+	UE_API virtual void Execute() override;
 };
 
 /**
  * Returns PI * 2.0
  */
 USTRUCT(meta=(DisplayName="Two Pi", TemplateName="TwoPi", Keywords="Tau"))
-struct RIGVM_API FRigVMFunction_MathFloatConstTwoPi : public FRigVMFunction_MathFloatConstant
+struct FRigVMFunction_MathFloatConstTwoPi : public FRigVMFunction_MathFloatConstant
 {
 	GENERATED_BODY()
 
 	RIGVM_METHOD()
-	virtual void Execute() override;
+	UE_API virtual void Execute() override;
 };
 
 /**
  * Returns E
  */
 USTRUCT(meta=(DisplayName="E", TemplateName="E", Keywords="Euler"))
-struct RIGVM_API FRigVMFunction_MathFloatConstE : public FRigVMFunction_MathFloatConstant
+struct FRigVMFunction_MathFloatConstE : public FRigVMFunction_MathFloatConstant
 {
 	GENERATED_BODY()
 	
 	RIGVM_METHOD()
-	virtual void Execute() override;
+	UE_API virtual void Execute() override;
 };
 
 /**
  * Returns the sum of the two values
  */
 USTRUCT(meta=(DisplayName="Add", TemplateName="Add", Keywords="Sum,+"))
-struct RIGVM_API FRigVMFunction_MathFloatAdd : public FRigVMFunction_MathFloatBinaryAggregateOp
+struct FRigVMFunction_MathFloatAdd : public FRigVMFunction_MathFloatBinaryAggregateOp
 {
 	GENERATED_BODY()
 
 	RIGVM_METHOD()
-	virtual void Execute() override;
+	UE_API virtual void Execute() override;
 };
 
 /**
  * Returns the difference of the two values
  */
 USTRUCT(meta=(DisplayName="Subtract", TemplateName="Subtract", Keywords="-"))
-struct RIGVM_API FRigVMFunction_MathFloatSub : public FRigVMFunction_MathFloatBinaryOp
+struct FRigVMFunction_MathFloatSub : public FRigVMFunction_MathFloatBinaryOp
 {
 	GENERATED_BODY()
 
 	RIGVM_METHOD()
-	virtual void Execute() override;
+	UE_API virtual void Execute() override;
 };
 
 /**
  * Returns the product of the two values
  */
 USTRUCT(meta=(DisplayName="Multiply", TemplateName="Multiply", Keywords="Product,*"))
-struct RIGVM_API FRigVMFunction_MathFloatMul : public FRigVMFunction_MathFloatBinaryAggregateOp
+struct FRigVMFunction_MathFloatMul : public FRigVMFunction_MathFloatBinaryAggregateOp
 {
 	GENERATED_BODY()
 
@@ -189,14 +191,14 @@ struct RIGVM_API FRigVMFunction_MathFloatMul : public FRigVMFunction_MathFloatBi
 	}
 
 	RIGVM_METHOD()
-	virtual void Execute() override;
+	UE_API virtual void Execute() override;
 };
 
 /**
  * Returns the division of the two values
  */
 USTRUCT(meta=(DisplayName="Divide", TemplateName="Divide", Keywords="Division,Divisor,/"))
-struct RIGVM_API FRigVMFunction_MathFloatDiv : public FRigVMFunction_MathFloatBinaryOp
+struct FRigVMFunction_MathFloatDiv : public FRigVMFunction_MathFloatBinaryOp
 {
 	GENERATED_BODY()
 
@@ -206,14 +208,14 @@ struct RIGVM_API FRigVMFunction_MathFloatDiv : public FRigVMFunction_MathFloatBi
 	}
 
 	RIGVM_METHOD()
-	virtual void Execute() override;
+	UE_API virtual void Execute() override;
 };
 
 /**
  * Returns the modulo of the two values
  */
 USTRUCT(meta=(DisplayName="Modulo", TemplateName="Modulo", Keywords="%,fmod"))
-struct RIGVM_API FRigVMFunction_MathFloatMod : public FRigVMFunction_MathFloatBinaryOp
+struct FRigVMFunction_MathFloatMod : public FRigVMFunction_MathFloatBinaryOp
 {
 	GENERATED_BODY()
 
@@ -224,38 +226,38 @@ struct RIGVM_API FRigVMFunction_MathFloatMod : public FRigVMFunction_MathFloatBi
 	}
 
 	RIGVM_METHOD()
-	virtual void Execute() override;
+	UE_API virtual void Execute() override;
 };
 
 /**
  * Returns the smaller of the two values
  */
 USTRUCT(meta=(DisplayName="Minimum", TemplateName="Minimum"))
-struct RIGVM_API FRigVMFunction_MathFloatMin : public FRigVMFunction_MathFloatBinaryAggregateOp
+struct FRigVMFunction_MathFloatMin : public FRigVMFunction_MathFloatBinaryAggregateOp
 {
 	GENERATED_BODY()
 
 	RIGVM_METHOD()
-	virtual void Execute() override;
+	UE_API virtual void Execute() override;
 };
 
 /**
  * Returns the larger of the two values
  */
 USTRUCT(meta=(DisplayName="Maximum", TemplateName="Maximum"))
-struct RIGVM_API FRigVMFunction_MathFloatMax : public FRigVMFunction_MathFloatBinaryAggregateOp
+struct FRigVMFunction_MathFloatMax : public FRigVMFunction_MathFloatBinaryAggregateOp
 {
 	GENERATED_BODY()
 
 	RIGVM_METHOD()
-	virtual void Execute() override;
+	UE_API virtual void Execute() override;
 };
 
 /**
  * Returns the value of A raised to the power of B.
  */
 USTRUCT(meta=(DisplayName="Power", TemplateName="Power"))
-struct RIGVM_API FRigVMFunction_MathFloatPow : public FRigVMFunction_MathFloatBinaryOp
+struct FRigVMFunction_MathFloatPow : public FRigVMFunction_MathFloatBinaryOp
 {
 	GENERATED_BODY()
 
@@ -266,50 +268,50 @@ struct RIGVM_API FRigVMFunction_MathFloatPow : public FRigVMFunction_MathFloatBi
 	}
 
 	RIGVM_METHOD()
-	virtual void Execute() override;
+	UE_API virtual void Execute() override;
 };
 
 /**
  * Returns the square root of the given value
  */
 USTRUCT(meta=(DisplayName="Sqrt", TemplateName="Sqrt", Keywords="Root,Square"))
-struct RIGVM_API FRigVMFunction_MathFloatSqrt : public FRigVMFunction_MathFloatUnaryOp
+struct FRigVMFunction_MathFloatSqrt : public FRigVMFunction_MathFloatUnaryOp
 {
 	GENERATED_BODY()
 
 	RIGVM_METHOD()
-	virtual void Execute() override;
+	UE_API virtual void Execute() override;
 };
 
 /**
  * Returns the negative value
  */
 USTRUCT(meta=(DisplayName="Negate", TemplateName="Negate", Keywords="-,Abs"))
-struct RIGVM_API FRigVMFunction_MathFloatNegate : public FRigVMFunction_MathFloatUnaryOp
+struct FRigVMFunction_MathFloatNegate : public FRigVMFunction_MathFloatUnaryOp
 {
 	GENERATED_BODY()
 
 	RIGVM_METHOD()
-	virtual void Execute() override;
+	UE_API virtual void Execute() override;
 };
 
 /**
  * Returns the absolute (positive) value
  */
 USTRUCT(meta=(DisplayName="Absolute", TemplateName="Absolute", Keywords="Abs,Neg"))
-struct RIGVM_API FRigVMFunction_MathFloatAbs : public FRigVMFunction_MathFloatUnaryOp
+struct FRigVMFunction_MathFloatAbs : public FRigVMFunction_MathFloatUnaryOp
 {
 	GENERATED_BODY()
 
 	RIGVM_METHOD()
-	virtual void Execute() override;
+	UE_API virtual void Execute() override;
 };
 
 /**
  * Returns the closest lower full number (integer) of the value
  */
 USTRUCT(meta=(DisplayName="Floor", TemplateName="Floor", Keywords="Round"))
-struct RIGVM_API FRigVMFunction_MathFloatFloor : public FRigVMFunction_MathFloatBase
+struct FRigVMFunction_MathFloatFloor : public FRigVMFunction_MathFloatBase
 {
 	GENERATED_BODY()
 
@@ -320,7 +322,7 @@ struct RIGVM_API FRigVMFunction_MathFloatFloor : public FRigVMFunction_MathFloat
 	}
 
 	RIGVM_METHOD()
-	virtual void Execute() override;
+	UE_API virtual void Execute() override;
 
 	UPROPERTY(meta=(Input))
 	float Value;
@@ -336,7 +338,7 @@ struct RIGVM_API FRigVMFunction_MathFloatFloor : public FRigVMFunction_MathFloat
  * Returns the closest higher full number (integer) of the value
  */
 USTRUCT(meta=(DisplayName="Ceiling", TemplateName="Ceiling", Keywords="Round"))
-struct RIGVM_API FRigVMFunction_MathFloatCeil : public FRigVMFunction_MathFloatBase
+struct FRigVMFunction_MathFloatCeil : public FRigVMFunction_MathFloatBase
 {
 	GENERATED_BODY()
 
@@ -347,7 +349,7 @@ struct RIGVM_API FRigVMFunction_MathFloatCeil : public FRigVMFunction_MathFloatB
 	}
 
 	RIGVM_METHOD()
-	virtual void Execute() override;
+	UE_API virtual void Execute() override;
 
 	UPROPERTY(meta=(Input))
 	float Value;
@@ -363,7 +365,7 @@ struct RIGVM_API FRigVMFunction_MathFloatCeil : public FRigVMFunction_MathFloatB
  * Returns the closest higher full number (integer) of the value
  */
 USTRUCT(meta=(DisplayName="Round", TemplateName="Round"))
-struct RIGVM_API FRigVMFunction_MathFloatRound : public FRigVMFunction_MathFloatBase
+struct FRigVMFunction_MathFloatRound : public FRigVMFunction_MathFloatBase
 {
 	GENERATED_BODY()
 
@@ -374,7 +376,7 @@ struct RIGVM_API FRigVMFunction_MathFloatRound : public FRigVMFunction_MathFloat
 	}
 
 	RIGVM_METHOD()
-	virtual void Execute() override;
+	UE_API virtual void Execute() override;
 
 	UPROPERTY(meta=(Input))
 	float Value;
@@ -390,7 +392,7 @@ struct RIGVM_API FRigVMFunction_MathFloatRound : public FRigVMFunction_MathFloat
  * Returns the float cast to an int (this uses floor)
  */
 USTRUCT(meta=(DisplayName="To Int", TemplateName="Cast", ExecuteContext="FRigVMExecuteContext"))
-struct RIGVM_API FRigVMFunction_MathFloatToInt : public FRigVMFunction_MathFloatBase
+struct FRigVMFunction_MathFloatToInt : public FRigVMFunction_MathFloatBase
 {
 	GENERATED_BODY()
 
@@ -401,7 +403,7 @@ struct RIGVM_API FRigVMFunction_MathFloatToInt : public FRigVMFunction_MathFloat
 	}
 
 	RIGVM_METHOD()
-	virtual void Execute() override;
+	UE_API virtual void Execute() override;
 
 	UPROPERTY(meta=(Input))
 	float Value;
@@ -414,7 +416,7 @@ struct RIGVM_API FRigVMFunction_MathFloatToInt : public FRigVMFunction_MathFloat
  * Returns the float cast to a double
  */
 USTRUCT(meta=(DisplayName="To Double", TemplateName="Cast", ExecuteContext="FRigVMExecuteContext"))
-struct RIGVM_API FRigVMFunction_MathFloatToDouble : public FRigVMFunction_MathFloatBase
+struct FRigVMFunction_MathFloatToDouble : public FRigVMFunction_MathFloatBase
 {
 	GENERATED_BODY()
 
@@ -425,7 +427,7 @@ struct RIGVM_API FRigVMFunction_MathFloatToDouble : public FRigVMFunction_MathFl
 	}
 
 	RIGVM_METHOD()
-	virtual void Execute() override;
+	UE_API virtual void Execute() override;
 
 	UPROPERTY(meta=(Input))
 	float Value;
@@ -438,19 +440,19 @@ struct RIGVM_API FRigVMFunction_MathFloatToDouble : public FRigVMFunction_MathFl
  * Returns the sign of the value (+1 for >= 0.f, -1 for < 0.f)
  */
 USTRUCT(meta=(DisplayName="Sign", TemplateName="Sign"))
-struct RIGVM_API FRigVMFunction_MathFloatSign : public FRigVMFunction_MathFloatUnaryOp
+struct FRigVMFunction_MathFloatSign : public FRigVMFunction_MathFloatUnaryOp
 {
 	GENERATED_BODY()
 
 	RIGVM_METHOD()
-	virtual void Execute() override;
+	UE_API virtual void Execute() override;
 };
 
 /**
  * Clamps the given value within the range provided by minimum and maximum
  */
 USTRUCT(meta=(DisplayName="Clamp", TemplateName="Clamp", Keywords="Range,Remap"))
-struct RIGVM_API FRigVMFunction_MathFloatClamp : public FRigVMFunction_MathFloatBase
+struct FRigVMFunction_MathFloatClamp : public FRigVMFunction_MathFloatBase
 {
 	GENERATED_BODY()
 
@@ -460,7 +462,7 @@ struct RIGVM_API FRigVMFunction_MathFloatClamp : public FRigVMFunction_MathFloat
 	}
 
 	RIGVM_METHOD()
-	virtual void Execute() override;
+	UE_API virtual void Execute() override;
 
 	UPROPERTY(meta=(Input))
 	float Value;
@@ -479,7 +481,7 @@ struct RIGVM_API FRigVMFunction_MathFloatClamp : public FRigVMFunction_MathFloat
  * Linearly interpolates between A and B using the ratio T
  */
 USTRUCT(meta=(DisplayName="Interpolate", TemplateName="Interpolate", Keywords="Lerp,Mix,Blend"))
-struct RIGVM_API FRigVMFunction_MathFloatLerp : public FRigVMFunction_MathFloatBase
+struct FRigVMFunction_MathFloatLerp : public FRigVMFunction_MathFloatBase
 {
 	GENERATED_BODY()
 
@@ -490,7 +492,7 @@ struct RIGVM_API FRigVMFunction_MathFloatLerp : public FRigVMFunction_MathFloatB
 	}
 
 	RIGVM_METHOD()
-	virtual void Execute() override;
+	UE_API virtual void Execute() override;
 
 	UPROPERTY(meta=(Input))
 	float A;
@@ -509,7 +511,7 @@ struct RIGVM_API FRigVMFunction_MathFloatLerp : public FRigVMFunction_MathFloatB
  * Remaps the given value from a source range to a target range.
  */
 USTRUCT(meta=(DisplayName="Remap", TemplateName="Remap", Keywords="Rescale,Scale"))
-struct RIGVM_API FRigVMFunction_MathFloatRemap : public FRigVMFunction_MathFloatBase
+struct FRigVMFunction_MathFloatRemap : public FRigVMFunction_MathFloatBase
 {
 	GENERATED_BODY()
 
@@ -521,7 +523,7 @@ struct RIGVM_API FRigVMFunction_MathFloatRemap : public FRigVMFunction_MathFloat
 	}
 
 	RIGVM_METHOD()
-	virtual void Execute() override;
+	UE_API virtual void Execute() override;
 
 	UPROPERTY(meta=(Input))
 	float Value;
@@ -550,11 +552,11 @@ struct RIGVM_API FRigVMFunction_MathFloatRemap : public FRigVMFunction_MathFloat
  * Returns true if the value A equals B
  */
 USTRUCT(meta=(DisplayName="Equals", TemplateName="Equals", Keywords="Same,==", Deprecated="5.1"))
-struct RIGVM_API FRigVMFunction_MathFloatEquals : public FRigVMFunction_MathFloatBase
+struct FRigVMFunction_MathFloatEquals : public FRigVMFunction_MathFloatBase
 {
 	GENERATED_BODY()
 	RIGVM_METHOD()
-	virtual void Execute() override;
+	UE_API virtual void Execute() override;
 
 	FRigVMFunction_MathFloatEquals()
 	{
@@ -579,11 +581,11 @@ struct RIGVM_API FRigVMFunction_MathFloatEquals : public FRigVMFunction_MathFloa
  * Returns true if the value A does not equal B
  */
 USTRUCT(meta=(DisplayName="Not Equals", TemplateName="NotEquals", Keywords="Different,!=", Deprecated="5.1"))
-struct RIGVM_API FRigVMFunction_MathFloatNotEquals : public FRigVMFunction_MathFloatBase
+struct FRigVMFunction_MathFloatNotEquals : public FRigVMFunction_MathFloatBase
 {
 	GENERATED_BODY()
 	RIGVM_METHOD()
-	virtual void Execute() override;
+	UE_API virtual void Execute() override;
 
 	FRigVMFunction_MathFloatNotEquals()
 	{
@@ -608,7 +610,7 @@ struct RIGVM_API FRigVMFunction_MathFloatNotEquals : public FRigVMFunction_MathF
  * Returns true if the value A is greater than B
  */
 USTRUCT(meta=(DisplayName="Greater", TemplateName="Greater", Keywords="Larger,Bigger,>"))
-struct RIGVM_API FRigVMFunction_MathFloatGreater : public FRigVMFunction_MathFloatBase
+struct FRigVMFunction_MathFloatGreater : public FRigVMFunction_MathFloatBase
 {
 	GENERATED_BODY()
 
@@ -619,7 +621,7 @@ struct RIGVM_API FRigVMFunction_MathFloatGreater : public FRigVMFunction_MathFlo
 	}
 
 	RIGVM_METHOD()
-	virtual void Execute() override;
+	UE_API virtual void Execute() override;
 
 	UPROPERTY(meta=(Input))
 	float A;
@@ -635,7 +637,7 @@ struct RIGVM_API FRigVMFunction_MathFloatGreater : public FRigVMFunction_MathFlo
  * Returns true if the value A is less than B
  */
 USTRUCT(meta=(DisplayName="Less", TemplateName="Less", Keywords="Smaller,<"))
-struct RIGVM_API FRigVMFunction_MathFloatLess : public FRigVMFunction_MathFloatBase
+struct FRigVMFunction_MathFloatLess : public FRigVMFunction_MathFloatBase
 {
 	GENERATED_BODY()
 	
@@ -646,7 +648,7 @@ struct RIGVM_API FRigVMFunction_MathFloatLess : public FRigVMFunction_MathFloatB
 	}
 
 	RIGVM_METHOD()
-	virtual void Execute() override;
+	UE_API virtual void Execute() override;
 
 	UPROPERTY(meta=(Input))
 	float A;
@@ -662,7 +664,7 @@ struct RIGVM_API FRigVMFunction_MathFloatLess : public FRigVMFunction_MathFloatB
  * Returns true if the value A is greater than or equal to B
  */
 USTRUCT(meta=(DisplayName="Greater Equal", TemplateName="GreaterEqual", Keywords="Larger,Bigger,>="))
-struct RIGVM_API FRigVMFunction_MathFloatGreaterEqual : public FRigVMFunction_MathFloatBase
+struct FRigVMFunction_MathFloatGreaterEqual : public FRigVMFunction_MathFloatBase
 {
 	GENERATED_BODY()
 
@@ -673,7 +675,7 @@ struct RIGVM_API FRigVMFunction_MathFloatGreaterEqual : public FRigVMFunction_Ma
 	}
 
 	RIGVM_METHOD()
-	virtual void Execute() override;
+	UE_API virtual void Execute() override;
 
 	UPROPERTY(meta=(Input))
 	float A;
@@ -689,7 +691,7 @@ struct RIGVM_API FRigVMFunction_MathFloatGreaterEqual : public FRigVMFunction_Ma
  * Returns true if the value A is less than or equal to B
  */
 USTRUCT(meta=(DisplayName="Less Equal", TemplateName="LessEqual", Keywords="Smaller,<="))
-struct RIGVM_API FRigVMFunction_MathFloatLessEqual : public FRigVMFunction_MathFloatBase
+struct FRigVMFunction_MathFloatLessEqual : public FRigVMFunction_MathFloatBase
 {
 	GENERATED_BODY()
 
@@ -700,7 +702,7 @@ struct RIGVM_API FRigVMFunction_MathFloatLessEqual : public FRigVMFunction_MathF
 	}
 
 	RIGVM_METHOD()
-	virtual void Execute() override;
+	UE_API virtual void Execute() override;
 
 	UPROPERTY(meta=(Input))
 	float A;
@@ -716,7 +718,7 @@ struct RIGVM_API FRigVMFunction_MathFloatLessEqual : public FRigVMFunction_MathF
  * Returns true if the value is nearly zero
  */
 USTRUCT(meta=(DisplayName="Is Nearly Zero", TemplateName="IsNearlyZero", Keywords="AlmostZero,0"))
-struct RIGVM_API FRigVMFunction_MathFloatIsNearlyZero : public FRigVMFunction_MathFloatBase
+struct FRigVMFunction_MathFloatIsNearlyZero : public FRigVMFunction_MathFloatBase
 {
 	GENERATED_BODY()
 	
@@ -727,7 +729,7 @@ struct RIGVM_API FRigVMFunction_MathFloatIsNearlyZero : public FRigVMFunction_Ma
 	}
 	
 	RIGVM_METHOD()
-	virtual void Execute() override;
+	UE_API virtual void Execute() override;
 
 	UPROPERTY(meta=(Input))
 	float Value;
@@ -743,7 +745,7 @@ struct RIGVM_API FRigVMFunction_MathFloatIsNearlyZero : public FRigVMFunction_Ma
  * Returns true if the value A is almost equal to B
  */
 USTRUCT(meta=(DisplayName="Is Nearly Equal", TemplateName="IsNearlyEqual", Keywords="AlmostEqual,=="))
-struct RIGVM_API FRigVMFunction_MathFloatIsNearlyEqual : public FRigVMFunction_MathFloatBase
+struct FRigVMFunction_MathFloatIsNearlyEqual : public FRigVMFunction_MathFloatBase
 {
 	GENERATED_BODY()
 
@@ -754,7 +756,7 @@ struct RIGVM_API FRigVMFunction_MathFloatIsNearlyEqual : public FRigVMFunction_M
 	}
 
 	RIGVM_METHOD()
-	virtual void Execute() override;
+	UE_API virtual void Execute() override;
 
 	UPROPERTY(meta=(Input))
 	float A;
@@ -773,7 +775,7 @@ struct RIGVM_API FRigVMFunction_MathFloatIsNearlyEqual : public FRigVMFunction_M
  * Return one of the two values based on the condition
  */
 USTRUCT(meta=(DisplayName="Select", Keywords="Pick,If", Deprecated = "4.26.0"))
-struct RIGVM_API FRigVMFunction_MathFloatSelectBool : public FRigVMFunction_MathFloatBase
+struct FRigVMFunction_MathFloatSelectBool : public FRigVMFunction_MathFloatBase
 {
 	GENERATED_BODY()
 
@@ -784,7 +786,7 @@ struct RIGVM_API FRigVMFunction_MathFloatSelectBool : public FRigVMFunction_Math
 	}
 
 	RIGVM_METHOD()
-	virtual void Execute() override;
+	UE_API virtual void Execute() override;
 
 	UPROPERTY(meta=(Input))
 	bool Condition;
@@ -806,115 +808,115 @@ struct RIGVM_API FRigVMFunction_MathFloatSelectBool : public FRigVMFunction_Math
  * Returns the degrees of a given value in radians
  */
 USTRUCT(meta=(DisplayName="Degrees", TemplateName="Degrees"))
-struct RIGVM_API FRigVMFunction_MathFloatDeg : public FRigVMFunction_MathFloatUnaryOp
+struct FRigVMFunction_MathFloatDeg : public FRigVMFunction_MathFloatUnaryOp
 {
 	GENERATED_BODY()
 
 	RIGVM_METHOD()
-	virtual void Execute() override;
+	UE_API virtual void Execute() override;
 };
 
 /**
  * Returns the radians of a given value in degrees
  */
 USTRUCT(meta=(DisplayName="Radians", TemplateName="Radians"))
-struct RIGVM_API FRigVMFunction_MathFloatRad : public FRigVMFunction_MathFloatUnaryOp
+struct FRigVMFunction_MathFloatRad : public FRigVMFunction_MathFloatUnaryOp
 {
 	GENERATED_BODY()
 
 	RIGVM_METHOD()
-	virtual void Execute() override;
+	UE_API virtual void Execute() override;
 };
 
 /**
  * Returns the sinus value of the given value (in radians)
  */
 USTRUCT(meta=(DisplayName="Sin", TemplateName="Sin"))
-struct RIGVM_API FRigVMFunction_MathFloatSin : public FRigVMFunction_MathFloatUnaryOp
+struct FRigVMFunction_MathFloatSin : public FRigVMFunction_MathFloatUnaryOp
 {
 	GENERATED_BODY()
 
 	RIGVM_METHOD()
-	virtual void Execute() override;
+	UE_API virtual void Execute() override;
 };
 
 /**
  * Returns the cosinus value of the given value (in radians)
  */
 USTRUCT(meta=(DisplayName="Cos", TemplateName="Cos"))
-struct RIGVM_API FRigVMFunction_MathFloatCos : public FRigVMFunction_MathFloatUnaryOp
+struct FRigVMFunction_MathFloatCos : public FRigVMFunction_MathFloatUnaryOp
 {
 	GENERATED_BODY()
 
 	RIGVM_METHOD()
-	virtual void Execute() override;
+	UE_API virtual void Execute() override;
 };
 
 /**
  * Returns the tangens value of the given value (in radians)
  */
 USTRUCT(meta=(DisplayName="Tan", TemplateName="Tan"))
-struct RIGVM_API FRigVMFunction_MathFloatTan : public FRigVMFunction_MathFloatUnaryOp
+struct FRigVMFunction_MathFloatTan : public FRigVMFunction_MathFloatUnaryOp
 {
 	GENERATED_BODY()
 
 	RIGVM_METHOD()
-	virtual void Execute() override;
+	UE_API virtual void Execute() override;
 };
 
 /**
  * Returns the inverse sinus value (in radians) of the given value
  */
 USTRUCT(meta=(DisplayName="Asin", TemplateName="Asin", Keywords="Arcsin"))
-struct RIGVM_API FRigVMFunction_MathFloatAsin : public FRigVMFunction_MathFloatUnaryOp
+struct FRigVMFunction_MathFloatAsin : public FRigVMFunction_MathFloatUnaryOp
 {
 	GENERATED_BODY()
 
 	RIGVM_METHOD()
-	virtual void Execute() override;
+	UE_API virtual void Execute() override;
 };
 
 /**
  * Returns the inverse cosinus value (in radians) of the given value
  */
 USTRUCT(meta=(DisplayName="Acos", TemplateName="Acos", Keywords="Arccos"))
-struct RIGVM_API FRigVMFunction_MathFloatAcos : public FRigVMFunction_MathFloatUnaryOp
+struct FRigVMFunction_MathFloatAcos : public FRigVMFunction_MathFloatUnaryOp
 {
 	GENERATED_BODY()
 
 	RIGVM_METHOD()
-	virtual void Execute() override;
+	UE_API virtual void Execute() override;
 };
 
 /**
  * Returns the inverse tangens value (in radians) of the given value
  */
 USTRUCT(meta=(DisplayName="Atan", TemplateName="Atan", Keywords="Arctan"))
-struct RIGVM_API FRigVMFunction_MathFloatAtan : public FRigVMFunction_MathFloatUnaryOp
+struct FRigVMFunction_MathFloatAtan : public FRigVMFunction_MathFloatUnaryOp
 {
 	GENERATED_BODY()
 
 	RIGVM_METHOD()
-	virtual void Execute() override;
+	UE_API virtual void Execute() override;
 };
 
 /**
  * Returns the arctangent of the specified A and B coordinates.
  */
 USTRUCT(meta=(DisplayName="Atan2", TemplateName="Atan2", Keywords="Arctan"))
-struct RIGVM_API FRigVMFunction_MathFloatAtan2 : public FRigVMFunction_MathFloatBinaryOp
+struct FRigVMFunction_MathFloatAtan2 : public FRigVMFunction_MathFloatBinaryOp
 {
 	GENERATED_BODY()
 
 	RIGVM_METHOD()
-	virtual void Execute() override;
+	UE_API virtual void Execute() override;
 };
 
 /**
  * Computes the angles alpha, beta and gamma (in radians) between the three sides A, B and C
  */
 USTRUCT(meta = (DisplayName = "Law Of Cosine", TemplateName="LawOfCosine"))
-struct RIGVM_API FRigVMFunction_MathFloatLawOfCosine : public FRigVMFunction_MathFloatBase
+struct FRigVMFunction_MathFloatLawOfCosine : public FRigVMFunction_MathFloatBase
 {
 	GENERATED_BODY()
 
@@ -925,7 +927,7 @@ struct RIGVM_API FRigVMFunction_MathFloatLawOfCosine : public FRigVMFunction_Mat
 	}
 
 	RIGVM_METHOD()
-	virtual void Execute() override;
+	UE_API virtual void Execute() override;
 
 	UPROPERTY(meta = (Input))
 	float A;
@@ -953,19 +955,19 @@ struct RIGVM_API FRigVMFunction_MathFloatLawOfCosine : public FRigVMFunction_Mat
  * Computes the base-e exponential of the given value 
  */
 USTRUCT(meta = (DisplayName = "Exponential", TemplateName="Exponential"))
-struct RIGVM_API FRigVMFunction_MathFloatExponential : public FRigVMFunction_MathFloatUnaryOp
+struct FRigVMFunction_MathFloatExponential : public FRigVMFunction_MathFloatUnaryOp
 {
 	GENERATED_BODY()
 
 	RIGVM_METHOD()
-	virtual void Execute() override;
+	UE_API virtual void Execute() override;
 };
 
 /**
  * Returns the sum of the given array
  */
 USTRUCT(meta = (DisplayName = "Array Sum", TemplateName = "ArraySum"))
-struct RIGVM_API FRigVMFunction_MathFloatArraySum : public FRigVMFunction_MathFloatBase
+struct FRigVMFunction_MathFloatArraySum : public FRigVMFunction_MathFloatBase
 {
 	GENERATED_BODY()
 
@@ -976,7 +978,7 @@ struct RIGVM_API FRigVMFunction_MathFloatArraySum : public FRigVMFunction_MathFl
 
 	/** Execute logic for this rig unit */
 	RIGVM_METHOD()
-	virtual void Execute() override;
+	UE_API virtual void Execute() override;
 
 	UPROPERTY(meta = (Input))
 	TArray<float> Array;
@@ -989,7 +991,7 @@ struct RIGVM_API FRigVMFunction_MathFloatArraySum : public FRigVMFunction_MathFl
  * Returns the average of the given array
  */
 USTRUCT(meta = (DisplayName = "Array Average", TemplateName = "ArrayAverage"))
-struct RIGVM_API FRigVMFunction_MathFloatArrayAverage : public FRigVMFunction_MathFloatBase
+struct FRigVMFunction_MathFloatArrayAverage : public FRigVMFunction_MathFloatBase
 {
 	GENERATED_BODY()
 
@@ -1000,11 +1002,13 @@ struct RIGVM_API FRigVMFunction_MathFloatArrayAverage : public FRigVMFunction_Ma
 
 	/** Execute logic for this rig unit */
 	RIGVM_METHOD()
-	virtual void Execute() override;
+	UE_API virtual void Execute() override;
 
 	UPROPERTY(meta = (Input))
 	TArray<float> Array;
 
 	UPROPERTY(meta = (Output))
 	float Average;
-};
\ No newline at end of file
+};
+
+#undef UE_API 
\ No newline at end of file
diff --git a/Engine/Plugins/Runtime/RigVM/Source/RigVM/Public/RigVMFunctions/Math/RigVMFunction_MathInt.h b/Engine/Plugins/Runtime/RigVM/Source/RigVM/Public/RigVMFunctions/Math/RigVMFunction_MathInt.h
--- a/Engine/Plugins/Runtime/RigVM/Source/RigVM/Public/RigVMFunctions/Math/RigVMFunction_MathInt.h
+++ b/Engine/Plugins/Runtime/RigVM/Source/RigVM/Public/RigVMFunctions/Math/RigVMFunction_MathInt.h
@@ -6,13 +6,13 @@
 #include "RigVMFunction_MathInt.generated.h"
 
 USTRUCT(meta=(Abstract, Category="Math|Int", MenuDescSuffix="(Int)"))
-struct RIGVM_API FRigVMFunction_MathIntBase : public FRigVMFunction_MathBase
+struct FRigVMFunction_MathIntBase : public FRigVMFunction_MathBase
 {
 	GENERATED_BODY()
 };
 
 USTRUCT(meta=(Abstract))
-struct RIGVM_API FRigVMFunction_MathIntUnaryOp : public FRigVMFunction_MathIntBase
+struct FRigVMFunction_MathIntUnaryOp : public FRigVMFunction_MathIntBase
 {
 	GENERATED_BODY()
 
@@ -29,7 +29,7 @@ struct RIGVM_API FRigVMFunction_MathIntUnaryOp : public FRigVMFunction_MathIntBa
 };
 
 USTRUCT(meta=(Abstract))
-struct RIGVM_API FRigVMFunction_MathIntBinaryOp : public FRigVMFunction_MathIntBase
+struct FRigVMFunction_MathIntBinaryOp : public FRigVMFunction_MathIntBase
 {
 	GENERATED_BODY()
 
@@ -49,7 +49,7 @@ struct RIGVM_API FRigVMFunction_MathIntBinaryOp : public FRigVMFunction_MathIntB
 };
 
 USTRUCT(meta=(Abstract))
-struct RIGVM_API FRigVMFunction_MathIntBinaryAggregateOp : public FRigVMFunction_MathIntBase
+struct FRigVMFunction_MathIntBinaryAggregateOp : public FRigVMFunction_MathIntBase
 {
 	GENERATED_BODY()
 
@@ -72,7 +72,7 @@ struct RIGVM_API FRigVMFunction_MathIntBinaryAggregateOp : public FRigVMFunction
  * A integer constant
  */
 USTRUCT(meta=(DisplayName="Integer", Keywords="Make,Construct,Constant"))
-struct RIGVM_API FRigVMFunction_MathIntMake : public FRigVMFunction_MathIntBase
+struct FRigVMFunction_MathIntMake : public FRigVMFunction_MathIntBase
 {
 	GENERATED_BODY()
 
@@ -82,7 +82,7 @@ struct RIGVM_API FRigVMFunction_MathIntMake : public FRigVMFunction_MathIntBase
 	}
 	
 	RIGVM_METHOD()
-	virtual void Execute() override;
+	RIGVM_API virtual void Execute() override;
 
 	UPROPERTY(meta=(Input, Output))
 	int32 Value;
@@ -92,31 +92,31 @@ struct RIGVM_API FRigVMFunction_MathIntMake : public FRigVMFunction_MathIntBase
  * Returns the sum of the two values
  */
 USTRUCT(meta=(DisplayName="Add", TemplateName="Add", Keywords="Sum,+"))
-struct RIGVM_API FRigVMFunction_MathIntAdd : public FRigVMFunction_MathIntBinaryAggregateOp
+struct FRigVMFunction_MathIntAdd : public FRigVMFunction_MathIntBinaryAggregateOp
 {
 	GENERATED_BODY()
 
 	RIGVM_METHOD()
-	virtual void Execute() override;
+	RIGVM_API virtual void Execute() override;
 };
 
 /**
  * Returns the difference of the two values
  */
 USTRUCT(meta=(DisplayName="Subtract", TemplateName="Subtract", Keywords="-"))
-struct RIGVM_API FRigVMFunction_MathIntSub : public FRigVMFunction_MathIntBinaryOp
+struct FRigVMFunction_MathIntSub : public FRigVMFunction_MathIntBinaryOp
 {
 	GENERATED_BODY()
 
 	RIGVM_METHOD()
-	virtual void Execute() override;
+	RIGVM_API virtual void Execute() override;
 };
 
 /**
  * Returns the product of the two values
  */
 USTRUCT(meta=(DisplayName="Multiply", TemplateName="Multiply", Keywords="Product,*"))
-struct RIGVM_API FRigVMFunction_MathIntMul : public FRigVMFunction_MathIntBinaryAggregateOp
+struct FRigVMFunction_MathIntMul : public FRigVMFunction_MathIntBinaryAggregateOp
 {
 	GENERATED_BODY()
 
@@ -127,14 +127,14 @@ struct RIGVM_API FRigVMFunction_MathIntMul : public FRigVMFunction_MathIntBinary
 	}
 
 	RIGVM_METHOD()
-	virtual void Execute() override;
+	RIGVM_API virtual void Execute() override;
 };
 
 /**
  * Returns the division of the two values
  */
 USTRUCT(meta=(DisplayName="Divide", TemplateName="Divide", Keywords="Division,Divisor,/"))
-struct RIGVM_API FRigVMFunction_MathIntDiv : public FRigVMFunction_MathIntBinaryOp
+struct FRigVMFunction_MathIntDiv : public FRigVMFunction_MathIntBinaryOp
 {
 	GENERATED_BODY()
 
@@ -144,14 +144,14 @@ struct RIGVM_API FRigVMFunction_MathIntDiv : public FRigVMFunction_MathIntBinary
 	}
 
 	RIGVM_METHOD()
-	virtual void Execute() override;
+	RIGVM_API virtual void Execute() override;
 };
 
 /**
  * Returns the modulo of the two values
  */
 USTRUCT(meta=(DisplayName="Modulo", TemplateName="Modulo", Keywords="%,fmod"))
-struct RIGVM_API FRigVMFunction_MathIntMod : public FRigVMFunction_MathIntBinaryOp
+struct FRigVMFunction_MathIntMod : public FRigVMFunction_MathIntBinaryOp
 {
 	GENERATED_BODY()
 
@@ -162,38 +162,38 @@ struct RIGVM_API FRigVMFunction_MathIntMod : public FRigVMFunction_MathIntBinary
 	}
 
 	RIGVM_METHOD()
-	virtual void Execute() override;
+	RIGVM_API virtual void Execute() override;
 };
 
 /**
  * Returns the smaller of the two values
  */
 USTRUCT(meta=(DisplayName="Minimum", TemplateName="Minimum"))
-struct RIGVM_API FRigVMFunction_MathIntMin : public FRigVMFunction_MathIntBinaryAggregateOp
+struct FRigVMFunction_MathIntMin : public FRigVMFunction_MathIntBinaryAggregateOp
 {
 	GENERATED_BODY()
 
 	RIGVM_METHOD()
-	virtual void Execute() override;
+	RIGVM_API virtual void Execute() override;
 };
 
 /**
  * Returns the larger of the two values
  */
 USTRUCT(meta=(DisplayName="Maximum", TemplateName="Maximum"))
-struct RIGVM_API FRigVMFunction_MathIntMax : public FRigVMFunction_MathIntBinaryAggregateOp
+struct FRigVMFunction_MathIntMax : public FRigVMFunction_MathIntBinaryAggregateOp
 {
 	GENERATED_BODY()
 
 	RIGVM_METHOD()
-	virtual void Execute() override;
+	RIGVM_API virtual void Execute() override;
 };
 
 /**
  * Returns the value of A raised to the power of B.
  */
 USTRUCT(meta=(DisplayName="Power", TemplateName="Power"))
-struct RIGVM_API FRigVMFunction_MathIntPow : public FRigVMFunction_MathIntBinaryOp
+struct FRigVMFunction_MathIntPow : public FRigVMFunction_MathIntBinaryOp
 {
 	GENERATED_BODY()
 
@@ -204,38 +204,38 @@ struct RIGVM_API FRigVMFunction_MathIntPow : public FRigVMFunction_MathIntBinary
 	}
 
 	RIGVM_METHOD()
-	virtual void Execute() override;
+	RIGVM_API virtual void Execute() override;
 };
 
 /**
  * Returns the negative value
  */
 USTRUCT(meta=(DisplayName="Negate", TemplateName="Negate", Keywords="-,Abs"))
-struct RIGVM_API FRigVMFunction_MathIntNegate : public FRigVMFunction_MathIntUnaryOp
+struct FRigVMFunction_MathIntNegate : public FRigVMFunction_MathIntUnaryOp
 {
 	GENERATED_BODY()
 
 	RIGVM_METHOD()
-	virtual void Execute() override;
+	RIGVM_API virtual void Execute() override;
 };
 
 /**
  * Returns the absolute (positive) value
  */
 USTRUCT(meta=(DisplayName="Absolute", TemplateName="Absolute", Keywords="Abs,Neg"))
-struct RIGVM_API FRigVMFunction_MathIntAbs : public FRigVMFunction_MathIntUnaryOp
+struct FRigVMFunction_MathIntAbs : public FRigVMFunction_MathIntUnaryOp
 {
 	GENERATED_BODY()
 
 	RIGVM_METHOD()
-	virtual void Execute() override;
+	RIGVM_API virtual void Execute() override;
 };
 
 /**
  * Returns the int cast to a float
  */
 USTRUCT(meta=(DisplayName="To Float", TemplateName="Cast", ExecuteContext="FRigVMExecuteContext"))
-struct RIGVM_API FRigVMFunction_MathIntToFloat : public FRigVMFunction_MathIntBase
+struct FRigVMFunction_MathIntToFloat : public FRigVMFunction_MathIntBase
 {
 	GENERATED_BODY()
 
@@ -246,7 +246,7 @@ struct RIGVM_API FRigVMFunction_MathIntToFloat : public FRigVMFunction_MathIntBa
 	}
 
 	RIGVM_METHOD()
-	virtual void Execute() override;
+	RIGVM_API virtual void Execute() override;
 
 	UPROPERTY(meta=(Input))
 	int32 Value;
@@ -259,7 +259,7 @@ struct RIGVM_API FRigVMFunction_MathIntToFloat : public FRigVMFunction_MathIntBa
  * Returns the int cast to a float
  */
 USTRUCT(meta=(DisplayName="To Double", TemplateName="Cast", ExecuteContext="FRigVMExecuteContext"))
-struct RIGVM_API FRigVMFunction_MathIntToDouble : public FRigVMFunction_MathIntBase
+struct FRigVMFunction_MathIntToDouble : public FRigVMFunction_MathIntBase
 {
 	GENERATED_BODY()
 
@@ -270,7 +270,7 @@ struct RIGVM_API FRigVMFunction_MathIntToDouble : public FRigVMFunction_MathIntB
 	}
 
 	RIGVM_METHOD()
-	virtual void Execute() override;
+	RIGVM_API virtual void Execute() override;
 
 	UPROPERTY(meta=(Input))
 	int32 Value;
@@ -283,19 +283,19 @@ struct RIGVM_API FRigVMFunction_MathIntToDouble : public FRigVMFunction_MathIntB
  * Returns the sign of the value (+1 for >= 0, -1 for < 0)
  */
 USTRUCT(meta=(DisplayName="Sign", TemplateName="Sign"))
-struct RIGVM_API FRigVMFunction_MathIntSign : public FRigVMFunction_MathIntUnaryOp
+struct FRigVMFunction_MathIntSign : public FRigVMFunction_MathIntUnaryOp
 {
 	GENERATED_BODY()
 
 	RIGVM_METHOD()
-	virtual void Execute() override;
+	RIGVM_API virtual void Execute() override;
 };
 
 /**
  * Clamps the given value within the range provided by minimum and maximum
  */
 USTRUCT(meta=(DisplayName="Clamp", TemplateName="Clamp", Keywords="Range,Remap"))
-struct RIGVM_API FRigVMFunction_MathIntClamp : public FRigVMFunction_MathIntBase
+struct FRigVMFunction_MathIntClamp : public FRigVMFunction_MathIntBase
 {
 	GENERATED_BODY()
 
@@ -305,7 +305,7 @@ struct RIGVM_API FRigVMFunction_MathIntClamp : public FRigVMFunction_MathIntBase
 	}
 
 	RIGVM_METHOD()
-	virtual void Execute() override;
+	RIGVM_API virtual void Execute() override;
 
 	UPROPERTY(meta=(Input))
 	int32 Value;
@@ -324,11 +324,11 @@ struct RIGVM_API FRigVMFunction_MathIntClamp : public FRigVMFunction_MathIntBase
  * Returns true if the value A equals B
  */
 USTRUCT(meta=(DisplayName="Equals", TemplateName="Equals", Keywords="Same,==", Deprecated="5.1"))
-struct RIGVM_API FRigVMFunction_MathIntEquals : public FRigVMFunction_MathIntBase
+struct FRigVMFunction_MathIntEquals : public FRigVMFunction_MathIntBase
 {
 	GENERATED_BODY()
 	RIGVM_METHOD()
-	virtual void Execute() override;
+	RIGVM_API virtual void Execute() override;
 
 	FRigVMFunction_MathIntEquals()
 	{
@@ -353,11 +353,11 @@ struct RIGVM_API FRigVMFunction_MathIntEquals : public FRigVMFunction_MathIntBas
  * Returns true if the value A does not equal B
  */
 USTRUCT(meta=(DisplayName="Not Equals", TemplateName="NotEquals", Keywords="Different,!=", Deprecated="5.1"))
-struct RIGVM_API FRigVMFunction_MathIntNotEquals : public FRigVMFunction_MathIntBase
+struct FRigVMFunction_MathIntNotEquals : public FRigVMFunction_MathIntBase
 {
 	GENERATED_BODY()
 	RIGVM_METHOD()
-	virtual void Execute() override;
+	RIGVM_API virtual void Execute() override;
 
 	FRigVMFunction_MathIntNotEquals()
 	{
@@ -382,7 +382,7 @@ struct RIGVM_API FRigVMFunction_MathIntNotEquals : public FRigVMFunction_MathInt
  * Returns true if the value A is greater than B
  */
 USTRUCT(meta=(DisplayName="Greater", TemplateName="Greater", Keywords="Larger,Bigger,>"))
-struct RIGVM_API FRigVMFunction_MathIntGreater : public FRigVMFunction_MathIntBase
+struct FRigVMFunction_MathIntGreater : public FRigVMFunction_MathIntBase
 {
 	GENERATED_BODY()
 
@@ -393,7 +393,7 @@ struct RIGVM_API FRigVMFunction_MathIntGreater : public FRigVMFunction_MathIntBa
 	}
 
 	RIGVM_METHOD()
-	virtual void Execute() override;
+	RIGVM_API virtual void Execute() override;
 
 	UPROPERTY(meta=(Input))
 	int32 A;
@@ -409,7 +409,7 @@ struct RIGVM_API FRigVMFunction_MathIntGreater : public FRigVMFunction_MathIntBa
  * Returns true if the value A is less than B
  */
 USTRUCT(meta=(DisplayName="Less", TemplateName="Less", Keywords="Smaller,<"))
-struct RIGVM_API FRigVMFunction_MathIntLess : public FRigVMFunction_MathIntBase
+struct FRigVMFunction_MathIntLess : public FRigVMFunction_MathIntBase
 {
 	GENERATED_BODY()
 	
@@ -420,7 +420,7 @@ struct RIGVM_API FRigVMFunction_MathIntLess : public FRigVMFunction_MathIntBase
 	}
 
 	RIGVM_METHOD()
-	virtual void Execute() override;
+	RIGVM_API virtual void Execute() override;
 
 	UPROPERTY(meta=(Input))
 	int32 A;
@@ -436,7 +436,7 @@ struct RIGVM_API FRigVMFunction_MathIntLess : public FRigVMFunction_MathIntBase
  * Returns true if the value A is greater than or equal to B
  */
 USTRUCT(meta=(DisplayName="Greater Equal", TemplateName="GreaterEqual", Keywords="Larger,Bigger,>="))
-struct RIGVM_API FRigVMFunction_MathIntGreaterEqual : public FRigVMFunction_MathIntBase
+struct FRigVMFunction_MathIntGreaterEqual : public FRigVMFunction_MathIntBase
 {
 	GENERATED_BODY()
 
@@ -447,7 +447,7 @@ struct RIGVM_API FRigVMFunction_MathIntGreaterEqual : public FRigVMFunction_Math
 	}
 
 	RIGVM_METHOD()
-	virtual void Execute() override;
+	RIGVM_API virtual void Execute() override;
 
 	UPROPERTY(meta=(Input))
 	int32 A;
@@ -463,7 +463,7 @@ struct RIGVM_API FRigVMFunction_MathIntGreaterEqual : public FRigVMFunction_Math
  * Returns true if the value A is less than or equal to B
  */
 USTRUCT(meta=(DisplayName="Less Equal", TemplateName="LessEqual", Keywords="Smaller,<="))
-struct RIGVM_API FRigVMFunction_MathIntLessEqual : public FRigVMFunction_MathIntBase
+struct FRigVMFunction_MathIntLessEqual : public FRigVMFunction_MathIntBase
 {
 	GENERATED_BODY()
 
@@ -474,7 +474,7 @@ struct RIGVM_API FRigVMFunction_MathIntLessEqual : public FRigVMFunction_MathInt
 	}
 
 	RIGVM_METHOD()
-	virtual void Execute() override;
+	RIGVM_API virtual void Execute() override;
 
 	UPROPERTY(meta=(Input))
 	int32 A;
@@ -490,7 +490,7 @@ struct RIGVM_API FRigVMFunction_MathIntLessEqual : public FRigVMFunction_MathInt
  * Returns the sum of the given array
  */
 USTRUCT(meta = (DisplayName = "Array Sum", TemplateName = "ArraySum"))
-struct RIGVM_API FRigVMFunction_MathIntArraySum : public FRigVMFunction_MathIntBase
+struct FRigVMFunction_MathIntArraySum : public FRigVMFunction_MathIntBase
 {
 	GENERATED_BODY()
 
@@ -501,7 +501,7 @@ struct RIGVM_API FRigVMFunction_MathIntArraySum : public FRigVMFunction_MathIntB
 
 	/** Execute logic for this rig unit */
 	RIGVM_METHOD()
-	virtual void Execute() override;
+	RIGVM_API virtual void Execute() override;
 
 	UPROPERTY(meta = (Input))
 	TArray<int32> Array;
@@ -515,7 +515,7 @@ struct RIGVM_API FRigVMFunction_MathIntArraySum : public FRigVMFunction_MathIntB
  * Returns the average of the given array
  */
 USTRUCT(meta = (DisplayName = "Array Average", TemplateName = "ArrayAverage"))
-struct RIGVM_API FRigVMFunction_MathIntArrayAverage : public FRigVMFunction_MathIntBase
+struct FRigVMFunction_MathIntArrayAverage : public FRigVMFunction_MathIntBase
 {
 	GENERATED_BODY()
 
@@ -526,7 +526,7 @@ struct RIGVM_API FRigVMFunction_MathIntArrayAverage : public FRigVMFunction_Math
 
 	/** Execute logic for this rig unit */
 	RIGVM_METHOD()
-	virtual void Execute() override;
+	RIGVM_API virtual void Execute() override;
 
 	UPROPERTY(meta = (Input))
 	TArray<int32> Array;
@@ -539,7 +539,7 @@ struct RIGVM_API FRigVMFunction_MathIntArrayAverage : public FRigVMFunction_Math
  * Converts an integer to a string
  */
 USTRUCT(meta = (DisplayName = "Int to String", TemplateName="Int to Name"))
-struct RIGVM_API FRigVMFunction_MathIntToString : public FRigVMFunction_MathIntBase
+struct FRigVMFunction_MathIntToString : public FRigVMFunction_MathIntBase
 {
 	GENERATED_BODY()
 
@@ -550,7 +550,7 @@ struct RIGVM_API FRigVMFunction_MathIntToString : public FRigVMFunction_MathIntB
 
 	/** Execute logic for this function */
 	RIGVM_METHOD()
-	virtual void Execute() override;
+	RIGVM_API virtual void Execute() override;
 
 	UPROPERTY(meta = (Input))
 	int32 Number;
@@ -570,7 +570,7 @@ struct RIGVM_API FRigVMFunction_MathIntToString : public FRigVMFunction_MathIntB
  * Converts an integer to a name
  */
 USTRUCT(meta = (DisplayName = "Int to Name", TemplateName="Int to Name"))
-struct RIGVM_API FRigVMFunction_MathIntToName : public FRigVMFunction_MathIntBase
+struct FRigVMFunction_MathIntToName : public FRigVMFunction_MathIntBase
 {
 	GENERATED_BODY()
 
@@ -581,7 +581,7 @@ struct RIGVM_API FRigVMFunction_MathIntToName : public FRigVMFunction_MathIntBas
 
 	/** Execute logic for this function */
 	RIGVM_METHOD()
-	virtual void Execute() override;
+	RIGVM_API virtual void Execute() override;
 
 	UPROPERTY(meta = (Input))
 	int32 Number;
diff --git a/Engine/Plugins/Runtime/RigVM/Source/RigVM/Public/RigVMFunctions/Math/RigVMFunction_MathMatrix.h b/Engine/Plugins/Runtime/RigVM/Source/RigVM/Public/RigVMFunctions/Math/RigVMFunction_MathMatrix.h
--- a/Engine/Plugins/Runtime/RigVM/Source/RigVM/Public/RigVMFunctions/Math/RigVMFunction_MathMatrix.h
+++ b/Engine/Plugins/Runtime/RigVM/Source/RigVM/Public/RigVMFunctions/Math/RigVMFunction_MathMatrix.h
@@ -7,13 +7,13 @@
 #include "RigVMFunction_MathMatrix.generated.h"
 
 USTRUCT(meta=(Abstract, Category="Math|Matrix", MenuDescSuffix="(Matrix)"))
-struct RIGVM_API FRigVMFunction_MathMatrixBase : public FRigVMFunction_MathBase
+struct FRigVMFunction_MathMatrixBase : public FRigVMFunction_MathBase
 {
 	GENERATED_BODY()
 };
 
 USTRUCT(meta=(Abstract))
-struct RIGVM_API FRigVMFunction_MathMatrixUnaryOp : public FRigVMFunction_MathMatrixBase
+struct FRigVMFunction_MathMatrixUnaryOp : public FRigVMFunction_MathMatrixBase
 {
 	GENERATED_BODY()
 
@@ -30,7 +30,7 @@ struct RIGVM_API FRigVMFunction_MathMatrixUnaryOp : public FRigVMFunction_MathMa
 };
 
 USTRUCT(meta=(Abstract))
-struct RIGVM_API FRigVMFunction_MathMatrixBinaryOp : public FRigVMFunction_MathMatrixBase
+struct FRigVMFunction_MathMatrixBinaryOp : public FRigVMFunction_MathMatrixBase
 {
 	GENERATED_BODY()
 
@@ -50,7 +50,7 @@ struct RIGVM_API FRigVMFunction_MathMatrixBinaryOp : public FRigVMFunction_MathM
 };
 
 USTRUCT(meta=(Abstract))
-struct RIGVM_API FRigVMFunction_MathMatrixBinaryAggregateOp : public FRigVMFunction_MathMatrixBase
+struct FRigVMFunction_MathMatrixBinaryAggregateOp : public FRigVMFunction_MathMatrixBase
 {
 	GENERATED_BODY()
 
@@ -73,7 +73,7 @@ struct RIGVM_API FRigVMFunction_MathMatrixBinaryAggregateOp : public FRigVMFunct
 * Makes a transform from a matrix
 */
 USTRUCT(meta=(DisplayName="To Transform", TemplateName="Cast", ExecuteContext="FRigVMExecuteContext", Keywords="Make,Construct"))
-struct RIGVM_API FRigVMFunction_MathMatrixToTransform : public FRigVMFunction_MathMatrixBase
+struct FRigVMFunction_MathMatrixToTransform : public FRigVMFunction_MathMatrixBase
 {
 	GENERATED_BODY()
 
@@ -84,7 +84,7 @@ struct RIGVM_API FRigVMFunction_MathMatrixToTransform : public FRigVMFunction_Ma
 	}
 
 	RIGVM_METHOD()
-	virtual void Execute() override;
+	RIGVM_API virtual void Execute() override;
 
 	UPROPERTY(meta=(Input))
 	FMatrix Value;
@@ -97,7 +97,7 @@ struct RIGVM_API FRigVMFunction_MathMatrixToTransform : public FRigVMFunction_Ma
  * Makes a matrix from a transform
  */
 USTRUCT(meta=(DisplayName="From Transform", Keywords="Make,Construct", Deprecated="5.0.1"))
-struct RIGVM_API FRigVMFunction_MathMatrixFromTransform : public FRigVMFunction_MathMatrixBase
+struct FRigVMFunction_MathMatrixFromTransform : public FRigVMFunction_MathMatrixBase
 {
 	GENERATED_BODY()
 
@@ -108,7 +108,7 @@ struct RIGVM_API FRigVMFunction_MathMatrixFromTransform : public FRigVMFunction_
 	}
 
 	RIGVM_METHOD()
-	virtual void Execute() override;
+	RIGVM_API virtual void Execute() override;
 
 	UPROPERTY(meta=(Input))
 	FTransform Transform;
@@ -124,7 +124,7 @@ struct RIGVM_API FRigVMFunction_MathMatrixFromTransform : public FRigVMFunction_
  * Makes a matrix from a transform
  */
 USTRUCT(meta=(DisplayName="From Transform", TemplateName="Cast", ExecuteContext="FRigVMExecuteContext", Keywords="Make,Construct"))
-struct RIGVM_API FRigVMFunction_MathMatrixFromTransformV2 : public FRigVMFunction_MathMatrixBase
+struct FRigVMFunction_MathMatrixFromTransformV2 : public FRigVMFunction_MathMatrixBase
 {
 	GENERATED_BODY()
 
@@ -135,7 +135,7 @@ struct RIGVM_API FRigVMFunction_MathMatrixFromTransformV2 : public FRigVMFunctio
 	}
 
 	RIGVM_METHOD()
-	virtual void Execute() override;
+	RIGVM_API virtual void Execute() override;
 
 	UPROPERTY(meta=(Input))
 	FTransform Value;
@@ -148,7 +148,7 @@ struct RIGVM_API FRigVMFunction_MathMatrixFromTransformV2 : public FRigVMFunctio
  * Converts the matrix to its vectors
  */
 USTRUCT(meta=(DisplayName="To Vectors", Keywords="Make,Construct"))
-struct RIGVM_API FRigVMFunction_MathMatrixToVectors : public FRigVMFunction_MathMatrixBase
+struct FRigVMFunction_MathMatrixToVectors : public FRigVMFunction_MathMatrixBase
 {
 	GENERATED_BODY()
 
@@ -162,7 +162,7 @@ struct RIGVM_API FRigVMFunction_MathMatrixToVectors : public FRigVMFunction_Math
 	}
 
 	RIGVM_METHOD()
-	virtual void Execute() override;
+	RIGVM_API virtual void Execute() override;
 
 	UPROPERTY(meta=(Input))
 	FMatrix Value;
@@ -184,7 +184,7 @@ struct RIGVM_API FRigVMFunction_MathMatrixToVectors : public FRigVMFunction_Math
 * Makes a matrix from its vectors
 */
 USTRUCT(meta=(DisplayName="From Vectors", Keywords="Make,Construct"))
-struct RIGVM_API FRigVMFunction_MathMatrixFromVectors : public FRigVMFunction_MathMatrixBase
+struct FRigVMFunction_MathMatrixFromVectors : public FRigVMFunction_MathMatrixBase
 {
 	GENERATED_BODY()
 
@@ -198,7 +198,7 @@ struct RIGVM_API FRigVMFunction_MathMatrixFromVectors : public FRigVMFunction_Ma
 	}
 
 	RIGVM_METHOD()
-	virtual void Execute() override;
+	RIGVM_API virtual void Execute() override;
 
 	UPROPERTY(meta=(Input))
 	FVector Origin;
@@ -220,20 +220,20 @@ struct RIGVM_API FRigVMFunction_MathMatrixFromVectors : public FRigVMFunction_Ma
  * Returns the product of the two values
  */
 USTRUCT(meta=(DisplayName="Multiply", TemplateName="Multiply", Keywords="Product,*,Global"))
-struct RIGVM_API FRigVMFunction_MathMatrixMul : public FRigVMFunction_MathMatrixBinaryAggregateOp
+struct FRigVMFunction_MathMatrixMul : public FRigVMFunction_MathMatrixBinaryAggregateOp
 {
 	GENERATED_BODY()
 	RIGVM_METHOD()
-	virtual void Execute() override;
+	RIGVM_API virtual void Execute() override;
 };
 
 /**
  * Returns the inverse value
  */
 USTRUCT(meta=(DisplayName="Inverse", TemplateName="Inverse"))
-struct RIGVM_API FRigVMFunction_MathMatrixInverse : public FRigVMFunction_MathMatrixUnaryOp
+struct FRigVMFunction_MathMatrixInverse : public FRigVMFunction_MathMatrixUnaryOp
 {
 	GENERATED_BODY()
 	RIGVM_METHOD()
-	virtual void Execute() override;
+	RIGVM_API virtual void Execute() override;
 };
diff --git a/Engine/Plugins/Runtime/RigVM/Source/RigVM/Public/RigVMFunctions/Math/RigVMFunction_MathQuaternion.h b/Engine/Plugins/Runtime/RigVM/Source/RigVM/Public/RigVMFunctions/Math/RigVMFunction_MathQuaternion.h
--- a/Engine/Plugins/Runtime/RigVM/Source/RigVM/Public/RigVMFunctions/Math/RigVMFunction_MathQuaternion.h
+++ b/Engine/Plugins/Runtime/RigVM/Source/RigVM/Public/RigVMFunctions/Math/RigVMFunction_MathQuaternion.h
@@ -6,14 +6,16 @@
 #include "EulerTransform.h"
 #include "RigVMFunction_MathQuaternion.generated.h"
 
+#define UE_API RIGVM_API
+
 USTRUCT(meta=(Abstract, Category="Math|Quaternion", MenuDescSuffix="(Quaternion)"))
-struct RIGVM_API FRigVMFunction_MathQuaternionBase : public FRigVMFunction_MathBase
+struct FRigVMFunction_MathQuaternionBase : public FRigVMFunction_MathBase
 {
 	GENERATED_BODY()
 };
 
 USTRUCT(meta=(Abstract))
-struct RIGVM_API FRigVMFunction_MathQuaternionUnaryOp : public FRigVMFunction_MathQuaternionBase
+struct FRigVMFunction_MathQuaternionUnaryOp : public FRigVMFunction_MathQuaternionBase
 {
 	GENERATED_BODY()
 
@@ -30,7 +32,7 @@ struct RIGVM_API FRigVMFunction_MathQuaternionUnaryOp : public FRigVMFunction_Ma
 };
 
 USTRUCT(meta=(Abstract))
-struct RIGVM_API FRigVMFunction_MathQuaternionBinaryOp : public FRigVMFunction_MathQuaternionBase
+struct FRigVMFunction_MathQuaternionBinaryOp : public FRigVMFunction_MathQuaternionBase
 {
 	GENERATED_BODY()
 
@@ -50,7 +52,7 @@ struct RIGVM_API FRigVMFunction_MathQuaternionBinaryOp : public FRigVMFunction_M
 };
 
 USTRUCT(meta=(Abstract))
-struct RIGVM_API FRigVMFunction_MathQuaternionBinaryAggregateOp : public FRigVMFunction_MathQuaternionBase
+struct FRigVMFunction_MathQuaternionBinaryAggregateOp : public FRigVMFunction_MathQuaternionBase
 {
 	GENERATED_BODY()
 
@@ -73,11 +75,11 @@ struct RIGVM_API FRigVMFunction_MathQuaternionBinaryAggregateOp : public FRigVMF
  * Makes a quaternion from its components
  */
 USTRUCT(meta=(DisplayName="Make Quat", Keywords="Make,Construct,Constant"))
-struct RIGVM_API FRigVMFunction_MathQuaternionMake : public FRigVMFunction_MathQuaternionBase
+struct FRigVMFunction_MathQuaternionMake : public FRigVMFunction_MathQuaternionBase
 {
 	GENERATED_BODY()
 	RIGVM_METHOD()
-	virtual void Execute() override;
+	UE_API virtual void Execute() override;
 
 	FRigVMFunction_MathQuaternionMake()
 	{
@@ -105,7 +107,7 @@ struct RIGVM_API FRigVMFunction_MathQuaternionMake : public FRigVMFunction_MathQ
  * Makes a quaternion from an axis and an angle in radians
  */
 USTRUCT(meta=(DisplayName="From Axis And Angle", TemplateName="FromAxisAndAngle", Keywords="Make,Construct"))
-struct RIGVM_API FRigVMFunction_MathQuaternionFromAxisAndAngle : public FRigVMFunction_MathQuaternionBase
+struct FRigVMFunction_MathQuaternionFromAxisAndAngle : public FRigVMFunction_MathQuaternionBase
 {
 	GENERATED_BODY()
 
@@ -117,7 +119,7 @@ struct RIGVM_API FRigVMFunction_MathQuaternionFromAxisAndAngle : public FRigVMFu
 	}
 
 	RIGVM_METHOD()
-	virtual void Execute() override;
+	UE_API virtual void Execute() override;
 
 	UPROPERTY(meta=(Input))
 	FVector Axis;
@@ -133,7 +135,7 @@ struct RIGVM_API FRigVMFunction_MathQuaternionFromAxisAndAngle : public FRigVMFu
  * Makes a quaternion from euler values in degrees
  */
 USTRUCT(meta=(DisplayName="From Euler", TemplateName="FromEuler", Keywords="Make,Construct"))
-struct RIGVM_API FRigVMFunction_MathQuaternionFromEuler : public FRigVMFunction_MathQuaternionBase
+struct FRigVMFunction_MathQuaternionFromEuler : public FRigVMFunction_MathQuaternionBase
 {
 	GENERATED_BODY()
 
@@ -145,7 +147,7 @@ struct RIGVM_API FRigVMFunction_MathQuaternionFromEuler : public FRigVMFunction_
 	}
 
 	RIGVM_METHOD()
-	virtual void Execute() override;
+	UE_API virtual void Execute() override;
 
 	UPROPERTY(meta = (Input))
 	FVector Euler;
@@ -161,7 +163,7 @@ struct RIGVM_API FRigVMFunction_MathQuaternionFromEuler : public FRigVMFunction_
  * Makes a quaternion from a rotator
  */
 USTRUCT(meta=(DisplayName="From Rotator", Keywords="Make,Construct", Deprecated="5.0.1"))
-struct RIGVM_API FRigVMFunction_MathQuaternionFromRotator : public FRigVMFunction_MathQuaternionBase
+struct FRigVMFunction_MathQuaternionFromRotator : public FRigVMFunction_MathQuaternionBase
 {
 	GENERATED_BODY()
 	
@@ -172,7 +174,7 @@ struct RIGVM_API FRigVMFunction_MathQuaternionFromRotator : public FRigVMFunctio
 	}
 
 	RIGVM_METHOD()
-	virtual void Execute() override;
+	UE_API virtual void Execute() override;
 
 	UPROPERTY(meta=(Input))
 	FRotator Rotator;
@@ -181,14 +183,14 @@ struct RIGVM_API FRigVMFunction_MathQuaternionFromRotator : public FRigVMFunctio
 	FQuat Result;
 	
 	RIGVM_METHOD()
-	virtual FRigVMStructUpgradeInfo GetUpgradeInfo() const override;
+	UE_API virtual FRigVMStructUpgradeInfo GetUpgradeInfo() const override;
 };
 
 /**
  * Makes a quaternion from a rotator
  */
 USTRUCT(meta=(DisplayName="From Rotator", TemplateName="Cast", ExecuteContext="FRigVMExecuteContext", Keywords="Make,Construct"))
-struct RIGVM_API FRigVMFunction_MathQuaternionFromRotatorV2 : public FRigVMFunction_MathQuaternionBase
+struct FRigVMFunction_MathQuaternionFromRotatorV2 : public FRigVMFunction_MathQuaternionBase
 {
 	GENERATED_BODY()
 	
@@ -199,7 +201,7 @@ struct RIGVM_API FRigVMFunction_MathQuaternionFromRotatorV2 : public FRigVMFunct
 	}
 
 	RIGVM_METHOD()
-	virtual void Execute() override;
+	UE_API virtual void Execute() override;
 
 	UPROPERTY(meta=(Input))
 	FRotator Value;
@@ -212,7 +214,7 @@ struct RIGVM_API FRigVMFunction_MathQuaternionFromRotatorV2 : public FRigVMFunct
  * Makes a quaternion from two vectors, representing the shortest rotation between the two vectors.
  */
 USTRUCT(meta=(DisplayName="From Two Vectors", TemplateName="FromTwoVectors", Keywords="Make,Construct"))
-struct RIGVM_API FRigVMFunction_MathQuaternionFromTwoVectors : public FRigVMFunction_MathQuaternionBase
+struct FRigVMFunction_MathQuaternionFromTwoVectors : public FRigVMFunction_MathQuaternionBase
 {
 	GENERATED_BODY()
 	
@@ -223,7 +225,7 @@ struct RIGVM_API FRigVMFunction_MathQuaternionFromTwoVectors : public FRigVMFunc
 	}
 
 	RIGVM_METHOD()
-	virtual void Execute() override;
+	UE_API virtual void Execute() override;
 
 	UPROPERTY(meta=(Input))
 	FVector A;
@@ -239,7 +241,7 @@ struct RIGVM_API FRigVMFunction_MathQuaternionFromTwoVectors : public FRigVMFunc
  * Retrieves the axis and angle of a quaternion in radians
  */
 USTRUCT(meta=(DisplayName="To Axis And Angle", TemplateName="ToAxisAndAngle", Keywords="Make,Construct,GetAxis,GetAngle"))
-struct RIGVM_API FRigVMFunction_MathQuaternionToAxisAndAngle : public FRigVMFunction_MathQuaternionBase
+struct FRigVMFunction_MathQuaternionToAxisAndAngle : public FRigVMFunction_MathQuaternionBase
 {
 	GENERATED_BODY()
 
@@ -251,7 +253,7 @@ struct RIGVM_API FRigVMFunction_MathQuaternionToAxisAndAngle : public FRigVMFunc
 	}
 
 	RIGVM_METHOD()
-	virtual void Execute() override;
+	UE_API virtual void Execute() override;
 
 	UPROPERTY(meta=(Input))
 	FQuat Value;
@@ -267,7 +269,7 @@ struct RIGVM_API FRigVMFunction_MathQuaternionToAxisAndAngle : public FRigVMFunc
  * Retrieves the forward, right and up vectors of a quaternion
  */
 USTRUCT(meta=(DisplayName="To Vectors", TemplateName="ToVectors", Keywords="Forward,Right,Up"))
-struct RIGVM_API FRigVMFunction_MathQuaternionToVectors : public FRigVMFunction_MathQuaternionBase
+struct FRigVMFunction_MathQuaternionToVectors : public FRigVMFunction_MathQuaternionBase
 {
 	GENERATED_BODY()
 
@@ -278,7 +280,7 @@ struct RIGVM_API FRigVMFunction_MathQuaternionToVectors : public FRigVMFunction_
 	}
 
 	RIGVM_METHOD()
-	virtual void Execute() override;
+	UE_API virtual void Execute() override;
 
 	UPROPERTY(meta=(Input))
 	FQuat Value;
@@ -297,7 +299,7 @@ struct RIGVM_API FRigVMFunction_MathQuaternionToVectors : public FRigVMFunction_
  * Scales a quaternion's angle
  */
 USTRUCT(meta=(DisplayName="Scale", Keywords="Multiply,Angle,Scale", Constant, Deprecated = "5.0.1"))
-struct RIGVM_API FRigVMFunction_MathQuaternionScale : public FRigVMFunction_MathQuaternionBase
+struct FRigVMFunction_MathQuaternionScale : public FRigVMFunction_MathQuaternionBase
 {
 	GENERATED_BODY()
 
@@ -308,7 +310,7 @@ struct RIGVM_API FRigVMFunction_MathQuaternionScale : public FRigVMFunction_Math
 	}
 
 	RIGVM_METHOD()
-	virtual void Execute() override;
+	UE_API virtual void Execute() override;
 
 	UPROPERTY(meta=(Input, Output))
 	FQuat Value;
@@ -317,14 +319,14 @@ struct RIGVM_API FRigVMFunction_MathQuaternionScale : public FRigVMFunction_Math
 	float Scale;
 
 	RIGVM_METHOD()
-	virtual FRigVMStructUpgradeInfo GetUpgradeInfo() const override;
+	UE_API virtual FRigVMStructUpgradeInfo GetUpgradeInfo() const override;
 };
 
 /**
  * Scales a quaternion's angle
  */
 USTRUCT(meta=(DisplayName="Scale", TemplateName="Scale", Keywords="Multiply,Angle,Scale", Constant))
-struct RIGVM_API FRigVMFunction_MathQuaternionScaleV2 : public FRigVMFunction_MathQuaternionBase
+struct FRigVMFunction_MathQuaternionScaleV2 : public FRigVMFunction_MathQuaternionBase
 {
 	GENERATED_BODY()
 
@@ -335,7 +337,7 @@ struct RIGVM_API FRigVMFunction_MathQuaternionScaleV2 : public FRigVMFunction_Ma
 	}
 
 	RIGVM_METHOD()
-	virtual void Execute() override;
+	UE_API virtual void Execute() override;
 
 	UPROPERTY(meta=(Input))
 	FQuat Value;
@@ -351,7 +353,7 @@ struct RIGVM_API FRigVMFunction_MathQuaternionScaleV2 : public FRigVMFunction_Ma
  * Retrieves the euler angles in degrees
  */
 USTRUCT(meta=(DisplayName="To Euler", TemplateName="ToEuler", Keywords="Make,Construct"))
-struct RIGVM_API FRigVMFunction_MathQuaternionToEuler : public FRigVMFunction_MathQuaternionBase
+struct FRigVMFunction_MathQuaternionToEuler : public FRigVMFunction_MathQuaternionBase
 {
 	GENERATED_BODY()
 
@@ -363,7 +365,7 @@ struct RIGVM_API FRigVMFunction_MathQuaternionToEuler : public FRigVMFunction_Ma
 	}
 
 	RIGVM_METHOD()
-	virtual void Execute() override;
+	UE_API virtual void Execute() override;
 
 	UPROPERTY(meta=(Input))
 	FQuat Value;
@@ -379,7 +381,7 @@ struct RIGVM_API FRigVMFunction_MathQuaternionToEuler : public FRigVMFunction_Ma
  * Retrieves the rotator
  */
 USTRUCT(meta=(DisplayName="To Rotator", TemplateName="Cast", ExecuteContext="FRigVMExecuteContext", Keywords="Make,Construct"))
-struct RIGVM_API FRigVMFunction_MathQuaternionToRotator : public FRigVMFunction_MathQuaternionBase
+struct FRigVMFunction_MathQuaternionToRotator : public FRigVMFunction_MathQuaternionBase
 {
 	GENERATED_BODY()
 
@@ -390,7 +392,7 @@ struct RIGVM_API FRigVMFunction_MathQuaternionToRotator : public FRigVMFunction_
 	}
 
 	RIGVM_METHOD()
-	virtual void Execute() override;
+	UE_API virtual void Execute() override;
 
 	UPROPERTY(meta=(Input))
 	FQuat Value;
@@ -403,29 +405,29 @@ struct RIGVM_API FRigVMFunction_MathQuaternionToRotator : public FRigVMFunction_
  * Returns the product of the two values
  */
 USTRUCT(meta=(DisplayName="Multiply", TemplateName="Multiply", Keywords="Product,*"))
-struct RIGVM_API FRigVMFunction_MathQuaternionMul : public FRigVMFunction_MathQuaternionBinaryAggregateOp
+struct FRigVMFunction_MathQuaternionMul : public FRigVMFunction_MathQuaternionBinaryAggregateOp
 {
 	GENERATED_BODY()
 	RIGVM_METHOD()
-	virtual void Execute() override;
+	UE_API virtual void Execute() override;
 };
 
 /**
  * Returns the negative value
  */
 USTRUCT(meta=(DisplayName="Inverse", TemplateName="Inverse"))
-struct RIGVM_API FRigVMFunction_MathQuaternionInverse : public FRigVMFunction_MathQuaternionUnaryOp
+struct FRigVMFunction_MathQuaternionInverse : public FRigVMFunction_MathQuaternionUnaryOp
 {
 	GENERATED_BODY()
 	RIGVM_METHOD()
-	virtual void Execute() override;
+	UE_API virtual void Execute() override;
 };
 
 /**
  * Linearly interpolates between A and B using the ratio T
  */
 USTRUCT(meta=(DisplayName="Interpolate", TemplateName="Interpolate", Keywords="Lerp,Mix,Blend,Slerp,SphericalInterpolate"))
-struct RIGVM_API FRigVMFunction_MathQuaternionSlerp : public FRigVMFunction_MathQuaternionBase
+struct FRigVMFunction_MathQuaternionSlerp : public FRigVMFunction_MathQuaternionBase
 {
 	GENERATED_BODY()
 
@@ -436,7 +438,7 @@ struct RIGVM_API FRigVMFunction_MathQuaternionSlerp : public FRigVMFunction_Math
 	}
 
 	RIGVM_METHOD()
-	virtual void Execute() override;
+	UE_API virtual void Execute() override;
 
 	UPROPERTY(meta=(Input))
 	FQuat A;
@@ -455,7 +457,7 @@ struct RIGVM_API FRigVMFunction_MathQuaternionSlerp : public FRigVMFunction_Math
  * Returns true if the value A equals B
  */
 USTRUCT(meta=(DisplayName="Equals", TemplateName="Equals", Keywords="Same,==", Deprecated="5.1"))
-struct RIGVM_API FRigVMFunction_MathQuaternionEquals : public FRigVMFunction_MathQuaternionBase
+struct FRigVMFunction_MathQuaternionEquals : public FRigVMFunction_MathQuaternionBase
 {
 	GENERATED_BODY()
 
@@ -466,7 +468,7 @@ struct RIGVM_API FRigVMFunction_MathQuaternionEquals : public FRigVMFunction_Mat
 	}	
 
 	RIGVM_METHOD()
-	virtual void Execute() override;
+	UE_API virtual void Execute() override;
 
 	UPROPERTY(meta=(Input))
 	FQuat A;
@@ -478,14 +480,14 @@ struct RIGVM_API FRigVMFunction_MathQuaternionEquals : public FRigVMFunction_Mat
 	bool Result;
 
 	RIGVM_METHOD()
-	virtual FRigVMStructUpgradeInfo GetUpgradeInfo() const override;
+	UE_API virtual FRigVMStructUpgradeInfo GetUpgradeInfo() const override;
 };
 
 /**
  * Returns true if the value A does not equal B
  */
 USTRUCT(meta=(DisplayName="Not Equals", TemplateName="NotEquals", Keywords="Different,!=", Deprecated="5.1"))
-struct RIGVM_API FRigVMFunction_MathQuaternionNotEquals : public FRigVMFunction_MathQuaternionBase
+struct FRigVMFunction_MathQuaternionNotEquals : public FRigVMFunction_MathQuaternionBase
 {
 	GENERATED_BODY()
 
@@ -496,7 +498,7 @@ struct RIGVM_API FRigVMFunction_MathQuaternionNotEquals : public FRigVMFunction_
 	}
 
 	RIGVM_METHOD()
-	virtual void Execute() override;
+	UE_API virtual void Execute() override;
 
 	UPROPERTY(meta=(Input))
 	FQuat A;
@@ -508,14 +510,14 @@ struct RIGVM_API FRigVMFunction_MathQuaternionNotEquals : public FRigVMFunction_
 	bool Result;
 
 	RIGVM_METHOD()
-	virtual FRigVMStructUpgradeInfo GetUpgradeInfo() const override;
+	UE_API virtual FRigVMStructUpgradeInfo GetUpgradeInfo() const override;
 };
 
 /**
  * Return one of the two values based on the condition
  */
 USTRUCT(meta=(DisplayName="Select", Keywords="Pick,If", Deprecated = "4.26.0"))
-struct RIGVM_API FRigVMFunction_MathQuaternionSelectBool : public FRigVMFunction_MathQuaternionBase
+struct FRigVMFunction_MathQuaternionSelectBool : public FRigVMFunction_MathQuaternionBase
 {
 	GENERATED_BODY()
 
@@ -526,7 +528,7 @@ struct RIGVM_API FRigVMFunction_MathQuaternionSelectBool : public FRigVMFunction
 	}
 
 	RIGVM_METHOD()
-	virtual void Execute() override;
+	UE_API virtual void Execute() override;
 
 	UPROPERTY(meta=(Input))
 	bool Condition;
@@ -541,14 +543,14 @@ struct RIGVM_API FRigVMFunction_MathQuaternionSelectBool : public FRigVMFunction
 	FQuat Result;
 
 	RIGVM_METHOD()
-	virtual FRigVMStructUpgradeInfo GetUpgradeInfo() const override;
+	UE_API virtual FRigVMStructUpgradeInfo GetUpgradeInfo() const override;
 };
 
 /**
  * Returns the dot product between two quaternions
  */
 USTRUCT(meta=(DisplayName="Dot", TemplateName="Dot,|"))
-struct RIGVM_API FRigVMFunction_MathQuaternionDot : public FRigVMFunction_MathQuaternionBase
+struct FRigVMFunction_MathQuaternionDot : public FRigVMFunction_MathQuaternionBase
 {
 	GENERATED_BODY()
 
@@ -559,7 +561,7 @@ struct RIGVM_API FRigVMFunction_MathQuaternionDot : public FRigVMFunction_MathQu
 	}
 
 	RIGVM_METHOD()
-	virtual void Execute() override;
+	UE_API virtual void Execute() override;
 
 	UPROPERTY(meta=(Input))
 	FQuat A;
@@ -575,18 +577,18 @@ struct RIGVM_API FRigVMFunction_MathQuaternionDot : public FRigVMFunction_MathQu
  * Returns the normalized quaternion
  */
 USTRUCT(meta=(DisplayName="Unit", TemplateName="Unit", Keywords="Normalize"))
-struct RIGVM_API FRigVMFunction_MathQuaternionUnit : public FRigVMFunction_MathQuaternionUnaryOp
+struct FRigVMFunction_MathQuaternionUnit : public FRigVMFunction_MathQuaternionUnaryOp
 {
 	GENERATED_BODY()
 	RIGVM_METHOD()
-	virtual void Execute() override;
+	UE_API virtual void Execute() override;
 };
 
 /**
  * Rotates a given vector by the quaternion
  */
 USTRUCT(meta=(DisplayName="Rotate Vector", TemplateName="Rotate Vector", Keywords="Transform,Multiply"))
-struct RIGVM_API FRigVMFunction_MathQuaternionRotateVector : public FRigVMFunction_MathQuaternionBase
+struct FRigVMFunction_MathQuaternionRotateVector : public FRigVMFunction_MathQuaternionBase
 {
 	GENERATED_BODY()
 
@@ -597,7 +599,7 @@ struct RIGVM_API FRigVMFunction_MathQuaternionRotateVector : public FRigVMFuncti
 	}
 
 	RIGVM_METHOD()
-	virtual void Execute() override;
+	UE_API virtual void Execute() override;
 
 	UPROPERTY(meta=(Input))
 	FQuat Transform;
@@ -613,7 +615,7 @@ struct RIGVM_API FRigVMFunction_MathQuaternionRotateVector : public FRigVMFuncti
  * Rotates a given vector by the quaternion
  */
 USTRUCT(meta = (DisplayName = "Axis", TemplateName="Axis", Keywords = "GetAxis,xAxis,yAxis,zAxis"))
-struct RIGVM_API FRigVMFunction_MathQuaternionGetAxis: public FRigVMFunction_MathQuaternionBase
+struct FRigVMFunction_MathQuaternionGetAxis: public FRigVMFunction_MathQuaternionBase
 {
 	GENERATED_BODY()
 
@@ -625,7 +627,7 @@ struct RIGVM_API FRigVMFunction_MathQuaternionGetAxis: public FRigVMFunction_Mat
 	}
 
 	RIGVM_METHOD()
-	virtual void Execute() override;
+	UE_API virtual void Execute() override;
 
 	UPROPERTY(meta = (Input))
 	FQuat Quaternion;
@@ -642,7 +644,7 @@ struct RIGVM_API FRigVMFunction_MathQuaternionGetAxis: public FRigVMFunction_Mat
  * Computes the swing and twist components of a quaternion
  */
 USTRUCT(meta = (DisplayName = "To Swing & Twist"))
-struct RIGVM_API FRigVMFunction_MathQuaternionSwingTwist : public FRigVMFunction_MathQuaternionBase
+struct FRigVMFunction_MathQuaternionSwingTwist : public FRigVMFunction_MathQuaternionBase
 {
 	GENERATED_BODY()
 
@@ -653,7 +655,7 @@ struct RIGVM_API FRigVMFunction_MathQuaternionSwingTwist : public FRigVMFunction
 	}
 
 	RIGVM_METHOD()
-	virtual void Execute() override;
+	UE_API virtual void Execute() override;
 
 	UPROPERTY(meta = (Input))
 	FQuat	Input;
@@ -672,7 +674,7 @@ struct RIGVM_API FRigVMFunction_MathQuaternionSwingTwist : public FRigVMFunction
  * Enum of possible rotation orders
  */
 USTRUCT(meta = (DisplayName = "Rotation Order", Category = "Math|Quaternion", Constant))
-struct RIGVM_API FRigVMFunction_MathQuaternionRotationOrder : public FRigVMFunction_MathBase
+struct FRigVMFunction_MathQuaternionRotationOrder : public FRigVMFunction_MathBase
 {
 	GENERATED_BODY()
 
@@ -686,14 +688,14 @@ struct RIGVM_API FRigVMFunction_MathQuaternionRotationOrder : public FRigVMFunct
 
 
 	RIGVM_METHOD()
-	virtual void Execute() override;
+	UE_API virtual void Execute() override;
 };
 
 /**
  * Returns the relative local transform within a parent's transform
  */
 USTRUCT(meta=(DisplayName="Make Relative", TemplateName="Make Relative", Keywords="Local,Global,Absolute"))
-struct RIGVM_API FRigVMFunction_MathQuaternionMakeRelative : public FRigVMFunction_MathQuaternionBase
+struct FRigVMFunction_MathQuaternionMakeRelative : public FRigVMFunction_MathQuaternionBase
 {
 	GENERATED_BODY()
 
@@ -703,7 +705,7 @@ struct RIGVM_API FRigVMFunction_MathQuaternionMakeRelative : public FRigVMFuncti
 	}
 
 	RIGVM_METHOD()
-	virtual void Execute() override;
+	UE_API virtual void Execute() override;
 
 	UPROPERTY(meta=(Input))
 	FQuat Global;
@@ -719,7 +721,7 @@ struct RIGVM_API FRigVMFunction_MathQuaternionMakeRelative : public FRigVMFuncti
  * Returns the absolute global transform within a parent's transform
  */
 USTRUCT(meta = (DisplayName = "Make Absolute", TemplateName="Make Absolute", Keywords = "Local,Global,Relative"))
-struct RIGVM_API FRigVMFunction_MathQuaternionMakeAbsolute : public FRigVMFunction_MathQuaternionBase
+struct FRigVMFunction_MathQuaternionMakeAbsolute : public FRigVMFunction_MathQuaternionBase
 {
 	GENERATED_BODY()
 
@@ -729,7 +731,7 @@ struct RIGVM_API FRigVMFunction_MathQuaternionMakeAbsolute : public FRigVMFuncti
 	}
 
 	RIGVM_METHOD()
-	virtual void Execute() override;
+	UE_API virtual void Execute() override;
 
 	UPROPERTY(meta = (Input))
 	FQuat Local;
@@ -745,7 +747,7 @@ struct RIGVM_API FRigVMFunction_MathQuaternionMakeAbsolute : public FRigVMFuncti
  * Mirror a rotation about a central transform.
  */
 USTRUCT(meta=(DisplayName="Mirror", TemplateName="Mirror"))
-struct RIGVM_API FRigVMFunction_MathQuaternionMirrorTransform : public FRigVMFunction_MathQuaternionBase
+struct FRigVMFunction_MathQuaternionMirrorTransform : public FRigVMFunction_MathQuaternionBase
 {
 	GENERATED_BODY()
 
@@ -758,7 +760,7 @@ struct RIGVM_API FRigVMFunction_MathQuaternionMirrorTransform : public FRigVMFun
 	}
 
 	RIGVM_METHOD()
-	virtual void Execute() override;
+	UE_API virtual void Execute() override;
 
 	UPROPERTY(meta=(Input))
 	FQuat Value;
@@ -778,3 +780,5 @@ struct RIGVM_API FRigVMFunction_MathQuaternionMirrorTransform : public FRigVMFun
 	UPROPERTY(meta=(Output))
 	FQuat Result;
 };
+
+#undef UE_API
\ No newline at end of file
