[RigVM]

diff --git a/Engine/Plugins/Runtime/RigVM/Source/RigVM/Public/RigVMFunctions/Math/RigVMFunction_MathRay.h b/Engine/Plugins/Runtime/RigVM/Source/RigVM/Public/RigVMFunctions/Math/RigVMFunction_MathRay.h
--- a/Engine/Plugins/Runtime/RigVM/Source/RigVM/Public/RigVMFunctions/Math/RigVMFunction_MathRay.h
+++ b/Engine/Plugins/Runtime/RigVM/Source/RigVM/Public/RigVMFunctions/Math/RigVMFunction_MathRay.h
@@ -8,7 +8,7 @@
 #include "RigVMFunction_MathRay.generated.h"
 
 USTRUCT(meta=(Abstract, Category="Math|Ray", MenuDescSuffix="(Ray)"))
-struct RIGVM_API FRigVMFunction_MathRayBase : public FRigVMFunction_MathBase
+struct FRigVMFunction_MathRayBase : public FRigVMFunction_MathBase
 {
 	GENERATED_BODY()
 };
@@ -17,7 +17,7 @@ struct RIGVM_API FRigVMFunction_MathRayBase : public FRigVMFunction_MathBase
  * Returns the closest point intersection of a ray with another ray
  */
 USTRUCT(meta = (DisplayName = "Intersect Ray", Keywords = "Closest,Ray,Cross"))
-struct RIGVM_API FRigVMFunction_MathRayIntersectRay : public FRigVMFunction_MathRayBase
+struct FRigVMFunction_MathRayIntersectRay : public FRigVMFunction_MathRayBase
 {
 	GENERATED_BODY()
 
@@ -29,7 +29,7 @@ struct RIGVM_API FRigVMFunction_MathRayIntersectRay : public FRigVMFunction_Math
 	}
 
 	RIGVM_METHOD()
-	virtual void Execute() override;
+	RIGVM_API virtual void Execute() override;
 
 	UPROPERTY(meta = (Input))
 	FRay A;
@@ -54,7 +54,7 @@ struct RIGVM_API FRigVMFunction_MathRayIntersectRay : public FRigVMFunction_Math
  * Returns the closest point intersection of a ray with a plane
  */
 USTRUCT(meta = (DisplayName = "Intersect Plane", Keywords = "Closest,Ray,Cross"))
-struct RIGVM_API FRigVMFunction_MathRayIntersectPlane : public FRigVMFunction_MathRayBase
+struct FRigVMFunction_MathRayIntersectPlane : public FRigVMFunction_MathRayBase
 {
 	GENERATED_BODY()
 
@@ -68,7 +68,7 @@ struct RIGVM_API FRigVMFunction_MathRayIntersectPlane : public FRigVMFunction_Ma
 	}
 
 	RIGVM_METHOD()
-	virtual void Execute() override;
+	RIGVM_API virtual void Execute() override;
 
 	UPROPERTY(meta = (Input))
 	FRay Ray;
@@ -93,7 +93,7 @@ struct RIGVM_API FRigVMFunction_MathRayIntersectPlane : public FRigVMFunction_Ma
  * Returns the position on a ray
  */
 USTRUCT(meta=(DisplayName="GetAt", Keywords="Ratio,Percentage"))
-struct RIGVM_API FRigVMFunction_MathRayGetAt : public FRigVMFunction_MathRayBase
+struct FRigVMFunction_MathRayGetAt : public FRigVMFunction_MathRayBase
 {
 	GENERATED_BODY()
 
@@ -105,7 +105,7 @@ struct RIGVM_API FRigVMFunction_MathRayGetAt : public FRigVMFunction_MathRayBase
 	}
 
 	RIGVM_METHOD()
-	virtual void Execute() override;
+	RIGVM_API virtual void Execute() override;
 
 	UPROPERTY(meta = (Input))
 	FRay Ray;
@@ -121,7 +121,7 @@ struct RIGVM_API FRigVMFunction_MathRayGetAt : public FRigVMFunction_MathRayBase
  * Returns the division of the two values
  */
 USTRUCT(meta=(DisplayName="Transform Ray", Keywords="Multiply,Project"))
-struct RIGVM_API FRigVMFunction_MathRayTransform : public FRigVMFunction_MathRayBase
+struct FRigVMFunction_MathRayTransform : public FRigVMFunction_MathRayBase
 {
 	GENERATED_BODY()
 
@@ -132,7 +132,7 @@ struct RIGVM_API FRigVMFunction_MathRayTransform : public FRigVMFunction_MathRay
 	}
 
 	RIGVM_METHOD()
-	virtual void Execute() override;
+	RIGVM_API virtual void Execute() override;
 
 	UPROPERTY(meta = (Input))
 	FRay Ray;
diff --git a/Engine/Plugins/Runtime/RigVM/Source/RigVM/Public/RigVMFunctions/Math/RigVMFunction_MathTransform.h b/Engine/Plugins/Runtime/RigVM/Source/RigVM/Public/RigVMFunctions/Math/RigVMFunction_MathTransform.h
--- a/Engine/Plugins/Runtime/RigVM/Source/RigVM/Public/RigVMFunctions/Math/RigVMFunction_MathTransform.h
+++ b/Engine/Plugins/Runtime/RigVM/Source/RigVM/Public/RigVMFunctions/Math/RigVMFunction_MathTransform.h
@@ -6,20 +6,22 @@
 #include "RigVMFunction_MathBase.h"
 #include "RigVMFunction_MathTransform.generated.h"
 
+#define UE_API RIGVM_API
+
 USTRUCT(meta=(Abstract, Category="Math|Transform", MenuDescSuffix="(Transform)"))
-struct RIGVM_API FRigVMFunction_MathTransformBase : public FRigVMFunction_MathBase
+struct FRigVMFunction_MathTransformBase : public FRigVMFunction_MathBase
 {
 	GENERATED_BODY()
 };
 
 USTRUCT(meta=(Abstract, Category="Math|Transform", MenuDescSuffix="(Transform)"))
-struct RIGVM_API FRigVMFunction_MathTransformMutableBase : public FRigVMFunction_MathMutableBase
+struct FRigVMFunction_MathTransformMutableBase : public FRigVMFunction_MathMutableBase
 {
 	GENERATED_BODY()
 };
 
 USTRUCT(meta=(Abstract))
-struct RIGVM_API FRigVMFunction_MathTransformUnaryOp : public FRigVMFunction_MathTransformBase
+struct FRigVMFunction_MathTransformUnaryOp : public FRigVMFunction_MathTransformBase
 {
 	GENERATED_BODY()
 
@@ -36,7 +38,7 @@ struct RIGVM_API FRigVMFunction_MathTransformUnaryOp : public FRigVMFunction_Mat
 };
 
 USTRUCT(meta=(Abstract))
-struct RIGVM_API FRigVMFunction_MathTransformBinaryOp : public FRigVMFunction_MathTransformBase
+struct FRigVMFunction_MathTransformBinaryOp : public FRigVMFunction_MathTransformBase
 {
 	GENERATED_BODY()
 
@@ -56,7 +58,7 @@ struct RIGVM_API FRigVMFunction_MathTransformBinaryOp : public FRigVMFunction_Ma
 };
 
 USTRUCT(meta=(Abstract))
-struct RIGVM_API FRigVMFunction_MathTransformBinaryAggregateOp : public FRigVMFunction_MathTransformBase
+struct FRigVMFunction_MathTransformBinaryAggregateOp : public FRigVMFunction_MathTransformBase
 {
 	GENERATED_BODY()
 
@@ -79,11 +81,11 @@ struct RIGVM_API FRigVMFunction_MathTransformBinaryAggregateOp : public FRigVMFu
  * Makes a transform from its components
  */
 USTRUCT(meta=(DisplayName="Make Transform", Keywords="Make,Construct,Constant"))
-struct RIGVM_API FRigVMFunction_MathTransformMake : public FRigVMFunction_MathTransformBase
+struct FRigVMFunction_MathTransformMake : public FRigVMFunction_MathTransformBase
 {
 	GENERATED_BODY()
 	RIGVM_METHOD()
-	virtual void Execute() override;
+	UE_API virtual void Execute() override;
 
 	FRigVMFunction_MathTransformMake()
 	{
@@ -110,7 +112,7 @@ struct RIGVM_API FRigVMFunction_MathTransformMake : public FRigVMFunction_MathTr
  * Makes a quaternion based transform from a euler based transform
  */
 USTRUCT(meta=(DisplayName="From Euler Transform", TemplateName="FromEulerTransform", Keywords="Make,Construct", Deprecated="5.0.1"))
-struct RIGVM_API FRigVMFunction_MathTransformFromEulerTransform : public FRigVMFunction_MathTransformBase
+struct FRigVMFunction_MathTransformFromEulerTransform : public FRigVMFunction_MathTransformBase
 {
 	GENERATED_BODY()
 
@@ -121,7 +123,7 @@ struct RIGVM_API FRigVMFunction_MathTransformFromEulerTransform : public FRigVMF
 	}
 
 	RIGVM_METHOD()
-	virtual void Execute() override;
+	UE_API virtual void Execute() override;
 
 	UPROPERTY(meta=(Input))
 	FEulerTransform EulerTransform;
@@ -130,14 +132,14 @@ struct RIGVM_API FRigVMFunction_MathTransformFromEulerTransform : public FRigVMF
 	FTransform Result;
 	
 	RIGVM_METHOD()
-	virtual FRigVMStructUpgradeInfo GetUpgradeInfo() const override;
+	UE_API virtual FRigVMStructUpgradeInfo GetUpgradeInfo() const override;
 };
 
 /**
  * Makes a quaternion based transform from a euler based transform
  */
 USTRUCT(meta=(DisplayName="To Transform", TemplateName="Cast", ExecuteContext="FRigVMExecuteContext", Keywords="Make,Construct"))
-struct RIGVM_API FRigVMFunction_MathTransformFromEulerTransformV2 : public FRigVMFunction_MathTransformBase
+struct FRigVMFunction_MathTransformFromEulerTransformV2 : public FRigVMFunction_MathTransformBase
 {
 	GENERATED_BODY()
 
@@ -148,7 +150,7 @@ struct RIGVM_API FRigVMFunction_MathTransformFromEulerTransformV2 : public FRigV
 	}
 
 	RIGVM_METHOD()
-	virtual void Execute() override;
+	UE_API virtual void Execute() override;
 
 	UPROPERTY(meta=(Input))
 	FEulerTransform Value;
@@ -161,7 +163,7 @@ struct RIGVM_API FRigVMFunction_MathTransformFromEulerTransformV2 : public FRigV
  * Retrieves a euler based transform from a quaternion based transform
  */
 USTRUCT(meta=(DisplayName="To Euler Transform", TemplateName="Cast", ExecuteContext="FRigVMExecuteContext", Keywords="Make,Construct"))
-struct RIGVM_API FRigVMFunction_MathTransformToEulerTransform : public FRigVMFunction_MathTransformBase
+struct FRigVMFunction_MathTransformToEulerTransform : public FRigVMFunction_MathTransformBase
 {
 	GENERATED_BODY()
 
@@ -172,7 +174,7 @@ struct RIGVM_API FRigVMFunction_MathTransformToEulerTransform : public FRigVMFun
 	}
 
 	RIGVM_METHOD()
-	virtual void Execute() override;
+	UE_API virtual void Execute() override;
 
 	UPROPERTY(meta=(Input))
 	FTransform Value;
@@ -185,7 +187,7 @@ struct RIGVM_API FRigVMFunction_MathTransformToEulerTransform : public FRigVMFun
  * Retrieves the forward, right and up vectors of a transform's quaternion
  */
 USTRUCT(meta=(DisplayName="To Vectors", TemplateName="ToVectors", Keywords="Forward,Right,Up"))
-struct RIGVM_API FRigVMFunction_MathTransformToVectors : public FRigVMFunction_MathTransformBase
+struct FRigVMFunction_MathTransformToVectors : public FRigVMFunction_MathTransformBase
 {
 	GENERATED_BODY()
 
@@ -196,7 +198,7 @@ struct RIGVM_API FRigVMFunction_MathTransformToVectors : public FRigVMFunction_M
 	}
 
 	RIGVM_METHOD()
-	virtual void Execute() override;
+	UE_API virtual void Execute() override;
 
 	UPROPERTY(meta=(Input))
 	FTransform Value;
@@ -215,18 +217,18 @@ struct RIGVM_API FRigVMFunction_MathTransformToVectors : public FRigVMFunction_M
  * Returns the product of the two values
  */
 USTRUCT(meta=(DisplayName="Multiply", TemplateName="Multiply", Keywords="Product,*,Global"))
-struct RIGVM_API FRigVMFunction_MathTransformMul : public FRigVMFunction_MathTransformBinaryAggregateOp
+struct FRigVMFunction_MathTransformMul : public FRigVMFunction_MathTransformBinaryAggregateOp
 {
 	GENERATED_BODY()
 	RIGVM_METHOD()
-	virtual void Execute() override;
+	UE_API virtual void Execute() override;
 };
 
 /**
  * Returns the relative local transform within a parent's transform
  */
 USTRUCT(meta=(DisplayName="Make Relative", TemplateName="Make Relative", Keywords="Local,Global,Absolute"))
-struct RIGVM_API FRigVMFunction_MathTransformMakeRelative : public FRigVMFunction_MathTransformBase
+struct FRigVMFunction_MathTransformMakeRelative : public FRigVMFunction_MathTransformBase
 {
 	GENERATED_BODY()
 
@@ -236,7 +238,7 @@ struct RIGVM_API FRigVMFunction_MathTransformMakeRelative : public FRigVMFunctio
 	}
 
 	RIGVM_METHOD()
-	virtual void Execute() override;
+	UE_API virtual void Execute() override;
 
 	UPROPERTY(meta=(Input))
 	FTransform Global;
@@ -252,7 +254,7 @@ struct RIGVM_API FRigVMFunction_MathTransformMakeRelative : public FRigVMFunctio
  * Returns the absolute global transform within a parent's transform
  */
 USTRUCT(meta = (DisplayName = "Make Absolute", TemplateName="Make Absolute", Keywords = "Local,Global,Relative"))
-struct RIGVM_API FRigVMFunction_MathTransformMakeAbsolute : public FRigVMFunction_MathTransformBase
+struct FRigVMFunction_MathTransformMakeAbsolute : public FRigVMFunction_MathTransformBase
 {
 	GENERATED_BODY()
 
@@ -262,7 +264,7 @@ struct RIGVM_API FRigVMFunction_MathTransformMakeAbsolute : public FRigVMFunctio
 	}
 
 	RIGVM_METHOD()
-	virtual void Execute() override;
+	UE_API virtual void Execute() override;
 
 	UPROPERTY(meta = (Input))
 	FTransform Local;
@@ -280,7 +282,7 @@ struct RIGVM_API FRigVMFunction_MathTransformMakeAbsolute : public FRigVMFunctio
 * a custom parent indices array, with which you can represent more than just chains.
 */
 USTRUCT(meta=(DisplayName="Make Transform Array Relative", Keywords="Local,Global,Absolute,Array,Accumulate"))
-struct RIGVM_API FRigVMFunction_MathTransformAccumulateArray : public FRigVMFunction_MathTransformMutableBase
+struct FRigVMFunction_MathTransformAccumulateArray : public FRigVMFunction_MathTransformMutableBase
 {
 	GENERATED_BODY()
 
@@ -290,10 +292,10 @@ struct RIGVM_API FRigVMFunction_MathTransformAccumulateArray : public FRigVMFunc
 		Root = FTransform::Identity;
 	}
 
-	virtual FString GetUnitLabel() const override;
+	UE_API virtual FString GetUnitLabel() const override;
 
 	RIGVM_METHOD()
-	virtual void Execute() override;
+	UE_API virtual void Execute() override;
 
 	UPROPERTY(meta=(Input, Output))
 	TArray<FTransform> Transforms;
@@ -322,18 +324,18 @@ struct RIGVM_API FRigVMFunction_MathTransformAccumulateArray : public FRigVMFunc
  * Returns the negative value
  */
 USTRUCT(meta=(DisplayName="Inverse", TemplateName="Inverse"))
-struct RIGVM_API FRigVMFunction_MathTransformInverse : public FRigVMFunction_MathTransformUnaryOp
+struct FRigVMFunction_MathTransformInverse : public FRigVMFunction_MathTransformUnaryOp
 {
 	GENERATED_BODY()
 	RIGVM_METHOD()
-	virtual void Execute() override;
+	UE_API virtual void Execute() override;
 };
 
 /**
  * Linearly interpolates between A and B using the ratio T
  */
 USTRUCT(meta=(DisplayName="Interpolate", TemplateName="Interpolate", Keywords="Lerp,Mix,Blend"))
-struct RIGVM_API FRigVMFunction_MathTransformLerp : public FRigVMFunction_MathTransformBase
+struct FRigVMFunction_MathTransformLerp : public FRigVMFunction_MathTransformBase
 {
 	GENERATED_BODY()
 
@@ -344,7 +346,7 @@ struct RIGVM_API FRigVMFunction_MathTransformLerp : public FRigVMFunction_MathTr
 	}
 
 	RIGVM_METHOD()
-	virtual void Execute() override;
+	UE_API virtual void Execute() override;
 
 	UPROPERTY(meta=(Input))
 	FTransform A;
@@ -363,7 +365,7 @@ struct RIGVM_API FRigVMFunction_MathTransformLerp : public FRigVMFunction_MathTr
  * Return one of the two values based on the condition
  */
 USTRUCT(meta=(DisplayName="Select", Keywords="Pick,If", Deprecated = "4.26.0"))
-struct RIGVM_API FRigVMFunction_MathTransformSelectBool : public FRigVMFunction_MathTransformBase
+struct FRigVMFunction_MathTransformSelectBool : public FRigVMFunction_MathTransformBase
 {
 	GENERATED_BODY()
 
@@ -374,7 +376,7 @@ struct RIGVM_API FRigVMFunction_MathTransformSelectBool : public FRigVMFunction_
 	}
 
 	RIGVM_METHOD()
-	virtual void Execute() override;
+	UE_API virtual void Execute() override;
 
 	UPROPERTY(meta=(Input))
 	bool Condition;
@@ -396,7 +398,7 @@ struct RIGVM_API FRigVMFunction_MathTransformSelectBool : public FRigVMFunction_
  * Rotates a given vector (direction) by the transform
  */
 USTRUCT(meta=(DisplayName="Rotate Vector", TemplateName="Rotate Vector", Keywords="Transform,Direction,TransformDirection"))
-struct RIGVM_API FRigVMFunction_MathTransformRotateVector : public FRigVMFunction_MathTransformBase
+struct FRigVMFunction_MathTransformRotateVector : public FRigVMFunction_MathTransformBase
 {
 	GENERATED_BODY()
 
@@ -407,7 +409,7 @@ struct RIGVM_API FRigVMFunction_MathTransformRotateVector : public FRigVMFunctio
 	}
 
 	RIGVM_METHOD()
-	virtual void Execute() override;
+	UE_API virtual void Execute() override;
 
 	UPROPERTY(meta=(Input))
 	FTransform Transform;
@@ -423,7 +425,7 @@ struct RIGVM_API FRigVMFunction_MathTransformRotateVector : public FRigVMFunctio
  * Multiplies a given vector (location) by the transform
  */
 USTRUCT(meta=(DisplayName="Transform Location", Keywords="Multiply"))
-struct RIGVM_API FRigVMFunction_MathTransformTransformVector : public FRigVMFunction_MathTransformBase
+struct FRigVMFunction_MathTransformTransformVector : public FRigVMFunction_MathTransformBase
 {
 	GENERATED_BODY()
 
@@ -434,7 +436,7 @@ struct RIGVM_API FRigVMFunction_MathTransformTransformVector : public FRigVMFunc
 	}
 	
 	RIGVM_METHOD()
-	virtual void Execute() override;
+	UE_API virtual void Execute() override;
 
 	UPROPERTY(meta=(Input))
 	FTransform Transform;
@@ -450,7 +452,7 @@ struct RIGVM_API FRigVMFunction_MathTransformTransformVector : public FRigVMFunc
  * Composes a Transform (and Euler Transform) from its components.
  */
 USTRUCT(meta=(DisplayName="Transform from SRT", Keywords ="EulerTransform,Scale,Rotation,Orientation,Translation,Location"))
-struct RIGVM_API FRigVMFunction_MathTransformFromSRT : public FRigVMFunction_MathTransformBase
+struct FRigVMFunction_MathTransformFromSRT : public FRigVMFunction_MathTransformBase
 {
 	GENERATED_BODY()
 
@@ -465,7 +467,7 @@ struct RIGVM_API FRigVMFunction_MathTransformFromSRT : public FRigVMFunction_Mat
 	}
 	
 	RIGVM_METHOD()
-	virtual void Execute() override;
+	UE_API virtual void Execute() override;
 
 	UPROPERTY(meta=(Input))
 	FVector Location;
@@ -491,7 +493,7 @@ struct RIGVM_API FRigVMFunction_MathTransformFromSRT : public FRigVMFunction_Mat
  * Decomposes a Transform Array to its components.
  */
 USTRUCT(meta = (DisplayName = "Transform Array to SRT", Keywords = "EulerTransform,Scale,Rotation,Orientation,Translation,Location"))
-struct RIGVM_API FRigVMFunction_MathTransformArrayToSRT : public FRigVMFunction_MathTransformBase
+struct FRigVMFunction_MathTransformArrayToSRT : public FRigVMFunction_MathTransformBase
 {
 	GENERATED_BODY()
 
@@ -500,7 +502,7 @@ struct RIGVM_API FRigVMFunction_MathTransformArrayToSRT : public FRigVMFunction_
 	}
 
 	RIGVM_METHOD()
-		virtual void Execute() override;
+	UE_API virtual void Execute() override;
 
 	UPROPERTY(meta = (Input))
 	TArray<FTransform> Transforms;
@@ -519,7 +521,7 @@ struct RIGVM_API FRigVMFunction_MathTransformArrayToSRT : public FRigVMFunction_
  * Clamps a position using a plane collision, cylindric collision or spherical collision.
  */
 USTRUCT(meta = (DisplayName = "Clamp Spatially", TemplateName="ClampSpatially", Keywords = "Collide,Collision"))
-struct RIGVM_API FRigVMFunction_MathTransformClampSpatially : public FRigVMFunction_MathTransformBase
+struct FRigVMFunction_MathTransformClampSpatially : public FRigVMFunction_MathTransformBase
 {
 	GENERATED_BODY()
 
@@ -537,7 +539,7 @@ struct RIGVM_API FRigVMFunction_MathTransformClampSpatially : public FRigVMFunct
 	}
 
 	RIGVM_METHOD()
-	virtual void Execute() override;
+	UE_API virtual void Execute() override;
 
 	UPROPERTY(meta = (Input))
 	FTransform Value;
@@ -583,7 +585,7 @@ struct RIGVM_API FRigVMFunction_MathTransformClampSpatially : public FRigVMFunct
  * Mirror a transform about a central transform.
  */
 USTRUCT(meta=(DisplayName="Mirror", TemplateName="Mirror"))
-struct RIGVM_API FRigVMFunction_MathTransformMirrorTransform : public FRigVMFunction_MathTransformBase
+struct FRigVMFunction_MathTransformMirrorTransform : public FRigVMFunction_MathTransformBase
 {
 	GENERATED_BODY()
 
@@ -596,7 +598,7 @@ struct RIGVM_API FRigVMFunction_MathTransformMirrorTransform : public FRigVMFunc
 	}
 
 	RIGVM_METHOD()
-	virtual void Execute() override;
+	UE_API virtual void Execute() override;
 
 	UPROPERTY(meta=(Input))
 	FTransform Value;
@@ -616,3 +618,5 @@ struct RIGVM_API FRigVMFunction_MathTransformMirrorTransform : public FRigVMFunc
 	UPROPERTY(meta=(Output))
 	FTransform Result;
 };
+
+#undef UE_API 
\ No newline at end of file
diff --git a/Engine/Plugins/Runtime/RigVM/Source/RigVM/Public/RigVMFunctions/Math/RigVMFunction_MathVector.h b/Engine/Plugins/Runtime/RigVM/Source/RigVM/Public/RigVMFunctions/Math/RigVMFunction_MathVector.h
--- a/Engine/Plugins/Runtime/RigVM/Source/RigVM/Public/RigVMFunctions/Math/RigVMFunction_MathVector.h
+++ b/Engine/Plugins/Runtime/RigVM/Source/RigVM/Public/RigVMFunctions/Math/RigVMFunction_MathVector.h
@@ -7,13 +7,13 @@
 #include "RigVMFunction_MathVector.generated.h"
 
 USTRUCT(meta=(Abstract, Category="Math|Vector", MenuDescSuffix="(Vector)"))
-struct RIGVM_API FRigVMFunction_MathVectorBase : public FRigVMFunction_MathBase
+struct FRigVMFunction_MathVectorBase : public FRigVMFunction_MathBase
 {
 	GENERATED_BODY()
 };
 
 USTRUCT(meta=(Abstract))
-struct RIGVM_API FRigVMFunction_MathVectorUnaryOp : public FRigVMFunction_MathVectorBase
+struct FRigVMFunction_MathVectorUnaryOp : public FRigVMFunction_MathVectorBase
 {
 	GENERATED_BODY()
 	
@@ -30,7 +30,7 @@ struct RIGVM_API FRigVMFunction_MathVectorUnaryOp : public FRigVMFunction_MathVe
 };
 
 USTRUCT(meta=(Abstract))
-struct RIGVM_API FRigVMFunction_MathVectorBinaryOp : public FRigVMFunction_MathVectorBase
+struct FRigVMFunction_MathVectorBinaryOp : public FRigVMFunction_MathVectorBase
 {
 	GENERATED_BODY()
 
@@ -50,7 +50,7 @@ struct RIGVM_API FRigVMFunction_MathVectorBinaryOp : public FRigVMFunction_MathV
 };
 
 USTRUCT(meta=(Abstract))
-struct RIGVM_API FRigVMFunction_MathVectorBinaryAggregateOp : public FRigVMFunction_MathVectorBase
+struct FRigVMFunction_MathVectorBinaryAggregateOp : public FRigVMFunction_MathVectorBase
 {
 	GENERATED_BODY()
 
@@ -73,11 +73,11 @@ struct RIGVM_API FRigVMFunction_MathVectorBinaryAggregateOp : public FRigVMFunct
  * Makes a vector from its components
  */
 USTRUCT(meta=(DisplayName="Make Vector", Keywords="Make,Construct,Constant"))
-struct RIGVM_API FRigVMFunction_MathVectorMake : public FRigVMFunction_MathVectorBase
+struct FRigVMFunction_MathVectorMake : public FRigVMFunction_MathVectorBase
 {
 	GENERATED_BODY()
 	RIGVM_METHOD()
-	virtual void Execute() override;
+	RIGVM_API virtual void Execute() override;
 
 	FRigVMFunction_MathVectorMake()
 	{
@@ -102,11 +102,11 @@ struct RIGVM_API FRigVMFunction_MathVectorMake : public FRigVMFunction_MathVecto
  * Makes a vector from a single float
  */
 USTRUCT(meta=(DisplayName="From Float", TemplateName="Cast", ExecuteContext="FRigVMExecuteContext", Keywords="Make,Construct"))
-struct RIGVM_API FRigVMFunction_MathVectorFromFloat : public FRigVMFunction_MathVectorBase
+struct FRigVMFunction_MathVectorFromFloat : public FRigVMFunction_MathVectorBase
 {
 	GENERATED_BODY()
 	RIGVM_METHOD()
-	virtual void Execute() override;
+	RIGVM_API virtual void Execute() override;
 
 	FRigVMFunction_MathVectorFromFloat()
 	{
@@ -125,11 +125,11 @@ struct RIGVM_API FRigVMFunction_MathVectorFromFloat : public FRigVMFunction_Math
  * Makes a vector from a single double
  */
 USTRUCT(meta=(DisplayName="From Float", TemplateName="Cast", ExecuteContext="FRigVMExecuteContext", Keywords="Make,Construct"))
-struct RIGVM_API FRigVMFunction_MathVectorFromDouble : public FRigVMFunction_MathVectorBase
+struct FRigVMFunction_MathVectorFromDouble : public FRigVMFunction_MathVectorBase
 {
 	GENERATED_BODY()
 	RIGVM_METHOD()
-	virtual void Execute() override;
+	RIGVM_API virtual void Execute() override;
 
 	FRigVMFunction_MathVectorFromDouble()
 	{
@@ -148,31 +148,31 @@ struct RIGVM_API FRigVMFunction_MathVectorFromDouble : public FRigVMFunction_Mat
  * Returns the sum of the two values
  */
 USTRUCT(meta=(DisplayName="Add", TemplateName="Add", Keywords="Sum,+"))
-struct RIGVM_API FRigVMFunction_MathVectorAdd : public FRigVMFunction_MathVectorBinaryAggregateOp
+struct FRigVMFunction_MathVectorAdd : public FRigVMFunction_MathVectorBinaryAggregateOp
 {
 	GENERATED_BODY()
 
 	RIGVM_METHOD()
-	virtual void Execute() override;
+	RIGVM_API virtual void Execute() override;
 };
 
 /**
  * Returns the difference of the two values
  */
 USTRUCT(meta=(DisplayName="Subtract", TemplateName="Subtract", Keywords="-"))
-struct RIGVM_API FRigVMFunction_MathVectorSub : public FRigVMFunction_MathVectorBinaryOp
+struct FRigVMFunction_MathVectorSub : public FRigVMFunction_MathVectorBinaryOp
 {
 	GENERATED_BODY()
 
 	RIGVM_METHOD()
-	virtual void Execute() override;
+	RIGVM_API virtual void Execute() override;
 };
 
 /**
  * Returns the product of the two values
  */
 USTRUCT(meta=(DisplayName="Multiply", TemplateName="Multiply", Keywords="Product,*"))
-struct RIGVM_API FRigVMFunction_MathVectorMul : public FRigVMFunction_MathVectorBinaryAggregateOp
+struct FRigVMFunction_MathVectorMul : public FRigVMFunction_MathVectorBinaryAggregateOp
 {
 	GENERATED_BODY()
 
@@ -182,14 +182,14 @@ struct RIGVM_API FRigVMFunction_MathVectorMul : public FRigVMFunction_MathVector
 	}
 
 	RIGVM_METHOD()
-	virtual void Execute() override;
+	RIGVM_API virtual void Execute() override;
 };
 
 /**
  * Returns the product of the the vector and the float value
  */
 USTRUCT(meta = (DisplayName = "Scale", TemplateName="Scale", Keywords = "Multiply,Product,*,ByScalar,ByFloat"))
-struct RIGVM_API FRigVMFunction_MathVectorScale : public FRigVMFunction_MathVectorBase
+struct FRigVMFunction_MathVectorScale : public FRigVMFunction_MathVectorBase
 {
 	GENERATED_BODY()
 
@@ -200,7 +200,7 @@ struct RIGVM_API FRigVMFunction_MathVectorScale : public FRigVMFunction_MathVect
 	}
 
 	RIGVM_METHOD()
-	virtual void Execute() override;
+	RIGVM_API virtual void Execute() override;
 
 	UPROPERTY(meta = (Input))
 	FVector Value;
@@ -216,7 +216,7 @@ struct RIGVM_API FRigVMFunction_MathVectorScale : public FRigVMFunction_MathVect
  * Returns the division of the two values
  */
 USTRUCT(meta=(DisplayName="Divide", TemplateName="Divide", Keywords="Division,Divisor,/"))
-struct RIGVM_API FRigVMFunction_MathVectorDiv : public FRigVMFunction_MathVectorBinaryOp
+struct FRigVMFunction_MathVectorDiv : public FRigVMFunction_MathVectorBinaryOp
 {
 	GENERATED_BODY()
 
@@ -226,14 +226,14 @@ struct RIGVM_API FRigVMFunction_MathVectorDiv : public FRigVMFunction_MathVector
 	}
 
 	RIGVM_METHOD()
-	virtual void Execute() override;
+	RIGVM_API virtual void Execute() override;
 };
 
 /**
  * Returns the modulo of the two values
  */
 USTRUCT(meta=(DisplayName="Modulo", TemplateName="Modulo", Keywords="%,fmod"))
-struct RIGVM_API FRigVMFunction_MathVectorMod : public FRigVMFunction_MathVectorBinaryOp
+struct FRigVMFunction_MathVectorMod : public FRigVMFunction_MathVectorBinaryOp
 {
 	GENERATED_BODY()
 
@@ -244,110 +244,110 @@ struct RIGVM_API FRigVMFunction_MathVectorMod : public FRigVMFunction_MathVector
 	}
 
 	RIGVM_METHOD()
-	virtual void Execute() override;
+	RIGVM_API virtual void Execute() override;
 };
 
 /**
  * Returns the smaller of the two values for each component
  */
 USTRUCT(meta=(DisplayName="Minimum", TemplateName="Minimum"))
-struct RIGVM_API FRigVMFunction_MathVectorMin : public FRigVMFunction_MathVectorBinaryAggregateOp
+struct FRigVMFunction_MathVectorMin : public FRigVMFunction_MathVectorBinaryAggregateOp
 {
 	GENERATED_BODY()
 
 	RIGVM_METHOD()
-	virtual void Execute() override;
+	RIGVM_API virtual void Execute() override;
 };
 
 /**
  * Returns the larger of the two values each component
  */
 USTRUCT(meta=(DisplayName="Maximum", TemplateName="Maximum"))
-struct RIGVM_API FRigVMFunction_MathVectorMax : public FRigVMFunction_MathVectorBinaryAggregateOp
+struct FRigVMFunction_MathVectorMax : public FRigVMFunction_MathVectorBinaryAggregateOp
 {
 	GENERATED_BODY()
 
 	RIGVM_METHOD()
-	virtual void Execute() override;
+	RIGVM_API virtual void Execute() override;
 };
 
 /**
  * Returns the negative value
  */
 USTRUCT(meta=(DisplayName="Negate", TemplateName="Negate", Keywords="-,Abs"))
-struct RIGVM_API FRigVMFunction_MathVectorNegate : public FRigVMFunction_MathVectorUnaryOp
+struct FRigVMFunction_MathVectorNegate : public FRigVMFunction_MathVectorUnaryOp
 {
 	GENERATED_BODY()
 
 	RIGVM_METHOD()
-	virtual void Execute() override;
+	RIGVM_API virtual void Execute() override;
 };
 
 /**
  * Returns the absolute (positive) value
  */
 USTRUCT(meta=(DisplayName="Absolute", TemplateName="Absolute", Keywords="Abs,Neg"))
-struct RIGVM_API FRigVMFunction_MathVectorAbs : public FRigVMFunction_MathVectorUnaryOp
+struct FRigVMFunction_MathVectorAbs : public FRigVMFunction_MathVectorUnaryOp
 {
 	GENERATED_BODY()
 
 	RIGVM_METHOD()
-	virtual void Execute() override;
+	RIGVM_API virtual void Execute() override;
 };
 
 /**
  * Returns the closest lower full number (integer) of the value for each component
  */
 USTRUCT(meta=(DisplayName="Floor", Keywords="Round"))
-struct RIGVM_API FRigVMFunction_MathVectorFloor : public FRigVMFunction_MathVectorUnaryOp
+struct FRigVMFunction_MathVectorFloor : public FRigVMFunction_MathVectorUnaryOp
 {
 	GENERATED_BODY()
 
 	RIGVM_METHOD()
-	virtual void Execute() override;
+	RIGVM_API virtual void Execute() override;
 };
 
 /**
  * Returns the closest higher full number (integer) of the value for each component
  */
 USTRUCT(meta=(DisplayName="Ceiling", Keywords="Round"))
-struct RIGVM_API FRigVMFunction_MathVectorCeil : public FRigVMFunction_MathVectorUnaryOp
+struct FRigVMFunction_MathVectorCeil : public FRigVMFunction_MathVectorUnaryOp
 {
 	GENERATED_BODY()
 
 	RIGVM_METHOD()
-	virtual void Execute() override;
+	RIGVM_API virtual void Execute() override;
 };
 
 /**
  * Returns the closest higher full number (integer) of the value for each component
  */
 USTRUCT(meta=(DisplayName="Round"))
-struct RIGVM_API FRigVMFunction_MathVectorRound : public FRigVMFunction_MathVectorUnaryOp
+struct FRigVMFunction_MathVectorRound : public FRigVMFunction_MathVectorUnaryOp
 {
 	GENERATED_BODY()
 
 	RIGVM_METHOD()
-	virtual void Execute() override;
+	RIGVM_API virtual void Execute() override;
 };
 
 /**
  * Returns the sign of the value (+1 for >= FVector(0.f, 0.f, 0.f), -1 for < 0.f) for each component
  */
 USTRUCT(meta=(DisplayName="Sign", TemplateName="Sign"))
-struct RIGVM_API FRigVMFunction_MathVectorSign : public FRigVMFunction_MathVectorUnaryOp
+struct FRigVMFunction_MathVectorSign : public FRigVMFunction_MathVectorUnaryOp
 {
 	GENERATED_BODY()
 
 	RIGVM_METHOD()
-	virtual void Execute() override;
+	RIGVM_API virtual void Execute() override;
 };
 
 /**
  * Clamps the given value within the range provided by minimum and maximum for each component
  */
 USTRUCT(meta=(DisplayName="Clamp", TemplateName="Clamp", Keywords="Range,Remap"))
-struct RIGVM_API FRigVMFunction_MathVectorClamp : public FRigVMFunction_MathVectorBase
+struct FRigVMFunction_MathVectorClamp : public FRigVMFunction_MathVectorBase
 {
 	GENERATED_BODY()
 
@@ -358,7 +358,7 @@ struct RIGVM_API FRigVMFunction_MathVectorClamp : public FRigVMFunction_MathVect
 	}
 
 	RIGVM_METHOD()
-	virtual void Execute() override;
+	RIGVM_API virtual void Execute() override;
 
 	UPROPERTY(meta=(Input))
 	FVector Value;
@@ -377,7 +377,7 @@ struct RIGVM_API FRigVMFunction_MathVectorClamp : public FRigVMFunction_MathVect
  * Linearly interpolates between A and B using the ratio T
  */
 USTRUCT(meta=(DisplayName="Interpolate", TemplateName="Interpolate", Keywords="Lerp,Mix,Blend"))
-struct RIGVM_API FRigVMFunction_MathVectorLerp : public FRigVMFunction_MathVectorBase
+struct FRigVMFunction_MathVectorLerp : public FRigVMFunction_MathVectorBase
 {
 	GENERATED_BODY()
 
@@ -389,7 +389,7 @@ struct RIGVM_API FRigVMFunction_MathVectorLerp : public FRigVMFunction_MathVecto
 	}
 
 	RIGVM_METHOD()
-	virtual void Execute() override;
+	RIGVM_API virtual void Execute() override;
 
 	UPROPERTY(meta=(Input))
 	FVector A;
@@ -408,7 +408,7 @@ struct RIGVM_API FRigVMFunction_MathVectorLerp : public FRigVMFunction_MathVecto
  * Remaps the given value from a source range to a target range for each component
  */
 USTRUCT(meta=(DisplayName="Remap", TemplateName="Remap", Keywords="Rescale,Scale"))
-struct RIGVM_API FRigVMFunction_MathVectorRemap : public FRigVMFunction_MathVectorBase
+struct FRigVMFunction_MathVectorRemap : public FRigVMFunction_MathVectorBase
 {
 	GENERATED_BODY()
 
@@ -420,7 +420,7 @@ struct RIGVM_API FRigVMFunction_MathVectorRemap : public FRigVMFunction_MathVect
 	}
 
 	RIGVM_METHOD()
-	virtual void Execute() override;
+	RIGVM_API virtual void Execute() override;
 
 	UPROPERTY(meta=(Input))
 	FVector Value;
@@ -449,7 +449,7 @@ struct RIGVM_API FRigVMFunction_MathVectorRemap : public FRigVMFunction_MathVect
  * Returns true if the value A equals B
  */
 USTRUCT(meta=(DisplayName="Equals", TemplateName="Equals", Keywords="Same,==", Deprecated="5.1"))
-struct RIGVM_API FRigVMFunction_MathVectorEquals : public FRigVMFunction_MathVectorBase
+struct FRigVMFunction_MathVectorEquals : public FRigVMFunction_MathVectorBase
 {
 	GENERATED_BODY()
 
@@ -460,7 +460,7 @@ struct RIGVM_API FRigVMFunction_MathVectorEquals : public FRigVMFunction_MathVec
 	}
 
 	RIGVM_METHOD()
-	virtual void Execute() override;
+	RIGVM_API virtual void Execute() override;
 
 	UPROPERTY(meta=(Input))
 	FVector A;
@@ -479,7 +479,7 @@ struct RIGVM_API FRigVMFunction_MathVectorEquals : public FRigVMFunction_MathVec
  * Returns true if the value A does not equal B
  */
 USTRUCT(meta=(DisplayName="Not Equals", TemplateName="NotEquals", Keywords="Different,!=", Deprecated="5.1"))
-struct RIGVM_API FRigVMFunction_MathVectorNotEquals : public FRigVMFunction_MathVectorBase
+struct FRigVMFunction_MathVectorNotEquals : public FRigVMFunction_MathVectorBase
 {
 	GENERATED_BODY()
 
@@ -490,7 +490,7 @@ struct RIGVM_API FRigVMFunction_MathVectorNotEquals : public FRigVMFunction_Math
 	}
 
 	RIGVM_METHOD()
-	virtual void Execute() override;
+	RIGVM_API virtual void Execute() override;
 
 	UPROPERTY(meta=(Input))
 	FVector A;
@@ -509,7 +509,7 @@ struct RIGVM_API FRigVMFunction_MathVectorNotEquals : public FRigVMFunction_Math
  * Returns true if the value is nearly zero
  */
 USTRUCT(meta=(DisplayName="Is Nearly Zero", TemplateName="IsNearlyZero", Keywords="AlmostZero,0"))
-struct RIGVM_API FRigVMFunction_MathVectorIsNearlyZero : public FRigVMFunction_MathVectorBase
+struct FRigVMFunction_MathVectorIsNearlyZero : public FRigVMFunction_MathVectorBase
 {
 	GENERATED_BODY()
 
@@ -521,7 +521,7 @@ struct RIGVM_API FRigVMFunction_MathVectorIsNearlyZero : public FRigVMFunction_M
 	}
 
 	RIGVM_METHOD()
-	virtual void Execute() override;
+	RIGVM_API virtual void Execute() override;
 
 	UPROPERTY(meta=(Input))
 	FVector Value;
@@ -537,7 +537,7 @@ struct RIGVM_API FRigVMFunction_MathVectorIsNearlyZero : public FRigVMFunction_M
  * Returns true if the value A is almost equal to B
  */
 USTRUCT(meta=(DisplayName="Is Nearly Equal", TemplateName="IsNearlyEqual", Keywords="AlmostEqual,=="))
-struct RIGVM_API FRigVMFunction_MathVectorIsNearlyEqual : public FRigVMFunction_MathVectorBase
+struct FRigVMFunction_MathVectorIsNearlyEqual : public FRigVMFunction_MathVectorBase
 {
 	GENERATED_BODY()
 
@@ -549,7 +549,7 @@ struct RIGVM_API FRigVMFunction_MathVectorIsNearlyEqual : public FRigVMFunction_
 	}
 
 	RIGVM_METHOD()
-	virtual void Execute() override;
+	RIGVM_API virtual void Execute() override;
 
 	UPROPERTY(meta=(Input))
 	FVector A;
@@ -568,7 +568,7 @@ struct RIGVM_API FRigVMFunction_MathVectorIsNearlyEqual : public FRigVMFunction_
  * Return one of the two values based on the condition
  */
 USTRUCT(meta=(DisplayName="Select", Keywords="Pick,If", Deprecated = "4.26.0"))
-struct RIGVM_API FRigVMFunction_MathVectorSelectBool : public FRigVMFunction_MathVectorBase
+struct FRigVMFunction_MathVectorSelectBool : public FRigVMFunction_MathVectorBase
 {
 	GENERATED_BODY()
 
@@ -579,7 +579,7 @@ struct RIGVM_API FRigVMFunction_MathVectorSelectBool : public FRigVMFunction_Mat
 	}
 
 	RIGVM_METHOD()
-	virtual void Execute() override;
+	RIGVM_API virtual void Execute() override;
 
 	UPROPERTY(meta=(Input))
 	bool Condition;
@@ -601,31 +601,31 @@ struct RIGVM_API FRigVMFunction_MathVectorSelectBool : public FRigVMFunction_Mat
  * Returns the degrees of a given value in radians
  */
 USTRUCT(meta=(DisplayName="Degrees", TemplateName="Degrees"))
-struct RIGVM_API FRigVMFunction_MathVectorDeg : public FRigVMFunction_MathVectorUnaryOp
+struct FRigVMFunction_MathVectorDeg : public FRigVMFunction_MathVectorUnaryOp
 {
 	GENERATED_BODY()
 
 	RIGVM_METHOD()
-	virtual void Execute() override;
+	RIGVM_API virtual void Execute() override;
 };
 
 /**
  * Returns the radians of a given value in degrees
  */
 USTRUCT(meta=(DisplayName="Radians", TemplateName="Radians"))
-struct RIGVM_API FRigVMFunction_MathVectorRad : public FRigVMFunction_MathVectorUnaryOp
+struct FRigVMFunction_MathVectorRad : public FRigVMFunction_MathVectorUnaryOp
 {
 	GENERATED_BODY()
 
 	RIGVM_METHOD()
-	virtual void Execute() override;
+	RIGVM_API virtual void Execute() override;
 };
 
 /**
  * Returns the squared length of the vector
  */
 USTRUCT(meta=(DisplayName="Length Squared", TemplateName="LengthSquared", Keywords="Length,Size,Magnitude"))
-struct RIGVM_API FRigVMFunction_MathVectorLengthSquared : public FRigVMFunction_MathVectorBase
+struct FRigVMFunction_MathVectorLengthSquared : public FRigVMFunction_MathVectorBase
 {
 	GENERATED_BODY()
 
@@ -636,7 +636,7 @@ struct RIGVM_API FRigVMFunction_MathVectorLengthSquared : public FRigVMFunction_
 	}
 
 	RIGVM_METHOD()
-	virtual void Execute() override;
+	RIGVM_API virtual void Execute() override;
 
 	UPROPERTY(meta=(Input))
 	FVector Value;
@@ -649,7 +649,7 @@ struct RIGVM_API FRigVMFunction_MathVectorLengthSquared : public FRigVMFunction_
  * Returns the length of the vector
  */
 USTRUCT(meta=(DisplayName="Length", TemplateName="Length", Keywords="Size,Magnitude"))
-struct RIGVM_API FRigVMFunction_MathVectorLength : public FRigVMFunction_MathVectorBase
+struct FRigVMFunction_MathVectorLength : public FRigVMFunction_MathVectorBase
 {
 	GENERATED_BODY()
 
@@ -660,7 +660,7 @@ struct RIGVM_API FRigVMFunction_MathVectorLength : public FRigVMFunction_MathVec
 	}
 
 	RIGVM_METHOD()
-	virtual void Execute() override;
+	RIGVM_API virtual void Execute() override;
 
 	UPROPERTY(meta=(Input))
 	FVector Value;
@@ -673,7 +673,7 @@ struct RIGVM_API FRigVMFunction_MathVectorLength : public FRigVMFunction_MathVec
  * Returns the distance from A to B
  */
 USTRUCT(meta=(DisplayName="Distance Between", TemplateName="Distance"))
-struct RIGVM_API FRigVMFunction_MathVectorDistance : public FRigVMFunction_MathVectorBase
+struct FRigVMFunction_MathVectorDistance : public FRigVMFunction_MathVectorBase
 {
 	GENERATED_BODY()
 
@@ -684,7 +684,7 @@ struct RIGVM_API FRigVMFunction_MathVectorDistance : public FRigVMFunction_MathV
 	}
 
 	RIGVM_METHOD()
-	virtual void Execute() override;
+	RIGVM_API virtual void Execute() override;
 
 	UPROPERTY(meta=(Input))
 	FVector A;
@@ -700,18 +700,18 @@ struct RIGVM_API FRigVMFunction_MathVectorDistance : public FRigVMFunction_MathV
  * Returns the cross product between two vectors
  */
 USTRUCT(meta=(DisplayName="Cross", TemplateName="Cross", Keywords="^"))
-struct RIGVM_API FRigVMFunction_MathVectorCross : public FRigVMFunction_MathVectorBinaryOp
+struct FRigVMFunction_MathVectorCross : public FRigVMFunction_MathVectorBinaryOp
 {
 	GENERATED_BODY()
 	RIGVM_METHOD()
-	virtual void Execute() override;
+	RIGVM_API virtual void Execute() override;
 };
 
 /**
  * Returns the dot product between two vectors
  */
 USTRUCT(meta=(DisplayName="Dot", TemplateName="Dot,|"))
-struct RIGVM_API FRigVMFunction_MathVectorDot : public FRigVMFunction_MathVectorBase
+struct FRigVMFunction_MathVectorDot : public FRigVMFunction_MathVectorBase
 {
 	GENERATED_BODY()
 
@@ -722,7 +722,7 @@ struct RIGVM_API FRigVMFunction_MathVectorDot : public FRigVMFunction_MathVector
 	}
 
 	RIGVM_METHOD()
-	virtual void Execute() override;
+	RIGVM_API virtual void Execute() override;
 
 	UPROPERTY(meta=(Input))
 	FVector A;
@@ -738,19 +738,19 @@ struct RIGVM_API FRigVMFunction_MathVectorDot : public FRigVMFunction_MathVector
  * Returns the normalized value
  */
 USTRUCT(meta=(DisplayName="Unit", TemplateName="Unit", Keywords="Normalize"))
-struct RIGVM_API FRigVMFunction_MathVectorUnit : public FRigVMFunction_MathVectorUnaryOp
+struct FRigVMFunction_MathVectorUnit : public FRigVMFunction_MathVectorUnaryOp
 {
 	GENERATED_BODY()
 
 	RIGVM_METHOD()
-	virtual void Execute() override;
+	RIGVM_API virtual void Execute() override;
 };
 
 /**
  * Sets the length of a given vector
  */
 USTRUCT(meta = (DisplayName = "SetLength", TemplateName="SetLength", Keywords = "Unit,Normalize,Scale"))
-struct RIGVM_API FRigVMFunction_MathVectorSetLength: public FRigVMFunction_MathVectorBase
+struct FRigVMFunction_MathVectorSetLength: public FRigVMFunction_MathVectorBase
 {
 	GENERATED_BODY()
 
@@ -761,7 +761,7 @@ struct RIGVM_API FRigVMFunction_MathVectorSetLength: public FRigVMFunction_MathV
 	}
 
 	RIGVM_METHOD()
-	virtual void Execute() override;
+	RIGVM_API virtual void Execute() override;
 
 	UPROPERTY(meta = (Input))
 	FVector Value;
@@ -777,7 +777,7 @@ struct RIGVM_API FRigVMFunction_MathVectorSetLength: public FRigVMFunction_MathV
  * Clamps the length of a given vector between a minimum and maximum
  */
 USTRUCT(meta = (DisplayName = "ClampLength", TemplateName="ClampLength", Keywords = "Unit,Normalize,Scale"))
-struct RIGVM_API FRigVMFunction_MathVectorClampLength: public FRigVMFunction_MathVectorBase
+struct FRigVMFunction_MathVectorClampLength: public FRigVMFunction_MathVectorBase
 {
 	GENERATED_BODY()
 
@@ -790,7 +790,7 @@ struct RIGVM_API FRigVMFunction_MathVectorClampLength: public FRigVMFunction_Mat
 
 
 	RIGVM_METHOD()
-	virtual void Execute() override;
+	RIGVM_API virtual void Execute() override;
 
 	UPROPERTY(meta = (Input))
 	FVector Value;
@@ -809,7 +809,7 @@ struct RIGVM_API FRigVMFunction_MathVectorClampLength: public FRigVMFunction_Mat
  * Mirror a vector about a normal vector.
  */
 USTRUCT(meta=(DisplayName="Mirror on Normal"))
-struct RIGVM_API FRigVMFunction_MathVectorMirror : public FRigVMFunction_MathVectorBase
+struct FRigVMFunction_MathVectorMirror : public FRigVMFunction_MathVectorBase
 {
 	GENERATED_BODY()
 
@@ -820,7 +820,7 @@ struct RIGVM_API FRigVMFunction_MathVectorMirror : public FRigVMFunction_MathVec
 	}
 
 	RIGVM_METHOD()
-	virtual void Execute() override;
+	RIGVM_API virtual void Execute() override;
 
 	UPROPERTY(meta=(Input))
 	FVector Value;
@@ -836,7 +836,7 @@ struct RIGVM_API FRigVMFunction_MathVectorMirror : public FRigVMFunction_MathVec
  * Returns the angle between two vectors in radians
  */
 USTRUCT(meta=(DisplayName="Angle Between", TemplateName="AngleBetween"))
-struct RIGVM_API FRigVMFunction_MathVectorAngle : public FRigVMFunction_MathVectorBase
+struct FRigVMFunction_MathVectorAngle : public FRigVMFunction_MathVectorBase
 {
 	GENERATED_BODY()
 
@@ -847,7 +847,7 @@ struct RIGVM_API FRigVMFunction_MathVectorAngle : public FRigVMFunction_MathVect
 	}
 
 	RIGVM_METHOD()
-	virtual void Execute() override;
+	RIGVM_API virtual void Execute() override;
 
 	UPROPERTY(meta=(Input))
 	FVector A;
@@ -863,7 +863,7 @@ struct RIGVM_API FRigVMFunction_MathVectorAngle : public FRigVMFunction_MathVect
  * Returns true if the two vectors are parallel
  */
 USTRUCT(meta=(DisplayName="Parallel", TemplateName="Parallel"))
-struct RIGVM_API FRigVMFunction_MathVectorParallel : public FRigVMFunction_MathVectorBase
+struct FRigVMFunction_MathVectorParallel : public FRigVMFunction_MathVectorBase
 {
 	GENERATED_BODY()
 
@@ -874,7 +874,7 @@ struct RIGVM_API FRigVMFunction_MathVectorParallel : public FRigVMFunction_MathV
 	}
 
 	RIGVM_METHOD()
-	virtual void Execute() override;
+	RIGVM_API virtual void Execute() override;
 
 	UPROPERTY(meta=(Input))
 	FVector A;
@@ -890,7 +890,7 @@ struct RIGVM_API FRigVMFunction_MathVectorParallel : public FRigVMFunction_MathV
  * Returns true if the two vectors are orthogonal
  */
 USTRUCT(meta=(DisplayName="Orthogonal", TemplateName="Orthogonal"))
-struct RIGVM_API FRigVMFunction_MathVectorOrthogonal : public FRigVMFunction_MathVectorBase
+struct FRigVMFunction_MathVectorOrthogonal : public FRigVMFunction_MathVectorBase
 {
 	GENERATED_BODY()
 
@@ -902,7 +902,7 @@ struct RIGVM_API FRigVMFunction_MathVectorOrthogonal : public FRigVMFunction_Mat
 	}
 
 	RIGVM_METHOD()
-	virtual void Execute() override;
+	RIGVM_API virtual void Execute() override;
 
 	UPROPERTY(meta=(Input))
 	FVector A;
@@ -918,7 +918,7 @@ struct RIGVM_API FRigVMFunction_MathVectorOrthogonal : public FRigVMFunction_Mat
  * Returns the 4 point bezier interpolation
  */
 USTRUCT(meta=(DisplayName="Bezier Four Point", Deprecated = "5.0"))
-struct RIGVM_API FRigVMFunction_MathVectorBezierFourPoint : public FRigVMFunction_MathVectorBase
+struct FRigVMFunction_MathVectorBezierFourPoint : public FRigVMFunction_MathVectorBase
 {
 	GENERATED_BODY()
 
@@ -930,7 +930,7 @@ struct RIGVM_API FRigVMFunction_MathVectorBezierFourPoint : public FRigVMFunctio
 	}
 
 	RIGVM_METHOD()
-	virtual void Execute() override;
+	RIGVM_API virtual void Execute() override;
 
 	UPROPERTY(meta=(Input))
 	FRigVMFourPointBezier Bezier;
@@ -952,7 +952,7 @@ struct RIGVM_API FRigVMFunction_MathVectorBezierFourPoint : public FRigVMFunctio
  * Creates a bezier four point
  */
 USTRUCT(meta = (DisplayName = "Make Bezier Four Point", Constant, Deprecated = "5.0"))
-struct RIGVM_API FRigVMFunction_MathVectorMakeBezierFourPoint : public FRigVMFunction_MathVectorBase
+struct FRigVMFunction_MathVectorMakeBezierFourPoint : public FRigVMFunction_MathVectorBase
 {
 	GENERATED_BODY()
 
@@ -962,7 +962,7 @@ struct RIGVM_API FRigVMFunction_MathVectorMakeBezierFourPoint : public FRigVMFun
 	}
 
 	RIGVM_METHOD()
-	virtual void Execute() override;
+	RIGVM_API virtual void Execute() override;
 
 	UPROPERTY(meta = (Input, Output))
 	FRigVMFourPointBezier Bezier;
@@ -977,7 +977,7 @@ struct RIGVM_API FRigVMFunction_MathVectorMakeBezierFourPoint : public FRigVMFun
  * You can disable the inner / outer envelope / collision by setting the minimum / maximum to 0.0.
  */
 USTRUCT(meta = (DisplayName = "Clamp Spatially", TemplateName="ClampSpatially", Keywords="Collide,Collision"))
-struct RIGVM_API FRigVMFunction_MathVectorClampSpatially: public FRigVMFunction_MathVectorBase
+struct FRigVMFunction_MathVectorClampSpatially: public FRigVMFunction_MathVectorBase
 {
 	GENERATED_BODY()
 
@@ -995,7 +995,7 @@ struct RIGVM_API FRigVMFunction_MathVectorClampSpatially: public FRigVMFunction_
 	}
 
 	RIGVM_METHOD()
-	virtual void Execute() override;
+	RIGVM_API virtual void Execute() override;
 
 	UPROPERTY(meta = (Input))
 	FVector Value;
@@ -1041,7 +1041,7 @@ struct RIGVM_API FRigVMFunction_MathVectorClampSpatially: public FRigVMFunction_
  * Intersects a plane with a vector given a start and direction
  */
 USTRUCT(meta = (DisplayName = "Intersect Plane", Keywords = "Collide,Intersect,Raycast"))
-struct RIGVM_API FRigVMFunction_MathIntersectPlane : public FRigVMFunction_MathVectorBase
+struct FRigVMFunction_MathIntersectPlane : public FRigVMFunction_MathVectorBase
 {
 	GENERATED_BODY()
 
@@ -1053,7 +1053,7 @@ struct RIGVM_API FRigVMFunction_MathIntersectPlane : public FRigVMFunction_MathV
 	}
 
 	RIGVM_METHOD()
-	virtual void Execute() override;
+	RIGVM_API virtual void Execute() override;
 
 	UPROPERTY(meta = (Input))
 	FVector Start;
@@ -1078,7 +1078,7 @@ struct RIGVM_API FRigVMFunction_MathIntersectPlane : public FRigVMFunction_MathV
  * Find the point on the plane that is closest to the given point and the distance between them.
  */
 USTRUCT(meta = (DisplayName = "Distance To Plane", Keywords = "Distance,Plane,Closest,Project"))
-struct RIGVM_API FRigVMFunction_MathDistanceToPlane : public FRigVMFunction_MathVectorBase
+struct FRigVMFunction_MathDistanceToPlane : public FRigVMFunction_MathVectorBase
 {
 	GENERATED_BODY()
 
@@ -1090,7 +1090,7 @@ struct RIGVM_API FRigVMFunction_MathDistanceToPlane : public FRigVMFunction_Math
 	}
 
 	RIGVM_METHOD()
-	virtual void Execute() override;
+	RIGVM_API virtual void Execute() override;
 
 	UPROPERTY(meta = (Input))
 	FVector Point;
@@ -1112,7 +1112,7 @@ struct RIGVM_API FRigVMFunction_MathDistanceToPlane : public FRigVMFunction_Math
  * Returns the relative local vector within a parent's vector
  */
 USTRUCT(meta=(DisplayName="Make Relative", TemplateName="Make Relative", Keywords="Local,Global,Absolute"))
-struct RIGVM_API FRigVMFunction_MathVectorMakeRelative : public FRigVMFunction_MathVectorBase
+struct FRigVMFunction_MathVectorMakeRelative : public FRigVMFunction_MathVectorBase
 {
 	GENERATED_BODY()
 
@@ -1122,7 +1122,7 @@ struct RIGVM_API FRigVMFunction_MathVectorMakeRelative : public FRigVMFunction_M
 	}
 
 	RIGVM_METHOD()
-	virtual void Execute() override;
+	RIGVM_API virtual void Execute() override;
 
 	UPROPERTY(meta=(Input))
 	FVector Global;
@@ -1138,7 +1138,7 @@ struct RIGVM_API FRigVMFunction_MathVectorMakeRelative : public FRigVMFunction_M
  * Returns the absolute global vector within a parent's vector
  */
 USTRUCT(meta = (DisplayName = "Make Absolute", TemplateName="Make Absolute", Keywords = "Local,Global,Relative"))
-struct RIGVM_API FRigVMFunction_MathVectorMakeAbsolute : public FRigVMFunction_MathVectorBase
+struct FRigVMFunction_MathVectorMakeAbsolute : public FRigVMFunction_MathVectorBase
 {
 	GENERATED_BODY()
 
@@ -1148,7 +1148,7 @@ struct RIGVM_API FRigVMFunction_MathVectorMakeAbsolute : public FRigVMFunction_M
 	}
 
 	RIGVM_METHOD()
-	virtual void Execute() override;
+	RIGVM_API virtual void Execute() override;
 
 	UPROPERTY(meta = (Input))
 	FVector Local;
@@ -1164,7 +1164,7 @@ struct RIGVM_API FRigVMFunction_MathVectorMakeAbsolute : public FRigVMFunction_M
  * Mirror a vector about a central transform.
  */
 USTRUCT(meta=(DisplayName="Mirror", TemplateName="Mirror"))
-struct RIGVM_API FRigVMFunction_MathVectorMirrorTransform : public FRigVMFunction_MathVectorBase
+struct FRigVMFunction_MathVectorMirrorTransform : public FRigVMFunction_MathVectorBase
 {
 	GENERATED_BODY()
 
@@ -1177,7 +1177,7 @@ struct RIGVM_API FRigVMFunction_MathVectorMirrorTransform : public FRigVMFunctio
 	}
 
 	RIGVM_METHOD()
-	virtual void Execute() override;
+	RIGVM_API virtual void Execute() override;
 
 	UPROPERTY(meta=(Input))
 	FVector Value;
@@ -1202,7 +1202,7 @@ struct RIGVM_API FRigVMFunction_MathVectorMirrorTransform : public FRigVMFunctio
  * Returns the sum of the given array
  */
 USTRUCT(meta = (DisplayName = "Array Sum", TemplateName = "ArraySum"))
-struct RIGVM_API FRigVMFunction_MathVectorArraySum : public FRigVMFunction_MathVectorBase
+struct FRigVMFunction_MathVectorArraySum : public FRigVMFunction_MathVectorBase
 {
 	GENERATED_BODY()
 
@@ -1213,7 +1213,7 @@ struct RIGVM_API FRigVMFunction_MathVectorArraySum : public FRigVMFunction_MathV
 
 	/** Execute logic for this rig unit */
 	RIGVM_METHOD()
-	virtual void Execute() override;
+	RIGVM_API virtual void Execute() override;
 
 	UPROPERTY(meta = (Input))
 	TArray<FVector> Array;
@@ -1227,7 +1227,7 @@ struct RIGVM_API FRigVMFunction_MathVectorArraySum : public FRigVMFunction_MathV
  * Returns the average of the given array
  */
 USTRUCT(meta = (DisplayName = "Array Average", TemplateName = "ArrayAverage"))
-struct RIGVM_API FRigVMFunction_MathVectorArrayAverage : public FRigVMFunction_MathVectorBase
+struct FRigVMFunction_MathVectorArrayAverage : public FRigVMFunction_MathVectorBase
 {
 	GENERATED_BODY()
 
@@ -1238,7 +1238,7 @@ struct RIGVM_API FRigVMFunction_MathVectorArrayAverage : public FRigVMFunction_M
 
 	/** Execute logic for this rig unit */
 	RIGVM_METHOD()
-	virtual void Execute() override;
+	RIGVM_API virtual void Execute() override;
 
 	UPROPERTY(meta = (Input))
 	TArray<FVector> Array;
diff --git a/Engine/Plugins/Runtime/RigVM/Source/RigVM/Public/RigVMFunctions/Math/RigVMFunction_Noise.h b/Engine/Plugins/Runtime/RigVM/Source/RigVM/Public/RigVMFunctions/Math/RigVMFunction_Noise.h
--- a/Engine/Plugins/Runtime/RigVM/Source/RigVM/Public/RigVMFunctions/Math/RigVMFunction_Noise.h
+++ b/Engine/Plugins/Runtime/RigVM/Source/RigVM/Public/RigVMFunctions/Math/RigVMFunction_Noise.h
@@ -9,7 +9,7 @@
  * Generates a float through a noise fluctuation process between a min and a max through speed
  */
 USTRUCT(meta=(DisplayName="Noise (Float)", Category="Math|Noise", TemplateName="Noise"))
-struct RIGVM_API FRigVMFunction_NoiseFloat : public FRigVMFunction_MathBase
+struct FRigVMFunction_NoiseFloat : public FRigVMFunction_MathBase
 {
 	GENERATED_BODY()
 
@@ -21,7 +21,7 @@ struct RIGVM_API FRigVMFunction_NoiseFloat : public FRigVMFunction_MathBase
 	}
 
 	RIGVM_METHOD()
-	virtual void Execute() override;
+	RIGVM_API virtual void Execute() override;
 
 	UPROPERTY(meta = (Input))
 	float Value;
@@ -49,7 +49,7 @@ struct RIGVM_API FRigVMFunction_NoiseFloat : public FRigVMFunction_MathBase
  * Generates a float through a noise fluctuation process between a min and a max through speed
  */
 USTRUCT(meta=(DisplayName="Noise (Double)", Category="Math|Noise", TemplateName="Noise"))
-struct RIGVM_API FRigVMFunction_NoiseDouble : public FRigVMFunction_MathBase
+struct FRigVMFunction_NoiseDouble : public FRigVMFunction_MathBase
 {
 	GENERATED_BODY()
 
@@ -61,7 +61,7 @@ struct RIGVM_API FRigVMFunction_NoiseDouble : public FRigVMFunction_MathBase
 	}
 
 	RIGVM_METHOD()
-	virtual void Execute() override;
+	RIGVM_API virtual void Execute() override;
 
 	UPROPERTY(meta = (Input))
 	double Value;
@@ -89,7 +89,7 @@ struct RIGVM_API FRigVMFunction_NoiseDouble : public FRigVMFunction_MathBase
  * Generates a vector through a noise fluctuation process between a min and a max through speed
  */
 USTRUCT(meta = (DisplayName = "Noise (Vector)", Category = "Math|Noise", Deprecated = "5.0.0"))
-struct RIGVM_API FRigVMFunction_NoiseVector : public FRigVMFunction_MathBase
+struct FRigVMFunction_NoiseVector : public FRigVMFunction_MathBase
 {
 	GENERATED_BODY()
 
@@ -103,7 +103,7 @@ struct RIGVM_API FRigVMFunction_NoiseVector : public FRigVMFunction_MathBase
 	}
 
 	RIGVM_METHOD()
-	virtual void Execute() override;
+	RIGVM_API virtual void Execute() override;
 
 	UPROPERTY(meta = (Input))
 	FVector Position;
@@ -134,7 +134,7 @@ struct RIGVM_API FRigVMFunction_NoiseVector : public FRigVMFunction_MathBase
  * Generates a vector through a noise fluctuation process between a min and a max through speed
  */
 USTRUCT(meta = (DisplayName = "Noise (Vector)", Category = "Math|Noise", TemplateName="Noise"))
-struct RIGVM_API FRigVMFunction_NoiseVector2 : public FRigVMFunction_MathBase
+struct FRigVMFunction_NoiseVector2 : public FRigVMFunction_MathBase
 {
 	GENERATED_BODY()
 
@@ -148,7 +148,7 @@ struct RIGVM_API FRigVMFunction_NoiseVector2 : public FRigVMFunction_MathBase
 	}
 
 	RIGVM_METHOD()
-	virtual void Execute() override;
+	RIGVM_API virtual void Execute() override;
 
 	UPROPERTY(meta = (Input))
 	FVector Value;
diff --git a/Engine/Plugins/Runtime/RigVM/Source/RigVM/Public/RigVMFunctions/Math/RigVMFunction_Random.h b/Engine/Plugins/Runtime/RigVM/Source/RigVM/Public/RigVMFunctions/Math/RigVMFunction_Random.h
--- a/Engine/Plugins/Runtime/RigVM/Source/RigVM/Public/RigVMFunctions/Math/RigVMFunction_Random.h
+++ b/Engine/Plugins/Runtime/RigVM/Source/RigVM/Public/RigVMFunctions/Math/RigVMFunction_Random.h
@@ -9,7 +9,7 @@
  * Generates a random float between a min and a max
  */
 USTRUCT(meta=(DisplayName="Random (Float)", Category="Math|Random", TemplateName="Random"))
-struct RIGVM_API FRigVMFunction_RandomFloat : public FRigVMFunction_MathBase
+struct FRigVMFunction_RandomFloat : public FRigVMFunction_MathBase
 {
 	GENERATED_BODY()
 
@@ -23,7 +23,7 @@ struct RIGVM_API FRigVMFunction_RandomFloat : public FRigVMFunction_MathBase
 	}
 
 	RIGVM_METHOD()
-	virtual void Execute() override;
+	RIGVM_API virtual void Execute() override;
 
 	UPROPERTY(meta = (Input))
 	int32 Seed;
@@ -62,7 +62,7 @@ struct RIGVM_API FRigVMFunction_RandomFloat : public FRigVMFunction_MathBase
  * Generates a random vector between a min and a max
  */
 USTRUCT(meta = (DisplayName = "Random (Vector)", Category = "Math|Random", TemplateName="Random"))
-struct RIGVM_API FRigVMFunction_RandomVector: public FRigVMFunction_MathBase
+struct FRigVMFunction_RandomVector: public FRigVMFunction_MathBase
 {
 	GENERATED_BODY()
 
@@ -77,7 +77,7 @@ struct RIGVM_API FRigVMFunction_RandomVector: public FRigVMFunction_MathBase
 	}
 
 	RIGVM_METHOD()
-	virtual void Execute() override;
+	RIGVM_API virtual void Execute() override;
 
 	UPROPERTY(meta = (Input))
 	int32 Seed;
diff --git a/Engine/Plugins/Runtime/RigVM/Source/RigVM/Public/RigVMFunctions/Math/RigVMMathLibrary.h b/Engine/Plugins/Runtime/RigVM/Source/RigVM/Public/RigVMFunctions/Math/RigVMMathLibrary.h
--- a/Engine/Plugins/Runtime/RigVM/Source/RigVM/Public/RigVMFunctions/Math/RigVMMathLibrary.h
+++ b/Engine/Plugins/Runtime/RigVM/Source/RigVM/Public/RigVMFunctions/Math/RigVMMathLibrary.h
@@ -45,7 +45,7 @@ enum class ERigVMAnimEasingType : uint8
 };
 
 USTRUCT(BlueprintType)
-struct RIGVM_API FRigVMFourPointBezier
+struct FRigVMFourPointBezier
 {
 	GENERATED_BODY()
 
@@ -68,7 +68,7 @@ struct RIGVM_API FRigVMFourPointBezier
 };
 
 USTRUCT(BlueprintType)
-struct RIGVM_API FRigVMMirrorSettings
+struct FRigVMMirrorSettings
 {
 	GENERATED_USTRUCT_BODY()
 
@@ -94,8 +94,8 @@ struct RIGVM_API FRigVMMirrorSettings
 	UPROPERTY(BlueprintReadWrite, EditAnywhere, Category = Settings, DisplayName = "Replace")
 	FString ReplaceString;
 
-	FTransform MirrorTransform(const FTransform& InTransform) const;
-	FVector MirrorVector(const FVector& InVector) const;
+	RIGVM_API FTransform MirrorTransform(const FTransform& InTransform) const;
+	RIGVM_API FVector MirrorVector(const FVector& InVector) const;
 };
 
 UENUM(meta = (RigVMTypeAllowed))
@@ -106,7 +106,7 @@ enum class ERigVMSimPointIntegrateType : uint8
 };
 
 USTRUCT(BlueprintType)
-struct RIGVM_API FRigVMSimPoint
+struct FRigVMSimPoint
 {
 	GENERATED_BODY()
 
@@ -157,25 +157,25 @@ struct RIGVM_API FRigVMSimPoint
 	UPROPERTY(EditAnywhere, Category=Simulation)
 	FVector LinearVelocity;
 
-	FRigVMSimPoint IntegrateVerlet(const FVector& InForce, float InBlend, float InDeltaTime) const;
-	FRigVMSimPoint IntegrateSemiExplicitEuler(const FVector& InForce, float InDeltaTime) const;
+	RIGVM_API FRigVMSimPoint IntegrateVerlet(const FVector& InForce, float InBlend, float InDeltaTime) const;
+	RIGVM_API FRigVMSimPoint IntegrateSemiExplicitEuler(const FVector& InForce, float InDeltaTime) const;
 };
 
-class RIGVM_API FRigVMMathLibrary
+class FRigVMMathLibrary
 {
 public:
-	static float AngleBetween(const FVector& A, const FVector& B);
-	static void FourPointBezier(const FVector& A, const FVector& B, const FVector& C, const FVector& D, float T, FVector& OutPosition, FVector& OutTangent);
-	static void FourPointBezier(const FRigVMFourPointBezier& Bezier, float T, FVector& OutPosition, FVector& OutTangent);
-	static float EaseFloat(float Value, ERigVMAnimEasingType Type);
-	static FTransform LerpTransform(const FTransform& A, const FTransform& B, float T);
-	static FVector ClampSpatially(const FVector& Value, EAxis::Type Axis, ERigVMClampSpatialMode::Type Type, float Minimum, float Maximum, FTransform Space);
-	static FQuat FindQuatBetweenVectors(const FVector& A, const FVector& B);
-	static FQuat FindQuatBetweenNormals(const FVector& A, const FVector& B);
+	static RIGVM_API float AngleBetween(const FVector& A, const FVector& B);
+	static RIGVM_API void FourPointBezier(const FVector& A, const FVector& B, const FVector& C, const FVector& D, float T, FVector& OutPosition, FVector& OutTangent);
+	static RIGVM_API void FourPointBezier(const FRigVMFourPointBezier& Bezier, float T, FVector& OutPosition, FVector& OutTangent);
+	static RIGVM_API float EaseFloat(float Value, ERigVMAnimEasingType Type);
+	static RIGVM_API FTransform LerpTransform(const FTransform& A, const FTransform& B, float T);
+	static RIGVM_API FVector ClampSpatially(const FVector& Value, EAxis::Type Axis, ERigVMClampSpatialMode::Type Type, float Minimum, float Maximum, FTransform Space);
+	static RIGVM_API FQuat FindQuatBetweenVectors(const FVector& A, const FVector& B);
+	static RIGVM_API FQuat FindQuatBetweenNormals(const FVector& A, const FVector& B);
 
 	// See - "Computing Euler angles from a rotation matrix" by Gregory G. Slabaugh
 	// Each spatial orientation can be mapped to two equivalent euler angles within range (-180, 180)
-	static FVector GetEquivalentEulerAngle(const FVector& InEulerAngle, const EEulerRotationOrder& InOrder);
+	static RIGVM_API FVector GetEquivalentEulerAngle(const FVector& InEulerAngle, const EEulerRotationOrder& InOrder);
 	
-	static FVector& ChooseBetterEulerAngleForAxisFilter(const FVector& Base, FVector& A, FVector& B);
+	static RIGVM_API FVector& ChooseBetterEulerAngleForAxisFilter(const FVector& Base, FVector& A, FVector& B);
 };
\ No newline at end of file
diff --git a/Engine/Plugins/Runtime/RigVM/Source/RigVM/Public/RigVMFunctions/RigVMDispatch_Array.h b/Engine/Plugins/Runtime/RigVM/Source/RigVM/Public/RigVMFunctions/RigVMDispatch_Array.h
--- a/Engine/Plugins/Runtime/RigVM/Source/RigVM/Public/RigVMFunctions/RigVMDispatch_Array.h
+++ b/Engine/Plugins/Runtime/RigVM/Source/RigVM/Public/RigVMFunctions/RigVMDispatch_Array.h
@@ -8,22 +8,22 @@
 #include "RigVMDispatch_Array.generated.h"
 
 USTRUCT(meta=(Abstract, Category = "Array", Keywords = "List,Collection", NodeColor = "1,1,1,1"))
-struct RIGVM_API FRigVMDispatch_ArrayBase : public FRigVMDispatch_CoreBase
+struct FRigVMDispatch_ArrayBase : public FRigVMDispatch_CoreBase
 {
 	GENERATED_BODY()
 
 public:
 	virtual ERigVMOpCode GetOpCode() const { return ERigVMOpCode::Invalid; }
-	static UScriptStruct* GetFactoryDispatchForOpCode(ERigVMOpCode InOpCode);
-	static FName GetFactoryNameForOpCode(ERigVMOpCode InOpCode);
+	static RIGVM_API UScriptStruct* GetFactoryDispatchForOpCode(ERigVMOpCode InOpCode);
+	static RIGVM_API FName GetFactoryNameForOpCode(ERigVMOpCode InOpCode);
 #if WITH_EDITOR
-	virtual FString GetArgumentDefaultValue(const FName& InArgumentName, TRigVMTypeIndex InTypeIndex) const override;
+	RIGVM_API virtual FString GetArgumentDefaultValue(const FName& InArgumentName, TRigVMTypeIndex InTypeIndex) const override;
 #endif
 	virtual bool IsSingleton() const override { return true; } 
 
 protected:
-	static FRigVMTemplateArgumentInfo CreateArgumentInfo(const FName& InName, ERigVMPinDirection InDirection);
-	static TMap<uint32, int32> GetArrayHash(FScriptArrayHelper& InArrayHelper, const FArrayProperty* InArrayProperty);
+	static RIGVM_API FRigVMTemplateArgumentInfo CreateArgumentInfo(const FName& InName, ERigVMPinDirection InDirection);
+	static RIGVM_API TMap<uint32, int32> GetArrayHash(FScriptArrayHelper& InArrayHelper, const FArrayProperty* InArrayProperty);
 
 	static inline FLazyName ArrayName = FLazyName(TEXT("Array"));
 	static inline FLazyName ValuesName = FLazyName(TEXT("Values"));
@@ -41,16 +41,16 @@ protected:
 };
 
 USTRUCT(meta=(Abstract))
-struct RIGVM_API FRigVMDispatch_ArrayBaseMutable : public FRigVMDispatch_ArrayBase
+struct FRigVMDispatch_ArrayBaseMutable : public FRigVMDispatch_ArrayBase
 {
 	GENERATED_BODY()
 
 public:
-	virtual const TArray<FRigVMExecuteArgument>& GetExecuteArguments_Impl(const FRigVMDispatchContext& InContext) const override;
+	RIGVM_API virtual const TArray<FRigVMExecuteArgument>& GetExecuteArguments_Impl(const FRigVMDispatchContext& InContext) const override;
 };
 
 USTRUCT(meta=(DisplayName = "Make Array", Keywords = "Make,MakeArray,Constant,Reroute"))
-struct RIGVM_API FRigVMDispatch_ArrayMake : public FRigVMDispatch_ArrayBase
+struct FRigVMDispatch_ArrayMake : public FRigVMDispatch_ArrayBase
 {
 	GENERATED_BODY()
 
@@ -60,23 +60,23 @@ public:
 		FactoryScriptStruct = StaticStruct();
 	}
 
-	virtual const TArray<FRigVMTemplateArgumentInfo>& GetArgumentInfos() const override;
-	virtual FRigVMTemplateTypeMap OnNewArgumentType(const FName& InArgumentName, TRigVMTypeIndex InTypeIndex) const override;
+	RIGVM_API virtual const TArray<FRigVMTemplateArgumentInfo>& GetArgumentInfos() const override;
+	RIGVM_API virtual FRigVMTemplateTypeMap OnNewArgumentType(const FName& InArgumentName, TRigVMTypeIndex InTypeIndex) const override;
 #if WITH_EDITOR
-	virtual FText GetNodeTooltip(const FRigVMTemplateTypeMap& InTypes) const override;
-	virtual FText GetArgumentTooltip(const FName& InArgumentName, TRigVMTypeIndex InTypeIndex) const override;
-	virtual FString GetArgumentMetaData(const FName& InArgumentName, const FName& InMetaDataKey) const override;
-	virtual FString GetArgumentDefaultValue(const FName& InArgumentName, TRigVMTypeIndex InTypeIndex) const override;
+	RIGVM_API virtual FText GetNodeTooltip(const FRigVMTemplateTypeMap& InTypes) const override;
+	RIGVM_API virtual FText GetArgumentTooltip(const FName& InArgumentName, TRigVMTypeIndex InTypeIndex) const override;
+	RIGVM_API virtual FString GetArgumentMetaData(const FName& InArgumentName, const FName& InMetaDataKey) const override;
+	RIGVM_API virtual FString GetArgumentDefaultValue(const FName& InArgumentName, TRigVMTypeIndex InTypeIndex) const override;
 #endif
-	virtual FName GetArgumentNameForOperandIndex(int32 InOperandIndex, int32 InTotalOperands) const override;
+	RIGVM_API virtual FName GetArgumentNameForOperandIndex(int32 InOperandIndex, int32 InTotalOperands) const override;
 
 protected:
 	virtual FRigVMFunctionPtr GetDispatchFunctionImpl(const FRigVMTemplateTypeMap& InTypes) const override { return &FRigVMDispatch_ArrayMake::Execute; }
-	static void Execute(FRigVMExtendedExecuteContext& InContext, FRigVMMemoryHandleArray Handles, FRigVMPredicateBranchArray Predicates);
+	static RIGVM_API void Execute(FRigVMExtendedExecuteContext& InContext, FRigVMMemoryHandleArray Handles, FRigVMPredicateBranchArray Predicates);
 };
 
 USTRUCT(meta=(DisplayName = "Reset"))
-struct RIGVM_API FRigVMDispatch_ArrayReset : public FRigVMDispatch_ArrayBaseMutable
+struct FRigVMDispatch_ArrayReset : public FRigVMDispatch_ArrayBaseMutable
 {
 	GENERATED_BODY()
 
@@ -87,21 +87,21 @@ public:
 	}
 
 	virtual ERigVMOpCode GetOpCode() const override { return ERigVMOpCode::ArrayReset; }
-	virtual FName GetArgumentNameForOperandIndex(int32 InOperandIndex, int32 InTotalOperands) const override;
-	virtual const TArray<FRigVMTemplateArgumentInfo>& GetArgumentInfos() const override;
-	virtual FRigVMTemplateTypeMap OnNewArgumentType(const FName& InArgumentName, TRigVMTypeIndex InTypeIndex) const override;
+	RIGVM_API virtual FName GetArgumentNameForOperandIndex(int32 InOperandIndex, int32 InTotalOperands) const override;
+	RIGVM_API virtual const TArray<FRigVMTemplateArgumentInfo>& GetArgumentInfos() const override;
+	RIGVM_API virtual FRigVMTemplateTypeMap OnNewArgumentType(const FName& InArgumentName, TRigVMTypeIndex InTypeIndex) const override;
 #if WITH_EDITOR
-	virtual FText GetNodeTooltip(const FRigVMTemplateTypeMap& InTypes) const override;
-	virtual FText GetArgumentTooltip(const FName& InArgumentName, TRigVMTypeIndex InTypeIndex) const override;
+	RIGVM_API virtual FText GetNodeTooltip(const FRigVMTemplateTypeMap& InTypes) const override;
+	RIGVM_API virtual FText GetArgumentTooltip(const FName& InArgumentName, TRigVMTypeIndex InTypeIndex) const override;
 #endif
 
 protected:
 	virtual FRigVMFunctionPtr GetDispatchFunctionImpl(const FRigVMTemplateTypeMap& InTypes) const override { return &FRigVMDispatch_ArrayReset::Execute; }
-	static void Execute(FRigVMExtendedExecuteContext& InContext, FRigVMMemoryHandleArray Handles, FRigVMPredicateBranchArray Predicates);
+	static RIGVM_API void Execute(FRigVMExtendedExecuteContext& InContext, FRigVMMemoryHandleArray Handles, FRigVMPredicateBranchArray Predicates);
 };
 
 USTRUCT(meta=(DisplayName = "Num"))
-struct RIGVM_API FRigVMDispatch_ArrayGetNum : public FRigVMDispatch_ArrayBase
+struct FRigVMDispatch_ArrayGetNum : public FRigVMDispatch_ArrayBase
 {
 	GENERATED_BODY()
 
@@ -112,21 +112,21 @@ public:
 	}
 
 	virtual ERigVMOpCode GetOpCode() const override { return ERigVMOpCode::ArrayGetNum; }
-	virtual FName GetArgumentNameForOperandIndex(int32 InOperandIndex, int32 InTotalOperands) const override;
-	virtual const TArray<FRigVMTemplateArgumentInfo>& GetArgumentInfos() const override;
-	virtual FRigVMTemplateTypeMap OnNewArgumentType(const FName& InArgumentName, TRigVMTypeIndex InTypeIndex) const override;
+	RIGVM_API virtual FName GetArgumentNameForOperandIndex(int32 InOperandIndex, int32 InTotalOperands) const override;
+	RIGVM_API virtual const TArray<FRigVMTemplateArgumentInfo>& GetArgumentInfos() const override;
+	RIGVM_API virtual FRigVMTemplateTypeMap OnNewArgumentType(const FName& InArgumentName, TRigVMTypeIndex InTypeIndex) const override;
 #if WITH_EDITOR
-	virtual FText GetNodeTooltip(const FRigVMTemplateTypeMap& InTypes) const override;
-	virtual FText GetArgumentTooltip(const FName& InArgumentName, TRigVMTypeIndex InTypeIndex) const override;
+	RIGVM_API virtual FText GetNodeTooltip(const FRigVMTemplateTypeMap& InTypes) const override;
+	RIGVM_API virtual FText GetArgumentTooltip(const FName& InArgumentName, TRigVMTypeIndex InTypeIndex) const override;
 #endif
 
 protected:
 	virtual FRigVMFunctionPtr GetDispatchFunctionImpl(const FRigVMTemplateTypeMap& InTypes) const override { return &FRigVMDispatch_ArrayGetNum::Execute; }
-	static void Execute(FRigVMExtendedExecuteContext& InContext, FRigVMMemoryHandleArray Handles, FRigVMPredicateBranchArray Predicates);
+	static RIGVM_API void Execute(FRigVMExtendedExecuteContext& InContext, FRigVMMemoryHandleArray Handles, FRigVMPredicateBranchArray Predicates);
 };
 
 USTRUCT(meta=(DisplayName = "Set Num"))
-struct RIGVM_API FRigVMDispatch_ArraySetNum : public FRigVMDispatch_ArrayBaseMutable
+struct FRigVMDispatch_ArraySetNum : public FRigVMDispatch_ArrayBaseMutable
 {
 	GENERATED_BODY()
 
@@ -137,21 +137,21 @@ public:
 	}
 
 	virtual ERigVMOpCode GetOpCode() const override { return ERigVMOpCode::ArraySetNum; }
-	virtual FName GetArgumentNameForOperandIndex(int32 InOperandIndex, int32 InTotalOperands) const override;
-	virtual const TArray<FRigVMTemplateArgumentInfo>& GetArgumentInfos() const override;
-	virtual FRigVMTemplateTypeMap OnNewArgumentType(const FName& InArgumentName, TRigVMTypeIndex InTypeIndex) const override;
+	RIGVM_API virtual FName GetArgumentNameForOperandIndex(int32 InOperandIndex, int32 InTotalOperands) const override;
+	RIGVM_API virtual const TArray<FRigVMTemplateArgumentInfo>& GetArgumentInfos() const override;
+	RIGVM_API virtual FRigVMTemplateTypeMap OnNewArgumentType(const FName& InArgumentName, TRigVMTypeIndex InTypeIndex) const override;
 #if WITH_EDITOR
-	virtual FText GetNodeTooltip(const FRigVMTemplateTypeMap& InTypes) const override;
-	virtual FText GetArgumentTooltip(const FName& InArgumentName, TRigVMTypeIndex InTypeIndex) const override;
+	RIGVM_API virtual FText GetNodeTooltip(const FRigVMTemplateTypeMap& InTypes) const override;
+	RIGVM_API virtual FText GetArgumentTooltip(const FName& InArgumentName, TRigVMTypeIndex InTypeIndex) const override;
 #endif
 
 protected:
 	virtual FRigVMFunctionPtr GetDispatchFunctionImpl(const FRigVMTemplateTypeMap& InTypes) const override { return &FRigVMDispatch_ArraySetNum::Execute; }
-	static void Execute(FRigVMExtendedExecuteContext& InContext, FRigVMMemoryHandleArray Handles, FRigVMPredicateBranchArray Predicates);
+	static RIGVM_API void Execute(FRigVMExtendedExecuteContext& InContext, FRigVMMemoryHandleArray Handles, FRigVMPredicateBranchArray Predicates);
 };
 
 USTRUCT(meta=(DisplayName = "Init"))
-struct RIGVM_API FRigVMDispatch_ArrayInit : public FRigVMDispatch_ArrayBaseMutable
+struct FRigVMDispatch_ArrayInit : public FRigVMDispatch_ArrayBaseMutable
 {
 	GENERATED_BODY()
 
@@ -161,21 +161,21 @@ public:
 		FactoryScriptStruct = StaticStruct();
 	}
 
-	virtual FName GetArgumentNameForOperandIndex(int32 InOperandIndex, int32 InTotalOperands) const override;
-	virtual const TArray<FRigVMTemplateArgumentInfo>& GetArgumentInfos() const override;
-	virtual FRigVMTemplateTypeMap OnNewArgumentType(const FName& InArgumentName, TRigVMTypeIndex InTypeIndex) const override;
+	RIGVM_API virtual FName GetArgumentNameForOperandIndex(int32 InOperandIndex, int32 InTotalOperands) const override;
+	RIGVM_API virtual const TArray<FRigVMTemplateArgumentInfo>& GetArgumentInfos() const override;
+	RIGVM_API virtual FRigVMTemplateTypeMap OnNewArgumentType(const FName& InArgumentName, TRigVMTypeIndex InTypeIndex) const override;
 #if WITH_EDITOR
-	virtual FText GetNodeTooltip(const FRigVMTemplateTypeMap& InTypes) const override;
-	virtual FText GetArgumentTooltip(const FName& InArgumentName, TRigVMTypeIndex InTypeIndex) const override;
+	RIGVM_API virtual FText GetNodeTooltip(const FRigVMTemplateTypeMap& InTypes) const override;
+	RIGVM_API virtual FText GetArgumentTooltip(const FName& InArgumentName, TRigVMTypeIndex InTypeIndex) const override;
 #endif
 
 protected:
 	virtual FRigVMFunctionPtr GetDispatchFunctionImpl(const FRigVMTemplateTypeMap& InTypes) const override { return &FRigVMDispatch_ArrayInit::Execute; }
-	static void Execute(FRigVMExtendedExecuteContext& InContext, FRigVMMemoryHandleArray Handles, FRigVMPredicateBranchArray Predicates);
+	static RIGVM_API void Execute(FRigVMExtendedExecuteContext& InContext, FRigVMMemoryHandleArray Handles, FRigVMPredicateBranchArray Predicates);
 };
 
 USTRUCT(meta=(DisplayName = "At", Keywords = "Get Index,At Index,[]"))
-struct RIGVM_API FRigVMDispatch_ArrayGetAtIndex : public FRigVMDispatch_ArrayBase
+struct FRigVMDispatch_ArrayGetAtIndex : public FRigVMDispatch_ArrayBase
 {
 	GENERATED_BODY()
 
@@ -186,21 +186,21 @@ public:
 	}
 
 	virtual ERigVMOpCode GetOpCode() const override { return ERigVMOpCode::ArrayGetAtIndex; }
-	virtual FName GetArgumentNameForOperandIndex(int32 InOperandIndex, int32 InTotalOperands) const override;
-	virtual const TArray<FRigVMTemplateArgumentInfo>& GetArgumentInfos() const override;
-	virtual FRigVMTemplateTypeMap OnNewArgumentType(const FName& InArgumentName, TRigVMTypeIndex InTypeIndex) const override;
+	RIGVM_API virtual FName GetArgumentNameForOperandIndex(int32 InOperandIndex, int32 InTotalOperands) const override;
+	RIGVM_API virtual const TArray<FRigVMTemplateArgumentInfo>& GetArgumentInfos() const override;
+	RIGVM_API virtual FRigVMTemplateTypeMap OnNewArgumentType(const FName& InArgumentName, TRigVMTypeIndex InTypeIndex) const override;
 #if WITH_EDITOR
-	virtual FText GetNodeTooltip(const FRigVMTemplateTypeMap& InTypes) const override;
-	virtual FText GetArgumentTooltip(const FName& InArgumentName, TRigVMTypeIndex InTypeIndex) const override;
+	RIGVM_API virtual FText GetNodeTooltip(const FRigVMTemplateTypeMap& InTypes) const override;
+	RIGVM_API virtual FText GetArgumentTooltip(const FName& InArgumentName, TRigVMTypeIndex InTypeIndex) const override;
 #endif
 
 protected:
 	virtual FRigVMFunctionPtr GetDispatchFunctionImpl(const FRigVMTemplateTypeMap& InTypes) const override { return &FRigVMDispatch_ArrayGetAtIndex::Execute; }
-	static void Execute(FRigVMExtendedExecuteContext& InContext, FRigVMMemoryHandleArray Handles, FRigVMPredicateBranchArray Predicates);
+	static RIGVM_API void Execute(FRigVMExtendedExecuteContext& InContext, FRigVMMemoryHandleArray Handles, FRigVMPredicateBranchArray Predicates);
 };
 
 USTRUCT(meta=(DisplayName = "Set At"))
-struct RIGVM_API FRigVMDispatch_ArraySetAtIndex : public FRigVMDispatch_ArrayBaseMutable
+struct FRigVMDispatch_ArraySetAtIndex : public FRigVMDispatch_ArrayBaseMutable
 {
 	GENERATED_BODY()
 
@@ -211,21 +211,21 @@ public:
 	}
 
 	virtual ERigVMOpCode GetOpCode() const override { return ERigVMOpCode::ArraySetAtIndex; }
-	virtual FName GetArgumentNameForOperandIndex(int32 InOperandIndex, int32 InTotalOperands) const override;
-	virtual const TArray<FRigVMTemplateArgumentInfo>& GetArgumentInfos() const override;
-	virtual FRigVMTemplateTypeMap OnNewArgumentType(const FName& InArgumentName, TRigVMTypeIndex InTypeIndex) const override;
+	RIGVM_API virtual FName GetArgumentNameForOperandIndex(int32 InOperandIndex, int32 InTotalOperands) const override;
+	RIGVM_API virtual const TArray<FRigVMTemplateArgumentInfo>& GetArgumentInfos() const override;
+	RIGVM_API virtual FRigVMTemplateTypeMap OnNewArgumentType(const FName& InArgumentName, TRigVMTypeIndex InTypeIndex) const override;
 #if WITH_EDITOR
-	virtual FText GetNodeTooltip(const FRigVMTemplateTypeMap& InTypes) const override;
-	virtual FText GetArgumentTooltip(const FName& InArgumentName, TRigVMTypeIndex InTypeIndex) const override;
+	RIGVM_API virtual FText GetNodeTooltip(const FRigVMTemplateTypeMap& InTypes) const override;
+	RIGVM_API virtual FText GetArgumentTooltip(const FName& InArgumentName, TRigVMTypeIndex InTypeIndex) const override;
 #endif
 
 protected:
 	virtual FRigVMFunctionPtr GetDispatchFunctionImpl(const FRigVMTemplateTypeMap& InTypes) const override { return &FRigVMDispatch_ArraySetAtIndex::Execute; }
-	static void Execute(FRigVMExtendedExecuteContext& InContext, FRigVMMemoryHandleArray Handles, FRigVMPredicateBranchArray Predicates);
+	static RIGVM_API void Execute(FRigVMExtendedExecuteContext& InContext, FRigVMMemoryHandleArray Handles, FRigVMPredicateBranchArray Predicates);
 };
 
 USTRUCT(meta=(DisplayName = "Add"))
-struct RIGVM_API FRigVMDispatch_ArrayAdd : public FRigVMDispatch_ArraySetAtIndex
+struct FRigVMDispatch_ArrayAdd : public FRigVMDispatch_ArraySetAtIndex
 {
 	GENERATED_BODY()
 
@@ -236,20 +236,20 @@ public:
 	}
 
 	virtual ERigVMOpCode GetOpCode() const override { return ERigVMOpCode::ArrayAdd; }
-	virtual FName GetArgumentNameForOperandIndex(int32 InOperandIndex, int32 InTotalOperands) const override;
-	virtual const TArray<FRigVMTemplateArgumentInfo>& GetArgumentInfos() const override;
+	RIGVM_API virtual FName GetArgumentNameForOperandIndex(int32 InOperandIndex, int32 InTotalOperands) const override;
+	RIGVM_API virtual const TArray<FRigVMTemplateArgumentInfo>& GetArgumentInfos() const override;
 #if WITH_EDITOR
-	virtual FText GetNodeTooltip(const FRigVMTemplateTypeMap& InTypes) const override;
-	virtual FText GetArgumentTooltip(const FName& InArgumentName, TRigVMTypeIndex InTypeIndex) const override;
+	RIGVM_API virtual FText GetNodeTooltip(const FRigVMTemplateTypeMap& InTypes) const override;
+	RIGVM_API virtual FText GetArgumentTooltip(const FName& InArgumentName, TRigVMTypeIndex InTypeIndex) const override;
 #endif
 
 protected:
 	virtual FRigVMFunctionPtr GetDispatchFunctionImpl(const FRigVMTemplateTypeMap& InTypes) const override { return &FRigVMDispatch_ArrayAdd::Execute; }
-	static void Execute(FRigVMExtendedExecuteContext& InContext, FRigVMMemoryHandleArray Handles, FRigVMPredicateBranchArray Predicates);
+	static RIGVM_API void Execute(FRigVMExtendedExecuteContext& InContext, FRigVMMemoryHandleArray Handles, FRigVMPredicateBranchArray Predicates);
 };
 
 USTRUCT(meta=(DisplayName = "Insert"))
-struct RIGVM_API FRigVMDispatch_ArrayInsert : public FRigVMDispatch_ArraySetAtIndex
+struct FRigVMDispatch_ArrayInsert : public FRigVMDispatch_ArraySetAtIndex
 {
 	GENERATED_BODY()
 
@@ -261,17 +261,17 @@ public:
 
 	virtual ERigVMOpCode GetOpCode() const override { return ERigVMOpCode::ArrayInsert; }
 #if WITH_EDITOR
-   	virtual FText GetNodeTooltip(const FRigVMTemplateTypeMap& InTypes) const override;
-   	virtual FText GetArgumentTooltip(const FName& InArgumentName, TRigVMTypeIndex InTypeIndex) const override;
+   	RIGVM_API virtual FText GetNodeTooltip(const FRigVMTemplateTypeMap& InTypes) const override;
+   	RIGVM_API virtual FText GetArgumentTooltip(const FName& InArgumentName, TRigVMTypeIndex InTypeIndex) const override;
 #endif
 
 protected:
 	virtual FRigVMFunctionPtr GetDispatchFunctionImpl(const FRigVMTemplateTypeMap& InTypes) const override { return &FRigVMDispatch_ArrayInsert::Execute; }
-	static void Execute(FRigVMExtendedExecuteContext& InContext, FRigVMMemoryHandleArray Handles, FRigVMPredicateBranchArray Predicates);
+	static RIGVM_API void Execute(FRigVMExtendedExecuteContext& InContext, FRigVMMemoryHandleArray Handles, FRigVMPredicateBranchArray Predicates);
 };
 
 USTRUCT(meta=(DisplayName = "Remove"))
-struct RIGVM_API FRigVMDispatch_ArrayRemove : public FRigVMDispatch_ArrayBaseMutable
+struct FRigVMDispatch_ArrayRemove : public FRigVMDispatch_ArrayBaseMutable
 {
 	GENERATED_BODY()
 
@@ -282,21 +282,21 @@ public:
 	}
 
 	virtual ERigVMOpCode GetOpCode() const override { return ERigVMOpCode::ArrayRemove; }
-	virtual FName GetArgumentNameForOperandIndex(int32 InOperandIndex, int32 InTotalOperands) const override;
-	virtual const TArray<FRigVMTemplateArgumentInfo>& GetArgumentInfos() const override;
-	virtual FRigVMTemplateTypeMap OnNewArgumentType(const FName& InArgumentName, TRigVMTypeIndex InTypeIndex) const override;
+	RIGVM_API virtual FName GetArgumentNameForOperandIndex(int32 InOperandIndex, int32 InTotalOperands) const override;
+	RIGVM_API virtual const TArray<FRigVMTemplateArgumentInfo>& GetArgumentInfos() const override;
+	RIGVM_API virtual FRigVMTemplateTypeMap OnNewArgumentType(const FName& InArgumentName, TRigVMTypeIndex InTypeIndex) const override;
 #if WITH_EDITOR
-   	virtual FText GetNodeTooltip(const FRigVMTemplateTypeMap& InTypes) const override;
-   	virtual FText GetArgumentTooltip(const FName& InArgumentName, TRigVMTypeIndex InTypeIndex) const override;
+   	RIGVM_API virtual FText GetNodeTooltip(const FRigVMTemplateTypeMap& InTypes) const override;
+   	RIGVM_API virtual FText GetArgumentTooltip(const FName& InArgumentName, TRigVMTypeIndex InTypeIndex) const override;
 #endif
 
 protected:
 	virtual FRigVMFunctionPtr GetDispatchFunctionImpl(const FRigVMTemplateTypeMap& InTypes) const override { return &FRigVMDispatch_ArrayRemove::Execute; }
-	static void Execute(FRigVMExtendedExecuteContext& InContext, FRigVMMemoryHandleArray Handles, FRigVMPredicateBranchArray Predicates);
+	static RIGVM_API void Execute(FRigVMExtendedExecuteContext& InContext, FRigVMMemoryHandleArray Handles, FRigVMPredicateBranchArray Predicates);
 };
 
 USTRUCT(meta=(DisplayName = "Reverse"))
-struct RIGVM_API FRigVMDispatch_ArrayReverse : public FRigVMDispatch_ArrayReset
+struct FRigVMDispatch_ArrayReverse : public FRigVMDispatch_ArrayReset
 {
 	GENERATED_BODY()
 
@@ -308,18 +308,18 @@ public:
 
 	virtual ERigVMOpCode GetOpCode() const override { return ERigVMOpCode::ArrayReverse; }
 #if WITH_EDITOR
-   	virtual FText GetNodeTooltip(const FRigVMTemplateTypeMap& InTypes) const override;
-   	virtual FText GetArgumentTooltip(const FName& InArgumentName, TRigVMTypeIndex InTypeIndex) const override;
+   	RIGVM_API virtual FText GetNodeTooltip(const FRigVMTemplateTypeMap& InTypes) const override;
+   	RIGVM_API virtual FText GetArgumentTooltip(const FName& InArgumentName, TRigVMTypeIndex InTypeIndex) const override;
 #endif
 
 protected:
 
 	virtual FRigVMFunctionPtr GetDispatchFunctionImpl(const FRigVMTemplateTypeMap& InTypes) const override { return &FRigVMDispatch_ArrayReverse::Execute; }
-	static void Execute(FRigVMExtendedExecuteContext& InContext, FRigVMMemoryHandleArray Handles, FRigVMPredicateBranchArray Predicates);
+	static RIGVM_API void Execute(FRigVMExtendedExecuteContext& InContext, FRigVMMemoryHandleArray Handles, FRigVMPredicateBranchArray Predicates);
 };
 
 USTRUCT(meta=(DisplayName = "Find"))
-struct RIGVM_API FRigVMDispatch_ArrayFind : public FRigVMDispatch_ArrayBase
+struct FRigVMDispatch_ArrayFind : public FRigVMDispatch_ArrayBase
 {
 	GENERATED_BODY()
 
@@ -330,21 +330,21 @@ public:
 	}
 
 	virtual ERigVMOpCode GetOpCode() const override { return ERigVMOpCode::ArrayFind; }
-	virtual FName GetArgumentNameForOperandIndex(int32 InOperandIndex, int32 InTotalOperands) const override;
-	virtual const TArray<FRigVMTemplateArgumentInfo>& GetArgumentInfos() const override;
-	virtual FRigVMTemplateTypeMap OnNewArgumentType(const FName& InArgumentName, TRigVMTypeIndex InTypeIndex) const override;
+	RIGVM_API virtual FName GetArgumentNameForOperandIndex(int32 InOperandIndex, int32 InTotalOperands) const override;
+	RIGVM_API virtual const TArray<FRigVMTemplateArgumentInfo>& GetArgumentInfos() const override;
+	RIGVM_API virtual FRigVMTemplateTypeMap OnNewArgumentType(const FName& InArgumentName, TRigVMTypeIndex InTypeIndex) const override;
 #if WITH_EDITOR
-   	virtual FText GetNodeTooltip(const FRigVMTemplateTypeMap& InTypes) const override;
-   	virtual FText GetArgumentTooltip(const FName& InArgumentName, TRigVMTypeIndex InTypeIndex) const override;
+   	RIGVM_API virtual FText GetNodeTooltip(const FRigVMTemplateTypeMap& InTypes) const override;
+   	RIGVM_API virtual FText GetArgumentTooltip(const FName& InArgumentName, TRigVMTypeIndex InTypeIndex) const override;
 #endif
 
 protected:
 	virtual FRigVMFunctionPtr GetDispatchFunctionImpl(const FRigVMTemplateTypeMap& InTypes) const override { return &FRigVMDispatch_ArrayFind::Execute; }
-	static void Execute(FRigVMExtendedExecuteContext& InContext, FRigVMMemoryHandleArray Handles, FRigVMPredicateBranchArray Predicates);
+	static RIGVM_API void Execute(FRigVMExtendedExecuteContext& InContext, FRigVMMemoryHandleArray Handles, FRigVMPredicateBranchArray Predicates);
 };
 
 USTRUCT(meta=(DisplayName = "Append"))
-struct RIGVM_API FRigVMDispatch_ArrayAppend : public FRigVMDispatch_ArrayBaseMutable
+struct FRigVMDispatch_ArrayAppend : public FRigVMDispatch_ArrayBaseMutable
 {
 	GENERATED_BODY()
 
@@ -355,21 +355,21 @@ public:
 	}
 
 	virtual ERigVMOpCode GetOpCode() const override { return ERigVMOpCode::ArrayAppend; }
-	virtual FName GetArgumentNameForOperandIndex(int32 InOperandIndex, int32 InTotalOperands) const override;
-	virtual const TArray<FRigVMTemplateArgumentInfo>& GetArgumentInfos() const override;
-	virtual FRigVMTemplateTypeMap OnNewArgumentType(const FName& InArgumentName, TRigVMTypeIndex InTypeIndex) const override;
+	RIGVM_API virtual FName GetArgumentNameForOperandIndex(int32 InOperandIndex, int32 InTotalOperands) const override;
+	RIGVM_API virtual const TArray<FRigVMTemplateArgumentInfo>& GetArgumentInfos() const override;
+	RIGVM_API virtual FRigVMTemplateTypeMap OnNewArgumentType(const FName& InArgumentName, TRigVMTypeIndex InTypeIndex) const override;
 #if WITH_EDITOR
-   	virtual FText GetNodeTooltip(const FRigVMTemplateTypeMap& InTypes) const override;
-   	virtual FText GetArgumentTooltip(const FName& InArgumentName, TRigVMTypeIndex InTypeIndex) const override;
+   	RIGVM_API virtual FText GetNodeTooltip(const FRigVMTemplateTypeMap& InTypes) const override;
+   	RIGVM_API virtual FText GetArgumentTooltip(const FName& InArgumentName, TRigVMTypeIndex InTypeIndex) const override;
 #endif
 
 protected:
 	virtual FRigVMFunctionPtr GetDispatchFunctionImpl(const FRigVMTemplateTypeMap& InTypes) const override { return &FRigVMDispatch_ArrayAppend::Execute; }
-	static void Execute(FRigVMExtendedExecuteContext& InContext, FRigVMMemoryHandleArray Handles, FRigVMPredicateBranchArray Predicates);
+	static RIGVM_API void Execute(FRigVMExtendedExecuteContext& InContext, FRigVMMemoryHandleArray Handles, FRigVMPredicateBranchArray Predicates);
 };
 
 USTRUCT(meta=(DisplayName = "Clone"))
-struct RIGVM_API FRigVMDispatch_ArrayClone : public FRigVMDispatch_ArrayBase
+struct FRigVMDispatch_ArrayClone : public FRigVMDispatch_ArrayBase
 {
 	GENERATED_BODY()
 
@@ -380,21 +380,21 @@ public:
 	}
 
 	virtual ERigVMOpCode GetOpCode() const override { return ERigVMOpCode::ArrayClone; }
-	virtual FName GetArgumentNameForOperandIndex(int32 InOperandIndex, int32 InTotalOperands) const override;
-	virtual const TArray<FRigVMTemplateArgumentInfo>& GetArgumentInfos() const override;
-	virtual FRigVMTemplateTypeMap OnNewArgumentType(const FName& InArgumentName, TRigVMTypeIndex InTypeIndex) const override;
+	RIGVM_API virtual FName GetArgumentNameForOperandIndex(int32 InOperandIndex, int32 InTotalOperands) const override;
+	RIGVM_API virtual const TArray<FRigVMTemplateArgumentInfo>& GetArgumentInfos() const override;
+	RIGVM_API virtual FRigVMTemplateTypeMap OnNewArgumentType(const FName& InArgumentName, TRigVMTypeIndex InTypeIndex) const override;
 #if WITH_EDITOR
-   	virtual FText GetNodeTooltip(const FRigVMTemplateTypeMap& InTypes) const override;
-   	virtual FText GetArgumentTooltip(const FName& InArgumentName, TRigVMTypeIndex InTypeIndex) const override;
+   	RIGVM_API virtual FText GetNodeTooltip(const FRigVMTemplateTypeMap& InTypes) const override;
+   	RIGVM_API virtual FText GetArgumentTooltip(const FName& InArgumentName, TRigVMTypeIndex InTypeIndex) const override;
 #endif
 
 protected:
 	virtual FRigVMFunctionPtr GetDispatchFunctionImpl(const FRigVMTemplateTypeMap& InTypes) const override { return &FRigVMDispatch_ArrayClone::Execute; }
-	static void Execute(FRigVMExtendedExecuteContext& InContext, FRigVMMemoryHandleArray Handles, FRigVMPredicateBranchArray Predicates);
+	static RIGVM_API void Execute(FRigVMExtendedExecuteContext& InContext, FRigVMMemoryHandleArray Handles, FRigVMPredicateBranchArray Predicates);
 };
 
 USTRUCT(meta=(DisplayName = "Union"))
-struct RIGVM_API FRigVMDispatch_ArrayUnion : public FRigVMDispatch_ArrayAppend
+struct FRigVMDispatch_ArrayUnion : public FRigVMDispatch_ArrayAppend
 {
 	GENERATED_BODY()
 
@@ -406,17 +406,17 @@ public:
 
 	virtual ERigVMOpCode GetOpCode() const override { return ERigVMOpCode::ArrayUnion; }
 #if WITH_EDITOR
-   	virtual FText GetNodeTooltip(const FRigVMTemplateTypeMap& InTypes) const override;
-   	virtual FText GetArgumentTooltip(const FName& InArgumentName, TRigVMTypeIndex InTypeIndex) const override;
+   	RIGVM_API virtual FText GetNodeTooltip(const FRigVMTemplateTypeMap& InTypes) const override;
+   	RIGVM_API virtual FText GetArgumentTooltip(const FName& InArgumentName, TRigVMTypeIndex InTypeIndex) const override;
 #endif
 
 protected:
 	virtual FRigVMFunctionPtr GetDispatchFunctionImpl(const FRigVMTemplateTypeMap& InTypes) const override { return &FRigVMDispatch_ArrayUnion::Execute; }
-	static void Execute(FRigVMExtendedExecuteContext& InContext, FRigVMMemoryHandleArray Handles, FRigVMPredicateBranchArray Predicates);
+	static RIGVM_API void Execute(FRigVMExtendedExecuteContext& InContext, FRigVMMemoryHandleArray Handles, FRigVMPredicateBranchArray Predicates);
 };
 
 USTRUCT(meta=(DisplayName = "Difference"))
-struct RIGVM_API FRigVMDispatch_ArrayDifference : public FRigVMDispatch_ArrayBase
+struct FRigVMDispatch_ArrayDifference : public FRigVMDispatch_ArrayBase
 {
 	GENERATED_BODY()
 
@@ -427,21 +427,21 @@ public:
 	}
 
 	virtual ERigVMOpCode GetOpCode() const override { return ERigVMOpCode::ArrayDifference; }
-	virtual FName GetArgumentNameForOperandIndex(int32 InOperandIndex, int32 InTotalOperands) const override;
-	virtual const TArray<FRigVMTemplateArgumentInfo>& GetArgumentInfos() const override;
-	virtual FRigVMTemplateTypeMap OnNewArgumentType(const FName& InArgumentName, TRigVMTypeIndex InTypeIndex) const override;
+	RIGVM_API virtual FName GetArgumentNameForOperandIndex(int32 InOperandIndex, int32 InTotalOperands) const override;
+	RIGVM_API virtual const TArray<FRigVMTemplateArgumentInfo>& GetArgumentInfos() const override;
+	RIGVM_API virtual FRigVMTemplateTypeMap OnNewArgumentType(const FName& InArgumentName, TRigVMTypeIndex InTypeIndex) const override;
 #if WITH_EDITOR
-   	virtual FText GetNodeTooltip(const FRigVMTemplateTypeMap& InTypes) const override;
-   	virtual FText GetArgumentTooltip(const FName& InArgumentName, TRigVMTypeIndex InTypeIndex) const override;
+   	RIGVM_API virtual FText GetNodeTooltip(const FRigVMTemplateTypeMap& InTypes) const override;
+   	RIGVM_API virtual FText GetArgumentTooltip(const FName& InArgumentName, TRigVMTypeIndex InTypeIndex) const override;
 #endif
 
 protected:
 	virtual FRigVMFunctionPtr GetDispatchFunctionImpl(const FRigVMTemplateTypeMap& InTypes) const override { return &FRigVMDispatch_ArrayDifference::Execute; }
-	static void Execute(FRigVMExtendedExecuteContext& InContext, FRigVMMemoryHandleArray Handles, FRigVMPredicateBranchArray Predicates);
+	static RIGVM_API void Execute(FRigVMExtendedExecuteContext& InContext, FRigVMMemoryHandleArray Handles, FRigVMPredicateBranchArray Predicates);
 };
 
 USTRUCT(meta=(DisplayName = "Intersection"))
-struct RIGVM_API FRigVMDispatch_ArrayIntersection : public FRigVMDispatch_ArrayDifference
+struct FRigVMDispatch_ArrayIntersection : public FRigVMDispatch_ArrayDifference
 {
 	GENERATED_BODY()
 
@@ -453,17 +453,17 @@ public:
 
 	virtual ERigVMOpCode GetOpCode() const override { return ERigVMOpCode::ArrayIntersection; }
 #if WITH_EDITOR
-   	virtual FText GetNodeTooltip(const FRigVMTemplateTypeMap& InTypes) const override;
-   	virtual FText GetArgumentTooltip(const FName& InArgumentName, TRigVMTypeIndex InTypeIndex) const override;
+   	RIGVM_API virtual FText GetNodeTooltip(const FRigVMTemplateTypeMap& InTypes) const override;
+   	RIGVM_API virtual FText GetArgumentTooltip(const FName& InArgumentName, TRigVMTypeIndex InTypeIndex) const override;
 #endif
 
 protected:
 	virtual FRigVMFunctionPtr GetDispatchFunctionImpl(const FRigVMTemplateTypeMap& InTypes) const override { return &FRigVMDispatch_ArrayIntersection::Execute; }
-	static void Execute(FRigVMExtendedExecuteContext& InContext, FRigVMMemoryHandleArray Handles, FRigVMPredicateBranchArray Predicates);
+	static RIGVM_API void Execute(FRigVMExtendedExecuteContext& InContext, FRigVMMemoryHandleArray Handles, FRigVMPredicateBranchArray Predicates);
 };
 
 USTRUCT(meta=(DisplayName = "For Each", Icon="EditorStyle|GraphEditor.Macro.ForEach_16x"))
-struct RIGVM_API FRigVMDispatch_ArrayIterator : public FRigVMDispatch_ArrayBaseMutable
+struct FRigVMDispatch_ArrayIterator : public FRigVMDispatch_ArrayBaseMutable
 {
 	GENERATED_BODY()
 
@@ -474,18 +474,18 @@ public:
 	}
 
 	virtual ERigVMOpCode GetOpCode() const override { return ERigVMOpCode::ArrayIterator; }
-	virtual FName GetArgumentNameForOperandIndex(int32 InOperandIndex, int32 InTotalOperands) const override;
-	virtual const TArray<FRigVMTemplateArgumentInfo>& GetArgumentInfos() const override;
-	virtual const TArray<FRigVMExecuteArgument>& GetExecuteArguments_Impl(const FRigVMDispatchContext& InContext) const override;
-	virtual const TArray<FName>& GetControlFlowBlocks_Impl(const FRigVMDispatchContext& InContext) const override;
-	virtual const bool IsControlFlowBlockSliced(const FName& InBlockName) const override;
-	virtual FRigVMTemplateTypeMap OnNewArgumentType(const FName& InArgumentName, TRigVMTypeIndex InTypeIndex) const override;
+	RIGVM_API virtual FName GetArgumentNameForOperandIndex(int32 InOperandIndex, int32 InTotalOperands) const override;
+	RIGVM_API virtual const TArray<FRigVMTemplateArgumentInfo>& GetArgumentInfos() const override;
+	RIGVM_API virtual const TArray<FRigVMExecuteArgument>& GetExecuteArguments_Impl(const FRigVMDispatchContext& InContext) const override;
+	RIGVM_API virtual const TArray<FName>& GetControlFlowBlocks_Impl(const FRigVMDispatchContext& InContext) const override;
+	RIGVM_API virtual const bool IsControlFlowBlockSliced(const FName& InBlockName) const override;
+	RIGVM_API virtual FRigVMTemplateTypeMap OnNewArgumentType(const FName& InArgumentName, TRigVMTypeIndex InTypeIndex) const override;
 #if WITH_EDITOR
-   	virtual FText GetNodeTooltip(const FRigVMTemplateTypeMap& InTypes) const override;
-   	virtual FText GetArgumentTooltip(const FName& InArgumentName, TRigVMTypeIndex InTypeIndex) const override;
+   	RIGVM_API virtual FText GetNodeTooltip(const FRigVMTemplateTypeMap& InTypes) const override;
+   	RIGVM_API virtual FText GetArgumentTooltip(const FName& InArgumentName, TRigVMTypeIndex InTypeIndex) const override;
 #endif
 
 protected:
 	virtual FRigVMFunctionPtr GetDispatchFunctionImpl(const FRigVMTemplateTypeMap& InTypes) const override { return &FRigVMDispatch_ArrayIterator::Execute; }
-	static void Execute(FRigVMExtendedExecuteContext& InContext, FRigVMMemoryHandleArray Handles, FRigVMPredicateBranchArray Predicates);
+	static RIGVM_API void Execute(FRigVMExtendedExecuteContext& InContext, FRigVMMemoryHandleArray Handles, FRigVMPredicateBranchArray Predicates);
 };
diff --git a/Engine/Plugins/Runtime/RigVM/Source/RigVM/Public/RigVMFunctions/RigVMDispatch_CastEnum.h b/Engine/Plugins/Runtime/RigVM/Source/RigVM/Public/RigVMFunctions/RigVMDispatch_CastEnum.h
--- a/Engine/Plugins/Runtime/RigVM/Source/RigVM/Public/RigVMFunctions/RigVMDispatch_CastEnum.h
+++ b/Engine/Plugins/Runtime/RigVM/Source/RigVM/Public/RigVMFunctions/RigVMDispatch_CastEnum.h
@@ -12,7 +12,7 @@ struct FRigVMDispatch_CastEnumBase
 };
 
 USTRUCT(meta=(DisplayName = "Cast", Category = "Enum", Keywords = "As", NodeColor = "1,1,1,1"))
-struct RIGVM_API FRigVMDispatch_CastEnumToInt : public FRigVMDispatchFactory, public FRigVMDispatch_CastEnumBase
+struct FRigVMDispatch_CastEnumToInt : public FRigVMDispatchFactory, public FRigVMDispatch_CastEnumBase
 {
 	GENERATED_BODY()
 
@@ -22,21 +22,21 @@ public:
 		FactoryScriptStruct = StaticStruct();
 	}
 	
-	virtual const TArray<FRigVMTemplateArgumentInfo>& GetArgumentInfos() const override;
-	virtual bool GetPermutationsFromArgumentType(const FName& InArgumentName, const TRigVMTypeIndex& InTypeIndex, TArray<FRigVMTemplateTypeMap, TInlineAllocator<1>>& OutPermutations) const override;
+	RIGVM_API virtual const TArray<FRigVMTemplateArgumentInfo>& GetArgumentInfos() const override;
+	RIGVM_API virtual bool GetPermutationsFromArgumentType(const FName& InArgumentName, const TRigVMTypeIndex& InTypeIndex, TArray<FRigVMTemplateTypeMap, TInlineAllocator<1>>& OutPermutations) const override;
 #if WITH_EDITOR
-	virtual FString GetNodeTitle(const FRigVMTemplateTypeMap& InTypes) const override;
-	virtual FText GetNodeTooltip(const FRigVMTemplateTypeMap& InTypes) const override;
+	RIGVM_API virtual FString GetNodeTitle(const FRigVMTemplateTypeMap& InTypes) const override;
+	RIGVM_API virtual FText GetNodeTooltip(const FRigVMTemplateTypeMap& InTypes) const override;
 #endif
 	virtual bool IsSingleton() const override { return true; }
 	
 protected:
 	virtual FRigVMFunctionPtr GetDispatchFunctionImpl(const FRigVMTemplateTypeMap& InTypes) const override { return &FRigVMDispatch_CastEnumToInt::Execute; }
-	static void Execute(FRigVMExtendedExecuteContext& InContext, FRigVMMemoryHandleArray Handles, FRigVMPredicateBranchArray RigVMBranches);
+	static RIGVM_API void Execute(FRigVMExtendedExecuteContext& InContext, FRigVMMemoryHandleArray Handles, FRigVMPredicateBranchArray RigVMBranches);
 };
 
 USTRUCT(meta=(DisplayName = "Cast", Category = "Enum", Keywords = "As", NodeColor = "1,1,1,1"))
-struct RIGVM_API FRigVMDispatch_CastIntToEnum : public FRigVMDispatchFactory, public FRigVMDispatch_CastEnumBase
+struct FRigVMDispatch_CastIntToEnum : public FRigVMDispatchFactory, public FRigVMDispatch_CastEnumBase
 {
 	GENERATED_BODY()
 
@@ -46,15 +46,15 @@ public:
 		FactoryScriptStruct = StaticStruct();
 	}
 	
-	virtual const TArray<FRigVMTemplateArgumentInfo>& GetArgumentInfos() const override;
-	virtual bool GetPermutationsFromArgumentType(const FName& InArgumentName, const TRigVMTypeIndex& InTypeIndex, TArray<FRigVMTemplateTypeMap, TInlineAllocator<1>>& OutPermutations) const override;
+	RIGVM_API virtual const TArray<FRigVMTemplateArgumentInfo>& GetArgumentInfos() const override;
+	RIGVM_API virtual bool GetPermutationsFromArgumentType(const FName& InArgumentName, const TRigVMTypeIndex& InTypeIndex, TArray<FRigVMTemplateTypeMap, TInlineAllocator<1>>& OutPermutations) const override;
 #if WITH_EDITOR
-	virtual FString GetNodeTitle(const FRigVMTemplateTypeMap& InTypes) const override;
-	virtual FText GetNodeTooltip(const FRigVMTemplateTypeMap& InTypes) const override;
+	RIGVM_API virtual FString GetNodeTitle(const FRigVMTemplateTypeMap& InTypes) const override;
+	RIGVM_API virtual FText GetNodeTooltip(const FRigVMTemplateTypeMap& InTypes) const override;
 #endif
 	virtual bool IsSingleton() const override { return true; }
 	
 protected:
 	virtual FRigVMFunctionPtr GetDispatchFunctionImpl(const FRigVMTemplateTypeMap& InTypes) const override { return &FRigVMDispatch_CastIntToEnum::Execute; }
-	static void Execute(FRigVMExtendedExecuteContext& InContext, FRigVMMemoryHandleArray Handles, FRigVMPredicateBranchArray RigVMBranches);
+	static RIGVM_API void Execute(FRigVMExtendedExecuteContext& InContext, FRigVMMemoryHandleArray Handles, FRigVMPredicateBranchArray RigVMBranches);
 };
