[RigVM]

diff --git a/Engine/Plugins/Runtime/RigVM/Source/RigVM/Public/RigVMFunctions/RigVMDispatch_CastObject.h b/Engine/Plugins/Runtime/RigVM/Source/RigVM/Public/RigVMFunctions/RigVMDispatch_CastObject.h
--- a/Engine/Plugins/Runtime/RigVM/Source/RigVM/Public/RigVMFunctions/RigVMDispatch_CastObject.h
+++ b/Engine/Plugins/Runtime/RigVM/Source/RigVM/Public/RigVMFunctions/RigVMDispatch_CastObject.h
@@ -1,4 +1,4 @@
-﻿// Copyright Epic Games, Inc. All Rights Reserved.
+// Copyright Epic Games, Inc. All Rights Reserved.
 
 #pragma once
 
@@ -6,7 +6,7 @@
 #include "RigVMDispatch_CastObject.generated.h"
 
 USTRUCT(meta=(DisplayName = "Cast", Category = "Object", Keywords = "As", NodeColor = "1,1,1,1"))
-struct RIGVM_API FRigVMDispatch_CastObject : public FRigVMDispatchFactory
+struct FRigVMDispatch_CastObject : public FRigVMDispatchFactory
 {
 	GENERATED_BODY()
 
@@ -16,17 +16,17 @@ public:
 		FactoryScriptStruct = StaticStruct();
 	}
 	
-	virtual const TArray<FRigVMTemplateArgumentInfo>& GetArgumentInfos() const override;
-	virtual bool GetPermutationsFromArgumentType(const FName& InArgumentName, const TRigVMTypeIndex& InTypeIndex, TArray<FRigVMTemplateTypeMap, TInlineAllocator<1>>& OutPermutations) const override;
+	RIGVM_API virtual const TArray<FRigVMTemplateArgumentInfo>& GetArgumentInfos() const override;
+	RIGVM_API virtual bool GetPermutationsFromArgumentType(const FName& InArgumentName, const TRigVMTypeIndex& InTypeIndex, TArray<FRigVMTemplateTypeMap, TInlineAllocator<1>>& OutPermutations) const override;
 #if WITH_EDITOR
-	virtual FString GetNodeTitle(const FRigVMTemplateTypeMap& InTypes) const override;
-	virtual FText GetNodeTooltip(const FRigVMTemplateTypeMap& InTypes) const override;
+	RIGVM_API virtual FString GetNodeTitle(const FRigVMTemplateTypeMap& InTypes) const override;
+	RIGVM_API virtual FText GetNodeTooltip(const FRigVMTemplateTypeMap& InTypes) const override;
 #endif
 	virtual bool IsSingleton() const override { return true; }
 	
 protected:
 	virtual FRigVMFunctionPtr GetDispatchFunctionImpl(const FRigVMTemplateTypeMap& InTypes) const override { return &FRigVMDispatch_CastObject::Execute; }
-	static void Execute(FRigVMExtendedExecuteContext& InContext, FRigVMMemoryHandleArray Handles, FRigVMPredicateBranchArray RigVMBranches);
+	static RIGVM_API void Execute(FRigVMExtendedExecuteContext& InContext, FRigVMMemoryHandleArray Handles, FRigVMPredicateBranchArray RigVMBranches);
 
 	static inline const FLazyName ValueName = FLazyName(TEXT("Value"));
 	static inline const FLazyName ResultName = FLazyName(TEXT("Result"));
diff --git a/Engine/Plugins/Runtime/RigVM/Source/RigVM/Public/RigVMFunctions/RigVMDispatch_Constant.h b/Engine/Plugins/Runtime/RigVM/Source/RigVM/Public/RigVMFunctions/RigVMDispatch_Constant.h
--- a/Engine/Plugins/Runtime/RigVM/Source/RigVM/Public/RigVMFunctions/RigVMDispatch_Constant.h
+++ b/Engine/Plugins/Runtime/RigVM/Source/RigVM/Public/RigVMFunctions/RigVMDispatch_Constant.h
@@ -8,7 +8,7 @@
 #include "RigVMDispatch_Constant.generated.h"
 
 USTRUCT(meta=(DisplayName = "Constant", Category = "Core", Keywords = "Value,Reroute", NodeColor = "1,1,1,1"))
-struct RIGVM_API FRigVMDispatch_Constant : public FRigVMDispatch_CoreBase
+struct FRigVMDispatch_Constant : public FRigVMDispatch_CoreBase
 {
 	GENERATED_BODY()
 
@@ -18,17 +18,17 @@ public:
 		FactoryScriptStruct = StaticStruct();
 	}
 
-	virtual const TArray<FRigVMTemplateArgumentInfo>& GetArgumentInfos() const override;
-	virtual FRigVMTemplateTypeMap OnNewArgumentType(const FName& InArgumentName, TRigVMTypeIndex InTypeIndex) const override;
+	RIGVM_API virtual const TArray<FRigVMTemplateArgumentInfo>& GetArgumentInfos() const override;
+	RIGVM_API virtual FRigVMTemplateTypeMap OnNewArgumentType(const FName& InArgumentName, TRigVMTypeIndex InTypeIndex) const override;
 #if WITH_EDITOR
-	virtual FString GetNodeTitle(const FRigVMTemplateTypeMap& InTypes) const override;
-	virtual FText GetNodeTooltip(const FRigVMTemplateTypeMap& InTypes) const override;
+	RIGVM_API virtual FString GetNodeTitle(const FRigVMTemplateTypeMap& InTypes) const override;
+	RIGVM_API virtual FText GetNodeTooltip(const FRigVMTemplateTypeMap& InTypes) const override;
 #endif
 	virtual bool IsSingleton() const override { return true; }
 	
 protected:
 	virtual FRigVMFunctionPtr GetDispatchFunctionImpl(const FRigVMTemplateTypeMap& InTypes) const override { return &FRigVMDispatch_Constant::Execute; }
-	static void Execute(FRigVMExtendedExecuteContext& InContext, FRigVMMemoryHandleArray Handles, FRigVMPredicateBranchArray RigVMBranches);
+	static RIGVM_API void Execute(FRigVMExtendedExecuteContext& InContext, FRigVMMemoryHandleArray Handles, FRigVMPredicateBranchArray RigVMBranches);
 
 	static inline const FLazyName ValueName = FLazyName(TEXT("Value"));
 
diff --git a/Engine/Plugins/Runtime/RigVM/Source/RigVM/Public/RigVMFunctions/RigVMDispatch_Core.h b/Engine/Plugins/Runtime/RigVM/Source/RigVM/Public/RigVMFunctions/RigVMDispatch_Core.h
--- a/Engine/Plugins/Runtime/RigVM/Source/RigVM/Public/RigVMFunctions/RigVMDispatch_Core.h
+++ b/Engine/Plugins/Runtime/RigVM/Source/RigVM/Public/RigVMFunctions/RigVMDispatch_Core.h
@@ -6,7 +6,7 @@
 #include "RigVMDispatch_Core.generated.h"
 
 USTRUCT(meta=(Abstract, Category = "Core", NodeColor = "0.762745, 1,0, 0.329412"))
-struct RIGVM_API FRigVMDispatch_CoreBase : public FRigVMDispatchFactory
+struct FRigVMDispatch_CoreBase : public FRigVMDispatchFactory
 {
 	GENERATED_BODY()
 };
@@ -15,7 +15,7 @@ struct RIGVM_API FRigVMDispatch_CoreBase : public FRigVMDispatchFactory
  * Compares any two values and return true if they are identical
  */
 USTRUCT(meta=(DisplayName = "Equals", Keywords = "Same,=="))
-struct RIGVM_API FRigVMDispatch_CoreEquals : public FRigVMDispatch_CoreBase
+struct FRigVMDispatch_CoreEquals : public FRigVMDispatch_CoreBase
 {
 	GENERATED_BODY()
 
@@ -26,15 +26,15 @@ public:
 		FactoryScriptStruct = StaticStruct();
 	}
 
-	virtual FName GetArgumentNameForOperandIndex(int32 InOperandIndex, int32 InTotalOperands) const override;
-	virtual const TArray<FRigVMTemplateArgumentInfo>& GetArgumentInfos() const override;
-	virtual FRigVMTemplateTypeMap OnNewArgumentType(const FName& InArgumentName, TRigVMTypeIndex InTypeIndex) const override;
+	RIGVM_API virtual FName GetArgumentNameForOperandIndex(int32 InOperandIndex, int32 InTotalOperands) const override;
+	RIGVM_API virtual const TArray<FRigVMTemplateArgumentInfo>& GetArgumentInfos() const override;
+	RIGVM_API virtual FRigVMTemplateTypeMap OnNewArgumentType(const FName& InArgumentName, TRigVMTypeIndex InTypeIndex) const override;
 
 protected:
 
-	virtual FRigVMFunctionPtr GetDispatchFunctionImpl(const FRigVMTemplateTypeMap& InTypes) const override;
-	static bool AdaptResult(bool bResult, const FRigVMExtendedExecuteContext& InContext);
-	static void Execute(FRigVMExtendedExecuteContext& InContext, FRigVMMemoryHandleArray Handles, FRigVMPredicateBranchArray Predicates);
+	RIGVM_API virtual FRigVMFunctionPtr GetDispatchFunctionImpl(const FRigVMTemplateTypeMap& InTypes) const override;
+	static RIGVM_API bool AdaptResult(bool bResult, const FRigVMExtendedExecuteContext& InContext);
+	static RIGVM_API void Execute(FRigVMExtendedExecuteContext& InContext, FRigVMMemoryHandleArray Handles, FRigVMPredicateBranchArray Predicates);
 
 	template<typename T>
 	static void Equals(FRigVMExtendedExecuteContext& InContext, FRigVMMemoryHandleArray Handles, FRigVMPredicateBranchArray Predicates)
@@ -91,7 +91,7 @@ protected:
  * Compares any two values and return true if they are identical
  */
 USTRUCT(meta=(DisplayName = "Not Equals", Keywords = "Different,!=,Xor"))
-struct RIGVM_API FRigVMDispatch_CoreNotEquals : public FRigVMDispatch_CoreEquals
+struct FRigVMDispatch_CoreNotEquals : public FRigVMDispatch_CoreEquals
 {
 	GENERATED_BODY()
 
@@ -100,7 +100,7 @@ struct RIGVM_API FRigVMDispatch_CoreNotEquals : public FRigVMDispatch_CoreEquals
 		FactoryScriptStruct = StaticStruct();
 	}
 
-	virtual const TArray<FRigVMTemplateArgumentInfo>& GetArgumentInfos() const override;
+	RIGVM_API virtual const TArray<FRigVMTemplateArgumentInfo>& GetArgumentInfos() const override;
 
 	// we are inheriting everything from the equals dispatch,
 	// and due to the check of the factory within FRigVMDispatch_CoreEquals::Execute we can
diff --git a/Engine/Plugins/Runtime/RigVM/Source/RigVM/Public/RigVMFunctions/RigVMDispatch_If.h b/Engine/Plugins/Runtime/RigVM/Source/RigVM/Public/RigVMFunctions/RigVMDispatch_If.h
--- a/Engine/Plugins/Runtime/RigVM/Source/RigVM/Public/RigVMFunctions/RigVMDispatch_If.h
+++ b/Engine/Plugins/Runtime/RigVM/Source/RigVM/Public/RigVMFunctions/RigVMDispatch_If.h
@@ -10,7 +10,7 @@
  * Chooses between two values based on a condition
  */
 USTRUCT(meta=(DisplayName = "If", Category = "Execution", Keywords = "Branch,Condition", NodeColor = "0,1,0,1"))
-struct RIGVM_API FRigVMDispatch_If : public FRigVMDispatch_CoreBase
+struct FRigVMDispatch_If : public FRigVMDispatch_CoreBase
 {
 	GENERATED_BODY()
 
@@ -21,19 +21,19 @@ public:
 		FactoryScriptStruct = StaticStruct();
 	}
 
-	virtual FName GetArgumentNameForOperandIndex(int32 InOperandIndex, int32 InTotalOperands) const override;
-	virtual const TArray<FRigVMTemplateArgumentInfo>& GetArgumentInfos() const override;
-	virtual FRigVMTemplateTypeMap OnNewArgumentType(const FName& InArgumentName, TRigVMTypeIndex InTypeIndex) const override;
+	RIGVM_API virtual FName GetArgumentNameForOperandIndex(int32 InOperandIndex, int32 InTotalOperands) const override;
+	RIGVM_API virtual const TArray<FRigVMTemplateArgumentInfo>& GetArgumentInfos() const override;
+	RIGVM_API virtual FRigVMTemplateTypeMap OnNewArgumentType(const FName& InArgumentName, TRigVMTypeIndex InTypeIndex) const override;
 	virtual bool IsSingleton() const override { return true; } 
 
 #if WITH_EDITOR
-	virtual FString GetArgumentMetaData(const FName& InArgumentName, const FName& InMetaDataKey) const override;
+	RIGVM_API virtual FString GetArgumentMetaData(const FName& InArgumentName, const FName& InMetaDataKey) const override;
 #endif
 
 protected:
 
-	virtual FRigVMFunctionPtr GetDispatchFunctionImpl(const FRigVMTemplateTypeMap& InTypes) const override;
-	static void Execute(FRigVMExtendedExecuteContext& InContext, FRigVMMemoryHandleArray Handles, FRigVMPredicateBranchArray Predicates);
+	RIGVM_API virtual FRigVMFunctionPtr GetDispatchFunctionImpl(const FRigVMTemplateTypeMap& InTypes) const override;
+	static RIGVM_API void Execute(FRigVMExtendedExecuteContext& InContext, FRigVMMemoryHandleArray Handles, FRigVMPredicateBranchArray Predicates);
 
 	static inline const FLazyName ConditionName = FLazyName(TEXT("Condition"));
 	static inline const FLazyName TrueName = FLazyName(TEXT("True"));
diff --git a/Engine/Plugins/Runtime/RigVM/Source/RigVM/Public/RigVMFunctions/RigVMDispatch_MakeStruct.h b/Engine/Plugins/Runtime/RigVM/Source/RigVM/Public/RigVMFunctions/RigVMDispatch_MakeStruct.h
--- a/Engine/Plugins/Runtime/RigVM/Source/RigVM/Public/RigVMFunctions/RigVMDispatch_MakeStruct.h
+++ b/Engine/Plugins/Runtime/RigVM/Source/RigVM/Public/RigVMFunctions/RigVMDispatch_MakeStruct.h
@@ -8,7 +8,7 @@
 #include "RigVMDispatch_MakeStruct.generated.h"
 
 USTRUCT(meta=(DisplayName = "Make", Category = "Core", Keywords = "Compose,Composition,Create,Constant", NodeColor = "1,1,1,1"))
-struct RIGVM_API FRigVMDispatch_MakeStruct : public FRigVMDispatch_CoreBase
+struct FRigVMDispatch_MakeStruct : public FRigVMDispatch_CoreBase
 {
 	GENERATED_BODY()
 
@@ -18,18 +18,18 @@ public:
 		FactoryScriptStruct = StaticStruct();
 	}
 	
-	virtual const TArray<FRigVMTemplateArgumentInfo>& GetArgumentInfos() const override;
-	virtual FRigVMTemplateTypeMap OnNewArgumentType(const FName& InArgumentName, TRigVMTypeIndex InTypeIndex) const override;
+	RIGVM_API virtual const TArray<FRigVMTemplateArgumentInfo>& GetArgumentInfos() const override;
+	RIGVM_API virtual FRigVMTemplateTypeMap OnNewArgumentType(const FName& InArgumentName, TRigVMTypeIndex InTypeIndex) const override;
 #if WITH_EDITOR
-	virtual FString GetNodeTitle(const FRigVMTemplateTypeMap& InTypes) const override;
-	virtual FText GetNodeTooltip(const FRigVMTemplateTypeMap& InTypes) const override;
-	virtual FString GetArgumentMetaData(const FName& InArgumentName, const FName& InMetaDataKey) const override;
-	virtual FString GetKeywords() const override;
+	RIGVM_API virtual FString GetNodeTitle(const FRigVMTemplateTypeMap& InTypes) const override;
+	RIGVM_API virtual FText GetNodeTooltip(const FRigVMTemplateTypeMap& InTypes) const override;
+	RIGVM_API virtual FString GetArgumentMetaData(const FName& InArgumentName, const FName& InMetaDataKey) const override;
+	RIGVM_API virtual FString GetKeywords() const override;
 #endif
 
 protected:
-	virtual FRigVMFunctionPtr GetDispatchFunctionImpl(const FRigVMTemplateTypeMap& InTypes) const override { return &FRigVMDispatch_MakeStruct::Execute; }
-	static void Execute(FRigVMExtendedExecuteContext& InContext, FRigVMMemoryHandleArray Handles, FRigVMPredicateBranchArray RigVMBranches);
+	RIGVM_API virtual FRigVMFunctionPtr GetDispatchFunctionImpl(const FRigVMTemplateTypeMap& InTypes) const override { return &FRigVMDispatch_MakeStruct::Execute; }
+	static RIGVM_API void Execute(FRigVMExtendedExecuteContext& InContext, FRigVMMemoryHandleArray Handles, FRigVMPredicateBranchArray RigVMBranches);
 
 	static inline const FLazyName ElementsName = FLazyName(TEXT("Elements"));
 	static inline const FLazyName StructName = FLazyName(TEXT("Struct"));
@@ -39,7 +39,7 @@ protected:
 };
 
 USTRUCT(meta=(DisplayName = "Break", Keywords = "Decompose,Decomposition"))
-struct RIGVM_API FRigVMDispatch_BreakStruct : public FRigVMDispatch_MakeStruct
+struct FRigVMDispatch_BreakStruct : public FRigVMDispatch_MakeStruct
 {
 	GENERATED_BODY()
 
@@ -49,10 +49,10 @@ public:
 		FactoryScriptStruct = StaticStruct();
 	}
 
-	virtual const TArray<FRigVMTemplateArgumentInfo>& GetArgumentInfos() const override;
+	RIGVM_API virtual const TArray<FRigVMTemplateArgumentInfo>& GetArgumentInfos() const override;
 #if WITH_EDITOR
-	virtual FText GetNodeTooltip(const FRigVMTemplateTypeMap& InTypes) const override;
-	virtual FString GetKeywords() const override;
+	RIGVM_API virtual FText GetNodeTooltip(const FRigVMTemplateTypeMap& InTypes) const override;
+	RIGVM_API virtual FString GetKeywords() const override;
 #endif
 };
 
diff --git a/Engine/Plugins/Runtime/RigVM/Source/RigVM/Public/RigVMFunctions/RigVMDispatch_Print.h b/Engine/Plugins/Runtime/RigVM/Source/RigVM/Public/RigVMFunctions/RigVMDispatch_Print.h
--- a/Engine/Plugins/Runtime/RigVM/Source/RigVM/Public/RigVMFunctions/RigVMDispatch_Print.h
+++ b/Engine/Plugins/Runtime/RigVM/Source/RigVM/Public/RigVMFunctions/RigVMDispatch_Print.h
@@ -9,7 +9,7 @@
  * Prints any value to the log
  */
 USTRUCT(meta=(DisplayName = "Print", NodeColor = "0.8, 0, 0.2, 1"))
-struct RIGVM_API FRigVMDispatch_Print : public FRigVMDispatchFactory
+struct FRigVMDispatch_Print : public FRigVMDispatchFactory
 {
 	GENERATED_BODY()
 
@@ -20,15 +20,15 @@ public:
 		FactoryScriptStruct = StaticStruct();
 	}
 
-	virtual FName GetArgumentNameForOperandIndex(int32 InOperandIndex, int32 InTotalOperands) const override;
-	virtual const TArray<FRigVMTemplateArgumentInfo>& GetArgumentInfos() const override;
-	virtual const TArray<FRigVMExecuteArgument>& GetExecuteArguments_Impl(const FRigVMDispatchContext& InContext) const override;
-	virtual FRigVMTemplateTypeMap OnNewArgumentType(const FName& InArgumentName, TRigVMTypeIndex InTypeIndex) const override;
+	RIGVM_API virtual FName GetArgumentNameForOperandIndex(int32 InOperandIndex, int32 InTotalOperands) const override;
+	RIGVM_API virtual const TArray<FRigVMTemplateArgumentInfo>& GetArgumentInfos() const override;
+	RIGVM_API virtual const TArray<FRigVMExecuteArgument>& GetExecuteArguments_Impl(const FRigVMDispatchContext& InContext) const override;
+	RIGVM_API virtual FRigVMTemplateTypeMap OnNewArgumentType(const FName& InArgumentName, TRigVMTypeIndex InTypeIndex) const override;
 	virtual bool IsSingleton() const override { return true; } 
 
 #if WITH_EDITOR
-	virtual FString GetArgumentDefaultValue(const FName& InArgumentName, TRigVMTypeIndex InTypeIndex) const override;
-	virtual FString GetArgumentMetaData(const FName& InArgumentName, const FName& InMetaDataKey) const override;
+	RIGVM_API virtual FString GetArgumentDefaultValue(const FName& InArgumentName, TRigVMTypeIndex InTypeIndex) const override;
+	RIGVM_API virtual FString GetArgumentMetaData(const FName& InArgumentName, const FName& InMetaDataKey) const override;
 #endif
 
 protected:
@@ -37,7 +37,7 @@ protected:
 	{
 		return &FRigVMDispatch_Print::Execute;
 	}
-	static void Execute(FRigVMExtendedExecuteContext& InContext, FRigVMMemoryHandleArray Handles, FRigVMPredicateBranchArray Predicates);
+	static RIGVM_API void Execute(FRigVMExtendedExecuteContext& InContext, FRigVMMemoryHandleArray Handles, FRigVMPredicateBranchArray Predicates);
 
 	static inline const FLazyName PrefixName = FLazyName(TEXT("Prefix"));
 	static inline const FLazyName ValueName = FLazyName(TEXT("Value"));
diff --git a/Engine/Plugins/Runtime/RigVM/Source/RigVM/Public/RigVMFunctions/RigVMDispatch_Select.h b/Engine/Plugins/Runtime/RigVM/Source/RigVM/Public/RigVMFunctions/RigVMDispatch_Select.h
--- a/Engine/Plugins/Runtime/RigVM/Source/RigVM/Public/RigVMFunctions/RigVMDispatch_Select.h
+++ b/Engine/Plugins/Runtime/RigVM/Source/RigVM/Public/RigVMFunctions/RigVMDispatch_Select.h
@@ -10,7 +10,7 @@
  * Pick from a list of values based on an integer index
  */
 USTRUCT(meta=(DisplayName = "Select", Category = "Execution", Keywords = "Switch,Case", NodeColor = "0,1,0,1"))
-struct RIGVM_API FRigVMDispatch_SelectInt32 : public FRigVMDispatch_CoreBase
+struct FRigVMDispatch_SelectInt32 : public FRigVMDispatch_CoreBase
 {
 	GENERATED_BODY()
 
@@ -21,19 +21,19 @@ public:
 		FactoryScriptStruct = StaticStruct();
 	}
 
-	virtual FName GetArgumentNameForOperandIndex(int32 InOperandIndex, int32 InTotalOperands) const override;
-	virtual const TArray<FRigVMTemplateArgumentInfo>& GetArgumentInfos() const override;
-	virtual FRigVMTemplateTypeMap OnNewArgumentType(const FName& InArgumentName, TRigVMTypeIndex InTypeIndex) const override;
+	RIGVM_API virtual FName GetArgumentNameForOperandIndex(int32 InOperandIndex, int32 InTotalOperands) const override;
+	RIGVM_API virtual const TArray<FRigVMTemplateArgumentInfo>& GetArgumentInfos() const override;
+	RIGVM_API virtual FRigVMTemplateTypeMap OnNewArgumentType(const FName& InArgumentName, TRigVMTypeIndex InTypeIndex) const override;
 #if WITH_EDITOR
-	virtual FString GetArgumentMetaData(const FName& InArgumentName, const FName& InMetaDataKey) const override;
-	virtual FString GetArgumentDefaultValue(const FName& InArgumentName, TRigVMTypeIndex InTypeIndex) const override;
+	RIGVM_API virtual FString GetArgumentMetaData(const FName& InArgumentName, const FName& InMetaDataKey) const override;
+	RIGVM_API virtual FString GetArgumentDefaultValue(const FName& InArgumentName, TRigVMTypeIndex InTypeIndex) const override;
 #endif
 	virtual bool IsSingleton() const override { return true; } 
 
 protected:
 
-	virtual FRigVMFunctionPtr GetDispatchFunctionImpl(const FRigVMTemplateTypeMap& InTypes) const override;
-	static void Execute(FRigVMExtendedExecuteContext& InContext, FRigVMMemoryHandleArray Handles, FRigVMPredicateBranchArray Predicates);
+	RIGVM_API virtual FRigVMFunctionPtr GetDispatchFunctionImpl(const FRigVMTemplateTypeMap& InTypes) const override;
+	static RIGVM_API void Execute(FRigVMExtendedExecuteContext& InContext, FRigVMMemoryHandleArray Handles, FRigVMPredicateBranchArray Predicates);
 
 	static inline FLazyName IndexName = FLazyName(TEXT("Index"));
 	static inline FLazyName ValuesName = FLazyName(TEXT("Values"));
diff --git a/Engine/Plugins/Runtime/RigVM/Source/RigVM/Public/RigVMFunctions/RigVMDispatch_Switch.h b/Engine/Plugins/Runtime/RigVM/Source/RigVM/Public/RigVMFunctions/RigVMDispatch_Switch.h
--- a/Engine/Plugins/Runtime/RigVM/Source/RigVM/Public/RigVMFunctions/RigVMDispatch_Switch.h
+++ b/Engine/Plugins/Runtime/RigVM/Source/RigVM/Public/RigVMFunctions/RigVMDispatch_Switch.h
@@ -10,7 +10,7 @@
  * Run a branch based on an integer index
  */
 USTRUCT(meta=(DisplayName = "Switch", Category = "Execution", Keywords = "Case", NodeColor = "0,1,0,1"))
-struct RIGVM_API FRigVMDispatch_SwitchInt32 : public FRigVMDispatch_CoreBase
+struct FRigVMDispatch_SwitchInt32 : public FRigVMDispatch_CoreBase
 {
 	GENERATED_BODY()
 
@@ -21,25 +21,25 @@ public:
 		FactoryScriptStruct = StaticStruct();
 	}
 
-	virtual FName GetArgumentNameForOperandIndex(int32 InOperandIndex, int32 InTotalOperands) const override;
-	virtual const TArray<FRigVMTemplateArgumentInfo>& GetArgumentInfos() const override;
-	virtual const TArray<FRigVMExecuteArgument>& GetExecuteArguments_Impl(const FRigVMDispatchContext& InContext) const override;
+	RIGVM_API virtual FName GetArgumentNameForOperandIndex(int32 InOperandIndex, int32 InTotalOperands) const override;
+	RIGVM_API virtual const TArray<FRigVMTemplateArgumentInfo>& GetArgumentInfos() const override;
+	RIGVM_API virtual const TArray<FRigVMExecuteArgument>& GetExecuteArguments_Impl(const FRigVMDispatchContext& InContext) const override;
 	virtual bool IsSingleton() const override { return true; } 
 
 #if WITH_EDITOR
-	virtual FString GetArgumentMetaData(const FName& InArgumentName, const FName& InMetaDataKey) const override;
-	virtual FString GetArgumentDefaultValue(const FName& InArgumentName, TRigVMTypeIndex InTypeIndex) const override;
-	virtual FName GetDisplayNameForArgument(const FName& InArgumentName) const override;
+	RIGVM_API virtual FString GetArgumentMetaData(const FName& InArgumentName, const FName& InMetaDataKey) const override;
+	RIGVM_API virtual FString GetArgumentDefaultValue(const FName& InArgumentName, TRigVMTypeIndex InTypeIndex) const override;
+	RIGVM_API virtual FName GetDisplayNameForArgument(const FName& InArgumentName) const override;
 #endif
 	
-	virtual const TArray<FName>& GetControlFlowBlocks_Impl(const FRigVMDispatchContext& InContext) const override;
+	RIGVM_API virtual const TArray<FName>& GetControlFlowBlocks_Impl(const FRigVMDispatchContext& InContext) const override;
 
 protected:
 
-	virtual FRigVMFunctionPtr GetDispatchFunctionImpl(const FRigVMTemplateTypeMap& InTypes) const override;
-	static void Execute(FRigVMExtendedExecuteContext& InContext, FRigVMMemoryHandleArray Handles, FRigVMPredicateBranchArray Predicates);
-	static FName GetCaseName(int32 InIndex);
-	static FName GetCaseDisplayName(int32 InIndex);
+	RIGVM_API virtual FRigVMFunctionPtr GetDispatchFunctionImpl(const FRigVMTemplateTypeMap& InTypes) const override;
+	static RIGVM_API void Execute(FRigVMExtendedExecuteContext& InContext, FRigVMMemoryHandleArray Handles, FRigVMPredicateBranchArray Predicates);
+	static RIGVM_API FName GetCaseName(int32 InIndex);
+	static RIGVM_API FName GetCaseDisplayName(int32 InIndex);
 
 	static inline const FLazyName IndexName = FLazyName(TEXT("Index"));
 	static inline const FLazyName CasesName = FLazyName(TEXT("Cases"));
diff --git a/Engine/Plugins/Runtime/RigVM/Source/RigVM/Public/RigVMFunctions/RigVMFunction_ControlFlow.h b/Engine/Plugins/Runtime/RigVM/Source/RigVM/Public/RigVMFunctions/RigVMFunction_ControlFlow.h
--- a/Engine/Plugins/Runtime/RigVM/Source/RigVM/Public/RigVMFunctions/RigVMFunction_ControlFlow.h
+++ b/Engine/Plugins/Runtime/RigVM/Source/RigVM/Public/RigVMFunctions/RigVMFunction_ControlFlow.h
@@ -6,7 +6,7 @@
 #include "RigVMFunction_ControlFlow.generated.h"
 
 USTRUCT(meta=(Abstract, Category="Execution", NodeColor = "0, 0, 0, 1"))
-struct RIGVM_API FRigVMFunction_ControlFlowBase : public FRigVMStruct
+struct FRigVMFunction_ControlFlowBase : public FRigVMStruct
 {
 	GENERATED_BODY()
 };
@@ -15,7 +15,7 @@ struct RIGVM_API FRigVMFunction_ControlFlowBase : public FRigVMStruct
  * Executes either the True or False branch based on the condition
  */
 USTRUCT(meta = (DisplayName = "Branch", Keywords = "If"))
-struct RIGVM_API FRigVMFunction_ControlFlowBranch : public FRigVMFunction_ControlFlowBase
+struct FRigVMFunction_ControlFlowBranch : public FRigVMFunction_ControlFlowBase
 {
 	GENERATED_BODY()
 
@@ -26,7 +26,7 @@ struct RIGVM_API FRigVMFunction_ControlFlowBranch : public FRigVMFunction_Contro
 	}
 
 	RIGVM_METHOD()
-	virtual void Execute();
+	RIGVM_API virtual void Execute();
 
 	virtual const TArray<FName>& GetControlFlowBlocks_Impl() const override
 	{
diff --git a/Engine/Plugins/Runtime/RigVM/Source/RigVM/Public/RigVMFunctions/RigVMFunction_Name.h b/Engine/Plugins/Runtime/RigVM/Source/RigVM/Public/RigVMFunctions/RigVMFunction_Name.h
--- a/Engine/Plugins/Runtime/RigVM/Source/RigVM/Public/RigVMFunctions/RigVMFunction_Name.h
+++ b/Engine/Plugins/Runtime/RigVM/Source/RigVM/Public/RigVMFunctions/RigVMFunction_Name.h
@@ -6,7 +6,7 @@
 #include "RigVMFunction_Name.generated.h"
 
 USTRUCT(meta=(Abstract, Category="Core|Name", NodeColor = "0.462745, 1,0, 0.329412"))
-struct RIGVM_API FRigVMFunction_NameBase : public FRigVMStruct
+struct FRigVMFunction_NameBase : public FRigVMStruct
 {
 	GENERATED_BODY()
 
@@ -17,7 +17,7 @@ struct RIGVM_API FRigVMFunction_NameBase : public FRigVMStruct
  * Concatenates two strings together to make a new string
  */
 USTRUCT(meta = (DisplayName = "Concat", TemplateName = "Concat", Keywords = "Add,+,Combine,Merge,Append"))
-struct RIGVM_API FRigVMFunction_NameConcat : public FRigVMFunction_NameBase
+struct FRigVMFunction_NameConcat : public FRigVMFunction_NameBase
 {
 	GENERATED_BODY()
 
@@ -27,7 +27,7 @@ struct RIGVM_API FRigVMFunction_NameConcat : public FRigVMFunction_NameBase
 	}
 
 	RIGVM_METHOD()
-	virtual void Execute() override;
+	RIGVM_API virtual void Execute() override;
 
 	UPROPERTY(meta=(Input, Aggregate))
 	FName A;
@@ -43,7 +43,7 @@ struct RIGVM_API FRigVMFunction_NameConcat : public FRigVMFunction_NameBase
  * Returns the left or right most characters from the string chopping the given number of characters from the start or the end
  */
 USTRUCT(meta = (DisplayName = "Chop", TemplateName = "Chop", Keywords = "Truncate,-,Remove,Subtract,Split"))
-struct RIGVM_API FRigVMFunction_NameTruncate : public FRigVMFunction_NameBase
+struct FRigVMFunction_NameTruncate : public FRigVMFunction_NameBase
 {
 	GENERATED_BODY()
 
@@ -55,7 +55,7 @@ struct RIGVM_API FRigVMFunction_NameTruncate : public FRigVMFunction_NameBase
 	}
 
 	RIGVM_METHOD()
-	virtual void Execute() override;
+	RIGVM_API virtual void Execute() override;
 
 	UPROPERTY(meta=(Input))
 	FName Name;
@@ -81,7 +81,7 @@ struct RIGVM_API FRigVMFunction_NameTruncate : public FRigVMFunction_NameBase
  * Replace all occurrences of a substring in this string
  */
 USTRUCT(meta = (DisplayName = "Replace", TemplateName = "Replace", Keywords = "Search,Emplace,Find"))
-struct RIGVM_API FRigVMFunction_NameReplace : public FRigVMFunction_NameBase
+struct FRigVMFunction_NameReplace : public FRigVMFunction_NameBase
 {
 	GENERATED_BODY()
 
@@ -91,7 +91,7 @@ struct RIGVM_API FRigVMFunction_NameReplace : public FRigVMFunction_NameBase
 	}
 
 	RIGVM_METHOD()
-	virtual void Execute() override;
+	RIGVM_API virtual void Execute() override;
 
 	UPROPERTY(meta = (Input))
 	FName Name;
@@ -110,7 +110,7 @@ struct RIGVM_API FRigVMFunction_NameReplace : public FRigVMFunction_NameBase
  * Tests whether this string ends with given string
  */
 USTRUCT(meta = (DisplayName = "Ends With", TemplateName = "EndsWith", Keywords = "Right"))
-struct RIGVM_API FRigVMFunction_EndsWith : public FRigVMFunction_NameBase
+struct FRigVMFunction_EndsWith : public FRigVMFunction_NameBase
 {
 	GENERATED_BODY()
 
@@ -121,7 +121,7 @@ struct RIGVM_API FRigVMFunction_EndsWith : public FRigVMFunction_NameBase
 	}
 
 	RIGVM_METHOD()
-		virtual void Execute() override;
+		RIGVM_API virtual void Execute() override;
 
 	UPROPERTY(meta = (Input))
 	FName Name;
@@ -137,7 +137,7 @@ struct RIGVM_API FRigVMFunction_EndsWith : public FRigVMFunction_NameBase
  * Tests whether this string starts with given string
  */
 USTRUCT(meta = (DisplayName = "Starts With", TemplateName = "StartsWith", Keywords = "Left"))
-struct RIGVM_API FRigVMFunction_StartsWith : public FRigVMFunction_NameBase
+struct FRigVMFunction_StartsWith : public FRigVMFunction_NameBase
 {
 	GENERATED_BODY()
 
@@ -148,7 +148,7 @@ struct RIGVM_API FRigVMFunction_StartsWith : public FRigVMFunction_NameBase
 	}
 
 	RIGVM_METHOD()
-	virtual void Execute() override;
+	RIGVM_API virtual void Execute() override;
 
 	UPROPERTY(meta = (Input))
 	FName Name;
@@ -164,7 +164,7 @@ struct RIGVM_API FRigVMFunction_StartsWith : public FRigVMFunction_NameBase
  * Returns true or false if a given name exists in another given name
  */
 USTRUCT(meta = (DisplayName = "Contains", TemplateName = "Contains", Keywords = "Contains,Find,Has,Search"))
-struct RIGVM_API FRigVMFunction_Contains : public FRigVMFunction_NameBase
+struct FRigVMFunction_Contains : public FRigVMFunction_NameBase
 {
 	GENERATED_BODY()
 
@@ -175,7 +175,7 @@ struct RIGVM_API FRigVMFunction_Contains : public FRigVMFunction_NameBase
 	}
 
 	RIGVM_METHOD()
-	virtual void Execute() override;
+	RIGVM_API virtual void Execute() override;
 
 	UPROPERTY(meta = (Input))
 	FName Name;
diff --git a/Engine/Plugins/Runtime/RigVM/Source/RigVM/Public/RigVMFunctions/RigVMFunction_String.h b/Engine/Plugins/Runtime/RigVM/Source/RigVM/Public/RigVMFunctions/RigVMFunction_String.h
--- a/Engine/Plugins/Runtime/RigVM/Source/RigVM/Public/RigVMFunctions/RigVMFunction_String.h
+++ b/Engine/Plugins/Runtime/RigVM/Source/RigVM/Public/RigVMFunctions/RigVMFunction_String.h
@@ -7,7 +7,7 @@
 #include "RigVMFunction_String.generated.h"
 
 USTRUCT(meta=(Abstract, Category="Core|String", NodeColor = "0.462745, 1,0, 0.329412"))
-struct RIGVM_API FRigVMFunction_StringBase : public FRigVMStruct
+struct FRigVMFunction_StringBase : public FRigVMStruct
 {
 	GENERATED_BODY()
 };
@@ -16,7 +16,7 @@ struct RIGVM_API FRigVMFunction_StringBase : public FRigVMStruct
  * Concatenates two strings together to make a new string
  */
 USTRUCT(meta = (DisplayName = "Concat", TemplateName = "Concat", Keywords = "Add,+,Combine,Merge,Append"))
-struct RIGVM_API FRigVMFunction_StringConcat : public FRigVMFunction_StringBase
+struct FRigVMFunction_StringConcat : public FRigVMFunction_StringBase
 {
 	GENERATED_BODY()
 
@@ -26,7 +26,7 @@ struct RIGVM_API FRigVMFunction_StringConcat : public FRigVMFunction_StringBase
 	}
 
 	RIGVM_METHOD()
-	virtual void Execute();
+	RIGVM_API virtual void Execute();
 
 	UPROPERTY(meta=(Input, Aggregate))
 	FString A;
@@ -42,7 +42,7 @@ struct RIGVM_API FRigVMFunction_StringConcat : public FRigVMFunction_StringBase
  * Returns the left or right most characters from the string chopping the given number of characters from the start or the end
  */
 USTRUCT(meta = (DisplayName = "Chop", TemplateName = "Chop", Keywords = "Truncate,-,Remove,Subtract,Split"))
-struct RIGVM_API FRigVMFunction_StringTruncate : public FRigVMFunction_StringBase
+struct FRigVMFunction_StringTruncate : public FRigVMFunction_StringBase
 {
 	GENERATED_BODY()
 
@@ -54,7 +54,7 @@ struct RIGVM_API FRigVMFunction_StringTruncate : public FRigVMFunction_StringBas
 	}
 
 	RIGVM_METHOD()
-	virtual void Execute();
+	RIGVM_API virtual void Execute();
 
 	UPROPERTY(meta=(Input))
 	FString Name;
@@ -80,7 +80,7 @@ struct RIGVM_API FRigVMFunction_StringTruncate : public FRigVMFunction_StringBas
  * Replace all occurrences of a substring in this string
  */
 USTRUCT(meta = (DisplayName = "Replace", TemplateName = "Replace", Keywords = "Search,Emplace,Find"))
-struct RIGVM_API FRigVMFunction_StringReplace : public FRigVMFunction_StringBase
+struct FRigVMFunction_StringReplace : public FRigVMFunction_StringBase
 {
 	GENERATED_BODY()
 
@@ -90,7 +90,7 @@ struct RIGVM_API FRigVMFunction_StringReplace : public FRigVMFunction_StringBase
 	}
 
 	RIGVM_METHOD()
-	virtual void Execute();
+	RIGVM_API virtual void Execute();
 
 	UPROPERTY(meta = (Input))
 	FString Name;
@@ -109,7 +109,7 @@ struct RIGVM_API FRigVMFunction_StringReplace : public FRigVMFunction_StringBase
  * Tests whether this string ends with given string
  */
 USTRUCT(meta = (DisplayName = "Ends With", TemplateName = "EndsWith", Keywords = "Right"))
-struct RIGVM_API FRigVMFunction_StringEndsWith : public FRigVMFunction_StringBase
+struct FRigVMFunction_StringEndsWith : public FRigVMFunction_StringBase
 {
 	GENERATED_BODY()
 
@@ -120,7 +120,7 @@ struct RIGVM_API FRigVMFunction_StringEndsWith : public FRigVMFunction_StringBas
 	}
 
 	RIGVM_METHOD()
-	virtual void Execute();
+	RIGVM_API virtual void Execute();
 
 	UPROPERTY(meta = (Input))
 	FString Name;
@@ -136,7 +136,7 @@ struct RIGVM_API FRigVMFunction_StringEndsWith : public FRigVMFunction_StringBas
  * Tests whether this string starts with given string
  */
 USTRUCT(meta = (DisplayName = "Starts With", TemplateName = "StartsWith", Keywords = "Left"))
-struct RIGVM_API FRigVMFunction_StringStartsWith : public FRigVMFunction_StringBase
+struct FRigVMFunction_StringStartsWith : public FRigVMFunction_StringBase
 {
 	GENERATED_BODY()
 
@@ -147,7 +147,7 @@ struct RIGVM_API FRigVMFunction_StringStartsWith : public FRigVMFunction_StringB
 	}
 
 	RIGVM_METHOD()
-	virtual void Execute();
+	RIGVM_API virtual void Execute();
 
 	UPROPERTY(meta = (Input))
 	FString Name;
@@ -163,7 +163,7 @@ struct RIGVM_API FRigVMFunction_StringStartsWith : public FRigVMFunction_StringB
  * Returns true or false if a given name exists in another given name
  */
 USTRUCT(meta = (DisplayName = "Contains", TemplateName = "Contains", Keywords = "Contains,Find,Has,Search"))
-struct RIGVM_API FRigVMFunction_StringContains : public FRigVMFunction_StringBase
+struct FRigVMFunction_StringContains : public FRigVMFunction_StringBase
 {
 	GENERATED_BODY()
 
@@ -174,7 +174,7 @@ struct RIGVM_API FRigVMFunction_StringContains : public FRigVMFunction_StringBas
 	}
 
 	RIGVM_METHOD()
-	virtual void Execute();
+	RIGVM_API virtual void Execute();
 
 	UPROPERTY(meta = (Input))
 	FString Name;
@@ -190,7 +190,7 @@ struct RIGVM_API FRigVMFunction_StringContains : public FRigVMFunction_StringBas
  * Returns the length of a string 
  */
 USTRUCT(meta = (DisplayName = "Length"))
-struct RIGVM_API FRigVMFunction_StringLength : public FRigVMFunction_StringBase
+struct FRigVMFunction_StringLength : public FRigVMFunction_StringBase
 {
 	GENERATED_BODY()
 
@@ -201,7 +201,7 @@ struct RIGVM_API FRigVMFunction_StringLength : public FRigVMFunction_StringBase
 	}
 
 	RIGVM_METHOD()
-	virtual void Execute();
+	RIGVM_API virtual void Execute();
 
 	UPROPERTY(meta = (Input))
 	FString Value;
@@ -214,7 +214,7 @@ struct RIGVM_API FRigVMFunction_StringLength : public FRigVMFunction_StringBase
  * Trims the whitespace from a string (start and end)
  */
 USTRUCT(meta = (DisplayName = "Trim Whitespace", Keywords = "Space,WhiteSpace,Remove,Truncate"))
-struct RIGVM_API FRigVMFunction_StringTrimWhitespace : public FRigVMFunction_StringBase
+struct FRigVMFunction_StringTrimWhitespace : public FRigVMFunction_StringBase
 {
 	GENERATED_BODY()
 
@@ -224,7 +224,7 @@ struct RIGVM_API FRigVMFunction_StringTrimWhitespace : public FRigVMFunction_Str
 	}
 
 	RIGVM_METHOD()
-	virtual void Execute();
+	RIGVM_API virtual void Execute();
 
 	UPROPERTY(meta = (Input))
 	FString Value;
@@ -237,7 +237,7 @@ struct RIGVM_API FRigVMFunction_StringTrimWhitespace : public FRigVMFunction_Str
  * Converts the string to upper case
  */
 USTRUCT(meta = (DisplayName = "To Uppercase"))
-struct RIGVM_API FRigVMFunction_StringToUppercase : public FRigVMFunction_StringBase
+struct FRigVMFunction_StringToUppercase : public FRigVMFunction_StringBase
 {
 	GENERATED_BODY()
 
@@ -247,7 +247,7 @@ struct RIGVM_API FRigVMFunction_StringToUppercase : public FRigVMFunction_String
 	}
 
 	RIGVM_METHOD()
-	virtual void Execute();
+	RIGVM_API virtual void Execute();
 
 	UPROPERTY(meta = (Input))
 	FString Value;
@@ -260,7 +260,7 @@ struct RIGVM_API FRigVMFunction_StringToUppercase : public FRigVMFunction_String
  * Converts the string to lower case
  */
 USTRUCT(meta = (DisplayName = "To Lowercase"))
-struct RIGVM_API FRigVMFunction_StringToLowercase : public FRigVMFunction_StringBase
+struct FRigVMFunction_StringToLowercase : public FRigVMFunction_StringBase
 {
 	GENERATED_BODY()
 
@@ -270,7 +270,7 @@ struct RIGVM_API FRigVMFunction_StringToLowercase : public FRigVMFunction_String
 	}
 
 	RIGVM_METHOD()
-	virtual void Execute();
+	RIGVM_API virtual void Execute();
 
 	UPROPERTY(meta = (Input))
 	FString Value;
@@ -283,7 +283,7 @@ struct RIGVM_API FRigVMFunction_StringToLowercase : public FRigVMFunction_String
  * Returns the reverse of the input string
  */
 USTRUCT(meta = (DisplayName = "Reverse"))
-struct RIGVM_API FRigVMFunction_StringReverse : public FRigVMFunction_StringBase
+struct FRigVMFunction_StringReverse : public FRigVMFunction_StringBase
 {
 	GENERATED_BODY()
 
@@ -293,7 +293,7 @@ struct RIGVM_API FRigVMFunction_StringReverse : public FRigVMFunction_StringBase
 	}
 
 	RIGVM_METHOD()
-	virtual void Execute();
+	RIGVM_API virtual void Execute();
 
 	UPROPERTY(meta = (Input))
 	FString Value;
@@ -306,7 +306,7 @@ struct RIGVM_API FRigVMFunction_StringReverse : public FRigVMFunction_StringBase
  * Returns the left most characters of a string
  */
 USTRUCT(meta = (DisplayName = "Left", Keywords = "Start,Begin"))
-struct RIGVM_API FRigVMFunction_StringLeft : public FRigVMFunction_StringBase
+struct FRigVMFunction_StringLeft : public FRigVMFunction_StringBase
 {
 	GENERATED_BODY()
 
@@ -317,7 +317,7 @@ struct RIGVM_API FRigVMFunction_StringLeft : public FRigVMFunction_StringBase
 	}
 
 	RIGVM_METHOD()
-	virtual void Execute();
+	RIGVM_API virtual void Execute();
 
 	UPROPERTY(meta = (Input))
 	FString Value;
@@ -333,7 +333,7 @@ struct RIGVM_API FRigVMFunction_StringLeft : public FRigVMFunction_StringBase
  * Returns the right most characters of a string
  */
 USTRUCT(meta = (DisplayName = "Right", Keywords = "End"))
-struct RIGVM_API FRigVMFunction_StringRight : public FRigVMFunction_StringBase
+struct FRigVMFunction_StringRight : public FRigVMFunction_StringBase
 {
 	GENERATED_BODY()
 
@@ -344,7 +344,7 @@ struct RIGVM_API FRigVMFunction_StringRight : public FRigVMFunction_StringBase
 	}
 
 	RIGVM_METHOD()
-	virtual void Execute();
+	RIGVM_API virtual void Execute();
 
 	UPROPERTY(meta = (Input))
 	FString Value;
@@ -360,7 +360,7 @@ struct RIGVM_API FRigVMFunction_StringRight : public FRigVMFunction_StringBase
  * Returns the middle section of a string
  */
 USTRUCT(meta = (DisplayName = "Middle", Keywords = "Within,Center"))
-struct RIGVM_API FRigVMFunction_StringMiddle : public FRigVMFunction_StringBase
+struct FRigVMFunction_StringMiddle : public FRigVMFunction_StringBase
 {
 	GENERATED_BODY()
 
@@ -372,7 +372,7 @@ struct RIGVM_API FRigVMFunction_StringMiddle : public FRigVMFunction_StringBase
 	}
 
 	RIGVM_METHOD()
-	virtual void Execute();
+	RIGVM_API virtual void Execute();
 
 	UPROPERTY(meta = (Input))
 	FString Value;
@@ -393,7 +393,7 @@ struct RIGVM_API FRigVMFunction_StringMiddle : public FRigVMFunction_StringBase
  * Finds a string within another string
  */
 USTRUCT(meta = (DisplayName = "Find", Keywords = "IndexOf"))
-struct RIGVM_API FRigVMFunction_StringFind : public FRigVMFunction_StringBase
+struct FRigVMFunction_StringFind : public FRigVMFunction_StringBase
 {
 	GENERATED_BODY()
 
@@ -405,7 +405,7 @@ struct RIGVM_API FRigVMFunction_StringFind : public FRigVMFunction_StringBase
 	}
 
 	RIGVM_METHOD()
-	virtual void Execute();
+	RIGVM_API virtual void Execute();
 
 	UPROPERTY(meta = (Input))
 	FString Value;
@@ -424,7 +424,7 @@ struct RIGVM_API FRigVMFunction_StringFind : public FRigVMFunction_StringBase
  * Splits a string into multiple sections given a separator
  */
 USTRUCT(meta = (DisplayName = "Split"))
-struct RIGVM_API FRigVMFunction_StringSplit : public FRigVMFunction_StringBase
+struct FRigVMFunction_StringSplit : public FRigVMFunction_StringBase
 {
 	GENERATED_BODY()
 
@@ -434,7 +434,7 @@ struct RIGVM_API FRigVMFunction_StringSplit : public FRigVMFunction_StringBase
 	}
 
 	RIGVM_METHOD()
-	virtual void Execute();
+	RIGVM_API virtual void Execute();
 
 	UPROPERTY(meta = (Input))
 	FString Value;
@@ -450,7 +450,7 @@ struct RIGVM_API FRigVMFunction_StringSplit : public FRigVMFunction_StringBase
  * Joins a string into multiple sections given a separator
  */
 USTRUCT(meta = (DisplayName = "Join", Keywords = "Combine"))
-struct RIGVM_API FRigVMFunction_StringJoin : public FRigVMFunction_StringBase
+struct FRigVMFunction_StringJoin : public FRigVMFunction_StringBase
 {
 	GENERATED_BODY()
 
@@ -460,7 +460,7 @@ struct RIGVM_API FRigVMFunction_StringJoin : public FRigVMFunction_StringBase
 	}
 
 	RIGVM_METHOD()
-	virtual void Execute();
+	RIGVM_API virtual void Execute();
 
 	UPROPERTY(meta = (Input))
 	TArray<FString> Values;
@@ -476,7 +476,7 @@ struct RIGVM_API FRigVMFunction_StringJoin : public FRigVMFunction_StringBase
  * Converts an integer number to a string with padding
  */
 USTRUCT(meta = (DisplayName = "Pad Integer", Keywords = "FromInt,Number,LeadingZeroes"))
-struct RIGVM_API FRigVMFunction_StringPadInteger : public FRigVMFunction_StringBase
+struct FRigVMFunction_StringPadInteger : public FRigVMFunction_StringBase
 {
 	GENERATED_BODY()
 
@@ -488,7 +488,7 @@ struct RIGVM_API FRigVMFunction_StringPadInteger : public FRigVMFunction_StringB
 	}
 
 	RIGVM_METHOD()
-	virtual void Execute();
+	RIGVM_API virtual void Execute();
 
 	UPROPERTY(meta = (Input))
 	int32 Value;
diff --git a/Engine/Plugins/Runtime/RigVM/Source/RigVM/Public/RigVMFunctions/Simulation/RigVMFunction_Accumulate.h b/Engine/Plugins/Runtime/RigVM/Source/RigVM/Public/RigVMFunctions/Simulation/RigVMFunction_Accumulate.h
--- a/Engine/Plugins/Runtime/RigVM/Source/RigVM/Public/RigVMFunctions/Simulation/RigVMFunction_Accumulate.h
+++ b/Engine/Plugins/Runtime/RigVM/Source/RigVM/Public/RigVMFunctions/Simulation/RigVMFunction_Accumulate.h
@@ -6,7 +6,7 @@
 #include "RigVMFunction_Accumulate.generated.h"
 
 USTRUCT(meta=(Category="Simulation|Accumulate"))
-struct RIGVM_API FRigVMFunction_AccumulateBase : public FRigVMFunction_SimBase
+struct FRigVMFunction_AccumulateBase : public FRigVMFunction_SimBase
 {
 	GENERATED_BODY()
 };
@@ -15,7 +15,7 @@ struct RIGVM_API FRigVMFunction_AccumulateBase : public FRigVMFunction_SimBase
  * Adds a value over time over and over again
  */
 USTRUCT(meta=(DisplayName="Accumulate Add (Float)", TemplateName="AccumulateAdd", Keywords="Simulate,++"))
-struct RIGVM_API FRigVMFunction_AccumulateFloatAdd : public FRigVMFunction_AccumulateBase
+struct FRigVMFunction_AccumulateFloatAdd : public FRigVMFunction_AccumulateBase
 {
 	GENERATED_BODY()
 	
@@ -29,7 +29,7 @@ struct RIGVM_API FRigVMFunction_AccumulateFloatAdd : public FRigVMFunction_Accum
 	virtual void Initialize() override { bIsInitialized = false; }
 
 	RIGVM_METHOD()
-	virtual void Execute() override;
+	RIGVM_API virtual void Execute() override;
 
 	UPROPERTY(meta=(Input))
 	float Increment;
@@ -54,7 +54,7 @@ struct RIGVM_API FRigVMFunction_AccumulateFloatAdd : public FRigVMFunction_Accum
  * Adds a vector over time over and over again
  */
 USTRUCT(meta = (DisplayName="Accumulate Add (Vector)", TemplateName="AccumulateAdd", Keywords="Simulate,++"))
-struct RIGVM_API FRigVMFunction_AccumulateVectorAdd : public FRigVMFunction_AccumulateBase
+struct FRigVMFunction_AccumulateVectorAdd : public FRigVMFunction_AccumulateBase
 {
 	GENERATED_BODY()
 
@@ -68,7 +68,7 @@ struct RIGVM_API FRigVMFunction_AccumulateVectorAdd : public FRigVMFunction_Accu
 	virtual void Initialize() override { bIsInitialized = false; }
 
 	RIGVM_METHOD()
-	virtual void Execute() override;
+	RIGVM_API virtual void Execute() override;
 
 	UPROPERTY(meta = (Input))
 	FVector Increment;
@@ -93,7 +93,7 @@ struct RIGVM_API FRigVMFunction_AccumulateVectorAdd : public FRigVMFunction_Accu
  * Multiplies a value over time over and over again
  */
 USTRUCT(meta=(DisplayName="Accumulate Mul (Float)", TemplateName="AccumulateMul", Keywords="Simulate,**"))
-struct RIGVM_API FRigVMFunction_AccumulateFloatMul : public FRigVMFunction_AccumulateBase
+struct FRigVMFunction_AccumulateFloatMul : public FRigVMFunction_AccumulateBase
 {
 	GENERATED_BODY()
 	
@@ -107,7 +107,7 @@ struct RIGVM_API FRigVMFunction_AccumulateFloatMul : public FRigVMFunction_Accum
 	virtual void Initialize() override { bIsInitialized = false; }
 
 	RIGVM_METHOD()
-	virtual void Execute() override;
+	RIGVM_API virtual void Execute() override;
 
 	UPROPERTY(meta=(Input))
 	float Multiplier;
@@ -132,7 +132,7 @@ struct RIGVM_API FRigVMFunction_AccumulateFloatMul : public FRigVMFunction_Accum
  * Multiplies a vector over time over and over again
  */
 USTRUCT(meta = (DisplayName="Accumulate Mul (Vector)", TemplateName="AccumulateMul", Keywords="Simulate,**"))
-struct RIGVM_API FRigVMFunction_AccumulateVectorMul : public FRigVMFunction_AccumulateBase
+struct FRigVMFunction_AccumulateVectorMul : public FRigVMFunction_AccumulateBase
 {
 	GENERATED_BODY()
 
@@ -146,7 +146,7 @@ struct RIGVM_API FRigVMFunction_AccumulateVectorMul : public FRigVMFunction_Accu
 	virtual void Initialize() override { bIsInitialized = false; }
 
 	RIGVM_METHOD()
-	virtual void Execute() override;
+	RIGVM_API virtual void Execute() override;
 
 	UPROPERTY(meta = (Input))
 	FVector Multiplier;
@@ -171,7 +171,7 @@ struct RIGVM_API FRigVMFunction_AccumulateVectorMul : public FRigVMFunction_Accu
  * Multiplies a quaternion over time over and over again
  */
 USTRUCT(meta=(DisplayName="Accumulate Mul (Quaternion)", TemplateName="AccumulateMul", Keywords="Simulate,**"))
-struct RIGVM_API FRigVMFunction_AccumulateQuatMul : public FRigVMFunction_AccumulateBase
+struct FRigVMFunction_AccumulateQuatMul : public FRigVMFunction_AccumulateBase
 {
 	GENERATED_BODY()
 	
@@ -185,7 +185,7 @@ struct RIGVM_API FRigVMFunction_AccumulateQuatMul : public FRigVMFunction_Accumu
 	virtual void Initialize() override { bIsInitialized = false; }
 
 	RIGVM_METHOD()
-	virtual void Execute() override;
+	RIGVM_API virtual void Execute() override;
 
 	UPROPERTY(meta=(Input))
 	FQuat Multiplier;
@@ -213,7 +213,7 @@ struct RIGVM_API FRigVMFunction_AccumulateQuatMul : public FRigVMFunction_Accumu
  * Multiplies a transform over time over and over again
  */
 USTRUCT(meta = (DisplayName="Accumulate Mul (Transform)", TemplateName="AccumulateMul", Keywords="Simulate,**"))
-struct RIGVM_API FRigVMFunction_AccumulateTransformMul : public FRigVMFunction_AccumulateBase
+struct FRigVMFunction_AccumulateTransformMul : public FRigVMFunction_AccumulateBase
 {
 	GENERATED_BODY()
 
@@ -227,7 +227,7 @@ struct RIGVM_API FRigVMFunction_AccumulateTransformMul : public FRigVMFunction_A
 	virtual void Initialize() override { bIsInitialized = false; }
 
 	RIGVM_METHOD()
-	virtual void Execute() override;
+	RIGVM_API virtual void Execute() override;
 
 	UPROPERTY(meta = (Input))
 	FTransform Multiplier;
@@ -255,7 +255,7 @@ struct RIGVM_API FRigVMFunction_AccumulateTransformMul : public FRigVMFunction_A
  * Interpolates two values over time over and over again
  */
 USTRUCT(meta=(DisplayName="Accumulate Lerp (Float)", TemplateName="AccumulateLerp", Keywords="Simulate,Ramp"))
-struct RIGVM_API FRigVMFunction_AccumulateFloatLerp : public FRigVMFunction_AccumulateBase
+struct FRigVMFunction_AccumulateFloatLerp : public FRigVMFunction_AccumulateBase
 {
 	GENERATED_BODY()
 	
@@ -269,7 +269,7 @@ struct RIGVM_API FRigVMFunction_AccumulateFloatLerp : public FRigVMFunction_Accu
 	virtual void Initialize() override { bIsInitialized = false; }
 
 	RIGVM_METHOD()
-	virtual void Execute() override;
+	RIGVM_API virtual void Execute() override;
 
 	UPROPERTY(meta=(Input))
 	float TargetValue;
@@ -297,7 +297,7 @@ struct RIGVM_API FRigVMFunction_AccumulateFloatLerp : public FRigVMFunction_Accu
  * Interpolates two vectors over time over and over again
  */
 USTRUCT(meta = (DisplayName="Accumulate Lerp (Vector)", TemplateName="AccumulateLerp", Keywords="Simulate,Ramp"))
-struct RIGVM_API FRigVMFunction_AccumulateVectorLerp : public FRigVMFunction_AccumulateBase
+struct FRigVMFunction_AccumulateVectorLerp : public FRigVMFunction_AccumulateBase
 {
 	GENERATED_BODY()
 
@@ -312,7 +312,7 @@ struct RIGVM_API FRigVMFunction_AccumulateVectorLerp : public FRigVMFunction_Acc
 	virtual void Initialize() override { bIsInitialized = false; }
 
 	RIGVM_METHOD()
-	virtual void Execute() override;
+	RIGVM_API virtual void Execute() override;
 
 	UPROPERTY(meta = (Input))
 	FVector TargetValue;
@@ -340,7 +340,7 @@ struct RIGVM_API FRigVMFunction_AccumulateVectorLerp : public FRigVMFunction_Acc
  * Interpolates two quaternions over time over and over again
  */
 USTRUCT(meta=(DisplayName="Accumulate Lerp (Quaternion)", TemplateName="AccumulateLerp", Keywords="Simulate,Ramp"))
-struct RIGVM_API FRigVMFunction_AccumulateQuatLerp : public FRigVMFunction_AccumulateBase
+struct FRigVMFunction_AccumulateQuatLerp : public FRigVMFunction_AccumulateBase
 {
 	GENERATED_BODY()
 	
@@ -355,7 +355,7 @@ struct RIGVM_API FRigVMFunction_AccumulateQuatLerp : public FRigVMFunction_Accum
 	virtual void Initialize() override { bIsInitialized = false; }
 
 	RIGVM_METHOD()
-	virtual void Execute() override;
+	RIGVM_API virtual void Execute() override;
 
 	UPROPERTY(meta=(Input))
 	FQuat TargetValue;
@@ -383,7 +383,7 @@ struct RIGVM_API FRigVMFunction_AccumulateQuatLerp : public FRigVMFunction_Accum
  * Interpolates two transforms over time over and over again
  */
 USTRUCT(meta = (DisplayName="Accumulate Lerp (Transform)", TemplateName="AccumulateLerp", Keywords="Simulate,Ramp"))
-struct RIGVM_API FRigVMFunction_AccumulateTransformLerp : public FRigVMFunction_AccumulateBase
+struct FRigVMFunction_AccumulateTransformLerp : public FRigVMFunction_AccumulateBase
 {
 	GENERATED_BODY()
 
@@ -398,7 +398,7 @@ struct RIGVM_API FRigVMFunction_AccumulateTransformLerp : public FRigVMFunction_
 	virtual void Initialize() override { bIsInitialized = false; }
 
 	RIGVM_METHOD()
-	virtual void Execute() override;
+	RIGVM_API virtual void Execute() override;
 
 	UPROPERTY(meta = (Input))
 	FTransform TargetValue;
@@ -426,7 +426,7 @@ struct RIGVM_API FRigVMFunction_AccumulateTransformLerp : public FRigVMFunction_
  * Accumulates the min and max values over time
  */
 USTRUCT(meta=(DisplayName="Accumulate Range (Float)", TemplateName="AccumulateRange", Keywords="Range"))
-struct RIGVM_API FRigVMFunction_AccumulateFloatRange : public FRigVMFunction_AccumulateBase
+struct FRigVMFunction_AccumulateFloatRange : public FRigVMFunction_AccumulateBase
 {
 	GENERATED_BODY()
 	
@@ -439,7 +439,7 @@ struct RIGVM_API FRigVMFunction_AccumulateFloatRange : public FRigVMFunction_Acc
 	virtual void Initialize() override { bIsInitialized = false; }
 
 	RIGVM_METHOD()
-	virtual void Execute() override;
+	RIGVM_API virtual void Execute() override;
 
 	UPROPERTY(meta = (Input))
 	float Value;
@@ -464,7 +464,7 @@ struct RIGVM_API FRigVMFunction_AccumulateFloatRange : public FRigVMFunction_Acc
  * Accumulates the min and max values over time
  */
 USTRUCT(meta = (DisplayName="Accumulate Range (Vector)", TemplateName="AccumulateRange", Keywords="Range"))
-struct RIGVM_API FRigVMFunction_AccumulateVectorRange : public FRigVMFunction_AccumulateBase
+struct FRigVMFunction_AccumulateVectorRange : public FRigVMFunction_AccumulateBase
 {
 	GENERATED_BODY()
 
@@ -477,7 +477,7 @@ struct RIGVM_API FRigVMFunction_AccumulateVectorRange : public FRigVMFunction_Ac
 	virtual void Initialize() override { bIsInitialized = false; }
 
 	RIGVM_METHOD()
-	virtual void Execute() override;
+	RIGVM_API virtual void Execute() override;
 
 	UPROPERTY(meta = (Input))
 	FVector Value;
diff --git a/Engine/Plugins/Runtime/RigVM/Source/RigVM/Public/RigVMFunctions/Simulation/RigVMFunction_AlphaInterp.h b/Engine/Plugins/Runtime/RigVM/Source/RigVM/Public/RigVMFunctions/Simulation/RigVMFunction_AlphaInterp.h
--- a/Engine/Plugins/Runtime/RigVM/Source/RigVM/Public/RigVMFunctions/Simulation/RigVMFunction_AlphaInterp.h
+++ b/Engine/Plugins/Runtime/RigVM/Source/RigVM/Public/RigVMFunctions/Simulation/RigVMFunction_AlphaInterp.h
@@ -10,7 +10,7 @@
  * Takes in a float value and outputs an accumulated value with a customized scale and clamp
  */
 USTRUCT(meta=(DisplayName="Alpha Interpolate", Keywords="Alpha,Lerp,LinearInterpolate", Category = "Simulation|Time", TemplateName="AlphaInterp", MenuDescSuffix = "(Float)"))
-struct RIGVM_API FRigVMFunction_AlphaInterp : public FRigVMFunction_SimBase
+struct FRigVMFunction_AlphaInterp : public FRigVMFunction_SimBase
 {
 	GENERATED_BODY()
 	
@@ -32,9 +32,9 @@ struct RIGVM_API FRigVMFunction_AlphaInterp : public FRigVMFunction_SimBase
 	}
 
 	RIGVM_METHOD()
-	virtual void Execute() override;
+	RIGVM_API virtual void Execute() override;
 
-	virtual FString ProcessPinLabelForInjection(const FString& InLabel) const override;
+	RIGVM_API virtual FString ProcessPinLabelForInjection(const FString& InLabel) const override;
 
 	UPROPERTY(meta=(Input))
 	float Value;
@@ -83,7 +83,7 @@ struct RIGVM_API FRigVMFunction_AlphaInterp : public FRigVMFunction_SimBase
  * Takes in a vector value and outputs an accumulated value with a customized scale and clamp
  */
 USTRUCT(meta=(DisplayName="Alpha Interpolate", Keywords="Alpha,Lerp,LinearInterpolate", Category = "Simulation|Time", TemplateName="AlphaInterp", MenuDescSuffix = "(Vector)"))
-struct RIGVM_API FRigVMFunction_AlphaInterpVector : public FRigVMFunction_SimBase
+struct FRigVMFunction_AlphaInterpVector : public FRigVMFunction_SimBase
 {
 	GENERATED_BODY()
 	
@@ -105,9 +105,9 @@ struct RIGVM_API FRigVMFunction_AlphaInterpVector : public FRigVMFunction_SimBas
 	}
 
 	RIGVM_METHOD()
-	virtual void Execute() override;
+	RIGVM_API virtual void Execute() override;
 
-	virtual FString ProcessPinLabelForInjection(const FString& InLabel) const override;
+	RIGVM_API virtual FString ProcessPinLabelForInjection(const FString& InLabel) const override;
 
 	UPROPERTY(meta=(Input))
 	FVector Value;
@@ -156,7 +156,7 @@ struct RIGVM_API FRigVMFunction_AlphaInterpVector : public FRigVMFunction_SimBas
  * Takes in a vector value and outputs an accumulated value with a customized scale and clamp
  */
 USTRUCT(meta=(DisplayName="Alpha Interpolate", Keywords="Alpha,Lerp,LinearInterpolate", Category = "Simulation|Time", TemplateName="AlphaInterp", MenuDescSuffix = "(Quat)"))
-struct RIGVM_API FRigVMFunction_AlphaInterpQuat : public FRigVMFunction_SimBase
+struct FRigVMFunction_AlphaInterpQuat : public FRigVMFunction_SimBase
 {
 	GENERATED_BODY()
 	
@@ -178,9 +178,9 @@ struct RIGVM_API FRigVMFunction_AlphaInterpQuat : public FRigVMFunction_SimBase
 	}
 
 	RIGVM_METHOD()
-	virtual void Execute() override;
+	RIGVM_API virtual void Execute() override;
 
-	virtual FString ProcessPinLabelForInjection(const FString& InLabel) const override;
+	RIGVM_API virtual FString ProcessPinLabelForInjection(const FString& InLabel) const override;
 
 	UPROPERTY(meta=(Input))
 	FQuat Value;
diff --git a/Engine/Plugins/Runtime/RigVM/Source/RigVM/Public/RigVMFunctions/Simulation/RigVMFunction_Kalman.h b/Engine/Plugins/Runtime/RigVM/Source/RigVM/Public/RigVMFunctions/Simulation/RigVMFunction_Kalman.h
--- a/Engine/Plugins/Runtime/RigVM/Source/RigVM/Public/RigVMFunctions/Simulation/RigVMFunction_Kalman.h
+++ b/Engine/Plugins/Runtime/RigVM/Source/RigVM/Public/RigVMFunctions/Simulation/RigVMFunction_Kalman.h
@@ -10,7 +10,7 @@
  * This uses a Kalman Filter internally.
  */
 USTRUCT(meta=(DisplayName="Average Over Time (Float)", Category = "Simulation|Time", TemplateName="KalmanFilter", Keywords="Average,Smooth,KalmanFilter"))
-struct RIGVM_API FRigVMFunction_KalmanFloat : public FRigVMFunction_SimBase
+struct FRigVMFunction_KalmanFloat : public FRigVMFunction_SimBase
 {
 	GENERATED_BODY()
 	
@@ -24,7 +24,7 @@ struct RIGVM_API FRigVMFunction_KalmanFloat : public FRigVMFunction_SimBase
 	virtual void Initialize() override { Buffer.Reset(); }
 
 	RIGVM_METHOD()
-	virtual void Execute() override;
+	RIGVM_API virtual void Execute() override;
 
 	UPROPERTY(meta = (Input))
 	float Value;
@@ -47,7 +47,7 @@ struct RIGVM_API FRigVMFunction_KalmanFloat : public FRigVMFunction_SimBase
  * This uses a Kalman Filter internally.
  */
 USTRUCT(meta=(DisplayName="Average Over Time (Vector)", Category = "Simulation|Time", TemplateName="KalmanFilter", Keywords="Average,Smooth,KalmanFilter"))
-struct RIGVM_API FRigVMFunction_KalmanVector : public FRigVMFunction_SimBase
+struct FRigVMFunction_KalmanVector : public FRigVMFunction_SimBase
 {
 	GENERATED_BODY()
 	
@@ -61,7 +61,7 @@ struct RIGVM_API FRigVMFunction_KalmanVector : public FRigVMFunction_SimBase
 	virtual void Initialize() override { Buffer.Reset(); }
 
 	RIGVM_METHOD()
-	virtual void Execute() override;
+	RIGVM_API virtual void Execute() override;
 
 	UPROPERTY(meta = (Input))
 	FVector Value;
@@ -84,7 +84,7 @@ struct RIGVM_API FRigVMFunction_KalmanVector : public FRigVMFunction_SimBase
  * This uses a Kalman Filter internally.
  */
 USTRUCT(meta=(DisplayName="Average Over Time (Transform)", Category = "Simulation|Time", TemplateName="KalmanFilter", Keywords="Average,Smooth,KalmanFilter"))
-struct RIGVM_API FRigVMFunction_KalmanTransform : public FRigVMFunction_SimBase
+struct FRigVMFunction_KalmanTransform : public FRigVMFunction_SimBase
 {
 	GENERATED_BODY()
 	
@@ -98,7 +98,7 @@ struct RIGVM_API FRigVMFunction_KalmanTransform : public FRigVMFunction_SimBase
 	virtual void Initialize() override { Buffer.Reset(); }
 
 	RIGVM_METHOD()
-	virtual void Execute() override;
+	RIGVM_API virtual void Execute() override;
 
 	UPROPERTY(meta = (Input))
 	FTransform Value;
diff --git a/Engine/Plugins/Runtime/RigVM/Source/RigVM/Public/RigVMFunctions/Simulation/RigVMFunction_SimBase.h b/Engine/Plugins/Runtime/RigVM/Source/RigVM/Public/RigVMFunctions/Simulation/RigVMFunction_SimBase.h
--- a/Engine/Plugins/Runtime/RigVM/Source/RigVM/Public/RigVMFunctions/Simulation/RigVMFunction_SimBase.h
+++ b/Engine/Plugins/Runtime/RigVM/Source/RigVM/Public/RigVMFunctions/Simulation/RigVMFunction_SimBase.h
@@ -6,13 +6,13 @@
 #include "RigVMFunction_SimBase.generated.h"
 
 USTRUCT(meta=(Abstract, Category = "Simulation", NodeColor = "0.25 0.05 0.05"))
-struct RIGVM_API FRigVMFunction_SimBase : public FRigVMStruct
+struct FRigVMFunction_SimBase : public FRigVMStruct
 {
 	GENERATED_BODY()
 };
 
 USTRUCT(meta=(Abstract, Category = "Simulation", NodeColor = "0.25 0.05 0.05"))
-struct RIGVM_API FRigVMFunction_SimBaseMutable : public FRigVMStructMutable
+struct FRigVMFunction_SimBaseMutable : public FRigVMStructMutable
 {
 	GENERATED_BODY()
 };
diff --git a/Engine/Plugins/Runtime/RigVM/Source/RigVM/Public/RigVMFunctions/Simulation/RigVMFunction_TimeOffset.h b/Engine/Plugins/Runtime/RigVM/Source/RigVM/Public/RigVMFunctions/Simulation/RigVMFunction_TimeOffset.h
--- a/Engine/Plugins/Runtime/RigVM/Source/RigVM/Public/RigVMFunctions/Simulation/RigVMFunction_TimeOffset.h
+++ b/Engine/Plugins/Runtime/RigVM/Source/RigVM/Public/RigVMFunctions/Simulation/RigVMFunction_TimeOffset.h
@@ -9,7 +9,7 @@
  * Records a value over time and can access the value from the past
  */
 USTRUCT(meta=(DisplayName="Value Over Time (Float)", Category = "Simulation|Time", TemplateName="TimeOffset", Keywords="Buffer,Delta,History,Previous,TimeOffset,Delay"))
-struct RIGVM_API FRigVMFunction_TimeOffsetFloat : public FRigVMFunction_SimBase
+struct FRigVMFunction_TimeOffsetFloat : public FRigVMFunction_SimBase
 {
 	GENERATED_BODY()
 	
@@ -26,7 +26,7 @@ struct RIGVM_API FRigVMFunction_TimeOffsetFloat : public FRigVMFunction_SimBase
 	virtual void Initialize() override { Buffer.Reset(); }
 
 	RIGVM_METHOD()
-	virtual void Execute() override;
+	RIGVM_API virtual void Execute() override;
 	
 	/** The value to record */
 	UPROPERTY(meta = (Input))
@@ -64,7 +64,7 @@ struct RIGVM_API FRigVMFunction_TimeOffsetFloat : public FRigVMFunction_SimBase
  * Records a value over time and can access the value from the past
  */
 USTRUCT(meta=(DisplayName="Value Over Time (Vector)", Category = "Simulation|Time", TemplateName="TimeOffset", Keywords="Buffer,Delta,History,Previous,TimeOffset,Delay"))
-struct RIGVM_API FRigVMFunction_TimeOffsetVector : public FRigVMFunction_SimBase
+struct FRigVMFunction_TimeOffsetVector : public FRigVMFunction_SimBase
 {
 	GENERATED_BODY()
 	
@@ -81,7 +81,7 @@ struct RIGVM_API FRigVMFunction_TimeOffsetVector : public FRigVMFunction_SimBase
 	virtual void Initialize() override { Buffer.Reset(); }
 
 	RIGVM_METHOD()
-	virtual void Execute() override;
+	RIGVM_API virtual void Execute() override;
 	
 	/** The value to record */
 	UPROPERTY(meta = (Input))
@@ -119,7 +119,7 @@ struct RIGVM_API FRigVMFunction_TimeOffsetVector : public FRigVMFunction_SimBase
  * Records a value over time and can access the value from the past
  */
 USTRUCT(meta=(DisplayName="Value Over Time (Transform)", Category = "Simulation|Time", TemplateName="TimeOffset", Keywords="Buffer,Delta,History,Previous,TimeOffset,Delay"))
-struct RIGVM_API FRigVMFunction_TimeOffsetTransform : public FRigVMFunction_SimBase
+struct FRigVMFunction_TimeOffsetTransform : public FRigVMFunction_SimBase
 {
 	GENERATED_BODY()
 	
@@ -136,7 +136,7 @@ struct RIGVM_API FRigVMFunction_TimeOffsetTransform : public FRigVMFunction_SimB
 	virtual void Initialize() override { Buffer.Reset(); }
 
 	RIGVM_METHOD()
-	virtual void Execute() override;
+	RIGVM_API virtual void Execute() override;
 	
 	/** The value to record */
 	UPROPERTY(meta = (Input))
diff --git a/Engine/Plugins/Runtime/RigVM/Source/RigVM/Public/RigVMFunctions/Simulation/RigVMFunction_Timeline.h b/Engine/Plugins/Runtime/RigVM/Source/RigVM/Public/RigVMFunctions/Simulation/RigVMFunction_Timeline.h
--- a/Engine/Plugins/Runtime/RigVM/Source/RigVM/Public/RigVMFunctions/Simulation/RigVMFunction_Timeline.h
+++ b/Engine/Plugins/Runtime/RigVM/Source/RigVM/Public/RigVMFunctions/Simulation/RigVMFunction_Timeline.h
@@ -9,7 +9,7 @@
  * Simulates a time value - can act as a timeline playing back
  */
 USTRUCT(meta=(DisplayName="Accumulated Time", Category = "Simulation|Accumulate", Keywords="Playback,Pause,Timeline"))
-struct RIGVM_API FRigVMFunction_Timeline : public FRigVMFunction_SimBase
+struct FRigVMFunction_Timeline : public FRigVMFunction_SimBase
 {
 	GENERATED_BODY()
 	
@@ -23,7 +23,7 @@ struct RIGVM_API FRigVMFunction_Timeline : public FRigVMFunction_SimBase
 	virtual void Initialize() override { bIsInitialized = false; }
 
 	RIGVM_METHOD()
-	virtual void Execute() override;
+	RIGVM_API virtual void Execute() override;
 
 	UPROPERTY(meta = (Input))
 	float Speed;
@@ -42,7 +42,7 @@ struct RIGVM_API FRigVMFunction_Timeline : public FRigVMFunction_SimBase
  * Simulates a time value - and outputs loop information
  */
 USTRUCT(meta=(DisplayName="Time Loop", Category = "Simulation|Accumulate", Keywords="Playback,Pause,Timeline"))
-struct RIGVM_API FRigVMFunction_TimeLoop : public FRigVMFunction_SimBase
+struct FRigVMFunction_TimeLoop : public FRigVMFunction_SimBase
 {
 	GENERATED_BODY()
 	
@@ -60,7 +60,7 @@ struct RIGVM_API FRigVMFunction_TimeLoop : public FRigVMFunction_SimBase
 	virtual void Initialize() override { bIsInitialized = false; }
 
 	RIGVM_METHOD()
-	virtual void Execute() override;
+	RIGVM_API virtual void Execute() override;
 
 	UPROPERTY(meta = (Input))
 	float Speed;
diff --git a/Engine/Plugins/Runtime/RigVM/Source/RigVM/Public/RigVMFunctions/Simulation/RigVMFunction_Verlet.h b/Engine/Plugins/Runtime/RigVM/Source/RigVM/Public/RigVMFunctions/Simulation/RigVMFunction_Verlet.h
--- a/Engine/Plugins/Runtime/RigVM/Source/RigVM/Public/RigVMFunctions/Simulation/RigVMFunction_Verlet.h
+++ b/Engine/Plugins/Runtime/RigVM/Source/RigVM/Public/RigVMFunctions/Simulation/RigVMFunction_Verlet.h
@@ -11,7 +11,7 @@
  * is more accurate and stable, and has more meaningful parameters.
  */
 USTRUCT(meta=(DisplayName="Verlet (Vector)", Category = "Simulation|Springs", TemplateName="Verlet", Keywords="Simulate,Integrate"))
-struct RIGVM_API FRigVMFunction_VerletIntegrateVector : public FRigVMFunction_SimBase
+struct FRigVMFunction_VerletIntegrateVector : public FRigVMFunction_SimBase
 {
 	GENERATED_BODY()
 	
@@ -26,7 +26,7 @@ struct RIGVM_API FRigVMFunction_VerletIntegrateVector : public FRigVMFunction_Si
 	}
 
 	RIGVM_METHOD()
-	virtual void Execute() override;
+	RIGVM_API virtual void Execute() override;
 
 	UPROPERTY(meta = (Input))
 	FVector Target;
diff --git a/Engine/Plugins/Runtime/RigVM/Source/RigVM/Public/RigVMHost.h b/Engine/Plugins/Runtime/RigVM/Source/RigVM/Public/RigVMHost.h
--- a/Engine/Plugins/Runtime/RigVM/Source/RigVM/Public/RigVMHost.h
+++ b/Engine/Plugins/Runtime/RigVM/Source/RigVM/Public/RigVMHost.h
@@ -15,29 +15,31 @@
 #endif
 #include "RigVMHost.generated.h"
 
+#define UE_API RIGVM_API 
+
 // set this to something larger than 0 to profile N runs
 #ifndef UE_RIGVM_PROFILE_EXECUTE_UNITS_NUM
 #define UE_RIGVM_PROFILE_EXECUTE_UNITS_NUM 0
 #endif
 
-UCLASS(Abstract, editinlinenew)
-class RIGVM_API URigVMHost : public UObject, public IInterface_AssetUserData
+UCLASS(Abstract, editinlinenew, MinimalAPI)
+class URigVMHost : public UObject, public IInterface_AssetUserData
 {
 	GENERATED_UCLASS_BODY()
 
 public:
 
 	UFUNCTION(BlueprintCallable, Category = RigVM)
-	static TArray<URigVMHost*> FindRigVMHosts(UObject* Outer, TSubclassOf<URigVMHost> OptionalClass);
+	static UE_API TArray<URigVMHost*> FindRigVMHosts(UObject* Outer, TSubclassOf<URigVMHost> OptionalClass);
 
-	static bool IsGarbageOrDestroyed(const UObject* InObject);
+	static UE_API bool IsGarbageOrDestroyed(const UObject* InObject);
 
 	/** UObject interface */
-	virtual UWorld* GetWorld() const override;
-	virtual void Serialize(FArchive& Ar) override;
-	virtual void PostLoad() override;
-	virtual void PreSave(FObjectPreSaveContext ObjectSaveContext) override;
-	virtual void BeginDestroy() override;
+	UE_API virtual UWorld* GetWorld() const override;
+	UE_API virtual void Serialize(FArchive& Ar) override;
+	UE_API virtual void PostLoad() override;
+	UE_API virtual void PreSave(FObjectPreSaveContext ObjectSaveContext) override;
+	UE_API virtual void BeginDestroy() override;
 
 	/** Gets the current absolute time */
 	UFUNCTION(BlueprintPure, Category = "RigVM")
@@ -49,67 +51,67 @@ public:
 
 	/** Set the current delta time */
 	UFUNCTION(BlueprintCallable, Category="RigVM")
-	void SetDeltaTime(float InDeltaTime);
+	UE_API void SetDeltaTime(float InDeltaTime);
 
 	/** Set the current absolute time */
 	UFUNCTION(BlueprintCallable, Category = "RigVM")
-	void SetAbsoluteTime(float InAbsoluteTime, bool InSetDeltaTimeZero = false);
+	UE_API void SetAbsoluteTime(float InAbsoluteTime, bool InSetDeltaTimeZero = false);
 
 	/** Set the current absolute and delta times */
 	UFUNCTION(BlueprintCallable, Category = "RigVM")
-	void SetAbsoluteAndDeltaTime(float InAbsoluteTime, float InDeltaTime);
+	UE_API void SetAbsoluteAndDeltaTime(float InAbsoluteTime, float InDeltaTime);
 
 	/** Set the current fps */
 	UFUNCTION(BlueprintCallable, Category = "RigVM")
-	void SetFramesPerSecond(float InFramesPerSecond);
+	UE_API void SetFramesPerSecond(float InFramesPerSecond);
 
 	/** Returns the current frames per second (this may change over time) */
 	UFUNCTION(BlueprintPure, Category = "RigVM")
-	float GetCurrentFramesPerSecond() const;
+	UE_API float GetCurrentFramesPerSecond() const;
 
 	/** Returns the public context script struct to use for this owner */
 	virtual UScriptStruct* GetPublicContextStruct() const { return FRigVMExecuteContext::StaticStruct(); }
 
 	/** Is valid for execution */
 	UFUNCTION(BlueprintPure, Category="RigVM")
-	virtual bool CanExecute() const;
+	UE_API virtual bool CanExecute() const;
 
 	/** Initialize things for the RigVM owner */
-	virtual void Initialize(bool bRequestInit = true);
+	UE_API virtual void Initialize(bool bRequestInit = true);
 
 	/** Initialize this Host VM Instance */
-	virtual bool InitializeVM(const FName& InEventName);
+	UE_API virtual bool InitializeVM(const FName& InEventName);
 
 	/** Evaluate at Any Thread */
-	virtual void Evaluate_AnyThread();
+	UE_API virtual void Evaluate_AnyThread();
 
 	/** Locks for the scope of Evaluate_AnyThread */
 	FTransactionallySafeCriticalSection& GetEvaluateMutex() { return EvaluateMutex; };
 
 	/** Returns the member properties as an external variable array */
-	TArray<FRigVMExternalVariable> GetExternalVariables() const;
+	UE_API TArray<FRigVMExternalVariable> GetExternalVariables() const;
 
 	/** Returns the public member properties as an external variable array */
-	TArray<FRigVMExternalVariable> GetPublicVariables() const;
+	UE_API TArray<FRigVMExternalVariable> GetPublicVariables() const;
 
 	/** Returns a public variable given its name */
-	FRigVMExternalVariable GetPublicVariableByName(const FName& InVariableName) const;
+	UE_API FRigVMExternalVariable GetPublicVariableByName(const FName& InVariableName) const;
 
 	/** Returns the names of variables accessible in scripting */
 	UFUNCTION(BlueprintPure, Category = "RigVM", meta=(DisplayName="Get Variables"))
-	TArray<FName> GetScriptAccessibleVariables() const;
+	UE_API TArray<FName> GetScriptAccessibleVariables() const;
 
 	/** Returns the type of a given variable */
 	UFUNCTION(BlueprintPure, Category = "RigVM")
-	FName GetVariableType(const FName& InVariableName) const;
+	UE_API FName GetVariableType(const FName& InVariableName) const;
 
 	/** Returns the value of a given variable as a string */
 	UFUNCTION(BlueprintPure, Category = "RigVM")
-	FString GetVariableAsString(const FName& InVariableName) const;
+	UE_API FString GetVariableAsString(const FName& InVariableName) const;
 
 	/** Returns the value of a given variable as a string */
 	UFUNCTION(BlueprintCallable, Category = "RigVM")
-	bool SetVariableFromString(const FName& InVariableName, const FString& InValue);
+	UE_API bool SetVariableFromString(const FName& InVariableName, const FString& InValue);
 
 	template<class T>
 	T GetPublicVariableValue(const FName& InVariableName)
@@ -142,14 +144,14 @@ public:
 	UPROPERTY()
 	FRigVMRuntimeSettings VMRuntimeSettings;
 
-	virtual void InvalidateCachedMemory();
+	UE_API virtual void InvalidateCachedMemory();
 
 	// Regenerates cached handles after a structural change (i.e. new UUserStruct)
-	virtual void RecreateCachedMemory();
+	UE_API virtual void RecreateCachedMemory();
 
 	/** Execute */
 	UFUNCTION(BlueprintCallable, Category = "RigVM")
-	virtual bool Execute(const FName& InEventName);
+	UE_API virtual bool Execute(const FName& InEventName);
 
 #if WITH_EDITOR
 
@@ -169,12 +171,12 @@ private:
 protected:
 
 	/** Initialize the CDO VM */
-	virtual bool InitializeCDOVM();
+	UE_API virtual bool InitializeCDOVM();
 
 	/** ExecuteUnits */
-	virtual bool Execute_Internal(const FName& InEventName);
+	UE_API virtual bool Execute_Internal(const FName& InEventName);
 
-	static bool DisableExecution();
+	static UE_API bool DisableExecution();
 
 public:
 
@@ -185,18 +187,18 @@ public:
 	}
 
 	UFUNCTION(BlueprintPure, Category = "RigVM")
-	virtual bool SupportsEvent(const FName& InEventName) const;
+	UE_API virtual bool SupportsEvent(const FName& InEventName) const;
 
 	UFUNCTION(BlueprintPure, Category = "RigVM")
-	virtual const TArray<FName>& GetSupportedEvents() const;
+	UE_API virtual const TArray<FName>& GetSupportedEvents() const;
 
 	/** Execute a user defined event */
 	UFUNCTION(BlueprintCallable, Category = "RigVM")
-	bool ExecuteEvent(const FName& InEventName);
+	UE_API bool ExecuteEvent(const FName& InEventName);
 
 	/** Requests to perform an init during the next execution */
 	UFUNCTION(BlueprintCallable, Category = "RigVM")
-	virtual void RequestInit();
+	UE_API virtual void RequestInit();
 
 	/** Returns true if this host requires the VM memory to be initialized */
 	UFUNCTION(BlueprintPure, Category = "RigVM")
@@ -210,29 +212,29 @@ public:
 	 * @param InEventIndex Deprecated argument. Not used.
 	 */
 	UFUNCTION(BlueprintCallable, Category = "RigVM")
-	void RequestRunOnceEvent(const FName& InEventName, int32 InEventIndex = -1);
+	UE_API void RequestRunOnceEvent(const FName& InEventName, int32 InEventIndex = -1);
 
 	/** Removes an event running once */
 	UFUNCTION(BlueprintCallable, Category = "RigVM")
-	bool RemoveRunOnceEvent(const FName& InEventName);
+	UE_API bool RemoveRunOnceEvent(const FName& InEventName);
 
 	/** Returns true if an event is queued to run once */
-	bool IsRunOnceEvent(const FName& InEventName) const;
+	UE_API bool IsRunOnceEvent(const FName& InEventName) const;
 
 	/** Returns the queue of events to run */
 	const TArray<FName>& GetEventQueue() const { return EventQueue; }
 
 	/** Sets the queue of events to run */
-	void SetEventQueue(const TArray<FName>& InEventNames);
+	UE_API void SetEventQueue(const TArray<FName>& InEventNames);
 
 	/** Provides the chance to a subclass to modify the event queue as needed */
 	virtual void AdaptEventQueueForEvaluate(TArray<FName>& InOutEventQueueToRun) {}
 
 	/** Update the settings such as array bound and log facilities */
-	void UpdateVMSettings();
+	UE_API void UpdateVMSettings();
 
 	UFUNCTION(BlueprintPure, Category = "RigVM")
-	URigVM* GetVM();
+	UE_API URigVM* GetVM();
 
 #if WITH_EDITOR
 	const FRigVMLog* GetLog() const { return RigVMLog; }
@@ -241,11 +243,11 @@ public:
 #endif
 
 	// Returns the compiler generated VM memory storage by type
-	virtual const FRigVMMemoryStorageStruct* GetDefaultMemoryByType(ERigVMMemoryType InMemoryType) const;
+	UE_API virtual const FRigVMMemoryStorageStruct* GetDefaultMemoryByType(ERigVMMemoryType InMemoryType) const;
 
 	// Returns an instanced VM memory storage by type
-	virtual FRigVMMemoryStorageStruct* GetMemoryByType(ERigVMMemoryType InMemoryType);
-	virtual const FRigVMMemoryStorageStruct* GetMemoryByType(ERigVMMemoryType InMemoryType) const;
+	UE_API virtual FRigVMMemoryStorageStruct* GetMemoryByType(ERigVMMemoryType InMemoryType);
+	UE_API virtual const FRigVMMemoryStorageStruct* GetMemoryByType(ERigVMMemoryType InMemoryType) const;
 
 	// The instanced mutable work memory
 	FRigVMMemoryStorageStruct* GetWorkMemory() { return GetMemoryByType(ERigVMMemoryType::Work); }
@@ -270,16 +272,16 @@ public:
 	const FRigVMDrawContainer& GetDrawContainer() const { return DrawContainer; };
 	FRigVMDrawContainer& GetDrawContainer() { return DrawContainer; };
 
-	void DrawIntoPDI(FPrimitiveDrawInterface* PDI, const FTransform& InTransform);
+	UE_API void DrawIntoPDI(FPrimitiveDrawInterface* PDI, const FTransform& InTransform);
 
-	virtual USceneComponent* GetOwningSceneComponent(); 
+	UE_API virtual USceneComponent* GetOwningSceneComponent(); 
 
 	virtual void PostInitInstanceIfRequired() {};
-	void SwapVMToNativizedIfRequired(UClass* InNativizedClass = nullptr);
-	static bool AreNativizedVMsDisabled() ;
+	UE_API void SwapVMToNativizedIfRequired(UClass* InNativizedClass = nullptr);
+	static UE_API bool AreNativizedVMsDisabled() ;
 
 #if WITH_EDITORONLY_DATA
-	static void DeclareConstructClasses(TArray<FTopLevelAssetPath>& OutConstructClasses, const UClass* SpecificSubclass);
+	static UE_API void DeclareConstructClasses(TArray<FTopLevelAssetPath>& OutConstructClasses, const UClass* SpecificSubclass);
 #endif
 
 #if UE_RIGVM_DEBUG_EXECUTION
@@ -320,11 +322,11 @@ public:
 		return *RigVMExtendedExecuteContext;
 	};
 
-	UObject* ResolveUserDefinedTypeById(const FString& InTypeName) const;
+	UE_API UObject* ResolveUserDefinedTypeById(const FString& InTypeName) const;
 
 protected:
 
-	virtual void PostInitInstance(URigVMHost* InCDO);
+	UE_API virtual void PostInitInstance(URigVMHost* InCDO);
 
 	/** Current delta time */
 	float DeltaTime;
@@ -351,8 +353,8 @@ protected:
 	bool bEnableLogging;
 #endif
 
-	void GenerateUserDefinedDependenciesData(FRigVMExtendedExecuteContext& Context);
-	TArray<const UObject*> GetUserDefinedDependencies(const TArray<const FRigVMMemoryStorageStruct*> InMemory);
+	UE_API void GenerateUserDefinedDependenciesData(FRigVMExtendedExecuteContext& Context);
+	UE_API TArray<const UObject*> GetUserDefinedDependencies(const TArray<const FRigVMMemoryStorageStruct*> InMemory);
 
 	UPROPERTY()
 	TMap<FString, FSoftObjectPath> UserDefinedStructGuidToPathName;
@@ -372,9 +374,9 @@ private:
 
 protected:
 	
-	void HandleExecutionReachedExit(const FName& InEventName);
+	UE_API void HandleExecutionReachedExit(const FName& InEventName);
 	
-	virtual TArray<FRigVMExternalVariable> GetExternalVariablesImpl(bool bFallbackToBlueprint) const;
+	UE_API virtual TArray<FRigVMExternalVariable> GetExternalVariablesImpl(bool bFallbackToBlueprint) const;
 
 	FProperty* GetPublicVariableProperty(const FName& InVariableName) const
 	{
@@ -409,10 +411,10 @@ public:
 	bool IsEvaluating() const { return !EventQueueToRun.IsEmpty(); }
 
 	/** Copy the VM from the default object */
-	void InstantiateVMFromCDO();
+	UE_API void InstantiateVMFromCDO();
 	
 	/** Copy the default values of external variables from the default object */
-	void CopyExternalVariableDefaultValuesFromCDO();
+	UE_API void CopyExternalVariableDefaultValuesFromCDO();
 	
 	/** Broadcasts a notification whenever the RigVMHost's memory is initialized. */
 	FRigVMExecutedEvent InitializedEvent;
@@ -425,16 +427,16 @@ public:
 
 protected: 
 
-	virtual void InitializeFromCDO();
+	UE_API virtual void InitializeFromCDO();
 
-	virtual void CopyVMMemory(FRigVMExtendedExecuteContext& TargetContext, const FRigVMExtendedExecuteContext& SourceContext);
+	UE_API virtual void CopyVMMemory(FRigVMExtendedExecuteContext& TargetContext, const FRigVMExtendedExecuteContext& SourceContext);
 
 public:
 	//~ Begin IInterface_AssetUserData Interface
-	virtual void AddAssetUserData(UAssetUserData* InUserData) override;
-	virtual void RemoveUserDataOfClass(TSubclassOf<UAssetUserData> InUserDataClass) override;
-	virtual UAssetUserData* GetAssetUserDataOfClass(TSubclassOf<UAssetUserData> InUserDataClass) override;
-	virtual const TArray<UAssetUserData*>* GetAssetUserDataArray() const override;
+	UE_API virtual void AddAssetUserData(UAssetUserData* InUserData) override;
+	UE_API virtual void RemoveUserDataOfClass(TSubclassOf<UAssetUserData> InUserDataClass) override;
+	UE_API virtual UAssetUserData* GetAssetUserDataOfClass(TSubclassOf<UAssetUserData> InUserDataClass) override;
+	UE_API virtual const TArray<UAssetUserData*>* GetAssetUserDataArray() const override;
 	//~ End IInterface_AssetUserData Interface
 
 protected:
@@ -491,7 +493,7 @@ protected:
 	FRigVMDebugInfo DebugInfo;
 	FRigVMProfilingInfo ProfilingInfo;
 	TMap<FString, bool> LoggedMessages;
-	void LogOnce(EMessageSeverity::Type InSeverity, int32 InInstructionIndex, const FString& InMessage);
+	UE_API void LogOnce(EMessageSeverity::Type InSeverity, int32 InInstructionIndex, const FString& InMessage);
 	
 public:
 
@@ -503,11 +505,11 @@ public:
 	}
 	
 	/** Adds a breakpoint in the VM at the InstructionIndex for the Node / Subject */
-	void AddBreakpoint(int32 InstructionIndex, UObject* InSubject, uint16 InDepth);
+	UE_API void AddBreakpoint(int32 InstructionIndex, UObject* InSubject, uint16 InDepth);
 
 	/** If the VM is halted at a breakpoint, it sets a breakpoint action so that
 	 *  it is applied on the next VM execution */
-	bool ExecuteBreakpointAction(const ERigVMBreakpointAction BreakpointAction);
+	UE_API bool ExecuteBreakpointAction(const ERigVMBreakpointAction BreakpointAction);
 
 	const FRigVMBreakpoint& GetHaltedAtBreakpoint() const
 	{
@@ -526,8 +528,8 @@ public:
 	const FRigVMProfilingInfo& GetProfilingInfo() const { return ProfilingInfo; }
 
 	/** Creates the snapshot VM if required and returns it */
-	URigVM* GetSnapshotVM(bool bCreateIfNeeded = true);
-	FRigVMExtendedExecuteContext& GetSnapshotContext();
+	UE_API URigVM* GetSnapshotVM(bool bCreateIfNeeded = true);
+	UE_API FRigVMExtendedExecuteContext& GetSnapshotContext();
 
 #endif	
 
@@ -547,7 +549,7 @@ protected:
 	friend class FRigVMCompileSettingsDetails;
 };
 
-class RIGVM_API FRigVMBracketScope
+class FRigVMBracketScope
 {
 public:
 
@@ -567,6 +569,8 @@ private:
 	int32& Bracket;
 };
 
+#undef UE_API
+
 #if UE_ENABLE_INCLUDE_ORDER_DEPRECATED_IN_5_6
 #include "SceneManagement.h"
 #endif
diff --git a/Engine/Plugins/Runtime/RigVM/Source/RigVM/Public/RigVMSettings.h b/Engine/Plugins/Runtime/RigVM/Source/RigVM/Public/RigVMSettings.h
--- a/Engine/Plugins/Runtime/RigVM/Source/RigVM/Public/RigVMSettings.h
+++ b/Engine/Plugins/Runtime/RigVM/Source/RigVM/Public/RigVMSettings.h
@@ -17,8 +17,8 @@ class UStaticMesh;
 /**
  * Customize RigVM Rig Editor.
  */
-UCLASS(config = EditorPerProjectUserSettings, meta=(DisplayName="RigVM Editor"))
-class RIGVM_API URigVMEditorSettings : public UDeveloperSettings
+UCLASS(config = EditorPerProjectUserSettings, meta=(DisplayName="RigVM Editor"), MinimalAPI)
+class URigVMEditorSettings : public UDeveloperSettings
 {
 	GENERATED_UCLASS_BODY()
 
@@ -40,8 +40,8 @@ public:
 };
 
 
-UCLASS(config = Editor, meta=(DisplayName="RigVM Project Settings"))
-class RIGVM_API URigVMProjectSettings : public UDeveloperSettings
+UCLASS(config = Editor, meta=(DisplayName="RigVM Project Settings"), MinimalAPI)
+class URigVMProjectSettings : public UDeveloperSettings
 {
 	GENERATED_UCLASS_BODY()
 
@@ -51,6 +51,6 @@ public:
 	TArray<FRigVMTag> VariantTags;
 
 	UFUNCTION(BlueprintPure, Category= Variants)
-	FRigVMTag GetTag(FName InTagName) const;
-	const FRigVMTag* FindTag(FName InTagName) const;
+	RIGVM_API FRigVMTag GetTag(FName InTagName) const;
+	RIGVM_API const FRigVMTag* FindTag(FName InTagName) const;
 };
\ No newline at end of file
