[RigVM]
* Fixed linker error caused by missing dll storage on ctor

diff --git a/Engine/Plugins/Runtime/RigVM/Source/RigVM/Public/RigVMCore/RigVMFunction.h b/Engine/Plugins/Runtime/RigVM/Source/RigVM/Public/RigVMCore/RigVMFunction.h
--- a/Engine/Plugins/Runtime/RigVM/Source/RigVM/Public/RigVMCore/RigVMFunction.h
+++ b/Engine/Plugins/Runtime/RigVM/Source/RigVM/Public/RigVMCore/RigVMFunction.h
@@ -149,7 +149,7 @@ struct FRigVMFunction
 	{
 	}
 
-	FRigVMFunction(const TCHAR* InName, FRigVMFunctionPtr InFunctionPtr, UScriptStruct* InStruct = nullptr, int32 InIndex = INDEX_NONE, const TArray<FRigVMFunctionArgument>& InArguments = TArray<FRigVMFunctionArgument>());
+	RIGVM_API FRigVMFunction(const TCHAR* InName, FRigVMFunctionPtr InFunctionPtr, UScriptStruct* InStruct = nullptr, int32 InIndex = INDEX_NONE, const TArray<FRigVMFunctionArgument>& InArguments = TArray<FRigVMFunctionArgument>());
 
 	FRigVMFunction(const FString& InName, FRigVMFunctionPtr InFunctionPtr, FRigVMDispatchFactory* InFactory, int32 InIndex = INDEX_NONE, const TArray<FRigVMFunctionArgument>& InArguments = TArray<FRigVMFunctionArgument>())
 	: Name(InName)
