RigVM: Introduce scope lock around events to run once

#rnx

diff --git a/Engine/Plugins/Runtime/RigVM/Source/RigVM/Private/RigVMHost.cpp b/Engine/Plugins/Runtime/RigVM/Source/RigVM/Private/RigVMHost.cpp
--- a/Engine/Plugins/Runtime/RigVM/Source/RigVM/Private/RigVMHost.cpp
+++ b/Engine/Plugins/Runtime/RigVM/Source/RigVM/Private/RigVMHost.cpp
@@ -348,7 +348,10 @@ void URigVMHost::Evaluate_AnyThread()
 	AdaptEventQueueForEvaluate(EventQueueToRun);
 	
 	// insert the events queued to run once at the beginning of the queue (in the order specified)
-	EventQueueToRun.Insert(EventsToRunOnce, 0);
+	{
+		UE::TScopeLock EventQueueToRunOnceLock(GetEventQueueToRunOnceMutex());
+		EventQueueToRun.Insert(EventsToRunOnce, 0);
+	}
 
 #if WITH_EDITOR
 	FName FirstEvent = NAME_None;
@@ -790,16 +793,19 @@ void URigVMHost::RequestInit()
 
 void URigVMHost::RequestRunOnceEvent(const FName& InEventName, int32 InEventIndex)
 {
+	UE::TScopeLock EventQueueToRunOnceLock(GetEventQueueToRunOnceMutex());
 	EventsToRunOnce.AddUnique(InEventName);
 }
 
 bool URigVMHost::RemoveRunOnceEvent(const FName& InEventName)
 {
+	UE::TScopeLock EventQueueToRunOnceLock(GetEventQueueToRunOnceMutex());
 	return EventsToRunOnce.Remove(InEventName) > 0;
 }
 
 bool URigVMHost::IsRunOnceEvent(const FName& InEventName) const
 {
+	UE::TScopeLock EventQueueToRunOnceLock(GetEventQueueToRunOnceMutex());
 	return EventsToRunOnce.Contains(InEventName);
 }
 
diff --git a/Engine/Plugins/Runtime/RigVM/Source/RigVM/Public/RigVMHost.h b/Engine/Plugins/Runtime/RigVM/Source/RigVM/Public/RigVMHost.h
--- a/Engine/Plugins/Runtime/RigVM/Source/RigVM/Public/RigVMHost.h
+++ b/Engine/Plugins/Runtime/RigVM/Source/RigVM/Public/RigVMHost.h
@@ -88,6 +88,9 @@ public:
 	/** Locks for the scope of Evaluate_AnyThread */
 	FTransactionallySafeCriticalSection& GetEvaluateMutex() { return EvaluateMutex; };
 
+	/** Lock for editing the event queue to run once */
+	FTransactionallySafeCriticalSection& GetEventQueueToRunOnceMutex() const { return EventQueueToRunOnceMutex; };
+
 	/** Returns the member properties as an external variable array */
 	UE_API TArray<FRigVMExternalVariable> GetExternalVariables() const;
 
@@ -536,6 +539,7 @@ public:
 private:
 
 	FTransactionallySafeCriticalSection EvaluateMutex;
+	mutable FTransactionallySafeCriticalSection EventQueueToRunOnceMutex;
 
 protected:
 
