RigVM: Fix cached pin defaults for adapted values such as quaternions

#rnx

diff --git a/Engine/Plugins/Runtime/RigVM/Source/RigVMDeveloper/Private/RigVMModel/RigVMPin.cpp b/Engine/Plugins/Runtime/RigVM/Source/RigVMDeveloper/Private/RigVMModel/RigVMPin.cpp
--- a/Engine/Plugins/Runtime/RigVM/Source/RigVMDeveloper/Private/RigVMModel/RigVMPin.cpp
+++ b/Engine/Plugins/Runtime/RigVM/Source/RigVMDeveloper/Private/RigVMModel/RigVMPin.cpp
@@ -123,6 +123,7 @@ URigVMPin::URigVMPin()
 	, LastKnownTypeIndex(INDEX_NONE)
 	, CachedIsStringType(this)
 	, CachedDefaultValue(this)
+	, CachedAdaptedDefaultValue(this)
 	, CachedCPPTypeObjectHash(this)
 	, CachedShowInDetailsPanelOnly(this)
 	, CachedPinPath(this)
@@ -914,9 +915,11 @@ FString URigVMPin::GetDefaultValue(const URigVMPin::FPinOverride& InOverride, bo
 		return OverrideValuePtr->DefaultValue;
 	}
 
-	if(CachedDefaultValue.IsValid())
+	TRigVMModelCachedValue<URigVMPin, FString>& Cache = bAdaptValueForPinType ? CachedAdaptedDefaultValue : CachedDefaultValue;
+
+	if(Cache.IsValid())
 	{
-		return CachedDefaultValue.GetValue();
+		return Cache.GetValue();
 	}
 
 	if (IsArray())
@@ -925,8 +928,8 @@ FString URigVMPin::GetDefaultValue(const URigVMPin::FPinOverride& InOverride, bo
 		const TRigVMTypeIndex ArrayType = GetTypeIndex();
 		if (!Registry.IsArrayType(ArrayType))
 		{
-			CachedDefaultValue = TEXT("()"); 
-			return CachedDefaultValue.GetValue();
+			Cache = TEXT("()"); 
+			return Cache.GetValue();
 		}
 		if (SubPins.Num() > 0)
 		{
@@ -936,8 +939,8 @@ FString URigVMPin::GetDefaultValue(const URigVMPin::FPinOverride& InOverride, bo
 			{
 				if (SubPin->GetTypeIndex() != ElementType)
 				{
-					CachedDefaultValue = TEXT("()"); 
-					return CachedDefaultValue.GetValue();
+					Cache = TEXT("()"); 
+					return Cache.GetValue();
 				}
 				FString ElementDefaultValue = SubPin->GetDefaultValue(InOverride, bAdaptValueForPinType);
 				if (SubPin->IsStringType())
@@ -949,11 +952,11 @@ FString URigVMPin::GetDefaultValue(const URigVMPin::FPinOverride& InOverride, bo
 			}
 			if (ElementDefaultValues.Num() == 0)
 			{
-				CachedDefaultValue = TEXT("()"); 
-				return CachedDefaultValue.GetValue();
+				Cache = TEXT("()"); 
+				return Cache.GetValue();
 			}
-			CachedDefaultValue = FString::Printf(TEXT("(%s)"), *FString::Join(ElementDefaultValues, TEXT(",")));
-			return CachedDefaultValue.GetValue();
+			Cache = FString::Printf(TEXT("(%s)"), *FString::Join(ElementDefaultValues, TEXT(",")));
+			return Cache.GetValue();
 		}
 
 		return DefaultValue.IsEmpty() ? TEXT("()") : DefaultValue;
@@ -1002,8 +1005,8 @@ FString URigVMPin::GetDefaultValue(const URigVMPin::FPinOverride& InOverride, bo
 					URigVMController::OverrideDefaultValueMember(SubPin->GetName(), MemberDefaultValue, FinalDefaultValue);
 				}
 
-				CachedDefaultValue = !FinalDefaultValue.IsEmpty() ? FinalDefaultValue : EmptyStructDefaultValue;
-				return CachedDefaultValue.GetValue();
+				Cache = !FinalDefaultValue.IsEmpty() ? FinalDefaultValue : EmptyStructDefaultValue;
+				return Cache.GetValue();
 			}
 		}
 
@@ -1024,8 +1027,8 @@ FString URigVMPin::GetDefaultValue(const URigVMPin::FPinOverride& InOverride, bo
 					const FQuat Quat = FQuat::MakeFromRotator(Rotator);
 					FString AdaptedDefaultValue;	
 					TBaseStructure<FQuat>::Get()->ExportText(AdaptedDefaultValue, &Quat, &Quat, nullptr, PPF_None, nullptr);
-					CachedDefaultValue = AdaptedDefaultValue;
-					return CachedDefaultValue.GetValue();
+					Cache = AdaptedDefaultValue;
+					return Cache.GetValue();
 				}
 			}
 		}
@@ -1042,8 +1045,8 @@ FString URigVMPin::GetDefaultValue(const URigVMPin::FPinOverride& InOverride, bo
 		// This default value is used during URigVMCompiler::FindOrAddRegister(...)
 		// Thus in this block, we have to return something like 0.0 instead of empty string
 
-		CachedDefaultValue = URigVMController::GetPinInitialDefaultValue(this);
-		return CachedDefaultValue.GetValue();
+		Cache = URigVMController::GetPinInitialDefaultValue(this);
+		return Cache.GetValue();
 	}
 
 	return DefaultValue;
diff --git a/Engine/Plugins/Runtime/RigVM/Source/RigVMDeveloper/Public/RigVMModel/RigVMPin.h b/Engine/Plugins/Runtime/RigVM/Source/RigVMDeveloper/Public/RigVMModel/RigVMPin.h
--- a/Engine/Plugins/Runtime/RigVM/Source/RigVMDeveloper/Public/RigVMModel/RigVMPin.h
+++ b/Engine/Plugins/Runtime/RigVM/Source/RigVMDeveloper/Public/RigVMModel/RigVMPin.h
@@ -622,6 +622,7 @@ private:
 	mutable TRigVMTypeIndex LastKnownTypeIndex;
 	mutable TRigVMModelCachedValue<URigVMPin, bool> CachedIsStringType;
 	mutable TRigVMModelCachedValue<URigVMPin, FString> CachedDefaultValue;
+	mutable TRigVMModelCachedValue<URigVMPin, FString> CachedAdaptedDefaultValue;
 	mutable TRigVMModelCachedValue<URigVMPin, uint32> CachedCPPTypeObjectHash;
 	mutable TRigVMModelCachedValue<URigVMPin, bool> CachedShowInDetailsPanelOnly;
 	mutable TRigVMModelCachedValue<URigVMPin, FString> CachedPinPath;
