RigVM: Fix reset default array for nodes

#rnx

diff --git a/Engine/Plugins/Runtime/RigVM/Source/RigVMEditor/Private/Editor/RigVMGraphDetailCustomization.cpp b/Engine/Plugins/Runtime/RigVM/Source/RigVMEditor/Private/Editor/RigVMGraphDetailCustomization.cpp
--- a/Engine/Plugins/Runtime/RigVM/Source/RigVMEditor/Private/Editor/RigVMGraphDetailCustomization.cpp
+++ b/Engine/Plugins/Runtime/RigVM/Source/RigVMEditor/Private/Editor/RigVMGraphDetailCustomization.cpp
@@ -2317,12 +2317,6 @@ void FRigVMWrappedNodeDetailCustomization::CustomizeDetails(IDetailLayoutBuilder
 			})
 		);
 
-		// disable the reset to default arrow
-		if(!CVarRigVMEnablePinOverrides.GetValueOnAnyThread())
-		{
-			ResetToDefault = FResetToDefaultOverride::Create(true, FSimpleDelegate(), true);
-		}
-
 		static const FSlateFontInfo NameFont = FAppStyle::GetFontStyle( TEXT("PropertyWindow.NormalFont") );
 
 		FText LabelOverride;
