Control Rig: Add autocompile toolbar button to legacy editor
#rnx

diff --git a/Engine/Plugins/Runtime/RigVM/Source/RigVMEditor/Private/Editor/RigVMEditor.cpp b/Engine/Plugins/Runtime/RigVM/Source/RigVMEditor/Private/Editor/RigVMEditor.cpp
--- a/Engine/Plugins/Runtime/RigVM/Source/RigVMEditor/Private/Editor/RigVMEditor.cpp
+++ b/Engine/Plugins/Runtime/RigVM/Source/RigVMEditor/Private/Editor/RigVMEditor.cpp
@@ -3619,6 +3619,7 @@ void FRigVMEditorBase::FillToolbar(FToolBarBuilder& ToolbarBuilder, bool bEndSec
 			FSlateIcon(FAppStyle::GetAppStyleSetName(), "LevelEditor.Recompile"),
 			true);
 
+		AddAutoCompileWidget(ToolbarBuilder);
 		AddSelectedDebugObjectWidget(ToolbarBuilder);
 
 		ToolbarBuilder.AddSeparator();
diff --git a/Engine/Plugins/Runtime/RigVM/Source/RigVMEditor/Private/Editor/RigVMLegacyEditor.cpp b/Engine/Plugins/Runtime/RigVM/Source/RigVMEditor/Private/Editor/RigVMLegacyEditor.cpp
--- a/Engine/Plugins/Runtime/RigVM/Source/RigVMEditor/Private/Editor/RigVMLegacyEditor.cpp
+++ b/Engine/Plugins/Runtime/RigVM/Source/RigVMEditor/Private/Editor/RigVMLegacyEditor.cpp
@@ -6,6 +6,7 @@
 #include "SMyBlueprint.h"
 #include "RigVMEditorModule.h"
 #include "SBlueprintEditorSelectedDebugObjectWidget.h"
+#include "RigVMEditorCommands.h"
 
 struct FRigVMEditorZoomLevelsContainer;
 
@@ -66,6 +67,12 @@ void FRigVMLegacyEditor::AddSelectedDebugObjectWidget(FToolBarBuilder& ToolbarBu
 	ToolbarBuilder.AddWidget(SNew(SBlueprintEditorSelectedDebugObjectWidget, SharedThis(this)));
 }
 
+void FRigVMLegacyEditor::AddAutoCompileWidget(FToolBarBuilder& ToolbarBuilder)
+{
+	ToolbarBuilder.AddToolBarButton(FRigVMEditorCommands::Get().AutoCompileGraph,
+		NAME_None, TAttribute<FText>(), TAttribute<FText>(), FSlateIcon(FRigVMEditorStyle::Get().GetStyleSetName(), "RigVM.AutoCompileGraph"));
+}
+
 void FRigVMLegacyEditor::Tick(float DeltaTime)
 {
 	FBlueprintEditor::Tick(DeltaTime);
diff --git a/Engine/Plugins/Runtime/RigVM/Source/RigVMEditor/Public/Editor/RigVMEditor.h b/Engine/Plugins/Runtime/RigVM/Source/RigVMEditor/Public/Editor/RigVMEditor.h
--- a/Engine/Plugins/Runtime/RigVM/Source/RigVMEditor/Public/Editor/RigVMEditor.h
+++ b/Engine/Plugins/Runtime/RigVM/Source/RigVMEditor/Public/Editor/RigVMEditor.h
@@ -244,6 +244,7 @@ protected:
 	virtual void AddCompileWidget(FToolBarBuilder& ToolbarBuilder) = 0;
 	virtual void AddSettingsAndDefaultWidget(FToolBarBuilder& ToolbarBuilder) = 0;
 	virtual void AddSelectedDebugObjectWidget(FToolBarBuilder& ToolbarBuilder) = 0;
+	virtual void AddAutoCompileWidget(FToolBarBuilder& ToolbarBuilder) = 0;
 };
 
 class RIGVMEDITOR_API FRigVMEditorBase : public IRigVMEditor
diff --git a/Engine/Plugins/Runtime/RigVM/Source/RigVMEditor/Public/Editor/RigVMLegacyEditor.h b/Engine/Plugins/Runtime/RigVM/Source/RigVMEditor/Public/Editor/RigVMLegacyEditor.h
--- a/Engine/Plugins/Runtime/RigVM/Source/RigVMEditor/Public/Editor/RigVMLegacyEditor.h
+++ b/Engine/Plugins/Runtime/RigVM/Source/RigVMEditor/Public/Editor/RigVMLegacyEditor.h
@@ -118,6 +118,7 @@ protected:
 	virtual void AddCompileWidget(FToolBarBuilder& ToolbarBuilder) override;
 	virtual void AddSettingsAndDefaultWidget(FToolBarBuilder& ToolbarBuilder) override {}
 	virtual void AddSelectedDebugObjectWidget(FToolBarBuilder& ToolbarBuilder) override;
+	virtual void AddAutoCompileWidget(FToolBarBuilder& ToolbarBuilder) override;
 	virtual FString GetCustomDebugObjectLabel(UObject* ObjectBeingDebugged) const override { return FRigVMEditorBase::GetCustomDebugObjectLabel(ObjectBeingDebugged); }
 public:
 	virtual void SummonSearchUI(bool bSetFindWithinBlueprint, FString NewSearchTerms, bool bSelectFirstResult) override { FBlueprintEditor::SummonSearchUI(bSetFindWithinBlueprint, NewSearchTerms, bSelectFirstResult); }
diff --git a/Engine/Plugins/Runtime/RigVM/Source/RigVMEditor/Public/Editor/RigVMNewEditor.h b/Engine/Plugins/Runtime/RigVM/Source/RigVMEditor/Public/Editor/RigVMNewEditor.h
--- a/Engine/Plugins/Runtime/RigVM/Source/RigVMEditor/Public/Editor/RigVMNewEditor.h
+++ b/Engine/Plugins/Runtime/RigVM/Source/RigVMEditor/Public/Editor/RigVMNewEditor.h
@@ -124,6 +124,7 @@ public:
 	virtual void AddCompileWidget(FToolBarBuilder& ToolbarBuilder) override;
 	virtual void AddSettingsAndDefaultWidget(FToolBarBuilder& ToolbarBuilder) override;
 	virtual void AddSelectedDebugObjectWidget(FToolBarBuilder& ToolbarBuilder) override;
+	virtual void AddAutoCompileWidget(FToolBarBuilder& ToolbarBuilder) override {}
 	virtual void Compile() { FRigVMEditorBase::Compile(); }
 	virtual void OnCreateGraphEditorCommands(TSharedPtr<FUICommandList> GraphEditorCommandsList) { FRigVMEditorBase::OnCreateGraphEditorCommands(GraphEditorCommandsList); }
 	virtual bool ShouldOpenGraphByDefault() const { return FRigVMEditorBase::ShouldOpenGraphByDefault(); }
