[AutoRTFM] Use a transactionally-safe CriticalSection in UNameSpacedUserData

This is necessary because this code is reachable from a Verse transaction.

diff --git a/Engine/Plugins/Runtime/RigVM/Source/RigVM/Private/RigVMCore/RigVMAssetUserData.cpp b/Engine/Plugins/Runtime/RigVM/Source/RigVM/Private/RigVMCore/RigVMAssetUserData.cpp
--- a/Engine/Plugins/Runtime/RigVM/Source/RigVM/Private/RigVMCore/RigVMAssetUserData.cpp
+++ b/Engine/Plugins/Runtime/RigVM/Source/RigVM/Private/RigVMCore/RigVMAssetUserData.cpp
@@ -79,7 +79,7 @@ void UNameSpacedUserData::Serialize(FArchive& Ar)
 
 const UNameSpacedUserData::FUserData* UNameSpacedUserData::GetUserData(const FString& InPath, FString* OutErrorMessage) const
 {
-	FScopeLock Lock(&CacheLock);
+	UE::TScopeLock Lock(CacheLock);
 	if(FUserData** UserDataInCache = CachedUserData.Find(InPath))
 	{
 		if(*UserDataInCache)
@@ -92,7 +92,7 @@ const UNameSpacedUserData::FUserData* UNameSpacedUserData::GetUserData(const FSt
 
 const TArray<const UNameSpacedUserData::FUserData*>& UNameSpacedUserData::GetUserDataArray(const FString& InParentPath, FString* OutErrorMessage) const
 {
-	FScopeLock Lock(&CacheLock);
+	UE::TScopeLock Lock(CacheLock);
 	if(const TArray<const FUserData*>* UserDataArrayInCache = CachedUserDataArray.Find(InParentPath))
 	{
 		return *UserDataArrayInCache;
@@ -193,7 +193,7 @@ const UNameSpacedUserData::FUserData* UNameSpacedUserData::GetUserDataWithinStru
 
 const TArray<const UNameSpacedUserData::FUserData*>& UNameSpacedUserData::GetUserDataArrayWithinStruct(UStruct* InStruct, const uint8* InMemory, const FString& InPath, FString* OutErrorMessage) const
 {
-	FScopeLock Lock(&CacheLock);
+	UE::TScopeLock Lock(CacheLock);
 
 	TArray<const FUserData*> UserDataArray;
 
@@ -241,7 +241,7 @@ const TArray<const UNameSpacedUserData::FUserData*>& UNameSpacedUserData::GetUse
 {
 	check(!InPath.IsEmpty());
 	
-	FScopeLock Lock(&CacheLock);
+	UE::TScopeLock Lock(CacheLock);
 
 	TArray<const FUserData*> UserDataArray;
 
@@ -294,7 +294,7 @@ const FProperty* UNameSpacedUserData::FindPropertyByName(const UStruct* InStruct
 
 void UNameSpacedUserData::InvalidateCache() const
 {
-	FScopeLock Lock(&CacheLock);
+	UE::TScopeLock Lock(CacheLock);
 	for(const TPair<FString, FUserData*>& Pair : CachedUserData)
 	{
 		delete(Pair.Value);
@@ -306,7 +306,7 @@ void UNameSpacedUserData::InvalidateCache() const
 const UNameSpacedUserData::FUserData* UNameSpacedUserData::StoreCacheForUserData(const FUserData& InUserData) const
 {
 	check(InUserData.IsValid());
-	FScopeLock Lock(&CacheLock);
+	UE::TScopeLock Lock(CacheLock);
 
 	if(const FUserData* const* UserDataInCache = CachedUserData.Find(InUserData.GetPath()))
 	{
@@ -329,7 +329,7 @@ const TArray<const UNameSpacedUserData::FUserData*>& UNameSpacedUserData::StoreC
 		return EmptyUserDatas;
 	}
 	
-	FScopeLock Lock(&CacheLock);
+	UE::TScopeLock Lock(CacheLock);
 
 	TArray<const FUserData*>& UserDataArrayInCache = CachedUserDataArray.FindOrAdd(InPath);
 	if(UserDataArrayInCache.IsEmpty())
diff --git a/Engine/Plugins/Runtime/RigVM/Source/RigVM/Public/RigVMCore/RigVMAssetUserData.h b/Engine/Plugins/Runtime/RigVM/Source/RigVM/Public/RigVMCore/RigVMAssetUserData.h
--- a/Engine/Plugins/Runtime/RigVM/Source/RigVM/Public/RigVMCore/RigVMAssetUserData.h
+++ b/Engine/Plugins/Runtime/RigVM/Source/RigVM/Public/RigVMCore/RigVMAssetUserData.h
@@ -5,6 +5,7 @@
 #include "CoreMinimal.h"
 #include "Engine/AssetUserData.h"
 #include "Engine/DataAsset.h"
+#include "Misc/TransactionallySafeCriticalSection.h"
 #include "RigVMAssetUserData.generated.h"
 
 /**
@@ -171,7 +172,7 @@ protected:
 
 	mutable TMap<FString, FUserData*> CachedUserData;
 	mutable TMap<FString, TArray<const FUserData*>> CachedUserDataArray;
-	mutable FCriticalSection CacheLock;
+	mutable FTransactionallySafeCriticalSection CacheLock;
 
 	static inline constexpr TCHAR PathNotFoundFormat[] = TEXT("UserData path '%s' could not be found.");
 	static inline constexpr TCHAR InvalidMemoryFormat[] = TEXT("UserData at path '%s' does not have any data backing it up.");
