RigVM: Avoid crash when using set container type

#rnx

diff --git a/Engine/Plugins/Runtime/RigVM/Source/RigVM/Private/RigVMCore/RigVMPropertyPath.cpp b/Engine/Plugins/Runtime/RigVM/Source/RigVM/Private/RigVMCore/RigVMPropertyPath.cpp
--- a/Engine/Plugins/Runtime/RigVM/Source/RigVM/Private/RigVMCore/RigVMPropertyPath.cpp
+++ b/Engine/Plugins/Runtime/RigVM/Source/RigVM/Private/RigVMCore/RigVMPropertyPath.cpp
@@ -80,7 +80,8 @@ FRigVMPropertyPath::FRigVMPropertyPath(const FProperty* InProperty, const FStrin
 		}
 		else if(CastField<FSetProperty>(Property))
 		{
-			checkNoEntry();
+			Segments.Empty();
+			return;
 		}
 		else
 		{
