RigVM: Fix crash when using map with anything other than FName as a key

#rnx

diff --git a/Engine/Plugins/Runtime/RigVM/Source/RigVM/Private/RigVMCore/RigVMPropertyPath.cpp b/Engine/Plugins/Runtime/RigVM/Source/RigVM/Private/RigVMCore/RigVMPropertyPath.cpp
--- a/Engine/Plugins/Runtime/RigVM/Source/RigVM/Private/RigVMCore/RigVMPropertyPath.cpp
+++ b/Engine/Plugins/Runtime/RigVM/Source/RigVM/Private/RigVMCore/RigVMPropertyPath.cpp
@@ -69,8 +69,6 @@ FRigVMPropertyPath::FRigVMPropertyPath(const FProperty* InProperty, const FStrin
 		}
 		else if(const FMapProperty* MapProperty = CastField<FMapProperty>(Property))
 		{
-			check(MapProperty->KeyProp->IsA<FNameProperty>());
-			
 			Segment.Type = ERigVMPropertyPathSegmentType::MapValue;
 			const FProperty* ValueProperty = MapProperty->ValueProp;
 			Segment.Name = *PathSegment;
