Control Rig: Optimize _switch_ rig unit to cache cases
#rnx

diff --git a/Engine/Plugins/Runtime/RigVM/Source/RigVM/Private/RigVMFunctions/RigVMDispatch_Switch.cpp b/Engine/Plugins/Runtime/RigVM/Source/RigVM/Private/RigVMFunctions/RigVMDispatch_Switch.cpp
--- a/Engine/Plugins/Runtime/RigVM/Source/RigVM/Private/RigVMFunctions/RigVMDispatch_Switch.cpp
+++ b/Engine/Plugins/Runtime/RigVM/Source/RigVM/Private/RigVMFunctions/RigVMDispatch_Switch.cpp
@@ -147,12 +147,28 @@ void FRigVMDispatch_SwitchInt32::Execute(FRigVMExtendedExecuteContext& InContext
 
 FName FRigVMDispatch_SwitchInt32::GetCaseName(int32 InIndex)
 {
-	return FRigVMBranchInfo::GetFixedArrayLabel(CasesName, *FString::FromInt(InIndex));
+	static TMap<int32, FName> NamesFromInt;
+	if (FName* Found = NamesFromInt.Find(InIndex))
+	{
+		return *Found;
+	}
+	
+	const FName Result = FRigVMBranchInfo::GetFixedArrayLabel(CasesName, *FString::FromInt(InIndex));
+	NamesFromInt.Add(InIndex, Result);
+	return Result;
 }
 
 FName FRigVMDispatch_SwitchInt32::GetCaseDisplayName(int32 InIndex)
 {
+	static TMap<int32, FName> NamesFromInt;
+    if (FName* Found = NamesFromInt.Find(InIndex))
+    {
+    	return *Found;
+    }
+    
 	static constexpr TCHAR Format[] = TEXT("Case %d");
-	return *FString::Printf(Format, InIndex);
+	const FName Result = *FString::Printf(Format, InIndex);
+    NamesFromInt.Add(InIndex, Result);
+    return Result;
 }
 
