ControlRig : Converted HandleExecutionReachedExit  event to function callback, to avoid worker thread contention

diff --git a/Engine/Plugins/Animation/ControlRig/Source/ControlRig/Private/ControlRig.cpp b/Engine/Plugins/Animation/ControlRig/Source/ControlRig/Private/ControlRig.cpp
--- a/Engine/Plugins/Animation/ControlRig/Source/ControlRig/Private/ControlRig.cpp
+++ b/Engine/Plugins/Animation/ControlRig/Source/ControlRig/Private/ControlRig.cpp
@@ -923,6 +923,11 @@ bool UControlRig::Execute(const FName& InEventName)
 	
 	FRigVMExtendedExecuteContext& ExtendedExecuteContext = GetRigVMExtendedExecuteContext();
 
+	TGuardValue<TFunction<void(const FName& InEventName)>> OnExecutionReachedExitCallback(ExtendedExecuteContext.OnExecutionReachedExitCallback, TFunction<void(const FName & InEventName)>([this](const FName& InEventName)
+		{
+			HandleExecutionReachedExit(InEventName);
+		}));
+
 	FControlRigExecuteContext& PublicContext = ExtendedExecuteContext.GetPublicDataSafe<FControlRigExecuteContext>();
 
 #if WITH_EDITOR
@@ -3688,8 +3693,6 @@ void UControlRig::PostInitInstance(URigVMHost* InCDO)
 	
 	Context.SetContextPublicDataStruct(FControlRigExecuteContext::StaticStruct());
 
-	Context.ExecutionReachedExit().RemoveAll(this);
-	Context.ExecutionReachedExit().AddUObject(this, &UControlRig::HandleExecutionReachedExit);
 	UpdateVMSettings();
 
 	// set up the hierarchy
diff --git a/Engine/Plugins/Runtime/RigVM/Source/RigVM/Private/RigVMCore/RigVM.cpp b/Engine/Plugins/Runtime/RigVM/Source/RigVM/Private/RigVMCore/RigVM.cpp
--- a/Engine/Plugins/Runtime/RigVM/Source/RigVM/Private/RigVMCore/RigVM.cpp
+++ b/Engine/Plugins/Runtime/RigVM/Source/RigVM/Private/RigVMCore/RigVM.cpp
@@ -2014,7 +2014,12 @@ ERigVMExecuteResult URigVM::ExecuteInstructions(FRigVMExtendedExecuteContext& Co
 				if(Context.bCurrentlyRunningRootEntry)
 				{
 					StopProfiling(Context);
-					Context.ExecutionReachedExit().Broadcast(Context.CurrentEntryName);
+
+					if (Context.OnExecutionReachedExitCallback)
+					{
+						Context.OnExecutionReachedExitCallback(Context.CurrentEntryName);
+					}
+
 #if WITH_EDITOR					
 					if (RigVMDebugInfo != nullptr)
 					{
diff --git a/Engine/Plugins/Runtime/RigVM/Source/RigVM/Private/RigVMCore/RigVMExecuteContext.cpp b/Engine/Plugins/Runtime/RigVM/Source/RigVM/Private/RigVMCore/RigVMExecuteContext.cpp
--- a/Engine/Plugins/Runtime/RigVM/Source/RigVM/Private/RigVMCore/RigVMExecuteContext.cpp
+++ b/Engine/Plugins/Runtime/RigVM/Source/RigVM/Private/RigVMCore/RigVMExecuteContext.cpp
@@ -95,7 +95,8 @@ void FRigVMExtendedExecuteContext::Reset()
 	DebugMemoryStorage = FRigVMMemoryStorageStruct();
 
 	CurrentVMMemory.Reset();
-	ExecutionReachedExit().Clear();
+	
+	OnExecutionReachedExitCallback.Reset();
 
 	CachedMemoryHandles.Reset();
 
diff --git a/Engine/Plugins/Runtime/RigVM/Source/RigVM/Private/RigVMHost.cpp b/Engine/Plugins/Runtime/RigVM/Source/RigVM/Private/RigVMHost.cpp
--- a/Engine/Plugins/Runtime/RigVM/Source/RigVM/Private/RigVMHost.cpp
+++ b/Engine/Plugins/Runtime/RigVM/Source/RigVM/Private/RigVMHost.cpp
@@ -210,11 +210,6 @@ void URigVMHost::BeginDestroy()
 	DebugInfo.Reset();
 #endif
 
-	if (RigVMExtendedExecuteContext != nullptr)
-	{
-		RigVMExtendedExecuteContext->ExecutionReachedExit().RemoveAll(this);
-	}
-
 #if WITH_EDITORONLY_DATA
 	if (VMSnapshotBeforeExecution)
 	{
@@ -525,6 +520,11 @@ bool URigVMHost::Execute(const FName& InEventName)
 	
 	FRigVMExtendedExecuteContext& ExtendedExecuteContext = GetRigVMExtendedExecuteContext();
 
+	TGuardValue<TFunction<void(const FName& InEventName)>> OnExecutionReachedExitCallback(ExtendedExecuteContext.OnExecutionReachedExitCallback, TFunction<void(const FName & InEventName)>([this](const FName& InEventName)
+		{
+			HandleExecutionReachedExit(InEventName);
+		}));
+
 	FRigVMExecuteContext& PublicContext = ExtendedExecuteContext.GetPublicData<>();
 	PublicContext.SetDeltaTime(DeltaTime);
 	PublicContext.SetAbsoluteTime(AbsoluteTime);
@@ -680,7 +680,7 @@ bool URigVMHost::Execute_Internal(const FName& InEventName)
 	}
 
 	FRigVMExtendedExecuteContext& ExtendedExecuteContext = GetRigVMExtendedExecuteContext();
-	
+
 	static constexpr TCHAR InvalidatedVMFormat[] = TEXT("%s: Invalidated VM - aborting execution.");
 	if(VM->IsNativized())
 	{
@@ -989,7 +989,7 @@ void URigVMHost::SwapVMToNativizedIfRequired(UClass* InNativizedClass)
 			VM->Rename(nullptr, GetTransientPackage(), REN_DoNotDirty | REN_DontCreateRedirectors | REN_NonTransactional);
 			VM->MarkAsGarbage();
 			VM = NewObject<URigVM>(this, InNativizedClass, TEXT("RigVM_NVMB"), PreviousFlags);
-			GetRigVMExtendedExecuteContext().ExecutionReachedExit().AddUObject(this, &URigVMHost::HandleExecutionReachedExit);
+
 #if UE_RIGVM_PROFILE_EXECUTE_UNITS_NUM
 			ProfilingRunsLeft = 0;
 			AccumulatedCycles = 0;
@@ -1079,9 +1079,6 @@ void URigVMHost::PostInitInstance(URigVMHost* InCDO)
 
 	ExtendedExecuteContext.SetContextPublicDataStruct(GetPublicContextStruct());
 
-	ExtendedExecuteContext.ExecutionReachedExit().RemoveAll(this);
-	ExtendedExecuteContext.ExecutionReachedExit().AddUObject(this, &URigVMHost::HandleExecutionReachedExit);
-
 #if WITH_EDITOR
 	ExtendedExecuteContext.GetPublicData<>().SetLog(RigVMLog); // may be nullptr
 #endif
diff --git a/Engine/Plugins/Runtime/RigVM/Source/RigVM/Public/RigVMCore/RigVMExecuteContext.h b/Engine/Plugins/Runtime/RigVM/Source/RigVM/Public/RigVMCore/RigVMExecuteContext.h
--- a/Engine/Plugins/Runtime/RigVM/Source/RigVM/Public/RigVMCore/RigVMExecuteContext.h
+++ b/Engine/Plugins/Runtime/RigVM/Source/RigVM/Public/RigVMCore/RigVMExecuteContext.h
@@ -865,8 +865,18 @@ struct FRigVMExtendedExecuteContext
 
 	/** Bindable event for external objects to be notified when the VM reaches an Exit Operation */
 	DECLARE_EVENT_OneParam(URigVM, FExecutionReachedExitEvent, const FName&);
+	UE_DEPRECATED(5.6, "OnExecutionReachedExit has been deprecated, please use OnExecutionReachedExitCallback.")
 	FExecutionReachedExitEvent OnExecutionReachedExit;
-	FExecutionReachedExitEvent& ExecutionReachedExit() { return OnExecutionReachedExit; }
+	UE_DEPRECATED(5.6, "OnExecutionReachedExit has been deprecated, please use OnExecutionReachedExitCallback.")
+	FExecutionReachedExitEvent& ExecutionReachedExit() 
+	{
+		PRAGMA_DISABLE_DEPRECATION_WARNINGS
+		return OnExecutionReachedExit; 
+		PRAGMA_ENABLE_DEPRECATION_WARNINGS
+	}
+
+	/** Callback for external objects to be notified when the VM reaches an Exit Operation */
+	TFunction<void(const FName& EventName)> OnExecutionReachedExitCallback;
 
 	TArray<FRigVMInstructionSetExecuteState> LazyBranchExecuteState;
 	TArray<FRigVMExternalVariableRuntimeData> ExternalVariableRuntimeData;
diff --git a/Engine/Plugins/Runtime/RigVM/Source/RigVM/Public/RigVMCore/RigVMNativized.h b/Engine/Plugins/Runtime/RigVM/Source/RigVM/Public/RigVMCore/RigVMNativized.h
--- a/Engine/Plugins/Runtime/RigVM/Source/RigVM/Public/RigVMCore/RigVMNativized.h
+++ b/Engine/Plugins/Runtime/RigVM/Source/RigVM/Public/RigVMCore/RigVMNativized.h
@@ -185,7 +185,10 @@ protected:
 	{
 		if(Context.EntriesBeingExecuted.Num() == 1)
 		{
-			Context.ExecutionReachedExit().Broadcast(Context.GetPublicData<>().GetEventName());
+			if (Context.OnExecutionReachedExitCallback)
+			{
+				Context.OnExecutionReachedExitCallback(Context.GetPublicData<>().GetEventName());
+			}
 		}
 		Context.GetPublicDataSafe<>().NumExecutions++;
 	}
