Control Rig: Fix array values not being passed by reference to a function reference when fold assignments is off
#rnx

diff --git a/Engine/Plugins/Runtime/RigVM/Source/RigVM/Private/RigVMCore/RigVM.cpp b/Engine/Plugins/Runtime/RigVM/Source/RigVM/Private/RigVMCore/RigVM.cpp
--- a/Engine/Plugins/Runtime/RigVM/Source/RigVM/Private/RigVMCore/RigVM.cpp
+++ b/Engine/Plugins/Runtime/RigVM/Source/RigVM/Private/RigVMCore/RigVM.cpp
@@ -2260,7 +2260,7 @@ ERigVMExecuteResult URigVM::ExecuteInstructions(FRigVMExtendedExecuteContext& Co
 					CurrentWorkMemory.Add(FString::Printf(TEXT(" -- %s"), *Line));
 				}
 			}
-			for (const FRigVMExternalVariable& ExternalVariable : ExternalVariables)
+			for (const FRigVMExternalVariable& ExternalVariable : GetExternalVariables(Context))
 			{
 				if (ExternalVariable.Memory == nullptr)
 				{
diff --git a/Engine/Plugins/Runtime/RigVM/Source/RigVM/Public/RigVMHost.h b/Engine/Plugins/Runtime/RigVM/Source/RigVM/Public/RigVMHost.h
--- a/Engine/Plugins/Runtime/RigVM/Source/RigVM/Public/RigVMHost.h
+++ b/Engine/Plugins/Runtime/RigVM/Source/RigVM/Public/RigVMHost.h
@@ -289,7 +289,7 @@ public:
 
 #if UE_RIGVM_DEBUG_EXECUTION
 	bool bDebugExecutionEnabled = false;
-	const FString GetDebugExecutionString();
+	UE_API const FString GetDebugExecutionString();
 #endif
 
 	/** Provide access to the ExtendedExecuteContext */
diff --git a/Engine/Plugins/Runtime/RigVM/Source/RigVMDeveloper/Private/RigVMModel/RigVMClient.cpp b/Engine/Plugins/Runtime/RigVM/Source/RigVMDeveloper/Private/RigVMModel/RigVMClient.cpp
--- a/Engine/Plugins/Runtime/RigVM/Source/RigVMDeveloper/Private/RigVMModel/RigVMClient.cpp
+++ b/Engine/Plugins/Runtime/RigVM/Source/RigVMDeveloper/Private/RigVMModel/RigVMClient.cpp
@@ -1315,6 +1315,14 @@ void FRigVMClient::PatchFunctionReferencesOnLoad()
 				{
 					FunctionReferenceNode->ReferencedFunctionHeader.LibraryPointer.SetLibraryNodePath(FunctionReferenceNode->ReferencedFunctionHeader.LibraryPointer.LibraryNode_DEPRECATED.ToString());
 				}
+
+				for (URigVMPin* Pin : FunctionReferenceNode->GetPins())
+				{
+					if (Pin->IsArray())
+					{
+						Pin->bIsDynamicArray = true;
+					}
+				}
 			}
 
 			if (URigVMCollapseNode* CollapseNode = Cast<URigVMCollapseNode>(Node))
@@ -1339,6 +1347,35 @@ void FRigVMClient::PatchFunctionsOnLoad(IRigVMGraphFunctionHost* FunctionHost, T
 
 			for (URigVMLibraryNode* LibraryNode : FunctionLibrary->GetFunctions())
 			{
+				// Fix array pins to be marked as dynamic array
+				for (URigVMPin* Pin : LibraryNode->GetPins())
+				{
+					if (Pin->IsArray())
+					{
+						Pin->bIsDynamicArray = true;
+					}
+				}
+				if (URigVMFunctionEntryNode* EntryNode = LibraryNode->GetEntryNode())
+				{
+					for (URigVMPin* Pin : EntryNode->GetPins())
+					{
+						if (Pin->IsArray())
+						{
+							Pin->bIsDynamicArray = true;
+						}
+					}
+				}
+				if (URigVMFunctionReturnNode* ReturnNode = LibraryNode->GetReturnNode())
+				{
+					for (URigVMPin* Pin : ReturnNode->GetPins())
+					{
+						if (Pin->IsArray())
+						{
+							Pin->bIsDynamicArray = true;
+						}
+					}
+				}
+
 				bool bIsPublic = FunctionLibrary->IsFunctionPublic(LibraryNode->GetFName());
 				if (!bIsPublic)
 				{
diff --git a/Engine/Plugins/Runtime/RigVM/Source/RigVMDeveloper/Private/RigVMModel/RigVMController.cpp b/Engine/Plugins/Runtime/RigVM/Source/RigVMDeveloper/Private/RigVMModel/RigVMController.cpp
--- a/Engine/Plugins/Runtime/RigVM/Source/RigVMDeveloper/Private/RigVMModel/RigVMController.cpp
+++ b/Engine/Plugins/Runtime/RigVM/Source/RigVMDeveloper/Private/RigVMModel/RigVMController.cpp
@@ -662,7 +662,11 @@ FRigVMPinInfoArray::FRigVMPinInfoArray(const URigVMNode* InNode, URigVMControlle
 		{
 			InController->ReportErrorf( TEXT("Invalid pin type %s for %s in %s"), *Pin->GetCPPType(), *Pin->GetPathName(), *InNode->GetPackage()->GetPathName());
 		}
-		(void)AddPin(InController, INDEX_NONE, Pin->GetFName(), Pin->GetDirection(), Pin->GetTypeIndex(), DefaultValue, DefaultValueType, nullptr, InPreviousPinInfos, bAddSubPins);
+		int32 Index = AddPin(InController, INDEX_NONE, Pin->GetFName(), Pin->GetDirection(), Pin->GetTypeIndex(), DefaultValue, DefaultValueType, nullptr, InPreviousPinInfos, bAddSubPins);
+		if (InNode->IsA<URigVMLibraryNode>())
+		{
+			Pins[Index].bIsDynamicArray = Pins[Index].bIsArray;
+		}
 	}
 }
 
@@ -684,7 +688,8 @@ FRigVMPinInfoArray::FRigVMPinInfoArray(const FRigVMGraphFunctionHeader& Function
 		{
 			DefaultValueType = PinInfo->DefaultValueType;
 		}
-		(void)AddPin(InController, INDEX_NONE, FunctionArgument.Name, FunctionArgument.Direction, TypeIndex, FunctionArgument.DefaultValue, DefaultValueType, nullptr, InPreviousPinInfos, true);
+		int32 Index = AddPin(InController, INDEX_NONE, FunctionArgument.Name, FunctionArgument.Direction, TypeIndex, FunctionArgument.DefaultValue, DefaultValueType, nullptr, InPreviousPinInfos, true);
+		Pins[Index].bIsDynamicArray = Pins[Index].bIsArray;
 	}
 }
 
@@ -13121,6 +13126,7 @@ FName URigVMController::AddExposedPin(const FName& InPinName, ERigVMPinDirection
 	Pin->CPPTypeObjectPath = InCPPTypeObjectPath;
 	Pin->bIsConstant = false;
 	Pin->Direction = InDirection;
+	Pin->bIsDynamicArray = RigVMTypeUtils::IsArrayType(CPPType);
 	AddNodePin(LibraryNode, Pin);
 
 	if (Pin->IsStruct())
diff --git a/Engine/Plugins/Runtime/RigVM/Source/RigVMDeveloper/Public/RigVMModel/RigVMPin.h b/Engine/Plugins/Runtime/RigVM/Source/RigVMDeveloper/Public/RigVMModel/RigVMPin.h
--- a/Engine/Plugins/Runtime/RigVM/Source/RigVMDeveloper/Public/RigVMModel/RigVMPin.h
+++ b/Engine/Plugins/Runtime/RigVM/Source/RigVMDeveloper/Public/RigVMModel/RigVMPin.h
@@ -643,6 +643,7 @@ private:
 	friend struct FRigVMSetPinCategoryAction;
 	friend class URigVMLibraryNode;
 	friend class URigVMEdGraphNode;
+	friend struct FRigVMClient;
 };
 
 class RIGVMDEVELOPER_API FRigVMPinDefaultValueImportErrorContext : public FOutputDevice
