RigVM: Various changes to enable AnimNext runtime refactor

Node subtitles are now deferred to the node class and are overrideable
Added the ability to specify a subclass to URigVMController::AddVariableNode
Exposed URigVMController::RewireLinks to subclasses
Exposed FRigVMMemoryStorageStruct::GeneratePropertyBagDescriptor publically (for use in type conversions)

diff --git a/Engine/Plugins/Runtime/RigVM/Source/RigVM/Public/RigVMCore/RigVMMemoryStorageStruct.h b/Engine/Plugins/Runtime/RigVM/Source/RigVM/Public/RigVMCore/RigVMMemoryStorageStruct.h
--- a/Engine/Plugins/Runtime/RigVM/Source/RigVM/Public/RigVMCore/RigVMMemoryStorageStruct.h
+++ b/Engine/Plugins/Runtime/RigVM/Source/RigVM/Public/RigVMCore/RigVMMemoryStorageStruct.h
@@ -374,6 +374,9 @@ struct FRigVMMemoryStorageStruct : public FInstancedPropertyBag
 		return GetPropertyBagStruct();
 	}
 
+	// Helper function to generate a property bag desc from a FRigVMPropertyDescription
+	static UE_API FPropertyBagPropertyDesc GeneratePropertyBagDescriptor(const FRigVMPropertyDescription& RigVMDescriptor);
+
 protected:
 	ERigVMMemoryType MemoryType = ERigVMMemoryType::Invalid;
 
@@ -396,7 +399,6 @@ protected:
 
 	UE_API void SetDefaultValues(const TArray<FRigVMPropertyDescription>& InPropertyDescriptions);
 
-	static UE_API FPropertyBagPropertyDesc GeneratePropertyBagDescriptor(const FRigVMPropertyDescription& RigVMDescriptor);
 	static UE_API bool GetPropertyTypeDataFromVMDescriptor(const FRigVMPropertyDescription& RigVMDescriptor, EPropertyBagPropertyType& OutBagPropertyType, FPropertyBagContainerTypes& OutBagContainerTypes);
 };
 
diff --git a/Engine/Plugins/Runtime/RigVM/Source/RigVMDeveloper/Private/EdGraph/RigVMEdGraphNode.cpp b/Engine/Plugins/Runtime/RigVM/Source/RigVMDeveloper/Private/EdGraph/RigVMEdGraphNode.cpp
--- a/Engine/Plugins/Runtime/RigVM/Source/RigVMDeveloper/Private/EdGraph/RigVMEdGraphNode.cpp
+++ b/Engine/Plugins/Runtime/RigVM/Source/RigVMDeveloper/Private/EdGraph/RigVMEdGraphNode.cpp
@@ -113,24 +113,9 @@ FText URigVMEdGraphNode::GetNodeTitle(ENodeTitleType::Type TitleType) const
 
 			else if(URigVMVariableNode* VariableNode = Cast<URigVMVariableNode>(ModelNode))
 			{
-				if(VariableNode->IsLocalVariable())
-				{
-					static const FString LocalVariableString = TEXT("Local Variable");
-					const FString DefaultValue = VariableNode->GetVariableDescription().DefaultValue;
-					if(DefaultValue.IsEmpty())
-					{
-						SubTitle = LocalVariableString;
-					}
-					else
-					{
-						SubTitle = FString::Printf(TEXT("%s\nDefault %s"), *LocalVariableString, *DefaultValue);
-					}
-				}
-				else if (VariableNode->IsInputArgument())
-				{
-					SubTitle = TEXT("Input parameter");
-				}
-				else
+				SubTitle = VariableNode->GetNodeSubTitle();
+
+				if (SubTitle.Len() == 0)
 				{
 					if(UBlueprint* Blueprint = GetBlueprint())
 					{
diff --git a/Engine/Plugins/Runtime/RigVM/Source/RigVMDeveloper/Private/RigVMModel/Nodes/RigVMVariableNode.cpp b/Engine/Plugins/Runtime/RigVM/Source/RigVMDeveloper/Private/RigVMModel/Nodes/RigVMVariableNode.cpp
--- a/Engine/Plugins/Runtime/RigVM/Source/RigVMDeveloper/Private/RigVMModel/Nodes/RigVMVariableNode.cpp
+++ b/Engine/Plugins/Runtime/RigVM/Source/RigVMDeveloper/Private/RigVMModel/Nodes/RigVMVariableNode.cpp
@@ -137,3 +137,27 @@ URigVMPin* URigVMVariableNode::GetValuePin() const
 	return FindPin(ValueName);
 }
 
+FString URigVMVariableNode::GetNodeSubTitle() const
+{
+	FString SubTitle;
+	
+	if(IsLocalVariable())
+	{
+		static const FString LocalVariableString = TEXT("Local Variable");
+		const FString DefaultValue = GetVariableDescription().DefaultValue;
+		if(DefaultValue.IsEmpty())
+		{
+			SubTitle = LocalVariableString;
+		}
+		else
+		{
+			SubTitle = FString::Printf(TEXT("%s\nDefault %s"), *LocalVariableString, *DefaultValue);
+		}
+	}
+	else if (IsInputArgument())
+	{
+		SubTitle = TEXT("Input parameter");
+	}
+
+	return SubTitle;
+}
diff --git a/Engine/Plugins/Runtime/RigVM/Source/RigVMDeveloper/Private/RigVMModel/RigVMController.cpp b/Engine/Plugins/Runtime/RigVM/Source/RigVMDeveloper/Private/RigVMModel/RigVMController.cpp
--- a/Engine/Plugins/Runtime/RigVM/Source/RigVMDeveloper/Private/RigVMModel/RigVMController.cpp
+++ b/Engine/Plugins/Runtime/RigVM/Source/RigVMDeveloper/Private/RigVMModel/RigVMController.cpp
@@ -1874,6 +1874,11 @@ bool URigVMController::SetUnitNodeDefaults(URigVMUnitNode* InNode, const FRigStr
 }
 
 URigVMVariableNode* URigVMController::AddVariableNode(const FName& InVariableName, const FString& InCPPType, UObject* InCPPTypeObject, bool bIsGetter, const FString& InDefaultValue, const FVector2D& InPosition, const FString& InNodeName, bool bSetupUndoRedo, bool bPrintPythonCommand)
+{
+	return AddVariableNode(InVariableName, URigVMVariableNode::StaticClass(), InCPPType, InCPPTypeObject, bIsGetter, InDefaultValue, InPosition, InNodeName, bSetupUndoRedo, bPrintPythonCommand);
+}
+
+URigVMVariableNode* URigVMController::AddVariableNode(const FName& InVariableName, TSubclassOf<URigVMVariableNode> InVariableNodeClass, const FString& InCPPType, UObject* InCPPTypeObject, bool bIsGetter, const FString& InDefaultValue, const FVector2D& InPosition, const FString& InNodeName, bool bSetupUndoRedo, bool bPrintPythonCommand)
 {
 	if (!IsValidGraph())
 	{
@@ -1969,7 +1974,7 @@ URigVMVariableNode* URigVMController::AddVariableNode(const FName& InVariableNam
 	}
 
 	FString Name = GetSchema()->GetValidNodeName(Graph, InNodeName.IsEmpty() ? FString(TEXT("VariableNode")) : InNodeName);
-	URigVMVariableNode* Node = NewObject<URigVMVariableNode>(Graph, *Name);
+	URigVMVariableNode* Node = NewObject<URigVMVariableNode>(Graph, InVariableNodeClass.Get(), *Name);
 	Node->Position = InPosition;
 
 	if (!bIsGetter)
diff --git a/Engine/Plugins/Runtime/RigVM/Source/RigVMDeveloper/Public/RigVMModel/Nodes/RigVMVariableNode.h b/Engine/Plugins/Runtime/RigVM/Source/RigVMDeveloper/Public/RigVMModel/Nodes/RigVMVariableNode.h
--- a/Engine/Plugins/Runtime/RigVM/Source/RigVMDeveloper/Public/RigVMModel/Nodes/RigVMVariableNode.h
+++ b/Engine/Plugins/Runtime/RigVM/Source/RigVMDeveloper/Public/RigVMModel/Nodes/RigVMVariableNode.h
@@ -72,6 +72,9 @@ public:
 	// Get the value pin
 	URigVMPin* GetValuePin() const;
 
+	// Get the subtitle of the node
+	virtual FString GetNodeSubTitle() const;
+
 private:
 
 	static const inline TCHAR* VariableName = TEXT("Variable");
diff --git a/Engine/Plugins/Runtime/RigVM/Source/RigVMDeveloper/Public/RigVMModel/RigVMController.h b/Engine/Plugins/Runtime/RigVM/Source/RigVMDeveloper/Public/RigVMModel/RigVMController.h
--- a/Engine/Plugins/Runtime/RigVM/Source/RigVMDeveloper/Public/RigVMModel/RigVMController.h
+++ b/Engine/Plugins/Runtime/RigVM/Source/RigVMDeveloper/Public/RigVMModel/RigVMController.h
@@ -388,6 +388,7 @@ public:
 	// This causes a NodeAdded modified event.
 	UFUNCTION(BlueprintCallable, Category = RigVMController)
 	URigVMVariableNode* AddVariableNode(const FName& InVariableName, const FString& InCPPType, UObject* InCPPTypeObject, bool bIsGetter, const FString& InDefaultValue, const FVector2D& InPosition = FVector2D::ZeroVector, const FString& InNodeName = TEXT(""), bool bSetupUndoRedo = true, bool bPrintPythonCommand = false);
+	URigVMVariableNode* AddVariableNode(const FName& InVariableName, TSubclassOf<URigVMVariableNode> InVariableNodeClass, const FString& InCPPType, UObject* InCPPTypeObject, bool bIsGetter, const FString& InDefaultValue, const FVector2D& InPosition = FVector2D::ZeroVector, const FString& InNodeName = TEXT(""), bool bSetupUndoRedo = true, bool bPrintPythonCommand = false);
 
 	// Adds a Variable Node to the edited Graph given a struct object path name.
 	// Variables represent local work state for the function and
@@ -1339,6 +1340,9 @@ protected:
 	bool IsValidSchema() const;
 	bool IsGraphEditable() const;
 
+#if WITH_EDITOR
+	void RewireLinks(URigVMPin* OldPin, URigVMPin* NewPin, bool bAsInput, bool bSetupUndoRedo, TArray<URigVMLink*> InLinks = TArray<URigVMLink*>());
+#endif
 
 	virtual const UClass* GetNodeClassForTemplate(FRigVMTemplate* InTemplate) const;
 
@@ -1420,10 +1424,6 @@ private:
 	bool ChangePinType(URigVMPin* InPin, const FString& InCPPType,UObject* InCPPTypeObject, bool bSetupUndoRedo, bool bSetupOrphanPins = true, bool bBreakLinks = true, bool bRemoveSubPins = true, bool bInitializeDefaultValue = true);
 	bool ChangePinType(URigVMPin* InPin, TRigVMTypeIndex InTypeIndex, bool bSetupUndoRedo, bool bSetupOrphanPins = true, bool bBreakLinks = true, bool bRemoveSubPins = true, bool bInitializeDefaultValue = true);
 
-#if WITH_EDITOR
-	void RewireLinks(URigVMPin* OldPin, URigVMPin* NewPin, bool bAsInput, bool bSetupUndoRedo, TArray<URigVMLink*> InLinks = TArray<URigVMLink*>());
-#endif
-
 	bool RenameObject(UObject* InObjectToRename, const TCHAR* InNewName, UObject* InNewOuter = nullptr, ERenameFlags InFlags = REN_None) const;
 	void DestroyObject(UObject* InObjectToDestroy) const ;
 	static URigVMPin* MakeExecutePin(URigVMNode* InNode, const FName& InName);
