Control Rig: Fix graph explorer not showing function/variable tooltip
#rnx

diff --git a/Engine/Plugins/Runtime/RigVM/Source/RigVMEditor/Private/Widgets/SRigVMEditorGraphExplorer.cpp b/Engine/Plugins/Runtime/RigVM/Source/RigVMEditor/Private/Widgets/SRigVMEditorGraphExplorer.cpp
--- a/Engine/Plugins/Runtime/RigVM/Source/RigVMEditor/Private/Widgets/SRigVMEditorGraphExplorer.cpp
+++ b/Engine/Plugins/Runtime/RigVM/Source/RigVMEditor/Private/Widgets/SRigVMEditorGraphExplorer.cpp
@@ -421,6 +421,8 @@ void SRigVMEditorGraphExplorer::CreateWidgets()
 	Delegates.OnCanRenameVariable = FRigVMGraphExplorer_OnCanRenameVariable::CreateSP(this, &SRigVMEditorGraphExplorer::OnCanRenameVariable);
 	Delegates.OnSetFunctionCategory = FRigVMGraphExplorer_OnSetFunctionCategory::CreateSP(this, &SRigVMEditorGraphExplorer::OnSetFunctionCategory);
 	Delegates.OnGetFunctionCategory = FRigVMGraphExplorer_OnGetFunctionCategory::CreateSP(this, &SRigVMEditorGraphExplorer::OnGetFunctionCategory);
+	Delegates.OnGetFunctionTooltip = FRigVMGraphExplorer_OnGetFunctionTooltip::CreateSP(this, &SRigVMEditorGraphExplorer::OnGetFunctionTooltip);
+	Delegates.OnGetVariableTooltip = FRigVMGraphExplorer_OnGetFunctionTooltip::CreateSP(this, &SRigVMEditorGraphExplorer::OnGetVariableTooltip);
 	Delegates.OnSetVariableCategory = FRigVMGraphExplorer_OnSetVariableCategory::CreateSP(this, &SRigVMEditorGraphExplorer::OnSetVariableCategory);
 	Delegates.OnGetVariableCategory = FRigVMGraphExplorer_OnGetVariableCategory::CreateSP(this, &SRigVMEditorGraphExplorer::OnGetVariableCategory);
 	Delegates.OnRequestContextMenu = FRigVMGraphExplorer_OnRequestContextMenu::CreateSP(this, &SRigVMEditorGraphExplorer::OnContextMenuOpening);
@@ -992,6 +994,38 @@ FString SRigVMEditorGraphExplorer::OnGetFunctionCategory(const FString& InFuncti
 	return FString();
 }
 
+FText SRigVMEditorGraphExplorer::OnGetFunctionTooltip(const FString& InFunctionPath) const
+{
+	TSharedPtr<IRigVMEditor> Editor = RigVMEditor.Pin();
+	if (!Editor.IsValid())
+	{
+		return FText();
+	}
+
+	const URigVMEdGraphSchema* Schema = GetDefault<URigVMEdGraphSchema>();
+	if (URigVMFunctionLibrary* Library = Editor->GetRigVMBlueprint()->GetLocalFunctionLibrary())
+	{
+		if (URigVMLibraryNode* Function = Library->FindFunction(*InFunctionPath))
+		{
+			return Function->GetToolTipText();
+		}
+	}
+	return FText();
+}
+
+FText SRigVMEditorGraphExplorer::OnGetVariableTooltip(const FString& InVariable) const
+{
+	TSharedPtr<IRigVMEditor> Editor = RigVMEditor.Pin();
+	if (!Editor.IsValid())
+	{
+		return FText();
+	}
+
+	FString Result;
+	FBlueprintEditorUtils::GetBlueprintVariableMetaData(Editor->GetRigVMBlueprint(), *InVariable, nullptr, TEXT("tooltip"), Result);
+	return FText::FromString(Result);
+}
+
 bool SRigVMEditorGraphExplorer::OnSetVariableCategory(const FString& InVariable, const FString& InCategory)
 {
 	TSharedPtr<IRigVMEditor> Editor = RigVMEditor.Pin();
diff --git a/Engine/Plugins/Runtime/RigVM/Source/RigVMEditor/Private/Widgets/SRigVMEditorGraphExplorerTreeView.cpp b/Engine/Plugins/Runtime/RigVM/Source/RigVMEditor/Private/Widgets/SRigVMEditorGraphExplorerTreeView.cpp
--- a/Engine/Plugins/Runtime/RigVM/Source/RigVMEditor/Private/Widgets/SRigVMEditorGraphExplorerTreeView.cpp
+++ b/Engine/Plugins/Runtime/RigVM/Source/RigVMEditor/Private/Widgets/SRigVMEditorGraphExplorerTreeView.cpp
@@ -308,13 +308,21 @@ FText SRigVMEditorGraphExplorerItem::GetItemTooltip() const
 	switch (Element->Key.Type)
 	{
 		case ERigVMExplorerElementType::Graph:
-			{
-				return Delegates.GetGraphTooltip(Element->Key.Name);
-			}
+		{
+			return Delegates.GetGraphTooltip(Element->Key.Name);
+		}
+		case ERigVMExplorerElementType::Function:
+		{
+			return Delegates.GetFunctionTooltip(Element->Key.Name);
+		}
+		case ERigVMExplorerElementType::Variable:
+		{
+			return Delegates.GetVariableTooltip(Element->Key.Name);
+		}
 		case ERigVMExplorerElementType::Event:
-			{
-				return Delegates.GetEventDisplayName(Element->Key.Name);
-			}
+		{
+			return Delegates.GetEventDisplayName(Element->Key.Name);
+		}
 	}
 	
 	return Tooltip;
diff --git a/Engine/Plugins/Runtime/RigVM/Source/RigVMEditor/Public/Widgets/SRigVMEditorGraphExplorer.h b/Engine/Plugins/Runtime/RigVM/Source/RigVMEditor/Public/Widgets/SRigVMEditorGraphExplorer.h
--- a/Engine/Plugins/Runtime/RigVM/Source/RigVMEditor/Public/Widgets/SRigVMEditorGraphExplorer.h
+++ b/Engine/Plugins/Runtime/RigVM/Source/RigVMEditor/Public/Widgets/SRigVMEditorGraphExplorer.h
@@ -127,6 +127,8 @@ private:
 	void OnFunctionDoubleClicked(const FString& InFunctionPath);
 	bool OnSetFunctionCategory(const FString& InFunctionPath, const FString& InCategory);
 	FString OnGetFunctionCategory(const FString& InPath) const;
+	FText OnGetFunctionTooltip(const FString& InFunctionPath) const;
+	FText OnGetVariableTooltip(const FString& InVariable) const;
 	bool OnSetVariableCategory(const FString& InVariable, const FString& InCategory);
 	FString OnGetVariableCategory(const FString& InVariable) const;
 	FEdGraphPinType OnGetVariablePinType(const FRigVMExplorerElementKey& InVariable);
diff --git a/Engine/Plugins/Runtime/RigVM/Source/RigVMEditor/Public/Widgets/SRigVMEditorGraphExplorerTreeView.h b/Engine/Plugins/Runtime/RigVM/Source/RigVMEditor/Public/Widgets/SRigVMEditorGraphExplorerTreeView.h
--- a/Engine/Plugins/Runtime/RigVM/Source/RigVMEditor/Public/Widgets/SRigVMEditorGraphExplorerTreeView.h
+++ b/Engine/Plugins/Runtime/RigVM/Source/RigVMEditor/Public/Widgets/SRigVMEditorGraphExplorerTreeView.h
@@ -83,6 +83,8 @@ DECLARE_DELEGATE_RetVal_ThreeParams(bool, FRigVMGraphExplorer_OnCanRenameFunctio
 DECLARE_DELEGATE_RetVal_ThreeParams(bool, FRigVMGraphExplorer_OnCanRenameVariable, const FRigVMExplorerElementKey&, const FString&, FText&);
 DECLARE_DELEGATE_RetVal_TwoParams(bool, FRigVMGraphExplorer_OnSetFunctionCategory, const FString&, const FString&);
 DECLARE_DELEGATE_RetVal_OneParam(FString, FRigVMGraphExplorer_OnGetFunctionCategory, const FString&);
+DECLARE_DELEGATE_RetVal_OneParam(FText, FRigVMGraphExplorer_OnGetFunctionTooltip, const FString&);
+DECLARE_DELEGATE_RetVal_OneParam(FText, FRigVMGraphExplorer_OnGetVariableTooltip, const FString&);
 DECLARE_DELEGATE_RetVal_TwoParams(bool, FRigVMGraphExplorer_OnSetVariableCategory, const FString&, const FString&);
 DECLARE_DELEGATE_RetVal_OneParam(FString, FRigVMGraphExplorer_OnGetVariableCategory, const FString&);
 DECLARE_DELEGATE_RetVal(TSharedPtr<SWidget>, FRigVMGraphExplorer_OnRequestContextMenu);
@@ -127,6 +129,8 @@ public:
 	FRigVMGraphExplorer_OnCanRenameVariable OnCanRenameVariable;
 	FRigVMGraphExplorer_OnSetFunctionCategory OnSetFunctionCategory;
 	FRigVMGraphExplorer_OnGetFunctionCategory OnGetFunctionCategory;
+	FRigVMGraphExplorer_OnGetFunctionTooltip OnGetFunctionTooltip;
+	FRigVMGraphExplorer_OnGetVariableTooltip OnGetVariableTooltip;
 	FRigVMGraphExplorer_OnSetVariableCategory OnSetVariableCategory;
 	FRigVMGraphExplorer_OnGetVariableCategory OnGetVariableCategory;
 	FRigVMGraphExplorer_OnGetVariablePinType OnGetVariablePinType;
@@ -398,6 +402,24 @@ public:
 		return FString(); 
 	}
 
+	FText GetFunctionTooltip(const FString& InPath) const
+	{
+		if (OnGetFunctionTooltip.IsBound())
+		{
+			return OnGetFunctionTooltip.Execute(InPath);
+		}
+		return FText();
+	}
+
+	FText GetVariableTooltip(const FString& InPath) const
+	{
+		if (OnGetVariableTooltip.IsBound())
+		{
+			return OnGetVariableTooltip.Execute(InPath);
+		}
+		return FText();
+	}
+
 	void SetVariableCategory(const FString& InPath, const FString& InCategory) const
 	{
 		if (OnSetVariableCategory.IsBound())
