RigVM: Make sure to pull in injection nodes from older assets.

#rnx

diff --git a/Engine/Plugins/Runtime/RigVM/Source/RigVMDeveloper/Private/RigVMModel/RigVMController.cpp b/Engine/Plugins/Runtime/RigVM/Source/RigVMDeveloper/Private/RigVMModel/RigVMController.cpp
--- a/Engine/Plugins/Runtime/RigVM/Source/RigVMDeveloper/Private/RigVMModel/RigVMController.cpp
+++ b/Engine/Plugins/Runtime/RigVM/Source/RigVMDeveloper/Private/RigVMModel/RigVMController.cpp
@@ -19491,6 +19491,13 @@ URigVMController::FPinState URigVMController::GetPinState(URigVMPin* InPin, bool
 		// no longer exists during ApplyPinState.
 		for(URigVMInjectionInfo* InjectionInfo : State.InjectionInfos)
 		{
+#if WITH_EDITOR
+			// Copy over the old unit node if the content is from prior to that property being deprecated.
+			if (!InjectionInfo->Node)
+			{
+				InjectionInfo->Node = InjectionInfo->UnitNode_DEPRECATED;
+			}
+#endif
 			RenameObject(InjectionInfo, nullptr, InPin->GetGraph());
 		}
 	}
diff --git a/Engine/Plugins/Runtime/RigVM/Source/RigVMDeveloper/Private/RigVMModel/RigVMPin.cpp b/Engine/Plugins/Runtime/RigVM/Source/RigVMDeveloper/Private/RigVMModel/RigVMPin.cpp
--- a/Engine/Plugins/Runtime/RigVM/Source/RigVMDeveloper/Private/RigVMModel/RigVMPin.cpp
+++ b/Engine/Plugins/Runtime/RigVM/Source/RigVMDeveloper/Private/RigVMModel/RigVMPin.cpp
@@ -49,6 +49,12 @@ URigVMInjectionInfo::FWeakInfo URigVMInjectionInfo::GetWeakInfo() const
 	FWeakInfo Info;
 	Info.bInjectedAsInput = bInjectedAsInput;
 	Info.Node = Node;
+#if WITH_EDITOR
+	if (!Node)
+	{
+		Info.Node = UnitNode_DEPRECATED;
+	}
+#endif
 	Info.InputPinName = InputPin != nullptr ? InputPin->GetFName() : NAME_None;
 	Info.OutputPinName = OutputPin != nullptr ? OutputPin->GetFName() : NAME_None;
 	return Info;
