Updated headers using UnrealCodeFixup to make sure dllstorage is on methods/staticvars instead of types. Part 4/n

diff --git a/Engine/Plugins/Runtime/RigVM/Source/RigVMDeveloper/Public/EdGraph/RigVMEdGraphNode.h b/Engine/Plugins/Runtime/RigVM/Source/RigVMDeveloper/Public/EdGraph/RigVMEdGraphNode.h
--- a/Engine/Plugins/Runtime/RigVM/Source/RigVMDeveloper/Public/EdGraph/RigVMEdGraphNode.h
+++ b/Engine/Plugins/Runtime/RigVM/Source/RigVMDeveloper/Public/EdGraph/RigVMEdGraphNode.h
@@ -9,14 +9,16 @@
 #include "Blueprint/BlueprintExtension.h"
 #include "RigVMEdGraphNode.generated.h"
 
+#define UE_API RIGVMDEVELOPER_API
+
 class FBlueprintActionDatabaseRegistrar;
 class UEdGraph;
 struct FSlateIcon;
 class URigVMBlueprint;
 
 /** Base class for RigVM editor side nodes */
-UCLASS()
-class RIGVMDEVELOPER_API URigVMEdGraphNode : public UEdGraphNode
+UCLASS(MinimalAPI)
+class URigVMEdGraphNode : public UEdGraphNode
 {
 	GENERATED_BODY()
 
@@ -99,7 +101,7 @@ public:
 		bool IsValid() const { return InputPin != nullptr || OutputPin != nullptr; }
 	};
 
-	URigVMEdGraphNode();
+	UE_API URigVMEdGraphNode();
 
 	// UObject Interface.
 #if WITH_EDITOR
@@ -107,37 +109,37 @@ public:
 #endif
 	
 	// UEdGraphNode Interface.
-	virtual FText GetNodeTitle(ENodeTitleType::Type TitleType) const override;
-	virtual FLinearColor GetNodeTitleColor() const override;
-	virtual FLinearColor GetNodeBodyTintColor() const;
-	virtual bool ShowPaletteIconOnNode() const override;
-	virtual FSlateIcon GetIconAndTint(FLinearColor& OutColor) const override;
-	virtual void AllocateDefaultPins() override;
-	virtual void ReconstructNode() override;
-	virtual void ReconstructNode_Internal(bool bForce = false);
-	virtual void GetNodeContextMenuActions(class UToolMenu* Menu, class UGraphNodeContextMenuContext* Context) const override;
-	virtual void DestroyNode() override;
-	virtual void PinDefaultValueChanged(UEdGraphPin* Pin) override;
-	virtual FText GetTooltipText() const override;
-	virtual bool CanCreateUnderSpecifiedSchema(const UEdGraphSchema* InSchema) const override;
-	virtual void AutowireNewNode(UEdGraphPin* FromPin) override;	
+	UE_API virtual FText GetNodeTitle(ENodeTitleType::Type TitleType) const override;
+	UE_API virtual FLinearColor GetNodeTitleColor() const override;
+	UE_API virtual FLinearColor GetNodeBodyTintColor() const;
+	UE_API virtual bool ShowPaletteIconOnNode() const override;
+	UE_API virtual FSlateIcon GetIconAndTint(FLinearColor& OutColor) const override;
+	UE_API virtual void AllocateDefaultPins() override;
+	UE_API virtual void ReconstructNode() override;
+	UE_API virtual void ReconstructNode_Internal(bool bForce = false);
+	UE_API virtual void GetNodeContextMenuActions(class UToolMenu* Menu, class UGraphNodeContextMenuContext* Context) const override;
+	UE_API virtual void DestroyNode() override;
+	UE_API virtual void PinDefaultValueChanged(UEdGraphPin* Pin) override;
+	UE_API virtual FText GetTooltipText() const override;
+	UE_API virtual bool CanCreateUnderSpecifiedSchema(const UEdGraphSchema* InSchema) const override;
+	UE_API virtual void AutowireNewNode(UEdGraphPin* FromPin) override;	
 	virtual bool SupportsCommentBubble() const override { return false; }
-	virtual bool IsSelectedInEditor() const;
-	virtual bool ShouldDrawNodeAsControlPointOnly(int32& OutInputPinIndex, int32& OutOutputPinIndex) const override;
-	virtual void BeginDestroy() override;
+	UE_API virtual bool IsSelectedInEditor() const;
+	UE_API virtual bool ShouldDrawNodeAsControlPointOnly(int32& OutInputPinIndex, int32& OutOutputPinIndex) const override;
+	UE_API virtual void BeginDestroy() override;
 #if WITH_EDITOR
 #if WITH_RIGVMLEGACYEDITOR
-	virtual void AddPinSearchMetaDataInfo(const UEdGraphPin* Pin, TArray<FSearchTagDataPair>& OutTaggedMetaData) const override;
+	UE_API virtual void AddPinSearchMetaDataInfo(const UEdGraphPin* Pin, TArray<FSearchTagDataPair>& OutTaggedMetaData) const override;
 #else
-	void AddRigVMSearchMetaDataInfo(TArray<struct UBlueprintExtension::FSearchTagDataPair>& OutTaggedMetaData) const;
-	void AddRigVMPinSearchMetaDataInfo(const UEdGraphPin* Pin, TArray<UBlueprintExtension::FSearchTagDataPair>& OutTaggedMetaData) const;
+	UE_API void AddRigVMSearchMetaDataInfo(TArray<struct UBlueprintExtension::FSearchTagDataPair>& OutTaggedMetaData) const;
+	UE_API void AddRigVMPinSearchMetaDataInfo(const UEdGraphPin* Pin, TArray<UBlueprintExtension::FSearchTagDataPair>& OutTaggedMetaData) const;
 #endif
 #endif
-	virtual FString GetPinMetaData(FName InPinName, FName InKey) override;
+	UE_API virtual FString GetPinMetaData(FName InPinName, FName InKey) override;
 
-	virtual bool IsDeprecated() const override;
-	bool IsOutDated() const;
-	virtual FEdGraphNodeDeprecationResponse GetDeprecationResponse(EEdGraphNodeDeprecationType DeprecationType) const override;
+	UE_API virtual bool IsDeprecated() const override;
+	UE_API bool IsOutDated() const;
+	UE_API virtual FEdGraphNodeDeprecationResponse GetDeprecationResponse(EEdGraphNodeDeprecationType DeprecationType) const override;
 
 	/** Set the cached dimensions of this node */
 	void SetDimensions(const FVector2D& InDimensions) { Dimensions = InDimensions; }
@@ -146,69 +148,69 @@ public:
 	const FVector2D& GetDimensions() const { return Dimensions; }
 
 	/** Check a pin's expansion state */
-	bool IsPinExpanded(const FString& InPinPath);
+	UE_API bool IsPinExpanded(const FString& InPinPath);
 
 	/** Propagate pin defaults to underlying properties if they have changed */
-	void CopyPinDefaultsToModel(UEdGraphPin* Pin, bool bUndo = false, bool bPrintPythonCommand = false);
+	UE_API void CopyPinDefaultsToModel(UEdGraphPin* Pin, bool bUndo = false, bool bPrintPythonCommand = false);
 
 	/** Get the blueprint that this node is contained within */
-	URigVMBlueprint* GetBlueprint() const;
+	UE_API URigVMBlueprint* GetBlueprint() const;
 
 	/** Get the VM model this node lives within */
-	URigVMGraph* GetModel() const;
+	UE_API URigVMGraph* GetModel() const;
 
 	/** Get the blueprint that this node is contained within */
-	URigVMController* GetController() const;
+	UE_API URigVMController* GetController() const;
 
 	/** Get the VM node this is node is wrapping */
-	URigVMNode* GetModelNode() const;
+	UE_API URigVMNode* GetModelNode() const;
 
 	/** Get the VM node name this node is wrapping */
-	FName GetModelNodeName() const;
+	UE_API FName GetModelNodeName() const;
 
 	/** Get the VM node path this node is wrapping */
-	const FString& GetModelNodePath() const;
+	UE_API const FString& GetModelNodePath() const;
 
-	URigVMPin* GetModelPinFromPinPath(const FString& InPinPath) const;
+	UE_API URigVMPin* GetModelPinFromPinPath(const FString& InPinPath) const;
 
 	/** Add a new element to the aggregate node referred to by the property path */
-	void HandleAddAggregateElement(const FString& InNodePath);
+	UE_API void HandleAddAggregateElement(const FString& InNodePath);
 
 	/** Add a new array element to the array referred to by the property path */
-	void HandleAddArrayElement(FString InPinPath);
+	UE_API void HandleAddArrayElement(FString InPinPath);
 	
 	/** Clear the array referred to by the property path */
-	void HandleClearArray(FString InPinPath);
+	UE_API void HandleClearArray(FString InPinPath);
 
 	/** Remove the array element referred to by the property path */
-	void HandleRemoveArrayElement(FString InPinPath);
+	UE_API void HandleRemoveArrayElement(FString InPinPath);
 
 	/** Insert a new array element after the element referred to by the property path */
-	void HandleInsertArrayElement(FString InPinPath);
+	UE_API void HandleInsertArrayElement(FString InPinPath);
 
-	int32 GetInstructionIndex(bool bAsInput) const;
+	UE_API int32 GetInstructionIndex(bool bAsInput) const;
 
-	const FRigVMTemplate* GetTemplate() const;
+	UE_API const FRigVMTemplate* GetTemplate() const;
 
-	void ClearErrorInfo();
+	UE_API void ClearErrorInfo();
 
-	void AddErrorInfo(const EMessageSeverity::Type& InSeverity, const FString& InMessage);
+	UE_API void AddErrorInfo(const EMessageSeverity::Type& InSeverity, const FString& InMessage);
 
-	void SetErrorInfo(const EMessageSeverity::Type& InSeverity, const FString& InMessage);
+	UE_API void SetErrorInfo(const EMessageSeverity::Type& InSeverity, const FString& InMessage);
 
-	URigVMPin* FindModelPinFromGraphPin(const UEdGraphPin* InGraphPin) const;
-	UEdGraphPin* FindGraphPinFromModelPin(const URigVMPin* InModelPin, bool bAsInput) const;
-	UEdGraphPin* FindGraphPinFromCategory(const FString& InCategory, bool bAsInput) const;
+	UE_API URigVMPin* FindModelPinFromGraphPin(const UEdGraphPin* InGraphPin) const;
+	UE_API UEdGraphPin* FindGraphPinFromModelPin(const URigVMPin* InModelPin, bool bAsInput) const;
+	UE_API UEdGraphPin* FindGraphPinFromCategory(const FString& InCategory, bool bAsInput) const;
 
 	/// Synchronize the stored name/value/type on the graph pin with the value stored on the node. 
 	/// If the pin has sub-pins, the value update is done recursively.
-	void SynchronizeGraphPinNameWithModelPin(const URigVMPin* InModelPin, bool bNotify = true);
-	void SynchronizeGraphPinValueWithModelPin(const URigVMPin* InModelPin);
-	void SynchronizeGraphPinTypeWithModelPin(const URigVMPin* InModelPin);
-	void SynchronizeGraphPinExpansionWithModelPin(const URigVMPin* InModelPin);
+	UE_API void SynchronizeGraphPinNameWithModelPin(const URigVMPin* InModelPin, bool bNotify = true);
+	UE_API void SynchronizeGraphPinValueWithModelPin(const URigVMPin* InModelPin);
+	UE_API void SynchronizeGraphPinTypeWithModelPin(const URigVMPin* InModelPin);
+	UE_API void SynchronizeGraphPinExpansionWithModelPin(const URigVMPin* InModelPin);
 	
-	void SyncGraphNodeTitleWithModelNodeTitle();
-	void SyncGraphNodeNameWithModelNodeName(const URigVMNode* InModelNode);
+	UE_API void SyncGraphNodeTitleWithModelNodeTitle();
+	UE_API void SyncGraphNodeNameWithModelNodeName(const URigVMNode* InModelNode);
 
 	FNodeTitleDirtied& OnNodeTitleDirtied() { return NodeTitleDirtied; }
 	FNodePinsChanged& OnNodePinsChanged() { return NodePinsChanged; }
@@ -216,74 +218,74 @@ public:
 	FNodeBeginRemoval& OnNodeBeginRemoval() { return NodeBeginRemoval; }
 
 	/** Called when there's a drastic change in the pins */
-	bool ModelPinsChanged(bool bForce = false);
+	UE_API bool ModelPinsChanged(bool bForce = false);
 
 	/** Called when a model pin is added after the node creation */
-	bool ModelPinAdded(const URigVMPin* InModelPin);
+	UE_API bool ModelPinAdded(const URigVMPin* InModelPin);
 
 	/** Called when a model pin is being removed */
-	bool ModelPinRemoved(const URigVMPin* InModelPin);
+	UE_API bool ModelPinRemoved(const URigVMPin* InModelPin);
 
 	/** Returns true if this node is relying on the cast template */
-	bool DrawAsCompactNode() const;
+	UE_API bool DrawAsCompactNode() const;
 
 	/** Sets the node's content from an external client, used for preview nodes without a graph */
-	void SetModelNode(URigVMNode* InModelNode);
+	UE_API void SetModelNode(URigVMNode* InModelNode);
 
 	const TArray<URigVMPin*>& GetInputPins() const { return InputPins; }
 
 protected:
 
-	FLinearColor GetNodeProfilingColor() const;
-	FLinearColor GetNodeOpacityColor() const;
+	UE_API FLinearColor GetNodeProfilingColor() const;
+	UE_API FLinearColor GetNodeOpacityColor() const;
 
 	/** Helper function for AllocateDefaultPins */
-	void UpdatePinLists();
-	bool CreateGraphPinFromCategory(const FString& InCategory, EEdGraphPinDirection InDirection);
-	bool CreateGraphPinFromModelPin(const URigVMPin* InModelPin, EEdGraphPinDirection InDirection,  UEdGraphPin* InParentPin = nullptr);
-	void RemoveGraphSubPins(UEdGraphPin *InParentPin, const TArray<UEdGraphPin*>& InPinsToKeep = TArray<UEdGraphPin*>());
-	bool ModelPinAdded_Internal(const URigVMPin* InModelPin);
-	bool ModelPinRemoved_Internal(const URigVMPin* InModelPin);
-	bool CategoryPinAdded_Internal(const FString& InCategory, EEdGraphPinDirection InDirection);
-	bool CategoryPinRemoved_Internal(const FString& InCategory);
+	UE_API void UpdatePinLists();
+	UE_API bool CreateGraphPinFromCategory(const FString& InCategory, EEdGraphPinDirection InDirection);
+	UE_API bool CreateGraphPinFromModelPin(const URigVMPin* InModelPin, EEdGraphPinDirection InDirection,  UEdGraphPin* InParentPin = nullptr);
+	UE_API void RemoveGraphSubPins(UEdGraphPin *InParentPin, const TArray<UEdGraphPin*>& InPinsToKeep = TArray<UEdGraphPin*>());
+	UE_API bool ModelPinAdded_Internal(const URigVMPin* InModelPin);
+	UE_API bool ModelPinRemoved_Internal(const URigVMPin* InModelPin);
+	UE_API bool CategoryPinAdded_Internal(const FString& InCategory, EEdGraphPinDirection InDirection);
+	UE_API bool CategoryPinRemoved_Internal(const FString& InCategory);
 
 	/** Copies default values from underlying properties into pin defaults, for editing */
-	void SetupPinDefaultsFromModel(UEdGraphPin* Pin, const URigVMPin* InModelPin = nullptr);
+	UE_API void SetupPinDefaultsFromModel(UEdGraphPin* Pin, const URigVMPin* InModelPin = nullptr);
 
 	/** Recreate pins when we reconstruct this node */
-	virtual void ReallocatePinsDuringReconstruction(const TArray<UEdGraphPin*>& OldPins);
+	UE_API virtual void ReallocatePinsDuringReconstruction(const TArray<UEdGraphPin*>& OldPins);
 
 	/** Wire-up new pins given old pin wiring */
-	virtual void RewireOldPinsToNewPins(TArray<UEdGraphPin*>& InOldPins, TArray<UEdGraphPin*>& InNewPins);
+	UE_API virtual void RewireOldPinsToNewPins(TArray<UEdGraphPin*>& InOldPins, TArray<UEdGraphPin*>& InNewPins);
 
 	/** Handle anything post-reconstruction */
-	virtual void PostReconstructNode();
+	UE_API virtual void PostReconstructNode();
 
 	/** Something that could change our title has changed */
-	void InvalidateNodeTitle() const;
+	UE_API void InvalidateNodeTitle() const;
 
 	/** Destroy all pins in an array */
-	void DestroyPinList(TArray<UEdGraphPin*>& InPins);
+	UE_API void DestroyPinList(TArray<UEdGraphPin*>& InPins);
 
 	/** Sets the body + title color from a color provided by the model */
-	void SetColorFromModel(const FLinearColor& InColor);
+	UE_API void SetColorFromModel(const FLinearColor& InColor);
 
-	UClass* GetRigVMGeneratedClass() const;
-	UClass* GetRigVMSkeletonGeneratedClass() const;
+	UE_API UClass* GetRigVMGeneratedClass() const;
+	UE_API UClass* GetRigVMSkeletonGeneratedClass() const;
 
-	static FEdGraphPinType GetPinTypeForModelPin(const URigVMPin* InModelPin);
-	static FEdGraphPinType GetPinTypeForCategoryPin();
+	static UE_API FEdGraphPinType GetPinTypeForModelPin(const URigVMPin* InModelPin);
+	static UE_API FEdGraphPinType GetPinTypeForCategoryPin();
 
-	virtual void ConfigurePin(UEdGraphPin* EdGraphPin, const URigVMPin* ModelPin) const;
+	UE_API virtual void ConfigurePin(UEdGraphPin* EdGraphPin, const URigVMPin* ModelPin) const;
 private:
 
 	int32 GetNodeTopologyVersion() const { return NodeTopologyVersion; }
 	int32 NodeTopologyVersion;
 
-	TArray<URigVMPin*>& PinListForPin(const URigVMPin* InModelPin);
+	UE_API TArray<URigVMPin*>& PinListForPin(const URigVMPin* InModelPin);
 
 #if WITH_EDITOR
-	void UpdateProfilingSettings();
+	UE_API void UpdateProfilingSettings();
 #endif
 	
 	FLinearColor CachedTitleColor;
@@ -323,3 +325,5 @@ private:
 	friend class URigVMEdGraphIfNodeSpawner;
 	friend class URigVMEdGraphSelectNodeSpawner;
 };
+
+#undef UE_API
diff --git a/Engine/Plugins/Runtime/RigVM/Source/RigVMDeveloper/Public/EdGraph/RigVMEdGraphSchema.h b/Engine/Plugins/Runtime/RigVM/Source/RigVMDeveloper/Public/EdGraph/RigVMEdGraphSchema.h
--- a/Engine/Plugins/Runtime/RigVM/Source/RigVMDeveloper/Public/EdGraph/RigVMEdGraphSchema.h
+++ b/Engine/Plugins/Runtime/RigVM/Source/RigVMDeveloper/Public/EdGraph/RigVMEdGraphSchema.h
@@ -12,6 +12,8 @@
 
 #include "RigVMEdGraphSchema.generated.h"
 
+#define UE_API RIGVMDEVELOPER_API
+
 class URigVMBlueprint;
 class URigVMEdGraph;
 class URigVMEdGraphNode;
@@ -19,34 +21,34 @@ class URigVMEdGraphNode;
 /**
  *Const values for Find-in-Blueprints to tag searchable data
  */
-struct RIGVMDEVELOPER_API FRigVMSearchTags
+struct FRigVMSearchTags
 {
 	// /** Begin const values for Find-in-Blueprint */
 	
 	/** Name tag */
-	static const FText FiB_Name;
+	static UE_API const FText FiB_Name;
 	/** Class Name tag */
-	static const FText FiB_ClassName;
+	static UE_API const FText FiB_ClassName;
 	/** NodeGuid tag */
-	static const FText FiB_NodeGuid;
+	static UE_API const FText FiB_NodeGuid;
 	
 	// /** Pin type tags */
 	/** Pin Category tag */
-	static const FText FiB_PinCategory;
+	static UE_API const FText FiB_PinCategory;
 	/** Pin Sub-Category tag */
-	static const FText FiB_PinSubCategory;
+	static UE_API const FText FiB_PinSubCategory;
 	/** Pin Binding tag */
-	static const FText FiB_PinBinding;
+	static UE_API const FText FiB_PinBinding;
 	/** Pin object class tag */
-	static const FText FiB_ObjectClass;
+	static UE_API const FText FiB_ObjectClass;
 	/** Pin IsArray tag */
-	static const FText FiB_IsArray;
+	static UE_API const FText FiB_IsArray;
 	/** Glyph icon tag */
-	static const FText FiB_Glyph;
+	static UE_API const FText FiB_Glyph;
 	/** Style set the glyph belongs to */
-	static const FText FiB_GlyphStyleSet;
+	static UE_API const FText FiB_GlyphStyleSet;
 	/** Glyph icon color tag */
-	static const FText FiB_GlyphColor;
+	static UE_API const FText FiB_GlyphColor;
 	
 	// /** End const values for Find-in-Blueprint */
 };
@@ -81,34 +83,34 @@ struct FRigVMPinConnectionResponse
 
 /////////////////////////////////////////////////////
 // FRigVMLocalVariableNameValidator
-class RIGVMDEVELOPER_API FRigVMLocalVariableNameValidator : public FStringSetNameValidator
+class FRigVMLocalVariableNameValidator : public FStringSetNameValidator
 {
 
 public:
-	FRigVMLocalVariableNameValidator(const class UBlueprint* Blueprint, const URigVMGraph* Graph, FName InExistingName = NAME_None);
+	UE_API FRigVMLocalVariableNameValidator(const class UBlueprint* Blueprint, const URigVMGraph* Graph, FName InExistingName = NAME_None);
 
 	// Begin FNameValidatorInterface
-	virtual EValidatorResult IsValid(const FString& Name, bool bOriginal) override;
-	virtual EValidatorResult IsValid(const FName& Name, bool bOriginal) override;
+	UE_API virtual EValidatorResult IsValid(const FString& Name, bool bOriginal) override;
+	UE_API virtual EValidatorResult IsValid(const FName& Name, bool bOriginal) override;
 	// End FNameValidatorInterface
 };
 
 /////////////////////////////////////////////////////
 // FRigVMNameValidator
-class RIGVMDEVELOPER_API FRigVMNameValidator : public FStringSetNameValidator
+class FRigVMNameValidator : public FStringSetNameValidator
 {
 
 public:
-	FRigVMNameValidator(const class UBlueprint* Blueprint, const UStruct* ValidationScope, FName InExistingName = NAME_None);
+	UE_API FRigVMNameValidator(const class UBlueprint* Blueprint, const UStruct* ValidationScope, FName InExistingName = NAME_None);
 
 	// Begin FNameValidatorInterface
-	virtual EValidatorResult IsValid(const FString& Name, bool bOriginal) override;
-	virtual EValidatorResult IsValid(const FName& Name, bool bOriginal) override;
+	UE_API virtual EValidatorResult IsValid(const FString& Name, bool bOriginal) override;
+	UE_API virtual EValidatorResult IsValid(const FName& Name, bool bOriginal) override;
 	// End FNameValidatorInterface
 };
 
 USTRUCT()
-struct RIGVMDEVELOPER_API FRigVMEdGraphSchemaAction_LocalVar : public FEdGraphSchemaAction_BlueprintVariableBase
+struct FRigVMEdGraphSchemaAction_LocalVar : public FEdGraphSchemaAction_BlueprintVariableBase
 {
 	GENERATED_BODY()
 
@@ -136,21 +138,21 @@ public:
 		return InType == GetTypeId() || InType == FEdGraphSchemaAction_BlueprintVariableBase::StaticGetTypeId();
 	}
 
-	virtual FEdGraphPinType GetPinType() const override;
+	UE_API virtual FEdGraphPinType GetPinType() const override;
 
-	virtual void ChangeVariableType(const FEdGraphPinType& NewPinType) override;
+	UE_API virtual void ChangeVariableType(const FEdGraphPinType& NewPinType) override;
 
-	virtual void RenameVariable(const FName& NewName) override;
+	UE_API virtual void RenameVariable(const FName& NewName) override;
 
-	virtual bool IsValidName(const FName& NewName, FText& OutErrorMessage) const override;
+	UE_API virtual bool IsValidName(const FName& NewName, FText& OutErrorMessage) const override;
 
-	virtual void DeleteVariable() override;
+	UE_API virtual void DeleteVariable() override;
 
-	virtual bool IsVariableUsed() override;
+	UE_API virtual bool IsVariableUsed() override;
 };
 
 USTRUCT()
-struct RIGVMDEVELOPER_API FRigVMEdGraphSchemaAction_PromoteToVariable : public FEdGraphSchemaAction
+struct FRigVMEdGraphSchemaAction_PromoteToVariable : public FEdGraphSchemaAction
 {
 	GENERATED_BODY()
 
@@ -170,14 +172,14 @@ public:
 		, bLocalVariable(false)
 	{}
 
-	FRigVMEdGraphSchemaAction_PromoteToVariable(UEdGraphPin* InEdGraphPin, bool InLocalVariable);
+	UE_API FRigVMEdGraphSchemaAction_PromoteToVariable(UEdGraphPin* InEdGraphPin, bool InLocalVariable);
 
 	virtual bool IsA(const FName& InType) const override
 	{
 		return InType == GetTypeId();
 	}
 
-	virtual UEdGraphNode* PerformAction(UEdGraph* ParentGraph, UEdGraphPin* FromPin, const FVector2f& Location, bool bSelectNewNode) override;
+	UE_API virtual UEdGraphNode* PerformAction(UEdGraph* ParentGraph, UEdGraphPin* FromPin, const FVector2f& Location, bool bSelectNewNode) override;
 
 private:
 
@@ -186,7 +188,7 @@ private:
 };
 
 USTRUCT()
-struct RIGVMDEVELOPER_API FRigVMEdGraphSchemaAction_PromoteToExposedPin : public FEdGraphSchemaAction
+struct FRigVMEdGraphSchemaAction_PromoteToExposedPin : public FEdGraphSchemaAction
 {
 	GENERATED_BODY()
 
@@ -205,14 +207,14 @@ public:
 		, EdGraphPin(nullptr)
 	{}
 
-	FRigVMEdGraphSchemaAction_PromoteToExposedPin(UEdGraphPin* InEdGraphPin);
+	UE_API FRigVMEdGraphSchemaAction_PromoteToExposedPin(UEdGraphPin* InEdGraphPin);
 
 	virtual bool IsA(const FName& InType) const override
 	{
 		return InType == GetTypeId();
 	}
 
-	virtual UEdGraphNode* PerformAction(UEdGraph* ParentGraph, UEdGraphPin* FromPin, const FVector2f& Location, bool bSelectNewNode) override;
+	UE_API virtual UEdGraphNode* PerformAction(UEdGraph* ParentGraph, UEdGraphPin* FromPin, const FVector2f& Location, bool bSelectNewNode) override;
 
 private:
 
@@ -220,7 +222,7 @@ private:
 };
 
 USTRUCT()
-struct RIGVMDEVELOPER_API FRigVMEdGraphSchemaAction_Event : public FEdGraphSchemaAction
+struct FRigVMEdGraphSchemaAction_Event : public FEdGraphSchemaAction
 {
 	GENERATED_BODY()
 
@@ -238,16 +240,16 @@ public:
 		: FEdGraphSchemaAction()
 	{}
 
-	FRigVMEdGraphSchemaAction_Event(const FName& InEventName, const FString& InNodePath, const FText& InNodeCategory);
+	UE_API FRigVMEdGraphSchemaAction_Event(const FName& InEventName, const FString& InNodePath, const FText& InNodeCategory);
 
 	virtual bool IsA(const FName& InType) const override
 	{
 		return InType == GetTypeId() || Super::IsA(InType);
 	}
 	virtual bool IsParentable() const override { return true; }
-	virtual FReply OnDoubleClick(UBlueprint* InBlueprint) override;
+	UE_API virtual FReply OnDoubleClick(UBlueprint* InBlueprint) override;
 	virtual bool CanBeRenamed() const override { return false; }
-	virtual FSlateBrush const* GetPaletteIcon() const override;
+	UE_API virtual FSlateBrush const* GetPaletteIcon() const override;
 
 private:
 
@@ -255,18 +257,18 @@ private:
 };
 
 /** DragDropAction class for drag and dropping an item from the My Blueprints tree (e.g., variable or function) */
-class RIGVMDEVELOPER_API FRigVMFunctionDragDropAction : public FGraphSchemaActionDragDropAction
+class FRigVMFunctionDragDropAction : public FGraphSchemaActionDragDropAction
 {
 public:
 
 	DRAG_DROP_OPERATOR_TYPE(FRigVMFunctionDragDropAction, FGraphSchemaActionDragDropAction)
 
 	// FGraphEditorDragDropAction interface
-	virtual FReply DroppedOnPanel(const TSharedRef< class SWidget >& Panel, const FVector2f& ScreenPosition, const FVector2f& GraphPosition, UEdGraph& Graph) override;
-	virtual FReply DroppedOnPin(const FVector2f& ScreenPosition, const FVector2f& GraphPosition) override;
-	virtual FReply DroppedOnAction(TSharedRef<FEdGraphSchemaAction> Action) override;
-	virtual FReply DroppedOnCategory(FText Category) override;
-	virtual void HoverTargetChanged() override;
+	UE_API virtual FReply DroppedOnPanel(const TSharedRef< class SWidget >& Panel, const FVector2f& ScreenPosition, const FVector2f& GraphPosition, UEdGraph& Graph) override;
+	UE_API virtual FReply DroppedOnPin(const FVector2f& ScreenPosition, const FVector2f& GraphPosition) override;
+	UE_API virtual FReply DroppedOnAction(TSharedRef<FEdGraphSchemaAction> Action) override;
+	UE_API virtual FReply DroppedOnCategory(FText Category) override;
+	UE_API virtual void HoverTargetChanged() override;
 	// End of FGraphEditorDragDropAction
 
 	/** Set if operation is modified by alt */
@@ -278,9 +280,9 @@ public:
 protected:
 
 	/** Constructor */
-	FRigVMFunctionDragDropAction();
+	UE_API FRigVMFunctionDragDropAction();
 
-	static TSharedRef<FRigVMFunctionDragDropAction> New(TSharedPtr<FEdGraphSchemaAction> InAction, URigVMBlueprint* InRigBlueprint, URigVMEdGraph* InRigGraph);
+	static UE_API TSharedRef<FRigVMFunctionDragDropAction> New(TSharedPtr<FEdGraphSchemaAction> InAction, URigVMBlueprint* InRigBlueprint, URigVMEdGraph* InRigGraph);
 
 protected:
 
@@ -292,8 +294,8 @@ protected:
 	friend class URigVMEdGraphSchema;
 };
 
-UCLASS()
-class RIGVMDEVELOPER_API URigVMEdGraphSchema : public UEdGraphSchema
+UCLASS(MinimalAPI)
+class URigVMEdGraphSchema : public UEdGraphSchema
 {
 	GENERATED_BODY()
 
@@ -302,92 +304,92 @@ public:
 	static inline const FLazyName GraphName_RigVM = FLazyName(TEXT("RigVM"));
 
 public:
-	URigVMEdGraphSchema();
+	UE_API URigVMEdGraphSchema();
 
 	virtual const FLazyName& GetRootGraphName() const { return GraphName_RigVM; }
 
 	// UEdGraphSchema interface
-	virtual void GetGraphContextActions(FGraphContextMenuBuilder& ContextMenuBuilder) const override;
-	virtual void GetContextMenuActions(class UToolMenu* Menu, class UGraphNodeContextMenuContext* Context) const override;
-	virtual bool TryCreateConnection(UEdGraphPin* PinA, UEdGraphPin* PinB) const override;
-	virtual const FPinConnectionResponse CanCreateConnection(const UEdGraphPin* A, const UEdGraphPin* B) const override;
-	virtual FLinearColor GetPinTypeColor(const FEdGraphPinType& PinType) const override;
-	virtual bool SupportsPinType(TWeakPtr<const FEdGraphSchemaAction> SchemaAction, const FEdGraphPinType& PinType) const override;
-	virtual bool SupportsPinTypeContainer(TWeakPtr<const FEdGraphSchemaAction> SchemaAction, const FEdGraphPinType& PinType, const EPinContainerType& ContainerType) const override;
-	virtual void BreakPinLinks(UEdGraphPin& TargetPin, bool bSendsNodeNotifcation) const override;
-	virtual void BreakSinglePinLink(UEdGraphPin* SourcePin, UEdGraphPin* TargetPin) const override;
-	virtual bool CanGraphBeDropped(TSharedPtr<FEdGraphSchemaAction> InAction) const override;
-	virtual FString GetFindReferenceSearchTerm(const FEdGraphSchemaAction* InGraphAction) const override;
-	virtual FReply BeginGraphDragAction(TSharedPtr<FEdGraphSchemaAction> InAction, const FPointerEvent& MouseEvent = FPointerEvent() ) const override;
-	virtual class FConnectionDrawingPolicy* CreateConnectionDrawingPolicy(int32 InBackLayerID, int32 InFrontLayerID, float InZoomFactor, const FSlateRect& InClippingRect, class FSlateWindowElementList& InDrawElements, class UEdGraph* InGraphObj) const override;
-	virtual bool ShouldHidePinDefaultValue(UEdGraphPin* Pin) const override;
-	virtual void TrySetDefaultValue(UEdGraphPin& InPin, const FString& InNewDefaultValue, bool bMarkAsModified = true) const override;
-	virtual void TrySetDefaultObject(UEdGraphPin& InPin, UObject* InNewDefaultObject, bool bMarkAsModified) const override;
-	virtual void TrySetDefaultText(UEdGraphPin& InPin, const FText& InNewDefaultText, bool bMarkAsModified) const override;
+	UE_API virtual void GetGraphContextActions(FGraphContextMenuBuilder& ContextMenuBuilder) const override;
+	UE_API virtual void GetContextMenuActions(class UToolMenu* Menu, class UGraphNodeContextMenuContext* Context) const override;
+	UE_API virtual bool TryCreateConnection(UEdGraphPin* PinA, UEdGraphPin* PinB) const override;
+	UE_API virtual const FPinConnectionResponse CanCreateConnection(const UEdGraphPin* A, const UEdGraphPin* B) const override;
+	UE_API virtual FLinearColor GetPinTypeColor(const FEdGraphPinType& PinType) const override;
+	UE_API virtual bool SupportsPinType(TWeakPtr<const FEdGraphSchemaAction> SchemaAction, const FEdGraphPinType& PinType) const override;
+	UE_API virtual bool SupportsPinTypeContainer(TWeakPtr<const FEdGraphSchemaAction> SchemaAction, const FEdGraphPinType& PinType, const EPinContainerType& ContainerType) const override;
+	UE_API virtual void BreakPinLinks(UEdGraphPin& TargetPin, bool bSendsNodeNotifcation) const override;
+	UE_API virtual void BreakSinglePinLink(UEdGraphPin* SourcePin, UEdGraphPin* TargetPin) const override;
+	UE_API virtual bool CanGraphBeDropped(TSharedPtr<FEdGraphSchemaAction> InAction) const override;
+	UE_API virtual FString GetFindReferenceSearchTerm(const FEdGraphSchemaAction* InGraphAction) const override;
+	UE_API virtual FReply BeginGraphDragAction(TSharedPtr<FEdGraphSchemaAction> InAction, const FPointerEvent& MouseEvent = FPointerEvent() ) const override;
+	UE_API virtual class FConnectionDrawingPolicy* CreateConnectionDrawingPolicy(int32 InBackLayerID, int32 InFrontLayerID, float InZoomFactor, const FSlateRect& InClippingRect, class FSlateWindowElementList& InDrawElements, class UEdGraph* InGraphObj) const override;
+	UE_API virtual bool ShouldHidePinDefaultValue(UEdGraphPin* Pin) const override;
+	UE_API virtual void TrySetDefaultValue(UEdGraphPin& InPin, const FString& InNewDefaultValue, bool bMarkAsModified = true) const override;
+	UE_API virtual void TrySetDefaultObject(UEdGraphPin& InPin, UObject* InNewDefaultObject, bool bMarkAsModified) const override;
+	UE_API virtual void TrySetDefaultText(UEdGraphPin& InPin, const FText& InNewDefaultText, bool bMarkAsModified) const override;
 	virtual bool ShouldAlwaysPurgeOnModification() const override { return false; }
-	virtual bool ArePinsCompatible(const UEdGraphPin* PinA, const UEdGraphPin* PinB, const UClass* CallingContext, bool bIgnoreArray /*= false*/) const override;
+	UE_API virtual bool ArePinsCompatible(const UEdGraphPin* PinA, const UEdGraphPin* PinB, const UClass* CallingContext, bool bIgnoreArray /*= false*/) const override;
 	virtual bool DoesSupportPinWatching() const	override { return true; }
-	virtual bool IsPinBeingWatched(UEdGraphPin const* Pin) const override;
-	virtual void ClearPinWatch(UEdGraphPin const* Pin) const override;
-	virtual void OnPinConnectionDoubleCicked(UEdGraphPin* PinA, UEdGraphPin* PinB, const FVector2f& GraphPosition) const override;
-	virtual bool MarkBlueprintDirtyFromNewNode(UBlueprint* InBlueprint, UEdGraphNode* InEdGraphNode) const override;
-	virtual bool SafeDeleteNodeFromGraph(UEdGraph* Graph, UEdGraphNode* Node) const override;
-	virtual bool CanVariableBeDropped(UEdGraph* InGraph, FProperty* InVariableToDrop) const override;
-	virtual bool RequestFunctionDropOnPanel(UEdGraph* InGraph, const FRigVMGraphFunctionIdentifier& InFunction, const FVector2D& InDropPosition, const FVector2D& InScreenPosition) const;
-	virtual bool RequestVariableDropOnPanel(UEdGraph* InGraph, FProperty* InVariableToDrop, const FVector2f& InDropPosition, const FVector2f& InScreenPosition) const override;
-	virtual bool RequestVariableDropOnPin(UEdGraph* InGraph, FProperty* InVariableToDrop, UEdGraphPin* InPin, const FVector2f& InDropPosition, const FVector2f& InScreenPosition) const override;
+	UE_API virtual bool IsPinBeingWatched(UEdGraphPin const* Pin) const override;
+	UE_API virtual void ClearPinWatch(UEdGraphPin const* Pin) const override;
+	UE_API virtual void OnPinConnectionDoubleCicked(UEdGraphPin* PinA, UEdGraphPin* PinB, const FVector2f& GraphPosition) const override;
+	UE_API virtual bool MarkBlueprintDirtyFromNewNode(UBlueprint* InBlueprint, UEdGraphNode* InEdGraphNode) const override;
+	UE_API virtual bool SafeDeleteNodeFromGraph(UEdGraph* Graph, UEdGraphNode* Node) const override;
+	UE_API virtual bool CanVariableBeDropped(UEdGraph* InGraph, FProperty* InVariableToDrop) const override;
+	UE_API virtual bool RequestFunctionDropOnPanel(UEdGraph* InGraph, const FRigVMGraphFunctionIdentifier& InFunction, const FVector2D& InDropPosition, const FVector2D& InScreenPosition) const;
+	UE_API virtual bool RequestVariableDropOnPanel(UEdGraph* InGraph, FProperty* InVariableToDrop, const FVector2f& InDropPosition, const FVector2f& InScreenPosition) const override;
+	UE_API virtual bool RequestVariableDropOnPin(UEdGraph* InGraph, FProperty* InVariableToDrop, UEdGraphPin* InPin, const FVector2f& InDropPosition, const FVector2f& InScreenPosition) const override;
 	using UEdGraphSchema::RequestVariableDropOnPin;
-	virtual bool RequestVariableDropOnPin(UEdGraph* InGraph, const FRigVMExternalVariable& InVariableToDrop, UEdGraphPin* InPin, const FVector2D& InDropPosition, const FVector2D& InScreenPosition) const;
-	virtual bool IsStructEditable(UStruct* InStruct) const;
-	virtual void SetNodePosition(UEdGraphNode* Node, const FVector2f& Position) const override;
+	UE_API virtual bool RequestVariableDropOnPin(UEdGraph* InGraph, const FRigVMExternalVariable& InVariableToDrop, UEdGraphPin* InPin, const FVector2D& InDropPosition, const FVector2D& InScreenPosition) const;
+	UE_API virtual bool IsStructEditable(UStruct* InStruct) const;
+	UE_API virtual void SetNodePosition(UEdGraphNode* Node, const FVector2f& Position) const override;
 	using UEdGraphSchema::SetNodePosition;
-	void SetNodePosition(UEdGraphNode* Node, const FVector2D& Position, bool bSetupUndo) const;
-	virtual void GetGraphDisplayInformation(const UEdGraph& Graph, /*out*/ FGraphDisplayInfo& DisplayInfo) const override;
-	virtual bool GetLocalVariables(const UEdGraph* InGraph, TArray<FBPVariableDescription>& OutLocalVariables) const override;
-	virtual TSharedPtr<FEdGraphSchemaAction> MakeActionFromVariableDescription(const UEdGraph* InEdGraph, const FBPVariableDescription& Variable) const override;
-	virtual FText GetGraphCategory(const UEdGraph* InGraph) const override;
-	virtual EGraphType GetGraphType(const UEdGraph* TestEdGraph) const override;
-	virtual FReply TrySetGraphCategory(const UEdGraph* InGraph, const FText& InCategory) const override;
-	virtual bool TryDeleteGraph(UEdGraph* GraphToDelete) const override;
-	virtual bool TryRenameGraph(UEdGraph* GraphToRename, const FName& InNewName) const override;
-	virtual bool TryToGetChildEvents(const UEdGraph* Graph, const int32 SectionId, TArray<TSharedPtr<FEdGraphSchemaAction>>& Actions, const FText& ParentCategory) const override;
+	UE_API void SetNodePosition(UEdGraphNode* Node, const FVector2D& Position, bool bSetupUndo) const;
+	UE_API virtual void GetGraphDisplayInformation(const UEdGraph& Graph, /*out*/ FGraphDisplayInfo& DisplayInfo) const override;
+	UE_API virtual bool GetLocalVariables(const UEdGraph* InGraph, TArray<FBPVariableDescription>& OutLocalVariables) const override;
+	UE_API virtual TSharedPtr<FEdGraphSchemaAction> MakeActionFromVariableDescription(const UEdGraph* InEdGraph, const FBPVariableDescription& Variable) const override;
+	UE_API virtual FText GetGraphCategory(const UEdGraph* InGraph) const override;
+	UE_API virtual EGraphType GetGraphType(const UEdGraph* TestEdGraph) const override;
+	UE_API virtual FReply TrySetGraphCategory(const UEdGraph* InGraph, const FText& InCategory) const override;
+	UE_API virtual bool TryDeleteGraph(UEdGraph* GraphToDelete) const override;
+	UE_API virtual bool TryRenameGraph(UEdGraph* GraphToRename, const FName& InNewName) const override;
+	UE_API virtual bool TryToGetChildEvents(const UEdGraph* Graph, const int32 SectionId, TArray<TSharedPtr<FEdGraphSchemaAction>>& Actions, const FText& ParentCategory) const override;
 	virtual bool CanDuplicateGraph(UEdGraph* InSourceGraph) const { return false; }
 	virtual bool AllowsFunctionVariants() const override { return CVarRigVMEnableVariants.GetValueOnAnyThread(); }
-	virtual UEdGraphPin* DropPinOnNode(UEdGraphNode* InTargetNode, const FName& InSourcePinName, const FEdGraphPinType& InSourcePinType, EEdGraphPinDirection InSourcePinDirection) const override;
-	virtual bool SupportsDropPinOnNode(UEdGraphNode* InTargetNode, const FEdGraphPinType& InSourcePinType, EEdGraphPinDirection InSourcePinDirection, FText& OutErrorMessage) const override;
+	UE_API virtual UEdGraphPin* DropPinOnNode(UEdGraphNode* InTargetNode, const FName& InSourcePinName, const FEdGraphPinType& InSourcePinType, EEdGraphPinDirection InSourcePinDirection) const override;
+	UE_API virtual bool SupportsDropPinOnNode(UEdGraphNode* InTargetNode, const FEdGraphPinType& InSourcePinType, EEdGraphPinDirection InSourcePinDirection, FText& OutErrorMessage) const override;
 	virtual void SetPinBeingDroppedOnNode(UEdGraphPin* InSourcePin) const override { PinBeingDropped = InSourcePin; }
-	virtual void InsertAdditionalActions(TArray<UBlueprint*> InBlueprints, TArray<UEdGraph*> InGraphs, TArray<UEdGraphPin*> InPins, FGraphActionListBuilderBase& OutAllActions) const override;
-	virtual TSharedPtr<INameValidatorInterface> GetNameValidator(const UBlueprint* BlueprintObj, const FName& OriginalName, const UStruct* ValidationScope, const FName& ActionTypeId) const override;
-	bool CanPromotePinToVariable(const UEdGraphPin& Pin, bool bCond) const;
+	UE_API virtual void InsertAdditionalActions(TArray<UBlueprint*> InBlueprints, TArray<UEdGraph*> InGraphs, TArray<UEdGraphPin*> InPins, FGraphActionListBuilderBase& OutAllActions) const override;
+	UE_API virtual TSharedPtr<INameValidatorInterface> GetNameValidator(const UBlueprint* BlueprintObj, const FName& OriginalName, const UStruct* ValidationScope, const FName& ActionTypeId) const override;
+	UE_API bool CanPromotePinToVariable(const UEdGraphPin& Pin, bool bCond) const;
 
 	/** Returns true if the schema supports the script type */
-	bool SupportsPinType(const UScriptStruct* ScriptStruct) const;
+	UE_API bool SupportsPinType(const UScriptStruct* ScriptStruct) const;
 
 	/** Create a graph node for a rig */
-	URigVMEdGraphNode* CreateGraphNode(URigVMEdGraph* InGraph, const FName& InPropertyName) const;
+	UE_API URigVMEdGraphNode* CreateGraphNode(URigVMEdGraph* InGraph, const FName& InPropertyName) const;
 
 	/** Helper function to rename a node */
-	void RenameNode(URigVMEdGraphNode* Node, const FName& InNewNodeName) const;
+	UE_API void RenameNode(URigVMEdGraphNode* Node, const FName& InNewNodeName) const;
 
 	/** Helper function to recursively reset the pin defaults */
-	virtual void ResetPinDefaultsRecursive(UEdGraphPin* InPin) const;
+	UE_API virtual void ResetPinDefaultsRecursive(UEdGraphPin* InPin) const;
 
 	/** Returns all of the applicable pin types for variables within a rigvm host */
-	virtual void GetVariablePinTypes(TArray<FEdGraphPinType>& PinTypes) const;
+	UE_API virtual void GetVariablePinTypes(TArray<FEdGraphPinType>& PinTypes) const;
 
-	void StartGraphNodeInteraction(UEdGraphNode* InNode) const;
-	void EndGraphNodeInteraction(UEdGraphNode* InNode) const;
-	static TArray<UEdGraphNode*> GetNodesToMoveForNode(UEdGraphNode* InNode);
-	FVector2D GetNodePositionAtStartOfInteraction(const UEdGraphNode* InNode) const;
+	UE_API void StartGraphNodeInteraction(UEdGraphNode* InNode) const;
+	UE_API void EndGraphNodeInteraction(UEdGraphNode* InNode) const;
+	static UE_API TArray<UEdGraphNode*> GetNodesToMoveForNode(UEdGraphNode* InNode);
+	UE_API FVector2D GetNodePositionAtStartOfInteraction(const UEdGraphNode* InNode) const;
 
-	bool AutowireNewNode(URigVMEdGraphNode* NewNode, UEdGraphPin* FromPin) const;
+	UE_API bool AutowireNewNode(URigVMEdGraphNode* NewNode, UEdGraphPin* FromPin) const;
 
-	void HandleModifiedEvent(ERigVMGraphNotifType InNotifType, URigVMGraph* InGraph, UObject* InSubject);
+	UE_API void HandleModifiedEvent(ERigVMGraphNotifType InNotifType, URigVMGraph* InGraph, UObject* InSubject);
 
 	// Allow derived classes to spawn derived node classes
-	virtual TSubclassOf<URigVMEdGraphNode> GetGraphNodeClass(const URigVMEdGraph* InGraph) const;
+	UE_API virtual TSubclassOf<URigVMEdGraphNode> GetGraphNodeClass(const URigVMEdGraph* InGraph) const;
 
-	virtual bool IsRigVMDefaultEvent(const FName& InEventName) const;
+	UE_API virtual bool IsRigVMDefaultEvent(const FName& InEventName) const;
 
 private:
 
@@ -405,3 +407,5 @@ private:
 	friend class URigVMEdGraphUnitNodeSpawner;
 	friend class URigVMEdGraphArrayNodeSpawner;
 };
+
+#undef UE_API
diff --git a/Engine/Plugins/Runtime/RigVM/Source/RigVMDeveloper/Public/RigVMBlueprint.h b/Engine/Plugins/Runtime/RigVM/Source/RigVMDeveloper/Public/RigVMBlueprint.h
--- a/Engine/Plugins/Runtime/RigVM/Source/RigVMDeveloper/Public/RigVMBlueprint.h
+++ b/Engine/Plugins/Runtime/RigVM/Source/RigVMDeveloper/Public/RigVMBlueprint.h
@@ -50,7 +50,7 @@ DECLARE_EVENT_OneParam(URigVMBlueprint, FOnRigVMRequestInspectObject, const TArr
 DECLARE_EVENT_OneParam(URigVMBlueprint, FOnRigVMRequestInspectMemoryStorage, const TArray<FRigVMMemoryStorageStruct*>&);
 
 USTRUCT()
-struct RIGVMDEVELOPER_API FRigVMPythonSettings
+struct FRigVMPythonSettings
 {
 	GENERATED_BODY();
 
@@ -69,7 +69,7 @@ enum class ERigVMTagDisplayMode : uint8
 };
 
 USTRUCT()
-struct RIGVMDEVELOPER_API FRigVMEdGraphDisplaySettings
+struct FRigVMEdGraphDisplaySettings
 {
 	GENERATED_BODY();
 
@@ -159,10 +159,10 @@ struct RIGVMDEVELOPER_API FRigVMEdGraphDisplaySettings
 	UPROPERTY(EditAnywhere, Category = "Graph Display Settings")
 	ERigVMTagDisplayMode TagDisplayMode;
 
-	void SetTotalMicroSeconds(double InTotalMicroSeconds);
-	void SetLastMinMicroSeconds(double InMinMicroSeconds);
-	void SetLastMaxMicroSeconds(double InMaxMicroSeconds);
-	double AggregateAverage(TArray<double>& InFrames, double InPrevious, double InNext) const;
+	RIGVMDEVELOPER_API void SetTotalMicroSeconds(double InTotalMicroSeconds);
+	RIGVMDEVELOPER_API void SetLastMinMicroSeconds(double InMinMicroSeconds);
+	RIGVMDEVELOPER_API void SetLastMaxMicroSeconds(double InMaxMicroSeconds);
+	RIGVMDEVELOPER_API double AggregateAverage(TArray<double>& InFrames, double InPrevious, double InNext) const;
 };
 
 enum class UE_DEPRECATED(5.4, "Pease, use ERigVMLoadType") ERigVMBlueprintLoadType : uint8
@@ -172,7 +172,7 @@ enum class UE_DEPRECATED(5.4, "Pease, use ERigVMLoadType") ERigVMBlueprintLoadTy
 };
 
 USTRUCT(meta = (Deprecated = "5.2"))
-struct RIGVMDEVELOPER_API FRigVMOldPublicFunctionArg
+struct FRigVMOldPublicFunctionArg
 {
 	GENERATED_BODY();
 	
@@ -199,11 +199,11 @@ struct RIGVMDEVELOPER_API FRigVMOldPublicFunctionArg
 	UPROPERTY()
 	ERigVMPinDirection Direction;
 
-	FEdGraphPinType GetPinType() const;
+	RIGVMDEVELOPER_API FEdGraphPinType GetPinType() const;
 };
 
 USTRUCT(meta = (Deprecated = "5.2"))
-struct RIGVMDEVELOPER_API FRigVMOldPublicFunctionData
+struct FRigVMOldPublicFunctionData
 {
 	GENERATED_BODY();
 
@@ -229,75 +229,75 @@ struct RIGVMDEVELOPER_API FRigVMOldPublicFunctionData
 	UPROPERTY()
 	TArray<FRigVMOldPublicFunctionArg> Arguments;
 
-	bool IsMutable() const;
+	RIGVMDEVELOPER_API bool IsMutable() const;
 };
 
 
-UCLASS(BlueprintType, meta=(IgnoreClassThumbnail))
-class RIGVMDEVELOPER_API URigVMBlueprint : public UBlueprint, public IRigVMClientHost, public IRigVMExternalDependencyManager
+UCLASS(MinimalAPI, BlueprintType, meta=(IgnoreClassThumbnail))
+class URigVMBlueprint : public UBlueprint, public IRigVMClientHost, public IRigVMExternalDependencyManager
 {
 	GENERATED_UCLASS_BODY()
 
 public:
-	URigVMBlueprint();
+	RIGVMDEVELOPER_API URigVMBlueprint();
 
-	void CommonInitialization(const FObjectInitializer& ObjectInitializer);
+	RIGVMDEVELOPER_API void CommonInitialization(const FObjectInitializer& ObjectInitializer);
 	
-	void InitializeModelIfRequired(bool bRecompileVM = true);
+	RIGVMDEVELOPER_API void InitializeModelIfRequired(bool bRecompileVM = true);
 
 	/** Get the (full) generated class for this rigvm blueprint */
-	URigVMBlueprintGeneratedClass* GetRigVMBlueprintGeneratedClass() const;
+	RIGVMDEVELOPER_API URigVMBlueprintGeneratedClass* GetRigVMBlueprintGeneratedClass() const;
 
 	/** Get the (skeleton) generated class for this rigvm blueprint */
-	URigVMBlueprintGeneratedClass* GetRigVMBlueprintSkeletonClass() const;
+	RIGVMDEVELOPER_API URigVMBlueprintGeneratedClass* GetRigVMBlueprintSkeletonClass() const;
 
 	/** Returns the class used as the super class for all generated classes */
 	virtual UClass* GetRigVMBlueprintGeneratedClassPrototype() const { return URigVMBlueprintGeneratedClass::StaticClass(); }
 
 	/** Returns the settings defaults for this blueprint */
-	URigVMEditorSettings* GetRigVMEditorSettings() const;
+	RIGVMDEVELOPER_API URigVMEditorSettings* GetRigVMEditorSettings() const;
 
 #if WITH_EDITOR
 	/** Returns true if a given panel node factory is compatible this blueprint */
-	virtual const FLazyName& GetPanelNodeFactoryName() const;
+	RIGVMDEVELOPER_API virtual const FLazyName& GetPanelNodeFactoryName() const;
 
 	/** Returns true if a given panel pin factory is compatible this blueprint */
-	virtual const FLazyName& GetPanelPinFactoryName() const;
+	RIGVMDEVELOPER_API virtual const FLazyName& GetPanelPinFactoryName() const;
 
 	static inline const FLazyName RigVMPanelNodeFactoryName = FLazyName(TEXT("FRigVMEdGraphPanelNodeFactory"));
 	static inline const FLazyName RigVMPanelPinFactoryName = FLazyName(TEXT("FRigVMEdGraphPanelPinFactory"));
 
 	/** Returns the editor module to be used for this blueprint */
-	virtual IRigVMEditorModule* GetEditorModule() const;
+	RIGVMDEVELOPER_API virtual IRigVMEditorModule* GetEditorModule() const;
 #endif
 
-	virtual void Serialize(FArchive& Ar) override;
+	RIGVMDEVELOPER_API virtual void Serialize(FArchive& Ar) override;
 
 #if WITH_EDITOR
 
 	// UBlueprint interface
-	virtual UClass* GetBlueprintClass() const override;
-	virtual UClass* RegenerateClass(UClass* ClassToRegenerate, UObject* PreviousCDO) override;
+	RIGVMDEVELOPER_API virtual UClass* GetBlueprintClass() const override;
+	RIGVMDEVELOPER_API virtual UClass* RegenerateClass(UClass* ClassToRegenerate, UObject* PreviousCDO) override;
 	virtual bool SupportedByDefaultBlueprintFactory() const override { return false; }
 	virtual bool IsValidForBytecodeOnlyRecompile() const override { return false; }
-	virtual void LoadModulesRequiredForCompilation() override;
-	virtual void GetTypeActions(FBlueprintActionDatabaseRegistrar& ActionRegistrar) const override;
-	virtual void GetInstanceActions(FBlueprintActionDatabaseRegistrar& ActionRegistrar) const override;
-	virtual void SetObjectBeingDebugged(UObject* NewObject) override;
-	virtual void PreSave(FObjectPreSaveContext ObjectSaveContext) override;
-	virtual void PostSaveRoot(FObjectPostSaveRootContext ObjectSaveContext) override;
-	virtual void PostLoad() override;
+	RIGVMDEVELOPER_API virtual void LoadModulesRequiredForCompilation() override;
+	RIGVMDEVELOPER_API virtual void GetTypeActions(FBlueprintActionDatabaseRegistrar& ActionRegistrar) const override;
+	RIGVMDEVELOPER_API virtual void GetInstanceActions(FBlueprintActionDatabaseRegistrar& ActionRegistrar) const override;
+	RIGVMDEVELOPER_API virtual void SetObjectBeingDebugged(UObject* NewObject) override;
+	RIGVMDEVELOPER_API virtual void PreSave(FObjectPreSaveContext ObjectSaveContext) override;
+	RIGVMDEVELOPER_API virtual void PostSaveRoot(FObjectPostSaveRootContext ObjectSaveContext) override;
+	RIGVMDEVELOPER_API virtual void PostLoad() override;
 #if WITH_EDITORONLY_DATA
-	static void DeclareConstructClasses(TArray<FTopLevelAssetPath>& OutConstructClasses, const UClass* SpecificSubclass);
+	static RIGVMDEVELOPER_API void DeclareConstructClasses(TArray<FTopLevelAssetPath>& OutConstructClasses, const UClass* SpecificSubclass);
 #endif
 	virtual bool IsPostLoadThreadSafe() const override { return false; }
-	virtual void PostTransacted(const FTransactionObjectEvent& TransactionEvent) override;
-	virtual void ReplaceDeprecatedNodes() override;
-	virtual void PreDuplicate(FObjectDuplicationParameters& DupParams) override;
-	virtual void PostDuplicate(bool bDuplicateForPIE) override;
-	virtual void GetAssetRegistryTags(FAssetRegistryTagsContext Context) const override;
+	RIGVMDEVELOPER_API virtual void PostTransacted(const FTransactionObjectEvent& TransactionEvent) override;
+	RIGVMDEVELOPER_API virtual void ReplaceDeprecatedNodes() override;
+	RIGVMDEVELOPER_API virtual void PreDuplicate(FObjectDuplicationParameters& DupParams) override;
+	RIGVMDEVELOPER_API virtual void PostDuplicate(bool bDuplicateForPIE) override;
+	RIGVMDEVELOPER_API virtual void GetAssetRegistryTags(FAssetRegistryTagsContext Context) const override;
 	UE_DEPRECATED(5.4, "Implement the version that takes FAssetRegistryTagsContext instead.")
-	virtual void GetAssetRegistryTags(TArray<FAssetRegistryTag>& OutTags) const override;
+	RIGVMDEVELOPER_API virtual void GetAssetRegistryTags(TArray<FAssetRegistryTag>& OutTags) const override;
 
 	virtual bool SupportsGlobalVariables() const override { return true; }
 	virtual bool SupportsLocalVariables() const override { return true; }
@@ -306,16 +306,16 @@ public:
 	virtual bool SupportsDelegates() const override { return false; }
 	virtual bool SupportsEventGraphs() const override { return true; }
 	virtual bool SupportsAnimLayers() const override { return false; }
-	virtual bool ExportGraphToText(UEdGraph* InEdGraph, FString& OutText) override;
-	virtual bool TryImportGraphFromText(const FString& InClipboardText, UEdGraph** OutGraphPtr = nullptr) override;
-	virtual bool CanImportGraphFromText(const FString& InClipboardText) override;
-	virtual bool RequiresForceLoadMembers(UObject* InObject) const override;
+	RIGVMDEVELOPER_API virtual bool ExportGraphToText(UEdGraph* InEdGraph, FString& OutText) override;
+	RIGVMDEVELOPER_API virtual bool TryImportGraphFromText(const FString& InClipboardText, UEdGraph** OutGraphPtr = nullptr) override;
+	RIGVMDEVELOPER_API virtual bool CanImportGraphFromText(const FString& InClipboardText) override;
+	RIGVMDEVELOPER_API virtual bool RequiresForceLoadMembers(UObject* InObject) const override;
 
 	// UObject interface
-	virtual void PostEditChangeChainProperty(struct FPropertyChangedChainEvent& PropertyChangedEvent) override;
-	virtual void PostRename(UObject* OldOuter, const FName OldName) override;
+	RIGVMDEVELOPER_API virtual void PostEditChangeChainProperty(struct FPropertyChangedChainEvent& PropertyChangedEvent) override;
+	RIGVMDEVELOPER_API virtual void PostRename(UObject* OldOuter, const FName OldName) override;
 	/** Called during cooking. Must return all objects that will be Preload()ed when this is serialized at load time. */
-	virtual void GetPreloadDependencies(TArray<UObject*>& OutDeps) override;
+	RIGVMDEVELOPER_API virtual void GetPreloadDependencies(TArray<UObject*>& OutDeps) override;
 
 	//  --- IRigVMClientHost interface Start---
 	virtual FString GetAssetName() const override { return GetName(); }
@@ -325,103 +325,103 @@ public:
 	virtual UClass* GetRigVMEdGraphNodeClass() const override { return URigVMEdGraphNode::StaticClass(); }
 	virtual UClass* GetRigVMEdGraphSchemaClass() const override { return URigVMEdGraphSchema::StaticClass(); }
 	virtual UClass* GetRigVMEditorSettingsClass() const override { return URigVMEditorSettings::StaticClass(); }
-	virtual FRigVMClient* GetRigVMClient() override;
-	virtual const FRigVMClient* GetRigVMClient() const override;
-	virtual IRigVMGraphFunctionHost* GetRigVMGraphFunctionHost() override;
-	virtual const IRigVMGraphFunctionHost* GetRigVMGraphFunctionHost() const override;
-	virtual UObject* GetEditorObjectForRigVMGraph(const URigVMGraph* InVMGraph) const override;
-	virtual URigVMGraph* GetRigVMGraphForEditorObject(UObject* InObject) const override;
-	virtual void HandleRigVMGraphAdded(const FRigVMClient* InClient, const FString& InNodePath) override;
-	virtual void HandleRigVMGraphRemoved(const FRigVMClient* InClient, const FString& InNodePath) override;
-	virtual void HandleRigVMGraphRenamed(const FRigVMClient* InClient, const FString& InOldNodePath, const FString& InNewNodePath) override;
-	virtual void HandleConfigureRigVMController(const FRigVMClient* InClient, URigVMController* InControllerToConfigure) override;
-	virtual UObject* ResolveUserDefinedTypeById(const FString& InTypeName) const override;
+	RIGVMDEVELOPER_API virtual FRigVMClient* GetRigVMClient() override;
+	RIGVMDEVELOPER_API virtual const FRigVMClient* GetRigVMClient() const override;
+	RIGVMDEVELOPER_API virtual IRigVMGraphFunctionHost* GetRigVMGraphFunctionHost() override;
+	RIGVMDEVELOPER_API virtual const IRigVMGraphFunctionHost* GetRigVMGraphFunctionHost() const override;
+	RIGVMDEVELOPER_API virtual UObject* GetEditorObjectForRigVMGraph(const URigVMGraph* InVMGraph) const override;
+	RIGVMDEVELOPER_API virtual URigVMGraph* GetRigVMGraphForEditorObject(UObject* InObject) const override;
+	RIGVMDEVELOPER_API virtual void HandleRigVMGraphAdded(const FRigVMClient* InClient, const FString& InNodePath) override;
+	RIGVMDEVELOPER_API virtual void HandleRigVMGraphRemoved(const FRigVMClient* InClient, const FString& InNodePath) override;
+	RIGVMDEVELOPER_API virtual void HandleRigVMGraphRenamed(const FRigVMClient* InClient, const FString& InOldNodePath, const FString& InNewNodePath) override;
+	RIGVMDEVELOPER_API virtual void HandleConfigureRigVMController(const FRigVMClient* InClient, URigVMController* InControllerToConfigure) override;
+	RIGVMDEVELOPER_API virtual UObject* ResolveUserDefinedTypeById(const FString& InTypeName) const override;
 
 
 	UFUNCTION(BlueprintCallable, Category = "RigVM Blueprint")
-	virtual void RecompileVM() override;
+	RIGVMDEVELOPER_API virtual void RecompileVM() override;
 	UFUNCTION(BlueprintCallable, Category = "RigVM Blueprint")
-	virtual void RecompileVMIfRequired() override;
+	RIGVMDEVELOPER_API virtual void RecompileVMIfRequired() override;
 	UFUNCTION(BlueprintCallable, Category = "RigVM Blueprint")
-	virtual void RequestAutoVMRecompilation() override;
+	RIGVMDEVELOPER_API virtual void RequestAutoVMRecompilation() override;
 	UFUNCTION(BlueprintCallable, Category = "RigVM Blueprint")
 	virtual void SetAutoVMRecompile(bool bAutoRecompile) override { bAutoRecompileVM = bAutoRecompile; }
 	UFUNCTION(BlueprintPure, Category = "RigVM Blueprint")
 	virtual bool GetAutoVMRecompile() const override { return bAutoRecompileVM; }
 
-	virtual void IncrementVMRecompileBracket() override;
-	virtual void DecrementVMRecompileBracket() override;
+	RIGVMDEVELOPER_API virtual void IncrementVMRecompileBracket() override;
+	RIGVMDEVELOPER_API virtual void DecrementVMRecompileBracket() override;
 
 	// this is needed since even after load
 	// model data can change while the RigVM BP is not opened
 	// for example, if a user defined struct changed after BP load,
 	// any pin that references the struct needs to be regenerated
-	virtual void RefreshAllModels(ERigVMLoadType InLoadType = ERigVMLoadType::PostLoad) override;
+	RIGVMDEVELOPER_API virtual void RefreshAllModels(ERigVMLoadType InLoadType = ERigVMLoadType::PostLoad) override;
 
 	// RigVMRegistry changes can be triggered when new user defined types(structs/enums) are added/removed
 	// in which case we have to refresh the model
-	virtual void OnRigVMRegistryChanged() override;
+	RIGVMDEVELOPER_API virtual void OnRigVMRegistryChanged() override;
 
 	UFUNCTION(BlueprintCallable, Category = "RigVM Blueprint")
-	virtual void RequestRigVMInit() override;
+	RIGVMDEVELOPER_API virtual void RequestRigVMInit() override;
 
 	UFUNCTION(BlueprintCallable, Category = "RigVM Blueprint")
-	virtual URigVMGraph* GetModel(const UEdGraph* InEdGraph = nullptr) const override;
-	virtual URigVMGraph* GetModel(const FString& InNodePath) const override;
+	RIGVMDEVELOPER_API virtual URigVMGraph* GetModel(const UEdGraph* InEdGraph = nullptr) const override;
+	RIGVMDEVELOPER_API virtual URigVMGraph* GetModel(const FString& InNodePath) const override;
 
 	UFUNCTION(BlueprintCallable, Category = "RigVM Blueprint")
-	virtual URigVMGraph* GetDefaultModel() const override;
+	RIGVMDEVELOPER_API virtual URigVMGraph* GetDefaultModel() const override;
 
 	UFUNCTION(BlueprintCallable, Category = "RigVM Blueprint")
-	virtual TArray<URigVMGraph*> GetAllModels() const override;
+	RIGVMDEVELOPER_API virtual TArray<URigVMGraph*> GetAllModels() const override;
 
 	UFUNCTION(BlueprintCallable, Category = "RigVM Blueprint")
-	virtual URigVMFunctionLibrary* GetLocalFunctionLibrary() const override;
+	RIGVMDEVELOPER_API virtual URigVMFunctionLibrary* GetLocalFunctionLibrary() const override;
 
 	UFUNCTION(BlueprintCallable, Category = "RigVM Blueprint")
-	virtual URigVMFunctionLibrary* GetOrCreateLocalFunctionLibrary(bool bSetupUndoRedo = true) override;
+	RIGVMDEVELOPER_API virtual URigVMFunctionLibrary* GetOrCreateLocalFunctionLibrary(bool bSetupUndoRedo = true) override;
 
 	UFUNCTION(BlueprintCallable, Category = "RigVM Blueprint")
-	virtual URigVMGraph* AddModel(FString InName = TEXT("Rig Graph"), bool bSetupUndoRedo = true, bool bPrintPythonCommand = true) override;
+	RIGVMDEVELOPER_API virtual URigVMGraph* AddModel(FString InName = TEXT("Rig Graph"), bool bSetupUndoRedo = true, bool bPrintPythonCommand = true) override;
 
 	UFUNCTION(BlueprintCallable, Category = "RigVM Blueprint")
-	virtual bool RemoveModel(FString InName = TEXT("Rig Graph"), bool bSetupUndoRedo = true, bool bPrintPythonCommand = true) override;
+	RIGVMDEVELOPER_API virtual bool RemoveModel(FString InName = TEXT("Rig Graph"), bool bSetupUndoRedo = true, bool bPrintPythonCommand = true) override;
 
-	virtual FRigVMGetFocusedGraph& OnGetFocusedGraph() override;
-	virtual const FRigVMGetFocusedGraph& OnGetFocusedGraph() const override;
+	RIGVMDEVELOPER_API virtual FRigVMGetFocusedGraph& OnGetFocusedGraph() override;
+	RIGVMDEVELOPER_API virtual const FRigVMGetFocusedGraph& OnGetFocusedGraph() const override;
 
 	UFUNCTION(BlueprintCallable, Category = "RigVM Blueprint")
-	virtual URigVMGraph* GetFocusedModel() const override;
+	RIGVMDEVELOPER_API virtual URigVMGraph* GetFocusedModel() const override;
 
 	UFUNCTION(BlueprintCallable, Category = "RigVM Blueprint")
-	virtual URigVMController* GetController(const URigVMGraph* InGraph = nullptr) const override;
+	RIGVMDEVELOPER_API virtual URigVMController* GetController(const URigVMGraph* InGraph = nullptr) const override;
 
 	UFUNCTION(BlueprintCallable, Category = "RigVM Blueprint")
-	virtual URigVMController* GetControllerByName(const FString InGraphName = TEXT("")) const override;
+	RIGVMDEVELOPER_API virtual URigVMController* GetControllerByName(const FString InGraphName = TEXT("")) const override;
 
 	UFUNCTION(BlueprintCallable, Category = "RigVM Blueprint")
-	virtual URigVMController* GetOrCreateController(URigVMGraph* InGraph = nullptr) override;
+	RIGVMDEVELOPER_API virtual URigVMController* GetOrCreateController(URigVMGraph* InGraph = nullptr) override;
 
-	virtual URigVMController* GetController(const UEdGraph* InEdGraph) const override;
-	virtual URigVMController* GetOrCreateController(const UEdGraph* InGraph) override;
+	RIGVMDEVELOPER_API virtual URigVMController* GetController(const UEdGraph* InEdGraph) const override;
+	RIGVMDEVELOPER_API virtual URigVMController* GetOrCreateController(const UEdGraph* InGraph) override;
 
 	UFUNCTION(BlueprintCallable, Category = "RigVM Blueprint")
-	virtual TArray<FString> GeneratePythonCommands(const FString InNewBlueprintName) override;
+	RIGVMDEVELOPER_API virtual TArray<FString> GeneratePythonCommands(const FString InNewBlueprintName) override;
 
 	virtual void SetupPinRedirectorsForBackwardsCompatibility() override {};
 
-	virtual FRigVMGraphModifiedEvent& OnModified() override;
+	RIGVMDEVELOPER_API virtual FRigVMGraphModifiedEvent& OnModified() override;
 
-	virtual bool IsFunctionPublic(const FName& InFunctionName) const override;
-	virtual void MarkFunctionPublic(const FName& InFunctionName, bool bIsPublic = true) override;
+	RIGVMDEVELOPER_API virtual bool IsFunctionPublic(const FName& InFunctionName) const override;
+	RIGVMDEVELOPER_API virtual void MarkFunctionPublic(const FName& InFunctionName, bool bIsPublic = true) override;
 
-	virtual void RenameGraph(const FString& InNodePath, const FName& InNewName) override;
+	RIGVMDEVELOPER_API virtual void RenameGraph(const FString& InNodePath, const FName& InNewName) override;
 
 	//  --- IRigVMClientHost interface End ---
 
 	//  --- IRigVMExternalDependencyManager interface Start ---
 
-	virtual TArray<FRigVMExternalDependency> GetExternalDependenciesForCategory(const FName& InCategory) const override;
+	RIGVMDEVELOPER_API virtual TArray<FRigVMExternalDependency> GetExternalDependenciesForCategory(const FName& InCategory) const override;
 	
 	//  --- IRigVMExternalDependencyManager interface End ---
 
@@ -432,7 +432,7 @@ public:
 	FOnRigVMRequestInspectMemoryStorage& OnRequestInspectMemoryStorage() { return OnRequestInspectMemoryStorageEvent; }
 	void RequestInspectMemoryStorage(const TArray<FRigVMMemoryStorageStruct*>& InMemoryStorageStructs) { OnRequestInspectMemoryStorageEvent.Broadcast(InMemoryStorageStructs); }
 
-	void AddVariableSearchMetaDataInfo(const FName InVariableName, TArray<UBlueprintExtension::FSearchTagDataPair>& OutTaggedMetaData) const;
+	RIGVMDEVELOPER_API void AddVariableSearchMetaDataInfo(const FName InVariableName, TArray<UBlueprintExtension::FSearchTagDataPair>& OutTaggedMetaData) const;
 
 #endif	// #if WITH_EDITOR
 
@@ -440,12 +440,12 @@ public:
 
 	virtual bool ShouldBeMarkedDirtyUponTransaction() const override { return false; }
 
-	URigVMGraph* GetTemplateModel(bool bIsFunctionLibrary = false);
-	URigVMController* GetTemplateController(bool bIsFunctionLibrary = false);
+	RIGVMDEVELOPER_API URigVMGraph* GetTemplateModel(bool bIsFunctionLibrary = false);
+	RIGVMDEVELOPER_API URigVMController* GetTemplateController(bool bIsFunctionLibrary = false);
 
 #if WITH_EDITOR
-	UEdGraph* GetEdGraph(const URigVMGraph* InModel) const;
-	UEdGraph* GetEdGraph(const FString& InNodePath) const;
+	RIGVMDEVELOPER_API UEdGraph* GetEdGraph(const URigVMGraph* InModel) const;
+	RIGVMDEVELOPER_API UEdGraph* GetEdGraph(const FString& InNodePath) const;
 #endif
 
 #if WITH_EDITORONLY_DATA
@@ -455,19 +455,19 @@ public:
 
 	// Returns a list of dependencies of this blueprint.
 	// Dependencies are blueprints that contain functions used in this blueprint
-	TArray<URigVMBlueprint*> GetDependencies(bool bRecursive = false) const;
+	RIGVMDEVELOPER_API TArray<URigVMBlueprint*> GetDependencies(bool bRecursive = false) const;
 
 	// Returns a list of dependents as unresolved soft object pointers.
 	// A dependent is a blueprint which uses a function defined in this blueprint.
 	// This function is not recursive, since it avoids opening the asset.
 	// Use GetDependentBlueprints as an alternative.
-	TArray<FAssetData> GetDependentAssets() const;
+	RIGVMDEVELOPER_API TArray<FAssetData> GetDependentAssets() const;
 
 	// Returns a list of dependents as resolved blueprints.
 	// A dependent is a blueprint which uses a function defined in this blueprint.
 	// If bOnlyLoaded is false, this function loads the dependent assets and can introduce a large cost
 	// depending on the size / count of assets in the project.
-	TArray<URigVMBlueprint*> GetDependentBlueprints(bool bRecursive = false, bool bOnlyLoaded = false) const;
+	RIGVMDEVELOPER_API TArray<URigVMBlueprint*> GetDependentBlueprints(bool bRecursive = false, bool bOnlyLoaded = false) const;
 
 	UPROPERTY(EditAnywhere, Category = "User Interface")
 	FRigVMEdGraphDisplaySettings RigGraphDisplaySettings;
@@ -542,56 +542,56 @@ public:
 	bool bSuspendModelNotificationsForSelf;
 	bool bSuspendAllNotifications;
 
-	void RebuildGraphFromModel();
+	RIGVMDEVELOPER_API void RebuildGraphFromModel();
 
-	FOnRigVMCompiledEvent& OnVMCompiled();
+	RIGVMDEVELOPER_API FOnRigVMCompiledEvent& OnVMCompiled();
 
 	UFUNCTION(BlueprintCallable, Category = "VM")
-	UClass* GetRigVMHostClass() const;
+	RIGVMDEVELOPER_API UClass* GetRigVMHostClass() const;
 
 	UFUNCTION(BlueprintCallable, Category = "VM")
-	URigVMHost* CreateRigVMHost();
+	RIGVMDEVELOPER_API URigVMHost* CreateRigVMHost();
 
 	UFUNCTION(BlueprintCallable, Category = "VM")
 	URigVMHost* GetDebuggedRigVMHost() { return Cast<URigVMHost>(GetObjectBeingDebugged()); }
 
 	UFUNCTION(BlueprintCallable, Category = "VM")
-	virtual TArray<UStruct*> GetAvailableRigVMStructs() const;
+	RIGVMDEVELOPER_API virtual TArray<UStruct*> GetAvailableRigVMStructs() const;
 
 #if WITH_EDITOR
 	UFUNCTION(BlueprintCallable, Category = "Variables")
-	TArray<FRigVMGraphVariableDescription> GetMemberVariables() const;
+	RIGVMDEVELOPER_API TArray<FRigVMGraphVariableDescription> GetMemberVariables() const;
 	
 	UFUNCTION(BlueprintCallable, Category = "Variables")
-	FName AddMemberVariable(const FName& InName, const FString& InCPPType, bool bIsPublic = false, bool bIsReadOnly = false, FString InDefaultValue = TEXT(""));
+	RIGVMDEVELOPER_API FName AddMemberVariable(const FName& InName, const FString& InCPPType, bool bIsPublic = false, bool bIsReadOnly = false, FString InDefaultValue = TEXT(""));
 
 	UFUNCTION(BlueprintCallable, Category = "Variables")
-	bool RemoveMemberVariable(const FName& InName);
+	RIGVMDEVELOPER_API bool RemoveMemberVariable(const FName& InName);
 
 	UFUNCTION(BlueprintCallable, Category = "Variables")
-	bool RenameMemberVariable(const FName& InOldName, const FName& InNewName);
+	RIGVMDEVELOPER_API bool RenameMemberVariable(const FName& InOldName, const FName& InNewName);
 
 	UFUNCTION(BlueprintCallable, Category = "Variables")
-	bool ChangeMemberVariableType(const FName& InName, const FString& InCPPType, bool bIsPublic = false, bool bIsReadOnly = false, FString InDefaultValue = TEXT(""));
+	RIGVMDEVELOPER_API bool ChangeMemberVariableType(const FName& InName, const FString& InCPPType, bool bIsPublic = false, bool bIsReadOnly = false, FString InDefaultValue = TEXT(""));
 
 	const FRigVMVariant& GetAssetVariant() const { return AssetVariant; }
 
 	UFUNCTION(BlueprintPure, Category = "Variants", meta = (DisplayName = "GetAssetVariant", ScriptName = "GetAssetVariant"))
-	FRigVMVariant GetAssetVariantBP() const;
+	RIGVMDEVELOPER_API FRigVMVariant GetAssetVariantBP() const;
 
 	UFUNCTION(BlueprintPure, Category = "Variants")
-	FRigVMVariantRef GetAssetVariantRef() const;
+	RIGVMDEVELOPER_API FRigVMVariantRef GetAssetVariantRef() const;
 
 	/** Resets the asset's guid to a new one and splits it from the former variant set */
 	UFUNCTION(BlueprintCallable, Category = "Variants")
-	bool SplitAssetVariant();
+	RIGVMDEVELOPER_API bool SplitAssetVariant();
 
 	/** Merges the asset's guid with a provided one to join the variant set */
 	UFUNCTION(BlueprintCallable, Category = "Variants")
-	bool JoinAssetVariant(const FGuid& InGuid);
+	RIGVMDEVELOPER_API bool JoinAssetVariant(const FGuid& InGuid);
 
 	UFUNCTION(BlueprintPure, Category = "Variants")
-	TArray<FRigVMVariantRef> GetMatchingVariants() const;
+	RIGVMDEVELOPER_API TArray<FRigVMVariantRef> GetMatchingVariants() const;
 
 #endif
 
@@ -617,15 +617,15 @@ private:
     bool bSkipDirtyBlueprintStatus;
 
 	FRigVMGraphModifiedEvent ModifiedEvent;
-	void Notify(ERigVMGraphNotifType InNotifType, UObject* InSubject);
-	void HandleModifiedEvent(ERigVMGraphNotifType InNotifType, URigVMGraph* InGraph, UObject* InSubject);
+	RIGVMDEVELOPER_API void Notify(ERigVMGraphNotifType InNotifType, UObject* InSubject);
+	RIGVMDEVELOPER_API void HandleModifiedEvent(ERigVMGraphNotifType InNotifType, URigVMGraph* InGraph, UObject* InSubject);
 
-	void ReplaceFunctionIdentifiers(const FString& OldAssetPath, const FString& NewAssetPath);
+	RIGVMDEVELOPER_API void ReplaceFunctionIdentifiers(const FString& OldAssetPath, const FString& NewAssetPath);
 
 #if WITH_EDITOR
 
 	UFUNCTION(BlueprintCallable, Category = "RigVM Blueprint")
-	void SuspendNotifications(bool bSuspendNotifs);
+	RIGVMDEVELOPER_API void SuspendNotifications(bool bSuspendNotifs);
 
 	FOnRigVMRefreshEditorEvent RefreshEditorEvent;
 	FOnRigVMVariableDroppedEvent VariableDroppedEvent;
@@ -645,37 +645,37 @@ private:
 	FOnRigVMCompiledEvent VMCompiledEvent;
 
 	virtual void PathDomainSpecificContentOnLoad() {}
-	virtual void PatchBoundVariables();
-	virtual void PatchVariableNodesWithIncorrectType();
+	RIGVMDEVELOPER_API virtual void PatchBoundVariables();
+	RIGVMDEVELOPER_API virtual void PatchVariableNodesWithIncorrectType();
 	virtual void PatchParameterNodesOnLoad() {}
-	virtual void PatchLinksWithCast();
-	virtual void GetBackwardsCompatibilityPublicFunctions(TArray<FName> &BackwardsCompatiblePublicFunctions, TMap<URigVMLibraryNode*, FRigVMGraphFunctionHeader>& OldHeaders);
+	RIGVMDEVELOPER_API virtual void PatchLinksWithCast();
+	RIGVMDEVELOPER_API virtual void GetBackwardsCompatibilityPublicFunctions(TArray<FName> &BackwardsCompatiblePublicFunctions, TMap<URigVMLibraryNode*, FRigVMGraphFunctionHeader>& OldHeaders);
 
 protected:
 
 #if WITH_EDITOR
-	static FName FindHostMemberVariableUniqueName(TSharedPtr<FKismetNameValidator> InNameValidator, const FString& InBaseName);
-	static int32 AddHostMemberVariable(URigVMBlueprint* InBlueprint, const FName& InVarName, FEdGraphPinType InVarType, bool bIsPublic, bool bIsReadOnly, FString InDefaultValue);
-	FName AddHostMemberVariableFromExternal(FRigVMExternalVariable InVariableToCreate, FString InDefaultValue = FString());
+	static RIGVMDEVELOPER_API FName FindHostMemberVariableUniqueName(TSharedPtr<FKismetNameValidator> InNameValidator, const FString& InBaseName);
+	static RIGVMDEVELOPER_API int32 AddHostMemberVariable(URigVMBlueprint* InBlueprint, const FName& InVarName, FEdGraphPinType InVarType, bool bIsPublic, bool bIsReadOnly, FString InDefaultValue);
+	RIGVMDEVELOPER_API FName AddHostMemberVariableFromExternal(FRigVMExternalVariable InVariableToCreate, FString InDefaultValue = FString());
 #endif
 	
-	virtual void CreateMemberVariablesOnLoad();
-	virtual void PatchVariableNodesOnLoad();
+	RIGVMDEVELOPER_API virtual void CreateMemberVariablesOnLoad();
+	RIGVMDEVELOPER_API virtual void PatchVariableNodesOnLoad();
 	TMap<FName, int32> AddedMemberVariableMap;
 	TArray<FBPVariableDescription> LastNewVariables;
 
 public:
-	void PropagateRuntimeSettingsFromBPToInstances();
-	void InitializeArchetypeInstances();
+	RIGVMDEVELOPER_API void PropagateRuntimeSettingsFromBPToInstances();
+	RIGVMDEVELOPER_API void InitializeArchetypeInstances();
 
 private:
 
 #if WITH_EDITOR
-	void HandlePackageDone(const FEndLoadPackageContext& Context);
+	RIGVMDEVELOPER_API void HandlePackageDone(const FEndLoadPackageContext& Context);
 
 protected:
 	
-	virtual void HandlePackageDone();
+	RIGVMDEVELOPER_API virtual void HandlePackageDone();
 
 	/** Our currently running rig vm instance */
 	UPROPERTY(transient)
@@ -684,7 +684,7 @@ protected:
 private:
 	
 	// RigVMBP, once end-loaded, will inform other RigVM-Dependent systems that Host instances are ready.
-	void BroadcastRigVMPackageDone();
+	RIGVMDEVELOPER_API void BroadcastRigVMPackageDone();
 
 	// Previously some memory classes were parented to the asset object
 	// however it is no longer supported since classes are now identified 
@@ -692,7 +692,7 @@ private:
 	// this function removes those deprecated class.
 	// new classes should be created by RecompileVM and parented to the Package
 	// during PostLoad
-	void RemoveDeprecatedVMMemoryClass();
+	RIGVMDEVELOPER_API void RemoveDeprecatedVMMemoryClass();
 
 #if WITH_EDITORONLY_DATA
 	// During load, we do not want the GC to destroy the generator classes until all URigVMMemoryStorage objects
@@ -710,24 +710,24 @@ public:
 
 	FOnRigVMExternalVariablesChanged& OnExternalVariablesChanged() { return ExternalVariablesChangedEvent; }
 
-	virtual void OnPreVariableChange(UObject* InObject);
-	virtual void OnPostVariableChange(UBlueprint* InBlueprint);
-	virtual void OnVariableAdded(const FName& InVarName);
-	virtual void OnVariableRemoved(const FName& InVarName);
-	virtual void OnVariableRenamed(const FName& InOldVarName, const FName& InNewVarName);
-	virtual void OnVariableTypeChanged(const FName& InVarName, FEdGraphPinType InOldPinType, FEdGraphPinType InNewPinType);
+	RIGVMDEVELOPER_API virtual void OnPreVariableChange(UObject* InObject);
+	RIGVMDEVELOPER_API virtual void OnPostVariableChange(UBlueprint* InBlueprint);
+	RIGVMDEVELOPER_API virtual void OnVariableAdded(const FName& InVarName);
+	RIGVMDEVELOPER_API virtual void OnVariableRemoved(const FName& InVarName);
+	RIGVMDEVELOPER_API virtual void OnVariableRenamed(const FName& InOldVarName, const FName& InNewVarName);
+	RIGVMDEVELOPER_API virtual void OnVariableTypeChanged(const FName& InVarName, FEdGraphPinType InOldPinType, FEdGraphPinType InNewPinType);
 
 	FOnRigVMNodeDoubleClicked& OnNodeDoubleClicked() { return NodeDoubleClickedEvent; }
-	void BroadcastNodeDoubleClicked(URigVMNode* InNode);
+	RIGVMDEVELOPER_API void BroadcastNodeDoubleClicked(URigVMNode* InNode);
 
 	FOnRigVMGraphImported& OnGraphImported() { return GraphImportedEvent; }
-	void BroadcastGraphImported(UEdGraph* InGraph);
+	RIGVMDEVELOPER_API void BroadcastGraphImported(UEdGraph* InGraph);
 
 	FOnRigVMPostEditChangeChainProperty& OnPostEditChangeChainProperty() { return PostEditChangeChainPropertyEvent; }
-	void BroadcastPostEditChangeChainProperty(FPropertyChangedChainEvent& PropertyChangedChainEvent);
+	RIGVMDEVELOPER_API void BroadcastPostEditChangeChainProperty(FPropertyChangedChainEvent& PropertyChangedChainEvent);
 
 	FOnRigVMLocalizeFunctionDialogRequested& OnRequestLocalizeFunctionDialog() { return RequestLocalizeFunctionDialog; }
-	void BroadcastRequestLocalizeFunctionDialog(FRigVMGraphFunctionIdentifier InFunction, bool bForce = false);
+	RIGVMDEVELOPER_API void BroadcastRequestLocalizeFunctionDialog(FRigVMGraphFunctionIdentifier InFunction, bool bForce = false);
 
 	FRigVMOnBulkEditDialogRequestedDelegate& OnRequestBulkEditDialog() { return RequestBulkEditDialog; }
 
@@ -738,7 +738,7 @@ public:
 	FRigVMController_RequestJumpToHyperlinkDelegate& OnRequestJumpToHyperlink() { return RequestJumpToHyperlink; };
 
 	FOnRigVMReportCompilerMessage& OnReportCompilerMessage() { return ReportCompilerMessageEvent; }
-	void BroadCastReportCompilerMessage(EMessageSeverity::Type InSeverity, UObject* InSubject, const FString& InMessage);
+	RIGVMDEVELOPER_API void BroadCastReportCompilerMessage(EMessageSeverity::Type InSeverity, UObject* InSubject, const FString& InMessage);
 
 	const FCompilerResultsLog& GetCompileLog() const { return CompileLog; }
 	FCompilerResultsLog& GetCompileLog() { return CompileLog; }
@@ -747,7 +747,7 @@ private:
 
 	FOnRigVMExternalVariablesChanged ExternalVariablesChangedEvent;
 	bool bUpdatingExternalVariables;
-	void BroadcastExternalVariablesChangedEvent();
+	RIGVMDEVELOPER_API void BroadcastExternalVariablesChangedEvent();
 	FCompilerResultsLog CompileLog;
 
 	FOnRigVMNodeDoubleClicked NodeDoubleClickedEvent;
@@ -762,16 +762,16 @@ private:
 
 #endif
 
-	UEdGraph* CreateEdGraph(URigVMGraph* InModel, bool bForce = false);
-	bool RemoveEdGraph(URigVMGraph* InModel);
-	void DestroyObject(UObject* InObject);
-	void CreateEdGraphForCollapseNodeIfNeeded(URigVMCollapseNode* InNode, bool bForce = false);
-	bool RemoveEdGraphForCollapseNode(URigVMCollapseNode* InNode, bool bNotify = false);
-	void HandleReportFromCompiler(EMessageSeverity::Type InSeverity, UObject* InSubject, const FString& InMessage);
+	RIGVMDEVELOPER_API UEdGraph* CreateEdGraph(URigVMGraph* InModel, bool bForce = false);
+	RIGVMDEVELOPER_API bool RemoveEdGraph(URigVMGraph* InModel);
+	RIGVMDEVELOPER_API void DestroyObject(UObject* InObject);
+	RIGVMDEVELOPER_API void CreateEdGraphForCollapseNodeIfNeeded(URigVMCollapseNode* InNode, bool bForce = false);
+	RIGVMDEVELOPER_API bool RemoveEdGraphForCollapseNode(URigVMCollapseNode* InNode, bool bNotify = false);
+	RIGVMDEVELOPER_API void HandleReportFromCompiler(EMessageSeverity::Type InSeverity, UObject* InSubject, const FString& InMessage);
 
 protected:
 	
-	TArray<IRigVMGraphFunctionHost*> GetReferencedFunctionHosts(bool bForceLoad);
+	RIGVMDEVELOPER_API TArray<IRigVMGraphFunctionHost*> GetReferencedFunctionHosts(bool bForceLoad);
 	
 #if WITH_EDITOR
 private:
@@ -785,50 +785,50 @@ private:
 public:
 
 	/** Sets the execution mode. In Release mode the rig will ignore all breakpoints. */
-	void SetDebugMode(const bool bValue);
+	RIGVMDEVELOPER_API void SetDebugMode(const bool bValue);
 
 	/** Returns the execution mode */
 	bool IsInDebugMode() const { return bCompileInDebugMode; }
 
 	/** Removes all the breakpoints from the blueprint and the VM */
-	void ClearBreakpoints();
+	RIGVMDEVELOPER_API void ClearBreakpoints();
 
 	/** Adds a breakpoint to all loaded blueprints which use the node indicated by InBreakpointNodePath 
 	  * If the node is inside a public function, it will add a breakpoint to all blueprints calling this function. */
-	bool AddBreakpoint(const FString& InBreakpointNodePath);
+	RIGVMDEVELOPER_API bool AddBreakpoint(const FString& InBreakpointNodePath);
 
 	/** Adds a breakpoint to all loaded blueprints which use the InBreakpointNode. 
 	  * If LibraryNode is not null, it indicates that the library uses the InBreapointNode, and the function will add
 	  * breakpoints to any other loaded blueprint that references this library. */
-	bool AddBreakpoint(URigVMNode* InBreakpointNode, URigVMLibraryNode* LibraryNode = nullptr);
+	RIGVMDEVELOPER_API bool AddBreakpoint(URigVMNode* InBreakpointNode, URigVMLibraryNode* LibraryNode = nullptr);
 
 	/** Adds a breakpoint to the first instruction of each callpath related to the InBreakpointNode */
-	bool AddBreakpointToHost(URigVMNode* InBreakpointNode);
+	RIGVMDEVELOPER_API bool AddBreakpointToHost(URigVMNode* InBreakpointNode);
 
 	/** Removes the given breakpoint from all the loaded blueprints that use this node, and recomputes all breakpoints
 	  * in the VM  */
-	bool RemoveBreakpoint(const FString& InBreakpointNodePath);
-	bool RemoveBreakpoint(URigVMNode* InBreakpointNode);
+	RIGVMDEVELOPER_API bool RemoveBreakpoint(const FString& InBreakpointNodePath);
+	RIGVMDEVELOPER_API bool RemoveBreakpoint(URigVMNode* InBreakpointNode);
 
 	/** Recomputes the instruction breakpoints given the node breakpoints in the blueprint */
-	void RefreshBreakpoints();
+	RIGVMDEVELOPER_API void RefreshBreakpoints();
 
 	/** Shape libraries to load during package load completed */ 
 	TArray<FString> ShapeLibrariesToLoadOnPackageLoaded;
 
-	TArray<FRigVMReferenceNodeData> GetReferenceNodeData() const;
+	RIGVMDEVELOPER_API TArray<FRigVMReferenceNodeData> GetReferenceNodeData() const;
 
 #endif
 
 protected:
 
-	static FSoftObjectPath PreDuplicateAssetPath;
-	static FSoftObjectPath PreDuplicateHostPath;
-	static TArray<URigVMBlueprint*> sCurrentlyOpenedRigVMBlueprints;
+	static RIGVMDEVELOPER_API FSoftObjectPath PreDuplicateAssetPath;
+	static RIGVMDEVELOPER_API FSoftObjectPath PreDuplicateHostPath;
+	static RIGVMDEVELOPER_API TArray<URigVMBlueprint*> sCurrentlyOpenedRigVMBlueprints;
 
 	void MarkDirtyDuringLoad() { bDirtyDuringLoad = true; }
 	
-	virtual void SetupDefaultObjectDuringCompilation(URigVMHost* InCDO);
+	RIGVMDEVELOPER_API virtual void SetupDefaultObjectDuringCompilation(URigVMHost* InCDO);
 
 public:
 	
@@ -844,13 +844,13 @@ private:
 
 public:
 
-	static void QueueCompilerMessageDelegate(const FOnRigVMReportCompilerMessage::FDelegate& InDelegate);
-	static void ClearQueuedCompilerMessageDelegates();
+	static RIGVMDEVELOPER_API void QueueCompilerMessageDelegate(const FOnRigVMReportCompilerMessage::FDelegate& InDelegate);
+	static RIGVMDEVELOPER_API void ClearQueuedCompilerMessageDelegates();
 	
 private:
 
-	static FCriticalSection QueuedCompilerMessageDelegatesMutex;
-	static TArray<FOnRigVMReportCompilerMessage::FDelegate> QueuedCompilerMessageDelegates;
+	static RIGVMDEVELOPER_API FCriticalSection QueuedCompilerMessageDelegatesMutex;
+	static RIGVMDEVELOPER_API TArray<FOnRigVMReportCompilerMessage::FDelegate> QueuedCompilerMessageDelegates;
 
 #endif
 
@@ -871,7 +871,7 @@ private:
 	friend class UE::RigVM::Editor::Tools::FFilterByAssetTag;
 };
 
-class RIGVMDEVELOPER_API FRigVMBlueprintCompileScope
+class FRigVMBlueprintCompileScope
 {
 public:
    
diff --git a/Engine/Plugins/Runtime/RigVM/Source/RigVMDeveloper/Public/RigVMBlueprintCompiler.h b/Engine/Plugins/Runtime/RigVM/Source/RigVMDeveloper/Public/RigVMBlueprintCompiler.h
--- a/Engine/Plugins/Runtime/RigVM/Source/RigVMDeveloper/Public/RigVMBlueprintCompiler.h
+++ b/Engine/Plugins/Runtime/RigVM/Source/RigVMDeveloper/Public/RigVMBlueprintCompiler.h
@@ -5,18 +5,20 @@
 #include "KismetCompilerModule.h"
 #include "KismetCompiler.h"
 
+#define UE_API RIGVMDEVELOPER_API
+
 class FRigVMBlueprintCompilerContext;
 class URigVMBlueprintGeneratedClass;
 
-class RIGVMDEVELOPER_API FRigVMBlueprintCompiler : public IBlueprintCompiler
+class FRigVMBlueprintCompiler : public IBlueprintCompiler
 {
 public:
 	/** IBlueprintCompiler interface */
 	virtual bool CanCompile(const UBlueprint* Blueprint) override { return false; } // to be implemented by the specializations
-	virtual void Compile(UBlueprint* Blueprint, const FKismetCompilerOptions& CompileOptions, FCompilerResultsLog& Results) override;
+	UE_API virtual void Compile(UBlueprint* Blueprint, const FKismetCompilerOptions& CompileOptions, FCompilerResultsLog& Results) override;
 };
 
-class RIGVMDEVELOPER_API FRigVMBlueprintCompilerContext : public FKismetCompilerContext
+class FRigVMBlueprintCompilerContext : public FKismetCompilerContext
 {
 public:
 	FRigVMBlueprintCompilerContext(UBlueprint* SourceSketch, FCompilerResultsLog& InMessageLog, const FKismetCompilerOptions& InCompilerOptions)
@@ -32,24 +34,26 @@ public:
 	virtual bool CanIgnoreNode(const UEdGraphNode* Node) const override { return true; }
 	virtual bool ShouldForceKeepNode(const UEdGraphNode* Node) const override { return false; }
 	virtual void PrecompileFunction(FKismetFunctionContext& Context, EInternalCompilerFlags InternalFlags) override {}
-	virtual void OnPostCDOCompiled(const UObject::FPostCDOCompiledContext& Context) override;
-	virtual void CopyTermDefaultsToDefaultObject(UObject* DefaultObject) override;
-	virtual void EnsureProperGeneratedClass(UClass*& TargetUClass) override;
-	virtual void SpawnNewClass(const FString& NewClassName) override;
-	virtual void OnNewClassSet(UBlueprintGeneratedClass* ClassToUse) override;
+	UE_API virtual void OnPostCDOCompiled(const UObject::FPostCDOCompiledContext& Context) override;
+	UE_API virtual void CopyTermDefaultsToDefaultObject(UObject* DefaultObject) override;
+	UE_API virtual void EnsureProperGeneratedClass(UClass*& TargetUClass) override;
+	UE_API virtual void SpawnNewClass(const FString& NewClassName) override;
+	UE_API virtual void OnNewClassSet(UBlueprintGeneratedClass* ClassToUse) override;
 	virtual void PruneIsolatedNodes(const TArray<UEdGraphNode*>& RootSet, TArray<UEdGraphNode*>& GraphNodes) override {}
-	virtual void CleanAndSanitizeClass(UBlueprintGeneratedClass* ClassToClean, UObject*& InOldCDO) override;
+	UE_API virtual void CleanAndSanitizeClass(UBlueprintGeneratedClass* ClassToClean, UObject*& InOldCDO) override;
 	virtual void MergeUbergraphPagesIn(UEdGraph* Ubergraph) override {};
-	virtual void PreCompileUpdateBlueprintOnLoad(UBlueprint* BP) override;
+	UE_API virtual void PreCompileUpdateBlueprintOnLoad(UBlueprint* BP) override;
 	virtual void CreateFunctionList() override {}
 	virtual void ProcessOneFunctionGraph(UEdGraph* SourceGraph, bool bInternalFunction = false) override {}
 
 private:
 
 	// used to fail a compilation and mark the blueprint in error
-	void MarkCompilationFailed(const FString& Message);
+	UE_API void MarkCompilationFailed(const FString& Message);
 
 protected:
 	/** the new class we are generating */
 	URigVMBlueprintGeneratedClass* NewRigVMBlueprintGeneratedClass;
-};
\ No newline at end of file
+};
+
+#undef UE_API
diff --git a/Engine/Plugins/Runtime/RigVM/Source/RigVMDeveloper/Public/RigVMBlueprintUtils.h b/Engine/Plugins/Runtime/RigVM/Source/RigVMDeveloper/Public/RigVMBlueprintUtils.h
--- a/Engine/Plugins/Runtime/RigVM/Source/RigVMDeveloper/Public/RigVMBlueprintUtils.h
+++ b/Engine/Plugins/Runtime/RigVM/Source/RigVMDeveloper/Public/RigVMBlueprintUtils.h
@@ -6,28 +6,32 @@
 #include "Templates/SubclassOf.h"
 #include "EdGraph/EdGraphPin.h"
 
+#define UE_API RIGVMDEVELOPER_API
+
 class UStruct;
 class UBlueprint;
 class URigVMEdGraphNode;
 class UEdGraph;
 class UEdGraphPin;
 
-struct RIGVMDEVELOPER_API FRigVMBlueprintUtils
+struct FRigVMBlueprintUtils
 {
 /** Call a function for each valid rig unit struct */
-static void ForAllRigVMStructs(TFunction<void(UScriptStruct*)> InFunction);
+static UE_API void ForAllRigVMStructs(TFunction<void(UScriptStruct*)> InFunction);
 
 /** Handle blueprint node reconstruction */
-static void HandleReconstructAllNodes(UBlueprint* InBlueprint);
+static UE_API void HandleReconstructAllNodes(UBlueprint* InBlueprint);
 
 /** Handle blueprint node refresh */
-static void HandleRefreshAllNodes(UBlueprint* InBlueprint);
+static UE_API void HandleRefreshAllNodes(UBlueprint* InBlueprint);
 
 /** Handle blueprint deleted */
-static void HandleAssetDeleted(const FAssetData& InAssetData);
+static UE_API void HandleAssetDeleted(const FAssetData& InAssetData);
 
 /** remove the variable if not used by anybody else but ToBeDeleted*/
-static void RemoveMemberVariableIfNotUsed(UBlueprint* Blueprint, const FName VarName);
+static UE_API void RemoveMemberVariableIfNotUsed(UBlueprint* Blueprint, const FName VarName);
+
+static UE_API FName ValidateName(UBlueprint* InBlueprint, const FString& InName);
+};
 
-static FName ValidateName(UBlueprint* InBlueprint, const FString& InName);
-};
\ No newline at end of file
+#undef UE_API
