Updated headers using UnrealCodeFixup to make sure dllstorage is on methods/staticvars instead of types. Part 4/n

diff --git a/Engine/Plugins/Runtime/RigVM/Source/RigVMDeveloper/Public/RigVMCompiler/RigVMAST.h b/Engine/Plugins/Runtime/RigVM/Source/RigVMDeveloper/Public/RigVMCompiler/RigVMAST.h
--- a/Engine/Plugins/Runtime/RigVM/Source/RigVMDeveloper/Public/RigVMCompiler/RigVMAST.h
+++ b/Engine/Plugins/Runtime/RigVM/Source/RigVMDeveloper/Public/RigVMCompiler/RigVMAST.h
@@ -11,6 +11,8 @@
 #include "Logging/TokenizedMessage.h"
 #include "RigVMAST.generated.h"
 
+#define UE_API RIGVMDEVELOPER_API
+
 class FRigVMParserAST;
 class FRigVMBlockExprAST;
 class FRigVMNodeExprAST;
@@ -37,7 +39,7 @@ class URigVMLibraryNode;
 /*
  * A structure to describe a link between two proxies
  */
-struct RIGVMDEVELOPER_API FRigVMASTLinkDescription
+struct FRigVMASTLinkDescription
 {
 	FRigVMASTLinkDescription() {}
 	
@@ -73,7 +75,7 @@ struct RIGVMDEVELOPER_API FRigVMASTLinkDescription
  * Expressions can only be constructed by an AST parser, and
  * are also memory-owned by the parser.
  */
-class RIGVMDEVELOPER_API FRigVMExprAST
+class FRigVMExprAST
 {
 public:
 
@@ -116,7 +118,7 @@ public:
 
 	/// returns the name of the expression's type
 	/// @return the name of the expression's type
-	FName GetTypeName() const;
+	UE_API FName GetTypeName() const;
 
 	/// provides type checking for inherited types
 	/// @param InType the type to check against
@@ -151,48 +153,48 @@ public:
 
 	/// returns the parent of this expression
 	/// @return the parent of this expression
-	const FRigVMExprAST* GetParent() const;
+	UE_API const FRigVMExprAST* GetParent() const;
 
 	/// returns the first parent in the tree of a given type
 	/// @param InExprType The type of expression to look for in the parent tree
 	/// @return the first parent in the tree of a given type
-	virtual const FRigVMExprAST* GetFirstParentOfType(EType InExprType) const;
+	UE_API virtual const FRigVMExprAST* GetFirstParentOfType(EType InExprType) const;
 
 	/// returns true if a given expression is parented to another expression
 	/// @param InParentExpr The potential parent to check
 	/// @return True if this expression is parented to the InParentExpr
-	bool IsParentedTo(const FRigVMExprAST* InParentExpr) const;
+	UE_API bool IsParentedTo(const FRigVMExprAST* InParentExpr) const;
 
 	/// returns true if a given expression is a parent of another expression
 	/// @param InChildExpr The potential child to check
 	/// @return True if this expression is a parent of the InChildExpr
-	bool IsParentOf(const FRigVMExprAST* InChildExpr) const;
+	UE_API bool IsParentOf(const FRigVMExprAST* InChildExpr) const;
 
 	/// returns the block of this expression
 	/// @return the block of this expression
-	const FRigVMBlockExprAST* GetBlock() const;
+	UE_API const FRigVMBlockExprAST* GetBlock() const;
 
 	/// returns the root / top level block of this expression
 	/// @return the root / top level block of this expression
-	const FRigVMBlockExprAST* GetRootBlock() const;
+	UE_API const FRigVMBlockExprAST* GetRootBlock() const;
 
 	/// returns all of the block this expression is in
 	/// @param bSortByDepth If true the blocks will be returned sorted (head to tail)
 	/// @return all of the block this expression is in
 	typedef TArray<const FRigVMBlockExprAST*, TInlineAllocator<5>> FRigVMBlockArray;
-	FRigVMBlockArray GetBlocks(bool bSortByDepth) const;
+	UE_API FRigVMBlockArray GetBlocks(bool bSortByDepth) const;
 
 	/// returns the unique block combination hash for this expression (or INDEX_NONE)
 	/// @return the unique block combination name for this expression (or INDEX_NONE)
-	TOptional<uint32> GetBlockCombinationHash() const;
+	UE_API TOptional<uint32> GetBlockCombinationHash() const;
 
 	/// returns the unique block combination name for this expression (or an empty string)
 	/// @return the unique block combination name for this expression (or an empty string)
-	const FString& GetBlockCombinationName() const;
+	UE_API const FString& GetBlockCombinationName() const;
 
 	/// returns the lowest child index found for this expression within a parent (or INDEX_NONE)
 	/// @return the lowest child index found for this expression within a parent (or INDEX_NONE)
-	int32 GetMinChildIndexWithinParent(const FRigVMExprAST* InParentExpr) const;
+	UE_API int32 GetMinChildIndexWithinParent(const FRigVMExprAST* InParentExpr) const;
 
 	/// returns the number of children of this expression
 	/// @return the number of children of this expression
@@ -200,7 +202,7 @@ public:
 
 	/// returns true if this expressions is constant (non varying)
 	/// @return true if this expressions is constant (non varying)
-	virtual bool IsConstant() const;
+	UE_API virtual bool IsConstant() const;
 
 	/// returns true if this expressions is varying (non constant)
 	/// @return true if this expressions is varying (non constant)
@@ -208,7 +210,7 @@ public:
 
 	/// returns the maximum depth of the expression in tree (number of recursive parents)
 	/// @return the maximum depth of the expression in tree (number of recursive parents)
-	int32 GetMaximumDepth() const;
+	UE_API int32 GetMaximumDepth() const;
 
 	/// accessor operator for a given child
 	/// @param InIndex the index of the child to retrieve (bound = NumChildren() - 1)
@@ -241,7 +243,7 @@ public:
 	/// returns the first parent in the tree of a given type
 	/// @param InExprType The type of expression to look for in the parent tree
 	/// @return the first parent in the tree of a given type
-	virtual const FRigVMExprAST* GetFirstChildOfType(EType InExprType) const;
+	UE_API virtual const FRigVMExprAST* GetFirstChildOfType(EType InExprType) const;
 
 	/// returns the number of parents of this expression
 	/// @return the number of parents of this expression
@@ -289,49 +291,49 @@ public:
 protected:
 
 	// default constructor (protected so that only parser can access it)
-	FRigVMExprAST(EType InType = EType::Invalid, const FRigVMASTProxy& InProxy = FRigVMASTProxy());
+	UE_API FRigVMExprAST(EType InType = EType::Invalid, const FRigVMASTProxy& InProxy = FRigVMASTProxy());
 
 	// adds a parent to this expression
 	// this in consequence also adds this as a child to the parent
 	// @param InParent the parent to add
-	void AddParent(FRigVMExprAST* InParent);
+	UE_API void AddParent(FRigVMExprAST* InParent);
 
 	// removes a parent from this expression
 	// this in consequence also removes this as a child from the parent
 	// @param InParent the parent to remove
-	void RemoveParent(FRigVMExprAST* InParent);
+	UE_API void RemoveParent(FRigVMExprAST* InParent);
 
 	// removes a child from this expression
 	// this in consequence also removes this as a parent from the child
 	// @param InChild the child to remove
-	void RemoveChild(FRigVMExprAST* InChild);
+	UE_API void RemoveChild(FRigVMExprAST* InChild);
 
 	// replaces a parent of this expression with a new one
 	// @param InCurrentParent the current parent to replace
 	// @param InNewParent the new parent to replace it with
-	void ReplaceParent(FRigVMExprAST* InCurrentParent, FRigVMExprAST* InNewParent);
+	UE_API void ReplaceParent(FRigVMExprAST* InCurrentParent, FRigVMExprAST* InNewParent);
 
 	// replaces a child of this expression with a new one
 	// @param InCurrentChild the current child to replace
 	// @param InNewChild the new child to replace it with
-	void ReplaceChild(FRigVMExprAST* InCurrentChild, FRigVMExprAST* InNewChild);
+	UE_API void ReplaceChild(FRigVMExprAST* InCurrentChild, FRigVMExprAST* InNewChild);
 
 	// replaces this expression with another one
 	// @param InReplacement the expression to replace this one
-	void ReplaceBy(FRigVMExprAST* InReplacement);
+	UE_API void ReplaceBy(FRigVMExprAST* InReplacement);
 
 	// computes all of the block this expression is in
-	void GetBlocksImpl(FRigVMBlockArray& InOutBlocks) const;
+	UE_API void GetBlocksImpl(FRigVMBlockArray& InOutBlocks) const;
 
 	// returns a string containing an indented tree structure
 	// for debugging purposes. this is only used by the parser
 	// @param InPrefix the prefix to use for indentation
 	// @return the text representation of this part of the tree
-	virtual FString DumpText(const FString& InPrefix = FString()) const;
+	UE_API virtual FString DumpText(const FString& InPrefix = FString()) const;
 
 	// Resets the internal caches of the expression
-	void InvalidateCaches();
-	void InvalidateCachesImpl(TArray<bool>& OutProcessed);
+	UE_API void InvalidateCaches();
+	UE_API void InvalidateCachesImpl(TArray<bool>& OutProcessed);
 
 	FName Name;
 	EType Type;
@@ -518,7 +520,7 @@ inline const FRigVMExitExprAST* FRigVMExprAST::To() const
  * of child expressions to be executed in order.
  * In C++ a block is represented by the curly braces { expr1, expr2, ...}.
  */ 
-class RIGVMDEVELOPER_API FRigVMBlockExprAST : public FRigVMExprAST
+class FRigVMBlockExprAST : public FRigVMExprAST
 {
 public:
 
@@ -531,16 +533,16 @@ public:
 	// returns true if this block needs to execute
 	// this is determined by the block containing an entry expression
 	// @return true if this block needs to execute
-	bool ShouldExecute() const;
+	UE_API bool ShouldExecute() const;
 
 	// returns true if this block contains an entry expression
 	// @return true if this block contains an entry expression
-	bool ContainsEntry() const;
+	UE_API bool ContainsEntry() const;
 
 	// returns true if this block contains a given expression
 	// @param InExpression the expression to check
 	// @return true if this block contains a given expression
-	bool Contains(const FRigVMExprAST* InExpression, TMap<const FRigVMExprAST*, bool>* ContainedExpressionsCache = nullptr) const;
+	UE_API bool Contains(const FRigVMExprAST* InExpression, TMap<const FRigVMExprAST*, bool>* ContainedExpressionsCache = nullptr) const;
 
 	// overload of the type checking mechanism
 	virtual bool IsA(EType InType) const override
@@ -549,7 +551,7 @@ public:
 	};
 
 	// returns true in case this block should not be executed
-	bool IsObsolete() const;
+	UE_API bool IsObsolete() const;
 
 protected:
 
@@ -570,7 +572,7 @@ private:
  * An abstract syntax tree node expression represents any expression
  * which references a node from the RigVM model.
  */
-class RIGVMDEVELOPER_API FRigVMNodeExprAST : public FRigVMBlockExprAST
+class FRigVMNodeExprAST : public FRigVMBlockExprAST
 {
 public:
 
@@ -592,14 +594,14 @@ public:
 	// @return the node from the model this expression is referencing
 	URigVMNode* GetNode() const { return GetProxy().GetSubjectChecked<URigVMNode>(); }
 
-	virtual bool IsConstant() const override;
-	const FRigVMExprAST* FindExprWithPinName(const FName& InPinName) const;
-	const FRigVMVarExprAST* FindVarWithPinName(const FName& InPinName) const;
+	UE_API virtual bool IsConstant() const override;
+	UE_API const FRigVMExprAST* FindExprWithPinName(const FName& InPinName) const;
+	UE_API const FRigVMVarExprAST* FindVarWithPinName(const FName& InPinName) const;
 
 protected:
 
 	// default constructor (protected so that only parser can access it)
-	FRigVMNodeExprAST(EType InType, const FRigVMASTProxy& InNodeProxy);
+	UE_API FRigVMNodeExprAST(EType InType, const FRigVMASTProxy& InNodeProxy);
 
 	// overload of the type checking mechanism
 	virtual bool IsA(EType InType) const override
@@ -618,7 +620,7 @@ private:
  * for a function or an event in an event graph.
  * In C++ the entry point is the declaration: void main(...);
  */
-class RIGVMDEVELOPER_API FRigVMEntryExprAST : public FRigVMNodeExprAST
+class FRigVMEntryExprAST : public FRigVMNodeExprAST
 {
 public:
 
@@ -640,7 +642,7 @@ public:
 
 	// returns the name of the entry / event
 	// @return the name of the entry / event
-	FName GetEventName() const;
+	UE_API FName GetEventName() const;
 
 protected:
 
@@ -658,7 +660,7 @@ private:
  * An abstract syntax tree entry expression represents an invocation
  * of an entry point.
  */
-class RIGVMDEVELOPER_API FRigVMInvokeEntryExprAST : public FRigVMNodeExprAST
+class FRigVMInvokeEntryExprAST : public FRigVMNodeExprAST
 {
 public:
 
@@ -680,7 +682,7 @@ public:
 
 	// returns the name of the entry / event
 	// @return the name of the entry / event
-	FName GetEventName() const;
+	UE_API FName GetEventName() const;
 
 protected:
 
@@ -701,7 +703,7 @@ private:
  * The call extern expression references a node (through parent class)
  * from the model providing all of the relevant information for the invocation.
  */
-class RIGVMDEVELOPER_API FRigVMCallExternExprAST: public FRigVMNodeExprAST
+class FRigVMCallExternExprAST: public FRigVMNodeExprAST
 {
 public:
 
@@ -732,7 +734,7 @@ private:
 /*
  *
  */
-class RIGVMDEVELOPER_API FRigVMInlineFunctionExprAST: public FRigVMNodeExprAST
+class FRigVMInlineFunctionExprAST: public FRigVMNodeExprAST
 {
 public:
 
@@ -770,7 +772,7 @@ private:
  * For the RigVM AST we use the no-op expression for representing reroute nodes
  * in the model as well as parameter and variable getter nodes.
  */
-class RIGVMDEVELOPER_API FRigVMNoOpExprAST : public FRigVMNodeExprAST
+class FRigVMNoOpExprAST : public FRigVMNodeExprAST
 {
 public:
 
@@ -804,7 +806,7 @@ private:
  * In C++ the var expression is a variable declaration: int A;
  * The var expression references a pin from the model.
  */
-class RIGVMDEVELOPER_API FRigVMVarExprAST: public FRigVMExprAST
+class FRigVMVarExprAST: public FRigVMExprAST
 {
 public:
 
@@ -825,7 +827,7 @@ public:
 		return true;
 	}
 
-	virtual bool IsConstant() const override;
+	UE_API virtual bool IsConstant() const override;
 
 	// returns the proxy this expression is using
 	virtual const FRigVMASTProxy& GetProxy() const override { return Proxy; }
@@ -836,33 +838,33 @@ public:
 
 	// returns the C++ data type of this variable
 	// @return the C++ data type of this variable
-	FString GetCPPType() const;
+	UE_API FString GetCPPType() const;
 
 	// returns the C++ data type object (ustruct / uenum)
 	// @return the C++ data type object (ustruct / uenum)
-	UObject* GetCPPTypeObject() const;
+	UE_API UObject* GetCPPTypeObject() const;
 
 	// returns the pin direction of this variable (input, output, hidden etc)
 	// @return the pin direction of this variable (input, output, hidden etc)
-	ERigVMPinDirection GetPinDirection() const;
+	UE_API ERigVMPinDirection GetPinDirection() const;
 
 	// returns the default value on the pin for this variable
 	// @return the default value on the pin for this variable
-	FString GetDefaultValue() const;
+	UE_API FString GetDefaultValue() const;
 
 	// returns true if this variable is an execute context
 	// @return true if this variable is an execute context
-	bool IsExecuteContext() const;
+	UE_API bool IsExecuteContext() const;
 
 	// returns true if this variable is a graph variable
 	// @return true if this variable is a graph variable
-	bool IsGraphVariable() const;
+	UE_API bool IsGraphVariable() const;
 
 	// returns true if this is a constant enum index
-	bool IsEnumValue() const;
+	UE_API bool IsEnumValue() const;
 
 	// returns true if this variable allows links to be "soft", so without a cache / computed value.
-	bool SupportsSoftLinks() const;
+	UE_API bool SupportsSoftLinks() const;
 
 protected:
 
@@ -887,7 +889,7 @@ private:
  * In C++ the literal expression is a literal declaration, for ex: const float PI = 3.14f;
  * The literal expression references a pin from the model.
  */
-class RIGVMDEVELOPER_API FRigVMLiteralExprAST : public FRigVMVarExprAST
+class FRigVMLiteralExprAST : public FRigVMVarExprAST
 {
 public:
 
@@ -928,7 +930,7 @@ private:
  * An abstract syntax tree external variable expression represents the reference 
  * of unowned / external memory.
  */
-class RIGVMDEVELOPER_API FRigVMExternalVarExprAST : public FRigVMVarExprAST
+class FRigVMExternalVarExprAST : public FRigVMVarExprAST
 {
 public:
 
@@ -967,7 +969,7 @@ private:
  * In C++ the assign expression used for construction and copy, for ex: int32 A = B + 1;
  * The assign expression references two pins / a link from the model.
  */
-class RIGVMDEVELOPER_API FRigVMAssignExprAST : public FRigVMExprAST
+class FRigVMAssignExprAST : public FRigVMExprAST
 {
 public:
 
@@ -1024,7 +1026,7 @@ private:
  * In C++ the copy expression is used for structures, for ex: FVector A = B;
  * The copy expression references two pins / a link from the model.
  */
-class RIGVMDEVELOPER_API FRigVMCopyExprAST : public FRigVMAssignExprAST
+class FRigVMCopyExprAST : public FRigVMAssignExprAST
 {
 public:
 
@@ -1063,7 +1065,7 @@ private:
  * The first child of the cached value expression is the var expression to be
  * computed / cached, the second child is the call extern expression to use.
  */
-class RIGVMDEVELOPER_API FRigVMCachedValueExprAST : public FRigVMExprAST
+class FRigVMCachedValueExprAST : public FRigVMExprAST
 {
 public:
 
@@ -1104,7 +1106,7 @@ private:
  * An abstract syntax tree exit expression represents the exit out of an entry expression.
  * In C++ the exit expression is a return from a main function.
  */
-class RIGVMDEVELOPER_API FRigVMExitExprAST : public FRigVMExprAST
+class FRigVMExitExprAST : public FRigVMExprAST
 {
 public:
 
@@ -1142,7 +1144,7 @@ private:
  * The folding settings can affect the performance of the parse dramatically.
  */
 USTRUCT(BlueprintType)
-struct RIGVMDEVELOPER_API FRigVMParserASTSettings
+struct FRigVMParserASTSettings
 {
 	GENERATED_BODY()
 
@@ -1187,7 +1189,7 @@ struct RIGVMDEVELOPER_API FRigVMParserASTSettings
 		return Settings;
 	}
 
-	void Report(EMessageSeverity::Type InSeverity, UObject* InSubject, const FString& InMessage) const;
+	UE_API void Report(EMessageSeverity::Type InSeverity, UObject* InSubject, const FString& InMessage) const;
 
 	template <typename... Types>
 		void Reportf(EMessageSeverity::Type InSeverity,
@@ -1206,17 +1208,17 @@ struct RIGVMDEVELOPER_API FRigVMParserASTSettings
  * The abstract syntax tree is then fed into the RigVMCompiler to 
  * generate the byte code for the virtual machine.
  */
-class RIGVMDEVELOPER_API FRigVMParserAST : public TSharedFromThis<FRigVMParserAST>
+class FRigVMParserAST : public TSharedFromThis<FRigVMParserAST>
 {
 public:
 
 	// default constructor
 	// @param InGraph The graph / model to parse
 	// @param InSettings The parse settings to use
-	FRigVMParserAST(TArray<URigVMGraph*> InGraphs, URigVMController* InController = nullptr, const FRigVMParserASTSettings& InSettings = FRigVMParserASTSettings::Fast(), const TArray<FRigVMExternalVariable>& InExternalVariables = TArray<FRigVMExternalVariable>());
+	UE_API FRigVMParserAST(TArray<URigVMGraph*> InGraphs, URigVMController* InController = nullptr, const FRigVMParserASTSettings& InSettings = FRigVMParserASTSettings::Fast(), const TArray<FRigVMExternalVariable>& InExternalVariables = TArray<FRigVMExternalVariable>());
 
 	// default destructor
-	~FRigVMParserAST();
+	UE_API ~FRigVMParserAST();
 
 	// returns the number of root expressions
 	int32 Num() const { return RootExpressions.Num(); }
@@ -1240,39 +1242,39 @@ public:
 	// returns the expression for a given subject. subjects include nodes and pins.
 	// @param InSubject the subject to retrieve the expression for (node or pin)
 	// @return the expressoin for the given subject (or nullptr)
-	const FRigVMExprAST* GetExprForSubject(const FRigVMASTProxy& InProxy) const;
+	UE_API const FRigVMExprAST* GetExprForSubject(const FRigVMASTProxy& InProxy) const;
 
 	// returns all expressions for a given subject. subjects include nodes and pins.
 	// @param InSubject the subject to retrieve the expression for (node or pin)
 	// @return all expressions for the given subject
-	TArray<const FRigVMExprAST*> GetExpressionsForSubject(UObject* InSubject) const;
+	UE_API TArray<const FRigVMExprAST*> GetExpressionsForSubject(UObject* InSubject) const;
 
 	// Prepares the parser for cycle checking on a given pin.
 	// This marks up the parents and childen of the corresponding expression in the graph,
 	// to allow the client to determine if a new parent / child relationship could cause a cycle.
 	// @param InPin the pin to initialize the cycle checking for
-	void PrepareCycleChecking(URigVMPin* InPin);
+	UE_API void PrepareCycleChecking(URigVMPin* InPin);
 
 	// Performs a cycle check for a new potential link (assign or copy) between two pins.
 	// @param InSourcePin the source (left) pin of the potential assign / copy
 	// @param InTargetPin the target (right) pin of the potential assign / copy
 	// @param OutFailureReason an optional storage for the possible failure reason description
 	// @return true if the potential link (assign / copy) can be established
-	bool CanLink(URigVMPin* InSourcePin, URigVMPin* InTargetPin, FString* OutFailureReason);
+	UE_API bool CanLink(URigVMPin* InSourcePin, URigVMPin* InTargetPin, FString* OutFailureReason);
 
 	// returns a string containing an indented tree structure
 	// for debugging purposes.
 	// @return the simple text representing this abstract syntax tree
-	FString DumpText() const;
+	UE_API FString DumpText() const;
 
 	// returns a string containing a dot file notation
 	// for debugging purposes.
 	// @return the dot file text representing this abstract syntax tree
-	FString DumpDot() const;
+	UE_API FString DumpDot() const;
 
 	// returns an obsolete block for unmanaged expression
-	FRigVMBlockExprAST* GetObsoleteBlock(bool bCreateIfMissing = true);
-	const FRigVMBlockExprAST* GetObsoleteBlock(bool bCreateIfMissing = true) const;
+	UE_API FRigVMBlockExprAST* GetObsoleteBlock(bool bCreateIfMissing = true);
+	UE_API const FRigVMBlockExprAST* GetObsoleteBlock(bool bCreateIfMissing = true) const;
 
 	// returns the AST's override table for pin defaults
 	const URigVMPin::FPinOverrideMap& GetPinOverrides() const { return PinOverrides; }
@@ -1354,75 +1356,75 @@ private:
 
 	// removes an array of expressions from the parser
 	// @param InExprs the expressions to remove
-	void RemoveExpressions(TArray<FRigVMExprAST*> InExprs);
+	UE_API void RemoveExpressions(TArray<FRigVMExprAST*> InExprs);
 
 	// a static helper function to traverse along all parents of an expression,
 	// provided a predicate to return true if the traverse should continue,
 	// and false if the traverse should stop.
 	// @param InExpr the expression to start the traverse at
 	// @param InContinuePredicate the predicate to use to break out of the traverse (return false)
-	static void TraverseParents(const FRigVMExprAST* InExpr, TFunctionRef<bool(const FRigVMExprAST*)> InContinuePredicate);
+	static UE_API void TraverseParents(const FRigVMExprAST* InExpr, TFunctionRef<bool(const FRigVMExprAST*)> InContinuePredicate);
 
 	// a static helper function to traverse along all children of an expression,
 	// provided a predicate to return true if the traverse should continue,
 	// and false if the traverse should stop.
 	// @param InExpr the expression to start the traverse at
 	// @param InContinuePredicate the predicate to use to break out of the traverse (return false)
-	static void TraverseChildren(const FRigVMExprAST* InExpr, TFunctionRef<bool(const FRigVMExprAST*)> InContinuePredicate);
+	static UE_API void TraverseChildren(const FRigVMExprAST* InExpr, TFunctionRef<bool(const FRigVMExprAST*)> InContinuePredicate);
 
 	// helper function to fold all entries with the same event name into one block
-	void FoldEntries();
+	UE_API void FoldEntries();
 
 	// helper function to inject an exit expression at the end of every entry expressoin
-	void InjectExitsToEntries();
+	UE_API void InjectExitsToEntries();
 
 	// helper function to refresh the expression indices (used after deleting an expression)
-	void RefreshExprIndices();
+	UE_API void RefreshExprIndices();
 
 	// helper function to fold / remove the no op expressions (parameters mostly)
-	void FoldNoOps();
+	UE_API void FoldNoOps();
 
 	// helper function to fold / merge redundant literals with the same value
-	void FoldLiterals();
+	UE_API void FoldLiterals();
 
 	// helper function to fold / remove obsolete assignments and reduce assignment chains
-	void FoldAssignments();
+	UE_API void FoldAssignments();
 
 	// helper function to inline all contributing nodes of the graph
-	void Inline(const TArray<URigVMGraph*>& InGraphs);
-	void Inline(const TArray<URigVMGraph*>& InGraphs, const TArray<FRigVMASTProxy>& InNodeProxies);
+	UE_API void Inline(const TArray<URigVMGraph*>& InGraphs);
+	UE_API void Inline(const TArray<URigVMGraph*>& InGraphs, const TArray<FRigVMASTProxy>& InNodeProxies);
 
 	// helper functions to retrieve links for a given pin
-	TArray<int32> GetSourceLinkIndices(const FRigVMASTProxy& InPinProxy, bool bRecursive = false) const;
-	TArray<int32> GetTargetLinkIndices(const FRigVMASTProxy& InPinProxy, bool bRecursive = false) const;
-	TArray<int32> GetLinkIndices(const FRigVMASTProxy& InPinProxy, bool bGetSource, bool bRecursive) const;
+	UE_API TArray<int32> GetSourceLinkIndices(const FRigVMASTProxy& InPinProxy, bool bRecursive = false) const;
+	UE_API TArray<int32> GetTargetLinkIndices(const FRigVMASTProxy& InPinProxy, bool bRecursive = false) const;
+	UE_API TArray<int32> GetLinkIndices(const FRigVMASTProxy& InPinProxy, bool bGetSource, bool bRecursive) const;
 
 	// function to retrieve a link given its index
-	const FRigVMASTLinkDescription& GetLink(int32 InLinkIndex) const;
+	UE_API const FRigVMASTLinkDescription& GetLink(int32 InLinkIndex) const;
 
 	// traverse a single mutable node (constructs entry, call extern and other expressions)
-	FRigVMExprAST* TraverseMutableNode(const FRigVMASTProxy& InNodeProxy, FRigVMExprAST* InParentExpr);
+	UE_API FRigVMExprAST* TraverseMutableNode(const FRigVMASTProxy& InNodeProxy, FRigVMExprAST* InParentExpr);
 
 	// traverse a single pure node (constructs call extern expressions)
-	FRigVMExprAST* TraverseNode(const FRigVMASTProxy& InNodeProxy, FRigVMExprAST* InParentExpr);
+	UE_API FRigVMExprAST* TraverseNode(const FRigVMASTProxy& InNodeProxy, FRigVMExprAST* InParentExpr);
 
 	// returns the expression for a given node
-	FRigVMExprAST* CreateExpressionForNode(const FRigVMASTProxy& InNodeProxy, FRigVMExprAST* InParentExpr);
+	UE_API FRigVMExprAST* CreateExpressionForNode(const FRigVMASTProxy& InNodeProxy, FRigVMExprAST* InParentExpr);
 
 	// traverse an array of pins for a given node
-	TArray<FRigVMExprAST*> TraversePins(const FRigVMASTProxy& InNodeProxy, FRigVMExprAST* InParentExpr);
+	UE_API TArray<FRigVMExprAST*> TraversePins(const FRigVMASTProxy& InNodeProxy, FRigVMExprAST* InParentExpr);
 
 	// traverse a single pin (constructs var + literal expressions)
-	FRigVMExprAST* TraversePin(const FRigVMASTProxy& InPinProxy, FRigVMExprAST* InParentExpr);
+	UE_API FRigVMExprAST* TraversePin(const FRigVMASTProxy& InPinProxy, FRigVMExprAST* InParentExpr);
 
 	// traverse a single link (constructs assign + copy expressions)
-	FRigVMExprAST* TraverseLink(int32 InLinkIndex, FRigVMExprAST* InParentExpr);
+	UE_API FRigVMExprAST* TraverseLink(int32 InLinkIndex, FRigVMExprAST* InParentExpr);
 
-	bool ShouldLinkBeSkipped(const FRigVMASTLinkDescription& InLink) const;
+	UE_API bool ShouldLinkBeSkipped(const FRigVMASTLinkDescription& InLink) const;
 
-	static FString GetLinkAsString(const FRigVMASTLinkDescription& InLink);
+	static UE_API FString GetLinkAsString(const FRigVMASTLinkDescription& InLink);
 
-	void IncrementCacheVersion() const;
+	UE_API void IncrementCacheVersion() const;
 
 	TMap<FRigVMASTProxy, FRigVMExprAST*> SubjectToExpression;
 	TMap<FRigVMASTProxy, int32> NodeExpressionIndex;
@@ -1438,7 +1440,7 @@ private:
 	TMap<FRigVMASTProxy, TArray<int32>> TargetLinkIndices;
 	TMap<FRigVMASTProxy, TArray<int32>> SourceLinkIndices;
 	TArray<FRigVMASTLinkDescription> Links;
-	static const TArray<FRigVMASTProxy> EmptyProxyArray;
+	static UE_API const TArray<FRigVMASTProxy> EmptyProxyArray;
 	URigVMPin::FPinOverrideMap PinOverrides;
 
 	enum ETraverseRelationShip
@@ -1463,3 +1465,4 @@ private:
 	friend class URigVMCompiler;
 };
 
+#undef UE_API
diff --git a/Engine/Plugins/Runtime/RigVM/Source/RigVMDeveloper/Public/RigVMCompiler/RigVMASTProxy.h b/Engine/Plugins/Runtime/RigVM/Source/RigVMDeveloper/Public/RigVMCompiler/RigVMASTProxy.h
--- a/Engine/Plugins/Runtime/RigVM/Source/RigVMDeveloper/Public/RigVMCompiler/RigVMASTProxy.h
+++ b/Engine/Plugins/Runtime/RigVM/Source/RigVMDeveloper/Public/RigVMCompiler/RigVMASTProxy.h
@@ -6,28 +6,30 @@
 #include "Templates/Casts.h"
 #include "UObject/UObjectBaseUtility.h"
 
+#define UE_API RIGVMDEVELOPER_API
+
 /*
  * A stack describing the model object which led to this invocation.
  * Objects on the stack can be URigVMGraph or URigVMNode.
  */
-class RIGVMDEVELOPER_API FRigVMCallstack
+class FRigVMCallstack
 {
 public:
 
 	// Returns the path
-	FString GetCallPath(bool bIncludeLast = true) const;
+	UE_API FString GetCallPath(bool bIncludeLast = true) const;
 
 	// Returns the number of entries on this callstack
-	int32 Num() const;
+	UE_API int32 Num() const;
 
 	// Returns the last object in the callstack
-	const UObject* Last() const;
+	UE_API const UObject* Last() const;
 
 	// Returns the element at a given index
-	const UObject* operator[](int32 InIndex) const;
+	UE_API const UObject* operator[](int32 InIndex) const;
 
 	// Returns true if the stack contains a given entry
-	bool Contains(const UObject* InEntry) const;
+	UE_API bool Contains(const UObject* InEntry) const;
 
 	friend uint32 GetTypeHash(const FRigVMCallstack& Callstack)
 	{
@@ -97,7 +99,7 @@ public:
 
 	TArray<TWeakObjectPtr<UObject>> GetStack() const { return Stack; }
 
-	FRigVMCallstack GetCallStackUpTo(int32 InIndex) const;
+	UE_API FRigVMCallstack GetCallStackUpTo(int32 InIndex) const;
 
 private:
 
@@ -127,7 +129,7 @@ private:
  * a graph. The subject can be a pin, a node, a link, etc.
  * The context is the callstack under which is occured.
  */
-class RIGVMDEVELOPER_API FRigVMASTProxy
+class FRigVMASTProxy
 {
 public:
 
@@ -159,10 +161,10 @@ public:
 		return *this;
 	}
 
-	static FRigVMASTProxy MakeFromUObject(UObject* InSubject);
-	static FRigVMASTProxy MakeFromCallPath(const FString& InCallPath, UObject* InRootObject);
-	static FRigVMASTProxy MakeFromCallstack(const FRigVMCallstack& InCallstack);
-	static FRigVMASTProxy MakeFromCallstack(const TArray<TWeakObjectPtr<UObject>>* InCallstack);
+	static UE_API FRigVMASTProxy MakeFromUObject(UObject* InSubject);
+	static UE_API FRigVMASTProxy MakeFromCallPath(const FString& InCallPath, UObject* InRootObject);
+	static UE_API FRigVMASTProxy MakeFromCallstack(const FRigVMCallstack& InCallstack);
+	static UE_API FRigVMASTProxy MakeFromCallstack(const TArray<TWeakObjectPtr<UObject>>* InCallstack);
 
 	FRigVMASTProxy GetSibling(UObject* InSubject) const
 	{
@@ -253,3 +255,5 @@ private:
 
 typedef TPair<FRigVMASTProxy, FRigVMASTProxy> FRigVMPinProxyPair;
 typedef TMap<FRigVMASTProxy, FString> FRigVMPinDefaultValueOverride;
+
+#undef UE_API
diff --git a/Engine/Plugins/Runtime/RigVM/Source/RigVMDeveloper/Public/RigVMCompiler/RigVMCodeGenerator.h b/Engine/Plugins/Runtime/RigVM/Source/RigVMDeveloper/Public/RigVMCompiler/RigVMCodeGenerator.h
--- a/Engine/Plugins/Runtime/RigVM/Source/RigVMDeveloper/Public/RigVMCompiler/RigVMCodeGenerator.h
+++ b/Engine/Plugins/Runtime/RigVM/Source/RigVMDeveloper/Public/RigVMCompiler/RigVMCodeGenerator.h
@@ -7,10 +7,12 @@
 
 #include "UObject/StrongObjectPtr.h"
 
+#define UE_API RIGVMDEVELOPER_API
+
 class URigVMGraph;
 
 
-struct RIGVMDEVELOPER_API FRigVMCodeGenerator
+struct FRigVMCodeGenerator
 {
 public:
 
@@ -22,28 +24,28 @@ public:
 		ParseVM(InClassName, InModuleName, InModelToNativize, InVMToNativize, InVMContext, PublicContextStruct, InPinToOperandMap, InMaxInstructionsPerFunction);
 	}
 
-	FString DumpIncludes(bool bLog = false);
+	UE_API FString DumpIncludes(bool bLog = false);
 	UE_DEPRECATED(5.4, "Please, use DumpExternalVariables with ExtendedExecuteContext parameter")
 	FString DumpExternalVariables(bool bForHeader, bool bLog = false) { return FString(); }
-	FString DumpExternalVariables(const FRigVMExtendedExecuteContext& Context, bool bForHeader, bool bLog = false);
-	FString DumpEntries(bool bForHeader, bool bLog = false);
-	FString DumpBlockNames(bool bForHeader, bool bLog = false);
+	UE_API FString DumpExternalVariables(const FRigVMExtendedExecuteContext& Context, bool bForHeader, bool bLog = false);
+	UE_API FString DumpEntries(bool bForHeader, bool bLog = false);
+	UE_API FString DumpBlockNames(bool bForHeader, bool bLog = false);
 	UE_DEPRECATED(5.4, "Please, use DumpProperties with ExtendedExecuteContext parameter")
 	FString DumpProperties(bool bForHeader, int32 InInstructionGroup, bool bLog = false) { return FString(); }
-	FString DumpProperties(const FRigVMExtendedExecuteContext& Context, bool bForHeader, int32 InInstructionGroup, bool bLog = false);
-	FString DumpDispatches(bool bLog = false);
-	FString DumpRequiredUProperties(bool bLog = false);
-	FString DumpInitialize(bool bLog = false);
+	UE_API FString DumpProperties(const FRigVMExtendedExecuteContext& Context, bool bForHeader, int32 InInstructionGroup, bool bLog = false);
+	UE_API FString DumpDispatches(bool bLog = false);
+	UE_API FString DumpRequiredUProperties(bool bLog = false);
+	UE_API FString DumpInitialize(bool bLog = false);
 	UE_DEPRECATED(5.4, "Please, use DumpInstructions with ExtendedExecuteContext parameter")
 	FString DumpInstructions(int32 InInstructionGroup, bool bLog = false) { return FString(); }
-	FString DumpInstructions(const FRigVMExtendedExecuteContext& Context, int32 InInstructionGroup, bool bLog = false);
+	UE_API FString DumpInstructions(const FRigVMExtendedExecuteContext& Context, int32 InInstructionGroup, bool bLog = false);
 	UE_DEPRECATED(5.4, "Please, use DumpHeader with ExtendedExecuteContext parameter")
 	FString DumpHeader(bool bLog = false) { return FString(); }
-	FString DumpHeader(const FRigVMExtendedExecuteContext& Context, bool bLog = false);
+	UE_API FString DumpHeader(const FRigVMExtendedExecuteContext& Context, bool bLog = false);
 	UE_DEPRECATED(5.4, "Please, use DumpSource with ExtendedExecuteContext parameter")
 	FString DumpSource(bool bLog = false) { return FString(); }
-	FString DumpSource(const FRigVMExtendedExecuteContext& Context, bool bLog = false);
-	FString DumpLines(const TArray<FString>& InLines, bool bLog = false);
+	UE_API FString DumpSource(const FRigVMExtendedExecuteContext& Context, bool bLog = false);
+	UE_API FString DumpLines(const TArray<FString>& InLines, bool bLog = false);
 	
 private:
 
@@ -98,39 +100,39 @@ private:
 		{}
 	};
 
-	void Reset();
-	void ParseVM(const FString& InClassName, const FString& InModuleName,
+	UE_API void Reset();
+	UE_API void ParseVM(const FString& InClassName, const FString& InModuleName,
 		URigVMGraph* InModelToNativize, URigVM* InVMToNativize, FRigVMExtendedExecuteContext& InVMContext,
 		const UScriptStruct* PublicContextStruct, TMap<FString,FRigVMOperand> InPinToOperandMap,
 		int32 InMaxInstructionsPerFunction);
-	void ParseInclude(UStruct* InDependency, const FName& InMethodName = NAME_None);
-	void ParseRequiredUProperties(const FRigVMExtendedExecuteContext& Context);
-	void ParseMemory(const FRigVMExtendedExecuteContext& Context, FRigVMMemoryStorageStruct* InMemory);
-	void ParseProperty(const FRigVMExtendedExecuteContext& Context, ERigVMMemoryType InMemoryType, const FProperty* InProperty, FRigVMMemoryStorageStruct* InMemory);
-	void ParseInstructionGroups(const FRigVMExtendedExecuteContext& Context);
-	FString DumpInstructions(const FRigVMExtendedExecuteContext& Context, const FString& InPrefix, int32 InFirstInstruction, int32 InLastInstruction, const FInstructionGroup& InGroup, bool bLog = false);
-	FString DumpInstructions(const FRigVMExtendedExecuteContext& Context, const FString& InPrefix, const TArray<int32> InInstructionIndices, const FInstructionGroup& InGroup, bool bLog = false);
-	FString GetOperandName(const FRigVMExtendedExecuteContext& Context, const FRigVMOperand& InOperand, bool bPerSlice = true, bool bAsInput = true) const;
-	FString GetOperandCPPType(const FRigVMExtendedExecuteContext& Context, const FRigVMOperand& InOperand) const;
-	FString GetOperandCPPBaseType(const FRigVMExtendedExecuteContext& Context, const FRigVMOperand& InOperand) const;
-	static FString SanitizeName(const FString& InName, const FString& CPPType = FString()); 
-	static FString SanitizeValue(const FString& InValue, const FString& CPPType, const UObject* CPPTypeObject); 
-	FRigVMPropertyDescription GetPropertyDescForOperand(const FRigVMExtendedExecuteContext& Context, const FRigVMOperand& InOperand) const;
-	FRigVMPropertyDescription GetPropertyForOperand(const FRigVMExtendedExecuteContext& Context, const FRigVMOperand& InOperand) const;
-	const FRigVMPropertyPathDescription& GetPropertyPathForOperand(const FRigVMExtendedExecuteContext& Context, const FRigVMOperand& InOperand) const;
-	int32 GetPropertyIndex(const FRigVMPropertyDescription& InProperty) const;
-	int32 GetPropertyIndex(const FRigVMExtendedExecuteContext& Context, const FRigVMOperand& InOperand) const;
-	ERigVMNativizedPropertyType GetPropertyType(const FRigVMPropertyDescription& InProperty) const;
-	ERigVMNativizedPropertyType GetPropertyType(const FRigVMExtendedExecuteContext& Context, const FRigVMOperand& InOperand) const;
-	void CheckOperand(const FRigVMExtendedExecuteContext& Context, const FRigVMOperand& InOperand) const;
-	FString GetMappedType(const FString& InCPPType) const;
-	const FString& GetMappedTypeSuffix(const FString& InCPPType) const;
-	FString GetMappedArrayTypeName(const FString InBaseElementType) const;
-	const FInstructionGroup& GetGroup(int32 InGroupIndex) const;
-	bool IsInstructionPartOfGroup(int32 InInstructionIndex, int32 InGroupIndex, bool bIncludeChildGroups) const;
-	bool IsPropertyPartOfGroup(int32 InPropertyIndex, int32 InGroupIndex) const;
-	FString GetEntryParameters() const;
-	static TArray<int32> GetInstructionIndicesFromRange(int32 First, int32 Last);
+	UE_API void ParseInclude(UStruct* InDependency, const FName& InMethodName = NAME_None);
+	UE_API void ParseRequiredUProperties(const FRigVMExtendedExecuteContext& Context);
+	UE_API void ParseMemory(const FRigVMExtendedExecuteContext& Context, FRigVMMemoryStorageStruct* InMemory);
+	UE_API void ParseProperty(const FRigVMExtendedExecuteContext& Context, ERigVMMemoryType InMemoryType, const FProperty* InProperty, FRigVMMemoryStorageStruct* InMemory);
+	UE_API void ParseInstructionGroups(const FRigVMExtendedExecuteContext& Context);
+	UE_API FString DumpInstructions(const FRigVMExtendedExecuteContext& Context, const FString& InPrefix, int32 InFirstInstruction, int32 InLastInstruction, const FInstructionGroup& InGroup, bool bLog = false);
+	UE_API FString DumpInstructions(const FRigVMExtendedExecuteContext& Context, const FString& InPrefix, const TArray<int32> InInstructionIndices, const FInstructionGroup& InGroup, bool bLog = false);
+	UE_API FString GetOperandName(const FRigVMExtendedExecuteContext& Context, const FRigVMOperand& InOperand, bool bPerSlice = true, bool bAsInput = true) const;
+	UE_API FString GetOperandCPPType(const FRigVMExtendedExecuteContext& Context, const FRigVMOperand& InOperand) const;
+	UE_API FString GetOperandCPPBaseType(const FRigVMExtendedExecuteContext& Context, const FRigVMOperand& InOperand) const;
+	static UE_API FString SanitizeName(const FString& InName, const FString& CPPType = FString()); 
+	static UE_API FString SanitizeValue(const FString& InValue, const FString& CPPType, const UObject* CPPTypeObject); 
+	UE_API FRigVMPropertyDescription GetPropertyDescForOperand(const FRigVMExtendedExecuteContext& Context, const FRigVMOperand& InOperand) const;
+	UE_API FRigVMPropertyDescription GetPropertyForOperand(const FRigVMExtendedExecuteContext& Context, const FRigVMOperand& InOperand) const;
+	UE_API const FRigVMPropertyPathDescription& GetPropertyPathForOperand(const FRigVMExtendedExecuteContext& Context, const FRigVMOperand& InOperand) const;
+	UE_API int32 GetPropertyIndex(const FRigVMPropertyDescription& InProperty) const;
+	UE_API int32 GetPropertyIndex(const FRigVMExtendedExecuteContext& Context, const FRigVMOperand& InOperand) const;
+	UE_API ERigVMNativizedPropertyType GetPropertyType(const FRigVMPropertyDescription& InProperty) const;
+	UE_API ERigVMNativizedPropertyType GetPropertyType(const FRigVMExtendedExecuteContext& Context, const FRigVMOperand& InOperand) const;
+	UE_API void CheckOperand(const FRigVMExtendedExecuteContext& Context, const FRigVMOperand& InOperand) const;
+	UE_API FString GetMappedType(const FString& InCPPType) const;
+	UE_API const FString& GetMappedTypeSuffix(const FString& InCPPType) const;
+	UE_API FString GetMappedArrayTypeName(const FString InBaseElementType) const;
+	UE_API const FInstructionGroup& GetGroup(int32 InGroupIndex) const;
+	UE_API bool IsInstructionPartOfGroup(int32 InInstructionIndex, int32 InGroupIndex, bool bIncludeChildGroups) const;
+	UE_API bool IsPropertyPartOfGroup(int32 InPropertyIndex, int32 InGroupIndex) const;
+	UE_API FString GetEntryParameters() const;
+	static UE_API TArray<int32> GetInstructionIndicesFromRange(int32 First, int32 Last);
 
 	static FString FormatArgs(const TCHAR* InFormatString, const FStringFormatOrderedArguments& InArgs)
 	{
@@ -268,5 +270,7 @@ private:
 	TMap<FName, FName> MappedPropertyNames;
 	TMap<FString, FString> OverriddenOperatorNames;
 
-	static const TMap<FName, TStructConstGenerator>& GetStructConstGenerators();
+	static UE_API const TMap<FName, TStructConstGenerator>& GetStructConstGenerators();
 };
+
+#undef UE_API
diff --git a/Engine/Plugins/Runtime/RigVM/Source/RigVMDeveloper/Public/RigVMCompiler/RigVMCompiler.h b/Engine/Plugins/Runtime/RigVM/Source/RigVMDeveloper/Public/RigVMCompiler/RigVMCompiler.h
--- a/Engine/Plugins/Runtime/RigVM/Source/RigVMDeveloper/Public/RigVMCompiler/RigVMCompiler.h
+++ b/Engine/Plugins/Runtime/RigVM/Source/RigVMDeveloper/Public/RigVMCompiler/RigVMCompiler.h
@@ -18,21 +18,23 @@
 
 #include "RigVMCompiler.generated.h"
 
+#define UE_API RIGVMDEVELOPER_API
+
 class URigVMFunctionReferenceNode;
 struct FRigVMGraphFunctionData;
 struct FRigVMGraphFunctionHeader;
 struct FRigVMFunctionCompilationData;
 
 USTRUCT(BlueprintType)
-struct RIGVMDEVELOPER_API FRigVMCompileSettings
+struct FRigVMCompileSettings
 {
 	GENERATED_BODY()
 
 public:
 
-	FRigVMCompileSettings();
+	UE_API FRigVMCompileSettings();
 
-	FRigVMCompileSettings(UScriptStruct* InExecuteContextScriptStruct);
+	UE_API FRigVMCompileSettings(UScriptStruct* InExecuteContextScriptStruct);
 
 	UScriptStruct* GetExecuteContextStruct() const { return ASTSettings.ExecuteContextStruct; }
 	void SetExecuteContextStruct(UScriptStruct* InExecuteContextScriptStruct) { ASTSettings.ExecuteContextStruct = InExecuteContextScriptStruct; }
@@ -85,9 +87,9 @@ public:
 		return Settings;
 	}
 
-	void ReportInfo(const FString& InMessage) const;
-	void ReportWarning(const FString& InMessage) const;
-	void ReportError(const FString& InMessage) const;
+	UE_API void ReportInfo(const FString& InMessage) const;
+	UE_API void ReportWarning(const FString& InMessage) const;
+	UE_API void ReportError(const FString& InMessage) const;
 
 	template <typename... Types>
 		void ReportInfof(UE::Core::TCheckedFormatString<FString::FmtCharType, Types...> Fmt, Types... Args) const
@@ -122,7 +124,7 @@ public:
 	}
 };
 
-class RIGVMDEVELOPER_API FRigVMCompileSettingsDuringLoadGuard
+class FRigVMCompileSettingsDuringLoadGuard
 {
 public:
 
@@ -135,7 +137,7 @@ private:
 	TGuardValue<bool> ASTErrorsAsNotifications;
 };
 
-struct RIGVMDEVELOPER_API FRigVMCompilerWorkData
+struct FRigVMCompilerWorkData
 {
 public:
 	FRigVMCompileSettings Settings;
@@ -186,13 +188,13 @@ public:
 	TMap<ERigVMMemoryType, TArray<FRigVMPropertyPathDescription>> PropertyPathDescriptions;
 	TMap<ERigVMMemoryType, TArray<FRigVMPropertyDescription>> PropertyDescriptions;
 
-	FRigVMOperand AddProperty(ERigVMMemoryType InMemoryType, const FName& InName, const FString& InCPPType, UObject* InCPPTypeObject, const FString& InDefaultValue = FString());
-	FRigVMOperand FindProperty(ERigVMMemoryType InMemoryType, const FName& InName) const;
-	FRigVMPropertyDescription GetProperty(const FRigVMOperand& InOperand);
-	int32 FindOrAddPropertyPath(const FRigVMOperand& InOperand, const FString& InHeadCPPType, const FString& InSegmentPath);
-	const FProperty* GetPropertyForOperand(const FRigVMOperand& InOperand) const;
-	TRigVMTypeIndex GetTypeIndexForOperand(const FRigVMOperand& InOperand) const;
-	FName GetUniquePropertyName(ERigVMMemoryType InMemoryType, const FName& InDesiredName) const;
+	UE_API FRigVMOperand AddProperty(ERigVMMemoryType InMemoryType, const FName& InName, const FString& InCPPType, UObject* InCPPTypeObject, const FString& InDefaultValue = FString());
+	UE_API FRigVMOperand FindProperty(ERigVMMemoryType InMemoryType, const FName& InName) const;
+	UE_API FRigVMPropertyDescription GetProperty(const FRigVMOperand& InOperand);
+	UE_API int32 FindOrAddPropertyPath(const FRigVMOperand& InOperand, const FString& InHeadCPPType, const FString& InSegmentPath);
+	UE_API const FProperty* GetPropertyForOperand(const FRigVMOperand& InOperand) const;
+	UE_API TRigVMTypeIndex GetTypeIndexForOperand(const FRigVMOperand& InOperand) const;
+	UE_API FName GetUniquePropertyName(ERigVMMemoryType InMemoryType, const FName& InDesiredName) const;
 
 	TSharedPtr<FRigVMParserAST> AST;
 	
@@ -229,9 +231,9 @@ public:
 	TArray<uint32> LazyBlocksToProcess;
 	TOptional<uint32> CurrentBlockHash;
 
-	void ReportInfo(const FString& InMessage) const;
-	void ReportWarning(const FString& InMessage) const;
-	void ReportError(const FString& InMessage) const;
+	UE_API void ReportInfo(const FString& InMessage) const;
+	UE_API void ReportWarning(const FString& InMessage) const;
+	UE_API void ReportError(const FString& InMessage) const;
 
 	template <typename... Types>
 	void ReportInfof(UE::Core::TCheckedFormatString<FString::FmtCharType, Types...> Fmt, Types... Args) const
@@ -252,8 +254,8 @@ public:
 	}
 
 	FRigVMReportDelegate OriginalReportDelegate;
-	void OverrideReportDelegate(bool& bEncounteredASTError, bool& bSurpressedASTError);
-	void RemoveOverrideReportDelegate();
+	UE_API void OverrideReportDelegate(bool& bEncounteredASTError, bool& bSurpressedASTError);
+	UE_API void RemoveOverrideReportDelegate();
 
 	void Clear()
 	{
@@ -291,14 +293,14 @@ public:
 
 DECLARE_DELEGATE_RetVal_OneParam(const FRigVMFunctionCompilationData*, FRigVMCompiler_GetFunctionCompilationData, const FRigVMGraphFunctionHeader& Header);
 
-UCLASS(BlueprintType)
-class RIGVMDEVELOPER_API URigVMCompiler : public UObject
+UCLASS(MinimalAPI, BlueprintType)
+class URigVMCompiler : public UObject
 {
 	GENERATED_BODY()
 
 public:
 
-	URigVMCompiler();
+	UE_API URigVMCompiler();
 
 	UPROPERTY()
 	FRigVMCompileSettings Settings_DEPRECATED;
@@ -316,58 +318,58 @@ public:
 	bool Compile(TArray<URigVMGraph*> InGraphs, URigVMController* InController, URigVM* OutVM, const TArray<FRigVMExternalVariable>& InExternalVariables, TMap<FString, FRigVMOperand>* OutOperands, TSharedPtr<FRigVMParserAST> InAST = TSharedPtr<FRigVMParserAST>(), FRigVMFunctionCompilationData* OutFunctionCompilationData = nullptr) { return false; }
 
 	UE_DEPRECATED(5.4, "Please use CompileFunction with Settings param.")
-	bool Compile(TArray<URigVMGraph*> InGraphs, URigVMController* InController, URigVM* OutVM, FRigVMExtendedExecuteContext& OutVMContext, const TArray<FRigVMExternalVariable>& InExternalVariables, TMap<FString, FRigVMOperand>* OutOperands, TSharedPtr<FRigVMParserAST> InAST = TSharedPtr<FRigVMParserAST>(), FRigVMFunctionCompilationData* OutFunctionCompilationData = nullptr);
+	UE_API bool Compile(TArray<URigVMGraph*> InGraphs, URigVMController* InController, URigVM* OutVM, FRigVMExtendedExecuteContext& OutVMContext, const TArray<FRigVMExternalVariable>& InExternalVariables, TMap<FString, FRigVMOperand>* OutOperands, TSharedPtr<FRigVMParserAST> InAST = TSharedPtr<FRigVMParserAST>(), FRigVMFunctionCompilationData* OutFunctionCompilationData = nullptr);
 
-	bool Compile(const FRigVMCompileSettings& InSettings, TArray<URigVMGraph*> InGraphs, URigVMController* InController, URigVM* OutVM, FRigVMExtendedExecuteContext& OutVMContext, const TArray<FRigVMExternalVariable>& InExternalVariables, TMap<FString, FRigVMOperand>* OutOperands, TSharedPtr<FRigVMParserAST> InAST = TSharedPtr<FRigVMParserAST>(), FRigVMFunctionCompilationData* OutFunctionCompilationData = nullptr);
+	UE_API bool Compile(const FRigVMCompileSettings& InSettings, TArray<URigVMGraph*> InGraphs, URigVMController* InController, URigVM* OutVM, FRigVMExtendedExecuteContext& OutVMContext, const TArray<FRigVMExternalVariable>& InExternalVariables, TMap<FString, FRigVMOperand>* OutOperands, TSharedPtr<FRigVMParserAST> InAST = TSharedPtr<FRigVMParserAST>(), FRigVMFunctionCompilationData* OutFunctionCompilationData = nullptr);
 
 	UE_DEPRECATED(5.3, "Please use CompileFunction with Context param.")
 	bool CompileFunction(const URigVMLibraryNode* InLibraryNode, URigVMController* InController, FRigVMFunctionCompilationData* OutFunctionCompilationData) { return false; }
 
 	UE_DEPRECATED(5.4, "Please use CompileFunction with Settings param.")
-	bool CompileFunction(const URigVMLibraryNode* InLibraryNode, URigVMController* InController, FRigVMFunctionCompilationData* OutFunctionCompilationData, FRigVMExtendedExecuteContext& OutVMContext);
+	UE_API bool CompileFunction(const URigVMLibraryNode* InLibraryNode, URigVMController* InController, FRigVMFunctionCompilationData* OutFunctionCompilationData, FRigVMExtendedExecuteContext& OutVMContext);
 
-	bool CompileFunction(const FRigVMCompileSettings& InSettings, const URigVMLibraryNode* InLibraryNode, URigVMController* InController, const TArray<FRigVMExternalVariable>& InExternalVariables, FRigVMFunctionCompilationData* OutFunctionCompilationData, FRigVMExtendedExecuteContext& OutVMContext);
+	UE_API bool CompileFunction(const FRigVMCompileSettings& InSettings, const URigVMLibraryNode* InLibraryNode, URigVMController* InController, const TArray<FRigVMExternalVariable>& InExternalVariables, FRigVMFunctionCompilationData* OutFunctionCompilationData, FRigVMExtendedExecuteContext& OutVMContext);
 
 	FRigVMCompiler_GetFunctionCompilationData GetFunctionCompilationData;
 	TMap<FString, const FRigVMFunctionCompilationData*> CompiledFunctions;
 
-	static UScriptStruct* GetScriptStructForCPPType(const FString& InCPPType);
-	static FString GetPinHash(const URigVMPin* InPin, const FRigVMVarExprAST* InVarExpr, bool bIsDebugValue = false, const URigVMLibraryNode* FunctionCompiling = nullptr, const FRigVMASTProxy& InPinProxy = FRigVMASTProxy());
+	static UE_API UScriptStruct* GetScriptStructForCPPType(const FString& InCPPType);
+	static UE_API FString GetPinHash(const URigVMPin* InPin, const FRigVMVarExprAST* InVarExpr, bool bIsDebugValue = false, const URigVMLibraryNode* FunctionCompiling = nullptr, const FRigVMASTProxy& InPinProxy = FRigVMASTProxy());
 
 	// follows assignment expressions to find the source ref counted containers
 	// since ref counted containers are not copied for assignments.
 	// this is currently only used for arrays.
-	static const FRigVMVarExprAST* GetSourceVarExpr(const FRigVMExprAST* InExpr);
+	static UE_API const FRigVMVarExprAST* GetSourceVarExpr(const FRigVMExprAST* InExpr);
 
-	void MarkDebugWatch(const FRigVMCompileSettings& InSettings, bool bRequired, URigVMPin* InPin, URigVM* OutVM, TMap<FString, FRigVMOperand>* OutOperands, TSharedPtr<FRigVMParserAST> InRuntimeAST);
+	UE_API void MarkDebugWatch(const FRigVMCompileSettings& InSettings, bool bRequired, URigVMPin* InPin, URigVM* OutVM, TMap<FString, FRigVMOperand>* OutOperands, TSharedPtr<FRigVMParserAST> InRuntimeAST);
 
 private:
 
 	const URigVMLibraryNode* CurrentCompilationFunction = nullptr;
 	TSet<const URigVMLibraryNode*> FunctionCompilationStack;
 
-	TArray<URigVMPin*> GetLinkedPins(URigVMPin* InPin, bool bInputs = true, bool bOutputs = true, bool bRecursive = true);
-	int32 GetElementSizeFromCPPType(const FString& InCPPType, UScriptStruct* InScriptStruct);
-
-	static FString GetPinHashImpl(const URigVMPin* InPin, const FRigVMVarExprAST* InVarExpr, bool bIsDebugValue = false, const URigVMLibraryNode* FunctionCompiling = nullptr, const FRigVMASTProxy& InPinProxy = FRigVMASTProxy());
-
-	bool TraverseExpression(const FRigVMExprAST* InExpr, FRigVMCompilerWorkData& WorkData);
-	bool TraverseChildren(const FRigVMExprAST* InExpr, FRigVMCompilerWorkData& WorkData);
-	bool TraverseBlock(const FRigVMBlockExprAST* InExpr, FRigVMCompilerWorkData& WorkData);
-	bool TraverseEntry(const FRigVMEntryExprAST* InExpr, FRigVMCompilerWorkData& WorkData);
-	bool TraverseCallExtern(const FRigVMCallExternExprAST* InExpr, FRigVMCompilerWorkData& WorkData);
-	bool TraverseInlineFunction(const FRigVMInlineFunctionExprAST* InExpr, FRigVMCompilerWorkData& WorkData);
-	bool TraverseNoOp(const FRigVMNoOpExprAST* InExpr, FRigVMCompilerWorkData& WorkData);
-	bool TraverseVar(const FRigVMVarExprAST* InExpr, FRigVMCompilerWorkData& WorkData);
-	bool TraverseLiteral(const FRigVMVarExprAST* InExpr, FRigVMCompilerWorkData& WorkData);
-	bool TraverseExternalVar(const FRigVMExternalVarExprAST* InExpr, FRigVMCompilerWorkData& WorkData);
-	bool TraverseAssign(const FRigVMAssignExprAST* InExpr, FRigVMCompilerWorkData& WorkData);
-	bool TraverseCopy(const FRigVMCopyExprAST* InExpr, FRigVMCompilerWorkData& WorkData);
-	bool TraverseCachedValue(const FRigVMCachedValueExprAST* InExpr, FRigVMCompilerWorkData& WorkData);
-	bool TraverseExit(const FRigVMExitExprAST* InExpr, FRigVMCompilerWorkData& WorkData);
-	bool TraverseInvokeEntry(const FRigVMInvokeEntryExprAST* InExpr, FRigVMCompilerWorkData& WorkData);
-
-	void AddCopyOperator(
+	UE_API TArray<URigVMPin*> GetLinkedPins(URigVMPin* InPin, bool bInputs = true, bool bOutputs = true, bool bRecursive = true);
+	UE_API int32 GetElementSizeFromCPPType(const FString& InCPPType, UScriptStruct* InScriptStruct);
+
+	static UE_API FString GetPinHashImpl(const URigVMPin* InPin, const FRigVMVarExprAST* InVarExpr, bool bIsDebugValue = false, const URigVMLibraryNode* FunctionCompiling = nullptr, const FRigVMASTProxy& InPinProxy = FRigVMASTProxy());
+
+	UE_API bool TraverseExpression(const FRigVMExprAST* InExpr, FRigVMCompilerWorkData& WorkData);
+	UE_API bool TraverseChildren(const FRigVMExprAST* InExpr, FRigVMCompilerWorkData& WorkData);
+	UE_API bool TraverseBlock(const FRigVMBlockExprAST* InExpr, FRigVMCompilerWorkData& WorkData);
+	UE_API bool TraverseEntry(const FRigVMEntryExprAST* InExpr, FRigVMCompilerWorkData& WorkData);
+	UE_API bool TraverseCallExtern(const FRigVMCallExternExprAST* InExpr, FRigVMCompilerWorkData& WorkData);
+	UE_API bool TraverseInlineFunction(const FRigVMInlineFunctionExprAST* InExpr, FRigVMCompilerWorkData& WorkData);
+	UE_API bool TraverseNoOp(const FRigVMNoOpExprAST* InExpr, FRigVMCompilerWorkData& WorkData);
+	UE_API bool TraverseVar(const FRigVMVarExprAST* InExpr, FRigVMCompilerWorkData& WorkData);
+	UE_API bool TraverseLiteral(const FRigVMVarExprAST* InExpr, FRigVMCompilerWorkData& WorkData);
+	UE_API bool TraverseExternalVar(const FRigVMExternalVarExprAST* InExpr, FRigVMCompilerWorkData& WorkData);
+	UE_API bool TraverseAssign(const FRigVMAssignExprAST* InExpr, FRigVMCompilerWorkData& WorkData);
+	UE_API bool TraverseCopy(const FRigVMCopyExprAST* InExpr, FRigVMCompilerWorkData& WorkData);
+	UE_API bool TraverseCachedValue(const FRigVMCachedValueExprAST* InExpr, FRigVMCompilerWorkData& WorkData);
+	UE_API bool TraverseExit(const FRigVMExitExprAST* InExpr, FRigVMCompilerWorkData& WorkData);
+	UE_API bool TraverseInvokeEntry(const FRigVMInvokeEntryExprAST* InExpr, FRigVMCompilerWorkData& WorkData);
+
+	UE_API void AddCopyOperator(
 		const FRigVMCopyOp& InOp,
 		const FRigVMAssignExprAST* InAssignExpr,
 		const FRigVMVarExprAST* InSourceExpr,
@@ -375,22 +377,22 @@ private:
 		FRigVMCompilerWorkData& WorkData,
 		bool bDelayCopyOperations = true);
 
-	void AddCopyOperator(
+	UE_API void AddCopyOperator(
 		const FRigVMCompilerWorkData::FCopyOpInfo& CopyOpInfo,
 		FRigVMCompilerWorkData& WorkData,
 		bool bDelayCopyOperations = true);
 
-	FRigVMOperand FindOrAddRegister(const FRigVMVarExprAST* InVarExpr, FRigVMCompilerWorkData& WorkData, bool bIsDebugValue = false);
-	const FRigVMCompilerWorkData::FRigVMASTProxyArray& FindProxiesWithSharedOperand(const FRigVMVarExprAST* InVarExpr, FRigVMCompilerWorkData& WorkData);
+	UE_API FRigVMOperand FindOrAddRegister(const FRigVMVarExprAST* InVarExpr, FRigVMCompilerWorkData& WorkData, bool bIsDebugValue = false);
+	UE_API const FRigVMCompilerWorkData::FRigVMASTProxyArray& FindProxiesWithSharedOperand(const FRigVMVarExprAST* InVarExpr, FRigVMCompilerWorkData& WorkData);
 
-	static FString GetPinNameWithDirectionPrefix(const URigVMPin* Pin);
-	static int32 GetOperandFunctionInterfaceParameterIndex(const TArray<FString>& OperandsPinNames, const FRigVMFunctionCompilationData* FunctionCompilationData, const FRigVMOperand& Operand);
+	static UE_API FString GetPinNameWithDirectionPrefix(const URigVMPin* Pin);
+	static UE_API int32 GetOperandFunctionInterfaceParameterIndex(const TArray<FString>& OperandsPinNames, const FRigVMFunctionCompilationData* FunctionCompilationData, const FRigVMOperand& Operand);
 
-	bool ValidateNode(const FRigVMCompileSettings& InSettings, URigVMNode* InNode, bool bCheck = true);
+	UE_API bool ValidateNode(const FRigVMCompileSettings& InSettings, URigVMNode* InNode, bool bCheck = true);
 	
-	void ReportInfo(const FRigVMCompileSettings& InSettings, const FString& InMessage);
-	void ReportWarning(const FRigVMCompileSettings& InSettings, const FString& InMessage);
-	void ReportError(const FRigVMCompileSettings& InSettings, const FString& InMessage);
+	UE_API void ReportInfo(const FRigVMCompileSettings& InSettings, const FString& InMessage);
+	UE_API void ReportWarning(const FRigVMCompileSettings& InSettings, const FString& InMessage);
+	UE_API void ReportError(const FRigVMCompileSettings& InSettings, const FString& InMessage);
 
 	template <typename... Types>
 	void ReportInfof(const FRigVMCompileSettings& InSettings, UE::Core::TCheckedFormatString<FString::FmtCharType, Types...> Fmt, Types... Args)
@@ -414,7 +416,7 @@ private:
 	friend class FFunctionCompilationScope;
 };
 
-class RIGVMDEVELOPER_API FFunctionCompilationScope
+class FFunctionCompilationScope
 {
 public:
 	FFunctionCompilationScope(URigVMCompiler* InCompiler, const URigVMLibraryNode* InLibraryNode)
@@ -432,3 +434,5 @@ private:
 	URigVMCompiler* Compiler;
 	const URigVMLibraryNode* LibraryNode;
 };
+
+#undef UE_API
diff --git a/Engine/Plugins/Runtime/RigVM/Source/RigVMDeveloper/Public/RigVMModel/Nodes/RigVMAggregateNode.h b/Engine/Plugins/Runtime/RigVM/Source/RigVMDeveloper/Public/RigVMModel/Nodes/RigVMAggregateNode.h
--- a/Engine/Plugins/Runtime/RigVM/Source/RigVMDeveloper/Public/RigVMModel/Nodes/RigVMAggregateNode.h
+++ b/Engine/Plugins/Runtime/RigVM/Source/RigVMDeveloper/Public/RigVMModel/Nodes/RigVMAggregateNode.h
@@ -7,50 +7,52 @@
 #include "RigVMModel/Nodes/RigVMUnitNode.h"
 #include "RigVMAggregateNode.generated.h"
 
+#define UE_API RIGVMDEVELOPER_API
+
 /**
  * The Aggregate Node contains a subgraph of nodes with aggregate pins (2in+1out or 1out+2in) connected
  * to each other. For example, a unit node IntAdd which adds 2 integers and provides Result=A+B could have
  * A, B and Result as aggregates in order to add additional input pins to construct an Aggregate Node that computes
  * Result=A+B+C.
  */
-UCLASS(BlueprintType)
-class RIGVMDEVELOPER_API URigVMAggregateNode : public URigVMCollapseNode
+UCLASS(MinimalAPI, BlueprintType)
+class URigVMAggregateNode : public URigVMCollapseNode
 {
 	GENERATED_BODY()
 
 public:
 
 	// default constructor
-	URigVMAggregateNode();
+	UE_API URigVMAggregateNode();
 
 	// URigVMUnitNode interface
-	virtual FString GetNodeTitle() const override;
-	virtual FLinearColor GetNodeColor() const override;
-	virtual FName GetMethodName() const;
-	virtual  FText GetToolTipText() const override;
-	virtual FText GetToolTipTextForPin(const URigVMPin* InPin) const override;
+	UE_API virtual FString GetNodeTitle() const override;
+	UE_API virtual FLinearColor GetNodeColor() const override;
+	UE_API virtual FName GetMethodName() const;
+	UE_API virtual  FText GetToolTipText() const override;
+	UE_API virtual FText GetToolTipTextForPin(const URigVMPin* InPin) const override;
 
 #if UE_RIGVM_AGGREGATE_NODES_ENABLED
 	virtual bool IsAggregate() const override { return false; };
 #endif
-	virtual bool IsInputAggregate() const override;
+	UE_API virtual bool IsInputAggregate() const override;
 
-	URigVMNode* GetFirstInnerNode() const;
-	URigVMNode* GetLastInnerNode() const;
+	UE_API URigVMNode* GetFirstInnerNode() const;
+	UE_API URigVMNode* GetLastInnerNode() const;
 	
-	virtual URigVMPin* GetFirstAggregatePin() const override;
-	virtual URigVMPin* GetSecondAggregatePin() const override;
-	virtual URigVMPin* GetOppositeAggregatePin() const override;
-	virtual TArray<URigVMPin*> GetAggregateInputs() const override;
-	virtual TArray<URigVMPin*> GetAggregateOutputs() const override;
+	UE_API virtual URigVMPin* GetFirstAggregatePin() const override;
+	UE_API virtual URigVMPin* GetSecondAggregatePin() const override;
+	UE_API virtual URigVMPin* GetOppositeAggregatePin() const override;
+	UE_API virtual TArray<URigVMPin*> GetAggregateInputs() const override;
+	UE_API virtual TArray<URigVMPin*> GetAggregateOutputs() const override;
 
 protected:
 
-	virtual FString GetOriginalDefaultValueForRootPin(const URigVMPin* InRootPin) const override;
+	UE_API virtual FString GetOriginalDefaultValueForRootPin(const URigVMPin* InRootPin) const override;
 
 private:
 
-	virtual void InvalidateCache() override;
+	UE_API virtual void InvalidateCache() override;
 
 	mutable URigVMNode* FirstInnerNodeCache;
 	mutable URigVMNode* LastInnerNodeCache;
@@ -58,3 +60,4 @@ private:
 	friend class URigVMController;
 };
 
+#undef UE_API
diff --git a/Engine/Plugins/Runtime/RigVM/Source/RigVMDeveloper/Public/RigVMModel/Nodes/RigVMArrayNode.h b/Engine/Plugins/Runtime/RigVM/Source/RigVMDeveloper/Public/RigVMModel/Nodes/RigVMArrayNode.h
--- a/Engine/Plugins/Runtime/RigVM/Source/RigVMDeveloper/Public/RigVMModel/Nodes/RigVMArrayNode.h
+++ b/Engine/Plugins/Runtime/RigVM/Source/RigVMDeveloper/Public/RigVMModel/Nodes/RigVMArrayNode.h
@@ -5,31 +5,33 @@
 #include "RigVMTemplateNode.h"
 #include "RigVMArrayNode.generated.h"
 
+#define UE_API RIGVMDEVELOPER_API
+
 /**
  * The Array Node represents one of a series available
  * array operations such as SetNum, GetAtIndex etc.
  */
-UCLASS(BlueprintType, Deprecated)
-class RIGVMDEVELOPER_API UDEPRECATED_RigVMArrayNode : public URigVMTemplateNode
+UCLASS(MinimalAPI, BlueprintType, Deprecated)
+class UDEPRECATED_RigVMArrayNode : public URigVMTemplateNode
 {
 	GENERATED_BODY()
 
 public:
 
 	// Default constructor
-	UDEPRECATED_RigVMArrayNode();
+	UE_API UDEPRECATED_RigVMArrayNode();
 
 	// Returns the op code of this node
 	UFUNCTION(BlueprintCallable, Category = RigVMArrayNode)
-	ERigVMOpCode GetOpCode() const;
+	UE_API ERigVMOpCode GetOpCode() const;
 
 	// Returns the C++ data type of the element
 	UFUNCTION(BlueprintCallable, Category = RigVMArrayNode)
-	FString GetCPPType() const;
+	UE_API FString GetCPPType() const;
 
 	// Returns the C++ data type struct of the array (or nullptr)
 	UFUNCTION(BlueprintCallable, Category = RigVMArrayNode)
-	UObject* GetCPPTypeObject() const;
+	UE_API UObject* GetCPPTypeObject() const;
 
 private:
 
@@ -42,3 +44,4 @@ private:
 	friend class FRigVMParserAST;
 };
 
+#undef UE_API
diff --git a/Engine/Plugins/Runtime/RigVM/Source/RigVMDeveloper/Public/RigVMModel/Nodes/RigVMBranchNode.h b/Engine/Plugins/Runtime/RigVM/Source/RigVMDeveloper/Public/RigVMModel/Nodes/RigVMBranchNode.h
--- a/Engine/Plugins/Runtime/RigVM/Source/RigVMDeveloper/Public/RigVMModel/Nodes/RigVMBranchNode.h
+++ b/Engine/Plugins/Runtime/RigVM/Source/RigVMDeveloper/Public/RigVMModel/Nodes/RigVMBranchNode.h
@@ -10,8 +10,8 @@ class UObject;
 /**
  * A branch node is used to branch between two blocks of execution
  */
-UCLASS(BlueprintType, Deprecated)
-class RIGVMDEVELOPER_API UDEPRECATED_RigVMBranchNode : public URigVMNode
+UCLASS(MinimalAPI, BlueprintType, Deprecated)
+class UDEPRECATED_RigVMBranchNode : public URigVMNode
 {
 	GENERATED_BODY()
 };
diff --git a/Engine/Plugins/Runtime/RigVM/Source/RigVMDeveloper/Public/RigVMModel/Nodes/RigVMCollapseNode.h b/Engine/Plugins/Runtime/RigVM/Source/RigVMDeveloper/Public/RigVMModel/Nodes/RigVMCollapseNode.h
--- a/Engine/Plugins/Runtime/RigVM/Source/RigVMDeveloper/Public/RigVMModel/Nodes/RigVMCollapseNode.h
+++ b/Engine/Plugins/Runtime/RigVM/Source/RigVMDeveloper/Public/RigVMModel/Nodes/RigVMCollapseNode.h
@@ -6,30 +6,32 @@
 #include "RigVMModel/RigVMGraph.h"
 #include "RigVMCollapseNode.generated.h"
 
+#define UE_API RIGVMDEVELOPER_API
+
 /**
  * The Collapse Node is a library node which stores the 
  * function and its nodes directly within the node itself.
  */
-UCLASS(BlueprintType)
-class RIGVMDEVELOPER_API URigVMCollapseNode : public URigVMLibraryNode
+UCLASS(MinimalAPI, BlueprintType)
+class URigVMCollapseNode : public URigVMLibraryNode
 {
 	GENERATED_BODY()
 
 public:
 
-	URigVMCollapseNode();
+	UE_API URigVMCollapseNode();
 
 	// RigVM node interface
-	virtual FText GetToolTipText() const override;
+	UE_API virtual FText GetToolTipText() const override;
 
 	// Library node interface
 	virtual FString GetNodeCategory() const override { return NodeCategory; }
 	virtual FString GetNodeKeywords() const override { return NodeKeywords; }
 	virtual FString GetNodeDescription() const override { return NodeDescription; }
-	virtual URigVMFunctionLibrary* GetLibrary() const override;
+	UE_API virtual URigVMFunctionLibrary* GetLibrary() const override;
 	virtual URigVMGraph* GetContainedGraph() const override { return ContainedGraph; }
 
-	FString GetEditorSubGraphName() const;
+	UE_API FString GetEditorSubGraphName() const;
 
 private:
 
@@ -48,3 +50,4 @@ private:
 	friend class URigVMController;
 };
 
+#undef UE_API
diff --git a/Engine/Plugins/Runtime/RigVM/Source/RigVMDeveloper/Public/RigVMModel/Nodes/RigVMCommentNode.h b/Engine/Plugins/Runtime/RigVM/Source/RigVMDeveloper/Public/RigVMModel/Nodes/RigVMCommentNode.h
--- a/Engine/Plugins/Runtime/RigVM/Source/RigVMDeveloper/Public/RigVMModel/Nodes/RigVMCommentNode.h
+++ b/Engine/Plugins/Runtime/RigVM/Source/RigVMDeveloper/Public/RigVMModel/Nodes/RigVMCommentNode.h
@@ -10,6 +10,8 @@
 
 #include "RigVMCommentNode.generated.h"
 
+#define UE_API RIGVMDEVELOPER_API
+
 class UObject;
 struct FFrame;
 
@@ -19,34 +21,34 @@ struct FFrame;
  * Comment Nodes are purely cosmetic and don't contribute
  * to the runtime result of the Graph / Function.
  */
-UCLASS(BlueprintType)
-class RIGVMDEVELOPER_API URigVMCommentNode : public URigVMNode
+UCLASS(MinimalAPI, BlueprintType)
+class URigVMCommentNode : public URigVMNode
 {
 	GENERATED_BODY()
 
 public:
 
 	// Default constructor
-	URigVMCommentNode();
+	UE_API URigVMCommentNode();
 
 	// Override of node title
 	virtual FString GetNodeTitle() const override { return GetCommentText(); }
 
 	// Returns the current user provided text of this comment.
 	UFUNCTION(BlueprintCallable, Category = RigVMCommentNode)
-	FString GetCommentText() const;
+	UE_API FString GetCommentText() const;
 
 	// Returns the current user provided font size of this comment.
 	UFUNCTION(BlueprintCallable, Category = RigVMCommentNode)
-	int32 GetCommentFontSize() const;
+	UE_API int32 GetCommentFontSize() const;
 
 	// Returns the current user provided bubble visibility of this comment.
 	UFUNCTION(BlueprintCallable, Category = RigVMCommentNode)
-	bool GetCommentBubbleVisible() const;
+	UE_API bool GetCommentBubbleVisible() const;
 
 	// Returns the current user provided bubble color inheritance of this comment.
 	UFUNCTION(BlueprintCallable, Category = RigVMCommentNode)
-	bool GetCommentColorBubble() const;
+	UE_API bool GetCommentColorBubble() const;
 
 private:
 
@@ -65,3 +67,4 @@ private:
 	friend class URigVMController;
 };
 
+#undef UE_API
diff --git a/Engine/Plugins/Runtime/RigVM/Source/RigVMDeveloper/Public/RigVMModel/Nodes/RigVMDispatchNode.h b/Engine/Plugins/Runtime/RigVM/Source/RigVMDeveloper/Public/RigVMModel/Nodes/RigVMDispatchNode.h
--- a/Engine/Plugins/Runtime/RigVM/Source/RigVMDeveloper/Public/RigVMModel/Nodes/RigVMDispatchNode.h
+++ b/Engine/Plugins/Runtime/RigVM/Source/RigVMDeveloper/Public/RigVMModel/Nodes/RigVMDispatchNode.h
@@ -16,6 +16,8 @@
 
 #include "RigVMDispatchNode.generated.h"
 
+#define UE_API RIGVMDEVELOPER_API
+
 class UObject;
 class URigVMPin;
 class UScriptStruct;
@@ -26,35 +28,35 @@ struct FRigVMDispatchFactory;
  * declared on a USTRUCT. Struct Nodes have input / output pins for all
  * struct UPROPERTY members.
  */
-UCLASS(BlueprintType)
-class RIGVMDEVELOPER_API URigVMDispatchNode : public URigVMTemplateNode
+UCLASS(MinimalAPI, BlueprintType)
+class URigVMDispatchNode : public URigVMTemplateNode
 {
 	GENERATED_BODY()
 
 public:
 
 	// Override node functions
-	virtual FString GetNodeTitle() const override;
-	virtual FText GetToolTipText() const override;
-	virtual FLinearColor GetNodeColor() const override;
-	virtual bool IsDefinedAsConstant() const override;
-	virtual bool IsDefinedAsVarying() const override;
-	virtual const TArray<FName>& GetControlFlowBlocks() const override;
-	virtual const bool IsControlFlowBlockSliced(const FName& InBlockName) const override;
-	virtual TArray<URigVMPin*> GetAggregateInputs() const override;
-	virtual TArray<URigVMPin*> GetAggregateOutputs() const override;
-	virtual FName GetNextAggregateName(const FName& InLastAggregatePinName) const override;
-	virtual FRigVMStructUpgradeInfo GetUpgradeInfo() const override;
-
-	const FRigVMDispatchFactory* GetFactory() const;
-	virtual bool IsOutDated() const override;
-	virtual FString GetDeprecatedMetadata() const override;
-
-	FRigVMDispatchContext GetDispatchContext() const;
+	UE_API virtual FString GetNodeTitle() const override;
+	UE_API virtual FText GetToolTipText() const override;
+	UE_API virtual FLinearColor GetNodeColor() const override;
+	UE_API virtual bool IsDefinedAsConstant() const override;
+	UE_API virtual bool IsDefinedAsVarying() const override;
+	UE_API virtual const TArray<FName>& GetControlFlowBlocks() const override;
+	UE_API virtual const bool IsControlFlowBlockSliced(const FName& InBlockName) const override;
+	UE_API virtual TArray<URigVMPin*> GetAggregateInputs() const override;
+	UE_API virtual TArray<URigVMPin*> GetAggregateOutputs() const override;
+	UE_API virtual FName GetNextAggregateName(const FName& InLastAggregatePinName) const override;
+	UE_API virtual FRigVMStructUpgradeInfo GetUpgradeInfo() const override;
+
+	UE_API const FRigVMDispatchFactory* GetFactory() const;
+	UE_API virtual bool IsOutDated() const override;
+	UE_API virtual FString GetDeprecatedMetadata() const override;
+
+	UE_API FRigVMDispatchContext GetDispatchContext() const;
 
 	// Returns an instance of the factory with the current values.
 	// @param bUseDefault If set to true the default struct will be created - otherwise the struct will contains the values from the node
-	TSharedPtr<FStructOnScope> ConstructFactoryInstance(bool bUseDefault = false, FString* OutFactoryDefault = nullptr) const;
+	UE_API TSharedPtr<FStructOnScope> ConstructFactoryInstance(bool bUseDefault = false, FString* OutFactoryDefault = nullptr) const;
 
 	// Returns a copy of the struct with the current values
 	template <
@@ -75,12 +77,12 @@ public:
 	
 protected:
 
-	virtual FText GetToolTipTextForPin(const URigVMPin* InPin) const override;
-	const UScriptStruct* GetFactoryStruct() const;
-	virtual void InvalidateCache() override;
-	virtual bool ShouldInputPinComputeLazily(const URigVMPin* InPin) const override;
-	virtual FString GetOriginalDefaultValueForRootPin(const URigVMPin* InRootPin) const override;
-	FString GetFactoryDefaultValue() const;
+	UE_API virtual FText GetToolTipTextForPin(const URigVMPin* InPin) const override;
+	UE_API const UScriptStruct* GetFactoryStruct() const;
+	UE_API virtual void InvalidateCache() override;
+	UE_API virtual bool ShouldInputPinComputeLazily(const URigVMPin* InPin) const override;
+	UE_API virtual FString GetOriginalDefaultValueForRootPin(const URigVMPin* InRootPin) const override;
+	UE_API FString GetFactoryDefaultValue() const;
 
 private:
 
@@ -90,3 +92,4 @@ private:
 	friend class URigVMController;
 };
 
+#undef UE_API
diff --git a/Engine/Plugins/Runtime/RigVM/Source/RigVMDeveloper/Public/RigVMModel/Nodes/RigVMEnumNode.h b/Engine/Plugins/Runtime/RigVM/Source/RigVMDeveloper/Public/RigVMModel/Nodes/RigVMEnumNode.h
--- a/Engine/Plugins/Runtime/RigVM/Source/RigVMDeveloper/Public/RigVMModel/Nodes/RigVMEnumNode.h
+++ b/Engine/Plugins/Runtime/RigVM/Source/RigVMDeveloper/Public/RigVMModel/Nodes/RigVMEnumNode.h
@@ -11,6 +11,8 @@
 
 #include "RigVMEnumNode.generated.h"
 
+#define UE_API RIGVMDEVELOPER_API
+
 class UEnum;
 class UObject;
 struct FFrame;
@@ -18,33 +20,33 @@ struct FFrame;
 /**
  * The Enum Node represents a constant enum value for use within the graph.
  */
-UCLASS(BlueprintType)
-class RIGVMDEVELOPER_API URigVMEnumNode : public URigVMNode
+UCLASS(MinimalAPI, BlueprintType)
+class URigVMEnumNode : public URigVMNode
 {
 	GENERATED_BODY()
 
 public:
 
 	// Default constructor
-	URigVMEnumNode();
+	UE_API URigVMEnumNode();
 
 	// Override of node title
-	virtual FString GetNodeTitle() const;
+	UE_API virtual FString GetNodeTitle() const;
 
 	// Returns the enum itself
 	UFUNCTION(BlueprintCallable, Category = RigVMEnumNode)
-	UEnum* GetEnum() const;
+	UE_API UEnum* GetEnum() const;
 
 	// Returns the C++ data type of the parameter
 	UFUNCTION(BlueprintCallable, Category = RigVMEnumNode)
-	FString GetCPPType() const;
+	UE_API FString GetCPPType() const;
 
 	// Returns the C++ data type struct of the parameter (or nullptr)
 	UFUNCTION(BlueprintCallable, Category = RigVMEnumNode)
-	UObject* GetCPPTypeObject() const;
+	UE_API UObject* GetCPPTypeObject() const;
 
 	// Returns the default value of the parameter as a string
-	FString GetDefaultValue(const URigVMPin::FPinOverride& InOverride = URigVMPin::EmptyPinOverride) const;
+	UE_API FString GetDefaultValue(const URigVMPin::FPinOverride& InOverride = URigVMPin::EmptyPinOverride) const;
 
 	// Override of node title
 	virtual FLinearColor GetNodeColor() const override { return FLinearColor::Blue; }
@@ -53,7 +55,7 @@ public:
 
 private:
 
-	static const TCHAR* EnumName;
+	static UE_API const TCHAR* EnumName;
 	static const inline TCHAR* EnumValueName = TEXT("EnumValue");
 	static const inline TCHAR* EnumIndexName = TEXT("EnumIndex");
 	
@@ -64,3 +66,4 @@ private:
 	friend class FRigVMParserAST;
 };
 
+#undef UE_API
diff --git a/Engine/Plugins/Runtime/RigVM/Source/RigVMDeveloper/Public/RigVMModel/Nodes/RigVMFunctionEntryNode.h b/Engine/Plugins/Runtime/RigVM/Source/RigVMDeveloper/Public/RigVMModel/Nodes/RigVMFunctionEntryNode.h
--- a/Engine/Plugins/Runtime/RigVM/Source/RigVMDeveloper/Public/RigVMModel/Nodes/RigVMFunctionEntryNode.h
+++ b/Engine/Plugins/Runtime/RigVM/Source/RigVMDeveloper/Public/RigVMModel/Nodes/RigVMFunctionEntryNode.h
@@ -5,12 +5,14 @@
 #include "RigVMModel/Nodes/RigVMFunctionInterfaceNode.h"
 #include "RigVMFunctionEntryNode.generated.h"
 
+#define UE_API RIGVMDEVELOPER_API
+
 /**
  * The Function Entry node is used to provide access to the 
  * input pins of the library node for links within.
  */
-UCLASS(BlueprintType)
-class RIGVMDEVELOPER_API URigVMFunctionEntryNode : public URigVMFunctionInterfaceNode
+UCLASS(MinimalAPI, BlueprintType)
+class URigVMFunctionEntryNode : public URigVMFunctionInterfaceNode
 {
 	GENERATED_BODY()
 
@@ -22,13 +24,14 @@ public:
 	virtual FName GetNotation() const override { return NAME_None; }
 
 	// Override node functions
-	virtual bool IsWithinLoop() const override;
+	UE_API virtual bool IsWithinLoop() const override;
 
 	// URigVMNode interface
-	virtual FString GetNodeTitle() const override;
+	UE_API virtual FString GetNodeTitle() const override;
 
 private:
 
 	friend class URigVMController;
 };
 
+#undef UE_API
diff --git a/Engine/Plugins/Runtime/RigVM/Source/RigVMDeveloper/Public/RigVMModel/Nodes/RigVMFunctionInterfaceNode.h b/Engine/Plugins/Runtime/RigVM/Source/RigVMDeveloper/Public/RigVMModel/Nodes/RigVMFunctionInterfaceNode.h
--- a/Engine/Plugins/Runtime/RigVM/Source/RigVMDeveloper/Public/RigVMModel/Nodes/RigVMFunctionInterfaceNode.h
+++ b/Engine/Plugins/Runtime/RigVM/Source/RigVMDeveloper/Public/RigVMModel/Nodes/RigVMFunctionInterfaceNode.h
@@ -5,35 +5,38 @@
 #include "RigVMModel/Nodes/RigVMTemplateNode.h"
 #include "RigVMFunctionInterfaceNode.generated.h"
 
+#define UE_API RIGVMDEVELOPER_API
+
 /**
  * The Function Interface node is is used as the base class for
  * both the entry and return nodes.
  */
-UCLASS(BlueprintType)
-class RIGVMDEVELOPER_API URigVMFunctionInterfaceNode : public URigVMTemplateNode
+UCLASS(MinimalAPI, BlueprintType)
+class URigVMFunctionInterfaceNode : public URigVMTemplateNode
 {
 	GENERATED_BODY()
 
 public:
 
 	// Override template node functions
-	virtual uint32 GetStructureHash() const override;
+	UE_API virtual uint32 GetStructureHash() const override;
 
 	// Override node functions
-	virtual FLinearColor GetNodeColor() const override;
-	virtual bool IsDefinedAsVarying() const override;
+	UE_API virtual FLinearColor GetNodeColor() const override;
+	UE_API virtual bool IsDefinedAsVarying() const override;
 	
 	// URigVMNode interface
-	virtual  FText GetToolTipText() const override;
-	virtual FText GetToolTipTextForPin(const URigVMPin* InPin) const override;	
+	UE_API virtual  FText GetToolTipText() const override;
+	UE_API virtual FText GetToolTipTextForPin(const URigVMPin* InPin) const override;	
 
 protected:
 
-	const URigVMPin* FindReferencedPin(const URigVMPin* InPin) const;
-	const URigVMPin* FindReferencedPin(const FString& InPinPath) const;
+	UE_API const URigVMPin* FindReferencedPin(const URigVMPin* InPin) const;
+	UE_API const URigVMPin* FindReferencedPin(const FString& InPinPath) const;
 
 private:
 
 	friend class URigVMController;
 };
 
+#undef UE_API
diff --git a/Engine/Plugins/Runtime/RigVM/Source/RigVMDeveloper/Public/RigVMModel/Nodes/RigVMFunctionReferenceNode.h b/Engine/Plugins/Runtime/RigVM/Source/RigVMDeveloper/Public/RigVMModel/Nodes/RigVMFunctionReferenceNode.h
--- a/Engine/Plugins/Runtime/RigVM/Source/RigVMDeveloper/Public/RigVMModel/Nodes/RigVMFunctionReferenceNode.h
+++ b/Engine/Plugins/Runtime/RigVM/Source/RigVMDeveloper/Public/RigVMModel/Nodes/RigVMFunctionReferenceNode.h
@@ -7,74 +7,76 @@
 #include "RigVMModel/Nodes/RigVMLibraryNode.h"
 #include "RigVMFunctionReferenceNode.generated.h"
 
+#define UE_API RIGVMDEVELOPER_API
+
 class URigVMFunctionLibrary;
 
 /**
  * The Function Reference Node is a library node which references
  * a library node from a separate function library graph.
  */
-UCLASS(BlueprintType)
-class RIGVMDEVELOPER_API URigVMFunctionReferenceNode : public URigVMLibraryNode
+UCLASS(MinimalAPI, BlueprintType)
+class URigVMFunctionReferenceNode : public URigVMLibraryNode
 {
 	GENERATED_BODY()
 
 public:
 
 	// URigVMNode interface
-	virtual FString GetNodeTitle() const override;
-	virtual FLinearColor GetNodeColor() const override;
-	virtual FText GetToolTipText() const override;
-	virtual FName GetDisplayNameForPin(const URigVMPin* InPin) const override;
-	virtual FString GetCategoryForPin(const FString& InPinPath) const override;
-	virtual int32 GetIndexInCategoryForPin(const FString& InPinPath) const override;
-	virtual FText GetToolTipTextForPin(const URigVMPin* InPin) const override;
-	virtual TArray<FString> GetPinCategories() const override;
-	virtual FRigVMNodeLayout GetNodeLayout(bool bIncludeEmptyCategories) const override;
+	UE_API virtual FString GetNodeTitle() const override;
+	UE_API virtual FLinearColor GetNodeColor() const override;
+	UE_API virtual FText GetToolTipText() const override;
+	UE_API virtual FName GetDisplayNameForPin(const URigVMPin* InPin) const override;
+	UE_API virtual FString GetCategoryForPin(const FString& InPinPath) const override;
+	UE_API virtual int32 GetIndexInCategoryForPin(const FString& InPinPath) const override;
+	UE_API virtual FText GetToolTipTextForPin(const URigVMPin* InPin) const override;
+	UE_API virtual TArray<FString> GetPinCategories() const override;
+	UE_API virtual FRigVMNodeLayout GetNodeLayout(bool bIncludeEmptyCategories) const override;
 	// end URigVMNode interface
 
 	// URigVMLibraryNode interface
-	virtual FString GetNodeCategory() const override;
-	virtual FString GetNodeKeywords() const override;
-	virtual TArray<FRigVMExternalVariable> GetExternalVariables() const override;
+	UE_API virtual FString GetNodeCategory() const override;
+	UE_API virtual FString GetNodeKeywords() const override;
+	UE_API virtual TArray<FRigVMExternalVariable> GetExternalVariables() const override;
 	virtual const FRigVMTemplate* GetTemplate() const override { return nullptr; }
-	virtual FRigVMGraphFunctionIdentifier GetFunctionIdentifier() const override;
+	UE_API virtual FRigVMGraphFunctionIdentifier GetFunctionIdentifier() const override;
 	// end URigVMLibraryNode interface
 
-	bool IsReferencedFunctionHostLoaded() const;
-	bool IsReferencedNodeLoaded() const;
-	URigVMLibraryNode* LoadReferencedNode() const;
+	UE_API bool IsReferencedFunctionHostLoaded() const;
+	UE_API bool IsReferencedNodeLoaded() const;
+	UE_API URigVMLibraryNode* LoadReferencedNode() const;
 
 	// Variable remapping
-	bool RequiresVariableRemapping() const;
-	bool IsFullyRemapped() const;
-	TArray<FRigVMExternalVariable> GetExternalVariables(bool bRemapped) const;
+	UE_API bool RequiresVariableRemapping() const;
+	UE_API bool IsFullyRemapped() const;
+	UE_API TArray<FRigVMExternalVariable> GetExternalVariables(bool bRemapped) const;
 	const TMap<FName, FName>& GetVariableMap() const { return VariableMap; }
-	FName GetOuterVariableName(const FName& InInnerVariableName) const;
+	UE_API FName GetOuterVariableName(const FName& InInnerVariableName) const;
 	// end Variable remapping
 
-	virtual uint32 GetStructureHash() const override;
+	UE_API virtual uint32 GetStructureHash() const override;
 
 	UFUNCTION(BlueprintCallable, Category = RigVMLibraryNode, meta = (DisplayName = "GetReferencedFunctionHeader", ScriptName = "GetReferencedFunctionHeader"))
 	FRigVMGraphFunctionHeader GetReferencedFunctionHeader_ForBlueprint() const { return GetReferencedFunctionHeader(); }
 
 	const FRigVMGraphFunctionHeader& GetReferencedFunctionHeader() const { return ReferencedFunctionHeader; }
 
-	void UpdateFunctionHeaderFromHost();
+	UE_API void UpdateFunctionHeaderFromHost();
 
-	const FRigVMGraphFunctionData* GetReferencedFunctionData(bool bLoadIfNecessary = true) const;
+	UE_API const FRigVMGraphFunctionData* GetReferencedFunctionData(bool bLoadIfNecessary = true) const;
 
-	TArray<FRigVMTag> GetVariantTags() const;
+	UE_API TArray<FRigVMTag> GetVariantTags() const;
 
 protected:
 
-	virtual FString GetOriginalDefaultValueForRootPin(const URigVMPin* InRootPin) const override;
+	UE_API virtual FString GetOriginalDefaultValueForRootPin(const URigVMPin* InRootPin) const override;
 	
 private:
 
-	bool RequiresVariableRemappingInternal(TArray<FRigVMExternalVariable>& InnerVariables) const;
-	virtual TArray<int32> GetInstructionsForVMImpl(const FRigVMExtendedExecuteContext& Context, URigVM* InVM, const FRigVMASTProxy& InProxy = FRigVMASTProxy()) const override;
-	const URigVMPin* FindReferencedPin(const URigVMPin* InPin) const;
-	const URigVMPin* FindReferencedPin(const FString& InPinPath) const;
+	UE_API bool RequiresVariableRemappingInternal(TArray<FRigVMExternalVariable>& InnerVariables) const;
+	UE_API virtual TArray<int32> GetInstructionsForVMImpl(const FRigVMExtendedExecuteContext& Context, URigVM* InVM, const FRigVMASTProxy& InProxy = FRigVMASTProxy()) const override;
+	UE_API const URigVMPin* FindReferencedPin(const URigVMPin* InPin) const;
+	UE_API const URigVMPin* FindReferencedPin(const FString& InPinPath) const;
 
 	//void SetReferencedFunctionData(FRigVMGraphFunctionData* Data);
 
@@ -94,3 +96,4 @@ private:
 	friend struct EngineTestRigVMFramework;
 };
 
+#undef UE_API
diff --git a/Engine/Plugins/Runtime/RigVM/Source/RigVMDeveloper/Public/RigVMModel/Nodes/RigVMFunctionReturnNode.h b/Engine/Plugins/Runtime/RigVM/Source/RigVMDeveloper/Public/RigVMModel/Nodes/RigVMFunctionReturnNode.h
--- a/Engine/Plugins/Runtime/RigVM/Source/RigVMDeveloper/Public/RigVMModel/Nodes/RigVMFunctionReturnNode.h
+++ b/Engine/Plugins/Runtime/RigVM/Source/RigVMDeveloper/Public/RigVMModel/Nodes/RigVMFunctionReturnNode.h
@@ -5,12 +5,14 @@
 #include "RigVMModel/Nodes/RigVMFunctionInterfaceNode.h"
 #include "RigVMFunctionReturnNode.generated.h"
 
+#define UE_API RIGVMDEVELOPER_API
+
 /**
  * The Function Return node is used to provide access to the 
  * output pins of the library node for links within.
  */
-UCLASS(BlueprintType)
-class RIGVMDEVELOPER_API URigVMFunctionReturnNode : public URigVMFunctionInterfaceNode
+UCLASS(MinimalAPI, BlueprintType)
+class URigVMFunctionReturnNode : public URigVMFunctionInterfaceNode
 {
 	GENERATED_BODY()
 
@@ -22,10 +24,11 @@ public:
 	virtual FName GetNotation() const override { return NAME_None; }
 
 	// URigVMNode interface
-	virtual FString GetNodeTitle() const override;
+	UE_API virtual FString GetNodeTitle() const override;
 
 private:
 
 	friend class URigVMController;
 };
 
+#undef UE_API
diff --git a/Engine/Plugins/Runtime/RigVM/Source/RigVMDeveloper/Public/RigVMModel/Nodes/RigVMIfNode.h b/Engine/Plugins/Runtime/RigVM/Source/RigVMDeveloper/Public/RigVMModel/Nodes/RigVMIfNode.h
--- a/Engine/Plugins/Runtime/RigVM/Source/RigVMDeveloper/Public/RigVMModel/Nodes/RigVMIfNode.h
+++ b/Engine/Plugins/Runtime/RigVM/Source/RigVMDeveloper/Public/RigVMModel/Nodes/RigVMIfNode.h
@@ -11,22 +11,24 @@
 
 #include "RigVMIfNode.generated.h"
 
+#define UE_API RIGVMDEVELOPER_API
+
 class UObject;
 struct FRigVMTemplate;
 
 /**
  * A if node is used to pick between two values
  */
-UCLASS(BlueprintType, Deprecated)
-class RIGVMDEVELOPER_API UDEPRECATED_RigVMIfNode : public URigVMTemplateNode
+UCLASS(MinimalAPI, BlueprintType, Deprecated)
+class UDEPRECATED_RigVMIfNode : public URigVMTemplateNode
 {
 	GENERATED_BODY()
 
 public:
 
 	// Override from URigVMTemplateNode
-	virtual FName GetNotation() const override;
-	virtual const FRigVMTemplate* GetTemplate() const override;
+	UE_API virtual FName GetNotation() const override;
+	UE_API virtual const FRigVMTemplate* GetTemplate() const override;
 	virtual bool IsSingleton() const override { return false; }
 
 	// Override from URigVMNode
@@ -48,3 +50,4 @@ private:
 	friend struct FRigVMRemoveNodeAction;
 };
 
+#undef UE_API
diff --git a/Engine/Plugins/Runtime/RigVM/Source/RigVMDeveloper/Public/RigVMModel/Nodes/RigVMInvokeEntryNode.h b/Engine/Plugins/Runtime/RigVM/Source/RigVMDeveloper/Public/RigVMModel/Nodes/RigVMInvokeEntryNode.h
--- a/Engine/Plugins/Runtime/RigVM/Source/RigVMDeveloper/Public/RigVMModel/Nodes/RigVMInvokeEntryNode.h
+++ b/Engine/Plugins/Runtime/RigVM/Source/RigVMDeveloper/Public/RigVMModel/Nodes/RigVMInvokeEntryNode.h
@@ -5,26 +5,28 @@
 #include "RigVMModel/RigVMNode.h"
 #include "RigVMInvokeEntryNode.generated.h"
 
+#define UE_API RIGVMDEVELOPER_API
+
 /**
  * The Invoke Entry Node is used to invoke another entry from the 
  * the currently running entry.
  */
-UCLASS(BlueprintType)
-class RIGVMDEVELOPER_API URigVMInvokeEntryNode : public URigVMNode
+UCLASS(MinimalAPI, BlueprintType)
+class URigVMInvokeEntryNode : public URigVMNode
 {
 	GENERATED_BODY()
 
 public:
 
 	// Default constructor
-	URigVMInvokeEntryNode();
+	UE_API URigVMInvokeEntryNode();
 
 	// Override of node title
-	virtual FString GetNodeTitle() const;
+	UE_API virtual FString GetNodeTitle() const;
 
 	// Returns the name of the entry to run
 	UFUNCTION(BlueprintCallable, Category = RigVMInvokeEntryNode)
-	FName GetEntryName() const;
+	UE_API FName GetEntryName() const;
 
 	// Override of node title
 	virtual FLinearColor GetNodeColor() const override { return FLinearColor::Red; }
@@ -34,7 +36,7 @@ private:
 
 	static const inline TCHAR* EntryName = TEXT("Entry");
 
-	URigVMPin* GetEntryNamePin() const;
+	UE_API URigVMPin* GetEntryNamePin() const;
 	
 
 	friend class URigVMController;
@@ -44,3 +46,4 @@ private:
 	friend class URigVMPin;
 };
 
+#undef UE_API
diff --git a/Engine/Plugins/Runtime/RigVM/Source/RigVMDeveloper/Public/RigVMModel/Nodes/RigVMLibraryNode.h b/Engine/Plugins/Runtime/RigVM/Source/RigVMDeveloper/Public/RigVMModel/Nodes/RigVMLibraryNode.h
--- a/Engine/Plugins/Runtime/RigVM/Source/RigVMDeveloper/Public/RigVMModel/Nodes/RigVMLibraryNode.h
+++ b/Engine/Plugins/Runtime/RigVM/Source/RigVMDeveloper/Public/RigVMModel/Nodes/RigVMLibraryNode.h
@@ -7,6 +7,8 @@
 #include "RigVMCore/RigVMGraphFunctionHost.h"
 #include "RigVMLibraryNode.generated.h"
 
+#define UE_API RIGVMDEVELOPER_API
+
 class URigVMGraph;
 class URigVMFunctionEntryNode;
 class URigVMFunctionReturnNode;
@@ -18,26 +20,26 @@ class URigVMFunctionLibrary;
  * expressed as a sub-graph (RigVMGroupNode) or as a 
  * referenced function (RigVMFunctionNode).
  */
-UCLASS(BlueprintType)
-class RIGVMDEVELOPER_API URigVMLibraryNode : public URigVMTemplateNode
+UCLASS(MinimalAPI, BlueprintType)
+class URigVMLibraryNode : public URigVMTemplateNode
 {
 	GENERATED_BODY()
 
 public:
 
 	// Override node functions
-	virtual bool IsDefinedAsConstant() const override;
-	virtual bool IsDefinedAsVarying() const override;
+	UE_API virtual bool IsDefinedAsConstant() const override;
+	UE_API virtual bool IsDefinedAsVarying() const override;
 	//virtual int32 GetInstructionVisitedCount(URigVM* InVM, const FRigVMASTProxy& InProxy = FRigVMASTProxy(), bool bConsolidatePerNode = false) const override;
 	
 	// Override template node functions
 	virtual UScriptStruct* GetScriptStruct() const override { return nullptr; }
 	virtual const FRigVMTemplate* GetTemplate() const override { return nullptr; }
 	virtual FName GetNotation() const override { return NAME_None; }
-	virtual uint32 GetStructureHash() const override;
+	UE_API virtual uint32 GetStructureHash() const override;
 
 	// URigVMNode interface
-	virtual  FText GetToolTipText() const override;
+	UE_API virtual  FText GetToolTipText() const override;
 	
 	// Library node interface
 	virtual FString GetNodeCategory() const { return FString(); }
@@ -50,31 +52,31 @@ public:
 	UFUNCTION(BlueprintCallable, Category = RigVMLibraryNode)
 	virtual URigVMGraph* GetContainedGraph() const { return nullptr; }
 	
-	virtual const TArray<URigVMNode*>& GetContainedNodes() const;
-	virtual const TArray<URigVMLink*>& GetContainedLinks() const;
-	virtual URigVMFunctionEntryNode* GetEntryNode() const;
-	virtual URigVMFunctionReturnNode* GetReturnNode() const;
-	virtual bool Contains(URigVMLibraryNode* InContainedNode, bool bRecursive = true) const;
-	virtual TArray<FRigVMExternalVariable> GetExternalVariables() const;
-	virtual TMap<FRigVMGraphFunctionIdentifier, uint32> GetDependencies() const;
-
-	virtual FRigVMGraphFunctionIdentifier GetFunctionIdentifier() const;
-	FRigVMGraphFunctionHeader GetFunctionHeader(IRigVMGraphFunctionHost* InHostObject = nullptr) const;
+	UE_API virtual const TArray<URigVMNode*>& GetContainedNodes() const;
+	UE_API virtual const TArray<URigVMLink*>& GetContainedLinks() const;
+	UE_API virtual URigVMFunctionEntryNode* GetEntryNode() const;
+	UE_API virtual URigVMFunctionReturnNode* GetReturnNode() const;
+	UE_API virtual bool Contains(URigVMLibraryNode* InContainedNode, bool bRecursive = true) const;
+	UE_API virtual TArray<FRigVMExternalVariable> GetExternalVariables() const;
+	UE_API virtual TMap<FRigVMGraphFunctionIdentifier, uint32> GetDependencies() const;
+
+	UE_API virtual FRigVMGraphFunctionIdentifier GetFunctionIdentifier() const;
+	UE_API FRigVMGraphFunctionHeader GetFunctionHeader(IRigVMGraphFunctionHost* InHostObject = nullptr) const;
 	
 	UFUNCTION(BlueprintPure, Category = RigVMLibraryNode)
-	FRigVMVariant GetFunctionVariant() const;
+	UE_API FRigVMVariant GetFunctionVariant() const;
 
 	UFUNCTION(BlueprintPure, Category = RigVMLibraryNode)
-	FRigVMVariantRef GetFunctionVariantRef() const;
+	UE_API FRigVMVariantRef GetFunctionVariantRef() const;
 
 	UFUNCTION(BlueprintPure, Category = RigVMLibraryNode)
-	TArray<FRigVMVariantRef> GetMatchingVariants() const;
+	UE_API TArray<FRigVMVariantRef> GetMatchingVariants() const;
 
 protected:
 
-	virtual TArray<int32> GetInstructionsForVMImpl(const FRigVMExtendedExecuteContext& Context, URigVM* InVM, const FRigVMASTProxy& InProxy = FRigVMASTProxy()) const override; 
-	const static TArray<URigVMNode*> EmptyNodes;
-	const static TArray<URigVMLink*> EmptyLinks;
+	UE_API virtual TArray<int32> GetInstructionsForVMImpl(const FRigVMExtendedExecuteContext& Context, URigVM* InVM, const FRigVMASTProxy& InProxy = FRigVMASTProxy()) const override; 
+	UE_API const static TArray<URigVMNode*> EmptyNodes;
+	UE_API const static TArray<URigVMLink*> EmptyLinks;
 
 private:
 
@@ -82,3 +84,4 @@ private:
 	friend struct FRigVMSetLibraryTemplateAction;
 };
 
+#undef UE_API
diff --git a/Engine/Plugins/Runtime/RigVM/Source/RigVMDeveloper/Public/RigVMModel/Nodes/RigVMParameterNode.h b/Engine/Plugins/Runtime/RigVM/Source/RigVMDeveloper/Public/RigVMModel/Nodes/RigVMParameterNode.h
--- a/Engine/Plugins/Runtime/RigVM/Source/RigVMDeveloper/Public/RigVMModel/Nodes/RigVMParameterNode.h
+++ b/Engine/Plugins/Runtime/RigVM/Source/RigVMDeveloper/Public/RigVMModel/Nodes/RigVMParameterNode.h
@@ -15,6 +15,8 @@
 
 #include "RigVMParameterNode.generated.h"
 
+#define UE_API RIGVMDEVELOPER_API
+
 class UObject;
 struct FFrame;
 
@@ -86,42 +88,42 @@ public:
  * The Parameter Node represents an input or output argument / parameter
  * of the Function / Graph. Parameter Node have only a single value pin.
  */
-UCLASS(BlueprintType, meta = (Deprecated = "5.1"))
-class RIGVMDEVELOPER_API URigVMParameterNode : public URigVMNode
+UCLASS(MinimalAPI, BlueprintType, meta = (Deprecated = "5.1"))
+class URigVMParameterNode : public URigVMNode
 {
 	GENERATED_BODY()
 
 public:
 
 	// Default constructor
-	URigVMParameterNode();
+	UE_API URigVMParameterNode();
 
 	// Override of node title
-	virtual FString GetNodeTitle() const;
+	UE_API virtual FString GetNodeTitle() const;
 
 	// Returns the name of the parameter
 	UFUNCTION(BlueprintCallable, Category = RigVMParameterNode)
-	FName GetParameterName() const;
+	UE_API FName GetParameterName() const;
 
 	// Returns true if this node is an input
 	UFUNCTION(BlueprintCallable, Category = RigVMParameterNode)
-	bool IsInput() const;
+	UE_API bool IsInput() const;
 
 	// Returns the C++ data type of the parameter
 	UFUNCTION(BlueprintCallable, Category = RigVMParameterNode)
-	FString GetCPPType() const;
+	UE_API FString GetCPPType() const;
 
 	// Returns the C++ data type struct of the parameter (or nullptr)
 	UFUNCTION(BlueprintCallable, Category = RigVMParameterNode)
-	UObject* GetCPPTypeObject() const;
+	UE_API UObject* GetCPPTypeObject() const;
 
 	// Returns the default value of the parameter as a string
 	UFUNCTION(BlueprintCallable, Category = RigVMParameterNode)
-	FString GetDefaultValue() const;
+	UE_API FString GetDefaultValue() const;
 
 	// Returns this parameter node's parameter description
 	UFUNCTION(BlueprintCallable, Category = RigVMParameterNode)
-	FRigVMGraphParameterDescription GetParameterDescription() const;
+	UE_API FRigVMGraphParameterDescription GetParameterDescription() const;
 
 	// Override of node title
 	virtual FLinearColor GetNodeColor() const override { return FLinearColor::Green; }
@@ -140,3 +142,4 @@ private:
 	friend class URigVMCompiler;
 };
 
+#undef UE_API
diff --git a/Engine/Plugins/Runtime/RigVM/Source/RigVMDeveloper/Public/RigVMModel/Nodes/RigVMRerouteNode.h b/Engine/Plugins/Runtime/RigVM/Source/RigVMDeveloper/Public/RigVMModel/Nodes/RigVMRerouteNode.h
--- a/Engine/Plugins/Runtime/RigVM/Source/RigVMDeveloper/Public/RigVMModel/Nodes/RigVMRerouteNode.h
+++ b/Engine/Plugins/Runtime/RigVM/Source/RigVMDeveloper/Public/RigVMModel/Nodes/RigVMRerouteNode.h
@@ -5,6 +5,8 @@
 #include "RigVMModel/RigVMNode.h"
 #include "RigVMRerouteNode.generated.h"
 
+#define UE_API RIGVMDEVELOPER_API
+
 /**
  * A reroute node is used to visually improve the 
  * data flow within a Graph. Reroutes are purely 
@@ -12,23 +14,23 @@
  * VM whatsoever. Reroutes can furthermore be
  * displayed as full nodes or as small circles.
  */
-UCLASS(BlueprintType)
-class RIGVMDEVELOPER_API URigVMRerouteNode : public URigVMNode
+UCLASS(MinimalAPI, BlueprintType)
+class URigVMRerouteNode : public URigVMNode
 {
 	GENERATED_BODY()
 
 public:
 
 	// Default constructor
-	URigVMRerouteNode();
+	UE_API URigVMRerouteNode();
 
 	// Override of node title
-	virtual FString GetNodeTitle() const override;
+	UE_API virtual FString GetNodeTitle() const override;
 
-	virtual FLinearColor GetNodeColor() const override;
+	UE_API virtual FLinearColor GetNodeColor() const override;
 
 	// Has no source connections
-	bool IsLiteral() const;
+	UE_API bool IsLiteral() const;
 
 private:
 
@@ -41,3 +43,4 @@ private:
 	friend class FRigVMDeveloperModule;
 };
 
+#undef UE_API
diff --git a/Engine/Plugins/Runtime/RigVM/Source/RigVMDeveloper/Public/RigVMModel/Nodes/RigVMSelectNode.h b/Engine/Plugins/Runtime/RigVM/Source/RigVMDeveloper/Public/RigVMModel/Nodes/RigVMSelectNode.h
--- a/Engine/Plugins/Runtime/RigVM/Source/RigVMDeveloper/Public/RigVMModel/Nodes/RigVMSelectNode.h
+++ b/Engine/Plugins/Runtime/RigVM/Source/RigVMDeveloper/Public/RigVMModel/Nodes/RigVMSelectNode.h
@@ -11,6 +11,8 @@
 
 #include "RigVMSelectNode.generated.h"
 
+#define UE_API RIGVMDEVELOPER_API
+
 class UObject;
 class URigVMPin;
 struct FRigVMTemplate;
@@ -18,16 +20,16 @@ struct FRigVMTemplate;
 /**
  * A select node is used to select between multiple values
  */
-UCLASS(BlueprintType, Deprecated)
-class RIGVMDEVELOPER_API UDEPRECATED_RigVMSelectNode : public URigVMTemplateNode
+UCLASS(MinimalAPI, BlueprintType, Deprecated)
+class UDEPRECATED_RigVMSelectNode : public URigVMTemplateNode
 {
 	GENERATED_BODY()
 
 public:
 
 	// Override from URigVMTemplateNode
-	virtual FName GetNotation() const override;
-	virtual const FRigVMTemplate* GetTemplate() const override;
+	UE_API virtual FName GetNotation() const override;
+	UE_API virtual const FRigVMTemplate* GetTemplate() const override;
 	virtual bool IsSingleton() const override { return false; }
 
 	// Override from URigVMNode
@@ -36,7 +38,7 @@ public:
 	
 protected:
 
-	virtual bool AllowsLinksOn(const URigVMPin* InPin) const override;
+	UE_API virtual bool AllowsLinksOn(const URigVMPin* InPin) const override;
 
 private:
 
@@ -55,3 +57,4 @@ private:
 	friend class URigVMPin;
 };
 
+#undef UE_API
