Updated headers using UnrealCodeFixup to make sure dllstorage is on methods/staticvars instead of types. Part 4/n

diff --git a/Engine/Plugins/Runtime/RigVM/Source/RigVMDeveloper/Public/RigVMModel/Nodes/RigVMTemplateNode.h b/Engine/Plugins/Runtime/RigVM/Source/RigVMDeveloper/Public/RigVMModel/Nodes/RigVMTemplateNode.h
--- a/Engine/Plugins/Runtime/RigVM/Source/RigVMDeveloper/Public/RigVMModel/Nodes/RigVMTemplateNode.h
+++ b/Engine/Plugins/Runtime/RigVM/Source/RigVMDeveloper/Public/RigVMModel/Nodes/RigVMTemplateNode.h
@@ -6,8 +6,10 @@
 #include "RigVMCore/RigVMTemplate.h"
 #include "RigVMTemplateNode.generated.h"
 
+#define UE_API RIGVMDEVELOPER_API
+
 USTRUCT()
-struct RIGVMDEVELOPER_API FRigVMTemplatePreferredType
+struct FRigVMTemplatePreferredType
 {
 	GENERATED_BODY()
 
@@ -43,8 +45,8 @@ public:
 		return Argument == Other.Argument && TypeIndex == Other.TypeIndex;
 	}
 	
-	void UpdateStringFromIndex();
-	void UpdateIndexFromString();
+	UE_API void UpdateStringFromIndex();
+	UE_API void UpdateIndexFromString();
 
 protected:
 	
@@ -66,86 +68,86 @@ protected:
  * Template nodes can morph into all functions implementing
  * the template's template.
  */
-UCLASS(BlueprintType)
-class RIGVMDEVELOPER_API URigVMTemplateNode : public URigVMNode
+UCLASS(MinimalAPI, BlueprintType)
+class URigVMTemplateNode : public URigVMNode
 {
 	GENERATED_BODY()
 
 public:
 
 	// default constructor
-	URigVMTemplateNode();
+	UE_API URigVMTemplateNode();
 
 	// UObject interface
-	virtual void PreSave(FObjectPreSaveContext SaveContext) override;
-	virtual void PostLoad() override;
+	UE_API virtual void PreSave(FObjectPreSaveContext SaveContext) override;
+	UE_API virtual void PostLoad() override;
 
 	// URigVMNode interface
-	virtual FString GetNodeTitle() const override;
-	virtual FName GetMethodName() const;
-	virtual  FText GetToolTipText() const override;
-	virtual FText GetToolTipTextForPin(const URigVMPin* InPin) const override;
-	virtual TArray<URigVMPin*> GetAggregateInputs() const override;
-	virtual TArray<URigVMPin*> GetAggregateOutputs() const override;
-	TArray<URigVMPin*> GetAggregatePins(const ERigVMPinDirection& InDirection) const;
-	virtual FName GetNextAggregateName(const FName& InLastAggregatePinName) const override;
+	UE_API virtual FString GetNodeTitle() const override;
+	UE_API virtual FName GetMethodName() const;
+	UE_API virtual  FText GetToolTipText() const override;
+	UE_API virtual FText GetToolTipTextForPin(const URigVMPin* InPin) const override;
+	UE_API virtual TArray<URigVMPin*> GetAggregateInputs() const override;
+	UE_API virtual TArray<URigVMPin*> GetAggregateOutputs() const override;
+	UE_API TArray<URigVMPin*> GetAggregatePins(const ERigVMPinDirection& InDirection) const;
+	UE_API virtual FName GetNextAggregateName(const FName& InLastAggregatePinName) const override;
 
 	// Returns the UStruct for this unit node
 	// (the struct declaring the RIGVM_METHOD)
 	UFUNCTION(BlueprintCallable, Category = RigVMUnitNode)
-	virtual UScriptStruct* GetScriptStruct() const;
+	UE_API virtual UScriptStruct* GetScriptStruct() const;
 
 	// Returns the notation of the node
 	UFUNCTION(BlueprintPure, Category = Template)
-	virtual FName GetNotation() const;
+	UE_API virtual FName GetNotation() const;
 
 	UFUNCTION(BlueprintCallable, Category = Template)
-	virtual bool IsSingleton() const;
+	UE_API virtual bool IsSingleton() const;
 
 	// returns true if a pin supports a given type
-	bool SupportsType(const URigVMPin* InPin, TRigVMTypeIndex InTypeIndex, TRigVMTypeIndex* OutTypeIndex = nullptr);
+	UE_API bool SupportsType(const URigVMPin* InPin, TRigVMTypeIndex InTypeIndex, TRigVMTypeIndex* OutTypeIndex = nullptr);
 
 	// returns the resolved functions for the template
-	TArray<const FRigVMFunction*> GetResolvedPermutations() const;
+	UE_API TArray<const FRigVMFunction*> GetResolvedPermutations() const;
 
 	// returns the resolved permutation indices for the template
-	TArray<int32> GetResolvedPermutationIndices(bool bAllowFloatingPointCasts) const;
+	UE_API TArray<int32> GetResolvedPermutationIndices(bool bAllowFloatingPointCasts) const;
 
 	// Returns a map of the resolved pin types
-	FRigVMTemplateTypeMap GetTemplatePinTypeMap(bool bIncludeHiddenPins = false, bool bIncludeExecutePins = false) const;
+	UE_API FRigVMTemplateTypeMap GetTemplatePinTypeMap(bool bIncludeHiddenPins = false, bool bIncludeExecutePins = false) const;
 
 	// returns the template used for this node
-	virtual const FRigVMTemplate* GetTemplate() const;
+	UE_API virtual const FRigVMTemplate* GetTemplate() const;
 
 	// returns the resolved function or nullptr if there are still unresolved pins left
-	const FRigVMFunction* GetResolvedFunction() const;
+	UE_API const FRigVMFunction* GetResolvedFunction() const;
 
 	// returns true if the template node is resolved
 	UFUNCTION(BlueprintPure, Category = Template)
-	bool IsResolved() const;
+	UE_API bool IsResolved() const;
 
 	// returns true if the template is fully unresolved
 	UFUNCTION(BlueprintPure, Category = Template)
-	bool IsFullyUnresolved() const;
+	UE_API bool IsFullyUnresolved() const;
 
 	// returns a default value for pin if it is known
-	FString GetInitialDefaultValueForPin(const FName& InRootPinName, const TArray<int32>& InPermutationIndices = TArray<int32>()) const;
+	UE_API FString GetInitialDefaultValueForPin(const FName& InRootPinName, const TArray<int32>& InPermutationIndices = TArray<int32>()) const;
 
 	// returns the display name for a pin
-	virtual FName GetDisplayNameForPin(const URigVMPin* InPin) const override;
+	UE_API virtual FName GetDisplayNameForPin(const URigVMPin* InPin) const override;
 
 	// Tries to reduce the input types to a single type, if all are compatible
 	// Will prioritize the InPreferredType if available
-	TRigVMTypeIndex TryReduceTypesToSingle(const TArray<TRigVMTypeIndex>& InTypes, const TRigVMTypeIndex PreferredType = TRigVMTypeIndex()) const;
+	UE_API TRigVMTypeIndex TryReduceTypesToSingle(const TArray<TRigVMTypeIndex>& InTypes, const TRigVMTypeIndex PreferredType = TRigVMTypeIndex()) const;
 
-	virtual uint32 GetStructureHash() const override;
+	UE_API virtual uint32 GetStructureHash() const override;
 
 protected:
 
-	virtual void InvalidateCache() override;
+	UE_API virtual void InvalidateCache() override;
 	
-	TArray<int32> FindPermutationsForTypes(const TArray<FRigVMTemplatePreferredType>& ArgumentTypes, bool bAllowCasting = false) const;
-	FRigVMTemplateTypeMap GetTypesForPermutation(const int32 InPermutationIndex) const;
+	UE_API TArray<int32> FindPermutationsForTypes(const TArray<FRigVMTemplatePreferredType>& ArgumentTypes, bool bAllowCasting = false) const;
+	UE_API FRigVMTemplateTypeMap GetTypesForPermutation(const int32 InPermutationIndex) const;
 
 	UPROPERTY()
 	FName TemplateNotation;
@@ -177,3 +179,4 @@ protected:
 	friend class URigVMCompiler;
 };
 
+#undef UE_API
diff --git a/Engine/Plugins/Runtime/RigVM/Source/RigVMDeveloper/Public/RigVMModel/Nodes/RigVMUnitNode.h b/Engine/Plugins/Runtime/RigVM/Source/RigVMDeveloper/Public/RigVMModel/Nodes/RigVMUnitNode.h
--- a/Engine/Plugins/Runtime/RigVM/Source/RigVMDeveloper/Public/RigVMModel/Nodes/RigVMUnitNode.h
+++ b/Engine/Plugins/Runtime/RigVM/Source/RigVMDeveloper/Public/RigVMModel/Nodes/RigVMUnitNode.h
@@ -8,6 +8,8 @@
 #include "RigVMModel/RigVMModelCachedValue.h"
 #include "RigVMUnitNode.generated.h"
 
+#define UE_API RIGVMDEVELOPER_API
+
 class URigVMHost;
 
 /**
@@ -15,66 +17,66 @@ class URigVMHost;
  * declared on a USTRUCT. Struct Nodes have input / output pins for all
  * struct UPROPERTY members.
  */
-UCLASS(BlueprintType)
-class RIGVMDEVELOPER_API URigVMUnitNode : public URigVMTemplateNode
+UCLASS(MinimalAPI, BlueprintType)
+class URigVMUnitNode : public URigVMTemplateNode
 {
 	GENERATED_BODY()
 
 public:
 
 	// UObject interface
-	virtual void PostLoad() override;
+	UE_API virtual void PostLoad() override;
 
 	// Override node functions
-	virtual FString GetNodeTitle() const override;
-	virtual FText GetToolTipText() const override;
-	virtual bool IsDefinedAsConstant() const override;
-	virtual bool IsDefinedAsVarying() const override;
-	virtual FName GetEventName() const override;
-	virtual bool CanOnlyExistOnce() const override;
-	virtual const TArray<FName>& GetControlFlowBlocks() const override;
-	virtual const bool IsControlFlowBlockSliced(const FName& InBlockName) const override;
-	virtual TArray<FRigVMUserWorkflow> GetSupportedWorkflows(ERigVMUserWorkflowType InType, const UObject* InSubject) const override;
-	virtual TArray<URigVMPin*> GetAggregateInputs() const override;
-	virtual TArray<URigVMPin*> GetAggregateOutputs() const override;
-	virtual FName GetNextAggregateName(const FName& InLastAggregatePinName) const override;
-	virtual FName GetDisplayNameForPin(const URigVMPin* InPin) const override;
-
-	virtual bool IsOutDated() const override;
-	virtual FString GetDeprecatedMetadata() const override;
+	UE_API virtual FString GetNodeTitle() const override;
+	UE_API virtual FText GetToolTipText() const override;
+	UE_API virtual bool IsDefinedAsConstant() const override;
+	UE_API virtual bool IsDefinedAsVarying() const override;
+	UE_API virtual FName GetEventName() const override;
+	UE_API virtual bool CanOnlyExistOnce() const override;
+	UE_API virtual const TArray<FName>& GetControlFlowBlocks() const override;
+	UE_API virtual const bool IsControlFlowBlockSliced(const FName& InBlockName) const override;
+	UE_API virtual TArray<FRigVMUserWorkflow> GetSupportedWorkflows(ERigVMUserWorkflowType InType, const UObject* InSubject) const override;
+	UE_API virtual TArray<URigVMPin*> GetAggregateInputs() const override;
+	UE_API virtual TArray<URigVMPin*> GetAggregateOutputs() const override;
+	UE_API virtual FName GetNextAggregateName(const FName& InLastAggregatePinName) const override;
+	UE_API virtual FName GetDisplayNameForPin(const URigVMPin* InPin) const override;
+
+	UE_API virtual bool IsOutDated() const override;
+	UE_API virtual FString GetDeprecatedMetadata() const override;
 
 	// URigVMTemplateNode interface
-	virtual UScriptStruct* GetScriptStruct() const override;
+	UE_API virtual UScriptStruct* GetScriptStruct() const override;
 
 	// Returns the name of the declared RIGVM_METHOD
 	UFUNCTION(BlueprintCallable, Category = RigVMUnitNode)
-	virtual FName GetMethodName() const;
+	UE_API virtual FName GetMethodName() const;
 
 	// Returns the default value for the struct as text
 	UFUNCTION(BlueprintCallable, Category = RigVMUnitNode)
-	FString GetStructDefaultValue() const;
+	UE_API FString GetStructDefaultValue() const;
 
 	// Returns an instance of the struct with the current values.
 	// @param bUseDefault If set to true the default struct will be created - otherwise the struct will contains the values from the node
-	TSharedPtr<FStructOnScope> ConstructStructInstance(bool bUseDefault = false) const;
+	UE_API TSharedPtr<FStructOnScope> ConstructStructInstance(bool bUseDefault = false) const;
 
 	// Returns an instance of the struct with values backed by the memory of a currently running host
-	TSharedPtr<FStructOnScope> ConstructLiveStructInstance(URigVMHost* InHost, int32 InSliceIndex = 0) const;
+	UE_API TSharedPtr<FStructOnScope> ConstructLiveStructInstance(URigVMHost* InHost, int32 InSliceIndex = 0) const;
 	
 	// Updates the memory of a host to match a specific struct instance
-	bool UpdateHostFromStructInstance(URigVMHost* InHost, TSharedPtr<FStructOnScope> InInstance, int32 InSliceIndex = 0) const;
+	UE_API bool UpdateHostFromStructInstance(URigVMHost* InHost, TSharedPtr<FStructOnScope> InInstance, int32 InSliceIndex = 0) const;
 
 	// Compares two struct instances and returns differences in pin values
-	void ComputePinValueDifferences(TSharedPtr<FStructOnScope> InCurrentInstance, TSharedPtr<FStructOnScope> InDesiredInstance, TMap<FString, FString>& OutNewPinDefaultValues) const;
+	UE_API void ComputePinValueDifferences(TSharedPtr<FStructOnScope> InCurrentInstance, TSharedPtr<FStructOnScope> InDesiredInstance, TMap<FString, FString>& OutNewPinDefaultValues) const;
 
 	// Compares a desired struct instance with the node's default and returns differences in pin values
-	void ComputePinValueDifferences(TSharedPtr<FStructOnScope> InDesiredInstance, TMap<FString, FString>& OutNewPinDefaultValues) const;
+	UE_API void ComputePinValueDifferences(TSharedPtr<FStructOnScope> InDesiredInstance, TMap<FString, FString>& OutNewPinDefaultValues) const;
 
 	// Returns true if the node is part of the debugged runtime rig
-	bool IsPartOfRuntime() const;
+	UE_API bool IsPartOfRuntime() const;
 
 	// Returns true if the node is part of the debugged runtime rig
-	bool IsPartOfRuntime(URigVMHost* InHost) const;
+	UE_API bool IsPartOfRuntime(URigVMHost* InHost) const;
 
 	// Returns a copy of the struct with the current values
 	template <
@@ -93,16 +95,16 @@ public:
 		return InstanceRef;
 	}
 
-	virtual FRigVMStructUpgradeInfo GetUpgradeInfo() const override;
+	UE_API virtual FRigVMStructUpgradeInfo GetUpgradeInfo() const override;
 
-	virtual uint32 GetStructureHash() const override;
+	UE_API virtual uint32 GetStructureHash() const override;
 
 protected:
 
-	virtual FText GetToolTipTextForPin(const URigVMPin* InPin) const override;
-	virtual bool ShouldInputPinComputeLazily(const URigVMPin* InPin) const override;
-	virtual FString GetOriginalDefaultValueForRootPin(const URigVMPin* InRootPin) const override;
-	void EnumeratePropertiesOnHostAndStructInstance(
+	UE_API virtual FText GetToolTipTextForPin(const URigVMPin* InPin) const override;
+	UE_API virtual bool ShouldInputPinComputeLazily(const URigVMPin* InPin) const override;
+	UE_API virtual FString GetOriginalDefaultValueForRootPin(const URigVMPin* InRootPin) const override;
+	UE_API void EnumeratePropertiesOnHostAndStructInstance(
 		URigVMHost* InHost,
 		TSharedPtr<FStructOnScope> InInstance, 
 		bool bPreferLiterals,
@@ -124,3 +126,4 @@ private:
 	friend class URigVMController;
 };
 
+#undef UE_API
diff --git a/Engine/Plugins/Runtime/RigVM/Source/RigVMDeveloper/Public/RigVMModel/Nodes/RigVMVariableNode.h b/Engine/Plugins/Runtime/RigVM/Source/RigVMDeveloper/Public/RigVMModel/Nodes/RigVMVariableNode.h
--- a/Engine/Plugins/Runtime/RigVM/Source/RigVMDeveloper/Public/RigVMModel/Nodes/RigVMVariableNode.h
+++ b/Engine/Plugins/Runtime/RigVM/Source/RigVMDeveloper/Public/RigVMModel/Nodes/RigVMVariableNode.h
@@ -7,70 +7,72 @@
 #include "RigVMModel/RigVMVariableDescription.h"
 #include "RigVMVariableNode.generated.h"
 
+#define UE_API RIGVMDEVELOPER_API
+
 /**
  * The Variable Node represents a mutable value / local state within the
  * the Function / Graph. Variable Node's can be a getter or a setter.
  * Getters are pure nodes with just an output value pin, while setters
  * are mutable nodes with an execute and input value pin.
  */
-UCLASS(BlueprintType)
-class RIGVMDEVELOPER_API URigVMVariableNode : public URigVMNode
+UCLASS(MinimalAPI, BlueprintType)
+class URigVMVariableNode : public URigVMNode
 {
 	GENERATED_BODY()
 
 public:
 
 	// Default constructor
-	URigVMVariableNode();
+	UE_API URigVMVariableNode();
 
 	// Override of node title
-	virtual FString GetNodeTitle() const;
+	UE_API virtual FString GetNodeTitle() const;
 
 	// Returns the name of the variable
 	UFUNCTION(BlueprintCallable, Category = RigVMVariableNode)
-	FName GetVariableName() const;
+	UE_API FName GetVariableName() const;
 
 	// Returns true if this node is a variable getter
 	UFUNCTION(BlueprintCallable, Category = RigVMVariableNode)
-	bool IsGetter() const;
+	UE_API bool IsGetter() const;
 
 	// Returns true if this variable is an external variable
 	UFUNCTION(BlueprintCallable, Category = RigVMVariableNode)
-	bool IsExternalVariable() const;
+	UE_API bool IsExternalVariable() const;
 
 	// Returns true if this variable is a local variable
 	UFUNCTION(BlueprintCallable, Category = RigVMVariableNode)
-	bool IsLocalVariable() const;
+	UE_API bool IsLocalVariable() const;
 
 	// Returns true if this variable is an input argument
 	UFUNCTION(BlueprintCallable, Category = RigVMVariableNode)
-	bool IsInputArgument() const;
+	UE_API bool IsInputArgument() const;
 
 	// Returns the C++ data type of the variable
 	UFUNCTION(BlueprintCallable, Category = RigVMVariableNode)
-	FString GetCPPType() const;
+	UE_API FString GetCPPType() const;
 
 	// Returns the C++ data type struct of the variable (or nullptr)
 	UFUNCTION(BlueprintCallable, Category = RigVMVariableNode)
-	UObject* GetCPPTypeObject() const;
+	UE_API UObject* GetCPPTypeObject() const;
 
 	// Returns the default value of the variable as a string
 	UFUNCTION(BlueprintCallable, Category = RigVMVariableNode)
-	FString GetDefaultValue() const;
+	UE_API FString GetDefaultValue() const;
 
 	// Returns this variable node's variable description
 	UFUNCTION(BlueprintCallable, Category = RigVMVariableNode)
-	FRigVMGraphVariableDescription GetVariableDescription() const;
+	UE_API FRigVMGraphVariableDescription GetVariableDescription() const;
 
 	// Override of node title
 	virtual FLinearColor GetNodeColor() const override { return FLinearColor::Blue; }
 	virtual bool IsDefinedAsVarying() const override { return true; }
 
 	// Get the variable name pin
-	URigVMPin* GetVariableNamePin() const;
+	UE_API URigVMPin* GetVariableNamePin() const;
 
 	// Get the value pin
-	URigVMPin* GetValuePin() const;
+	UE_API URigVMPin* GetValuePin() const;
 
 	// Get the subtitle of the node
 	virtual FString GetNodeSubTitle() const;
@@ -92,3 +94,4 @@ private:
 	friend class URigVMEdGraphVariableNodeSpawner;
 };
 
+#undef UE_API
diff --git a/Engine/Plugins/Runtime/RigVM/Source/RigVMDeveloper/Public/RigVMModel/RigVMBuildData.h b/Engine/Plugins/Runtime/RigVM/Source/RigVMDeveloper/Public/RigVMModel/RigVMBuildData.h
--- a/Engine/Plugins/Runtime/RigVM/Source/RigVMDeveloper/Public/RigVMModel/RigVMBuildData.h
+++ b/Engine/Plugins/Runtime/RigVM/Source/RigVMDeveloper/Public/RigVMModel/RigVMBuildData.h
@@ -5,8 +5,10 @@
 #include "RigVMModel/Nodes/RigVMFunctionReferenceNode.h"
 #include "RigVMBuildData.generated.h"
 
+#define UE_API RIGVMDEVELOPER_API
+
 USTRUCT(BlueprintType)
-struct RIGVMDEVELOPER_API FRigVMFunctionReferenceArray
+struct FRigVMFunctionReferenceArray
 {
 	GENERATED_BODY()
 
@@ -27,7 +29,7 @@ struct RIGVMDEVELOPER_API FRigVMFunctionReferenceArray
 };
 
 USTRUCT()
-struct RIGVMDEVELOPER_API FRigVMReferenceNodeData
+struct FRigVMReferenceNodeData
 {
 	GENERATED_BODY();
 
@@ -37,7 +39,7 @@ public:
 		:ReferenceNodePath(), ReferencedFunctionIdentifier()
 	{}
 
-	FRigVMReferenceNodeData(URigVMFunctionReferenceNode* InReferenceNode);
+	UE_API FRigVMReferenceNodeData(URigVMFunctionReferenceNode* InReferenceNode);
 
 	UPROPERTY()
 	FString ReferenceNodePath;
@@ -51,8 +53,8 @@ public:
 	UPROPERTY()
 	FRigVMGraphFunctionIdentifier ReferencedFunctionIdentifier;
 
-	TSoftObjectPtr<URigVMFunctionReferenceNode> GetReferenceNodeObjectPath();
-	URigVMFunctionReferenceNode* GetReferenceNode();
+	UE_API TSoftObjectPtr<URigVMFunctionReferenceNode> GetReferenceNodeObjectPath();
+	UE_API URigVMFunctionReferenceNode* GetReferenceNode();
 
 private:
 
@@ -64,8 +66,8 @@ private:
  * for the RigVM graph to improve the user experience.
  * This object is never serialized.
  */
-UCLASS(BlueprintType)
-class RIGVMDEVELOPER_API URigVMBuildData : public UObject
+UCLASS(MinimalAPI, BlueprintType)
+class URigVMBuildData : public UObject
 {
 	GENERATED_BODY()
 
@@ -73,15 +75,15 @@ public:
 
 	// Returns the singleton build data
 	UFUNCTION(BlueprintPure, Category=RigVMBuildData)
-	static URigVMBuildData* Get();
+	static UE_API URigVMBuildData* Get();
 	
 	// Looks for all function references (in RigVMClientHost metadata) and initializes the URigVMBuildData
-	void InitializeIfNeeded();
+	UE_API void InitializeIfNeeded();
 
-	static void RegisterReferencesFromAsset(const FAssetData& AssetData);
+	static UE_API void RegisterReferencesFromAsset(const FAssetData& AssetData);
 
 	// Returns the list of references for a given function definition
-	const FRigVMFunctionReferenceArray* FindFunctionReferences(const FRigVMGraphFunctionIdentifier& InFunction) const;
+	UE_API const FRigVMFunctionReferenceArray* FindFunctionReferences(const FRigVMGraphFunctionIdentifier& InFunction) const;
 
 	/**
 	 * Iterator function to invoke a lambda / TFunction for each reference of a function
@@ -89,114 +91,114 @@ public:
 	 * @param PerReferenceFunction The function to invoke for each reference
 	 * @param bLoadIfNecessary If true, will load packages if necessary
 	 */
-	void ForEachFunctionReference(const FRigVMGraphFunctionIdentifier& InFunction, TFunction<void(URigVMFunctionReferenceNode*)> PerReferenceFunction, bool bLoadIfNecessary = true) const;
+	UE_API void ForEachFunctionReference(const FRigVMGraphFunctionIdentifier& InFunction, TFunction<void(URigVMFunctionReferenceNode*)> PerReferenceFunction, bool bLoadIfNecessary = true) const;
 
 	/**
 	* Iterator function to invoke a lambda / TFunction for each reference of a function
 	* @param InFunction The function to iterate all references for
 	* @param PerReferenceFunction The function to invoke for each reference
 	*/
-	void ForEachFunctionReferenceSoftPtr(const FRigVMGraphFunctionIdentifier& InFunction, TFunction<void(TSoftObjectPtr<URigVMFunctionReferenceNode>)>
+	UE_API void ForEachFunctionReferenceSoftPtr(const FRigVMGraphFunctionIdentifier& InFunction, TFunction<void(TSoftObjectPtr<URigVMFunctionReferenceNode>)>
 	                                     PerReferenceFunction) const;
 
 	// registers a new reference node for a given function
-	void RegisterFunctionReference(const FRigVMGraphFunctionIdentifier& InFunction, URigVMFunctionReferenceNode* InReference);
+	UE_API void RegisterFunctionReference(const FRigVMGraphFunctionIdentifier& InFunction, URigVMFunctionReferenceNode* InReference);
 
 	// registers a new reference node for a given function
-	void RegisterFunctionReference(const FRigVMGraphFunctionIdentifier& InFunction, TSoftObjectPtr<URigVMFunctionReferenceNode> InReference);
+	UE_API void RegisterFunctionReference(const FRigVMGraphFunctionIdentifier& InFunction, TSoftObjectPtr<URigVMFunctionReferenceNode> InReference);
 
 	// registers a new reference node for a given function
-	void RegisterFunctionReference(FRigVMReferenceNodeData InReferenceNodeData);
+	UE_API void RegisterFunctionReference(FRigVMReferenceNodeData InReferenceNodeData);
 
 	// unregisters a new reference node for a given function
-	void UnregisterFunctionReference(const FRigVMGraphFunctionIdentifier& InFunction, URigVMFunctionReferenceNode* InReference);
+	UE_API void UnregisterFunctionReference(const FRigVMGraphFunctionIdentifier& InFunction, URigVMFunctionReferenceNode* InReference);
 
 	// unregisters a new reference node for a given function
-	void UnregisterFunctionReference(const FRigVMGraphFunctionIdentifier& InFunction, TSoftObjectPtr<URigVMFunctionReferenceNode> InReference);
+	UE_API void UnregisterFunctionReference(const FRigVMGraphFunctionIdentifier& InFunction, TSoftObjectPtr<URigVMFunctionReferenceNode> InReference);
 
 	// Clear references to temp assets
-	void ClearInvalidReferences();
+	UE_API void ClearInvalidReferences();
 
 	// Helper function to disable clearing transient package references
 	void SetIsRunningUnitTest(bool bIsRunning) { bIsRunningUnitTest = bIsRunning; }
 
 	// Will find all public function variant refs, and private function variant refs from loaded assets
 	UFUNCTION(BlueprintPure, Category=RigVMBuildData)
-	TArray<FRigVMVariantRef> GatherAllFunctionVariantRefs() const;
+	UE_API TArray<FRigVMVariantRef> GatherAllFunctionVariantRefs() const;
 
 	// Will find the public function variant refs inside this asset, and private function variant refs if the asset is loaded
 	UFUNCTION(BlueprintPure, Category=RigVMBuildData)
-	TArray<FRigVMVariantRef> GatherFunctionVariantRefsForAsset(const FAssetData& InAssetData) const;
+	UE_API TArray<FRigVMVariantRef> GatherFunctionVariantRefsForAsset(const FAssetData& InAssetData) const;
 
 	// Will find all the function variants matching the given variant guid
 	UFUNCTION(BlueprintPure, Category=RigVMBuildData)
-	TArray<FRigVMVariantRef> FindFunctionVariantRefs(const FGuid& InGuid) const;
+	UE_API TArray<FRigVMVariantRef> FindFunctionVariantRefs(const FGuid& InGuid) const;
 
 	// Returns the function identifier given a variant reference
 	UFUNCTION(BlueprintPure, Category=RigVMBuildData)
-	FRigVMGraphFunctionIdentifier GetFunctionIdentifierForVariant(const FRigVMVariantRef& InVariantRef) const;
+	UE_API FRigVMGraphFunctionIdentifier GetFunctionIdentifierForVariant(const FRigVMVariantRef& InVariantRef) const;
 
 	// Creates a new variant of a graph function within the same asset
 	UFUNCTION(BlueprintCallable, Category=RigVMBuildData)
-	FRigVMVariantRef CreateFunctionVariant(const FRigVMGraphFunctionIdentifier& InFunctionIdentifier, FName InName = NAME_None);
+	UE_API FRigVMVariantRef CreateFunctionVariant(const FRigVMGraphFunctionIdentifier& InFunctionIdentifier, FName InName = NAME_None);
 
 	// Will find all asset variant refs
 	UFUNCTION(BlueprintPure, Category=RigVMBuildData)
-	TArray<FRigVMVariantRef> GatherAllAssetVariantRefs() const;
+	UE_API TArray<FRigVMVariantRef> GatherAllAssetVariantRefs() const;
 
 	// Will find all the asset variants matching the given variant guid
 	UFUNCTION(BlueprintPure, Category=RigVMBuildData)
-	TArray<FRigVMVariantRef> FindAssetVariantRefs(const FGuid& InGuid) const;
+	UE_API TArray<FRigVMVariantRef> FindAssetVariantRefs(const FGuid& InGuid) const;
 
 	// Creates a new variant of an asset
 	UFUNCTION(BlueprintCallable, Category=RigVMBuildData)
-	FRigVMVariantRef CreateAssetVariant(const FAssetData& InAssetData, FName InName = NAME_None);
+	UE_API FRigVMVariantRef CreateAssetVariant(const FAssetData& InAssetData, FName InName = NAME_None);
 
 	// Returns the asset data given an object path
 	UFUNCTION(BlueprintPure, Category=RigVMBuildData)
-	FAssetData GetAssetDataForPath(const FSoftObjectPath& InObjectPath) const;
+	UE_API FAssetData GetAssetDataForPath(const FSoftObjectPath& InObjectPath) const;
 
 	// Returns the asset data given a variant reference
 	UFUNCTION(BlueprintPure, Category=RigVMBuildData)
-	FRigVMVariantRef GetVariantRefForAsset(const FAssetData& InAssetData) const;
+	UE_API FRigVMVariantRef GetVariantRefForAsset(const FAssetData& InAssetData) const;
 	
 	// Returns the asset data given a variant reference
 	UFUNCTION(BlueprintPure, Category=RigVMBuildData)
-	FAssetData GetAssetDataForVariant(const FRigVMVariantRef& InVariantRef) const;
+	UE_API FAssetData GetAssetDataForVariant(const FRigVMVariantRef& InVariantRef) const;
 
 	// Splits a variant from its variant set (by assigning a new, unique guid
 	UFUNCTION(BlueprintPure, Category=RigVMBuildData)
-	FRigVMVariantRef SplitVariantFromSet(const FRigVMVariantRef& InVariantRef);
+	UE_API FRigVMVariantRef SplitVariantFromSet(const FRigVMVariantRef& InVariantRef);
 
 	// Joins a variant with another variant's set
 	UFUNCTION(BlueprintPure, Category=RigVMBuildData)
-	FRigVMVariantRef JoinVariantSet(const FRigVMVariantRef& InVariantRef, const FGuid& InGuid);
+	UE_API FRigVMVariantRef JoinVariantSet(const FRigVMVariantRef& InVariantRef, const FGuid& InGuid);
 
 #if WITH_EDITOR
 	// Returns all known public function identifiers used in the project
 	UFUNCTION(BlueprintPure, Category=RigVMBuildData)
-	TArray<FRigVMGraphFunctionIdentifier> GetAllFunctionIdentifiers(bool bOnlyPublic = true) const;
+	UE_API TArray<FRigVMGraphFunctionIdentifier> GetAllFunctionIdentifiers(bool bOnlyPublic = true) const;
 #endif
 
 	// Returns all known public function identifiers used in the project
 	UFUNCTION(BlueprintPure, Category=RigVMBuildData)
-	TArray<FRigVMGraphFunctionIdentifier> GetUsedFunctionIdentifiers(bool bOnlyPublic = true) const;
+	UE_API TArray<FRigVMGraphFunctionIdentifier> GetUsedFunctionIdentifiers(bool bOnlyPublic = true) const;
 
 	// Returns all known function references
-	FRigVMFunctionReferenceArray GetAllFunctionReferences() const;
+	UE_API FRigVMFunctionReferenceArray GetAllFunctionReferences() const;
 
 private:
 
 	// disable default constructor
-	URigVMBuildData();
+	UE_API URigVMBuildData();
 
-	static TArray<UClass*> FindAllRigVMAssetClasses();
+	static UE_API TArray<UClass*> FindAllRigVMAssetClasses();
 
-	void SetupRigVMGraphFunctionPointers();
-	void TearDownRigVMGraphFunctionPointers();
-	static TArray<FRigVMGraphFunctionHeader> GetFunctionHeadersForAsset(const FAssetData& InAssetData);
+	UE_API void SetupRigVMGraphFunctionPointers();
+	UE_API void TearDownRigVMGraphFunctionPointers();
+	static UE_API TArray<FRigVMGraphFunctionHeader> GetFunctionHeadersForAsset(const FAssetData& InAssetData);
 	
-	static bool bInitialized;
+	static UE_API bool bInitialized;
 
 	UPROPERTY(meta=(DeprecatedProperty))
 	TMap< TSoftObjectPtr<URigVMLibraryNode>, FRigVMFunctionReferenceArray > FunctionReferences_DEPRECATED;
@@ -212,3 +214,4 @@ private:
 	friend class FRigVMDeveloperModule;
 };
 
+#undef UE_API
diff --git a/Engine/Plugins/Runtime/RigVM/Source/RigVMDeveloper/Public/RigVMModel/RigVMClient.h b/Engine/Plugins/Runtime/RigVM/Source/RigVMDeveloper/Public/RigVMModel/RigVMClient.h
--- a/Engine/Plugins/Runtime/RigVM/Source/RigVMDeveloper/Public/RigVMModel/RigVMClient.h
+++ b/Engine/Plugins/Runtime/RigVM/Source/RigVMDeveloper/Public/RigVMModel/RigVMClient.h
@@ -9,10 +9,12 @@
 #include "UObject/WeakInterfacePtr.h"
 #include "RigVMClient.generated.h"
 
+#define UE_API RIGVMDEVELOPER_API
+
 struct FRigVMClient;
 
-UINTERFACE()
-class RIGVMDEVELOPER_API URigVMClientHost : public UInterface
+UINTERFACE(MinimalAPI)
+class URigVMClientHost : public UInterface
 {
 	GENERATED_BODY()
 };
@@ -26,7 +28,7 @@ enum class ERigVMLoadType : uint8
 DECLARE_DELEGATE_RetVal(URigVMGraph*, FRigVMGetFocusedGraph);
 
 // Interface that allows an object to host a rig VM client. Used by graph edting code to interact with the controller.
-class RIGVMDEVELOPER_API IRigVMClientHost
+class IRigVMClientHost
 {
 	GENERATED_BODY()
 
@@ -85,7 +87,7 @@ public:
 
 	// Given a type name for a user-defined type, either struct or an enum, returns a pointer to the object definition
 	// or nullptr if the client host has no knowledge of such a type.
-	virtual UObject* ResolveUserDefinedTypeById(const FString& InTypeName) const;
+	UE_API virtual UObject* ResolveUserDefinedTypeById(const FString& InTypeName) const;
 
 	// Recompiles the VM if not already being compiled
 	virtual void RecompileVM() = 0;
@@ -156,14 +158,14 @@ public:
 	virtual void RenameGraph(const FString& InNodePath, const FName& InNewName) = 0;
 };
 
-UINTERFACE()
-class RIGVMDEVELOPER_API URigVMEditorSideObject : public UInterface
+UINTERFACE(MinimalAPI)
+class URigVMEditorSideObject : public UInterface
 {
 	GENERATED_BODY()
 };
 
 // Interface that allows a UI graph to identify itself against a model graph
-class RIGVMDEVELOPER_API IRigVMEditorSideObject
+class IRigVMEditorSideObject
 {
 	GENERATED_BODY()
 
@@ -179,14 +181,14 @@ public:
 	virtual void HandleRigVMGraphRenamed(const FString& InOldNodePath, const FString& InNewNodePath) = 0;
 };
 
-UINTERFACE()
-class RIGVMDEVELOPER_API URigVMClientExternalModelHost : public UInterface
+UINTERFACE(MinimalAPI)
+class URigVMClientExternalModelHost : public UInterface
 {
 	GENERATED_BODY()
 };
 
 // Interface that allows access to externally-hosted models
-class RIGVMDEVELOPER_API IRigVMClientExternalModelHost
+class IRigVMClientExternalModelHost
 {
 	GENERATED_BODY()
 
@@ -202,7 +204,7 @@ public:
 
 // A management struct containing graphs and controllers.
 USTRUCT()
-struct RIGVMDEVELOPER_API FRigVMClient
+struct FRigVMClient
 {
 public:
 
@@ -225,108 +227,108 @@ public:
 
 	UE_DEPRECATED(5.5, "Please use SetDefaultSchemaClass or set a schema per controller/graph.")
 	void SetSchemaClass(TSubclassOf<URigVMSchema> InSchemaClass) { SetDefaultSchemaClass(InSchemaClass); }
-	void SetDefaultSchemaClass(TSubclassOf<URigVMSchema> InSchemaClass);
-	void SetControllerClass(TSubclassOf<URigVMController> InControllerClass);
-	void SetOuterClientHost(UObject* InOuterClientHost, const FName& InOuterClientHostPropertyName);
-	void SetFromDeprecatedData(URigVMGraph* InDefaultGraph, URigVMFunctionLibrary* InFunctionLibrary);
-	void SetExternalModelHost(IRigVMClientExternalModelHost* InExternalModelHost);
+	UE_API void SetDefaultSchemaClass(TSubclassOf<URigVMSchema> InSchemaClass);
+	UE_API void SetControllerClass(TSubclassOf<URigVMController> InControllerClass);
+	UE_API void SetOuterClientHost(UObject* InOuterClientHost, const FName& InOuterClientHostPropertyName);
+	UE_API void SetFromDeprecatedData(URigVMGraph* InDefaultGraph, URigVMFunctionLibrary* InFunctionLibrary);
+	UE_API void SetExternalModelHost(IRigVMClientExternalModelHost* InExternalModelHost);
 
-	void Reset();
+	UE_API void Reset();
 	int32 Num() const { return GetModels().Num(); }
 	UE_DEPRECATED(5.5, "Please use GetDefaultSchema or get a schema per controller/graph.")
 	URigVMSchema* GetSchema() { return GetDefaultSchema(); }
-	URigVMSchema* GetDefaultSchema() const;
+	UE_API URigVMSchema* GetDefaultSchema() const;
 	TSubclassOf<URigVMSchema> GetDefaultSchemaClass() const { return DefaultSchemaClass; }
 	UE_DEPRECATED(5.5, "Please use GetSchema or get a schema per controller/graph.")
 	URigVMSchema* GetOrCreateSchema() { return GetDefaultSchema(); }
-	URigVMGraph* GetDefaultModel() const;
-	URigVMGraph* GetModel(int32 InIndex) const;
-	URigVMGraph* GetModel(const UEdGraph* InEdGraph = nullptr) const;
-	URigVMGraph* GetModel(const FString& InNodePathOrName) const;
-	URigVMGraph* GetModel(const UObject* InEditorSideObject) const;
-	void RefreshAllModels(ERigVMLoadType InLoadType, bool bEnablePostLoadHashing, bool& bIsCompiling);
-	const TArray<TObjectPtr<URigVMGraph>>& GetModels() const;
-	TArray<URigVMGraph*> GetAllModels(bool bIncludeFunctionLibrary, bool bRecursive) const;
-	TArray<URigVMGraph*> GetAllModelsLeavesFirst(bool bIncludeFunctionLibrary) const;
-	URigVMController* GetController(int32 InIndex) const;
-	URigVMController* GetController(const FString& InNodePathOrName) const;
-	URigVMController* GetController(const URigVMGraph* InModel) const;
-	URigVMController* GetController(const UObject* InEditorSideObject) const;
-	URigVMController* GetOrCreateController(int32 InIndex);
-	URigVMController* GetOrCreateController(const FString& InNodePathOrName);
-	URigVMController* GetOrCreateController(const URigVMGraph* InModel);
-	URigVMController* GetOrCreateController(const UObject* InEditorSideObject);
-	URigVMController* GetControllerByName(const FString InGraphName) const;
-	bool RemoveController(const URigVMGraph* InModel);
+	UE_API URigVMGraph* GetDefaultModel() const;
+	UE_API URigVMGraph* GetModel(int32 InIndex) const;
+	UE_API URigVMGraph* GetModel(const UEdGraph* InEdGraph = nullptr) const;
+	UE_API URigVMGraph* GetModel(const FString& InNodePathOrName) const;
+	UE_API URigVMGraph* GetModel(const UObject* InEditorSideObject) const;
+	UE_API void RefreshAllModels(ERigVMLoadType InLoadType, bool bEnablePostLoadHashing, bool& bIsCompiling);
+	UE_API const TArray<TObjectPtr<URigVMGraph>>& GetModels() const;
+	UE_API TArray<URigVMGraph*> GetAllModels(bool bIncludeFunctionLibrary, bool bRecursive) const;
+	UE_API TArray<URigVMGraph*> GetAllModelsLeavesFirst(bool bIncludeFunctionLibrary) const;
+	UE_API URigVMController* GetController(int32 InIndex) const;
+	UE_API URigVMController* GetController(const FString& InNodePathOrName) const;
+	UE_API URigVMController* GetController(const URigVMGraph* InModel) const;
+	UE_API URigVMController* GetController(const UObject* InEditorSideObject) const;
+	UE_API URigVMController* GetOrCreateController(int32 InIndex);
+	UE_API URigVMController* GetOrCreateController(const FString& InNodePathOrName);
+	UE_API URigVMController* GetOrCreateController(const URigVMGraph* InModel);
+	UE_API URigVMController* GetOrCreateController(const UObject* InEditorSideObject);
+	UE_API URigVMController* GetControllerByName(const FString InGraphName) const;
+	UE_API bool RemoveController(const URigVMGraph* InModel);
 	URigVMFunctionLibrary* GetFunctionLibrary() const { return FunctionLibrary; }
-	URigVMFunctionLibrary* GetOrCreateFunctionLibrary(bool bSetupUndoRedo, const FObjectInitializer* ObjectInitializer = nullptr, bool bCreateController = true);
-	URigVMFunctionLibrary* GetOrCreateFunctionLibrary(TSubclassOf<URigVMSchema> InSchemaClass, bool bSetupUndoRedo, const FObjectInitializer* ObjectInitializer = nullptr, bool bCreateController = true);
-	TArray<FName> GetEntryNames(UScriptStruct* InUnitScriptStructFilter = nullptr) const;
+	UE_API URigVMFunctionLibrary* GetOrCreateFunctionLibrary(bool bSetupUndoRedo, const FObjectInitializer* ObjectInitializer = nullptr, bool bCreateController = true);
+	UE_API URigVMFunctionLibrary* GetOrCreateFunctionLibrary(TSubclassOf<URigVMSchema> InSchemaClass, bool bSetupUndoRedo, const FObjectInitializer* ObjectInitializer = nullptr, bool bCreateController = true);
+	UE_API TArray<FName> GetEntryNames(UScriptStruct* InUnitScriptStructFilter = nullptr) const;
 	UE_DEPRECATED(5.5, "Please use GetDefaultExecuteContextStruct or get an execute context from a graph/controller schema.")
 	UScriptStruct* GetExecuteContextStruct() const { return GetDefaultExecuteContextStruct(); }
-	UScriptStruct* GetDefaultExecuteContextStruct() const;
+	UE_API UScriptStruct* GetDefaultExecuteContextStruct() const;
 	UE_DEPRECATED(5.5, "Please use SetDefaultExecuteContextStruct or set an execute context on a graph/controller schema.")
 	void SetExecuteContextStruct(UScriptStruct* InExecuteContextStruct) { SetDefaultExecuteContextStruct(InExecuteContextStruct); }
-	void SetDefaultExecuteContextStruct(UScriptStruct* InExecuteContextStruct);
+	UE_API void SetDefaultExecuteContextStruct(UScriptStruct* InExecuteContextStruct);
 
 	FRigVMGetFocusedGraph& OnGetFocusedGraph() { return OnGetFocusedGraphDelegate;}
 	const FRigVMGetFocusedGraph& OnGetFocusedGraph() const { return OnGetFocusedGraphDelegate; }
-	URigVMGraph* GetFocusedModel() const;
-	URigVMGraph* AddModel(const FString InName, bool bSetupUndoRedo, bool bPrintPythonCommand);
-	URigVMGraph* AddModel(const FName& InName, bool bSetupUndoRedo, const FObjectInitializer* ObjectInitializer = nullptr, bool bCreateController = true);
-	URigVMGraph* AddModel(const FName& InName, TSubclassOf<URigVMSchema> InSchemaClass, bool bSetupUndoRedo, const FObjectInitializer* ObjectInitializer = nullptr, bool bCreateController = true);
+	UE_API URigVMGraph* GetFocusedModel() const;
+	UE_API URigVMGraph* AddModel(const FString InName, bool bSetupUndoRedo, bool bPrintPythonCommand);
+	UE_API URigVMGraph* AddModel(const FName& InName, bool bSetupUndoRedo, const FObjectInitializer* ObjectInitializer = nullptr, bool bCreateController = true);
+	UE_API URigVMGraph* AddModel(const FName& InName, TSubclassOf<URigVMSchema> InSchemaClass, bool bSetupUndoRedo, const FObjectInitializer* ObjectInitializer = nullptr, bool bCreateController = true);
 
-	URigVMGraph* CreateModel(const FName& InName, TSubclassOf<URigVMSchema> InSchemaClass, bool bSetupUndoRedo, UObject* InOuter, const FObjectInitializer* ObjectInitializer = nullptr, bool bCreateController = true);
-	TObjectPtr<URigVMGraph> CreateContainedGraphModel(URigVMCollapseNode* CollapseNode, const FName& Name);
+	UE_API URigVMGraph* CreateModel(const FName& InName, TSubclassOf<URigVMSchema> InSchemaClass, bool bSetupUndoRedo, UObject* InOuter, const FObjectInitializer* ObjectInitializer = nullptr, bool bCreateController = true);
+	UE_API TObjectPtr<URigVMGraph> CreateContainedGraphModel(URigVMCollapseNode* CollapseNode, const FName& Name);
 	
-	void AddModel(URigVMGraph* InModel, bool bCreateController);
-	bool RemoveModel(FString InName, bool bSetupUndoRedo, bool bPrintPythonCommand);
-	bool RemoveModel(const FString& InNodePathOrName, bool bSetupUndoRedo);
-	FName RenameModel(const FString& InNodePathOrName, const FName& InNewName, bool bSetupUndoRedo);
-	void PostTransacted(const FTransactionObjectEvent& TransactionEvent);
-	void OnCollapseNodeRenamed(const URigVMCollapseNode* InCollapseNode);
-	void OnCollapseNodeRemoved(const URigVMCollapseNode* InCollapseNode);
-
-	URigVMNode* FindNode(const FString& InNodePathOrName) const;
-	URigVMPin* FindPin(const FString& InPinPath) const;
+	UE_API void AddModel(URigVMGraph* InModel, bool bCreateController);
+	UE_API bool RemoveModel(FString InName, bool bSetupUndoRedo, bool bPrintPythonCommand);
+	UE_API bool RemoveModel(const FString& InNodePathOrName, bool bSetupUndoRedo);
+	UE_API FName RenameModel(const FString& InNodePathOrName, const FName& InNewName, bool bSetupUndoRedo);
+	UE_API void PostTransacted(const FTransactionObjectEvent& TransactionEvent);
+	UE_API void OnCollapseNodeRenamed(const URigVMCollapseNode* InCollapseNode);
+	UE_API void OnCollapseNodeRemoved(const URigVMCollapseNode* InCollapseNode);
+
+	UE_API URigVMNode* FindNode(const FString& InNodePathOrName) const;
+	UE_API URigVMPin* FindPin(const FString& InPinPath) const;
 	
 	TArray<TObjectPtr<URigVMGraph>>::RangedForIteratorType      begin() { return const_cast<TArray<TObjectPtr<URigVMGraph>>&>(GetModels()).begin(); }
 	TArray<TObjectPtr<URigVMGraph>>::RangedForConstIteratorType begin() const { return GetModels().begin(); }
 	TArray<TObjectPtr<URigVMGraph>>::RangedForIteratorType      end() { return const_cast<TArray<TObjectPtr<URigVMGraph>>&>(GetModels()).end(); }
 	TArray<TObjectPtr<URigVMGraph>>::RangedForConstIteratorType end() const { return GetModels().end(); }
 
-	UObject* GetOuter() const;
-	FProperty* GetOuterClientProperty() const;
-	void NotifyOuterOfPropertyChange(EPropertyChangeType::Type ChangeType = EPropertyChangeType::Interactive) const;
-	FName GetUniqueName(const FName& InDesiredName) const;
-	static FName GetUniqueName(UObject* InOuter, const FName& InDesiredName);
-	static void DestroyObject(UObject* InObject);
+	UE_API UObject* GetOuter() const;
+	UE_API FProperty* GetOuterClientProperty() const;
+	UE_API void NotifyOuterOfPropertyChange(EPropertyChangeType::Type ChangeType = EPropertyChangeType::Interactive) const;
+	UE_API FName GetUniqueName(const FName& InDesiredName) const;
+	static UE_API FName GetUniqueName(UObject* InOuter, const FName& InDesiredName);
+	static UE_API void DestroyObject(UObject* InObject);
 
-	uint32 GetStructureHash() const;
-	uint32 GetSerializedStructureHash() const;
+	UE_API uint32 GetStructureHash() const;
+	UE_API uint32 GetSerializedStructureHash() const;
 
 	// backwards compatibility
-	FRigVMClientPatchResult PatchModelsOnLoad();
-	void PatchFunctionReferencesOnLoad();
-	void PatchFunctionsOnLoad(IRigVMGraphFunctionHost* FunctionHost, TArray<FName>& BackwardsCompatiblePublicFunctions, TMap<URigVMLibraryNode*, FRigVMGraphFunctionHeader>& OldHeaders);
-	FRigVMClientPatchResult PatchPinDefaultValues();
+	UE_API FRigVMClientPatchResult PatchModelsOnLoad();
+	UE_API void PatchFunctionReferencesOnLoad();
+	UE_API void PatchFunctionsOnLoad(IRigVMGraphFunctionHost* FunctionHost, TArray<FName>& BackwardsCompatiblePublicFunctions, TMap<URigVMLibraryNode*, FRigVMGraphFunctionHeader>& OldHeaders);
+	UE_API FRigVMClientPatchResult PatchPinDefaultValues();
 
 	// try to reattach detached links and delete remaining ones
-	void ProcessDetachedLinks();
+	UE_API void ProcessDetachedLinks();
 
 	// work to be done before saving
-	void PreSave(FObjectPreSaveContext ObjectSaveContext);
+	UE_API void PreSave(FObjectPreSaveContext ObjectSaveContext);
 
-	void HandleGraphModifiedEvent(ERigVMGraphNotifType InNotifType, URigVMGraph* InGraph, UObject* InSubject);
+	UE_API void HandleGraphModifiedEvent(ERigVMGraphNotifType InNotifType, URigVMGraph* InGraph, UObject* InSubject);
 
-	FRigVMGraphFunctionStore* FindFunctionStore(const URigVMLibraryNode* InLibraryNode);
-	bool UpdateGraphFunctionData(const URigVMLibraryNode* InLibraryNode);
-	bool UpdateExternalVariablesForFunction(const URigVMLibraryNode* InLibraryNode);
-	bool UpdateDependenciesForFunction(const URigVMLibraryNode* InLibraryNode);
-	bool UpdateFunctionReferences(const FRigVMGraphFunctionHeader& InHeader, bool bUpdateDependencies, bool bUpdateExternalVariables);
-	bool DirtyGraphFunctionCompilationData(URigVMLibraryNode* InLibraryNode);
-	bool UpdateGraphFunctionSerializedGraph(URigVMLibraryNode* InLibraryNode);
-	bool IsFunctionPublic(URigVMLibraryNode* InLibraryNode);
+	UE_API FRigVMGraphFunctionStore* FindFunctionStore(const URigVMLibraryNode* InLibraryNode);
+	UE_API bool UpdateGraphFunctionData(const URigVMLibraryNode* InLibraryNode);
+	UE_API bool UpdateExternalVariablesForFunction(const URigVMLibraryNode* InLibraryNode);
+	UE_API bool UpdateDependenciesForFunction(const URigVMLibraryNode* InLibraryNode);
+	UE_API bool UpdateFunctionReferences(const FRigVMGraphFunctionHeader& InHeader, bool bUpdateDependencies, bool bUpdateExternalVariables);
+	UE_API bool DirtyGraphFunctionCompilationData(URigVMLibraryNode* InLibraryNode);
+	UE_API bool UpdateGraphFunctionSerializedGraph(URigVMLibraryNode* InLibraryNode);
+	UE_API bool IsFunctionPublic(URigVMLibraryNode* InLibraryNode);
 
 	static constexpr TCHAR RigVMModelPrefix[] = TEXT("RigVMModel");
 
@@ -346,9 +348,9 @@ private:
 		FString OtherNodePath;
 	};
 
-	URigVMController* CreateController(const URigVMGraph* InModel);
-	URigVMActionStack* GetOrCreateActionStack();
-	void ResetActionStack();
+	UE_API URigVMController* CreateController(const URigVMGraph* InModel);
+	UE_API URigVMActionStack* GetOrCreateActionStack();
+	UE_API void ResetActionStack();
 
 	FRigVMGetFocusedGraph OnGetFocusedGraphDelegate;
 
@@ -390,3 +392,5 @@ private:
 	friend class UEngineTestClientHost;
 	friend class URigVMBlueprint;
 };
+
+#undef UE_API
diff --git a/Engine/Plugins/Runtime/RigVM/Source/RigVMDeveloper/Public/RigVMModel/RigVMControllerActions.h b/Engine/Plugins/Runtime/RigVM/Source/RigVMDeveloper/Public/RigVMModel/RigVMControllerActions.h
--- a/Engine/Plugins/Runtime/RigVM/Source/RigVMDeveloper/Public/RigVMModel/RigVMControllerActions.h
+++ b/Engine/Plugins/Runtime/RigVM/Source/RigVMDeveloper/Public/RigVMModel/RigVMControllerActions.h
@@ -72,15 +72,15 @@ struct FRigVMActionKey
 /**
  * The action wrapper is used to extract an action from a serialized key.
  */
-struct RIGVMDEVELOPER_API FRigVMActionWrapper
+struct FRigVMActionWrapper
 {
 public:
-	FRigVMActionWrapper(const FRigVMActionKey& Key);
-	~FRigVMActionWrapper();
+	RIGVMDEVELOPER_API FRigVMActionWrapper(const FRigVMActionKey& Key);
+	RIGVMDEVELOPER_API ~FRigVMActionWrapper();
 
-	const UScriptStruct* GetScriptStruct() const;
-	FRigVMBaseAction* GetAction(URigVMController* InLastController = nullptr) const;
-	FString ExportText() const;
+	RIGVMDEVELOPER_API const UScriptStruct* GetScriptStruct() const;
+	RIGVMDEVELOPER_API FRigVMBaseAction* GetAction(URigVMController* InLastController = nullptr) const;
+	RIGVMDEVELOPER_API FString ExportText() const;
 
 private:
 	FRigVMActionWrapper(const FRigVMActionWrapper& Other) = delete;
@@ -112,7 +112,7 @@ struct FRigVMActionNodeContent
  * base implementations.
  */
 USTRUCT()
-struct RIGVMDEVELOPER_API FRigVMBaseAction
+struct FRigVMBaseAction
 {
 	GENERATED_BODY()
 
@@ -140,7 +140,7 @@ public:
 	virtual ~FRigVMBaseAction() {};
 
 	// Returns the controller of this action
-	URigVMController* GetController() const;
+	RIGVMDEVELOPER_API URigVMController* GetController() const;
 
 	// Access to the actions script struct. this needs to be overloaded
 	virtual UScriptStruct* GetScriptStruct() const { return FRigVMBaseAction::StaticStruct(); }
@@ -156,19 +156,19 @@ public:
 
 	// Trys to merge the action with another action and 
 	// returns true if successfull.
-	virtual bool Merge(const FRigVMBaseAction* Other);
+	RIGVMDEVELOPER_API virtual bool Merge(const FRigVMBaseAction* Other);
 
 	// Returns true if this action makes the other action obsolete.
-	virtual bool MakesObsolete(const FRigVMBaseAction* Other) const;
+	RIGVMDEVELOPER_API virtual bool MakesObsolete(const FRigVMBaseAction* Other) const;
 
 	// Returns true if this action can undo / redo
 	bool CanUndoRedo() const { return GetController() != nullptr; }
 
 	// Un-does the action and returns true if successfull.
-	virtual bool Undo();
+	RIGVMDEVELOPER_API virtual bool Undo();
 
 	// Re-does the action and returns true if successfull.
-	virtual bool Redo();
+	RIGVMDEVELOPER_API virtual bool Redo();
 
 	// Adds a child / sub action to this one
 	template<class ActionType>
@@ -180,13 +180,13 @@ public:
 		SubActions.Add(Key);
 	}
 
-	bool StoreNode(const URigVMNode* InNode, bool bIsPriorChange = true);
-	bool RestoreNode(const FName& InNodeName, bool bIsUndoing);
+	RIGVMDEVELOPER_API bool StoreNode(const URigVMNode* InNode, bool bIsPriorChange = true);
+	RIGVMDEVELOPER_API bool RestoreNode(const FName& InNodeName, bool bIsUndoing);
 
 #if RIGVM_ACTIONSTACK_VERBOSE_LOG
 	
 	// Logs the action to the console / output log
-	void LogAction(const FString& InPrefix) const;
+	RIGVMDEVELOPER_API void LogAction(const FString& InPrefix) const;
 
 #endif
 
@@ -199,7 +199,7 @@ protected:
 	{
 	}
 
-	void EnsureControllerValidity() const;
+	RIGVMDEVELOPER_API void EnsureControllerValidity() const;
 
 	mutable TWeakObjectPtr<URigVMController> WeakController;
 
@@ -228,8 +228,8 @@ protected:
  * to open / close a bracket. Open brackets automatically record additional
  * actions occuring during the bracket's lifetime.
  */
-UCLASS()
-class RIGVMDEVELOPER_API URigVMActionStack : public UObject
+UCLASS(MinimalAPI)
+class URigVMActionStack : public UObject
 {
 	GENERATED_BODY()
 
@@ -240,7 +240,7 @@ class RIGVMDEVELOPER_API URigVMActionStack : public UObject
 
 public:
 
-	static URigVMActionStack* GetDisabledActionStack();
+	static RIGVMDEVELOPER_API URigVMActionStack* GetDisabledActionStack();
 
 	// Begins an action and opens a bracket / scope.
 	template<class ActionType>
@@ -393,32 +393,32 @@ public:
 	// Opens an undo bracket / scope to record actions into.
 	// This is primary useful for Python.
 	UFUNCTION()
-	bool OpenUndoBracket(URigVMController* InController, const FString& InTitle);
+	RIGVMDEVELOPER_API bool OpenUndoBracket(URigVMController* InController, const FString& InTitle);
 
 	// Closes an undo bracket / scope.
 	// This is primary useful for Python.
 	UFUNCTION()
-	bool CloseUndoBracket(URigVMController* InController);
+	RIGVMDEVELOPER_API bool CloseUndoBracket(URigVMController* InController);
 
 	// Cancels an undo bracket / scope.
 	// This is primary useful for Python.
 	UFUNCTION()
-	bool CancelUndoBracket(URigVMController* InController);
+	RIGVMDEVELOPER_API bool CancelUndoBracket(URigVMController* InController);
 
 	// Pops the last action from the undo stack and perform undo on it.
 	// Note: This should really only be used for unit tests,
 	// use the GEditor's main Undo method instead.
 	UFUNCTION()
-	bool Undo(URigVMController* InController);
+	RIGVMDEVELOPER_API bool Undo(URigVMController* InController);
 
 	// Pops the last action from the redo stack and perform redo on it.
 	// Note: This should really only be used for unit tests,
 	// use the GEditor's main Redo method instead.
 	UFUNCTION()
-	bool Redo(URigVMController* InController);
+	RIGVMDEVELOPER_API bool Redo(URigVMController* InController);
 
 #if WITH_EDITOR
-	virtual void PostTransacted(const FTransactionObjectEvent& TransactionEvent) override;
+	RIGVMDEVELOPER_API virtual void PostTransacted(const FTransactionObjectEvent& TransactionEvent) override;
 #endif
 
 	FRigVMGraphModifiedEvent& OnModified() { return ModifiedEvent; }
@@ -432,10 +432,10 @@ private:
 		LogAction(ActionType::StaticStruct(), InAction, InPrefix);
 	}
 	
-	void LogAction(const UScriptStruct* InActionStruct, const FRigVMBaseAction& InAction, const FString& InPrefix);
+	RIGVMDEVELOPER_API void LogAction(const UScriptStruct* InActionStruct, const FRigVMBaseAction& InAction, const FString& InPrefix);
 #endif
 
-	bool IsDisabled() const;
+	RIGVMDEVELOPER_API bool IsDisabled() const;
 	
 	UPROPERTY()
 	int32 ActionIndex;
@@ -1762,4 +1762,4 @@ public:
 	virtual UScriptStruct* GetScriptStruct() const override { return FRigVMRemoveTraitAction::StaticStruct(); }
 	virtual bool Undo() override;
 	virtual bool Redo() override;
-};
\ No newline at end of file
+};
diff --git a/Engine/Plugins/Runtime/RigVM/Source/RigVMDeveloper/Public/RigVMModel/RigVMExternalDependency.h b/Engine/Plugins/Runtime/RigVM/Source/RigVMDeveloper/Public/RigVMModel/RigVMExternalDependency.h
--- a/Engine/Plugins/Runtime/RigVM/Source/RigVMDeveloper/Public/RigVMModel/RigVMExternalDependency.h
+++ b/Engine/Plugins/Runtime/RigVM/Source/RigVMDeveloper/Public/RigVMModel/RigVMExternalDependency.h
@@ -6,8 +6,10 @@
 #include "RigVMCore/RigVMGraphFunctionHost.h"
 #include "RigVMExternalDependency.generated.h"
 
+#define UE_API RIGVMDEVELOPER_API
+
 USTRUCT()
-struct RIGVMDEVELOPER_API FRigVMExternalDependency
+struct FRigVMExternalDependency
 {
 	GENERATED_BODY()
 
@@ -51,22 +53,22 @@ private:
 	FName Category;
 };
 
-UINTERFACE()
-class RIGVMDEVELOPER_API URigVMExternalDependencyManager : public UInterface
+UINTERFACE(MinimalAPI)
+class URigVMExternalDependencyManager : public UInterface
 {
 	GENERATED_BODY()
 };
 
 // Interface to deal with mapping external dependencies
-class RIGVMDEVELOPER_API IRigVMExternalDependencyManager
+class IRigVMExternalDependencyManager
 {
 	GENERATED_BODY()
 
 public:
 	
-	virtual const TArray<FName>& GetExternalDependencyCategories() const;
+	UE_API virtual const TArray<FName>& GetExternalDependencyCategories() const;
 	virtual TArray<FRigVMExternalDependency> GetExternalDependenciesForCategory(const FName& InCategory) const = 0;
-	TArray<FRigVMExternalDependency> GetAllExternalDependencies() const;
+	UE_API TArray<FRigVMExternalDependency> GetAllExternalDependencies() const;
 
 	static inline const FLazyName UserDefinedEnumCategory = FLazyName(TEXT("UserDefinedEnum"));
 	static inline const FLazyName UserDefinedStructCategory = FLazyName(TEXT("UserDefinedStruct"));
@@ -74,18 +76,19 @@ public:
 
 protected:
 
-	virtual void CollectExternalDependencies(TArray<FRigVMExternalDependency>& OutDependencies, const FName& InCategory, const FRigVMClient* InClient) const;
-	virtual void CollectExternalDependencies(TArray<FRigVMExternalDependency>& OutDependencies, const FName& InCategory, const FRigVMGraphFunctionStore* InFunctionStore) const;
-	virtual void CollectExternalDependencies(TArray<FRigVMExternalDependency>& OutDependencies, const FName& InCategory, const FRigVMGraphFunctionData* InFunction) const;
-	virtual void CollectExternalDependencies(TArray<FRigVMExternalDependency>& OutDependencies, const FName& InCategory, const FRigVMGraphFunctionHeader* InHeader) const;
-	virtual void CollectExternalDependencies(TArray<FRigVMExternalDependency>& OutDependencies, const FName& InCategory, const FRigVMFunctionCompilationData* InCompilationData) const;
-	virtual void CollectExternalDependencies(TArray<FRigVMExternalDependency>& OutDependencies, const FName& InCategory, const URigVMGraph* InGraph) const;
-	virtual void CollectExternalDependencies(TArray<FRigVMExternalDependency>& OutDependencies, const FName& InCategory, const URigVMNode* InNode) const;
-	virtual void CollectExternalDependencies(TArray<FRigVMExternalDependency>& OutDependencies, const FName& InCategory, const URigVMPin* InPin) const;
-	virtual void CollectExternalDependencies(TArray<FRigVMExternalDependency>& OutDependencies, const FName& InCategory, const UStruct* InStruct) const;
-	virtual void CollectExternalDependencies(TArray<FRigVMExternalDependency>& OutDependencies, const FName& InCategory, const UEnum* InEnum) const;
-	virtual void CollectExternalDependenciesForCPPTypeObject(TArray<FRigVMExternalDependency>& OutDependencies, const FName& InCategory, const UObject* InObject) const;
-
-	static TArray<FName> DependencyCategories;
+	UE_API virtual void CollectExternalDependencies(TArray<FRigVMExternalDependency>& OutDependencies, const FName& InCategory, const FRigVMClient* InClient) const;
+	UE_API virtual void CollectExternalDependencies(TArray<FRigVMExternalDependency>& OutDependencies, const FName& InCategory, const FRigVMGraphFunctionStore* InFunctionStore) const;
+	UE_API virtual void CollectExternalDependencies(TArray<FRigVMExternalDependency>& OutDependencies, const FName& InCategory, const FRigVMGraphFunctionData* InFunction) const;
+	UE_API virtual void CollectExternalDependencies(TArray<FRigVMExternalDependency>& OutDependencies, const FName& InCategory, const FRigVMGraphFunctionHeader* InHeader) const;
+	UE_API virtual void CollectExternalDependencies(TArray<FRigVMExternalDependency>& OutDependencies, const FName& InCategory, const FRigVMFunctionCompilationData* InCompilationData) const;
+	UE_API virtual void CollectExternalDependencies(TArray<FRigVMExternalDependency>& OutDependencies, const FName& InCategory, const URigVMGraph* InGraph) const;
+	UE_API virtual void CollectExternalDependencies(TArray<FRigVMExternalDependency>& OutDependencies, const FName& InCategory, const URigVMNode* InNode) const;
+	UE_API virtual void CollectExternalDependencies(TArray<FRigVMExternalDependency>& OutDependencies, const FName& InCategory, const URigVMPin* InPin) const;
+	UE_API virtual void CollectExternalDependencies(TArray<FRigVMExternalDependency>& OutDependencies, const FName& InCategory, const UStruct* InStruct) const;
+	UE_API virtual void CollectExternalDependencies(TArray<FRigVMExternalDependency>& OutDependencies, const FName& InCategory, const UEnum* InEnum) const;
+	UE_API virtual void CollectExternalDependenciesForCPPTypeObject(TArray<FRigVMExternalDependency>& OutDependencies, const FName& InCategory, const UObject* InObject) const;
+
+	static UE_API TArray<FName> DependencyCategories;
 };
 
+#undef UE_API
diff --git a/Engine/Plugins/Runtime/RigVM/Source/RigVMDeveloper/Public/RigVMModel/RigVMFunctionLibrary.h b/Engine/Plugins/Runtime/RigVM/Source/RigVMDeveloper/Public/RigVMModel/RigVMFunctionLibrary.h
--- a/Engine/Plugins/Runtime/RigVM/Source/RigVMDeveloper/Public/RigVMModel/RigVMFunctionLibrary.h
+++ b/Engine/Plugins/Runtime/RigVM/Source/RigVMDeveloper/Public/RigVMModel/RigVMFunctionLibrary.h
@@ -9,46 +9,48 @@
 #include "RigVMModel/Nodes/RigVMLibraryNode.h"
 #include "RigVMFunctionLibrary.generated.h"
 
+#define UE_API RIGVMDEVELOPER_API
+
 DECLARE_DELEGATE_RetVal(const FSoftObjectPath, URigVMFunctionLibrary_GetFunctionHostObjectPath);
 
 /**
  * The Function Library is a graph used only to store
  * the sub graphs used for functions.
  */
-UCLASS(BlueprintType)
-class RIGVMDEVELOPER_API URigVMFunctionLibrary : public URigVMGraph
+UCLASS(MinimalAPI, BlueprintType)
+class URigVMFunctionLibrary : public URigVMGraph
 {
 	GENERATED_BODY()
 
 public:
 
 	// Default constructor
-	URigVMFunctionLibrary();
+	UE_API URigVMFunctionLibrary();
 
 	// URigVMGraph interface
-	virtual FString GetNodePath() const override;
-	virtual URigVMFunctionLibrary* GetDefaultFunctionLibrary() const override;
+	UE_API virtual FString GetNodePath() const override;
+	UE_API virtual URigVMFunctionLibrary* GetDefaultFunctionLibrary() const override;
 	// end URigVMGraph interface
 
 	// Returns all of the stored functions
 	UFUNCTION(BlueprintCallable, Category = RigVMGraph)
-	TArray<URigVMLibraryNode*> GetFunctions() const;
+	UE_API TArray<URigVMLibraryNode*> GetFunctions() const;
 
 	// Finds a function by name
 	UFUNCTION(BlueprintCallable, Category = RigVMGraph)
-	URigVMLibraryNode* FindFunction(const FName& InFunctionName) const;
+	UE_API URigVMLibraryNode* FindFunction(const FName& InFunctionName) const;
 
 	// Finds a function by a node within a function (or a sub graph of that)
 	UFUNCTION(BlueprintCallable, Category = RigVMGraph)
-    URigVMLibraryNode* FindFunctionForNode(URigVMNode* InNode) const;
+    UE_API URigVMLibraryNode* FindFunctionForNode(URigVMNode* InNode) const;
 
 	// Returns all references for a given function name
 	UFUNCTION(BlueprintCallable, Category = RigVMGraph)
-	TArray< TSoftObjectPtr<URigVMFunctionReferenceNode> > GetReferencesForFunction(const FName& InFunctionName);
+	UE_API TArray< TSoftObjectPtr<URigVMFunctionReferenceNode> > GetReferencesForFunction(const FName& InFunctionName);
 
 	// Returns all references for a given function name
 	UFUNCTION(BlueprintCallable, Category = RigVMGraph)
-	TArray< FString > GetReferencePathsForFunction(const FName& InFunctionName);
+	UE_API TArray< FString > GetReferencePathsForFunction(const FName& InFunctionName);
 
 	/**
 	* Iterator function to invoke a lambda / TFunction for each reference of a function
@@ -56,27 +58,27 @@ public:
 	* @param PerReferenceFunction The function to invoke for each reference
 	* @param bLoadIfNecessary If true, will load packages when needed
 	*/
-	void ForEachReference(const FName& InFunctionName, TFunction<void(URigVMFunctionReferenceNode*)> PerReferenceFunction, bool bLoadIfNecessary = true) const;
+	UE_API void ForEachReference(const FName& InFunctionName, TFunction<void(URigVMFunctionReferenceNode*)> PerReferenceFunction, bool bLoadIfNecessary = true) const;
 
 	/**
 	* Iterator function to invoke a lambda / TFunction for each reference of a function
 	* @param InFunctionName The function name to iterate all references for
 	* @param PerReferenceFunction The function to invoke for each reference
 	*/
-	void ForEachReferenceSoftPtr(const FName& InFunctionName, TFunction<void(TSoftObjectPtr<URigVMFunctionReferenceNode>)> PerReferenceFunction) const;
+	UE_API void ForEachReferenceSoftPtr(const FName& InFunctionName, TFunction<void(TSoftObjectPtr<URigVMFunctionReferenceNode>)> PerReferenceFunction) const;
 
 	// Returns a function that has been previously localized based on the provided function to localize.
 	// We maintain meta data on what functions have been created locally based on which other ones,
 	// and use this method to avoid redundant localizations.
-	URigVMLibraryNode* FindPreviouslyLocalizedFunction(FRigVMGraphFunctionIdentifier InFunctionToLocalize);
+	UE_API URigVMLibraryNode* FindPreviouslyLocalizedFunction(FRigVMGraphFunctionIdentifier InFunctionToLocalize);
 
-	const FSoftObjectPath GetFunctionHostObjectPath() const;
+	UE_API const FSoftObjectPath GetFunctionHostObjectPath() const;
 	URigVMFunctionLibrary_GetFunctionHostObjectPath GetFunctionHostObjectPathDelegate;
 
 	bool IsFunctionPublic(const FName& InFunctionName) const { return PublicFunctionNames.Contains(InFunctionName); }
 
-	const FRigVMVariant* GetFunctionVariant(const FName& InFunctionName) const;
-	FRigVMVariant* GetFunctionVariant(const FName& InFunctionName);
+	UE_API const FRigVMVariant* GetFunctionVariant(const FName& InFunctionName) const;
+	UE_API FRigVMVariant* GetFunctionVariant(const FName& InFunctionName);
 
 private:
 
@@ -104,3 +106,4 @@ private:
 	friend class UControlRigBlueprint;
 };
 
+#undef UE_API
diff --git a/Engine/Plugins/Runtime/RigVM/Source/RigVMDeveloper/Public/RigVMModel/RigVMGraph.h b/Engine/Plugins/Runtime/RigVM/Source/RigVMDeveloper/Public/RigVMModel/RigVMGraph.h
--- a/Engine/Plugins/Runtime/RigVM/Source/RigVMDeveloper/Public/RigVMModel/RigVMGraph.h
+++ b/Engine/Plugins/Runtime/RigVM/Source/RigVMDeveloper/Public/RigVMModel/RigVMGraph.h
@@ -13,6 +13,8 @@
 #include "UObject/Interface.h"
 #include "RigVMGraph.generated.h"
 
+#define UE_API RIGVMDEVELOPER_API
+
 class URigVMFunctionLibrary;
 struct FRigVMClient;
 
@@ -24,129 +26,129 @@ struct FRigVMClient;
  * Graphs provide access to its Nodes, Pins and
  * Links.
  */
-UCLASS(BlueprintType)
-class RIGVMDEVELOPER_API URigVMGraph : public UObject
+UCLASS(MinimalAPI, BlueprintType)
+class URigVMGraph : public UObject
 {
 	GENERATED_BODY()
 
 public:
 
 	// Default constructor
-	URigVMGraph();
+	UE_API URigVMGraph();
 
-	virtual void PostLoad() override;
+	UE_API virtual void PostLoad() override;
 
 	// Returns all of the Nodes within this Graph.
 	UFUNCTION(BlueprintCallable, Category = RigVMGraph)
-	const TArray<URigVMNode*>& GetNodes() const;
+	UE_API const TArray<URigVMNode*>& GetNodes() const;
 
 	// Returns all of the Links within this Graph.
 	UFUNCTION(BlueprintCallable, Category = RigVMGraph)
-	const TArray<URigVMLink*>& GetLinks() const;
+	UE_API const TArray<URigVMLink*>& GetLinks() const;
 
 	// Returns true if the graph contains a link given its string representation 
 	UFUNCTION(BlueprintCallable, Category = RigVMGraph)
-	bool ContainsLink(const FString& InPinPathRepresentation) const;
+	UE_API bool ContainsLink(const FString& InPinPathRepresentation) const;
 
 	// Returns all of the contained graphs
 	UFUNCTION(BlueprintCallable, Category = RigVMGraph)
-	TArray<URigVMGraph*> GetContainedGraphs(bool bRecursive = false) const;
+	UE_API TArray<URigVMGraph*> GetContainedGraphs(bool bRecursive = false) const;
 
 	// Returns the parent graph of this graph
 	UFUNCTION(BlueprintCallable, Category = RigVMGraph)
-    URigVMGraph* GetParentGraph() const;
+    UE_API URigVMGraph* GetParentGraph() const;
 
 	// Returns the root / top level parent graph of this graph (or this if it is the root graph)
 	UFUNCTION(BlueprintCallable, Category = RigVMGraph)
-    URigVMGraph* GetRootGraph() const;
+    UE_API URigVMGraph* GetRootGraph() const;
 
 	// Returns the root / top level parent graph of this graph (or this if it is the root graph)
 	UFUNCTION(BlueprintPure, Category = RigVMGraph)
-	int32 GetGraphDepth() const;
+	UE_API int32 GetGraphDepth() const;
 
 	// Returns true if this graph is a root / top level graph
 	UFUNCTION(BlueprintCallable, Category = RigVMGraph)
-    bool IsRootGraph() const;
+    UE_API bool IsRootGraph() const;
 
 	// Returns the entry node of this graph
 	UFUNCTION(BlueprintCallable, Category = RigVMGraph)
-	URigVMFunctionEntryNode* GetEntryNode() const;
+	UE_API URigVMFunctionEntryNode* GetEntryNode() const;
 	
 	// Returns the return node of this graph
 	UFUNCTION(BlueprintCallable, Category = RigVMGraph)
-	URigVMFunctionReturnNode* GetReturnNode() const;
+	UE_API URigVMFunctionReturnNode* GetReturnNode() const;
 
 	// Returns array of event names
 	UFUNCTION(BlueprintCallable, Category = RigVMGraph)
-	TArray<FName> GetEventNames() const;
+	UE_API TArray<FName> GetEventNames() const;
 
 	// Returns a list of unique Variable descriptions within this Graph.
 	// Multiple Variable Nodes can share the same description.
 	UFUNCTION(BlueprintCallable, Category = RigVMGraph)
-	TArray<FRigVMGraphVariableDescription> GetVariableDescriptions() const;
+	UE_API TArray<FRigVMGraphVariableDescription> GetVariableDescriptions() const;
 
 	// Returns the path of this graph as defined by its invoking nodes
 	UFUNCTION(BlueprintCallable, Category = RigVMGraph)
-	virtual FString GetNodePath() const;
+	UE_API virtual FString GetNodePath() const;
 
 	// Returns the name of this graph (as defined by the node path)
 	UFUNCTION(BlueprintCallable, Category = RigVMGraph)
-	FString GetGraphName() const;
+	UE_API FString GetGraphName() const;
 
 	// Returns a Node given its name (or nullptr).
 	UFUNCTION(BlueprintCallable, Category = RigVMGraph)
-	URigVMNode* FindNodeByName(const FName& InNodeName) const;
+	UE_API URigVMNode* FindNodeByName(const FName& InNodeName) const;
 
 	// Returns a Node given its path (or nullptr).
 	// (for now this is the same as finding a node by its name.)
 	UFUNCTION(BlueprintCallable, Category = RigVMGraph)
-	URigVMNode* FindNode(const FString& InNodePath) const;
+	UE_API URigVMNode* FindNode(const FString& InNodePath) const;
 
 	// Returns a Pin given its path, for example "Node.Color.R".
 	UFUNCTION(BlueprintCallable, Category = RigVMGraph)
-	URigVMPin* FindPin(const FString& InPinPath) const;
+	UE_API URigVMPin* FindPin(const FString& InPinPath) const;
 
 	// Returns a link given its string representation,
 	// for example "NodeA.Color.R -> NodeB.Translation.X"
 	UFUNCTION(BlueprintCallable, Category = RigVMGraph)
-	URigVMLink* FindLink(const FString& InLinkPinPathRepresentation) const;
+	UE_API URigVMLink* FindLink(const FString& InLinkPinPathRepresentation) const;
 
 	// Returns true if a Node with a given name is selected.
 	UFUNCTION(BlueprintCallable, Category = RigVMGraph)
-	bool IsNodeSelected(const FName& InNodeName) const;
+	UE_API bool IsNodeSelected(const FName& InNodeName) const;
 
 	// Returns the names of all currently selected Nodes.
 	UFUNCTION(BlueprintCallable, Category = RigVMGraph)
-	const TArray<FName>& GetSelectNodes() const;
+	UE_API const TArray<FName>& GetSelectNodes() const;
 
 	// Returns true if this graph is the top level graph
 	UFUNCTION(BlueprintCallable, Category = RigVMGraph)
-	bool IsTopLevelGraph() const;
+	UE_API bool IsTopLevelGraph() const;
 
 	// Returns the locally available function library
 	UFUNCTION(BlueprintCallable, Category = RigVMGraph)
-	virtual URigVMFunctionLibrary* GetDefaultFunctionLibrary() const;
+	UE_API virtual URigVMFunctionLibrary* GetDefaultFunctionLibrary() const;
 
 	UFUNCTION(BlueprintCallable, Category = RigVMGraph)
-	void SetDefaultFunctionLibrary(URigVMFunctionLibrary* InFunctionLibrary);
+	UE_API void SetDefaultFunctionLibrary(URigVMFunctionLibrary* InFunctionLibrary);
 
-	TArray<FRigVMExternalVariable> GetExternalVariables() const;
+	UE_API TArray<FRigVMExternalVariable> GetExternalVariables() const;
 
 	// Returns the local variables of this function
 	UFUNCTION(BlueprintCallable, Category = RigVMGraph)
-	TArray<FRigVMGraphVariableDescription> GetLocalVariables(bool bIncludeInputArguments = false) const;
+	UE_API TArray<FRigVMGraphVariableDescription> GetLocalVariables(bool bIncludeInputArguments = false) const;
 
 	// Returns the input arguments of this graph
 	UFUNCTION(BlueprintCallable, Category = RigVMGraph)
-	TArray<FRigVMGraphVariableDescription> GetInputArguments() const;
+	UE_API TArray<FRigVMGraphVariableDescription> GetInputArguments() const;
 	
 	// Returns the output arguments of this graph
     UFUNCTION(BlueprintCallable, Category = RigVMGraph)
-    TArray<FRigVMGraphVariableDescription> GetOutputArguments() const;
+    UE_API TArray<FRigVMGraphVariableDescription> GetOutputArguments() const;
 
 	// Returns the schema used by this graph
 	UFUNCTION(BlueprintCallable, Category = RigVMGraph)
-	URigVMSchema* GetSchema() const;
+	UE_API URigVMSchema* GetSchema() const;
 
 	// Returns the schema class used by this graph
 	UFUNCTION(BlueprintCallable, Category = RigVMGraph)
@@ -154,33 +156,33 @@ public:
 
 	// Sets the schema class on the graph
 	UFUNCTION(BlueprintCallable, Category = RigVMGraph)
-	void SetSchemaClass(TSubclassOf<URigVMSchema> InSchemaClass);
+	UE_API void SetSchemaClass(TSubclassOf<URigVMSchema> InSchemaClass);
 
 	// Returns the modified event, which can be used to 
 	// subscribe to changes happening within the Graph.
-	FRigVMGraphModifiedEvent& OnModified();
+	UE_API FRigVMGraphModifiedEvent& OnModified();
 
 	// Sets the execute context struct type to use
-	void SetExecuteContextStruct(UScriptStruct* InExecuteContextStruct);
+	UE_API void SetExecuteContextStruct(UScriptStruct* InExecuteContextStruct);
 	
-	UScriptStruct* GetExecuteContextStruct() const;
+	UE_API UScriptStruct* GetExecuteContextStruct() const;
 
-	void PrepareCycleChecking(URigVMPin* InPin, bool bAsInput);
-	virtual bool CanLink(URigVMPin* InSourcePin, URigVMPin* InTargetPin, FString* OutFailureReason, const FRigVMByteCode* InByteCode, ERigVMPinDirection InUserLinkDirection = ERigVMPinDirection::IO, bool bEnableTypeCasting = true);
+	UE_API void PrepareCycleChecking(URigVMPin* InPin, bool bAsInput);
+	UE_API virtual bool CanLink(URigVMPin* InSourcePin, URigVMPin* InTargetPin, FString* OutFailureReason, const FRigVMByteCode* InByteCode, ERigVMPinDirection InUserLinkDirection = ERigVMPinDirection::IO, bool bEnableTypeCasting = true);
 	
-	TSharedPtr<FRigVMParserAST> GetDiagnosticsAST(bool bForceRefresh = false, TArray<URigVMLink*> InLinksToSkip = TArray<URigVMLink*>());
-	TSharedPtr<FRigVMParserAST> GetRuntimeAST(const FRigVMParserASTSettings& InSettings = FRigVMParserASTSettings::Optimized(), bool bForceRefresh = false);
-	void ClearAST(bool bClearDiagnostics = true, bool bClearRuntime = true);
+	UE_API TSharedPtr<FRigVMParserAST> GetDiagnosticsAST(bool bForceRefresh = false, TArray<URigVMLink*> InLinksToSkip = TArray<URigVMLink*>());
+	UE_API TSharedPtr<FRigVMParserAST> GetRuntimeAST(const FRigVMParserASTSettings& InSettings = FRigVMParserASTSettings::Optimized(), bool bForceRefresh = false);
+	UE_API void ClearAST(bool bClearDiagnostics = true, bool bClearRuntime = true);
 
-	virtual uint32 GetStructureHash() const;
+	UE_API virtual uint32 GetStructureHash() const;
 	uint32 GetSerializedStructureHash() const { return LastStructureHash; }
 
-	virtual void PreSave(FObjectPreSaveContext SaveContext) override;
+	UE_API virtual void PreSave(FObjectPreSaveContext SaveContext) override;
 
 private:
 
 	FRigVMGraphModifiedEvent ModifiedEvent;
-	void Notify(ERigVMGraphNotifType InNotifType, UObject* InSubject);
+	UE_API void Notify(ERigVMGraphNotifType InNotifType, UObject* InSubject);
 
 	UPROPERTY()
 	TArray<TObjectPtr<URigVMNode>> Nodes;
@@ -220,7 +222,7 @@ private:
 	UPROPERTY()
 	TSubclassOf<URigVMSchema> SchemaClass;
 
-	bool IsNameAvailable(const FString& InName) const;
+	UE_API bool IsNameAvailable(const FString& InName) const;
 
 	friend class URigVMController;
 	friend class URigVMBlueprint;
@@ -231,3 +233,4 @@ private:
 	friend struct FRigVMControllerObjectFactory;
 };
 
+#undef UE_API
diff --git a/Engine/Plugins/Runtime/RigVM/Source/RigVMDeveloper/Public/RigVMModel/RigVMLink.h b/Engine/Plugins/Runtime/RigVM/Source/RigVMDeveloper/Public/RigVMModel/RigVMLink.h
--- a/Engine/Plugins/Runtime/RigVM/Source/RigVMDeveloper/Public/RigVMModel/RigVMLink.h
+++ b/Engine/Plugins/Runtime/RigVM/Source/RigVMDeveloper/Public/RigVMModel/RigVMLink.h
@@ -7,6 +7,8 @@
 #include "RigVMPin.h"
 #include "RigVMLink.generated.h"
 
+#define UE_API RIGVMDEVELOPER_API
+
 class URigVMGraph;
 
 /**
@@ -15,8 +17,8 @@ class URigVMGraph;
  * Graph itself - or through the URigVMPin::GetLinks()
  * method.
  */
-UCLASS(BlueprintType)
-class RIGVMDEVELOPER_API URigVMLink : public UObject
+UCLASS(MinimalAPI, BlueprintType)
+class URigVMLink : public UObject
 {
 	GENERATED_BODY()
 
@@ -26,82 +28,82 @@ public:
 	URigVMLink() = default;
 
 	// Serialization override
-	virtual void Serialize(FArchive& Ar) override;
+	UE_API virtual void Serialize(FArchive& Ar) override;
 
 	// Returns the current index of this Link within its owning Graph.
 	UFUNCTION(BlueprintCallable, Category = RigVMLink)
-	int32 GetLinkIndex() const;
+	UE_API int32 GetLinkIndex() const;
 
 	// Returns the Link's owning Graph/
 	UFUNCTION(BlueprintCallable, Category = RigVMLink)
-	URigVMGraph* GetGraph() const;
+	UE_API URigVMGraph* GetGraph() const;
 
 	// Returns the graph nesting depth of this link
 	UFUNCTION(BlueprintCallable, Category = RigVMLink)
-	int32 GetGraphDepth() const;
+	UE_API int32 GetGraphDepth() const;
 
 	// Returns the source Pin of this Link (or nullptr)
 	UFUNCTION(BlueprintCallable, Category = RigVMLink)
-	URigVMPin* GetSourcePin() const;
+	UE_API URigVMPin* GetSourcePin() const;
 
 	// Returns the target Pin of this Link (or nullptr)
 	UFUNCTION(BlueprintCallable, Category = RigVMLink)
-	URigVMPin* GetTargetPin() const;
+	UE_API URigVMPin* GetTargetPin() const;
 
 	// Returns the source Node of this Link (or nullptr)
 	UFUNCTION(BlueprintCallable, Category = RigVMLink)
-	URigVMNode* GetSourceNode() const;
+	UE_API URigVMNode* GetSourceNode() const;
 
 	// Returns the target Node of this Link (or nullptr)
 	UFUNCTION(BlueprintCallable, Category = RigVMLink)
-	URigVMNode* GetTargetNode() const;
+	UE_API URigVMNode* GetTargetNode() const;
 
 	// Returns the source pin's path pin of this Link
-	FString GetSourcePinPath() const;
+	UE_API FString GetSourcePinPath() const;
 
 	// Returns the target pin's path pin of this Link
-	FString GetTargetPinPath() const;
+	UE_API FString GetTargetPinPath() const;
 
 	// Sets the source pin's path pin of this Link
-	bool SetSourcePinPath(const FString& InPinPath);
+	UE_API bool SetSourcePinPath(const FString& InPinPath);
 
 	// Sets the target pin's path pin of this Link
-	bool SetTargetPinPath(const FString& InPinPath);
+	UE_API bool SetTargetPinPath(const FString& InPinPath);
 
 	// Sets the target pin's path pin of this Link
-	bool SetSourceAndTargetPinPaths(const FString& InSourcePinPath, const FString& InTargetPinPath);
+	UE_API bool SetSourceAndTargetPinPaths(const FString& InSourcePinPath, const FString& InTargetPinPath);
 
 	// Returns the opposite Pin of this Link given one of its edges (or nullptr)
 	UFUNCTION(BlueprintCallable, Category = RigVMLink)
-	URigVMPin* GetOppositePin(const URigVMPin* InPin) const;
+	UE_API URigVMPin* GetOppositePin(const URigVMPin* InPin) const;
 
 	// Returns a string representation of the Link,
 	// for example: "NodeA.Color.R -> NodeB.Translation.X"
 	// note: can be split again using SplitPinPathRepresentation
 	UFUNCTION(BlueprintCallable, Category = RigVMLink)
-	FString GetPinPathRepresentation() const;
+	UE_API FString GetPinPathRepresentation() const;
 
 	// Returns a string representation of the Link given the two pin paths
 	// for example: "NodeA.Color.R -> NodeB.Translation.X"
 	// note: can be split again using SplitPinPathRepresentation
-	static FString GetPinPathRepresentation(const FString& InSourcePinPath, const FString& InTargetPinPath);
+	static UE_API FString GetPinPathRepresentation(const FString& InSourcePinPath, const FString& InTargetPinPath);
 
 	// Splits a pin path representation of a link
 	// for example: "NodeA.Color.R -> NodeB.Translation.X"
 	// into its two pin paths
-	static bool SplitPinPathRepresentation(const FString& InString, FString& OutSource, FString& OutTarget);
+	static UE_API bool SplitPinPathRepresentation(const FString& InString, FString& OutSource, FString& OutTarget);
 
-	void UpdatePinPaths();
-	void UpdatePinPointers() const;
-	bool Detach();
+	UE_API void UpdatePinPaths();
+	UE_API void UpdatePinPointers() const;
+	UE_API bool Detach();
 	
 private:
 
 	// Returns true if the link is attached.
 	// Attached links rely on the pin pointers first and the pin path second.
 	// Deattached links never rely on the pin pointers and always try to resolve from string.
-	bool IsAttached() const;
-	bool Attach(FString* OutFailureReason = nullptr);
+	UE_API bool IsAttached() const;
+	UE_API bool Attach(FString* OutFailureReason = nullptr);
 
 	UPROPERTY()
 	FString SourcePinPath;
@@ -113,3 +115,4 @@ private:
 	mutable URigVMPin* TargetPin = nullptr;
 };
 
+#undef UE_API
