Updated headers using UnrealCodeFixup to make sure dllstorage is on methods/staticvars instead of types. Part 4/n

diff --git a/Engine/Plugins/Runtime/RigVM/Source/RigVMDeveloper/Public/RigVMModel/RigVMNode.h b/Engine/Plugins/Runtime/RigVM/Source/RigVMDeveloper/Public/RigVMModel/RigVMNode.h
--- a/Engine/Plugins/Runtime/RigVM/Source/RigVMDeveloper/Public/RigVMModel/RigVMNode.h
+++ b/Engine/Plugins/Runtime/RigVM/Source/RigVMDeveloper/Public/RigVMModel/RigVMNode.h
@@ -11,6 +11,8 @@
 #include "UObject/StructOnScope.h"
 #include "RigVMNode.generated.h"
 
+#define UE_API RIGVMDEVELOPER_API
+
 class URigVMGraph;
 
 namespace ERigVMNodeDefaultValueOverrideState
@@ -32,115 +34,115 @@ namespace ERigVMNodeDefaultValueOverrideState
  * Nodes contain Pins to represent parameters for Function
  * Invocations or Value access on Variables / Parameters.
  */
-UCLASS(BlueprintType)
-class RIGVMDEVELOPER_API URigVMNode : public UObject
+UCLASS(MinimalAPI, BlueprintType)
+class URigVMNode : public UObject
 {
 	GENERATED_BODY()
 
 public:
 
 	// Default constructor
-	URigVMNode();
+	UE_API URigVMNode();
 
 	// Default destructor
-	virtual ~URigVMNode();
+	UE_API virtual ~URigVMNode();
 
-	void Serialize(FArchive& Ar) override;
+	UE_API void Serialize(FArchive& Ar) override;
 
 	// Returns the a . separated string containing all of the
 	// names used to reach this Node within the Graph.
 	// (for now this is the same as the Node's name)
 	UFUNCTION(BlueprintCallable, Category = RigVMNode)
-	FString GetNodePath(bool bRecursive = false) const;
+	UE_API FString GetNodePath(bool bRecursive = false) const;
 
 	// Splits a NodePath at the start, so for example "CollapseNodeA|CollapseNodeB|CollapseNodeC" becomes "CollapseNodeA" and "CollapseNodeB|CollapseNodeC"
-	static bool SplitNodePathAtStart(const FString& InNodePath, FString& LeftMost, FString& Right);
+	static UE_API bool SplitNodePathAtStart(const FString& InNodePath, FString& LeftMost, FString& Right);
 
 	// Splits a NodePath at the end, so for example "CollapseNodeA|CollapseNodeB|CollapseNodeC" becomes "CollapseNodeA|CollapseNodeB" and "CollapseNodeC"
-	static bool SplitNodePathAtEnd(const FString& InNodePath, FString& Left, FString& RightMost);
+	static UE_API bool SplitNodePathAtEnd(const FString& InNodePath, FString& Left, FString& RightMost);
 
 	// Splits a NodePath into all segments, so for example "Node.Color.R" becomes ["Node", "Color", "R"]
-	static bool SplitNodePath(const FString& InNodePath, TArray<FString>& Parts);
+	static UE_API bool SplitNodePath(const FString& InNodePath, TArray<FString>& Parts);
 
 	// Joins a NodePath from to segments, so for example "CollapseNodeA" and "CollapseNodeB|CollapseNodeC" becomes "CollapseNodeA|CollapseNodeB|CollapseNodeC"
-	static FString JoinNodePath(const FString& Left, const FString& Right);
+	static UE_API FString JoinNodePath(const FString& Left, const FString& Right);
 
 	// Joins a NodePath from to segments, so for example ["CollapseNodeA", "CollapseNodeB", "CollapseNodeC"] becomes "CollapseNodeA|CollapseNodeB|CollapseNodeC"
-	static FString JoinNodePath(const TArray<FString>& InParts);
+	static UE_API FString JoinNodePath(const TArray<FString>& InParts);
 
 	// Returns the current index of the Node within the Graph.
 	UFUNCTION(BlueprintCallable, Category = RigVMNode)
-	int32 GetNodeIndex() const;
+	UE_API int32 GetNodeIndex() const;
 
 	// Returns all of the top-level Pins of this Node.
 	UFUNCTION(BlueprintCallable, Category = RigVMNode)
-	virtual const TArray<URigVMPin*>& GetPins() const;
+	UE_API virtual const TArray<URigVMPin*>& GetPins() const;
 
 	// Returns all of the Pins of this Node (including SubPins).
 	UFUNCTION(BlueprintCallable, Category = RigVMNode)
-	TArray<URigVMPin*> GetAllPinsRecursively() const;
+	UE_API TArray<URigVMPin*> GetAllPinsRecursively() const;
 
 	// Returns all user defined categories on this node
 	UFUNCTION(BlueprintCallable, Category = RigVMNode)
-	virtual TArray<FString> GetPinCategories() const;
+	UE_API virtual TArray<FString> GetPinCategories() const;
 
 	// Returns the name of pin category
 	UFUNCTION(BlueprintCallable, Category = RigVMNode)
-	FString GetPinCategoryName(const FString InCategory) const;
+	UE_API FString GetPinCategoryName(const FString InCategory) const;
 
 	// Returns all sub user defined categories of a given parent category
 	UFUNCTION(BlueprintCallable, Category = RigVMNode)
-	TArray<FString> GetSubPinCategories(const FString InCategory, bool bOnlyExisting = false, bool bRecursive = false) const;
+	UE_API TArray<FString> GetSubPinCategories(const FString InCategory, bool bOnlyExisting = false, bool bRecursive = false) const;
 
 	// Returns the parent pin category of the given category (or an empty string in case there's no parent)
 	UFUNCTION(BlueprintCallable, Category = RigVMNode)
-	FString GetParentPinCategory(const FString InCategory, bool bOnlyExisting = false) const;
+	UE_API FString GetParentPinCategory(const FString InCategory, bool bOnlyExisting = false) const;
 
 	// Returns all parent categories of a given 
 	UFUNCTION(BlueprintCallable, Category = RigVMNode)
-	TArray<FString> GetParentPinCategories(const FString InCategory, bool bOnlyExisting = false, bool bIncludeSelf = false) const;
+	UE_API TArray<FString> GetParentPinCategories(const FString InCategory, bool bOnlyExisting = false, bool bIncludeSelf = false) const;
 
 	// Returns the depth of the category (starting with 0 for "Foo" and 2 for "Foo|Bar|Waldo")
-	static int32 GetPinCategoryDepth(const FString& InCategory);
+	static UE_API int32 GetPinCategoryDepth(const FString& InCategory);
 
 	// Returns all pins for a given category
 	UFUNCTION(BlueprintCallable, Category = RigVMNode)
-	TArray<URigVMPin*> GetPinsForCategory(FString InCategory) const;
+	UE_API TArray<URigVMPin*> GetPinsForCategory(FString InCategory) const;
 
 	// Returns all pins for a given category
 	UFUNCTION(BlueprintCallable, Category = RigVMNode)
-	bool IsPinCategoryExpanded(FString InCategory) const;
+	UE_API bool IsPinCategoryExpanded(FString InCategory) const;
 
 	const FString& GetLastAffectedPinCategory() const { return LastAffectedPinCategory; }
 
 	// Returns the pin UI layout for this node
 	UFUNCTION(BlueprintCallable, Category = RigVMNode)
-	virtual FRigVMNodeLayout GetNodeLayout(bool bIncludeEmptyCategories = false) const;
+	UE_API virtual FRigVMNodeLayout GetNodeLayout(bool bIncludeEmptyCategories = false) const;
 
 	// Returns the default value for a given pin
-	FString GetOriginalPinDefaultValue(const URigVMPin* InPin) const;
+	UE_API FString GetOriginalPinDefaultValue(const URigVMPin* InPin) const;
 
 	// Returns 0 if there's no override on this node,
 	// 1 if some pins have an override and
 	// 2 if all pins are overriden
-	ERigVMNodeDefaultValueOverrideState::Type GetPinDefaultValueOverrideState() const;
+	UE_API ERigVMNodeDefaultValueOverrideState::Type GetPinDefaultValueOverrideState() const;
 
 	// Returns a Pin given it's partial pin path below
 	// this node (for example: "Color.R")
 	UFUNCTION(BlueprintCallable, Category = RigVMNode)
-	URigVMPin* FindPin(const FString& InPinPath) const;
+	UE_API URigVMPin* FindPin(const FString& InPinPath) const;
 
 	// Returns a root pin given its name
 	UFUNCTION(BlueprintCallable, Category = RigVMNode)
-	URigVMPin* FindRootPinByName(const FName& InPinName) const;
+	UE_API URigVMPin* FindRootPinByName(const FName& InPinName) const;
 
 	// Returns the first execute pin
 	UFUNCTION(BlueprintCallable, Category = RigVMNode)
-	URigVMPin* FindExecutePin() const;
+	UE_API URigVMPin* FindExecutePin() const;
 
 	// Returns all of the top-level orphaned Pins of this Node.
 	UFUNCTION(BlueprintCallable, Category = RigVMNode)
-    virtual const TArray<URigVMPin*>& GetOrphanedPins() const;
+    UE_API virtual const TArray<URigVMPin*>& GetOrphanedPins() const;
 
 	// Returns true if the node has orphaned pins - which leads to a compiler error
 	UFUNCTION(BlueprintPure, Category = RigVMNode)
@@ -148,61 +150,61 @@ public:
 
 	// Returns the Graph of this Node
 	UFUNCTION(BlueprintCallable, Category = RigVMNode)
-	URigVMGraph* GetGraph() const;
+	UE_API URigVMGraph* GetGraph() const;
 
 	// Returns the top level / root Graph of this Node
 	UFUNCTION(BlueprintCallable, Category = RigVMNode)
-    URigVMGraph* GetRootGraph() const;
+    UE_API URigVMGraph* GetRootGraph() const;
 
 	// Returns the graph nesting depth of this node
 	UFUNCTION(BlueprintCallable, Category = RigVMNode)
-	int32 GetGraphDepth() const;
+	UE_API int32 GetGraphDepth() const;
 
 	// Returns the injection info of this Node (or nullptr)
 	UFUNCTION(BlueprintCallable, Category = RigVMNode)
-	URigVMInjectionInfo* GetInjectionInfo() const;
+	UE_API URigVMInjectionInfo* GetInjectionInfo() const;
 
 	// Returns the title of this Node - used for UI.
 	UFUNCTION(BlueprintCallable, Category = RigVMNode)
-	virtual FString GetNodeTitle() const;
+	UE_API virtual FString GetNodeTitle() const;
 
 	// Returns currently node title set in the node, without any transformations
 	// Warning : The raw title might be ignored by some nodes overrides to GetNodeTitle
-	const FString& GetNodeTitleRaw() const;
+	UE_API const FString& GetNodeTitleRaw() const;
 
 	// Returns the 2d position of this node - used for UI.
 	UFUNCTION(BlueprintCallable, Category = RigVMNode)
-	FVector2D GetPosition() const;
+	UE_API FVector2D GetPosition() const;
 
 	// Returns the 2d size of this node - used for UI.
 	UFUNCTION(BlueprintCallable, Category = RigVMNode)
-	FVector2D GetSize() const;
+	UE_API FVector2D GetSize() const;
 
 	// Returns the color of this node - used for UI.
 	UFUNCTION(BlueprintCallable, Category = RigVMNode)
-	virtual FLinearColor GetNodeColor() const;
+	UE_API virtual FLinearColor GetNodeColor() const;
 
 	// Returns the tooltip of this node
 	UFUNCTION(BlueprintCallable, Category = RigVMNode)
-	virtual FText GetToolTipText() const;
+	UE_API virtual FText GetToolTipText() const;
 
 	// Returns true if this Node is currently selected.
 	UFUNCTION(BlueprintCallable, Category = RigVMNode)
-	bool IsSelected() const;
+	UE_API bool IsSelected() const;
 
 	// Returns true if this is an injected node.
 	// Injected nodes are managed by pins are are not visible to the user.
 	UFUNCTION(BlueprintCallable, Category = RigVMNode)
-	bool IsInjected() const;
+	UE_API bool IsInjected() const;
 
 	// Returns true if this should be visible in the UI
 	UFUNCTION(BlueprintCallable, Category = RigVMNode)
-	bool IsVisibleInUI() const;
+	UE_API bool IsVisibleInUI() const;
 
 	// Returns true if this Node has no side-effects
 	// and no internal state.
 	UFUNCTION(BlueprintCallable, Category = RigVMNode)
-	virtual bool IsPure() const;
+	UE_API virtual bool IsPure() const;
 
 	// Returns true if the node is defined as non-varying
 	UFUNCTION(BlueprintCallable, Category = RigVMNode)
@@ -215,63 +217,63 @@ public:
 	// Returns true if this Node has side effects or
 	// internal state.
 	UFUNCTION(BlueprintCallable, Category = RigVMNode)
-	virtual bool IsMutable() const;
+	UE_API virtual bool IsMutable() const;
 
 	// Returns true if this node has an unknown type pin
-	bool HasWildCardPin() const;
+	UE_API bool HasWildCardPin() const;
 
 	virtual bool ContributesToResult() const { return IsMutable(); }
 
 	// Returns true if this Node is the beginning of a scope
 	UFUNCTION(BlueprintCallable, Category = RigVMNode)
-	virtual bool IsEvent() const;
+	UE_API virtual bool IsEvent() const;
 
 	// Returns the name of the event
 	UFUNCTION(BlueprintCallable, Category = RigVMNode)
-	virtual FName GetEventName() const;
+	UE_API virtual FName GetEventName() const;
 
 	// Returns true if this node can only exist once in a graph
 	UFUNCTION(BlueprintCallable, Category = RigVMNode)
-	virtual bool CanOnlyExistOnce() const;
+	UE_API virtual bool CanOnlyExistOnce() const;
 
 	// Returns true if the node has any input pins
 	UFUNCTION(BlueprintCallable, Category = RigVMNode)
-	virtual bool HasInputPin(bool bIncludeIO = true) const;
+	UE_API virtual bool HasInputPin(bool bIncludeIO = true) const;
 
 	// Returns true if the node has any io pins
 	UFUNCTION(BlueprintCallable, Category = RigVMNode)
-	virtual bool HasIOPin() const;
+	UE_API virtual bool HasIOPin() const;
 
 	// Returns true if the node has any lazily evaluating pins
 	UFUNCTION(BlueprintCallable, Category = RigVMNode)
-	virtual bool HasLazyPin(bool bOnlyConsiderPinsWithLinks = false) const;
+	UE_API virtual bool HasLazyPin(bool bOnlyConsiderPinsWithLinks = false) const;
 
 	// Returns true if the node has any output pins
 	UFUNCTION(BlueprintCallable, Category = RigVMNode)
-	virtual bool HasOutputPin(bool bIncludeIO = true) const;
+	UE_API virtual bool HasOutputPin(bool bIncludeIO = true) const;
 
 	// Returns true if the node has any pins of the provided direction
 	UFUNCTION(BlueprintCallable, Category = RigVMNode)
-	virtual bool HasPinOfDirection(ERigVMPinDirection InDirection) const;
+	UE_API virtual bool HasPinOfDirection(ERigVMPinDirection InDirection) const;
 
 	// Returns true if this Node is linked to another 
 	// given node through any of the Nodes' Pins.
 	UFUNCTION(BlueprintCallable, Category = RigVMNode)
-	bool IsLinkedTo(URigVMNode* InNode) const;
+	UE_API bool IsLinkedTo(URigVMNode* InNode) const;
 
 	// Returns all links to any pin on this node
 	UFUNCTION(BlueprintCallable, Category = RigVMNode)
-	TArray<URigVMLink*> GetLinks() const;
+	UE_API TArray<URigVMLink*> GetLinks() const;
 
 	// Returns a list of Nodes connected as sources to
 	// this Node as the target.
 	UFUNCTION(BlueprintCallable, Category = RigVMNode)
-	TArray<URigVMNode*> GetLinkedSourceNodes() const;
+	UE_API TArray<URigVMNode*> GetLinkedSourceNodes() const;
 
 	// Returns a list of Nodes connected as targets to
 	// this Node as the source.
 	UFUNCTION(BlueprintCallable, Category = RigVMNode)
-	TArray<URigVMNode*> GetLinkedTargetNodes() const;
+	UE_API TArray<URigVMNode*> GetLinkedTargetNodes() const;
 
 	// Returns the name of the node prior to the renaming
 	UFUNCTION(BlueprintCallable, Category = RigVMNode)
@@ -281,36 +283,36 @@ public:
 	const TArray<int32>& GetInstructionsForVM(URigVM* InVM, const FRigVMASTProxy& InProxy = FRigVMASTProxy()) const	{ static TArray<int32> Dummy; return Dummy;	}
 
 	// Returns the indices of associated instructions for this node
-	const TArray<int32>& GetInstructionsForVM(const FRigVMExtendedExecuteContext& Context, URigVM* InVM, const FRigVMASTProxy& InProxy = FRigVMASTProxy()) const;
+	UE_API const TArray<int32>& GetInstructionsForVM(const FRigVMExtendedExecuteContext& Context, URigVM* InVM, const FRigVMASTProxy& InProxy = FRigVMASTProxy()) const;
 
 	UE_DEPRECATED(5.3, "Please, use GetInstructionVisitedCount with Context param")
 	virtual int32 GetInstructionVisitedCount(URigVM* InVM, const FRigVMASTProxy& InProxy = FRigVMASTProxy()) const { return 0; }
 
 	// Returns the number of visited / run instructions for this node
-	virtual int32 GetInstructionVisitedCount(const FRigVMExtendedExecuteContext& Context, URigVM* InVM, const FRigVMASTProxy& InProxy = FRigVMASTProxy()) const;
+	UE_API virtual int32 GetInstructionVisitedCount(const FRigVMExtendedExecuteContext& Context, URigVM* InVM, const FRigVMASTProxy& InProxy = FRigVMASTProxy()) const;
 
 	UE_DEPRECATED(5.3, "Please, use GetInstructionMicroSeconds with Context param")
 	double GetInstructionMicroSeconds(URigVM* InVM, const FRigVMASTProxy& InProxy = FRigVMASTProxy()) const { return 0.0; }
 
 	// Returns the accumulated duration of all of instructions for this node 
-	double GetInstructionMicroSeconds(const FRigVMExtendedExecuteContext& Context, URigVM* InVM, const FRigVMASTProxy& InProxy = FRigVMASTProxy()) const;
+	UE_API double GetInstructionMicroSeconds(const FRigVMExtendedExecuteContext& Context, URigVM* InVM, const FRigVMASTProxy& InProxy = FRigVMASTProxy()) const;
 
 	// return true if this node is a loop node
 	UFUNCTION(BlueprintPure, Category = RigVMNode)
-	bool IsLoopNode() const;
+	UE_API bool IsLoopNode() const;
 
 	// return true if this node is a control flow node
 	UFUNCTION(BlueprintPure, Category = RigVMNode)
-	bool IsControlFlowNode() const;
+	UE_API bool IsControlFlowNode() const;
 
 	// returns the names of the control flow blocks of this node
-	virtual const TArray<FName>& GetControlFlowBlocks() const;
+	UE_API virtual const TArray<FName>& GetControlFlowBlocks() const;
 
 	// returns true if a control flow block requires slicing
-	virtual const bool IsControlFlowBlockSliced(const FName& InBlockName) const;
+	UE_API virtual const bool IsControlFlowBlockSliced(const FName& InBlockName) const;
 
 	// returns true if the node is within a loop
-	virtual bool IsWithinLoop() const;
+	UE_API virtual bool IsWithinLoop() const;
 
 	// returns true if the node can be upgraded
 	UFUNCTION(BlueprintPure, Category = RigVMNode)
@@ -318,7 +320,7 @@ public:
 
 	// returns all supported workflows of the node
 	UFUNCTION(BlueprintCallable, Category = RigVMNode)
-	virtual TArray<FRigVMUserWorkflow> GetSupportedWorkflows(ERigVMUserWorkflowType InType, const UObject* InSubject) const;
+	UE_API virtual TArray<FRigVMUserWorkflow> GetSupportedWorkflows(ERigVMUserWorkflowType InType, const UObject* InSubject) const;
 
 	UFUNCTION(BlueprintCallable, Category = RigVMNode)
 	bool HasBreakpoint() const { return bHasBreakpoint; }
@@ -333,19 +335,19 @@ public:
 	void SetExecutionIsHaltedAtThisNode(const bool bValue) { bHaltedAtThisNode = bValue; }
 
 	UFUNCTION(BlueprintPure, Category = RigVMNode)
-	virtual bool IsAggregate() const;
+	UE_API virtual bool IsAggregate() const;
 
 	UFUNCTION(BlueprintPure, Category = RigVMNode)
-	virtual URigVMPin* GetFirstAggregatePin() const;
+	UE_API virtual URigVMPin* GetFirstAggregatePin() const;
 
 	UFUNCTION(BlueprintPure, Category = RigVMNode)
-	virtual URigVMPin* GetSecondAggregatePin() const;
+	UE_API virtual URigVMPin* GetSecondAggregatePin() const;
 
 	UFUNCTION(BlueprintPure, Category = RigVMNode)
-	virtual URigVMPin* GetOppositeAggregatePin() const;
+	UE_API virtual URigVMPin* GetOppositeAggregatePin() const;
 
 	UFUNCTION(BlueprintPure, Category = RigVMNode)
-	virtual bool IsInputAggregate() const;
+	UE_API virtual bool IsInputAggregate() const;
 
 	UFUNCTION(BlueprintPure, Category = RigVMNode)
 	virtual TArray<URigVMPin*> GetAggregateInputs() const { return {}; }
@@ -357,36 +359,36 @@ public:
 	virtual FName GetNextAggregateName(const FName& InLastAggregatePinName) const { return NAME_None; }
 
 	UFUNCTION(BlueprintCallable, Category = RigVMNode)
-	virtual URigVMLibraryNode* FindFunctionForNode() const;
+	UE_API virtual URigVMLibraryNode* FindFunctionForNode() const;
 
 	virtual FRigVMStructUpgradeInfo GetUpgradeInfo() const { return FRigVMStructUpgradeInfo(); }
 
-	virtual uint32 GetStructureHash() const;
+	UE_API virtual uint32 GetStructureHash() const;
 
 	// allows the node to support non-native pins
 	virtual bool HasNonNativePins() const { return !TraitRootPinNames.IsEmpty(); }
 
 	UFUNCTION(BlueprintPure, Category = RigVMNode)
-	TArray<URigVMPin*> GetTraitPins() const;
+	UE_API TArray<URigVMPin*> GetTraitPins() const;
 
 	const TArray<FString>& GetTraitNames() const { return TraitRootPinNames; }
 
 	UFUNCTION(BlueprintPure, Category = RigVMNode)
-	bool IsTraitPin(FName InName) const;
+	UE_API bool IsTraitPin(FName InName) const;
 
-	bool IsTraitPin(const URigVMPin* InTraitPin) const;
+	UE_API bool IsTraitPin(const URigVMPin* InTraitPin) const;
 
-	URigVMPin* FindTrait(const FName& InName, const FString& InSubPinPath = FString()) const;
+	UE_API URigVMPin* FindTrait(const FName& InName, const FString& InSubPinPath = FString()) const;
 
-	URigVMPin* FindTrait(const URigVMPin* InTraitPin) const;
+	UE_API URigVMPin* FindTrait(const URigVMPin* InTraitPin) const;
 
-	TSharedPtr<FStructOnScope> GetTraitInstance(const FName& InName, bool bUseDefaultValueFromPin = true) const;
+	UE_API TSharedPtr<FStructOnScope> GetTraitInstance(const FName& InName, bool bUseDefaultValueFromPin = true) const;
 
-	TSharedPtr<FStructOnScope> GetTraitInstance(const URigVMPin* InTraitPin, bool bUseDefaultValueFromPin = true) const;
+	UE_API TSharedPtr<FStructOnScope> GetTraitInstance(const URigVMPin* InTraitPin, bool bUseDefaultValueFromPin = true) const;
 
-	UScriptStruct* GetTraitScriptStruct(const FName& InName) const;
+	UE_API UScriptStruct* GetTraitScriptStruct(const FName& InName) const;
 
-	UScriptStruct* GetTraitScriptStruct(const URigVMPin* InTraitPin) const;
+	UE_API UScriptStruct* GetTraitScriptStruct(const URigVMPin* InTraitPin) const;
 	
 	virtual bool IsOutDated() const
 	{
@@ -399,41 +401,41 @@ public:
 	}
 
 	// returns the display name for a pin
-	FName GetDisplayNameForPin(const FString& InPinPath) const;
-	virtual FName GetDisplayNameForPin(const URigVMPin* InPin) const;
+	UE_API FName GetDisplayNameForPin(const FString& InPinPath) const;
+	UE_API virtual FName GetDisplayNameForPin(const URigVMPin* InPin) const;
 
 	// returns the display name for a struct member
-	static FName GetDisplayNameForStructMember(const URigVMPin* InPin);
-	static FName GetDisplayNameForStructMember(const UStruct* InStruct, const FString& InPath);
+	static UE_API FName GetDisplayNameForStructMember(const URigVMPin* InPin);
+	static UE_API FName GetDisplayNameForStructMember(const UStruct* InStruct, const FString& InPath);
 
 	// returns the display name within a property
-	static FName GetDisplayNameForProperty(const FProperty* InProperty, const FString& InRemainingPath);
+	static UE_API FName GetDisplayNameForProperty(const FProperty* InProperty, const FString& InRemainingPath);
 
 	// returns the category for a pin
-	virtual FString GetCategoryForPin(const FString& InPinPath) const;
+	UE_API virtual FString GetCategoryForPin(const FString& InPinPath) const;
 
 	// returns the index of a pin within a category
-	virtual int32 GetIndexInCategoryForPin(const FString& InPinPath) const;
+	UE_API virtual int32 GetIndexInCategoryForPin(const FString& InPinPath) const;
 
-	virtual FText GetTypedToolTipText(const URigVMPin* InPin, const FText& ToolTipBody) const;
+	UE_API virtual FText GetTypedToolTipText(const URigVMPin* InPin, const FText& ToolTipBody) const;
 
 private:
 
 	static const inline TCHAR* NodeColorName = TEXT("NodeColor");
 
-	bool IsLinkedToRecursive(URigVMPin* InPin, URigVMNode* InNode) const;
-	void GetLinkedNodesRecursive(URigVMPin* InPin, bool bLookForSources, TArray<URigVMNode*>& OutNodes) const;
+	UE_API bool IsLinkedToRecursive(URigVMPin* InPin, URigVMNode* InNode) const;
+	UE_API void GetLinkedNodesRecursive(URigVMPin* InPin, bool bLookForSources, TArray<URigVMNode*>& OutNodes) const;
 
 protected:
 
-	virtual void InvalidateCache();
-	virtual TArray<int32> GetInstructionsForVMImpl(const FRigVMExtendedExecuteContext& Context, URigVM* InVM, const FRigVMASTProxy& InProxy = FRigVMASTProxy()) const; 
-	virtual FText GetToolTipTextForPin(const URigVMPin* InPin) const;
+	UE_API virtual void InvalidateCache();
+	UE_API virtual TArray<int32> GetInstructionsForVMImpl(const FRigVMExtendedExecuteContext& Context, URigVM* InVM, const FRigVMASTProxy& InProxy = FRigVMASTProxy()) const; 
+	UE_API virtual FText GetToolTipTextForPin(const URigVMPin* InPin) const;
 	virtual bool AllowsLinksOn(const URigVMPin* InPin) const { return true; }
 	virtual bool ShouldInputPinComputeLazily(const URigVMPin* InPin) const { return false; }
-	virtual FString GetOriginalDefaultValueForRootPin(const URigVMPin* InRootPin) const;
+	UE_API virtual FString GetOriginalDefaultValueForRootPin(const URigVMPin* InRootPin) const;
 
-	void UpdateTraitRootPinNames();
+	UE_API void UpdateTraitRootPinNames();
 
 public:
 	const uint32& GetCachedValueVersion() const
@@ -442,7 +444,7 @@ public:
 	}
 
 protected:
-	void IncrementVersion();
+	UE_API void IncrementVersion();
 
 	UPROPERTY()
 	FString NodeTitle;
@@ -502,10 +504,10 @@ private:
 		mutable double MicroSeconds;
 		mutable TArray<int32> Instructions;
 	};
-	const FProfilingCache* UpdateProfilingCacheIfNeeded(const FRigVMExtendedExecuteContext& Context, URigVM* InVM, const FRigVMASTProxy& InProxy) const;
+	UE_API const FProfilingCache* UpdateProfilingCacheIfNeeded(const FRigVMExtendedExecuteContext& Context, URigVM* InVM, const FRigVMASTProxy& InProxy) const;
 	mutable uint32 ProfilingHash;
 	mutable TMap<uint32, TSharedPtr<FProfilingCache>> ProfilingCache;
-	static TArray<int32> EmptyInstructionArray;
+	static UE_API TArray<int32> EmptyInstructionArray;
 #endif
 
 	mutable TMap<FString, FString> CachedOriginalPinDefaultValues;
@@ -517,3 +519,5 @@ private:
 	friend class FRigVMLexer;
 	friend class URigVMSchema;
 };
+
+#undef UE_API
diff --git a/Engine/Plugins/Runtime/RigVM/Source/RigVMDeveloper/Public/RigVMModel/RigVMPin.h b/Engine/Plugins/Runtime/RigVM/Source/RigVMDeveloper/Public/RigVMModel/RigVMPin.h
--- a/Engine/Plugins/Runtime/RigVM/Source/RigVMDeveloper/Public/RigVMModel/RigVMPin.h
+++ b/Engine/Plugins/Runtime/RigVM/Source/RigVMDeveloper/Public/RigVMModel/RigVMPin.h
@@ -14,6 +14,8 @@
 #include "RigVMModelCachedValue.h"
 #include "RigVMPin.generated.h"
 
+#define UE_API RIGVMDEVELOPER_API
+
 class URigVMGraph;
 class URigVMNode;
 class URigVMUnitNode;
@@ -28,8 +30,8 @@ extern RIGVMDEVELOPER_API TAutoConsoleVariable<bool> CVarRigVMEnablePinOverrides
  * Injected nodes are not visible to the user, but they are normal
  * nodes on the graph.
  */
-UCLASS(BlueprintType)
-class RIGVMDEVELOPER_API URigVMInjectionInfo: public UObject
+UCLASS(MinimalAPI, BlueprintType)
+class URigVMInjectionInfo: public UObject
 {
 	GENERATED_BODY()
 
@@ -59,11 +61,11 @@ public:
 
 	// Returns the graph of this injected node.
 	UFUNCTION(BlueprintCallable, Category = RigVMInjectionInfo)
-	URigVMGraph* GetGraph() const;
+	UE_API URigVMGraph* GetGraph() const;
 
 	// Returns the pin of this injected node.
 	UFUNCTION(BlueprintCallable, Category = RigVMInjectionInfo)
-	URigVMPin* GetPin() const;
+	UE_API URigVMPin* GetPin() const;
 
 	struct FWeakInfo
 	{
@@ -73,7 +75,7 @@ public:
 		FName OutputPinName;
 	};
 
-	FWeakInfo GetWeakInfo() const;
+	UE_API FWeakInfo GetWeakInfo() const;
 };
 
 UENUM()
@@ -96,8 +98,8 @@ enum class ERigVMPinDefaultValueType : uint8
  * SubPins, and so on.
  * A URigVMPin is owned solely by a URigVMNode.
  */
-UCLASS(BlueprintType)
-class RIGVMDEVELOPER_API URigVMPin : public UObject
+UCLASS(MinimalAPI, BlueprintType)
+class URigVMPin : public UObject
 {
 	GENERATED_BODY()
 
@@ -130,54 +132,54 @@ public:
 	// A map used to override pin default values
 	typedef TMap<FRigVMASTProxy, FPinOverrideValue> FPinOverrideMap;
 	typedef TPair<FRigVMASTProxy, const FPinOverrideMap&> FPinOverride;
-	static const URigVMPin::FPinOverrideMap EmptyPinOverrideMap;
-	static const FPinOverride EmptyPinOverride;
+	static UE_API const URigVMPin::FPinOverrideMap EmptyPinOverrideMap;
+	static UE_API const FPinOverride EmptyPinOverride;
 
 	// Splits a PinPath at the start, so for example "Node.Color.R" becomes "Node" and "Color.R"
-	static bool SplitPinPathAtStart(const FString& InPinPath, FString& LeftMost, FString& Right);
+	static UE_API bool SplitPinPathAtStart(const FString& InPinPath, FString& LeftMost, FString& Right);
 
 	// Splits a PinPath at the start, so for example "Node.Color.R" becomes "Node.Color" and "R"
-	static bool SplitPinPathAtEnd(const FString& InPinPath, FString& Left, FString& RightMost);
+	static UE_API bool SplitPinPathAtEnd(const FString& InPinPath, FString& Left, FString& RightMost);
 
 	// Splits a PinPath into all segments, so for example "Node.Color.R" becomes ["Node", "Color", "R"]
-	static bool SplitPinPath(const FString& InPinPath, TArray<FString>& Parts);
+	static UE_API bool SplitPinPath(const FString& InPinPath, TArray<FString>& Parts);
 
 	// Joins a PinPath from to segments, so for example "Node.Color" and "R" becomes "Node.Color.R"
-	static FString JoinPinPath(const FString& Left, const FString& Right);
+	static UE_API FString JoinPinPath(const FString& Left, const FString& Right);
 
 	// Joins a PinPath from to segments, so for example ["Node", "Color", "R"] becomes "Node.Color.R"
-	static FString JoinPinPath(const TArray<FString>& InParts);
+	static UE_API FString JoinPinPath(const TArray<FString>& InParts);
 
 	// Splits the default value into name-value pairs
-	static TArray<FString> SplitDefaultValue(const FString& InDefaultValue);
+	static UE_API TArray<FString> SplitDefaultValue(const FString& InDefaultValue);
 	// Joins a collection of element DefaultValues into a default value for an array of those elements
-	static FString GetDefaultValueForArray(TConstArrayView<FString> InDefaultValues);
+	static UE_API FString GetDefaultValueForArray(TConstArrayView<FString> InDefaultValues);
 
 	// Default constructor
-	URigVMPin();
+	UE_API URigVMPin();
 
 	// returns true if the name of this pin matches a given name
-	bool NameEquals(const FString& InName, bool bFollowCoreRedirectors = false) const;
+	UE_API bool NameEquals(const FString& InName, bool bFollowCoreRedirectors = false) const;
 
 	// Returns a . separated path containing all names of the pin and its owners,
 	// this includes the node name, for example "Node.Color.R"
 	UFUNCTION(BlueprintCallable, Category = RigVMPin)
-	FString GetPinPath(bool bUseNodePath = false) const;
+	UE_API FString GetPinPath(bool bUseNodePath = false) const;
 
 	// Returns a . separated path containing all names of the pin and its owners
 	// until we hit the provided parent pin.
 	UFUNCTION(BlueprintCallable, Category = RigVMPin)
-	FString GetSubPinPath(const URigVMPin* InParentPin, bool bIncludeParentPinName = false) const;
+	UE_API FString GetSubPinPath(const URigVMPin* InParentPin, bool bIncludeParentPinName = false) const;
 
 	// Returns the category on a pin. The category is UI relevant only and used
 	// to order pins in the user interface of the node as well as on the details panel.
 	UFUNCTION(BlueprintCallable, Category = RigVMPin)
-	FString GetCategory() const;
+	UE_API FString GetCategory() const;
 
 	// Returns index within a category on a pin. The category is UI relevant only and used
 	// to order pins in the user interface of the node as well as on the details panel.
 	UFUNCTION(BlueprintCallable, Category = RigVMPin)
-	int32 GetIndexInCategory() const;
+	UE_API int32 GetIndexInCategory() const;
 
 	// Returns a . separated path containing all names of the pin within its main
 	// memory owner / storage. This is typically used to create an offset pointer
@@ -186,64 +188,64 @@ public:
 	// corresponding SegmentPath is "Translation.X", since the transform is the
 	// storage / memory.
 	UFUNCTION(BlueprintCallable, Category = RigVMPin)
-	FString GetSegmentPath(bool bIncludeRootPin = false) const;
+	UE_API FString GetSegmentPath(bool bIncludeRootPin = false) const;
 
 	// Populates an array of pins which will be reduced to the same operand in the
 	// VM. This includes Source-Target pins in different nodes, pins in collapse and
 	// referenced function nodes, and their corresponding entry and return nodes.
-	void GetExposedPinChain(TArray<const URigVMPin*>& OutExposedPins) const;
+	UE_API void GetExposedPinChain(TArray<const URigVMPin*>& OutExposedPins) const;
 
 	// Returns the display label of the pin
 	UFUNCTION(BlueprintCallable, Category = RigVMPin)
-	FName GetDisplayName() const;
+	UE_API FName GetDisplayName() const;
 
 	// Returns the direction of the pin
 	UFUNCTION(BlueprintCallable, Category = RigVMPin)
-	ERigVMPinDirection GetDirection() const;
+	UE_API ERigVMPinDirection GetDirection() const;
 
 	// Returns true if the pin is currently expanded
 	UFUNCTION(BlueprintCallable, Category = RigVMPin)
-	bool IsExpanded() const;
+	UE_API bool IsExpanded() const;
 
 	// Returns true if the pin is defined as a constant value / literal
 	UFUNCTION(BlueprintCallable, Category = RigVMPin)
-	bool IsDefinedAsConstant() const;
+	UE_API bool IsDefinedAsConstant() const;
 
 	// Returns true if the pin should be watched
 	UFUNCTION(BlueprintCallable, Category = RigVMPin)
-	bool RequiresWatch(const bool bCheckExposedPinChain = false) const;
+	UE_API bool RequiresWatch(const bool bCheckExposedPinChain = false) const;
 	
 	// Returns true if the data type of the Pin is a enum
 	UFUNCTION(BlueprintCallable, Category = RigVMPin)
-	bool IsEnum() const;
+	UE_API bool IsEnum() const;
 
 	// Returns true if the data type of the Pin is a struct
 	UFUNCTION(BlueprintCallable, Category = RigVMPin)
-	bool IsStruct() const;
+	UE_API bool IsStruct() const;
 
 	// Returns true if the Pin is a SubPin within a struct
 	UFUNCTION(BlueprintCallable, Category = RigVMPin)
-	bool IsStructMember() const;
+	UE_API bool IsStructMember() const;
 
 	// Returns true if the data type of the Pin is a uobject
 	UFUNCTION(BlueprintCallable, Category = RigVMPin)
-	bool IsUObject() const;
+	UE_API bool IsUObject() const;
 
 	// Returns true if the data type of the Pin is a interface
 	UFUNCTION(BlueprintCallable, Category = RigVMPin)
-	bool IsInterface() const;
+	UE_API bool IsInterface() const;
 
 	// Returns true if the data type of the Pin is an array
 	UFUNCTION(BlueprintCallable, Category = RigVMPin)
-	bool IsArray() const;
+	UE_API bool IsArray() const;
 
 	// Returns true if the Pin is a SubPin within an array
 	UFUNCTION(BlueprintCallable, Category = RigVMPin)
-	bool IsArrayElement() const;
+	UE_API bool IsArrayElement() const;
 
 	// Returns true if this pin represents a dynamic array
 	UFUNCTION(BlueprintCallable, Category = RigVMPin)
-	bool IsDynamicArray() const;
+	UE_API bool IsDynamicArray() const;
 
 	// Returns true if this data type is referenced counted
 	UFUNCTION(BlueprintCallable, Category = RigVMPin)
@@ -251,67 +253,67 @@ public:
 
 	// Returns true if this pin's value may be executed lazily
 	UFUNCTION(BlueprintCallable, Category = RigVMPin)
-	bool IsLazy() const;
+	UE_API bool IsLazy() const;
 
 	// Returns the index of the Pin within the node / parent Pin
 	UFUNCTION(BlueprintCallable, Category = RigVMPin)
-	int32 GetPinIndex() const;
+	UE_API int32 GetPinIndex() const;
 
 	// Returns the absolute index of the Pin within the node / parent Pin
 	UFUNCTION(BlueprintCallable, Category = RigVMPin)
-	int32 GetAbsolutePinIndex() const;
+	UE_API int32 GetAbsolutePinIndex() const;
 
 	// Returns the number of elements within an array Pin
 	UFUNCTION(BlueprintCallable, Category = RigVMPin)
-	int32 GetArraySize() const;
+	UE_API int32 GetArraySize() const;
 
 	// Returns the C++ data type of the pin
 	UFUNCTION(BlueprintCallable, Category = RigVMPin)
-	FString GetCPPType() const;
+	UE_API FString GetCPPType() const;
 
 	// Returns the C++ data type of an element of the Pin array
 	UFUNCTION(BlueprintCallable, Category = RigVMPin)
-	FString GetArrayElementCppType() const;
+	UE_API FString GetArrayElementCppType() const;
 
 	// Returns the argument type this pin would represent within a template
-	FRigVMTemplateArgumentType GetTemplateArgumentType() const;
+	UE_API FRigVMTemplateArgumentType GetTemplateArgumentType() const;
 
 	// Returns the argument type index this pin would represent within a template
-	TRigVMTypeIndex GetTypeIndex() const;
+	UE_API TRigVMTypeIndex GetTypeIndex() const;
 
 	// Returns true if the C++ data type is FString or FName
 	UFUNCTION(BlueprintCallable, Category = RigVMPin)
-	bool IsStringType() const;
+	UE_API bool IsStringType() const;
 
 	// Returns true if the C++ data type is an execute context
 	UFUNCTION(BlueprintCallable, Category = RigVMPin)
-	bool IsExecuteContext() const;
+	UE_API bool IsExecuteContext() const;
 
 	// Returns true if the C++ data type is unknown
 	UFUNCTION(BlueprintCallable, Category = RigVMPin)
-	bool IsWildCard() const;
+	UE_API bool IsWildCard() const;
 
 	// Returns true if any of the subpins is a wildcard
 	UFUNCTION(BlueprintCallable, Category = RigVMPin)
-	bool ContainsWildCardSubPin() const;
+	UE_API bool ContainsWildCardSubPin() const;
 
 	// Returns true if this pin is an array that should be displayed as elements only
 	UFUNCTION(BlueprintCallable, Category = RigVMPin)
-	bool IsFixedSizeArray() const;
+	UE_API bool IsFixedSizeArray() const;
 
 	// Returns true if this pin is an array that should be displayed as elements only
 	UFUNCTION(BlueprintCallable, Category = RigVMPin)
-	bool ShouldOnlyShowSubPins() const;
+	UE_API bool ShouldOnlyShowSubPins() const;
 
 	// Returns true if this pin's subpins should be hidden in the UI
 	UFUNCTION(BlueprintCallable, Category = RigVMPin)
-	bool ShouldHideSubPins() const;
+	UE_API bool ShouldHideSubPins() const;
 
 	UFUNCTION(BlueprintCallable, Category = RigVMPin)
-	FString GetOriginalDefaultValue() const;
+	UE_API FString GetOriginalDefaultValue() const;
 
 	UFUNCTION(BlueprintCallable, Category = RigVMPin)
-	bool HasOriginalDefaultValue() const;
+	UE_API bool HasOriginalDefaultValue() const;
 	
 	// Returns the default value of the Pin as a string.
 	// Note that this value is computed based on the Pin's
@@ -319,26 +321,26 @@ public:
 	// the default value is actually composed out of the
 	// default values of the X, Y and Z SubPins.
 	UFUNCTION(BlueprintCallable, Category = RigVMPin)
-	FString GetDefaultValue() const;
+	UE_API FString GetDefaultValue() const;
 
 	// Returns the default value with an additional override ma
-	FString GetDefaultValue(const FPinOverride& InOverride, bool bAdaptValueForPinType = true) const;
+	UE_API FString GetDefaultValue(const FPinOverride& InOverride, bool bAdaptValueForPinType = true) const;
 
 	// Returns the default value as stored by the user.
-	FString GetDefaultValueStoredByUserInterface() const;
+	UE_API FString GetDefaultValueStoredByUserInterface() const;
 
 	// Returns true if the default value provided is valid
 	UFUNCTION(BlueprintPure, Category = RigVMPin)
-	bool IsValidDefaultValue(const FString& InDefaultValue) const;
+	UE_API bool IsValidDefaultValue(const FString& InDefaultValue) const;
 
 	// Returns true if the default value was ever changed by the user
 	UE_DEPRECATED(5.6, "Please use URigVMPin::HasDefaultValueOverride()")
 	UFUNCTION(BlueprintPure, Category = RigVMPin)
-	bool HasUserProvidedDefaultValue() const;
+	UE_API bool HasUserProvidedDefaultValue() const;
 
 	// Returns true if the default value was ever changed by the user
     UFUNCTION(BlueprintPure, Category = RigVMPin)
-	bool HasDefaultValueOverride() const;
+	UE_API bool HasDefaultValueOverride() const;
 
 	// Returns true if the pin can / may provide a default value 
 	UFUNCTION(BlueprintPure, Category = RigVMPin)
@@ -346,207 +348,207 @@ public:
 
 	// Returns true if the pin can / may provide a default value 
 	UFUNCTION(BlueprintPure, Category = RigVMPin)
-	bool CanProvideDefaultValue() const;
+	UE_API bool CanProvideDefaultValue() const;
 
 	// Returns the default value clamped with the limit meta values defined by the UPROPERTY in URigVMUnitNodes 
-	FString ClampDefaultValueFromMetaData(const FString& InDefaultValue) const;
+	UE_API FString ClampDefaultValueFromMetaData(const FString& InDefaultValue) const;
 
 	// Returns the keyed metadata associated with this pin, if any
 	UFUNCTION(BlueprintCallable, Category = RigVMPin)
-	FString GetMetaData(FName InKey) const;
+	UE_API FString GetMetaData(FName InKey) const;
 
 	// Returns the name of a custom widget to be used
 	// for editing the Pin.
 	UFUNCTION(BlueprintCallable, Category = RigVMPin)
-	FName GetCustomWidgetName() const;
+	UE_API FName GetCustomWidgetName() const;
 
 	// Returns the tooltip of this pin
 	UFUNCTION(BlueprintCallable, Category = RigVMPin)
-	FText GetToolTipText() const;
+	UE_API FText GetToolTipText() const;
 
 	// Returns the struct of the data type of the Pin,
 	// or nullptr otherwise.
 	UFUNCTION(BlueprintCallable, Category = RigVMPin)
-	UObject* GetCPPTypeObject() const;
+	UE_API UObject* GetCPPTypeObject() const;
 
 	// Returns the struct of the data type of the Pin,
 	// or nullptr otherwise.
 	UFUNCTION(BlueprintCallable, Category = RigVMPin)
-	UScriptStruct* GetScriptStruct() const;
+	UE_API UScriptStruct* GetScriptStruct() const;
 
 	// Returns the parent struct of the data type of the Pin
 	UFUNCTION(BlueprintCallable, Category = RigVMPin)
-	UScriptStruct* GetParentScriptStruct(const URigVMUnitNode* FallbackNode) const;
+	UE_API UScriptStruct* GetParentScriptStruct(const URigVMUnitNode* FallbackNode) const;
 
 	// Returns the enum of the data type of the Pin,
 	// or nullptr otherwise.
 	UFUNCTION(BlueprintCallable, Category = RigVMPin)
-	UEnum* GetEnum() const;
+	UE_API UEnum* GetEnum() const;
 
 	// Returns the parent Pin - or nullptr if the Pin
 	// is nested directly below a node.
 	UFUNCTION(BlueprintCallable, Category = RigVMPin)
-	URigVMPin* GetParentPin() const;
+	UE_API URigVMPin* GetParentPin() const;
 
 	// Returns the top-most parent Pin, so for example
 	// for "Node.Transform.Translation.X" this returns
 	// the Pin for "Node.Transform".
 	UFUNCTION(BlueprintCallable, Category = RigVMPin)
-	URigVMPin* GetRootPin() const;
+	UE_API URigVMPin* GetRootPin() const;
 
 	// Returns true if this pin is a root pin
 	UFUNCTION(BlueprintCallable, Category = RigVMPin)
-	bool IsRootPin() const;
+	UE_API bool IsRootPin() const;
 
 	// Returns the pin to be used for a link.
 	// This might differ from this actual pin, since
 	// the pin might contain injected nodes.
 	UFUNCTION(BlueprintCallable, Category = RigVMPin)
-	URigVMPin* GetPinForLink() const;
+	UE_API URigVMPin* GetPinForLink() const;
 
 	// Returns the link that represents the connection
 	// between this pin and InOtherPin. nullptr is returned
 	// if the pins are not connected.
 	UFUNCTION(BlueprintCallable, Category = RigVMPin)
-	URigVMLink* FindLinkForPin(const URigVMPin* InOtherPin) const;
+	UE_API URigVMLink* FindLinkForPin(const URigVMPin* InOtherPin) const;
 
 	// Returns the original pin for a pin on an injected
 	// node. This can be used to determine where a link
 	// should go in the user interface
 	UFUNCTION(BlueprintCallable, Category = RigVMPin)
-	URigVMPin* GetOriginalPinFromInjectedNode() const;
+	UE_API URigVMPin* GetOriginalPinFromInjectedNode() const;
 
 	// Returns all of the SubPins of this one.
 	UFUNCTION(BlueprintCallable, Category = RigVMPin)
-	const TArray<URigVMPin*>& GetSubPins() const;
+	UE_API const TArray<URigVMPin*>& GetSubPins() const;
 
 	// Returns all of the SubPins of this one including sub-sub-pins
 	UFUNCTION(BlueprintCallable, Category = RigVMPin)
-	TArray<URigVMPin*> GetAllSubPinsRecursively() const;
+	UE_API TArray<URigVMPin*> GetAllSubPinsRecursively() const;
 
 	// Returns a SubPin given a name / path or nullptr.
 	UFUNCTION(BlueprintCallable, Category = RigVMPin)
-	URigVMPin* FindSubPin(const FString& InPinPath) const;
+	UE_API URigVMPin* FindSubPin(const FString& InPinPath) const;
 
 	// Returns true if this Pin is linked to another Pin
 	UFUNCTION(BlueprintCallable, Category = RigVMPin)
-	bool IsLinkedTo(const URigVMPin* InPin) const;
+	UE_API bool IsLinkedTo(const URigVMPin* InPin) const;
 
 	// Returns true if the pin has any link
-	bool IsLinked(bool bRecursive = false) const;
+	UE_API bool IsLinked(bool bRecursive = false) const;
 
 	// Returns all of the links linked to this Pin.
 	UFUNCTION(BlueprintCallable, Category = RigVMPin)
-	const TArray<URigVMLink*>& GetLinks() const;
+	UE_API const TArray<URigVMLink*>& GetLinks() const;
 
 	// Returns all of the linked source Pins,
 	// using this Pin as the target.
 	UFUNCTION(BlueprintCallable, Category = RigVMPin)
-	TArray<URigVMPin*> GetLinkedSourcePins(bool bRecursive = false) const;
+	UE_API TArray<URigVMPin*> GetLinkedSourcePins(bool bRecursive = false) const;
 
 	// Returns all of the linked target Pins,
 	// using this Pin as the source.
 	UFUNCTION(BlueprintCallable, Category = RigVMPin)
-	TArray<URigVMPin*> GetLinkedTargetPins(bool bRecursive = false) const;
+	UE_API TArray<URigVMPin*> GetLinkedTargetPins(bool bRecursive = false) const;
 
 	// Returns all of the source pins
 	// using this Pin as the target.
 	UFUNCTION(BlueprintCallable, Category = RigVMPin)
-	TArray<URigVMLink*> GetSourceLinks(bool bRecursive = false) const;
+	UE_API TArray<URigVMLink*> GetSourceLinks(bool bRecursive = false) const;
 
 	// Returns all of the target links,
 	// using this Pin as the source.
 	UFUNCTION(BlueprintCallable, Category = RigVMPin)
-	TArray<URigVMLink*> GetTargetLinks(bool bRecursive = false) const;
+	UE_API TArray<URigVMLink*> GetTargetLinks(bool bRecursive = false) const;
 
 	// Returns the node of this Pin.
 	UFUNCTION(BlueprintCallable, Category = RigVMPin)
-	URigVMNode* GetNode() const;
+	UE_API URigVMNode* GetNode() const;
 
 	// Returns the graph of this Pin.
 	UFUNCTION(BlueprintCallable, Category = RigVMPin)
-	URigVMGraph* GetGraph() const;
+	UE_API URigVMGraph* GetGraph() const;
 
 	// Returns true is the two provided source and target Pins
 	// can be linked to one another.
-	static bool CanLink(const URigVMPin* InSourcePin, const URigVMPin* InTargetPin, FString* OutFailureReason, const FRigVMByteCode* InByteCode, ERigVMPinDirection InUserLinkDirection = ERigVMPinDirection::IO, bool bInAllowWildcard = false, bool bEnableTypeCasting = true);
+	static UE_API bool CanLink(const URigVMPin* InSourcePin, const URigVMPin* InTargetPin, FString* OutFailureReason, const FRigVMByteCode* InByteCode, ERigVMPinDirection InUserLinkDirection = ERigVMPinDirection::IO, bool bInAllowWildcard = false, bool bEnableTypeCasting = true);
 
 	// Returns true if this pin has injected nodes
 	bool HasInjectedNodes() const { return InjectionInfos.Num() > 0; }
 
 	// Returns true if this pin has injected nodes
-	bool HasInjectedUnitNodes() const;
+	UE_API bool HasInjectedUnitNodes() const;
 
 	// Returns the injected nodes this pin contains.
 	const TArray<URigVMInjectionInfo*> GetInjectedNodes() const { return InjectionInfos; }
 
-	URigVMVariableNode* GetBoundVariableNode() const;
+	UE_API URigVMVariableNode* GetBoundVariableNode() const;
 
 	// Returns the variable bound to this pin (or NAME_None)
-	const FString GetBoundVariablePath() const;
+	UE_API const FString GetBoundVariablePath() const;
 
 	// Returns the variable bound to this pin (or NAME_None)
-	const FString GetBoundVariablePath(const FPinOverride& InOverride) const;
+	UE_API const FString GetBoundVariablePath(const FPinOverride& InOverride) const;
 
 	// Returns the variable bound to this pin (or NAME_None)
-	FString GetBoundVariableName() const;
+	UE_API FString GetBoundVariableName() const;
 
 	// Returns true if this pin is bound to a variable
-	bool IsBoundToVariable() const;
+	UE_API bool IsBoundToVariable() const;
 
 	// Returns true if this pin is bound to a variable
-	bool IsBoundToVariable(const FPinOverride& InOverride) const;
+	UE_API bool IsBoundToVariable(const FPinOverride& InOverride) const;
 
 	// Returns true if this pin is bound to an external variable
-	bool IsBoundToExternalVariable() const;
+	UE_API bool IsBoundToExternalVariable() const;
 
 	// Returns true if this pin is bound to a local variable
-	bool IsBoundToLocalVariable() const;
+	UE_API bool IsBoundToLocalVariable() const;
 
 	// Returns true if this pin is bound to an input argument
-	bool IsBoundToInputArgument() const;
+	UE_API bool IsBoundToInputArgument() const;
 
 	// Returns true if the pin can be bound to a given variable
-	bool CanBeBoundToVariable(const FRigVMExternalVariable& InExternalVariable, const FString& InSegmentPath = FString()) const;
+	UE_API bool CanBeBoundToVariable(const FRigVMExternalVariable& InExternalVariable, const FString& InSegmentPath = FString()) const;
 
 	// Returns true if the pin should not show up on a node, but in the details panel
-	bool ShowInDetailsPanelOnly() const;
+	UE_API bool ShowInDetailsPanelOnly() const;
 
 	// Returns an external variable matching this pin's type
-	FRigVMExternalVariable ToExternalVariable() const;
+	UE_API FRigVMExternalVariable ToExternalVariable() const;
 
 	// Returns true if the pin has been orphaned
-	bool IsOrphanPin() const;
+	UE_API bool IsOrphanPin() const;
 
-	uint32 GetStructureHash() const;
+	UE_API uint32 GetStructureHash() const;
 
 	// Returns true if this pin represents a trait 
 	UFUNCTION(BlueprintPure, Category = RigVMPin)
-	bool IsTraitPin() const;
+	UE_API bool IsTraitPin() const;
 
 	// Returns true if this pin represents a trait's programmatic pin
-	bool IsProgrammaticPin() const;
+	UE_API bool IsProgrammaticPin() const;
 
 	// Get all the sub-pins that are programmatic
-	TArray<URigVMPin*> GetProgrammaticSubPins() const;
+	UE_API TArray<URigVMPin*> GetProgrammaticSubPins() const;
 
 	// Returns the trait backing up this pin
-	TSharedPtr<FStructOnScope> GetTraitInstance(bool bUseDefaultValueFromPin = true) const;
+	UE_API TSharedPtr<FStructOnScope> GetTraitInstance(bool bUseDefaultValueFromPin = true) const;
 
 	// Returns the struct of the trait backing up this pin
-	UScriptStruct* GetTraitScriptStruct() const;
+	UE_API UScriptStruct* GetTraitScriptStruct() const;
 
-	const uint32& GetNodeCachedValueVersion() const;
-	const uint32& GetCachedValueVersion() const;
+	UE_API const uint32& GetNodeCachedValueVersion() const;
+	UE_API const uint32& GetCachedValueVersion() const;
 
 private:
 
-	void UpdateTypeInformationIfRequired() const;
-	void SetNameFromIndex();
-	void SetDisplayName(const FName& InDisplayName);
-	void IncrementVersion(bool bAffectParentPin = true, bool bAffectSubPins = true);
+	UE_API void UpdateTypeInformationIfRequired() const;
+	UE_API void SetNameFromIndex();
+	UE_API void SetDisplayName(const FName& InDisplayName);
+	UE_API void IncrementVersion(bool bAffectParentPin = true, bool bAffectSubPins = true);
 
-	void GetExposedPinChainImpl(TArray<const URigVMPin*>& OutExposedPins, TArray<const URigVMPin*>& VisitedPins) const;
+	UE_API void GetExposedPinChainImpl(TArray<const URigVMPin*>& OutExposedPins, TArray<const URigVMPin*>& VisitedPins) const;
 
 	UPROPERTY()
 	FName DisplayName;
@@ -646,7 +648,7 @@ private:
 	friend struct FRigVMClient;
 };
 
-class RIGVMDEVELOPER_API FRigVMPinDefaultValueImportErrorContext : public FOutputDevice
+class FRigVMPinDefaultValueImportErrorContext : public FOutputDevice
 {
 public:
 
@@ -672,3 +674,5 @@ public:
 private:
 	ELogVerbosity::Type MaxVerbosity;
 };
+
+#undef UE_API
diff --git a/Engine/Plugins/Runtime/RigVM/Source/RigVMDeveloper/Public/RigVMModel/RigVMSchema.h b/Engine/Plugins/Runtime/RigVM/Source/RigVMDeveloper/Public/RigVMModel/RigVMSchema.h
--- a/Engine/Plugins/Runtime/RigVM/Source/RigVMDeveloper/Public/RigVMModel/RigVMSchema.h
+++ b/Engine/Plugins/Runtime/RigVM/Source/RigVMDeveloper/Public/RigVMModel/RigVMSchema.h
@@ -10,6 +10,8 @@
 #include "RigVMGraph.h"
 #include "RigVMSchema.generated.h"
 
+#define UE_API RIGVMDEVELOPER_API
+
 #define RIGVMSCHEMA_DEFAULT_FUNCTION_BODY \
 const URigVMGraph* Graph = InController->GetGraph(); \
 check(Graph); \
@@ -24,115 +26,115 @@ class URigVMController;
  * The Schema is used to determine which actions are allowed
  * on a graph. This includes any topological change.
  */
-UCLASS()
-class RIGVMDEVELOPER_API URigVMSchema : public UObject
+UCLASS(MinimalAPI)
+class URigVMSchema : public UObject
 {
 	GENERATED_UCLASS_BODY()
 
 public:
 
 	// Default constructor
-	URigVMSchema();
+	UE_API URigVMSchema();
 
 	// Returns the execute context struct this schema is using
 	UScriptStruct* GetExecuteContextStruct() const { return ExecuteContextStruct; }
 
 	// Returns true if a graph supports a given type
-	virtual bool SupportsType(URigVMController* InController, TRigVMTypeIndex InTypeIndex) const;
+	UE_API virtual bool SupportsType(URigVMController* InController, TRigVMTypeIndex InTypeIndex) const;
 
 	// Returns true if a graph supports a given unit function
-	virtual bool SupportsUnitFunction(URigVMController* InController, const FRigVMFunction* InUnitFunction) const;
+	UE_API virtual bool SupportsUnitFunction(URigVMController* InController, const FRigVMFunction* InUnitFunction) const;
 
 	// Returns true if a graph supports a given dispatch factory
-	virtual bool SupportsDispatchFactory(URigVMController* InController, const FRigVMDispatchFactory* InDispatchFactory) const;
+	UE_API virtual bool SupportsDispatchFactory(URigVMController* InController, const FRigVMDispatchFactory* InDispatchFactory) const;
 
 	// Returns true if a graph supports a given template
-	virtual bool SupportsTemplate(URigVMController* InController, const FRigVMTemplate* InTemplate) const;
+	UE_API virtual bool SupportsTemplate(URigVMController* InController, const FRigVMTemplate* InTemplate) const;
 
 	// Returns true if a graph supports a given graph based function
-	virtual bool SupportsGraphFunction(URigVMController* InController, const FRigVMGraphFunctionHeader* InGraphFunction) const;
+	UE_API virtual bool SupportsGraphFunction(URigVMController* InController, const FRigVMGraphFunctionHeader* InGraphFunction) const;
 
 	// Returns true if a graph supports a given external variable
-	virtual bool SupportsExternalVariable(URigVMController* InController, const FRigVMExternalVariable* InExternalVariable) const;
+	UE_API virtual bool SupportsExternalVariable(URigVMController* InController, const FRigVMExternalVariable* InExternalVariable) const;
 
 	// Returns true if a graph supports pin categories / node layouts
-	virtual bool SupportsNodeLayouts(const URigVMGraph* InGraph) const;
+	UE_API virtual bool SupportsNodeLayouts(const URigVMGraph* InGraph) const;
 
 	// Returns true if the pin for a given struct should be unfolded into subpins for a graph
-	virtual bool ShouldUnfoldStruct(URigVMController* InController, const UStruct* InStruct) const;
+	UE_API virtual bool ShouldUnfoldStruct(URigVMController* InController, const UStruct* InStruct) const;
 
 	// Returns true if a node name is valid for a given graph
-	virtual bool IsValidNodeName(const URigVMGraph* InGraph, const FName& InNodeName) const;
+	UE_API virtual bool IsValidNodeName(const URigVMGraph* InGraph, const FName& InNodeName) const;
 
 	// Returns true if a given node can be added the graph
-	virtual bool CanAddNode(URigVMController* InController, const URigVMNode* InNode) const;
+	UE_API virtual bool CanAddNode(URigVMController* InController, const URigVMNode* InNode) const;
 
 	// Returns true if a given node can be removed from the graph
-	virtual bool CanRemoveNode(URigVMController* InController, const URigVMNode* InNode) const;
+	UE_API virtual bool CanRemoveNode(URigVMController* InController, const URigVMNode* InNode) const;
 
 	// Returns true if a given node can be added the graph
-	virtual bool CanRenameNode(URigVMController* InController, const URigVMNode* InNode, const FName& InNewNodeName) const;
+	UE_API virtual bool CanRenameNode(URigVMController* InController, const URigVMNode* InNode, const FName& InNewNodeName) const;
 
 	// Returns true if a node can moved to a new position
-	virtual bool CanMoveNode(URigVMController* InController, const URigVMNode* InNode, const FVector2D& InNewPosition = FVector2D::ZeroVector) const;
+	UE_API virtual bool CanMoveNode(URigVMController* InController, const URigVMNode* InNode, const FVector2D& InNewPosition = FVector2D::ZeroVector) const;
 
 	// Returns true if a node can resized to a new size
-	virtual bool CanResizeNode(URigVMController* InController, const URigVMNode* InNode, const FVector2D& InNewSize = {100, 100}) const;
+	UE_API virtual bool CanResizeNode(URigVMController* InController, const URigVMNode* InNode, const FVector2D& InNewSize = {100, 100}) const;
 
 	// Returns true if a node supports setting a node title
-	virtual bool CanSetNodeTitle(URigVMController* InController, const URigVMNode* InNode) const;
+	UE_API virtual bool CanSetNodeTitle(URigVMController* InController, const URigVMNode* InNode) const;
 
 	// Returns true if a node can recolored to a new color
-	virtual bool CanRecolorNode(URigVMController* InController, const URigVMNode* InNode, const FLinearColor& InNewColor = FLinearColor::White) const;
+	UE_API virtual bool CanRecolorNode(URigVMController* InController, const URigVMNode* InNode, const FLinearColor& InNewColor = FLinearColor::White) const;
 
 	// Returns true if a link can be added between two pins
-	virtual bool CanAddLink(URigVMController* InController, const URigVMPin* InSourcePin, const URigVMPin* InTargetPin, const FRigVMByteCode* InByteCode, ERigVMPinDirection InUserLinkDirection = ERigVMPinDirection::IO, bool bInAllowWildcard = false, bool bEnableTypeCasting = true, FString* OutFailureReason = nullptr) const;
+	UE_API virtual bool CanAddLink(URigVMController* InController, const URigVMPin* InSourcePin, const URigVMPin* InTargetPin, const FRigVMByteCode* InByteCode, ERigVMPinDirection InUserLinkDirection = ERigVMPinDirection::IO, bool bInAllowWildcard = false, bool bEnableTypeCasting = true, FString* OutFailureReason = nullptr) const;
 
 	// Returns true if a link exists between two pins and can be broken / removed
-	virtual bool CanBreakLink(URigVMController* InController, const URigVMPin* InSourcePin, const URigVMPin* InTargetPin) const;
+	UE_API virtual bool CanBreakLink(URigVMController* InController, const URigVMPin* InSourcePin, const URigVMPin* InTargetPin) const;
 
 	// Returns true if a set of nodes can be collapsed
-	virtual bool CanCollapseNodes(URigVMController* InController, const TArrayView<URigVMNode* const>& InNodesToCollapse) const;
+	UE_API virtual bool CanCollapseNodes(URigVMController* InController, const TArrayView<URigVMNode* const>& InNodesToCollapse) const;
 
 	// Returns true if a node can be expanded within a graph
-	virtual bool CanExpandNode(URigVMController* InController, const URigVMNode* InNodeToExpand) const;
+	UE_API virtual bool CanExpandNode(URigVMController* InController, const URigVMNode* InNodeToExpand) const;
 
 	// Returns true if a pin should be unfolded (represented by its subpins)
-	virtual bool CanUnfoldPin(URigVMController* InController, const URigVMPin* InPinToUnfold) const;
+	UE_API virtual bool CanUnfoldPin(URigVMController* InController, const URigVMPin* InPinToUnfold) const;
 
 	// Returns true if a variable can be bound to a pin
-	virtual bool CanBindVariable(URigVMController* InController, const URigVMPin* InPinToBind, const FRigVMExternalVariable* InVariableToBind, const FString& InNewBoundVariablePath) const;
+	UE_API virtual bool CanBindVariable(URigVMController* InController, const URigVMPin* InPinToBind, const FRigVMExternalVariable* InVariableToBind, const FString& InNewBoundVariablePath) const;
 
 	// Returns true if a variable can be unbound from a pin
-	virtual bool CanUnbindVariable(URigVMController* InController, const URigVMPin* InBoundPin) const;
+	UE_API virtual bool CanUnbindVariable(URigVMController* InController, const URigVMPin* InBoundPin) const;
 
 	// Returns true if functions definitions (not refs) can be added to the graph
-	virtual bool CanAddFunction(URigVMController* InController, const URigVMNode* InFunctionNode) const;
+	UE_API virtual bool CanAddFunction(URigVMController* InController, const URigVMNode* InFunctionNode) const;
 
 	// Returns true if a function can be removed
-	virtual bool CanRemoveFunction(URigVMController* InController, const URigVMNode* InFunctionNode) const;
+	UE_API virtual bool CanRemoveFunction(URigVMController* InController, const URigVMNode* InFunctionNode) const;
 	
 	static int32 GetMaxNameLength() { return 100; }
-	virtual FString GetSanitizedName(const FString& InName, bool bAllowPeriod, bool bAllowSpace) const;
-	virtual FString GetSanitizedGraphName(const FString& InName) const;
-	virtual FString GetSanitizedNodeName(const FString& InName) const;
-	virtual FString GetSanitizedVariableName(const FString& InName) const;
-	virtual FString GetSanitizedPinName(const FString& InName) const;
-	virtual FString GetSanitizedPinPath(const FString& InName) const;
-	virtual FString GetGraphOuterName(const URigVMGraph* InGraph) const;
-	virtual	FString GetValidNodeName(const URigVMGraph* InGraph, const FString& InPrefix) const;
-	static void SanitizeName(FString& InOutName, bool bAllowPeriod, bool bAllowSpace);
+	UE_API virtual FString GetSanitizedName(const FString& InName, bool bAllowPeriod, bool bAllowSpace) const;
+	UE_API virtual FString GetSanitizedGraphName(const FString& InName) const;
+	UE_API virtual FString GetSanitizedNodeName(const FString& InName) const;
+	UE_API virtual FString GetSanitizedVariableName(const FString& InName) const;
+	UE_API virtual FString GetSanitizedPinName(const FString& InName) const;
+	UE_API virtual FString GetSanitizedPinPath(const FString& InName) const;
+	UE_API virtual FString GetGraphOuterName(const URigVMGraph* InGraph) const;
+	UE_API virtual	FString GetValidNodeName(const URigVMGraph* InGraph, const FString& InPrefix) const;
+	static UE_API void SanitizeName(FString& InOutName, bool bAllowPeriod, bool bAllowSpace);
 
 	// Returns a unique name based on a IsNameAvailable predicate
-	static FName GetUniqueName(const FName& InName, TFunction<bool(const FName&)> IsNameAvailablePredicate, bool bAllowPeriod, bool bAllowSpace);
+	static UE_API FName GetUniqueName(const FName& InName, TFunction<bool(const FName&)> IsNameAvailablePredicate, bool bAllowPeriod, bool bAllowSpace);
 
 protected:
 
 	// Sets the execute context struct this schema is using
-	void SetExecuteContextStruct(UScriptStruct* InExecuteContextStruct);
+	UE_API void SetExecuteContextStruct(UScriptStruct* InExecuteContextStruct);
 
-	bool IsGraphEditable(const URigVMGraph* InGraph) const;
-	TObjectPtr<URigVMNode> FindEventNode(URigVMController* InController, const UScriptStruct* InScriptStruct) const;
+	UE_API bool IsGraphEditable(const URigVMGraph* InGraph) const;
+	UE_API TObjectPtr<URigVMNode> FindEventNode(URigVMController* InController, const UScriptStruct* InScriptStruct) const;
 
 	UPROPERTY(transient)
 	TObjectPtr<UScriptStruct> ExecuteContextStruct;
@@ -151,4 +153,6 @@ protected:
 	friend struct FRigVMClient;
 	friend class URigVMController;
 	friend struct FRigVMBaseAction;
-};
\ No newline at end of file
+};
+
+#undef UE_API
diff --git a/Engine/Plugins/Runtime/RigVM/Source/RigVMDeveloper/Public/RigVMUserWorkflowRegistry.h b/Engine/Plugins/Runtime/RigVM/Source/RigVMDeveloper/Public/RigVMUserWorkflowRegistry.h
--- a/Engine/Plugins/Runtime/RigVM/Source/RigVMDeveloper/Public/RigVMUserWorkflowRegistry.h
+++ b/Engine/Plugins/Runtime/RigVM/Source/RigVMDeveloper/Public/RigVMUserWorkflowRegistry.h
@@ -16,32 +16,36 @@
 
 #include "RigVMUserWorkflowRegistry.generated.h"
 
+#define UE_API RIGVMDEVELOPER_API
+
 class UScriptStruct;
 struct FFrame;
 
 DECLARE_DYNAMIC_DELEGATE_RetVal_OneParam(TArray<FRigVMUserWorkflow>, FRigVMUserWorkflowProvider, const UObject*, InSubject);
 
-UCLASS(BlueprintType)
-class RIGVMDEVELOPER_API URigVMUserWorkflowRegistry : public UObject
+UCLASS(MinimalAPI, BlueprintType)
+class URigVMUserWorkflowRegistry : public UObject
 {
 	GENERATED_BODY()
 
 public:
 
 	UFUNCTION(BlueprintPure, Category = FRigVMUserWorkflowRegistry)
-   	static URigVMUserWorkflowRegistry* Get();
+   	static UE_API URigVMUserWorkflowRegistry* Get();
 
 	UFUNCTION(BlueprintCallable, Category = FRigVMUserWorkflowRegistry)
-	int32 RegisterProvider(const UScriptStruct* InStruct, FRigVMUserWorkflowProvider InProvider);
+	UE_API int32 RegisterProvider(const UScriptStruct* InStruct, FRigVMUserWorkflowProvider InProvider);
 
 	UFUNCTION(BlueprintCallable, Category = FRigVMUserWorkflowRegistry)
-	void UnregisterProvider(int32 InHandle);
+	UE_API void UnregisterProvider(int32 InHandle);
 
 	UFUNCTION(BlueprintPure, Category = FRigVMUserWorkflowRegistry)
-	TArray<FRigVMUserWorkflow> GetWorkflows(ERigVMUserWorkflowType InType, const UScriptStruct* InStruct, const UObject* InSubject) const;
+	UE_API TArray<FRigVMUserWorkflow> GetWorkflows(ERigVMUserWorkflowType InType, const UScriptStruct* InStruct, const UObject* InSubject) const;
 
 private:
 
 	int32 MaxHandle = 0;
 	mutable TArray<TTuple<int32, const UScriptStruct*, FRigVMUserWorkflowProvider>> Providers;
 };
+
+#undef UE_API
diff --git a/Engine/Plugins/Runtime/RigVM/Source/RigVMEditor/Private/Editor/Kismet/RigVMBlueprintCompilationManager.h b/Engine/Plugins/Runtime/RigVM/Source/RigVMEditor/Private/Editor/Kismet/RigVMBlueprintCompilationManager.h
--- a/Engine/Plugins/Runtime/RigVM/Source/RigVMEditor/Private/Editor/Kismet/RigVMBlueprintCompilationManager.h
+++ b/Engine/Plugins/Runtime/RigVM/Source/RigVMEditor/Private/Editor/Kismet/RigVMBlueprintCompilationManager.h
@@ -35,10 +35,10 @@ struct FRigVMBPCompileRequest
 	FCompilerResultsLog* ClientResultsLog;
 };
 
-struct RIGVMEDITOR_API FRigVMBlueprintCompilationManager
+struct FRigVMBlueprintCompilationManager
 {
-	static void Initialize();
-	static void Shutdown();
+	static RIGVMEDITOR_API void Initialize();
+	static RIGVMEDITOR_API void Shutdown();
 
 	/**
 	 * Compiles all blueprints that have been placed in the compilation queue. 
@@ -47,46 +47,46 @@ struct RIGVMEDITOR_API FRigVMBlueprintCompilationManager
 	 * of objects the compilation manager has loaded. The compilation manager
 	 * will not load data while processing the compilation queue)
 	 */
-	static void FlushCompilationQueue(FUObjectSerializeContext* InLoadContext);
+	static RIGVMEDITOR_API void FlushCompilationQueue(FUObjectSerializeContext* InLoadContext);
 	
 	/**
 	 * Flushes the compilation queue and finishes reinstancing
 	 */
-	static void FlushCompilationQueueAndReinstance();
+	static RIGVMEDITOR_API void FlushCompilationQueueAndReinstance();
 
 	/**
 	 * Immediately compiles the blueprint, no expectation that related blueprints be subsequently compiled.
 	 * It will be significantly more efficient to queue blueprints and then flush the compilation queue
 	 * if there are several blueprints that require compilation (e.g. typical case on PIE):
 	 */
-	static void CompileSynchronously(const FRigVMBPCompileRequest& Request);
+	static RIGVMEDITOR_API void CompileSynchronously(const FRigVMBPCompileRequest& Request);
 
 	/**
 	 * Adds a newly loaded blueprint to the compilation queue
 	 */
-	static void NotifyBlueprintLoaded(UBlueprint* BPLoaded);
+	static RIGVMEDITOR_API void NotifyBlueprintLoaded(UBlueprint* BPLoaded);
 
 	/**
 	 * Adds a blueprint to the compilation queue - useful for batch compilation
 	 */
-	static void QueueForCompilation(UBlueprint* BP);
+	static RIGVMEDITOR_API void QueueForCompilation(UBlueprint* BP);
 
 	/** Returns true when UBlueprint::GeneratedClass members are up to date */
-	static bool IsGeneratedClassLayoutReady();
+	static RIGVMEDITOR_API bool IsGeneratedClassLayoutReady();
 	
 	/** 
 	 * Returns the Default Value associated with ForClass::Property, if ForClass is currently 
 	 * being compiled this function can look at the old version of the CDO and read the default
 	 * value from there
 	 */
-	static bool GetDefaultValue(const UClass* ForClass, const FProperty* Property, FString& OutDefaultValueAsString);
+	static RIGVMEDITOR_API bool GetDefaultValue(const UClass* ForClass, const FProperty* Property, FString& OutDefaultValueAsString);
 
 	/**
 	 * Safely reparents all child classes of every Key in OldClassToNewClass to the class in 
 	 * the corresponding Value. Typically this means every child type will be reinstanced - although
 	 * reinstancing could be avoided when layouts match.
 	 */
-	static void ReparentHierarchies(const TMap<UClass*, UClass*>& OldClassToNewClass);
+	static RIGVMEDITOR_API void ReparentHierarchies(const TMap<UClass*, UClass*>& OldClassToNewClass);
 
 	/** 
 	 * Registers a blueprint compiler extension - anytime a blueprint of the provided type is compiled
@@ -94,9 +94,9 @@ struct RIGVMEDITOR_API FRigVMBlueprintCompilationManager
 	 * compilation there may be blueprints compiled before the extension is registed unless special
 	 * care has been taken.
 	 */
-	static void RegisterCompilerExtension(TSubclassOf<UBlueprint> BlueprintType, URigVMBlueprintCompilerExtension* Extension);
+	static RIGVMEDITOR_API void RegisterCompilerExtension(TSubclassOf<UBlueprint> BlueprintType, URigVMBlueprintCompilerExtension* Extension);
 private:
-	FRigVMBlueprintCompilationManager();
+	RIGVMEDITOR_API FRigVMBlueprintCompilationManager();
 };
 
-#endif
\ No newline at end of file
+#endif
diff --git a/Engine/Plugins/Runtime/RigVM/Source/RigVMEditor/Private/Editor/Kismet/RigVMBlueprintCompilerExtension.h b/Engine/Plugins/Runtime/RigVM/Source/RigVMEditor/Private/Editor/Kismet/RigVMBlueprintCompilerExtension.h
--- a/Engine/Plugins/Runtime/RigVM/Source/RigVMEditor/Private/Editor/Kismet/RigVMBlueprintCompilerExtension.h
+++ b/Engine/Plugins/Runtime/RigVM/Source/RigVMEditor/Private/Editor/Kismet/RigVMBlueprintCompilerExtension.h
@@ -11,6 +11,8 @@
 
 #include "RigVMBlueprintCompilerExtension.generated.h"
 
+#define UE_API RIGVMEDITOR_API
+
 class UEdGraph;
 
 USTRUCT()
@@ -22,16 +24,16 @@ struct FRigVMBlueprintCompiledData
 	TArray<TObjectPtr<UEdGraph>> IntermediateGraphs;
 };
 
-UCLASS(Abstract)
-class RIGVMEDITOR_API URigVMBlueprintCompilerExtension : public UObject
+UCLASS(MinimalAPI, Abstract)
+class URigVMBlueprintCompilerExtension : public UObject
 {
 	GENERATED_BODY()
 	
 public:
-	URigVMBlueprintCompilerExtension(const FObjectInitializer& ObjectInitializer);
+	UE_API URigVMBlueprintCompilerExtension(const FObjectInitializer& ObjectInitializer);
 #if !WITH_RIGVMLEGACYEDITOR
 
-	void BlueprintCompiled(const FKismetCompilerContext& CompilationContext, const FRigVMBlueprintCompiledData& Data);
+	UE_API void BlueprintCompiled(const FKismetCompilerContext& CompilationContext, const FRigVMBlueprintCompiledData& Data);
 
 protected:
 	/** 
@@ -42,3 +44,4 @@ protected:
 #endif
 };
 
+#undef UE_API
diff --git a/Engine/Plugins/Runtime/RigVM/Source/RigVMEditor/Private/Editor/Kismet/RigVMFindInBlueprintManager.h b/Engine/Plugins/Runtime/RigVM/Source/RigVMEditor/Private/Editor/Kismet/RigVMFindInBlueprintManager.h
--- a/Engine/Plugins/Runtime/RigVM/Source/RigVMEditor/Private/Editor/Kismet/RigVMFindInBlueprintManager.h
+++ b/Engine/Plugins/Runtime/RigVM/Source/RigVMEditor/Private/Editor/Kismet/RigVMFindInBlueprintManager.h
@@ -70,83 +70,83 @@ typedef TSharedPtr<FRigVMImaginaryFiBData, ESPMode::ThreadSafe> FRigVMImaginaryF
 /**
  *Const values for Find-in-Blueprints to tag searchable data
  */
-struct RIGVMEDITOR_API FRigVMFindInBlueprintSearchTags
+struct FRigVMFindInBlueprintSearchTags
 {
 	/** Properties tag, for Blueprint variables */
-	static const FText FiB_Properties;
+	static RIGVMEDITOR_API const FText FiB_Properties;
 
 	/** Components tags */
-	static const FText FiB_Components;
-	static const FText FiB_IsSCSComponent;
+	static RIGVMEDITOR_API const FText FiB_Components;
+	static RIGVMEDITOR_API const FText FiB_IsSCSComponent;
 	/** End Components tags */
 
 	/** Nodes tag */
-	static const FText FiB_Nodes;
+	static RIGVMEDITOR_API const FText FiB_Nodes;
 
 	/** Schema Name tag, to identify the schema that a graph uses */
-	static const FText FiB_SchemaName;
+	static RIGVMEDITOR_API const FText FiB_SchemaName;
 	/** Uber graphs tag */
-	static const FText FiB_UberGraphs;
+	static RIGVMEDITOR_API const FText FiB_UberGraphs;
 	/** Function graph tag */
-	static const FText FiB_Functions;
+	static RIGVMEDITOR_API const FText FiB_Functions;
 	/** Macro graph tag */
-	static const FText FiB_Macros;
+	static RIGVMEDITOR_API const FText FiB_Macros;
 	/** Sub graph tag, for any sub-graphs in a Blueprint */
-	static const FText FiB_SubGraphs;
+	static RIGVMEDITOR_API const FText FiB_SubGraphs;
 	/** Extension graph tag, for any graphs in a Blueprint Extension */
-	static const FText FiB_ExtensionGraphs;
+	static RIGVMEDITOR_API const FText FiB_ExtensionGraphs;
 	/** Blueprint extension tag. */
-	static const FText FiB_Extensions;
+	static RIGVMEDITOR_API const FText FiB_Extensions;
 
 	/** Name tag */
-	static const FText FiB_Name;
+	static RIGVMEDITOR_API const FText FiB_Name;
 	/** Native Name tag */
-	static const FText FiB_NativeName;
+	static RIGVMEDITOR_API const FText FiB_NativeName;
 	/** Class Name tag */
-	static const FText FiB_ClassName;
+	static RIGVMEDITOR_API const FText FiB_ClassName;
 	/** NodeGuid tag */
-	static const FText FiB_NodeGuid;
+	static RIGVMEDITOR_API const FText FiB_NodeGuid;
 	/** Default value */
-	static const FText FiB_DefaultValue;
+	static RIGVMEDITOR_API const FText FiB_DefaultValue;
 	/** Tooltip tag */
-	static const FText FiB_Tooltip;
+	static RIGVMEDITOR_API const FText FiB_Tooltip;
 	/** Description tag */
-	static const FText FiB_Description;
+	static RIGVMEDITOR_API const FText FiB_Description;
 	/** Comment tag */
-	static const FText FiB_Comment;
+	static RIGVMEDITOR_API const FText FiB_Comment;
 	/** Path tag */
-	static const FText FiB_Path;
+	static RIGVMEDITOR_API const FText FiB_Path;
 	/** Parent Class tag */
-	static const FText FiB_ParentClass;
+	static RIGVMEDITOR_API const FText FiB_ParentClass;
 	/** Interfaces tag */
-	static const FText FiB_Interfaces;
+	static RIGVMEDITOR_API const FText FiB_Interfaces;
 	/** Class that originally defined the function tag */
-	static const FText FiB_FuncOriginClass;
+	static RIGVMEDITOR_API const FText FiB_FuncOriginClass;
 
 	/** Pin type tags */
 
 	/** Pins tag */
-	static const FText FiB_Pins;
+	static RIGVMEDITOR_API const FText FiB_Pins;
 
 	/** Pin Category tag */
-	static const FText FiB_PinCategory;
+	static RIGVMEDITOR_API const FText FiB_PinCategory;
 	/** Pin Sub-Category tag */
-	static const FText FiB_PinSubCategory;
+	static RIGVMEDITOR_API const FText FiB_PinSubCategory;
 	/** Pin object class tag */
-	static const FText FiB_ObjectClass;
+	static RIGVMEDITOR_API const FText FiB_ObjectClass;
 	/** Pin IsArray tag */
-	static const FText FiB_IsArray;
+	static RIGVMEDITOR_API const FText FiB_IsArray;
 	/** Pin IsReference tag */
-	static const FText FiB_IsReference;
+	static RIGVMEDITOR_API const FText FiB_IsReference;
 	/** Glyph icon tag */
-	static const FText FiB_Glyph;
+	static RIGVMEDITOR_API const FText FiB_Glyph;
 	/** Style set the glyph belongs to */
-	static const FText FiB_GlyphStyleSet;
+	static RIGVMEDITOR_API const FText FiB_GlyphStyleSet;
 	/** Glyph icon color tag */
-	static const FText FiB_GlyphColor;
+	static RIGVMEDITOR_API const FText FiB_GlyphColor;
 
 	// Identifier for metadata storage, completely unsearchable tag
-	static const FText FiBMetaDataTag;
+	static RIGVMEDITOR_API const FText FiBMetaDataTag;
 	/** End const values for Find-in-Blueprint */
 };
 #endif
@@ -288,13 +288,13 @@ enum ERigVMSearchQueryFilter
 // 	FText Value;
 // };
 
-struct RIGVMEDITOR_API FRigVMFiBMD
+struct FRigVMFiBMD
 {
-	static const FString FiBSearchableMD;
-	static const FString FiBSearchableShallowMD;
-	static const FString FiBSearchableExplicitMD;
-	static const FString FiBSearchableHiddenExplicitMD;
-	static const FString FiBSearchableFormatVersionMD;
+	static RIGVMEDITOR_API const FString FiBSearchableMD;
+	static RIGVMEDITOR_API const FString FiBSearchableShallowMD;
+	static RIGVMEDITOR_API const FString FiBSearchableExplicitMD;
+	static RIGVMEDITOR_API const FString FiBSearchableHiddenExplicitMD;
+	static RIGVMEDITOR_API const FString FiBSearchableFormatVersionMD;
 };
 
 /** Which assets to index for caching */
@@ -335,7 +335,7 @@ enum class ERigVMFiBCacheOpFlags
 ENUM_CLASS_FLAGS(ERigVMFiBCacheOpFlags);
 
 /** Options to configure the bulk caching task */
-struct RIGVMEDITOR_API FRigVMFindInBlueprintCachingOptions
+struct FRigVMFindInBlueprintCachingOptions
 {
 	/** Type of caching operation */
 	ERigVMFiBCacheOpType OpType = ERigVMFiBCacheOpType::CachePendingAssets;
@@ -366,32 +366,32 @@ ENUM_CLASS_FLAGS(ERigVMAddOrUpdateBlueprintSearchMetadataFlags);
 // FRigVMFindInBlueprintsResult
 
 /* Item that matched the search results */
-class RIGVMEDITOR_API FRigVMFindInBlueprintsResult : public TSharedFromThis< FRigVMFindInBlueprintsResult >
+class FRigVMFindInBlueprintsResult : public TSharedFromThis< FRigVMFindInBlueprintsResult >
 {
 public:
 	FRigVMFindInBlueprintsResult() = default;
 	virtual ~FRigVMFindInBlueprintsResult() = default;
 
 	/* Create a root */
-	explicit FRigVMFindInBlueprintsResult(const FText& InDisplayText);
+	RIGVMEDITOR_API explicit FRigVMFindInBlueprintsResult(const FText& InDisplayText);
 
 	/* Called when user clicks on the search item */
-	virtual FReply OnClick();
+	RIGVMEDITOR_API virtual FReply OnClick();
 
 	/* Get Category for this search result */
-	virtual FText GetCategory() const;
+	RIGVMEDITOR_API virtual FText GetCategory() const;
 
 	/* Create an icon to represent the result */
-	virtual TSharedRef<SWidget>	CreateIcon() const;
+	RIGVMEDITOR_API virtual TSharedRef<SWidget>	CreateIcon() const;
 
 	/** Finalizes any content for the search data that was unsafe to do on a separate thread */
 	virtual void FinalizeSearchData() {};
 
 	/* Gets the comment on this node if any */
-	FString GetCommentText() const;
+	RIGVMEDITOR_API FString GetCommentText() const;
 
 	/** gets the blueprint housing all these search results */
-	UBlueprint* GetParentBlueprint() const;
+	RIGVMEDITOR_API UBlueprint* GetParentBlueprint() const;
 
 	/**
 	* Parses search info for specific data important for displaying the search result in an easy to understand format
@@ -404,10 +404,10 @@ public:
 	virtual void ParseSearchInfo(FText InKey, FText InValue) {};
 
 	/** Returns the Object represented by this search information give the Blueprint it can be found in */
-	virtual UObject* GetObject(UBlueprint* InBlueprint) const;
+	RIGVMEDITOR_API virtual UObject* GetObject(UBlueprint* InBlueprint) const;
 
 	/** Returns the display string for the row */
-	FText GetDisplayString() const;
+	RIGVMEDITOR_API FText GetDisplayString() const;
 
 public:
 	/*Any children listed under this category */
@@ -428,7 +428,7 @@ typedef TSharedPtr<FRigVMFindInBlueprintsResult> FRigVMSearchResult;
 ////////////////////////////////////
 // FRigVMStreamSearch
 
-struct RIGVMEDITOR_API FRigVMStreamSearchOptions
+struct FRigVMStreamSearchOptions
 {
 	/** Filter to limit the FilteredImaginaryResults to */
 	enum ERigVMSearchQueryFilter ImaginaryDataFilter;
@@ -447,40 +447,40 @@ struct RIGVMEDITOR_API FRigVMStreamSearchOptions
 /**
  * Async task for searching Blueprints
  */
-class RIGVMEDITOR_API FRigVMStreamSearch : public FRunnable
+class FRigVMStreamSearch : public FRunnable
 {
 public:
 	/** Constructor */
-	FRigVMStreamSearch(const FString& InSearchValue, const FRigVMStreamSearchOptions& InSearchOptions = FRigVMStreamSearchOptions());
+	RIGVMEDITOR_API FRigVMStreamSearch(const FString& InSearchValue, const FRigVMStreamSearchOptions& InSearchOptions = FRigVMStreamSearchOptions());
 
 	/** Begin FRunnable Interface */
-	virtual bool Init() override;
+	RIGVMEDITOR_API virtual bool Init() override;
 
-	virtual uint32 Run() override;
+	RIGVMEDITOR_API virtual uint32 Run() override;
 
-	virtual void Stop() override;
+	RIGVMEDITOR_API virtual void Stop() override;
 
-	virtual void Exit() override;
+	RIGVMEDITOR_API virtual void Exit() override;
 	/** End FRunnable Interface */
 
 	/** Brings the thread to a safe stop before continuing. */
-	void EnsureCompletion();
+	RIGVMEDITOR_API void EnsureCompletion();
 
 	/** Returns TRUE if the thread is done with it's work. */
-	bool IsComplete() const;
+	RIGVMEDITOR_API bool IsComplete() const;
 
 	/** Returns TRUE if Stop() was called while work is still pending. */
-	bool WasStopped() const;
+	RIGVMEDITOR_API bool WasStopped() const;
 
 	/**
 	 * Appends the items filtered through the search filter to the passed array
 	 *
 	 * @param OutItemsFound		All the items found since last queried
 	 */
-	void GetFilteredItems(TArray<TSharedPtr<class FRigVMFindInBlueprintsResult>>& OutItemsFound);
+	RIGVMEDITOR_API void GetFilteredItems(TArray<TSharedPtr<class FRigVMFindInBlueprintsResult>>& OutItemsFound);
 
 	/** Helper function to query the percent complete this search is */
-	float GetPercentComplete() const;
+	RIGVMEDITOR_API float GetPercentComplete() const;
 
 	/** Returns the Out-of-Date Blueprint count */
 	int32 GetOutOfDateCount() const
@@ -489,7 +489,7 @@ public:
 	}
 
 	/** Returns the FilteredImaginaryResults from the search query, these results have been filtered by the ImaginaryDataFilter. */
-	void GetFilteredImaginaryResults(TArray<FRigVMImaginaryFiBDataSharedPtr>& OutFilteredImaginaryResults);
+	RIGVMEDITOR_API void GetFilteredImaginaryResults(TArray<FRigVMImaginaryFiBDataSharedPtr>& OutFilteredImaginaryResults);
 	
 public:
 	/** Thread to run the cleanup FRunnable on */
@@ -528,19 +528,19 @@ private:
 // FRigVMFindInBlueprintSearchManager
 
 /** Singleton manager for handling all Blueprint searches, helps to manage the going progress of Blueprints, and is thread-safe. */
-class RIGVMEDITOR_API FRigVMFindInBlueprintSearchManager : public FTickableEditorObject
+class FRigVMFindInBlueprintSearchManager : public FTickableEditorObject
 {
 public:
-	static FRigVMFindInBlueprintSearchManager* Instance;
-	static FRigVMFindInBlueprintSearchManager& Get();
+	static RIGVMEDITOR_API FRigVMFindInBlueprintSearchManager* Instance;
+	static RIGVMEDITOR_API FRigVMFindInBlueprintSearchManager& Get();
 
-	FRigVMFindInBlueprintSearchManager();
-	~FRigVMFindInBlueprintSearchManager();
+	RIGVMEDITOR_API FRigVMFindInBlueprintSearchManager();
+	RIGVMEDITOR_API ~FRigVMFindInBlueprintSearchManager();
 
 	//~ Begin FTickableObject Interface
-	virtual void Tick(float DeltaTime) override;
-	virtual bool IsTickable() const override;
-	virtual TStatId GetStatId() const override;
+	RIGVMEDITOR_API virtual void Tick(float DeltaTime) override;
+	RIGVMEDITOR_API virtual bool IsTickable() const override;
+	RIGVMEDITOR_API virtual TStatId GetStatId() const override;
 	//~ End FTickableObject Interface
 
 	/**
@@ -549,7 +549,7 @@ public:
 	 * @param InSearchData		Search data to be applied
 	 * @param bAllowNewEntry	Whether to allow a new entry to be added if a match is not found
 	 */
-	void ApplySearchDataToDatabase(FRigVMSearchData InSearchData, bool bAllowNewEntry = false);
+	RIGVMEDITOR_API void ApplySearchDataToDatabase(FRigVMSearchData InSearchData, bool bAllowNewEntry = false);
 
 	/**
 	 * Given an asset path, locate and return a copy of its matching search data in the index cache.
@@ -557,7 +557,7 @@ public:
 	 * @param InAssetPath				Asset path (search index key).
 	 * @return							Matching search data from the index cache. Will return invalid (empty) search data if a matching entry was not found.
 	 */
-	FRigVMSearchData GetSearchDataForAssetPath(const FSoftObjectPath& InAssetPath);
+	RIGVMEDITOR_API FRigVMSearchData GetSearchDataForAssetPath(const FSoftObjectPath& InAssetPath);
 
 	/**
 	 * Gathers the Blueprint's search metadata and adds or updates it in the cache
@@ -566,14 +566,14 @@ public:
 	 * @param InFlags			Flags to assist with controlling this function's behavior
 	 * @param InVersion			Allows the caller to override the format version to use for caching
 	 */
-	void AddOrUpdateBlueprintSearchMetadata(UBlueprint* InBlueprint, ERigVMAddOrUpdateBlueprintSearchMetadataFlags InFlags = ERigVMAddOrUpdateBlueprintSearchMetadataFlags::None, ERigVMFiBVersion InVersion = ERigVMFiBVersion::RIGVM_FIB_VER_NONE);
+	RIGVMEDITOR_API void AddOrUpdateBlueprintSearchMetadata(UBlueprint* InBlueprint, ERigVMAddOrUpdateBlueprintSearchMetadataFlags InFlags = ERigVMAddOrUpdateBlueprintSearchMetadataFlags::None, ERigVMFiBVersion InVersion = ERigVMFiBVersion::RIGVM_FIB_VER_NONE);
 
 	/**
 	 * Starts a search query, the FiB manager handles where the thread is at in the search query at all times so that post-save of the cache to disk it can correct the index
 	 *
 	 * @param InSearchOriginator		Pointer to the thread object that the query originates from
 	 */
-	void BeginSearchQuery(const class FRigVMStreamSearch* InSearchOriginator);
+	RIGVMEDITOR_API void BeginSearchQuery(const class FRigVMStreamSearch* InSearchOriginator);
 
 	/**
 	 * Continues a search query, returning a single piece of search data
@@ -582,7 +582,7 @@ public:
 	 * @param OutSearchData				Result of the search, if there is any Blueprint search data still available to query
 	 * @return							TRUE if the search was successful, FALSE if the search is complete
 	 */
-	bool ContinueSearchQuery(const class FRigVMStreamSearch* InSearchOriginator, FRigVMSearchData& OutSearchData);
+	RIGVMEDITOR_API bool ContinueSearchQuery(const class FRigVMStreamSearch* InSearchOriginator, FRigVMSearchData& OutSearchData);
 
 	/**
 	 * This function ensures that the passed in search query ends in a safe manner. The search will no longer be valid to this manager, though it does not destroy any objects.
@@ -590,7 +590,7 @@ public:
 	 *
 	 * @param InSearchOriginator	Identifying search stream to be stopped
 	 */
-	void EnsureSearchQueryEnds(const class FRigVMStreamSearch* InSearchOriginator);
+	RIGVMEDITOR_API void EnsureSearchQueryEnds(const class FRigVMStreamSearch* InSearchOriginator);
 
 	/**
 	 * Query how far along a search thread is
@@ -598,7 +598,7 @@ public:
 	 * @param OutSearchData				Result of the search, if there is any Blueprint search data still available to query
 	 * @return							Percent along the search thread is
 	 */
-	float GetPercentComplete(const class FRigVMStreamSearch* InSearchOriginator) const;
+	RIGVMEDITOR_API float GetPercentComplete(const class FRigVMStreamSearch* InSearchOriginator) const;
 
 	/**
 	 * Query for a single, specific Blueprint's search data.
@@ -607,16 +607,16 @@ public:
 	 * @param bInRebuildSearchData		When TRUE the search data will be freshly collected
 	 * @return							The search data, or invalid (empty) search data if not found
 	 */
-	FRigVMSearchData QuerySingleBlueprint(UBlueprint* InBlueprint, bool bInRebuildSearchData);
+	RIGVMEDITOR_API FRigVMSearchData QuerySingleBlueprint(UBlueprint* InBlueprint, bool bInRebuildSearchData);
 
 	/** Processes the encoded string value in the SearchData into the intermediate format, return true if string and version were valid */
-	bool ProcessEncodedValueForUnloadedBlueprint(FRigVMSearchData& SearchData);
+	RIGVMEDITOR_API bool ProcessEncodedValueForUnloadedBlueprint(FRigVMSearchData& SearchData);
 
 	/** Returns the number of unindexed Blueprints, either due to not having been indexed before, or AR data being out-of-date */
-	int32 GetNumberUnindexedAssets() const;
+	RIGVMEDITOR_API int32 GetNumberUnindexedAssets() const;
 
 	/** Returns the number of uncached assets during an active indexing operation */
-	int32 GetNumberUncachedAssets() const;
+	RIGVMEDITOR_API int32 GetNumberUncachedAssets() const;
 
 	/**
 	 * Starts a task to cache Blueprints at a rate of 1 per tick
@@ -624,12 +624,12 @@ public:
 	 * @param InSourceWidget				The source RigVMFindInBlueprints widget, this widget will be informed when caching is complete
 	 * @param InCachingOptions				Options to configure the caching task
 	 */
-	void CacheAllAssets(TWeakPtr< class SRigVMFindInBlueprints > InSourceWidget, const FRigVMFindInBlueprintCachingOptions& InCachingOptions);
+	RIGVMEDITOR_API void CacheAllAssets(TWeakPtr< class SRigVMFindInBlueprints > InSourceWidget, const FRigVMFindInBlueprintCachingOptions& InCachingOptions);
 
 	/**
 	 * Exports a list of all unindexed assets to Saved/RigVMFindInBlueprints_OutdatedAssetList.txt
 	 */
-	void ExportOutdatedAssetList();
+	RIGVMEDITOR_API void ExportOutdatedAssetList();
 
 	/**
 	 * Starts the actual caching process
@@ -637,19 +637,19 @@ public:
 	 * @param bInSourceControlActive		TRUE if source control is active
 	 * @param bInCheckoutAndSave			TRUE if the system should checkout and save all assets that need to be reindexed
 	 */
-	void OnCacheAllUnindexedAssets(bool bInSourceControlActive, bool bInCheckoutAndSave);
+	RIGVMEDITOR_API void OnCacheAllUnindexedAssets(bool bInSourceControlActive, bool bInCheckoutAndSave);
 
 	/** Stops the caching process where it currently is at, the rest can be continued later */
-	void CancelCacheAll(SRigVMFindInBlueprints* InFindInBlueprintWidget);
+	RIGVMEDITOR_API void CancelCacheAll(SRigVMFindInBlueprints* InFindInBlueprintWidget);
 
 	/** Returns the current index in the caching */
-	int32 GetCurrentCacheIndex() const;
+	RIGVMEDITOR_API int32 GetCurrentCacheIndex() const;
 
 	/** Returns the path of the current Blueprint being cached */
-	FSoftObjectPath GetCurrentCacheBlueprintPath() const;
+	RIGVMEDITOR_API FSoftObjectPath GetCurrentCacheBlueprintPath() const;
 
 	/** Returns the progress complete on the caching */
-	float GetCacheProgress() const;
+	RIGVMEDITOR_API float GetCacheProgress() const;
 
 	/** Returns the list of Blueprint paths that failed to cache */
 	TSet<FSoftObjectPath> GetFailedToCachePathList() const { return FailedToCachePaths; }
@@ -661,26 +661,26 @@ public:
 	bool HasCachingFailed() const { return FailedToCachePaths.Num() > 0; };
 
 	/** Callback to note that Blueprint caching is started */
-	void StartedCachingBlueprints(ERigVMFiBCacheOpType InCacheOpType, ERigVMFiBCacheOpFlags InCacheOpFlags);
+	RIGVMEDITOR_API void StartedCachingBlueprints(ERigVMFiBCacheOpType InCacheOpType, ERigVMFiBCacheOpFlags InCacheOpFlags);
 
 	/**
 	 * Callback to note that Blueprint caching is complete
 	 *
 	 * @param InNumberCached		The number of Blueprints cached, to be chopped off the existing array so the rest (if any) can be finished later
 	 */
-	void FinishedCachingBlueprints(ERigVMFiBCacheOpType InCacheOpType, ERigVMFiBCacheOpFlags InCacheOpFlags, int32 InNumberCached, TSet<FSoftObjectPath>& InFailedToCacheList);
+	RIGVMEDITOR_API void FinishedCachingBlueprints(ERigVMFiBCacheOpType InCacheOpType, ERigVMFiBCacheOpFlags InCacheOpFlags, int32 InNumberCached, TSet<FSoftObjectPath>& InFailedToCacheList);
 
 	/** Returns TRUE if Blueprints are being cached. */
-	bool IsCacheInProgress() const;
+	RIGVMEDITOR_API bool IsCacheInProgress() const;
 
 	/** Returns TRUE if unindexed Blueprints are being cached (since this can block the UI) */
-	bool IsUnindexedCacheInProgress() const;
+	RIGVMEDITOR_API bool IsUnindexedCacheInProgress() const;
 
 	/** Returns TRUE if we're still inside the initial asset discovery and registration stage */
-	bool IsAssetDiscoveryInProgress() const;
+	RIGVMEDITOR_API bool IsAssetDiscoveryInProgress() const;
 
 	/** Returns TRUE if there are one or more active asynchronous search queries */
-	bool IsAsyncSearchQueryInProgress() const;
+	RIGVMEDITOR_API bool IsAsyncSearchQueryInProgress() const;
 
 	/** Returns a weak reference to the widget that initiated the current caching operation */
 	TWeakPtr<SRigVMFindInBlueprints> GetSourceCachingWidget() const { return SourceCachingWidget; }
@@ -696,77 +696,77 @@ public:
 	bool ShouldEnableSearchResultTemplates() const { return !bDisableSearchResultTemplates; }
 
 	/** Find or create the global find results widget */
-	TSharedPtr<SRigVMFindInBlueprints> GetGlobalFindResults();
+	RIGVMEDITOR_API TSharedPtr<SRigVMFindInBlueprints> GetGlobalFindResults();
 
 	/** Enable or disable the global find results tab feature */
-	void EnableGlobalFindResults(bool bEnable);
+	RIGVMEDITOR_API void EnableGlobalFindResults(bool bEnable);
 
 	/** Close any orphaned global find results tabs for a particular tab manager */
-	void CloseOrphanedGlobalFindResultsTabs(TSharedPtr<class FTabManager> TabManager);
+	RIGVMEDITOR_API void CloseOrphanedGlobalFindResultsTabs(TSharedPtr<class FTabManager> TabManager);
 
 	/** Returns true if a global find results tab is currently open */
 	bool IsGlobalFindResultsOpen() const { return GlobalFindResults.Num() > 0; }
 
-	void GlobalFindResultsClosed(const TSharedRef<SRigVMFindInBlueprints>& FindResults);
+	RIGVMEDITOR_API void GlobalFindResultsClosed(const TSharedRef<SRigVMFindInBlueprints>& FindResults);
 
 	/** Dumps the full index cache to the given stream (for debugging purposes) */
-	void DumpCache(FArchive& Ar);
+	RIGVMEDITOR_API void DumpCache(FArchive& Ar);
 
 	/** Randomly initiates a partial rebuild of the index cache (for debugging purposes) */
-	void DoRandomIndexCacheRebuild();
+	RIGVMEDITOR_API void DoRandomIndexCacheRebuild();
 
 public:
 	/** Converts a string of hex characters, previously converted by ConvertFTextToHexString, to an FText. */
-	static FText ConvertHexStringToFText(FString InHexString);
+	static RIGVMEDITOR_API FText ConvertHexStringToFText(FString InHexString);
 
 	/** Serializes an FText to memory and converts the memory into a string of hex characters */
-	static FString ConvertFTextToHexString(FText InValue);
+	static RIGVMEDITOR_API FString ConvertFTextToHexString(FText InValue);
 
 	/** Given a fully constructed Find-in-Blueprint FString of searchable data, will parse and construct a JsonObject */
-	static TSharedPtr< class FJsonObject > ConvertJsonStringToObject(FRigVMSearchDataVersionInfo InVersionInfo, const FString& InJsonString, TMap<int32, FText>& OutFTextLookupTable);
+	static RIGVMEDITOR_API TSharedPtr< class FJsonObject > ConvertJsonStringToObject(FRigVMSearchDataVersionInfo InVersionInfo, const FString& InJsonString, TMap<int32, FText>& OutFTextLookupTable);
 
 	/** Generates a human-readable search index for the given Blueprint (for debugging purposes) */
-	static FString GenerateSearchIndexForDebugging(UBlueprint* InBlueprint);
+	static RIGVMEDITOR_API FString GenerateSearchIndexForDebugging(UBlueprint* InBlueprint);
 
 private:
 	/** Initializes the FiB manager */
-	void Initialize();
+	RIGVMEDITOR_API void Initialize();
 
 	/** Callback hook during pre-garbage collection, pauses all processing searches so that the GC can do it's job */ 
-	void PauseFindInBlueprintSearch();
+	RIGVMEDITOR_API void PauseFindInBlueprintSearch();
 
 	/** Callback hook during post-garbage collection, saves the cache to disk and unpauses all processing searches */
-	void UnpauseFindInBlueprintSearch();
+	RIGVMEDITOR_API void UnpauseFindInBlueprintSearch();
 
 	/** Callback hook from the Asset Registry when an asset is added */
-	void OnAssetAdded(const struct FAssetData& InAssetData);
+	RIGVMEDITOR_API void OnAssetAdded(const struct FAssetData& InAssetData);
 
 	/** Callback hook from the Asset Registry, marks the asset for deletion from the cache */
-	void OnAssetRemoved(const struct FAssetData& InAssetData);
+	RIGVMEDITOR_API void OnAssetRemoved(const struct FAssetData& InAssetData);
 
 	/** Callback hook from the Asset Registry, marks the asset for deletion from the cache */
-	void OnAssetRenamed(const struct FAssetData& InAssetData, const FString& InOldName);
+	RIGVMEDITOR_API void OnAssetRenamed(const struct FAssetData& InAssetData, const FString& InOldName);
 
 	/** Callback hook from the Asset Registry, signals that the initial asset discovery stage has been completed */
-	void OnAssetRegistryFilesLoaded();
+	RIGVMEDITOR_API void OnAssetRegistryFilesLoaded();
 
 	/** Callback hook from the Asset Registry when an asset is loaded */
-	void OnAssetLoaded(class UObject* InAsset);
+	RIGVMEDITOR_API void OnAssetLoaded(class UObject* InAsset);
 
 	/** Callback from Kismet when a Blueprint is unloaded */
-	void OnBlueprintUnloaded(class UBlueprint* InBlueprint);
+	RIGVMEDITOR_API void OnBlueprintUnloaded(class UBlueprint* InBlueprint);
 
 	/** Callback hook from the Reload manager that indicates that a module has been reloaded */
-	void OnReloadComplete(EReloadCompleteReason Reason);
+	RIGVMEDITOR_API void OnReloadComplete(EReloadCompleteReason Reason);
 
 	/** Returns a copy of the search data that's cached at the given index. Will return invalid (empty) search data if the index is out of range */
-	FRigVMSearchData GetSearchDataForIndex(int32 CacheIndex);
+	RIGVMEDITOR_API FRigVMSearchData GetSearchDataForIndex(int32 CacheIndex);
 
 	/** Cleans the cache of any excess data from Blueprints that have been moved, renamed, or deleted. Occurs during post-garbage collection */
-	void CleanCache();
+	RIGVMEDITOR_API void CleanCache();
 
 	/** Builds the cache from all available Blueprint assets that the asset registry has discovered at the time of this function. Occurs on startup */
-	void BuildCache();
+	RIGVMEDITOR_API void BuildCache();
 
 	/**
 	 * Helper to properly add a Blueprint's SearchData to the database
@@ -774,25 +774,25 @@ private:
 	 * @param InSearchData		Data to add to the database
 	 * @return					Index into the SearchArray for looking up the added item
 	 */
-	int32 AddSearchDataToDatabase(FRigVMSearchData InSearchData);
+	RIGVMEDITOR_API int32 AddSearchDataToDatabase(FRigVMSearchData InSearchData);
 
 	/** Removes a Blueprint from being managed by the FiB system by passing in the UBlueprint's path */
-	void RemoveBlueprintByPath(const FSoftObjectPath& InPath);
+	RIGVMEDITOR_API void RemoveBlueprintByPath(const FSoftObjectPath& InPath);
 
 	/** Adds a new search database entry for unloaded asset data */
-	void AddUnloadedBlueprintSearchMetadata(const FAssetData& InAssetData);
+	RIGVMEDITOR_API void AddUnloadedBlueprintSearchMetadata(const FAssetData& InAssetData);
 
 	/** Begins the process of extracting FiB data from an unloaded asset */
-	void ExtractUnloadedFiBData(const FAssetData& InAssetData, FString* InFiBData, FName InKeyForFiBData, ERigVMFiBVersion InFiBDataVersion);
+	RIGVMEDITOR_API void ExtractUnloadedFiBData(const FAssetData& InAssetData, FString* InFiBData, FName InKeyForFiBData, ERigVMFiBVersion InFiBDataVersion);
 
 	/** Determines the global find results tab label */
-	FText GetGlobalFindResultsTabLabel(int32 TabIdx);
+	RIGVMEDITOR_API FText GetGlobalFindResultsTabLabel(int32 TabIdx);
 
 	/** Handler for a request to spawn a new global find results tab */
-	TSharedRef<SDockTab> SpawnGlobalFindResultsTab(const FSpawnTabArgs& SpawnTabArgs, int32 TabIdx);
+	RIGVMEDITOR_API TSharedRef<SDockTab> SpawnGlobalFindResultsTab(const FSpawnTabArgs& SpawnTabArgs, int32 TabIdx);
 
 	/** Creates and opens a new global find results tab */
-	TSharedPtr<SRigVMFindInBlueprints> OpenGlobalFindResultsTab();
+	RIGVMEDITOR_API TSharedPtr<SRigVMFindInBlueprints> OpenGlobalFindResultsTab();
 
 protected:
 	/** Contains info about an active search query */
@@ -815,13 +815,13 @@ protected:
 	typedef TSharedPtr<FActiveSearchQuery, ESPMode::ThreadSafe> FActiveSearchQueryPtr;
 
 	/** Thread-safe access to the active search query that's mapped to the given stream search */
-	FActiveSearchQueryPtr FindSearchQuery(const class FRigVMStreamSearch* InSearchOriginator) const;
+	RIGVMEDITOR_API FActiveSearchQueryPtr FindSearchQuery(const class FRigVMStreamSearch* InSearchOriginator) const;
 
 	/** Returns the next pending search data for the given query and advances the index to the next entry */
-	FRigVMSearchData GetNextSearchDataForQuery(const FRigVMStreamSearch* InSearchOriginator, FActiveSearchQueryPtr InSearchQueryPtr, bool bCheckDeferredList);
+	RIGVMEDITOR_API FRigVMSearchData GetNextSearchDataForQuery(const FRigVMStreamSearch* InSearchOriginator, FActiveSearchQueryPtr InSearchQueryPtr, bool bCheckDeferredList);
 
 	/** If searches are paused, blocks the calling thread until searching is resumed */
-	void BlockSearchQueryIfPaused();
+	RIGVMEDITOR_API void BlockSearchQueryIfPaused();
 
 private:
 	/** Maps the Blueprint paths to their index in the SearchArray */
@@ -913,7 +913,7 @@ private:
 	bool bDisableImmediateAssetDiscovery;
 };
 
-struct RIGVMEDITOR_API FRigVMDisableGatheringDataOnScope
+struct FRigVMDisableGatheringDataOnScope
 {
 	bool bOriginallyEnabled;
 	FRigVMDisableGatheringDataOnScope() : bOriginallyEnabled(FRigVMFindInBlueprintSearchManager::Get().IsGatheringDataEnabled())
@@ -925,4 +925,4 @@ struct RIGVMEDITOR_API FRigVMDisableGatheringDataOnScope
 		FRigVMFindInBlueprintSearchManager::Get().EnableGatheringData(bOriginallyEnabled);
 	}
 };
-#endif
\ No newline at end of file
+#endif
