Updated headers using UnrealCodeFixup to make sure dllstorage is on methods/staticvars instead of types. Part 4/n

diff --git a/Engine/Plugins/Runtime/RigVM/Source/RigVMEditor/Private/Editor/Kismet/RigVMFindInBlueprints.h b/Engine/Plugins/Runtime/RigVM/Source/RigVMEditor/Private/Editor/Kismet/RigVMFindInBlueprints.h
--- a/Engine/Plugins/Runtime/RigVM/Source/RigVMEditor/Private/Editor/Kismet/RigVMFindInBlueprints.h
+++ b/Engine/Plugins/Runtime/RigVM/Source/RigVMEditor/Private/Editor/Kismet/RigVMFindInBlueprints.h
@@ -285,7 +285,7 @@ enum class ERigVMFiBIndexAllPermission
 #if !WITH_RIGVMLEGACYEDITOR
 
 /*Widget for searching for (functions/events) across all blueprints or just a single blueprint */
-class RIGVMEDITOR_API SRigVMFindInBlueprints: public SCompoundWidget
+class SRigVMFindInBlueprints: public SCompoundWidget
 {
 public:
 	SLATE_BEGIN_ARGS( SRigVMFindInBlueprints )
@@ -300,11 +300,11 @@ public:
 		SLATE_ARGUMENT(TSharedPtr<SDockTab>, ContainingTab)
 	SLATE_END_ARGS()
 
-	void Construct(const FArguments& InArgs, TSharedPtr<class FRigVMEditorBase> InBlueprintEditor = nullptr);
-	~SRigVMFindInBlueprints();
+	RIGVMEDITOR_API void Construct(const FArguments& InArgs, TSharedPtr<class FRigVMEditorBase> InBlueprintEditor = nullptr);
+	RIGVMEDITOR_API ~SRigVMFindInBlueprints();
 
 	/** Focuses this widget's search box, and changes the mode as well, and optionally the search terms */
-	void FocusForUse(bool bSetFindWithinBlueprint, FString NewSearchTerms = FString(), bool bSelectFirstResult = false);
+	RIGVMEDITOR_API void FocusForUse(bool bSetFindWithinBlueprint, FString NewSearchTerms = FString(), bool bSelectFirstResult = false);
 
 	/**
 	 * Submits a search query
@@ -314,26 +314,26 @@ public:
 	 * @param InSearchOptions						Optional search parameters.
 	 * @param InOnSearchComplete					Callback when the search is complete, passing the filtered imaginary data (if any).
 	 */
-	void MakeSearchQuery(FString InSearchString, bool bInIsFindWithinBlueprint, const FRigVMStreamSearchOptions& InSearchOptions = FRigVMStreamSearchOptions(), FRigVMOnSearchComplete InOnSearchComplete = FRigVMOnSearchComplete());
+	RIGVMEDITOR_API void MakeSearchQuery(FString InSearchString, bool bInIsFindWithinBlueprint, const FRigVMStreamSearchOptions& InSearchOptions = FRigVMStreamSearchOptions(), FRigVMOnSearchComplete InOnSearchComplete = FRigVMOnSearchComplete());
 
 	/** Called when caching Blueprints is started */
-	void OnCacheStarted(ERigVMFiBCacheOpType InOpType, ERigVMFiBCacheOpFlags InOpFlags);
+	RIGVMEDITOR_API void OnCacheStarted(ERigVMFiBCacheOpType InOpType, ERigVMFiBCacheOpFlags InOpFlags);
 	
 	/** Called when caching Blueprints is complete */
-	void OnCacheComplete(ERigVMFiBCacheOpType InOpType, ERigVMFiBCacheOpFlags InOpFlags);
+	RIGVMEDITOR_API void OnCacheComplete(ERigVMFiBCacheOpType InOpType, ERigVMFiBCacheOpFlags InOpFlags);
 
 	/**
 	 * Asynchronously caches all Blueprints below a specified version.
 	 *
 	 * @param InOptions								Options to configure the caching task
 	 */
-	void CacheAllBlueprints(const FRigVMFindInBlueprintCachingOptions& InOptions);
+	RIGVMEDITOR_API void CacheAllBlueprints(const FRigVMFindInBlueprintCachingOptions& InOptions);
 
 	/** If this is a global find results widget, returns the host tab's unique ID. Otherwise, returns NAME_None. */
-	FName GetHostTabId() const;
+	RIGVMEDITOR_API FName GetHostTabId() const;
 
 	/** If this is a global find results widget, ask the host tab to close */
-	void CloseHostTab();
+	RIGVMEDITOR_API void CloseHostTab();
 
 	/** Determines if this context does not accept syncing from an external source */
 	bool IsLocked() const
@@ -342,113 +342,113 @@ public:
 	}
 
 	/** Determines whether a search query is actively in progress */
-	bool IsSearchInProgress() const;
+	RIGVMEDITOR_API bool IsSearchInProgress() const;
 
 	/** SWidget overrides */
-	virtual FReply OnKeyDown(const FGeometry& MyGeometry, const FKeyEvent& InKeyEvent) override;
+	RIGVMEDITOR_API virtual FReply OnKeyDown(const FGeometry& MyGeometry, const FKeyEvent& InKeyEvent) override;
 
 	/** Clears the currently visible results */
-	void ClearResults();
+	RIGVMEDITOR_API void ClearResults();
 
 private:
 	/** Processes results of the ongoing async stream search */
-	EActiveTimerReturnType UpdateSearchResults( double InCurrentTime, float InDeltaTime );
+	RIGVMEDITOR_API EActiveTimerReturnType UpdateSearchResults( double InCurrentTime, float InDeltaTime );
 
 	/** Register any Find-in-Blueprint commands */
-	void RegisterCommands();
+	RIGVMEDITOR_API void RegisterCommands();
 
 	/*Called when user changes the text they are searching for */
-	void OnSearchTextChanged(const FText& Text);
+	RIGVMEDITOR_API void OnSearchTextChanged(const FText& Text);
 
 	/*Called when user changes commits text to the search box */
-	void OnSearchTextCommitted(const FText& Text, ETextCommit::Type CommitType);
+	RIGVMEDITOR_API void OnSearchTextCommitted(const FText& Text, ETextCommit::Type CommitType);
 
 	/* Get the children of a row */
-	void OnGetChildren( FRigVMSearchResult InItem, TArray< FRigVMSearchResult >& OutChildren );
+	RIGVMEDITOR_API void OnGetChildren( FRigVMSearchResult InItem, TArray< FRigVMSearchResult >& OutChildren );
 
 	/* Called when user double clicks on a new result */
-	void OnTreeSelectionDoubleClicked( FRigVMSearchResult Item );
+	RIGVMEDITOR_API void OnTreeSelectionDoubleClicked( FRigVMSearchResult Item );
 
 	/* Called when a new row is being generated */
-	TSharedRef<ITableRow> OnGenerateRow(FRigVMSearchResult InItem, const TSharedRef<STableViewBase>& OwnerTable);
+	RIGVMEDITOR_API TSharedRef<ITableRow> OnGenerateRow(FRigVMSearchResult InItem, const TSharedRef<STableViewBase>& OwnerTable);
 	
 	/** Launches a thread for streaming more content into the results widget */
-	void LaunchStreamThread(const FString& InSearchValue, const FRigVMStreamSearchOptions& InSearchOptions = FRigVMStreamSearchOptions(), FRigVMOnSearchComplete InOnSearchComplete = FRigVMOnSearchComplete());
+	RIGVMEDITOR_API void LaunchStreamThread(const FString& InSearchValue, const FRigVMStreamSearchOptions& InSearchOptions = FRigVMStreamSearchOptions(), FRigVMOnSearchComplete InOnSearchComplete = FRigVMOnSearchComplete());
 
 	/** Returns the percent complete on the search for the progress bar */
-	TOptional<float> GetPercentCompleteSearch() const;
+	RIGVMEDITOR_API TOptional<float> GetPercentCompleteSearch() const;
 
 	/** Returns the search bar visiblity for the given widget */
-	EVisibility GetSearchBarWidgetVisiblity(ERigVMFiBSearchBarWidget InSearchBarWidget) const;
+	RIGVMEDITOR_API EVisibility GetSearchBarWidgetVisiblity(ERigVMFiBSearchBarWidget InSearchBarWidget) const;
 
 	/** Adds the "cache" bar at the bottom of the Find-in-Blueprints widget, to notify the user that the search is incomplete */
-	void ConditionallyAddCacheBar();
+	RIGVMEDITOR_API void ConditionallyAddCacheBar();
 
 	/** Callback to remove the "cache" bar when a button is pressed */
-	FReply OnRemoveCacheBar();
+	RIGVMEDITOR_API FReply OnRemoveCacheBar();
 
 	/** Callback to return the cache bar's display text, informing the user of the situation */
-	FText GetCacheBarStatusText() const;
+	RIGVMEDITOR_API FText GetCacheBarStatusText() const;
 
 	/** Callback to return the current asset name during a cache operation */
-	FText GetCacheBarCurrentAssetName() const;
+	RIGVMEDITOR_API FText GetCacheBarCurrentAssetName() const;
 
 	/** Whether user is allowed to initiate loading and indexing all blueprints with out-of-date metadata */
-	bool CanCacheAllUnindexedBlueprints() const;
+	RIGVMEDITOR_API bool CanCacheAllUnindexedBlueprints() const;
 
 	/** Callback to cache all unindexed Blueprints */
-	FReply OnCacheAllUnindexedBlueprints();
+	RIGVMEDITOR_API FReply OnCacheAllUnindexedBlueprints();
 
 	/** Callback to export a list of all blueprints that need reindexing */
-	FReply OnExportUnindexedAssetList();
+	RIGVMEDITOR_API FReply OnExportUnindexedAssetList();
 
 	/** Callback to cache all Blueprints according to the given options */
-	FReply OnCacheAllBlueprints(const FRigVMFindInBlueprintCachingOptions& InOptions);
+	RIGVMEDITOR_API FReply OnCacheAllBlueprints(const FRigVMFindInBlueprintCachingOptions& InOptions);
 
 	/** Callback to cancel the caching process */
-	FReply OnCancelCacheAll();
+	RIGVMEDITOR_API FReply OnCancelCacheAll();
 
 	/** Retrieves the current index of the Blueprint caching process */
-	int32 GetCurrentCacheIndex() const;
+	RIGVMEDITOR_API int32 GetCurrentCacheIndex() const;
 
 	/** Gets the percent complete of the caching process */
-	TOptional<float> GetPercentCompleteCache() const;
+	RIGVMEDITOR_API TOptional<float> GetPercentCompleteCache() const;
 
 	/** Returns the caching bar's visibility, it goes invisible when there is nothing to be cached. The next search will remove this bar or make it visible again */
-	EVisibility GetCacheBarVisibility() const;
+	RIGVMEDITOR_API EVisibility GetCacheBarVisibility() const;
 
 	/** Returns the cache bar visibility for the given widget */
-	EVisibility GetCacheBarWidgetVisibility(ERigVMFiBCacheBarWidget InCacheBarWidget) const;
+	RIGVMEDITOR_API EVisibility GetCacheBarWidgetVisibility(ERigVMFiBCacheBarWidget InCacheBarWidget) const;
 
 	/** Returns TRUE if Blueprint caching is in progress */
-	bool IsCacheInProgress() const;
+	RIGVMEDITOR_API bool IsCacheInProgress() const;
 
 	/** Returns the BG image used for the cache bar */
-	const FSlateBrush* GetCacheBarImage() const;
+	RIGVMEDITOR_API const FSlateBrush* GetCacheBarImage() const;
 
 	/** Callback to build the context menu when right clicking in the tree */
-	TSharedPtr<SWidget> OnContextMenuOpening();
+	RIGVMEDITOR_API TSharedPtr<SWidget> OnContextMenuOpening();
 
 	/** Helper function to select all items */
-	void SelectAllItemsHelper(FRigVMSearchResult InItemToSelect);
+	RIGVMEDITOR_API void SelectAllItemsHelper(FRigVMSearchResult InItemToSelect);
 
 	/** Callback when user attempts to select all items in the search results */
-	void OnSelectAllAction();
+	RIGVMEDITOR_API void OnSelectAllAction();
 
 	/** Callback when user attempts to copy their selection in the Find-in-Blueprints */
-	void OnCopyAction();
+	RIGVMEDITOR_API void OnCopyAction();
 
 	/** Called when the user clicks the global find results button */
-	FReply OnOpenGlobalFindResults();
+	RIGVMEDITOR_API FReply OnOpenGlobalFindResults();
 
 	/** Called when the host tab is closed (if valid) */
-	void OnHostTabClosed(TSharedRef<SDockTab> DockTab);
+	RIGVMEDITOR_API void OnHostTabClosed(TSharedRef<SDockTab> DockTab);
 
 	/** Called when the lock button is clicked in a global find results tab */
-	FReply OnLockButtonClicked();
+	RIGVMEDITOR_API FReply OnLockButtonClicked();
 
 	/** Returns the image used for the lock button in a global find results tab */
-	const FSlateBrush* OnGetLockButtonImage() const;
+	RIGVMEDITOR_API const FSlateBrush* OnGetLockButtonImage() const;
 
 private:
 	/** Pointer back to the blueprint editor that owns us */
@@ -520,4 +520,4 @@ private:
 	/** True if cache bar should remain visible after a caching operation has ended */
 	bool bKeepCacheBarProgressVisible;
 };
-#endif
\ No newline at end of file
+#endif
diff --git a/Engine/Plugins/Runtime/RigVM/Source/RigVMEditor/Public/EdGraph/NodeSpawners/RigVMEdGraphEnumNodeSpawner.h b/Engine/Plugins/Runtime/RigVM/Source/RigVMEditor/Public/EdGraph/NodeSpawners/RigVMEdGraphEnumNodeSpawner.h
--- a/Engine/Plugins/Runtime/RigVM/Source/RigVMEditor/Public/EdGraph/NodeSpawners/RigVMEdGraphEnumNodeSpawner.h
+++ b/Engine/Plugins/Runtime/RigVM/Source/RigVMEditor/Public/EdGraph/NodeSpawners/RigVMEdGraphEnumNodeSpawner.h
@@ -12,10 +12,12 @@
 #include "EdGraph/NodeSpawners/RigVMEdGraphNodeSpawner.h"
 #include "RigVMEdGraphEnumNodeSpawner.generated.h"
 
+#define UE_API RIGVMEDITOR_API
+
 class URigVMEdGraphNode;
 
-UCLASS(Transient)
-class RIGVMEDITOR_API URigVMEdGraphEnumNodeSpawner : public URigVMEdGraphNodeSpawner
+UCLASS(MinimalAPI, Transient)
+class URigVMEdGraphEnumNodeSpawner : public URigVMEdGraphNodeSpawner
 {
 	GENERATED_BODY()
 
@@ -26,12 +28,12 @@ public:
 	 * 
 	 * @return A newly allocated instance of this class.
 	 */
-	static URigVMEdGraphEnumNodeSpawner* CreateForEnum(UEnum* InEnum, const FText& InMenuDesc, const FText& InCategory, const FText& InTooltip);
+	static UE_API URigVMEdGraphEnumNodeSpawner* CreateForEnum(UEnum* InEnum, const FText& InMenuDesc, const FText& InCategory, const FText& InTooltip);
 
 	// UBlueprintNodeSpawner interface
-	virtual FBlueprintNodeSignature GetSpawnerSignature() const override;
-	virtual FBlueprintActionUiSpec GetUiSpec(FBlueprintActionContext const& Context, FBindingSet const& Bindings) const override;
-	virtual UEdGraphNode* Invoke(UEdGraph* ParentGraph, FBindingSet const& Bindings, FVector2D const Location) const override;
+	UE_API virtual FBlueprintNodeSignature GetSpawnerSignature() const override;
+	UE_API virtual FBlueprintActionUiSpec GetUiSpec(FBlueprintActionContext const& Context, FBindingSet const& Bindings) const override;
+	UE_API virtual UEdGraphNode* Invoke(UEdGraph* ParentGraph, FBindingSet const& Bindings, FVector2D const Location) const override;
 	// End UBlueprintNodeSpawner interface
 
 private:
@@ -39,3 +41,5 @@ private:
 	UEnum* Enum;
 	friend class UEngineTestControlRig;
 };
+
+#undef UE_API
diff --git a/Engine/Plugins/Runtime/RigVM/Source/RigVMEditor/Public/EdGraph/NodeSpawners/RigVMEdGraphFunctionRefNodeSpawner.h b/Engine/Plugins/Runtime/RigVM/Source/RigVMEditor/Public/EdGraph/NodeSpawners/RigVMEdGraphFunctionRefNodeSpawner.h
--- a/Engine/Plugins/Runtime/RigVM/Source/RigVMEditor/Public/EdGraph/NodeSpawners/RigVMEdGraphFunctionRefNodeSpawner.h
+++ b/Engine/Plugins/Runtime/RigVM/Source/RigVMEditor/Public/EdGraph/NodeSpawners/RigVMEdGraphFunctionRefNodeSpawner.h
@@ -16,10 +16,12 @@
 #include "RigVMModel/Nodes/RigVMUnitNode.h"
 #include "RigVMEdGraphFunctionRefNodeSpawner.generated.h"
 
+#define UE_API RIGVMEDITOR_API
+
 class URigVMEdGraphNode;
 
-UCLASS(Transient)
-class RIGVMEDITOR_API URigVMEdGraphFunctionRefNodeSpawner : public URigVMEdGraphNodeSpawner
+UCLASS(MinimalAPI, Transient)
+class URigVMEdGraphFunctionRefNodeSpawner : public URigVMEdGraphNodeSpawner
 {
 	GENERATED_BODY()
 
@@ -30,22 +32,22 @@ public:
 	 * 
 	 * @return A newly allocated instance of this class.
 	 */
-	static URigVMEdGraphFunctionRefNodeSpawner* CreateFromFunction(URigVMLibraryNode* InFunction);
+	static UE_API URigVMEdGraphFunctionRefNodeSpawner* CreateFromFunction(URigVMLibraryNode* InFunction);
 
 	/**
 	 * Creates a new URigVMEdGraphFunctionRefNodeSpawner, charged with spawning a function reference
 	 * 
 	 * @return A newly allocated instance of this class.
 	 */
-	static URigVMEdGraphFunctionRefNodeSpawner* CreateFromAssetData(const FAssetData& InAssetData, const FRigVMGraphFunctionHeader& InPublicFunction);
-	static URigVMEdGraphFunctionRefNodeSpawner* CreateFromAssetData(const FAssetData& InAssetData, const FRigVMOldPublicFunctionData& InPublicFunction);
+	static UE_API URigVMEdGraphFunctionRefNodeSpawner* CreateFromAssetData(const FAssetData& InAssetData, const FRigVMGraphFunctionHeader& InPublicFunction);
+	static UE_API URigVMEdGraphFunctionRefNodeSpawner* CreateFromAssetData(const FAssetData& InAssetData, const FRigVMOldPublicFunctionData& InPublicFunction);
 
 	// UBlueprintNodeSpawner interface
-	virtual void Prime() override;
-	virtual FBlueprintNodeSignature GetSpawnerSignature() const override;
-	virtual FBlueprintActionUiSpec GetUiSpec(FBlueprintActionContext const& Context, FBindingSet const& Bindings) const override;
-	virtual UEdGraphNode* Invoke(UEdGraph* ParentGraph, FBindingSet const& Bindings, FVector2D const Location) const override;
-	virtual bool IsTemplateNodeFilteredOut(FBlueprintActionFilter const& Filter) const override;
+	UE_API virtual void Prime() override;
+	UE_API virtual FBlueprintNodeSignature GetSpawnerSignature() const override;
+	UE_API virtual FBlueprintActionUiSpec GetUiSpec(FBlueprintActionContext const& Context, FBindingSet const& Bindings) const override;
+	UE_API virtual UEdGraphNode* Invoke(UEdGraph* ParentGraph, FBindingSet const& Bindings, FVector2D const Location) const override;
+	UE_API virtual bool IsTemplateNodeFilteredOut(FBlueprintActionFilter const& Filter) const override;
 	// End UBlueprintNodeSpawner interface
 
 private:
@@ -61,8 +63,10 @@ private:
 	UPROPERTY()
 	FSoftObjectPath AssetPath;
 
-	static URigVMEdGraphNode* SpawnNode(UEdGraph* ParentGraph, UBlueprint* Blueprint, FRigVMGraphFunctionHeader& InFunction, FVector2D const Location);
+	static UE_API URigVMEdGraphNode* SpawnNode(UEdGraph* ParentGraph, UBlueprint* Blueprint, FRigVMGraphFunctionHeader& InFunction, FVector2D const Location);
 
 	friend class UEngineTestControlRig;
 	friend class FRigVMEditorBase;
 };
+
+#undef UE_API
diff --git a/Engine/Plugins/Runtime/RigVM/Source/RigVMEditor/Public/EdGraph/NodeSpawners/RigVMEdGraphInvokeEntryNodeSpawner.h b/Engine/Plugins/Runtime/RigVM/Source/RigVMEditor/Public/EdGraph/NodeSpawners/RigVMEdGraphInvokeEntryNodeSpawner.h
--- a/Engine/Plugins/Runtime/RigVM/Source/RigVMEditor/Public/EdGraph/NodeSpawners/RigVMEdGraphInvokeEntryNodeSpawner.h
+++ b/Engine/Plugins/Runtime/RigVM/Source/RigVMEditor/Public/EdGraph/NodeSpawners/RigVMEdGraphInvokeEntryNodeSpawner.h
@@ -7,24 +7,26 @@
 #include "RigVMModel/RigVMGraph.h"
 #include "RigVMEdGraphInvokeEntryNodeSpawner.generated.h"
 
+#define UE_API RIGVMEDITOR_API
+
 class URigVMEdGraphNode;
 class URigVMBlueprint;
 
-UCLASS(Transient)
-class RIGVMEDITOR_API URigVMEdGraphInvokeEntryNodeSpawner : public URigVMEdGraphNodeSpawner
+UCLASS(MinimalAPI, Transient)
+class URigVMEdGraphInvokeEntryNodeSpawner : public URigVMEdGraphNodeSpawner
 {
 	GENERATED_BODY()
 
 public:
 
-	static URigVMEdGraphInvokeEntryNodeSpawner* CreateForEntry(URigVMBlueprint* InBlueprint, const FName& InEntryName, const FText& InMenuDesc, const FText& InCategory, const FText& InTooltip);
+	static UE_API URigVMEdGraphInvokeEntryNodeSpawner* CreateForEntry(URigVMBlueprint* InBlueprint, const FName& InEntryName, const FText& InMenuDesc, const FText& InCategory, const FText& InTooltip);
 
 	// UBlueprintNodeSpawner interface
-	virtual void Prime() override;
-	virtual FBlueprintNodeSignature GetSpawnerSignature() const override;
-	virtual FBlueprintActionUiSpec GetUiSpec(FBlueprintActionContext const& Context, FBindingSet const& Bindings) const override;
-	virtual UEdGraphNode* Invoke(UEdGraph* ParentGraph, FBindingSet const& Bindings, FVector2D const Location) const override;
-	virtual bool IsTemplateNodeFilteredOut(FBlueprintActionFilter const& Filter) const override;
+	UE_API virtual void Prime() override;
+	UE_API virtual FBlueprintNodeSignature GetSpawnerSignature() const override;
+	UE_API virtual FBlueprintActionUiSpec GetUiSpec(FBlueprintActionContext const& Context, FBindingSet const& Bindings) const override;
+	UE_API virtual UEdGraphNode* Invoke(UEdGraph* ParentGraph, FBindingSet const& Bindings, FVector2D const Location) const override;
+	UE_API virtual bool IsTemplateNodeFilteredOut(FBlueprintActionFilter const& Filter) const override;
 	// End UBlueprintNodeSpawner interface
 
 private:
@@ -36,3 +38,5 @@ private:
 
 	friend class UEngineTestControlRig;
 };
+
+#undef UE_API
diff --git a/Engine/Plugins/Runtime/RigVM/Source/RigVMEditor/Public/EdGraph/NodeSpawners/RigVMEdGraphNodeSpawner.h b/Engine/Plugins/Runtime/RigVM/Source/RigVMEditor/Public/EdGraph/NodeSpawners/RigVMEdGraphNodeSpawner.h
--- a/Engine/Plugins/Runtime/RigVM/Source/RigVMEditor/Public/EdGraph/NodeSpawners/RigVMEdGraphNodeSpawner.h
+++ b/Engine/Plugins/Runtime/RigVM/Source/RigVMEditor/Public/EdGraph/NodeSpawners/RigVMEdGraphNodeSpawner.h
@@ -9,18 +9,20 @@
 #include "BlueprintActionFilter.h"
 #include "RigVMEdGraphNodeSpawner.generated.h"
 
+#define UE_API RIGVMEDITOR_API
+
 class URigVMEdGraphNode;
 
-UCLASS(Transient)
-class RIGVMEDITOR_API URigVMEdGraphNodeSpawner : public UBlueprintNodeSpawner
+UCLASS(MinimalAPI, Transient)
+class URigVMEdGraphNodeSpawner : public UBlueprintNodeSpawner
 {
 	GENERATED_BODY()
 
 public:
 
-	virtual bool IsTemplateNodeFilteredOut(FBlueprintActionFilter const& Filter) const override;
+	UE_API virtual bool IsTemplateNodeFilteredOut(FBlueprintActionFilter const& Filter) const override;
 
-	void SetRelatedBlueprintClass(TSubclassOf<URigVMBlueprint> InClass);
+	UE_API void SetRelatedBlueprintClass(TSubclassOf<URigVMBlueprint> InClass);
 
 	struct FPinInfo
 	{
@@ -44,9 +46,11 @@ public:
 		{}
 	};
 	
-	static URigVMEdGraphNode* SpawnTemplateNode(UEdGraph* InParentGraph, const TArray<FPinInfo>& InPins, const FName& InNodeName = NAME_None);
+	static UE_API URigVMEdGraphNode* SpawnTemplateNode(UEdGraph* InParentGraph, const TArray<FPinInfo>& InPins, const FName& InNodeName = NAME_None);
 	
 private:
 
 	TSubclassOf<URigVMBlueprint> RelatedBlueprintClass;
 };
+
+#undef UE_API
diff --git a/Engine/Plugins/Runtime/RigVM/Source/RigVMEditor/Public/EdGraph/NodeSpawners/RigVMEdGraphTemplateNodeSpawner.h b/Engine/Plugins/Runtime/RigVM/Source/RigVMEditor/Public/EdGraph/NodeSpawners/RigVMEdGraphTemplateNodeSpawner.h
--- a/Engine/Plugins/Runtime/RigVM/Source/RigVMEditor/Public/EdGraph/NodeSpawners/RigVMEdGraphTemplateNodeSpawner.h
+++ b/Engine/Plugins/Runtime/RigVM/Source/RigVMEditor/Public/EdGraph/NodeSpawners/RigVMEdGraphTemplateNodeSpawner.h
@@ -15,10 +15,12 @@
 #include "RigVMModel/Nodes/RigVMTemplateNode.h"
 #include "RigVMEdGraphTemplateNodeSpawner.generated.h"
 
+#define UE_API RIGVMEDITOR_API
+
 class URigVMEdGraphNode;
 
-UCLASS(Transient)
-class RIGVMEDITOR_API URigVMEdGraphTemplateNodeSpawner : public URigVMEdGraphNodeSpawner
+UCLASS(MinimalAPI, Transient)
+class URigVMEdGraphTemplateNodeSpawner : public URigVMEdGraphNodeSpawner
 {
 	GENERATED_BODY()
 
@@ -29,14 +31,14 @@ public:
 	 * 
 	 * @return A newly allocated instance of this class.
 	 */
-	static URigVMEdGraphTemplateNodeSpawner* CreateFromNotation(const FName& InNotation, const FText& InMenuDesc, const FText& InCategory, const FText& InTooltip);
+	static UE_API URigVMEdGraphTemplateNodeSpawner* CreateFromNotation(const FName& InNotation, const FText& InMenuDesc, const FText& InCategory, const FText& InTooltip);
 
 	// UBlueprintNodeSpawner interface
-	virtual void Prime() override;
-	virtual FBlueprintNodeSignature GetSpawnerSignature() const override;
-	bool IsTemplateNodeFilteredOut(FBlueprintActionFilter const& Filter) const override;
-	virtual FBlueprintActionUiSpec GetUiSpec(FBlueprintActionContext const& Context, FBindingSet const& Bindings) const override;
-	virtual UEdGraphNode* Invoke(UEdGraph* ParentGraph, FBindingSet const& Bindings, FVector2D const Location) const override;
+	UE_API virtual void Prime() override;
+	UE_API virtual FBlueprintNodeSignature GetSpawnerSignature() const override;
+	UE_API bool IsTemplateNodeFilteredOut(FBlueprintActionFilter const& Filter) const override;
+	UE_API virtual FBlueprintActionUiSpec GetUiSpec(FBlueprintActionContext const& Context, FBindingSet const& Bindings) const override;
+	UE_API virtual UEdGraphNode* Invoke(UEdGraph* ParentGraph, FBindingSet const& Bindings, FVector2D const Location) const override;
 	// End UBlueprintNodeSpawner interface
 
 private:
@@ -46,8 +48,10 @@ private:
 
 	const FRigVMTemplate* Template = nullptr;
 
-	static URigVMEdGraphNode* SpawnNode(UEdGraph* ParentGraph, UBlueprint* Blueprint, const FRigVMTemplate* Template, FVector2D const Location);
+	static UE_API URigVMEdGraphNode* SpawnNode(UEdGraph* ParentGraph, UBlueprint* Blueprint, const FRigVMTemplate* Template, FVector2D const Location);
 
 	friend class UEngineTestControlRig;
 	friend class FRigVMEditorBase;
 };
+
+#undef UE_API
diff --git a/Engine/Plugins/Runtime/RigVM/Source/RigVMEditor/Public/EdGraph/NodeSpawners/RigVMEdGraphUnitNodeSpawner.h b/Engine/Plugins/Runtime/RigVM/Source/RigVMEditor/Public/EdGraph/NodeSpawners/RigVMEdGraphUnitNodeSpawner.h
--- a/Engine/Plugins/Runtime/RigVM/Source/RigVMEditor/Public/EdGraph/NodeSpawners/RigVMEdGraphUnitNodeSpawner.h
+++ b/Engine/Plugins/Runtime/RigVM/Source/RigVMEditor/Public/EdGraph/NodeSpawners/RigVMEdGraphUnitNodeSpawner.h
@@ -16,10 +16,12 @@
 #include "RigVMModel/Nodes/RigVMUnitNode.h"
 #include "RigVMEdGraphUnitNodeSpawner.generated.h"
 
+#define UE_API RIGVMEDITOR_API
+
 class URigVMEdGraphNode;
 
-UCLASS(Transient)
-class RIGVMEDITOR_API URigVMEdGraphUnitNodeSpawner : public URigVMEdGraphNodeSpawner
+UCLASS(MinimalAPI, Transient)
+class URigVMEdGraphUnitNodeSpawner : public URigVMEdGraphNodeSpawner
 {
 	GENERATED_BODY()
 
@@ -30,14 +32,14 @@ public:
 	 * 
 	 * @return A newly allocated instance of this class.
 	 */
-	static URigVMEdGraphUnitNodeSpawner* CreateFromStruct(UScriptStruct* InStruct, const FName& InMethodName, const FText& InMenuDesc, const FText& InCategory, const FText& InTooltip);
+	static UE_API URigVMEdGraphUnitNodeSpawner* CreateFromStruct(UScriptStruct* InStruct, const FName& InMethodName, const FText& InMenuDesc, const FText& InCategory, const FText& InTooltip);
 
 	// UBlueprintNodeSpawner interface
-	virtual void Prime() override;
-	virtual FBlueprintNodeSignature GetSpawnerSignature() const override;
-	virtual FBlueprintActionUiSpec GetUiSpec(FBlueprintActionContext const& Context, FBindingSet const& Bindings) const override;
-	virtual UEdGraphNode* Invoke(UEdGraph* ParentGraph, FBindingSet const& Bindings, FVector2D const Location) const override;
-	static void HookupMutableNode(URigVMNode* InModelNode, URigVMBlueprint* InRigBlueprint);
+	UE_API virtual void Prime() override;
+	UE_API virtual FBlueprintNodeSignature GetSpawnerSignature() const override;
+	UE_API virtual FBlueprintActionUiSpec GetUiSpec(FBlueprintActionContext const& Context, FBindingSet const& Bindings) const override;
+	UE_API virtual UEdGraphNode* Invoke(UEdGraph* ParentGraph, FBindingSet const& Bindings, FVector2D const Location) const override;
+	static UE_API void HookupMutableNode(URigVMNode* InModelNode, URigVMBlueprint* InRigBlueprint);
 	// End UBlueprintNodeSpawner interface
 
 private:
@@ -48,8 +50,10 @@ private:
 	UPROPERTY(Transient)
 	FName MethodName;
 
-	static URigVMEdGraphNode* SpawnNode(UEdGraph* ParentGraph, UBlueprint* Blueprint, UScriptStruct* StructTemplate, const FName& InMethodName, FVector2D const Location);
+	static UE_API URigVMEdGraphNode* SpawnNode(UEdGraph* ParentGraph, UBlueprint* Blueprint, UScriptStruct* StructTemplate, const FName& InMethodName, FVector2D const Location);
 
 	friend class UEngineTestControlRig;
 	friend class FRigVMEditorBase;
 };
+
+#undef UE_API
diff --git a/Engine/Plugins/Runtime/RigVM/Source/RigVMEditor/Public/EdGraph/NodeSpawners/RigVMEdGraphVariableNodeSpawner.h b/Engine/Plugins/Runtime/RigVM/Source/RigVMEditor/Public/EdGraph/NodeSpawners/RigVMEdGraphVariableNodeSpawner.h
--- a/Engine/Plugins/Runtime/RigVM/Source/RigVMEditor/Public/EdGraph/NodeSpawners/RigVMEdGraphVariableNodeSpawner.h
+++ b/Engine/Plugins/Runtime/RigVM/Source/RigVMEditor/Public/EdGraph/NodeSpawners/RigVMEdGraphVariableNodeSpawner.h
@@ -14,11 +14,13 @@
 #include "RigVMModel/RigVMVariableDescription.h"
 #include "RigVMEdGraphVariableNodeSpawner.generated.h"
 
+#define UE_API RIGVMEDITOR_API
+
 class URigVMEdGraphNode;
 class URigVMBlueprint;
 
-UCLASS(Transient)
-class RIGVMEDITOR_API URigVMEdGraphVariableNodeSpawner : public URigVMEdGraphNodeSpawner
+UCLASS(MinimalAPI, Transient)
+class URigVMEdGraphVariableNodeSpawner : public URigVMEdGraphNodeSpawner
 {
 	GENERATED_BODY()
 
@@ -30,15 +32,15 @@ public:
 	 * 
 	 * @return A newly allocated instance of this class.
 	 */
-	static URigVMEdGraphVariableNodeSpawner* CreateFromExternalVariable(URigVMBlueprint* InBlueprint, const FRigVMExternalVariable& InExternalVariable, bool bInIsGetter, const FText& InMenuDesc, const FText& InCategory, const FText& InTooltip);
-	static URigVMEdGraphVariableNodeSpawner* CreateFromLocalVariable(URigVMBlueprint* InBlueprint, URigVMGraph* InGraphOwner, const FRigVMGraphVariableDescription& InLocalVariable, bool bInIsGetter, const FText& InMenuDesc, const FText& InCategory, const FText& InTooltip);
+	static UE_API URigVMEdGraphVariableNodeSpawner* CreateFromExternalVariable(URigVMBlueprint* InBlueprint, const FRigVMExternalVariable& InExternalVariable, bool bInIsGetter, const FText& InMenuDesc, const FText& InCategory, const FText& InTooltip);
+	static UE_API URigVMEdGraphVariableNodeSpawner* CreateFromLocalVariable(URigVMBlueprint* InBlueprint, URigVMGraph* InGraphOwner, const FRigVMGraphVariableDescription& InLocalVariable, bool bInIsGetter, const FText& InMenuDesc, const FText& InCategory, const FText& InTooltip);
 
 	// UBlueprintNodeSpawner interface
-	virtual void Prime() override;
-	virtual FBlueprintNodeSignature GetSpawnerSignature() const override;
-	virtual FBlueprintActionUiSpec GetUiSpec(FBlueprintActionContext const& Context, FBindingSet const& Bindings) const override;
-	virtual UEdGraphNode* Invoke(UEdGraph* ParentGraph, FBindingSet const& Bindings, FVector2D const Location) const override;
-	virtual bool IsTemplateNodeFilteredOut(FBlueprintActionFilter const& Filter) const override;
+	UE_API virtual void Prime() override;
+	UE_API virtual FBlueprintNodeSignature GetSpawnerSignature() const override;
+	UE_API virtual FBlueprintActionUiSpec GetUiSpec(FBlueprintActionContext const& Context, FBindingSet const& Bindings) const override;
+	UE_API virtual UEdGraphNode* Invoke(UEdGraph* ParentGraph, FBindingSet const& Bindings, FVector2D const Location) const override;
+	UE_API virtual bool IsTemplateNodeFilteredOut(FBlueprintActionFilter const& Filter) const override;
 	// End UBlueprintNodeSpawner interface
 
 private:
@@ -52,3 +54,5 @@ private:
 
 	friend class UEngineTestControlRig;
 };
+
+#undef UE_API
diff --git a/Engine/Plugins/Runtime/RigVM/Source/RigVMEditor/Public/EdGraph/RigVMEdGraphConnectionDrawingPolicy.h b/Engine/Plugins/Runtime/RigVM/Source/RigVMEditor/Public/EdGraph/RigVMEdGraphConnectionDrawingPolicy.h
--- a/Engine/Plugins/Runtime/RigVM/Source/RigVMEditor/Public/EdGraph/RigVMEdGraphConnectionDrawingPolicy.h
+++ b/Engine/Plugins/Runtime/RigVM/Source/RigVMEditor/Public/EdGraph/RigVMEdGraphConnectionDrawingPolicy.h
@@ -4,9 +4,11 @@
 
 #include "BlueprintConnectionDrawingPolicy.h"
 
+#define UE_API RIGVMEDITOR_API
+
 class URigVMEdGraphNode;
 
-class RIGVMEDITOR_API FRigVMEdGraphConnectionDrawingPolicy : public FKismetConnectionDrawingPolicy
+class FRigVMEdGraphConnectionDrawingPolicy : public FKismetConnectionDrawingPolicy
 {
 public:
 	FRigVMEdGraphConnectionDrawingPolicy(int32 InBackLayerID, int32 InFrontLayerID, float InZoomFactor, const FSlateRect& InClippingRect, FSlateWindowElementList& InDrawElements, UEdGraph* InGraphObj)
@@ -19,18 +21,18 @@ public:
 		MidpointRadius = CastImage->ImageSize * InZoomFactor * 0.5f;
 	}
 
-	virtual void SetIncompatiblePinDrawState(const TSharedPtr<SGraphPin>& StartPin, const TSet< TSharedRef<SWidget> >& VisiblePins) override;
-	virtual void ResetIncompatiblePinDrawState(const TSet< TSharedRef<SWidget> >& VisiblePins) override;
-	virtual void BuildPinToPinWidgetMap(TMap<TSharedRef<SWidget>, FArrangedWidget>& InPinGeometries) override;
-	virtual void DrawPinGeometries(TMap<TSharedRef<SWidget>, FArrangedWidget>& InPinGeometries, FArrangedChildren& ArrangedNodes) override;
-	virtual void DetermineLinkGeometry(
+	UE_API virtual void SetIncompatiblePinDrawState(const TSharedPtr<SGraphPin>& StartPin, const TSet< TSharedRef<SWidget> >& VisiblePins) override;
+	UE_API virtual void ResetIncompatiblePinDrawState(const TSet< TSharedRef<SWidget> >& VisiblePins) override;
+	UE_API virtual void BuildPinToPinWidgetMap(TMap<TSharedRef<SWidget>, FArrangedWidget>& InPinGeometries) override;
+	UE_API virtual void DrawPinGeometries(TMap<TSharedRef<SWidget>, FArrangedWidget>& InPinGeometries, FArrangedChildren& ArrangedNodes) override;
+	UE_API virtual void DetermineLinkGeometry(
 	FArrangedChildren& ArrangedNodes, 
 	TSharedRef<SWidget>& OutputPinWidget,
 	UEdGraphPin* OutputPin,
 	UEdGraphPin* InputPin,
 	/*out*/ FArrangedWidget*& StartWidgetGeometry,
 	/*out*/ FArrangedWidget*& EndWidgetGeometry) override;
-	virtual void DetermineWiringStyle(UEdGraphPin* OutputPin, UEdGraphPin* InputPin, /*inout*/ FConnectionParams& Params) override;
+	UE_API virtual void DetermineWiringStyle(UEdGraphPin* OutputPin, UEdGraphPin* InputPin, /*inout*/ FConnectionParams& Params) override;
 	virtual bool UseDrawStateCaching() const override { return true; }
 
 private:
@@ -39,9 +41,11 @@ private:
 
 	bool UseLowDetailConnections() const { return ZoomFactor <= 0.175f; /* zoom level -9 */ }
 
-	bool ShouldChangeTangentForReouteControlPoint(URigVMEdGraphNode* Node);
+	UE_API bool ShouldChangeTangentForReouteControlPoint(URigVMEdGraphNode* Node);
 	// Average of the positions of all pins connected to InPin
-	bool GetAverageConnectedPositionForPin(UEdGraphPin* InPin, FVector2f& OutPos) const;
+	UE_API bool GetAverageConnectedPositionForPin(UEdGraphPin* InPin, FVector2f& OutPos) const;
 
 	const FSlateBrush* CastImage;
 };
+
+#undef UE_API
diff --git a/Engine/Plugins/Runtime/RigVM/Source/RigVMEditor/Public/EdGraph/RigVMEdGraphPanelNodeFactory.h b/Engine/Plugins/Runtime/RigVM/Source/RigVMEditor/Public/EdGraph/RigVMEdGraphPanelNodeFactory.h
--- a/Engine/Plugins/Runtime/RigVM/Source/RigVMEditor/Public/EdGraph/RigVMEdGraphPanelNodeFactory.h
+++ b/Engine/Plugins/Runtime/RigVM/Source/RigVMEditor/Public/EdGraph/RigVMEdGraphPanelNodeFactory.h
@@ -4,10 +4,14 @@
 
 #include "EdGraphUtilities.h"
 
-class RIGVMEDITOR_API FRigVMEdGraphPanelNodeFactory : public FGraphPanelNodeFactory
+#define UE_API RIGVMEDITOR_API
+
+class FRigVMEdGraphPanelNodeFactory : public FGraphPanelNodeFactory
 {
 public:
-	virtual FName GetFactoryName() const;
+	UE_API virtual FName GetFactoryName() const;
 private:
-	virtual TSharedPtr<class SGraphNode> CreateNode(UEdGraphNode* Node) const override;
+	UE_API virtual TSharedPtr<class SGraphNode> CreateNode(UEdGraphNode* Node) const override;
 };
+
+#undef UE_API
diff --git a/Engine/Plugins/Runtime/RigVM/Source/RigVMEditor/Public/EdGraph/RigVMEdGraphPanelPinFactory.h b/Engine/Plugins/Runtime/RigVM/Source/RigVMEditor/Public/EdGraph/RigVMEdGraphPanelPinFactory.h
--- a/Engine/Plugins/Runtime/RigVM/Source/RigVMEditor/Public/EdGraph/RigVMEdGraphPanelPinFactory.h
+++ b/Engine/Plugins/Runtime/RigVM/Source/RigVMEditor/Public/EdGraph/RigVMEdGraphPanelPinFactory.h
@@ -5,14 +5,18 @@
 #include "CoreMinimal.h"
 #include "EdGraphUtilities.h"
 
-class RIGVMEDITOR_API FRigVMEdGraphPanelPinFactory : public FGraphPanelPinFactory
+#define UE_API RIGVMEDITOR_API
+
+class FRigVMEdGraphPanelPinFactory : public FGraphPanelPinFactory
 {
 public:
-	virtual FName GetFactoryName() const;
+	UE_API virtual FName GetFactoryName() const;
 	
 	// FGraphPanelPinFactory interface
-	virtual TSharedPtr<class SGraphPin> CreatePin(class UEdGraphPin* InPin) const override;
+	UE_API virtual TSharedPtr<class SGraphPin> CreatePin(class UEdGraphPin* InPin) const override;
 
 
-	virtual TSharedPtr<class SGraphPin> CreatePin_Internal(class UEdGraphPin* InPin) const;
+	UE_API virtual TSharedPtr<class SGraphPin> CreatePin_Internal(class UEdGraphPin* InPin) const;
 };
+
+#undef UE_API
diff --git a/Engine/Plugins/Runtime/RigVM/Source/RigVMEditor/Public/Editor/RigVMActionMenuBuilder.h b/Engine/Plugins/Runtime/RigVM/Source/RigVMEditor/Public/Editor/RigVMActionMenuBuilder.h
--- a/Engine/Plugins/Runtime/RigVM/Source/RigVMEditor/Public/Editor/RigVMActionMenuBuilder.h
+++ b/Engine/Plugins/Runtime/RigVM/Source/RigVMEditor/Public/Editor/RigVMActionMenuBuilder.h
@@ -9,6 +9,8 @@
 #include "Internationalization/Text.h"
 #include "Templates/SharedPointer.h"
 
+#define UE_API RIGVMEDITOR_API
+
 class FBlueprintActionFilter;
 class FRigVMEditorBase;
 
@@ -29,7 +31,7 @@ struct FBlueprintActionInfo;
  * FRigVMActionMenuItems for actions that pass. Takes care of generating the 
  * each menu item's category/name/etc.
  */
-struct RIGVMEDITOR_API FRigVMActionMenuBuilder : public FGraphActionListBuilderBase
+struct FRigVMActionMenuBuilder : public FGraphActionListBuilderBase
 {
 public:
 	/** Flags used to configure the builder. */
@@ -51,10 +53,10 @@ public:
 	 * 
 	 * @param Flags	- Optional configuration flags.
 	 */
-	FRigVMActionMenuBuilder(EConfigFlags Flags = EConfigFlags::DefaultConfig);
+	UE_API FRigVMActionMenuBuilder(EConfigFlags Flags = EConfigFlags::DefaultConfig);
 	
 	// FGraphActionListBuilderBase interface
-	virtual void Empty() override;
+	UE_API virtual void Empty() override;
 	// End FGraphActionListBuilderBase interface
 	
 	/**
@@ -69,27 +71,27 @@ public:
 	 * @param  MenuOrder The sort order to assign this section of the menu (higher numbers come first).
 	 * @param  Flags	 Set of ESectionFlags to customize this menu section.
 	 */
-	void AddMenuSection(FBlueprintActionFilter const& Filter, FText const& Heading = FText::GetEmpty());
+	UE_API void AddMenuSection(FBlueprintActionFilter const& Filter, FText const& Heading = FText::GetEmpty());
 	
 	/**
 	 * Regenerates the entire menu list from the cached menu sections. Filters 
 	 * and adds action items from the blueprint action database (as defined by 
 	 * the MenuSections list).
 	 */
-	void RebuildActionList();
+	UE_API void RebuildActionList();
 
 	/** Returns the current number of actions that are still pending */
-	int32 GetNumPendingActions() const;
+	UE_API int32 GetNumPendingActions() const;
 
 	/** Processes any actions that may be added asynchronously or across multiple frames. Returns true if one or more actions were added into the list. */
-	bool ProcessPendingActions();
+	UE_API bool ProcessPendingActions();
 
 	/** Returns the normalized completion state when processing pending actions (e.g. for a status indicator) */
-	float GetPendingActionsProgress() const;
+	UE_API float GetPendingActionsProgress() const;
 
 protected:
 	// Adds menu items for the given database action.
-	void MakeMenuItems(FBlueprintActionInfo& InAction);
+	UE_API void MakeMenuItems(FBlueprintActionInfo& InAction);
 
 
 private:
@@ -110,4 +112,6 @@ private:
 	bool bUsePendingActionList;
 };
 
-ENUM_CLASS_FLAGS(FRigVMActionMenuBuilder::EConfigFlags);
\ No newline at end of file
+ENUM_CLASS_FLAGS(FRigVMActionMenuBuilder::EConfigFlags);
+
+#undef UE_API
diff --git a/Engine/Plugins/Runtime/RigVM/Source/RigVMEditor/Public/Editor/RigVMActionMenuItem.h b/Engine/Plugins/Runtime/RigVM/Source/RigVMEditor/Public/Editor/RigVMActionMenuItem.h
--- a/Engine/Plugins/Runtime/RigVM/Source/RigVMEditor/Public/Editor/RigVMActionMenuItem.h
+++ b/Engine/Plugins/Runtime/RigVM/Source/RigVMEditor/Public/Editor/RigVMActionMenuItem.h
@@ -18,6 +18,8 @@
 
 #include "RigVMActionMenuItem.generated.h"
 
+#define UE_API RIGVMEDITOR_API
+
 class FReferenceCollector;
 class UBlueprintNodeSpawner;
 class UEdGraph;
@@ -34,7 +36,7 @@ struct FSlateBrush;
  * invoked from this class (separated to divide ui and node-spawning).
  */
 USTRUCT()
-struct RIGVMEDITOR_API FRigVMActionMenuItem : public FEdGraphSchemaAction
+struct FRigVMActionMenuItem : public FEdGraphSchemaAction
 {
 	GENERATED_USTRUCT_BODY();
 	static FName StaticGetTypeId() { static FName const TypeId("FRigVMActionMenuItem"); return TypeId; }
@@ -42,17 +44,17 @@ struct RIGVMEDITOR_API FRigVMActionMenuItem : public FEdGraphSchemaAction
 public:	
 	/** Constructors */
 	FRigVMActionMenuItem(UBlueprintNodeSpawner const* NodeSpawner = nullptr) : Action(NodeSpawner), IconTint(FLinearColor::White), IconBrush(nullptr) {}
-	FRigVMActionMenuItem(UBlueprintNodeSpawner const* NodeSpawner, FBlueprintActionUiSpec const& UiDef, IBlueprintNodeBinder::FBindingSet const& Bindings = IBlueprintNodeBinder::FBindingSet(), FText InNodeCategory = FText(), int32 InGrouping = 0);
+	UE_API FRigVMActionMenuItem(UBlueprintNodeSpawner const* NodeSpawner, FBlueprintActionUiSpec const& UiDef, IBlueprintNodeBinder::FBindingSet const& Bindings = IBlueprintNodeBinder::FBindingSet(), FText InNodeCategory = FText(), int32 InGrouping = 0);
 	
 	// FEdGraphSchemaAction interface
 	virtual FName         GetTypeId() const final { return StaticGetTypeId(); }
-	virtual UEdGraphNode* PerformAction(UEdGraph* ParentGraph, UEdGraphPin* FromPin, const FVector2f& Location, bool bSelectNewNode = true) final;
-	virtual UEdGraphNode* PerformAction(UEdGraph* ParentGraph, TArray<UEdGraphPin*>& FromPins, const FVector2f& Location, bool bSelectNewNode = true) final;
-	virtual void          AddReferencedObjects(FReferenceCollector& Collector) final;
+	UE_API virtual UEdGraphNode* PerformAction(UEdGraph* ParentGraph, UEdGraphPin* FromPin, const FVector2f& Location, bool bSelectNewNode = true) final;
+	UE_API virtual UEdGraphNode* PerformAction(UEdGraph* ParentGraph, TArray<UEdGraphPin*>& FromPins, const FVector2f& Location, bool bSelectNewNode = true) final;
+	UE_API virtual void          AddReferencedObjects(FReferenceCollector& Collector) final;
 	// End FEdGraphSchemaAction interface
 
 	/** @return */
-	UBlueprintNodeSpawner const* GetRawAction() const;
+	UE_API UBlueprintNodeSpawner const* GetRawAction() const;
 
 	/**
 	 * Retrieves the icon brush for this menu entry (to be displayed alongside
@@ -61,7 +63,7 @@ public:
 	 * @param  ColorOut	An output parameter that's filled in with the color to tint the brush with.
 	 * @return An slate brush to be used for this menu item in the action menu.
 	 */
-	FSlateBrush const* GetMenuIcon(FSlateColor& ColorOut);
+	UE_API FSlateBrush const* GetMenuIcon(FSlateColor& ColorOut);
 
 
 
@@ -75,3 +77,5 @@ private:
 	/** */
 	IBlueprintNodeBinder::FBindingSet Bindings;
 };
+
+#undef UE_API
diff --git a/Engine/Plugins/Runtime/RigVM/Source/RigVMEditor/Public/Editor/RigVMCommentNodeDetailCustomization.h b/Engine/Plugins/Runtime/RigVM/Source/RigVMEditor/Public/Editor/RigVMCommentNodeDetailCustomization.h
--- a/Engine/Plugins/Runtime/RigVM/Source/RigVMEditor/Public/Editor/RigVMCommentNodeDetailCustomization.h
+++ b/Engine/Plugins/Runtime/RigVM/Source/RigVMEditor/Public/Editor/RigVMCommentNodeDetailCustomization.h
@@ -10,10 +10,12 @@
 #include "Styling/SlateTypes.h"
 #include "IPropertyUtilities.h"
 
+#define UE_API RIGVMEDITOR_API
+
 class IRigVMEditor;
 class IPropertyHandle;
 
-class RIGVMEDITOR_API FRigVMCommentNodeDetailCustomization : public IDetailCustomization
+class FRigVMCommentNodeDetailCustomization : public IDetailCustomization
 {
 	FRigVMCommentNodeDetailCustomization()
 	: BlueprintBeingCustomized(nullptr)
@@ -29,31 +31,31 @@ public:
 	}
 
 	/** IDetailCustomization interface */
-	virtual void CustomizeDetails(IDetailLayoutBuilder& DetailBuilder) override;
+	UE_API virtual void CustomizeDetails(IDetailLayoutBuilder& DetailBuilder) override;
 
 private:
 
 	URigVMBlueprint* BlueprintBeingCustomized;
 	TArray<TWeakObjectPtr<URigVMCommentNode>> ObjectsBeingCustomized;
 
-	void GetValuesFromNode(TWeakObjectPtr<URigVMCommentNode> Node);
-	void SetValues(TWeakObjectPtr<URigVMCommentNode> Node);
+	UE_API void GetValuesFromNode(TWeakObjectPtr<URigVMCommentNode> Node);
+	UE_API void SetValues(TWeakObjectPtr<URigVMCommentNode> Node);
 	
-	FText GetText() const;
-	void SetText(const FText& InNewText, ETextCommit::Type InCommitType);
+	UE_API FText GetText() const;
+	UE_API void SetText(const FText& InNewText, ETextCommit::Type InCommitType);
 
-	FLinearColor GetColor() const;
-	FReply OnChooseColor(const FGeometry& MyGeometry, const FPointerEvent& MouseEvent);
-	void OnColorPicked(FLinearColor LinearColor);
+	UE_API FLinearColor GetColor() const;
+	UE_API FReply OnChooseColor(const FGeometry& MyGeometry, const FPointerEvent& MouseEvent);
+	UE_API void OnColorPicked(FLinearColor LinearColor);
 
-	ECheckBoxState IsShowingBubbleEnabled() const;
-	void OnShowingBubbleStateChanged(ECheckBoxState InValue);
+	UE_API ECheckBoxState IsShowingBubbleEnabled() const;
+	UE_API void OnShowingBubbleStateChanged(ECheckBoxState InValue);
 
-	ECheckBoxState IsColorBubbleEnabled() const;
-	void OnColorBubbleStateChanged(ECheckBoxState InValue);
+	UE_API ECheckBoxState IsColorBubbleEnabled() const;
+	UE_API void OnColorBubbleStateChanged(ECheckBoxState InValue);
 
-	TOptional<int> GetFontSize() const;
-	void OnFontSizeChanged(int32 InValue, ETextCommit::Type Arg);
+	UE_API TOptional<int> GetFontSize() const;
+	UE_API void OnFontSizeChanged(int32 InValue, ETextCommit::Type Arg);
 	
 	FString CommentText;
 	bool bShowingBubble;
@@ -63,3 +65,4 @@ private:
 
 
 
+#undef UE_API
diff --git a/Engine/Plugins/Runtime/RigVM/Source/RigVMEditor/Public/Editor/RigVMCompilerResultsTabSummoner.h b/Engine/Plugins/Runtime/RigVM/Source/RigVMEditor/Public/Editor/RigVMCompilerResultsTabSummoner.h
--- a/Engine/Plugins/Runtime/RigVM/Source/RigVMEditor/Public/Editor/RigVMCompilerResultsTabSummoner.h
+++ b/Engine/Plugins/Runtime/RigVM/Source/RigVMEditor/Public/Editor/RigVMCompilerResultsTabSummoner.h
@@ -7,21 +7,23 @@
 #include "Toolkits/AssetEditorToolkit.h"
 #include "WorkflowOrientedApp/WorkflowTabFactory.h"
 
+#define UE_API RIGVMEDITOR_API
+
 #define LOCTEXT_NAMESPACE "RigVMEditor"
 
 //////////////////////////////////////////////////////////////////////////
 // FCompilerResultsSummoner
 
-struct RIGVMEDITOR_API FRigVMCompilerResultsTabSummoner : public FWorkflowTabFactory
+struct FRigVMCompilerResultsTabSummoner : public FWorkflowTabFactory
 {
 
 public:
 
 	static const FName TabID() { static FName ID = TEXT("RigVM Compiler Results"); return ID; }
 	
-	FRigVMCompilerResultsTabSummoner(TSharedPtr<class FAssetEditorToolkit> InHostingApp);
+	UE_API FRigVMCompilerResultsTabSummoner(TSharedPtr<class FAssetEditorToolkit> InHostingApp);
 
-	virtual TSharedRef<SWidget> CreateTabBody(const FWorkflowTabSpawnInfo& Info) const override;
+	UE_API virtual TSharedRef<SWidget> CreateTabBody(const FWorkflowTabSpawnInfo& Info) const override;
 
 	virtual FText GetTabToolTipText(const FWorkflowTabSpawnInfo& Info) const override
 	{
@@ -29,4 +31,6 @@ public:
 	}
 };
 
-#undef LOCTEXT_NAMESPACE
\ No newline at end of file
+#undef LOCTEXT_NAMESPACE
+
+#undef UE_API
diff --git a/Engine/Plugins/Runtime/RigVM/Source/RigVMEditor/Public/Editor/RigVMDetailsInspectorTabSummoner.h b/Engine/Plugins/Runtime/RigVM/Source/RigVMEditor/Public/Editor/RigVMDetailsInspectorTabSummoner.h
--- a/Engine/Plugins/Runtime/RigVM/Source/RigVMEditor/Public/Editor/RigVMDetailsInspectorTabSummoner.h
+++ b/Engine/Plugins/Runtime/RigVM/Source/RigVMEditor/Public/Editor/RigVMDetailsInspectorTabSummoner.h
@@ -5,20 +5,24 @@
 #include "CoreMinimal.h"
 #include "WorkflowOrientedApp/WorkflowTabFactory.h"
 
+#define UE_API RIGVMEDITOR_API
+
 class FRigVMNewEditor;
 
-struct RIGVMEDITOR_API FRigVMDetailsInspectorTabSummoner : public FWorkflowTabFactory
+struct FRigVMDetailsInspectorTabSummoner : public FWorkflowTabFactory
 {
 public:
 
 	static const FName TabID() { static FName ID = TEXT("RigVM Details"); return ID; }
 
 public:
-	FRigVMDetailsInspectorTabSummoner(const TSharedRef<FRigVMNewEditor>& InRigVMEditor);
+	UE_API FRigVMDetailsInspectorTabSummoner(const TSharedRef<FRigVMNewEditor>& InRigVMEditor);
 
-	virtual TSharedRef<SWidget> CreateTabBody(const FWorkflowTabSpawnInfo& Info) const override;
-	virtual TSharedRef<SDockTab> SpawnTab(const FWorkflowTabSpawnInfo& Info) const override;
+	UE_API virtual TSharedRef<SWidget> CreateTabBody(const FWorkflowTabSpawnInfo& Info) const override;
+	UE_API virtual TSharedRef<SDockTab> SpawnTab(const FWorkflowTabSpawnInfo& Info) const override;
 
 protected:
 	TWeakPtr<FRigVMNewEditor> RigVMEditor;
 };
+
+#undef UE_API
diff --git a/Engine/Plugins/Runtime/RigVM/Source/RigVMEditor/Public/Editor/RigVMDetailsViewWrapperObject.h b/Engine/Plugins/Runtime/RigVM/Source/RigVMEditor/Public/Editor/RigVMDetailsViewWrapperObject.h
--- a/Engine/Plugins/Runtime/RigVM/Source/RigVMEditor/Public/Editor/RigVMDetailsViewWrapperObject.h
+++ b/Engine/Plugins/Runtime/RigVM/Source/RigVMEditor/Public/Editor/RigVMDetailsViewWrapperObject.h
@@ -7,31 +7,33 @@
 #include "RigVMModel/RigVMNode.h"
 #include "RigVMDetailsViewWrapperObject.generated.h"
 
+#define UE_API RIGVMEDITOR_API
+
 class URigVMDetailsViewWrapperObject;
 
 DECLARE_EVENT_ThreeParams(URigVMDetailsViewWrapperObject, FWrappedPropertyChangedChainEvent, URigVMDetailsViewWrapperObject*, const FString&, FPropertyChangedChainEvent&);
 
-UCLASS()
-class RIGVMEDITOR_API URigVMDetailsViewWrapperObject : public UObject
+UCLASS(MinimalAPI)
+class URigVMDetailsViewWrapperObject : public UObject
 {
 public:
 	GENERATED_BODY()
 
-	URigVMDetailsViewWrapperObject();
+	UE_API URigVMDetailsViewWrapperObject();
 
 	// Creating wrappers from a given struct
-	virtual UClass* GetClassForStruct(UScriptStruct* InStruct, bool bCreateIfNeeded = true) const;
-	static URigVMDetailsViewWrapperObject* MakeInstance(UClass* InWrapperObjectClass, UObject* InOuter, UScriptStruct* InStruct, uint8* InStructMemory, UObject* InSubject = nullptr);
-	UScriptStruct* GetWrappedStruct() const;
+	UE_API virtual UClass* GetClassForStruct(UScriptStruct* InStruct, bool bCreateIfNeeded = true) const;
+	static UE_API URigVMDetailsViewWrapperObject* MakeInstance(UClass* InWrapperObjectClass, UObject* InOuter, UScriptStruct* InStruct, uint8* InStructMemory, UObject* InSubject = nullptr);
+	UE_API UScriptStruct* GetWrappedStruct() const;
 
 	// Creating wrappers from a RigVMNode
-	virtual UClass* GetClassForNodes(TArray<URigVMNode*> InNodes, bool bCreateIfNeeded = true) const;
-	static URigVMDetailsViewWrapperObject* MakeInstance(UClass* InWrapperObjectClass, UObject* InOuter, TArray<URigVMNode*> InNodes, URigVMNode* InSubject);
+	UE_API virtual UClass* GetClassForNodes(TArray<URigVMNode*> InNodes, bool bCreateIfNeeded = true) const;
+	static UE_API URigVMDetailsViewWrapperObject* MakeInstance(UClass* InWrapperObjectClass, UObject* InOuter, TArray<URigVMNode*> InNodes, URigVMNode* InSubject);
 	
-	static void MarkOutdatedClass(UClass* InClass);
-	static bool IsValidClass(UClass* InClass);
+	static UE_API void MarkOutdatedClass(UClass* InClass);
+	static UE_API bool IsValidClass(UClass* InClass);
 	
-	FString GetWrappedNodeNotation() const;
+	UE_API FString GetWrappedNodeNotation() const;
 	
 	bool IsChildOf(const UStruct* InStruct) const
 	{
@@ -45,9 +47,9 @@ public:
 		return IsChildOf(T::StaticStruct());
 	}
 
-	virtual void SetContent(const uint8* InStructMemory, const UStruct* InStruct);
-	virtual void GetContent(uint8* OutStructMemory, const UStruct* InStruct) const;
-	void SetContent(URigVMNode* InNode);
+	UE_API virtual void SetContent(const uint8* InStructMemory, const UStruct* InStruct);
+	UE_API virtual void GetContent(uint8* OutStructMemory, const UStruct* InStruct) const;
+	UE_API void SetContent(URigVMNode* InNode);
 
 	template<typename T>
 	T GetContent() const
@@ -67,17 +69,17 @@ public:
 		SetContent((const uint8*)&InValue, T::StaticStruct());
 	}
 
-	UObject* GetSubject() const;
-	void SetSubject(UObject* InSubject);
+	UE_API UObject* GetSubject() const;
+	UE_API void SetSubject(UObject* InSubject);
 
 	FWrappedPropertyChangedChainEvent& GetWrappedPropertyChangedChainEvent() { return WrappedPropertyChangedChainEvent; }
-	virtual void PostEditChangeChainProperty(FPropertyChangedChainEvent& PropertyChangedEvent) override;
+	UE_API virtual void PostEditChangeChainProperty(FPropertyChangedChainEvent& PropertyChangedEvent) override;
 
 private:
 
-	static void CopyPropertiesForUnrelatedStructs(uint8* InTargetMemory, const UStruct* InTargetStruct, const uint8* InSourceMemory, const UStruct* InSourceStruct);
-	void HandleModifiedEvent(ERigVMGraphNotifType InNotifType, URigVMGraph* InGraph, UObject* InSubject);
-	void SetContentForPin(URigVMPin* InPin);
+	static UE_API void CopyPropertiesForUnrelatedStructs(uint8* InTargetMemory, const UStruct* InTargetStruct, const uint8* InSourceMemory, const UStruct* InSourceStruct);
+	UE_API void HandleModifiedEvent(ERigVMGraphNotifType InNotifType, URigVMGraph* InGraph, UObject* InSubject);
+	UE_API void SetContentForPin(URigVMPin* InPin);
 
 	struct FPerClassInfo
 	{
@@ -115,9 +117,9 @@ private:
 		}
 	};
 	
-	static TMap<FPerClassInfo, UClass*> InfoToClass;
-	static TMap<UClass*, FPerClassInfo> ClassToInfo;
-	static TSet<UClass*> OutdatedClassToRecreate;
+	static UE_API TMap<FPerClassInfo, UClass*> InfoToClass;
+	static UE_API TMap<UClass*, FPerClassInfo> ClassToInfo;
+	static UE_API TSet<UClass*> OutdatedClassToRecreate;
 	bool bIsSettingValue;
 
 	UPROPERTY()
@@ -125,3 +127,5 @@ private:
 	
 	FWrappedPropertyChangedChainEvent WrappedPropertyChangedChainEvent;
 };
+
+#undef UE_API
diff --git a/Engine/Plugins/Runtime/RigVM/Source/RigVMEditor/Public/Editor/RigVMEditor.h b/Engine/Plugins/Runtime/RigVM/Source/RigVMEditor/Public/Editor/RigVMEditor.h
--- a/Engine/Plugins/Runtime/RigVM/Source/RigVMEditor/Public/Editor/RigVMEditor.h
+++ b/Engine/Plugins/Runtime/RigVM/Source/RigVMEditor/Public/Editor/RigVMEditor.h
@@ -21,6 +21,8 @@
 #include "SNodePanel.h"
 #include "Editor/RigVMDetailsViewWrapperObject.h"
 
+#define UE_API RIGVMEDITOR_API
+
 #define LOCTEXT_NAMESPACE "RigVMEditor"
 
 class FDocumentTracker;
@@ -56,7 +58,7 @@ namespace RigVMNodeSectionID
 struct FRigVMEditorModes
 {
 	// Mode constants
-	RIGVMEDITOR_API static inline const FLazyName RigVMEditorMode = FLazyName(TEXT("RigVM"));
+	static inline const FLazyName RigVMEditorMode = FLazyName(TEXT("RigVM"));
 	static FText GetLocalizedMode(const FName InMode)
 	{
 		static TMap< FName, FText > LocModes;
@@ -97,7 +99,7 @@ public:
 	}
 };
 
-class RIGVMEDITOR_API IRigVMEditor
+class IRigVMEditor
 {
 public:
 
@@ -247,7 +249,7 @@ protected:
 	virtual void AddAutoCompileWidget(FToolBarBuilder& ToolbarBuilder) = 0;
 };
 
-class RIGVMEDITOR_API FRigVMEditorBase : public IRigVMEditor
+class FRigVMEditorBase : public IRigVMEditor
 {
 public:
 	
@@ -258,39 +260,39 @@ public:
 	 * @param	InitToolkitHost			When Mode is WorldCentric, this is the level editor instance to spawn this editor within
 	 * @param	InRigVMBlueprint	The blueprint object to start editing.
 	 */
-	virtual void InitRigVMEditor(const EToolkitMode::Type Mode, const TSharedPtr< class IToolkitHost >& InitToolkitHost, class URigVMBlueprint* InRigVMBlueprint);
+	UE_API virtual void InitRigVMEditor(const EToolkitMode::Type Mode, const TSharedPtr< class IToolkitHost >& InitToolkitHost, class URigVMBlueprint* InRigVMBlueprint);
 
-	static FRigVMEditorBase* GetFromAssetEditorInstance(IAssetEditorInstance* Instance);
+	static UE_API FRigVMEditorBase* GetFromAssetEditorInstance(IAssetEditorInstance* Instance);
 	// returns the blueprint being edited
-	virtual URigVMBlueprint* GetRigVMBlueprint() const override;
+	UE_API virtual URigVMBlueprint* GetRigVMBlueprint() const override;
 	
-	void HandleJumpToHyperlink(const UObject* InSubject);
-
-	void Compile();
-	bool IsCompilingEnabled() const;
-
-	void DeleteSelectedNodes();
-	bool CanDeleteNodes() const;
-	void CopySelectedNodes();
-	bool CanCopyNodes() const;
-	void PasteNodes();
-	bool CanPasteNodes() const;
-	void CutSelectedNodes();
-	bool CanCutNodes() const;
-	void DuplicateNodes();
-	bool CanDuplicateNodes() const;
+	UE_API void HandleJumpToHyperlink(const UObject* InSubject);
+
+	UE_API void Compile();
+	UE_API bool IsCompilingEnabled() const;
+
+	UE_API void DeleteSelectedNodes();
+	UE_API bool CanDeleteNodes() const;
+	UE_API void CopySelectedNodes();
+	UE_API bool CanCopyNodes() const;
+	UE_API void PasteNodes();
+	UE_API bool CanPasteNodes() const;
+	UE_API void CutSelectedNodes();
+	UE_API bool CanCutNodes() const;
+	UE_API void DuplicateNodes();
+	UE_API bool CanDuplicateNodes() const;
 	
-	void OnStartWatchingPin();
-	bool CanStartWatchingPin() const;
-	void OnStopWatchingPin();
-	bool CanStopWatchingPin() const;
+	UE_API void OnStartWatchingPin();
+	UE_API bool CanStartWatchingPin() const;
+	UE_API void OnStopWatchingPin();
+	UE_API bool CanStopWatchingPin() const;
 
 	virtual FOnRefreshEvent OnRefresh() override { return OnRefreshEvent; }
 
-	FText GetGraphDecorationString(UEdGraph* InGraph) const;
-	virtual bool IsEditable(UEdGraph* InGraph) const override;
+	UE_API FText GetGraphDecorationString(UEdGraph* InGraph) const;
+	UE_API virtual bool IsEditable(UEdGraph* InGraph) const override;
 
-	virtual void OnGraphNodeDropToPerform(TSharedPtr<FDragDropOperation> InDragDropOp, UEdGraph* InGraph, const FVector2f& InNodePosition, const FVector2f& InScreenPosition) override;
+	UE_API virtual void OnGraphNodeDropToPerform(TSharedPtr<FDragDropOperation> InDragDropOp, UEdGraph* InGraph, const FVector2f& InNodePosition, const FVector2f& InScreenPosition) override;
 
 	/**
 	 * Util for finding a glyph and color for a variable.
@@ -301,7 +303,7 @@ public:
 	 * @param SecondaryColorOut The resulting secondary color for the glyph (used for Map types)
 	 * @return					The resulting glyph brush
 	 */
-	static FSlateBrush const* GetVarIconAndColorFromProperty(const FProperty* Property, FSlateColor& IconColorOut, FSlateBrush const*& SecondaryBrushOut, FSlateColor& SecondaryColorOut);
+	static UE_API FSlateBrush const* GetVarIconAndColorFromProperty(const FProperty* Property, FSlateColor& IconColorOut, FSlateBrush const*& SecondaryBrushOut, FSlateColor& SecondaryColorOut);
 
 	/**
 	* Util for finding a glyph and color for a variable.
@@ -312,77 +314,77 @@ public:
 	* @param SecondaryColorOut The resulting secondary color for the glyph (used for Map types)
 	* @return					The resulting glyph brush
 	*/
-	static FSlateBrush const* GetVarIconAndColorFromPinType(const FEdGraphPinType& PinType, FSlateColor& IconColorOut, FSlateBrush const*& SecondaryBrushOut, FSlateColor& SecondaryColorOut);
+	static UE_API FSlateBrush const* GetVarIconAndColorFromPinType(const FEdGraphPinType& PinType, FSlateColor& IconColorOut, FSlateBrush const*& SecondaryBrushOut, FSlateColor& SecondaryColorOut);
 	
 protected:
 	
-	FRigVMEditorBase();
+	UE_API FRigVMEditorBase();
 	virtual ~FRigVMEditorBase(){}
 
-	void UnbindEditor();
+	UE_API void UnbindEditor();
 	
 
-	void HandleAssetRequestedOpen(UObject* InObject);
-	void HandleAssetRequestClose(UObject* InObject, EAssetEditorCloseReason InReason);
+	UE_API void HandleAssetRequestedOpen(UObject* InObject);
+	UE_API void HandleAssetRequestClose(UObject* InObject, EAssetEditorCloseReason InReason);
 	bool bRequestedReopen = false;
 
-	virtual const FName GetEditorModeName() const;
+	UE_API virtual const FName GetEditorModeName() const;
 
 	// FBlueprintEditor interface
-	UBlueprint* GetBlueprintObj() const;
-	virtual bool InEditingMode() const override;
-	TSubclassOf<UEdGraphSchema> GetDefaultSchemaClass() const;
-	void OnGraphEditorFocused(const TSharedRef<SGraphEditor>& InGraphEditor);
+	UE_API UBlueprint* GetBlueprintObj() const;
+	UE_API virtual bool InEditingMode() const override;
+	UE_API TSubclassOf<UEdGraphSchema> GetDefaultSchemaClass() const;
+	UE_API void OnGraphEditorFocused(const TSharedRef<SGraphEditor>& InGraphEditor);
 
 	//  FTickableEditorObject Interface
-	void Tick(float DeltaTime);
+	UE_API void Tick(float DeltaTime);
 
 	// IToolkit Interface
-	void BringToolkitToFront();
-	FName GetToolkitFName() const;
-	FName GetToolkitContextFName() const;
-	FText GetBaseToolkitName() const;
-	FText GetToolkitToolTipText() const;
-	FString GetWorldCentricTabPrefix() const;
-	FLinearColor GetWorldCentricTabColorScale() const;	
-	void InitToolMenuContextImpl(FToolMenuContext& MenuContext);
+	UE_API void BringToolkitToFront();
+	UE_API FName GetToolkitFName() const;
+	UE_API FName GetToolkitContextFName() const;
+	UE_API FText GetBaseToolkitName() const;
+	UE_API FText GetToolkitToolTipText() const;
+	UE_API FString GetWorldCentricTabPrefix() const;
+	UE_API FLinearColor GetWorldCentricTabColorScale() const;	
+	UE_API void InitToolMenuContextImpl(FToolMenuContext& MenuContext);
 
 	// BlueprintEditor interface
-	bool TransactionObjectAffectsBlueprintImpl(UObject* InTransactedObject);
-	virtual bool CanAddNewLocalVariable() const override;
-	virtual void OnAddNewLocalVariable() override;
-	virtual void OnPasteNewLocalVariable(const FBPVariableDescription& VariableDescription) override;
+	UE_API bool TransactionObjectAffectsBlueprintImpl(UObject* InTransactedObject);
+	UE_API virtual bool CanAddNewLocalVariable() const override;
+	UE_API virtual void OnAddNewLocalVariable() override;
+	UE_API virtual void OnPasteNewLocalVariable(const FBPVariableDescription& VariableDescription) override;
 
 
 	bool IsNativeParentClassCodeLinkEnabled() const { return false; }
 	bool ReparentBlueprint_IsVisible() const { return false; }
-	FReply OnSpawnGraphNodeByShortcut(FInputChord InChord, const FVector2f& InPosition, UEdGraph* InGraph);
+	UE_API FReply OnSpawnGraphNodeByShortcut(FInputChord InChord, const FVector2f& InPosition, UEdGraph* InGraph);
 	bool ShouldLoadBPLibrariesFromAssetRegistry() { return false; }
-	bool JumpToHyperlinkImpl(const UObject* ObjectReference, bool bRequestRename = false);
+	UE_API bool JumpToHyperlinkImpl(const UObject* ObjectReference, bool bRequestRename = false);
 	bool ShouldOpenGraphByDefault() const { return true; }
-	virtual void AddNewFunctionVariant(const UEdGraph* InOriginalFunction) override;
+	UE_API virtual void AddNewFunctionVariant(const UEdGraph* InOriginalFunction) override;
 
 	// FEditorUndoClient Interface
-	void PostUndoImpl(bool bSuccess);
-	void PostRedoImpl(bool bSuccess);
+	UE_API void PostUndoImpl(bool bSuccess);
+	UE_API void PostRedoImpl(bool bSuccess);
 
-	virtual void PostTransaction(bool bSuccess, const FTransaction* Transaction, bool bIsRedo);
+	UE_API virtual void PostTransaction(bool bSuccess, const FTransaction* Transaction, bool bIsRedo);
 
 
 	// IToolkitHost Interface
-	void OnToolkitHostingStarted(const TSharedRef<class IToolkit>& Toolkit);
-	void OnToolkitHostingFinished(const TSharedRef<class IToolkit>& Toolkit);
+	UE_API void OnToolkitHostingStarted(const TSharedRef<class IToolkit>& Toolkit);
+	UE_API void OnToolkitHostingFinished(const TSharedRef<class IToolkit>& Toolkit);
 
 	//  FTickableEditorObject Interface
-	TStatId GetStatId() const;
+	UE_API TStatId GetStatId() const;
 
 	// returns the currently debugged / viewed host
-	virtual URigVMHost* GetRigVMHost() const override;
+	UE_API virtual URigVMHost* GetRigVMHost() const override;
 
-	virtual UObject* GetOuterForHost() const;
+	UE_API virtual UObject* GetOuterForHost() const;
 
 	// returns the class to use for detail wrapper objects (UI shim layer)
-	virtual UClass* GetDetailWrapperClass() const;
+	UE_API virtual UClass* GetDetailWrapperClass() const;
 
 	// allows the editor to fill an empty graph
 	virtual void CreateEmptyGraphContent(URigVMController* InController) {}
@@ -395,7 +397,7 @@ protected:
 	/** Get the toolbox hosting widget */
 	TSharedRef<SBorder> GetToolbox() { return Toolbox.ToSharedRef(); }
 
-	virtual bool SelectLocalVariable(const UEdGraph* Graph, const FName& VariableName) override;
+	UE_API virtual bool SelectLocalVariable(const UEdGraph* Graph, const FName& VariableName) override;
 
 
 	enum ERigVMEditorExecutionModeType
@@ -405,172 +407,172 @@ protected:
 	};
 
 	// FBlueprintEditor Interface
-	void CreateDefaultCommandsImpl();
-	void OnCreateGraphEditorCommands(TSharedPtr<FUICommandList> GraphEditorCommandsList);
+	UE_API void CreateDefaultCommandsImpl();
+	UE_API void OnCreateGraphEditorCommands(TSharedPtr<FUICommandList> GraphEditorCommandsList);
 
-	void SaveAsset_Execute();
-	void SaveAssetAs_Execute();
+	UE_API void SaveAsset_Execute();
+	UE_API void SaveAssetAs_Execute();
 	bool IsInAScriptingMode() const { return true; }
-	virtual void OnNewDocumentClicked(ECreatedDocumentType GraphType) override;
-	bool IsSectionVisibleImpl(RigVMNodeSectionID::Type InSectionID) const;
-	bool AreEventGraphsAllowed() const;
-	bool AreMacrosAllowed() const;
-	bool AreDelegatesAllowed() const;
-	bool NewDocument_IsVisibleForTypeImpl(ECreatedDocumentType GraphType) const;
-	FGraphAppearanceInfo GetGraphAppearance(class UEdGraph* InGraph) const;
-	void OnActiveTabChanged( TSharedPtr<SDockTab> PreviouslyActive, TSharedPtr<SDockTab> NewlyActivated );
-	void OnSelectedNodesChangedImpl(const TSet<class UObject*>& NewSelection);
-	void OnBlueprintChangedImpl(UBlueprint* InBlueprint, bool bIsJustBeingCompiled);
-	virtual void ForceEditorRefresh(ERefreshRigVMEditorReason::Type Reason = ERefreshRigVMEditorReason::UnknownReason) override;
-	void SetupGraphEditorEvents(UEdGraph* InGraph, SGraphEditor::FGraphEditorEvents& InEvents);
-	void FocusInspectorOnGraphSelection(const TSet<class UObject*>& NewSelection, bool bForceRefresh = false);
+	UE_API virtual void OnNewDocumentClicked(ECreatedDocumentType GraphType) override;
+	UE_API bool IsSectionVisibleImpl(RigVMNodeSectionID::Type InSectionID) const;
+	UE_API bool AreEventGraphsAllowed() const;
+	UE_API bool AreMacrosAllowed() const;
+	UE_API bool AreDelegatesAllowed() const;
+	UE_API bool NewDocument_IsVisibleForTypeImpl(ECreatedDocumentType GraphType) const;
+	UE_API FGraphAppearanceInfo GetGraphAppearance(class UEdGraph* InGraph) const;
+	UE_API void OnActiveTabChanged( TSharedPtr<SDockTab> PreviouslyActive, TSharedPtr<SDockTab> NewlyActivated );
+	UE_API void OnSelectedNodesChangedImpl(const TSet<class UObject*>& NewSelection);
+	UE_API void OnBlueprintChangedImpl(UBlueprint* InBlueprint, bool bIsJustBeingCompiled);
+	UE_API virtual void ForceEditorRefresh(ERefreshRigVMEditorReason::Type Reason = ERefreshRigVMEditorReason::UnknownReason) override;
+	UE_API void SetupGraphEditorEvents(UEdGraph* InGraph, SGraphEditor::FGraphEditorEvents& InEvents);
+	UE_API void FocusInspectorOnGraphSelection(const TSet<class UObject*>& NewSelection, bool bForceRefresh = false);
 #if WITH_RIGVMLEGACYEDITOR
 	virtual TSharedPtr<SKismetInspector> GetKismetInspector() const override { return nullptr; }
 #endif
 
-	virtual void HandleModifiedEvent(ERigVMGraphNotifType InNotifType, URigVMGraph* InGraph, UObject* InSubject) override;
+	UE_API virtual void HandleModifiedEvent(ERigVMGraphNotifType InNotifType, URigVMGraph* InGraph, UObject* InSubject) override;
 	UE_DEPRECATED(5.4, "Please use HandleVMCompiledEvent with ExtendedExecuteContext param.")
 	virtual void HandleVMCompiledEvent(UObject* InCompiledObject, URigVM* InVM) {}
-	virtual void HandleVMCompiledEvent(UObject* InCompiledObject, URigVM* InVM, FRigVMExtendedExecuteContext& InContext);
-	virtual void HandleVMExecutedEvent(URigVMHost* InHost, const FName& InEventName);
-	void HandleVMExecutionHalted(const int32 InstructionIndex, UObject* InNode, const FName& InEntryName);
-	void SetHaltedNode(URigVMNode* Node);
+	UE_API virtual void HandleVMCompiledEvent(UObject* InCompiledObject, URigVM* InVM, FRigVMExtendedExecuteContext& InContext);
+	UE_API virtual void HandleVMExecutedEvent(URigVMHost* InHost, const FName& InEventName);
+	UE_API void HandleVMExecutionHalted(const int32 InstructionIndex, UObject* InNode, const FName& InEntryName);
+	UE_API void SetHaltedNode(URigVMNode* Node);
 	
 	// FNotifyHook Interface
-	void NotifyPreChange(FProperty* PropertyAboutToChange);
-	void NotifyPostChange(const FPropertyChangedEvent& PropertyChangedEvent, FProperty* PropertyThatChanged);
+	UE_API void NotifyPreChange(FProperty* PropertyAboutToChange);
+	UE_API void NotifyPostChange(const FPropertyChangedEvent& PropertyChangedEvent, FProperty* PropertyThatChanged);
 	/** delegate for changing property */
-	virtual void OnFinishedChangingProperties(const FPropertyChangedEvent& PropertyChangedEvent) override;
-	void OnPropertyChanged(UObject* InObject, FPropertyChangedEvent& InEvent);
-	virtual void OnWrappedPropertyChangedChainEvent(URigVMDetailsViewWrapperObject* InWrapperObject, const FString& InPropertyPath, FPropertyChangedChainEvent& InPropertyChangedChainEvent);
-	void OnRequestLocalizeFunctionDialog(FRigVMGraphFunctionIdentifier& InFunction, URigVMController* InTargetController, IRigVMGraphFunctionHost* InTargetFunctionHost, bool bForce);
-	FRigVMController_BulkEditResult OnRequestBulkEditDialog(URigVMBlueprint* InBlueprint, URigVMController* InController, URigVMLibraryNode* InFunction, ERigVMControllerBulkEditType InEditType);
-	bool OnRequestBreakLinksDialog(TArray<URigVMLink*> InLinks);
-	TRigVMTypeIndex OnRequestPinTypeSelectionDialog(const TArray<TRigVMTypeIndex>& InTypes);
+	UE_API virtual void OnFinishedChangingProperties(const FPropertyChangedEvent& PropertyChangedEvent) override;
+	UE_API void OnPropertyChanged(UObject* InObject, FPropertyChangedEvent& InEvent);
+	UE_API virtual void OnWrappedPropertyChangedChainEvent(URigVMDetailsViewWrapperObject* InWrapperObject, const FString& InPropertyPath, FPropertyChangedChainEvent& InPropertyChangedChainEvent);
+	UE_API void OnRequestLocalizeFunctionDialog(FRigVMGraphFunctionIdentifier& InFunction, URigVMController* InTargetController, IRigVMGraphFunctionHost* InTargetFunctionHost, bool bForce);
+	UE_API FRigVMController_BulkEditResult OnRequestBulkEditDialog(URigVMBlueprint* InBlueprint, URigVMController* InController, URigVMLibraryNode* InFunction, ERigVMControllerBulkEditType InEditType);
+	UE_API bool OnRequestBreakLinksDialog(TArray<URigVMLink*> InLinks);
+	UE_API TRigVMTypeIndex OnRequestPinTypeSelectionDialog(const TArray<TRigVMTypeIndex>& InTypes);
 
-	bool UpdateDefaultValueForVariable(FBPVariableDescription& InVariable, bool bUseCDO);
+	UE_API bool UpdateDefaultValueForVariable(FBPVariableDescription& InVariable, bool bUseCDO);
 
 	URigVMController* ActiveController;
 
 	/** Push a newly compiled/opened host to the editor */
-	virtual void UpdateRigVMHost();
+	UE_API virtual void UpdateRigVMHost();
 	virtual void UpdateRigVMHost_PreClearOldHost(URigVMHost* InPreviousHost) {};
 
 	/** Update the name lists for use in name combo boxes */
-	virtual void CacheNameLists();
+	UE_API virtual void CacheNameLists();
 
 	// FGCObject Interface
-	void AddReferencedObjects( FReferenceCollector& Collector );
+	UE_API void AddReferencedObjects( FReferenceCollector& Collector );
 
-	virtual void BindCommands();
+	UE_API virtual void BindCommands();
 
-	void ToggleAutoCompileGraph();
-	bool IsAutoCompileGraphOn() const;
+	UE_API void ToggleAutoCompileGraph();
+	UE_API bool IsAutoCompileGraphOn() const;
 	bool CanAutoCompileGraph() const { return true; }
-	void ToggleEventQueue();
-	void ToggleExecutionMode();
-	TSharedRef<SWidget> GenerateEventQueueMenuContent();
-	TSharedRef<SWidget> GenerateExecutionModeMenuContent();
-	virtual FMenuBuilder GenerateBulkEditMenu();
-	TSharedRef<SWidget> GenerateBulkEditMenuContent();
-	virtual void GenerateEventQueueMenuContent(FMenuBuilder& MenuBuilder);
+	UE_API void ToggleEventQueue();
+	UE_API void ToggleExecutionMode();
+	UE_API TSharedRef<SWidget> GenerateEventQueueMenuContent();
+	UE_API TSharedRef<SWidget> GenerateExecutionModeMenuContent();
+	UE_API virtual FMenuBuilder GenerateBulkEditMenu();
+	UE_API TSharedRef<SWidget> GenerateBulkEditMenuContent();
+	UE_API virtual void GenerateEventQueueMenuContent(FMenuBuilder& MenuBuilder);
 
 	/** Wraps the normal blueprint editor's action menu creation callback */
-	FActionMenuContent HandleCreateGraphActionMenu(UEdGraph* InGraph, const FVector2f& InNodePosition, const TArray<UEdGraphPin*>& InDraggedPins, bool bAutoExpand, SGraphEditor::FActionMenuClosed InOnMenuClosed);
-	void OnNodeTitleCommitted(const FText& NewText, ETextCommit::Type CommitInfo, UEdGraphNode* NodeBeingChanged);
+	UE_API FActionMenuContent HandleCreateGraphActionMenu(UEdGraph* InGraph, const FVector2f& InNodePosition, const TArray<UEdGraphPin*>& InDraggedPins, bool bAutoExpand, SGraphEditor::FActionMenuClosed InOnMenuClosed);
+	UE_API void OnNodeTitleCommitted(const FText& NewText, ETextCommit::Type CommitInfo, UEdGraphNode* NodeBeingChanged);
 
 	/** Undo Action**/
-	void UndoAction();
+	UE_API void UndoAction();
 
 	/** Redo Action **/
-	void RedoAction();
+	UE_API void RedoAction();
 	
-	void OnCreateComment();
+	UE_API void OnCreateComment();
 
 	bool IsDetailsPanelRefreshSuspended() const { return bSuspendDetailsPanelRefresh; }
 	bool& GetSuspendDetailsPanelRefreshFlag() { return bSuspendDetailsPanelRefresh; }
-	TArray<TWeakObjectPtr<UObject>> GetSelectedObjects() const;
-	virtual void SetDetailObjects(const TArray<UObject*>& InObjects);
-	virtual void SetDetailObjects(const TArray<UObject*>& InObjects, bool bChangeUISelectionState);
-	virtual void SetDetailObjectFilter(TSharedPtr<FDetailsViewObjectFilter> InObjectFilter);
-	void SetMemoryStorageDetails(const TArray<FRigVMMemoryStorageStruct*>& InStructs);
-	void SetDetailViewForGraph(URigVMGraph* InGraph);
-	void SetDetailViewForFocusedGraph();
-	void SetDetailViewForLocalVariable();
-	virtual void RefreshDetailView();
-	bool DetailViewShowsAnyRigUnit() const;
-	bool DetailViewShowsLocalVariable() const;
-	bool DetailViewShowsStruct(UScriptStruct* InStruct) const;
-	void ClearDetailObject(bool bChangeUISelectionState = true);
-	void ClearDetailsViewWrapperObjects();
+	UE_API TArray<TWeakObjectPtr<UObject>> GetSelectedObjects() const;
+	UE_API virtual void SetDetailObjects(const TArray<UObject*>& InObjects);
+	UE_API virtual void SetDetailObjects(const TArray<UObject*>& InObjects, bool bChangeUISelectionState);
+	UE_API virtual void SetDetailObjectFilter(TSharedPtr<FDetailsViewObjectFilter> InObjectFilter);
+	UE_API void SetMemoryStorageDetails(const TArray<FRigVMMemoryStorageStruct*>& InStructs);
+	UE_API void SetDetailViewForGraph(URigVMGraph* InGraph);
+	UE_API void SetDetailViewForFocusedGraph();
+	UE_API void SetDetailViewForLocalVariable();
+	UE_API virtual void RefreshDetailView();
+	UE_API bool DetailViewShowsAnyRigUnit() const;
+	UE_API bool DetailViewShowsLocalVariable() const;
+	UE_API bool DetailViewShowsStruct(UScriptStruct* InStruct) const;
+	UE_API void ClearDetailObject(bool bChangeUISelectionState = true);
+	UE_API void ClearDetailsViewWrapperObjects();
 	const TArray<TStrongObjectPtr<URigVMDetailsViewWrapperObject>>& GetWrapperObjects() const { return WrapperObjects; }
 
-	void SetHost(URigVMHost* InHost);
+	UE_API void SetHost(URigVMHost* InHost);
 
-	virtual URigVMGraph* GetFocusedModel() const override;
-	URigVMController* GetFocusedController() const;
-	TSharedPtr<SGraphEditor> GetGraphEditor(UEdGraph* InEdGraph) const;
+	UE_API virtual URigVMGraph* GetFocusedModel() const override;
+	UE_API URigVMController* GetFocusedController() const;
+	UE_API TSharedPtr<SGraphEditor> GetGraphEditor(UEdGraph* InEdGraph) const;
 
 	/** Extend menu */
-	void ExtendMenu();
+	UE_API void ExtendMenu();
 
 	/** Extend toolbar */
-	void ExtendToolbar();
+	UE_API void ExtendToolbar();
 	
 	/** Fill the toolbar with content */
-	virtual void FillToolbar(FToolBarBuilder& ToolbarBuilder, bool bEndSection = true);
+	UE_API virtual void FillToolbar(FToolBarBuilder& ToolbarBuilder, bool bEndSection = true);
 	
-	virtual TArray<FName> GetDefaultEventQueue() const;
-	TArray<FName> GetEventQueue() const;
-	void SetEventQueue(TArray<FName> InEventQueue);
-	virtual void SetEventQueue(TArray<FName> InEventQueue, bool bCompile);
+	UE_API virtual TArray<FName> GetDefaultEventQueue() const;
+	UE_API TArray<FName> GetEventQueue() const;
+	UE_API void SetEventQueue(TArray<FName> InEventQueue);
+	UE_API virtual void SetEventQueue(TArray<FName> InEventQueue, bool bCompile);
 	virtual int32 GetEventQueueComboValue() const { return INDEX_NONE; }
 	virtual FText GetEventQueueLabel() const { return FText(); }
-	virtual FSlateIcon GetEventQueueIcon(const TArray<FName>& InEventQueue) const;
-	FSlateIcon GetEventQueueIcon() const;
-
-	void SetExecutionMode(const ERigVMEditorExecutionModeType InExecutionMode);
-	int32 GetExecutionModeComboValue() const;
-	FText GetExecutionModeLabel() const;
-	static FSlateIcon GetExecutionModeIcon(const ERigVMEditorExecutionModeType InExecutionMode);
-	FSlateIcon GetExecutionModeIcon() const;
+	UE_API virtual FSlateIcon GetEventQueueIcon(const TArray<FName>& InEventQueue) const;
+	UE_API FSlateIcon GetEventQueueIcon() const;
+
+	UE_API void SetExecutionMode(const ERigVMEditorExecutionModeType InExecutionMode);
+	UE_API int32 GetExecutionModeComboValue() const;
+	UE_API FText GetExecutionModeLabel() const;
+	static UE_API FSlateIcon GetExecutionModeIcon(const ERigVMEditorExecutionModeType InExecutionMode);
+	UE_API FSlateIcon GetExecutionModeIcon() const;
 	
-	virtual void GetDebugObjects(TArray<FRigVMCustomDebugObject>& DebugList) const override;
+	UE_API virtual void GetDebugObjects(TArray<FRigVMCustomDebugObject>& DebugList) const override;
 	virtual bool OnlyShowCustomDebugObjects() const override { return true; }
-	void HandleSetObjectBeingDebugged(UObject* InObject);
-	virtual FString GetCustomDebugObjectLabel(UObject* ObjectBeingDebugged) const override;
+	UE_API void HandleSetObjectBeingDebugged(UObject* InObject);
+	UE_API virtual FString GetCustomDebugObjectLabel(UObject* ObjectBeingDebugged) const override;
 
 	/** Handle hiding items in the graph */
-	void HandleHideItem();
-	bool CanHideItem() const;
+	UE_API void HandleHideItem();
+	UE_API bool CanHideItem() const;
 
 	/** Update stale watch pins */
-	void UpdateStaleWatchedPins();
+	UE_API void UpdateStaleWatchedPins();
 	
-	virtual void HandleRefreshEditorFromBlueprint(URigVMBlueprint* InBlueprint);
-	void HandleVariableDroppedFromBlueprint(UObject* InSubject, FProperty* InVariableToDrop, const FVector2D& InDropPosition, const FVector2D& InScreenPosition);
-	void HandleBreakpointAdded();
-	void OnGraphNodeClicked(URigVMEdGraphNode* InNode);
-	void OnNodeDoubleClicked(URigVMBlueprint* InBlueprint, URigVMNode* InNode);
-	void OnGraphImported(UEdGraph* InEdGraph);
-	bool OnActionMatchesName(FEdGraphSchemaAction* InAction, const FName& InName) const;
-	void HandleShowCurrentStatement();
-	void HandleBreakpointActionRequested(const ERigVMBreakpointAction BreakpointAction);
-	bool IsHaltedAtBreakpoint() const;
-	void FrameSelection();
-	void SwapFunctionWithinAsset();
-	void SwapFunctionAcrossProject();
-	void SwapFunctionForAssets(const TArray<FAssetData>& InAssets, bool bSetupUndo);
-	void SwapAssetReferences();
+	UE_API virtual void HandleRefreshEditorFromBlueprint(URigVMBlueprint* InBlueprint);
+	UE_API void HandleVariableDroppedFromBlueprint(UObject* InSubject, FProperty* InVariableToDrop, const FVector2D& InDropPosition, const FVector2D& InScreenPosition);
+	UE_API void HandleBreakpointAdded();
+	UE_API void OnGraphNodeClicked(URigVMEdGraphNode* InNode);
+	UE_API void OnNodeDoubleClicked(URigVMBlueprint* InBlueprint, URigVMNode* InNode);
+	UE_API void OnGraphImported(UEdGraph* InEdGraph);
+	UE_API bool OnActionMatchesName(FEdGraphSchemaAction* InAction, const FName& InName) const;
+	UE_API void HandleShowCurrentStatement();
+	UE_API void HandleBreakpointActionRequested(const ERigVMBreakpointAction BreakpointAction);
+	UE_API bool IsHaltedAtBreakpoint() const;
+	UE_API void FrameSelection();
+	UE_API void SwapFunctionWithinAsset();
+	UE_API void SwapFunctionAcrossProject();
+	UE_API void SwapFunctionForAssets(const TArray<FAssetData>& InAssets, bool bSetupUndo);
+	UE_API void SwapAssetReferences();
 
 
 	/** Once the log is collected update the graph */
-	void UpdateGraphCompilerErrors();
+	UE_API void UpdateGraphCompilerErrors();
 
 	/** Returns true if PIE is currently running */
-	static bool IsPIERunning();
+	static UE_API bool IsPIERunning();
 
-	void OnPIEStopped(bool bSimulation);
+	UE_API void OnPIEStopped(bool bSimulation);
 
 	/** Our currently running rig vm instance */
 	//TObjectPtr<URigVMHost> Host;
@@ -587,7 +589,7 @@ protected:
 	virtual void SetGraphExplorerWidget(TSharedPtr<SRigVMEditorGraphExplorer> InWidget) override { GraphExplorerWidget = InWidget; }
 	virtual TSharedPtr<SRigVMEditorGraphExplorer> GetGraphExplorerWidget() override { return GraphExplorerWidget; }
 
-	bool IsEditingSingleBlueprint() const;
+	UE_API bool IsEditingSingleBlueprint() const;
 	
 protected:
 	bool bAnyErrorsLeft;
@@ -717,4 +719,6 @@ struct FRigVMEditorZoomLevelsContainer : public FZoomLevelsContainer
 	TArray<FRigVMEditorZoomLevelEntry> ZoomLevels;
 };
 
-#undef LOCTEXT_NAMESPACE
\ No newline at end of file
+#undef LOCTEXT_NAMESPACE
+
+#undef UE_API
diff --git a/Engine/Plugins/Runtime/RigVM/Source/RigVMEditor/Public/Editor/RigVMEditorGraphExplorerTabSummoner.h b/Engine/Plugins/Runtime/RigVM/Source/RigVMEditor/Public/Editor/RigVMEditorGraphExplorerTabSummoner.h
--- a/Engine/Plugins/Runtime/RigVM/Source/RigVMEditor/Public/Editor/RigVMEditorGraphExplorerTabSummoner.h
+++ b/Engine/Plugins/Runtime/RigVM/Source/RigVMEditor/Public/Editor/RigVMEditorGraphExplorerTabSummoner.h
@@ -5,19 +5,23 @@
 #include "CoreMinimal.h"
 #include "WorkflowOrientedApp/WorkflowTabFactory.h"
 
+#define UE_API RIGVMEDITOR_API
+
 class IRigVMEditor;
 
-struct RIGVMEDITOR_API FRigVMEditorGraphExplorerTabSummoner : public FWorkflowTabFactory
+struct FRigVMEditorGraphExplorerTabSummoner : public FWorkflowTabFactory
 {
 public:
 
 	static const FName TabID() { static FName ID = TEXT("RigVM Graph Explorer"); return ID; }
 	
 public:
-	FRigVMEditorGraphExplorerTabSummoner(const TSharedRef<IRigVMEditor>& InRigVMEditor);
+	UE_API FRigVMEditorGraphExplorerTabSummoner(const TSharedRef<IRigVMEditor>& InRigVMEditor);
 	
-	virtual TSharedRef<SWidget> CreateTabBody(const FWorkflowTabSpawnInfo& Info) const override;
+	UE_API virtual TSharedRef<SWidget> CreateTabBody(const FWorkflowTabSpawnInfo& Info) const override;
 	
 protected:
 	TWeakPtr<IRigVMEditor> RigVMEditor;
 };
+
+#undef UE_API
diff --git a/Engine/Plugins/Runtime/RigVM/Source/RigVMEditor/Public/Editor/RigVMEditorMenuContext.h b/Engine/Plugins/Runtime/RigVM/Source/RigVMEditor/Public/Editor/RigVMEditorMenuContext.h
--- a/Engine/Plugins/Runtime/RigVM/Source/RigVMEditor/Public/Editor/RigVMEditorMenuContext.h
+++ b/Engine/Plugins/Runtime/RigVM/Source/RigVMEditor/Public/Editor/RigVMEditorMenuContext.h
@@ -7,6 +7,8 @@
 
 #include "RigVMEditorMenuContext.generated.h"
 
+#define UE_API RIGVMEDITOR_API
+
 class URigVMHost;
 class URigVMBlueprint;
 class IRigVMEditor;
@@ -15,7 +17,7 @@ class URigVMNode;
 class URigVMPin;
 
 USTRUCT(BlueprintType)
-struct RIGVMEDITOR_API FRigVMEditorGraphMenuContext
+struct FRigVMEditorGraphMenuContext
 {
 	GENERATED_BODY()
 
@@ -46,8 +48,8 @@ struct RIGVMEDITOR_API FRigVMEditorGraphMenuContext
 	TObjectPtr<const URigVMPin> Pin;
 };
 
-UCLASS(BlueprintType)
-class RIGVMEDITOR_API URigVMEditorMenuContext : public UObject
+UCLASS(MinimalAPI, BlueprintType)
+class URigVMEditorMenuContext : public UObject
 {
 	GENERATED_BODY()
 
@@ -57,28 +59,30 @@ public:
 	 * @param InRigVMEditor 	    The RigVM Editor hosting the menus
 	 * @param InGraphMenuContext 	Additional context for specific menus
 	*/
-	void Init(TWeakPtr<IRigVMEditor> InRigVMEditor, const FRigVMEditorGraphMenuContext& InGraphMenuContext = FRigVMEditorGraphMenuContext());
+	UE_API void Init(TWeakPtr<IRigVMEditor> InRigVMEditor, const FRigVMEditorGraphMenuContext& InGraphMenuContext = FRigVMEditorGraphMenuContext());
 
 	/** Get the rigvm blueprint that we are editing */
 	UFUNCTION(BlueprintCallable, Category = RigVMEditor)
-    URigVMBlueprint* GetRigVMBlueprint() const;
+    UE_API URigVMBlueprint* GetRigVMBlueprint() const;
 	
 	/** Get the active rigvm host instance in the viewport */
 	UFUNCTION(BlueprintCallable, Category = RigVMEditor)
-    URigVMHost* GetRigVMHost() const;
+    UE_API URigVMHost* GetRigVMHost() const;
 
 	/** Returns true if either alt key is down */
 	UFUNCTION(BlueprintCallable, Category = RigVMEditor)
-	bool IsAltDown() const;
+	UE_API bool IsAltDown() const;
 	
 	/** Returns context for graph node context menu */
 	UFUNCTION(BlueprintCallable, Category = RigVMEditor)
-	FRigVMEditorGraphMenuContext GetGraphMenuContext();
+	UE_API FRigVMEditorGraphMenuContext GetGraphMenuContext();
 	
-	IRigVMEditor* GetRigVMEditor() const;
+	UE_API IRigVMEditor* GetRigVMEditor() const;
 
 private:
 	TWeakPtr<IRigVMEditor> WeakRigVMEditor;
 
 	FRigVMEditorGraphMenuContext GraphMenuContext;
-};
\ No newline at end of file
+};
+
+#undef UE_API
diff --git a/Engine/Plugins/Runtime/RigVM/Source/RigVMEditor/Public/Editor/RigVMEditorStyle.h b/Engine/Plugins/Runtime/RigVM/Source/RigVMEditor/Public/Editor/RigVMEditorStyle.h
--- a/Engine/Plugins/Runtime/RigVM/Source/RigVMEditor/Public/Editor/RigVMEditorStyle.h
+++ b/Engine/Plugins/Runtime/RigVM/Source/RigVMEditor/Public/Editor/RigVMEditorStyle.h
@@ -10,7 +10,7 @@
 #include "Styling/StyleColors.h"
 #include "Brushes/SlateRoundedBoxBrush.h"
 
-class RIGVMEDITOR_API FRigVMEditorStyle final
+class FRigVMEditorStyle final
 	: public FSlateStyleSet
 {
 	class FContentRootBracket
diff --git a/Engine/Plugins/Runtime/RigVM/Source/RigVMEditor/Public/Editor/RigVMEditorTools.h b/Engine/Plugins/Runtime/RigVM/Source/RigVMEditor/Public/Editor/RigVMEditorTools.h
--- a/Engine/Plugins/Runtime/RigVM/Source/RigVMEditor/Public/Editor/RigVMEditorTools.h
+++ b/Engine/Plugins/Runtime/RigVM/Source/RigVMEditor/Public/Editor/RigVMEditorTools.h
@@ -7,6 +7,8 @@
 #include "FrontendFilterBase.h"
 #include "RigVMCore/RigVMVariant.h"
 
+#define UE_API RIGVMEDITOR_API
+
 class URigVMController;
 class URigVMGraph;
 class URigVMFunctionLibrary;
@@ -41,17 +43,17 @@ RIGVMEDITOR_API void OnRequestLocalizeFunctionDialog(FRigVMGraphFunctionIdentifi
 	
 RIGVMEDITOR_API FAssetData FindAssetFromAnyPath(const FString& InPartialOrFullPath, bool bConvertToRootPath);
 
-class RIGVMEDITOR_API FFilterByAssetTag : public FFrontendFilter
+class FFilterByAssetTag : public FFrontendFilter
 {
 public:
-	FFilterByAssetTag(TSharedPtr<FFrontendFilterCategory> InCategory, const FRigVMTag& InTag);
+	UE_API FFilterByAssetTag(TSharedPtr<FFrontendFilterCategory> InCategory, const FRigVMTag& InTag);
 
 	FString GetName() const override { return Tag.bMarksSubjectAsInvalid ? TEXT("Exclude ") + Tag.Name.ToString() : Tag.Name.ToString(); }
 	FText GetDisplayName() const override { return FText::FromString(Tag.bMarksSubjectAsInvalid ? TEXT("Exclude ") + Tag.GetLabel() : Tag.GetLabel()); }
 	FText GetToolTipText() const override { return Tag.ToolTip; }
 	FLinearColor GetColor() const override { return Tag.Color; };
 
-	bool PassesFilter(const FContentBrowserItem& InItem) const override;
+	UE_API bool PassesFilter(const FContentBrowserItem& InItem) const override;
 	bool ShouldBeMarkedAsInvalid() const { return Tag.bMarksSubjectAsInvalid; }
 
 private:
@@ -59,3 +61,5 @@ private:
 };
 
 } // end namespace UE::RigVM::Editor::Tools
+
+#undef UE_API
diff --git a/Engine/Plugins/Runtime/RigVM/Source/RigVMEditor/Public/Editor/RigVMExecutionStackTabSummoner.h b/Engine/Plugins/Runtime/RigVM/Source/RigVMEditor/Public/Editor/RigVMExecutionStackTabSummoner.h
--- a/Engine/Plugins/Runtime/RigVM/Source/RigVMEditor/Public/Editor/RigVMExecutionStackTabSummoner.h
+++ b/Engine/Plugins/Runtime/RigVM/Source/RigVMEditor/Public/Editor/RigVMExecutionStackTabSummoner.h
@@ -5,19 +5,23 @@
 #include "CoreMinimal.h"
 #include "WorkflowOrientedApp/WorkflowTabFactory.h"
 
+#define UE_API RIGVMEDITOR_API
+
 class IRigVMEditor;
 
-struct RIGVMEDITOR_API FRigVMExecutionStackTabSummoner : public FWorkflowTabFactory
+struct FRigVMExecutionStackTabSummoner : public FWorkflowTabFactory
 {
 public:
 
 	static inline const FLazyName TabID = FLazyName(TEXT("Execution Stack"));
 	
 public:
-	FRigVMExecutionStackTabSummoner(const TSharedRef<IRigVMEditor>& InRigVMEditor);
+	UE_API FRigVMExecutionStackTabSummoner(const TSharedRef<IRigVMEditor>& InRigVMEditor);
 	
-	virtual TSharedRef<SWidget> CreateTabBody(const FWorkflowTabSpawnInfo& Info) const override;
+	UE_API virtual TSharedRef<SWidget> CreateTabBody(const FWorkflowTabSpawnInfo& Info) const override;
 	
 protected:
 	TWeakPtr<IRigVMEditor> RigVMEditor;
 };
+
+#undef UE_API
diff --git a/Engine/Plugins/Runtime/RigVM/Source/RigVMEditor/Public/Editor/RigVMFindReferences.h b/Engine/Plugins/Runtime/RigVM/Source/RigVMEditor/Public/Editor/RigVMFindReferences.h
--- a/Engine/Plugins/Runtime/RigVM/Source/RigVMEditor/Public/Editor/RigVMFindReferences.h
+++ b/Engine/Plugins/Runtime/RigVM/Source/RigVMEditor/Public/Editor/RigVMFindReferences.h
@@ -9,6 +9,8 @@
 #include "Textures/SlateIcon.h"
 #include "EdGraph/EdGraphSchema.h"
 
+#define UE_API RIGVMEDITOR_API
+
 class SWidget;
 class SDockTab;
 class UBlueprint;
@@ -16,30 +18,30 @@ class FUICommandList;
 class URigVMBlueprint;
 
 
-class RIGVMEDITOR_API FRigVMFindResult : public TSharedFromThis< FRigVMFindResult >
+class FRigVMFindResult : public TSharedFromThis< FRigVMFindResult >
 {
 public:
 	FRigVMFindResult() = default;
 	virtual ~FRigVMFindResult() = default;
 
 	/* Create a root */
-	explicit FRigVMFindResult(TWeakObjectPtr<URigVMBlueprint> InBlueprint);
-	explicit FRigVMFindResult(TWeakObjectPtr<URigVMBlueprint> InBlueprint, const FText& InDisplayText);
+	UE_API explicit FRigVMFindResult(TWeakObjectPtr<URigVMBlueprint> InBlueprint);
+	UE_API explicit FRigVMFindResult(TWeakObjectPtr<URigVMBlueprint> InBlueprint, const FText& InDisplayText);
 
 	/* Called when user clicks on the search item */
-	virtual FReply OnClick();
+	UE_API virtual FReply OnClick();
 
 	/* Get Category for this search result */
-	virtual FText GetCategory() const;
+	UE_API virtual FText GetCategory() const;
 
 	/* Create an icon to represent the result */
-	virtual TSharedRef<SWidget>	CreateIcon() const;
+	UE_API virtual TSharedRef<SWidget>	CreateIcon() const;
 
 	/** Finalizes any content for the search data that was unsafe to do on a separate thread */
 	virtual void FinalizeSearchData() {};
 
 	/** gets the blueprint housing all these search results */
-	URigVMBlueprint* GetBlueprint() const;
+	UE_API URigVMBlueprint* GetBlueprint() const;
 
 	/**
 	* Parses search info for specific data important for displaying the search result in an easy to understand format
@@ -50,10 +52,10 @@ public:
 	virtual void ParseSearchInfo(FText InKey, FText InValue) {};
 
 	/** Returns the Object represented by this search information give the Blueprint it can be found in */
-	virtual UObject* GetObject(URigVMBlueprint* InBlueprint) const;
+	UE_API virtual UObject* GetObject(URigVMBlueprint* InBlueprint) const;
 
 	/** Returns the display string for the row */
-	FText GetDisplayString() const;
+	UE_API FText GetDisplayString() const;
 
 public:
 	/*Any children listed under this category */
@@ -200,55 +202,55 @@ private:
 };
 
 /*Widget for searching for (functions/events) across all blueprints or just a single blueprint */
-class RIGVMEDITOR_API SRigVMFindReferences: public SCompoundWidget
+class SRigVMFindReferences: public SCompoundWidget
 {
 public:
 	SLATE_BEGIN_ARGS( SRigVMFindReferences )
 	{}
 	SLATE_END_ARGS()
 
-	void Construct(const FArguments& InArgs, TSharedPtr<class FRigVMEditorBase> InBlueprintEditor = nullptr);
+	UE_API void Construct(const FArguments& InArgs, TSharedPtr<class FRigVMEditorBase> InBlueprintEditor = nullptr);
 
 	/** Focuses this widget's search box, and changes the mode as well, and optionally the search terms */
-	void FocusForUse(bool bSetFindWithinBlueprint, FString NewSearchTerms = FString(), bool bSelectFirstResult = false);
+	UE_API void FocusForUse(bool bSetFindWithinBlueprint, FString NewSearchTerms = FString(), bool bSelectFirstResult = false);
 
 	/** SWidget overrides */
-	virtual FReply OnKeyDown(const FGeometry& MyGeometry, const FKeyEvent& InKeyEvent) override;
+	UE_API virtual FReply OnKeyDown(const FGeometry& MyGeometry, const FKeyEvent& InKeyEvent) override;
 
 private:
 	
 	/** The main function that will find references and build the tree */
-	void FindReferences(const FString& SearchTerms);
+	UE_API void FindReferences(const FString& SearchTerms);
 
 	/** Register any Find-in-Blueprint commands */
-	void RegisterCommands();
+	UE_API void RegisterCommands();
 
 	/*Called when user changes the text they are searching for */
-	void OnSearchTextChanged(const FText& Text);
+	UE_API void OnSearchTextChanged(const FText& Text);
 
 	/*Called when user changes commits text to the search box */
-	void OnSearchTextCommitted(const FText& Text, ETextCommit::Type CommitType);
+	UE_API void OnSearchTextCommitted(const FText& Text, ETextCommit::Type CommitType);
 
 	/* Get the children of a row */
-	void OnGetChildren( FRigVMSearchResult InItem, TArray< FRigVMSearchResult >& OutChildren );
+	UE_API void OnGetChildren( FRigVMSearchResult InItem, TArray< FRigVMSearchResult >& OutChildren );
 
 	/* Called when user double clicks on a new result */
-	void OnTreeSelectionDoubleClicked( FRigVMSearchResult Item );
+	UE_API void OnTreeSelectionDoubleClicked( FRigVMSearchResult Item );
 
 	/* Called when a new row is being generated */
-	TSharedRef<ITableRow> OnGenerateRow(FRigVMSearchResult InItem, const TSharedRef<STableViewBase>& OwnerTable);
+	UE_API TSharedRef<ITableRow> OnGenerateRow(FRigVMSearchResult InItem, const TSharedRef<STableViewBase>& OwnerTable);
 
 	/** Callback to build the context menu when right clicking in the tree */
-	TSharedPtr<SWidget> OnContextMenuOpening();
+	UE_API TSharedPtr<SWidget> OnContextMenuOpening();
 
 	/** Helper function to select all items */
-	void SelectAllItemsHelper(FRigVMSearchResult InItemToSelect);
+	UE_API void SelectAllItemsHelper(FRigVMSearchResult InItemToSelect);
 
 	/** Callback when user attempts to select all items in the search results */
-	void OnSelectAllAction();
+	UE_API void OnSelectAllAction();
 
 	/** Callback when user attempts to copy their selection in the Find-in-Blueprints */
-	void OnCopyAction();
+	UE_API void OnCopyAction();
 
 private:
 	/** Pointer back to the blueprint editor that owns us */
@@ -274,4 +276,6 @@ private:
 
 	/** Commands handled by this widget */
 	TSharedPtr< FUICommandList > CommandList;
-};
\ No newline at end of file
+};
+
+#undef UE_API
diff --git a/Engine/Plugins/Runtime/RigVM/Source/RigVMEditor/Public/Editor/RigVMFindReferencesTabSummoner.h b/Engine/Plugins/Runtime/RigVM/Source/RigVMEditor/Public/Editor/RigVMFindReferencesTabSummoner.h
--- a/Engine/Plugins/Runtime/RigVM/Source/RigVMEditor/Public/Editor/RigVMFindReferencesTabSummoner.h
+++ b/Engine/Plugins/Runtime/RigVM/Source/RigVMEditor/Public/Editor/RigVMFindReferencesTabSummoner.h
@@ -7,21 +7,23 @@
 #include "Toolkits/AssetEditorToolkit.h"
 #include "WorkflowOrientedApp/WorkflowTabFactory.h"
 
+#define UE_API RIGVMEDITOR_API
+
 #define LOCTEXT_NAMESPACE "RigVMEditor"
 
 //////////////////////////////////////////////////////////////////////////
 // FCompilerResultsSummoner
 
-struct RIGVMEDITOR_API FRigVMFindReferencesTabSummoner : public FWorkflowTabFactory
+struct FRigVMFindReferencesTabSummoner : public FWorkflowTabFactory
 {
 
 public:
 
 	static const FName TabID() { static FName ID = TEXT("RigVM Find References"); return ID; }
 	
-	FRigVMFindReferencesTabSummoner(TSharedPtr<class FAssetEditorToolkit> InHostingApp);
+	UE_API FRigVMFindReferencesTabSummoner(TSharedPtr<class FAssetEditorToolkit> InHostingApp);
 
-	virtual TSharedRef<SWidget> CreateTabBody(const FWorkflowTabSpawnInfo& Info) const override;
+	UE_API virtual TSharedRef<SWidget> CreateTabBody(const FWorkflowTabSpawnInfo& Info) const override;
 
 	virtual FText GetTabToolTipText(const FWorkflowTabSpawnInfo& Info) const override
 	{
@@ -29,4 +31,6 @@ public:
 	}
 };
 
-#undef LOCTEXT_NAMESPACE
\ No newline at end of file
+#undef LOCTEXT_NAMESPACE
+
+#undef UE_API
