Updated headers using UnrealCodeFixup to make sure dllstorage is on methods/staticvars instead of types. Part 4/n

diff --git a/Engine/Plugins/Runtime/RigVM/Source/RigVMEditor/Public/Editor/RigVMGraphDetailCustomization.h b/Engine/Plugins/Runtime/RigVM/Source/RigVMEditor/Public/Editor/RigVMGraphDetailCustomization.h
--- a/Engine/Plugins/Runtime/RigVM/Source/RigVMEditor/Public/Editor/RigVMGraphDetailCustomization.h
+++ b/Engine/Plugins/Runtime/RigVM/Source/RigVMEditor/Public/Editor/RigVMGraphDetailCustomization.h
@@ -21,24 +21,26 @@
 #include "Widgets/SRigVMLogWidget.h"
 #include "HAL/PlatformApplicationMisc.h"
 
+#define UE_API RIGVMEDITOR_API
+
 class IDetailLayoutBuilder;
 class FRigVMGraphDetailCustomizationImpl;
 
-class RIGVMEDITOR_API FRigVMFunctionArgumentGroupLayout : public IDetailCustomNodeBuilder, public TSharedFromThis<FRigVMFunctionArgumentGroupLayout>
+class FRigVMFunctionArgumentGroupLayout : public IDetailCustomNodeBuilder, public TSharedFromThis<FRigVMFunctionArgumentGroupLayout>
 {
 public:
-	FRigVMFunctionArgumentGroupLayout(
+	UE_API FRigVMFunctionArgumentGroupLayout(
 		const TWeakObjectPtr<URigVMGraph>& InGraph,
 		const TWeakInterfacePtr<IRigVMClientHost>& InRigVMClientHost,
 		const TWeakPtr<IRigVMEditor>& InEditor,
 		bool bInputs);
-	virtual ~FRigVMFunctionArgumentGroupLayout();
+	UE_API virtual ~FRigVMFunctionArgumentGroupLayout();
 
 private:
 	/** IDetailCustomNodeBuilder Interface*/
 	virtual void SetOnRebuildChildren(FSimpleDelegate InOnRegenerateChildren) override { OnRebuildChildren = InOnRegenerateChildren; }
 	virtual void GenerateHeaderRowContent(FDetailWidgetRow& NodeRow) override {}
-	virtual void GenerateChildContent(IDetailChildrenBuilder& ChildrenBuilder) override;
+	UE_API virtual void GenerateChildContent(IDetailChildrenBuilder& ChildrenBuilder) override;
 	virtual void Tick(float DeltaTime) override {}
 	virtual bool RequiresTick() const override { return false; }
 	virtual FName GetName() const override { return NAME_None; }
@@ -46,7 +48,7 @@ private:
 
 private:
 
-	void HandleModifiedEvent(ERigVMGraphNotifType InNotifType, URigVMGraph* InGraph, UObject* InSubject);
+	UE_API void HandleModifiedEvent(ERigVMGraphNotifType InNotifType, URigVMGraph* InGraph, UObject* InSubject);
 
 	TWeakObjectPtr<URigVMGraph> GraphPtr;
 	TWeakInterfacePtr<IRigVMClientHost> WeakRigVMClientHost;
@@ -55,7 +57,7 @@ private:
 	FSimpleDelegate OnRebuildChildren;
 };
 
-class RIGVMEDITOR_API FRigVMFunctionArgumentLayout : public IDetailCustomNodeBuilder, public TSharedFromThis<FRigVMFunctionArgumentLayout>
+class FRigVMFunctionArgumentLayout : public IDetailCustomNodeBuilder, public TSharedFromThis<FRigVMFunctionArgumentLayout>
 {
 public:
 
@@ -74,8 +76,8 @@ public:
 private:
 
 	/** IDetailCustomNodeBuilder Interface*/
-	virtual void GenerateHeaderRowContent(FDetailWidgetRow& NodeRow) override;
-	virtual void GenerateChildContent(IDetailChildrenBuilder& ChildrenBuilder) override;
+	UE_API virtual void GenerateHeaderRowContent(FDetailWidgetRow& NodeRow) override;
+	UE_API virtual void GenerateChildContent(IDetailChildrenBuilder& ChildrenBuilder) override;
 	virtual void Tick(float DeltaTime) override {}
 	virtual bool RequiresTick() const override { return false; }
 	virtual FName GetName() const override { return PinPtr.Get()->GetFName(); }
@@ -84,26 +86,26 @@ private:
 private:
 
 	/** Determines if this pin should not be editable */
-	bool ShouldPinBeReadOnly(bool bIsEditingPinType = false) const;
+	UE_API bool ShouldPinBeReadOnly(bool bIsEditingPinType = false) const;
 
 	/** Determines if editing the pins on the node should be read only */
-	bool IsPinEditingReadOnly(bool bIsEditingPinType = false) const;
+	UE_API bool IsPinEditingReadOnly(bool bIsEditingPinType = false) const;
 
 	/** Determines if an argument can be moved up or down */
-	bool CanArgumentBeMoved(bool bMoveUp) const;
+	UE_API bool CanArgumentBeMoved(bool bMoveUp) const;
 
 	/** Callbacks for all the functionality for modifying arguments */
-	void OnRemoveClicked();
-	FReply OnArgMoveUp();
-	FReply OnArgMoveDown();
+	UE_API void OnRemoveClicked();
+	UE_API FReply OnArgMoveUp();
+	UE_API FReply OnArgMoveDown();
 
-	FText OnGetArgNameText() const;
-	FText OnGetArgToolTipText() const;
-	void OnArgNameTextCommitted(const FText& NewText, ETextCommit::Type InTextCommit);
+	UE_API FText OnGetArgNameText() const;
+	UE_API FText OnGetArgToolTipText() const;
+	UE_API void OnArgNameTextCommitted(const FText& NewText, ETextCommit::Type InTextCommit);
 
-	FEdGraphPinType OnGetPinInfo() const;
-	void PinInfoChanged(const FEdGraphPinType& PinType);
-	void OnPrePinInfoChange(const FEdGraphPinType& PinType);
+	UE_API FEdGraphPinType OnGetPinInfo() const;
+	UE_API void PinInfoChanged(const FEdGraphPinType& PinType);
+	UE_API void OnPrePinInfoChange(const FEdGraphPinType& PinType);
 
 private:
 
@@ -126,20 +128,20 @@ private:
 	FRigVMLocalVariableNameValidator NameValidator;
 };
 
-class RIGVMEDITOR_API FRigVMFunctionArgumentDefaultNode : public IDetailCustomNodeBuilder, public TSharedFromThis<FRigVMFunctionArgumentDefaultNode>
+class FRigVMFunctionArgumentDefaultNode : public IDetailCustomNodeBuilder, public TSharedFromThis<FRigVMFunctionArgumentDefaultNode>
 {
 public:
-	FRigVMFunctionArgumentDefaultNode(
+	UE_API FRigVMFunctionArgumentDefaultNode(
 		const TWeakObjectPtr<URigVMGraph>& InGraph,
 		const TWeakInterfacePtr<IRigVMClientHost>& InClientHost
 	);
-	virtual ~FRigVMFunctionArgumentDefaultNode();
+	UE_API virtual ~FRigVMFunctionArgumentDefaultNode();
 
 private:
 	/** IDetailCustomNodeBuilder Interface*/
 	virtual void SetOnRebuildChildren(FSimpleDelegate InOnRegenerateChildren) override { OnRebuildChildren = InOnRegenerateChildren; }
 	virtual void GenerateHeaderRowContent(FDetailWidgetRow& NodeRow) override {}
-	virtual void GenerateChildContent(IDetailChildrenBuilder& ChildrenBuilder) override;
+	UE_API virtual void GenerateChildContent(IDetailChildrenBuilder& ChildrenBuilder) override;
 	virtual void Tick(float DeltaTime) override {}
 	virtual bool RequiresTick() const override { return false; }
 	virtual FName GetName() const override { return NAME_None; }
@@ -147,8 +149,8 @@ private:
 
 private:
 
-	void OnGraphChanged(const FEdGraphEditAction& InAction);
-	void HandleModifiedEvent(ERigVMGraphNotifType InNotifType, URigVMGraph* InGraph, UObject* InSubject);
+	UE_API void OnGraphChanged(const FEdGraphEditAction& InAction);
+	UE_API void HandleModifiedEvent(ERigVMGraphNotifType InNotifType, URigVMGraph* InGraph, UObject* InSubject);
 
 	TWeakObjectPtr<URigVMGraph> GraphPtr;
 	TWeakObjectPtr<URigVMEdGraph> EdGraphOuterPtr;
@@ -159,20 +161,20 @@ private:
 };
 
 /** Customization for editing rig vm graphs */
-class RIGVMEDITOR_API FRigVMGraphDetailCustomization : public IDetailCustomization
+class FRigVMGraphDetailCustomization : public IDetailCustomization
 {
 public:
 	/** Makes a new instance of this detail layout class for a specific detail view requesting it */
-	static TSharedRef<IDetailCustomization> MakeInstance(TSharedPtr<IRigVMEditor> InBlueprintEditor, const UClass* InExpectedBlueprintClass);
-	FRigVMGraphDetailCustomization(TSharedPtr<IRigVMEditor> RigVMigEditor, URigVMBlueprint* RigVMBlueprint);
+	static UE_API TSharedRef<IDetailCustomization> MakeInstance(TSharedPtr<IRigVMEditor> InBlueprintEditor, const UClass* InExpectedBlueprintClass);
+	UE_API FRigVMGraphDetailCustomization(TSharedPtr<IRigVMEditor> RigVMigEditor, URigVMBlueprint* RigVMBlueprint);
 #if WITH_RIGVMLEGACYEDITOR
-	static TSharedPtr<IDetailCustomization> MakeLegacyInstance(TSharedPtr<IBlueprintEditor> InBlueprintEditor, const UClass* InExpectedBlueprintClass);
-	FRigVMGraphDetailCustomization(TSharedPtr<IBlueprintEditor> RigVMigEditor, URigVMBlueprint* RigVMBlueprint);
+	static UE_API TSharedPtr<IDetailCustomization> MakeLegacyInstance(TSharedPtr<IBlueprintEditor> InBlueprintEditor, const UClass* InExpectedBlueprintClass);
+	UE_API FRigVMGraphDetailCustomization(TSharedPtr<IBlueprintEditor> RigVMigEditor, URigVMBlueprint* RigVMBlueprint);
 #endif
 
 
 	// IDetailCustomization interface
-	virtual void CustomizeDetails(IDetailLayoutBuilder& DetailLayout) override;
+	UE_API virtual void CustomizeDetails(IDetailLayoutBuilder& DetailLayout) override;
 
 private:
 	/** The editor we are embedded in */
@@ -187,10 +189,10 @@ private:
 	TSharedPtr<FRigVMGraphDetailCustomizationImpl> RigVMGraphDetailCustomizationImpl;
 };
 
-class RIGVMEDITOR_API FRigVMGraphDetailCustomizationImpl : public TSharedFromThis<FRigVMGraphDetailCustomizationImpl>
+class FRigVMGraphDetailCustomizationImpl : public TSharedFromThis<FRigVMGraphDetailCustomizationImpl>
 {
 public:
-	void CustomizeDetails(IDetailLayoutBuilder& DetailLayout,
+	UE_API void CustomizeDetails(IDetailLayoutBuilder& DetailLayout,
 		URigVMGraph* Model,
 		URigVMController* Controller,
 		IRigVMClientHost* InRigVMClientHost,
@@ -198,59 +200,59 @@ public:
 
 private:
 
-	bool IsAddNewInputOutputEnabled() const;
-	EVisibility GetAddNewInputOutputVisibility() const;
-	FReply OnAddNewInputClicked();
-	FReply OnAddNewOutputClicked();
-	FText GetNodeCategory() const;
-	void SetNodeCategory(const FText& InNewText, ETextCommit::Type InCommitType);
-	FText GetNodeKeywords() const;
-	void SetNodeKeywords(const FText& InNewText, ETextCommit::Type InCommitType);
-	FText GetNodeDescription() const;
-	void SetNodeDescription(const FText& InNewText, ETextCommit::Type InCommitType);
-	FLinearColor GetNodeColor() const;
-	void SetNodeColor(FLinearColor InColor, bool bSetupUndoRedo);
-	void OnNodeColorBegin();
-	void OnNodeColorEnd();
-	void OnNodeColorCancelled(FLinearColor OriginalColor);
-	FReply OnNodeColorClicked();
-	FText GetCurrentAccessSpecifierName() const;
-	void OnAccessSpecifierSelected( TSharedPtr<FRigVMStringWithTag> SpecifierName, ESelectInfo::Type SelectInfo );
-	TSharedRef<ITableRow> HandleGenerateRowAccessSpecifier( TSharedPtr<FRigVMStringWithTag> SpecifierName, const TSharedRef<STableViewBase>& OwnerTable );
-	bool IsValidFunction() const;
-	FRigVMVariant GetVariant() const;
-	FRigVMVariantRef GetSubjectVariantRef() const;
-	TArray<FRigVMVariantRef> GetVariantRefs() const;
+	UE_API bool IsAddNewInputOutputEnabled() const;
+	UE_API EVisibility GetAddNewInputOutputVisibility() const;
+	UE_API FReply OnAddNewInputClicked();
+	UE_API FReply OnAddNewOutputClicked();
+	UE_API FText GetNodeCategory() const;
+	UE_API void SetNodeCategory(const FText& InNewText, ETextCommit::Type InCommitType);
+	UE_API FText GetNodeKeywords() const;
+	UE_API void SetNodeKeywords(const FText& InNewText, ETextCommit::Type InCommitType);
+	UE_API FText GetNodeDescription() const;
+	UE_API void SetNodeDescription(const FText& InNewText, ETextCommit::Type InCommitType);
+	UE_API FLinearColor GetNodeColor() const;
+	UE_API void SetNodeColor(FLinearColor InColor, bool bSetupUndoRedo);
+	UE_API void OnNodeColorBegin();
+	UE_API void OnNodeColorEnd();
+	UE_API void OnNodeColorCancelled(FLinearColor OriginalColor);
+	UE_API FReply OnNodeColorClicked();
+	UE_API FText GetCurrentAccessSpecifierName() const;
+	UE_API void OnAccessSpecifierSelected( TSharedPtr<FRigVMStringWithTag> SpecifierName, ESelectInfo::Type SelectInfo );
+	UE_API TSharedRef<ITableRow> HandleGenerateRowAccessSpecifier( TSharedPtr<FRigVMStringWithTag> SpecifierName, const TSharedRef<STableViewBase>& OwnerTable );
+	UE_API bool IsValidFunction() const;
+	UE_API FRigVMVariant GetVariant() const;
+	UE_API FRigVMVariantRef GetSubjectVariantRef() const;
+	UE_API TArray<FRigVMVariantRef> GetVariantRefs() const;
 
 private:
 
-	void OnVariantChanged(const FRigVMVariant& InVariant);
-	void OnBrowseVariantRef(const FRigVMVariantRef& InVariantRef);
-	TArray<FRigVMTag> OnGetAssignedTags() const;
-	void OnAddAssignedTag(const FName& InTagName);
-	void OnRemoveAssignedTag(const FName& InTagName);
-
-	URigVMLibraryNode* GetLibraryNode() const;
-	URigVMNode* GetNodeForLayout() const;
-	const FRigVMNodeLayout* GetNodeLayout() const;
-	TArray<FString> GetUncategorizedPins() const;
-	TArray<FRigVMPinCategory> GetPinCategories() const;
-	FString GetPinCategory(FString InPinPath) const;
-	int32 GetPinIndexInCategory(FString InPinPath) const;
-	FString GetPinLabel(FString InPinPath) const;
-	FLinearColor GetPinColor(FString InPinPath) const;
-	const FSlateBrush* GetPinIcon(FString InPinPath) const;
-	void HandleCategoryAdded(FString InCategory);
-	void HandleCategoryRemoved(FString InCategory);
-	void HandleCategoryRenamed(FString InOldCategory, FString InNewCategory);
-	void HandlePinCategoryChanged(FString InPinPath, FString InCategory);
-	void HandlePinLabelChanged(FString InPinPath, FString InNewLabel);
-	void HandlePinIndexInCategoryChanged(FString InPinPath, int32 InIndexInCategory);
-	static bool ValidateName(FString InNewName, FText& OutErrorMessage);
-	bool HandleValidateCategoryName(FString InCategoryPath, FString InNewName, FText& OutErrorMessage);
-	bool HandleValidatePinDisplayName(FString InPinPath, FString InNewName, FText& OutErrorMessage);
-
-	uint32 GetNodeLayoutHash() const;
+	UE_API void OnVariantChanged(const FRigVMVariant& InVariant);
+	UE_API void OnBrowseVariantRef(const FRigVMVariantRef& InVariantRef);
+	UE_API TArray<FRigVMTag> OnGetAssignedTags() const;
+	UE_API void OnAddAssignedTag(const FName& InTagName);
+	UE_API void OnRemoveAssignedTag(const FName& InTagName);
+
+	UE_API URigVMLibraryNode* GetLibraryNode() const;
+	UE_API URigVMNode* GetNodeForLayout() const;
+	UE_API const FRigVMNodeLayout* GetNodeLayout() const;
+	UE_API TArray<FString> GetUncategorizedPins() const;
+	UE_API TArray<FRigVMPinCategory> GetPinCategories() const;
+	UE_API FString GetPinCategory(FString InPinPath) const;
+	UE_API int32 GetPinIndexInCategory(FString InPinPath) const;
+	UE_API FString GetPinLabel(FString InPinPath) const;
+	UE_API FLinearColor GetPinColor(FString InPinPath) const;
+	UE_API const FSlateBrush* GetPinIcon(FString InPinPath) const;
+	UE_API void HandleCategoryAdded(FString InCategory);
+	UE_API void HandleCategoryRemoved(FString InCategory);
+	UE_API void HandleCategoryRenamed(FString InOldCategory, FString InNewCategory);
+	UE_API void HandlePinCategoryChanged(FString InPinPath, FString InCategory);
+	UE_API void HandlePinLabelChanged(FString InPinPath, FString InNewLabel);
+	UE_API void HandlePinIndexInCategoryChanged(FString InPinPath, int32 InIndexInCategory);
+	static UE_API bool ValidateName(FString InNewName, FText& OutErrorMessage);
+	UE_API bool HandleValidateCategoryName(FString InCategoryPath, FString InNewName, FText& OutErrorMessage);
+	UE_API bool HandleValidatePinDisplayName(FString InPinPath, FString InNewName, FText& OutErrorMessage);
+
+	UE_API uint32 GetNodeLayoutHash() const;
 
 	/** The graph we are editing */
 	TWeakObjectPtr<URigVMGraph> WeakModel;
@@ -270,30 +272,30 @@ private:
 	/** Set to true if the UI is currently picking a color */
 	bool bIsPickingColor;
 
-	static TArray<TSharedPtr<FRigVMStringWithTag>> AccessSpecifierStrings;
+	static UE_API TArray<TSharedPtr<FRigVMStringWithTag>> AccessSpecifierStrings;
 	mutable TOptional<FRigVMNodeLayout> CachedNodeLayout;
 };
 
 /** Customization for editing a rig vm node */
-class RIGVMEDITOR_API FRigVMWrappedNodeDetailCustomization : public IDetailCustomization
+class FRigVMWrappedNodeDetailCustomization : public IDetailCustomization
 {
 public:
 	
-	FRigVMWrappedNodeDetailCustomization();
+	UE_API FRigVMWrappedNodeDetailCustomization();
 
 	/** Makes a new instance of this detail layout class for a specific detail view requesting it */
-	static TSharedRef<IDetailCustomization> MakeInstance();
+	static UE_API TSharedRef<IDetailCustomization> MakeInstance();
 
 	// IDetailCustomization interface
-	virtual void CustomizeDetails(IDetailLayoutBuilder& DetailLayout) override;
+	UE_API virtual void CustomizeDetails(IDetailLayoutBuilder& DetailLayout) override;
 
-	TSharedRef<SWidget> MakeNameListItemWidget(TSharedPtr<FRigVMStringWithTag> InItem);
-	FText GetNameListText(const FNameProperty* InProperty) const;
-	TSharedPtr<FRigVMStringWithTag> GetCurrentlySelectedItem(const FNameProperty* InProperty, const TArray<TSharedPtr<FRigVMStringWithTag>>* InNameList) const;
-	void SetNameListText(const FText& NewTypeInValue, ETextCommit::Type, const FNameProperty* InProperty, TSharedRef<IPropertyUtilities> PropertyUtilities);
-	void OnNameListChanged(TSharedPtr<FRigVMStringWithTag> NewSelection, ESelectInfo::Type SelectInfo, const FNameProperty* InProperty, TSharedRef<IPropertyUtilities> PropertyUtilities);
-	void OnNameListComboBox(const FNameProperty* InProperty, const TArray<TSharedPtr<FRigVMStringWithTag>>* InNameList);
-	void CustomizeLiveValues(IDetailLayoutBuilder& DetailLayout);
+	UE_API TSharedRef<SWidget> MakeNameListItemWidget(TSharedPtr<FRigVMStringWithTag> InItem);
+	UE_API FText GetNameListText(const FNameProperty* InProperty) const;
+	UE_API TSharedPtr<FRigVMStringWithTag> GetCurrentlySelectedItem(const FNameProperty* InProperty, const TArray<TSharedPtr<FRigVMStringWithTag>>* InNameList) const;
+	UE_API void SetNameListText(const FText& NewTypeInValue, ETextCommit::Type, const FNameProperty* InProperty, TSharedRef<IPropertyUtilities> PropertyUtilities);
+	UE_API void OnNameListChanged(TSharedPtr<FRigVMStringWithTag> NewSelection, ESelectInfo::Type SelectInfo, const FNameProperty* InProperty, TSharedRef<IPropertyUtilities> PropertyUtilities);
+	UE_API void OnNameListComboBox(const FNameProperty* InProperty, const TArray<TSharedPtr<FRigVMStringWithTag>>* InNameList);
+	UE_API void CustomizeLiveValues(IDetailLayoutBuilder& DetailLayout);
 
 	URigVMBlueprint* BlueprintBeingCustomized;
 	TArray<TWeakObjectPtr<URigVMDetailsViewWrapperObject>> ObjectsBeingCustomized;
@@ -302,11 +304,11 @@ public:
 };
 
 /** Customization for editing a rig vm integer control enum class */
-class RIGVMEDITOR_API FRigVMGraphEnumDetailCustomization : public IPropertyTypeCustomization
+class FRigVMGraphEnumDetailCustomization : public IPropertyTypeCustomization
 {
 public:
 
-	FRigVMGraphEnumDetailCustomization();
+	UE_API FRigVMGraphEnumDetailCustomization();
 	
 	static TSharedRef<IPropertyTypeCustomization> MakeInstance()
 	{
@@ -314,8 +316,8 @@ public:
 	}
 
 	/** IPropertyTypeCustomization interface */
-	virtual void CustomizeHeader(TSharedRef<class IPropertyHandle> InPropertyHandle, class FDetailWidgetRow& HeaderRow, IPropertyTypeCustomizationUtils& StructCustomizationUtils) override;
-	virtual void CustomizeChildren(TSharedRef<class IPropertyHandle> InPropertyHandle, class IDetailChildrenBuilder& StructBuilder, IPropertyTypeCustomizationUtils& StructCustomizationUtils) override;
+	UE_API virtual void CustomizeHeader(TSharedRef<class IPropertyHandle> InPropertyHandle, class FDetailWidgetRow& HeaderRow, IPropertyTypeCustomizationUtils& StructCustomizationUtils) override;
+	UE_API virtual void CustomizeChildren(TSharedRef<class IPropertyHandle> InPropertyHandle, class IDetailChildrenBuilder& StructBuilder, IPropertyTypeCustomizationUtils& StructCustomizationUtils) override;
 
 protected:
 
@@ -422,7 +424,7 @@ protected:
 		return (UEnum**)MemoryPtr;
 	}
 
-	void HandleControlEnumChanged(TSharedPtr<FString> InEnumPath, ESelectInfo::Type InSelectType, TSharedRef<IPropertyHandle> InPropertyHandle);
+	UE_API void HandleControlEnumChanged(TSharedPtr<FString> InEnumPath, ESelectInfo::Type InSelectType, TSharedRef<IPropertyHandle> InPropertyHandle);
 
 	URigVMBlueprint* BlueprintBeingCustomized;
 	URigVMGraph* GraphBeingCustomized;
@@ -431,11 +433,11 @@ protected:
 };
 
 /** Customization for editing a rig vm node */
-class RIGVMEDITOR_API FRigVMGraphMathTypeDetailCustomization : public IPropertyTypeCustomization
+class FRigVMGraphMathTypeDetailCustomization : public IPropertyTypeCustomization
 {
 public:
 
-	FRigVMGraphMathTypeDetailCustomization();
+	UE_API FRigVMGraphMathTypeDetailCustomization();
 	
 	static TSharedRef<IPropertyTypeCustomization> MakeInstance()
 	{
@@ -443,8 +445,8 @@ public:
 	}
 
 	/** IPropertyTypeCustomization interface */
-	virtual void CustomizeHeader(TSharedRef<class IPropertyHandle> InPropertyHandle, class FDetailWidgetRow& HeaderRow, IPropertyTypeCustomizationUtils& StructCustomizationUtils) override;
-	virtual void CustomizeChildren(TSharedRef<class IPropertyHandle> InPropertyHandle, class IDetailChildrenBuilder& StructBuilder, IPropertyTypeCustomizationUtils& StructCustomizationUtils) override;
+	UE_API virtual void CustomizeHeader(TSharedRef<class IPropertyHandle> InPropertyHandle, class FDetailWidgetRow& HeaderRow, IPropertyTypeCustomizationUtils& StructCustomizationUtils) override;
+	UE_API virtual void CustomizeChildren(TSharedRef<class IPropertyHandle> InPropertyHandle, class IDetailChildrenBuilder& StructBuilder, IPropertyTypeCustomizationUtils& StructCustomizationUtils) override;
 
 protected:
 
@@ -812,3 +814,5 @@ inline FName FRigVMGraphMathTypeDetailCustomization::GetScalePropertyName<FEuler
 {
 	return TEXT("Scale");
 }
+
+#undef UE_API
diff --git a/Engine/Plugins/Runtime/RigVM/Source/RigVMEditor/Public/Editor/RigVMLegacyEditor.h b/Engine/Plugins/Runtime/RigVM/Source/RigVMEditor/Public/Editor/RigVMLegacyEditor.h
--- a/Engine/Plugins/Runtime/RigVM/Source/RigVMEditor/Public/Editor/RigVMLegacyEditor.h
+++ b/Engine/Plugins/Runtime/RigVM/Source/RigVMEditor/Public/Editor/RigVMLegacyEditor.h
@@ -18,11 +18,11 @@ class FRigVMGraphExplorerDragDropOp;
 class SRigVMEditorGraphExplorer;
 class FRigVMLegacyEditor;
 
-class RIGVMEDITOR_API FRigVMLegacyEditor : public FBlueprintEditor, public FRigVMEditorBase
+class FRigVMLegacyEditor : public FBlueprintEditor, public FRigVMEditorBase
 {
 public:
 
-	FRigVMLegacyEditor();
+	RIGVMEDITOR_API FRigVMLegacyEditor();
 
 	virtual TSharedRef<IRigVMEditor> SharedRef() override { return StaticCastSharedRef<IRigVMEditor>(SharedThis(this)); }
 	virtual TSharedRef<const IRigVMEditor> SharedRef() const override { return StaticCastSharedRef<const IRigVMEditor>(SharedThis(this)); }
@@ -32,8 +32,8 @@ protected:
 	// IRigVMEditor overrides
 	virtual TSharedPtr<FAssetEditorToolkit> GetHostingApp() override { return AsShared(); }
 	virtual const TSharedPtr<FAssetEditorToolkit> GetHostingApp() const override { return ConstCastSharedRef<FAssetEditorToolkit>(AsShared()); }
-	virtual void InitAssetEditor(const EToolkitMode::Type Mode, const TSharedPtr<IToolkitHost>& InitToolkitHost, const FName AppIdentifier, const TSharedRef<FTabManager::FLayout>& StandaloneDefaultLayout, const bool bCreateDefaultStandaloneMenu, const bool bCreateDefaultToolbar, const TArray<UObject*>& ObjectsToEdit, const bool bInIsToolbarFocusable, const bool bInUseSmallToolbarIcons, const TOptional<EAssetOpenMethod>& InOpenMethod) override;
-	virtual void CreateEditorToolbar() override;
+	RIGVMEDITOR_API virtual void InitAssetEditor(const EToolkitMode::Type Mode, const TSharedPtr<IToolkitHost>& InitToolkitHost, const FName AppIdentifier, const TSharedRef<FTabManager::FLayout>& StandaloneDefaultLayout, const bool bCreateDefaultStandaloneMenu, const bool bCreateDefaultToolbar, const TArray<UObject*>& ObjectsToEdit, const bool bInIsToolbarFocusable, const bool bInUseSmallToolbarIcons, const TOptional<EAssetOpenMethod>& InOpenMethod) override;
+	RIGVMEDITOR_API virtual void CreateEditorToolbar() override;
 	virtual void CommonInitialization(const TArray<UBlueprint*>& InitBlueprints, bool bShouldOpenInDefaultsMode) override { FBlueprintEditor::CommonInitialization(InitBlueprints, bShouldOpenInDefaultsMode); }
 	virtual TSharedPtr<FDocumentTracker> GetDocumentManager() const override { return DocumentManager; }
 	virtual void AddApplicationMode(FName ModeName, TSharedRef<FApplicationMode> Mode) override { FBlueprintEditor::AddApplicationMode(ModeName, Mode); }
@@ -50,29 +50,29 @@ protected:
 	virtual TSharedPtr<SKismetInspector> GetKismetInspector() const override { return Inspector.IsValid() ? Inspector.ToSharedRef().ToSharedPtr() : nullptr; }
 #endif
 	virtual TSharedPtr<SRigVMDetailsInspector> GetRigVMInspector() const override { return nullptr; }
-	virtual TSharedPtr<FApplicationMode> CreateEditorMode() override;
-	virtual const FName GetEditorAppName() const override;
+	RIGVMEDITOR_API virtual TSharedPtr<FApplicationMode> CreateEditorMode() override;
+	RIGVMEDITOR_API virtual const FName GetEditorAppName() const override;
 	virtual const TArray< UObject* >& GetEditingBlueprints() const override { return FBlueprintEditor::GetEditingObjects(); }
-	virtual void OnGraphEditorFocused(const TSharedRef<SGraphEditor>& InGraphEditor) override;
+	RIGVMEDITOR_API virtual void OnGraphEditorFocused(const TSharedRef<SGraphEditor>& InGraphEditor) override;
 	virtual const TSharedRef<IToolkitHost> GetToolkitHost() const override { return FBlueprintEditor::GetToolkitHost(); }
 	virtual bool IsHosted() const override { return FBlueprintEditor::IsHosted(); }
 	virtual void BringToolkitToFrontImpl() override { FBlueprintEditor::BringToolkitToFront(); }
-	virtual void InitToolMenuContext(FToolMenuContext& MenuContext) override;
+	RIGVMEDITOR_API virtual void InitToolMenuContext(FToolMenuContext& MenuContext) override;
 	virtual TWeakPtr<SGraphEditor> GetFocusedGraphEditor() override { return FocusedGraphEdPtr; }
 	virtual TWeakPtr<FDocumentTabFactory> GetGraphEditorTabFactory() const override { return GraphEditorTabFactoryPtr; }
-	virtual bool TransactionObjectAffectsBlueprint(UObject* InTransactedObject) override;
-	virtual FEdGraphPinType GetLastPinTypeUsed() override;
+	RIGVMEDITOR_API virtual bool TransactionObjectAffectsBlueprint(UObject* InTransactedObject) override;
+	RIGVMEDITOR_API virtual FEdGraphPinType GetLastPinTypeUsed() override;
 	virtual void LogSimpleMessage(const FText& MessageText) override { FBlueprintEditor::LogSimpleMessage(MessageText); }
 	virtual void RenameNewlyAddedAction(FName InActionName) override { FBlueprintEditor::RenameNewlyAddedAction(InActionName); }
 	virtual FGraphPanelSelectionSet GetSelectedNodes() const override { return FBlueprintEditor::GetSelectedNodes(); }
 	virtual void SetUISelectionState(FName SelectionOwner) override { FBlueprintEditor::SetUISelectionState(SelectionOwner); }
 	virtual void AnalyticsTrackNodeEvent(UBlueprint* Blueprint, UEdGraphNode* GraphNode, bool bNodeDelete) const override { FBlueprintEditor::AnalyticsTrackNodeEvent(Blueprint, GraphNode, bNodeDelete); }
-	virtual void JumpToHyperlink(const UObject* ObjectReference, bool bRequestRename) override;
-	virtual void PostUndo(bool bSuccess) override;
-	virtual void PostRedo(bool bSuccess) override;
+	RIGVMEDITOR_API virtual void JumpToHyperlink(const UObject* ObjectReference, bool bRequestRename) override;
+	RIGVMEDITOR_API virtual void PostUndo(bool bSuccess) override;
+	RIGVMEDITOR_API virtual void PostRedo(bool bSuccess) override;
 	virtual UEdGraphPin* GetCurrentlySelectedPin() const override { return FBlueprintEditor::GetCurrentlySelectedPin();}
-	virtual void CreateDefaultCommands() override;
-	virtual TSharedRef<SGraphEditor> CreateGraphEditorWidget(TSharedRef<FTabInfo> InTabInfo, UEdGraph* InGraph) override;
+	RIGVMEDITOR_API virtual void CreateDefaultCommands() override;
+	RIGVMEDITOR_API virtual TSharedRef<SGraphEditor> CreateGraphEditorWidget(TSharedRef<FTabInfo> InTabInfo, UEdGraph* InGraph) override;
 	virtual void CompileImpl() override { FBlueprintEditor::Compile(); }
 	virtual void SaveAsset_Execute_Impl() override { FBlueprintEditor::SaveAsset_Execute(); }
 	virtual void SaveAssetAs_Execute_Impl() override { FBlueprintEditor::SaveAssetAs_Execute(); }
@@ -87,18 +87,18 @@ protected:
 	virtual void RemoveMenuExtender(TSharedPtr<FExtender> Extender) override { FBlueprintEditor::RemoveMenuExtender(Extender); }
 	virtual TSharedPtr<IMessageLogListing> GetCompilerResultsListing() override { return CompilerResultsListing; }
 	virtual void OnBlueprintChangedInnerImpl(UBlueprint* InBlueprint, bool bIsJustBeingCompiled) override { FBlueprintEditor::OnBlueprintChangedImpl(InBlueprint, bIsJustBeingCompiled); }
-	virtual void RefreshEditorsImpl(ERefreshRigVMEditorReason::Type Reason) override;
+	RIGVMEDITOR_API virtual void RefreshEditorsImpl(ERefreshRigVMEditorReason::Type Reason) override;
 	virtual void SetupGraphEditorEvents(UEdGraph* InGraph, SGraphEditor::FGraphEditorEvents& InEvents) override { FRigVMEditorBase::SetupGraphEditorEvents(InGraph, InEvents); }
 	virtual void SetupGraphEditorEventsImpl(UEdGraph* InGraph, SGraphEditor::FGraphEditorEvents& InEvents) override { FBlueprintEditor::SetupGraphEditorEvents(InGraph, InEvents); }
 	virtual FActionMenuContent OnCreateGraphActionMenu(UEdGraph* InGraph, const FVector2f& InNodePosition, const TArray<UEdGraphPin*>& InDraggedPins, bool bAutoExpand, SGraphEditor::FActionMenuClosed InOnMenuClosed) override { return FBlueprintEditor::OnCreateGraphActionMenu(InGraph, InNodePosition, InDraggedPins, bAutoExpand, InOnMenuClosed); }
 	virtual void AddReferencedObjects(FReferenceCollector& Collector) override { FRigVMEditorBase::AddReferencedObjects(Collector); }
 	virtual void AddReferencedObjectsImpl(FReferenceCollector& Collector) override { FBlueprintEditor::AddReferencedObjects(Collector); }
-	virtual bool NewDocument_IsVisibleForType(FBlueprintEditor::ECreatedDocumentType GraphType) const override;
-	virtual bool IsSectionVisible(NodeSectionID::Type InSectionID) const override;
+	RIGVMEDITOR_API virtual bool NewDocument_IsVisibleForType(FBlueprintEditor::ECreatedDocumentType GraphType) const override;
+	RIGVMEDITOR_API virtual bool IsSectionVisible(NodeSectionID::Type InSectionID) const override;
 	virtual FGraphAppearanceInfo GetGraphAppearanceImpl(UEdGraph* InGraph) const override { return FBlueprintEditor::GetGraphAppearance(InGraph); }
 	virtual void NotifyPreChangeImpl(FProperty* PropertyAboutToChange) override { FBlueprintEditor::NotifyPreChange(PropertyAboutToChange); }
 	virtual void NotifyPostChangeImpl(const FPropertyChangedEvent& PropertyChangedEvent, FProperty* PropertyThatChanged) override { FBlueprintEditor::NotifyPostChange(PropertyChangedEvent, PropertyThatChanged); }
-	virtual FName GetSelectedVariableName() override;
+	RIGVMEDITOR_API virtual FName GetSelectedVariableName() override;
 	virtual bool IsNodeTitleVisible(const UEdGraphNode* Node, bool bRequestRename) override { return FBlueprintEditor::IsNodeTitleVisible(Node, bRequestRename); }
 	virtual void EditClassDefaults_Clicked() override { FBlueprintEditor::EditClassDefaults_Clicked(); }
 	virtual void EditGlobalOptions_Clicked() override { FBlueprintEditor::EditGlobalOptions_Clicked(); }
@@ -115,20 +115,20 @@ protected:
 	virtual void SetIsContextSensitive(const bool bIsContextSensitive) override { FBlueprintEditor::GetIsContextSensitive() = bIsContextSensitive; }
 	virtual void RegisterToolbarTab(const TSharedRef<FTabManager>& InTabManager) override { FBlueprintEditor::RegisterToolbarTab(InTabManager); }
 	virtual const TArray<UObject*>* GetObjectsCurrentlyBeingEdited() const override { return FBlueprintEditor::GetObjectsCurrentlyBeingEdited(); }
-	virtual void AddCompileWidget(FToolBarBuilder& ToolbarBuilder) override;
+	RIGVMEDITOR_API virtual void AddCompileWidget(FToolBarBuilder& ToolbarBuilder) override;
 	virtual void AddSettingsAndDefaultWidget(FToolBarBuilder& ToolbarBuilder) override {}
-	virtual void AddSelectedDebugObjectWidget(FToolBarBuilder& ToolbarBuilder) override;
-	virtual void AddAutoCompileWidget(FToolBarBuilder& ToolbarBuilder) override;
+	RIGVMEDITOR_API virtual void AddSelectedDebugObjectWidget(FToolBarBuilder& ToolbarBuilder) override;
+	RIGVMEDITOR_API virtual void AddAutoCompileWidget(FToolBarBuilder& ToolbarBuilder) override;
 	virtual FString GetCustomDebugObjectLabel(UObject* ObjectBeingDebugged) const override { return FRigVMEditorBase::GetCustomDebugObjectLabel(ObjectBeingDebugged); }
 public:
 	virtual void SummonSearchUI(bool bSetFindWithinBlueprint, FString NewSearchTerms, bool bSelectFirstResult) override { FBlueprintEditor::SummonSearchUI(bSetFindWithinBlueprint, NewSearchTerms, bSelectFirstResult); }
 	virtual TSharedRef<FUICommandList> GetToolkitCommands() override { return ToolkitCommands; }
-	virtual void OnClose() override;
+	RIGVMEDITOR_API virtual void OnClose() override;
 	virtual FPreviewScene* GetPreviewScene() override { return &PreviewScene; }
 
 	
 protected:
-	virtual void Tick(float DeltaTime) override;
+	RIGVMEDITOR_API virtual void Tick(float DeltaTime) override;
 	// IRigVMEditor overrides
 
 	///////////////////////////////////////////////////////////////////////////////////////////////
@@ -185,7 +185,7 @@ protected:
 	virtual void OnStopWatchingPin() override { FRigVMEditorBase::OnStopWatchingPin(); }
 	virtual bool CanStopWatchingPin() const override { return FRigVMEditorBase::CanStopWatchingPin(); }
 	virtual void OnCreateComment() override { FRigVMEditorBase::OnCreateComment(); }
-	virtual void GetCustomDebugObjects(TArray<FCustomDebugObject>& DebugList) const override;
+	RIGVMEDITOR_API virtual void GetCustomDebugObjects(TArray<FCustomDebugObject>& DebugList) const override;
 	virtual bool OnlyShowCustomDebugObjects() const override { return FRigVMEditorBase::OnlyShowCustomDebugObjects(); }
 	virtual void HandleSetObjectBeingDebugged(UObject* InObject) override { FRigVMEditorBase::HandleSetObjectBeingDebugged(InObject); }
 	virtual bool OnActionMatchesName(FEdGraphSchemaAction* InAction, const FName& InName) const override { return FRigVMEditorBase::OnActionMatchesName(InAction, InName); }
@@ -199,4 +199,4 @@ protected:
 
 };
 
-#endif
\ No newline at end of file
+#endif
diff --git a/Engine/Plugins/Runtime/RigVM/Source/RigVMEditor/Public/Editor/RigVMLocalVariableDetailCustomization.h b/Engine/Plugins/Runtime/RigVM/Source/RigVMEditor/Public/Editor/RigVMLocalVariableDetailCustomization.h
--- a/Engine/Plugins/Runtime/RigVM/Source/RigVMEditor/Public/Editor/RigVMLocalVariableDetailCustomization.h
+++ b/Engine/Plugins/Runtime/RigVM/Source/RigVMEditor/Public/Editor/RigVMLocalVariableDetailCustomization.h
@@ -14,9 +14,11 @@
 #include "IPropertyUtilities.h"
 #include "EdGraph/RigVMEdGraphSchema.h"
 
+#define UE_API RIGVMEDITOR_API
+
 class IPropertyHandle;
 
-class RIGVMEDITOR_API FRigVMLocalVariableDetailCustomization : public IDetailCustomization
+class FRigVMLocalVariableDetailCustomization : public IDetailCustomization
 {
 	FRigVMLocalVariableDetailCustomization()
 	: GraphBeingCustomized(nullptr)
@@ -34,7 +36,7 @@ public:
 	}
 
 	/** IDetailCustomization interface */
-	virtual void CustomizeDetails(IDetailLayoutBuilder& DetailBuilder) override;
+	UE_API virtual void CustomizeDetails(IDetailLayoutBuilder& DetailBuilder) override;
 
 private:
 
@@ -51,16 +53,17 @@ private:
 	FRigVMLocalVariableNameValidator NameValidator;
 	TArray<TSharedPtr<FString>> EnumOptions;
 
-	FText GetName() const;
-	void SetName(const FText& InNewText, ETextCommit::Type InCommitType);
-	bool OnVerifyNameChanged(const FText& InText, FText& OutErrorMessage);
+	UE_API FText GetName() const;
+	UE_API void SetName(const FText& InNewText, ETextCommit::Type InCommitType);
+	UE_API bool OnVerifyNameChanged(const FText& InText, FText& OutErrorMessage);
 	
-	FEdGraphPinType OnGetPinInfo() const;
-	void HandlePinInfoChanged(const FEdGraphPinType& PinType);
+	UE_API FEdGraphPinType OnGetPinInfo() const;
+	UE_API void HandlePinInfoChanged(const FEdGraphPinType& PinType);
 
-	ECheckBoxState HandleBoolDefaultValueIsChecked( ) const;
-	void OnBoolDefaultValueChanged(ECheckBoxState InCheckBoxState);
+	UE_API ECheckBoxState HandleBoolDefaultValueIsChecked( ) const;
+	UE_API void OnBoolDefaultValueChanged(ECheckBoxState InCheckBoxState);
 };
 
 
 
+#undef UE_API
diff --git a/Engine/Plugins/Runtime/RigVM/Source/RigVMEditor/Public/Editor/RigVMMinimalEnvironment.h b/Engine/Plugins/Runtime/RigVM/Source/RigVMEditor/Public/Editor/RigVMMinimalEnvironment.h
--- a/Engine/Plugins/Runtime/RigVM/Source/RigVMEditor/Public/Editor/RigVMMinimalEnvironment.h
+++ b/Engine/Plugins/Runtime/RigVM/Source/RigVMEditor/Public/Editor/RigVMMinimalEnvironment.h
@@ -10,27 +10,29 @@
 #include "EdGraph/RigVMEdGraphNode.h"
 #include "UObject/StrongObjectPtr.h"
 
-class RIGVMEDITOR_API FRigVMMinimalEnvironment : public TSharedFromThis<FRigVMMinimalEnvironment>
+#define UE_API RIGVMEDITOR_API
+
+class FRigVMMinimalEnvironment : public TSharedFromThis<FRigVMMinimalEnvironment>
 {
 public:
-	FRigVMMinimalEnvironment(const UClass* InRigVMBlueprintClass = nullptr);
+	UE_API FRigVMMinimalEnvironment(const UClass* InRigVMBlueprintClass = nullptr);
 
-	URigVMGraph* GetModel() const;
-	URigVMController* GetController() const;
-	URigVMNode* GetNode() const;
-	URigVMEdGraph* GetEdGraph() const;
-	URigVMEdGraphNode* GetEdGraphNode() const;
+	UE_API URigVMGraph* GetModel() const;
+	UE_API URigVMController* GetController() const;
+	UE_API URigVMNode* GetNode() const;
+	UE_API URigVMEdGraph* GetEdGraph() const;
+	UE_API URigVMEdGraphNode* GetEdGraphNode() const;
 
-	void SetSchemata(const UClass* InRigVMBlueprintClass);
-	void SetNode(URigVMNode* InModelNode);
-	void SetFunctionNode(const FRigVMGraphFunctionIdentifier& InIdentifier);
+	UE_API void SetSchemata(const UClass* InRigVMBlueprintClass);
+	UE_API void SetNode(URigVMNode* InModelNode);
+	UE_API void SetFunctionNode(const FRigVMGraphFunctionIdentifier& InIdentifier);
 	
-	FSimpleDelegate& OnChanged();
-	void Tick_GameThead(float InDeltaTime);
+	UE_API FSimpleDelegate& OnChanged();
+	UE_API void Tick_GameThead(float InDeltaTime);
 
 private:
 
-	void HandleModified(ERigVMGraphNotifType InNotification, URigVMGraph* InGraph, UObject* InSubject);
+	UE_API void HandleModified(ERigVMGraphNotifType InNotification, URigVMGraph* InGraph, UObject* InSubject);
 	
 	TStrongObjectPtr<URigVMGraph> ModelGraph;
 	TStrongObjectPtr<URigVMController> ModelController;
@@ -43,3 +45,5 @@ private:
 	FSimpleDelegate ChangedDelegate;
 	FDelegateHandle ModelHandle;
 };
+
+#undef UE_API
diff --git a/Engine/Plugins/Runtime/RigVM/Source/RigVMEditor/Public/Editor/RigVMNewEditor.h b/Engine/Plugins/Runtime/RigVM/Source/RigVMEditor/Public/Editor/RigVMNewEditor.h
--- a/Engine/Plugins/Runtime/RigVM/Source/RigVMEditor/Public/Editor/RigVMNewEditor.h
+++ b/Engine/Plugins/Runtime/RigVM/Source/RigVMEditor/Public/Editor/RigVMNewEditor.h
@@ -6,20 +6,22 @@
 #include "RigVMEditor.h"
 #include "WorkflowOrientedApp/WorkflowCentricApplication.h"
 
+#define UE_API RIGVMEDITOR_API
+
 class SRigVMFindReferences;
 class SRigVMDetailsInspector;
 
-struct RIGVMEDITOR_API FRigVMNewEditorTabs
+struct FRigVMNewEditorTabs
 {
 	// Tab identifiers
-	static const FName CompilerResultsID();
+	static UE_API const FName CompilerResultsID();
 };
 
-class RIGVMEDITOR_API FRigVMNewEditor : public FWorkflowCentricApplication, public FRigVMEditorBase, public FGCObject, public FNotifyHook, public FTickableEditorObject, public FEditorUndoClient, public FNoncopyable
+class FRigVMNewEditor : public FWorkflowCentricApplication, public FRigVMEditorBase, public FGCObject, public FNotifyHook, public FTickableEditorObject, public FEditorUndoClient, public FNoncopyable
 {
 public:
-	FRigVMNewEditor();
-	virtual void OnClose() override;
+	UE_API FRigVMNewEditor();
+	UE_API virtual void OnClose() override;
 
 	virtual TSharedRef<IRigVMEditor> SharedRef() override { return StaticCastSharedRef<IRigVMEditor>(SharedThis(this)); }
 	virtual TSharedRef<const IRigVMEditor> SharedRef() const override { return StaticCastSharedRef<const IRigVMEditor>(SharedThis(this)); }
@@ -28,21 +30,21 @@ public:
 	virtual const TSharedPtr<FAssetEditorToolkit> GetHostingApp() const override { return ConstCastSharedRef<FAssetEditorToolkit>(AsShared()); }
 protected:
 	
-	virtual void InitAssetEditor(const EToolkitMode::Type Mode, const TSharedPtr<IToolkitHost>& InitToolkitHost, const FName AppIdentifier, const TSharedRef<FTabManager::FLayout>& StandaloneDefaultLayout, const bool bCreateDefaultStandaloneMenu, const bool bCreateDefaultToolbar, const TArray<UObject*>& ObjectsToEdit, const bool bInIsToolbarFocusable, const bool bInUseSmallToolbarIcons, const TOptional<EAssetOpenMethod>& InOpenMethod) override;
+	UE_API virtual void InitAssetEditor(const EToolkitMode::Type Mode, const TSharedPtr<IToolkitHost>& InitToolkitHost, const FName AppIdentifier, const TSharedRef<FTabManager::FLayout>& StandaloneDefaultLayout, const bool bCreateDefaultStandaloneMenu, const bool bCreateDefaultToolbar, const TArray<UObject*>& ObjectsToEdit, const bool bInIsToolbarFocusable, const bool bInUseSmallToolbarIcons, const TOptional<EAssetOpenMethod>& InOpenMethod) override;
 	virtual void CreateEditorToolbar() override {}
-	virtual void CommonInitialization(const TArray<UBlueprint*>& InitBlueprints, bool bShouldOpenInDefaultsMode) override;
-	void OnBlueprintChanged(UBlueprint* InBlueprint);
-	void SaveEditedObjectState();
+	UE_API virtual void CommonInitialization(const TArray<UBlueprint*>& InitBlueprints, bool bShouldOpenInDefaultsMode) override;
+	UE_API void OnBlueprintChanged(UBlueprint* InBlueprint);
+	UE_API void SaveEditedObjectState();
 	virtual TSharedPtr<FDocumentTracker> GetDocumentManager() const override { return DocumentManager; }
 	virtual void AddApplicationMode(FName ModeName, TSharedRef<FApplicationMode> Mode) override { FWorkflowCentricApplication::AddApplicationMode(ModeName, Mode); }
 	virtual void RegenerateMenusAndToolbars() override { FWorkflowCentricApplication::RegenerateMenusAndToolbars(); }
-	virtual void SetCurrentMode(FName NewMode) override;
+	UE_API virtual void SetCurrentMode(FName NewMode) override;
 	virtual FEditorModeTools& GetToolkitEditorModeManager() const override { return FWorkflowCentricApplication::GetEditorModeManager(); }
-	virtual void PostLayoutBlueprintEditorInitialization() override;
-	virtual TSharedPtr<SGraphEditor> OpenGraphAndBringToFront(UEdGraph* Graph, bool bSetFocus = true) override;
-	virtual bool FindOpenTabsContainingDocument(const UObject* DocumentID, /*inout*/ TArray< TSharedPtr<SDockTab> >& Results) override;
-	virtual TSharedPtr<SDockTab> OpenDocument(const UObject* DocumentID, FDocumentTracker::EOpenDocumentCause Cause) override;
-	virtual void CloseDocumentTab(const UObject* DocumentID) override;
+	UE_API virtual void PostLayoutBlueprintEditorInitialization() override;
+	UE_API virtual TSharedPtr<SGraphEditor> OpenGraphAndBringToFront(UEdGraph* Graph, bool bSetFocus = true) override;
+	UE_API virtual bool FindOpenTabsContainingDocument(const UObject* DocumentID, /*inout*/ TArray< TSharedPtr<SDockTab> >& Results) override;
+	UE_API virtual TSharedPtr<SDockTab> OpenDocument(const UObject* DocumentID, FDocumentTracker::EOpenDocumentCause Cause) override;
+	UE_API virtual void CloseDocumentTab(const UObject* DocumentID) override;
 	virtual TSharedPtr<FTabManager> GetTabManager() override { return FWorkflowCentricApplication::GetTabManager(); }
 public:
 	virtual TSharedPtr<SRigVMDetailsInspector> GetRigVMInspector() const override { return Inspector.IsValid() ? Inspector.ToSharedRef().ToSharedPtr() : nullptr; }
@@ -50,80 +52,80 @@ public:
 	
 	TSharedRef<SWidget> GetCompilerResults() const { return CompilerResults.ToSharedRef(); }
 	TSharedRef<SRigVMFindReferences> GetFindResults() const { return FindResults.ToSharedRef(); }
-	virtual void JumpToHyperlink(const UObject* ObjectReference, bool bRequestRename) override;
+	UE_API virtual void JumpToHyperlink(const UObject* ObjectReference, bool bRequestRename) override;
 protected:
 	virtual TSharedPtr<IMessageLogListing> GetCompilerResultsListing() override { return CompilerResultsListing; }
-	virtual TSharedPtr<FApplicationMode> CreateEditorMode() override;
-	virtual const FName GetEditorAppName() const override;
+	UE_API virtual TSharedPtr<FApplicationMode> CreateEditorMode() override;
+	UE_API virtual const FName GetEditorAppName() const override;
 	virtual const TArray< UObject* >& GetEditingBlueprints() const override { return FWorkflowCentricApplication::GetEditingObjects(); }
 	virtual const TSharedRef<IToolkitHost> GetToolkitHost() const override { return FWorkflowCentricApplication::GetToolkitHost(); }
 	virtual bool IsHosted() const override { return FWorkflowCentricApplication::IsHosted(); }
 	virtual void BringToolkitToFrontImpl() override { FWorkflowCentricApplication::BringToolkitToFront(); }
-	virtual void InitToolMenuContext(FToolMenuContext& MenuContext) override;
+	UE_API virtual void InitToolMenuContext(FToolMenuContext& MenuContext) override;
 	virtual TSharedRef<FUICommandList> GetToolkitCommands() override { return ToolkitCommands; }
 	virtual TWeakPtr<SGraphEditor> GetFocusedGraphEditor() override { return FocusedGraphEdPtr; }
 	virtual TWeakPtr<FDocumentTabFactory> GetGraphEditorTabFactory() const override { return GraphEditorTabFactoryPtr; }
-	virtual bool TransactionObjectAffectsBlueprint(UObject* InTransactedObject) override;
-	virtual FEdGraphPinType GetLastPinTypeUsed() override;
-	virtual void LogSimpleMessage(const FText& MessageText) override;
-	virtual void RenameNewlyAddedAction(FName InActionName) override;
-	virtual FGraphPanelSelectionSet GetSelectedNodes() const override;
-	virtual void SetUISelectionState(FName SelectionOwner) override;
-	virtual void AnalyticsTrackNodeEvent(UBlueprint* Blueprint, UEdGraphNode* GraphNode, bool bNodeDelete) const override;
-	virtual void PostUndo(bool bSuccess) override;
-	virtual void PostRedo(bool bSuccess) override;
-	virtual UEdGraphPin* GetCurrentlySelectedPin() const override;
-	virtual void CreateDefaultCommands() override;
-	virtual TSharedRef<SGraphEditor> CreateGraphEditorWidget(TSharedRef<FTabInfo> InTabInfo, UEdGraph* InGraph) override;
-	virtual void CompileImpl() override;
+	UE_API virtual bool TransactionObjectAffectsBlueprint(UObject* InTransactedObject) override;
+	UE_API virtual FEdGraphPinType GetLastPinTypeUsed() override;
+	UE_API virtual void LogSimpleMessage(const FText& MessageText) override;
+	UE_API virtual void RenameNewlyAddedAction(FName InActionName) override;
+	UE_API virtual FGraphPanelSelectionSet GetSelectedNodes() const override;
+	UE_API virtual void SetUISelectionState(FName SelectionOwner) override;
+	UE_API virtual void AnalyticsTrackNodeEvent(UBlueprint* Blueprint, UEdGraphNode* GraphNode, bool bNodeDelete) const override;
+	UE_API virtual void PostUndo(bool bSuccess) override;
+	UE_API virtual void PostRedo(bool bSuccess) override;
+	UE_API virtual UEdGraphPin* GetCurrentlySelectedPin() const override;
+	UE_API virtual void CreateDefaultCommands() override;
+	UE_API virtual TSharedRef<SGraphEditor> CreateGraphEditorWidget(TSharedRef<FTabInfo> InTabInfo, UEdGraph* InGraph) override;
+	UE_API virtual void CompileImpl() override;
 	virtual void SaveAsset_Execute_Impl() override { FWorkflowCentricApplication::SaveAsset_Execute(); }
 	virtual void SaveAssetAs_Execute_Impl() override { FWorkflowCentricApplication::SaveAssetAs_Execute(); }
-	virtual bool IsGraphInCurrentBlueprint(const UEdGraph* InGraph) const override;
-	virtual bool IsEditableImpl(UEdGraph* InGraph) const override;
-	virtual UEdGraph* GetFocusedGraph() const override;
-	virtual void JumpToNode(const UEdGraphNode* Node, bool bRequestRename) override;
-	virtual void JumpToPin(const UEdGraphPin* Pin) override;
+	UE_API virtual bool IsGraphInCurrentBlueprint(const UEdGraph* InGraph) const override;
+	UE_API virtual bool IsEditableImpl(UEdGraph* InGraph) const override;
+	UE_API virtual UEdGraph* GetFocusedGraph() const override;
+	UE_API virtual void JumpToNode(const UEdGraphNode* Node, bool bRequestRename) override;
+	UE_API virtual void JumpToPin(const UEdGraphPin* Pin) override;
 	virtual void AddToolbarExtender(TSharedPtr<FExtender> Extender) override { FWorkflowCentricApplication::AddToolbarExtender(Extender); }
 	virtual void RemoveToolbarExtender(TSharedPtr<FExtender> Extender) override { FWorkflowCentricApplication::RemoveToolbarExtender(Extender); };
 	virtual void AddMenuExtender(TSharedPtr<FExtender> Extender) override { FWorkflowCentricApplication::AddMenuExtender(Extender); }
 	virtual void RemoveMenuExtender(TSharedPtr<FExtender> Extender) override { FWorkflowCentricApplication::RemoveMenuExtender(Extender); }
-	virtual void OnBlueprintChangedInnerImpl(UBlueprint* InBlueprint, bool bIsJustBeingCompiled) override;
-	virtual void RefreshEditorsImpl(ERefreshRigVMEditorReason::Type Reason) override;
-	virtual void SetupGraphEditorEventsImpl(UEdGraph* InGraph, SGraphEditor::FGraphEditorEvents& InEvents) override;
-	virtual FActionMenuContent OnCreateGraphActionMenu(UEdGraph* InGraph, const FVector2f& InNodePosition, const TArray<UEdGraphPin*>& InDraggedPins, bool bAutoExpand, SGraphEditor::FActionMenuClosed InOnMenuClosed) override;
-	virtual void AddReferencedObjectsImpl(FReferenceCollector& Collector) override;
-	virtual bool IsSectionVisible(RigVMNodeSectionID::Type InSectionID) const; 
-	virtual bool NewDocument_IsVisibleForType(ECreatedDocumentType GraphType) const;
-	virtual FGraphAppearanceInfo GetGraphAppearanceImpl(UEdGraph* InGraph) const override;
-	virtual void NotifyPreChangeImpl(FProperty* PropertyAboutToChange) override;
-	virtual void NotifyPostChangeImpl(const FPropertyChangedEvent& PropertyChangedEvent, FProperty* PropertyThatChanged) override;
-	virtual FName GetSelectedVariableName() override;
-	virtual bool IsNodeTitleVisible(const UEdGraphNode* Node, bool bRequestRename) override;
-	virtual void EditClassDefaults_Clicked() override;
-	virtual void EditGlobalOptions_Clicked() override;
-	bool IsDetailsPanelEditingGlobalOptions() const;
-	bool IsDetailsPanelEditingClassDefaults() const;
-	virtual void TryInvokingDetailsTab(bool bFlash = true) override;
+	UE_API virtual void OnBlueprintChangedInnerImpl(UBlueprint* InBlueprint, bool bIsJustBeingCompiled) override;
+	UE_API virtual void RefreshEditorsImpl(ERefreshRigVMEditorReason::Type Reason) override;
+	UE_API virtual void SetupGraphEditorEventsImpl(UEdGraph* InGraph, SGraphEditor::FGraphEditorEvents& InEvents) override;
+	UE_API virtual FActionMenuContent OnCreateGraphActionMenu(UEdGraph* InGraph, const FVector2f& InNodePosition, const TArray<UEdGraphPin*>& InDraggedPins, bool bAutoExpand, SGraphEditor::FActionMenuClosed InOnMenuClosed) override;
+	UE_API virtual void AddReferencedObjectsImpl(FReferenceCollector& Collector) override;
+	UE_API virtual bool IsSectionVisible(RigVMNodeSectionID::Type InSectionID) const; 
+	UE_API virtual bool NewDocument_IsVisibleForType(ECreatedDocumentType GraphType) const;
+	UE_API virtual FGraphAppearanceInfo GetGraphAppearanceImpl(UEdGraph* InGraph) const override;
+	UE_API virtual void NotifyPreChangeImpl(FProperty* PropertyAboutToChange) override;
+	UE_API virtual void NotifyPostChangeImpl(const FPropertyChangedEvent& PropertyChangedEvent, FProperty* PropertyThatChanged) override;
+	UE_API virtual FName GetSelectedVariableName() override;
+	UE_API virtual bool IsNodeTitleVisible(const UEdGraphNode* Node, bool bRequestRename) override;
+	UE_API virtual void EditClassDefaults_Clicked() override;
+	UE_API virtual void EditGlobalOptions_Clicked() override;
+	UE_API bool IsDetailsPanelEditingGlobalOptions() const;
+	UE_API bool IsDetailsPanelEditingClassDefaults() const;
+	UE_API virtual void TryInvokingDetailsTab(bool bFlash = true) override;
 	virtual FName GetGraphExplorerWidgetID() override { return FRigVMEditorGraphExplorerTabSummoner::TabID(); }
-	virtual void RefreshInspector() override;
-	virtual void RefreshStandAloneDefaultsEditor() override;
-	virtual void Tick(float DeltaTime) override;
-	virtual void GetPinTypeSelectorFilters(TArray<TSharedPtr<IPinTypeSelectorFilter>>& OutFilters) const override;
-	virtual void OnAddNewVariable() override;
-	virtual void ZoomToSelection_Clicked() override;
+	UE_API virtual void RefreshInspector() override;
+	UE_API virtual void RefreshStandAloneDefaultsEditor() override;
+	UE_API virtual void Tick(float DeltaTime) override;
+	UE_API virtual void GetPinTypeSelectorFilters(TArray<TSharedPtr<IPinTypeSelectorFilter>>& OutFilters) const override;
+	UE_API virtual void OnAddNewVariable() override;
+	UE_API virtual void ZoomToSelection_Clicked() override;
 public:
-	virtual void SummonSearchUI(bool bSetFindWithinBlueprint, FString NewSearchTerms, bool bSelectFirstResult) override;
-	virtual void RestoreEditedObjectState() override;
-	virtual void SetupViewForBlueprintEditingMode() override;
-	virtual void OnGraphEditorFocused(const TSharedRef<SGraphEditor>& InGraphEditor);
-	virtual void OnGraphEditorBackgrounded(const TSharedRef<SGraphEditor>& InGraphEditor);
+	UE_API virtual void SummonSearchUI(bool bSetFindWithinBlueprint, FString NewSearchTerms, bool bSelectFirstResult) override;
+	UE_API virtual void RestoreEditedObjectState() override;
+	UE_API virtual void SetupViewForBlueprintEditingMode() override;
+	UE_API virtual void OnGraphEditorFocused(const TSharedRef<SGraphEditor>& InGraphEditor);
+	UE_API virtual void OnGraphEditorBackgrounded(const TSharedRef<SGraphEditor>& InGraphEditor);
 	virtual bool GetIsContextSensitive() override { return bIsActionMenuContextSensitive; }
 	virtual void SetIsContextSensitive(const bool bIsContextSensitive) override { bIsActionMenuContextSensitive = bIsContextSensitive; }
 	virtual void RegisterToolbarTab(const TSharedRef<FTabManager>& InTabManager) override { FAssetEditorToolkit::RegisterTabSpawners(InTabManager); }
 	virtual const TArray<UObject*>* GetObjectsCurrentlyBeingEdited() const override { return FAssetEditorToolkit::GetObjectsCurrentlyBeingEdited(); }
-	virtual void AddCompileWidget(FToolBarBuilder& ToolbarBuilder) override;
-	virtual void AddSettingsAndDefaultWidget(FToolBarBuilder& ToolbarBuilder) override;
-	virtual void AddSelectedDebugObjectWidget(FToolBarBuilder& ToolbarBuilder) override;
+	UE_API virtual void AddCompileWidget(FToolBarBuilder& ToolbarBuilder) override;
+	UE_API virtual void AddSettingsAndDefaultWidget(FToolBarBuilder& ToolbarBuilder) override;
+	UE_API virtual void AddSelectedDebugObjectWidget(FToolBarBuilder& ToolbarBuilder) override;
 	virtual void AddAutoCompileWidget(FToolBarBuilder& ToolbarBuilder) override {}
 	virtual void Compile() { FRigVMEditorBase::Compile(); }
 	virtual void OnCreateGraphEditorCommands(TSharedPtr<FUICommandList> GraphEditorCommandsList) { FRigVMEditorBase::OnCreateGraphEditorCommands(GraphEditorCommandsList); }
@@ -134,7 +136,7 @@ public:
 	virtual FPreviewScene* GetPreviewScene() override { return nullptr; }
 
 	//~ Begin IToolkit Interface
-	virtual void RegisterTabSpawners(const TSharedRef<class FTabManager>& InTabManager) override;
+	UE_API virtual void RegisterTabSpawners(const TSharedRef<class FTabManager>& InTabManager) override;
 	//~ End IToolkit Interface
 
 
@@ -142,53 +144,53 @@ public:
 	virtual void AddReferencedObjects(FReferenceCollector& Collector) override { return FRigVMEditorBase::AddReferencedObjects(Collector); }
 	virtual FString GetReferencerName() const override { return TEXT("FRigVMNewEditor"); }
 	virtual TStatId GetStatId() const override { RETURN_QUICK_DECLARE_CYCLE_STAT(FRigVMNewEditor, STATGROUP_Tickables); }
-	virtual void StartEditingDefaults(bool bAutoFocus = true, bool bForceRefresh = false);
+	UE_API virtual void StartEditingDefaults(bool bAutoFocus = true, bool bForceRefresh = false);
 
-	float GetInstructionTextOpacity(UEdGraph* InGraph) const;
-	virtual void ClearSelectionStateFor(FName SelectionOwner);
+	UE_API float GetInstructionTextOpacity(UEdGraph* InGraph) const;
+	UE_API virtual void ClearSelectionStateFor(FName SelectionOwner);
 protected:
-	void OnLogTokenClicked(const TSharedRef<IMessageToken>& MessageToken);
+	UE_API void OnLogTokenClicked(const TSharedRef<IMessageToken>& MessageToken);
 
 	/** Dumps messages to the compiler log, with an option to force it to display/come to front */
-	void DumpMessagesToCompilerLog(const TArray<TSharedRef<class FTokenizedMessage>>& Messages, bool bForceMessageDisplay);
+	UE_API void DumpMessagesToCompilerLog(const TArray<TSharedRef<class FTokenizedMessage>>& Messages, bool bForceMessageDisplay);
 public:
-	void CreateDefaultTabContents(const TArray<UBlueprint*> InBlueprints);
-
-	TSharedRef<SWidget> GenerateCompileOptionsMenu();
-	void MakeSaveOnCompileSubMenu(FMenuBuilder& InMenu);
-	void SetSaveOnCompileSetting(ESaveOnCompile NewSetting);
-	bool IsSaveOnCompileEnabled() const;
-	bool IsSaveOnCompileOptionSet(TWeakPtr<FRigVMNewEditor> Editor, ESaveOnCompile Option);
-	void ToggleJumpToErrorNodeSetting();
-	bool IsJumpToErrorNodeOptionSet();
-	UEdGraphNode* FindNodeWithError(UBlueprint* Blueprint, EMessageSeverity::Type Severity);
-	UEdGraphNode* FindNodeWithError(FCompilerResultsLog const& ErrorLog, EMessageSeverity::Type Severity);
-	FText GetCompileStatusTooltip() const;
-	FSlateIcon GetCompileStatusImage() const;
+	UE_API void CreateDefaultTabContents(const TArray<UBlueprint*> InBlueprints);
+
+	UE_API TSharedRef<SWidget> GenerateCompileOptionsMenu();
+	UE_API void MakeSaveOnCompileSubMenu(FMenuBuilder& InMenu);
+	UE_API void SetSaveOnCompileSetting(ESaveOnCompile NewSetting);
+	UE_API bool IsSaveOnCompileEnabled() const;
+	UE_API bool IsSaveOnCompileOptionSet(TWeakPtr<FRigVMNewEditor> Editor, ESaveOnCompile Option);
+	UE_API void ToggleJumpToErrorNodeSetting();
+	UE_API bool IsJumpToErrorNodeOptionSet();
+	UE_API UEdGraphNode* FindNodeWithError(UBlueprint* Blueprint, EMessageSeverity::Type Severity);
+	UE_API UEdGraphNode* FindNodeWithError(FCompilerResultsLog const& ErrorLog, EMessageSeverity::Type Severity);
+	UE_API FText GetCompileStatusTooltip() const;
+	UE_API FSlateIcon GetCompileStatusImage() const;
 
 public:
-	static const FSlateBrush* GetGlyphForGraph(const UEdGraph* Graph, bool bInLargeIcon);
+	static UE_API const FSlateBrush* GetGlyphForGraph(const UEdGraph* Graph, bool bInLargeIcon);
 
-	static const FName SelectionState_GraphExplorer();
-	static const FName SelectionState_Graph();
-	static const FName SelectionState_ClassSettings();
-	static const FName SelectionState_ClassDefaults();
+	static UE_API const FName SelectionState_GraphExplorer();
+	static UE_API const FName SelectionState_Graph();
+	static UE_API const FName SelectionState_ClassSettings();
+	static UE_API const FName SelectionState_ClassDefaults();
 
 	virtual FNotifyHook* GetNotifyHook() override { return this; }
 
-	void OnSelectedNodesChanged(const FGraphPanelSelectionSet& NewSelection);
+	UE_API void OnSelectedNodesChanged(const FGraphPanelSelectionSet& NewSelection);
 
-	void OnAlignTop();
-	void OnAlignMiddle();
-	void OnAlignBottom();
-	void OnAlignLeft();
-	void OnAlignCenter();
-	void OnAlignRight();
-	void OnStraightenConnections();
-	void OnDistributeNodesH();
-	void OnDistributeNodesV();
-	void SelectAllNodes();
-	bool CanSelectAllNodes() const;
+	UE_API void OnAlignTop();
+	UE_API void OnAlignMiddle();
+	UE_API void OnAlignBottom();
+	UE_API void OnAlignLeft();
+	UE_API void OnAlignCenter();
+	UE_API void OnAlignRight();
+	UE_API void OnStraightenConnections();
+	UE_API void OnDistributeNodesH();
+	UE_API void OnDistributeNodesV();
+	UE_API void SelectAllNodes();
+	UE_API bool CanSelectAllNodes() const;
 
 protected:
 	
@@ -236,3 +238,4 @@ protected:
 };
 
 
+#undef UE_API
diff --git a/Engine/Plugins/Runtime/RigVM/Source/RigVMEditor/Public/Editor/RigVMNewEditorMode.h b/Engine/Plugins/Runtime/RigVM/Source/RigVMEditor/Public/Editor/RigVMNewEditorMode.h
--- a/Engine/Plugins/Runtime/RigVM/Source/RigVMEditor/Public/Editor/RigVMNewEditorMode.h
+++ b/Engine/Plugins/Runtime/RigVM/Source/RigVMEditor/Public/Editor/RigVMNewEditorMode.h
@@ -6,30 +6,32 @@
 #include "WorkflowOrientedApp/WorkflowTabManager.h"
 #include "WorkflowOrientedApp/ApplicationMode.h"
 
+#define UE_API RIGVMEDITOR_API
+
 class URigVMBlueprint;
 class FRigVMNewEditor;
 
-struct RIGVMEDITOR_API FRigVMNewEditorApplicationModes
+struct FRigVMNewEditorApplicationModes
 {
 	// Mode identifiers
-	static const FName StandardRigVMEditorMode();
-	static const FName RigVMDefaultsMode();
+	static UE_API const FName StandardRigVMEditorMode();
+	static UE_API const FName RigVMDefaultsMode();
 	
 private:
 	FRigVMNewEditorApplicationModes() {}
 };
 
-class RIGVMEDITOR_API FRigVMNewEditorMode : public FApplicationMode
+class FRigVMNewEditorMode : public FApplicationMode
 {
 public:
-	FRigVMNewEditorMode(const TSharedRef<FRigVMNewEditor>& InRigVMEditor); 
+	UE_API FRigVMNewEditorMode(const TSharedRef<FRigVMNewEditor>& InRigVMEditor); 
 
 	// FApplicationMode interface
-	virtual void RegisterTabFactories(TSharedPtr<FTabManager> InTabManager) override;
+	UE_API virtual void RegisterTabFactories(TSharedPtr<FTabManager> InTabManager) override;
 
-	void RegisterToolbarTab(const TSharedRef<class FTabManager>& TabManager);
+	UE_API void RegisterToolbarTab(const TSharedRef<class FTabManager>& TabManager);
 
-	virtual void PostActivateMode() override;
+	UE_API virtual void PostActivateMode() override;
 
 protected:
 	// Set of spawnable tabs
@@ -37,3 +39,5 @@ protected:
 
 	TWeakPtr<FRigVMNewEditor> Editor;
 };
+
+#undef UE_API
diff --git a/Engine/Plugins/Runtime/RigVM/Source/RigVMEditor/Public/Editor/RigVMVariableDetailCustomization.h b/Engine/Plugins/Runtime/RigVM/Source/RigVMEditor/Public/Editor/RigVMVariableDetailCustomization.h
--- a/Engine/Plugins/Runtime/RigVM/Source/RigVMEditor/Public/Editor/RigVMVariableDetailCustomization.h
+++ b/Engine/Plugins/Runtime/RigVM/Source/RigVMEditor/Public/Editor/RigVMVariableDetailCustomization.h
@@ -12,6 +12,8 @@
 #include "Widgets/Views/STableViewBase.h"
 #include "Widgets/Views/SListView.h"
 
+#define UE_API RIGVMEDITOR_API
+
 class URigVMBlueprint;
 class SRigVMEditorGraphExplorer;
 class IDetailLayoutBuilder;
@@ -24,66 +26,66 @@ class FStructOnScope;
 class SEditableTextBox;
 class SComboButton;
 
-class RIGVMEDITOR_API FRigVMVariableDetailCustomization : public IDetailCustomization
+class FRigVMVariableDetailCustomization : public IDetailCustomization
 {
 public:
 	/** Makes a new instance of this detail layout class for a specific detail view requesting it */
-	static TSharedRef<IDetailCustomization> MakeInstance(TSharedPtr<IRigVMEditor> InEditor);
-	FRigVMVariableDetailCustomization(TSharedPtr<IRigVMEditor> InEditor, URigVMBlueprint* Blueprint);
+	static UE_API TSharedRef<IDetailCustomization> MakeInstance(TSharedPtr<IRigVMEditor> InEditor);
+	UE_API FRigVMVariableDetailCustomization(TSharedPtr<IRigVMEditor> InEditor, URigVMBlueprint* Blueprint);
 
 #if WITH_RIGVMLEGACYEDITOR
-	static TSharedPtr<IDetailCustomization> MakeLegacyInstance(TSharedPtr<IBlueprintEditor> InBlueprintEditor);
-	FRigVMVariableDetailCustomization(TSharedPtr<IBlueprintEditor> RigVMigEditor, UBlueprint* Blueprint);
+	static UE_API TSharedPtr<IDetailCustomization> MakeLegacyInstance(TSharedPtr<IBlueprintEditor> InBlueprintEditor);
+	UE_API FRigVMVariableDetailCustomization(TSharedPtr<IBlueprintEditor> RigVMigEditor, UBlueprint* Blueprint);
 #endif
 
 	// IDetailCustomization interface
-	virtual void CustomizeDetails(IDetailLayoutBuilder& DetailLayout) override;
+	UE_API virtual void CustomizeDetails(IDetailLayoutBuilder& DetailLayout) override;
 
 private:
-	void PopulateCategories();
+	UE_API void PopulateCategories();
 
 private:
 	/** Accessors passed to parent */
-	FName GetVariableName() const;
-	FText OnGetVariableName() const;
-	void OnVarNameCommitted(const FText& InNewName, ETextCommit::Type InTextCommit);
+	UE_API FName GetVariableName() const;
+	UE_API FText OnGetVariableName() const;
+	UE_API void OnVarNameCommitted(const FText& InNewName, ETextCommit::Type InTextCommit);
 
 	// Callbacks for uproperty details customization
-	FEdGraphPinType OnGetVarType() const;
-	void OnVarTypeChanged(const FEdGraphPinType& NewPinType);
+	UE_API FEdGraphPinType OnGetVarType() const;
+	UE_API void OnVarTypeChanged(const FEdGraphPinType& NewPinType);
 
-	void OnBrowseToVarType() const;
-	bool CanBrowseToVarType() const;
+	UE_API void OnBrowseToVarType() const;
+	UE_API bool CanBrowseToVarType() const;
 
 	
-	FText OnGetTooltipText() const;
-	void OnTooltipTextCommitted(const FText& NewText, ETextCommit::Type InTextCommit, FName VarName);
-	EVisibility IsToolTipVisible() const;
+	UE_API FText OnGetTooltipText() const;
+	UE_API void OnTooltipTextCommitted(const FText& NewText, ETextCommit::Type InTextCommit, FName VarName);
+	UE_API EVisibility IsToolTipVisible() const;
 	
-	FText OnGetCategoryText() const;
-	void OnCategoryTextCommitted(const FText& NewText, ETextCommit::Type InTextCommit, FName VarName);
-	TSharedRef< ITableRow > MakeCategoryViewWidget( TSharedPtr<FText> Item, const TSharedRef< STableViewBase >& OwnerTable );
-	void OnCategorySelectionChanged( TSharedPtr<FText> ProposedSelection, ESelectInfo::Type /*SelectInfo*/ );
+	UE_API FText OnGetCategoryText() const;
+	UE_API void OnCategoryTextCommitted(const FText& NewText, ETextCommit::Type InTextCommit, FName VarName);
+	UE_API TSharedRef< ITableRow > MakeCategoryViewWidget( TSharedPtr<FText> Item, const TSharedRef< STableViewBase >& OwnerTable );
+	UE_API void OnCategorySelectionChanged( TSharedPtr<FText> ProposedSelection, ESelectInfo::Type /*SelectInfo*/ );
 	
 
-	ECheckBoxState OnGetExposedToSpawnCheckboxState() const;
-	void OnExposedToSpawnChanged(ECheckBoxState InNewState);
+	UE_API ECheckBoxState OnGetExposedToSpawnCheckboxState() const;
+	UE_API void OnExposedToSpawnChanged(ECheckBoxState InNewState);
 
-	ECheckBoxState OnGetPrivateCheckboxState() const;
-	void OnPrivateChanged(ECheckBoxState InNewState);
+	UE_API ECheckBoxState OnGetPrivateCheckboxState() const;
+	UE_API void OnPrivateChanged(ECheckBoxState InNewState);
 	
-	ECheckBoxState OnGetExposedToCinematicsCheckboxState() const;
-	void OnExposedToCinematicsChanged(ECheckBoxState InNewState);
+	UE_API ECheckBoxState OnGetExposedToCinematicsCheckboxState() const;
+	UE_API void OnExposedToCinematicsChanged(ECheckBoxState InNewState);
 
 
-	FText OnGetMetaKeyValue(FName Key) const;
-	void OnMetaKeyValueChanged(const FText& NewMinValue, ETextCommit::Type CommitInfo, FName Key);
-	EVisibility RangeVisibility() const;
+	UE_API FText OnGetMetaKeyValue(FName Key) const;
+	UE_API void OnMetaKeyValueChanged(const FText& NewMinValue, ETextCommit::Type CommitInfo, FName Key);
+	UE_API EVisibility RangeVisibility() const;
 
 
 	
 	/** Refreshes cached data that changes after a Blueprint recompile */
-	void OnPostEditorRefresh();
+	UE_API void OnPostEditorRefresh();
 
 private:
 	/** The Blueprint editor we are embedded in */
@@ -107,3 +109,5 @@ private:
 	/** Cached name for the variable we are affecting */
 	mutable FName CachedVariableName;
 };
+
+#undef UE_API
diff --git a/Engine/Plugins/Runtime/RigVM/Source/RigVMEditor/Public/Editor/SRigVMDetailsInspector.h b/Engine/Plugins/Runtime/RigVM/Source/RigVMEditor/Public/Editor/SRigVMDetailsInspector.h
--- a/Engine/Plugins/Runtime/RigVM/Source/RigVMEditor/Public/Editor/SRigVMDetailsInspector.h
+++ b/Engine/Plugins/Runtime/RigVM/Source/RigVMEditor/Public/Editor/SRigVMDetailsInspector.h
@@ -23,6 +23,8 @@
 #include "Widgets/SCompoundWidget.h"
 #include "Widgets/SWidget.h"
 
+#define UE_API RIGVMEDITOR_API
+
 class IRigVMEditor;
 class FProperty;
 class FStructOnScope;
@@ -42,7 +44,7 @@ typedef TSet<class UObject*> FInspectorSelectionSet;
 // SRigVMDetailsInspector
 
 /** Widget that shows properties and tools related to the selected node(s) */
-class RIGVMEDITOR_API SRigVMDetailsInspector : public SCompoundWidget, public FGCObject
+class SRigVMDetailsInspector : public SCompoundWidget, public FGCObject
 {
 public:
 	SLATE_BEGIN_ARGS( SRigVMDetailsInspector )
@@ -70,7 +72,7 @@ public:
 		SLATE_ARGUMENT( bool, ShowSectionSelector)
 	SLATE_END_ARGS()
 	
-	void Construct(const FArguments& InArgs);
+	UE_API void Construct(const FArguments& InArgs);
 
 	/** Options for ShowDetails */
 	struct FShowDetailsOptions
@@ -96,42 +98,42 @@ public:
 	};
 
 	// SWidget interface
-	virtual void Tick( const FGeometry& AllottedGeometry, const double InCurrentTime, const float InDeltaTime ) override;
+	UE_API virtual void Tick( const FGeometry& AllottedGeometry, const double InCurrentTime, const float InDeltaTime ) override;
 	// End of SWidget interface
 
 	/** Update the inspector window to show information on the supplied object */
-	void ShowDetailsForSingleObject(UObject* Object, const FShowDetailsOptions& Options = FShowDetailsOptions());
+	UE_API void ShowDetailsForSingleObject(UObject* Object, const FShowDetailsOptions& Options = FShowDetailsOptions());
 
 	/** Update the inspector window to show information on the supplied objects */
-	void ShowDetailsForObjects(const TArray<UObject*>& PropertyObjects, const FShowDetailsOptions& Options = FShowDetailsOptions());
+	UE_API void ShowDetailsForObjects(const TArray<UObject*>& PropertyObjects, const FShowDetailsOptions& Options = FShowDetailsOptions());
 
 	/** Update the inspector window to show single struct. This invalidates ShowDetailsForObjects */
-	void ShowSingleStruct(TSharedPtr<FStructOnScope> InStructToDisplay);
+	UE_API void ShowSingleStruct(TSharedPtr<FStructOnScope> InStructToDisplay);
 
 	/** Used to control visibility of a property in the property window */
-	bool IsPropertyVisible( const struct FPropertyAndParent& PropertyAndParent ) const;
+	UE_API bool IsPropertyVisible( const struct FPropertyAndParent& PropertyAndParent ) const;
 
 
 	TSharedPtr<class IDetailsView> GetPropertyView() const { return PropertyView; }
 
-	TSharedPtr<SDockTab> GetOwnerTab() const;
-	void SetOwnerTab(TSharedRef<SDockTab> Tab);
+	UE_API TSharedPtr<SDockTab> GetOwnerTab() const;
+	UE_API void SetOwnerTab(TSharedRef<SDockTab> Tab);
 
 
 	/** returns the list of selected objects */
-	const TArray< TWeakObjectPtr<UObject> >& GetSelectedObjects() const;
+	UE_API const TArray< TWeakObjectPtr<UObject> >& GetSelectedObjects() const;
 
-	void OnEditorClose(const IRigVMEditor* RigVMEditorBase, URigVMBlueprint* RigVMBlueprint);
+	UE_API void OnEditorClose(const IRigVMEditor* RigVMEditorBase, URigVMBlueprint* RigVMBlueprint);
 
-	virtual void AddReferencedObjects(FReferenceCollector& Collector) override;
-	virtual FString GetReferencerName() const override;
+	UE_API virtual void AddReferencedObjects(FReferenceCollector& Collector) override;
+	UE_API virtual FString GetReferencerName() const override;
 
 protected:
 	/** Update the inspector window to show information on the supplied objects */
-	void UpdateFromObjects(const TArray<UObject*>& PropertyObjects, TArray<UObject*>& SelectionInfo, const FShowDetailsOptions& Options);
+	UE_API void UpdateFromObjects(const TArray<UObject*>& PropertyObjects, TArray<UObject*>& SelectionInfo, const FShowDetailsOptions& Options);
 
 	/** Add this property and all its child properties to SelectedObjectProperties */
-	void AddPropertiesRecursive(FProperty* Property);
+	UE_API void AddPropertiesRecursive(FProperty* Property);
 
 	/** Pointer back to the kismet 2 tool that owns us */
 	TWeakPtr<IRigVMEditor> WeakEditor;
@@ -188,36 +190,38 @@ protected:
 	TSharedPtr<class IStructureDetailsView> StructureDetailsView;
 
 	/** Update the inspector window to show information on the single struct */
-	void UpdateFromSingleStruct(const TSharedPtr<FStructOnScope>& InStructToDisplay);
+	UE_API void UpdateFromSingleStruct(const TSharedPtr<FStructOnScope>& InStructToDisplay);
 
 	/** Is struct view property read only */
-	bool IsStructViewPropertyReadOnly(const struct FPropertyAndParent& PropertyAndParent) const;
+	UE_API bool IsStructViewPropertyReadOnly(const struct FPropertyAndParent& PropertyAndParent) const;
 
 protected:
 
 	/** Returns whether the property view be visible */
-	EVisibility GetPropertyViewVisibility() const;
+	UE_API EVisibility GetPropertyViewVisibility() const;
 
 	/** Returns whether the properties in the view should be editable */
-	bool IsPropertyEditingEnabled() const;
+	UE_API bool IsPropertyEditingEnabled() const;
 
 	/**
 	 * Generates a widget that is used to edit the specified object array contextually.  This widget
 	 * will be displayed along with a property view in the level editor
 	 */
-	TSharedRef<SWidget> MakeContextualEditingWidget(TArray<UObject*>& SelectionInfo, const FShowDetailsOptions& Options);
+	UE_API TSharedRef<SWidget> MakeContextualEditingWidget(TArray<UObject*>& SelectionInfo, const FShowDetailsOptions& Options);
 
 	/**
 	 * Generates the text for the title in the contextual editing widget
 	 */
-	FText GetContextualEditingWidgetTitle() const;
+	UE_API FText GetContextualEditingWidgetTitle() const;
 
-	ECheckBoxState GetPublicViewCheckboxState() const;
-	void SetPublicViewCheckboxState(ECheckBoxState InIsChecked);
+	UE_API ECheckBoxState GetPublicViewCheckboxState() const;
+	UE_API void SetPublicViewCheckboxState(ECheckBoxState InIsChecked);
 
-	bool IsAnyParentOrContainerSelected(const FPropertyAndParent& PropertyAndParent) const;
+	UE_API bool IsAnyParentOrContainerSelected(const FPropertyAndParent& PropertyAndParent) const;
 
 	/** Callback invoked after a value change on the selected object(s) */
-	void OnFinishedChangingProperties(const FPropertyChangedEvent& InPropertyChangedEvent);
+	UE_API void OnFinishedChangingProperties(const FPropertyChangedEvent& InPropertyChangedEvent);
 
 };
+
+#undef UE_API
diff --git a/Engine/Plugins/Runtime/RigVM/Source/RigVMEditor/Public/RigVMEditorBlueprintLibrary.h b/Engine/Plugins/Runtime/RigVM/Source/RigVMEditor/Public/RigVMEditorBlueprintLibrary.h
--- a/Engine/Plugins/Runtime/RigVM/Source/RigVMEditor/Public/RigVMEditorBlueprintLibrary.h
+++ b/Engine/Plugins/Runtime/RigVM/Source/RigVMEditor/Public/RigVMEditorBlueprintLibrary.h
@@ -13,6 +13,8 @@
 #include "Kismet/BlueprintFunctionLibrary.h"
 #include "RigVMEditorBlueprintLibrary.generated.h"
 
+#define UE_API RIGVMEDITOR_API
+
 UENUM(BlueprintType)
 enum ERigVMBlueprintLoadLogSeverity : uint8
 {
@@ -22,7 +24,7 @@ enum ERigVMBlueprintLoadLogSeverity : uint8
 };
 
 USTRUCT(BlueprintType)
-struct RIGVMEDITOR_API FRigVMBlueprintLoadLogEntry
+struct FRigVMBlueprintLoadLogEntry
 {
 	GENERATED_BODY()
 
@@ -54,81 +56,83 @@ DECLARE_DYNAMIC_DELEGATE_RetVal_OneParam(bool, FRigVMAssetDataFilterDynamic, FAs
 DECLARE_DYNAMIC_DELEGATE_RetVal_TwoParams(bool, FRigVMBlueprintFilterDynamic, const URigVMBlueprint*, Blueprint, TArray<FRigVMBlueprintLoadLogEntry>, LogDuringLoad);
 DECLARE_DYNAMIC_DELEGATE_RetVal_TwoParams(bool, FRigVMNodeFilterDynamic, const URigVMBlueprint*, Blueprint, const URigVMNode*, Node);
 
-UCLASS(BlueprintType, meta=(ScriptName="RigVMBlueprintLibrary"))
-class RIGVMEDITOR_API URigVMEditorBlueprintLibrary : public UBlueprintFunctionLibrary
+UCLASS(MinimalAPI, BlueprintType, meta=(ScriptName="RigVMBlueprintLibrary"))
+class URigVMEditorBlueprintLibrary : public UBlueprintFunctionLibrary
 {
 	GENERATED_BODY()
 
 public:
 
 	UFUNCTION(BlueprintCallable, Category = "RigVM Blueprint")
-	static void RecompileVM(URigVMBlueprint* InBlueprint);
+	static UE_API void RecompileVM(URigVMBlueprint* InBlueprint);
 
 	UFUNCTION(BlueprintCallable, Category = "RigVM Blueprint")
-	static void RecompileVMIfRequired(URigVMBlueprint* InBlueprint);
+	static UE_API void RecompileVMIfRequired(URigVMBlueprint* InBlueprint);
 	
 	UFUNCTION(BlueprintCallable, Category = "RigVM Blueprint")
-	static void RequestAutoVMRecompilation(URigVMBlueprint* InBlueprint);
+	static UE_API void RequestAutoVMRecompilation(URigVMBlueprint* InBlueprint);
 
 	UFUNCTION(BlueprintCallable, BlueprintPure, Category = "RigVM Blueprint")
-	static URigVMGraph* GetModel(URigVMBlueprint* InBlueprint);
+	static UE_API URigVMGraph* GetModel(URigVMBlueprint* InBlueprint);
 
 	UFUNCTION(BlueprintCallable, BlueprintPure, Category = "RigVM Blueprint")
-	static URigVMController* GetController(URigVMBlueprint* InBlueprint);
+	static UE_API URigVMController* GetController(URigVMBlueprint* InBlueprint);
 
 	UFUNCTION(BlueprintCallable, Category = "RigVM Blueprint")
-	static TArray<URigVMBlueprint*> LoadAssets();
+	static UE_API TArray<URigVMBlueprint*> LoadAssets();
 
 	UFUNCTION(BlueprintCallable, Category = "RigVM Blueprint")
-	static TArray<URigVMBlueprint*> LoadAssetsByClass(TSubclassOf<URigVMBlueprint> InClass);
+	static UE_API TArray<URigVMBlueprint*> LoadAssetsByClass(TSubclassOf<URigVMBlueprint> InClass);
 
 	UFUNCTION(BlueprintCallable, Category = "RigVM Blueprint", meta = (DisplayName = "LoadAssetsWithBlueprintFilter", ScriptName = "LoadAssetsWithBlueprintFilter"))
-		static TArray<URigVMBlueprint*> LoadAssetsWithBlueprintFilter_ForBlueprint(
+		static UE_API TArray<URigVMBlueprint*> LoadAssetsWithBlueprintFilter_ForBlueprint(
 		TSubclassOf<URigVMBlueprint> InClass,
 		FRigVMBlueprintFilterDynamic InBlueprintFilter);
 
 	UFUNCTION(BlueprintCallable, Category = "RigVM Blueprint", meta = (DisplayName = "LoadAssetsWithAssetDataFilter", ScriptName = "LoadAssetsWithAssetDataFilter"))
-		static TArray<URigVMBlueprint*> LoadAssetsWithAssetDataFilter_ForBlueprint(
+		static UE_API TArray<URigVMBlueprint*> LoadAssetsWithAssetDataFilter_ForBlueprint(
 		TSubclassOf<URigVMBlueprint> InClass,
 		FRigVMAssetDataFilterDynamic InAssetDataFilter);
 
 	UFUNCTION(BlueprintCallable, Category = "RigVM Blueprint", meta = (DisplayName = "LoadAssetsWithNodeFilter", ScriptName = "LoadAssetsWithNodeFilter"))
-	static TArray<URigVMBlueprint*> LoadAssetsWithNodeFilter_ForBlueprint(
+	static UE_API TArray<URigVMBlueprint*> LoadAssetsWithNodeFilter_ForBlueprint(
 		TSubclassOf<URigVMBlueprint> InClass,
 		FRigVMNodeFilterDynamic InNodeFilter);
 
-	static TArray<URigVMBlueprint*> LoadAssetsWithNodeFilter(
+	static UE_API TArray<URigVMBlueprint*> LoadAssetsWithNodeFilter(
 		TSubclassOf<URigVMBlueprint> InClass,
 		FRigVMNodeFilter InNodeFilter);
 
 	UFUNCTION(BlueprintCallable, Category = "RigVM Blueprint", meta = (DisplayName = "LoadAssetsWithAssetDataAndBlueprintFilters", ScriptName = "LoadAssetsWithAssetDataAndBlueprintFilters"))
-	static TArray<URigVMBlueprint*> LoadAssetsWithAssetDataAndBlueprintFilters_ForBlueprint(
+	static UE_API TArray<URigVMBlueprint*> LoadAssetsWithAssetDataAndBlueprintFilters_ForBlueprint(
 		TSubclassOf<URigVMBlueprint> InClass,
 		FRigVMAssetDataFilterDynamic InAssetDataFilter,
 		FRigVMBlueprintFilterDynamic InBlueprintFilter);
 
-	static TArray<URigVMBlueprint*> LoadAssetsWithAssetDataAndBlueprintFilters(
+	static UE_API TArray<URigVMBlueprint*> LoadAssetsWithAssetDataAndBlueprintFilters(
 		TSubclassOf<URigVMBlueprint> InClass,
 		FRigVMAssetDataFilter InAssetDataFilter,
 		FRigVMBlueprintFilter InBlueprintFilter);
 
 	UFUNCTION(BlueprintCallable, Category = "RigVM Blueprint", meta = (DisplayName = "LoadAssetsWithAssetDataAndNodeFilters", ScriptName = "LoadAssetsWithAssetDataAndNodeFilters"))
-	static TArray<URigVMBlueprint*> LoadAssetsWithAssetDataAndNodeFilters_ForBlueprint(
+	static UE_API TArray<URigVMBlueprint*> LoadAssetsWithAssetDataAndNodeFilters_ForBlueprint(
 		TSubclassOf<URigVMBlueprint> InClass,
 		FRigVMAssetDataFilterDynamic InAssetDataFilter,
 		FRigVMNodeFilterDynamic InNodeFilter);
 
-	static TArray<URigVMBlueprint*> LoadAssetsWithAssetDataAndNodeFilters(
+	static UE_API TArray<URigVMBlueprint*> LoadAssetsWithAssetDataAndNodeFilters(
 		TSubclassOf<URigVMBlueprint> InClass,
 		FRigVMAssetDataFilter InAssetDataFilter,
 		FRigVMNodeFilter InNodeFilter);
 
 	UFUNCTION(BlueprintCallable, Category = "RigVM Blueprint", meta = (DisplayName = "GetAssetsWithFilter", ScriptName = "GetAssetsWithFilter"))
-	static TArray<FAssetData> GetAssetsWithFilter_ForBlueprint(
+	static UE_API TArray<FAssetData> GetAssetsWithFilter_ForBlueprint(
 		TSubclassOf<URigVMBlueprint> InClass,
 		FRigVMAssetDataFilterDynamic InAssetDataFilter);
 
-	static TArray<FAssetData> GetAssetsWithFilter(
+	static UE_API TArray<FAssetData> GetAssetsWithFilter(
 		TSubclassOf<URigVMBlueprint> InClass,
 		FRigVMAssetDataFilter InAssetDataFilter);
 };
+
+#undef UE_API
diff --git a/Engine/Plugins/Runtime/RigVM/Source/RigVMEditor/Public/RigVMEditorModule.h b/Engine/Plugins/Runtime/RigVM/Source/RigVMEditor/Public/RigVMEditorModule.h
--- a/Engine/Plugins/Runtime/RigVM/Source/RigVMEditor/Public/RigVMEditorModule.h
+++ b/Engine/Plugins/Runtime/RigVM/Source/RigVMEditor/Public/RigVMEditorModule.h
@@ -15,6 +15,8 @@
 #include "Kismet2/StructureEditorUtils.h"
 #include "Modules/ModuleManager.h"
 
+#define UE_API RIGVMEDITOR_API
+
 class URigVMUserWorkflowOptions;
 class URigVMPin;
 class URigVMNode;
@@ -60,26 +62,26 @@ public:
 	virtual bool AssetsPublicFunctionsAllowed(const FAssetData& InAssetData) const = 0;
 };
 
-class RIGVMEDITOR_API FRigVMEditorModule : public IRigVMEditorModule
+class FRigVMEditorModule : public IRigVMEditorModule
 {
 public:
 
-	static FRigVMEditorModule& Get();
+	static UE_API FRigVMEditorModule& Get();
 
 	/** IModuleInterface implementation */
-	virtual void StartupModule() override;
-	virtual void ShutdownModule() override;
+	UE_API virtual void StartupModule() override;
+	UE_API virtual void ShutdownModule() override;
 
-	void StartupModuleCommon();
-	void ShutdownModuleCommon();
+	UE_API void StartupModuleCommon();
+	UE_API void ShutdownModuleCommon();
 
 	/** RigVMEditorModule interface */
-	virtual UClass* GetRigVMBlueprintClass() const;
-	const URigVMBlueprint* GetRigVMBlueprintCDO() const;
-	virtual void GetTypeActions(URigVMBlueprint* RigVMBlueprint, FBlueprintActionDatabaseRegistrar& ActionRegistrar);
-	virtual void GetInstanceActions(URigVMBlueprint* RigVMBlueprint, FBlueprintActionDatabaseRegistrar& ActionRegistrar);
-	virtual void GetNodeContextMenuActions(IRigVMClientHost* RigVMClientHost, const URigVMEdGraphNode* EdGraphNode, URigVMNode* ModelNode, class UToolMenu* Menu) const;
-	virtual void GetPinContextMenuActions(IRigVMClientHost* RigVMClientHost, const UEdGraphPin* EdGraphPin, URigVMPin* ModelPin, class UToolMenu* Menu) const;
+	UE_API virtual UClass* GetRigVMBlueprintClass() const;
+	UE_API const URigVMBlueprint* GetRigVMBlueprintCDO() const;
+	UE_API virtual void GetTypeActions(URigVMBlueprint* RigVMBlueprint, FBlueprintActionDatabaseRegistrar& ActionRegistrar);
+	UE_API virtual void GetInstanceActions(URigVMBlueprint* RigVMBlueprint, FBlueprintActionDatabaseRegistrar& ActionRegistrar);
+	UE_API virtual void GetNodeContextMenuActions(IRigVMClientHost* RigVMClientHost, const URigVMEdGraphNode* EdGraphNode, URigVMNode* ModelNode, class UToolMenu* Menu) const;
+	UE_API virtual void GetPinContextMenuActions(IRigVMClientHost* RigVMClientHost, const UEdGraphPin* EdGraphPin, URigVMPin* ModelPin, class UToolMenu* Menu) const;
 	virtual bool AssetsPublicFunctionsAllowed(const FAssetData& InAssetData) const override { return true; }
 
 	/** IHasMenuExtensibility interface */
@@ -89,45 +91,45 @@ public:
 	virtual TSharedPtr<FExtensibilityManager> GetToolBarExtensibilityManager() override { return ToolBarExtensibilityManager; }
 
 	/* FStructureEditorUtils::INotifyOnStructChanged Interface, used to respond to changes to user defined structs */
-	virtual void PreChange(const UUserDefinedStruct* Changed, FStructureEditorUtils::EStructureEditorChangeInfo ChangedType) override;
-	virtual void PostChange(const UUserDefinedStruct* Changed, FStructureEditorUtils::EStructureEditorChangeInfo ChangedType) override;
+	UE_API virtual void PreChange(const UUserDefinedStruct* Changed, FStructureEditorUtils::EStructureEditorChangeInfo ChangedType) override;
+	UE_API virtual void PostChange(const UUserDefinedStruct* Changed, FStructureEditorUtils::EStructureEditorChangeInfo ChangedType) override;
 
 	/** Get all toolbar extenders */
 	DECLARE_DELEGATE_RetVal_TwoParams(TSharedRef<FExtender>, FRigVMEditorToolbarExtender, const TSharedRef<FUICommandList> /*InCommandList*/, TSharedRef<FRigVMEditorBase> /*InRigVMEditor*/);
-	virtual TArray<FRigVMEditorToolbarExtender>& GetAllRigVMEditorToolbarExtenders();
+	UE_API virtual TArray<FRigVMEditorToolbarExtender>& GetAllRigVMEditorToolbarExtenders();
 
-	virtual void GetContextMenuActions(const URigVMEdGraphSchema* Schema, class UToolMenu* Menu, class UGraphNodeContextMenuContext* Context) const override;
+	UE_API virtual void GetContextMenuActions(const URigVMEdGraphSchema* Schema, class UToolMenu* Menu, class UGraphNodeContextMenuContext* Context) const override;
 
 	/** Make sure to create the root graph for a given blueprint */
-	void CreateRootGraphIfRequired(URigVMBlueprint* InBlueprint) const;
+	UE_API void CreateRootGraphIfRequired(URigVMBlueprint* InBlueprint) const;
 
-	virtual FConnectionDrawingPolicy* CreateConnectionDrawingPolicy(int32 InBackLayerID, int32 InFrontLayerID, float InZoomFactor, const FSlateRect& InClippingRect, FSlateWindowElementList& InDrawElements, UEdGraph* InGraphObj) const override;
+	UE_API virtual FConnectionDrawingPolicy* CreateConnectionDrawingPolicy(int32 InBackLayerID, int32 InFrontLayerID, float InZoomFactor, const FSlateRect& InClippingRect, FSlateWindowElementList& InDrawElements, UEdGraph* InGraphObj) const override;
 
 protected:
 
-	bool IsRigVMEditorModuleBase() const;
+	UE_API bool IsRigVMEditorModuleBase() const;
 
 	/** Specific section callbacks for the context menu */
-	virtual void GetNodeWorkflowContextMenuActions(IRigVMClientHost* RigVMClientHost, const URigVMEdGraphNode* EdGraphNode, URigVMNode* ModelNode, UToolMenu* Menu) const;
-	virtual void GetNodeEventsContextMenuActions(IRigVMClientHost* RigVMClientHost, const URigVMEdGraphNode* EdGraphNode, URigVMNode* ModelNode, UToolMenu* Menu) const;
-	virtual void GetNodeDefaultValueContextMenuActions(IRigVMClientHost* RigVMClientHost, const URigVMEdGraphNode* EdGraphNode, URigVMNode* ModelNode, UToolMenu* Menu) const;
-	virtual void GetNodeConversionContextMenuActions(IRigVMClientHost* RigVMClientHost, const URigVMEdGraphNode* EdGraphNode, URigVMNode* ModelNode, UToolMenu* Menu) const;
-	virtual void GetNodeDebugContextMenuActions(IRigVMClientHost* RigVMClientHost, const URigVMEdGraphNode* EdGraphNode, URigVMNode* ModelNode, UToolMenu* Menu) const;
-	virtual void GetNodeVariablesContextMenuActions(IRigVMClientHost* RigVMClientHost, const URigVMEdGraphNode* EdGraphNode, URigVMNode* ModelNode, UToolMenu* Menu) const;
-	virtual void GetNodeTemplatesContextMenuActions(IRigVMClientHost* RigVMClientHost, const URigVMEdGraphNode* EdGraphNode, URigVMNode* ModelNode, UToolMenu* Menu) const;
-	virtual void GetNodeOrganizationContextMenuActions(IRigVMClientHost* RigVMClientHost, const URigVMEdGraphNode* EdGraphNode, URigVMNode* ModelNode, UToolMenu* Menu) const;
-	virtual void GetNodeVariantContextMenuActions(IRigVMClientHost* RigVMClientHost, const URigVMEdGraphNode* EdGraphNode, URigVMNode* ModelNode, UToolMenu* Menu) const;
-	virtual void GetNodeTestContextMenuActions(IRigVMClientHost* RigVMClientHost, const URigVMEdGraphNode* EdGraphNode, URigVMNode* ModelNode, UToolMenu* Menu) const;
-	virtual void GetNodeDisplayContextMenuActions(IRigVMClientHost* RigVMClientHost, const URigVMEdGraphNode* EdGraphNode, URigVMNode* ModelNode, UToolMenu* Menu) const;
-	virtual void GetPinWorkflowContextMenuActions(IRigVMClientHost* RigVMClientHost, const UEdGraphPin* EdGraphPin, URigVMPin* ModelPin, UToolMenu* Menu) const;
-	virtual void GetPinDebugContextMenuActions(IRigVMClientHost* RigVMClientHost, const UEdGraphPin* EdGraphPin, URigVMPin* ModelPin, UToolMenu* Menu) const;
-	virtual void GetPinArrayContextMenuActions(IRigVMClientHost* RigVMClientHost, const UEdGraphPin* EdGraphPin, URigVMPin* ModelPin, UToolMenu* Menu) const;
-	virtual void GetPinAggregateContextMenuActions(IRigVMClientHost* RigVMClientHost, const UEdGraphPin* EdGraphPin, URigVMPin* ModelPin, UToolMenu* Menu) const;
-	virtual void GetPinTemplateContextMenuActions(IRigVMClientHost* RigVMClientHost, const UEdGraphPin* EdGraphPin, URigVMPin* ModelPin, UToolMenu* Menu) const;
-	virtual void GetPinConversionContextMenuActions(IRigVMClientHost* RigVMClientHost, const UEdGraphPin* EdGraphPin, URigVMPin* ModelPin, UToolMenu* Menu) const;
-	virtual void GetPinVariableContextMenuActions(IRigVMClientHost* RigVMClientHost, const UEdGraphPin* EdGraphPin, URigVMPin* ModelPin, UToolMenu* Menu) const;
-	virtual void GetPinResetDefaultContextMenuActions(IRigVMClientHost* RigVMClientHost, const UEdGraphPin* EdGraphPin, URigVMPin* ModelPin, UToolMenu* Menu) const;
-	virtual void GetPinInjectedNodesContextMenuActions(IRigVMClientHost* RigVMClientHost, const UEdGraphPin* EdGraphPin, URigVMPin* ModelPin, UToolMenu* Menu) const;
+	UE_API virtual void GetNodeWorkflowContextMenuActions(IRigVMClientHost* RigVMClientHost, const URigVMEdGraphNode* EdGraphNode, URigVMNode* ModelNode, UToolMenu* Menu) const;
+	UE_API virtual void GetNodeEventsContextMenuActions(IRigVMClientHost* RigVMClientHost, const URigVMEdGraphNode* EdGraphNode, URigVMNode* ModelNode, UToolMenu* Menu) const;
+	UE_API virtual void GetNodeDefaultValueContextMenuActions(IRigVMClientHost* RigVMClientHost, const URigVMEdGraphNode* EdGraphNode, URigVMNode* ModelNode, UToolMenu* Menu) const;
+	UE_API virtual void GetNodeConversionContextMenuActions(IRigVMClientHost* RigVMClientHost, const URigVMEdGraphNode* EdGraphNode, URigVMNode* ModelNode, UToolMenu* Menu) const;
+	UE_API virtual void GetNodeDebugContextMenuActions(IRigVMClientHost* RigVMClientHost, const URigVMEdGraphNode* EdGraphNode, URigVMNode* ModelNode, UToolMenu* Menu) const;
+	UE_API virtual void GetNodeVariablesContextMenuActions(IRigVMClientHost* RigVMClientHost, const URigVMEdGraphNode* EdGraphNode, URigVMNode* ModelNode, UToolMenu* Menu) const;
+	UE_API virtual void GetNodeTemplatesContextMenuActions(IRigVMClientHost* RigVMClientHost, const URigVMEdGraphNode* EdGraphNode, URigVMNode* ModelNode, UToolMenu* Menu) const;
+	UE_API virtual void GetNodeOrganizationContextMenuActions(IRigVMClientHost* RigVMClientHost, const URigVMEdGraphNode* EdGraphNode, URigVMNode* ModelNode, UToolMenu* Menu) const;
+	UE_API virtual void GetNodeVariantContextMenuActions(IRigVMClientHost* RigVMClientHost, const URigVMEdGraphNode* EdGraphNode, URigVMNode* ModelNode, UToolMenu* Menu) const;
+	UE_API virtual void GetNodeTestContextMenuActions(IRigVMClientHost* RigVMClientHost, const URigVMEdGraphNode* EdGraphNode, URigVMNode* ModelNode, UToolMenu* Menu) const;
+	UE_API virtual void GetNodeDisplayContextMenuActions(IRigVMClientHost* RigVMClientHost, const URigVMEdGraphNode* EdGraphNode, URigVMNode* ModelNode, UToolMenu* Menu) const;
+	UE_API virtual void GetPinWorkflowContextMenuActions(IRigVMClientHost* RigVMClientHost, const UEdGraphPin* EdGraphPin, URigVMPin* ModelPin, UToolMenu* Menu) const;
+	UE_API virtual void GetPinDebugContextMenuActions(IRigVMClientHost* RigVMClientHost, const UEdGraphPin* EdGraphPin, URigVMPin* ModelPin, UToolMenu* Menu) const;
+	UE_API virtual void GetPinArrayContextMenuActions(IRigVMClientHost* RigVMClientHost, const UEdGraphPin* EdGraphPin, URigVMPin* ModelPin, UToolMenu* Menu) const;
+	UE_API virtual void GetPinAggregateContextMenuActions(IRigVMClientHost* RigVMClientHost, const UEdGraphPin* EdGraphPin, URigVMPin* ModelPin, UToolMenu* Menu) const;
+	UE_API virtual void GetPinTemplateContextMenuActions(IRigVMClientHost* RigVMClientHost, const UEdGraphPin* EdGraphPin, URigVMPin* ModelPin, UToolMenu* Menu) const;
+	UE_API virtual void GetPinConversionContextMenuActions(IRigVMClientHost* RigVMClientHost, const UEdGraphPin* EdGraphPin, URigVMPin* ModelPin, UToolMenu* Menu) const;
+	UE_API virtual void GetPinVariableContextMenuActions(IRigVMClientHost* RigVMClientHost, const UEdGraphPin* EdGraphPin, URigVMPin* ModelPin, UToolMenu* Menu) const;
+	UE_API virtual void GetPinResetDefaultContextMenuActions(IRigVMClientHost* RigVMClientHost, const UEdGraphPin* EdGraphPin, URigVMPin* ModelPin, UToolMenu* Menu) const;
+	UE_API virtual void GetPinInjectedNodesContextMenuActions(IRigVMClientHost* RigVMClientHost, const UEdGraphPin* EdGraphPin, URigVMPin* ModelPin, UToolMenu* Menu) const;
 
 	/** Extensibility managers */
 	TSharedPtr<FExtensibilityManager> MenuExtensibilityManager;
@@ -150,9 +152,11 @@ private:
 	/** StaticStruct is not safe on shutdown, so we cache the name, and use this to unregister on shut down */
 	TArray<FName> PropertiesToUnregisterOnShutdown;
 
-	void HandleNewBlueprintCreated(UBlueprint* InBlueprint);
+	UE_API void HandleNewBlueprintCreated(UBlueprint* InBlueprint);
 
-	bool ShowWorkflowOptionsDialog(URigVMUserWorkflowOptions* InOptions) const;
+	UE_API bool ShowWorkflowOptionsDialog(URigVMUserWorkflowOptions* InOptions) const;
 
 	friend class URigVMEdGraphNode;
-};
\ No newline at end of file
+};
+
+#undef UE_API
diff --git a/Engine/Plugins/Runtime/RigVM/Source/RigVMEditor/Public/Widgets/SRigVMBulkEditDialog.h b/Engine/Plugins/Runtime/RigVM/Source/RigVMEditor/Public/Widgets/SRigVMBulkEditDialog.h
--- a/Engine/Plugins/Runtime/RigVM/Source/RigVMEditor/Public/Widgets/SRigVMBulkEditDialog.h
+++ b/Engine/Plugins/Runtime/RigVM/Source/RigVMEditor/Public/Widgets/SRigVMBulkEditDialog.h
@@ -8,7 +8,7 @@
 #include "Widgets/SWindow.h"
 
 template<typename InWidgetType = SRigVMBulkEditWidget>
-class RIGVMEDITOR_API SRigVMBulkEditDialog : public SWindow
+class SRigVMBulkEditDialog : public SWindow
 {
 public:
 
diff --git a/Engine/Plugins/Runtime/RigVM/Source/RigVMEditor/Public/Widgets/SRigVMBulkEditWidget.h b/Engine/Plugins/Runtime/RigVM/Source/RigVMEditor/Public/Widgets/SRigVMBulkEditWidget.h
--- a/Engine/Plugins/Runtime/RigVM/Source/RigVMEditor/Public/Widgets/SRigVMBulkEditWidget.h
+++ b/Engine/Plugins/Runtime/RigVM/Source/RigVMEditor/Public/Widgets/SRigVMBulkEditWidget.h
@@ -10,9 +10,11 @@
 #include "Interfaces/IMainFrameModule.h"
 #include "ScopedTransaction.h"
 
+#define UE_API RIGVMEDITOR_API
+
 DECLARE_DELEGATE_OneParam(FOnPhaseActivated, TSharedRef<FRigVMTreePhase>);
 
-class RIGVMEDITOR_API SRigVMBulkEditWidget : public SCompoundWidget
+class SRigVMBulkEditWidget : public SCompoundWidget
 {
 public:
 	SLATE_BEGIN_ARGS(SRigVMBulkEditWidget)
@@ -37,16 +39,16 @@ public:
 	SLATE_ATTRIBUTE(FString, BulkEditConfirmIniField)
 	SLATE_END_ARGS()
 
-	virtual ~SRigVMBulkEditWidget() override;
+	UE_API virtual ~SRigVMBulkEditWidget() override;
 
-	void Construct(const FArguments& InArgs);
+	UE_API void Construct(const FArguments& InArgs);
 
-	virtual void Tick(const FGeometry& AllottedGeometry, const double InCurrentTime, const float InDeltaTime) override;
+	UE_API virtual void Tick(const FGeometry& AllottedGeometry, const double InCurrentTime, const float InDeltaTime) override;
 
-	FText GetDialogTitle() const;
+	UE_API FText GetDialogTitle() const;
 
-	void QueueTasks(const TArray<TSharedRef<FRigVMTreeTask>>& InTasks);
-	void CancelTasks();
+	UE_API void QueueTasks(const TArray<TSharedRef<FRigVMTreeTask>>& InTasks);
+	UE_API void CancelTasks();
 
 	TSharedRef<SRigVMBulkEditWidget> GetBulkEditWidget()
 	{
@@ -63,21 +65,21 @@ public:
 		return BulkEditLogWidget;
 	}
 
-	TSharedPtr<FRigVMTreePhase> GetActivePhasePtr() const;
-	TSharedRef<FRigVMTreePhase> GetActivePhase() const;
-	TSharedPtr<FRigVMTreePhase> FindPhase(int32 InID) const;
-	bool ActivatePhase(int32 InID);
-	TSharedRef<FRigVMTreeContext> GetContext() const;
+	UE_API TSharedPtr<FRigVMTreePhase> GetActivePhasePtr() const;
+	UE_API TSharedRef<FRigVMTreePhase> GetActivePhase() const;
+	UE_API TSharedPtr<FRigVMTreePhase> FindPhase(int32 InID) const;
+	UE_API bool ActivatePhase(int32 InID);
+	UE_API TSharedRef<FRigVMTreeContext> GetContext() const;
 
-	TArray<TSharedRef<FRigVMTreeNode>> GetSelectedNodes() const;
-	bool HasAnyVisibleCheckedNode() const;
-	TArray<TSharedRef<FRigVMTreeNode>> GetCheckedNodes() const;
+	UE_API TArray<TSharedRef<FRigVMTreeNode>> GetSelectedNodes() const;
+	UE_API bool HasAnyVisibleCheckedNode() const;
+	UE_API TArray<TSharedRef<FRigVMTreeNode>> GetCheckedNodes() const;
 	
-	FReply OnBackButtonClicked();
-	FReply OnCancelButtonClicked();
-	FReply OnPrimaryButtonClicked();
+	UE_API FReply OnBackButtonClicked();
+	UE_API FReply OnCancelButtonClicked();
+	UE_API FReply OnPrimaryButtonClicked();
 
-	void CloseDialog();
+	UE_API void CloseDialog();
 
 private:
 
@@ -93,12 +95,12 @@ private:
 	TAttribute<FText> BulkEditConfirmMessage;
 	TAttribute<FString> BulkEditConfirmIniField;
 
-	bool AreTasksInProgress() const;
-	EVisibility GetTasksProgressVisibility() const;
-	TOptional<float> GetTasksProgressPercentage() const;
+	UE_API bool AreTasksInProgress() const;
+	UE_API EVisibility GetTasksProgressVisibility() const;
+	UE_API TOptional<float> GetTasksProgressPercentage() const;
 
-	void OnLogMessage(const TSharedRef<FTokenizedMessage>& InMessage) const;
-	void OnScriptException(ELogVerbosity::Type InVerbosity, const TCHAR* InMessage, const TCHAR* InStackMessage);
+	UE_API void OnLogMessage(const TSharedRef<FTokenizedMessage>& InMessage) const;
+	UE_API void OnScriptException(ELogVerbosity::Type InVerbosity, const TCHAR* InMessage, const TCHAR* InStackMessage);
 
 	mutable FCriticalSection TasksCriticalSection;
 	TArray<TSharedRef<FRigVMTreeTask>> RemainingTasks;
@@ -109,12 +111,14 @@ private:
 	bool bCloseOnSuccess = false;
 	TSharedPtr<FScopedTransaction> Transaction;
 
-	EVisibility GetBackButtonVisibility() const;
-	bool IsBackButtonEnabled() const;
-	EVisibility GetCancelButtonVisibility() const;
-	bool IsCancelButtonEnabled() const;
-	EVisibility GetPrimaryButtonVisibility() const;
-	bool IsPrimaryButtonEnabled() const;
-	FText GetPrimaryButtonText() const;
-	bool IsReadyToClose() const;
+	UE_API EVisibility GetBackButtonVisibility() const;
+	UE_API bool IsBackButtonEnabled() const;
+	UE_API EVisibility GetCancelButtonVisibility() const;
+	UE_API bool IsCancelButtonEnabled() const;
+	UE_API EVisibility GetPrimaryButtonVisibility() const;
+	UE_API bool IsPrimaryButtonEnabled() const;
+	UE_API FText GetPrimaryButtonText() const;
+	UE_API bool IsReadyToClose() const;
 };
+
+#undef UE_API
