Updated headers using UnrealCodeFixup to make sure dllstorage is on methods/staticvars instead of types. Part 4/n

diff --git a/Engine/Plugins/Runtime/RigVM/Source/RigVMEditor/Public/Widgets/SRigVMChangesTreeView.h b/Engine/Plugins/Runtime/RigVM/Source/RigVMEditor/Public/Widgets/SRigVMChangesTreeView.h
--- a/Engine/Plugins/Runtime/RigVM/Source/RigVMEditor/Public/Widgets/SRigVMChangesTreeView.h
+++ b/Engine/Plugins/Runtime/RigVM/Source/RigVMEditor/Public/Widgets/SRigVMChangesTreeView.h
@@ -10,9 +10,11 @@
 #include "Widgets/Views/STreeView.h"
 #include "Widgets/Input/SSearchBox.h"
 
+#define UE_API RIGVMEDITOR_API
+
 class SRigVMChangesTreeView;
 
-class RIGVMEDITOR_API SRigVMChangesTreeRow
+class SRigVMChangesTreeRow
 	: public STableRow<TSharedRef<FRigVMTreeNode>>
 {
 public:
@@ -23,25 +25,25 @@ public:
 	SLATE_ARGUMENT(TSharedPtr<SRigVMChangesTreeView>, OwningWidget)
 	SLATE_END_ARGS()
 
-	virtual ~SRigVMChangesTreeRow() override;
-	void Construct(const FArguments& InArgs, const TSharedRef<STableViewBase>& OwnerTableView);
+	UE_API virtual ~SRigVMChangesTreeRow() override;
+	UE_API void Construct(const FArguments& InArgs, const TSharedRef<STableViewBase>& OwnerTableView);
 
 private:
 
-	EVisibility GetCheckBoxVisibility() const;
-	ECheckBoxState GetCheckBoxState() const;
-	void OnCheckBoxStateChanged(ECheckBoxState InNewState);
-	const FSlateBrush* GetBackgroundImage() const;
-	FSlateColor GetBackgroundColor() const;
-	FOptionalSize GetIndentWidth() const;
-	FSlateColor GetTextColor() const;
-	const FSlateBrush* GetIcon() const;
-	FSlateColor GetIconColor() const;
-	const FSlateBrush* GetExpanderImage() const;
-	FReply OnExpanderMouseButtonDown(const FGeometry& SenderGeometry, const FPointerEvent& MouseEvent);
-	EVisibility GetExpanderVisibility() const;
-	TArray<FRigVMTag> GetVariantTags() const;
-	void RequestRefresh(bool bForce);
+	UE_API EVisibility GetCheckBoxVisibility() const;
+	UE_API ECheckBoxState GetCheckBoxState() const;
+	UE_API void OnCheckBoxStateChanged(ECheckBoxState InNewState);
+	UE_API const FSlateBrush* GetBackgroundImage() const;
+	UE_API FSlateColor GetBackgroundColor() const;
+	UE_API FOptionalSize GetIndentWidth() const;
+	UE_API FSlateColor GetTextColor() const;
+	UE_API const FSlateBrush* GetIcon() const;
+	UE_API FSlateColor GetIconColor() const;
+	UE_API const FSlateBrush* GetExpanderImage() const;
+	UE_API FReply OnExpanderMouseButtonDown(const FGeometry& SenderGeometry, const FPointerEvent& MouseEvent);
+	UE_API EVisibility GetExpanderVisibility() const;
+	UE_API TArray<FRigVMTag> GetVariantTags() const;
+	UE_API void RequestRefresh(bool bForce);
 
 	TSharedPtr<FRigVMTreeNode> Node;
 	SRigVMChangesTreeView* OwningWidget = nullptr;
@@ -51,7 +53,7 @@ private:
 DECLARE_DELEGATE_RetVal_OneParam(FReply,FOnRigVMTreeNodeSelected,TSharedRef<FRigVMTreeNode>);
 DECLARE_DELEGATE_RetVal_OneParam(FReply,FOnRigVMTreeNodeDoubleClicked,TSharedRef<FRigVMTreeNode>);
 
-class RIGVMEDITOR_API SRigVMChangesTreeView : public SCompoundWidget
+class SRigVMChangesTreeView : public SCompoundWidget
 {
 public:
 	SLATE_BEGIN_ARGS(SRigVMChangesTreeView)
@@ -61,42 +63,42 @@ public:
 	SLATE_EVENT(FOnRigVMTreeNodeDoubleClicked, OnNodeDoubleClicked)
 	SLATE_END_ARGS()
 
-	virtual ~SRigVMChangesTreeView() override;
+	UE_API virtual ~SRigVMChangesTreeView() override;
 
-	void Construct(const FArguments& InArgs);
+	UE_API void Construct(const FArguments& InArgs);
 
-	virtual void Tick(const FGeometry& AllottedGeometry, const double InCurrentTime, const float InDeltaTime) override;
+	UE_API virtual void Tick(const FGeometry& AllottedGeometry, const double InCurrentTime, const float InDeltaTime) override;
 
-	EVisibility GetPathFilterVisibility() const;
-	void OnPathFilterTextChanged(const FText& SearchText);
-	void OnPathFilterTextCommitted(const FText& SearchText, ETextCommit::Type CommitInfo);
+	UE_API EVisibility GetPathFilterVisibility() const;
+	UE_API void OnPathFilterTextChanged(const FText& SearchText);
+	UE_API void OnPathFilterTextCommitted(const FText& SearchText, ETextCommit::Type CommitInfo);
 
-	void RequestRefresh_AnyThread(bool bForce = false);
-	void RefreshFilteredNodes(bool bForce = false);
-	void RefreshFilteredNodesIfRequired();
+	UE_API void RequestRefresh_AnyThread(bool bForce = false);
+	UE_API void RefreshFilteredNodes(bool bForce = false);
+	UE_API void RefreshFilteredNodesIfRequired();
 
-	void OnPhaseChanged();
-	void SetSelection(const TSharedPtr<FRigVMTreeNode>& InNode, bool bRequestScrollIntoView = false);
+	UE_API void OnPhaseChanged();
+	UE_API void SetSelection(const TSharedPtr<FRigVMTreeNode>& InNode, bool bRequestScrollIntoView = false);
 
 	TSharedPtr<STreeView<TSharedRef<FRigVMTreeNode>>> GetTreeView()
 	{
 		return TreeView;
 	}
 
-	TSharedRef<FRigVMTreePhase> GetPhase() const;
-	TSharedRef<FRigVMTreeContext> GetContext() const;
-	TSharedPtr<FRigVMTreePathFilter> GetPathFilter() const;
+	UE_API TSharedRef<FRigVMTreePhase> GetPhase() const;
+	UE_API TSharedRef<FRigVMTreeContext> GetContext() const;
+	UE_API TSharedPtr<FRigVMTreePathFilter> GetPathFilter() const;
 
-	TArray<TSharedRef<FRigVMTreeNode>> GetSelectedNodes() const;
-	bool HasAnyVisibleCheckedNode() const;
-	TArray<TSharedRef<FRigVMTreeNode>> GetCheckedNodes() const;
+	UE_API TArray<TSharedRef<FRigVMTreeNode>> GetSelectedNodes() const;
+	UE_API bool HasAnyVisibleCheckedNode() const;
+	UE_API TArray<TSharedRef<FRigVMTreeNode>> GetCheckedNodes() const;
 	
 private:
 
 	TAttribute<TSharedPtr<FRigVMTreePhase>> PhaseAttribute;
 	TArray<TSharedRef<FRigVMTreeNode>> FilteredNodes;
 	
-	FText GetPathFilterText();
+	UE_API FText GetPathFilterText();
 	TSharedPtr<SSearchBox> PathFilterBox;
 
 	TSharedPtr<STreeView<TSharedRef<FRigVMTreeNode>>> TreeView;
@@ -108,18 +110,20 @@ private:
 	TAttribute<FText> BulkEditConfirmMessage;
 	TAttribute<FString> BulkEditConfirmIniField;
 
-	TSharedRef<ITableRow> MakeTreeRowWidget(TSharedRef<FRigVMTreeNode> InNode, const TSharedRef<STableViewBase>& OwnerTable);
+	UE_API TSharedRef<ITableRow> MakeTreeRowWidget(TSharedRef<FRigVMTreeNode> InNode, const TSharedRef<STableViewBase>& OwnerTable);
 
-	void GetChildrenForNode(TSharedRef<FRigVMTreeNode> InNode, TArray<TSharedRef<FRigVMTreeNode>>& OutChildren);
-	void OnSelectionChanged(TSharedPtr<FRigVMTreeNode> Selection, ESelectInfo::Type SelectInfo);
-	void OnTreeElementDoubleClicked(TSharedRef<FRigVMTreeNode> InNode);
-	ESelectionMode::Type GetSelectionMode() const;
+	UE_API void GetChildrenForNode(TSharedRef<FRigVMTreeNode> InNode, TArray<TSharedRef<FRigVMTreeNode>>& OutChildren);
+	UE_API void OnSelectionChanged(TSharedPtr<FRigVMTreeNode> Selection, ESelectInfo::Type SelectInfo);
+	UE_API void OnTreeElementDoubleClicked(TSharedRef<FRigVMTreeNode> InNode);
+	UE_API ESelectionMode::Type GetSelectionMode() const;
 
-	TSharedPtr<SWidget> OnGetNodeContextMenuContent();
+	UE_API TSharedPtr<SWidget> OnGetNodeContextMenuContent();
 
-	EVisibility GetSettingsButtonVisibility() const;
-	FReply OnSettingsButtonClicked();
+	UE_API EVisibility GetSettingsButtonVisibility() const;
+	UE_API FReply OnSettingsButtonClicked();
 
 	std::atomic<int32> RequestRefreshCount;
 	std::atomic<int32> RequestRefreshForceCount;
 };
+
+#undef UE_API
diff --git a/Engine/Plugins/Runtime/RigVM/Source/RigVMEditor/Public/Widgets/SRigVMEditorSelectedDebugObjectWidget.h b/Engine/Plugins/Runtime/RigVM/Source/RigVMEditor/Public/Widgets/SRigVMEditorSelectedDebugObjectWidget.h
--- a/Engine/Plugins/Runtime/RigVM/Source/RigVMEditor/Public/Widgets/SRigVMEditorSelectedDebugObjectWidget.h
+++ b/Engine/Plugins/Runtime/RigVM/Source/RigVMEditor/Public/Widgets/SRigVMEditorSelectedDebugObjectWidget.h
@@ -5,6 +5,8 @@
 #include "Editor/RigVMEditor.h"
 #include "Widgets/Input/SComboBox.h"
 
+#define UE_API RIGVMEDITOR_API
+
 struct FRigVMDebugObjectInstance
 {
 	/** Actual object to debug, can be null */
@@ -44,68 +46,68 @@ struct FRigVMDebugObjectInstance
 	}
 };
 
-class RIGVMEDITOR_API SRigVMEditorSelectedDebugObjectWidget : public SCompoundWidget
+class SRigVMEditorSelectedDebugObjectWidget : public SCompoundWidget
 {
 public:
 	SLATE_BEGIN_ARGS(SRigVMEditorSelectedDebugObjectWidget){}
 	SLATE_END_ARGS()
 
-	void Construct(const FArguments& InArgs, TSharedPtr<IRigVMEditor> InEditor);
+	UE_API void Construct(const FArguments& InArgs, TSharedPtr<IRigVMEditor> InEditor);
 
 	// SWidget interface
-	virtual void Tick( const FGeometry& AllottedGeometry, const double InCurrentTime, const float InDeltaTime ) override;
+	UE_API virtual void Tick( const FGeometry& AllottedGeometry, const double InCurrentTime, const float InDeltaTime ) override;
 	// End of SWidget interface
 
 	/** Adds an object to the list of debug choices */
-	void AddDebugObject(UObject* TestObject, const FString& TestObjectName = FString());
+	UE_API void AddDebugObject(UObject* TestObject, const FString& TestObjectName = FString());
 
 private:
-	UBlueprint* GetBlueprintObj() const;
+	UE_API UBlueprint* GetBlueprintObj() const;
 
 	/** Creates a list of all debug objects **/
-	void GenerateDebugObjectInstances(bool bRestoreSelection);
+	UE_API void GenerateDebugObjectInstances(bool bRestoreSelection);
 
 	/** Generate list of active PIE worlds to debug **/
-	void GenerateDebugWorldNames(bool bRestoreSelection);
+	UE_API void GenerateDebugWorldNames(bool bRestoreSelection);
 
 	/** Refresh the widget. **/
-	void OnRefresh();
+	UE_API void OnRefresh();
 
 	/** Returns the entry for the current debug actor */
-	TSharedPtr<FRigVMDebugObjectInstance> GetDebugObjectInstance() const;
+	UE_API TSharedPtr<FRigVMDebugObjectInstance> GetDebugObjectInstance() const;
 
 	/** Returns the name of the current debug actor */
-	TSharedPtr<FString> GetDebugWorldName() const;
+	UE_API TSharedPtr<FString> GetDebugWorldName() const;
 
 	/** Handles the selection changed event for the debug actor combo box */
-	void DebugObjectSelectionChanged(TSharedPtr<FRigVMDebugObjectInstance> NewSelection, ESelectInfo::Type SelectInfo);
+	UE_API void DebugObjectSelectionChanged(TSharedPtr<FRigVMDebugObjectInstance> NewSelection, ESelectInfo::Type SelectInfo);
 
-	void DebugWorldSelectionChanged(TSharedPtr<FString> NewSelection, ESelectInfo::Type SelectInfo);
+	UE_API void DebugWorldSelectionChanged(TSharedPtr<FString> NewSelection, ESelectInfo::Type SelectInfo);
 
 	/** Called when user clicks button to select the current object being debugged */
-	void SelectedDebugObject_OnClicked();
+	UE_API void SelectedDebugObject_OnClicked();
 
 	/** Returns true if a debug actor is currently selected */
-	bool IsDebugObjectSelected() const;
+	UE_API bool IsDebugObjectSelected() const;
 
-	EVisibility IsDebugWorldComboVisible() const;
+	UE_API EVisibility IsDebugWorldComboVisible() const;
 
 	/* Returns the string to indicate no debug object is selected */
-	const FString& GetNoDebugString() const;
+	UE_API const FString& GetNoDebugString() const;
 
-	const FString& GetDebugAllWorldsString() const;
+	UE_API const FString& GetDebugAllWorldsString() const;
 
 	/** Helper method to construct a debug object label string */
-	FString MakeDebugObjectLabel(UObject* TestObject, bool bAddContextIfSelectedInEditor, bool bAddSpawnedContext) const;
+	UE_API FString MakeDebugObjectLabel(UObject* TestObject, bool bAddContextIfSelectedInEditor, bool bAddSpawnedContext) const;
 
 	/** Fills in data for a specific instance */
-	void FillDebugObjectInstance(TSharedPtr<FRigVMDebugObjectInstance> Instance);
+	UE_API void FillDebugObjectInstance(TSharedPtr<FRigVMDebugObjectInstance> Instance);
 
 	/** Called to create a widget for each debug object item */
-	TSharedRef<SWidget> CreateDebugObjectItemWidget(TSharedPtr<FRigVMDebugObjectInstance> InItem);
+	UE_API TSharedRef<SWidget> CreateDebugObjectItemWidget(TSharedPtr<FRigVMDebugObjectInstance> InItem);
 
 	/** Returns the combo button label to use for the currently-selected debug object item */
-	FText GetSelectedDebugObjectTextLabel() const;
+	UE_API FText GetSelectedDebugObjectTextLabel() const;
 
 private:
 	/** Pointer back to the rigvm editor tool that owns us */
@@ -125,3 +127,5 @@ private:
 
 	TWeakObjectPtr<UObject> LastObjectObserved;
 };
+
+#undef UE_API
diff --git a/Engine/Plugins/Runtime/RigVM/Source/RigVMEditor/Public/Widgets/SRigVMExecutionStackView.h b/Engine/Plugins/Runtime/RigVM/Source/RigVMEditor/Public/Widgets/SRigVMExecutionStackView.h
--- a/Engine/Plugins/Runtime/RigVM/Source/RigVMEditor/Public/Widgets/SRigVMExecutionStackView.h
+++ b/Engine/Plugins/Runtime/RigVM/Source/RigVMEditor/Public/Widgets/SRigVMExecutionStackView.h
@@ -10,6 +10,8 @@
 #include "RigVMCore/RigVM.h"
 #include "RigVMBlueprint.h"
 
+#define UE_API RIGVMEDITOR_API
+
 class IRigVMEditor;
 class SRigVMExecutionStackView;
 class FUICommandList;
@@ -27,12 +29,12 @@ namespace ERigStackEntry
 }
 
 /** An item in the stack */
-class RIGVMEDITOR_API FRigStackEntry : public TSharedFromThis<FRigStackEntry>
+class FRigStackEntry : public TSharedFromThis<FRigStackEntry>
 {
 public:
-	FRigStackEntry(int32 InEntryIndex, ERigStackEntry::Type InEntryType, int32 InInstructionIndex, ERigVMOpCode InOpCode, const FString& InLabel, const FRigVMASTProxy& InProxy);
+	UE_API FRigStackEntry(int32 InEntryIndex, ERigStackEntry::Type InEntryType, int32 InInstructionIndex, ERigVMOpCode InOpCode, const FString& InLabel, const FRigVMASTProxy& InProxy);
 
-	TSharedRef<ITableRow> MakeTreeRowWidget(const TSharedRef<STableViewBase>& InOwnerTable, TSharedRef<FRigStackEntry> InEntry, TSharedRef<FUICommandList> InCommandList, TSharedPtr<SRigVMExecutionStackView> InStackView, TWeakObjectPtr<URigVMBlueprint> InBlueprint);
+	UE_API TSharedRef<ITableRow> MakeTreeRowWidget(const TSharedRef<STableViewBase>& InOwnerTable, TSharedRef<FRigStackEntry> InEntry, TSharedRef<FUICommandList> InCommandList, TSharedPtr<SRigVMExecutionStackView> InStackView, TWeakObjectPtr<URigVMBlueprint> InBlueprint);
 
 	int32 EntryIndex;
 	ERigStackEntry::Type EntryType;
@@ -67,70 +69,70 @@ private:
 	FText GetDurationText() const;
 };
 
-class RIGVMEDITOR_API SRigVMExecutionStackView : public SCompoundWidget
+class SRigVMExecutionStackView : public SCompoundWidget
 {
 public:
 	SLATE_BEGIN_ARGS(SRigVMExecutionStackView) {}
 	SLATE_END_ARGS()
 
-	~SRigVMExecutionStackView();
+	UE_API ~SRigVMExecutionStackView();
 
-	void Construct(const FArguments& InArgs, TSharedRef<IRigVMEditor> InRigVMEditor);
+	UE_API void Construct(const FArguments& InArgs, TSharedRef<IRigVMEditor> InRigVMEditor);
 
 	/** Set Selection Changed */
-	void OnSelectionChanged(TSharedPtr<FRigStackEntry> Selection, ESelectInfo::Type SelectInfo);
+	UE_API void OnSelectionChanged(TSharedPtr<FRigStackEntry> Selection, ESelectInfo::Type SelectInfo);
 
-	TSharedPtr< SWidget > CreateContextMenu();
+	UE_API TSharedPtr< SWidget > CreateContextMenu();
 
 protected:
 
 	/** Rebuild the tree view */
-	void RefreshTreeView(URigVM* InVM, FRigVMExtendedExecuteContext* InVMContext);
+	UE_API void RefreshTreeView(URigVM* InVM, FRigVMExtendedExecuteContext* InVMContext);
 
 private:
 
 	/** Bind commands that this widget handles */
-	void BindCommands();
+	UE_API void BindCommands();
 
 	/** Make a row widget for the table */
-	TSharedRef<ITableRow> MakeTableRowWidget(TSharedPtr<FRigStackEntry> InItem, const TSharedRef<STableViewBase>& OwnerTable, TWeakObjectPtr<URigVMBlueprint> InBlueprint);
+	UE_API TSharedRef<ITableRow> MakeTableRowWidget(TSharedPtr<FRigStackEntry> InItem, const TSharedRef<STableViewBase>& OwnerTable, TWeakObjectPtr<URigVMBlueprint> InBlueprint);
 
 	/** Get children for the tree */
-	void HandleGetChildrenForTree(TSharedPtr<FRigStackEntry> InItem, TArray<TSharedPtr<FRigStackEntry>>& OutChildren);
+	UE_API void HandleGetChildrenForTree(TSharedPtr<FRigStackEntry> InItem, TArray<TSharedPtr<FRigStackEntry>>& OutChildren);
 
 	/** Focus on the selected operator in the graph*/
-	void HandleFocusOnSelectedGraphNode();
+	UE_API void HandleFocusOnSelectedGraphNode();
 
 	/** Offers a dialog to move to a specific instruction */
-	void HandleGoToInstruction();
+	UE_API void HandleGoToInstruction();
 
 	/** Selects the target instructions for the current selection */
-	void HandleSelectTargetInstructions();
+	UE_API void HandleSelectTargetInstructions();
 
-	TArray<TSharedPtr<FRigStackEntry>> GetTargetItems(const TArray<TSharedPtr<FRigStackEntry>>& InItems) const;
+	UE_API TArray<TSharedPtr<FRigStackEntry>> GetTargetItems(const TArray<TSharedPtr<FRigStackEntry>>& InItems) const;
 
-	void UpdateTargetItemHighlighting();
+	UE_API void UpdateTargetItemHighlighting();
 
-	void OnVMCompiled(UObject* InCompiledObject, URigVM* InCompiledVM, FRigVMExtendedExecuteContext& InVMContext);
+	UE_API void OnVMCompiled(UObject* InCompiledObject, URigVM* InCompiledVM, FRigVMExtendedExecuteContext& InVMContext);
 
 	//* Focus on the instruction when the execution is halted */
-	void HandleExecutionHalted(const int32 HaltedAtInstruction, UObject* InNode, const FName& InEntryName);
+	UE_API void HandleExecutionHalted(const int32 HaltedAtInstruction, UObject* InNode, const FName& InEntryName);
 
 	/** Search box widget */
 	TSharedPtr<SSearchBox> FilterBox;
 	FText FilterText;
-	void OnFilterTextChanged(const FText& SearchText);
+	UE_API void OnFilterTextChanged(const FText& SearchText);
 
 
 	bool bSuspendModelNotifications;
 	bool bSuspendControllerSelection;
-	void HandleModifiedEvent(ERigVMGraphNotifType InNotifType, URigVMGraph* InGraph, UObject* InSubject);
-	void HandleHostInitializedEvent(URigVMHost* InHost, const FName& InEventName);
-	void HandlePreviewHostUpdated(IRigVMEditor* InEditor);
-	void HandleItemMouseDoubleClick(TSharedPtr<FRigStackEntry> InItem);
+	UE_API void HandleModifiedEvent(ERigVMGraphNotifType InNotifType, URigVMGraph* InGraph, UObject* InSubject);
+	UE_API void HandleHostInitializedEvent(URigVMHost* InHost, const FName& InEventName);
+	UE_API void HandlePreviewHostUpdated(IRigVMEditor* InEditor);
+	UE_API void HandleItemMouseDoubleClick(TSharedPtr<FRigStackEntry> InItem);
 
 	/** Populate the execution stack with descriptive names for each instruction */
-	void PopulateStackView(URigVM* InVM, FRigVMExtendedExecuteContext* InVMContext);
+	UE_API void PopulateStackView(URigVM* InVM, FRigVMExtendedExecuteContext* InVMContext);
 
 	TSharedPtr<STreeView<TSharedPtr<FRigStackEntry>>> TreeView;
 
@@ -149,3 +151,5 @@ private:
 	FDelegateHandle OnVMCompiledHandle;
 	FDelegateHandle OnPreviewHostUpdatedHandle;
 };
+
+#undef UE_API
diff --git a/Engine/Plugins/Runtime/RigVM/Source/RigVMEditor/Public/Widgets/SRigVMGraphBreakLinksWidget.h b/Engine/Plugins/Runtime/RigVM/Source/RigVMEditor/Public/Widgets/SRigVMGraphBreakLinksWidget.h
--- a/Engine/Plugins/Runtime/RigVM/Source/RigVMEditor/Public/Widgets/SRigVMGraphBreakLinksWidget.h
+++ b/Engine/Plugins/Runtime/RigVM/Source/RigVMEditor/Public/Widgets/SRigVMGraphBreakLinksWidget.h
@@ -7,36 +7,38 @@
 #include "Dialogs/Dialogs.h"
 #include "IAssetTypeActions.h"
 
+#define UE_API RIGVMEDITOR_API
+
 class ITableRow;
 class STableViewBase;
 
 DECLARE_DELEGATE_OneParam(FRigVMOnFocusOnLinkRequestedDelegate, URigVMLink*);
 
-class RIGVMEDITOR_API SRigVMGraphBreakLinksWidget : public SCompoundWidget
+class SRigVMGraphBreakLinksWidget : public SCompoundWidget
 {
 public:
 	SLATE_BEGIN_ARGS(SRigVMGraphBreakLinksWidget) {}
 	SLATE_ARGUMENT(FRigVMOnFocusOnLinkRequestedDelegate, OnFocusOnLink)
 	SLATE_END_ARGS()
 
-	void Construct(const FArguments& InArgs, TArray<URigVMLink*> InLinks);
+	UE_API void Construct(const FArguments& InArgs, TArray<URigVMLink*> InLinks);
 
 private:
 
 	/** Handler for when the user double clicks, presses enter, or presses space on an asset */
-	void OnLinkDoubleClicked(const URigVMLink* InLink, EAssetTypeActivationMethod::Type ActivationMethod);
+	UE_API void OnLinkDoubleClicked(const URigVMLink* InLink, EAssetTypeActivationMethod::Type ActivationMethod);
 
-	TSharedRef<ITableRow> GenerateItemRow(URigVMLink* Item, const TSharedRef<STableViewBase>& OwnerTable);
+	UE_API TSharedRef<ITableRow> GenerateItemRow(URigVMLink* Item, const TSharedRef<STableViewBase>& OwnerTable);
  
 	TArray<URigVMLink*> Links;
 	FRigVMOnFocusOnLinkRequestedDelegate OnFocusOnLink;
 
-	void HandleItemMouseDoubleClick(URigVMLink* InItem);
+	UE_API void HandleItemMouseDoubleClick(URigVMLink* InItem);
 
 	friend class SRigVMGraphBreakLinksDialog;
 };
 
-class RIGVMEDITOR_API SRigVMGraphBreakLinksDialog : public SWindow
+class SRigVMGraphBreakLinksDialog : public SWindow
 {
 public:
 	
@@ -48,15 +50,17 @@ public:
 	SLATE_ARGUMENT(FRigVMOnFocusOnLinkRequestedDelegate, OnFocusOnLink)
 	SLATE_END_ARGS()
 
-	void Construct(const FArguments& InArgs);
+	UE_API void Construct(const FArguments& InArgs);
 
-	EAppReturnType::Type ShowModal();
+	UE_API EAppReturnType::Type ShowModal();
 
 protected:
 
 	TSharedPtr<SRigVMGraphBreakLinksWidget> BreakLinksWidget;
 	
-	FReply OnButtonClick(EAppReturnType::Type ButtonID);
-	bool IsOkButtonEnabled() const;
+	UE_API FReply OnButtonClick(EAppReturnType::Type ButtonID);
+	UE_API bool IsOkButtonEnabled() const;
 	EAppReturnType::Type UserResponse;
-};
\ No newline at end of file
+};
+
+#undef UE_API
diff --git a/Engine/Plugins/Runtime/RigVM/Source/RigVMEditor/Public/Widgets/SRigVMGraphChangePinType.h b/Engine/Plugins/Runtime/RigVM/Source/RigVMEditor/Public/Widgets/SRigVMGraphChangePinType.h
--- a/Engine/Plugins/Runtime/RigVM/Source/RigVMEditor/Public/Widgets/SRigVMGraphChangePinType.h
+++ b/Engine/Plugins/Runtime/RigVM/Source/RigVMEditor/Public/Widgets/SRigVMGraphChangePinType.h
@@ -10,9 +10,11 @@
 #include "IPropertyAccessEditor.h"
 #include "RigVMBlueprint.h"
 
+#define UE_API RIGVMEDITOR_API
+
 DECLARE_DELEGATE_OneParam(FOnTypeSelected, TRigVMTypeIndex);
 
-class RIGVMEDITOR_API SRigVMGraphChangePinType : public SCompoundWidget
+class SRigVMGraphChangePinType : public SCompoundWidget
 {
 public:
 
@@ -26,22 +28,24 @@ public:
 
 	SLATE_END_ARGS()
 
-	void Construct(const FArguments& InArgs);
+	UE_API void Construct(const FArguments& InArgs);
 
 protected:
 
-	static FText GetBindingText(const FRigVMTemplateArgumentType& InType);
-	FText GetBindingText(URigVMPin* ModelPin) const;
-	FText GetBindingText() const;
-	const FSlateBrush* GetBindingImage() const;
-	FLinearColor GetBindingColor() const;
-	bool OnCanBindProperty(FProperty* InProperty) const;
-	bool OnCanBindToClass(UClass* InClass) const;
-	void OnAddBinding(FName InPropertyName, const TArray<FBindingChainElement>& InBindingChain);
-	void FillPinTypeMenu( FMenuBuilder& MenuBuilder );
-	void HandlePinTypeChanged(FRigVMTemplateArgumentType InType);
+	static UE_API FText GetBindingText(const FRigVMTemplateArgumentType& InType);
+	UE_API FText GetBindingText(URigVMPin* ModelPin) const;
+	UE_API FText GetBindingText() const;
+	UE_API const FSlateBrush* GetBindingImage() const;
+	UE_API FLinearColor GetBindingColor() const;
+	UE_API bool OnCanBindProperty(FProperty* InProperty) const;
+	UE_API bool OnCanBindToClass(UClass* InClass) const;
+	UE_API void OnAddBinding(FName InPropertyName, const TArray<FBindingChainElement>& InBindingChain);
+	UE_API void FillPinTypeMenu( FMenuBuilder& MenuBuilder );
+	UE_API void HandlePinTypeChanged(FRigVMTemplateArgumentType InType);
 
 	TArray<TRigVMTypeIndex> Types;
 	FOnTypeSelected OnTypeSelected;
 	FPropertyBindingWidgetArgs BindingArgs;
 };
+
+#undef UE_API
diff --git a/Engine/Plugins/Runtime/RigVM/Source/RigVMEditor/Public/Widgets/SRigVMGraphFunctionBulkEditWidget.h b/Engine/Plugins/Runtime/RigVM/Source/RigVMEditor/Public/Widgets/SRigVMGraphFunctionBulkEditWidget.h
--- a/Engine/Plugins/Runtime/RigVM/Source/RigVMEditor/Public/Widgets/SRigVMGraphFunctionBulkEditWidget.h
+++ b/Engine/Plugins/Runtime/RigVM/Source/RigVMEditor/Public/Widgets/SRigVMGraphFunctionBulkEditWidget.h
@@ -8,38 +8,40 @@
 #include "Dialogs/Dialogs.h"
 #include "IAssetTypeActions.h"
 
+#define UE_API RIGVMEDITOR_API
+
 class SRigVMGraphFunctionBulkEditWidget;
 
-class RIGVMEDITOR_API SRigVMGraphFunctionBulkEditWidget : public SCompoundWidget
+class SRigVMGraphFunctionBulkEditWidget : public SCompoundWidget
 {
 public:
 	SLATE_BEGIN_ARGS(SRigVMGraphFunctionBulkEditWidget) {}
 	SLATE_END_ARGS()
 
-	void Construct(const FArguments& InArgs, URigVMBlueprint* InBlueprint, URigVMController* InController, URigVMLibraryNode* InFunction, ERigVMControllerBulkEditType InEditType);
+	UE_API void Construct(const FArguments& InArgs, URigVMBlueprint* InBlueprint, URigVMController* InController, URigVMLibraryNode* InFunction, ERigVMControllerBulkEditType InEditType);
 
 private:
 
 	/** Handler for when an asset context menu has been requested. */
-	TSharedPtr<SWidget> OnGetAssetContextMenu( const TArray<FAssetData>& SelectedAssets );
+	UE_API TSharedPtr<SWidget> OnGetAssetContextMenu( const TArray<FAssetData>& SelectedAssets );
 
 	/** Handler for when the user double clicks, presses enter, or presses space on an asset */
-	void OnAssetsActivated(const TArray<FAssetData>& ActivatedAssets, EAssetTypeActivationMethod::Type ActivationMethod);
+	UE_API void OnAssetsActivated(const TArray<FAssetData>& ActivatedAssets, EAssetTypeActivationMethod::Type ActivationMethod);
 
 	/** Show a progress bar and load all of the assets */
-	void LoadAffectedAssets();
+	UE_API void LoadAffectedAssets();
 
 	URigVMBlueprint* Blueprint;
 	URigVMController* Controller;
 	URigVMLibraryNode* Function;
 	ERigVMControllerBulkEditType EditType;
 
-	TSharedRef<SWidget> MakeAssetViewForReferencedAssets();
+	UE_API TSharedRef<SWidget> MakeAssetViewForReferencedAssets();
 
 	friend class SRigVMGraphFunctionBulkEditDialog;
 };
 
-class RIGVMEDITOR_API SRigVMGraphFunctionBulkEditDialog : public SWindow
+class SRigVMGraphFunctionBulkEditDialog : public SWindow
 {
 public:
 	
@@ -53,15 +55,17 @@ public:
 	SLATE_ARGUMENT(ERigVMControllerBulkEditType, EditType)
 	SLATE_END_ARGS()
 
-	void Construct(const FArguments& InArgs);
+	UE_API void Construct(const FArguments& InArgs);
 
-	EAppReturnType::Type ShowModal();
+	UE_API EAppReturnType::Type ShowModal();
 
 protected:
 
 	TSharedPtr<SRigVMGraphFunctionBulkEditWidget> BulkEditWidget;
 	
-	FReply OnButtonClick(EAppReturnType::Type ButtonID);
-	bool IsOkButtonEnabled() const;
+	UE_API FReply OnButtonClick(EAppReturnType::Type ButtonID);
+	UE_API bool IsOkButtonEnabled() const;
 	EAppReturnType::Type UserResponse;
-};
\ No newline at end of file
+};
+
+#undef UE_API
diff --git a/Engine/Plugins/Runtime/RigVM/Source/RigVMEditor/Public/Widgets/SRigVMGraphFunctionLocalizationWidget.h b/Engine/Plugins/Runtime/RigVM/Source/RigVMEditor/Public/Widgets/SRigVMGraphFunctionLocalizationWidget.h
--- a/Engine/Plugins/Runtime/RigVM/Source/RigVMEditor/Public/Widgets/SRigVMGraphFunctionLocalizationWidget.h
+++ b/Engine/Plugins/Runtime/RigVM/Source/RigVMEditor/Public/Widgets/SRigVMGraphFunctionLocalizationWidget.h
@@ -8,38 +8,40 @@
 #include "Widgets/Views/SListView.h"
 #include "Dialogs/Dialogs.h"
 
+#define UE_API RIGVMEDITOR_API
+
 class SRigVMGraphFunctionLocalizationWidget;
 
-class RIGVMEDITOR_API SRigVMGraphFunctionLocalizationItem : public TSharedFromThis<SRigVMGraphFunctionLocalizationItem>
+class SRigVMGraphFunctionLocalizationItem : public TSharedFromThis<SRigVMGraphFunctionLocalizationItem>
 {
 public:
-	SRigVMGraphFunctionLocalizationItem(const FRigVMGraphFunctionIdentifier& InFunction);
+	UE_API SRigVMGraphFunctionLocalizationItem(const FRigVMGraphFunctionIdentifier& InFunction);
 
 	FText DisplayText;
 	FText ToolTipText;
 	const FRigVMGraphFunctionIdentifier Function;
 };
 
-class RIGVMEDITOR_API SRigVMGraphFunctionLocalizationTableRow : public STableRow<TSharedPtr<SRigVMGraphFunctionLocalizationItem>>
+class SRigVMGraphFunctionLocalizationTableRow : public STableRow<TSharedPtr<SRigVMGraphFunctionLocalizationItem>>
 {
 	SLATE_BEGIN_ARGS(SRigVMGraphFunctionLocalizationTableRow) {}
 	SLATE_END_ARGS()
 
-	void Construct(const FArguments& InArgs, const TSharedRef<STableViewBase>& OwnerTable, SRigVMGraphFunctionLocalizationWidget* InLocalizationWidget, TSharedRef<SRigVMGraphFunctionLocalizationItem> InFunctionItem);
+	UE_API void Construct(const FArguments& InArgs, const TSharedRef<STableViewBase>& OwnerTable, SRigVMGraphFunctionLocalizationWidget* InLocalizationWidget, TSharedRef<SRigVMGraphFunctionLocalizationItem> InFunctionItem);
 };
 
-class RIGVMEDITOR_API SRigVMGraphFunctionLocalizationWidget : public SCompoundWidget
+class SRigVMGraphFunctionLocalizationWidget : public SCompoundWidget
 {
 public:
 	SLATE_BEGIN_ARGS(SRigVMGraphFunctionLocalizationWidget) {}
 	SLATE_END_ARGS()
 
-	void Construct(const FArguments& InArgs, const FRigVMGraphFunctionIdentifier& InFunctionToLocalize, IRigVMGraphFunctionHost* InTargetFunctionHost);
+	UE_API void Construct(const FArguments& InArgs, const FRigVMGraphFunctionIdentifier& InFunctionToLocalize, IRigVMGraphFunctionHost* InTargetFunctionHost);
 
-	TSharedRef<ITableRow> GenerateFunctionListRow(TSharedPtr<SRigVMGraphFunctionLocalizationItem> InItem, const TSharedRef<STableViewBase>& InOwningTable);
-	ECheckBoxState IsFunctionEnabled(const FRigVMGraphFunctionIdentifier InFunction) const;
-	void SetFunctionEnabled(ECheckBoxState NewState, const FRigVMGraphFunctionIdentifier InFunction);
-	bool IsFunctionPublic(const FRigVMGraphFunctionIdentifier InFunction) const;
+	UE_API TSharedRef<ITableRow> GenerateFunctionListRow(TSharedPtr<SRigVMGraphFunctionLocalizationItem> InItem, const TSharedRef<STableViewBase>& InOwningTable);
+	UE_API ECheckBoxState IsFunctionEnabled(const FRigVMGraphFunctionIdentifier InFunction) const;
+	UE_API void SetFunctionEnabled(ECheckBoxState NewState, const FRigVMGraphFunctionIdentifier InFunction);
+	UE_API bool IsFunctionPublic(const FRigVMGraphFunctionIdentifier InFunction) const;
 
 private:
 
@@ -51,7 +53,7 @@ private:
 	friend class SRigVMGraphFunctionLocalizationDialog;
 };
 
-class RIGVMEDITOR_API SRigVMGraphFunctionLocalizationDialog : public SWindow
+class SRigVMGraphFunctionLocalizationDialog : public SWindow
 {
 public:
 	
@@ -63,16 +65,18 @@ public:
 	SLATE_ARGUMENT(IRigVMGraphFunctionHost*, GraphFunctionHost)
 	SLATE_END_ARGS()
 
-	void Construct(const FArguments& InArgs);
+	UE_API void Construct(const FArguments& InArgs);
 
-	EAppReturnType::Type ShowModal();
-	TArray<FRigVMGraphFunctionIdentifier>& GetFunctionsToLocalize();
+	UE_API EAppReturnType::Type ShowModal();
+	UE_API TArray<FRigVMGraphFunctionIdentifier>& GetFunctionsToLocalize();
 
 protected:
 
 	TSharedPtr<SRigVMGraphFunctionLocalizationWidget> FunctionsWidget;
 	
-	FReply OnButtonClick(EAppReturnType::Type ButtonID);
-	bool IsOkButtonEnabled() const;
+	UE_API FReply OnButtonClick(EAppReturnType::Type ButtonID);
+	UE_API bool IsOkButtonEnabled() const;
 	EAppReturnType::Type UserResponse;
-};
\ No newline at end of file
+};
+
+#undef UE_API
diff --git a/Engine/Plugins/Runtime/RigVM/Source/RigVMEditor/Public/Widgets/SRigVMGraphNode.h b/Engine/Plugins/Runtime/RigVM/Source/RigVMEditor/Public/Widgets/SRigVMGraphNode.h
--- a/Engine/Plugins/Runtime/RigVM/Source/RigVMEditor/Public/Widgets/SRigVMGraphNode.h
+++ b/Engine/Plugins/Runtime/RigVM/Source/RigVMEditor/Public/Widgets/SRigVMGraphNode.h
@@ -11,13 +11,15 @@
 #include "RigVMBlueprint.h"
 #include "Overrides/SOverrideStatusWidget.h"
 
+#define UE_API RIGVMEDITOR_API
+
 class URigVMEdGraphNode;
 class STableViewBase;
 class SOverlay;
 class SGraphPin;
 class UEdGraphPin;
 
-class RIGVMEDITOR_API SRigVMGraphNode : public SGraphNode
+class SRigVMGraphNode : public SGraphNode
 {
 public:
 	SLATE_BEGIN_ARGS(SRigVMGraphNode)
@@ -28,80 +30,80 @@ public:
 
 	SLATE_END_ARGS()
 
-	void Construct(const FArguments& InArgs);
+	UE_API void Construct(const FArguments& InArgs);
 
 	// SGraphNode interface
-	virtual TSharedRef<SWidget> CreateTitleWidget(TSharedPtr<SNodeTitle> InNodeTitle) override;
+	UE_API virtual TSharedRef<SWidget> CreateTitleWidget(TSharedPtr<SNodeTitle> InNodeTitle) override;
 
-	virtual void EndUserInteraction() const override;
-	virtual void MoveTo(const FVector2f& NewPosition, FNodeSet& NodeFilter, bool bMarkDirty = true) override;
-	virtual void AddPin( const TSharedRef<SGraphPin>& PinToAdd ) override;
-	virtual void CreateStandardPinWidget(UEdGraphPin* CurPin) override;
+	UE_API virtual void EndUserInteraction() const override;
+	UE_API virtual void MoveTo(const FVector2f& NewPosition, FNodeSet& NodeFilter, bool bMarkDirty = true) override;
+	UE_API virtual void AddPin( const TSharedRef<SGraphPin>& PinToAdd ) override;
+	UE_API virtual void CreateStandardPinWidget(UEdGraphPin* CurPin) override;
 	virtual void SetDefaultTitleAreaWidget(TSharedRef<SOverlay> DefaultTitleAreaWidget) override
 	{
 		TitleAreaWidget = DefaultTitleAreaWidget;
 	}
-	virtual const FSlateBrush * GetNodeBodyBrush() const override;
-	virtual FReply OnMouseButtonDown(const FGeometry& MyGeometry, const FPointerEvent& MouseEvent) override;
-	virtual FReply OnMouseButtonDoubleClick(const FGeometry& InMyGeometry, const FPointerEvent& InMouseEvent) override;
+	UE_API virtual const FSlateBrush * GetNodeBodyBrush() const override;
+	UE_API virtual FReply OnMouseButtonDown(const FGeometry& MyGeometry, const FPointerEvent& MouseEvent) override;
+	UE_API virtual FReply OnMouseButtonDoubleClick(const FGeometry& InMyGeometry, const FPointerEvent& InMouseEvent) override;
 
-	virtual TSharedRef<SWidget> CreateNodeContentArea() override;
-	virtual void GetOverlayBrushes(bool bSelected, const FVector2f& WidgetSize, TArray<FOverlayBrushInfo>& Brushes) const override;
-	virtual void GetNodeInfoPopups(FNodeInfoContext* Context, TArray<FGraphInformationPopupInfo>& Popups) const override;
-	virtual TArray<FOverlayWidgetInfo> GetOverlayWidgets(bool bSelected, const FVector2f& WidgetSize) const override;
+	UE_API virtual TSharedRef<SWidget> CreateNodeContentArea() override;
+	UE_API virtual void GetOverlayBrushes(bool bSelected, const FVector2f& WidgetSize, TArray<FOverlayBrushInfo>& Brushes) const override;
+	UE_API virtual void GetNodeInfoPopups(FNodeInfoContext* Context, TArray<FGraphInformationPopupInfo>& Popups) const override;
+	UE_API virtual TArray<FOverlayWidgetInfo> GetOverlayWidgets(bool bSelected, const FVector2f& WidgetSize) const override;
 
-	virtual void RefreshErrorInfo() override;
-	virtual void Tick(const FGeometry& AllottedGeometry, const double InCurrentTime, const float InDeltaTime) override;
+	UE_API virtual void RefreshErrorInfo() override;
+	UE_API virtual void Tick(const FGeometry& AllottedGeometry, const double InCurrentTime, const float InDeltaTime) override;
 	virtual bool IsHidingPinWidgets() const override { return UseLowDetailNodeContent(); }
-	virtual bool UseLowDetailPinNames() const override;
-	virtual void UpdateGraphNode() override;
-	void UpdateStandardNode();
-	void UpdateCompactNode();
+	UE_API virtual bool UseLowDetailPinNames() const override;
+	UE_API virtual void UpdateGraphNode() override;
+	UE_API void UpdateStandardNode();
+	UE_API void UpdateCompactNode();
 
-	void CreateAddPinButton();
+	UE_API void CreateAddPinButton();
 
-	void CreateWorkflowWidgets();
+	UE_API void CreateWorkflowWidgets();
 
 	/** Callback function executed when Add pin button is clicked */
-	virtual FReply OnAddPin() override;
+	UE_API virtual FReply OnAddPin() override;
 
 protected:
 
-	bool UseLowDetailNodeContent() const;
+	UE_API bool UseLowDetailNodeContent() const;
 
-	FVector2D GetLowDetailDesiredSize() const;
+	UE_API FVector2D GetLowDetailDesiredSize() const;
 
-	EVisibility GetTitleVisibility() const;
-	EVisibility GetArrayPlusButtonVisibility(URigVMPin* InModelPin) const;
+	UE_API EVisibility GetTitleVisibility() const;
+	UE_API EVisibility GetArrayPlusButtonVisibility(URigVMPin* InModelPin) const;
 
-	FText GetPinLabel(TWeakPtr<SGraphPin> GraphPin) const;
+	UE_API FText GetPinLabel(TWeakPtr<SGraphPin> GraphPin) const;
 
-	virtual TOptional<FSlateColor> GetHighlightColor(const SGraphPin* InGraphPin) const override;
-	FSlateColor GetVariableLabelTextColor(TWeakObjectPtr<URigVMFunctionReferenceNode> FunctionReferenceNode, FName InVariableName) const;
-	FText GetVariableLabelTooltipText(TWeakObjectPtr<URigVMBlueprint> InBlueprint, FName InVariableName) const;
+	UE_API virtual TOptional<FSlateColor> GetHighlightColor(const SGraphPin* InGraphPin) const override;
+	UE_API FSlateColor GetVariableLabelTextColor(TWeakObjectPtr<URigVMFunctionReferenceNode> FunctionReferenceNode, FName InVariableName) const;
+	UE_API FText GetVariableLabelTooltipText(TWeakObjectPtr<URigVMBlueprint> InBlueprint, FName InVariableName) const;
 
-	FReply HandleAddArrayElement(FString InModelPinPath);
+	UE_API FReply HandleAddArrayElement(FString InModelPinPath);
 
-	void HandleNodeTitleDirtied();
-	void HandleNodePinsChanged();
-	void HandleNodeBeginRemoval();
+	UE_API void HandleNodeTitleDirtied();
+	UE_API void HandleNodePinsChanged();
+	UE_API void HandleNodeBeginRemoval();
 
-	FText GetInstructionCountText() const;
-	FText GetInstructionDurationText() const;
+	UE_API FText GetInstructionCountText() const;
+	UE_API FText GetInstructionDurationText() const;
 
-	TSharedRef<SWidget> OnOverrideWidgetMenu() const;
+	UE_API TSharedRef<SWidget> OnOverrideWidgetMenu() const;
 
 protected:
 
-	int32 GetNodeTopologyVersion() const;
-	EVisibility GetPinVisibility(int32 InPinInfoIndex, bool bAskingForSubPin) const;
-	const FSlateBrush * GetExpanderImage(int32 InPinInfoIndex, bool bLeft, bool bHovered) const;
-	FReply OnExpanderArrowClicked(int32 InPinInfoIndex);
-	void HandleModifiedEvent(ERigVMGraphNotifType InNotifType, URigVMGraph* InGraph, UObject* InSubject);
-	virtual void UpdatePinTreeView();
-	ERigVMNodeDefaultValueOverrideState::Type GetPinDefaultValueOverrideState() const;
-	bool IsCategoryPin(const SGraphPin* InGraphPin) const;
-	EOverrideWidgetStatus::Type GetOverrideStatus() const;
+	UE_API int32 GetNodeTopologyVersion() const;
+	UE_API EVisibility GetPinVisibility(int32 InPinInfoIndex, bool bAskingForSubPin) const;
+	UE_API const FSlateBrush * GetExpanderImage(int32 InPinInfoIndex, bool bLeft, bool bHovered) const;
+	UE_API FReply OnExpanderArrowClicked(int32 InPinInfoIndex);
+	UE_API void HandleModifiedEvent(ERigVMGraphNotifType InNotifType, URigVMGraph* InGraph, UObject* InSubject);
+	UE_API virtual void UpdatePinTreeView();
+	UE_API ERigVMNodeDefaultValueOverrideState::Type GetPinDefaultValueOverrideState() const;
+	UE_API bool IsCategoryPin(const SGraphPin* InGraphPin) const;
+	UE_API EOverrideWidgetStatus::Type GetOverrideStatus() const;
 
 	/** Cached widget title area */
 	TSharedPtr<SOverlay> TitleAreaWidget;
@@ -113,8 +115,8 @@ protected:
 	TSharedPtr<STextBlock> InstructionDurationTextBlockWidget;
 	TSharedPtr<SOverrideStatusWidget> OverrideStatusWidget; 
 
-	static const FSlateBrush* CachedImg_CR_Pin_Connected;
-	static const FSlateBrush* CachedImg_CR_Pin_Disconnected;
+	static UE_API const FSlateBrush* CachedImg_CR_Pin_Connected;
+	static UE_API const FSlateBrush* CachedImg_CR_Pin_Disconnected;
 
 	/** Cache the node title so we can invalidate it */
 	TSharedPtr<SNodeTitle> NodeTitle;
@@ -170,3 +172,5 @@ protected:
 	// but we still have a pointer to them in our pin widget.
 	TSet<UEdGraphPin *> PinsToDelete;
 };
+
+#undef UE_API
diff --git a/Engine/Plugins/Runtime/RigVM/Source/RigVMEditor/Public/Widgets/SRigVMGraphNodeComment.h b/Engine/Plugins/Runtime/RigVM/Source/RigVMEditor/Public/Widgets/SRigVMGraphNodeComment.h
--- a/Engine/Plugins/Runtime/RigVM/Source/RigVMEditor/Public/Widgets/SRigVMGraphNodeComment.h
+++ b/Engine/Plugins/Runtime/RigVM/Source/RigVMEditor/Public/Widgets/SRigVMGraphNodeComment.h
@@ -5,25 +5,27 @@
 #include "CoreMinimal.h"
 #include "SGraphNodeComment.h"
 
+#define UE_API RIGVMEDITOR_API
+
 class UEdGraphNode_Comment;
 
-class RIGVMEDITOR_API SRigVMGraphNodeComment : public SGraphNodeComment
+class SRigVMGraphNodeComment : public SGraphNodeComment
 {
 public:
 
-	SRigVMGraphNodeComment();
+	UE_API SRigVMGraphNodeComment();
 
-	virtual FReply OnMouseButtonUp(const FGeometry& MyGeometry, const FPointerEvent& MouseEvent) override;
-	virtual void EndUserInteraction() const override;
-	virtual void Tick(const FGeometry& AllottedGeometry, const double InCurrentTime, const float InDeltaTime) override;
-	virtual void MoveTo(const FVector2f& NewPosition, FNodeSet& NodeFilter, bool bMarkDirty = true) override;
+	UE_API virtual FReply OnMouseButtonUp(const FGeometry& MyGeometry, const FPointerEvent& MouseEvent) override;
+	UE_API virtual void EndUserInteraction() const override;
+	UE_API virtual void Tick(const FGeometry& AllottedGeometry, const double InCurrentTime, const float InDeltaTime) override;
+	UE_API virtual void MoveTo(const FVector2f& NewPosition, FNodeSet& NodeFilter, bool bMarkDirty = true) override;
 
 protected:
 
-	virtual bool IsNodeUnderComment(UEdGraphNode_Comment* InCommentNode, const TSharedRef<SGraphNode> InNodeWidget) const override;
-	bool IsNodeUnderComment(UEdGraphNode* InNode) const;
-	bool IsNodeUnderComment(const FVector2f& InNodePosition) const;
-	bool IsNodeUnderComment(const FVector2f& InCommentPosition, const FVector2f& InNodePosition) const;
+	UE_API virtual bool IsNodeUnderComment(UEdGraphNode_Comment* InCommentNode, const TSharedRef<SGraphNode> InNodeWidget) const override;
+	UE_API bool IsNodeUnderComment(UEdGraphNode* InNode) const;
+	UE_API bool IsNodeUnderComment(const FVector2f& InNodePosition) const;
+	UE_API bool IsNodeUnderComment(const FVector2f& InCommentPosition, const FVector2f& InNodePosition) const;
 
 private:
 
@@ -31,3 +33,5 @@ private:
 
 	int8 CachedColorBubble;
 };
+
+#undef UE_API
diff --git a/Engine/Plugins/Runtime/RigVM/Source/RigVMEditor/Public/Widgets/SRigVMGraphNodeKnot.h b/Engine/Plugins/Runtime/RigVM/Source/RigVMEditor/Public/Widgets/SRigVMGraphNodeKnot.h
--- a/Engine/Plugins/Runtime/RigVM/Source/RigVMEditor/Public/Widgets/SRigVMGraphNodeKnot.h
+++ b/Engine/Plugins/Runtime/RigVM/Source/RigVMEditor/Public/Widgets/SRigVMGraphNodeKnot.h
@@ -6,16 +6,20 @@
 #include "Widgets/DeclarativeSyntaxSupport.h"
 #include "SGraphNodeKnot.h"
 
-class RIGVMEDITOR_API SRigVMGraphNodeKnot : public SGraphNodeKnot
+#define UE_API RIGVMEDITOR_API
+
+class SRigVMGraphNodeKnot : public SGraphNodeKnot
 {
 	SLATE_BEGIN_ARGS(SRigVMGraphNodeKnot) {} 
 	SLATE_END_ARGS()
 
-	void Construct(const FArguments& InArgs, UEdGraphNode* InKnot);
+	UE_API void Construct(const FArguments& InArgs, UEdGraphNode* InKnot);
 
 	// SGraphNode interface
-	virtual void EndUserInteraction() const override;
-	virtual void MoveTo(const FVector2f& NewPosition, FNodeSet& NodeFilter, bool bMarkDirty = true) override;
+	UE_API virtual void EndUserInteraction() const override;
+	UE_API virtual void MoveTo(const FVector2f& NewPosition, FNodeSet& NodeFilter, bool bMarkDirty = true) override;
+
+	UE_API void HandleNodeBeginRemoval();
+};
 
-	void HandleNodeBeginRemoval();
-};
\ No newline at end of file
+#undef UE_API
diff --git a/Engine/Plugins/Runtime/RigVM/Source/RigVMEditor/Public/Widgets/SRigVMGraphPinCategory.h b/Engine/Plugins/Runtime/RigVM/Source/RigVMEditor/Public/Widgets/SRigVMGraphPinCategory.h
--- a/Engine/Plugins/Runtime/RigVM/Source/RigVMEditor/Public/Widgets/SRigVMGraphPinCategory.h
+++ b/Engine/Plugins/Runtime/RigVM/Source/RigVMEditor/Public/Widgets/SRigVMGraphPinCategory.h
@@ -7,7 +7,9 @@
 #include "Widgets/SWidget.h"
 #include "SGraphPin.h"
 
-class RIGVMEDITOR_API SRigVMGraphPinCategory : public SGraphPin
+#define UE_API RIGVMEDITOR_API
+
+class SRigVMGraphPinCategory : public SGraphPin
 {
 public:
 	SLATE_BEGIN_ARGS(SRigVMGraphPinCategory)
@@ -15,13 +17,15 @@ public:
 	SLATE_END_ARGS()
 
 
-	void Construct(const FArguments& InArgs, UEdGraphPin* InGraphPinObj);
-	virtual FReply OnMouseButtonDown(const FGeometry& MyGeometry, const FPointerEvent& MouseEvent) override;
+	UE_API void Construct(const FArguments& InArgs, UEdGraphPin* InGraphPinObj);
+	UE_API virtual FReply OnMouseButtonDown(const FGeometry& MyGeometry, const FPointerEvent& MouseEvent) override;
 
 protected:
 
 	//~ Begin SGraphPin Interface
-	virtual TSharedRef<SWidget>	GetDefaultValueWidget() override;
+	UE_API virtual TSharedRef<SWidget>	GetDefaultValueWidget() override;
 	virtual const FSlateBrush* GetPinIcon() const override { return nullptr; }
 	//~ End SGraphPin Interface
 };
+
+#undef UE_API
diff --git a/Engine/Plugins/Runtime/RigVM/Source/RigVMEditor/Public/Widgets/SRigVMGraphPinCurveFloat.h b/Engine/Plugins/Runtime/RigVM/Source/RigVMEditor/Public/Widgets/SRigVMGraphPinCurveFloat.h
--- a/Engine/Plugins/Runtime/RigVM/Source/RigVMEditor/Public/Widgets/SRigVMGraphPinCurveFloat.h
+++ b/Engine/Plugins/Runtime/RigVM/Source/RigVMEditor/Public/Widgets/SRigVMGraphPinCurveFloat.h
@@ -9,32 +9,36 @@
 #include "Curves/CurveFloat.h"
 #include "SGraphPin.h"
 
-class RIGVMEDITOR_API SRigVMGraphPinCurveFloat : public SGraphPin, public FCurveOwnerInterface
+#define UE_API RIGVMEDITOR_API
+
+class SRigVMGraphPinCurveFloat : public SGraphPin, public FCurveOwnerInterface
 {
 public:
 	SLATE_BEGIN_ARGS(SRigVMGraphPinCurveFloat) {}
 	SLATE_END_ARGS()
 
-	void Construct(const FArguments& InArgs, UEdGraphPin* InGraphPinObj);
+	UE_API void Construct(const FArguments& InArgs, UEdGraphPin* InGraphPinObj);
 
 	/** FCurveOwnerInterface interface */
-	virtual TArray<FRichCurveEditInfoConst> GetCurves() const override;
-	virtual TArray<FRichCurveEditInfo> GetCurves() override;
-	virtual void ModifyOwner() override;
-	virtual TArray<const UObject*> GetOwners() const override;
-	virtual void MakeTransactional() override;
-	virtual void OnCurveChanged(const TArray<FRichCurveEditInfo>& ChangedCurveEditInfos) override;
-	virtual bool IsValidCurve(FRichCurveEditInfo CurveInfo) override;
+	UE_API virtual TArray<FRichCurveEditInfoConst> GetCurves() const override;
+	UE_API virtual TArray<FRichCurveEditInfo> GetCurves() override;
+	UE_API virtual void ModifyOwner() override;
+	UE_API virtual TArray<const UObject*> GetOwners() const override;
+	UE_API virtual void MakeTransactional() override;
+	UE_API virtual void OnCurveChanged(const TArray<FRichCurveEditInfo>& ChangedCurveEditInfos) override;
+	UE_API virtual bool IsValidCurve(FRichCurveEditInfo CurveInfo) override;
 
 protected:
 
 	//~ Begin SGraphPin Interface
-	virtual TSharedRef<SWidget>	GetDefaultValueWidget() override;
-	virtual void Tick(const FGeometry& AllottedGeometry, const double InCurrentTime, const float InDeltaTime) override;
+	UE_API virtual TSharedRef<SWidget>	GetDefaultValueWidget() override;
+	UE_API virtual void Tick(const FGeometry& AllottedGeometry, const double InCurrentTime, const float InDeltaTime) override;
 	//~ End SGraphPin Interface
 
-	FRuntimeFloatCurve& UpdateAndGetCurve();
+	UE_API FRuntimeFloatCurve& UpdateAndGetCurve();
 
 	TSharedPtr<SCurveEditor> CurveEditor;
 	FRuntimeFloatCurve Curve;
 };
+
+#undef UE_API
diff --git a/Engine/Plugins/Runtime/RigVM/Source/RigVMEditor/Public/Widgets/SRigVMGraphPinEditableNameValueWidget.h b/Engine/Plugins/Runtime/RigVM/Source/RigVMEditor/Public/Widgets/SRigVMGraphPinEditableNameValueWidget.h
--- a/Engine/Plugins/Runtime/RigVM/Source/RigVMEditor/Public/Widgets/SRigVMGraphPinEditableNameValueWidget.h
+++ b/Engine/Plugins/Runtime/RigVM/Source/RigVMEditor/Public/Widgets/SRigVMGraphPinEditableNameValueWidget.h
@@ -10,10 +10,12 @@
 #include "Widgets/Views/SListView.h"
 #include "Widgets/Input/SComboBox.h"
 
+#define UE_API RIGVMEDITOR_API
+
 /**
  * A searchable text combo box
  */
-class RIGVMEDITOR_API SRigVMGraphPinEditableNameValueWidget : public SComboButton
+class SRigVMGraphPinEditableNameValueWidget : public SComboButton
 {
 public:
 	/** Type of list used for showing menu options. */
@@ -70,41 +72,41 @@ public:
 	 *
 	 * @param InArgs   Declaration from which to construct the combo box
 	 */
-	void Construct(const FArguments& InArgs);
+	UE_API void Construct(const FArguments& InArgs);
 
-	void ClearSelection();
+	UE_API void ClearSelection();
 
-	void SetSelectedItem(TSharedPtr<FString> InSelectedItem);
+	UE_API void SetSelectedItem(TSharedPtr<FString> InSelectedItem);
 
 	/** @return the item currently selected by the combo box. */
-	TSharedPtr<FString> GetSelectedItem();
+	UE_API TSharedPtr<FString> GetSelectedItem();
 
 	/**
 	 * Requests a list refresh after updating options
 	 * Call SetSelectedItem to update the selected item if required
 	 * @see SetSelectedItem
 	 */
-	void RefreshOptions();
+	UE_API void RefreshOptions();
 
 private:
 
 	/** Generate a row for the InItem in the combo box's list (passed in as OwnerTable). Do this by calling the user-specified OnGenerateWidget */
-	TSharedRef<ITableRow> GenerateMenuItemRow(TSharedPtr<FString> InItem, const TSharedRef<STableViewBase>& OwnerTable);
+	UE_API TSharedRef<ITableRow> GenerateMenuItemRow(TSharedPtr<FString> InItem, const TSharedRef<STableViewBase>& OwnerTable);
 
 	//** Called if the menu is closed
-	void OnMenuOpenChanged(bool bOpen);
+	UE_API void OnMenuOpenChanged(bool bOpen);
 
 	/** Invoked when the selection in the list changes */
-	void OnSelectionChanged_Internal(TSharedPtr<FString> ProposedSelection, ESelectInfo::Type SelectInfo);
+	UE_API void OnSelectionChanged_Internal(TSharedPtr<FString> ProposedSelection, ESelectInfo::Type SelectInfo);
 
 	/** Invoked when the search text changes */
-	void OnSearchTextChanged(const FText& ChangedText);
+	UE_API void OnSearchTextChanged(const FText& ChangedText);
 
 	/** Invoked when the search is committed*/
-	void OnSearchTextCommitted(const FText& ChangedText, ETextCommit::Type CommitType);
+	UE_API void OnSearchTextCommitted(const FText& ChangedText, ETextCommit::Type CommitType);
 
 	/** Handle clicking on the content menu */
-	virtual FReply OnButtonClicked() override;
+	UE_API virtual FReply OnButtonClicked() override;
 
 	/** The item style to use. */
 	const FTableRowStyle* ItemStyle;
@@ -128,7 +130,8 @@ private:
 	const TArray< TSharedPtr<FString> >* OptionsSource;
 
 	/** Used to focus the name box immediately following construction */
-	EActiveTimerReturnType SetFocusPostConstruct(double InCurrentTime, float InDeltaTime);
+	UE_API EActiveTimerReturnType SetFocusPostConstruct(double InCurrentTime, float InDeltaTime);
 };
 
 
+#undef UE_API
diff --git a/Engine/Plugins/Runtime/RigVM/Source/RigVMEditor/Public/Widgets/SRigVMGraphPinEnumPicker.h b/Engine/Plugins/Runtime/RigVM/Source/RigVMEditor/Public/Widgets/SRigVMGraphPinEnumPicker.h
--- a/Engine/Plugins/Runtime/RigVM/Source/RigVMEditor/Public/Widgets/SRigVMGraphPinEnumPicker.h
+++ b/Engine/Plugins/Runtime/RigVM/Source/RigVMEditor/Public/Widgets/SRigVMGraphPinEnumPicker.h
@@ -9,7 +9,9 @@
 #include "SSearchableComboBox.h"
 #include "RigVMModel/RigVMPin.h"
 
-class RIGVMEDITOR_API SRigVMEnumPicker : public SCompoundWidget
+#define UE_API RIGVMEDITOR_API
+
+class SRigVMEnumPicker : public SCompoundWidget
 {
 public:
 	DECLARE_DELEGATE_RetVal(TObjectPtr<UEnum>, FGetCurrentEnum);
@@ -20,7 +22,7 @@ public:
 		SLATE_ARGUMENT(bool, IsEnabled)
 	SLATE_END_ARGS()
 
-	void Construct(const FArguments& InArgs);
+	UE_API void Construct(const FArguments& InArgs);
 
 protected:
 
@@ -41,8 +43,8 @@ protected:
 		return nullptr;
 	}
 
-	TSharedRef<SWidget> OnGetEnumNameWidget(TSharedPtr<FString> InItem);
-	void PopulateEnumOptions();
+	UE_API TSharedRef<SWidget> OnGetEnumNameWidget(TSharedPtr<FString> InItem);
+	UE_API void PopulateEnumOptions();
 
 	TArray<TSharedPtr<FString>> EnumOptions;
 	SSearchableComboBox::FOnSelectionChanged OnEnumChangedDelegate;
@@ -50,7 +52,7 @@ protected:
 	bool bIsEnabled;
 };
 
-class RIGVMEDITOR_API SRigVMGraphPinEnumPicker : public SGraphPin
+class SRigVMGraphPinEnumPicker : public SGraphPin
 {
 public:
 	SLATE_BEGIN_ARGS(SRigVMGraphPinEnumPicker)
@@ -58,15 +60,17 @@ public:
 		SLATE_ARGUMENT(URigVMPin*, ModelPin)
 	SLATE_END_ARGS()
 
-	void Construct(const FArguments& InArgs, UEdGraphPin* InGraphPinObj);
+	UE_API void Construct(const FArguments& InArgs, UEdGraphPin* InGraphPinObj);
 
 protected:
 	//~ Begin SGraphPin Interface
-	virtual TSharedRef<SWidget>	GetDefaultValueWidget() override;
+	UE_API virtual TSharedRef<SWidget>	GetDefaultValueWidget() override;
 	//~ End SGraphPin Interface
 
-	void HandleControlEnumChanged(TSharedPtr<FString> InItem, ESelectInfo::Type InSelectionInfo);
+	UE_API void HandleControlEnumChanged(TSharedPtr<FString> InItem, ESelectInfo::Type InSelectionInfo);
 	
 	URigVMPin* ModelPin;
 
 };
+
+#undef UE_API
diff --git a/Engine/Plugins/Runtime/RigVM/Source/RigVMEditor/Public/Widgets/SRigVMGraphPinNameList.h b/Engine/Plugins/Runtime/RigVM/Source/RigVMEditor/Public/Widgets/SRigVMGraphPinNameList.h
--- a/Engine/Plugins/Runtime/RigVM/Source/RigVMEditor/Public/Widgets/SRigVMGraphPinNameList.h
+++ b/Engine/Plugins/Runtime/RigVM/Source/RigVMEditor/Public/Widgets/SRigVMGraphPinNameList.h
@@ -9,9 +9,11 @@
 #include "RigVMModel/RigVMPin.h"
 #include "Widgets/SRigVMGraphPinNameListValueWidget.h"
 
+#define UE_API RIGVMEDITOR_API
+
 class URigVMEdGraph;
 
-class RIGVMEDITOR_API SRigVMGraphPinNameList : public SGraphPin
+class SRigVMGraphPinNameList : public SGraphPin
 {
 public:
 
@@ -40,22 +42,22 @@ public:
 
 	SLATE_END_ARGS()
 
-	void Construct(const FArguments& InArgs, UEdGraphPin* InGraphPinObj);
+	UE_API void Construct(const FArguments& InArgs, UEdGraphPin* InGraphPinObj);
 
 protected:
 	//~ Begin SGraphPin Interface
-	virtual TSharedRef<SWidget>	GetDefaultValueWidget() override;
+	UE_API virtual TSharedRef<SWidget>	GetDefaultValueWidget() override;
 	//~ End SGraphPin Interface
 
-	const TArray<TSharedPtr<FRigVMStringWithTag>>* GetNameList(bool bForContent = true) const;
-	FText GetNameListText() const;
-	FSlateColor GetNameColor() const;
-	virtual void SetNameListText(const FText& NewTypeInValue, ETextCommit::Type CommitInfo);
-	void UpdateNameLists(bool bUpdateCurrent = true, bool bUpdateValidation = true);
+	UE_API const TArray<TSharedPtr<FRigVMStringWithTag>>* GetNameList(bool bForContent = true) const;
+	UE_API FText GetNameListText() const;
+	UE_API FSlateColor GetNameColor() const;
+	UE_API virtual void SetNameListText(const FText& NewTypeInValue, ETextCommit::Type CommitInfo);
+	UE_API void UpdateNameLists(bool bUpdateCurrent = true, bool bUpdateValidation = true);
 
-	TSharedRef<SWidget> MakeNameListItemWidget(TSharedPtr<FRigVMStringWithTag> InItem);
-	void OnNameListChanged(TSharedPtr<FRigVMStringWithTag> NewSelection, ESelectInfo::Type SelectInfo);
-	void OnNameListComboBox();
+	UE_API TSharedRef<SWidget> MakeNameListItemWidget(TSharedPtr<FRigVMStringWithTag> InItem);
+	UE_API void OnNameListChanged(TSharedPtr<FRigVMStringWithTag> NewSelection, ESelectInfo::Type SelectInfo);
+	UE_API void OnNameListComboBox();
 
 	FOnGetNameListContent OnGetNameListContent;
 	FOnGetNameListContent OnGetNameListContentForValidation;
@@ -70,11 +72,13 @@ protected:
 	bool AllowUserProvidedText;
 
 	/** Helper buttons. */
-	FSlateColor OnGetWidgetForeground() const;
-	FSlateColor OnGetWidgetBackground() const;
-	FReply HandleGetSelectedClicked();
-	FReply HandleBrowseClicked();
+	UE_API FSlateColor OnGetWidgetForeground() const;
+	UE_API FSlateColor OnGetWidgetBackground() const;
+	UE_API FReply HandleGetSelectedClicked();
+	UE_API FReply HandleBrowseClicked();
 	FOnGetNameFromSelection OnGetNameFromSelection;
 	FOnGetSelectedClicked OnGetSelectedClicked;
 	FOnBrowseClicked OnBrowseClicked;
 };
+
+#undef UE_API
diff --git a/Engine/Plugins/Runtime/RigVM/Source/RigVMEditor/Public/Widgets/SRigVMGraphPinNameListValueWidget.h b/Engine/Plugins/Runtime/RigVM/Source/RigVMEditor/Public/Widgets/SRigVMGraphPinNameListValueWidget.h
--- a/Engine/Plugins/Runtime/RigVM/Source/RigVMEditor/Public/Widgets/SRigVMGraphPinNameListValueWidget.h
+++ b/Engine/Plugins/Runtime/RigVM/Source/RigVMEditor/Public/Widgets/SRigVMGraphPinNameListValueWidget.h
@@ -11,10 +11,12 @@
 #include "Widgets/Input/SComboBox.h"
 #include "EdGraph/RigVMEdGraph.h"
 
+#define UE_API RIGVMEDITOR_API
+
 /**
  * A searchable text combo box
  */
-class RIGVMEDITOR_API SRigVMGraphPinNameListValueWidget : public SComboButton
+class SRigVMGraphPinNameListValueWidget : public SComboButton
 {
 public:
 	/** Type of list used for showing menu options. */
@@ -83,48 +85,48 @@ public:
 	 *
 	 * @param InArgs   Declaration from which to construct the combo box
 	 */
-	void Construct(const FArguments& InArgs);
+	UE_API void Construct(const FArguments& InArgs);
 
-	void ClearSelection();
+	UE_API void ClearSelection();
 
-	void SetSelectedItem(TSharedPtr<FRigVMStringWithTag> InSelectedItem);
-	void SetOptionsSource(const TArray< TSharedPtr<FRigVMStringWithTag> >* InOptionsSource);
+	UE_API void SetSelectedItem(TSharedPtr<FRigVMStringWithTag> InSelectedItem);
+	UE_API void SetOptionsSource(const TArray< TSharedPtr<FRigVMStringWithTag> >* InOptionsSource);
 
 	/** @return the item currently selected by the combo box. */
-	TSharedPtr<FRigVMStringWithTag> GetSelectedItem();
+	UE_API TSharedPtr<FRigVMStringWithTag> GetSelectedItem();
 
 	/**
 	 * Requests a list refresh after updating options
 	 * Call SetSelectedItem to update the selected item if required
 	 * @see SetSelectedItem
 	 */
-	void RefreshOptions();
+	UE_API void RefreshOptions();
 
 private:
 
 	/** Generate a row for the InItem in the combo box's list (passed in as OwnerTable). Do this by calling the user-specified OnGenerateWidget */
-	TSharedRef<ITableRow> GenerateMenuItemRow(TSharedPtr<FRigVMStringWithTag> InItem, const TSharedRef<STableViewBase>& OwnerTable);
+	UE_API TSharedRef<ITableRow> GenerateMenuItemRow(TSharedPtr<FRigVMStringWithTag> InItem, const TSharedRef<STableViewBase>& OwnerTable);
 
 	//** Called if the menu is closed
-	void OnMenuOpenChanged(bool bOpen);
+	UE_API void OnMenuOpenChanged(bool bOpen);
 
 	/** Invoked when the selection in the list changes */
-	void OnSelectionChanged_Internal(TSharedPtr<FRigVMStringWithTag> ProposedSelection, ESelectInfo::Type SelectInfo, bool bForce = false);
+	UE_API void OnSelectionChanged_Internal(TSharedPtr<FRigVMStringWithTag> ProposedSelection, ESelectInfo::Type SelectInfo, bool bForce = false);
 
 	/** Invoked when the search text changes */
-	void OnSearchTextChanged(const FText& ChangedText);
+	UE_API void OnSearchTextChanged(const FText& ChangedText);
 
 	/** Invoked when the search is committed*/
-	void OnSearchTextCommitted(const FText& ChangedText, ETextCommit::Type CommitType);
+	UE_API void OnSearchTextCommitted(const FText& ChangedText, ETextCommit::Type CommitType);
 
 	/** Special case handling for search box key commands */
-	FReply OnSearchTextKeyDown(const FGeometry& MyGeometry, const FKeyEvent& InKeyEvent);
+	UE_API FReply OnSearchTextKeyDown(const FGeometry& MyGeometry, const FKeyEvent& InKeyEvent);
 
 	/** Special case handling for combo list key commands */
-	FReply OnComboListKeyDown(const FGeometry& MyGeometry, const FKeyEvent& InKeyEvent);
+	UE_API FReply OnComboListKeyDown(const FGeometry& MyGeometry, const FKeyEvent& InKeyEvent);
 
 	/** Handle clicking on the content menu */
-	virtual FReply OnButtonClicked() override;
+	UE_API virtual FReply OnButtonClicked() override;
 
 	/** The item style to use. */
 	const FTableRowStyle* ItemStyle;
@@ -149,7 +151,8 @@ private:
 	bool AllowUserProvidedText;
 
 	/** Used to focus the name box immediately following construction */
-	EActiveTimerReturnType SetFocusPostConstruct(double InCurrentTime, float InDeltaTime);
+	UE_API EActiveTimerReturnType SetFocusPostConstruct(double InCurrentTime, float InDeltaTime);
 };
 
 
+#undef UE_API
diff --git a/Engine/Plugins/Runtime/RigVM/Source/RigVMEditor/Public/Widgets/SRigVMGraphPinQuat.h b/Engine/Plugins/Runtime/RigVM/Source/RigVMEditor/Public/Widgets/SRigVMGraphPinQuat.h
--- a/Engine/Plugins/Runtime/RigVM/Source/RigVMEditor/Public/Widgets/SRigVMGraphPinQuat.h
+++ b/Engine/Plugins/Runtime/RigVM/Source/RigVMEditor/Public/Widgets/SRigVMGraphPinQuat.h
@@ -9,7 +9,9 @@
 #include "RigVMModel/RigVMPin.h"
 #include "SGraphPin.h"
 
-class RIGVMEDITOR_API SRigVMGraphPinQuat : public SGraphPin
+#define UE_API RIGVMEDITOR_API
+
+class SRigVMGraphPinQuat : public SGraphPin
 {
 public:
 	SLATE_BEGIN_ARGS(SRigVMGraphPinQuat)
@@ -19,19 +21,21 @@ public:
 	SLATE_END_ARGS()
 
 
-	void Construct(const FArguments& InArgs, UEdGraphPin* InGraphPinObj);
+	UE_API void Construct(const FArguments& InArgs, UEdGraphPin* InGraphPinObj);
 
 protected:
 
 	//~ Begin SGraphPin Interface
-	virtual TSharedRef<SWidget>	GetDefaultValueWidget() override;
+	UE_API virtual TSharedRef<SWidget>	GetDefaultValueWidget() override;
 	//~ End SGraphPin Interface
 
-	TOptional<FRotator> GetRotator() const;
-	void OnRotatorCommitted(FRotator InRotator, ETextCommit::Type InCommitType, bool bUndoRedo);
-	TOptional<float> GetRotatorComponent(int32 InComponent) const;
-	void OnRotatorComponentChanged(float InValue, int32 InComponent);
-	void OnRotatorComponentCommitted(float InValue, ETextCommit::Type InCommitType, int32 InComponent, bool bUndoRedo);
+	UE_API TOptional<FRotator> GetRotator() const;
+	UE_API void OnRotatorCommitted(FRotator InRotator, ETextCommit::Type InCommitType, bool bUndoRedo);
+	UE_API TOptional<float> GetRotatorComponent(int32 InComponent) const;
+	UE_API void OnRotatorComponentChanged(float InValue, int32 InComponent);
+	UE_API void OnRotatorComponentCommitted(float InValue, ETextCommit::Type InCommitType, int32 InComponent, bool bUndoRedo);
 
 	URigVMPin* ModelPin;
 };
+
+#undef UE_API
diff --git a/Engine/Plugins/Runtime/RigVM/Source/RigVMEditor/Public/Widgets/SRigVMGraphPinUserDataNameSpace.h b/Engine/Plugins/Runtime/RigVM/Source/RigVMEditor/Public/Widgets/SRigVMGraphPinUserDataNameSpace.h
--- a/Engine/Plugins/Runtime/RigVM/Source/RigVMEditor/Public/Widgets/SRigVMGraphPinUserDataNameSpace.h
+++ b/Engine/Plugins/Runtime/RigVM/Source/RigVMEditor/Public/Widgets/SRigVMGraphPinUserDataNameSpace.h
@@ -8,28 +8,32 @@
 #include "SGraphPin.h"
 #include "Widgets/SRigVMGraphPinEditableNameValueWidget.h"
 
-class RIGVMEDITOR_API SRigVMGraphPinUserDataNameSpace : public SGraphPin
+#define UE_API RIGVMEDITOR_API
+
+class SRigVMGraphPinUserDataNameSpace : public SGraphPin
 {
 public:
 	SLATE_BEGIN_ARGS(SRigVMGraphPinUserDataNameSpace) {}
 	SLATE_END_ARGS()
 
-	void Construct(const FArguments& InArgs, UEdGraphPin* InGraphPinObj);
+	UE_API void Construct(const FArguments& InArgs, UEdGraphPin* InGraphPinObj);
 
 protected:
 	//~ Begin SGraphPin Interface
-	virtual TSharedRef<SWidget>	GetDefaultValueWidget() override;
+	UE_API virtual TSharedRef<SWidget>	GetDefaultValueWidget() override;
 	//~ End SGraphPin Interface
 
-	FText GetNameSpaceText() const;
-	virtual void SetNameSpaceText(const FText& NewTypeInValue, ETextCommit::Type CommitInfo);
+	UE_API FText GetNameSpaceText() const;
+	UE_API virtual void SetNameSpaceText(const FText& NewTypeInValue, ETextCommit::Type CommitInfo);
 
-	TSharedRef<SWidget> MakeNameSpaceItemWidget(TSharedPtr<FString> InItem);
-	void OnNameSpaceChanged(TSharedPtr<FString> NewSelection, ESelectInfo::Type SelectInfo);
-	void OnNameSpaceComboBox();
-	TArray<TSharedPtr<FString>>& GetNameSpaces();
+	UE_API TSharedRef<SWidget> MakeNameSpaceItemWidget(TSharedPtr<FString> InItem);
+	UE_API void OnNameSpaceChanged(TSharedPtr<FString> NewSelection, ESelectInfo::Type SelectInfo);
+	UE_API void OnNameSpaceComboBox();
+	UE_API TArray<TSharedPtr<FString>>& GetNameSpaces();
 
 	TArray<TSharedPtr<FString>> NameSpaces;
 	TSharedPtr<SRigVMGraphPinEditableNameValueWidget> NameComboBox;
 
 };
+
+#undef UE_API
diff --git a/Engine/Plugins/Runtime/RigVM/Source/RigVMEditor/Public/Widgets/SRigVMGraphPinUserDataPath.h b/Engine/Plugins/Runtime/RigVM/Source/RigVMEditor/Public/Widgets/SRigVMGraphPinUserDataPath.h
--- a/Engine/Plugins/Runtime/RigVM/Source/RigVMEditor/Public/Widgets/SRigVMGraphPinUserDataPath.h
+++ b/Engine/Plugins/Runtime/RigVM/Source/RigVMEditor/Public/Widgets/SRigVMGraphPinUserDataPath.h
@@ -13,7 +13,9 @@
 #include "RigVMCore/RigVMAssetUserData.h"
 #include "IPropertyAccessEditor.h"
 
-class RIGVMEDITOR_API SRigVMUserDataPath : public SCompoundWidget
+#define UE_API RIGVMEDITOR_API
+
+class SRigVMUserDataPath : public SCompoundWidget
 {
 public:
 
@@ -25,33 +27,33 @@ public:
 
 	SLATE_END_ARGS()
 
-	void Construct(const FArguments& InArgs);
+	UE_API void Construct(const FArguments& InArgs);
 
 protected:
 
-	FString GetUserDataPath(URigVMPin* ModelPin) const;
-	FString GetUserDataPath() const;
+	UE_API FString GetUserDataPath(URigVMPin* ModelPin) const;
+	UE_API FString GetUserDataPath() const;
 	FText GetUserDataPathText() const { return FText::FromString(GetUserDataPath()); }
-	const FSlateBrush* GetUserDataIcon() const;
-	FLinearColor GetUserDataColor() const;
-	const FSlateBrush* GetUserDataIcon(const UNameSpacedUserData::FUserData* InUserData) const;
-	FLinearColor GetUserDataColor(const UNameSpacedUserData::FUserData* InUserData) const;
-	TSharedRef<SWidget> GetTopLevelMenuContent();
-	void FillUserDataPathMenu( FMenuBuilder& InMenuBuilder, FString InParentPath );
-	void HandleSetUserDataPath(FString InUserDataPath);
-
-	URigVMBlueprint* GetBlueprint() const;
-	FString GetUserDataNameSpace(URigVMPin* ModelPin) const;
-	FString GetUserDataNameSpace() const;
-	const UNameSpacedUserData* GetUserDataObject() const;
-	const UNameSpacedUserData::FUserData* GetUserData() const;
+	UE_API const FSlateBrush* GetUserDataIcon() const;
+	UE_API FLinearColor GetUserDataColor() const;
+	UE_API const FSlateBrush* GetUserDataIcon(const UNameSpacedUserData::FUserData* InUserData) const;
+	UE_API FLinearColor GetUserDataColor(const UNameSpacedUserData::FUserData* InUserData) const;
+	UE_API TSharedRef<SWidget> GetTopLevelMenuContent();
+	UE_API void FillUserDataPathMenu( FMenuBuilder& InMenuBuilder, FString InParentPath );
+	UE_API void HandleSetUserDataPath(FString InUserDataPath);
+
+	UE_API URigVMBlueprint* GetBlueprint() const;
+	UE_API FString GetUserDataNameSpace(URigVMPin* ModelPin) const;
+	UE_API FString GetUserDataNameSpace() const;
+	UE_API const UNameSpacedUserData* GetUserDataObject() const;
+	UE_API const UNameSpacedUserData::FUserData* GetUserData() const;
 
 	TArray<URigVMPin*> ModelPins;
 	TSharedPtr<SMenuAnchor> MenuAnchor;
 	bool bAllowUObjects = false;
 };
 
-class RIGVMEDITOR_API  SRigVMGraphPinUserDataPath : public SGraphPin
+class  SRigVMGraphPinUserDataPath : public SGraphPin
 {
 public:
 
@@ -61,13 +63,15 @@ public:
 
 	SLATE_END_ARGS()
 
-	void Construct(const FArguments& InArgs, UEdGraphPin* InGraphPinObj);
+	UE_API void Construct(const FArguments& InArgs, UEdGraphPin* InGraphPinObj);
 
 protected:
 
 	//~ Begin SGraphPin Interface
-	virtual TSharedRef<SWidget>	GetDefaultValueWidget() override;
+	UE_API virtual TSharedRef<SWidget>	GetDefaultValueWidget() override;
 	//~ End SGraphPin Interface
 
 	TArray<URigVMPin*> ModelPins;
 };
+
+#undef UE_API
diff --git a/Engine/Plugins/Runtime/RigVM/Source/RigVMEditor/Public/Widgets/SRigVMGraphPinVariableBinding.h b/Engine/Plugins/Runtime/RigVM/Source/RigVMEditor/Public/Widgets/SRigVMGraphPinVariableBinding.h
--- a/Engine/Plugins/Runtime/RigVM/Source/RigVMEditor/Public/Widgets/SRigVMGraphPinVariableBinding.h
+++ b/Engine/Plugins/Runtime/RigVM/Source/RigVMEditor/Public/Widgets/SRigVMGraphPinVariableBinding.h
@@ -11,7 +11,9 @@
 #include "IPropertyAccessEditor.h"
 #include "RigVMBlueprint.h"
 
-class RIGVMEDITOR_API SRigVMGraphVariableBinding : public SCompoundWidget
+#define UE_API RIGVMEDITOR_API
+
+class SRigVMGraphVariableBinding : public SCompoundWidget
 {
 public:
 
@@ -31,21 +33,21 @@ public:
 
 	SLATE_END_ARGS()
 
-	void Construct(const FArguments& InArgs);
+	UE_API void Construct(const FArguments& InArgs);
 
 protected:
 
-	FText GetBindingText(URigVMPin* ModelPin) const;
-	FText GetBindingText() const;
-	const FSlateBrush* GetBindingImage() const;
-	FLinearColor GetBindingColor() const;
-	bool OnCanBindProperty(FProperty* InProperty) const;
-	bool OnCanBindToClass(UClass* InClass) const;
-	void OnAddBinding(FName InPropertyName, const TArray<FBindingChainElement>& InBindingChain);
-	bool OnCanRemoveBinding(FName InPropertyName);
-	void OnRemoveBinding(FName InPropertyName);
-	void FillLocalVariableMenu( FMenuBuilder& MenuBuilder );
-	void HandleBindToLocalVariable(FRigVMGraphVariableDescription InLocalVariable);
+	UE_API FText GetBindingText(URigVMPin* ModelPin) const;
+	UE_API FText GetBindingText() const;
+	UE_API const FSlateBrush* GetBindingImage() const;
+	UE_API FLinearColor GetBindingColor() const;
+	UE_API bool OnCanBindProperty(FProperty* InProperty) const;
+	UE_API bool OnCanBindToClass(UClass* InClass) const;
+	UE_API void OnAddBinding(FName InPropertyName, const TArray<FBindingChainElement>& InBindingChain);
+	UE_API bool OnCanRemoveBinding(FName InPropertyName);
+	UE_API void OnRemoveBinding(FName InPropertyName);
+	UE_API void FillLocalVariableMenu( FMenuBuilder& MenuBuilder );
+	UE_API void HandleBindToLocalVariable(FRigVMGraphVariableDescription InLocalVariable);
 
 	TArray<URigVMPin*> ModelPins;
 	URigVMFunctionReferenceNode* FunctionReferenceNode;
@@ -55,7 +57,7 @@ protected:
 	bool bCanRemoveBinding;
 };
 
-class RIGVMEDITOR_API SRigVMGraphPinVariableBinding : public SGraphPin
+class SRigVMGraphPinVariableBinding : public SGraphPin
 {
 public:
 
@@ -66,14 +68,16 @@ public:
 
 	SLATE_END_ARGS()
 
-	void Construct(const FArguments& InArgs, UEdGraphPin* InGraphPinObj);
+	UE_API void Construct(const FArguments& InArgs, UEdGraphPin* InGraphPinObj);
 
 protected:
 
 	//~ Begin SGraphPin Interface
-	virtual TSharedRef<SWidget>	GetDefaultValueWidget() override;
+	UE_API virtual TSharedRef<SWidget>	GetDefaultValueWidget() override;
 	//~ End SGraphPin Interface
 
 	TArray<URigVMPin*> ModelPins;
 	URigVMBlueprint* Blueprint;
 };
+
+#undef UE_API
diff --git a/Engine/Plugins/Runtime/RigVM/Source/RigVMEditor/Public/Widgets/SRigVMGraphPinVariableName.h b/Engine/Plugins/Runtime/RigVM/Source/RigVMEditor/Public/Widgets/SRigVMGraphPinVariableName.h
--- a/Engine/Plugins/Runtime/RigVM/Source/RigVMEditor/Public/Widgets/SRigVMGraphPinVariableName.h
+++ b/Engine/Plugins/Runtime/RigVM/Source/RigVMEditor/Public/Widgets/SRigVMGraphPinVariableName.h
@@ -8,28 +8,32 @@
 #include "SGraphPin.h"
 #include "Widgets/SRigVMGraphPinEditableNameValueWidget.h"
 
-class RIGVMEDITOR_API SRigVMGraphPinVariableName : public SGraphPin
+#define UE_API RIGVMEDITOR_API
+
+class SRigVMGraphPinVariableName : public SGraphPin
 {
 public:
 	SLATE_BEGIN_ARGS(SRigVMGraphPinVariableName) {}
 	SLATE_END_ARGS()
 
-		void Construct(const FArguments& InArgs, UEdGraphPin* InGraphPinObj);
+		UE_API void Construct(const FArguments& InArgs, UEdGraphPin* InGraphPinObj);
 
 protected:
 	//~ Begin SGraphPin Interface
-	virtual TSharedRef<SWidget>	GetDefaultValueWidget() override;
+	UE_API virtual TSharedRef<SWidget>	GetDefaultValueWidget() override;
 	//~ End SGraphPin Interface
 
-	FText GetVariableNameText() const;
-	virtual void SetVariableNameText(const FText& NewTypeInValue, ETextCommit::Type CommitInfo);
+	UE_API FText GetVariableNameText() const;
+	UE_API virtual void SetVariableNameText(const FText& NewTypeInValue, ETextCommit::Type CommitInfo);
 
-	TSharedRef<SWidget> MakeVariableNameItemWidget(TSharedPtr<FString> InItem);
-	void OnVariableNameChanged(TSharedPtr<FString> NewSelection, ESelectInfo::Type SelectInfo);
-	void OnVariableNameComboBox();
-	TArray<TSharedPtr<FString>>& GetVariableNames();
+	UE_API TSharedRef<SWidget> MakeVariableNameItemWidget(TSharedPtr<FString> InItem);
+	UE_API void OnVariableNameChanged(TSharedPtr<FString> NewSelection, ESelectInfo::Type SelectInfo);
+	UE_API void OnVariableNameComboBox();
+	UE_API TArray<TSharedPtr<FString>>& GetVariableNames();
 
 	TArray<TSharedPtr<FString>> VariableNames;
 	TSharedPtr<SRigVMGraphPinEditableNameValueWidget> NameComboBox;
 
 };
+
+#undef UE_API
diff --git a/Engine/Plugins/Runtime/RigVM/Source/RigVMEditor/Public/Widgets/SRigVMNodeLayoutWidget.h b/Engine/Plugins/Runtime/RigVM/Source/RigVMEditor/Public/Widgets/SRigVMNodeLayoutWidget.h
--- a/Engine/Plugins/Runtime/RigVM/Source/RigVMEditor/Public/Widgets/SRigVMNodeLayoutWidget.h
+++ b/Engine/Plugins/Runtime/RigVM/Source/RigVMEditor/Public/Widgets/SRigVMNodeLayoutWidget.h
@@ -8,6 +8,8 @@
 #include "RigVMModel/RigVMNode.h"
 #include "Input/DragAndDrop.h"
 
+#define UE_API RIGVMEDITOR_API
+
 DECLARE_DELEGATE_RetVal(TArray<FString> /* pin paths  */, FRigVMNodeLayoutWidget_OnGetUncategorizedPins);
 DECLARE_DELEGATE_RetVal(TArray<FRigVMPinCategory> /* categories */, FRigVMNodeLayoutWidget_OnGetCategories);
 DECLARE_DELEGATE_OneParam(FRigVMNodeLayoutWidget_OnCategoryAdded, FString /* category */);
@@ -24,7 +26,7 @@ DECLARE_DELEGATE_TwoParams(FRigVMNodeLayoutWidget_OnElementCategoryChanged, FStr
 DECLARE_DELEGATE_RetVal(uint32, FRigVMNodeLayoutWidget_OnGetStructuralHash);
 DECLARE_DELEGATE_RetVal_ThreeParams(bool, FRigVMNodeLayoutWidget_ValidateName, FString /* in path */, FString /* new name */, FText& /* out error */);
 
-class RIGVMEDITOR_API SRigVMNodeLayoutWidget : public SBox
+class SRigVMNodeLayoutWidget : public SBox
 {
 public:
 	
@@ -51,13 +53,13 @@ public:
 	SLATE_ATTRIBUTE(float, MaxScrollBoxSize)
 	SLATE_END_ARGS()
 
-	SRigVMNodeLayoutWidget();
-	virtual ~SRigVMNodeLayoutWidget() override;
+	UE_API SRigVMNodeLayoutWidget();
+	UE_API virtual ~SRigVMNodeLayoutWidget() override;
 	
-	void Construct(const FArguments& InArgs);
-	void Refresh();
-	virtual void Tick(const FGeometry& AllottedGeometry, const double InCurrentTime, const float InDeltaTime) override;
-	virtual FReply OnKeyDown(const FGeometry& MyGeometry, const FKeyEvent& InKeyEvent) override;
+	UE_API void Construct(const FArguments& InArgs);
+	UE_API void Refresh();
+	UE_API virtual void Tick(const FGeometry& AllottedGeometry, const double InCurrentTime, const float InDeltaTime) override;
+	UE_API virtual FReply OnKeyDown(const FGeometry& MyGeometry, const FKeyEvent& InKeyEvent) override;
 
 private:
 
@@ -139,7 +141,7 @@ private:
 	DECLARE_DELEGATE_RetVal_ThreeParams(TOptional<EItemDropZone>, FRigVMNodeLayoutWidget_OnElementCanDrop, TSharedPtr<FNodeLayoutRow>, TSharedPtr<FNodeLayoutRow>, EItemDropZone);
 	DECLARE_DELEGATE_RetVal_ThreeParams(FReply, FRigVMNodeLayoutWidget_OnElementAcceptDrop, TSharedPtr<FNodeLayoutRow>, TSharedPtr<FNodeLayoutRow>, EItemDropZone);
 
-	class RIGVMEDITOR_API SRigVMNowLayoutRow
+	class SRigVMNowLayoutRow
 	: public STableRow<TSharedPtr<FNodeLayoutRow>>
 	{
 	public:
@@ -165,13 +167,13 @@ private:
 
 		SLATE_END_ARGS()
 
-		virtual ~SRigVMNowLayoutRow() override;
-		void Construct(const FArguments& InArgs, const TSharedRef<STableViewBase>& OwnerTableView);
-		static TSharedPtr<SInlineEditableTextBlock> ConstructLabel(TSharedPtr<FNodeLayoutRow> InNodeLayoutRow, TSharedRef<SHorizontalBox> OutHorizontalBox, SRigVMNowLayoutRow* InRow);
+		UE_API virtual ~SRigVMNowLayoutRow() override;
+		UE_API void Construct(const FArguments& InArgs, const TSharedRef<STableViewBase>& OwnerTableView);
+		static UE_API TSharedPtr<SInlineEditableTextBlock> ConstructLabel(TSharedPtr<FNodeLayoutRow> InNodeLayoutRow, TSharedRef<SHorizontalBox> OutHorizontalBox, SRigVMNowLayoutRow* InRow);
 
 		TSharedPtr<FNodeLayoutRow> GetNodeLayoutRow() const { return NodeLayoutRow; }
-		void OnLabelCommitted(const FText& InText, ETextCommit::Type InCommitType) const;
-		bool OnVerifyLabelChanged(const FText& InText, FText& OutErrorMessage);
+		UE_API void OnLabelCommitted(const FText& InText, ETextCommit::Type InCommitType) const;
+		UE_API bool OnVerifyLabelChanged(const FText& InText, FText& OutErrorMessage);
 
 	private:
 		TSharedPtr<FNodeLayoutRow> NodeLayoutRow;
@@ -186,12 +188,12 @@ private:
 		TSharedPtr<SInlineEditableTextBlock> LabelEditWidget;
 	};
 
-	TSharedRef<ITableRow> GenerateRow(TSharedPtr<FNodeLayoutRow> InRow, const TSharedRef<STableViewBase>& OwnerTable);
-	void GetChildrenForRow(TSharedPtr<FNodeLayoutRow> InRow, TArray<TSharedPtr<FNodeLayoutRow>>& OutChildren);
-	void OnItemExpansionChanged(TSharedPtr<FNodeLayoutRow> InRow, bool bExpanded);
-	void OnItemSelectionChanged(TSharedPtr<FNodeLayoutRow> InRow, ESelectInfo::Type InSelectInfo);
-	bool IsNodeLayoutEditable() const;
-	FReply HandleAddCategory();
+	UE_API TSharedRef<ITableRow> GenerateRow(TSharedPtr<FNodeLayoutRow> InRow, const TSharedRef<STableViewBase>& OwnerTable);
+	UE_API void GetChildrenForRow(TSharedPtr<FNodeLayoutRow> InRow, TArray<TSharedPtr<FNodeLayoutRow>>& OutChildren);
+	UE_API void OnItemExpansionChanged(TSharedPtr<FNodeLayoutRow> InRow, bool bExpanded);
+	UE_API void OnItemSelectionChanged(TSharedPtr<FNodeLayoutRow> InRow, ESelectInfo::Type InSelectInfo);
+	UE_API bool IsNodeLayoutEditable() const;
+	UE_API FReply HandleAddCategory();
 
 	class FRigVMNodeLayoutDragDropOp : public FDragDropOperation
 	{
@@ -215,12 +217,14 @@ private:
 		TArray<TSharedPtr<FNodeLayoutRow>> NodeLayoutRows;
 	};
 
-	FReply OnDragDetectedForRow(TSharedPtr<FNodeLayoutRow> InSourceRow);
-	TOptional<EItemDropZone> OnCanAcceptDrop(const FDragDropEvent& InDragDropEvent, TSharedPtr<FNodeLayoutRow> InTargetRow, EItemDropZone InDropZone);
-	FReply OnAcceptDrop(const FDragDropEvent& InDragDropEvent, TSharedPtr<FNodeLayoutRow> InTargetRow, EItemDropZone InDropZone);
-	virtual FReply OnDrop(const FGeometry& InGeometry, const FDragDropEvent& InDragDropEvent);
+	UE_API FReply OnDragDetectedForRow(TSharedPtr<FNodeLayoutRow> InSourceRow);
+	UE_API TOptional<EItemDropZone> OnCanAcceptDrop(const FDragDropEvent& InDragDropEvent, TSharedPtr<FNodeLayoutRow> InTargetRow, EItemDropZone InDropZone);
+	UE_API FReply OnAcceptDrop(const FDragDropEvent& InDragDropEvent, TSharedPtr<FNodeLayoutRow> InTargetRow, EItemDropZone InDropZone);
+	UE_API virtual FReply OnDrop(const FGeometry& InGeometry, const FDragDropEvent& InDragDropEvent);
 	
 	TArray<TSharedPtr<FNodeLayoutRow>> NodeLayoutRows;
 	TSharedPtr<STreeView<TSharedPtr<FNodeLayoutRow>>> TreeView;
 	TOptional<uint32> LastStructuralHash;
-};
\ No newline at end of file
+};
+
+#undef UE_API
diff --git a/Engine/Plugins/Runtime/RigVM/Source/RigVMEditor/Public/Widgets/SRigVMNodePreviewWidget.h b/Engine/Plugins/Runtime/RigVM/Source/RigVMEditor/Public/Widgets/SRigVMNodePreviewWidget.h
--- a/Engine/Plugins/Runtime/RigVM/Source/RigVMEditor/Public/Widgets/SRigVMNodePreviewWidget.h
+++ b/Engine/Plugins/Runtime/RigVM/Source/RigVMEditor/Public/Widgets/SRigVMNodePreviewWidget.h
@@ -6,7 +6,9 @@
 #include "SRigVMGraphNode.h"
 #include "Editor/RigVMMinimalEnvironment.h"
 
-class RIGVMEDITOR_API SRigVMNodePreviewWidget : public SCompoundWidget
+#define UE_API RIGVMEDITOR_API
+
+class SRigVMNodePreviewWidget : public SCompoundWidget
 {
 public:
 	
@@ -17,12 +19,14 @@ public:
 	SLATE_ARGUMENT(TSharedPtr<FRigVMMinimalEnvironment>, Environment)
 	SLATE_END_ARGS()
 	
-	void Construct(const FArguments& InArgs);
+	UE_API void Construct(const FArguments& InArgs);
 
-	virtual void Tick(const FGeometry& AllottedGeometry, const double InCurrentTime, const float InDeltaTime) override;
+	UE_API virtual void Tick(const FGeometry& AllottedGeometry, const double InCurrentTime, const float InDeltaTime) override;
 
 private:
-	void UpdateNodeWidget();
+	UE_API void UpdateNodeWidget();
 	
 	TSharedPtr<FRigVMMinimalEnvironment> Environment;
 };
+
+#undef UE_API
diff --git a/Engine/Plugins/Runtime/RigVM/Source/RigVMEditor/Public/Widgets/SRigVMSwapAssetReferencesWidget.h b/Engine/Plugins/Runtime/RigVM/Source/RigVMEditor/Public/Widgets/SRigVMSwapAssetReferencesWidget.h
--- a/Engine/Plugins/Runtime/RigVM/Source/RigVMEditor/Public/Widgets/SRigVMSwapAssetReferencesWidget.h
+++ b/Engine/Plugins/Runtime/RigVM/Source/RigVMEditor/Public/Widgets/SRigVMSwapAssetReferencesWidget.h
@@ -6,6 +6,8 @@
 #include "Widgets/SRigVMBulkEditWidget.h"
 #include "RigVMEditorBlueprintLibrary.h"
 
+#define UE_API RIGVMEDITOR_API
+
 DECLARE_DELEGATE_RetVal_OneParam(TArray<FSoftObjectPath>, FOnGetReferences, FAssetData);
 DECLARE_DELEGATE_RetVal_TwoParams(bool, FOnSwapReference, FSoftObjectPath, FAssetData);
 
@@ -167,7 +169,7 @@ private:
 	FOnSwapReference SwapFunction;
 };
 
-class RIGVMEDITOR_API SRigVMSwapAssetReferencesWidget : public SCompoundWidget
+class SRigVMSwapAssetReferencesWidget : public SCompoundWidget
 {
 public:
 
@@ -188,7 +190,7 @@ public:
 	SLATE_EVENT(FOnSwapReference, OnSwapReference)
 	SLATE_END_ARGS()
 
-	void Construct(const FArguments& InArgs);
+	UE_API void Construct(const FArguments& InArgs);
 
 	TSharedRef<SRigVMBulkEditWidget> GetBulkEditWidget()
 	{
@@ -208,14 +210,16 @@ private:
 	FOnGetReferences OnGetReferences;
 	FOnSwapReference OnSwapReference;
 
-	TArray<TSharedRef<FRigVMTreeNode>> GetAssetNodes(const FArguments& InArgs, const int32& InPhase);
-	void OnPhaseActivated(TSharedRef<FRigVMTreePhase> Phase);
-	FReply OnNodeSelected(TSharedRef<FRigVMTreeNode> Node);
-	FReply OnNodeDoubleClicked(TSharedRef<FRigVMTreeNode> Node);
-	void SetSourceAsset(const FAssetData& InAsset);
-	void SetTargetAsset(const FAssetData& InAsset);
+	UE_API TArray<TSharedRef<FRigVMTreeNode>> GetAssetNodes(const FArguments& InArgs, const int32& InPhase);
+	UE_API void OnPhaseActivated(TSharedRef<FRigVMTreePhase> Phase);
+	UE_API FReply OnNodeSelected(TSharedRef<FRigVMTreeNode> Node);
+	UE_API FReply OnNodeDoubleClicked(TSharedRef<FRigVMTreeNode> Node);
+	UE_API void SetSourceAsset(const FAssetData& InAsset);
+	UE_API void SetTargetAsset(const FAssetData& InAsset);
 
 	static constexpr int32 PHASE_PICKSOURCE = 0;
 	static constexpr int32 PHASE_PICKTARGET = 1;
 	static constexpr int32 PHASE_PICKASSETREFS = 2;
 };
+
+#undef UE_API
diff --git a/Engine/Plugins/Runtime/RigVM/Source/RigVMEditor/Public/Widgets/SRigVMSwapFunctionsWidget.h b/Engine/Plugins/Runtime/RigVM/Source/RigVMEditor/Public/Widgets/SRigVMSwapFunctionsWidget.h
--- a/Engine/Plugins/Runtime/RigVM/Source/RigVMEditor/Public/Widgets/SRigVMSwapFunctionsWidget.h
+++ b/Engine/Plugins/Runtime/RigVM/Source/RigVMEditor/Public/Widgets/SRigVMSwapFunctionsWidget.h
@@ -11,6 +11,8 @@
 #include "Widgets/SRigVMBulkEditWidget.h"
 #include "Widgets/SRigVMNodePreviewWidget.h"
 
+#define UE_API RIGVMEDITOR_API
+
 class FRigVMSwapFunctionContext : public FRigVMTreeContext
 {
 public:
@@ -276,7 +278,7 @@ private:
 	FRigVMGraphFunctionIdentifier Identifier;
 };
 
-class RIGVMEDITOR_API SRigVMSwapFunctionsWidget : public SCompoundWidget
+class SRigVMSwapFunctionsWidget : public SCompoundWidget
 {
 public:
 
@@ -295,7 +297,7 @@ public:
 	SLATE_ARGUMENT(bool, CloseOnSuccess)
 	SLATE_END_ARGS()
 
-	void Construct(const FArguments& InArgs);
+	UE_API void Construct(const FArguments& InArgs);
 
 	TSharedRef<SRigVMBulkEditWidget> GetBulkEditWidget()
 	{
@@ -311,15 +313,17 @@ private:
 	TSharedPtr<FRigVMMinimalEnvironment> TargetPreviewEnvironment;
 	bool bSkipPickingFunctionRefs = false;;
 
-	static TArray<TSharedRef<FRigVMTreeNode>> GetFunctionIdentifierNodes(const FArguments& InArgs);
-	static TArray<TSharedRef<FRigVMTreeNode>> GetFunctionRefNodes(const FArguments& InArgs);
-	void OnPhaseActivated(TSharedRef<FRigVMTreePhase> Phase);
-	FReply OnNodeSelected(TSharedRef<FRigVMTreeNode> Node);
-	FReply OnNodeDoubleClicked(TSharedRef<FRigVMTreeNode> Node);
-	void SetSourceFunction(const FRigVMGraphFunctionIdentifier& InIdentifier);
-	void SetTargetFunction(const FRigVMGraphFunctionIdentifier& InIdentifier);
+	static UE_API TArray<TSharedRef<FRigVMTreeNode>> GetFunctionIdentifierNodes(const FArguments& InArgs);
+	static UE_API TArray<TSharedRef<FRigVMTreeNode>> GetFunctionRefNodes(const FArguments& InArgs);
+	UE_API void OnPhaseActivated(TSharedRef<FRigVMTreePhase> Phase);
+	UE_API FReply OnNodeSelected(TSharedRef<FRigVMTreeNode> Node);
+	UE_API FReply OnNodeDoubleClicked(TSharedRef<FRigVMTreeNode> Node);
+	UE_API void SetSourceFunction(const FRigVMGraphFunctionIdentifier& InIdentifier);
+	UE_API void SetTargetFunction(const FRigVMGraphFunctionIdentifier& InIdentifier);
 
 	static constexpr int32 PHASE_PICKSOURCE = 0;
 	static constexpr int32 PHASE_PICKTARGET = 1;
 	static constexpr int32 PHASE_PICKFUNCTIONREFS = 2;
 };
+
+#undef UE_API
diff --git a/Engine/Plugins/Runtime/RigVM/Source/RigVMEditor/Public/Widgets/SRigVMVariableMappingWidget.h b/Engine/Plugins/Runtime/RigVM/Source/RigVMEditor/Public/Widgets/SRigVMVariableMappingWidget.h
--- a/Engine/Plugins/Runtime/RigVM/Source/RigVMEditor/Public/Widgets/SRigVMVariableMappingWidget.h
+++ b/Engine/Plugins/Runtime/RigVM/Source/RigVMEditor/Public/Widgets/SRigVMVariableMappingWidget.h
@@ -9,11 +9,13 @@
 #include "Widgets/Views/STreeView.h"
 #include "Widgets/Input/SSearchBox.h"
 
+#define UE_API RIGVMEDITOR_API
+
 //=====================================================================================
 // UI for displaying property info
 //=====================================================================================
 
-struct RIGVMEDITOR_API FRigVMVariableMappingInfo
+struct FRigVMVariableMappingInfo
 {
 public:
 	// This is the property that is the most shallow type
@@ -77,7 +79,7 @@ DECLARE_DELEGATE_TwoParams(FOnRigVMPinCheckStateChanged, ECheckBoxState, FName /
 DECLARE_DELEGATE_RetVal_OneParam(ECheckBoxState, FOnRigVMPinGetCheckState, FName /** Property Name */);
 DECLARE_DELEGATE_RetVal_OneParam(bool, FOnRigVMPinIsCheckEnabled, FName /** Property Name */);
 
-class RIGVMEDITOR_API SRigVMVariableMappingTreeRow
+class SRigVMVariableMappingTreeRow
 	: public STableRow< TSharedPtr<FRigVMVariableMappingInfoPtr> >
 {
 public:
@@ -96,7 +98,7 @@ public:
 	SLATE_EVENT(FOnRigVMPinIsCheckEnabled, OnPinIsEnabledCheckState)
 	SLATE_END_ARGS()
 
-	void Construct(const FArguments& InArgs, const TSharedRef<STableViewBase>& OwnerTableView);
+	UE_API void Construct(const FArguments& InArgs, const TSharedRef<STableViewBase>& OwnerTableView);
 
 private:
 
@@ -108,20 +110,20 @@ private:
 	FRigVMVariableMappingInfoPtr	Item;
 
 	// Curve combo box options
-	FReply OnClearButtonClicked();
-	FText GetFilterText() const;
+	UE_API FReply OnClearButtonClicked();
+	UE_API FText GetFilterText() const;
 
-	TSharedRef<SWidget> MakeVarOptionComboWidget(TSharedPtr<FString> InItem);
-	void OnVarOptionSourceChanged(TSharedPtr<FString> NewSelection, ESelectInfo::Type SelectInfo);
-	FText GetVarOptionComboBoxContent() const;
-	FText GetVarOptionComboBoxToolTip() const;
-	void OnVarOptionComboOpening();
-	bool IsVarOptionEnabled() const;
-	TSharedPtr<FString> GetVarOptionString(FName VarOptionName) const;
+	UE_API TSharedRef<SWidget> MakeVarOptionComboWidget(TSharedPtr<FString> InItem);
+	UE_API void OnVarOptionSourceChanged(TSharedPtr<FString> NewSelection, ESelectInfo::Type SelectInfo);
+	UE_API FText GetVarOptionComboBoxContent() const;
+	UE_API FText GetVarOptionComboBoxToolTip() const;
+	UE_API void OnVarOptionComboOpening();
+	UE_API bool IsVarOptionEnabled() const;
+	UE_API TSharedPtr<FString> GetVarOptionString(FName VarOptionName) const;
 
-	ECheckBoxState IsPinChecked() const;
-	void OnPinCheckStatusChanged(ECheckBoxState NewState);
-	bool IsPinEnabled() const;
+	UE_API ECheckBoxState IsPinChecked() const;
+	UE_API void OnPinCheckStatusChanged(ECheckBoxState NewState);
+	UE_API bool IsPinEnabled() const;
 
 	FOnRigVMVariableMappingChanged	OnVariableMappingChanged;
 	FOnRigVMGetVariableMapping		OnGetVariableMapping;
@@ -135,7 +137,7 @@ private:
 
 typedef STreeView< TSharedPtr<FRigVMVariableMappingInfo> > SRigVMVariableMappingTreeView;
 
-class RIGVMEDITOR_API SRigVMVariableMappingWidget : public SCompoundWidget
+class SRigVMVariableMappingWidget : public SCompoundWidget
 {
 public:
 	SLATE_BEGIN_ARGS(SRigVMVariableMappingWidget)
@@ -158,7 +160,7 @@ public:
 	* @param InArgs - Arguments passed from Slate
 	*
 	*/
-	void Construct(const FArguments& InArgs/*, FSimpleMulticastDelegate& InOnPostUndo*/);
+	UE_API void Construct(const FArguments& InArgs/*, FSimpleMulticastDelegate& InOnPostUndo*/);
 
 	/**
 	* Filters the SListView when the user changes the search text box (NameFilterBox)
@@ -166,7 +168,7 @@ public:
 	* @param SearchText - The text the user has typed
 	*
 	*/
-	void OnFilterTextChanged(const FText& SearchText);
+	UE_API void OnFilterTextChanged(const FText& SearchText);
 
 	/**
 	* Filters the SListView when the user hits enter or clears the search box
@@ -176,7 +178,7 @@ public:
 	* @param CommitInfo - Not used
 	*
 	*/
-	void OnFilterTextCommitted(const FText& SearchText, ETextCommit::Type CommitInfo);
+	UE_API void OnFilterTextCommitted(const FText& SearchText, ETextCommit::Type CommitInfo);
 
 	/**
 	* Create a widget for an entry in the tree from an info
@@ -186,12 +188,12 @@ public:
 	*
 	* @return A new Slate widget, containing the UI for this row
 	*/
-	TSharedRef<ITableRow> GenerateVariableMappingRow(TSharedPtr<FRigVMVariableMappingInfo> InInfo, const TSharedRef<STableViewBase>& OwnerTable);
+	UE_API TSharedRef<ITableRow> GenerateVariableMappingRow(TSharedPtr<FRigVMVariableMappingInfo> InInfo, const TSharedRef<STableViewBase>& OwnerTable);
 
 	/**
 	* Handler for the delete of retarget source
 	*/
-	void RefreshVariableMappingList();
+	UE_API void RefreshVariableMappingList();
 
 private:
 
@@ -214,13 +216,13 @@ private:
 	FText FilterText;
 
 	/** Delegate for undo/redo transaction **/
-	void PostUndo();
+	UE_API void PostUndo();
 
 	// Make a single tree row widget
-	TSharedRef<ITableRow> MakeTreeRowWidget(TSharedPtr<FRigVMVariableMappingInfo> InInfo, const TSharedRef<STableViewBase>& OwnerTable);
+	UE_API TSharedRef<ITableRow> MakeTreeRowWidget(TSharedPtr<FRigVMVariableMappingInfo> InInfo, const TSharedRef<STableViewBase>& OwnerTable);
 
 	// Get the children for the provided bone info
-	void GetChildrenForInfo(TSharedPtr<FRigVMVariableMappingInfo> InInfo, TArray< TSharedPtr<FRigVMVariableMappingInfo> >& OutChildren);
+	UE_API void GetChildrenForInfo(TSharedPtr<FRigVMVariableMappingInfo> InInfo, TArray< TSharedPtr<FRigVMVariableMappingInfo> >& OutChildren);
 
 	FOnRigVMGetAvailableMapping		OnGetAvailableMappingDelegate;
 	FOnRigVMVariableMappingChanged	OnVariableMappingChangedDelegate;
@@ -233,3 +235,4 @@ private:
 
 };
 
+#undef UE_API
diff --git a/Engine/Plugins/Runtime/RigVM/Source/RigVMEditor/Public/Widgets/SRigVMVariantTagWidget.h b/Engine/Plugins/Runtime/RigVM/Source/RigVMEditor/Public/Widgets/SRigVMVariantTagWidget.h
--- a/Engine/Plugins/Runtime/RigVM/Source/RigVMEditor/Public/Widgets/SRigVMVariantTagWidget.h
+++ b/Engine/Plugins/Runtime/RigVM/Source/RigVMEditor/Public/Widgets/SRigVMVariantTagWidget.h
@@ -8,6 +8,8 @@
 #include "Widgets/SBoxPanel.h"
 #include "RigVMCore/RigVMVariant.h"
 
+#define UE_API RIGVMEDITOR_API
+
 DECLARE_DELEGATE_RetVal(TArray<FRigVMTag>, FRigVMVariant_OnGetTags);
 DECLARE_DELEGATE_OneParam(FRigVMVariant_OnAddTag, const FName&);
 DECLARE_DELEGATE_OneParam(FRigVMVariant_OnRemoveTag, const FName&);
@@ -55,7 +57,7 @@ private:
 	TAttribute<bool> EnableContextMenu;
 };
 
-class RIGVMEDITOR_API SRigVMVariantTagWidget : public SBox
+class SRigVMVariantTagWidget : public SBox
 {
 public:
 
@@ -79,18 +81,18 @@ public:
 	SLATE_ATTRIBUTE(const FSlateBrush*, CapsuleTagBorder)
 	SLATE_END_ARGS()
 
-	SRigVMVariantTagWidget();
-	virtual ~SRigVMVariantTagWidget() override;
+	UE_API SRigVMVariantTagWidget();
+	UE_API virtual ~SRigVMVariantTagWidget() override;
 	
-	void Construct(const FArguments& InArgs);
+	UE_API void Construct(const FArguments& InArgs);
 
-	virtual void Tick(const FGeometry& AllottedGeometry, const double InCurrentTime, const float InDeltaTime) override;
+	UE_API virtual void Tick(const FGeometry& AllottedGeometry, const double InCurrentTime, const float InDeltaTime) override;
 
 private:
 
-	void UpdateCapsules();
-	FReply OnAddTagClicked(const FName& InTagName) const;
-	TSharedRef<SWidget> OnBuildAddTagMenuContent() const;
+	UE_API void UpdateCapsules();
+	UE_API FReply OnAddTagClicked(const FName& InTagName) const;
+	UE_API TSharedRef<SWidget> OnBuildAddTagMenuContent() const;
 
 	TSharedPtr<SVerticalBox> VerticalCapsuleBox;
 	TSharedPtr<SHorizontalBox> HorizontalCapsuleBox;
@@ -105,3 +107,5 @@ private:
 	uint32 LastTagHash;
 	float MinDesiredLabelWidth;
 };
+
+#undef UE_API
diff --git a/Engine/Plugins/Runtime/RigVM/Source/RigVMEditor/Public/Widgets/SRigVMVariantWidget.h b/Engine/Plugins/Runtime/RigVM/Source/RigVMEditor/Public/Widgets/SRigVMVariantWidget.h
--- a/Engine/Plugins/Runtime/RigVM/Source/RigVMEditor/Public/Widgets/SRigVMVariantWidget.h
+++ b/Engine/Plugins/Runtime/RigVM/Source/RigVMEditor/Public/Widgets/SRigVMVariantWidget.h
@@ -10,13 +10,15 @@
 #include "RigVMCore/RigVMVariant.h"
 #include "Widgets/SRigVMVariantTagWidget.h"
 
+#define UE_API RIGVMEDITOR_API
+
 DECLARE_DELEGATE_OneParam(FRigVMVariantWidget_OnVariantChanged, const FRigVMVariant&);
 DECLARE_DELEGATE_RetVal_OneParam(TSharedPtr<SWidget>, FRigVMVariantWidget_OnCreateVariantRefRow, const FRigVMVariantRef&);
 DECLARE_DELEGATE_OneParam(FRigVMVariantWidget_OnBrowseVariantRef, const FRigVMVariantRef&);
 DECLARE_DELEGATE_RetVal_OneParam(TSharedPtr<SWidget>, FRigVMVariantWidget_OnVariantRefContextMenu, const FRigVMVariantRef&);
 DECLARE_DELEGATE_RetVal(TSharedPtr<SWidget>, FRigVMVariantWidget_OnContextMenu);
 
-struct RIGVMEDITOR_API FRigVMVariantWidgetContext
+struct FRigVMVariantWidgetContext
 {
 	FRigVMVariantWidgetContext()
 		: ParentPath()
@@ -64,7 +66,7 @@ private:
 	FRigVMVariantWidget_OnContextMenu OnContextMenu;
 };
 
-class RIGVMEDITOR_API SRigVMVariantWidget : public SBox
+class SRigVMVariantWidget : public SBox
 {
 public:
 	
@@ -91,12 +93,12 @@ public:
 	SLATE_ATTRIBUTE(bool, EnableTagContextMenu)
 	SLATE_END_ARGS()
 
-	SRigVMVariantWidget();
-	virtual ~SRigVMVariantWidget() override;
+	UE_API SRigVMVariantWidget();
+	UE_API virtual ~SRigVMVariantWidget() override;
 	
-	void Construct(const FArguments& InArgs);
+	UE_API void Construct(const FArguments& InArgs);
 
-	virtual void Tick(const FGeometry& AllottedGeometry, const double InCurrentTime, const float InDeltaTime) override;
+	UE_API virtual void Tick(const FGeometry& AllottedGeometry, const double InCurrentTime, const float InDeltaTime) override;
 
 private:
 
@@ -107,7 +109,7 @@ private:
 		TWeakPtr<ITableRow> RowWidget;
 	};
 
-	class RIGVMEDITOR_API SRigVMVariantRefTreeRow
+	class SRigVMVariantRefTreeRow
 	: public STableRow<TSharedPtr<FVariantTreeRowInfo>>
 	{
 	public:
@@ -118,19 +120,19 @@ private:
 		SLATE_EVENT(FRigVMVariantWidget_OnVariantRefContextMenu, OnVariantRefContextMenu)
 		SLATE_END_ARGS()
 
-		virtual ~SRigVMVariantRefTreeRow() override;
-		void Construct(const FArguments& InArgs, const TSharedRef<STableViewBase>& OwnerTableView);
+		UE_API virtual ~SRigVMVariantRefTreeRow() override;
+		UE_API void Construct(const FArguments& InArgs, const TSharedRef<STableViewBase>& OwnerTableView);
 	};
 
-	const FRigVMVariantWidgetContext& GetVariantContext() const;
-	EVisibility GetVariantRefListVisibility() const;
-	TSharedRef<ITableRow> GenerateVariantTreeRow(TSharedPtr<FVariantTreeRowInfo> InRowInfo, const TSharedRef<STableViewBase>& OwnerTable);
-	void GetChildrenForVariantInfo(TSharedPtr<FVariantTreeRowInfo> InInfo, TArray<TSharedPtr<FVariantTreeRowInfo>>& OutChildren);
-	TSharedPtr<SWidget> CreateDefaultVariantRefRow(const FRigVMVariantRef& InVariantRef) const;
-	void RebuildVariantRefList();
-	const FSlateBrush* GetThumbnailBorder(TSharedRef<SBorder> InThumbnailBorder) const;
-	TSharedPtr<SWidget> OnVariantRefTreeContextMenu();
-	TSharedPtr<SWidget> CreateDefaultVariantRefContextMenu(const FRigVMVariantRef& InVariantRef);
+	UE_API const FRigVMVariantWidgetContext& GetVariantContext() const;
+	UE_API EVisibility GetVariantRefListVisibility() const;
+	UE_API TSharedRef<ITableRow> GenerateVariantTreeRow(TSharedPtr<FVariantTreeRowInfo> InRowInfo, const TSharedRef<STableViewBase>& OwnerTable);
+	UE_API void GetChildrenForVariantInfo(TSharedPtr<FVariantTreeRowInfo> InInfo, TArray<TSharedPtr<FVariantTreeRowInfo>>& OutChildren);
+	UE_API TSharedPtr<SWidget> CreateDefaultVariantRefRow(const FRigVMVariantRef& InVariantRef) const;
+	UE_API void RebuildVariantRefList();
+	UE_API const FSlateBrush* GetThumbnailBorder(TSharedRef<SBorder> InThumbnailBorder) const;
+	UE_API TSharedPtr<SWidget> OnVariantRefTreeContextMenu();
+	UE_API TSharedPtr<SWidget> CreateDefaultVariantRefContextMenu(const FRigVMVariantRef& InVariantRef);
 
 	TAttribute<FRigVMVariant> VariantAttribute;
 	TAttribute<FRigVMVariantRef> SubjectVariantRefAttribute;
@@ -148,4 +150,6 @@ private:
 	TSharedPtr<SVerticalBox> VariantRefListBox;
 	TSharedPtr<STreeView<TSharedPtr<FVariantTreeRowInfo>>> VariantRefTreeView;
 	TAttribute<FRigVMVariantWidgetContext> ContextAttribute;
-};
\ No newline at end of file
+};
+
+#undef UE_API
diff --git a/Engine/Plugins/Runtime/ScriptableToolsFramework/Source/EditorScriptableToolsFramework/Public/BaseTools/EditorScriptableClickDragTool.h b/Engine/Plugins/Runtime/ScriptableToolsFramework/Source/EditorScriptableToolsFramework/Public/BaseTools/EditorScriptableClickDragTool.h
--- a/Engine/Plugins/Runtime/ScriptableToolsFramework/Source/EditorScriptableToolsFramework/Public/BaseTools/EditorScriptableClickDragTool.h
+++ b/Engine/Plugins/Runtime/ScriptableToolsFramework/Source/EditorScriptableToolsFramework/Public/BaseTools/EditorScriptableClickDragTool.h
@@ -8,10 +8,10 @@
 /**
  * Editor-Only variant of UScriptableClickDragTool, which gives access to Editor-Only BP functions
  */
-UCLASS(Transient, Blueprintable)
-class EDITORSCRIPTABLETOOLSFRAMEWORK_API UEditorScriptableClickDragTool : public UScriptableClickDragTool
+UCLASS(MinimalAPI, Transient, Blueprintable)
+class UEditorScriptableClickDragTool : public UScriptableClickDragTool
 {
 	GENERATED_BODY()
 public:
 
-};
\ No newline at end of file
+};
diff --git a/Engine/Plugins/Runtime/ScriptableToolsFramework/Source/EditorScriptableToolsFramework/Public/BaseTools/EditorScriptableModularBehaviorTool.h b/Engine/Plugins/Runtime/ScriptableToolsFramework/Source/EditorScriptableToolsFramework/Public/BaseTools/EditorScriptableModularBehaviorTool.h
--- a/Engine/Plugins/Runtime/ScriptableToolsFramework/Source/EditorScriptableToolsFramework/Public/BaseTools/EditorScriptableModularBehaviorTool.h
+++ b/Engine/Plugins/Runtime/ScriptableToolsFramework/Source/EditorScriptableToolsFramework/Public/BaseTools/EditorScriptableModularBehaviorTool.h
@@ -8,10 +8,10 @@
 /**
  * Editor-Only variant of UScriptableModularBehaviorTool, which gives access to Editor-Only BP functions
  */
-UCLASS(Transient, Blueprintable)
-class EDITORSCRIPTABLETOOLSFRAMEWORK_API UEditorScriptableModularBehaviorTool : public UScriptableModularBehaviorTool
+UCLASS(MinimalAPI, Transient, Blueprintable)
+class UEditorScriptableModularBehaviorTool : public UScriptableModularBehaviorTool
 {
 	GENERATED_BODY()
 public:
 
-};
\ No newline at end of file
+};
diff --git a/Engine/Plugins/Runtime/ScriptableToolsFramework/Source/EditorScriptableToolsFramework/Public/BaseTools/EditorScriptableSingleClickTool.h b/Engine/Plugins/Runtime/ScriptableToolsFramework/Source/EditorScriptableToolsFramework/Public/BaseTools/EditorScriptableSingleClickTool.h
--- a/Engine/Plugins/Runtime/ScriptableToolsFramework/Source/EditorScriptableToolsFramework/Public/BaseTools/EditorScriptableSingleClickTool.h
+++ b/Engine/Plugins/Runtime/ScriptableToolsFramework/Source/EditorScriptableToolsFramework/Public/BaseTools/EditorScriptableSingleClickTool.h
@@ -8,10 +8,10 @@
 /**
  * Editor-Only variant of UScriptableSingleClickTool, which gives access to Editor-Only BP functions
  */
-UCLASS(Transient, Blueprintable)
-class EDITORSCRIPTABLETOOLSFRAMEWORK_API UEditorScriptableSingleClickTool : public UScriptableSingleClickTool
+UCLASS(MinimalAPI, Transient, Blueprintable)
+class UEditorScriptableSingleClickTool : public UScriptableSingleClickTool
 {
 	GENERATED_BODY()
 public:
 
-};
\ No newline at end of file
+};
