RigVM: Localized registry in cooked environments - first steps.

diff --git a/Engine/Plugins/Experimental/AnimNext/Source/AnimNext/Private/Variables/RigVMDispatch_CallObjectAccessorFunction.cpp b/Engine/Plugins/Experimental/AnimNext/Source/AnimNext/Private/Variables/RigVMDispatch_CallObjectAccessorFunction.cpp
--- a/Engine/Plugins/Experimental/AnimNext/Source/AnimNext/Private/Variables/RigVMDispatch_CallObjectAccessorFunction.cpp
+++ b/Engine/Plugins/Experimental/AnimNext/Source/AnimNext/Private/Variables/RigVMDispatch_CallObjectAccessorFunction.cpp
@@ -9,7 +9,7 @@ FRigVMDispatch_CallObjectAccessorFunctionBase::FRigVMDispatch_CallObjectAccessor
 	FactoryScriptStruct = StaticStruct();
 }
 
-FName FRigVMDispatch_CallObjectAccessorFunctionBase::GetArgumentNameForOperandIndex(int32 InOperandIndex, int32 InTotalOperands) const
+FName FRigVMDispatch_CallObjectAccessorFunctionBase::GetArgumentNameForOperandIndex(int32 InOperandIndex, int32 InTotalOperands, FRigVMRegistryHandle& InRegistry) const
 {
 	static const FName ArgumentNames[] =
 	{
@@ -21,17 +21,17 @@ FName FRigVMDispatch_CallObjectAccessorFunctionBase::GetArgumentNameForOperandIn
 	return ArgumentNames[InOperandIndex];
 }
 
-void FRigVMDispatch_CallObjectAccessorFunctionBase::RegisterDependencyTypes_NoLock(FRigVMRegistry_NoLock& InRegistry) const
+void FRigVMDispatch_CallObjectAccessorFunctionBase::RegisterDependencyTypes_NoLock(FRigVMRegistryHandle& InRegistry) const
 {
 	static UScriptStruct* const AllowedStructTypes[] =
 	{
 		FAnimNextSoftFunctionPtr::StaticStruct(),
 	};
 
-	InRegistry.RegisterStructTypes_NoLock(AllowedStructTypes);
+	InRegistry->RegisterStructTypes_NoLock(AllowedStructTypes);
 	for(UScriptStruct* const ScriptStruct : AllowedStructTypes)
 	{
-		InRegistry.FindOrAddType_NoLock(FRigVMTemplateArgumentType(ScriptStruct));
+		InRegistry->FindOrAddType_NoLock(FRigVMTemplateArgumentType(ScriptStruct));
 	}
 
 	static TPair<UClass*, FRigVMRegistry::ERegisterObjectOperation> const AllowedObjectTypes[] =
@@ -39,18 +39,17 @@ void FRigVMDispatch_CallObjectAccessorFunctionBase::RegisterDependencyTypes_NoLo
 		{ UObject::StaticClass(), FRigVMRegistry::ERegisterObjectOperation::Class },
 	};
 
-	InRegistry.RegisterObjectTypes_NoLock(AllowedObjectTypes);
+	InRegistry->RegisterObjectTypes_NoLock(AllowedObjectTypes);
 
 	for(const TPair<UClass*, FRigVMRegistry::ERegisterObjectOperation>& ObjectType : AllowedObjectTypes)
 	{
-		InRegistry.FindOrAddType_NoLock(FRigVMTemplateArgumentType(ObjectType.Key));
+		InRegistry->FindOrAddType_NoLock(FRigVMTemplateArgumentType(ObjectType.Key));
 	}
 }
 
-const TArray<FRigVMTemplateArgumentInfo>& FRigVMDispatch_CallObjectAccessorFunctionBase::GetArgumentInfos() const
+const TArray<FRigVMTemplateArgumentInfo>& FRigVMDispatch_CallObjectAccessorFunctionBase::GetArgumentInfos(FRigVMRegistryHandle& InRegistry) const
 {
-	static TArray<FRigVMTemplateArgumentInfo> Infos;
-	if(Infos.IsEmpty())
+	if(CachedArgumentInfos.IsEmpty())
 	{
 		static const TArray<FRigVMTemplateArgument::ETypeCategory> ValueCategories =
 		{
@@ -58,22 +57,19 @@ const TArray<FRigVMTemplateArgumentInfo>& FRigVMDispatch_CallObjectAccessorFunct
 			FRigVMTemplateArgument::ETypeCategory_ArrayAnyValue
 		};
 
-		const FRigVMRegistry_NoLock& Registry = FRigVMRegistry_NoLock::GetForRead();
-		Infos.Emplace(ObjectName, ERigVMPinDirection::Input, Registry.GetTypeIndex_NoLock<UObject>());
-		Infos.Emplace(FunctionName, ERigVMPinDirection::Input, Registry.GetTypeIndex_NoLock<FAnimNextSoftFunctionPtr>());
-		Infos.Emplace(ValueName, ERigVMPinDirection::Output, ValueCategories);
+		CachedArgumentInfos.Emplace(ObjectName, ERigVMPinDirection::Input, InRegistry->GetTypeIndex_NoLock<UObject>());
+		CachedArgumentInfos.Emplace(FunctionName, ERigVMPinDirection::Input, InRegistry->GetTypeIndex_NoLock<FAnimNextSoftFunctionPtr>());
+		CachedArgumentInfos.Emplace(ValueName, ERigVMPinDirection::Output, ValueCategories);
 	}
 
-	return Infos;
+	return CachedArgumentInfos;
 }
 
-FRigVMTemplateTypeMap FRigVMDispatch_CallObjectAccessorFunctionBase::OnNewArgumentType(const FName& InArgumentName, TRigVMTypeIndex InTypeIndex) const
+FRigVMTemplateTypeMap FRigVMDispatch_CallObjectAccessorFunctionBase::OnNewArgumentType(const FName& InArgumentName, TRigVMTypeIndex InTypeIndex, FRigVMRegistryHandle& InRegistry) const
 {
-	const FRigVMRegistry_NoLock& Registry = FRigVMRegistry_NoLock::GetForRead();
-	
 	FRigVMTemplateTypeMap Types;
-	Types.Add(ObjectName, Registry.GetTypeIndex_NoLock<UObject>());
-	Types.Add(FunctionName, Registry.GetTypeIndex_NoLock<FAnimNextSoftFunctionPtr>());
+	Types.Add(ObjectName, InRegistry->GetTypeIndex_NoLock<UObject>());
+	Types.Add(FunctionName, InRegistry->GetTypeIndex_NoLock<FAnimNextSoftFunctionPtr>());
 	Types.Add(ValueName, InTypeIndex);
 	return Types;
 }
diff --git a/Engine/Plugins/Experimental/AnimNext/Source/AnimNext/Private/Variables/RigVMDispatch_GetObjectProperty.cpp b/Engine/Plugins/Experimental/AnimNext/Source/AnimNext/Private/Variables/RigVMDispatch_GetObjectProperty.cpp
--- a/Engine/Plugins/Experimental/AnimNext/Source/AnimNext/Private/Variables/RigVMDispatch_GetObjectProperty.cpp
+++ b/Engine/Plugins/Experimental/AnimNext/Source/AnimNext/Private/Variables/RigVMDispatch_GetObjectProperty.cpp
@@ -9,7 +9,7 @@ FRigVMDispatch_GetObjectProperty::FRigVMDispatch_GetObjectProperty()
 	FactoryScriptStruct = StaticStruct();
 }
 
-FName FRigVMDispatch_GetObjectProperty::GetArgumentNameForOperandIndex(int32 InOperandIndex, int32 InTotalOperands) const
+FName FRigVMDispatch_GetObjectProperty::GetArgumentNameForOperandIndex(int32 InOperandIndex, int32 InTotalOperands, FRigVMRegistryHandle& InRegistry) const
 {
 	static const FName ArgumentNames[] =
 	{
@@ -21,17 +21,17 @@ FName FRigVMDispatch_GetObjectProperty::GetArgumentNameForOperandIndex(int32 InO
 	return ArgumentNames[InOperandIndex];
 }
 
-void FRigVMDispatch_GetObjectProperty::RegisterDependencyTypes_NoLock(FRigVMRegistry_NoLock& InRegistry) const
+void FRigVMDispatch_GetObjectProperty::RegisterDependencyTypes_NoLock(FRigVMRegistryHandle& InRegistry) const
 {
 	static UScriptStruct* const AllowedStructTypes[] =
 	{
 		FAnimNextFieldPath::StaticStruct(),
 	};
 
-	InRegistry.RegisterStructTypes_NoLock(AllowedStructTypes);
+	InRegistry->RegisterStructTypes_NoLock(AllowedStructTypes);
 	for(UScriptStruct* const ScriptStruct : AllowedStructTypes)
 	{
-		InRegistry.FindOrAddType_NoLock(FRigVMTemplateArgumentType(ScriptStruct));
+		InRegistry->FindOrAddType_NoLock(FRigVMTemplateArgumentType(ScriptStruct));
 	}
 
 	static TPair<UClass*, FRigVMRegistry::ERegisterObjectOperation> const AllowedObjectTypes[] =
@@ -39,18 +39,17 @@ void FRigVMDispatch_GetObjectProperty::RegisterDependencyTypes_NoLock(FRigVMRegi
 		{ UObject::StaticClass(), FRigVMRegistry::ERegisterObjectOperation::Class },
 	};
 
-	InRegistry.RegisterObjectTypes_NoLock(AllowedObjectTypes);
+	InRegistry->RegisterObjectTypes_NoLock(AllowedObjectTypes);
 
 	for(const TPair<UClass*, FRigVMRegistry::ERegisterObjectOperation>& ObjectType : AllowedObjectTypes)
 	{
-		InRegistry.FindOrAddType_NoLock(FRigVMTemplateArgumentType(ObjectType.Key));
+		InRegistry->FindOrAddType_NoLock(FRigVMTemplateArgumentType(ObjectType.Key));
 	}
 }
 
-const TArray<FRigVMTemplateArgumentInfo>& FRigVMDispatch_GetObjectProperty::GetArgumentInfos() const
+const TArray<FRigVMTemplateArgumentInfo>& FRigVMDispatch_GetObjectProperty::GetArgumentInfos(FRigVMRegistryHandle& InRegistry) const
 {
-	static TArray<FRigVMTemplateArgumentInfo> Infos;
-	if(Infos.IsEmpty())
+	if(CachedArgumentInfos.IsEmpty())
 	{
 		static const TArray<FRigVMTemplateArgument::ETypeCategory> ValueCategories =
 		{
@@ -58,22 +57,19 @@ const TArray<FRigVMTemplateArgumentInfo>& FRigVMDispatch_GetObjectProperty::GetA
 			FRigVMTemplateArgument::ETypeCategory_ArrayAnyValue
 		};
 
-		const FRigVMRegistry_NoLock& Registry = FRigVMRegistry_NoLock::GetForRead();
-		Infos.Emplace(ObjectName, ERigVMPinDirection::Input, Registry.GetTypeIndex_NoLock<UObject>());
-		Infos.Emplace(PropertyName, ERigVMPinDirection::Input, Registry.GetTypeIndex_NoLock<FAnimNextFieldPath>());
-		Infos.Emplace(ValueName, ERigVMPinDirection::Output, ValueCategories);
+		CachedArgumentInfos.Emplace(ObjectName, ERigVMPinDirection::Input, InRegistry->GetTypeIndex_NoLock<UObject>());
+		CachedArgumentInfos.Emplace(PropertyName, ERigVMPinDirection::Input, InRegistry->GetTypeIndex_NoLock<FAnimNextFieldPath>());
+		CachedArgumentInfos.Emplace(ValueName, ERigVMPinDirection::Output, ValueCategories);
 	}
 
-	return Infos;
+	return CachedArgumentInfos;
 }
 
-FRigVMTemplateTypeMap FRigVMDispatch_GetObjectProperty::OnNewArgumentType(const FName& InArgumentName, TRigVMTypeIndex InTypeIndex) const
+FRigVMTemplateTypeMap FRigVMDispatch_GetObjectProperty::OnNewArgumentType(const FName& InArgumentName, TRigVMTypeIndex InTypeIndex, FRigVMRegistryHandle& InRegistry) const
 {
-	const FRigVMRegistry_NoLock& Registry = FRigVMRegistry_NoLock::GetForRead();
-	
 	FRigVMTemplateTypeMap Types;
-	Types.Add(ObjectName, Registry.GetTypeIndex_NoLock<UObject>());
-	Types.Add(PropertyName, Registry.GetTypeIndex_NoLock<FAnimNextFieldPath>());
+	Types.Add(ObjectName, InRegistry->GetTypeIndex_NoLock<UObject>());
+	Types.Add(PropertyName, InRegistry->GetTypeIndex_NoLock<FAnimNextFieldPath>());
 	Types.Add(ValueName, InTypeIndex);
 	return Types;
 }
diff --git a/Engine/Plugins/Experimental/AnimNextChooser/Source/AnimNextChooser/Private/RigUnit_EvaluateChooser.cpp b/Engine/Plugins/Experimental/AnimNextChooser/Source/AnimNextChooser/Private/RigUnit_EvaluateChooser.cpp
--- a/Engine/Plugins/Experimental/AnimNextChooser/Source/AnimNextChooser/Private/RigUnit_EvaluateChooser.cpp
+++ b/Engine/Plugins/Experimental/AnimNextChooser/Source/AnimNextChooser/Private/RigUnit_EvaluateChooser.cpp
@@ -56,7 +56,7 @@ FRigVMDispatch_EvaluateChooser::FRigVMDispatch_EvaluateChooser()
 	FactoryScriptStruct = StaticStruct();
 }
 
-FName FRigVMDispatch_EvaluateChooser::GetArgumentNameForOperandIndex(int32 InOperandIndex, int32 InTotalOperands) const
+FName FRigVMDispatch_EvaluateChooser::GetArgumentNameForOperandIndex(int32 InOperandIndex, int32 InTotalOperands, FRigVMRegistryHandle& InRegistry) const
 {
 	static const FName ArgumentNames[] =
 	{
@@ -68,10 +68,9 @@ FName FRigVMDispatch_EvaluateChooser::GetArgumentNameForOperandIndex(int32 InOpe
 	return ArgumentNames[InOperandIndex];
 }
 
-const TArray<FRigVMTemplateArgumentInfo>& FRigVMDispatch_EvaluateChooser::GetArgumentInfos() const
+const TArray<FRigVMTemplateArgumentInfo>& FRigVMDispatch_EvaluateChooser::GetArgumentInfos(FRigVMRegistryHandle& InRegistry) const
 {
-	static TArray<FRigVMTemplateArgumentInfo> Infos;
-	if(Infos.IsEmpty())
+	if(CachedArgumentInfos.IsEmpty())
 	{
 		static const TArray<FRigVMTemplateArgument::ETypeCategory> ResultCategories =
 		{
@@ -79,15 +78,14 @@ const TArray<FRigVMTemplateArgumentInfo>& FRigVMDispatch_EvaluateChooser::GetArg
 			FRigVMTemplateArgument::ETypeCategory_ArrayObjectValue
 		};
 		
-		const FRigVMRegistry_NoLock& Registry = FRigVMRegistry_NoLock::GetForRead();
-		const TRigVMTypeIndex ChooserTypeIndex = Registry.GetTypeIndex_NoLock<UChooserTable>();
-		const TRigVMTypeIndex ObjectTypeIndex = Registry.GetTypeIndex_NoLock<UObject>();
-		Infos.Emplace(ChooserName, ERigVMPinDirection::Input, ChooserTypeIndex);
-		Infos.Emplace(ContextObjectName, ERigVMPinDirection::Input, ObjectTypeIndex);
-		Infos.Emplace(ResultName, ERigVMPinDirection::Output, ResultCategories);
+		const TRigVMTypeIndex ChooserTypeIndex = InRegistry->GetTypeIndex_NoLock<UChooserTable>();
+		const TRigVMTypeIndex ObjectTypeIndex = InRegistry->GetTypeIndex_NoLock<UObject>();
+		CachedArgumentInfos.Emplace(ChooserName, ERigVMPinDirection::Input, ChooserTypeIndex);
+		CachedArgumentInfos.Emplace(ContextObjectName, ERigVMPinDirection::Input, ObjectTypeIndex);
+		CachedArgumentInfos.Emplace(ResultName, ERigVMPinDirection::Output, ResultCategories);
 	}
 
-	return Infos;
+	return CachedArgumentInfos;
 }
 
 void FRigVMDispatch_EvaluateChooser::Execute(FRigVMExtendedExecuteContext& InContext, FRigVMMemoryHandleArray Handles, FRigVMPredicateBranchArray RigVMBranches)
diff --git a/Engine/Plugins/Experimental/AnimNextChooser/Source/AnimNextChooser/Private/RigUnit_EvaluateChooser.h b/Engine/Plugins/Experimental/AnimNextChooser/Source/AnimNextChooser/Private/RigUnit_EvaluateChooser.h
--- a/Engine/Plugins/Experimental/AnimNextChooser/Source/AnimNextChooser/Private/RigUnit_EvaluateChooser.h
+++ b/Engine/Plugins/Experimental/AnimNextChooser/Source/AnimNextChooser/Private/RigUnit_EvaluateChooser.h
@@ -28,12 +28,12 @@ private:
 	friend struct UE::AnimNext::UncookedOnly::FUtils;
 
 	virtual UScriptStruct* GetExecuteContextStruct() const { return FRigVMExecuteContext::StaticStruct(); }
-	virtual FName GetArgumentNameForOperandIndex(int32 InOperandIndex, int32 InTotalOperands) const override;
-	virtual const TArray<FRigVMTemplateArgumentInfo>& GetArgumentInfos() const override;
+	virtual FName GetArgumentNameForOperandIndex(int32 InOperandIndex, int32 InTotalOperands, FRigVMRegistryHandle& InRegistry) const override;
+	virtual const TArray<FRigVMTemplateArgumentInfo>& GetArgumentInfos(FRigVMRegistryHandle& InRegistry) const override;
 
 	virtual bool IsSingleton() const override { return true; }
 
-	virtual FRigVMFunctionPtr GetDispatchFunctionImpl(const FRigVMTemplateTypeMap& InTypes) const override
+	virtual FRigVMFunctionPtr GetDispatchFunctionImpl(const FRigVMTemplateTypeMap& InTypes, const FRigVMRegistryHandle& InRegistry) const override
 	{
 		return &FRigVMDispatch_EvaluateChooser::Execute;
 	}
diff --git a/Engine/Plugins/Runtime/RigVM/Source/RigVM/Private/RigVMBlueprintGeneratedClass.cpp b/Engine/Plugins/Runtime/RigVM/Source/RigVM/Private/RigVMBlueprintGeneratedClass.cpp
--- a/Engine/Plugins/Runtime/RigVM/Source/RigVM/Private/RigVMBlueprintGeneratedClass.cpp
+++ b/Engine/Plugins/Runtime/RigVM/Source/RigVM/Private/RigVMBlueprintGeneratedClass.cpp
@@ -74,8 +74,9 @@ void URigVMBlueprintGeneratedClass::Serialize(FArchive& Ar)
 	{
 		if (const URigVMHost* CDO = Cast<URigVMHost>(GetDefaultObject(true)))
 		{
-			if (Ar.IsSaving() && CDO->VM)
+			if (Ar.IsSaving() && Ar.IsPersistent() && CDO->VM)
 			{
+				VM->CreateLocalizedRegistryIfRequired();
 				VM->CopyDataForSerialization(CDO->VM);
 			}
 		}
diff --git a/Engine/Plugins/Runtime/RigVM/Source/RigVM/Private/RigVMCore/RigVM.cpp b/Engine/Plugins/Runtime/RigVM/Source/RigVM/Private/RigVMCore/RigVM.cpp
--- a/Engine/Plugins/Runtime/RigVM/Source/RigVM/Private/RigVMCore/RigVM.cpp
+++ b/Engine/Plugins/Runtime/RigVM/Source/RigVM/Private/RigVMCore/RigVM.cpp
@@ -88,6 +88,14 @@ UScriptStruct* FRigVMParameter::GetScriptStruct() const
 	return ScriptStruct;
 }
 
+#define UE_RIGVM_GET_READ_REGISTRY_CONTEXT() \
+TUniquePtr<FRigVMRegistryReadLock> RegistryReadLock; \
+if(!LocalizedRegistry.IsValid()) \
+{ \
+	RegistryReadLock = MakeUnique<FRigVMRegistryReadLock>(); \
+} \
+FRigVMRegistryHandle& RegistryHandle = RegistryReadLock.IsValid() ? *RegistryReadLock.Get() : LocalizedRegistry->GetHandle_NoLock();
+
 URigVM::URigVM()
 	: ByteCodePtr(&ByteCodeStorage)
     , FunctionNamesPtr(&FunctionNamesStorage)
@@ -163,8 +171,40 @@ void URigVM::Save(FArchive& Ar)
 	UE_RIGVM_ARCHIVETRACE_ENTRY(Ar, TEXT("LiteralMemoryStorage"));
 	Ar << DefaultWorkMemoryStorage;
 	UE_RIGVM_ARCHIVETRACE_ENTRY(Ar, TEXT("DefaultWorkMemoryStorage"));
-	Ar << DefaultDebugMemoryStorage;
-	UE_RIGVM_ARCHIVETRACE_ENTRY(Ar, TEXT("DefaultDebugMemoryStorage"));
+
+	if(Ar.IsPersistent() && CVarRigVMEnableLocalizedRegistry.GetValueOnAnyThread() && LocalizedRegistry.IsValid())
+	{
+		FRigVMMemoryStorageStruct EmptyDebugMemoryStruct;
+		Ar << EmptyDebugMemoryStruct;
+		UE_RIGVM_ARCHIVETRACE_ENTRY(Ar, TEXT("DefaultDebugMemoryStorage"));
+
+		bool bStoredLocalizedRegistry = true;
+		Ar << bStoredLocalizedRegistry;
+
+		const int64 ArchivePosBeforeSerializedRegistry = Ar.Tell();
+
+		int64 ArchivePosAfterSerializedRegistry = 0;
+		Ar << ArchivePosAfterSerializedRegistry;
+
+		LocalizedRegistry->Serialize_NoLock(Ar);
+
+		// store the archive positions after the registry so we
+		// can skip it during load
+		ArchivePosAfterSerializedRegistry = Ar.Tell();
+		Ar.Seek(ArchivePosBeforeSerializedRegistry);
+		Ar << ArchivePosAfterSerializedRegistry;
+		Ar.Seek(ArchivePosAfterSerializedRegistry);
+		UE_RIGVM_ARCHIVETRACE_ENTRY(Ar, TEXT("LocalizedRegistry"));
+	}
+	else
+	{
+		Ar << DefaultDebugMemoryStorage;
+		UE_RIGVM_ARCHIVETRACE_ENTRY(Ar, TEXT("DefaultDebugMemoryStorage"));
+
+		bool bStoredLocalizedRegistry = false;
+		Ar << bStoredLocalizedRegistry;
+		UE_RIGVM_ARCHIVETRACE_ENTRY(Ar, TEXT("LocalizedRegistry"));
+	}
 }
 
 void URigVM::Load(FArchive& Ar)
@@ -280,6 +320,30 @@ void URigVM::Load(FArchive& Ar)
 		Ar << DefaultWorkMemoryStorage;
 		Ar << DefaultDebugMemoryStorage;
 	}
+
+	if (Ar.CustomVer(FRigVMObjectVersion::GUID) >= FRigVMObjectVersion::LocalizedRegistry)
+	{
+		bool bStoredLocalizedRegistry = false;
+		Ar << bStoredLocalizedRegistry;
+
+		LocalizedRegistry.Reset();
+		if(bStoredLocalizedRegistry)
+		{
+			int64 ArchivePosAfterSerializedRegistry = 0;
+			Ar << ArchivePosAfterSerializedRegistry;
+
+			if(CVarRigVMEnableLocalizedRegistry.GetValueOnAnyThread() == false)
+			{
+				UE_LOG(LogRigVM, Warning, TEXT("VM '%s' was serialized with a localized registry, but CVarRigVMEnableLocalizedRegistry is disabled."), *GetPathName());
+				Ar.Seek(ArchivePosAfterSerializedRegistry);
+			}
+			else
+			{
+				LocalizedRegistry = FRigVMRegistry_NoLock::CreateLocalizedRegistry();
+				LocalizedRegistry->Serialize_NoLock(Ar);
+			}
+		}
+	}
 }
 
 void URigVM::CopyDataForSerialization(URigVM* InVM)
@@ -300,6 +364,8 @@ void URigVM::CopyDataForSerialization(URigVM* InVM)
 	LiteralMemoryStorage = InVM->LiteralMemoryStorage;
 	DefaultWorkMemoryStorage = InVM->DefaultWorkMemoryStorage;
 	DefaultDebugMemoryStorage = InVM->DefaultDebugMemoryStorage;
+
+	LocalizedRegistry = InVM->LocalizedRegistry;
 }
 
 void URigVM::PostLoad()
@@ -342,6 +408,18 @@ void URigVM::RefreshArgumentNameCaches()
 	// the argument name caches are being updated.
 	FScopeLock ResolveFunctionsScopeLock(&ResolveFunctionsMutex);
 
+	int32 TotalNumFactories;
+	if(LocalizedRegistry.IsValid())
+	{
+		TotalNumFactories = LocalizedRegistry->GetFactories_NoLock().Num();
+	}
+	else
+	{
+		TotalNumFactories = FRigVMRegistry::Get().GetFactories().Num();
+	}
+
+	TArray<const FRigVMDispatchFactory*> Factories;
+
 	// make sure to update all argument name caches
 	const TArray<const FRigVMFunction*>& Functions = GetFunctions();
 	for (const FRigVMInstruction& Instruction : Instructions)
@@ -353,8 +431,26 @@ void URigVM::RefreshArgumentNameCaches()
 			{
 				if (const FRigVMDispatchFactory* Factory = Functions[Op.FunctionIndex]->Factory)
 				{
-					FRigVMOperandArray Operands = ByteCodeStorage.GetOperandsForExecuteOp(Instruction);
-					(void)Factory->UpdateArgumentNameCache(Operands.Num());
+					if(Factories.IsEmpty())
+					{
+						Factories.AddZeroed(TotalNumFactories);
+					}
+					
+					if(Factories[Factory->GetFactoryIndex()] == nullptr)
+					{
+						Factories[Factory->GetFactoryIndex()] = Factory;
+						
+						FRigVMOperandArray Operands = ByteCodeStorage.GetOperandsForExecuteOp(Instruction);
+						if(LocalizedRegistry.IsValid())
+						{
+							(void)Factory->UpdateArgumentNameCache_NoLock(Operands.Num(), LocalizedRegistry->GetHandle_NoLock());
+						}
+						else
+						{
+							FRigVMRegistryReadLock ReadLock;
+							(void)Factory->UpdateArgumentNameCache(Operands.Num(), ReadLock);
+						}
+					}
 				}
 			}
 		}
@@ -622,6 +718,7 @@ void URigVM::Reset_Internal()
 	LazyBranches.Reset();
 
 	InvalidateCachedMemory_Internal();
+	LocalizedRegistry.Reset();
 }
 
 void URigVM::Reset(FRigVMExtendedExecuteContext& Context)
@@ -667,7 +764,15 @@ int32 URigVM::AddRigVMFunction(const FString& InFunctionName)
 		return FunctionIndex;
 	}
 
-	const FRigVMFunction* Function = FRigVMRegistry_RWLock::Get().FindFunction(*InFunctionName);
+	const FRigVMFunction* Function = nullptr;
+	if(LocalizedRegistry.IsValid())
+	{
+		Function = LocalizedRegistry.Get()->FindFunction_NoLock(*InFunctionName);
+	}
+	else
+	{
+		Function = FRigVMRegistry_RWLock::Get().FindFunction(*InFunctionName);
+	}
 	if (Function == nullptr)
 	{
 		return INDEX_NONE;
@@ -984,7 +1089,15 @@ FRigVMParameter URigVM::GetParameterByName(const FName& InParameterName)
 
 bool URigVM::ResolveFunctionsIfRequired()
 {
-	FScopeLock ResolveFunctionsScopeLock(&ResolveFunctionsMutex);
+	// only create the function resolval lock if we are running
+	// on the global registry.
+#if WITH_EDITOR
+	TUniquePtr<FScopeLock> ResolveFunctionsScopeLock;
+	if(!LocalizedRegistry.IsValid())
+	{
+		ResolveFunctionsScopeLock = MakeUnique<FScopeLock>(&ResolveFunctionsMutex);
+	}
+#endif
 	
 	bool bSuccess = true;
 	
@@ -1004,11 +1117,13 @@ bool URigVM::ResolveFunctionsIfRequired()
 			});
 		}
 
+		UE_RIGVM_GET_READ_REGISTRY_CONTEXT();
+		
 		TArray<FName>& FunctionNames = GetFunctionNames();
 		for (int32 FunctionIndex = 0; FunctionIndex < FunctionNames.Num(); FunctionIndex++)
 		{
 			const FString FunctionNameString = FunctionNames[FunctionIndex].ToString();
-			if(const FRigVMFunction* Function = FRigVMRegistry_RWLock::Get().FindFunction(*FunctionNameString, TypeResolver))
+			if(const FRigVMFunction* Function = RegistryHandle->FindFunction_NoLock(*FunctionNameString, TypeResolver))
 			{
 				GetFunctions()[FunctionIndex] = Function;
 				GetFactories()[FunctionIndex] = Function->Factory;
@@ -1059,7 +1174,18 @@ void URigVM::InvalidateCachedMemory(FRigVMExtendedExecuteContext& Context)
 	Context.InvalidateCachedMemory();
 }
 
-void URigVM::InstructionOpEval(FRigVMExtendedExecuteContext& Context, int32 InstructionIndex, int32 InHandleBaseIndex, const TFunctionRef<void(FRigVMExtendedExecuteContext& Context, int32 InHandleIndex, const FRigVMBranchInfoKey& InBranchInfoKey, const FRigVMOperand& InArg)>& InOpFunc)
+void URigVM::InstructionOpEval(
+	FRigVMExtendedExecuteContext& Context,
+	FRigVMRegistryHandle& RegistryHandle,
+	int32 InstructionIndex,
+	int32 InHandleBaseIndex,
+	const TFunctionRef<void(
+		FRigVMExtendedExecuteContext& /*Context*/,
+		FRigVMRegistryHandle& /*RegistryHandle*/,
+		int32 /*InHandleIndex*/,
+		const FRigVMBranchInfoKey& /*InBranchInfoKey*/,
+		const FRigVMOperand& /*InArg*/
+	)>& InOpFunc)
 {
 	const FRigVMByteCode& ByteCode = GetByteCode();
 	const TArray<const FRigVMFunction*>& Functions = GetFunctions();
@@ -1071,13 +1197,18 @@ void URigVM::InstructionOpEval(FRigVMExtendedExecuteContext& Context, int32 Inst
 		const FRigVMExecuteOp& Op = ByteCode.GetOpAt<FRigVMExecuteOp>(Instructions[InstructionIndex]);
 		FRigVMOperandArray Operands = ByteCode.GetOperandsForExecuteOp(Instructions[InstructionIndex]);
 		const FRigVMFunction* Function = Functions[Op.FunctionIndex];
+		if(Function == nullptr)
+		{
+			return;
+		}
 
 		for (int32 ArgIndex = 0; ArgIndex < Operands.Num(); ArgIndex++)
 		{
 			InOpFunc(
 				Context,
+				RegistryHandle,
 				InHandleBaseIndex++,
-				{ InstructionIndex, ArgIndex, Function->GetArgumentNameForOperandIndex(ArgIndex, Operands.Num()) },
+				{ InstructionIndex, ArgIndex, Function->GetArgumentNameForOperandIndex(ArgIndex, Operands.Num(), RegistryHandle) },
 				Operands[ArgIndex]);
 		}
 	}
@@ -1092,14 +1223,14 @@ void URigVM::InstructionOpEval(FRigVMExtendedExecuteContext& Context, int32 Inst
 		case ERigVMOpCode::Decrement:
 		{
 			const FRigVMUnaryOp& Op = ByteCode.GetOpAt<FRigVMUnaryOp>(Instructions[InstructionIndex]);
-			InOpFunc(Context, InHandleBaseIndex, {}, Op.Arg);
+			InOpFunc(Context, RegistryHandle, InHandleBaseIndex, {}, Op.Arg);
 			break;
 		}
 		case ERigVMOpCode::Copy:
 		{
 			const FRigVMCopyOp& Op = ByteCode.GetOpAt<FRigVMCopyOp>(Instructions[InstructionIndex]);
-			InOpFunc(Context, InHandleBaseIndex + 0, {}, Op.Source);
-			InOpFunc(Context, InHandleBaseIndex + 1, {}, Op.Target);
+			InOpFunc(Context, RegistryHandle, InHandleBaseIndex + 0, {}, Op.Source);
+			InOpFunc(Context, RegistryHandle, InHandleBaseIndex + 1, {}, Op.Target);
 			break;
 		}
 		case ERigVMOpCode::Equals:
@@ -1107,11 +1238,11 @@ void URigVM::InstructionOpEval(FRigVMExtendedExecuteContext& Context, int32 Inst
 		{
 			const FRigVMComparisonOp& Op = ByteCode.GetOpAt<FRigVMComparisonOp>(Instructions[InstructionIndex]);
 			FRigVMOperand Arg = Op.A;
-			InOpFunc(Context, InHandleBaseIndex + 0, {}, Arg);
+			InOpFunc(Context, RegistryHandle, InHandleBaseIndex + 0, {}, Arg);
 			Arg = Op.B;
-			InOpFunc(Context, InHandleBaseIndex + 1, {}, Arg);
+			InOpFunc(Context, RegistryHandle, InHandleBaseIndex + 1, {}, Arg);
 			Arg = Op.Result;
-			InOpFunc(Context, InHandleBaseIndex + 2, {}, Arg);
+			InOpFunc(Context, RegistryHandle, InHandleBaseIndex + 2, {}, Arg);
 			break;
 		}
 		case ERigVMOpCode::JumpAbsolute:
@@ -1126,14 +1257,14 @@ void URigVM::InstructionOpEval(FRigVMExtendedExecuteContext& Context, int32 Inst
 		{
 			const FRigVMJumpIfOp& Op = ByteCode.GetOpAt<FRigVMJumpIfOp>(Instructions[InstructionIndex]);
 			const FRigVMOperand& Arg = Op.Arg;
-			InOpFunc(Context, InHandleBaseIndex, {}, Arg);
+			InOpFunc(Context, RegistryHandle, InHandleBaseIndex, {}, Arg);
 			break;
 		}
 		case ERigVMOpCode::ChangeType:
 		{
 			const FRigVMChangeTypeOp& Op = ByteCode.GetOpAt<FRigVMChangeTypeOp>(Instructions[InstructionIndex]);
 			const FRigVMOperand& Arg = Op.Arg;
-			InOpFunc(Context, InHandleBaseIndex, {}, Arg);
+			InOpFunc(Context, RegistryHandle, InHandleBaseIndex, {}, Arg);
 			break;
 		}
 		case ERigVMOpCode::Exit:
@@ -1143,8 +1274,8 @@ void URigVM::InstructionOpEval(FRigVMExtendedExecuteContext& Context, int32 Inst
 		case ERigVMOpCode::BeginBlock:
 		{
 			const FRigVMBinaryOp& Op = ByteCode.GetOpAt<FRigVMBinaryOp>(Instructions[InstructionIndex]);
-			InOpFunc(Context, InHandleBaseIndex + 0, {}, Op.ArgA);
-			InOpFunc(Context, InHandleBaseIndex + 1, {}, Op.ArgB);
+			InOpFunc(Context, RegistryHandle, InHandleBaseIndex + 0, {}, Op.ArgA);
+			InOpFunc(Context, RegistryHandle, InHandleBaseIndex + 1, {}, Op.ArgB);
 			break;
 		}
 		case ERigVMOpCode::EndBlock:
@@ -1159,21 +1290,21 @@ void URigVM::InstructionOpEval(FRigVMExtendedExecuteContext& Context, int32 Inst
 		{
 			const FRigVMJumpToBranchOp& Op = ByteCode.GetOpAt<FRigVMJumpToBranchOp>(Instructions[InstructionIndex]);
 			const FRigVMOperand& Arg = Op.Arg;
-			InOpFunc(Context, InHandleBaseIndex, {}, Arg);
+			InOpFunc(Context, RegistryHandle, InHandleBaseIndex, {}, Arg);
 			break;
 		}
 		case ERigVMOpCode::RunInstructions:
 		{
 			const FRigVMRunInstructionsOp& Op = ByteCode.GetOpAt<FRigVMRunInstructionsOp>(Instructions[InstructionIndex]);
 			const FRigVMOperand& Arg = Op.Arg;
-			InOpFunc(Context, InHandleBaseIndex, {}, Arg);
+			InOpFunc(Context, RegistryHandle, InHandleBaseIndex, {}, Arg);
 			break;
 		}
 		case ERigVMOpCode::SetupTraits:
 		{
 			const FRigVMSetupTraitsOp& Op = ByteCode.GetOpAt<FRigVMSetupTraitsOp>(Instructions[InstructionIndex]);
 			const FRigVMOperand& Arg = Op.Arg;
-			InOpFunc(Context, InHandleBaseIndex, {}, Arg);
+			InOpFunc(Context, RegistryHandle, InHandleBaseIndex, {}, Arg);
 			break;
 		}
 		case ERigVMOpCode::Invalid:
@@ -1211,13 +1342,20 @@ void URigVM::PrepareMemoryForExecution(FRigVMExtendedExecuteContext& Context)
 	// Make sure we have enough room to prevent repeated allocations.
 	FirstHandleForInstruction.Reset(Instructions.Num() + 1);
 
+	TUniquePtr<FRigVMRegistryReadLock> RegistryReadLock;
+	if(!LocalizedRegistry.IsValid())
+	{
+		RegistryReadLock = MakeUnique<FRigVMRegistryReadLock>();
+	}
+	FRigVMRegistryHandle& RegistryHandle = RegistryReadLock.IsValid() ? *RegistryReadLock.Get() : LocalizedRegistry->GetHandle_NoLock();
+
 	// Count how many handles we need and set up the indirection offsets for the handles.
 	int32 HandleCount = 0;
 	FirstHandleForInstruction.Add(0);
 	for (int32 InstructionIndex = 0; InstructionIndex < Instructions.Num(); InstructionIndex++)
 	{
-		InstructionOpEval(Context, InstructionIndex, INDEX_NONE,
-			[&HandleCount](FRigVMExtendedExecuteContext& Context, int32, const FRigVMBranchInfoKey&, const FRigVMOperand&)
+		InstructionOpEval(Context, RegistryHandle, InstructionIndex, INDEX_NONE,
+			[&HandleCount](FRigVMExtendedExecuteContext&, const FRigVMRegistryHandle&, int32, const FRigVMBranchInfoKey&, const FRigVMOperand&)
 			{
 				HandleCount++;
 			});
@@ -1244,15 +1382,17 @@ void URigVM::CacheMemoryHandlesIfRequired(FRigVMExtendedExecuteContext& Context)
 
 	// Allocate all the space and zero it out to ensure all pages required for it are paged in immediately.
 	Context.CachedMemoryHandles.SetNumUninitialized(MemoryHandleCount);
-	
+
+	UE_RIGVM_GET_READ_REGISTRY_CONTEXT();
+
 	// Now cache the handles as needed, in several batches as needed
 	const int32 TaskBatchSize = 10;
-	ParallelFor(TEXT("CacheMemoryHandlesIfRequired"), Instructions.Num(), TaskBatchSize, [this, &Context](int32 InstructionIndex)
+	ParallelFor(TEXT("CacheMemoryHandlesIfRequired"), Instructions.Num(), TaskBatchSize, [this, &Context, &RegistryHandle](int32 InstructionIndex)
 		{
-			InstructionOpEval(Context, InstructionIndex, FirstHandleForInstruction[InstructionIndex],
-				[&](FRigVMExtendedExecuteContext& Context, int32 InHandleIndex, const FRigVMBranchInfoKey& InBranchInfoKey, const FRigVMOperand& InOp)
+			InstructionOpEval(Context, RegistryHandle, InstructionIndex, FirstHandleForInstruction[InstructionIndex],
+				[&](FRigVMExtendedExecuteContext& InContext, const FRigVMRegistryHandle& InRegistry, int32 InHandleIndex, const FRigVMBranchInfoKey& InBranchInfoKey, const FRigVMOperand& InOp)
 				{
-					CacheSingleMemoryHandle(Context, InHandleIndex, InBranchInfoKey, InOp);
+					CacheSingleMemoryHandle(InContext, InRegistry, InHandleIndex, InBranchInfoKey, InOp);
 				});
 		}
 	);
@@ -1494,8 +1634,8 @@ bool URigVM::Initialize(FRigVMExtendedExecuteContext& Context)
 	
 	TGuardValue<int32> GuardThreadId(Context.ExecutingThreadId, FPlatformTLS::GetCurrentThreadId());
 
-	ResolveFunctionsIfRequired();
 	RefreshInstructionsIfRequired();
+	ResolveFunctionsIfRequired();
 
 	if (Instructions.Num() == 0)
 	{
@@ -2361,6 +2501,212 @@ void URigVM::SetPropertyValueFromString(FRigVMExtendedExecuteContext& Context, c
 	Memory->SetDataFromString(InOperand.GetRegisterIndex(), InDefaultValue);
 }
 
+FRigVMRegistry_NoLock* URigVM::GetLocalizedRegistry() const
+{
+	if(LocalizedRegistry.IsValid())
+	{
+		return LocalizedRegistry.Get();
+	}
+	return nullptr;
+}
+
+bool URigVM::CreateLocalizedRegistryIfRequired(bool bForce)
+{
+	if(LocalizedRegistry.IsValid())
+	{
+		return true;
+	}
+
+	if(!bForce && CVarRigVMEnableLocalizedRegistry.GetValueOnAnyThread() == false)
+	{
+		return false;
+	}
+		
+	// create a new registry
+	TSharedPtr<FRigVMRegistry_NoLock> TemporaryRegistryPtr = FRigVMRegistry_NoLock::CreateLocalizedRegistry();
+	FRigVMRegistry_NoLock& TemporaryRegistry = *TemporaryRegistryPtr.Get();
+	FRigVMRegistryHandle TemporaryContext(&TemporaryRegistry);
+
+	RefreshInstructionsIfRequired();
+	if(ExternalPropertyPaths.IsEmpty())
+	{
+		RefreshExternalPropertyPaths();
+	}
+	ResolveFunctionsIfRequired();
+
+	const FRigVMByteCode& ByteCode = GetByteCode();
+
+	// add all required types to it
+	for (int32 InstructionIndex = 0; InstructionIndex < Instructions.Num(); InstructionIndex++)
+	{
+		FRigVMOperandArray Operands = ByteCode.GetOperandsForOp(Instructions[InstructionIndex]);
+		for(const FRigVMOperand& Operand : Operands)
+		{
+			const FProperty* Property = nullptr;
+			switch(Operand.GetMemoryType())
+			{
+				case ERigVMMemoryType::Literal:
+				{
+					if(Operand.GetRegisterOffset() != INDEX_NONE)
+					{
+						Property = LiteralMemoryStorage.GetPropertyPaths()[Operand.GetRegisterOffset()].GetTailProperty();
+					}
+					else
+					{
+						Property = LiteralMemoryStorage.GetProperty(Operand.GetRegisterIndex());
+					}
+					break;
+				}
+				case ERigVMMemoryType::Work:
+				{
+					if(Operand.GetRegisterOffset() != INDEX_NONE)
+					{
+						Property = DefaultWorkMemoryStorage.GetPropertyPaths()[Operand.GetRegisterOffset()].GetTailProperty();
+					}
+					else
+					{
+						Property = DefaultWorkMemoryStorage.GetProperty(Operand.GetRegisterIndex());
+					}
+					break;
+				}
+				case ERigVMMemoryType::External:
+				{
+					if(Operand.GetRegisterOffset() != INDEX_NONE)
+					{
+						if(!ExternalPropertyPaths.IsValidIndex(Operand.GetRegisterOffset()))
+						{
+							RefreshExternalPropertyPaths();
+						}
+						Property = ExternalPropertyPaths[Operand.GetRegisterOffset()].GetTailProperty();
+					}
+					else
+					{
+						Property = ExternalVariables[Operand.GetRegisterIndex()].Property;
+					}
+					break;
+				}
+				case ERigVMMemoryType::Debug:
+				default:
+				{
+					break;
+				}
+			}
+
+			if(Property == nullptr)
+			{
+				continue;
+			}
+
+			// just creating the argument will register the type as needed
+			FRigVMTemplateArgument::Make_NoLock(const_cast<FProperty*>(Property), TemporaryContext);
+		}
+	}
+
+	// for dispatch functions we need to register all of the argument types for the given permutations
+	// otherwise we won't be able to retrieve the permutation from string again later.
+	for(const FRigVMFunction* Function : FunctionsStorage)
+	{
+		if(Function && Function->Factory && Function->OwnerRegistry)
+		{
+			for(const TRigVMTypeIndex& ArgumentTypeIndex :  Function->GetArgumentTypeIndices())
+			{
+				const FRigVMTemplateArgumentType& Type = Function->OwnerRegistry->GetType_NoLock(ArgumentTypeIndex);
+				TemporaryRegistry.FindOrAddType_NoLock(Type);
+			}
+		}
+	}
+
+	// register all of the factories
+	TMap<const FRigVMDispatchFactory*, TArray<const FRigVMFunction*>> RequiredFactories;
+	for(const FRigVMFunction* Function : FunctionsStorage)
+	{
+		if(Function == nullptr)
+		{
+			continue;
+		}
+		if(Function->Factory)
+		{
+			RequiredFactories.FindOrAdd(Function->Factory).Add(Function);
+		}
+	}
+	for(const TPair<const FRigVMDispatchFactory*, TArray<const FRigVMFunction*>>& Pair : RequiredFactories)
+	{
+		const FRigVMDispatchFactory* Factory = Pair.Key;
+		check(Factory);
+		UScriptStruct* FactoryStruct = Factory->GetScriptStruct();
+		check(FactoryStruct);
+
+		// flatten the argument infos to the permutations used only
+		TArray<FRigVMTemplateArgumentInfo> TemporaryArgumentInfos;
+		{
+			FRigVMRegistryReadLock ReadLock;
+			const TArray<FRigVMTemplateArgumentInfo> FlattenedArgumentInfos = Factory->GetTemplate_NoLock(ReadLock)->GetFlattenedArgumentInfosForFunctions_NoLock(Pair.Value, ReadLock);
+			TemporaryArgumentInfos.Reserve(FlattenedArgumentInfos.Num());
+
+			for(const FRigVMTemplateArgumentInfo& ArgumentInfo : FlattenedArgumentInfos)
+			{
+				TArray<TRigVMTypeIndex> TypeIndices = ArgumentInfo.GetArgument_NoLock(ReadLock).TypeIndices;
+				for(int32 TypeIndex = 0; TypeIndex < TypeIndices.Num(); TypeIndex++)
+				{
+					const FRigVMTemplateArgumentType Type = ReadLock.GetRegistry().GetType_NoLock(TypeIndices[TypeIndex]);
+					TypeIndices[TypeIndex] = TemporaryRegistry.FindOrAddType_NoLock(Type);
+					check(TypeIndices[TypeIndex] != INDEX_NONE);
+				}
+				TemporaryArgumentInfos.Emplace(ArgumentInfo.Name, ArgumentInfo.Direction, TypeIndices);
+			}
+		}
+
+		// register the factory
+		FRigVMDispatchFactory* TemporaryFactory = const_cast<FRigVMDispatchFactory*>(TemporaryRegistry.RegisterFactory_NoLock(FactoryStruct, TemporaryArgumentInfos));
+
+		// update all of the argument name caches
+		if(const FRigVMTemplate* Template = TemporaryFactory->GetTemplate_NoLock(TemporaryContext))
+		{
+			(void)TemporaryFactory->UpdateArgumentNameCache_NoLock(Template->NumArguments(), TemporaryContext);
+		}
+	}
+
+	// register all struct based functions
+	for(const FRigVMFunction* Function : FunctionsStorage)
+	{
+		if(Function == nullptr)
+		{
+			continue;
+		}
+		
+		if(Function->Factory)
+		{
+			continue;
+		}
+		else
+		{
+			// registering the function will also register it's potential template
+			TemporaryRegistry.Register_NoLock(*Function->Name, Function->FunctionPtr, Function->Struct, Function->Arguments);
+		}
+	}
+
+	Swap(LocalizedRegistry, TemporaryRegistryPtr);
+
+	// now that we've registered all functions - let's re-resolve the functions
+	FunctionsStorage.Reset();
+	FactoriesStorage.Reset();
+
+	// if any of the functions fail to resolve we'll revert back to the global registry
+	if(!ResolveFunctionsIfRequired())
+	{
+		UE_LOG(LogRigVM, Error, TEXT("Cannot localize registry due to missing function."));
+
+		LocalizedRegistry.Reset();
+
+		FunctionsStorage.Reset();
+		FactoriesStorage.Reset();
+		verify(ResolveFunctionsIfRequired());
+		return false;
+	}
+	
+	return true;
+}
+
 #if WITH_EDITOR
 
 TArray<FString> URigVM::DumpByteCodeAsTextArray(FRigVMExtendedExecuteContext& Context, const TArray<int32>& InInstructionOrder, bool bIncludeLineNumbers, TFunction<FString(const FString& RegisterName, const FString& RegisterOffsetName)> OperandFormatFunction)
@@ -2661,7 +3007,7 @@ void URigVM::ClearDebugMemory(FRigVMExtendedExecuteContext& Context)
 #endif
 }
 
-void URigVM::CacheSingleMemoryHandle(FRigVMExtendedExecuteContext& Context, int32 InHandleIndex, const FRigVMBranchInfoKey& InBranchInfoKey, const FRigVMOperand& InArg, bool bForExecute)
+void URigVM::CacheSingleMemoryHandle(FRigVMExtendedExecuteContext& Context, const FRigVMRegistryHandle& InRegistry, int32 InHandleIndex, const FRigVMBranchInfoKey& InBranchInfoKey, const FRigVMOperand& InArg, bool bForExecute)
 {
 	FRigVMMemoryStorageStruct* Memory = GetMemoryByType(Context, InArg.GetMemoryType()/*, false*/);
 
@@ -2848,3 +3194,5 @@ void URigVM::StopProfiling(FRigVMExtendedExecuteContext& Context)
 	}
 #endif
 }
+
+#undef UE_RIGVM_GET_READ_REGISTRY_CONTEXT
\ No newline at end of file
diff --git a/Engine/Plugins/Runtime/RigVM/Source/RigVM/Private/RigVMCore/RigVMDispatchFactory.cpp b/Engine/Plugins/Runtime/RigVM/Source/RigVM/Private/RigVMCore/RigVMDispatchFactory.cpp
--- a/Engine/Plugins/Runtime/RigVM/Source/RigVM/Private/RigVMCore/RigVMDispatchFactory.cpp
+++ b/Engine/Plugins/Runtime/RigVM/Source/RigVM/Private/RigVMCore/RigVMDispatchFactory.cpp
@@ -151,10 +151,10 @@ bool FRigVMDispatchFactory::IsLazyInputArgument(const FName& InArgumentName) con
 
 #endif
 
-FName FRigVMDispatchFactory::GetArgumentNameForOperandIndex(int32 InOperandIndex, int32 InTotalOperands) const
+FName FRigVMDispatchFactory::GetArgumentNameForOperandIndex(int32 InOperandIndex, int32 InTotalOperands, FRigVMRegistryHandle& InRegistry) const
 {
-	check(GetArgumentInfos().Num() == InTotalOperands);
-	return GetArgumentInfos()[InOperandIndex].Name;
+	check(GetArgumentInfos(InRegistry).Num() == InTotalOperands);
+	return GetArgumentInfos(InRegistry)[InOperandIndex].Name;
 }
 
 const TArray<FName>& FRigVMDispatchFactory::GetControlFlowBlocks(const FRigVMDispatchContext& InContext) const
@@ -177,14 +177,16 @@ const TArray<FName>& FRigVMDispatchFactory::GetControlFlowBlocks_Impl(const FRig
 	return EmptyArray;
 }
 
-const TArray<FName>* FRigVMDispatchFactory::UpdateArgumentNameCache(int32 InNumberOperands) const
+const TArray<FName>* FRigVMDispatchFactory::UpdateArgumentNameCache(int32 InNumberOperands, FRigVMRegistryHandle& InRegistry) const
 {
 	FScopeLock ArgumentNamesLock(ArgumentNamesMutex);
-	return UpdateArgumentNameCache_NoLock(InNumberOperands);
+	return UpdateArgumentNameCache_NoLock(InNumberOperands, InRegistry);
 }
 
-const TArray<FName>* FRigVMDispatchFactory::UpdateArgumentNameCache_NoLock(int32 InNumberOperands) const
+const TArray<FName>* FRigVMDispatchFactory::UpdateArgumentNameCache_NoLock(int32 InNumberOperands, FRigVMRegistryHandle& InRegistry) const
 {
+	check(InRegistry.GetRegistry() == OwnerRegistry);			
+
 	TSharedPtr<TArray<FName>>& ArgumentNames = ArgumentNamesMap.FindOrAdd(InNumberOperands);
 	if(!ArgumentNames.IsValid())
 	{
@@ -192,24 +194,26 @@ const TArray<FName>* FRigVMDispatchFactory::UpdateArgumentNameCache_NoLock(int32
 		ArgumentNames->Reserve(InNumberOperands);
 		for(int32 OperandIndex = 0; OperandIndex < InNumberOperands; OperandIndex++)
 		{
-			ArgumentNames->Add(GetArgumentNameForOperandIndex(OperandIndex, InNumberOperands));
+			ArgumentNames->Add(GetArgumentNameForOperandIndex(OperandIndex, InNumberOperands, InRegistry));
 		}
 	}
 	return ArgumentNames.Get();
 }
 
-const TArray<FRigVMTemplateArgumentInfo>& FRigVMDispatchFactory::GetArgumentInfos() const
+const TArray<FRigVMTemplateArgumentInfo>& FRigVMDispatchFactory::GetArgumentInfos(FRigVMRegistryHandle& InRegistry) const
 {
 	static const TArray<FRigVMTemplateArgumentInfo> EmptyArguments;
 	return EmptyArguments;
 }
 
-TArray<FRigVMExecuteArgument> FRigVMDispatchFactory::GetExecuteArguments_NoLock(const FRigVMDispatchContext& InContext) const
+TArray<FRigVMExecuteArgument> FRigVMDispatchFactory::GetExecuteArguments_NoLock(const FRigVMDispatchContext& InContext, const FRigVMRegistryHandle& InRegistry) const
 {
+	check(InRegistry.GetRegistry() == OwnerRegistry);
+	
 	TArray<FRigVMExecuteArgument> Arguments = GetExecuteArguments_Impl(InContext);
 	for(FRigVMExecuteArgument& Argument : Arguments)
 	{
-		if(Argument.TypeIndex != INDEX_NONE && FRigVMRegistry_NoLock::GetForRead().IsArrayType_NoLock(Argument.TypeIndex))
+		if(Argument.TypeIndex != INDEX_NONE && InRegistry->IsArrayType_NoLock(Argument.TypeIndex))
 		{
 			Argument.TypeIndex = RigVMTypeUtils::TypeIndex::ExecuteArray;
 		}
@@ -230,35 +234,65 @@ const TArray<FRigVMExecuteArgument>& FRigVMDispatchFactory::GetExecuteArguments_
 
 FRigVMFunctionPtr FRigVMDispatchFactory::GetOrCreateDispatchFunction(const FRigVMTemplateTypeMap& InTypes) const
 {
-	const FRigVMRegistryWriteLock _;
-	return GetOrCreateDispatchFunction_NoLock(InTypes);
+	return GetOrCreateDispatchFunction_NoLock(InTypes, FRigVMRegistryReadLock());
+}
+
+const FRigVMTemplate* FRigVMDispatchFactory::CreateTemplateForArgumentInfos_NoLock(const TArray<FRigVMTemplateArgumentInfo>& InArguments, FRigVMRegistryHandle& InRegistry)
+{
+	check(InRegistry.GetRegistry() == OwnerRegistry);
+	check(CachedTemplate == nullptr);
+
+	const FName FactoryName = GetFactoryName();
+
+	FRigVMTemplateDelegates Delegates;
+	Delegates.GetDispatchFactoryDelegate = FRigVMTemplate_GetDispatchFactoryDelegate::CreateLambda(
+	[FactoryName](const FRigVMRegistryHandle& InRegistry)
+	{
+		return InRegistry->FindDispatchFactory_NoLock(FactoryName);
+	});
+	
+	CachedTemplate = InRegistry->AddTemplateFromArguments_NoLock(GetFactoryName(), InArguments, Delegates);
+	(void)UpdateArgumentNameCache_NoLock(InArguments.Num(), OwnerRegistry->Context);
+	
+	return CachedTemplate;
 }
 
-FRigVMFunctionPtr FRigVMDispatchFactory::GetOrCreateDispatchFunction_NoLock(const FRigVMTemplateTypeMap& InTypes) const
+FRigVMFunctionPtr FRigVMDispatchFactory::GetOrCreateDispatchFunction_NoLock(const FRigVMTemplateTypeMap& InTypes, const FRigVMRegistryHandle& InRegistry) const
 {
-	const FString PermutationName = GetPermutationNameImpl(InTypes, false);
-	if(const FRigVMFunction* ExistingFunction = FRigVMRegistry_NoLock::GetForRead().FindFunction_NoLock(*PermutationName))
+	check(InRegistry.GetRegistry() == OwnerRegistry);
+
+	const FString PermutationName = GetPermutationNameImpl(InTypes, InRegistry);
+	if(const FRigVMFunction* ExistingFunction = InRegistry->FindFunction_NoLock(*PermutationName))
 	{
 		return ExistingFunction->FunctionPtr;
 	}
 	
-	return CreateDispatchFunction_NoLock(InTypes);
+	return CreateDispatchFunction_NoLock(InTypes, InRegistry);
+}
+
+FRigVMFunctionPtr FRigVMDispatchFactory::CreateDispatchFunction_NoLock(const FRigVMTemplateTypeMap& InTypes, const FRigVMRegistryHandle& InRegistry) const
+{
+	check(InRegistry.GetRegistry() == OwnerRegistry);
+	return GetDispatchFunctionImpl(InTypes, InRegistry);
 }
 
-FRigVMFunctionPtr FRigVMDispatchFactory::CreateDispatchFunction_NoLock(const FRigVMTemplateTypeMap& InTypes) const
+TArray<FRigVMFunction> FRigVMDispatchFactory::CreateDispatchPredicates_NoLock(const FRigVMTemplateTypeMap& InTypes, const FRigVMRegistryHandle& InRegistry) const
 {
-	return GetDispatchFunctionImpl(InTypes);
+	check(InRegistry.GetRegistry() == OwnerRegistry);
+	return GetDispatchPredicatesImpl(InTypes, InRegistry);
 }
 
-TArray<FRigVMFunction> FRigVMDispatchFactory::CreateDispatchPredicates_NoLock(const FRigVMTemplateTypeMap& InTypes) const
+FString FRigVMDispatchFactory::GetPermutationName(const FRigVMTemplateTypeMap& InTypes) const
 {
-	return GetDispatchPredicatesImpl(InTypes);
+	FRigVMRegistryReadLock ReadLock;
+	return GetPermutationName_NoLock(InTypes, ReadLock);
 }
 
-FString FRigVMDispatchFactory::GetPermutationName(const FRigVMTemplateTypeMap& InTypes, const bool bLockRegistry) const
+FString FRigVMDispatchFactory::GetPermutationName_NoLock(const FRigVMTemplateTypeMap& InTypes, FRigVMRegistryHandle& InRegistry) const
 {
+	check(InRegistry.GetRegistry() == OwnerRegistry);
 #if WITH_EDITOR
-	const TArray<FRigVMTemplateArgumentInfo>& Arguments = GetArgumentInfos();
+	const TArray<FRigVMTemplateArgumentInfo>& Arguments = GetArgumentInfos(InRegistry);
 	
 	checkf(InTypes.Num() == Arguments.Num(), TEXT("Failed getting permutation names for '%s' "), *GetFactoryName().ToString());
 	
@@ -267,10 +301,10 @@ FString FRigVMDispatchFactory::GetPermutationName(const FRigVMTemplateTypeMap& I
 		check(InTypes.Contains(Argument.Name));
 	}
 #endif
-	return GetPermutationNameImpl(InTypes, bLockRegistry);
+	return GetPermutationNameImpl(InTypes, InRegistry);
 }
 
-TArray<FRigVMTemplateArgumentInfo> FRigVMDispatchFactory::BuildArgumentListFromPrimaryArgument(const TArray<FRigVMTemplateArgumentInfo>& InInfos, const FName& InPrimaryArgumentName) const
+TArray<FRigVMTemplateArgumentInfo> FRigVMDispatchFactory::BuildArgumentListFromPrimaryArgument(const TArray<FRigVMTemplateArgumentInfo>& InInfos, const FName& InPrimaryArgumentName, FRigVMRegistryHandle& InRegistry) const
 {
 	TArray<FRigVMTemplateArgumentInfo> NewInfos;
 	
@@ -288,14 +322,14 @@ TArray<FRigVMTemplateArgumentInfo> FRigVMDispatchFactory::BuildArgumentListFromP
 	TArray< TArray<TRigVMTypeIndex> > TypeIndicesArray;
 	TypeIndicesArray.SetNum(NumInfos);
 
-	const FRigVMTemplateArgument PrimaryArgument = PrimaryInfo->GetArgument();
+	const FRigVMTemplateArgument PrimaryArgument = PrimaryInfo->GetArgument_NoLock(InRegistry);
 	bool bFoundArg = true;
 	PrimaryArgument.ForEachType([&](const TRigVMTypeIndex Type)
 	{
 		if (bFoundArg)
 		{
 			TArray<FRigVMTemplateTypeMap, TInlineAllocator<1>> Permutations;
-			GetPermutationsFromArgumentType(InPrimaryArgumentName, Type, Permutations);
+			GetPermutationsFromArgumentType(InPrimaryArgumentName, Type, Permutations, InRegistry);
 			for (const FRigVMTemplateTypeMap& Permutation : Permutations)
 			{
 				for (int32 Index=0; Index < InInfos.Num(); ++Index)
@@ -317,7 +351,7 @@ TArray<FRigVMTemplateArgumentInfo> FRigVMDispatchFactory::BuildArgumentListFromP
 			}
 		}
 		return true;
-	});
+	}, InRegistry);
 
 	if (!bFoundArg)
 	{
@@ -327,11 +361,11 @@ TArray<FRigVMTemplateArgumentInfo> FRigVMDispatchFactory::BuildArgumentListFromP
 	NewInfos.Reserve(NumInfos);
 	for (int32 Index=0; Index < NumInfos; ++Index)
 	{
-		const FRigVMTemplateArgument Argument = InInfos[Index].GetArgument();
+		const FRigVMTemplateArgument Argument = InInfos[Index].GetArgument_NoLock(InRegistry);
 		const TArray<TRigVMTypeIndex>& TypeIndices = TypeIndicesArray[Index];
 		if (TypeIndices.IsEmpty())
 		{
-			NewInfos.Emplace(InInfos[Index].Name, Argument.Direction, Argument.TypeCategories);
+			NewInfos.Emplace(InInfos[Index].Name, Argument.Direction, Argument.TypeCategories, nullptr);
 		}
 		else
 		{
@@ -342,9 +376,9 @@ TArray<FRigVMTemplateArgumentInfo> FRigVMDispatchFactory::BuildArgumentListFromP
 	return NewInfos;
 }
 
-FString FRigVMDispatchFactory::GetPermutationNameImpl(const FRigVMTemplateTypeMap& InTypes, const bool bLockRegistry) const
+FString FRigVMDispatchFactory::GetPermutationNameImpl(const FRigVMTemplateTypeMap& InTypes, const FRigVMRegistryHandle& InRegistry) const
 {
-	const FString TypePairStrings = FRigVMTemplate::GetStringFromArgumentTypes(InTypes, bLockRegistry);
+	const FString TypePairStrings = FRigVMTemplate::GetStringFromArgumentTypes(InTypes, InRegistry);
 	return RigVMStringUtils::JoinStrings(GetFactoryName().ToString(), TypePairStrings, TEXT("::"));
 }
 
@@ -354,23 +388,24 @@ bool FRigVMDispatchFactory::CopyProperty(const FProperty* InTargetProperty, uint
 	return URigVMMemoryStorage::CopyProperty(InTargetProperty, InTargetPtr, InSourceProperty, InSourcePtr);
 }
 
-const FRigVMTemplate* FRigVMDispatchFactory::GetTemplate_RWLock() const
+const FRigVMTemplate* FRigVMDispatchFactory::GetTemplate() const
 {
-	const FRigVMRegistryWriteLock _;
-	return GetTemplate_NoLock();
+	FRigVMRegistryReadLock Lock;
+	return GetTemplate_NoLock(Lock);
 }
 
-const FRigVMTemplate* FRigVMDispatchFactory::GetTemplate_NoLock() const
+const FRigVMTemplate* FRigVMDispatchFactory::GetTemplate_NoLock(FRigVMRegistryHandle& InRegistry) const
 {
 	// make sure to rely on the instance of this factory that's stored under the registry
-	FRigVMRegistry_NoLock& Registry = FRigVMRegistry_NoLock::GetForWrite();
 	const FName FactoryName = GetFactoryName();
-	const FRigVMDispatchFactory* ThisFactory = Registry.FindDispatchFactory_NoLock(FactoryName);
+	const FRigVMDispatchFactory* ThisFactory = InRegistry->FindDispatchFactory_NoLock(FactoryName);
 	if(ThisFactory != this)
 	{
-		return ThisFactory->GetTemplate_NoLock();
+		return ThisFactory->GetTemplate_NoLock(InRegistry);
 	}
 
+	check(InRegistry.GetRegistry() == OwnerRegistry);
+
 	FScopeLock GetTemplateScopeLock(&GetTemplateMutex);
 	
 	if(CachedTemplate)
@@ -379,14 +414,14 @@ const FRigVMTemplate* FRigVMDispatchFactory::GetTemplate_NoLock() const
 	}
 	
 	// we don't allow execute types on arguments	
-	const TArray<FRigVMTemplateArgumentInfo>& Infos = GetArgumentInfos();
+	const TArray<FRigVMTemplateArgumentInfo>& Infos = GetArgumentInfos(InRegistry);
 	for (const FRigVMTemplateArgumentInfo& Info : Infos)
 	{
-		const FRigVMTemplateArgument Argument = Info.GetArgument();
+		const FRigVMTemplateArgument Argument = Info.GetArgument_NoLock(InRegistry);
 		const int32 Index = Argument.IndexOfByPredicate([&](const TRigVMTypeIndex TypeIndex)
 		{
-			return Registry.IsExecuteType_NoLock(TypeIndex);
-		});
+			return InRegistry->IsExecuteType_NoLock(TypeIndex);
+		}, InRegistry);
 		
 		if (Index != INDEX_NONE)
 		{
@@ -398,18 +433,24 @@ const FRigVMTemplate* FRigVMDispatchFactory::GetTemplate_NoLock() const
 
 	FRigVMTemplateDelegates Delegates;
 	Delegates.GetDispatchFactoryDelegate = FRigVMTemplate_GetDispatchFactoryDelegate::CreateLambda(
-	[FactoryName]()
+	[FactoryName](const FRigVMRegistryHandle& InRegistry)
 	{
-		return FRigVMRegistry_NoLock::GetForRead().FindDispatchFactory_NoLock(FactoryName);
+		return InRegistry->FindDispatchFactory_NoLock(FactoryName);
 	});
 
-	CachedTemplate = Registry.AddTemplateFromArguments_NoLock(GetFactoryName(), Infos, Delegates); 
+	CachedTemplate = InRegistry->AddTemplateFromArguments_NoLock(GetFactoryName(), Infos, Delegates); 
 	return CachedTemplate;
 }
 
 FName FRigVMDispatchFactory::GetTemplateNotation() const
 {
-	if(const FRigVMTemplate* Template = GetTemplate_RWLock())
+	FRigVMRegistryWriteLock Lock;
+	return GetTemplateNotation_NoLock(Lock);
+}
+
+FName FRigVMDispatchFactory::GetTemplateNotation_NoLock(FRigVMRegistryHandle& InRegistry) const
+{
+	if(const FRigVMTemplate* Template = GetTemplate_NoLock(InRegistry))
 	{
 		return Template->GetNotation();
 	}
diff --git a/Engine/Plugins/Runtime/RigVM/Source/RigVM/Private/RigVMCore/RigVMFunction.cpp b/Engine/Plugins/Runtime/RigVM/Source/RigVM/Private/RigVMCore/RigVMFunction.cpp
--- a/Engine/Plugins/Runtime/RigVM/Source/RigVM/Private/RigVMCore/RigVMFunction.cpp
+++ b/Engine/Plugins/Runtime/RigVM/Source/RigVM/Private/RigVMCore/RigVMFunction.cpp
@@ -8,13 +8,55 @@
 #include "RigVMPropertyUtils.h"
 #include "UObject/Package.h"
 
+void FRigVMFunctionArgument::Serialize_NoLock(FArchive& Ar)
+{
+	if(Ar.IsLoading())
+	{
+		NameString = MakeShared<FString>();
+		TypeString = MakeShared<FString>();
+		Ar << *NameString.Get();
+		Ar << *TypeString.Get();
+		
+		Name = NameString->operator*();
+		Type = TypeString->operator*();
+
+		int32 DirectionAsInt32 = 0;
+		Ar << DirectionAsInt32;
+		Direction = (ERigVMFunctionArgumentDirection)DirectionAsInt32;
+	}
+	else
+	{
+		FString LocalNameString = Name;
+		FString LocalTypeString = Type;
+		Ar << LocalNameString;
+		Ar << LocalTypeString;
+
+		int32 DirectionAsInt32 = (int32)(Direction);
+		Ar << DirectionAsInt32;
+	}
+}
+
+bool FRigVMFunctionArgument::operator==(const FRigVMFunctionArgument& InOther) const
+{
+	if(FCString::Strcmp(Name, InOther.Name) != 0)
+	{
+		return false;
+	}
+	if(FCString::Strcmp(Type, InOther.Type) != 0)
+	{
+		return false;
+	}
+	return Direction == InOther.Direction;
+}
+
 FRigVMFunction::FRigVMFunction(const TCHAR* InName, FRigVMFunctionPtr InFunctionPtr, UScriptStruct* InStruct, int32 InIndex, const TArray<FRigVMFunctionArgument>& InArguments)
-: Name(InName)
-, Struct(InStruct)
+: Name(InName), Struct(InStruct)
 , Factory(nullptr)
+, OwnerRegistry(nullptr)
 , FunctionPtr(InFunctionPtr)
 , Index(InIndex)
 , TemplateIndex(INDEX_NONE)
+, PermutationIndex(INDEX_NONE)
 , Arguments(InArguments)
 , ArgumentNames()
 {
@@ -79,10 +121,16 @@ FString FRigVMFunction::GetModuleRelativeHeaderPath() const
 
 const TArray<TRigVMTypeIndex>& FRigVMFunction::GetArgumentTypeIndices() const
 {
+	FRigVMRegistryWriteLock Lock;
+	return GetArgumentTypeIndices_NoLock(Lock);
+}
+
+const TArray<TRigVMTypeIndex>& FRigVMFunction::GetArgumentTypeIndices_NoLock(FRigVMRegistryHandle& InRegistry) const
+{
+	check(InRegistry.GetRegistry() == OwnerRegistry);
+	
 	if(ArgumentTypeIndices.IsEmpty() && !Arguments.IsEmpty())
 	{
-		const FRigVMRegistryWriteLock _;
-
 		if(Struct)
 		{
 			for(const FRigVMFunctionArgument& Argument : Arguments)
@@ -94,20 +142,20 @@ const TArray<TRigVMTypeIndex>& FRigVMFunction::GetArgumentTypeIndices() const
 					RigVMPropertyUtils::GetTypeFromProperty(Property, CPPType, CPPTypeObject);
 
 					const FRigVMTemplateArgumentType Type(CPPType, CPPTypeObject);
-					ArgumentTypeIndices.Add(FRigVMRegistry_NoLock::GetForWrite().FindOrAddType_NoLock(Type));
+					ArgumentTypeIndices.Add(InRegistry->FindOrAddType_NoLock(Type));
 				}
 			}
 		}
-		else if(const FRigVMTemplate* Template = GetTemplate_NoLock())
+		else if(const FRigVMTemplate* Template = GetTemplate_NoLock(InRegistry))
 		{
-			const int32 PermutationIndex = Template->FindPermutation(this);
 			check(PermutationIndex != INDEX_NONE);
+			check(PermutationIndex == Template->FindPermutation_NoLock(this, InRegistry));
 
 			for(const FRigVMFunctionArgument& FunctionArgument : Arguments)
 			{
 				const FRigVMTemplateArgument* TemplateArgument = Template->FindArgument(FunctionArgument.Name);
 				check(TemplateArgument);
-				ArgumentTypeIndices.Add(TemplateArgument->GetTypeIndex_NoLock(PermutationIndex));
+				ArgumentTypeIndices.Add(TemplateArgument->GetTypeIndex_NoLock(PermutationIndex, InRegistry));
 			}
 		}
 		else
@@ -120,19 +168,20 @@ const TArray<TRigVMTypeIndex>& FRigVMFunction::GetArgumentTypeIndices() const
 
 const FRigVMTemplate* FRigVMFunction::GetTemplate() const
 {
-	const FRigVMRegistryReadLock _;
-	return GetTemplate_NoLock();
+	return GetTemplate_NoLock(FRigVMRegistryReadLock());
 }
 
-const FRigVMTemplate* FRigVMFunction::GetTemplate_NoLock() const
+const FRigVMTemplate* FRigVMFunction::GetTemplate_NoLock(const FRigVMRegistryHandle& InRegistry) const
 {
+	check(InRegistry.GetRegistry() == OwnerRegistry);
+	
 	if(TemplateIndex == INDEX_NONE)
 	{
 		return nullptr;
 	}
 
-	const FRigVMTemplate* Template = &FRigVMRegistry_NoLock::GetForRead().GetTemplates_NoLock()[TemplateIndex];
-	if(Template->NumPermutations() <= 1)
+	const FRigVMTemplate* Template = &InRegistry->GetTemplates_NoLock()[TemplateIndex];
+	if(Template->NumPermutations_NoLock(InRegistry) <= 1)
 	{
 		return nullptr;
 	}
@@ -140,8 +189,15 @@ const FRigVMTemplate* FRigVMFunction::GetTemplate_NoLock() const
 	return Template;
 }
 
-const UScriptStruct* FRigVMFunction::GetExecuteContextStruct(bool bLockRegistry) const
+const UScriptStruct* FRigVMFunction::GetExecuteContextStruct() const
+{
+	return GetExecuteContextStruct_NoLock(FRigVMRegistryReadLock());
+}
+
+const UScriptStruct* FRigVMFunction::GetExecuteContextStruct_NoLock(const FRigVMRegistryHandle& InRegistry) const
 {
+	check(InRegistry.GetRegistry() == OwnerRegistry);
+
 	if(Factory)
 	{
 		return Factory->GetExecuteContextStruct();
@@ -158,14 +214,23 @@ const UScriptStruct* FRigVMFunction::GetExecuteContextStruct(bool bLockRegistry)
 
 bool FRigVMFunction::SupportsExecuteContextStruct(const UScriptStruct* InExecuteContextStruct) const
 {
-	return InExecuteContextStruct->IsChildOf(GetExecuteContextStruct());
+	return SupportsExecuteContextStruct_NoLock(InExecuteContextStruct, FRigVMRegistryReadLock());
 }
 
-const FName& FRigVMFunction::GetArgumentNameForOperandIndex(int32 InOperandIndex, int32 InTotalOperands) const
+bool FRigVMFunction::SupportsExecuteContextStruct_NoLock(const UScriptStruct* InExecuteContextStruct, const FRigVMRegistryHandle& InRegistry) const
 {
+	check(InRegistry.GetRegistry() == OwnerRegistry);
+
+	return InExecuteContextStruct->IsChildOf(GetExecuteContextStruct_NoLock(InRegistry));
+}
+
+const FName& FRigVMFunction::GetArgumentNameForOperandIndex(int32 InOperandIndex, int32 InTotalOperands, FRigVMRegistryHandle& InRegistry) const
+{
+	check(InRegistry.GetRegistry() == OwnerRegistry);
+
 	if(Factory)
 	{
-		const TArray<FName>* FactoryArgumentNames = Factory->UpdateArgumentNameCache_NoLock(InTotalOperands);
+		const TArray<FName>* FactoryArgumentNames = Factory->UpdateArgumentNameCache_NoLock(InTotalOperands, InRegistry);
 		check(FactoryArgumentNames);
 		check(FactoryArgumentNames->IsValidIndex(InOperandIndex));
 		return (*FactoryArgumentNames)[InOperandIndex];
@@ -174,3 +239,76 @@ const FName& FRigVMFunction::GetArgumentNameForOperandIndex(int32 InOperandIndex
 	check(ArgumentNames.IsValidIndex(InOperandIndex));
 	return ArgumentNames[InOperandIndex];
 }
+
+void FRigVMFunction::Serialize_NoLock(FArchive& Ar, FRigVMRegistryHandle& InRegistry)
+{
+	if(Ar.IsSaving())
+	{
+		verify(InRegistry.GetRegistry() == OwnerRegistry);
+	}
+	
+	Ar << Name;
+	Ar << Struct;
+
+	if(Ar.IsLoading())
+	{
+		Index = INDEX_NONE;
+		TemplateIndex = INDEX_NONE;
+		FunctionPtr = nullptr;
+		Arguments.Reset();
+		ArgumentTypeIndices.Reset();
+		ArgumentNames.Reset();
+
+		UScriptStruct* FactoryStruct = nullptr;
+		Ar << FactoryStruct;
+
+		int32 NumArguments = 0;
+		Ar << NumArguments;
+		Arguments.AddZeroed(NumArguments);
+		for(int32 ArgumentIndex = 0; ArgumentIndex < NumArguments; ArgumentIndex++)
+		{
+			Arguments[ArgumentIndex].Serialize_NoLock(Ar);
+		}
+
+		if(FactoryStruct)
+		{
+			Factory = InRegistry->FindOrAddDispatchFactory_NoLock(FactoryStruct);
+		}
+		else
+		{
+			Factory = nullptr;
+
+			// try to retrieve the function pointer from the global registry.
+			// in non-editor environments the global registry will still contain
+			// function pointers (but not the types and other expensive storage).
+			if(const FRigVMFunction* GlobalFunction = FRigVMRegistry::Get().FindFunction(*Name))
+			{
+				FunctionPtr = GlobalFunction->FunctionPtr;
+#if WITH_EDITOR
+				check(GlobalFunction->Struct == Struct);
+				check(GlobalFunction->Arguments.Num() == Arguments.Num());
+				for(int32 Arg = 0; Arg< Arguments.Num(); Arg++)
+				{
+					check(GlobalFunction->Arguments[Arg] == Arguments[Arg]);
+				}
+#endif
+			}
+		}
+	}
+	else if(Ar.IsSaving() || Ar.IsCountingMemory() || Ar.IsObjectReferenceCollector())
+	{
+		UScriptStruct* FactoryStruct = nullptr;
+		if(Factory)
+		{
+			FactoryStruct = Factory->GetScriptStruct();
+		}
+		Ar << FactoryStruct;
+
+		int32 NumArguments = Arguments.Num();
+		Ar << NumArguments;
+		for(FRigVMFunctionArgument& Argument : Arguments)
+		{
+			Argument.Serialize_NoLock(Ar);
+		}
+	}
+}
