RigVM: Localized registry in cooked environments - first steps.

diff --git a/Engine/Plugins/Runtime/RigVM/Source/RigVM/Private/RigVMCore/RigVMRegistry.cpp b/Engine/Plugins/Runtime/RigVM/Source/RigVM/Private/RigVMCore/RigVMRegistry.cpp
--- a/Engine/Plugins/Runtime/RigVM/Source/RigVM/Private/RigVMCore/RigVMRegistry.cpp
+++ b/Engine/Plugins/Runtime/RigVM/Source/RigVM/Private/RigVMCore/RigVMRegistry.cpp
@@ -21,20 +21,22 @@
 static FDelayedAutoRegisterHelper GRigVMRegistrySingletonHelper(EDelayedRegisterRunPhase::EndOfEngineInit, &FRigVMRegistry_NoLock::OnEngineInit);
 
 static TAutoConsoleVariable<bool> CVarRigVMUpdateDispatchFactoriesGreedily(TEXT("RigVM.UpdateDispatchFactoriesGreedily"), true, TEXT("Set this to false to avoid loading dispatch factories during engine init / plugin mount."));
+TAutoConsoleVariable<bool> CVarRigVMEnableLocalizedRegistry(TEXT("RigVM.EnableLocalizedRegistry"), false, TEXT("Set this to true to allow VMs to serialize their own localized registries."));
 
-FRigVMRegistry_NoLock::FRigVMRegistry_NoLock() :
-	bAvoidTypePropagation(false),
-	bEverRefreshedEngineTypes(false),
-	bEverRefreshedDispatchFactoriesAfterEngineInit(false)
+FRigVMRegistry_NoLock::FRigVMRegistry_NoLock(bool InIsGlobalRegistry)
+	: Context(this)
+	, bIsGlobalRegistry(InIsGlobalRegistry)
+	, bAvoidTypePropagation(false)
+	, bEverRefreshedEngineTypes(false)
+	, bEverRefreshedDispatchFactoriesAfterEngineInit(false)
 {
 }
 
-FRigVMRegistry_NoLock& FRigVMRegistry_NoLock::Get(ELockType InLockType)
+TSharedRef<FRigVMRegistry_NoLock> FRigVMRegistry_NoLock::CreateLocalizedRegistry()
 {
-#if WITH_EDITOR
-	FRigVMRegistry_RWLock::EnsureLocked(InLockType);
-#endif
-	return FRigVMRegistry_RWLock::Get();
+	TSharedRef<FRigVMRegistry_NoLock> LocalizedRegistry = MakeShareable(new FRigVMRegistry_NoLock(false));
+	LocalizedRegistry->InitializeBaseTypes_NoLock();
+	return LocalizedRegistry;
 }
 
 FRigVMRegistry_NoLock::~FRigVMRegistry_NoLock()
@@ -126,8 +128,7 @@ uint32 FRigVMRegistry_NoLock::GetHashForType_NoLock(TRigVMTypeIndex InTypeIndex)
 		return UINT32_MAX;
 	}
 
-	FRigVMRegistry_NoLock* MutableThis = (FRigVMRegistry_NoLock*)this; 
-	FTypeInfo& TypeInfo = MutableThis->Types[InTypeIndex];
+	const FTypeInfo& TypeInfo = Types[InTypeIndex];
 	
 	if(TypeInfo.Hash != UINT32_MAX)
 	{
@@ -262,102 +263,349 @@ void FRigVMRegistry_NoLock::RebuildRegistry_NoLock()
 	UserDefinedTypeToIndex.Reset();
 	AllowedClasses.Reset();
 
-	Initialize(false);
+	InitializeBaseTypes_NoLock();
+}
+
+void FRigVMRegistry_NoLock::FTypeInfo::Serialize_NoLock(FArchive& Ar)
+{
+	Type.Serialize_NoLock(Ar);
+	
+	if(Ar.IsLoading())
+	{
+		BaseTypeIndex = INDEX_NONE;
+		ArrayTypeIndex = INDEX_NONE;
+		bIsArray = false;
+		bIsExecute = false;
+		Hash = UINT32_MAX;
+	}
+}
+
+const FRigVMRegistryHandle& FRigVMRegistry_NoLock::GetHandle_NoLock() const
+{
+	return Context;
 }
 
-void FRigVMRegistry_NoLock::Initialize_NoLock()
+FRigVMRegistryHandle& FRigVMRegistry_NoLock::GetHandle_NoLock()
 {
-	Types.Reserve(512);
-	TypeToIndex.Reserve(512);
+	return Context;
+}
+
+void FRigVMRegistry_NoLock::InitializeBaseTypes_NoLock()
+{
+	Types.Reserve(bIsGlobalRegistry ? 512 : 0);
+	TypeToIndex.Reserve(bIsGlobalRegistry ? 512 : 0);
 	TypesPerCategory.Reserve(19);
 	TemplatesPerCategory.Reserve(19);
 	
-	TypesPerCategory.Add(FRigVMTemplateArgument::ETypeCategory_Execute, TArray<TRigVMTypeIndex>()).Reserve(8);
-	TypesPerCategory.Add(FRigVMTemplateArgument::ETypeCategory_SingleAnyValue, TArray<TRigVMTypeIndex>()).Reserve(256);
-	TypesPerCategory.Add(FRigVMTemplateArgument::ETypeCategory_ArrayAnyValue, TArray<TRigVMTypeIndex>()).Reserve(256);
-	TypesPerCategory.Add(FRigVMTemplateArgument::ETypeCategory_ArrayArrayAnyValue, TArray<TRigVMTypeIndex>()).Reserve(256);
+	TypesPerCategory.Add(FRigVMTemplateArgument::ETypeCategory_Execute, TArray<TRigVMTypeIndex>()).Reserve(bIsGlobalRegistry ? 8 : 0);
+	TypesPerCategory.Add(FRigVMTemplateArgument::ETypeCategory_SingleAnyValue, TArray<TRigVMTypeIndex>()).Reserve(bIsGlobalRegistry ? 256 : 0);
+	TypesPerCategory.Add(FRigVMTemplateArgument::ETypeCategory_ArrayAnyValue, TArray<TRigVMTypeIndex>()).Reserve(bIsGlobalRegistry ? 256 : 0);
+	TypesPerCategory.Add(FRigVMTemplateArgument::ETypeCategory_ArrayArrayAnyValue, TArray<TRigVMTypeIndex>()).Reserve(bIsGlobalRegistry ? 256 : 0);
 	TypesPerCategory.Add(FRigVMTemplateArgument::ETypeCategory_SingleSimpleValue, TArray<TRigVMTypeIndex>()).Reserve(8);
 	TypesPerCategory.Add(FRigVMTemplateArgument::ETypeCategory_ArraySimpleValue, TArray<TRigVMTypeIndex>()).Reserve(8);
 	TypesPerCategory.Add(FRigVMTemplateArgument::ETypeCategory_ArrayArraySimpleValue, TArray<TRigVMTypeIndex>()).Reserve(8);
 	TypesPerCategory.Add(FRigVMTemplateArgument::ETypeCategory_SingleMathStructValue, TArray<TRigVMTypeIndex>()).Reserve(GetMathTypes().Num());
 	TypesPerCategory.Add(FRigVMTemplateArgument::ETypeCategory_ArrayMathStructValue, TArray<TRigVMTypeIndex>()).Reserve(GetMathTypes().Num());
 	TypesPerCategory.Add(FRigVMTemplateArgument::ETypeCategory_ArrayArrayMathStructValue, TArray<TRigVMTypeIndex>()).Reserve(GetMathTypes().Num());
-	TypesPerCategory.Add(FRigVMTemplateArgument::ETypeCategory_SingleScriptStructValue, TArray<TRigVMTypeIndex>()).Reserve(128);
-	TypesPerCategory.Add(FRigVMTemplateArgument::ETypeCategory_ArrayScriptStructValue, TArray<TRigVMTypeIndex>()).Reserve(128);
-	TypesPerCategory.Add(FRigVMTemplateArgument::ETypeCategory_ArrayArrayScriptStructValue, TArray<TRigVMTypeIndex>()).Reserve(128);
-	TypesPerCategory.Add(FRigVMTemplateArgument::ETypeCategory_SingleEnumValue, TArray<TRigVMTypeIndex>()).Reserve(128);
-	TypesPerCategory.Add(FRigVMTemplateArgument::ETypeCategory_ArrayEnumValue, TArray<TRigVMTypeIndex>()).Reserve(128);
-	TypesPerCategory.Add(FRigVMTemplateArgument::ETypeCategory_ArrayArrayEnumValue, TArray<TRigVMTypeIndex>()).Reserve(128);
-	TypesPerCategory.Add(FRigVMTemplateArgument::ETypeCategory_SingleObjectValue, TArray<TRigVMTypeIndex>()).Reserve(128);
-	TypesPerCategory.Add(FRigVMTemplateArgument::ETypeCategory_ArrayObjectValue, TArray<TRigVMTypeIndex>()).Reserve(128);
-	TypesPerCategory.Add(FRigVMTemplateArgument::ETypeCategory_ArrayArrayObjectValue, TArray<TRigVMTypeIndex>()).Reserve(128);
-
-	TemplatesPerCategory.Add(FRigVMTemplateArgument::ETypeCategory_Execute, TArray<int32>()).Reserve(8);
-	TemplatesPerCategory.Add(FRigVMTemplateArgument::ETypeCategory_SingleAnyValue, TArray<int32>()).Reserve(64);
-	TemplatesPerCategory.Add(FRigVMTemplateArgument::ETypeCategory_ArrayAnyValue, TArray<int32>()).Reserve(64);
-	TemplatesPerCategory.Add(FRigVMTemplateArgument::ETypeCategory_ArrayArrayAnyValue, TArray<int32>()).Reserve(64);
-	TemplatesPerCategory.Add(FRigVMTemplateArgument::ETypeCategory_SingleSimpleValue, TArray<int32>()).Reserve(64);
-	TemplatesPerCategory.Add(FRigVMTemplateArgument::ETypeCategory_ArraySimpleValue, TArray<int32>()).Reserve(64);
-	TemplatesPerCategory.Add(FRigVMTemplateArgument::ETypeCategory_ArrayArraySimpleValue, TArray<int32>()).Reserve(64);
-	TemplatesPerCategory.Add(FRigVMTemplateArgument::ETypeCategory_SingleMathStructValue, TArray<int32>()).Reserve(64);
-	TemplatesPerCategory.Add(FRigVMTemplateArgument::ETypeCategory_ArrayMathStructValue, TArray<int32>()).Reserve(64);
-	TemplatesPerCategory.Add(FRigVMTemplateArgument::ETypeCategory_ArrayArrayMathStructValue, TArray<int32>()).Reserve(64);
-	TemplatesPerCategory.Add(FRigVMTemplateArgument::ETypeCategory_SingleScriptStructValue, TArray<int32>()).Reserve(64);
-	TemplatesPerCategory.Add(FRigVMTemplateArgument::ETypeCategory_ArrayScriptStructValue, TArray<int32>()).Reserve(64);
-	TemplatesPerCategory.Add(FRigVMTemplateArgument::ETypeCategory_ArrayArrayScriptStructValue, TArray<int32>()).Reserve(64);
-	TemplatesPerCategory.Add(FRigVMTemplateArgument::ETypeCategory_SingleEnumValue, TArray<int32>()).Reserve(64);
-	TemplatesPerCategory.Add(FRigVMTemplateArgument::ETypeCategory_ArrayEnumValue, TArray<int32>()).Reserve(64);
-	TemplatesPerCategory.Add(FRigVMTemplateArgument::ETypeCategory_ArrayArrayEnumValue, TArray<int32>()).Reserve(64);
-	TemplatesPerCategory.Add(FRigVMTemplateArgument::ETypeCategory_SingleObjectValue, TArray<int32>()).Reserve(64);
-	TemplatesPerCategory.Add(FRigVMTemplateArgument::ETypeCategory_ArrayObjectValue, TArray<int32>()).Reserve(64);
-	TemplatesPerCategory.Add(FRigVMTemplateArgument::ETypeCategory_ArrayArrayObjectValue, TArray<int32>()).Reserve(64);
-
-	RigVMTypeUtils::TypeIndex::Execute = FindOrAddType_NoLock(FRigVMTemplateArgumentType(FRigVMExecuteContext::StaticStruct()), false);
-	RigVMTypeUtils::TypeIndex::ExecuteArray = FindOrAddType_NoLock(FRigVMTemplateArgumentType(FRigVMExecuteContext::StaticStruct()).ConvertToArray(), false);
-	RigVMTypeUtils::TypeIndex::Bool = FindOrAddType_NoLock(FRigVMTemplateArgumentType(RigVMTypeUtils::BoolTypeName, nullptr), false);
-	RigVMTypeUtils::TypeIndex::Float = FindOrAddType_NoLock(FRigVMTemplateArgumentType(RigVMTypeUtils::FloatTypeName, nullptr), false);
-	RigVMTypeUtils::TypeIndex::Double = FindOrAddType_NoLock(FRigVMTemplateArgumentType(RigVMTypeUtils::DoubleTypeName, nullptr), false);
-	RigVMTypeUtils::TypeIndex::Int32 = FindOrAddType_NoLock(FRigVMTemplateArgumentType(RigVMTypeUtils::Int32TypeName, nullptr), false);
-	RigVMTypeUtils::TypeIndex::UInt32 = FindOrAddType_NoLock(FRigVMTemplateArgumentType(RigVMTypeUtils::UInt32TypeName, nullptr), false);
-	RigVMTypeUtils::TypeIndex::UInt8 = FindOrAddType_NoLock(FRigVMTemplateArgumentType(RigVMTypeUtils::UInt8TypeName, nullptr), false);
-	RigVMTypeUtils::TypeIndex::FName = FindOrAddType_NoLock(FRigVMTemplateArgumentType(RigVMTypeUtils::FNameTypeName, nullptr), false);
-	RigVMTypeUtils::TypeIndex::FString = FindOrAddType_NoLock(FRigVMTemplateArgumentType(RigVMTypeUtils::FStringTypeName, nullptr), false);
-	RigVMTypeUtils::TypeIndex::WildCard = FindOrAddType_NoLock(FRigVMTemplateArgumentType(RigVMTypeUtils::GetWildCardCPPTypeName(), RigVMTypeUtils::GetWildCardCPPTypeObject()), false);
-	RigVMTypeUtils::TypeIndex::BoolArray = FindOrAddType_NoLock(FRigVMTemplateArgumentType(RigVMTypeUtils::BoolArrayTypeName, nullptr), false);
-	RigVMTypeUtils::TypeIndex::FloatArray = FindOrAddType_NoLock(FRigVMTemplateArgumentType(RigVMTypeUtils::FloatArrayTypeName, nullptr), false);
-	RigVMTypeUtils::TypeIndex::DoubleArray = FindOrAddType_NoLock(FRigVMTemplateArgumentType(RigVMTypeUtils::DoubleArrayTypeName, nullptr), false);
-	RigVMTypeUtils::TypeIndex::Int32Array = FindOrAddType_NoLock(FRigVMTemplateArgumentType(RigVMTypeUtils::Int32ArrayTypeName, nullptr), false);
-	RigVMTypeUtils::TypeIndex::UInt32Array = FindOrAddType_NoLock(FRigVMTemplateArgumentType(RigVMTypeUtils::UInt32ArrayTypeName, nullptr), false);
-	RigVMTypeUtils::TypeIndex::UInt8Array = FindOrAddType_NoLock(FRigVMTemplateArgumentType(RigVMTypeUtils::UInt8ArrayTypeName, nullptr), false);
-	RigVMTypeUtils::TypeIndex::FNameArray = FindOrAddType_NoLock(FRigVMTemplateArgumentType(RigVMTypeUtils::FNameArrayTypeName, nullptr), false);
-	RigVMTypeUtils::TypeIndex::FStringArray = FindOrAddType_NoLock(FRigVMTemplateArgumentType(RigVMTypeUtils::FStringArrayTypeName, nullptr), false);
-	RigVMTypeUtils::TypeIndex::WildCardArray = FindOrAddType_NoLock(FRigVMTemplateArgumentType(RigVMTypeUtils::GetWildCardArrayCPPTypeName(), RigVMTypeUtils::GetWildCardCPPTypeObject()), false);
-
-	// register the default math types
-	for(UScriptStruct* MathType : GetMathTypes())
-	{
-		FindOrAddType_NoLock(FRigVMTemplateArgumentType(MathType), false);
-	}
-
-	// hook the registry to prepare for engine shutdown
-	FCoreDelegates::OnExit.AddLambda([&]()
-	{
-		Reset_NoLock();
-
-		if (FAssetRegistryModule* AssetRegistryModule = FModuleManager::GetModulePtr<FAssetRegistryModule>(TEXT("AssetRegistry")))
-		{
-			if (AssetRegistryModule->TryGet())
+	TypesPerCategory.Add(FRigVMTemplateArgument::ETypeCategory_SingleScriptStructValue, TArray<TRigVMTypeIndex>()).Reserve(bIsGlobalRegistry ? 128 : 0);
+	TypesPerCategory.Add(FRigVMTemplateArgument::ETypeCategory_ArrayScriptStructValue, TArray<TRigVMTypeIndex>()).Reserve(bIsGlobalRegistry ? 128 : 0);
+	TypesPerCategory.Add(FRigVMTemplateArgument::ETypeCategory_ArrayArrayScriptStructValue, TArray<TRigVMTypeIndex>()).Reserve(bIsGlobalRegistry ? 128 : 0);
+	TypesPerCategory.Add(FRigVMTemplateArgument::ETypeCategory_SingleEnumValue, TArray<TRigVMTypeIndex>()).Reserve(bIsGlobalRegistry ? 128 : 0);
+	TypesPerCategory.Add(FRigVMTemplateArgument::ETypeCategory_ArrayEnumValue, TArray<TRigVMTypeIndex>()).Reserve(bIsGlobalRegistry ? 128 : 0);
+	TypesPerCategory.Add(FRigVMTemplateArgument::ETypeCategory_ArrayArrayEnumValue, TArray<TRigVMTypeIndex>()).Reserve(bIsGlobalRegistry ? 128 : 0);
+	TypesPerCategory.Add(FRigVMTemplateArgument::ETypeCategory_SingleObjectValue, TArray<TRigVMTypeIndex>()).Reserve(bIsGlobalRegistry ? 128 : 0);
+	TypesPerCategory.Add(FRigVMTemplateArgument::ETypeCategory_ArrayObjectValue, TArray<TRigVMTypeIndex>()).Reserve(bIsGlobalRegistry ? 128 : 0);
+	TypesPerCategory.Add(FRigVMTemplateArgument::ETypeCategory_ArrayArrayObjectValue, TArray<TRigVMTypeIndex>()).Reserve(bIsGlobalRegistry ? 128 : 0);
+
+	TemplatesPerCategory.Add(FRigVMTemplateArgument::ETypeCategory_Execute, TArray<int32>()).Reserve(bIsGlobalRegistry ? 8 : 0);
+	TemplatesPerCategory.Add(FRigVMTemplateArgument::ETypeCategory_SingleAnyValue, TArray<int32>()).Reserve(bIsGlobalRegistry ? 64 : 0);
+	TemplatesPerCategory.Add(FRigVMTemplateArgument::ETypeCategory_ArrayAnyValue, TArray<int32>()).Reserve(bIsGlobalRegistry ? 64 : 0);
+	TemplatesPerCategory.Add(FRigVMTemplateArgument::ETypeCategory_ArrayArrayAnyValue, TArray<int32>()).Reserve(bIsGlobalRegistry ? 64 : 0);
+	TemplatesPerCategory.Add(FRigVMTemplateArgument::ETypeCategory_SingleSimpleValue, TArray<int32>()).Reserve(bIsGlobalRegistry ? 64 : 0);
+	TemplatesPerCategory.Add(FRigVMTemplateArgument::ETypeCategory_ArraySimpleValue, TArray<int32>()).Reserve(bIsGlobalRegistry ? 64 : 0);
+	TemplatesPerCategory.Add(FRigVMTemplateArgument::ETypeCategory_ArrayArraySimpleValue, TArray<int32>()).Reserve(bIsGlobalRegistry ? 64 : 0);
+	TemplatesPerCategory.Add(FRigVMTemplateArgument::ETypeCategory_SingleMathStructValue, TArray<int32>()).Reserve(bIsGlobalRegistry ? 64 : 0);
+	TemplatesPerCategory.Add(FRigVMTemplateArgument::ETypeCategory_ArrayMathStructValue, TArray<int32>()).Reserve(bIsGlobalRegistry ? 64 : 0);
+	TemplatesPerCategory.Add(FRigVMTemplateArgument::ETypeCategory_ArrayArrayMathStructValue, TArray<int32>()).Reserve(bIsGlobalRegistry ? 64 : 0);
+	TemplatesPerCategory.Add(FRigVMTemplateArgument::ETypeCategory_SingleScriptStructValue, TArray<int32>()).Reserve(bIsGlobalRegistry ? 64 : 0);
+	TemplatesPerCategory.Add(FRigVMTemplateArgument::ETypeCategory_ArrayScriptStructValue, TArray<int32>()).Reserve(bIsGlobalRegistry ? 64 : 0);
+	TemplatesPerCategory.Add(FRigVMTemplateArgument::ETypeCategory_ArrayArrayScriptStructValue, TArray<int32>()).Reserve(bIsGlobalRegistry ? 64 : 0);
+	TemplatesPerCategory.Add(FRigVMTemplateArgument::ETypeCategory_SingleEnumValue, TArray<int32>()).Reserve(bIsGlobalRegistry ? 64 : 0);
+	TemplatesPerCategory.Add(FRigVMTemplateArgument::ETypeCategory_ArrayEnumValue, TArray<int32>()).Reserve(bIsGlobalRegistry ? 64 : 0);
+	TemplatesPerCategory.Add(FRigVMTemplateArgument::ETypeCategory_ArrayArrayEnumValue, TArray<int32>()).Reserve(bIsGlobalRegistry ? 64 : 0);
+	TemplatesPerCategory.Add(FRigVMTemplateArgument::ETypeCategory_SingleObjectValue, TArray<int32>()).Reserve(bIsGlobalRegistry ? 64 : 0);
+	TemplatesPerCategory.Add(FRigVMTemplateArgument::ETypeCategory_ArrayObjectValue, TArray<int32>()).Reserve(bIsGlobalRegistry ? 64 : 0);
+	TemplatesPerCategory.Add(FRigVMTemplateArgument::ETypeCategory_ArrayArrayObjectValue, TArray<int32>()).Reserve(bIsGlobalRegistry ? 64 : 0);
+
+	// register the simple types the same between all registries.
+	// we rely on the static RigVMTypeUtils::TypeIndex lookup - so all registries have to use the same type indices for this.
+	TRigVMTypeIndex LocalExecuteIndex, LocalExecuteArrayIndex, LocalBoolIndex, LocalFloatIndex, LocalDoubleIndex, LocalInt32Index, LocalUInt32Index,
+		LocalUInt8Index, LocalFNameIndex, LocalFStringIndex, LocalWildCardIndex, LocalBoolArrayIndex, LocalFloatArrayIndex, LocalDoubleArrayIndex,
+		LocalInt32ArrayIndex, LocalUInt32ArrayIndex, LocalUInt8ArrayIndex, LocalFNameArrayIndex, LocalFStringArrayIndex, LocalWildCardArrayIndex;
+	{
+		const TGuardValue<bool> EnableAllArrayDimensions(bIsGlobalRegistry, true);
+		LocalExecuteIndex = FindOrAddType_NoLock(FRigVMTemplateArgumentType(FRigVMExecuteContext::StaticStruct()), false);
+		LocalExecuteArrayIndex = FindOrAddType_NoLock(FRigVMTemplateArgumentType(FRigVMExecuteContext::StaticStruct()).ConvertToArray(), false);
+		LocalBoolIndex = FindOrAddType_NoLock(FRigVMTemplateArgumentType(RigVMTypeUtils::BoolTypeName, nullptr), false);
+		LocalFloatIndex = FindOrAddType_NoLock(FRigVMTemplateArgumentType(RigVMTypeUtils::FloatTypeName, nullptr), false);
+		LocalDoubleIndex = FindOrAddType_NoLock(FRigVMTemplateArgumentType(RigVMTypeUtils::DoubleTypeName, nullptr), false);
+		LocalInt32Index = FindOrAddType_NoLock(FRigVMTemplateArgumentType(RigVMTypeUtils::Int32TypeName, nullptr), false);
+		LocalUInt32Index = FindOrAddType_NoLock(FRigVMTemplateArgumentType(RigVMTypeUtils::UInt32TypeName, nullptr), false);
+		LocalUInt8Index = FindOrAddType_NoLock(FRigVMTemplateArgumentType(RigVMTypeUtils::UInt8TypeName, nullptr), false);
+		LocalFNameIndex = FindOrAddType_NoLock(FRigVMTemplateArgumentType(RigVMTypeUtils::FNameTypeName, nullptr), false);
+		LocalFStringIndex = FindOrAddType_NoLock(FRigVMTemplateArgumentType(RigVMTypeUtils::FStringTypeName, nullptr), false);
+		LocalWildCardIndex = FindOrAddType_NoLock(FRigVMTemplateArgumentType(RigVMTypeUtils::GetWildCardCPPTypeName(), RigVMTypeUtils::GetWildCardCPPTypeObject()), false);
+		LocalBoolArrayIndex = FindOrAddType_NoLock(FRigVMTemplateArgumentType(RigVMTypeUtils::BoolArrayTypeName, nullptr), false);
+		LocalFloatArrayIndex = FindOrAddType_NoLock(FRigVMTemplateArgumentType(RigVMTypeUtils::FloatArrayTypeName, nullptr), false);
+		LocalDoubleArrayIndex = FindOrAddType_NoLock(FRigVMTemplateArgumentType(RigVMTypeUtils::DoubleArrayTypeName, nullptr), false);
+		LocalInt32ArrayIndex = FindOrAddType_NoLock(FRigVMTemplateArgumentType(RigVMTypeUtils::Int32ArrayTypeName, nullptr), false);
+		LocalUInt32ArrayIndex = FindOrAddType_NoLock(FRigVMTemplateArgumentType(RigVMTypeUtils::UInt32ArrayTypeName, nullptr), false);
+		LocalUInt8ArrayIndex = FindOrAddType_NoLock(FRigVMTemplateArgumentType(RigVMTypeUtils::UInt8ArrayTypeName, nullptr), false);
+		LocalFNameArrayIndex = FindOrAddType_NoLock(FRigVMTemplateArgumentType(RigVMTypeUtils::FNameArrayTypeName, nullptr), false);
+		LocalFStringArrayIndex = FindOrAddType_NoLock(FRigVMTemplateArgumentType(RigVMTypeUtils::FStringArrayTypeName, nullptr), false);
+		LocalWildCardArrayIndex = FindOrAddType_NoLock(FRigVMTemplateArgumentType(RigVMTypeUtils::GetWildCardArrayCPPTypeName(), RigVMTypeUtils::GetWildCardCPPTypeObject()), false);
+	}
+
+	if(bIsGlobalRegistry)
+	{
+		RigVMTypeUtils::TypeIndex::Execute = LocalExecuteIndex;
+		RigVMTypeUtils::TypeIndex::ExecuteArray = LocalExecuteArrayIndex;
+		RigVMTypeUtils::TypeIndex::Bool = LocalBoolIndex;
+		RigVMTypeUtils::TypeIndex::Float = LocalFloatIndex;
+		RigVMTypeUtils::TypeIndex::Double = LocalDoubleIndex;
+		RigVMTypeUtils::TypeIndex::Int32 = LocalInt32Index;
+		RigVMTypeUtils::TypeIndex::UInt32 = LocalUInt32Index;
+		RigVMTypeUtils::TypeIndex::UInt8 = LocalUInt8Index;
+		RigVMTypeUtils::TypeIndex::FName = LocalFNameIndex;
+		RigVMTypeUtils::TypeIndex::FString = LocalFStringIndex;
+		RigVMTypeUtils::TypeIndex::WildCard = LocalWildCardIndex;
+		RigVMTypeUtils::TypeIndex::BoolArray = LocalBoolArrayIndex;
+		RigVMTypeUtils::TypeIndex::FloatArray = LocalFloatArrayIndex;
+		RigVMTypeUtils::TypeIndex::DoubleArray = LocalDoubleArrayIndex;
+		RigVMTypeUtils::TypeIndex::Int32Array = LocalInt32ArrayIndex;
+		RigVMTypeUtils::TypeIndex::UInt32Array = LocalUInt32ArrayIndex;
+		RigVMTypeUtils::TypeIndex::UInt8Array = LocalUInt8ArrayIndex;
+		RigVMTypeUtils::TypeIndex::FNameArray = LocalFNameArrayIndex;
+		RigVMTypeUtils::TypeIndex::FStringArray = LocalFStringArrayIndex;
+		RigVMTypeUtils::TypeIndex::WildCardArray = LocalWildCardArrayIndex;
+		
+		// register the default math types
+		for(UScriptStruct* MathType : GetMathTypes())
+		{
+			FindOrAddType_NoLock(FRigVMTemplateArgumentType(MathType), false);
+		}
+
+		// hook the registry to prepare for engine shutdown
+		FCoreDelegates::OnExit.AddLambda([&]()
+		{
+			Reset_NoLock();
+
+			if (FAssetRegistryModule* AssetRegistryModule = FModuleManager::GetModulePtr<FAssetRegistryModule>(TEXT("AssetRegistry")))
 			{
-				AssetRegistryModule->Get().OnAssetRemoved().RemoveAll(this);
-				AssetRegistryModule->Get().OnAssetRenamed().RemoveAll(this);
+				if (AssetRegistryModule->TryGet())
+				{
+					AssetRegistryModule->Get().OnAssetRemoved().RemoveAll(this);
+					AssetRegistryModule->Get().OnAssetRenamed().RemoveAll(this);
+				}
+			}
+
+			IPluginManager::Get().OnNewPluginMounted().RemoveAll(this);
+			IPluginManager::Get().OnPluginUnmounted().RemoveAll(this);
+
+			UE::Anim::AttributeTypes::GetOnAttributeTypesChanged().RemoveAll(this);
+		});
+	}
+	else
+	{
+		check(RigVMTypeUtils::TypeIndex::Execute == LocalExecuteIndex);
+		check(RigVMTypeUtils::TypeIndex::ExecuteArray == LocalExecuteArrayIndex);
+		check(RigVMTypeUtils::TypeIndex::Bool == LocalBoolIndex);
+		check(RigVMTypeUtils::TypeIndex::Float == LocalFloatIndex);
+		check(RigVMTypeUtils::TypeIndex::Double == LocalDoubleIndex);
+		check(RigVMTypeUtils::TypeIndex::Int32 == LocalInt32Index);
+		check(RigVMTypeUtils::TypeIndex::UInt32 == LocalUInt32Index);
+		check(RigVMTypeUtils::TypeIndex::UInt8 == LocalUInt8Index);
+		check(RigVMTypeUtils::TypeIndex::FName == LocalFNameIndex);
+		check(RigVMTypeUtils::TypeIndex::FString == LocalFStringIndex);
+		check(RigVMTypeUtils::TypeIndex::WildCard == LocalWildCardIndex);
+		check(RigVMTypeUtils::TypeIndex::BoolArray == LocalBoolArrayIndex);
+		check(RigVMTypeUtils::TypeIndex::FloatArray == LocalFloatArrayIndex);
+		check(RigVMTypeUtils::TypeIndex::DoubleArray == LocalDoubleArrayIndex);
+		check(RigVMTypeUtils::TypeIndex::Int32Array == LocalInt32ArrayIndex);
+		check(RigVMTypeUtils::TypeIndex::UInt32Array == LocalUInt32ArrayIndex);
+		check(RigVMTypeUtils::TypeIndex::UInt8Array == LocalUInt8ArrayIndex);
+		check(RigVMTypeUtils::TypeIndex::FNameArray == LocalFNameArrayIndex);
+		check(RigVMTypeUtils::TypeIndex::FStringArray == LocalFStringArrayIndex);
+		check(RigVMTypeUtils::TypeIndex::WildCardArray == LocalWildCardArrayIndex);
+	}
+}
+
+void FRigVMRegistry_NoLock::Serialize_NoLock(FArchive& Ar)
+{
+	check(!bIsGlobalRegistry);
+
+	if(Ar.IsLoading())
+	{
+		RebuildRegistry_NoLock();
+		Ar << AllowedClasses;
+
+		// types
+		{
+			int32 NumTypes = 0;
+			Ar << NumTypes;
+			Types.Reserve(NumTypes * 3); // 3 times due to array dimensions [] and [][]
+		
+			for(int32 TypeIndex = 0; TypeIndex < NumTypes; TypeIndex++)
+			{
+				FTypeInfo TypeInfo;
+				TypeInfo.Serialize_NoLock(Ar);
+				(void)FindOrAddType_NoLock(TypeInfo.Type);
 			}
 		}
 
-		IPluginManager::Get().OnNewPluginMounted().RemoveAll(this);
-		IPluginManager::Get().OnPluginUnmounted().RemoveAll(this);
+		// factories
+		{
+			int32 NumFactories = 0;
+			Ar << NumFactories;
+			
+			for(int32 FactoryIndex = 0; FactoryIndex < NumFactories; FactoryIndex++)
+			{
+				UScriptStruct* FactoryScriptStruct = nullptr;
+				Ar << FactoryScriptStruct;
 
-		UE::Anim::AttributeTypes::GetOnAttributeTypesChanged().RemoveAll(this);
-	});
+				int32 NumFlattenedArguments = 0;
+				Ar << NumFlattenedArguments;
+
+				TArray<FRigVMTemplateArgumentInfo> FlattenedArguments;
+				FlattenedArguments.Reserve(NumFlattenedArguments);
+
+				for(int32 ArgumentIndex = 0; ArgumentIndex < NumFlattenedArguments; ArgumentIndex++)
+				{
+					FName ArgumentName(NAME_None);
+					Ar << ArgumentName;
+
+					int32 Direction = 0;
+					Ar << Direction;
+
+					int32 NumTypes = 0;
+					Ar << NumTypes;
+
+					TArray<TRigVMTypeIndex> TypeIndices;
+					TypeIndices.Reserve(NumTypes);
+					
+					for(int32 TypeIndex = 0; TypeIndex < NumTypes; TypeIndex++)
+					{
+						FName CPPType(NAME_None);
+						Ar << CPPType;
+
+						UObject* CPPTypeObject = nullptr;
+						Ar << CPPTypeObject;
+						
+						TypeIndices.Add(FindOrAddType_NoLock({CPPType, CPPTypeObject}));
+					}
+
+					if(TypeIndices.Num() == 1)
+					{
+						FlattenedArguments.Emplace(ArgumentName, static_cast<ERigVMPinDirection>(Direction), TypeIndices[0]);
+					}
+					else
+					{
+						FlattenedArguments.Emplace(ArgumentName, static_cast<ERigVMPinDirection>(Direction), TypeIndices);
+					}
+				}
+
+				if(const FRigVMDispatchFactory* Factory = RegisterFactory_NoLock(FactoryScriptStruct, FlattenedArguments))
+				{
+					// pulling on this caches the template
+					Factory->GetTemplate_NoLock(Context);
+				}
+			}
+		}
+
+		// functions
+		{
+			int32 NumFunctions = 0;
+			Ar << NumFunctions;
+			for(int32 FunctionIndex = 0; FunctionIndex < NumFunctions; FunctionIndex++)
+			{
+				FRigVMFunction Function;
+				Function.Serialize_NoLock(Ar, Context);
+
+				if(Function.Struct)
+				{
+					Register_NoLock(*Function.Name, Function.FunctionPtr, Function.Struct, Function.Arguments);
+				}
+				else if(Function.Factory)
+				{
+					// finding function will pull on the factory to create it
+					(void)FindFunction_NoLock(*Function.Name);
+				}
+			}
+		}
+	}
+	else if(Ar.IsSaving() || Ar.IsCountingMemory() || Ar.IsObjectReferenceCollector())
+	{
+		Ar << AllowedClasses;
+		
+		// types
+		{
+			int32 NumTypes = Types.Num();
+			Ar << NumTypes;
+			for(FTypeInfo& TypeInfo : Types)
+			{
+				TypeInfo.Serialize_NoLock(Ar);
+			}
+		}
+
+		// factories
+		{
+			int32 NumFactories = Factories.Num();
+			Ar << NumFactories;
+			for(FRigVMDispatchFactory* Factory : Factories)
+			{
+				UScriptStruct* FactoryScriptStruct = Factory->GetScriptStruct();
+				Ar << FactoryScriptStruct;
+
+				// find all used permutations of the factory and store their types
+				TArray<FRigVMTemplateArgumentInfo> FlattenedArguments = Factory->GetTemplate_NoLock(Context)->GetFlattenedArgumentInfos_NoLock(Context);
+
+				int32 NumFlattenedArguments = FlattenedArguments.Num();
+				Ar << NumFlattenedArguments;
+
+				for(int32 ArgumentIndex = 0; ArgumentIndex < NumFlattenedArguments; ArgumentIndex++)
+				{
+					FRigVMTemplateArgumentInfo& ArgumentInfo = FlattenedArguments[ArgumentIndex];
+					FRigVMTemplateArgument Argument = ArgumentInfo.GetArgument_NoLock(Context);
+					check(!Argument.bUseCategories);
+					
+					Ar << Argument.Name;
+
+					int32 Direction = static_cast<int32>(Argument.Direction);
+					Ar << Direction;
+
+					int32 NumTypes = Argument.TypeIndices.Num(); 
+					Ar << NumTypes;
+
+					for(int32 TypeIndex = 0; TypeIndex < NumTypes; TypeIndex++)
+					{
+						FRigVMTemplateArgumentType Type = GetType_NoLock(Argument.TypeIndices[TypeIndex]);
+						Ar << Type.CPPType;
+						Ar << Type.CPPTypeObject;
+					}
+				}
+			}
+		}
+
+		// functions
+		{
+			int32 NumFunctions = Functions.Num();
+			Ar << NumFunctions;
+			for(FRigVMFunction& Function : Functions)
+			{
+				Function.Serialize_NoLock(Ar, Context);
+			}
+		}
+	}
 }
 
 void FRigVMRegistry_NoLock::RefreshEngineTypesIfRequired_NoLock()
@@ -366,11 +614,26 @@ void FRigVMRegistry_NoLock::RefreshEngineTypesIfRequired_NoLock()
 	{
 		return;
 	}
+
 	RefreshEngineTypes_NoLock();
 }
 
 void FRigVMRegistry_NoLock::RefreshEngineTypes_NoLock()
 {
+#if !WITH_EDITOR
+	// if we are not in editor and VMs rely on their own localized registry
+	// we don't have to initialize all of the types here.
+	if(CVarRigVMEnableLocalizedRegistry.GetValueOnAnyThread() == true)
+	{
+		return;
+	}
+#endif
+
+	if(!bIsGlobalRegistry)
+	{
+		return;
+	}
+
 	TRACE_CPUPROFILER_EVENT_SCOPE(FRigVMRegistry::RefreshEngineTypes);
 
 	TGuardValue<bool> AvoidTypePropagationGuard(bAvoidTypePropagation, true);
@@ -446,7 +709,7 @@ void FRigVMRegistry_NoLock::RefreshEngineTypes_NoLock()
 				if(!TemplateProcessed[TemplateIndex])
 				{
 					FRigVMTemplate& Template = Templates[TemplateIndex];
-					(void)Template.UpdateAllArgumentTypesSlow();
+					(void)Template.UpdateAllArgumentTypesSlow(GetHandle_NoLock());
 					TemplateProcessed[TemplateIndex] = true;
 				}
 			}
@@ -479,7 +742,7 @@ bool FRigVMRegistry_NoLock::RefreshFunctionsAndDispatches_NoLock()
 		// the game.
 		if(Factory->CachedTemplate == nullptr)
 		{
-			(void)Factory->GetTemplate_NoLock();
+			(void)Factory->GetTemplate_NoLock(Context);
 			bRegistryChanged = true;
 		}
 	}
@@ -582,12 +845,16 @@ TRigVMTypeIndex FRigVMRegistry_NoLock::FindOrAddType_NoLock(const FRigVMTemplate
 	TRigVMTypeIndex Index = GetTypeIndex_NoLock(InType);
 	if(Index == INDEX_NONE)
 	{
+		const bool bRegisterAllArrayDimensions = bIsGlobalRegistry;
+		
+		static constexpr int32 MaxArrayDimension = 3;
+		
 		FRigVMTemplateArgumentType ElementType = InType;
 		while(ElementType.IsArray())
 		{
 			ElementType.ConvertToBaseElement();
 		}
-		
+
 		const UObject* CPPTypeObject = ElementType.CPPTypeObject;
 		if(!bForce && (CPPTypeObject != nullptr))
 		{
@@ -620,9 +887,10 @@ TRigVMTypeIndex FRigVMRegistry_NoLock::FindOrAddType_NoLock(const FRigVMTemplate
 			bIsExecute = ScriptStruct->IsChildOf(FRigVMExecutePin::StaticStruct());
 		}
 
-		TArray<TRigVMTypeIndex> Indices;
-		Indices.Reserve(3);
-		for (int32 ArrayDimension=0; ArrayDimension<3; ++ArrayDimension)
+		// 3 is the max array dimension allowed within RigVM
+		TArray<TRigVMTypeIndex, TInlineAllocator<3>> Indices;
+		TArray<bool, TInlineAllocator<3>> AddedTypeIndex;
+		for (int32 ArrayDimension=0; ArrayDimension<MaxArrayDimension; ++ArrayDimension)
 		{
 			if (bIsExecute && ArrayDimension > 1)
 			{
@@ -639,6 +907,17 @@ TRigVMTypeIndex FRigVMRegistry_NoLock::FindOrAddType_NoLock(const FRigVMTemplate
 			Info.Type = CurType;
 			Info.bIsArray = ArrayDimension > 0;
 			Info.bIsExecute = bIsExecute;
+
+			// if this is a localized registry we may already have this type.
+			if(!bRegisterAllArrayDimensions)
+			{
+				if(const TRigVMTypeIndex* ExistingIndex = TypeToIndex.Find(CurType))
+				{
+					Indices.Add(*ExistingIndex);
+					AddedTypeIndex.Add(false);
+					continue;
+				}
+			}
 			
 			Index = Types.Add(Info);
 #if UE_RIGVM_DEBUG_TYPEINDEX
@@ -647,15 +926,19 @@ TRigVMTypeIndex FRigVMRegistry_NoLock::FindOrAddType_NoLock(const FRigVMTemplate
 			TypeToIndex.Add(CurType, Index);
 
 			Indices.Add(Index);
+			AddedTypeIndex.Add(true);
 		}
 
-		Types[Indices[1]].BaseTypeIndex = Indices[0];
-		Types[Indices[0]].ArrayTypeIndex = Indices[1];
-
-		if (!bIsExecute)
+		if(Indices.Num() > 1)
 		{
-			Types[Indices[2]].BaseTypeIndex = Indices[1];
-			Types[Indices[1]].ArrayTypeIndex = Indices[2];
+			Types[Indices[1]].BaseTypeIndex = Indices[0];
+			Types[Indices[0]].ArrayTypeIndex = Indices[1];
+
+			if (!bIsExecute && Indices.Num() > 2)
+			{
+				Types[Indices[2]].BaseTypeIndex = Indices[1];
+				Types[Indices[1]].ArrayTypeIndex = Indices[2];
+			}
 		}
 
 		// update the categories first then propagate to TemplatesPerCategory once all categories up to date
@@ -666,12 +949,17 @@ TRigVMTypeIndex FRigVMRegistry_NoLock::FindOrAddType_NoLock(const FRigVMTemplate
 			ToPropagate.Emplace(InCategory, NewIndex);
 		}; 
 
-		for (int32 ArrayDimension=0; ArrayDimension<3; ++ArrayDimension)
+		for (int32 ArrayDimension=0; ArrayDimension<MaxArrayDimension; ++ArrayDimension)
 		{
 			if(bIsExecute && ArrayDimension > 1)
 			{
 				break;
 			}
+			if(!AddedTypeIndex[ArrayDimension])
+			{
+				continue;
+			}
+			
 			Index = Indices[ArrayDimension];
 
 			// Add to category
@@ -814,7 +1102,7 @@ TRigVMTypeIndex FRigVMRegistry_NoLock::FindOrAddType_NoLock(const FRigVMTemplate
 		// propagate new type to templates once they have all been added to the categories
 		for (const auto& [Category, NewIndex]: ToPropagate)
 		{
-			PropagateTypeAddedToCategory_NoLock(Category, NewIndex);	
+			PropagateTypeAddedToCategory_NoLock(Category, NewIndex);
 		}
 
 		// if the type is a structure
@@ -828,7 +1116,7 @@ TRigVMTypeIndex FRigVMRegistry_NoLock::FindOrAddType_NoLock(const FRigVMTemplate
 				{
 					// by creating a template argument for the child property
 					// the type will be added by calling ::FindOrAddType_Internal recursively.
-					(void)FRigVMTemplateArgument::Make_NoLock(Property, *this);
+					(void)FRigVMTemplateArgument::Make_NoLock(Property, Context);
 				}
 #if WITH_EDITOR
 				else
@@ -880,7 +1168,7 @@ void FRigVMRegistry_NoLock::PropagateTypeAddedToCategory_NoLock(const FRigVMTemp
 		for(const int32 TemplateIndex : TemplatesToUseType)
 		{
 			FRigVMTemplate& Template = Templates[TemplateIndex];
-			(void)Template.HandlePropagatedArgumentType(InTypeIndex);
+			(void)Template.HandlePropagatedArgumentType(InTypeIndex, GetHandle_NoLock());
 		}
 	}
 }
@@ -992,7 +1280,7 @@ void FRigVMRegistry_NoLock::RemoveTypeInCategory_NoLock(FRigVMTemplateArgument::
 	for (const int32 TemplateIndex : TemplatesToUseType)
 	{
 		FRigVMTemplate& Template = Templates[TemplateIndex];
-		Template.HandleTypeRemoval(InTypeIndex);
+		Template.HandleTypeRemoval(InTypeIndex, GetHandle_NoLock());
 	}
 }
 
@@ -1007,9 +1295,7 @@ void FRigVMRegistry_NoLock::OnEngineInit()
 // to avoid confusing certain compilers into creating multiple copies of the registry
 FRigVMRegistry_RWLock& FRigVMRegistry_RWLock::Get()
 {
-	// static in a function scope ensures that the GC system is initiated before 
-	// the registry constructor is called
-	static FRigVMRegistry_RWLock s_RigVMRegistry;
+	static FRigVMRegistry_RWLock s_RigVMRegistry(true /* global */);
 	return s_RigVMRegistry;
 }
 
@@ -1555,13 +1841,23 @@ void FRigVMRegistry_NoLock::Register_NoLock(const TCHAR* InName, FRigVMFunctionP
 
 	const FRigVMFunction Function(InName, InFunctionPtr, InStruct, Functions.Num(), InArguments);
 	Functions.AddElement(Function);
+	Functions[Function.Index].OwnerRegistry = this;
 	FunctionNameToIndex.Add(InName, Function.Index);
 
+#if !WITH_EDITOR
+	// if we are not in editor and VMs rely on their own localized registry
+	// let's only register the function pointer - but not the types etc... keep it as shallow as possible.
+	if(bIsGlobalRegistry && CVarRigVMEnableLocalizedRegistry.GetValueOnAnyThread() == true)
+	{
+		return;
+	}
+#endif
+
 	// register all of the types used by the function
 	for (TFieldIterator<FProperty> It(InStruct); It; ++It)
 	{
 		// creating the argument causes the registration
-		(void)FRigVMTemplateArgument::Make_NoLock(*It, *this);
+		(void)FRigVMTemplateArgument::Make_NoLock(*It, Context);
 	}
 
 #if WITH_EDITOR
@@ -1577,20 +1873,31 @@ void FRigVMRegistry_NoLock::Register_NoLock(const TCHAR* InName, FRigVMFunctionP
 		if (FString(InName).Split(TEXT("::"), nullptr, &MethodName))
 		{
 			const FString TemplateName = RigVMStringUtils::JoinStrings(TemplateMetadata, MethodName, TEXT("::"));
-			FRigVMTemplate Template(InStruct, TemplateName, Function.Index);
+			FRigVMTemplate Template(InStruct, TemplateName, Function.Index, Context);
 			if (Template.IsValid())
 			{
+				Functions[Function.Index].PermutationIndex = Template.NumPermutations_NoLock(Context) - 1;
+				
 				bool bWasMerged = false;
 
 				const int32* ExistingTemplateIndexPtr = NotationToIndex.Find(Template.GetNotation());
 				if(ExistingTemplateIndexPtr)
 				{
 					FRigVMTemplate& ExistingTemplate = TemplateArray[*ExistingTemplateIndexPtr];
-					if (ExistingTemplate.Merge(Template))
+					if (ExistingTemplate.Merge(Template, GetHandle_NoLock()))
 					{
 						if (!bIsDeprecated)
 						{
 							Functions[Function.Index].TemplateIndex = ExistingTemplate.Index;
+							Functions[Function.Index].PermutationIndex = ExistingTemplate.NumPermutations_NoLock(Context) - 1;
+#if WITH_EDITOR
+							int32 MaxNumPermutations = 0;
+							for(const FRigVMTemplateArgument& Argument : ExistingTemplate.Arguments)
+							{
+								MaxNumPermutations = FMath::Max(MaxNumPermutations, Argument.GetNumTypes_NoLock(Context));
+							}
+							verify(MaxNumPermutations == ExistingTemplate.NumPermutations_NoLock(Context));
+#endif
 						}
 						bWasMerged = true;
 					}
@@ -1617,7 +1924,7 @@ void FRigVMRegistry_NoLock::Register_NoLock(const TCHAR* InName, FRigVMFunctionP
 #endif
 }
 
-const FRigVMDispatchFactory* FRigVMRegistry_NoLock::RegisterFactory_NoLock(UScriptStruct* InFactoryStruct)
+const FRigVMDispatchFactory* FRigVMRegistry_NoLock::RegisterFactory_NoLock(UScriptStruct* InFactoryStruct, const TArray<FRigVMTemplateArgumentInfo>& InArgumentInfos)
 {
 	check(InFactoryStruct);
 	check(InFactoryStruct != FRigVMDispatchFactory::StaticStruct());
@@ -1646,17 +1953,33 @@ const FRigVMDispatchFactory* FRigVMRegistry_NoLock::RegisterFactory_NoLock(UScri
 		return nullptr;
 	}
 #endif
-	
+
 	FRigVMDispatchFactory* Factory = (FRigVMDispatchFactory*)FMemory::Malloc(InFactoryStruct->GetStructureSize());
 	InFactoryStruct->InitializeStruct(Factory, 1);
 	Factory->FactoryScriptStruct = InFactoryStruct;
-	Factories.Add(Factory);
-	Factory->RegisterDependencyTypes_NoLock(*this);
+	Factory->FactoryIndex = Factories.Add(Factory);
+	Factory->OwnerRegistry = this;
+	Factory->RegisterDependencyTypes_NoLock(Context);
+
+	if(!InArgumentInfos.IsEmpty())
+	{
+		Factory->CreateTemplateForArgumentInfos_NoLock(InArgumentInfos, Context);
+	}
+	
 	return Factory;
 }
 
 void FRigVMRegistry_NoLock::RegisterPredicate_NoLock(UScriptStruct* InStruct, const TCHAR* InName, const TArray<FRigVMFunctionArgument>& InArguments)
 {
+#if !WITH_EDITOR
+	// if we are not in editor and VMs rely on their own localized registry
+	// let's not register any predicates
+	if(bIsGlobalRegistry && CVarRigVMEnableLocalizedRegistry.GetValueOnAnyThread() == true)
+	{
+		return;
+	}
+#endif
+
 	// Make sure the predicate does not already exist
 	TArray<FRigVMFunction>& Predicates = StructNameToPredicates.FindOrAdd(InStruct->GetFName());
 	if (Predicates.ContainsByPredicate([InName](const FRigVMFunction& Predicate)
@@ -1741,15 +2064,15 @@ const FRigVMFunction* FRigVMRegistry_NoLock::FindFunction_NoLock(const TCHAR* In
 		// if the factory has never been registered - FindDispatchFactory will try to look it up and register
 		if(const FRigVMDispatchFactory* Factory = FindDispatchFactory_NoLock(*StructOrFactoryName))
 		{
-			if(const FRigVMTemplate* Template = Factory->GetTemplate_NoLock())
+			if(const FRigVMTemplate* Template = Factory->GetTemplate_NoLock(Context))
 			{
-				const FRigVMTemplateTypeMap ArgumentTypes = Template->GetArgumentTypesFromString_Impl(SuffixString, &InTypeResolver, false);
+				const FRigVMTemplateTypeMap ArgumentTypes = Template->GetArgumentTypesFromString_NoLock(SuffixString, &InTypeResolver, Context);
 				if(ArgumentTypes.Num() == Template->NumArguments())
 				{
-					const int32 PermutationIndex = Template->FindPermutation(ArgumentTypes, false);
+					const int32 PermutationIndex = Template->FindPermutation_NoLock(ArgumentTypes, Context);
 					if(PermutationIndex != INDEX_NONE)
 					{
-						return ((FRigVMTemplate*)Template)->GetOrCreatePermutation_NoLock(PermutationIndex);
+						return ((FRigVMTemplate*)Template)->GetOrCreatePermutation_NoLock(PermutationIndex, Context);
 					}
 				}
 			}
@@ -1793,8 +2116,8 @@ const FRigVMFunction* FRigVMRegistry_NoLock::FindFunction_NoLock(const TCHAR* In
 					const FRigVMFunction* RedirectedFunction = FindFunction_NoLock(*(NewStructOrFactoryName + TEXT("::") + SuffixString), InTypeResolver);
 					if(RedirectedFunction)
 					{
-						FRigVMRegistry_NoLock& MutableRegistry = Get(LockType_Write);
-						MutableRegistry.FunctionNameToIndex.Add(InName, RedirectedFunction->Index);
+						FRigVMRegistry_NoLock* MutableRegistry = const_cast<FRigVMRegistry_NoLock*>(this);
+						MutableRegistry->FunctionNameToIndex.Add(InName, RedirectedFunction->Index);
 						return RedirectedFunction;
 					}
 				}
@@ -1857,7 +2180,7 @@ const FRigVMTemplate* FRigVMRegistry_NoLock::FindTemplate_NoLock(const FName& In
 		
 		if(const FRigVMDispatchFactory* Factory = FindDispatchFactory_NoLock(*FactoryName))
 		{
-			return Factory->GetTemplate_NoLock();
+			return Factory->GetTemplate_NoLock(Context);
 		}
 	}
 
@@ -1888,7 +2211,7 @@ const FRigVMTemplate* FRigVMRegistry_NoLock::FindTemplate_NoLock(const FName& In
 				const FString NewDispatchFactoryName = FRigVMDispatchFactory::DispatchPrefix + Redirect->NewName.ObjectName.ToString();
 				if(const FRigVMDispatchFactory* NewDispatchFactory = FindDispatchFactory_NoLock(*NewDispatchFactoryName))
 				{
-					return NewDispatchFactory->GetTemplate_NoLock();
+					return NewDispatchFactory->GetTemplate_NoLock(Context);
 				}
 			}
 		}
@@ -1953,11 +2276,11 @@ const FRigVMTemplate* FRigVMRegistry_NoLock::AddTemplateFromArguments_NoLock(con
 {
 	// we only support to ask for templates here which provide singleton types
 	int32 NumPermutations = 0;
-	FRigVMTemplate Template(InName, InInfos);
+	FRigVMTemplate Template(InName, InInfos, GetHandle_NoLock());
 	for(const FRigVMTemplateArgument& Argument : Template.Arguments)
 	{
-		const int32 NumIndices = Argument.GetNumTypes_NoLock();
-		if(!Argument.IsSingleton_NoLock() && NumPermutations > 1)
+		const int32 NumIndices = Argument.GetNumTypes_NoLock(GetHandle_NoLock());
+		if(!Argument.IsSingleton_NoLock({}, GetHandle_NoLock()) && NumPermutations > 1)
 		{
 			if(NumIndices != NumPermutations)
 			{
@@ -1971,10 +2294,10 @@ const FRigVMTemplate* FRigVMRegistry_NoLock::AddTemplateFromArguments_NoLock(con
 	// if any of the arguments are wildcards we'll need to update the types
 	for(FRigVMTemplateArgument& Argument : Template.Arguments)
 	{
-		const int32 NumTypes = Argument.GetNumTypes_NoLock();
+		const int32 NumTypes = Argument.GetNumTypes_NoLock(GetHandle_NoLock());
 		if(NumTypes == 1)
 		{
-			const TRigVMTypeIndex FirstTypeIndex = Argument.GetTypeIndex_NoLock(0);
+			const TRigVMTypeIndex FirstTypeIndex = Argument.GetTypeIndex_NoLock(0, GetHandle_NoLock());
 			if(IsWildCardType_NoLock(FirstTypeIndex))
 			{
 #if WITH_EDITOR
@@ -1991,7 +2314,7 @@ const FRigVMTemplate* FRigVMRegistry_NoLock::AddTemplateFromArguments_NoLock(con
 				Argument.bUseCategories = true;
 				Argument.TypeIndices.Reset();
 		
-				NumPermutations = FMath::Max(NumPermutations, Argument.GetNumTypes_NoLock());
+				NumPermutations = FMath::Max(NumPermutations, Argument.GetNumTypes_NoLock(GetHandle_NoLock()));
 			}
 		}
 	}
@@ -2019,7 +2342,7 @@ const FRigVMTemplate* FRigVMRegistry_NoLock::AddTemplateFromArguments_NoLock(con
 		bool bAnyArgumentWithZeroTypes = false;
 		for(int32 ArgIndex = 0; ArgIndex < NumArguments; ArgIndex++)
 		{
-			(void)TypeCaches[ArgIndex].UpdateIfRequired(Template.Arguments[ArgIndex]);
+			(void)TypeCaches[ArgIndex].UpdateIfRequired(GetHandle_NoLock(), Template.Arguments[ArgIndex]);
 			bAnyArgumentWithZeroTypes = bAnyArgumentWithZeroTypes || TypeCaches[ArgIndex].GetNumTypes_NoLock() == 0;
 		}
 
@@ -2077,15 +2400,16 @@ const FRigVMTemplate* FRigVMRegistry_NoLock::AddTemplateFromArguments_NoLock(con
 #if WITH_EDITOR
 	for(FRigVMTemplateArgument& Argument : Template.Arguments)
 	{
-		Argument.UpdateTypeToPermutationsSlow();
+		Argument.UpdateTypeToPermutationsSlow(GetHandle_NoLock());
 	}
 #endif
 
 	Template.Permutations.Init(INDEX_NONE, NumPermutations);
-	Template.RecomputeTypesHashToPermutations(TypeCaches);
+	Template.RecomputeTypesHashToPermutations(TypeCaches, GetHandle_NoLock());
 
 	const int32 Index = Templates.AddElement(Template);
 	Templates[Index].Index = Index;
+	Templates[Index].OwnerRegistry = this;
 	Templates[Index].Delegates = InDelegates;
 	TemplateNotationToIndex.Add(Template.GetNotation(), Index);
 
@@ -2111,6 +2435,11 @@ FRigVMDispatchFactory* FRigVMRegistry_NoLock::FindDispatchFactory_NoLock(const F
 		return *FactoryPtr;
 	}
 
+	if(!bIsGlobalRegistry)
+	{
+		return nullptr;
+	}
+	
 	FString FactoryName = InFactoryName.ToString();
 	
 	// if the factory has never been registered - we should try to look it up	
@@ -2138,12 +2467,12 @@ FString FRigVMRegistry_NoLock::FindOrAddSingletonDispatchFunction_NoLock(UScript
 	{
 		if(Factory->IsSingleton())
 		{
-			if(const FRigVMTemplate* Template = Factory->GetTemplate_NoLock())
+			if(const FRigVMTemplate* Template = Factory->GetTemplate_NoLock(Context))
 			{
 				// use the types for the first permutation - since we don't care
 				// for a singleton dispatch
-				const FRigVMTemplateTypeMap TypesForPrimaryPermutation = Template->GetTypesForPermutation_NoLock(0);
-				const FString Name = Factory->GetPermutationName(TypesForPrimaryPermutation, false);
+				const FRigVMTemplateTypeMap TypesForPrimaryPermutation = Template->GetTypesForPermutation_NoLock(0, GetHandle_NoLock());
+				const FString Name = Factory->GetPermutationName_NoLock(TypesForPrimaryPermutation, Context);
 				if(const FRigVMFunction* Function = FindFunction_NoLock(*Name))
 				{
 					return Function->Name;
@@ -2163,20 +2492,24 @@ const TArray<FRigVMFunction>* FRigVMRegistry_NoLock::GetPredicatesForStruct_NoLo
 	return StructNameToPredicates.Find(InStructName);
 }
 
-FRigVMRegistry_RWLock::FRigVMRegistry_RWLock()
-	: FRigVMRegistry_NoLock()
+FRigVMRegistry_RWLock::FRigVMRegistry_RWLock(bool InIsGlobalRegistry)
+	: FRigVMRegistry_NoLock(InIsGlobalRegistry)
 {
-	Initialize(true);
+	Initialize();
 }
 
-void FRigVMRegistry_RWLock::Initialize(bool bLockRegistry)
+void FRigVMRegistry_RWLock::Initialize()
 {
 	LockType = LockType_Invalid;
 	LockCount = 0;
 	
-	const FConditionalWriteScopeLock _(*this, bLockRegistry);
-	
+	FRigVMRegistryWriteLock WriteLock(*this);
 	Initialize_NoLock();
+}
+
+void FRigVMRegistry_RWLock::Initialize_NoLock()
+{
+	FRigVMRegistry_NoLock::InitializeBaseTypes_NoLock();
 	
 	const FAssetRegistryModule& AssetRegistryModule = FModuleManager::LoadModuleChecked<FAssetRegistryModule>(TEXT("AssetRegistry"));
 	AssetRegistryModule.Get().OnAssetRemoved().AddRaw(this, &FRigVMRegistry_RWLock::OnAssetRemoved);
diff --git a/Engine/Plugins/Runtime/RigVM/Source/RigVM/Private/RigVMCore/RigVMStruct.cpp b/Engine/Plugins/Runtime/RigVM/Source/RigVM/Private/RigVMCore/RigVMStruct.cpp
--- a/Engine/Plugins/Runtime/RigVM/Source/RigVM/Private/RigVMCore/RigVMStruct.cpp
+++ b/Engine/Plugins/Runtime/RigVM/Source/RigVM/Private/RigVMCore/RigVMStruct.cpp
@@ -99,8 +99,9 @@ FRigVMStructUpgradeInfo FRigVMStructUpgradeInfo::MakeFromStructToFactory(UScript
 {
 	check(InRigVMStruct);
 	check(InFactoryStruct);
-	
-	const FRigVMDispatchFactory* Factory = FRigVMRegistry::Get().FindOrAddDispatchFactory(InFactoryStruct);
+
+	FRigVMRegistry& Registry = FRigVMRegistry::Get();
+	const FRigVMDispatchFactory* Factory = Registry.FindOrAddDispatchFactory(InFactoryStruct);
 	check(Factory);
 	const FRigVMTemplateTypeMap Types = GetTypeMapFromStruct(InRigVMStruct);
 	const FString PermutationName = Factory->GetPermutationName(Types);
