RigVM: Localized registry in cooked environments - first steps.

diff --git a/Engine/Plugins/Runtime/RigVM/Source/RigVM/Private/RigVMCore/RigVMTemplate.cpp b/Engine/Plugins/Runtime/RigVM/Source/RigVM/Private/RigVMCore/RigVMTemplate.cpp
--- a/Engine/Plugins/Runtime/RigVM/Source/RigVM/Private/RigVMCore/RigVMTemplate.cpp
+++ b/Engine/Plugins/Runtime/RigVM/Source/RigVM/Private/RigVMCore/RigVMTemplate.cpp
@@ -40,6 +40,12 @@ FRigVMTemplateArgumentType::FRigVMTemplateArgumentType(const FName& InCPPType, U
 #endif
 }
 
+void FRigVMTemplateArgumentType::Serialize_NoLock(FArchive& Ar)
+{
+	Ar << CPPType;
+	Ar << CPPTypeObject;
+}
+
 FRigVMTemplateArgument::FRigVMTemplateArgument()
 {}
 
@@ -48,7 +54,7 @@ FRigVMTemplateArgument::FRigVMTemplateArgument(const FName& InName, ERigVMPinDir
 	, Direction(InDirection)
 {}
 
-FRigVMTemplateArgument::FRigVMTemplateArgument(FProperty* InProperty, FRigVMRegistry_NoLock& InRegistry)
+FRigVMTemplateArgument::FRigVMTemplateArgument(FProperty* InProperty, FRigVMRegistryHandle& InRegistry)
 	: Name(InProperty->GetFName())
 {
 #if WITH_EDITOR
@@ -93,55 +99,50 @@ FRigVMTemplateArgument::FRigVMTemplateArgument(FProperty* InProperty, FRigVMRegi
 	}
 	
 	const FRigVMTemplateArgumentType Type(CPPTypeName, CPPTypeObject);
-	const TRigVMTypeIndex TypeIndex = InRegistry.FindOrAddType_NoLock(Type, true); 
+	const TRigVMTypeIndex TypeIndex = InRegistry->FindOrAddType_NoLock(Type, true); 
 
 	TypeIndices.Add(TypeIndex);
 	EnsureValidExecuteType_NoLock(InRegistry);
 #if WITH_EDITOR
-	UpdateTypeToPermutationsSlow();
+	UpdateTypeToPermutationsSlow(InRegistry);
 #endif
 }
 
 FRigVMTemplateArgument FRigVMTemplateArgument::Make(FProperty* InProperty)
 {
-	FRigVMRegistryWriteLock WriteLock;
-	return Make_NoLock(InProperty, WriteLock.GetRegistry());
-}
-
-FRigVMTemplateArgument FRigVMTemplateArgument::Make_NoLock(FProperty* InProperty)
-{
-	return Make_NoLock(InProperty, FRigVMRegistry_NoLock::GetForWrite());
+	FRigVMRegistryReadLock Lock;;
+	return Make_NoLock(InProperty, Lock);
 }
 
-FRigVMTemplateArgument FRigVMTemplateArgument::Make_NoLock(FProperty* InProperty, FRigVMRegistry_NoLock& InRegistry)
+FRigVMTemplateArgument FRigVMTemplateArgument::Make_NoLock(FProperty* InProperty, FRigVMRegistryHandle& InRegistry)
 {
 	return FRigVMTemplateArgument(InProperty, InRegistry);
 }
 
-FRigVMTemplateArgument::FRigVMTemplateArgument(const FName& InName, ERigVMPinDirection InDirection, TRigVMTypeIndex InTypeIndex)
+FRigVMTemplateArgument::FRigVMTemplateArgument(const FName& InName, ERigVMPinDirection InDirection, TRigVMTypeIndex InTypeIndex, const FRigVMRegistryHandle& InRegistry)
 	: Name(InName)
 	, Direction(InDirection)
 	, TypeIndices({InTypeIndex})
 {
-	EnsureValidExecuteType_NoLock(FRigVMRegistry_RWLock::Get());
+	EnsureValidExecuteType_NoLock(InRegistry);
 #if WITH_EDITOR
-	UpdateTypeToPermutationsSlow();
+	UpdateTypeToPermutationsSlow(InRegistry);
 #endif
 }
 
-FRigVMTemplateArgument::FRigVMTemplateArgument(const FName& InName, ERigVMPinDirection InDirection, const TArray<TRigVMTypeIndex>& InTypeIndices)
+FRigVMTemplateArgument::FRigVMTemplateArgument(const FName& InName, ERigVMPinDirection InDirection, const TArray<TRigVMTypeIndex>& InTypeIndices, const FRigVMRegistryHandle& InRegistry)
 	: Name(InName)
 	, Direction(InDirection)
 	, TypeIndices(InTypeIndices)
 {
 	check(TypeIndices.Num() > 0);
-	EnsureValidExecuteType_NoLock(FRigVMRegistry_RWLock::Get());
+	EnsureValidExecuteType_NoLock(InRegistry);
 #if WITH_EDITOR
-	UpdateTypeToPermutationsSlow();
+	UpdateTypeToPermutationsSlow(InRegistry);
 #endif
 }
 
-FRigVMTemplateArgument::FRigVMTemplateArgument(const FName& InName, ERigVMPinDirection InDirection, const TArray<ETypeCategory>& InTypeCategories, TFunction<bool(const TRigVMTypeIndex&)> InFilterType)
+FRigVMTemplateArgument::FRigVMTemplateArgument(const FName& InName, ERigVMPinDirection InDirection, const TArray<ETypeCategory>& InTypeCategories, TFunction<bool(const TRigVMTypeIndex&, const FRigVMRegistryHandle&)> InFilterType, const FRigVMRegistryHandle& InRegistry)
 	: Name(InName)
 	, Direction(InDirection)
 	, TypeCategories(InTypeCategories)
@@ -164,14 +165,15 @@ FRigVMTemplateArgument::FRigVMTemplateArgument(const FName& InName, ERigVMPinDir
 				int32 NumTypes = 0;
 				for (const ETypeCategory TypeCategory : TypeCategories)
 				{
-					NumTypes += FRigVMRegistry_NoLock::GetForRead().GetTypesForCategory_NoLock(TypeCategory).Num();
+					NumTypes += InRegistry->GetTypesForCategory_NoLock(TypeCategory).Num();
 				}
 				AllTypes.Reserve(NumTypes);
 			}
 
 			for (const ETypeCategory TypeCategory : TypeCategories)
 			{
-				const TArray<TRigVMTypeIndex>& Types = FRigVMRegistry_NoLock::GetForRead().GetTypesForCategory_NoLock(TypeCategory);
+				const TArray<TRigVMTypeIndex>& Types = InRegistry->GetTypesForCategory_NoLock(TypeCategory);
+				
 				for (const TRigVMTypeIndex Type: Types)
 				{
 					AllTypes.Add(Type);
@@ -192,51 +194,51 @@ FRigVMTemplateArgument::FRigVMTemplateArgument(const FName& InName, ERigVMPinDir
 			TArray<TRigVMTypeIndex> Indices = AllTypes.Array();
 			if (FilterType)
 			{
-				Indices = Indices.FilterByPredicate([this](const TRigVMTypeIndex& Type)
+				Indices = Indices.FilterByPredicate([this, &InRegistry](const TRigVMTypeIndex& Type)
 				{
-					return FilterType(Type);
+					return FilterType(Type, InRegistry);
 				});
 			}
 			TypeIndices = MoveTemp(Indices);
-			EnsureValidExecuteType_NoLock(FRigVMRegistry_RWLock::Get());
+			EnsureValidExecuteType_NoLock(InRegistry);
 		}
 
 #if WITH_EDITOR
-		UpdateTypeToPermutationsSlow();
+		UpdateTypeToPermutationsSlow(InRegistry);
 #endif
 	}
 }
 
-void FRigVMTemplateArgument::EnsureValidExecuteType_NoLock(FRigVMRegistry_NoLock& InRegistry)
+void FRigVMTemplateArgument::EnsureValidExecuteType_NoLock(const FRigVMRegistryHandle& InRegistry)
 {
 	for(TRigVMTypeIndex& TypeIndex : TypeIndices)
 	{
-		InRegistry.ConvertExecuteContextToBaseType_NoLock(TypeIndex);
+		InRegistry->ConvertExecuteContextToBaseType_NoLock(TypeIndex);
 	}
 }
 
 #if WITH_EDITOR
 
-void FRigVMTemplateArgument::UpdateTypeToPermutationsSlow()
+void FRigVMTemplateArgument::UpdateTypeToPermutationsSlow(const FRigVMRegistryHandle& InRegistry)
 {
 	TypeToPermutations.Reset();
-	TypeToPermutations.Reserve(GetNumTypes_NoLock());
+	TypeToPermutations.Reserve(GetNumTypes_NoLock(InRegistry));
 
 	int32 TypeIndex = 0;
 	ForEachType([&](const TRigVMTypeIndex Type)
 	{
 		TypeToPermutations.FindOrAdd(Type).Add(TypeIndex++);
 		return true;
-	});
+	}, InRegistry);
 }
 
-bool FRigVMTemplateArgument::SupportsTypeIndex(TRigVMTypeIndex InTypeIndex, TRigVMTypeIndex* OutTypeIndex, const bool bLockRegistry) const
+bool FRigVMTemplateArgument::SupportsTypeIndex(TRigVMTypeIndex InTypeIndex, TRigVMTypeIndex* OutTypeIndex) const
 {
-	const FRigVMRegistryReadLock _(bLockRegistry);
-	return SupportsTypeIndex_NoLock(InTypeIndex, OutTypeIndex);
+	FRigVMRegistryReadLock Lock;;
+	return SupportsTypeIndex_NoLock(Lock, InTypeIndex, OutTypeIndex);
 }
 
-bool FRigVMTemplateArgument::SupportsTypeIndex_NoLock(TRigVMTypeIndex InTypeIndex, TRigVMTypeIndex* OutTypeIndex) const
+bool FRigVMTemplateArgument::SupportsTypeIndex_NoLock(const FRigVMRegistryHandle& InRegistry, TRigVMTypeIndex InTypeIndex, TRigVMTypeIndex* OutTypeIndex) const
 {
 	if(InTypeIndex == INDEX_NONE)
 	{
@@ -244,37 +246,36 @@ bool FRigVMTemplateArgument::SupportsTypeIndex_NoLock(TRigVMTypeIndex InTypeInde
 	}
 	
 	// convert any execute type into the base execute
-	const FRigVMRegistry_NoLock& Registry = FRigVMRegistry_NoLock::GetForRead();
-	if(Registry.IsExecuteType_NoLock(InTypeIndex))
+	if(InRegistry->IsExecuteType_NoLock(InTypeIndex))
 	{
-		const bool bIsArray = Registry.IsArrayType_NoLock(InTypeIndex);
+		const bool bIsArray = InRegistry->IsArrayType_NoLock(InTypeIndex);
 		InTypeIndex = RigVMTypeUtils::TypeIndex::Execute;
 		if(bIsArray)
 		{
-			InTypeIndex = Registry.GetArrayTypeFromBaseTypeIndex_NoLock(InTypeIndex);
+			InTypeIndex = InRegistry->GetArrayTypeFromBaseTypeIndex_NoLock(InTypeIndex);
 		}
 	}
 
-	const TArray<int32>& Permutations = GetPermutations_NoLock(InTypeIndex);
+	const TArray<int32>& Permutations = GetPermutations_NoLock(InTypeIndex, InRegistry);
 	if (!Permutations.IsEmpty())
 	{
 		if(OutTypeIndex)
 		{
-			(*OutTypeIndex) = GetTypeIndex_NoLock(Permutations[0]);
+			(*OutTypeIndex) = GetTypeIndex_NoLock(Permutations[0], InRegistry);
 		}
 		return true;
 	}
 
 	// Try to find compatible type
-	const TArray<TRigVMTypeIndex>& CompatibleTypes = Registry.GetCompatibleTypes_NoLock(InTypeIndex);
+	const TArray<TRigVMTypeIndex>& CompatibleTypes = InRegistry->GetCompatibleTypes_NoLock(InTypeIndex);
 	for (const TRigVMTypeIndex& CompatibleTypeIndex : CompatibleTypes)
 	{
-		const TArray<int32>& CompatiblePermutations = GetPermutations_NoLock(CompatibleTypeIndex);
+		const TArray<int32>& CompatiblePermutations = GetPermutations_NoLock(CompatibleTypeIndex, InRegistry);
 		if (!CompatiblePermutations.IsEmpty())
 		{
 			if(OutTypeIndex)
 			{
-				(*OutTypeIndex) = GetTypeIndex_NoLock(CompatiblePermutations[0]);
+				(*OutTypeIndex) = GetTypeIndex_NoLock(CompatiblePermutations[0], InRegistry);
 			}
 			return true;
 		}
@@ -285,13 +286,13 @@ bool FRigVMTemplateArgument::SupportsTypeIndex_NoLock(TRigVMTypeIndex InTypeInde
 
 #endif
 
-bool FRigVMTemplateArgument::IsSingleton(const TArray<int32>& InPermutationIndices, const bool bLockRegistry) const
+bool FRigVMTemplateArgument::IsSingleton(const TArray<int32>& InPermutationIndices) const
 {
-	const FRigVMRegistryReadLock _(bLockRegistry);
-	return IsSingleton_NoLock(InPermutationIndices);
+	FRigVMRegistryReadLock Lock;;
+	return IsSingleton_NoLock(InPermutationIndices, Lock);
 }
 
-bool FRigVMTemplateArgument::IsSingleton_NoLock(const TArray<int32>& InPermutationIndices) const
+bool FRigVMTemplateArgument::IsSingleton_NoLock(const TArray<int32>& InPermutationIndices, const FRigVMRegistryHandle& InRegistry) const
 {
 #if WITH_EDITOR
 	if (TypeToPermutations.Num() == 1)
@@ -308,11 +309,11 @@ bool FRigVMTemplateArgument::IsSingleton_NoLock(const TArray<int32>& InPermutati
 	}
 
 	const bool bUsesPermutations = !InPermutationIndices.IsEmpty(); 
-	const int32 NumPermutations = bUsesPermutations ? InPermutationIndices.Num() : GetNumTypes_NoLock();
-	const TRigVMTypeIndex InType0 = GetTypeIndex_NoLock(bUsesPermutations ? InPermutationIndices[0] : 0);
+	const int32 NumPermutations = bUsesPermutations ? InPermutationIndices.Num() : GetNumTypes_NoLock(InRegistry);
+	const TRigVMTypeIndex InType0 = GetTypeIndex_NoLock(bUsesPermutations ? InPermutationIndices[0] : 0, InRegistry);
 	for (int32 PermutationIndex = 1; PermutationIndex < NumPermutations; PermutationIndex++)
 	{
-		if (GetTypeIndex_NoLock(bUsesPermutations ? InPermutationIndices[PermutationIndex] : PermutationIndex) != InType0)
+		if (GetTypeIndex_NoLock(bUsesPermutations ? InPermutationIndices[PermutationIndex] : PermutationIndex, InRegistry) != InType0)
 		{
 			return false;
 		}
@@ -323,41 +324,38 @@ bool FRigVMTemplateArgument::IsSingleton_NoLock(const TArray<int32>& InPermutati
 
 bool FRigVMTemplateArgument::IsExecute() const
 {
-	const FRigVMRegistry_NoLock& Registry = FRigVMRegistry_NoLock::GetForRead();
-	return IsExecute_NoLock(Registry);
+	return IsExecute_NoLock(FRigVMRegistryReadLock());
 }
 
-bool FRigVMTemplateArgument::IsExecute_NoLock(const FRigVMRegistry_NoLock& InRegistry) const
+bool FRigVMTemplateArgument::IsExecute_NoLock(const FRigVMRegistryHandle& InRegistry) const
 {
 	const int32 FoundAnyNotExec = IndexOfByPredicate([&InRegistry](const TRigVMTypeIndex Type)
 	{
-		return !InRegistry.IsExecuteType_NoLock(Type);
-	});
+		return !InRegistry->IsExecuteType_NoLock(Type);
+	}, InRegistry);
 	return FoundAnyNotExec == INDEX_NONE;
 }
 
-FRigVMTemplateArgument::EArrayType FRigVMTemplateArgument::GetArrayType(const bool bLockRegistry) const
+FRigVMTemplateArgument::EArrayType FRigVMTemplateArgument::GetArrayType() const
 {
-	const FRigVMRegistryReadLock _(bLockRegistry);
-	return GetArrayType_NoLock();
+	return GetArrayType_NoLock(FRigVMRegistryReadLock());
 }
 
-FRigVMTemplateArgument::EArrayType FRigVMTemplateArgument::GetArrayType_NoLock() const
+FRigVMTemplateArgument::EArrayType FRigVMTemplateArgument::GetArrayType_NoLock(const FRigVMRegistryHandle& InRegistry) const
 {
 	if(CachedArrayType.IsSet())
 	{
 		return CachedArrayType.GetValue();
 	}
 	
-	TArray<TRigVMTypeIndex> Types; GetAllTypes_NoLock(Types);
+	TArray<TRigVMTypeIndex> Types; GetAllTypes_NoLock(Types, InRegistry);
 
 	const int32 NumTypes = Types.Num();
 	if (NumTypes > 0)
 	{
-		const FRigVMRegistry_NoLock& Registry = FRigVMRegistry_NoLock::GetForRead();
-		const EArrayType ArrayType = Registry.IsArrayType_NoLock(Types[0]) ? EArrayType_ArrayValue : EArrayType_SingleValue;
+		const EArrayType ArrayType = InRegistry->IsArrayType_NoLock(Types[0]) ? EArrayType_ArrayValue : EArrayType_SingleValue;
 		
-		if(IsSingleton_NoLock())
+		if(IsSingleton_NoLock({}, InRegistry))
 		{
 			CachedArrayType = ArrayType;
 			return CachedArrayType.GetValue();
@@ -372,7 +370,7 @@ FRigVMTemplateArgument::EArrayType FRigVMTemplateArgument::GetArrayType_NoLock()
 				continue;
 			}
 			
-			const EArrayType OtherArrayType = Registry.IsArrayType_NoLock(TypeIndex) ? EArrayType_ArrayValue : EArrayType_SingleValue;
+			const EArrayType OtherArrayType = InRegistry->IsArrayType_NoLock(TypeIndex) ? EArrayType_ArrayValue : EArrayType_SingleValue;
 			if(OtherArrayType != ArrayType)
 			{
 				CachedArrayType = EArrayType_Mixed;
@@ -389,13 +387,12 @@ FRigVMTemplateArgument::EArrayType FRigVMTemplateArgument::GetArrayType_NoLock()
 
 #if WITH_EDITOR
 
-const TArray<int32>& FRigVMTemplateArgument::GetPermutations(const TRigVMTypeIndex InType, bool bLockRegistry) const
+const TArray<int32>& FRigVMTemplateArgument::GetPermutations(const TRigVMTypeIndex InType) const
 {
-	FRigVMRegistryReadLock _(bLockRegistry);
-	return GetPermutations_NoLock(InType);
+	return GetPermutations_NoLock(InType, FRigVMRegistryReadLock());
 }
 
-const TArray<int32>& FRigVMTemplateArgument::GetPermutations_NoLock(const TRigVMTypeIndex InType) const
+const TArray<int32>& FRigVMTemplateArgument::GetPermutations_NoLock(const TRigVMTypeIndex InType, const FRigVMRegistryHandle& InRegistry) const
 {
 	if (const TArray<int32>* Found = TypeToPermutations.Find(InType))
 	{
@@ -412,7 +409,7 @@ const TArray<int32>& FRigVMTemplateArgument::GetPermutations_NoLock(const TRigVM
 		}
 		IndexInTypes++;
 		return true;
-	});
+	}, InRegistry);
 
 	if (!Permutations.IsEmpty())
 	{
@@ -430,13 +427,12 @@ void FRigVMTemplateArgument::InvalidatePermutations(const TRigVMTypeIndex InType
 
 #endif
 
-void FRigVMTemplateArgument::GetAllTypes(TArray<TRigVMTypeIndex>& OutTypes, const bool bLockRegistry) const
+void FRigVMTemplateArgument::GetAllTypes(TArray<TRigVMTypeIndex>& OutTypes) const
 {
-	const FRigVMRegistryReadLock _(bLockRegistry);
-	GetAllTypes_NoLock(OutTypes);
+	GetAllTypes_NoLock(OutTypes, FRigVMRegistryReadLock());
 }
 
-void FRigVMTemplateArgument::GetAllTypes_NoLock(TArray<TRigVMTypeIndex>& OutTypes) const
+void FRigVMTemplateArgument::GetAllTypes_NoLock(TArray<TRigVMTypeIndex>& OutTypes, const FRigVMRegistryHandle& InRegistry) const
 {
 	if (!bUseCategories)
 	{
@@ -449,14 +445,14 @@ void FRigVMTemplateArgument::GetAllTypes_NoLock(TArray<TRigVMTypeIndex>& OutType
 	{
 		if (FilterType == nullptr)
 		{
-			OutTypes.Append(FRigVMRegistry_NoLock::GetForRead().GetTypesForCategory_NoLock(Category));
+			OutTypes.Append(InRegistry->GetTypesForCategory_NoLock(Category));
 		}
 		else
 		{
-			const TArray<TRigVMTypeIndex>& CategoryTypes = FRigVMRegistry_NoLock::GetForRead().GetTypesForCategory_NoLock(Category);
+			const TArray<TRigVMTypeIndex>& CategoryTypes = InRegistry->GetTypesForCategory_NoLock(Category);
 			for (const TRigVMTypeIndex& Type : CategoryTypes)
 			{
-				if (FilterType(Type))
+				if (FilterType(Type, InRegistry))
 				{
 					OutTypes.Add(Type);
 				}
@@ -465,13 +461,12 @@ void FRigVMTemplateArgument::GetAllTypes_NoLock(TArray<TRigVMTypeIndex>& OutType
 	}
 }
 
-TRigVMTypeIndex FRigVMTemplateArgument::GetTypeIndex(const int32 InIndex, bool bLockRegistry) const
+TRigVMTypeIndex FRigVMTemplateArgument::GetTypeIndex(const int32 InIndex) const
 {
-	const FRigVMRegistryReadLock _(bLockRegistry);
-	return GetTypeIndex_NoLock(InIndex);
+	return GetTypeIndex_NoLock(InIndex, FRigVMRegistryReadLock());
 }
 
-TRigVMTypeIndex FRigVMTemplateArgument::GetTypeIndex_NoLock(const int32 InIndex) const
+TRigVMTypeIndex FRigVMTemplateArgument::GetTypeIndex_NoLock(const int32 InIndex, const FRigVMRegistryHandle& InRegistry) const
 {
 	if (!bUseCategories)
 	{
@@ -483,9 +478,9 @@ TRigVMTypeIndex FRigVMTemplateArgument::GetTypeIndex_NoLock(const int32 InIndex)
 	{
 		TRigVMTypeIndex ValidType = INDEX_NONE;
 		int32 ValidIndex = 0;
-		CategoryViews(TypeCategories).ForEachType([this, &ValidIndex, InIndex, &ValidType](const TRigVMTypeIndex& Type) -> bool
+		CategoryViews(TypeCategories, InRegistry).ForEachType([this, &ValidIndex, InIndex, &ValidType, &InRegistry](const TRigVMTypeIndex& Type) -> bool
 		{
-			if (FilterType(Type))
+			if (FilterType(Type, InRegistry))
 			{
 				if (ValidIndex == InIndex)
 				{
@@ -499,22 +494,21 @@ TRigVMTypeIndex FRigVMTemplateArgument::GetTypeIndex_NoLock(const int32 InIndex)
 		return ValidType;
 	}
 
-	return CategoryViews(TypeCategories).GetTypeIndex(InIndex);
+	return CategoryViews(TypeCategories, InRegistry).GetTypeIndex(InIndex);
 }
 
-TOptional<TRigVMTypeIndex> FRigVMTemplateArgument::TryToGetTypeIndex(const int32 InIndex, const bool bLockRegistry) const
+TOptional<TRigVMTypeIndex> FRigVMTemplateArgument::TryToGetTypeIndex(const int32 InIndex) const
 {
-	const FRigVMRegistryReadLock _(bLockRegistry);
-	return TryToGetTypeIndex_NoLock(InIndex);
+	return TryToGetTypeIndex_NoLock(InIndex, FRigVMRegistryReadLock());
 }
 
-TOptional<TRigVMTypeIndex> FRigVMTemplateArgument::TryToGetTypeIndex_NoLock(const int32 InIndex) const
+TOptional<TRigVMTypeIndex> FRigVMTemplateArgument::TryToGetTypeIndex_NoLock(const int32 InIndex, const FRigVMRegistryHandle& InRegistry) const
 {
-	if(IsSingleton_NoLock())
+	if(IsSingleton_NoLock({}, InRegistry))
 	{
-		return GetTypeIndex_NoLock(0);
+		return GetTypeIndex_NoLock(0, InRegistry);
 	}
-	const TRigVMTypeIndex TypeIndex = GetTypeIndex_NoLock(InIndex);
+	const TRigVMTypeIndex TypeIndex = GetTypeIndex_NoLock(InIndex, InRegistry);
 	if(TypeIndex != INDEX_NONE)
 	{
 		return TypeIndex;
@@ -522,7 +516,7 @@ TOptional<TRigVMTypeIndex> FRigVMTemplateArgument::TryToGetTypeIndex_NoLock(cons
 	return TOptional<TRigVMTypeIndex>();
 }
 
-int32 FRigVMTemplateArgument::FindTypeIndex(const TRigVMTypeIndex InTypeIndex) const
+int32 FRigVMTemplateArgument::FindTypeIndex(const TRigVMTypeIndex InTypeIndex, const FRigVMRegistryHandle& InRegistry) const
 {
 	if (!bUseCategories)
 	{
@@ -533,14 +527,14 @@ int32 FRigVMTemplateArgument::FindTypeIndex(const TRigVMTypeIndex InTypeIndex) c
 	{
 		bool bFound = false;
 		int32 ValidIndex = 0;
-		CategoryViews(TypeCategories).ForEachType([this, &ValidIndex, &bFound, InTypeIndex](const TRigVMTypeIndex& Type) -> bool
+		CategoryViews(TypeCategories, InRegistry).ForEachType([this, &ValidIndex, &bFound, InTypeIndex, &InRegistry](const TRigVMTypeIndex& Type) -> bool
 		{
 			if (Type == InTypeIndex)
 			{
 				bFound = true;
 				return false;
 			}
-			if (FilterType(Type))
+			if (FilterType(Type, InRegistry))
 			{
 				ValidIndex++;
 			}
@@ -553,10 +547,10 @@ int32 FRigVMTemplateArgument::FindTypeIndex(const TRigVMTypeIndex InTypeIndex) c
 		return INDEX_NONE;
 	}
 
-	return CategoryViews(TypeCategories).FindIndex(InTypeIndex);
+	return CategoryViews(TypeCategories, InRegistry).FindIndex(InTypeIndex);
 }
 
-int32 FRigVMTemplateArgument::GetNumTypes() const
+int32 FRigVMTemplateArgument::GetNumTypes(const FRigVMRegistryHandle& InRegistry) const
 {
 	if (!bUseCategories)
 	{
@@ -566,9 +560,9 @@ int32 FRigVMTemplateArgument::GetNumTypes() const
 	if (FilterType)
 	{
 		int32 NumTypes = 0;
-		CategoryViews(TypeCategories).ForEachType([this, &NumTypes](const TRigVMTypeIndex& Type) -> bool
+		CategoryViews(TypeCategories, InRegistry).ForEachType([this, &NumTypes, &InRegistry](const TRigVMTypeIndex& Type) -> bool
 		{
-			if (FilterType(Type))
+			if (FilterType(Type, InRegistry))
 			{
 				NumTypes++;
 			}
@@ -577,13 +571,13 @@ int32 FRigVMTemplateArgument::GetNumTypes() const
 		return NumTypes;
 	}
 	
-	return Algo::Accumulate(TypeCategories, 0, [](int32 Sum, const ETypeCategory Category)
+	return Algo::Accumulate(TypeCategories, 0, [&InRegistry](int32 Sum, const ETypeCategory Category)
 	{
-		return Sum + FRigVMRegistry_RWLock::Get().GetTypesForCategory(Category).Num();
+		return Sum + InRegistry->GetTypesForCategory_NoLock(Category).Num();
 	});
 }
 
-int32 FRigVMTemplateArgument::GetNumTypes_NoLock() const
+int32 FRigVMTemplateArgument::GetNumTypes_NoLock(const FRigVMRegistryHandle& InRegistry) const
 {
 	if (!bUseCategories)
 	{
@@ -594,9 +588,9 @@ int32 FRigVMTemplateArgument::GetNumTypes_NoLock() const
 	if (FilterType)
 	{
 		int32 NumTypes = 0;
-		CategoryViews(TypeCategories).ForEachType([this, &NumTypes](const TRigVMTypeIndex& Type) -> bool
+		CategoryViews(TypeCategories, InRegistry).ForEachType([this, &NumTypes, &InRegistry](const TRigVMTypeIndex& Type) -> bool
 		{
-			if (FilterType(Type))
+			if (FilterType(Type, InRegistry))
 			{
 				NumTypes++;
 			}
@@ -605,9 +599,9 @@ int32 FRigVMTemplateArgument::GetNumTypes_NoLock() const
 		return NumTypes;
 	}
 
-	return Algo::Accumulate(TypeCategories, 0, [](int32 Sum, const ETypeCategory Category)
+	return Algo::Accumulate(TypeCategories, 0, [&InRegistry](int32 Sum, const ETypeCategory Category)
 	{
-		return Sum + FRigVMRegistry_NoLock::GetForRead().GetTypesForCategory_NoLock(Category).Num();
+		return Sum + InRegistry->GetTypesForCategory_NoLock(Category).Num();
 	});
 }
 
@@ -623,7 +617,7 @@ void FRigVMTemplateArgument::RemoveType(const int32 InIndex)
 	TypeIndices.RemoveAt(InIndex);
 }
 
-void FRigVMTemplateArgument::ForEachType(TFunction<bool(const TRigVMTypeIndex InType)>&& InCallback) const
+void FRigVMTemplateArgument::ForEachType(TFunction<bool(const TRigVMTypeIndex InType)>&& InCallback, const FRigVMRegistryHandle& InRegistry) const
 {
 	if (!bUseCategories)
 	{
@@ -632,9 +626,9 @@ void FRigVMTemplateArgument::ForEachType(TFunction<bool(const TRigVMTypeIndex In
 
 	if (FilterType)
 	{
-		CategoryViews(TypeCategories).ForEachType([this, InCallback](const TRigVMTypeIndex Type)
+		CategoryViews(TypeCategories, InRegistry).ForEachType([this, InCallback, &InRegistry](const TRigVMTypeIndex Type)
 		{
-			if (FilterType(Type))
+			if (FilterType(Type, InRegistry))
 			{
 				return InCallback(Type);
 			}
@@ -643,20 +637,22 @@ void FRigVMTemplateArgument::ForEachType(TFunction<bool(const TRigVMTypeIndex In
 		return;
 	}
 	
-	return CategoryViews(TypeCategories).ForEachType(MoveTemp(InCallback));
+	return CategoryViews(TypeCategories, InRegistry).ForEachType(MoveTemp(InCallback));
 }
 
 TArray<TRigVMTypeIndex> FRigVMTemplateArgument::GetSupportedTypeIndices(const TArray<int32>& InPermutationIndices) const
 {
-	const FRigVMRegistryReadLock _;
+	return GetSupportedTypeIndices_NoLock(InPermutationIndices, FRigVMRegistryReadLock());
+}
 
+TArray<TRigVMTypeIndex> FRigVMTemplateArgument::GetSupportedTypeIndices_NoLock(const TArray<int32>& InPermutationIndices, const FRigVMRegistryHandle& InRegistry) const
+{
 #if UE_RIGVM_DEBUG_TYPEINDEX
-	const FRigVMRegistry_NoLock& Registry = _.GetRegistry();
-	auto UpdateTypeIndex = [&Registry](const TRigVMTypeIndex& TypeIndex) -> TRigVMTypeIndex
+	auto UpdateTypeIndex = [&InRegistry](const TRigVMTypeIndex& TypeIndex) -> TRigVMTypeIndex
 	{
 		if(TypeIndex.Name.IsNone())
 		{
-			return Registry.GetTypeIndex_NoLock(Registry.GetType_NoLock(TypeIndex));
+			return InRegistry->GetTypeIndex_NoLock(InRegistry->GetType_NoLock(TypeIndex));
 		}
 		return TypeIndex;
 	};
@@ -677,14 +673,14 @@ TArray<TRigVMTypeIndex> FRigVMTemplateArgument::GetSupportedTypeIndices(const TA
 #endif
 			}
 			return true;
-		});
+		}, InRegistry);
 	}
 	else
 	{
 		for(const int32 PermutationIndex : InPermutationIndices)
 		{
 			// INDEX_NONE indicates deleted permutation
-			const TRigVMTypeIndex TypeIndex = GetTypeIndex_NoLock(PermutationIndex);
+			const TRigVMTypeIndex TypeIndex = GetTypeIndex_NoLock(PermutationIndex, InRegistry);
 			if (TypeIndex != INDEX_NONE)
 			{
 #if UE_RIGVM_DEBUG_TYPEINDEX
@@ -702,32 +698,34 @@ TArray<TRigVMTypeIndex> FRigVMTemplateArgument::GetSupportedTypeIndices(const TA
 
 TArray<FString> FRigVMTemplateArgument::GetSupportedTypeStrings(const TArray<int32>& InPermutationIndices) const
 {
-	const FRigVMRegistryReadLock _;
-	const FRigVMRegistry_NoLock& Registry = FRigVMRegistry_NoLock::GetForRead();
+	return GetSupportedTypeStrings_NoLock(InPermutationIndices, FRigVMRegistryReadLock());
+}
 
+TArray<FString> FRigVMTemplateArgument::GetSupportedTypeStrings_NoLock(const TArray<int32>& InPermutationIndices, const FRigVMRegistryHandle& InRegistry) const
+{
 	TArray<FString> SupportedTypes;
 	if(InPermutationIndices.IsEmpty())
 	{
-		ForEachType([&](const TRigVMTypeIndex TypeIndex)
+		ForEachType([&InRegistry, &SupportedTypes](const TRigVMTypeIndex TypeIndex)
 		{
 			// INDEX_NONE indicates deleted permutation
 			if (TypeIndex != INDEX_NONE)
 			{
-				const FString TypeString = Registry.GetType_NoLock(TypeIndex).CPPType.ToString();
+				const FString TypeString = InRegistry->GetType_NoLock(TypeIndex).CPPType.ToString();
 				SupportedTypes.AddUnique(TypeString);
 			}
 			return true;
-		});
+		}, InRegistry);
 	}
 	else
 	{
 		for(const int32 PermutationIndex : InPermutationIndices)
 		{
-			const TRigVMTypeIndex TypeIndex = GetTypeIndex_NoLock(PermutationIndex);
+			const TRigVMTypeIndex TypeIndex = GetTypeIndex_NoLock(PermutationIndex, InRegistry);
 			// INDEX_NONE indicates deleted permutation
 			if (TypeIndex != INDEX_NONE)
 			{
-				const FString TypeString = Registry.GetType_NoLock(TypeIndex).CPPType.ToString();
+				const FString TypeString = InRegistry->GetType_NoLock(TypeIndex).CPPType.ToString();
 				SupportedTypes.AddUnique(TypeString);
 			}
 		}
@@ -738,12 +736,12 @@ TArray<FString> FRigVMTemplateArgument::GetSupportedTypeStrings(const TArray<int
 #endif
 
 
-FRigVMTemplateArgument::CategoryViews::CategoryViews(const TArray<ETypeCategory>& InCategories)
+FRigVMTemplateArgument::CategoryViews::CategoryViews(const TArray<ETypeCategory>& InCategories, const FRigVMRegistryHandle& InRegistry)
 {
 	Types.Reserve(InCategories.Num());
 	for (const ETypeCategory Category: InCategories)
 	{
-		Types.Emplace(FRigVMRegistry_NoLock::GetForRead().GetTypesForCategory_NoLock(Category));
+		Types.Emplace(InRegistry->GetTypesForCategory_NoLock(Category));
 	}
 }
 
@@ -796,13 +794,13 @@ int32 FRigVMTemplateArgument::CategoryViews::FindIndex(const TRigVMTypeIndex InT
 FRigVMTemplateArgumentInfo::FRigVMTemplateArgumentInfo(const FName InName, ERigVMPinDirection InDirection, const TArray<TRigVMTypeIndex>& InTypeIndices)
 	: Name(InName)
 	, Direction(InDirection)
-	, FactoryCallback([InTypeIndices](const FName InName, ERigVMPinDirection InDirection){ return FRigVMTemplateArgument(InName, InDirection, InTypeIndices); } )
+	, FactoryCallback([InTypeIndices](const FName InName, ERigVMPinDirection InDirection, const FRigVMRegistryHandle& InRegistry){ return FRigVMTemplateArgument(InName, InDirection, InTypeIndices, InRegistry); } )
 {}
 
 FRigVMTemplateArgumentInfo::FRigVMTemplateArgumentInfo(const FName InName, ERigVMPinDirection InDirection, TRigVMTypeIndex InTypeIndex)
 	: Name(InName)
 	, Direction(InDirection)
-	, FactoryCallback( [InTypeIndex](const FName InName, ERigVMPinDirection InDirection){ return FRigVMTemplateArgument(InName, InDirection, InTypeIndex); } )
+	, FactoryCallback( [InTypeIndex](const FName InName, ERigVMPinDirection InDirection, const FRigVMRegistryHandle& InRegistry){ return FRigVMTemplateArgument(InName, InDirection, InTypeIndex, InRegistry); } )
 {}
 
 FRigVMTemplateArgumentInfo::FRigVMTemplateArgumentInfo(
@@ -811,13 +809,13 @@ FRigVMTemplateArgumentInfo::FRigVMTemplateArgumentInfo(
 	TypeFilterCallback InTypeFilter)
 	: Name(InName)
 	, Direction(InDirection)
-	, FactoryCallback([InTypeCategories, InTypeFilter](const FName InName, ERigVMPinDirection InDirection){ return FRigVMTemplateArgument(InName, InDirection, InTypeCategories, InTypeFilter); })
+	, FactoryCallback([InTypeCategories, InTypeFilter](const FName InName, ERigVMPinDirection InDirection, const FRigVMRegistryHandle& InRegistry){ return FRigVMTemplateArgument(InName, InDirection, InTypeCategories, InTypeFilter, InRegistry); })
 {}
 
 FRigVMTemplateArgumentInfo::FRigVMTemplateArgumentInfo(const FName InName, ERigVMPinDirection InDirection)
 	: Name(InName)
 	, Direction(InDirection)
-	, FactoryCallback([](const FName InName, ERigVMPinDirection InDirection){ return FRigVMTemplateArgument(InName, InDirection); })
+	, FactoryCallback([](const FName InName, ERigVMPinDirection InDirection, const FRigVMRegistryHandle&){ return FRigVMTemplateArgument(InName, InDirection); })
 {}
 
 FRigVMTemplateArgumentInfo::FRigVMTemplateArgumentInfo(const FName InName, ERigVMPinDirection InDirection, ArgumentCallback&& InCallback)
@@ -828,7 +826,12 @@ FRigVMTemplateArgumentInfo::FRigVMTemplateArgumentInfo(const FName InName, ERigV
 
 FRigVMTemplateArgument FRigVMTemplateArgumentInfo::GetArgument() const
 {
-	FRigVMTemplateArgument Argument = FactoryCallback(Name, Direction);
+	return GetArgument_NoLock(FRigVMRegistryReadLock());
+}
+
+FRigVMTemplateArgument FRigVMTemplateArgumentInfo::GetArgument_NoLock(const FRigVMRegistryHandle& InRegistry) const
+{
+	FRigVMTemplateArgument Argument = FactoryCallback(Name, Direction, InRegistry);
 	return MoveTemp(Argument);
 }
 
@@ -858,6 +861,14 @@ FName FRigVMTemplateArgumentInfo::ComputeTemplateNotation(const FName InTemplate
 TArray<TRigVMTypeIndex> FRigVMTemplateArgumentInfo::GetTypesFromCategories(
 	const TArray<FRigVMTemplateArgument::ETypeCategory>& InTypeCategories,
 	const FRigVMTemplateArgument::FTypeFilter& InTypeFilter)
+{
+	return GetTypesFromCategories_NoLock(InTypeCategories, InTypeFilter, FRigVMRegistryReadLock());
+}
+
+TArray<TRigVMTypeIndex> FRigVMTemplateArgumentInfo::GetTypesFromCategories_NoLock(
+	const TArray<FRigVMTemplateArgument::ETypeCategory>& InTypeCategories,
+	const FRigVMTemplateArgument::FTypeFilter& InTypeFilter,
+	const FRigVMRegistryHandle& InRegistry)
 {
 	TSet<TRigVMTypeIndex> AllTypes;
 
@@ -866,14 +877,14 @@ TArray<TRigVMTypeIndex> FRigVMTemplateArgumentInfo::GetTypesFromCategories(
 		int32 NumTypes = 0;
 		for (const FRigVMTemplateArgument::ETypeCategory TypeCategory : InTypeCategories)
 		{
-			NumTypes += FRigVMRegistry_NoLock::GetForRead().GetTypesForCategory_NoLock(TypeCategory).Num();
+			NumTypes += InRegistry->GetTypesForCategory_NoLock(TypeCategory).Num();
 		}
 		AllTypes.Reserve(NumTypes);
 	}
 
 	for (const FRigVMTemplateArgument::ETypeCategory TypeCategory : InTypeCategories)
 	{
-		AllTypes.Append(FRigVMRegistry_NoLock::GetForRead().GetTypesForCategory_NoLock(TypeCategory));
+		AllTypes.Append(InRegistry->GetTypesForCategory_NoLock(TypeCategory));
 	}
 
 	TArray<TRigVMTypeIndex> Types;
@@ -886,7 +897,7 @@ TArray<TRigVMTypeIndex> FRigVMTemplateArgumentInfo::GetTypesFromCategories(
 		Types.Reserve(AllTypes.Num());
 		for (const TRigVMTypeIndex& Type : AllTypes)
 		{
-			if (InTypeFilter.Execute(Type))
+			if (InTypeFilter.Execute(Type, InRegistry))
 			{
 				Types.Add(Type);
 			}
@@ -902,16 +913,17 @@ FRigVMTemplate::FRigVMTemplate()
 	: Index(INDEX_NONE)
 	, Notation(NAME_None)
 	, Hash(UINT32_MAX)
+	, OwnerRegistry(nullptr)
 {
 
 }
 
-FRigVMTemplate::FRigVMTemplate(UScriptStruct* InStruct, const FString& InTemplateName, int32 InFunctionIndex)
+FRigVMTemplate::FRigVMTemplate(UScriptStruct* InStruct, const FString& InTemplateName, int32 InFunctionIndex, FRigVMRegistryHandle& InRegistry)
 	: Index(INDEX_NONE)
 	, Notation(NAME_None)
 	, Hash(UINT32_MAX)
+	, OwnerRegistry(InRegistry.GetRegistry())
 {
-	FRigVMRegistry_NoLock& Registry = FRigVMRegistry_NoLock::GetForWrite(); 
 	TArray<FString> ArgumentNotations;
 
 	// create the arguments sorted by super -> child struct.
@@ -921,10 +933,10 @@ FRigVMTemplate::FRigVMTemplate(UScriptStruct* InStruct, const FString& InTemplat
 		// only iterate on this struct's fields, not the super structs'
 		for (TFieldIterator<FProperty> It(Struct, EFieldIterationFlags::None); It; ++It)
 		{
-			FRigVMTemplateArgument Argument = FRigVMTemplateArgument::Make_NoLock(*It, Registry);
+			FRigVMTemplateArgument Argument = FRigVMTemplateArgument::Make_NoLock(*It, InRegistry);
 			Argument.Index = Arguments.Num();
 
-			if(!Argument.IsExecute_NoLock(Registry) && IsValidArgumentForTemplate(Argument.GetDirection()) && Argument.GetDirection() != ERigVMPinDirection::Hidden)
+			if(!Argument.IsExecute_NoLock(InRegistry) && IsValidArgumentForTemplate(Argument.GetDirection()) && Argument.GetDirection() != ERigVMPinDirection::Hidden)
 			{
 				Arguments.Add(Argument);
 			}
@@ -937,7 +949,7 @@ FRigVMTemplate::FRigVMTemplate(UScriptStruct* InStruct, const FString& InTemplat
 	{
 		if(const FRigVMTemplateArgument* Argument = FindArgument(It->GetFName()))
 		{
-			if(!Argument->IsExecute_NoLock(Registry) && Argument->GetDirection() != ERigVMPinDirection::Hidden)
+			if(!Argument->IsExecute_NoLock(InRegistry) && Argument->GetDirection() != ERigVMPinDirection::Hidden)
 			{
 				ArgumentNotations.Add(GetArgumentNotation(Argument->Name, Argument->Direction));
 			}
@@ -957,27 +969,28 @@ FRigVMTemplate::FRigVMTemplate(UScriptStruct* InStruct, const FString& InTemplat
 			}
 		}
 
-		UpdateTypesHashToPermutation(Permutations.Num()-1);
+		UpdateTypesHashToPermutation(Permutations.Num()-1, InRegistry);
 	}
 }
 
-FRigVMTemplate::FRigVMTemplate(const FName& InTemplateName, const TArray<FRigVMTemplateArgumentInfo>& InInfos)
+FRigVMTemplate::FRigVMTemplate(const FName& InTemplateName, const TArray<FRigVMTemplateArgumentInfo>& InInfos, const FRigVMRegistryHandle& InRegistry)
 	: Index(INDEX_NONE)
 	, Notation(NAME_None)
 	, Hash(UINT32_MAX)
+	, OwnerRegistry(const_cast<FRigVMRegistry_NoLock*>(InRegistry.GetRegistry()))
 {
 	for (const FRigVMTemplateArgumentInfo& InInfo : InInfos)
 	{
 		if(IsValidArgumentForTemplate(InInfo.Direction))
 		{
-			FRigVMTemplateArgument Argument = InInfo.GetArgument();
+			FRigVMTemplateArgument Argument = InInfo.GetArgument_NoLock(InRegistry);
 			Argument.Index = Arguments.Num();
 			Arguments.Emplace(MoveTemp(Argument));
 		}
 	}
 	
 	Notation = FRigVMTemplateArgumentInfo::ComputeTemplateNotation(InTemplateName, InInfos);
-	UpdateTypesHashToPermutation(Permutations.Num()-1);
+	UpdateTypesHashToPermutation(Permutations.Num()-1, InRegistry);
 }
 
 FLinearColor FRigVMTemplate::GetColorFromMetadata(FString InMetadata)
@@ -1090,17 +1103,17 @@ TArray<UStruct*> FRigVMTemplate::GetSuperStructs(UStruct* InStruct, bool bInclud
 
 FRigVMTemplate::FTypeMap FRigVMTemplate::GetArgumentTypesFromString(const FString& InTypeString, const FRigVMUserDefinedTypeResolver* InTypeResolver) const
 {
-	return GetArgumentTypesFromString_Impl(InTypeString, InTypeResolver, true);
+	FRigVMRegistryWriteLock Lock;
+	return GetArgumentTypesFromString_NoLock(InTypeString, InTypeResolver, Lock);
 }
 
-FRigVMTemplate::FTypeMap FRigVMTemplate::GetArgumentTypesFromString_Impl(const FString& InTypeString, const FRigVMUserDefinedTypeResolver* InTypeResolver, bool bLockRegistry) const
+FRigVMTemplate::FTypeMap FRigVMTemplate::GetArgumentTypesFromString_NoLock(const FString& InTypeString, const FRigVMUserDefinedTypeResolver* InTypeResolver, FRigVMRegistryHandle& InRegistry) const
 {
+	check(InRegistry.GetRegistry() == OwnerRegistry);
+	
 	FTypeMap Types;
 	if(!InTypeString.IsEmpty())
 	{
-		const FRigVMRegistryWriteLock _(bLockRegistry);
-		FRigVMRegistry_NoLock& Registry = FRigVMRegistry_NoLock::GetForWrite();
-
 		FString Left, Right = InTypeString;
 		while(!Right.IsEmpty())
 		{
@@ -1115,7 +1128,7 @@ FRigVMTemplate::FTypeMap FRigVMTemplate::GetArgumentTypesFromString_Impl(const F
 			{
 				if(const FRigVMTemplateArgument* Argument = FindArgument(*ArgumentName))
 				{
-					TRigVMTypeIndex TypeIndex = Registry.GetTypeIndexFromCPPType_NoLock(TypeName);
+					TRigVMTypeIndex TypeIndex = InRegistry->GetTypeIndexFromCPPType_NoLock(TypeName);
 
 					// If the type was not found, check if it's a user-defined type that hasn't been registered yet.
 					if (TypeIndex == INDEX_NONE && RigVMTypeUtils::RequiresCPPTypeObject(TypeName))
@@ -1123,7 +1136,7 @@ FRigVMTemplate::FTypeMap FRigVMTemplate::GetArgumentTypesFromString_Impl(const F
 						UObject* CPPTypeObject = RigVMTypeUtils::ObjectFromCPPType(TypeName, true, InTypeResolver);
 						
 						FRigVMTemplateArgumentType ArgType(*TypeName, CPPTypeObject);
-						TypeIndex = Registry.FindOrAddType_NoLock(ArgType, false);
+						TypeIndex = InRegistry->FindOrAddType_NoLock(ArgType, false);
 					}
 					
 					if(TypeIndex != INDEX_NONE)
@@ -1137,15 +1150,12 @@ FRigVMTemplate::FTypeMap FRigVMTemplate::GetArgumentTypesFromString_Impl(const F
 	return Types;
 }
 
-FString FRigVMTemplate::GetStringFromArgumentTypes(const FTypeMap& InTypes, bool bLockRegistry)
+FString FRigVMTemplate::GetStringFromArgumentTypes(const FTypeMap& InTypes, const FRigVMRegistryHandle& InRegistry)
 {
-	const FRigVMRegistryReadLock _(bLockRegistry); 
-	const FRigVMRegistry_NoLock& Registry = FRigVMRegistry_NoLock::GetForRead();
-	
 	TArray<FString> TypePairStrings;
 	for(const TPair<FName,TRigVMTypeIndex>& Pair : InTypes)
 	{
-		const FRigVMTemplateArgumentType& Type = Registry.GetType_NoLock(Pair.Value);
+		const FRigVMTemplateArgumentType& Type = InRegistry->GetType_NoLock(Pair.Value);
 		static constexpr TCHAR Format[] = TEXT("%s:%s");
 		FString PairString;
 		const FString KeyString = Pair.Key.ToString();
@@ -1467,8 +1477,10 @@ FString FRigVMTemplate::GetArgumentMetaData(const FName& InArgumentName, const F
 
 #endif
 
-bool FRigVMTemplate::Merge(const FRigVMTemplate& InOther)
+bool FRigVMTemplate::Merge(const FRigVMTemplate& InOther, const FRigVMRegistryHandle& InRegistry)
 {
+	check(InRegistry.GetRegistry() == OwnerRegistry);
+
 	if (!IsValid() || !InOther.IsValid())
 	{
 		return false;
@@ -1479,15 +1491,17 @@ bool FRigVMTemplate::Merge(const FRigVMTemplate& InOther)
 		return false;
 	}
 
-	if(InOther.GetExecuteContextStruct(false) != GetExecuteContextStruct(false))
+	check(OwnerRegistry == InOther.OwnerRegistry);
+
+	if(InOther.GetExecuteContextStruct_NoLock(InRegistry) != GetExecuteContextStruct_NoLock(InRegistry))
 	{
 		// find the previously defined permutation.
 		UE_LOG(LogRigVM, Display, TEXT("RigVMFunction '%s' cannot be merged into the '%s' template. ExecuteContext Types differ ('%s' vs '%s' from '%s')."),
-			*InOther.GetPrimaryPermutation(false)->Name,
+			*InOther.GetPrimaryPermutation_NoLock(InRegistry)->Name,
 			*GetNotation().ToString(),
-			*InOther.GetExecuteContextStruct(false)->GetStructCPPName(),
-			*GetExecuteContextStruct(false)->GetStructCPPName(),
-			*GetPrimaryPermutation(false)->Name);
+			*InOther.GetExecuteContextStruct_NoLock(InRegistry)->GetStructCPPName(),
+			*GetExecuteContextStruct_NoLock(InRegistry)->GetStructCPPName(),
+			*GetPrimaryPermutation_NoLock(InRegistry)->Name);
 		return false;
 	}
 
@@ -1497,12 +1511,12 @@ bool FRigVMTemplate::Merge(const FRigVMTemplate& InOther)
 	}
 
 	// find colliding permutations
-	for(int32 PermutationIndex = 0;PermutationIndex<NumPermutations();PermutationIndex++)
+	for(int32 PermutationIndex = 0;PermutationIndex<NumPermutations_NoLock(InRegistry);PermutationIndex++)
 	{
 		int32 MatchingArguments = 0;
 		for(int32 ArgumentIndex = 0; ArgumentIndex < Arguments.Num(); ArgumentIndex++)
 		{
-			if (Arguments[ArgumentIndex].GetTypeIndex_NoLock(PermutationIndex) == InOther.Arguments[ArgumentIndex].GetTypeIndex_NoLock(0))
+			if (Arguments[ArgumentIndex].GetTypeIndex_NoLock(PermutationIndex, InRegistry) == InOther.Arguments[ArgumentIndex].GetTypeIndex_NoLock(0, InRegistry))
 			{
 				MatchingArguments++;
 			}
@@ -1511,9 +1525,9 @@ bool FRigVMTemplate::Merge(const FRigVMTemplate& InOther)
 		{
 			// find the previously defined permutation.
 			UE_LOG(LogRigVM, Display, TEXT("RigVMFunction '%s' cannot be merged into the '%s' template. It collides with '%s'."),
-				*InOther.GetPrimaryPermutation(false)->Name,
+				*InOther.GetPrimaryPermutation_NoLock(InRegistry)->Name,
 				*GetNotation().ToString(),
-				*GetPermutation_NoLock(PermutationIndex)->Name);
+				*GetPermutation_NoLock(PermutationIndex, InRegistry)->Name);
 			return false;
 		}
 	}
@@ -1523,7 +1537,7 @@ bool FRigVMTemplate::Merge(const FRigVMTemplate& InOther)
 	for (int32 ArgumentIndex = 0; ArgumentIndex < Arguments.Num(); ArgumentIndex++)
 	{
 		const FRigVMTemplateArgument& OtherArg = InOther.Arguments[ArgumentIndex];
-		if (OtherArg.GetNumTypes_NoLock() != 1)
+		if (OtherArg.GetNumTypes_NoLock(InRegistry) != 1)
 		{
 			return false;
 		}
@@ -1531,9 +1545,9 @@ bool FRigVMTemplate::Merge(const FRigVMTemplate& InOther)
 		// Add Other argument information into the TypeToPermutations map
 		{
 			FRigVMTemplateArgument& NewArg = NewArgs.Add_GetRef(Arguments[ArgumentIndex]);
-			const TRigVMTypeIndex OtherTypeIndex = OtherArg.GetTypeIndex_NoLock(0);
+			const TRigVMTypeIndex OtherTypeIndex = OtherArg.GetTypeIndex_NoLock(0, InRegistry);
 #if WITH_EDITOR
-			const int32 NewPermutationIndex = NewArg.GetNumTypes_NoLock();
+			const int32 NewPermutationIndex = NewArg.GetNumTypes_NoLock(InRegistry);
 			if (TArray<int32>* ArgTypePermutations = NewArg.TypeToPermutations.Find(OtherTypeIndex))
 			{
 				ArgTypePermutations->Add(NewPermutationIndex);
@@ -1551,7 +1565,7 @@ bool FRigVMTemplate::Merge(const FRigVMTemplate& InOther)
 
 	Permutations.Add(InOther.Permutations[0]);
 
-	UpdateTypesHashToPermutation(Permutations.Num()-1);
+	UpdateTypesHashToPermutation(Permutations.Num()-1, InRegistry);
 	
 	return true;
 }
@@ -1590,32 +1604,37 @@ const FRigVMExecuteArgument* FRigVMTemplate::FindExecuteArgument(const FName& In
 
 const TArray<FRigVMExecuteArgument>& FRigVMTemplate::GetExecuteArguments(const FRigVMDispatchContext& InContext) const
 {
+	return GetExecuteArguments_NoLock(InContext, FRigVMRegistryReadLock());
+}
+
+const TArray<FRigVMExecuteArgument>& FRigVMTemplate::GetExecuteArguments_NoLock(const FRigVMDispatchContext& InContext, const FRigVMRegistryHandle& InRegistry) const
+{
+	check(InRegistry.GetRegistry() == OwnerRegistry);
+
 	if(ExecuteArguments.IsEmpty())
 	{
-		FRigVMRegistryWriteLock WriteLock;
-		FRigVMRegistry_NoLock& Registry = WriteLock.GetRegistry();
-
 		if(UsesDispatch())
 		{
-			const FRigVMDispatchFactory* Factory = Delegates.GetDispatchFactoryDelegate.Execute();
+			const FRigVMDispatchFactory* Factory = Delegates.GetDispatchFactoryDelegate.Execute(InRegistry);
 			check(Factory);
 
-			ExecuteArguments = Factory->GetExecuteArguments_NoLock(InContext);
+			ExecuteArguments = Factory->GetExecuteArguments_NoLock(InContext, InRegistry);
 		}
-		else if(const FRigVMFunction* PrimaryPermutation = GetPrimaryPermutation(false /* lock */))
+		else if(const FRigVMFunction* PrimaryPermutation = GetPrimaryPermutation_NoLock(FRigVMRegistryReadLock()))
 		{
 			if(PrimaryPermutation->Struct)
 			{
+				FRigVMRegistryHandle& MutableRegistryHandle = *const_cast<FRigVMRegistryHandle*>(&InRegistry);
 				TArray<UStruct*> Structs = GetSuperStructs(PrimaryPermutation->Struct, true);
 				for(const UStruct* Struct : Structs)
 				{
 					// only iterate on this struct's fields, not the super structs'
 					for (TFieldIterator<FProperty> It(Struct, EFieldIterationFlags::None); It; ++It)
 					{
-						FRigVMTemplateArgument Argument = FRigVMTemplateArgument::Make_NoLock(*It, Registry);
+						FRigVMTemplateArgument Argument = FRigVMTemplateArgument::Make_NoLock(*It, MutableRegistryHandle);
 						if(Argument.IsExecute())
 						{
-							ExecuteArguments.Emplace(Argument.Name, Argument.Direction, Argument.GetTypeIndex_NoLock(0));
+							ExecuteArguments.Emplace(Argument.Name, Argument.Direction, Argument.GetTypeIndex_NoLock(0, InRegistry));
 						}
 					}
 				}
@@ -1625,19 +1644,33 @@ const TArray<FRigVMExecuteArgument>& FRigVMTemplate::GetExecuteArguments(const F
 	return ExecuteArguments;
 }
 
-const UScriptStruct* FRigVMTemplate::GetExecuteContextStruct(bool bLockRegistry) const
+const UScriptStruct* FRigVMTemplate::GetExecuteContextStruct() const
+{
+	return GetExecuteContextStruct_NoLock(FRigVMRegistryReadLock());
+}
+
+const UScriptStruct* FRigVMTemplate::GetExecuteContextStruct_NoLock(const FRigVMRegistryHandle& InRegistry) const
 {
-	if(const FRigVMDispatchFactory* Factory = GetDispatchFactory(bLockRegistry))
+	check(InRegistry.GetRegistry() == OwnerRegistry);
+
+	if(const FRigVMDispatchFactory* Factory = GetDispatchFactory_NoLock(InRegistry))
 	{
 		return Factory->GetExecuteContextStruct();
 	}
 	check(!Permutations.IsEmpty());
-	return GetPrimaryPermutation(bLockRegistry)->GetExecuteContextStruct(bLockRegistry);
+	return GetPrimaryPermutation_NoLock(InRegistry)->GetExecuteContextStruct_NoLock(InRegistry);
 }
 
 bool FRigVMTemplate::SupportsExecuteContextStruct(const UScriptStruct* InExecuteContextStruct) const
 {
-	return InExecuteContextStruct->IsChildOf(GetExecuteContextStruct(true));
+	return SupportsExecuteContextStruct_NoLock(InExecuteContextStruct, FRigVMRegistryReadLock());
+}
+
+bool FRigVMTemplate::SupportsExecuteContextStruct_NoLock(const UScriptStruct* InExecuteContextStruct, const FRigVMRegistryHandle& InRegistry) const
+{
+	check(InRegistry.GetRegistry() == OwnerRegistry);
+	
+	return InExecuteContextStruct->IsChildOf(GetExecuteContextStruct_NoLock(InRegistry));
 }
 
 #if WITH_EDITOR
@@ -1653,72 +1686,86 @@ bool FRigVMTemplate::ArgumentSupportsTypeIndex(const FName& InArgumentName, TRig
 
 #endif
 
-const FRigVMFunction* FRigVMTemplate::GetPrimaryPermutation(bool bLockRegistry) const
+int32 FRigVMTemplate::NumPermutations() const
+{
+	return NumPermutations_NoLock(FRigVMRegistryReadLock());
+}
+
+const FRigVMFunction* FRigVMTemplate::GetPrimaryPermutation() const
+{
+	return GetPrimaryPermutation_NoLock(FRigVMRegistryReadLock());
+}
+
+const FRigVMFunction* FRigVMTemplate::GetPrimaryPermutation_NoLock(const FRigVMRegistryHandle& InRegistry) const
 {
-	if (NumPermutations() > 0)
+	check(InRegistry.GetRegistry() == OwnerRegistry);
+
+	if (NumPermutations_NoLock(InRegistry) > 0)
 	{
-		return GetPermutation(0, bLockRegistry);
+		return GetPermutation_NoLock(0, InRegistry);
 	}
 	return nullptr;
 }
 
-const FRigVMFunction* FRigVMTemplate::GetPermutation(int32 InIndex, bool bLockRegistry) const
+const FRigVMFunction* FRigVMTemplate::GetPermutation(int32 InIndex) const
 {
-	const FRigVMRegistryReadLock _(bLockRegistry);
-	return GetPermutation_NoLock(InIndex);
+	return GetPermutation_NoLock(InIndex, FRigVMRegistryReadLock());
 }
 
-const FRigVMFunction* FRigVMTemplate::GetPermutation_NoLock(int32 InIndex) const
+const FRigVMFunction* FRigVMTemplate::GetPermutation_NoLock(int32 InIndex, const FRigVMRegistryHandle& InRegistry) const
 {
-	const FRigVMRegistry_NoLock& Registry = FRigVMRegistry_NoLock::GetForRead();
+	check(InRegistry.GetRegistry() == OwnerRegistry);
+
 	const int32 FunctionIndex = Permutations[InIndex];
-	if(Registry.GetFunctions_NoLock().IsValidIndex(FunctionIndex))
+	if(InRegistry->GetFunctions_NoLock().IsValidIndex(FunctionIndex))
 	{
-		return &Registry.GetFunctions_NoLock()[Permutations[InIndex]];
+		const FRigVMFunction* Function = &InRegistry->GetFunctions_NoLock()[Permutations[InIndex]];
+		check(Function->PermutationIndex == InIndex);
+		return Function;
 	}
 	return nullptr;
 }
 
-const FRigVMFunction* FRigVMTemplate::GetOrCreatePermutation(int32 InIndex, bool bLockRegistry)
+const FRigVMFunction* FRigVMTemplate::GetOrCreatePermutation(int32 InIndex)
 {
-	const FRigVMRegistryWriteLock _(bLockRegistry);
-	return GetOrCreatePermutation_NoLock(InIndex);
+	FRigVMRegistryWriteLock Lock;
+	return GetOrCreatePermutation_NoLock(InIndex, Lock);
 }
 
-const FRigVMFunction* FRigVMTemplate::GetOrCreatePermutation_NoLock(int32 InIndex)
+const FRigVMFunction* FRigVMTemplate::GetOrCreatePermutation_NoLock(int32 InIndex, FRigVMRegistryHandle& InRegistry)
 {
-	if(const FRigVMFunction* Function = GetPermutation_NoLock(InIndex))
+	check(InRegistry.GetRegistry() == OwnerRegistry);
+
+	if(const FRigVMFunction* Function = GetPermutation_NoLock(InIndex, InRegistry))
 	{
 		return Function;
 	}
 
 	if(Permutations[InIndex] == INDEX_NONE && UsesDispatch())
 	{
-		FRigVMRegistry_NoLock& Registry = FRigVMRegistry_NoLock::GetForWrite();
-		
 		FTypeMap Types;
 		for(const FRigVMTemplateArgument& Argument : Arguments)
 		{
-			Types.Add(Argument.GetName(), Argument.GetTypeIndex_NoLock(InIndex));
+			Types.Add(Argument.GetName(), Argument.GetTypeIndex_NoLock(InIndex, InRegistry));
 		}
 
-		FRigVMDispatchFactory* Factory = Delegates.GetDispatchFactoryDelegate.Execute();
+		FRigVMDispatchFactory* Factory = Delegates.GetDispatchFactoryDelegate.Execute(InRegistry);
 		if (ensure(Factory))
 		{
-			const FRigVMFunctionPtr DispatchFunction = Factory->CreateDispatchFunction_NoLock(Types);
+			const FRigVMFunctionPtr DispatchFunction = Factory->CreateDispatchFunction_NoLock(Types, InRegistry);
 
 			TArray<FRigVMFunctionArgument> FunctionArguments;
 			for(const FRigVMTemplateArgument& Argument : Arguments)
 			{
-				const FRigVMTemplateArgumentType& Type = Registry.GetType_NoLock(Argument.GetTypeIndex_NoLock(InIndex));
+				const FRigVMTemplateArgumentType& Type = InRegistry->GetType_NoLock(Argument.GetTypeIndex_NoLock(InIndex, InRegistry));
 				FunctionArguments.Add(FRigVMFunctionArgument(Argument.Name.ToString(), Type.CPPType.ToString()));
 			}
 			
 			static constexpr TCHAR Format[] = TEXT("%s::%s");
-			const FString PermutationName = Factory->GetPermutationNameImpl(Types, false);
-			const int32 FunctionIndex = Permutations[InIndex] = Registry.Functions.Num();
+			const FString PermutationName = Factory->GetPermutationNameImpl(Types, InRegistry);
+			const int32 FunctionIndex = Permutations[InIndex] = InRegistry->Functions.Num();
 			
-			Registry.Functions.AddElement(
+			InRegistry->Functions.AddElement(
 				FRigVMFunction(
 					PermutationName,
 					DispatchFunction,
@@ -1727,51 +1774,76 @@ const FRigVMFunction* FRigVMTemplate::GetOrCreatePermutation_NoLock(int32 InInde
 					FunctionArguments
 				)
 			);
-			Registry.Functions[FunctionIndex].TemplateIndex = Index;
-			Registry.FunctionNameToIndex.Add(*PermutationName, FunctionIndex);
+			InRegistry->Functions[FunctionIndex].TemplateIndex = Index;
+			InRegistry->Functions[FunctionIndex].PermutationIndex = InIndex;
+			InRegistry->Functions[FunctionIndex].OwnerRegistry = InRegistry.GetRegistry();
+			InRegistry->FunctionNameToIndex.Add(*PermutationName, FunctionIndex);
 
-			TArray<FRigVMFunction> Predicates = Factory->CreateDispatchPredicates_NoLock(Types);
-			Registry.StructNameToPredicates.Add(*PermutationName, MoveTemp(Predicates));
+			TArray<FRigVMFunction> Predicates = Factory->CreateDispatchPredicates_NoLock(Types, InRegistry);
+			InRegistry->StructNameToPredicates.Add(*PermutationName, MoveTemp(Predicates));
 
-			return &Registry.Functions[FunctionIndex];
+			return &InRegistry->Functions[FunctionIndex];
 		}
 	}
 	
 	return nullptr;
 }
 
-bool FRigVMTemplate::ContainsPermutation(const FRigVMFunction* InPermutation, bool bLockRegistry) const
+bool FRigVMTemplate::ContainsPermutation(const FRigVMFunction* InPermutation) const
 {
-	FRigVMRegistryReadLock _(bLockRegistry);
-	return ContainsPermutation_NoLock(InPermutation);
+	return ContainsPermutation_NoLock(InPermutation, FRigVMRegistryReadLock());
 }
 
-bool FRigVMTemplate::ContainsPermutation_NoLock(const FRigVMFunction* InPermutation) const
+bool FRigVMTemplate::ContainsPermutation_NoLock(const FRigVMFunction* InPermutation, const FRigVMRegistryHandle& InRegistry) const
 {
-	return FindPermutation(InPermutation) != INDEX_NONE;
+	check(InRegistry.GetRegistry() == OwnerRegistry);
+
+	return FindPermutation_NoLock(InPermutation, InRegistry) != INDEX_NONE;
 }
 
 int32 FRigVMTemplate::FindPermutation(const FRigVMFunction* InPermutation) const
+{
+	return FindPermutation_NoLock(InPermutation, FRigVMRegistryReadLock());
+}
+
+int32 FRigVMTemplate::FindPermutation_NoLock(const FRigVMFunction* InPermutation, const FRigVMRegistryHandle& InRegistry) const
 {
 	check(InPermutation);
-	return Permutations.Find(InPermutation->Index);
+	check(InPermutation->TemplateIndex == Index);
+	check(InRegistry.GetRegistry() == OwnerRegistry);
+	check(Permutations.Find(InPermutation->Index) == InPermutation->PermutationIndex);
+	return InPermutation->PermutationIndex;
+}
+
+int32 FRigVMTemplate::FindPermutation(const FTypeMap& InTypes) const
+{
+	return FindPermutation_NoLock(InTypes, FRigVMRegistryReadLock());
 }
 
-int32 FRigVMTemplate::FindPermutation(const FTypeMap& InTypes, bool bLockRegistry) const
+int32 FRigVMTemplate::FindPermutation_NoLock(const FTypeMap& InTypes, const FRigVMRegistryHandle& InRegistry) const
 {
+	check(InRegistry.GetRegistry() == OwnerRegistry);
+
 	FTypeMap Types = InTypes;
 	int32 PermutationIndex = INDEX_NONE;
-	if(FullyResolve(Types, PermutationIndex, bLockRegistry))
+	if(FullyResolve_NoLock(Types, PermutationIndex, InRegistry))
 	{
 		return PermutationIndex;
 	}
 	return INDEX_NONE;
 }
 
-bool FRigVMTemplate::FullyResolve(FRigVMTemplate::FTypeMap& InOutTypes, int32& OutPermutationIndex, bool bLockRegistry) const
+bool FRigVMTemplate::FullyResolve(FRigVMTemplate::FTypeMap& InOutTypes, int32& OutPermutationIndex) const
 {
+	return FullyResolve_NoLock(InOutTypes, OutPermutationIndex, FRigVMRegistryReadLock());
+}
+
+bool FRigVMTemplate::FullyResolve_NoLock(FTypeMap& InOutTypes, int32& OutPermutationIndex, const FRigVMRegistryHandle& InRegistry) const
+{
+	check(InRegistry.GetRegistry() == OwnerRegistry);
+
 	TArray<int32> PermutationIndices;
-	Resolve(InOutTypes, PermutationIndices, false, bLockRegistry);
+	Resolve_NoLock(InOutTypes, PermutationIndices, false, InRegistry);
 	if(PermutationIndices.Num() == 1)
 	{
 		OutPermutationIndex = PermutationIndices[0];
@@ -1783,16 +1855,23 @@ bool FRigVMTemplate::FullyResolve(FRigVMTemplate::FTypeMap& InOutTypes, int32& O
 	return OutPermutationIndex != INDEX_NONE;
 }
 
-bool FRigVMTemplate::Resolve(FTypeMap& InOutTypes, TArray<int32>& OutPermutationIndices, bool bAllowFloatingPointCasts, bool bLockRegistry) const
+bool FRigVMTemplate::Resolve(FTypeMap& InOutTypes, TArray<int32>& OutPermutationIndices, bool bAllowFloatingPointCasts) const
 {
-	const FRigVMRegistryReadLock _(bLockRegistry);
-	const FRigVMRegistry_NoLock& Registry = FRigVMRegistry_NoLock::GetForRead();
+	return Resolve_NoLock(InOutTypes, OutPermutationIndices, bAllowFloatingPointCasts, FRigVMRegistryReadLock());
+}
+
+bool FRigVMTemplate::Resolve_NoLock(FTypeMap& InOutTypes, TArray<int32>& OutPermutationIndices, bool bAllowFloatingPointCasts, const FRigVMRegistryHandle& InRegistry) const
+{
+	check(InRegistry.GetRegistry() == OwnerRegistry);
 	
 	FTypeMap InputTypes = InOutTypes;
 	InOutTypes.Reset();
 
+	const int32 CurrentNumPermutations = NumPermutations_NoLock(InRegistry);
 	OutPermutationIndices.Reset();
-	for (int32 PermutationIndex = 0; PermutationIndex < NumPermutations(); PermutationIndex++)
+	OutPermutationIndices.Reserve(CurrentNumPermutations);
+	
+	for (int32 PermutationIndex = 0; PermutationIndex < CurrentNumPermutations; PermutationIndex++)
 	{
 		OutPermutationIndices.Add(PermutationIndex);
 	}
@@ -1801,16 +1880,16 @@ bool FRigVMTemplate::Resolve(FTypeMap& InOutTypes, TArray<int32>& OutPermutation
 	{
 		const TRigVMTypeIndex* InputType = InputTypes.Find(Argument.Name);
 		
-		if (Argument.IsSingleton_NoLock())
+		if (Argument.IsSingleton_NoLock({}, InRegistry))
 		{
 			// if we are singleton - we still need to check if the potentially provided
 			// type is compatible with the singleton type.
-			const TRigVMTypeIndex SingleTypeIndex = Argument.GetTypeIndex_NoLock(0);
+			const TRigVMTypeIndex SingleTypeIndex = Argument.GetTypeIndex_NoLock(0, InRegistry);
 			if (InputType)
 			{
-				if (!Registry.IsWildCardType_NoLock(*InputType))
+				if (!InRegistry->IsWildCardType_NoLock(*InputType))
 				{
-					if (!Registry.CanMatchTypes_NoLock(*InputType, SingleTypeIndex, bAllowFloatingPointCasts))
+					if (!InRegistry->CanMatchTypes_NoLock(*InputType, SingleTypeIndex, bAllowFloatingPointCasts))
 					{
 						OutPermutationIndices.Reset();
 						return false;
@@ -1822,7 +1901,7 @@ bool FRigVMTemplate::Resolve(FTypeMap& InOutTypes, TArray<int32>& OutPermutation
 		}
 		else if (InputType)
 		{
-			TArray<TRigVMTypeIndex> AllTypes; Argument.GetAllTypes_NoLock(AllTypes);
+			TArray<TRigVMTypeIndex> AllTypes; Argument.GetAllTypes_NoLock(AllTypes, InRegistry);
 
 			TRigVMTypeIndex MatchedType = *InputType;
 			bool bFoundMatch = false;
@@ -1831,17 +1910,18 @@ bool FRigVMTemplate::Resolve(FTypeMap& InOutTypes, TArray<int32>& OutPermutation
 			// Using a map to collect all permutations that we can keep/remove
 			// instead of removing them one by one, which can be costly
 			TMap<int32, bool> PermutationsToKeep;
+			PermutationsToKeep.Reserve(AllTypes.Num());
 
 			for (int32 PermutationIndex = 0; PermutationIndex < AllTypes.Num(); PermutationIndex++)
 			{
 				const TRigVMTypeIndex Type = AllTypes[PermutationIndex];
-				if(!Registry.CanMatchTypes_NoLock(Type, *InputType, bAllowFloatingPointCasts))
+				if(!InRegistry->CanMatchTypes_NoLock(Type, *InputType, bAllowFloatingPointCasts))
 				{
-					PermutationsToKeep.FindOrAdd(PermutationIndex) = false;
+					PermutationsToKeep.Add(PermutationIndex, false);
 				}
 				else
 				{
-					PermutationsToKeep.FindOrAdd(PermutationIndex) = true;
+					PermutationsToKeep.Add(PermutationIndex, true);
 					bFoundMatch = true;
 
 					// if the type matches - but it's not the exact same
@@ -1858,7 +1938,7 @@ bool FRigVMTemplate::Resolve(FTypeMap& InOutTypes, TArray<int32>& OutPermutation
 				}
 			}
 
-			OutPermutationIndices = OutPermutationIndices.FilterByPredicate([PermutationsToKeep](int32 PermutationIndex) -> bool
+			OutPermutationIndices = OutPermutationIndices.FilterByPredicate([&PermutationsToKeep](int32 PermutationIndex) -> bool
 			{
 				const bool* Value = PermutationsToKeep.Find(PermutationIndex);
 				return Value ? *Value : false;
@@ -1872,14 +1952,14 @@ bool FRigVMTemplate::Resolve(FTypeMap& InOutTypes, TArray<int32>& OutPermutation
 			}
 		}
 
-		const FRigVMTemplateArgument::EArrayType ArrayType = Argument.GetArrayType_NoLock();
+		const FRigVMTemplateArgument::EArrayType ArrayType = Argument.GetArrayType_NoLock(InRegistry);
 		if(ArrayType == FRigVMTemplateArgument::EArrayType_Mixed)
 		{
 			InOutTypes.Add(Argument.Name, RigVMTypeUtils::TypeIndex::WildCard);
 
 			if(InputType)
 			{
-				if(Registry.IsArrayType_NoLock(*InputType))
+				if(InRegistry->IsArrayType_NoLock(*InputType))
 				{
 					InOutTypes.FindChecked(Argument.Name) = RigVMTypeUtils::TypeIndex::WildCardArray;
 				}
@@ -1900,16 +1980,16 @@ bool FRigVMTemplate::Resolve(FTypeMap& InOutTypes, TArray<int32>& OutPermutation
 		InOutTypes.Reset();
 		for (const FRigVMTemplateArgument& Argument : Arguments)
 		{
-			InOutTypes.Add(Argument.Name, Argument.GetTypeIndex_NoLock(OutPermutationIndices[0]));
+			InOutTypes.Add(Argument.Name, Argument.GetTypeIndex_NoLock(OutPermutationIndices[0], InRegistry));
 		}
 	}
 	else if (OutPermutationIndices.Num() > 1)
 	{
 		for (const FRigVMTemplateArgument& Argument : Arguments)
 		{
-			if (Argument.IsSingleton_NoLock(OutPermutationIndices))
+			if (Argument.IsSingleton_NoLock(OutPermutationIndices, InRegistry))
 			{
-				InOutTypes.FindChecked(Argument.Name) = Argument.GetTypeIndex_NoLock(OutPermutationIndices[0]);
+				InOutTypes.FindChecked(Argument.Name) = Argument.GetTypeIndex_NoLock(OutPermutationIndices[0], InRegistry);
 			}
 		}
 	}
@@ -1917,7 +1997,7 @@ bool FRigVMTemplate::Resolve(FTypeMap& InOutTypes, TArray<int32>& OutPermutation
 	return !OutPermutationIndices.IsEmpty();
 }
 
-uint32 FRigVMTemplate::GetTypesHashFromTypes(const FTypeMap& InTypes) const
+uint32 FRigVMTemplate::GetTypesHashFromTypes(const FTypeMap& InTypes, const FRigVMRegistryHandle& InRegistry) const
 {
 	// It is only a valid type map if it includes all arguments, and non of the types is a wildcard
 	
@@ -1927,10 +2007,9 @@ uint32 FRigVMTemplate::GetTypesHashFromTypes(const FTypeMap& InTypes) const
 		return TypeHash;
 	}
 
-	const FRigVMRegistry_NoLock& Registry = FRigVMRegistry_NoLock::GetForRead();
 	for (const TPair<FName, TRigVMTypeIndex>& Pair : InTypes)
 	{
-		if (Registry.IsWildCardType_NoLock(Pair.Value))
+		if (InRegistry->IsWildCardType_NoLock(Pair.Value))
 		{
 			return TypeHash;
 		}
@@ -1948,18 +2027,19 @@ uint32 FRigVMTemplate::GetTypesHashFromTypes(const FTypeMap& InTypes) const
 	return TypeHash;
 }
 
-bool FRigVMTemplate::ContainsPermutation(const FTypeMap& InTypes, bool bLockRegistry) const
+bool FRigVMTemplate::ContainsPermutation(const FTypeMap& InTypes) const
 {
-	const FRigVMRegistryReadLock Readlock(bLockRegistry);
-	return ContainsPermutation_NoLock((InTypes));
+	return ContainsPermutation_NoLock(InTypes, FRigVMRegistryReadLock());
 }
 
-bool FRigVMTemplate::ContainsPermutation_NoLock(const FTypeMap& InTypes) const
+bool FRigVMTemplate::ContainsPermutation_NoLock(const FTypeMap& InTypes, const FRigVMRegistryHandle& InRegistry) const
 {
+	check(InRegistry.GetRegistry() == OwnerRegistry);
+
 	// If they type map is valid (full description of arguments), then we can rely on
 	// the TypesHashToPermutation cache. Otherwise, we will have to search for a specific permutation
 	// by filtering types
-	const uint32 TypesHash = GetTypesHashFromTypes(InTypes);
+	const uint32 TypesHash = GetTypesHashFromTypes(InTypes, InRegistry);
 	if (TypesHash != 0)
 	{
 		return TypesHashToPermutation.Contains(TypesHash);
@@ -1971,7 +2051,7 @@ bool FRigVMTemplate::ContainsPermutation_NoLock(const FTypeMap& InTypes) const
 	{
 		if (const FRigVMTemplateArgument* Argument = FindArgument(Pair.Key))
 		{
-			const TArray<int32>& ArgumentPermutations = Argument->GetPermutations_NoLock(Pair.Value);
+			const TArray<int32>& ArgumentPermutations = Argument->GetPermutations_NoLock(Pair.Value, InRegistry);
 			if (!ArgumentPermutations.IsEmpty())
 			{
 				// If possible permutations is empty, initialize it
@@ -2012,17 +2092,22 @@ bool FRigVMTemplate::ContainsPermutation_NoLock(const FTypeMap& InTypes) const
 }
 
 bool FRigVMTemplate::ResolveArgument(const FName& InArgumentName, const TRigVMTypeIndex InTypeIndex,
-	FTypeMap& InOutTypes, bool bLockRegistry) const
+	FTypeMap& InOutTypes) const
 {
-	const FRigVMRegistryReadLock _(bLockRegistry);
-	const FRigVMRegistry_NoLock& Registry = FRigVMRegistry_NoLock::GetForRead();
+	return ResolveArgument_NoLock(InArgumentName, InTypeIndex, InOutTypes, FRigVMRegistryReadLock());
+}
 
-	auto RemoveWildCardTypes = [&Registry](const FTypeMap& InTypes)
+bool FRigVMTemplate::ResolveArgument_NoLock(const FName& InArgumentName, const TRigVMTypeIndex InTypeIndex,
+	FTypeMap& InOutTypes, const FRigVMRegistryHandle& InRegistry) const
+{
+	check(InRegistry.GetRegistry() == OwnerRegistry);
+
+	auto RemoveWildCardTypes = [&InRegistry](const FTypeMap& InTypes)
 	{
 		FTypeMap FilteredTypes;
 		for(const FTypePair& Pair: InTypes)
 		{
-			if(!Registry.IsWildCardType_NoLock(Pair.Value))
+			if(!InRegistry->IsWildCardType_NoLock(Pair.Value))
 			{
 				FilteredTypes.Add(Pair);
 			}
@@ -2039,11 +2124,11 @@ bool FRigVMTemplate::ResolveArgument(const FName& InArgumentName, const TRigVMTy
 	TArray<int32> PermutationIndices;
 	FTypeMap RemainingTypesToResolve;
 	
-	if(Resolve(ResolvedTypes, PermutationIndices, true, false))
+	if(Resolve_NoLock(ResolvedTypes, PermutationIndices, true, InRegistry))
 	{
 		// let's see if the input argument resolved into the expected type
 		const TRigVMTypeIndex ResolvedInputType = ResolvedTypes.FindChecked(InArgumentName);
-		if(!Registry.CanMatchTypes_NoLock(ResolvedInputType, InTypeIndex, true))
+		if(!InRegistry->CanMatchTypes_NoLock(ResolvedInputType, InTypeIndex, true))
 		{
 			return false;
 		}
@@ -2080,7 +2165,7 @@ bool FRigVMTemplate::ResolveArgument(const FName& InArgumentName, const TRigVMTy
 			FTypeMap NewResolvedTypes = RemoveWildCardTypes(ResolvedTypes);
 			NewResolvedTypes.FindOrAdd(TypeToResolve.Key) = TypeToResolve.Value;
 
-			if(Resolve(NewResolvedTypes, PermutationIndices, true, false))
+			if(Resolve_NoLock(NewResolvedTypes, PermutationIndices, true, InRegistry))
 			{
 				ResolvedTypes = NewResolvedTypes;
 				RemoveResolvedTypesFromRemainingList(InOutTypes, ResolvedTypes, RemainingTypesToResolve);
@@ -2100,19 +2185,20 @@ bool FRigVMTemplate::ResolveArgument(const FName& InArgumentName, const TRigVMTy
 	return false;
 }
 
-FRigVMTemplateTypeMap FRigVMTemplate::GetTypesForPermutation(const int32 InPermutationIndex, const bool bLockRegistry) const
+FRigVMTemplateTypeMap FRigVMTemplate::GetTypesForPermutation(const int32 InPermutationIndex) const
 {
-	const FRigVMRegistryReadLock _;
-	return GetTypesForPermutation_NoLock(InPermutationIndex);
+	return GetTypesForPermutation_NoLock(InPermutationIndex, FRigVMRegistryReadLock());
 }
 
-FRigVMTemplateTypeMap FRigVMTemplate::GetTypesForPermutation_NoLock(const int32 InPermutationIndex) const
+FRigVMTemplateTypeMap FRigVMTemplate::GetTypesForPermutation_NoLock(const int32 InPermutationIndex, const FRigVMRegistryHandle& InRegistry) const
 {
+	check(InRegistry.GetRegistry() == OwnerRegistry);
+
 	FTypeMap TypeMap;
 	for (int32 ArgIndex = 0; ArgIndex < NumArguments(); ++ArgIndex)
 	{
 		const FRigVMTemplateArgument* Argument = GetArgument(ArgIndex);
-		const TOptional<TRigVMTypeIndex> TypeIndex = Argument->TryToGetTypeIndex_NoLock(InPermutationIndex);
+		const TOptional<TRigVMTypeIndex> TypeIndex = Argument->TryToGetTypeIndex_NoLock(InPermutationIndex, InRegistry);
 		if(TypeIndex.IsSet())
 		{
 			TypeMap.Add(Argument->GetName(), TypeIndex.GetValue());
@@ -2197,8 +2283,10 @@ FString FRigVMTemplate::GetKeywords() const
 
 #endif
 
-bool FRigVMTemplate::UpdateAllArgumentTypesSlow()
+bool FRigVMTemplate::UpdateAllArgumentTypesSlow(FRigVMRegistryHandle& InRegistry)
 {
+	check(InRegistry.GetRegistry() == OwnerRegistry);
+
 	const int32 PrimaryArgumentIndex = Arguments.IndexOfByPredicate([](const FRigVMTemplateArgument& Argument)
 	{
 		return Argument.bUseCategories;
@@ -2215,7 +2303,7 @@ bool FRigVMTemplate::UpdateAllArgumentTypesSlow()
 	for(int32 ArgumentIndex = 0; ArgumentIndex < Arguments.Num(); ArgumentIndex++)
 	{
 		FRigVMTemplateArgument& Argument = Arguments[ArgumentIndex];
-		if(Argument.bUseCategories || Argument.IsSingleton_NoLock())
+		if(Argument.bUseCategories || Argument.IsSingleton_NoLock({}, InRegistry))
 		{
 			continue;
 		}
@@ -2226,7 +2314,6 @@ bool FRigVMTemplate::UpdateAllArgumentTypesSlow()
 #endif
 	}
 
-	const FRigVMRegistry_NoLock& Registry = FRigVMRegistry_NoLock::GetForRead();
 	const FRigVMTemplateArgument& PrimaryArgument = Arguments[PrimaryArgumentIndex];
 	TArray<FRigVMTemplateTypeMap, TInlineAllocator<1>> TypesArray;
 	bool bResult = true;
@@ -2234,18 +2321,18 @@ bool FRigVMTemplate::UpdateAllArgumentTypesSlow()
 	const FRigVMDispatchFactory* Factory = nullptr;
 	if(Delegates.GetDispatchFactoryDelegate.IsBound())
 	{
-		Factory = Delegates.GetDispatchFactoryDelegate.Execute();
+		Factory = Delegates.GetDispatchFactoryDelegate.Execute(InRegistry);
 		ensure(Factory);
 	}
 
 	PrimaryArgument.ForEachType([&](const TRigVMTypeIndex PrimaryTypeIndex)
 	{
-		if(!UpdateArgumentTypes_Impl(PrimaryArgument, PrimaryTypeIndex, Registry, Factory, TypesArray))
+		if(!UpdateArgumentTypes_Impl(PrimaryArgument, PrimaryTypeIndex, InRegistry, Factory, TypesArray))
 		{
 			bResult = false;
 		}
 		return true;
-	});
+	}, InRegistry);
 
 #if WITH_EDITOR
 	for(int32 ArgumentIndex = 0; ArgumentIndex < Arguments.Num(); ArgumentIndex++)
@@ -2253,7 +2340,7 @@ bool FRigVMTemplate::UpdateAllArgumentTypesSlow()
 		FRigVMTemplateArgument& Argument = Arguments[ArgumentIndex];
 		if(Argument.bUseCategories)
 		{
-			Argument.UpdateTypeToPermutationsSlow();
+			Argument.UpdateTypeToPermutationsSlow(InRegistry);
 		}
 	}
 #endif
@@ -2264,18 +2351,18 @@ bool FRigVMTemplate::UpdateAllArgumentTypesSlow()
 bool FRigVMTemplate::UpdateArgumentTypes_Impl(
 	const FRigVMTemplateArgument& InPrimaryArgument,
 	const TRigVMTypeIndex InPrimaryTypeIndex,
-	const FRigVMRegistry_NoLock& InRegistry, 
+	FRigVMRegistryHandle& InRegistry, 
 	const FRigVMDispatchFactory* InFactory,
 	TArray<FRigVMTemplateTypeMap, TInlineAllocator<1>>& InOutTypesArray)
 {
 	InOutTypesArray.Reset();
 	if(InFactory)
 	{
-		InFactory->GetPermutationsFromArgumentType(InPrimaryArgument.Name, InPrimaryTypeIndex, InOutTypesArray);
+		InFactory->GetPermutationsFromArgumentType(InPrimaryArgument.Name, InPrimaryTypeIndex, InOutTypesArray, InRegistry);
 	}
 	else if(OnNewArgumentType().IsBound())
 	{
-		FRigVMTemplateTypeMap Types = OnNewArgumentType().Execute(InPrimaryArgument.Name, InPrimaryTypeIndex);
+		FRigVMTemplateTypeMap Types = OnNewArgumentType().Execute(InPrimaryArgument.Name, InPrimaryTypeIndex, InRegistry);
 		InOutTypesArray.Add(Types);
 	}
 
@@ -2288,11 +2375,11 @@ bool FRigVMTemplate::UpdateArgumentTypes_Impl(
 				for (TPair<FName, TRigVMTypeIndex>& ArgumentAndType : Types)
 				{
 					// similar to FRigVMTemplateArgument::EnsureValidExecuteType
-					InRegistry.ConvertExecuteContextToBaseType_NoLock(ArgumentAndType.Value);
+					InRegistry->ConvertExecuteContextToBaseType_NoLock(ArgumentAndType.Value);
 				}
 
 				// Find if these types were already registered
-				if (ContainsPermutation_NoLock(Types))
+				if (ContainsPermutation_NoLock(Types, InRegistry))
 				{
 					return true;
 				}
@@ -2314,11 +2401,10 @@ bool FRigVMTemplate::UpdateArgumentTypes_Impl(
 
 					const TRigVMTypeIndex& TypeIndex = *TypeIndexPtr;
 					TypeHash = HashCombine(TypeHash, GetTypeHash(TypeIndex));
-
 #if WITH_EDITOR
 					Argument.TypeToPermutations.FindOrAdd(TypeIndex).Add(Permutations.Num());
 #endif
-					if(Argument.bUseCategories || Argument.IsSingleton_NoLock())
+					if(Argument.bUseCategories || Argument.IsSingleton_NoLock({}, InRegistry))
 					{
 						continue;
 					}
@@ -2341,8 +2427,10 @@ bool FRigVMTemplate::UpdateArgumentTypes_Impl(
 	return true;
 }
 
-bool FRigVMTemplate::HandlePropagatedArgumentType(const TRigVMTypeIndex InTypeIndex)
+bool FRigVMTemplate::HandlePropagatedArgumentType(const TRigVMTypeIndex InTypeIndex, FRigVMRegistryHandle& InRegistry)
 {
+	check(InRegistry.GetRegistry() == OwnerRegistry);
+
 	const int32 PrimaryArgumentIndex = Arguments.IndexOfByPredicate([](const FRigVMTemplateArgument& Argument)
 	{
 		return Argument.bUseCategories;
@@ -2356,18 +2444,17 @@ bool FRigVMTemplate::HandlePropagatedArgumentType(const TRigVMTypeIndex InTypeIn
 
 	InvalidateHash();
 
-	const FRigVMRegistry_NoLock& Registry = FRigVMRegistry_NoLock::GetForRead();
 	const FRigVMTemplateArgument& PrimaryArgument = Arguments[PrimaryArgumentIndex];
 	TArray<FRigVMTemplateTypeMap, TInlineAllocator<1>> TypesArray;
 
 	const FRigVMDispatchFactory* Factory = nullptr;
 	if(Delegates.GetDispatchFactoryDelegate.IsBound())
 	{
-		Factory = Delegates.GetDispatchFactoryDelegate.Execute();
+		Factory = Delegates.GetDispatchFactoryDelegate.Execute(InRegistry);
 		ensure(Factory);
 	}
 
-	const bool bResult = UpdateArgumentTypes_Impl(PrimaryArgument, InTypeIndex, Registry, Factory, TypesArray);
+	const bool bResult = UpdateArgumentTypes_Impl(PrimaryArgument, InTypeIndex, InRegistry, Factory, TypesArray);
 #if WITH_EDITOR
 	if(bResult)
 	{
@@ -2376,22 +2463,81 @@ bool FRigVMTemplate::HandlePropagatedArgumentType(const TRigVMTypeIndex InTypeIn
 			FRigVMTemplateArgument& Argument = Arguments[ArgumentIndex];
 			if(Argument.bUseCategories)
 			{
-				Argument.UpdateTypeToPermutationsSlow();
+				Argument.UpdateTypeToPermutationsSlow(InRegistry);
 			}
 		}
 	}
 #endif
+
+	// we may need to update the template permutations since we have have registered a new type
+	// in between other types - especially if the template uses multiple categories somewhere
+	bool bUsesMultipleCategories = false; 
+	for(int32 ArgumentIndex = 0; ArgumentIndex < Arguments.Num(); ArgumentIndex++)
+	{
+		FRigVMTemplateArgument& Argument = Arguments[ArgumentIndex];
+		if(Argument.bUseCategories && Argument.TypeCategories.Num() > 1)
+		{
+			bUsesMultipleCategories = true;
+			break;
+		}
+	}
+	if(bUsesMultipleCategories)
+	{
+		TArray<int32> ValidFunctionIndices;
+		bool bHasChangedPermutations = false;
+
+		for(int32 PermutationIndex = 0; PermutationIndex < Permutations.Num(); PermutationIndex++)
+		{
+			const int32& FunctionIndex = Permutations[PermutationIndex];
+			if(FunctionIndex == INDEX_NONE)
+			{
+				continue;
+			}
+			ValidFunctionIndices.Add(FunctionIndex);
+
+			FRigVMTemplateTypeMap Types;
+			Types.Reserve(OwnerRegistry->Functions[FunctionIndex].Arguments.Num());
+			for(const FRigVMFunctionArgument& Argument : OwnerRegistry->Functions[FunctionIndex].Arguments)
+			{
+				const TRigVMTypeIndex TypeIndex = OwnerRegistry->GetTypeIndexFromCPPType_NoLock(Argument.Type);
+				check(TypeIndex != INDEX_NONE);
+				Types.Add(Argument.Name, TypeIndex);
+			}
+			
+			const int32 NewPermutationIndex = FindPermutation_NoLock(Types, InRegistry);
+			if(NewPermutationIndex != PermutationIndex)
+			{
+				OwnerRegistry->Functions[FunctionIndex].PermutationIndex = NewPermutationIndex;
+				bHasChangedPermutations = true;
+			}
+		}
+
+		if(bHasChangedPermutations)
+		{
+			for(int32 PermutationIndex = 0; PermutationIndex < Permutations.Num(); PermutationIndex++)
+			{
+				Permutations[PermutationIndex] = INDEX_NONE;
+			}
+			for(int32 FunctionIndex : ValidFunctionIndices)
+			{
+				Permutations[OwnerRegistry->Functions[FunctionIndex].PermutationIndex] = FunctionIndex;
+			}
+		}
+	}
+
 	return bResult;
 }
 
-void FRigVMTemplate::HandleTypeRemoval(TRigVMTypeIndex InTypeIndex)
+void FRigVMTemplate::HandleTypeRemoval(TRigVMTypeIndex InTypeIndex, const FRigVMRegistryHandle& InRegistry)
 {
+	check(InRegistry.GetRegistry() == OwnerRegistry);
+
 	InvalidateHash();
 
 	TArray<int32> PermutationsToRemove;
-	for (int32 PermutationIndex = 0; PermutationIndex < NumPermutations(); PermutationIndex++)
+	for (int32 PermutationIndex = 0; PermutationIndex < NumPermutations_NoLock(InRegistry); PermutationIndex++)
 	{
-		FRigVMTemplateTypeMap TypeMap = GetTypesForPermutation_NoLock(PermutationIndex);
+		FRigVMTemplateTypeMap TypeMap = GetTypesForPermutation_NoLock(PermutationIndex, InRegistry);
 
 		TArray<TRigVMTypeIndex> Types;
 		TypeMap.GenerateValueArray(Types);
@@ -2423,27 +2569,38 @@ void FRigVMTemplate::HandleTypeRemoval(TRigVMTypeIndex InTypeIndex)
 #endif
 }
 
-const FRigVMDispatchFactory* FRigVMTemplate::GetDispatchFactory(const bool bLockRegistry) const
+const FRigVMDispatchFactory* FRigVMTemplate::GetDispatchFactory() const
+{
+	FRigVMRegistryReadLock Lock;
+	return GetDispatchFactory_NoLock(Lock);
+}
+
+const FRigVMDispatchFactory* FRigVMTemplate::GetDispatchFactory_NoLock(const FRigVMRegistryHandle& InRegistry) const
 {
-	FRigVMRegistryReadLock _(bLockRegistry);
-	return GetDispatchFactory_NoLock();
+	check(InRegistry.GetRegistry() == OwnerRegistry);
+
+	return UsesDispatch() ? Delegates.GetDispatchFactoryDelegate.Execute(InRegistry) : nullptr;
 }
 
-void FRigVMTemplate::RecomputeTypesHashToPermutations()
+void FRigVMTemplate::RecomputeTypesHashToPermutations(const FRigVMRegistryHandle& InRegistry)
 {
+	check(InRegistry.GetRegistry() == OwnerRegistry);
+
 	TArray<FRigVMTypeCacheScope_NoLock> TypeCaches;
 	TypeCaches.SetNum(NumArguments());
 	
 	for(int32 ArgIndex = 0; ArgIndex < NumArguments(); ArgIndex++)
 	{
-		(void)TypeCaches[ArgIndex].UpdateIfRequired(Arguments[ArgIndex]);
+		(void)TypeCaches[ArgIndex].UpdateIfRequired(InRegistry, Arguments[ArgIndex]);
 	}
 	
-	RecomputeTypesHashToPermutations(TypeCaches);
+	RecomputeTypesHashToPermutations(TypeCaches, InRegistry);
 }
 
-void FRigVMTemplate::RecomputeTypesHashToPermutations(const TArray<FRigVMTypeCacheScope_NoLock>& InTypeCaches)
+void FRigVMTemplate::RecomputeTypesHashToPermutations(const TArray<FRigVMTypeCacheScope_NoLock>& InTypeCaches, const FRigVMRegistryHandle& InRegistry)
 {
+	check(InRegistry.GetRegistry() == OwnerRegistry);
+
 	TypesHashToPermutation.Reset();
 
 	bool bAnyArgumentWithZeroTypes = false;
@@ -2454,7 +2611,7 @@ void FRigVMTemplate::RecomputeTypesHashToPermutations(const TArray<FRigVMTypeCac
 
 	if(!bAnyArgumentWithZeroTypes)
 	{
-		for (int32 PermutationIndex=0; PermutationIndex<NumPermutations(); ++PermutationIndex)
+		for (int32 PermutationIndex=0; PermutationIndex<NumPermutations_NoLock(InRegistry); ++PermutationIndex)
 		{
 			uint32 TypesHash=0;
 			for (int32 ArgIndex=0; ArgIndex < InTypeCaches.Num(); ++ArgIndex)
@@ -2466,8 +2623,10 @@ void FRigVMTemplate::RecomputeTypesHashToPermutations(const TArray<FRigVMTypeCac
 	}
 }
 
-void FRigVMTemplate::UpdateTypesHashToPermutation(const int32 InPermutation)
+void FRigVMTemplate::UpdateTypesHashToPermutation(const int32 InPermutation, const FRigVMRegistryHandle& InRegistry)
 {
+	check(InRegistry.GetRegistry() == OwnerRegistry);
+
 	if (!Permutations.IsValidIndex(InPermutation))
 	{
 		return;
@@ -2476,30 +2635,127 @@ void FRigVMTemplate::UpdateTypesHashToPermutation(const int32 InPermutation)
 	uint32 TypeHash=0;
 	for (const FRigVMTemplateArgument& Argument : Arguments)
 	{
-		TypeHash = HashCombine(TypeHash, GetTypeHash(Argument.GetTypeIndex_NoLock(InPermutation)));
+		TypeHash = HashCombine(TypeHash, GetTypeHash(Argument.GetTypeIndex_NoLock(InPermutation, InRegistry)));
 	}
 	TypesHashToPermutation.Add(TypeHash, InPermutation);
 }
 
-uint32 GetTypeHash_NoLock(const FRigVMTemplateArgument& InArgument)
+TArray<FRigVMTemplateArgumentInfo> FRigVMTemplate::GetFlattenedArgumentInfos_NoLock(FRigVMRegistryHandle& InRegistry) const
+{
+	check(InRegistry.GetRegistry() == OwnerRegistry);
+
+	TArray<int32> PermutationIndices;
+	PermutationIndices.Reserve(Permutations.Num());
+	for(int32 PermutationIndex = 0; PermutationIndex < Permutations.Num(); PermutationIndex++)
+	{
+		PermutationIndices.Add(PermutationIndex);
+	}
+	return GetFlattenedArgumentInfosForPermutations_NoLock(PermutationIndices, InRegistry);
+}
+
+TArray<FRigVMTemplateArgumentInfo> FRigVMTemplate::GetFlattenedArgumentInfosForFunctions_NoLock(const TArray<const FRigVMFunction*>& InFunctions, FRigVMRegistryHandle& InRegistry) const
+{
+	check(InRegistry.GetRegistry() == OwnerRegistry);
+
+	TArray<int32> PermutationIndices;
+	PermutationIndices.Reserve(InFunctions.Num());
+	for(const FRigVMFunction* Function : InFunctions)
+	{
+		check(Function->OwnerRegistry == this->OwnerRegistry);
+		check(Function->TemplateIndex == this->Index);
+		PermutationIndices.Add(Function->PermutationIndex);
+	}
+	return GetFlattenedArgumentInfosForPermutations_NoLock(PermutationIndices, InRegistry);
+}
+
+TArray<FRigVMTemplateArgumentInfo> FRigVMTemplate::GetFlattenedArgumentInfosForPermutations_NoLock(const TArray<int32>& InPermutations, FRigVMRegistryHandle& InRegistry) const
+{
+	if(InPermutations.IsEmpty())
+	{
+		return {};
+	}
+	
+	FRigVMTemplateTypeMap Types;
+	TMap<FName, TArray<TRigVMTypeIndex>> TypeIndices;
+	
+	for(const int32& PermutationIndex : InPermutations)
+	{
+		Types = GetTypesForPermutation_NoLock(PermutationIndex, InRegistry);
+		for(const TPair<FName, TRigVMTypeIndex>& Pair : Types)
+		{
+			TypeIndices.FindOrAdd(Pair.Key).Add(Pair.Value);
+		}
+	}
+
+	TArray<FRigVMTemplateArgumentInfo> Infos;
+	for(const FRigVMTemplateArgument& Argument : Arguments)
+	{
+		if(Argument.IsSingleton_NoLock(InPermutations, InRegistry))
+		{
+			Infos.Emplace(Argument.Name, Argument.Direction, TypeIndices.FindChecked(Argument.Name)[0]);
+		}
+		else
+		{
+			Infos.Emplace(Argument.Name, Argument.Direction, TypeIndices.FindChecked(Argument.Name));
+		}
+	}
+
+	return Infos;
+}
+
+/*
+bool FRigVMTemplate::SetPermutationsOverride_NoLock(const TArray<FRigVMTemplateTypeMap>& InPermutations, FRigVMRegistryHandle& InRegistry)
+{
+	check(InRegistry.GetRegistry() == OwnerRegistry);
+
+	for(const FRigVMTemplateTypeMap& Permutation : InPermutations)
+	{
+		if(!Permutation.Num() == Template->NumArguments())
+		{
+			UE_LOG(LogRigVM, Error, TEXT("%s - SetPermutationsOverride_NoLock: Provided permutation has incorrect number of types (%d), expected %d."),
+				*GetFactoryName().ToString(), Permutation.Num(), Template->NumArguments());
+			return false;
+		}
+		for(const TPair<FName, TRigVMTypeIndex>& TypePair : Permutation)
+		{
+			const FRigVMTemplateArgument* Argument = Template->FindArgument(TypePair.Key);
+			if(Argument == nullptr)
+			{
+				UE_LOG(LogRigVM, Error, TEXT("%s - SetPermutationsOverride_NoLock: Argument '%s' not found in template."),
+					*GetFactoryName().ToString(), TypePair.Key.ToString());
+				return false;
+			}
+
+			if(!Argument->SupportsTypeIndex_NoLock(InRegistry, TypePair.Value))
+			{
+				const FRigVMTemplateArgumentType& Type = InRegistry->GetType_NoLock(TypePair.Value);
+				UE_LOG(LogRigVM, Error, TEXT("%s - SetPermutationsOverride_NoLock: Argument '%s' doesn't support type '%s'."),
+					*GetFactoryName().ToString(), TypePair.Key.ToString(), *Type.CPPType.ToString());
+				return false;
+			}
+		}
+	}
+}
+*/
+
+uint32 GetTypeHash_NoLock(const FRigVMTemplateArgument& InArgument, const FRigVMRegistryHandle& InRegistry)
 {
-	const FRigVMRegistry_NoLock& Registry = FRigVMRegistry_NoLock::GetForRead();
 	uint32 Hash = GetTypeHash(InArgument.Name.ToString());
 	Hash = HashCombine(Hash, GetTypeHash((int32)InArgument.Direction));
 	InArgument.ForEachType([&](const TRigVMTypeIndex TypeIndex)
 	{
-		Hash = HashCombine(Hash, Registry.GetHashForType_NoLock(TypeIndex));
+		Hash = HashCombine(Hash, InRegistry->GetHashForType_NoLock(TypeIndex));
 		return true;
-	});
+	}, InRegistry);
 	return Hash;
 }
 
 uint32 GetTypeHash(const FRigVMTemplate& InTemplate)
 {
-	return InTemplate.ComputeTypeHash();
+	return InTemplate.ComputeTypeHash(FRigVMRegistryReadLock());
 }
 
-uint32 GetTypeHash_NoLock(const FRigVMTemplate& InTemplate)
+uint32 GetTypeHash_NoLock(const FRigVMTemplate& InTemplate, const FRigVMRegistryHandle& InRegistry)
 {
 	if(InTemplate.Hash != UINT32_MAX)
 	{
@@ -2509,12 +2765,12 @@ uint32 GetTypeHash_NoLock(const FRigVMTemplate& InTemplate)
 	uint32 Hash = GetTypeHash(InTemplate.GetNotation().ToString());
 	for(const FRigVMTemplateArgument& Argument : InTemplate.Arguments)
 	{
-		Hash = HashCombine(Hash, GetTypeHash_NoLock(Argument));
+		Hash = HashCombine(Hash, GetTypeHash_NoLock(Argument, InRegistry));
 	}
 
 	// todo: in Dev-EngineMerge we should add the execute arguments to the hash as well
 
-	if(const FRigVMDispatchFactory* Factory = InTemplate.GetDispatchFactory_NoLock())
+	if(const FRigVMDispatchFactory* Factory = InTemplate.GetDispatchFactory_NoLock(InRegistry))
 	{
 		Hash = HashCombine(Hash, GetTypeHash(Factory->GetFactoryName().ToString()));
 	}
@@ -2523,22 +2779,23 @@ uint32 GetTypeHash_NoLock(const FRigVMTemplate& InTemplate)
 	return Hash;
 }
 
-uint32 FRigVMTemplate::ComputeTypeHash() const
+uint32 FRigVMTemplate::ComputeTypeHash(const FRigVMRegistryHandle& InRegistry) const
 {
-	const FRigVMRegistryReadLock _; 
-	return GetTypeHash_NoLock(*this);
+	check(InRegistry.GetRegistry() == OwnerRegistry);
+
+	return GetTypeHash_NoLock(*this, InRegistry);
 }
 
 FRigVMTypeCacheScope_NoLock::FRigVMTypeCacheScope_NoLock()
-	: Argument(nullptr)
-	, bShouldCopyTypes(true)
-	, NumTypes(0)
-	, Types(TArray<TRigVMTypeIndex>())
+	: Registry(nullptr)
+	, Argument(nullptr)
+	, bShouldCopyTypes(false)
 {
 }
 
-FRigVMTypeCacheScope_NoLock::FRigVMTypeCacheScope_NoLock(const FRigVMTemplateArgument& InArgument)
-	: Argument(&InArgument)
+FRigVMTypeCacheScope_NoLock::FRigVMTypeCacheScope_NoLock(const FRigVMRegistryHandle& InRegistry, const FRigVMTemplateArgument& InArgument)
+	: Registry(const_cast<FRigVMRegistry_NoLock*>(InRegistry.GetRegistry()))
+	, Argument(&InArgument)
 	, bShouldCopyTypes(InArgument.FilterType != nullptr)//  ||bUseCategories)??
 {
 }
@@ -2547,17 +2804,24 @@ FRigVMTypeCacheScope_NoLock::~FRigVMTypeCacheScope_NoLock()
 {
 }
 
-const FRigVMTypeCacheScope_NoLock& FRigVMTypeCacheScope_NoLock::UpdateIfRequired(const FRigVMTemplateArgument& InArgument)
+const FRigVMTypeCacheScope_NoLock& FRigVMTypeCacheScope_NoLock::UpdateIfRequired(const FRigVMRegistryHandle& InRegistry, const FRigVMTemplateArgument& InArgument)
 {
+	if(Registry != InRegistry.GetRegistry())
+	{
+		Registry = const_cast<FRigVMRegistry_NoLock*>(InRegistry.GetRegistry());
+		Argument = nullptr;
+	}
 	if(Argument != &InArgument)
 	{
-		*this = FRigVMTypeCacheScope_NoLock(InArgument);
+		Argument = &InArgument;
+		bShouldCopyTypes = InArgument.FilterType != nullptr;
 	}
 	return *this;
 }
 
 int32 FRigVMTypeCacheScope_NoLock::GetNumTypes_NoLock() const
 {
+	check(Registry);
 	if(!NumTypes.IsSet())
 	{
 		if(bShouldCopyTypes)
@@ -2568,7 +2832,7 @@ int32 FRigVMTypeCacheScope_NoLock::GetNumTypes_NoLock() const
 		else
 		{
 			check(Argument);
-			NumTypes = Argument->GetNumTypes_NoLock();
+			NumTypes = Argument->GetNumTypes_NoLock(Registry->GetHandle_NoLock());
 		}
 	}
 	return NumTypes.GetValue();
@@ -2576,6 +2840,7 @@ int32 FRigVMTypeCacheScope_NoLock::GetNumTypes_NoLock() const
 
 TRigVMTypeIndex FRigVMTypeCacheScope_NoLock::GetTypeIndex_NoLock(int32 InIndex) const
 {
+	check(Registry);
 	if(!Types.IsSet())
 	{
 		if(bShouldCopyTypes)
@@ -2584,7 +2849,7 @@ TRigVMTypeIndex FRigVMTypeCacheScope_NoLock::GetTypeIndex_NoLock(int32 InIndex)
 		}
 		else
 		{
-			return Argument->GetTypeIndex_NoLock(InIndex);
+			return Argument->GetTypeIndex_NoLock(InIndex, Registry->GetHandle_NoLock());
 		}
 	}
 
@@ -2595,6 +2860,7 @@ TRigVMTypeIndex FRigVMTypeCacheScope_NoLock::GetTypeIndex_NoLock(int32 InIndex)
 
 void FRigVMTypeCacheScope_NoLock::UpdateTypesIfRequired() const
 {
+	check(Registry);
 	check(bShouldCopyTypes);
 	if(Types.IsSet())
 	{
@@ -2602,11 +2868,11 @@ void FRigVMTypeCacheScope_NoLock::UpdateTypesIfRequired() const
 	}
 	Types = TArray<TRigVMTypeIndex>();
 	TArray<TRigVMTypeIndex>& TypesArray = Types.GetValue();
-	TypesArray.Reserve(Argument->GetNumTypes_NoLock());
+	TypesArray.Reserve(Argument->GetNumTypes_NoLock(Registry->GetHandle_NoLock()));
 
 	Argument->ForEachType([&TypesArray](const TRigVMTypeIndex Type)
 	{
 		TypesArray.Add(Type);
 		return true;
-	});
+	}, Registry->GetHandle_NoLock());
 }
