RigVM: Localized registry in cooked environments - first steps.

diff --git a/Engine/Plugins/Runtime/RigVM/Source/RigVM/Private/RigVMFunctions/RigVMDispatch_Array.cpp b/Engine/Plugins/Runtime/RigVM/Source/RigVM/Private/RigVMFunctions/RigVMDispatch_Array.cpp
--- a/Engine/Plugins/Runtime/RigVM/Source/RigVM/Private/RigVMFunctions/RigVMDispatch_Array.cpp
+++ b/Engine/Plugins/Runtime/RigVM/Source/RigVM/Private/RigVMFunctions/RigVMDispatch_Array.cpp
@@ -192,18 +192,17 @@ const TArray<FRigVMExecuteArgument>& FRigVMDispatch_ArrayBaseMutable::GetExecute
 
 ////////////////////////////////////////////////////////////////////////////////////////////////////////////////////////
 
-const TArray<FRigVMTemplateArgumentInfo>& FRigVMDispatch_ArrayMake::GetArgumentInfos() const
+const TArray<FRigVMTemplateArgumentInfo>& FRigVMDispatch_ArrayMake::GetArgumentInfos(FRigVMRegistryHandle& InRegistry) const
 {
-	static TArray<FRigVMTemplateArgumentInfo> Infos;
-	if (Infos.IsEmpty())
+	if (CachedArgumentInfos.IsEmpty())
 	{
-		Infos.Emplace(CreateArgumentInfo(ValuesName, ERigVMPinDirection::Input));
-		Infos.Emplace(CreateArgumentInfo(ArrayName, ERigVMPinDirection::Output));
+		CachedArgumentInfos.Emplace(CreateArgumentInfo(ValuesName, ERigVMPinDirection::Input));
+		CachedArgumentInfos.Emplace(CreateArgumentInfo(ArrayName, ERigVMPinDirection::Output));
 	}
-	return Infos;
+	return CachedArgumentInfos;
 }
 
-FRigVMTemplateTypeMap FRigVMDispatch_ArrayMake::OnNewArgumentType(const FName& InArgumentName, TRigVMTypeIndex InTypeIndex) const
+FRigVMTemplateTypeMap FRigVMDispatch_ArrayMake::OnNewArgumentType(const FName& InArgumentName, TRigVMTypeIndex InTypeIndex, FRigVMRegistryHandle& InRegistry) const
 {
 	check(InArgumentName == ValuesName || InArgumentName == ArrayName);
 	return {
@@ -263,7 +262,7 @@ FString FRigVMDispatch_ArrayMake::GetArgumentDefaultValue(const FName& InArgumen
 
 #endif
 
-FName FRigVMDispatch_ArrayMake::GetArgumentNameForOperandIndex(int32 InOperandIndex, int32 InTotalOperands) const
+FName FRigVMDispatch_ArrayMake::GetArgumentNameForOperandIndex(int32 InOperandIndex, int32 InTotalOperands, FRigVMRegistryHandle& InRegistry) const
 {
 	if(InOperandIndex < InTotalOperands - 1)
 	{
@@ -298,23 +297,22 @@ void FRigVMDispatch_ArrayMake::Execute(FRigVMExtendedExecuteContext& InContext,
 
 ////////////////////////////////////////////////////////////////////////////////////////////////////////////////////////
 
-FName FRigVMDispatch_ArrayReset::GetArgumentNameForOperandIndex(int32 InOperandIndex, int32 InTotalOperands) const
+FName FRigVMDispatch_ArrayReset::GetArgumentNameForOperandIndex(int32 InOperandIndex, int32 InTotalOperands, FRigVMRegistryHandle& InRegistry) const
 {
 	check(InTotalOperands == 1);
 	return ArrayName;
 }
 
-const TArray<FRigVMTemplateArgumentInfo>& FRigVMDispatch_ArrayReset::GetArgumentInfos() const
+const TArray<FRigVMTemplateArgumentInfo>& FRigVMDispatch_ArrayReset::GetArgumentInfos(FRigVMRegistryHandle& InRegistry) const
 {
-	static TArray<FRigVMTemplateArgumentInfo> Infos;
-	if(Infos.IsEmpty())
+	if(CachedArgumentInfos.IsEmpty())
 	{
-		Infos.Emplace(CreateArgumentInfo(ArrayName, ERigVMPinDirection::IO));
+		CachedArgumentInfos.Emplace(CreateArgumentInfo(ArrayName, ERigVMPinDirection::IO));
 	}
-	return Infos;
+	return CachedArgumentInfos;
 }
 
-FRigVMTemplateTypeMap FRigVMDispatch_ArrayReset::OnNewArgumentType(const FName& InArgumentName, TRigVMTypeIndex InTypeIndex) const
+FRigVMTemplateTypeMap FRigVMDispatch_ArrayReset::OnNewArgumentType(const FName& InArgumentName, TRigVMTypeIndex InTypeIndex, FRigVMRegistryHandle& InRegistry) const
 {
 	check(InArgumentName == ArrayName);
 	return {
@@ -349,7 +347,7 @@ FText FRigVMDispatch_ArrayReset::GetArgumentTooltip(const FName& InArgumentName,
 
 ////////////////////////////////////////////////////////////////////////////////////////////////////////////////////////
 
-FName FRigVMDispatch_ArrayGetNum::GetArgumentNameForOperandIndex(int32 InOperandIndex, int32 InTotalOperands) const
+FName FRigVMDispatch_ArrayGetNum::GetArgumentNameForOperandIndex(int32 InOperandIndex, int32 InTotalOperands, FRigVMRegistryHandle& InRegistry) const
 {
 	static const FLazyName ArgumentNames[] = {
 		ArrayName,
@@ -359,18 +357,17 @@ FName FRigVMDispatch_ArrayGetNum::GetArgumentNameForOperandIndex(int32 InOperand
 	return ArgumentNames[InOperandIndex];
 }
 
-const TArray<FRigVMTemplateArgumentInfo>& FRigVMDispatch_ArrayGetNum::GetArgumentInfos() const
+const TArray<FRigVMTemplateArgumentInfo>& FRigVMDispatch_ArrayGetNum::GetArgumentInfos(FRigVMRegistryHandle& InRegistry) const
 {
-	static TArray<FRigVMTemplateArgumentInfo> Infos;
-	if(Infos.IsEmpty())
+	if(CachedArgumentInfos.IsEmpty())
 	{
-		Infos.Emplace( CreateArgumentInfo(ArrayName, ERigVMPinDirection::Input) );
-		Infos.Emplace( CreateArgumentInfo(NumName, ERigVMPinDirection::Output) );
+		CachedArgumentInfos.Emplace( CreateArgumentInfo(ArrayName, ERigVMPinDirection::Input) );
+		CachedArgumentInfos.Emplace( CreateArgumentInfo(NumName, ERigVMPinDirection::Output) );
 	}
-	return Infos;
+	return CachedArgumentInfos;
 }
 
-FRigVMTemplateTypeMap FRigVMDispatch_ArrayGetNum::OnNewArgumentType(const FName& InArgumentName, TRigVMTypeIndex InTypeIndex) const
+FRigVMTemplateTypeMap FRigVMDispatch_ArrayGetNum::OnNewArgumentType(const FName& InArgumentName, TRigVMTypeIndex InTypeIndex, FRigVMRegistryHandle& InRegistry) const
 {
 	check(InArgumentName == ArrayName);
 	return {
@@ -411,7 +408,7 @@ FText FRigVMDispatch_ArrayGetNum::GetArgumentTooltip(const FName& InArgumentName
 
 ////////////////////////////////////////////////////////////////////////////////////////////////////////////////////////
 
-FName FRigVMDispatch_ArraySetNum::GetArgumentNameForOperandIndex(int32 InOperandIndex, int32 InTotalOperands) const
+FName FRigVMDispatch_ArraySetNum::GetArgumentNameForOperandIndex(int32 InOperandIndex, int32 InTotalOperands, FRigVMRegistryHandle& InRegistry) const
 {
 	static const FLazyName ArgumentNames[] = {
 		ArrayName,
@@ -421,18 +418,17 @@ FName FRigVMDispatch_ArraySetNum::GetArgumentNameForOperandIndex(int32 InOperand
 	return ArgumentNames[InOperandIndex];
 }
 
-const TArray<FRigVMTemplateArgumentInfo>& FRigVMDispatch_ArraySetNum::GetArgumentInfos() const
+const TArray<FRigVMTemplateArgumentInfo>& FRigVMDispatch_ArraySetNum::GetArgumentInfos(FRigVMRegistryHandle& InRegistry) const
 {
-	static TArray<FRigVMTemplateArgumentInfo> Infos;
-	if (Infos.IsEmpty())
+	if (CachedArgumentInfos.IsEmpty())
 	{
-		Infos.Emplace(CreateArgumentInfo(ArrayName, ERigVMPinDirection::IO));
-		Infos.Emplace(CreateArgumentInfo(NumName, ERigVMPinDirection::Input));
+		CachedArgumentInfos.Emplace(CreateArgumentInfo(ArrayName, ERigVMPinDirection::IO));
+		CachedArgumentInfos.Emplace(CreateArgumentInfo(NumName, ERigVMPinDirection::Input));
 	};
-	return Infos;
+	return CachedArgumentInfos;
 }
 
-FRigVMTemplateTypeMap FRigVMDispatch_ArraySetNum::OnNewArgumentType(const FName& InArgumentName, TRigVMTypeIndex InTypeIndex) const
+FRigVMTemplateTypeMap FRigVMDispatch_ArraySetNum::OnNewArgumentType(const FName& InArgumentName, TRigVMTypeIndex InTypeIndex, FRigVMRegistryHandle& InRegistry) const
 {
 	check(InArgumentName == ArrayName);
 	return {
@@ -476,7 +472,7 @@ FText FRigVMDispatch_ArraySetNum::GetArgumentTooltip(const FName& InArgumentName
 
 ////////////////////////////////////////////////////////////////////////////////////////////////////////////////////////
 
-FName FRigVMDispatch_ArrayInit::GetArgumentNameForOperandIndex(int32 InOperandIndex, int32 InTotalOperands) const
+FName FRigVMDispatch_ArrayInit::GetArgumentNameForOperandIndex(int32 InOperandIndex, int32 InTotalOperands, FRigVMRegistryHandle& InRegistry) const
 {
 	static const FLazyName ArgumentNames[] = {
 		ArrayName,
@@ -487,25 +483,24 @@ FName FRigVMDispatch_ArrayInit::GetArgumentNameForOperandIndex(int32 InOperandIn
 	return ArgumentNames[InOperandIndex];
 }
 
-const TArray<FRigVMTemplateArgumentInfo>& FRigVMDispatch_ArrayInit::GetArgumentInfos() const
+const TArray<FRigVMTemplateArgumentInfo>& FRigVMDispatch_ArrayInit::GetArgumentInfos(FRigVMRegistryHandle& InRegistry) const
 {
-	static TArray<FRigVMTemplateArgumentInfo> Infos;
-	if (Infos.IsEmpty())
+	if (CachedArgumentInfos.IsEmpty())
 	{
-		Infos.Emplace(CreateArgumentInfo(ArrayName, ERigVMPinDirection::IO));
-		Infos.Emplace(CreateArgumentInfo(NumName, ERigVMPinDirection::Input));
-		Infos.Emplace(CreateArgumentInfo(ElementName, ERigVMPinDirection::Input));
+		CachedArgumentInfos.Emplace(CreateArgumentInfo(ArrayName, ERigVMPinDirection::IO));
+		CachedArgumentInfos.Emplace(CreateArgumentInfo(NumName, ERigVMPinDirection::Input));
+		CachedArgumentInfos.Emplace(CreateArgumentInfo(ElementName, ERigVMPinDirection::Input));
 	};
-	return Infos;
+	return CachedArgumentInfos;
 }
 
-FRigVMTemplateTypeMap FRigVMDispatch_ArrayInit::OnNewArgumentType(const FName& InArgumentName, TRigVMTypeIndex InTypeIndex) const
+FRigVMTemplateTypeMap FRigVMDispatch_ArrayInit::OnNewArgumentType(const FName& InArgumentName, TRigVMTypeIndex InTypeIndex, FRigVMRegistryHandle& InRegistry) const
 {
 	check(InArgumentName == ArrayName);
 	return {
 		{ArrayName, InTypeIndex},
 		{NumName, RigVMTypeUtils::TypeIndex::Int32},
-		{ElementName, FRigVMRegistry_NoLock::GetForRead().GetBaseTypeFromArrayTypeIndex_NoLock(InTypeIndex)}
+		{ElementName, InRegistry->GetBaseTypeFromArrayTypeIndex_NoLock(InTypeIndex)}
 	};
 }
 
@@ -555,7 +550,7 @@ FText FRigVMDispatch_ArrayInit::GetArgumentTooltip(const FName& InArgumentName,
 
 ////////////////////////////////////////////////////////////////////////////////////////////////////////////////////////
 
-FName FRigVMDispatch_ArrayGetAtIndex::GetArgumentNameForOperandIndex(int32 InOperandIndex, int32 InTotalOperands) const
+FName FRigVMDispatch_ArrayGetAtIndex::GetArgumentNameForOperandIndex(int32 InOperandIndex, int32 InTotalOperands, FRigVMRegistryHandle& InRegistry) const
 {
 	static const FLazyName ArgumentNames[] = {
 		ArrayName,
@@ -566,19 +561,18 @@ FName FRigVMDispatch_ArrayGetAtIndex::GetArgumentNameForOperandIndex(int32 InOpe
 	return ArgumentNames[InOperandIndex];
 }
 
-const TArray<FRigVMTemplateArgumentInfo>& FRigVMDispatch_ArrayGetAtIndex::GetArgumentInfos() const
+const TArray<FRigVMTemplateArgumentInfo>& FRigVMDispatch_ArrayGetAtIndex::GetArgumentInfos(FRigVMRegistryHandle& InRegistry) const
 {
-	static TArray<FRigVMTemplateArgumentInfo> Infos;
-	if (Infos.IsEmpty())
+	if (CachedArgumentInfos.IsEmpty())
 	{
-		Infos.Emplace( CreateArgumentInfo(ArrayName, ERigVMPinDirection::Input) );
-		Infos.Emplace( CreateArgumentInfo(IndexName, ERigVMPinDirection::Input) );
-		Infos.Emplace( CreateArgumentInfo(ElementName, ERigVMPinDirection::Output) );
+		CachedArgumentInfos.Emplace( CreateArgumentInfo(ArrayName, ERigVMPinDirection::Input) );
+		CachedArgumentInfos.Emplace( CreateArgumentInfo(IndexName, ERigVMPinDirection::Input) );
+		CachedArgumentInfos.Emplace( CreateArgumentInfo(ElementName, ERigVMPinDirection::Output) );
 	}
-	return Infos;
+	return CachedArgumentInfos;
 }
 
-FRigVMTemplateTypeMap FRigVMDispatch_ArrayGetAtIndex::OnNewArgumentType(const FName& InArgumentName, TRigVMTypeIndex InTypeIndex) const
+FRigVMTemplateTypeMap FRigVMDispatch_ArrayGetAtIndex::OnNewArgumentType(const FName& InArgumentName, TRigVMTypeIndex InTypeIndex, FRigVMRegistryHandle& InRegistry) const
 {
 	check(InArgumentName == ArrayName || InArgumentName == ElementName);
 
@@ -587,12 +581,12 @@ FRigVMTemplateTypeMap FRigVMDispatch_ArrayGetAtIndex::OnNewArgumentType(const FN
 		return {
 			{ArrayName, InTypeIndex},
 			{IndexName, RigVMTypeUtils::TypeIndex::Int32},
-			{ElementName, FRigVMRegistry_NoLock::GetForRead().GetBaseTypeFromArrayTypeIndex_NoLock(InTypeIndex)}
+			{ElementName, InRegistry->GetBaseTypeFromArrayTypeIndex_NoLock(InTypeIndex)}
 		};
 	}
 
 	return {
-		{ArrayName, FRigVMRegistry_NoLock::GetForRead().GetArrayTypeFromBaseTypeIndex_NoLock(InTypeIndex)},
+		{ArrayName, InRegistry->GetArrayTypeFromBaseTypeIndex_NoLock(InTypeIndex)},
 		{IndexName, RigVMTypeUtils::TypeIndex::Int32},
 		{ElementName, InTypeIndex}
 	};
@@ -640,7 +634,7 @@ FText FRigVMDispatch_ArrayGetAtIndex::GetArgumentTooltip(const FName& InArgument
 
 ////////////////////////////////////////////////////////////////////////////////////////////////////////////////////////
 
-FName FRigVMDispatch_ArraySetAtIndex::GetArgumentNameForOperandIndex(int32 InOperandIndex, int32 InTotalOperands) const
+FName FRigVMDispatch_ArraySetAtIndex::GetArgumentNameForOperandIndex(int32 InOperandIndex, int32 InTotalOperands, FRigVMRegistryHandle& InRegistry) const
 {
 	static const FLazyName ArgumentNames[] = {
 		ArrayName,
@@ -651,19 +645,18 @@ FName FRigVMDispatch_ArraySetAtIndex::GetArgumentNameForOperandIndex(int32 InOpe
 	return ArgumentNames[InOperandIndex];
 }
 
-const TArray<FRigVMTemplateArgumentInfo>& FRigVMDispatch_ArraySetAtIndex::GetArgumentInfos() const
+const TArray<FRigVMTemplateArgumentInfo>& FRigVMDispatch_ArraySetAtIndex::GetArgumentInfos(FRigVMRegistryHandle& InRegistry) const
 {
-	static TArray<FRigVMTemplateArgumentInfo> Infos;
-	if (Infos.IsEmpty())
+	if (CachedArgumentInfos.IsEmpty())
 	{
-		Infos.Emplace( CreateArgumentInfo(ArrayName, ERigVMPinDirection::IO) );
-		Infos.Emplace( CreateArgumentInfo(IndexName, ERigVMPinDirection::Input) );
-		Infos.Emplace( CreateArgumentInfo(ElementName, ERigVMPinDirection::Input) );
+		CachedArgumentInfos.Emplace( CreateArgumentInfo(ArrayName, ERigVMPinDirection::IO) );
+		CachedArgumentInfos.Emplace( CreateArgumentInfo(IndexName, ERigVMPinDirection::Input) );
+		CachedArgumentInfos.Emplace( CreateArgumentInfo(ElementName, ERigVMPinDirection::Input) );
 	}
-	return Infos;
+	return CachedArgumentInfos;
 }
 
-FRigVMTemplateTypeMap FRigVMDispatch_ArraySetAtIndex::OnNewArgumentType(const FName& InArgumentName, TRigVMTypeIndex InTypeIndex) const
+FRigVMTemplateTypeMap FRigVMDispatch_ArraySetAtIndex::OnNewArgumentType(const FName& InArgumentName, TRigVMTypeIndex InTypeIndex, FRigVMRegistryHandle& InRegistry) const
 {
 	check(InArgumentName == ArrayName || InArgumentName == ElementName);
 
@@ -672,12 +665,12 @@ FRigVMTemplateTypeMap FRigVMDispatch_ArraySetAtIndex::OnNewArgumentType(const FN
 		return {
 			{ArrayName, InTypeIndex},
 			{IndexName, RigVMTypeUtils::TypeIndex::Int32},
-			{ElementName, FRigVMRegistry_NoLock::GetForRead().GetBaseTypeFromArrayTypeIndex_NoLock(InTypeIndex)}
+			{ElementName, InRegistry->GetBaseTypeFromArrayTypeIndex_NoLock(InTypeIndex)}
 		};
 	}
 
 	return {
-		{ArrayName, FRigVMRegistry_NoLock::GetForRead().GetArrayTypeFromBaseTypeIndex_NoLock(InTypeIndex)},
+		{ArrayName, InRegistry->GetArrayTypeFromBaseTypeIndex_NoLock(InTypeIndex)},
 		{IndexName, RigVMTypeUtils::TypeIndex::Int32},
 		{ElementName, InTypeIndex}
 	};
@@ -725,7 +718,7 @@ FText FRigVMDispatch_ArraySetAtIndex::GetArgumentTooltip(const FName& InArgument
 
 ////////////////////////////////////////////////////////////////////////////////////////////////////////////////////////
 
-FName FRigVMDispatch_ArrayAdd::GetArgumentNameForOperandIndex(int32 InOperandIndex, int32 InTotalOperands) const
+FName FRigVMDispatch_ArrayAdd::GetArgumentNameForOperandIndex(int32 InOperandIndex, int32 InTotalOperands, FRigVMRegistryHandle& InRegistry) const
 {
 	static const FLazyName ArgumentNames[] = {
 		ArrayName,
@@ -736,16 +729,15 @@ FName FRigVMDispatch_ArrayAdd::GetArgumentNameForOperandIndex(int32 InOperandInd
 	return ArgumentNames[InOperandIndex];
 }
 
-const TArray<FRigVMTemplateArgumentInfo>& FRigVMDispatch_ArrayAdd::GetArgumentInfos() const
+const TArray<FRigVMTemplateArgumentInfo>& FRigVMDispatch_ArrayAdd::GetArgumentInfos(FRigVMRegistryHandle& InRegistry) const
 {
-	static TArray<FRigVMTemplateArgumentInfo> Infos;
-	if (Infos.IsEmpty())
+	if (CachedArgumentInfos.IsEmpty())
 	{
-		Infos.Emplace( CreateArgumentInfo(ArrayName, ERigVMPinDirection::IO) );
-		Infos.Emplace( CreateArgumentInfo(ElementName, ERigVMPinDirection::Input) );
-		Infos.Emplace( CreateArgumentInfo(IndexName, ERigVMPinDirection::Output) );
+		CachedArgumentInfos.Emplace( CreateArgumentInfo(ArrayName, ERigVMPinDirection::IO) );
+		CachedArgumentInfos.Emplace( CreateArgumentInfo(ElementName, ERigVMPinDirection::Input) );
+		CachedArgumentInfos.Emplace( CreateArgumentInfo(IndexName, ERigVMPinDirection::Output) );
 	}
-	return Infos;
+	return CachedArgumentInfos;
 }
 
 void FRigVMDispatch_ArrayAdd::Execute(FRigVMExtendedExecuteContext& InContext, FRigVMMemoryHandleArray Handles, FRigVMPredicateBranchArray Predicates)
@@ -846,7 +838,7 @@ FText FRigVMDispatch_ArrayInsert::GetArgumentTooltip(const FName& InArgumentName
 
 ////////////////////////////////////////////////////////////////////////////////////////////////////////////////////////
 
-FName FRigVMDispatch_ArrayRemove::GetArgumentNameForOperandIndex(int32 InOperandIndex, int32 InTotalOperands) const
+FName FRigVMDispatch_ArrayRemove::GetArgumentNameForOperandIndex(int32 InOperandIndex, int32 InTotalOperands, FRigVMRegistryHandle& InRegistry) const
 {
 	static const FLazyName ArgumentNames[] = {
 		ArrayName,
@@ -856,18 +848,17 @@ FName FRigVMDispatch_ArrayRemove::GetArgumentNameForOperandIndex(int32 InOperand
 	return ArgumentNames[InOperandIndex];
 }
 
-const TArray<FRigVMTemplateArgumentInfo>& FRigVMDispatch_ArrayRemove::GetArgumentInfos() const
+const TArray<FRigVMTemplateArgumentInfo>& FRigVMDispatch_ArrayRemove::GetArgumentInfos(FRigVMRegistryHandle& InRegistry) const
 {
-	static TArray<FRigVMTemplateArgumentInfo> Infos;
-	if (Infos.IsEmpty())
+	if (CachedArgumentInfos.IsEmpty())
 	{
-		Infos.Emplace( CreateArgumentInfo(ArrayName, ERigVMPinDirection::IO) );
-		Infos.Emplace( CreateArgumentInfo(IndexName, ERigVMPinDirection::Input) );
+		CachedArgumentInfos.Emplace( CreateArgumentInfo(ArrayName, ERigVMPinDirection::IO) );
+		CachedArgumentInfos.Emplace( CreateArgumentInfo(IndexName, ERigVMPinDirection::Input) );
 	}
-	return Infos;
+	return CachedArgumentInfos;
 }
 
-FRigVMTemplateTypeMap FRigVMDispatch_ArrayRemove::OnNewArgumentType(const FName& InArgumentName, TRigVMTypeIndex InTypeIndex) const
+FRigVMTemplateTypeMap FRigVMDispatch_ArrayRemove::OnNewArgumentType(const FName& InArgumentName, TRigVMTypeIndex InTypeIndex, FRigVMRegistryHandle& InRegistry) const
 {
 	check(InArgumentName == ArrayName);
 
@@ -943,7 +934,7 @@ FText FRigVMDispatch_ArrayReverse::GetArgumentTooltip(const FName& InArgumentNam
 
 ////////////////////////////////////////////////////////////////////////////////////////////////////////////////////////
 
-FName FRigVMDispatch_ArrayFind::GetArgumentNameForOperandIndex(int32 InOperandIndex, int32 InTotalOperands) const
+FName FRigVMDispatch_ArrayFind::GetArgumentNameForOperandIndex(int32 InOperandIndex, int32 InTotalOperands, FRigVMRegistryHandle& InRegistry) const
 {
 	static const FLazyName ArgumentNames[] = {
 		ArrayName,
@@ -955,20 +946,19 @@ FName FRigVMDispatch_ArrayFind::GetArgumentNameForOperandIndex(int32 InOperandIn
 	return ArgumentNames[InOperandIndex];
 }
 
-const TArray<FRigVMTemplateArgumentInfo>& FRigVMDispatch_ArrayFind::GetArgumentInfos() const
+const TArray<FRigVMTemplateArgumentInfo>& FRigVMDispatch_ArrayFind::GetArgumentInfos(FRigVMRegistryHandle& InRegistry) const
 {
-	static TArray<FRigVMTemplateArgumentInfo> Infos;
-	if (Infos.IsEmpty())
+	if (CachedArgumentInfos.IsEmpty())
 	{
-		Infos.Emplace(CreateArgumentInfo(ArrayName, ERigVMPinDirection::Input));
-		Infos.Emplace(CreateArgumentInfo(ElementName, ERigVMPinDirection::Input));
-		Infos.Emplace(CreateArgumentInfo(IndexName, ERigVMPinDirection::Output));
-		Infos.Emplace(CreateArgumentInfo(SuccessName, ERigVMPinDirection::Output));
+		CachedArgumentInfos.Emplace(CreateArgumentInfo(ArrayName, ERigVMPinDirection::Input));
+		CachedArgumentInfos.Emplace(CreateArgumentInfo(ElementName, ERigVMPinDirection::Input));
+		CachedArgumentInfos.Emplace(CreateArgumentInfo(IndexName, ERigVMPinDirection::Output));
+		CachedArgumentInfos.Emplace(CreateArgumentInfo(SuccessName, ERigVMPinDirection::Output));
 	}
-	return Infos;
+	return CachedArgumentInfos;
 }
 
-FRigVMTemplateTypeMap FRigVMDispatch_ArrayFind::OnNewArgumentType(const FName& InArgumentName, TRigVMTypeIndex InTypeIndex) const
+FRigVMTemplateTypeMap FRigVMDispatch_ArrayFind::OnNewArgumentType(const FName& InArgumentName, TRigVMTypeIndex InTypeIndex, FRigVMRegistryHandle& InRegistry) const
 {
 	check(InArgumentName == ArrayName || InArgumentName == ElementName);
 
@@ -976,14 +966,14 @@ FRigVMTemplateTypeMap FRigVMDispatch_ArrayFind::OnNewArgumentType(const FName& I
 	{
 		return {
 			{ArrayName, InTypeIndex},
-			{ElementName, FRigVMRegistry_NoLock::GetForRead().GetBaseTypeFromArrayTypeIndex_NoLock(InTypeIndex)},
+			{ElementName, InRegistry->GetBaseTypeFromArrayTypeIndex_NoLock(InTypeIndex)},
 			{IndexName, RigVMTypeUtils::TypeIndex::Int32},
 			{SuccessName, RigVMTypeUtils::TypeIndex::Bool}
 		};
 	}
 
 	return {
-		{ArrayName, FRigVMRegistry_NoLock::GetForRead().GetArrayTypeFromBaseTypeIndex_NoLock(InTypeIndex)},
+		{ArrayName, InRegistry->GetArrayTypeFromBaseTypeIndex_NoLock(InTypeIndex)},
 		{ElementName, InTypeIndex},
 		{IndexName, RigVMTypeUtils::TypeIndex::Int32},
 		{SuccessName, RigVMTypeUtils::TypeIndex::Bool}
@@ -1057,7 +1047,7 @@ FText FRigVMDispatch_ArrayFind::GetArgumentTooltip(const FName& InArgumentName,
 
 ////////////////////////////////////////////////////////////////////////////////////////////////////////////////////////
 
-FName FRigVMDispatch_ArrayAppend::GetArgumentNameForOperandIndex(int32 InOperandIndex, int32 InTotalOperands) const
+FName FRigVMDispatch_ArrayAppend::GetArgumentNameForOperandIndex(int32 InOperandIndex, int32 InTotalOperands, FRigVMRegistryHandle& InRegistry) const
 {
 	static const FLazyName ArgumentNames[] = {
 		ArrayName,
@@ -1067,19 +1057,18 @@ FName FRigVMDispatch_ArrayAppend::GetArgumentNameForOperandIndex(int32 InOperand
 	return ArgumentNames[InOperandIndex];
 }
 
-const TArray<FRigVMTemplateArgumentInfo>& FRigVMDispatch_ArrayAppend::GetArgumentInfos() const
+const TArray<FRigVMTemplateArgumentInfo>& FRigVMDispatch_ArrayAppend::GetArgumentInfos(FRigVMRegistryHandle& InRegistry) const
 {
-	static TArray<FRigVMTemplateArgumentInfo> Infos;
-	if (Infos.IsEmpty())
+	if (CachedArgumentInfos.IsEmpty())
 	{
-		Infos.Emplace(CreateArgumentInfo(ArrayName, ERigVMPinDirection::IO));
-		Infos.Emplace(CreateArgumentInfo(OtherName, ERigVMPinDirection::Input));
+		CachedArgumentInfos.Emplace(CreateArgumentInfo(ArrayName, ERigVMPinDirection::IO));
+		CachedArgumentInfos.Emplace(CreateArgumentInfo(OtherName, ERigVMPinDirection::Input));
 	}
-	return Infos;
+	return CachedArgumentInfos;
 }
 
 
-FRigVMTemplateTypeMap FRigVMDispatch_ArrayAppend::OnNewArgumentType(const FName& InArgumentName, TRigVMTypeIndex InTypeIndex) const
+FRigVMTemplateTypeMap FRigVMDispatch_ArrayAppend::OnNewArgumentType(const FName& InArgumentName, TRigVMTypeIndex InTypeIndex, FRigVMRegistryHandle& InRegistry) const
 {
 	check(InArgumentName == ArrayName || InArgumentName == OtherName);
 
@@ -1140,7 +1129,7 @@ FText FRigVMDispatch_ArrayAppend::GetArgumentTooltip(const FName& InArgumentName
 
 ////////////////////////////////////////////////////////////////////////////////////////////////////////////////////////
 
-FName FRigVMDispatch_ArrayClone::GetArgumentNameForOperandIndex(int32 InOperandIndex, int32 InTotalOperands) const
+FName FRigVMDispatch_ArrayClone::GetArgumentNameForOperandIndex(int32 InOperandIndex, int32 InTotalOperands, FRigVMRegistryHandle& InRegistry) const
 {
 	static const FLazyName ArgumentNames[] = {
 		ArrayName,
@@ -1150,18 +1139,17 @@ FName FRigVMDispatch_ArrayClone::GetArgumentNameForOperandIndex(int32 InOperandI
 	return ArgumentNames[InOperandIndex];
 }
 
-const TArray<FRigVMTemplateArgumentInfo>& FRigVMDispatch_ArrayClone::GetArgumentInfos() const
+const TArray<FRigVMTemplateArgumentInfo>& FRigVMDispatch_ArrayClone::GetArgumentInfos(FRigVMRegistryHandle& InRegistry) const
 {
-	static TArray<FRigVMTemplateArgumentInfo> Infos;
-	if (Infos.IsEmpty())
+	if (CachedArgumentInfos.IsEmpty())
 	{
-		Infos.Emplace(CreateArgumentInfo(ArrayName, ERigVMPinDirection::Input));
-		Infos.Emplace(CreateArgumentInfo(CloneName, ERigVMPinDirection::Output));
+		CachedArgumentInfos.Emplace(CreateArgumentInfo(ArrayName, ERigVMPinDirection::Input));
+		CachedArgumentInfos.Emplace(CreateArgumentInfo(CloneName, ERigVMPinDirection::Output));
 	}
-	return Infos;
+	return CachedArgumentInfos;
 }
 
-FRigVMTemplateTypeMap FRigVMDispatch_ArrayClone::OnNewArgumentType(const FName& InArgumentName, TRigVMTypeIndex InTypeIndex) const
+FRigVMTemplateTypeMap FRigVMDispatch_ArrayClone::OnNewArgumentType(const FName& InArgumentName, TRigVMTypeIndex InTypeIndex, FRigVMRegistryHandle& InRegistry) const
 {
 	check(InArgumentName == ArrayName || InArgumentName == CloneName);
 
@@ -1325,7 +1313,7 @@ FText FRigVMDispatch_ArrayUnion::GetArgumentTooltip(const FName& InArgumentName,
 
 ////////////////////////////////////////////////////////////////////////////////////////////////////////////////////////
 
-FName FRigVMDispatch_ArrayDifference::GetArgumentNameForOperandIndex(int32 InOperandIndex, int32 InTotalOperands) const
+FName FRigVMDispatch_ArrayDifference::GetArgumentNameForOperandIndex(int32 InOperandIndex, int32 InTotalOperands, FRigVMRegistryHandle& InRegistry) const
 {
 	static const FLazyName ArgumentNames[] = {
 		ArrayName,
@@ -1336,19 +1324,18 @@ FName FRigVMDispatch_ArrayDifference::GetArgumentNameForOperandIndex(int32 InOpe
 	return ArgumentNames[InOperandIndex];
 }
 
-const TArray<FRigVMTemplateArgumentInfo>& FRigVMDispatch_ArrayDifference::GetArgumentInfos() const
+const TArray<FRigVMTemplateArgumentInfo>& FRigVMDispatch_ArrayDifference::GetArgumentInfos(FRigVMRegistryHandle& InRegistry) const
 {
-	static TArray<FRigVMTemplateArgumentInfo> Infos;
-	if (Infos.IsEmpty())
+	if (CachedArgumentInfos.IsEmpty())
 	{
-		Infos.Emplace(CreateArgumentInfo(ArrayName, ERigVMPinDirection::Input));
-		Infos.Emplace(CreateArgumentInfo(OtherName, ERigVMPinDirection::Input));
-		Infos.Emplace(CreateArgumentInfo(ResultName, ERigVMPinDirection::Output));
+		CachedArgumentInfos.Emplace(CreateArgumentInfo(ArrayName, ERigVMPinDirection::Input));
+		CachedArgumentInfos.Emplace(CreateArgumentInfo(OtherName, ERigVMPinDirection::Input));
+		CachedArgumentInfos.Emplace(CreateArgumentInfo(ResultName, ERigVMPinDirection::Output));
 	}
-	return Infos;
+	return CachedArgumentInfos;
 }
 
-FRigVMTemplateTypeMap FRigVMDispatch_ArrayDifference::OnNewArgumentType(const FName& InArgumentName, TRigVMTypeIndex InTypeIndex) const
+FRigVMTemplateTypeMap FRigVMDispatch_ArrayDifference::OnNewArgumentType(const FName& InArgumentName, TRigVMTypeIndex InTypeIndex, FRigVMRegistryHandle& InRegistry) const
 {
 	check(InArgumentName == ArrayName || InArgumentName == OtherName || InArgumentName == ResultName);
 
@@ -1493,7 +1480,7 @@ FText FRigVMDispatch_ArrayIntersection::GetArgumentTooltip(const FName& InArgume
 
 ////////////////////////////////////////////////////////////////////////////////////////////////////////////////////////
 
-FName FRigVMDispatch_ArrayIterator::GetArgumentNameForOperandIndex(int32 InOperandIndex, int32 InTotalOperands) const
+FName FRigVMDispatch_ArrayIterator::GetArgumentNameForOperandIndex(int32 InOperandIndex, int32 InTotalOperands, FRigVMRegistryHandle& InRegistry) const
 {
 	static const FLazyName ArgumentNames[] = {
 		ArrayName,
@@ -1507,19 +1494,18 @@ FName FRigVMDispatch_ArrayIterator::GetArgumentNameForOperandIndex(int32 InOpera
 	return ArgumentNames[InOperandIndex];
 }
 
-const TArray<FRigVMTemplateArgumentInfo>& FRigVMDispatch_ArrayIterator::GetArgumentInfos() const
+const TArray<FRigVMTemplateArgumentInfo>& FRigVMDispatch_ArrayIterator::GetArgumentInfos(FRigVMRegistryHandle& InRegistry) const
 {
-	static TArray<FRigVMTemplateArgumentInfo> Infos;
-	if (Infos.IsEmpty())
+	if (CachedArgumentInfos.IsEmpty())
 	{
-		Infos.Emplace(CreateArgumentInfo(ArrayName, ERigVMPinDirection::Input));
-		Infos.Emplace(CreateArgumentInfo(ElementName, ERigVMPinDirection::Output));
-		Infos.Emplace(CreateArgumentInfo(IndexName, ERigVMPinDirection::Output));
-		Infos.Emplace(CreateArgumentInfo(CountName, ERigVMPinDirection::Output));
-		Infos.Emplace(CreateArgumentInfo(RatioName, ERigVMPinDirection::Output));
-		Infos.Emplace(FRigVMStruct::ControlFlowBlockToRunName, ERigVMPinDirection::Hidden, RigVMTypeUtils::TypeIndex::FName);
+		CachedArgumentInfos.Emplace(CreateArgumentInfo(ArrayName, ERigVMPinDirection::Input));
+		CachedArgumentInfos.Emplace(CreateArgumentInfo(ElementName, ERigVMPinDirection::Output));
+		CachedArgumentInfos.Emplace(CreateArgumentInfo(IndexName, ERigVMPinDirection::Output));
+		CachedArgumentInfos.Emplace(CreateArgumentInfo(CountName, ERigVMPinDirection::Output));
+		CachedArgumentInfos.Emplace(CreateArgumentInfo(RatioName, ERigVMPinDirection::Output));
+		CachedArgumentInfos.Emplace(FRigVMStruct::ControlFlowBlockToRunName, ERigVMPinDirection::Hidden, RigVMTypeUtils::TypeIndex::FName);
 	}
-	return Infos;
+	return CachedArgumentInfos;
 }
 
 const TArray<FRigVMExecuteArgument>& FRigVMDispatch_ArrayIterator::GetExecuteArguments_Impl(
@@ -1547,7 +1533,7 @@ const bool FRigVMDispatch_ArrayIterator::IsControlFlowBlockSliced(const FName& I
 }
 
 FRigVMTemplateTypeMap FRigVMDispatch_ArrayIterator::OnNewArgumentType(const FName& InArgumentName,
-	TRigVMTypeIndex InTypeIndex) const
+	TRigVMTypeIndex InTypeIndex, FRigVMRegistryHandle& InRegistry) const
 {
 	check(InArgumentName == ArrayName || InArgumentName == ElementName);
 
@@ -1555,7 +1541,7 @@ FRigVMTemplateTypeMap FRigVMDispatch_ArrayIterator::OnNewArgumentType(const FNam
 	{
 		return {
 			{ArrayName, InTypeIndex},
-			{ElementName, FRigVMRegistry_NoLock::GetForRead().GetBaseTypeFromArrayTypeIndex_NoLock(InTypeIndex)},
+			{ElementName, InRegistry->GetBaseTypeFromArrayTypeIndex_NoLock(InTypeIndex)},
 			{IndexName, RigVMTypeUtils::TypeIndex::Int32},
 			{CountName, RigVMTypeUtils::TypeIndex::Int32},
 			{RatioName, RigVMTypeUtils::TypeIndex::Float},
@@ -1564,7 +1550,7 @@ FRigVMTemplateTypeMap FRigVMDispatch_ArrayIterator::OnNewArgumentType(const FNam
 	}
 
 	return {
-		{ArrayName, FRigVMRegistry_NoLock::GetForRead().GetArrayTypeFromBaseTypeIndex_NoLock(InTypeIndex)},
+		{ArrayName, InRegistry->GetArrayTypeFromBaseTypeIndex_NoLock(InTypeIndex)},
 		{ElementName, InTypeIndex},
 		{IndexName, RigVMTypeUtils::TypeIndex::Int32},
 		{CountName, RigVMTypeUtils::TypeIndex::Int32},
diff --git a/Engine/Plugins/Runtime/RigVM/Source/RigVM/Private/RigVMFunctions/RigVMDispatch_CastEnum.cpp b/Engine/Plugins/Runtime/RigVM/Source/RigVM/Private/RigVMFunctions/RigVMDispatch_CastEnum.cpp
--- a/Engine/Plugins/Runtime/RigVM/Source/RigVM/Private/RigVMFunctions/RigVMDispatch_CastEnum.cpp
+++ b/Engine/Plugins/Runtime/RigVM/Source/RigVM/Private/RigVMFunctions/RigVMDispatch_CastEnum.cpp
@@ -6,24 +6,23 @@
 
 #define LOCTEXT_NAMESPACE "RigVMDispatch_CastEnum"
 
-const TArray<FRigVMTemplateArgumentInfo>& FRigVMDispatch_CastEnumToInt::GetArgumentInfos() const
+const TArray<FRigVMTemplateArgumentInfo>& FRigVMDispatch_CastEnumToInt::GetArgumentInfos(FRigVMRegistryHandle& InRegistry) const
 {
-	static TArray<FRigVMTemplateArgumentInfo> OutInfos;
-	if (OutInfos.IsEmpty())
+	if (CachedArgumentInfos.IsEmpty())
 	{
 		static const TArray<FRigVMTemplateArgument::ETypeCategory> ElementCategories =
 		{
 			FRigVMTemplateArgument::ETypeCategory_SingleEnumValue
 		};
 		
-		OutInfos.Emplace(ValueName, ERigVMPinDirection::Input, ElementCategories);
-		OutInfos.Emplace(ResultName, ERigVMPinDirection::Output, RigVMTypeUtils::TypeIndex::Int32);
+		CachedArgumentInfos.Emplace(ValueName, ERigVMPinDirection::Input, ElementCategories);
+		CachedArgumentInfos.Emplace(ResultName, ERigVMPinDirection::Output, RigVMTypeUtils::TypeIndex::Int32);
 	}
 
-	return OutInfos;
+	return CachedArgumentInfos;
 }
 
-bool FRigVMDispatch_CastEnumToInt::GetPermutationsFromArgumentType(const FName& InArgumentName, const TRigVMTypeIndex& InTypeIndex, TArray<FRigVMTemplateTypeMap, TInlineAllocator<1>>& OutPermutations) const
+bool FRigVMDispatch_CastEnumToInt::GetPermutationsFromArgumentType(const FName& InArgumentName, const TRigVMTypeIndex& InTypeIndex, TArray<FRigVMTemplateTypeMap, TInlineAllocator<1>>& OutPermutations, FRigVMRegistryHandle& InRegistry) const
 {
 	if (InArgumentName == ValueName)
 	{
@@ -101,24 +100,23 @@ void FRigVMDispatch_CastEnumToInt::Execute(FRigVMExtendedExecuteContext& InConte
 
 
 
-const TArray<FRigVMTemplateArgumentInfo>& FRigVMDispatch_CastIntToEnum::GetArgumentInfos() const
+const TArray<FRigVMTemplateArgumentInfo>& FRigVMDispatch_CastIntToEnum::GetArgumentInfos(FRigVMRegistryHandle& InRegistry) const
 {
-	static TArray<FRigVMTemplateArgumentInfo> OutInfos;
-	if (OutInfos.IsEmpty())
+	if (CachedArgumentInfos.IsEmpty())
 	{
 		static const TArray<FRigVMTemplateArgument::ETypeCategory> ElementCategories =
 		{
 			FRigVMTemplateArgument::ETypeCategory_SingleEnumValue
 		};
 		
-		OutInfos.Emplace(ValueName, ERigVMPinDirection::Input, RigVMTypeUtils::TypeIndex::Int32);
-		OutInfos.Emplace(ResultName, ERigVMPinDirection::Output, ElementCategories);
+		CachedArgumentInfos.Emplace(ValueName, ERigVMPinDirection::Input, RigVMTypeUtils::TypeIndex::Int32);
+		CachedArgumentInfos.Emplace(ResultName, ERigVMPinDirection::Output, ElementCategories);
 	}
 
-	return OutInfos;
+	return CachedArgumentInfos;
 }
 
-bool FRigVMDispatch_CastIntToEnum::GetPermutationsFromArgumentType(const FName& InArgumentName, const TRigVMTypeIndex& InTypeIndex, TArray<FRigVMTemplateTypeMap, TInlineAllocator<1>>& OutPermutations) const
+bool FRigVMDispatch_CastIntToEnum::GetPermutationsFromArgumentType(const FName& InArgumentName, const TRigVMTypeIndex& InTypeIndex, TArray<FRigVMTemplateTypeMap, TInlineAllocator<1>>& OutPermutations, FRigVMRegistryHandle& InRegistry) const
 {
 	if (InArgumentName == ResultName)
 	{
diff --git a/Engine/Plugins/Runtime/RigVM/Source/RigVM/Private/RigVMFunctions/RigVMDispatch_CastObject.cpp b/Engine/Plugins/Runtime/RigVM/Source/RigVM/Private/RigVMFunctions/RigVMDispatch_CastObject.cpp
--- a/Engine/Plugins/Runtime/RigVM/Source/RigVM/Private/RigVMFunctions/RigVMDispatch_CastObject.cpp
+++ b/Engine/Plugins/Runtime/RigVM/Source/RigVM/Private/RigVMFunctions/RigVMDispatch_CastObject.cpp
@@ -5,31 +5,30 @@
 
 #define LOCTEXT_NAMESPACE "RigVMDispatch_CastObject"
 
-const TArray<FRigVMTemplateArgumentInfo>& FRigVMDispatch_CastObject::GetArgumentInfos() const
+const TArray<FRigVMTemplateArgumentInfo>& FRigVMDispatch_CastObject::GetArgumentInfos(FRigVMRegistryHandle& InRegistry) const
 {
-	static TArray<FRigVMTemplateArgumentInfo> OutInfos;
-	if (OutInfos.IsEmpty())
+	if (CachedArgumentInfos.IsEmpty())
 	{
 		static const TArray<FRigVMTemplateArgument::ETypeCategory> ElementCategories =
 		{
 			FRigVMTemplateArgument::ETypeCategory_SingleObjectValue
 		};
 		
-		static TArray<FRigVMTemplateArgumentInfo> Infos;
+		TArray<FRigVMTemplateArgumentInfo> Infos;
 		Infos.Emplace(ValueName, ERigVMPinDirection::Input, ElementCategories);
 		Infos.Emplace(ResultName, ERigVMPinDirection::Output, INDEX_NONE);
 		
-		OutInfos = BuildArgumentListFromPrimaryArgument(Infos, ValueName);
+		CachedArgumentInfos = BuildArgumentListFromPrimaryArgument(Infos, ValueName, InRegistry);
 	}
 
-	return OutInfos;
+	return CachedArgumentInfos;
 }
 
-bool FRigVMDispatch_CastObject::GetPermutationsFromArgumentType(const FName& InArgumentName, const TRigVMTypeIndex& InTypeIndex, TArray<FRigVMTemplateTypeMap, TInlineAllocator<1>>& OutPermutations) const
+bool FRigVMDispatch_CastObject::GetPermutationsFromArgumentType(const FName& InArgumentName, const TRigVMTypeIndex& InTypeIndex, TArray<FRigVMTemplateTypeMap, TInlineAllocator<1>>& OutPermutations, FRigVMRegistryHandle& InRegistry) const
 {
 	if (InArgumentName == ValueName)
 	{
-		const TArray<TRigVMTypeIndex>& ObjectTypes = FRigVMRegistry_NoLock::GetForRead().GetTypesForCategory_NoLock(FRigVMTemplateArgument::ETypeCategory_SingleObjectValue);
+		const TArray<TRigVMTypeIndex>& ObjectTypes = InRegistry->GetTypesForCategory_NoLock(FRigVMTemplateArgument::ETypeCategory_SingleObjectValue);
 		for (const TRigVMTypeIndex& Type : ObjectTypes)
 		{
 			OutPermutations.Add(
@@ -41,7 +40,7 @@ bool FRigVMDispatch_CastObject::GetPermutationsFromArgumentType(const FName& InA
 	}
 	else if (InArgumentName == ResultName)
 	{
-		const TArray<TRigVMTypeIndex>& ObjectTypes = FRigVMRegistry_NoLock::GetForRead().GetTypesForCategory_NoLock(FRigVMTemplateArgument::ETypeCategory_SingleObjectValue);
+		const TArray<TRigVMTypeIndex>& ObjectTypes = InRegistry->GetTypesForCategory_NoLock(FRigVMTemplateArgument::ETypeCategory_SingleObjectValue);
 		for (const TRigVMTypeIndex& Type : ObjectTypes)
 		{
 			OutPermutations.Add(
diff --git a/Engine/Plugins/Runtime/RigVM/Source/RigVM/Private/RigVMFunctions/RigVMDispatch_Constant.cpp b/Engine/Plugins/Runtime/RigVM/Source/RigVM/Private/RigVMFunctions/RigVMDispatch_Constant.cpp
--- a/Engine/Plugins/Runtime/RigVM/Source/RigVM/Private/RigVMFunctions/RigVMDispatch_Constant.cpp
+++ b/Engine/Plugins/Runtime/RigVM/Source/RigVM/Private/RigVMFunctions/RigVMDispatch_Constant.cpp
@@ -8,21 +8,21 @@
 #include UE_INLINE_GENERATED_CPP_BY_NAME(RigVMDispatch_Constant)
 #define LOCTEXT_NAMESPACE "RigVMDispatch_Constant"
 
-const TArray<FRigVMTemplateArgumentInfo>& FRigVMDispatch_Constant::GetArgumentInfos() const
+const TArray<FRigVMTemplateArgumentInfo>& FRigVMDispatch_Constant::GetArgumentInfos(FRigVMRegistryHandle& InRegistry) const
 {
-	static TArray<FRigVMTemplateArgumentInfo> Infos;
-	if (Infos.IsEmpty())
+	if (CachedArgumentInfos.IsEmpty())
 	{
 		static const TArray<FRigVMTemplateArgument::ETypeCategory> Categories = {
-			FRigVMTemplateArgument::ETypeCategory_SingleAnyValue
+			FRigVMTemplateArgument::ETypeCategory_SingleAnyValue,
+			FRigVMTemplateArgument::ETypeCategory_ArrayAnyValue
 		};
 		
-		Infos.Emplace(FRigVMTemplateArgumentInfo(ValueName, ERigVMPinDirection::IO, Categories));
+		CachedArgumentInfos.Emplace(FRigVMTemplateArgumentInfo(ValueName, ERigVMPinDirection::IO, Categories));
 	}
-	return Infos;
+	return CachedArgumentInfos;
 }
 
-FRigVMTemplateTypeMap FRigVMDispatch_Constant::OnNewArgumentType(const FName& InArgumentName, TRigVMTypeIndex InTypeIndex) const
+FRigVMTemplateTypeMap FRigVMDispatch_Constant::OnNewArgumentType(const FName& InArgumentName, TRigVMTypeIndex InTypeIndex, FRigVMRegistryHandle& InRegistry) const
 {
 	return {
 		{ValueName, InTypeIndex},
diff --git a/Engine/Plugins/Runtime/RigVM/Source/RigVM/Private/RigVMFunctions/RigVMDispatch_Core.cpp b/Engine/Plugins/Runtime/RigVM/Source/RigVM/Private/RigVMFunctions/RigVMDispatch_Core.cpp
--- a/Engine/Plugins/Runtime/RigVM/Source/RigVM/Private/RigVMFunctions/RigVMDispatch_Core.cpp
+++ b/Engine/Plugins/Runtime/RigVM/Source/RigVM/Private/RigVMFunctions/RigVMDispatch_Core.cpp
@@ -6,7 +6,7 @@
 
 #include UE_INLINE_GENERATED_CPP_BY_NAME(RigVMDispatch_Core)
 
-FName FRigVMDispatch_CoreEquals::GetArgumentNameForOperandIndex(int32 InOperandIndex, int32 InTotalOperands) const
+FName FRigVMDispatch_CoreEquals::GetArgumentNameForOperandIndex(int32 InOperandIndex, int32 InTotalOperands, FRigVMRegistryHandle& InRegistry) const
 {
 	static const FLazyName ArgumentNames[] = {
 		AName,
@@ -17,24 +17,23 @@ FName FRigVMDispatch_CoreEquals::GetArgumentNameForOperandIndex(int32 InOperandI
 	return ArgumentNames[InOperandIndex];
 }
 
-const TArray<FRigVMTemplateArgumentInfo>& FRigVMDispatch_CoreEquals::GetArgumentInfos() const
+const TArray<FRigVMTemplateArgumentInfo>& FRigVMDispatch_CoreEquals::GetArgumentInfos(FRigVMRegistryHandle& InRegistry) const
 {
-	static TArray<FRigVMTemplateArgumentInfo> Infos;
-	if (Infos.IsEmpty())
+	if (CachedArgumentInfos.IsEmpty())
 	{
 		static const TArray<FRigVMTemplateArgument::ETypeCategory> ValueCategories = {
 			FRigVMTemplateArgument::ETypeCategory_SingleAnyValue,
 			FRigVMTemplateArgument::ETypeCategory_ArrayAnyValue
 		};
-		Infos.Emplace(AName, ERigVMPinDirection::Input, ValueCategories);
-		Infos.Emplace(BName, ERigVMPinDirection::Input, ValueCategories);
-		Infos.Emplace(ResultName, ERigVMPinDirection::Output, RigVMTypeUtils::TypeIndex::Bool);
+		CachedArgumentInfos.Emplace(AName, ERigVMPinDirection::Input, ValueCategories);
+		CachedArgumentInfos.Emplace(BName, ERigVMPinDirection::Input, ValueCategories);
+		CachedArgumentInfos.Emplace(ResultName, ERigVMPinDirection::Output, RigVMTypeUtils::TypeIndex::Bool);
 	}
-	return Infos;
+	return CachedArgumentInfos;
 }
 
 FRigVMTemplateTypeMap FRigVMDispatch_CoreEquals::OnNewArgumentType(const FName& InArgumentName,
-	TRigVMTypeIndex InTypeIndex) const
+	TRigVMTypeIndex InTypeIndex, FRigVMRegistryHandle& InRegistry) const
 {
 	FRigVMTemplateTypeMap Types;
 	Types.Add(AName, InTypeIndex);
@@ -43,7 +42,7 @@ FRigVMTemplateTypeMap FRigVMDispatch_CoreEquals::OnNewArgumentType(const FName&
 	return Types;
 }
 
-FRigVMFunctionPtr FRigVMDispatch_CoreEquals::GetDispatchFunctionImpl(const FRigVMTemplateTypeMap& InTypes) const
+FRigVMFunctionPtr FRigVMDispatch_CoreEquals::GetDispatchFunctionImpl(const FRigVMTemplateTypeMap& InTypes, const FRigVMRegistryHandle& InRegistry) const
 {
 	const TRigVMTypeIndex TypeIndex = InTypes.FindChecked(AName);
 	check(TypeIndex == InTypes.FindChecked(BName));
@@ -69,27 +68,27 @@ FRigVMFunctionPtr FRigVMDispatch_CoreEquals::GetDispatchFunctionImpl(const FRigV
 	{
 		return &FRigVMDispatch_CoreEquals::StringEquals;
 	}
-	if(TypeIndex == FRigVMRegistry_NoLock::GetForRead().GetTypeIndex_NoLock<FVector>())
+	if(TypeIndex == InRegistry->GetTypeIndex_NoLock<FVector>())
 	{
 		return &FRigVMDispatch_CoreEquals::MathTypeEquals<FVector>;
 	}
-	if(TypeIndex == FRigVMRegistry_NoLock::GetForRead().GetTypeIndex_NoLock<FVector2D>())
+	if(TypeIndex == InRegistry->GetTypeIndex_NoLock<FVector2D>())
 	{
 		return &FRigVMDispatch_CoreEquals::MathTypeEquals<FVector2D>;
 	}
-	if(TypeIndex == FRigVMRegistry_NoLock::GetForRead().GetTypeIndex_NoLock<FRotator>())
+	if(TypeIndex == InRegistry->GetTypeIndex_NoLock<FRotator>())
 	{
 		return &FRigVMDispatch_CoreEquals::MathTypeEquals<FRotator>;
 	}
-	if(TypeIndex == FRigVMRegistry_NoLock::GetForRead().GetTypeIndex_NoLock<FQuat>())
+	if(TypeIndex == InRegistry->GetTypeIndex_NoLock<FQuat>())
 	{
 		return &FRigVMDispatch_CoreEquals::MathTypeEquals<FQuat>;
 	}
-	if(TypeIndex == FRigVMRegistry_NoLock::GetForRead().GetTypeIndex_NoLock<FTransform>())
+	if(TypeIndex == InRegistry->GetTypeIndex_NoLock<FTransform>())
 	{
 		return &FRigVMDispatch_CoreEquals::MathTypeEquals<FTransform>;
 	}
-	if(TypeIndex == FRigVMRegistry_NoLock::GetForRead().GetTypeIndex_NoLock<FLinearColor>())
+	if(TypeIndex == InRegistry->GetTypeIndex_NoLock<FLinearColor>())
 	{
 		return &FRigVMDispatch_CoreEquals::MathTypeEquals<FLinearColor>;
 	}
@@ -124,18 +123,17 @@ void FRigVMDispatch_CoreEquals::Execute(FRigVMExtendedExecuteContext& InContext,
 
 // duplicate the code here so that the FRigVMDispatch_CoreNotEquals has it's own static variables
 // to store the types are registration time.
-const TArray<FRigVMTemplateArgumentInfo>& FRigVMDispatch_CoreNotEquals::GetArgumentInfos() const
+const TArray<FRigVMTemplateArgumentInfo>& FRigVMDispatch_CoreNotEquals::GetArgumentInfos(FRigVMRegistryHandle& InRegistry) const
 {
-	static TArray<FRigVMTemplateArgumentInfo> Infos;
-	if (Infos.IsEmpty())
+	if (CachedArgumentInfos.IsEmpty())
 	{
 		static const TArray<FRigVMTemplateArgument::ETypeCategory> ValueCategories = {
 			FRigVMTemplateArgument::ETypeCategory_SingleAnyValue,
 			FRigVMTemplateArgument::ETypeCategory_ArrayAnyValue
 		};
-		Infos.Emplace(AName, ERigVMPinDirection::Input, ValueCategories);
-		Infos.Emplace(BName, ERigVMPinDirection::Input, ValueCategories);
-		Infos.Emplace(ResultName, ERigVMPinDirection::Output, RigVMTypeUtils::TypeIndex::Bool);
+		CachedArgumentInfos.Emplace(AName, ERigVMPinDirection::Input, ValueCategories);
+		CachedArgumentInfos.Emplace(BName, ERigVMPinDirection::Input, ValueCategories);
+		CachedArgumentInfos.Emplace(ResultName, ERigVMPinDirection::Output, RigVMTypeUtils::TypeIndex::Bool);
 	}
-	return Infos;
+	return CachedArgumentInfos;
 }
diff --git a/Engine/Plugins/Runtime/RigVM/Source/RigVM/Private/RigVMFunctions/RigVMDispatch_If.cpp b/Engine/Plugins/Runtime/RigVM/Source/RigVM/Private/RigVMFunctions/RigVMDispatch_If.cpp
--- a/Engine/Plugins/Runtime/RigVM/Source/RigVM/Private/RigVMFunctions/RigVMDispatch_If.cpp
+++ b/Engine/Plugins/Runtime/RigVM/Source/RigVM/Private/RigVMFunctions/RigVMDispatch_If.cpp
@@ -6,7 +6,7 @@
 
 #include UE_INLINE_GENERATED_CPP_BY_NAME(RigVMDispatch_If)
 
-FName FRigVMDispatch_If::GetArgumentNameForOperandIndex(int32 InOperandIndex, int32 InTotalOperands) const
+FName FRigVMDispatch_If::GetArgumentNameForOperandIndex(int32 InOperandIndex, int32 InTotalOperands, FRigVMRegistryHandle& InRegistry) const
 {
 	static const FLazyName ArgumentNames[] = {
 		ConditionName,
@@ -18,26 +18,25 @@ FName FRigVMDispatch_If::GetArgumentNameForOperandIndex(int32 InOperandIndex, in
 	return ArgumentNames[InOperandIndex];
 }
 
-const TArray<FRigVMTemplateArgumentInfo>& FRigVMDispatch_If::GetArgumentInfos() const
+const TArray<FRigVMTemplateArgumentInfo>& FRigVMDispatch_If::GetArgumentInfos(FRigVMRegistryHandle& InRegistry) const
 {
-	static TArray<FRigVMTemplateArgumentInfo> Infos;
-	if(Infos.IsEmpty())
+	if(CachedArgumentInfos.IsEmpty())
 	{
 		static const TArray<FRigVMTemplateArgument::ETypeCategory> ValueCategories = {
 			FRigVMTemplateArgument::ETypeCategory_SingleAnyValue,
 			FRigVMTemplateArgument::ETypeCategory_ArrayAnyValue
 		};
-		
-		Infos.Emplace(ConditionName, ERigVMPinDirection::Input, RigVMTypeUtils::TypeIndex::Bool);
-		Infos.Emplace(TrueName, ERigVMPinDirection::Input, ValueCategories);
-		Infos.Emplace(FalseName, ERigVMPinDirection::Input, ValueCategories);
-		Infos.Emplace(ResultName, ERigVMPinDirection::Output, ValueCategories);
+	
+		CachedArgumentInfos.Emplace(ConditionName, ERigVMPinDirection::Input, RigVMTypeUtils::TypeIndex::Bool);
+		CachedArgumentInfos.Emplace(TrueName, ERigVMPinDirection::Input, ValueCategories);
+		CachedArgumentInfos.Emplace(FalseName, ERigVMPinDirection::Input, ValueCategories);
+		CachedArgumentInfos.Emplace(ResultName, ERigVMPinDirection::Output, ValueCategories);
 	}
-	return Infos;
+	return CachedArgumentInfos;
 }
 
 FRigVMTemplateTypeMap FRigVMDispatch_If::OnNewArgumentType(const FName& InArgumentName,
-	TRigVMTypeIndex InTypeIndex) const
+	TRigVMTypeIndex InTypeIndex, FRigVMRegistryHandle& InRegistry) const
 {
 	FRigVMTemplateTypeMap Types;
 	Types.Add(ConditionName, RigVMTypeUtils::TypeIndex::Bool);
@@ -63,7 +62,7 @@ FString FRigVMDispatch_If::GetArgumentMetaData(const FName& InArgumentName, cons
 
 #endif
 
-FRigVMFunctionPtr FRigVMDispatch_If::GetDispatchFunctionImpl(const FRigVMTemplateTypeMap& InTypes) const
+FRigVMFunctionPtr FRigVMDispatch_If::GetDispatchFunctionImpl(const FRigVMTemplateTypeMap& InTypes, const FRigVMRegistryHandle& InRegistry) const
 {
 	return &FRigVMDispatch_If::Execute;
 }
diff --git a/Engine/Plugins/Runtime/RigVM/Source/RigVM/Private/RigVMFunctions/RigVMDispatch_MakeStruct.cpp b/Engine/Plugins/Runtime/RigVM/Source/RigVM/Private/RigVMFunctions/RigVMDispatch_MakeStruct.cpp
--- a/Engine/Plugins/Runtime/RigVM/Source/RigVM/Private/RigVMFunctions/RigVMDispatch_MakeStruct.cpp
+++ b/Engine/Plugins/Runtime/RigVM/Source/RigVM/Private/RigVMFunctions/RigVMDispatch_MakeStruct.cpp
@@ -8,22 +8,21 @@
 #include UE_INLINE_GENERATED_CPP_BY_NAME(RigVMDispatch_MakeStruct)
 #define LOCTEXT_NAMESPACE "RigVMDispatch_MakeStruct"
 
-const TArray<FRigVMTemplateArgumentInfo>& FRigVMDispatch_MakeStruct::GetArgumentInfos() const
+const TArray<FRigVMTemplateArgumentInfo>& FRigVMDispatch_MakeStruct::GetArgumentInfos(FRigVMRegistryHandle& InRegistry) const
 {
-	static TArray<FRigVMTemplateArgumentInfo> Infos;
-	if(Infos.IsEmpty())
+	if(CachedArgumentInfos.IsEmpty())
 	{
 		static const TArray<FRigVMTemplateArgument::ETypeCategory> Categories = {
 			FRigVMTemplateArgument::ETypeCategory_SingleScriptStructValue
 		};
 		
-		Infos.Emplace(ElementsName, ERigVMPinDirection::Input, Categories);
-		Infos.Emplace(StructName, ERigVMPinDirection::Output, Categories);
+		CachedArgumentInfos.Emplace(ElementsName, ERigVMPinDirection::Input, Categories);
+		CachedArgumentInfos.Emplace(StructName, ERigVMPinDirection::Output, Categories);
 	}
-	return Infos;
+	return CachedArgumentInfos;
 }
 
-FRigVMTemplateTypeMap FRigVMDispatch_MakeStruct::OnNewArgumentType(const FName& InArgumentName, TRigVMTypeIndex InTypeIndex) const
+FRigVMTemplateTypeMap FRigVMDispatch_MakeStruct::OnNewArgumentType(const FName& InArgumentName, TRigVMTypeIndex InTypeIndex, FRigVMRegistryHandle& InRegistry) const
 {
 	return {
 		{ElementsName, InTypeIndex},
@@ -102,19 +101,18 @@ void FRigVMDispatch_MakeStruct::Execute(FRigVMExtendedExecuteContext& InContext,
 	URigVMMemoryStorage::CopyProperty(TargetProperty, TargetMemory, SourceProperty, SourceMemory);
 }
 
-const TArray<FRigVMTemplateArgumentInfo>& FRigVMDispatch_BreakStruct::GetArgumentInfos() const
+const TArray<FRigVMTemplateArgumentInfo>& FRigVMDispatch_BreakStruct::GetArgumentInfos(FRigVMRegistryHandle& InRegistry) const
 {
-	static TArray<FRigVMTemplateArgumentInfo> Infos;
-	if(Infos.IsEmpty())
+	if(CachedArgumentInfos.IsEmpty())
 	{
 		static const TArray<FRigVMTemplateArgument::ETypeCategory> Categories = {
 			FRigVMTemplateArgument::ETypeCategory_SingleScriptStructValue
 		};
 
-		Infos.Emplace(StructName, ERigVMPinDirection::Input, Categories);
-		Infos.Emplace(ElementsName, ERigVMPinDirection::Output, Categories);
+		CachedArgumentInfos.Emplace(StructName, ERigVMPinDirection::Input, Categories);
+		CachedArgumentInfos.Emplace(ElementsName, ERigVMPinDirection::Output, Categories);
 	}
-	return Infos;
+	return CachedArgumentInfos;
 }
 
 #if WITH_EDITOR
diff --git a/Engine/Plugins/Runtime/RigVM/Source/RigVM/Private/RigVMFunctions/RigVMDispatch_Print.cpp b/Engine/Plugins/Runtime/RigVM/Source/RigVM/Private/RigVMFunctions/RigVMDispatch_Print.cpp
--- a/Engine/Plugins/Runtime/RigVM/Source/RigVM/Private/RigVMFunctions/RigVMDispatch_Print.cpp
+++ b/Engine/Plugins/Runtime/RigVM/Source/RigVM/Private/RigVMFunctions/RigVMDispatch_Print.cpp
@@ -7,7 +7,7 @@
 
 #include UE_INLINE_GENERATED_CPP_BY_NAME(RigVMDispatch_Print)
 
-FName FRigVMDispatch_Print::GetArgumentNameForOperandIndex(int32 InOperandIndex, int32 InTotalOperands) const
+FName FRigVMDispatch_Print::GetArgumentNameForOperandIndex(int32 InOperandIndex, int32 InTotalOperands, FRigVMRegistryHandle& InRegistry) const
 {
 	static const FLazyName ArgumentNames[] = {
 		PrefixName,
@@ -20,23 +20,22 @@ FName FRigVMDispatch_Print::GetArgumentNameForOperandIndex(int32 InOperandIndex,
 	return ArgumentNames[InOperandIndex];
 }
 
-const TArray<FRigVMTemplateArgumentInfo>& FRigVMDispatch_Print::GetArgumentInfos() const
+const TArray<FRigVMTemplateArgumentInfo>& FRigVMDispatch_Print::GetArgumentInfos(FRigVMRegistryHandle& InRegistry) const
 {
-	static TArray<FRigVMTemplateArgumentInfo> Infos;
-	if (Infos.IsEmpty())
+	if(CachedArgumentInfos.IsEmpty())
 	{
 		static const TArray<FRigVMTemplateArgument::ETypeCategory> ValueCategories = {
 			FRigVMTemplateArgument::ETypeCategory_SingleAnyValue,
 			FRigVMTemplateArgument::ETypeCategory_ArrayAnyValue
 		};
-	
-		Infos.Emplace(PrefixName, ERigVMPinDirection::Input, RigVMTypeUtils::TypeIndex::FString);
-		Infos.Emplace(ValueName, ERigVMPinDirection::Input, ValueCategories);
-		Infos.Emplace(EnabledName, ERigVMPinDirection::Input, RigVMTypeUtils::TypeIndex::Bool);
-		Infos.Emplace(ScreenDurationName, ERigVMPinDirection::Input, RigVMTypeUtils::TypeIndex::Float);
-		Infos.Emplace(ScreenColorName, ERigVMPinDirection::Input, FRigVMRegistry_NoLock::GetForRead().GetTypeIndex_NoLock<FLinearColor>());
+
+		CachedArgumentInfos.Emplace(PrefixName, ERigVMPinDirection::Input, RigVMTypeUtils::TypeIndex::FString);
+		CachedArgumentInfos.Emplace(ValueName, ERigVMPinDirection::Input, ValueCategories);
+		CachedArgumentInfos.Emplace(EnabledName, ERigVMPinDirection::Input, RigVMTypeUtils::TypeIndex::Bool);
+		CachedArgumentInfos.Emplace(ScreenDurationName, ERigVMPinDirection::Input, RigVMTypeUtils::TypeIndex::Float);
+		CachedArgumentInfos.Emplace(ScreenColorName, ERigVMPinDirection::Input, InRegistry->GetTypeIndex_NoLock<FLinearColor>());
 	}
-	return Infos;
+	return CachedArgumentInfos;
 }
 
 const TArray<FRigVMExecuteArgument>& FRigVMDispatch_Print::GetExecuteArguments_Impl(const FRigVMDispatchContext& InContext) const
@@ -48,14 +47,14 @@ const TArray<FRigVMExecuteArgument>& FRigVMDispatch_Print::GetExecuteArguments_I
 }
 
 FRigVMTemplateTypeMap FRigVMDispatch_Print::OnNewArgumentType(const FName& InArgumentName,
-                                                              TRigVMTypeIndex InTypeIndex) const
+                                                              TRigVMTypeIndex InTypeIndex, FRigVMRegistryHandle& InRegistry) const
 {
 	FRigVMTemplateTypeMap Types;
 	Types.Add(PrefixName, RigVMTypeUtils::TypeIndex::FString);
 	Types.Add(ValueName, InTypeIndex);
 	Types.Add(EnabledName, RigVMTypeUtils::TypeIndex::Bool);
 	Types.Add(ScreenDurationName, RigVMTypeUtils::TypeIndex::Float);
-	Types.Add(ScreenColorName, FRigVMRegistry_NoLock::GetForRead().GetTypeIndex_NoLock<FLinearColor>());
+	Types.Add(ScreenColorName, InRegistry->GetTypeIndex_NoLock<FLinearColor>());
 	return Types;
 }
 
diff --git a/Engine/Plugins/Runtime/RigVM/Source/RigVM/Private/RigVMFunctions/RigVMDispatch_Select.cpp b/Engine/Plugins/Runtime/RigVM/Source/RigVM/Private/RigVMFunctions/RigVMDispatch_Select.cpp
--- a/Engine/Plugins/Runtime/RigVM/Source/RigVM/Private/RigVMFunctions/RigVMDispatch_Select.cpp
+++ b/Engine/Plugins/Runtime/RigVM/Source/RigVM/Private/RigVMFunctions/RigVMDispatch_Select.cpp
@@ -7,7 +7,7 @@
 
 #include UE_INLINE_GENERATED_CPP_BY_NAME(RigVMDispatch_Select)
 
-FName FRigVMDispatch_SelectInt32::GetArgumentNameForOperandIndex(int32 InOperandIndex, int32 InTotalOperands) const
+FName FRigVMDispatch_SelectInt32::GetArgumentNameForOperandIndex(int32 InOperandIndex, int32 InTotalOperands, FRigVMRegistryHandle& InRegistry) const
 {
 	if(InOperandIndex > 0 && InOperandIndex < InTotalOperands - 1)
 	{
@@ -18,10 +18,9 @@ FName FRigVMDispatch_SelectInt32::GetArgumentNameForOperandIndex(int32 InOperand
 }
 
 
-const TArray<FRigVMTemplateArgumentInfo>& FRigVMDispatch_SelectInt32::GetArgumentInfos() const
+const TArray<FRigVMTemplateArgumentInfo>& FRigVMDispatch_SelectInt32::GetArgumentInfos(FRigVMRegistryHandle& InRegistry) const
 {
-	static TArray<FRigVMTemplateArgumentInfo> Infos;
-	if (Infos.IsEmpty())
+	if (CachedArgumentInfos.IsEmpty())
 	{
 		static const TArray<FRigVMTemplateArgument::ETypeCategory> ArrayCategories = {
 			FRigVMTemplateArgument::ETypeCategory_ArrayAnyValue,
@@ -32,27 +31,27 @@ const TArray<FRigVMTemplateArgumentInfo>& FRigVMDispatch_SelectInt32::GetArgumen
 			FRigVMTemplateArgument::ETypeCategory_ArrayAnyValue
 		};
 		
-		Infos.Emplace(IndexName, ERigVMPinDirection::Input, RigVMTypeUtils::TypeIndex::Int32);
-		Infos.Emplace(ValuesName, ERigVMPinDirection::Input, ArrayCategories);
-		Infos.Emplace(ResultName, ERigVMPinDirection::Output, SingleCategories);
+		CachedArgumentInfos.Emplace(IndexName, ERigVMPinDirection::Input, RigVMTypeUtils::TypeIndex::Int32);
+		CachedArgumentInfos.Emplace(ValuesName, ERigVMPinDirection::Input, ArrayCategories);
+		CachedArgumentInfos.Emplace(ResultName, ERigVMPinDirection::Output, SingleCategories);
 	}
 
-	return Infos;
+	return CachedArgumentInfos;
 }
 
 FRigVMTemplateTypeMap FRigVMDispatch_SelectInt32::OnNewArgumentType(const FName& InArgumentName,
-	TRigVMTypeIndex InTypeIndex) const
+	TRigVMTypeIndex InTypeIndex, FRigVMRegistryHandle& InRegistry) const
 {
 	FRigVMTemplateTypeMap Types;
 	Types.Add(IndexName, RigVMTypeUtils::TypeIndex::Int32);
 	if(InArgumentName == ValuesName)
 	{
 		Types.Add(ValuesName, InTypeIndex);
-		Types.Add(ResultName, FRigVMRegistry_NoLock::GetForRead().GetBaseTypeFromArrayTypeIndex_NoLock(InTypeIndex));
+		Types.Add(ResultName, InRegistry->GetBaseTypeFromArrayTypeIndex_NoLock(InTypeIndex));
 	}
 	else
 	{
-		Types.Add(ValuesName, FRigVMRegistry_NoLock::GetForRead().GetArrayTypeFromBaseTypeIndex_NoLock(InTypeIndex));
+		Types.Add(ValuesName, InRegistry->GetArrayTypeFromBaseTypeIndex_NoLock(InTypeIndex));
 		Types.Add(ResultName, InTypeIndex);
 	}
 	return Types;
@@ -103,7 +102,7 @@ FString FRigVMDispatch_SelectInt32::GetArgumentDefaultValue(const FName& InArgum
 #endif
 
 
-FRigVMFunctionPtr FRigVMDispatch_SelectInt32::GetDispatchFunctionImpl(const FRigVMTemplateTypeMap& InTypes) const
+FRigVMFunctionPtr FRigVMDispatch_SelectInt32::GetDispatchFunctionImpl(const FRigVMTemplateTypeMap& InTypes, const FRigVMRegistryHandle& InRegistry) const
 {
 	return &FRigVMDispatch_SelectInt32::Execute;
 }
diff --git a/Engine/Plugins/Runtime/RigVM/Source/RigVM/Private/RigVMFunctions/RigVMDispatch_Switch.cpp b/Engine/Plugins/Runtime/RigVM/Source/RigVM/Private/RigVMFunctions/RigVMDispatch_Switch.cpp
--- a/Engine/Plugins/Runtime/RigVM/Source/RigVM/Private/RigVMFunctions/RigVMDispatch_Switch.cpp
+++ b/Engine/Plugins/Runtime/RigVM/Source/RigVM/Private/RigVMFunctions/RigVMDispatch_Switch.cpp
@@ -7,7 +7,7 @@
 
 #include UE_INLINE_GENERATED_CPP_BY_NAME(RigVMDispatch_Switch)
 
-FName FRigVMDispatch_SwitchInt32::GetArgumentNameForOperandIndex(int32 InOperandIndex, int32 InTotalOperands) const
+FName FRigVMDispatch_SwitchInt32::GetArgumentNameForOperandIndex(int32 InOperandIndex, int32 InTotalOperands, FRigVMRegistryHandle& InRegistry) const
 {
 	static const FLazyName ArgumentNames[] = {
 		IndexName,
@@ -17,15 +17,14 @@ FName FRigVMDispatch_SwitchInt32::GetArgumentNameForOperandIndex(int32 InOperand
 	return ArgumentNames[InOperandIndex];
 }
 
-const TArray<FRigVMTemplateArgumentInfo>& FRigVMDispatch_SwitchInt32::GetArgumentInfos() const
+const TArray<FRigVMTemplateArgumentInfo>& FRigVMDispatch_SwitchInt32::GetArgumentInfos(FRigVMRegistryHandle& InRegistry) const
 {
-	static TArray<FRigVMTemplateArgumentInfo> Arguments;
-	if(Arguments.IsEmpty())
+	if(CachedArgumentInfos.IsEmpty())
 	{
-		Arguments.Emplace(IndexName, ERigVMPinDirection::Input, RigVMTypeUtils::TypeIndex::Int32);
-		Arguments.Emplace(FRigVMStruct::ControlFlowBlockToRunName, ERigVMPinDirection::Hidden, RigVMTypeUtils::TypeIndex::FName);
+		CachedArgumentInfos.Emplace(IndexName, ERigVMPinDirection::Input, RigVMTypeUtils::TypeIndex::Int32);
+		CachedArgumentInfos.Emplace(FRigVMStruct::ControlFlowBlockToRunName, ERigVMPinDirection::Hidden, RigVMTypeUtils::TypeIndex::FName);
 	}
-	return Arguments;
+	return CachedArgumentInfos;
 }
 
 const TArray<FRigVMExecuteArgument>& FRigVMDispatch_SwitchInt32::GetExecuteArguments_Impl(const FRigVMDispatchContext& InContext) const
@@ -120,7 +119,7 @@ const TArray<FName>& FRigVMDispatch_SwitchInt32::GetControlFlowBlocks_Impl(const
 	return DefaultBlocks;
 }
 
-FRigVMFunctionPtr FRigVMDispatch_SwitchInt32::GetDispatchFunctionImpl(const FRigVMTemplateTypeMap& InTypes) const
+FRigVMFunctionPtr FRigVMDispatch_SwitchInt32::GetDispatchFunctionImpl(const FRigVMTemplateTypeMap& InTypes, const FRigVMRegistryHandle& InRegistry) const
 {
 	return &FRigVMDispatch_SwitchInt32::Execute;
 }
diff --git a/Engine/Plugins/Runtime/RigVM/Source/RigVM/Private/RigVMFunctions/RigVMFunction_String.cpp b/Engine/Plugins/Runtime/RigVM/Source/RigVM/Private/RigVMFunctions/RigVMFunction_String.cpp
--- a/Engine/Plugins/Runtime/RigVM/Source/RigVM/Private/RigVMFunctions/RigVMFunction_String.cpp
+++ b/Engine/Plugins/Runtime/RigVM/Source/RigVM/Private/RigVMFunctions/RigVMFunction_String.cpp
@@ -253,24 +253,23 @@ FRigVMFunction_StringPadInteger_Execute()
 	}
 }
 
-const TArray<FRigVMTemplateArgumentInfo>& FRigDispatch_ToString::GetArgumentInfos() const
+const TArray<FRigVMTemplateArgumentInfo>& FRigDispatch_ToString::GetArgumentInfos(FRigVMRegistryHandle& InRegistry) const
 {
-	static TArray<FRigVMTemplateArgumentInfo> Infos;
-	if(Infos.IsEmpty())
+	if(CachedArgumentInfos.IsEmpty())
 	{
 		static const TArray<FRigVMTemplateArgument::ETypeCategory> ValueCategories = {
 			FRigVMTemplateArgument::ETypeCategory_SingleAnyValue,
 			FRigVMTemplateArgument::ETypeCategory_ArrayAnyValue
 		};
 		
-		Infos.Emplace(TEXT("Value"), ERigVMPinDirection::Input, ValueCategories);
-		Infos.Emplace(TEXT("Result"), ERigVMPinDirection::Output, RigVMTypeUtils::TypeIndex::FString);
+		CachedArgumentInfos.Emplace(TEXT("Value"), ERigVMPinDirection::Input, ValueCategories);
+		CachedArgumentInfos.Emplace(TEXT("Result"), ERigVMPinDirection::Output, RigVMTypeUtils::TypeIndex::FString);
 	}
-	return Infos;
+	return CachedArgumentInfos;
 }
 
 FRigVMTemplateTypeMap FRigDispatch_ToString::OnNewArgumentType(const FName& InArgumentName,
-	TRigVMTypeIndex InTypeIndex) const
+	TRigVMTypeIndex InTypeIndex, FRigVMRegistryHandle& InRegistry) const
 {
 	FRigVMTemplateTypeMap Types;
 	Types.Add(TEXT("Value"), InTypeIndex);
@@ -278,7 +277,7 @@ FRigVMTemplateTypeMap FRigDispatch_ToString::OnNewArgumentType(const FName& InAr
 	return Types;
 }
 
-FRigVMFunctionPtr FRigDispatch_ToString::GetDispatchFunctionImpl(const FRigVMTemplateTypeMap& InTypes) const
+FRigVMFunctionPtr FRigDispatch_ToString::GetDispatchFunctionImpl(const FRigVMTemplateTypeMap& InTypes, const FRigVMRegistryHandle& InRegistry) const
 {
 	// treat name and string special
 	const TRigVMTypeIndex& ValueTypeIndex = InTypes.FindChecked(TEXT("Value"));
@@ -322,24 +321,23 @@ void FRigDispatch_ToString::Execute(FRigVMExtendedExecuteContext& InContext, FRi
 	ValueProperty->ExportText_Direct(Result, Value, Value, nullptr, PPF_None, nullptr);
 }
 
-const TArray<FRigVMTemplateArgumentInfo>& FRigDispatch_FromString::GetArgumentInfos() const
+const TArray<FRigVMTemplateArgumentInfo>& FRigDispatch_FromString::GetArgumentInfos(FRigVMRegistryHandle& InRegistry) const
 {
-	static TArray<FRigVMTemplateArgumentInfo> Infos;
-	if(Infos.IsEmpty())
+	if(CachedArgumentInfos.IsEmpty())
 	{
 		static const TArray<FRigVMTemplateArgument::ETypeCategory> ValueCategories = {
 			FRigVMTemplateArgument::ETypeCategory_SingleAnyValue,
 			FRigVMTemplateArgument::ETypeCategory_ArrayAnyValue
 		};
 		
-		Infos.Emplace(TEXT("String"), ERigVMPinDirection::Input, RigVMTypeUtils::TypeIndex::FString);
-		Infos.Emplace(TEXT("Result"), ERigVMPinDirection::Output, ValueCategories);
+		CachedArgumentInfos.Emplace(TEXT("String"), ERigVMPinDirection::Input, RigVMTypeUtils::TypeIndex::FString);
+		CachedArgumentInfos.Emplace(TEXT("Result"), ERigVMPinDirection::Output, ValueCategories);
 	}
-	return Infos;
+	return CachedArgumentInfos;
 }
 
 FRigVMTemplateTypeMap FRigDispatch_FromString::OnNewArgumentType(const FName& InArgumentName,
-	TRigVMTypeIndex InTypeIndex) const
+	TRigVMTypeIndex InTypeIndex, FRigVMRegistryHandle& InRegistry) const
 {
 	FRigVMTemplateTypeMap Types;
 	Types.Add(TEXT("String"), RigVMTypeUtils::TypeIndex::FString);
@@ -347,7 +345,7 @@ FRigVMTemplateTypeMap FRigDispatch_FromString::OnNewArgumentType(const FName& In
 	return Types;
 }
 
-FRigVMFunctionPtr FRigDispatch_FromString::GetDispatchFunctionImpl(const FRigVMTemplateTypeMap& InTypes) const
+FRigVMFunctionPtr FRigDispatch_FromString::GetDispatchFunctionImpl(const FRigVMTemplateTypeMap& InTypes, const FRigVMRegistryHandle& InRegistry) const
 {
 	// treat name and string special
 	const TRigVMTypeIndex& ResultTypeIndex = InTypes.FindChecked(TEXT("Result"));
diff --git a/Engine/Plugins/Runtime/RigVM/Source/RigVM/Private/RigVMHost.cpp b/Engine/Plugins/Runtime/RigVM/Source/RigVM/Private/RigVMHost.cpp
--- a/Engine/Plugins/Runtime/RigVM/Source/RigVM/Private/RigVMHost.cpp
+++ b/Engine/Plugins/Runtime/RigVM/Source/RigVM/Private/RigVMHost.cpp
@@ -1197,11 +1197,22 @@ TArray<const UObject*> URigVMHost::GetUserDefinedDependencies(const TArray<const
 	const TArray<const FRigVMFunction*>& Functions = VM->GetFunctions();
 	for (const FRigVMFunction* Function : Functions)
 	{
-		const FRigVMRegistry_RWLock& Registry = FRigVMRegistry_RWLock::Get();
-		const TArray<TRigVMTypeIndex>& TypeIndices = Function->GetArgumentTypeIndices();
+		FRigVMRegistry_NoLock* Registry = VM->GetLocalizedRegistry();
+
+		// if the VM doesn't have a localized registry
+		// then access the global one and lock it for reading.
+		TUniquePtr<FRigVMRegistryReadLock> ReadLock;
+		if(Registry == nullptr)
+		{
+			FRigVMRegistry_RWLock& RegistryRef = FRigVMRegistry_RWLock::Get();
+			Registry = &RegistryRef;
+			ReadLock = MakeUnique<FRigVMRegistryReadLock>(RegistryRef);
+		}
+		
+		const TArray<TRigVMTypeIndex>& TypeIndices = Function->GetArgumentTypeIndices_NoLock(Registry->GetHandle_NoLock());
 		for (const TRigVMTypeIndex& TypeIndex : TypeIndices)
 		{
-			const FRigVMTemplateArgumentType& Type = Registry.GetType(TypeIndex);
+			const FRigVMTemplateArgumentType& Type = Registry->GetType_NoLock(TypeIndex);
 			if (Cast<UUserDefinedStruct>(Type.CPPTypeObject) ||
 				Cast<UUserDefinedEnum>(Type.CPPTypeObject))
 			{
diff --git a/Engine/Plugins/Runtime/RigVM/Source/RigVM/Public/RigVMCore/RigVM.h b/Engine/Plugins/Runtime/RigVM/Source/RigVM/Public/RigVMCore/RigVM.h
--- a/Engine/Plugins/Runtime/RigVM/Source/RigVM/Public/RigVMCore/RigVM.h
+++ b/Engine/Plugins/Runtime/RigVM/Source/RigVM/Public/RigVMCore/RigVM.h
@@ -889,6 +889,17 @@ public:
 		return Statistics;
 	}
 
+	// Returns the localized registry this VM uses or nullptr
+	RIGVM_API FRigVMRegistry_NoLock* GetLocalizedRegistry() const;
+
+	// Creates a localized registry for this VM
+	RIGVM_API bool CreateLocalizedRegistryIfRequired(bool bForce = false);
+
+private:
+	//  The local registry to use by the VM
+	TSharedPtr<FRigVMRegistry_NoLock> LocalizedRegistry;
+
+public:
 
 #if WITH_EDITOR
 	UE_DEPRECATED(5.4, "Please use DumpByteCodeAsTextArray with ExtendedExecuteContext parameter.")
@@ -986,7 +997,18 @@ public:
 	RIGVM_API void InvalidateCachedMemory(FRigVMExtendedExecuteContext& Context);
 
 private:
-	RIGVM_API void InstructionOpEval(FRigVMExtendedExecuteContext& Context, int32 InstructionIndex, int32 InHandleBaseIndex, const TFunctionRef<void(FRigVMExtendedExecuteContext& Context, int32 InHandleIndex, const FRigVMBranchInfoKey& InBranchInfoKey, const FRigVMOperand& InArg)>& InOpFunc);
+	RIGVM_API void InstructionOpEval(
+		FRigVMExtendedExecuteContext& Context,
+		FRigVMRegistryHandle& RegistryHandle,
+		int32 InstructionIndex,
+		int32 InHandleBaseIndex,
+		const TFunctionRef<void(
+			FRigVMExtendedExecuteContext& /*Context*/,
+			FRigVMRegistryHandle& /*RegistryHandle*/,
+			int32 /*InHandleIndex*/,
+			const FRigVMBranchInfoKey& /*InBranchInfoKey*/,
+			const FRigVMOperand& /*InArg*/
+		)>& InOpFunc);
 	RIGVM_API void PrepareMemoryForExecution(FRigVMExtendedExecuteContext& Context);
 	RIGVM_API void CacheMemoryHandlesIfRequired(FRigVMExtendedExecuteContext& Context);
 	RIGVM_API void RebuildByteCodeOnLoad();
@@ -1063,7 +1085,7 @@ private:
 	// debug watch register memory needs to be cleared for each execution
 	RIGVM_API void ClearDebugMemory(FRigVMExtendedExecuteContext& Context);
 
-	RIGVM_API void CacheSingleMemoryHandle(FRigVMExtendedExecuteContext& Context, int32 InHandleIndex, const FRigVMBranchInfoKey& InBranchInfoKey, const FRigVMOperand& InArg, bool bForExecute = false);
+	RIGVM_API void CacheSingleMemoryHandle(FRigVMExtendedExecuteContext& Context, const FRigVMRegistryHandle& InRegistry, int32 InHandleIndex, const FRigVMBranchInfoKey& InBranchInfoKey, const FRigVMOperand& InArg, bool bForExecute = false);
 
 	void CopyOperandForDebuggingIfNeeded(FRigVMExtendedExecuteContext& Context, const FRigVMOperand& InArg, const FRigVMMemoryHandle& InHandle)
 	{
diff --git a/Engine/Plugins/Runtime/RigVM/Source/RigVM/Public/RigVMCore/RigVMDispatchFactory.h b/Engine/Plugins/Runtime/RigVM/Source/RigVM/Public/RigVMCore/RigVMDispatchFactory.h
--- a/Engine/Plugins/Runtime/RigVM/Source/RigVM/Public/RigVMCore/RigVMDispatchFactory.h
+++ b/Engine/Plugins/Runtime/RigVM/Source/RigVM/Public/RigVMCore/RigVMDispatchFactory.h
@@ -9,6 +9,7 @@
 #include "HAL/Platform.h"
 #include "Internationalization/Text.h"
 #include "Math/Color.h"
+#include "RigVMCore/RigVMRegistryHandle.h"
 #include "RigVMCore/RigVMFunction.h"
 #include "RigVMCore/RigVMStructUpgradeInfo.h"
 #include "RigVMCore/RigVMTemplate.h"
@@ -34,8 +35,9 @@ struct FRigVMDispatchContext
 		: Instance()
 		, Subject(nullptr)
 		, StringRepresentation()
-	{}
-	
+	{
+	}
+
 	TSharedPtr<FStructOnScope> Instance;
 	const UObject* Subject;
 	FString StringRepresentation;
@@ -52,8 +54,10 @@ struct FRigVMDispatchFactory
 public:
 
 	FRigVMDispatchFactory()
-		: FactoryScriptStruct(nullptr)
+		: FactoryIndex(INDEX_NONE)
+		, FactoryScriptStruct(nullptr)
 		, CachedTemplate(nullptr)
+		, OwnerRegistry(nullptr)
 	{
 		ArgumentNamesMutex = new FCriticalSection();
 	}
@@ -66,9 +70,15 @@ public:
 	// returns the name of the factory template
 	RIGVM_API FName GetFactoryName() const;
 
+	// returns the index of this factory in its registry
+	RIGVM_API int32 GetFactoryIndex() const { return FactoryIndex; }
+
 	// returns the struct for this dispatch factory
 	UScriptStruct* GetScriptStruct() const { return FactoryScriptStruct; };
 
+	// returns the registry this factory belongs top
+	const FRigVMRegistry_NoLock* GetOwnerRegistry() const { return OwnerRegistry; }
+
 #if WITH_EDITOR
 	
 	// returns the title of the node for a given type set.
@@ -121,7 +131,7 @@ public:
 	// Returns the name to use for the branch info / argument based on the operand index.
 	// operand index count may be different than the number of arguments, since fixed size array
 	// arguments get unrolled.
-	RIGVM_API virtual FName GetArgumentNameForOperandIndex(int32 InOperandIndex, int32 InTotalOperands) const;
+	RIGVM_API virtual FName GetArgumentNameForOperandIndex(int32 InOperandIndex, int32 InTotalOperands, FRigVMRegistryHandle& InRegistry) const;
 
 	// returns true if the dispatch is a control flow dispatch
 	bool IsControlFlowDispatch(const FRigVMDispatchContext& InContext) const { return !GetControlFlowBlocks(InContext).IsEmpty(); }
@@ -141,23 +151,23 @@ public:
 	// registered needed types during registration of the factory.
 	// this is called within a NoLock code path - so make sure to access
 	// the FRigVMRegistry_NoLock within as needed.
-	virtual void RegisterDependencyTypes_NoLock(FRigVMRegistry_NoLock& InRegistry) const {}
+	virtual void RegisterDependencyTypes_NoLock(FRigVMRegistryHandle& InRegistry) const {}
 
 	// returns the arguments of the template
-	RIGVM_API virtual const TArray<FRigVMTemplateArgumentInfo>& GetArgumentInfos() const;
+	RIGVM_API virtual const TArray<FRigVMTemplateArgumentInfo>& GetArgumentInfos(FRigVMRegistryHandle& InRegistry) const;
 
 	// returns the execute arguments of the template
-	RIGVM_API TArray<FRigVMExecuteArgument> GetExecuteArguments_NoLock(const FRigVMDispatchContext& InContext) const;
+	RIGVM_API TArray<FRigVMExecuteArgument> GetExecuteArguments_NoLock(const FRigVMDispatchContext& InContext, const FRigVMRegistryHandle& InRegistry) const;
 
 	// this function is deprecated, please use GetPermutationsFromArgumentType
 	// returns the new permutation argument types after a new type is defined for one argument
 	// this happens if types are being loaded later after this factory has already been deployed (like UUserDefinedStruct)
-	virtual FRigVMTemplateTypeMap OnNewArgumentType(const FName& InArgumentName, TRigVMTypeIndex InTypeIndex) const { return FRigVMTemplateTypeMap(); }
+	virtual FRigVMTemplateTypeMap OnNewArgumentType(const FName& InArgumentName, TRigVMTypeIndex InTypeIndex, FRigVMRegistryHandle& InRegistry) const { return FRigVMTemplateTypeMap(); }
 
 	// returns the new permutations argument types after a new type is defined for one argument
-	virtual bool GetPermutationsFromArgumentType(const FName& InArgumentName, const TRigVMTypeIndex& InTypeIndex, TArray<FRigVMTemplateTypeMap, TInlineAllocator<1>>& OutPermutations) const
+	virtual bool GetPermutationsFromArgumentType(const FName& InArgumentName, const TRigVMTypeIndex& InTypeIndex, TArray<FRigVMTemplateTypeMap, TInlineAllocator<1>>& OutPermutations, FRigVMRegistryHandle& InRegistry) const
 	{
-		FRigVMTemplateTypeMap Permutation = OnNewArgumentType(InArgumentName, InTypeIndex);
+		FRigVMTemplateTypeMap Permutation = OnNewArgumentType(InArgumentName, InTypeIndex, InRegistry);
 		if (!Permutation.IsEmpty())
 		{
 			OutPermutations.Add(Permutation);
@@ -173,44 +183,43 @@ public:
 	RIGVM_API FRigVMFunctionPtr GetOrCreateDispatchFunction(const FRigVMTemplateTypeMap& InTypes) const;
 
 	// builds and returns the template
-	const FRigVMTemplate* GetTemplate() const
-	{
-		return GetTemplate_RWLock();
-	}
+	RIGVM_API const FRigVMTemplate* GetTemplate() const;
+
+	// builds and returns the template
+	RIGVM_API const FRigVMTemplate* GetTemplate_NoLock(FRigVMRegistryHandle& InRegistry) const;
 
 	// returns the name of the factory template
 	RIGVM_API FName GetTemplateNotation() const;
+	RIGVM_API FName GetTemplateNotation_NoLock(FRigVMRegistryHandle& InRegistry) const;
 
 	// returns the name of the permutation for a given set of types
-	RIGVM_API FString GetPermutationName(const FRigVMTemplateTypeMap& InTypes, const bool bLockRegistry = true) const;
+	RIGVM_API FString GetPermutationName(const FRigVMTemplateTypeMap& InTypes) const;
+	RIGVM_API FString GetPermutationName_NoLock(const FRigVMTemplateTypeMap& InTypes, FRigVMRegistryHandle& InRegistry) const;
 
 	// returns true if the dispatch uses the same function ptr for all permutations
-	virtual bool IsSingleton() const { return false; } 
+	virtual bool IsSingleton() const { return false; }
 
 protected:
 
-	// returns the dispatch function for a given type set
-	RIGVM_API FRigVMFunctionPtr GetOrCreateDispatchFunction_NoLock(const FRigVMTemplateTypeMap& InTypes) const;
+	// creates the template based on existing argument infos
+	RIGVM_API const FRigVMTemplate* CreateTemplateForArgumentInfos_NoLock(const TArray<FRigVMTemplateArgumentInfo>& InArguments, FRigVMRegistryHandle& InRegistry);
 
-	// builds and returns the template
-	RIGVM_API const FRigVMTemplate* GetTemplate_RWLock() const;
-
-	// builds and returns the template
-	RIGVM_API const FRigVMTemplate* GetTemplate_NoLock() const;
+	// returns the dispatch function for a given type set
+	RIGVM_API FRigVMFunctionPtr GetOrCreateDispatchFunction_NoLock(const FRigVMTemplateTypeMap& InTypes, const FRigVMRegistryHandle& InRegistry) const;
 
 	// for each type defined in the primary argument, this function will call GetPermutationsFromArgumentType to construct an array of arguments with the appropiate permutations
-	RIGVM_API TArray<FRigVMTemplateArgumentInfo> BuildArgumentListFromPrimaryArgument(const TArray<FRigVMTemplateArgumentInfo>& InInfos, const FName& InPrimaryArgumentName) const;
+	RIGVM_API TArray<FRigVMTemplateArgumentInfo> BuildArgumentListFromPrimaryArgument(const TArray<FRigVMTemplateArgumentInfo>& InInfos, const FName& InPrimaryArgumentName, FRigVMRegistryHandle& InRegistry) const;
 
 	// returns the name of the permutation for a given set of types
-	RIGVM_API FString GetPermutationNameImpl(const FRigVMTemplateTypeMap& InTypes, const bool bLockRegistry = true) const;
+	RIGVM_API FString GetPermutationNameImpl(const FRigVMTemplateTypeMap& InTypes, const FRigVMRegistryHandle& InRegistry) const;
 
-	FRigVMFunctionPtr CreateDispatchFunction_NoLock(const FRigVMTemplateTypeMap& InTypes) const;
+	FRigVMFunctionPtr CreateDispatchFunction_NoLock(const FRigVMTemplateTypeMap& InTypes, const FRigVMRegistryHandle& InRegistry) const;
 
-	virtual FRigVMFunctionPtr GetDispatchFunctionImpl(const FRigVMTemplateTypeMap& InTypes) const { return nullptr; }
+	virtual FRigVMFunctionPtr GetDispatchFunctionImpl(const FRigVMTemplateTypeMap& InTypes, const FRigVMRegistryHandle& InRegistry) const { return nullptr; }
 	
-	TArray<FRigVMFunction> CreateDispatchPredicates_NoLock(const FRigVMTemplateTypeMap& InTypes) const;
+	TArray<FRigVMFunction> CreateDispatchPredicates_NoLock(const FRigVMTemplateTypeMap& InTypes, const FRigVMRegistryHandle& InRegistry) const;
 	
-	virtual TArray<FRigVMFunction> GetDispatchPredicatesImpl(const FRigVMTemplateTypeMap& InTypes) const { return TArray<FRigVMFunction>(); }
+	virtual TArray<FRigVMFunction> GetDispatchPredicatesImpl(const FRigVMTemplateTypeMap& InTypes, const FRigVMRegistryHandle& InRegistry) const { return TArray<FRigVMFunction>(); }
 
 	RIGVM_API virtual const TArray<FRigVMExecuteArgument>& GetExecuteArguments_Impl(const FRigVMDispatchContext& InContext) const;
 
@@ -264,14 +273,17 @@ protected:
 	static const inline TCHAR* DispatchPrefix = TEXT("DISPATCH_");
 	static const inline TCHAR* TrueString = TEXT("True");
 
-	const TArray<FName>* UpdateArgumentNameCache(int32 InNumberOperands) const;
-	const TArray<FName>* UpdateArgumentNameCache_NoLock(int32 InNumberOperands) const;
+	const TArray<FName>* UpdateArgumentNameCache(int32 InNumberOperands, FRigVMRegistryHandle& InRegistry) const;
+	const TArray<FName>* UpdateArgumentNameCache_NoLock(int32 InNumberOperands, FRigVMRegistryHandle& InRegistry) const;
 
+	int32 FactoryIndex;
 	UScriptStruct* FactoryScriptStruct;
 	mutable const FRigVMTemplate* CachedTemplate;
+	const FRigVMRegistry_NoLock* OwnerRegistry;
 	static FCriticalSection GetTemplateMutex;
 	mutable TMap<int32, TSharedPtr<TArray<FName>>> ArgumentNamesMap;
 	FCriticalSection* ArgumentNamesMutex;
+	mutable TArray<FRigVMTemplateArgumentInfo> CachedArgumentInfos;
 	friend struct FRigVMTemplate;
 	friend struct FRigVMRegistry_NoLock;
 	friend struct FRigVMFunction;
diff --git a/Engine/Plugins/Runtime/RigVM/Source/RigVM/Public/RigVMCore/RigVMFunction.h b/Engine/Plugins/Runtime/RigVM/Source/RigVM/Public/RigVMCore/RigVMFunction.h
--- a/Engine/Plugins/Runtime/RigVM/Source/RigVM/Public/RigVMCore/RigVMFunction.h
+++ b/Engine/Plugins/Runtime/RigVM/Source/RigVM/Public/RigVMCore/RigVMFunction.h
@@ -8,6 +8,7 @@
 #include "Containers/UnrealString.h"
 #include "CoreMinimal.h"
 #include "CoreTypes.h"
+#include "RigVMRegistryHandle.h"
 #include "RigVMExecuteContext.h"
 #include "RigVMMemory.h"
 #include "RigVMTypeIndex.h"
@@ -21,6 +22,7 @@ class UScriptStruct;
 struct FRigVMExtendedExecuteContext;
 struct FRigVMMemoryHandle;
 struct FRigVMPredicateBranch;
+struct FRigVMRegistry_NoLock;
 
 typedef TArrayView<FRigVMMemoryHandle> FRigVMMemoryHandleArray;
 typedef TArrayView<FRigVMPredicateBranch> FRigVMPredicateBranchArray;
@@ -92,6 +94,10 @@ struct FRigVMFunctionArgument
 		Name = NameString->operator*();
 		Type = TypeString->operator*();
 	}
+
+	void Serialize_NoLock(FArchive& Ar);
+
+	RIGVM_API bool operator ==(const FRigVMFunctionArgument& InOther) const;
 };
 
 typedef TArray<FRigVMFunctionArgument> FRigVMFunctionArgumentArray;
@@ -129,9 +135,11 @@ struct FRigVMFunction
 	FString Name;
 	UScriptStruct* Struct;
 	FRigVMDispatchFactory* Factory;
+	FRigVMRegistry_NoLock* OwnerRegistry;
 	FRigVMFunctionPtr FunctionPtr;
 	int32 Index;
 	int32 TemplateIndex;
+	int32 PermutationIndex;
 	TArray<FRigVMFunctionArgument> Arguments;
 	mutable TArray<TRigVMTypeIndex> ArgumentTypeIndices;
 	mutable TArray<FName> ArgumentNames;
@@ -140,9 +148,11 @@ struct FRigVMFunction
 		: Name()
 		, Struct(nullptr)
 		, Factory(nullptr)
+		, OwnerRegistry(nullptr)
 		, FunctionPtr(nullptr)
 		, Index(INDEX_NONE)
 		, TemplateIndex(INDEX_NONE)
+		, PermutationIndex(INDEX_NONE)
 		, Arguments()
 		, ArgumentTypeIndices()
 		, ArgumentNames()
@@ -155,9 +165,11 @@ struct FRigVMFunction
 	: Name(InName)
 	, Struct(nullptr)
 	, Factory(InFactory)
+	, OwnerRegistry(nullptr)
 	, FunctionPtr(InFunctionPtr)
 	, Index(InIndex)
 	, TemplateIndex(INDEX_NONE)
+	, PermutationIndex(INDEX_NONE)
 	, Arguments(InArguments)
 	, ArgumentNames()
 	{
@@ -170,9 +182,13 @@ struct FRigVMFunction
 	RIGVM_API FString GetModuleRelativeHeaderPath() const;
 	const TArray<FRigVMFunctionArgument>& GetArguments() const { return Arguments; }
 	RIGVM_API const TArray<TRigVMTypeIndex>& GetArgumentTypeIndices() const;
+	RIGVM_API const TArray<TRigVMTypeIndex>& GetArgumentTypeIndices_NoLock(FRigVMRegistryHandle& InRegistry) const;
 	RIGVM_API const FRigVMTemplate* GetTemplate() const;
-	RIGVM_API const FRigVMTemplate* GetTemplate_NoLock() const;
-	RIGVM_API const UScriptStruct* GetExecuteContextStruct(bool bLockRegistry = true) const;
+	RIGVM_API const FRigVMTemplate* GetTemplate_NoLock(const FRigVMRegistryHandle& InRegistry) const;
+	RIGVM_API const UScriptStruct* GetExecuteContextStruct() const;
+	RIGVM_API const UScriptStruct* GetExecuteContextStruct_NoLock(const FRigVMRegistryHandle& InRegistry) const;
 	RIGVM_API bool SupportsExecuteContextStruct(const UScriptStruct* InExecuteContextStruct) const;
-	RIGVM_API const FName& GetArgumentNameForOperandIndex(int32 InOperandIndex, int32 InTotalOperands) const;
+	RIGVM_API bool SupportsExecuteContextStruct_NoLock(const UScriptStruct* InExecuteContextStruct, const FRigVMRegistryHandle& InRegistry) const;
+	RIGVM_API const FName& GetArgumentNameForOperandIndex(int32 InOperandIndex, int32 InTotalOperands, FRigVMRegistryHandle& InRegistry) const;
+	RIGVM_API void Serialize_NoLock(FArchive& Ar, FRigVMRegistryHandle& InRegistry);
 };
